Copying //UE4/Dev-Editor to Dev-Main (//UE4/Dev-Main) for 4.24

diff --git a/Engine/Plugins/Experimental/MeshModelingToolset/Source/MeshModelingTools/Private/DrawPolygonTool.cpp b/Engine/Plugins/Experimental/MeshModelingToolset/Source/MeshModelingTools/Private/DrawPolygonTool.cpp
new file mode 100644
--- /dev/null
+++ b/Engine/Plugins/Experimental/MeshModelingToolset/Source/MeshModelingTools/Private/DrawPolygonTool.cpp
@@ -0,0 +1,1116 @@
+// Copyright 1998-2019 Epic Games, Inc. All Rights Reserved.
+
+#include "DrawPolygonTool.h"
+#include "InteractiveToolManager.h"
+#include "InteractiveGizmoManager.h"
+#include "ToolBuilderUtil.h"
+#include "BaseBehaviors/MultiClickSequenceInputBehavior.h"
+#include "BaseBehaviors/KeyAsModifierInputBehavior.h"
+
+#include "Polygon2.h"
+#include "FrameTypes.h"
+#include "MatrixTypes.h"
+#include "DynamicMeshAttributeSet.h"
+
+#include "MeshDescriptionBuilder.h"
+#include "Generators/PlanarPolygonMeshGenerator.h"
+#include "Operations/ExtrudeMesh.h"
+#include "Distance/DistLine3Ray3.h"
+#include "Intersection/IntrSegment2Segment2.h"
+#include "MeshQueries.h"
+#include "ToolSceneQueriesUtil.h"
+
+#include "DynamicMeshEditor.h"
+
+#include "BaseGizmos/GizmoComponents.h"
+#include "BaseGizmos/TransformGizmo.h"
+#include "Drawing/MeshDebugDrawing.h"
+
+#include "Selection/SelectClickedAction.h"
+#include "Selection/ToolSelectionUtil.h"
+#include "AssetGenerationUtil.h"
+
+#include "StaticMeshComponentBuilder.h"
+
+// [RMS] need these for now
+#include "Components/StaticMeshComponent.h"
+#include "Engine/StaticMesh.h"
+#include "Engine/StaticMeshActor.h"
+
+
+#define LOCTEXT_NAMESPACE "UDrawPolygonTool"
+
+
+/*
+ * ToolBuilder
+ */
+constexpr int StartPointSnapID = FPointPlanarSnapSolver::BaseExternalPointID + 1;
+constexpr int CurrentSceneSnapID = FPointPlanarSnapSolver::BaseExternalPointID + 2;
+
+bool UDrawPolygonToolBuilder::CanBuildTool(const FToolBuilderState& SceneState) const
+{
+	return (this->AssetAPI != nullptr);
+}
+
+UInteractiveTool* UDrawPolygonToolBuilder::BuildTool(const FToolBuilderState& SceneState) const
+{
+	UDrawPolygonTool* NewTool = NewObject<UDrawPolygonTool>(SceneState.ToolManager);
+	NewTool->SetWorld(SceneState.World);
+	NewTool->SetAssetAPI(AssetAPI);
+	return NewTool;
+}
+
+/*
+ * Properties
+ */
+UDrawPolygonToolStandardProperties::UDrawPolygonToolStandardProperties()
+{
+}
+
+
+void UDrawPolygonToolStandardProperties::SaveProperties(UInteractiveTool* SaveFromTool)
+{
+	UDrawPolygonToolStandardProperties* PropertyCache = GetPropertyCache<UDrawPolygonToolStandardProperties>();
+	PropertyCache->PolygonType = this->PolygonType;
+	PropertyCache->OutputMode = this->OutputMode;
+	PropertyCache->ExtrudeHeight = this->ExtrudeHeight;
+	PropertyCache->Steps = this->Steps;
+	PropertyCache->bAllowSelfIntersections = this->bAllowSelfIntersections;
+	PropertyCache->bShowGizmo = this->bShowGizmo;
+}
+
+void UDrawPolygonToolStandardProperties::RestoreProperties(UInteractiveTool* RestoreToTool)
+{
+	UDrawPolygonToolStandardProperties* PropertyCache = GetPropertyCache<UDrawPolygonToolStandardProperties>();
+	this->PolygonType = PropertyCache->PolygonType;
+	this->OutputMode = PropertyCache->OutputMode;
+	this->ExtrudeHeight = PropertyCache->ExtrudeHeight;
+	this->Steps = PropertyCache->Steps;
+	this->bAllowSelfIntersections = PropertyCache->bAllowSelfIntersections;
+	this->bShowGizmo = PropertyCache->bShowGizmo;
+}
+
+
+
+void UDrawPolygonToolSnapProperties::SaveProperties(UInteractiveTool* SaveFromTool)
+{
+	UDrawPolygonToolSnapProperties* PropertyCache = GetPropertyCache<UDrawPolygonToolSnapProperties>();
+	PropertyCache->bEnableSnapping = this->bEnableSnapping;
+	PropertyCache->bSnapToVertices = this->bSnapToVertices;
+	PropertyCache->bSnapToEdges = this->bSnapToEdges;
+	PropertyCache->bSnapToAngles = this->bSnapToAngles;
+	PropertyCache->bSnapToLengths = this->bSnapToLengths;
+	PropertyCache->bHitSceneObjects = this->bHitSceneObjects;
+	//PropertyCache->SegmentLength = this->Length;		// this is purely a feedback property
+	PropertyCache->HitNormalOffset = this->HitNormalOffset;
+}
+
+void UDrawPolygonToolSnapProperties::RestoreProperties(UInteractiveTool* RestoreToTool)
+{
+	UDrawPolygonToolSnapProperties* PropertyCache = GetPropertyCache<UDrawPolygonToolSnapProperties>();
+	this->bEnableSnapping = PropertyCache->bEnableSnapping;
+	this->bSnapToVertices = PropertyCache->bSnapToVertices;
+	this->bSnapToEdges = PropertyCache->bSnapToEdges;
+	this->bSnapToAngles = PropertyCache->bSnapToAngles;
+	this->bSnapToLengths = PropertyCache->bSnapToLengths;
+	this->bHitSceneObjects = PropertyCache->bHitSceneObjects;
+	//this->SegmentLength = PropertyCache->Length;
+	this->HitNormalOffset = PropertyCache->HitNormalOffset;
+}
+
+
+
+
+/*
+ * Tool
+ */
+UDrawPolygonTool::UDrawPolygonTool()
+{
+	DrawPlaneOrigin = FVector::ZeroVector;
+	DrawPlaneOrientation = FQuat::Identity;
+	bInInteractiveExtrude = false;
+}
+
+void UDrawPolygonTool::SetWorld(UWorld* World)
+{
+	this->TargetWorld = World;
+}
+
+void UDrawPolygonTool::SetAssetAPI(IToolsContextAssetAPI* AssetAPIIn)
+{
+	this->AssetAPI = AssetAPIIn;
+}
+
+void UDrawPolygonTool::Setup()
+{
+	UInteractiveTool::Setup();
+
+	// add default button input behaviors for devices
+	UMultiClickSequenceInputBehavior* MouseBehavior = NewObject<UMultiClickSequenceInputBehavior>();
+	MouseBehavior->Initialize(this);
+	MouseBehavior->Modifiers.RegisterModifier(IgnoreSnappingModifier, FInputDeviceState::IsShiftKeyDown);
+	AddInputBehavior(MouseBehavior);
+
+	// Register a click behavior/action pair, that sets the draw plane to the clicked world position
+	FSelectClickedAction* SetPlaneAction = new FSelectClickedAction();
+	SetPlaneAction->World = this->TargetWorld;
+	SetPlaneAction->OnClickedPositionFunc = [this](const FHitResult& Hit) {
+		SetDrawPlaneFromWorldPos(Hit.ImpactPoint, Hit.ImpactNormal);
+	};
+	SetPointInWorldConnector = SetPlaneAction;
+
+	USingleClickInputBehavior* ClickToSetPlaneBehavior = NewObject<USingleClickInputBehavior>();
+	ClickToSetPlaneBehavior->ModifierCheckFunc = FInputDeviceState::IsCtrlKeyDown;
+	ClickToSetPlaneBehavior->Initialize(SetPointInWorldConnector);
+	ClickToSetPlaneBehavior->SetDefaultPriority(MouseBehavior->GetPriority().MakeHigher());
+	AddInputBehavior(ClickToSetPlaneBehavior);
+
+	// register modifier key behaviors   (disabled because it is not implemented yet)
+	//UKeyAsModifierInputBehavior* AKeyBehavior = NewObject<UKeyAsModifierInputBehavior>();
+	//AKeyBehavior->Initialize(this, AngleSnapModifier, EKeys::A);
+	//AddInputBehavior(AKeyBehavior);
+
+
+	PolygonProperties = NewObject<UDrawPolygonToolStandardProperties>(this, TEXT("Polygon Settings"));
+	PolygonProperties->RestoreProperties(this);
+	ShowGizmoWatcher.Initialize(
+		[this]() { return this->PolygonProperties->bShowGizmo; }, 
+		[this](bool bNewValue) { this->UpdateShowGizmoState(bNewValue); },
+		true);
+
+	// Create a new TransformGizmo and associated TransformProxy. The TransformProxy will not be the
+	// parent of any Components in this case, we just use it's transform and change delegate.
+	PlaneTransformProxy = NewObject<UTransformProxy>(this);
+	PlaneTransformProxy->SetTransform(FTransform(DrawPlaneOrientation, DrawPlaneOrigin));
+	PlaneTransformGizmo = GetToolManager()->GetPairedGizmoManager()->Create3AxisTransformGizmo(this);
+	PlaneTransformGizmo->SetActiveTarget(PlaneTransformProxy, GetToolManager());
+	// listen for changes to the proxy and update the plane when that happens
+	PlaneTransformProxy->OnTransformChanged.AddUObject(this, &UDrawPolygonTool::PlaneTransformChanged);
+
+	// initialize material properties for new objects
+	MaterialProperties = NewObject<UNewMeshMaterialProperties>(this);
+	MaterialProperties->RestoreProperties(this);
+
+	// create preview mesh object
+	PreviewMesh = NewObject<UPreviewMesh>(this, TEXT("DrawPolygonPreviewMesh"));
+	PreviewMesh->CreateInWorld(this->TargetWorld, FTransform::Identity);
+	PreviewMesh->SetVisible(false);
+	PreviewMesh->SetMaterial(MaterialProperties->Material);
+	bPreviewUpdatePending = false;
+
+	// initialize snapping engine and properties
+	SnapEngine.SnapMetricTolerance = ToolSceneQueriesUtil::GetDefaultVisualAngleSnapThreshD();
+	SnapEngine.SnapMetricFunc = [this](const FVector3d& Position1, const FVector3d& Position2) {
+		return ToolSceneQueriesUtil::CalculateViewVisualAngleD(this->CameraState, Position1, Position2);
+	};
+	SnapEngine.Plane = FFrame3d((FVector3d)DrawPlaneOrigin, (FQuaterniond)DrawPlaneOrientation);
+
+	SnapProperties = NewObject<UDrawPolygonToolSnapProperties>(this, TEXT("Snapping"));
+	SnapProperties->RestoreProperties(this);
+
+	// register tool properties
+	AddToolPropertySource(PolygonProperties);
+	AddToolPropertySource(SnapProperties);
+	AddToolPropertySource(MaterialProperties);
+
+	ShowStartupMessage();
+}
+
+
+void UDrawPolygonTool::Shutdown(EToolShutdownType ShutdownType)
+{
+	PreviewMesh->Disconnect();
+	PreviewMesh = nullptr;
+
+	if (SetPointInWorldConnector != nullptr)
+	{
+		delete SetPointInWorldConnector;
+	}
+
+	GetToolManager()->GetPairedGizmoManager()->DestroyAllGizmosByOwner(this);
+
+	PolygonProperties->SaveProperties(this);
+	SnapProperties->SaveProperties(this);
+	MaterialProperties->SaveProperties(this);
+}
+
+void UDrawPolygonTool::RegisterActions(FInteractiveToolActionSet& ActionSet)
+{
+	ActionSet.RegisterAction(this, (int32)EStandardToolActions::BaseClientDefinedActionID + 1,
+		TEXT("PopLastVertex"), 
+		LOCTEXT("PopLastVertex", "Pop Last Vertex"),
+		LOCTEXT("PopLastVertexTooltip", "Pop last vertex added to polygon"),
+		EModifierKey::None, EKeys::BackSpace,
+		[this]() { PopLastVertexAction(); });
+
+
+	ActionSet.RegisterAction(this, (int32)EStandardToolActions::BaseClientDefinedActionID + 2,
+		TEXT("ToggleGizmo"),
+		LOCTEXT("ToggleGizmo", "Toggle Gizmo"),
+		LOCTEXT("ToggleGizmoTooltip", "Toggle visibility of the transformation Gizmo"),
+		EModifierKey::None, EKeys::A,
+		[this]() { PolygonProperties->bShowGizmo = !PolygonProperties->bShowGizmo; });
+
+}
+
+
+
+void UDrawPolygonTool::PopLastVertexAction()
+{
+	if (bInInteractiveExtrude || PolygonVertices.Num() == 0)
+	{
+		return;
+	}
+
+	if (bInFixedPolygonMode == false)
+	{
+		int NumVertices = PolygonVertices.Num();
+		if (NumVertices > 1)
+		{
+			PolygonVertices.RemoveAt(NumVertices-1);
+		}
+		else
+		{
+			PolygonVertices.RemoveAt(0);
+			bAbortActivePolygonDraw = true;
+		}
+	}
+	else
+	{
+		int NumVertices = FixedPolygonClickPoints.Num();
+		if (NumVertices > 1)
+		{
+			FixedPolygonClickPoints.RemoveAt(NumVertices - 1);
+		}
+		else
+		{
+			FixedPolygonClickPoints.RemoveAt(0);
+			bAbortActivePolygonDraw = true;
+		}
+	}
+}
+
+void DrawEdgeTicks(FPrimitiveDrawInterface* PDI, 
+	const FSegment3f& Segment, float Height,
+	const FVector3f& PlaneNormal, 
+	const FLinearColor& Color, uint8 DepthPriorityGroup, float LineThickness, bool bIsScreenSpace)
+{
+	FVector3f Center = Segment.Center;
+	FVector3f X = Segment.Direction;
+	FVector3f Y = X.Cross(PlaneNormal);
+	Y.Normalize();
+	FVector A = Center - Height * 0.25f*X - Height * Y;
+	FVector B = Center + Height * 0.25f*X + Height * Y;
+	PDI->DrawLine(A, B, Color, DepthPriorityGroup, LineThickness, 0.0f, bIsScreenSpace);
+	A += Height * 0.5f*X;
+	B += Height * 0.5f*X;
+	PDI->DrawLine(A, B, Color, DepthPriorityGroup, LineThickness, 0.0f, bIsScreenSpace);
+}
+
+void UDrawPolygonTool::Render(IToolsContextRenderAPI* RenderAPI)
+{
+	FPrimitiveDrawInterface* PDI = RenderAPI->GetPrimitiveDrawInterface();
+	GetToolManager()->GetContextQueriesAPI()->GetCurrentViewState(CameraState);
+
+	if (bPreviewUpdatePending)
+	{
+		UpdateLivePreview();
+		bPreviewUpdatePending = false;
+	}
+
+	double CurViewSizeFactor = ToolSceneQueriesUtil::CalculateDimensionFromVisualAngleD(CameraState, PreviewVertex, 1.0);
+
+	FColor OpenPolygonColor(240, 16, 240);
+	FColor ClosedPolygonColor(16, 240, 16);
+	FColor ErrorColor(240, 16, 16);
+	float HiddenLineThickness = 1.0f;
+	float LineThickness = 4.0f;
+	float SelfIntersectThickness = 8.0f;
+	FColor GridColor(128, 128, 128, 32);
+	float GridThickness = 0.5f;
+	float GridLineSpacing = 25.0f;   // @todo should be relative to view
+	int NumGridLines = 21;
+	FColor SnapHighlightColor(240, 200, 16);
+	float ElementSize = CurViewSizeFactor;
+
+	bool bIsClosed = SnapEngine.HaveActiveSnap() && SnapEngine.GetActiveSnapTargetID() == StartPointSnapID;
+
+	if (bInInteractiveExtrude == false)
+	{
+		FFrame3f DrawFrame(DrawPlaneOrigin, DrawPlaneOrientation);
+		MeshDebugDraw::DrawSimpleGrid(DrawFrame, NumGridLines, GridLineSpacing, GridThickness, GridColor, false, PDI, FTransform::Identity);
+	}
+
+	if (bInFixedPolygonMode)
+	{
+		if (bInInteractiveExtrude == false)		// once we are in extrude, polygon is done
+		{
+			FixedPolygonClickPoints.Add(PreviewVertex);
+			GenerateFixedPolygon(FixedPolygonClickPoints, PolygonVertices);
+			FixedPolygonClickPoints.Pop(false);
+		}
+		bIsClosed = true;
+	}
+
+	int NumVerts = PolygonVertices.Num();
+
+	if (SnapEngine.HaveActiveSnap())
+	{
+		PDI->DrawPoint((FVector)SnapEngine.GetActiveSnapToPoint(), ClosedPolygonColor, 10, SDPG_Foreground);
+		
+		PDI->DrawPoint((FVector)SnapEngine.GetActiveSnapFromPoint(), OpenPolygonColor, 15, SDPG_Foreground);
+		PDI->DrawLine((FVector)SnapEngine.GetActiveSnapToPoint(), (FVector)SnapEngine.GetActiveSnapFromPoint(),
+			ClosedPolygonColor, SDPG_Foreground, 0.5f, 0.0f, true);
+		if (SnapEngine.GetActiveSnapTargetID() == CurrentSceneSnapID)
+		{
+			if (LastSnapGeometry.PointCount == 1) {
+				DrawCircle(PDI, (FVector)LastSnapGeometry.Points[0], CameraState.Right(), CameraState.Up(),
+					SnapHighlightColor, ElementSize, 32, SDPG_Foreground, 1.0f, 0.0f, true);
+			} 
+			else
+			{
+				PDI->DrawLine((FVector)LastSnapGeometry.Points[0], (FVector)LastSnapGeometry.Points[1],
+					SnapHighlightColor, SDPG_Foreground, 1.0f, 0.0f, true);
+			}
+		}
+
+		if (SnapEngine.HaveActiveSnapLine())
+		{
+			FLine3d SnapLine = SnapEngine.GetActiveSnapLine();
+			PDI->DrawLine((FVector)SnapLine.PointAt(-9999), (FVector)SnapLine.PointAt(9999),
+				ClosedPolygonColor, SDPG_Foreground, 0.5, 0.0f, true);
+
+			if (SnapEngine.HaveActiveSnapDistance())
+			{
+				int iSegment = SnapEngine.GetActiveSnapDistanceID();
+				TArray<FVector>& HistoryPoints = (bInFixedPolygonMode) ? FixedPolygonClickPoints : PolygonVertices;
+				FVector UseNormal = DrawPlaneOrientation.GetAxisZ();
+				DrawEdgeTicks(PDI, FSegment3f(HistoryPoints[iSegment], HistoryPoints[iSegment+1]),
+					0.75f*ElementSize, UseNormal, SnapHighlightColor, SDPG_Foreground, 1.0f, true);
+				DrawEdgeTicks(PDI, FSegment3f(HistoryPoints[HistoryPoints.Num()-1], PreviewVertex),
+					0.75f*ElementSize, UseNormal, SnapHighlightColor, SDPG_Foreground, 1.0f, true);
+				PDI->DrawLine(HistoryPoints[iSegment], HistoryPoints[iSegment + 1],
+					SnapHighlightColor, SDPG_Foreground, 2.0f, 0.0f, true);
+			}
+		}
+	}
+
+
+	if (bHaveSurfaceHit)
+	{
+		PDI->DrawPoint((FVector)SurfaceHitPoint, ClosedPolygonColor, 10, SDPG_Foreground);
+		if (SnapProperties->HitNormalOffset != 0)
+		{
+			PDI->DrawPoint((FVector)SurfaceOffsetPoint, OpenPolygonColor, 15, SDPG_Foreground);
+			PDI->DrawLine((FVector)SurfaceOffsetPoint, (FVector)SurfaceHitPoint,
+				ClosedPolygonColor, SDPG_Foreground, 0.5f, 0.0f, true);
+		}
+		PDI->DrawLine((FVector)SurfaceOffsetPoint, (FVector)PreviewVertex,
+			ClosedPolygonColor, SDPG_Foreground, 0.5f, 0.0f, true);
+	}
+
+
+	if (PolygonVertices.Num() > 0)
+	{
+		FColor UseColor = (bIsClosed) ? ClosedPolygonColor : OpenPolygonColor;
+		FVector UseLastVertex = (bIsClosed) ? PolygonVertices[0] : PreviewVertex;
+		float UseThickness = LineThickness;
+		if (bHaveSelfIntersection)
+		{
+			UseColor = ErrorColor;
+			UseThickness = SelfIntersectThickness;
+		}
+
+		// draw thin no-depth
+		for (int i = 0; i < NumVerts - 1; ++i)
+		{
+			PDI->DrawLine(PolygonVertices[i], PolygonVertices[i + 1],
+				UseColor, SDPG_Foreground, HiddenLineThickness, 0.0f, true);
+		}
+		PDI->DrawLine(PolygonVertices[NumVerts - 1], UseLastVertex,
+			UseColor, SDPG_Foreground, HiddenLineThickness, 0.0f, true);
+
+
+		// draw thick depth-tested
+		for (int i = 0; i < NumVerts-1; ++i)
+		{
+			PDI->DrawLine(PolygonVertices[i], PolygonVertices[i+1],
+				UseColor, SDPG_World, LineThickness, 0.0f, true);
+		}
+		PDI->DrawLine(PolygonVertices[NumVerts-1], UseLastVertex,
+			UseColor, SDPG_World, LineThickness, 0.0f, true);
+
+		if (bHaveSelfIntersection)
+		{
+			PDI->DrawPoint(SelfIntersectionPoint, ErrorColor, 10, SDPG_Foreground);
+		}
+	}
+
+	// draw preview vertex
+	PDI->DrawPoint(PreviewVertex, ClosedPolygonColor, 10, SDPG_Foreground);
+
+
+	// todo should be an indicator
+	if (bInInteractiveExtrude)
+	{
+		float Length = 10; float Thickness = 2;
+		FColor HitFrameColor(0, 128, 128);
+		PDI->DrawLine(
+			HitPosFrameWorld.PointAt(-Length, -Length, 0), HitPosFrameWorld.PointAt(Length, Length, 0),
+			HitFrameColor, 1, Thickness, 0.0f, true);
+		PDI->DrawLine(
+			HitPosFrameWorld.PointAt(-Length, Length, 0), HitPosFrameWorld.PointAt(Length, -Length, 0),
+			HitFrameColor, 1, Thickness, 0.0f, true);
+
+		FVector PreviewOrigin = (FVector)PreviewHeightFrame.Origin;
+
+		FVector DrawPlaneNormal = DrawPlaneOrientation.GetAxisZ();
+
+		FColor AxisColor(128, 128, 0);
+		PDI->DrawLine(
+			PreviewOrigin -1000*DrawPlaneNormal, PreviewOrigin +1000*DrawPlaneNormal,
+			AxisColor, 1, 1.0f, 0.0f, true);
+
+		FColor HeightPosColor(128, 0, 128);
+		PDI->DrawLine(
+			PreviewOrigin + PolygonProperties->ExtrudeHeight*DrawPlaneNormal, HitPosFrameWorld.Origin,
+			HeightPosColor, 1, 1.0f, 0.0f, true);
+	}
+
+
+	ShowGizmoWatcher.CheckAndUpdate();
+}
+
+
+void UDrawPolygonTool::ResetPolygon()
+{
+	PolygonVertices.Reset();
+	SnapEngine.Reset();
+	bHaveSurfaceHit = false;
+	bInFixedPolygonMode = false;
+}
+
+void UDrawPolygonTool::UpdatePreviewVertex(const FVector& PreviewVertexIn)
+{
+	PreviewVertex = PreviewVertexIn;
+
+	// update length and angle
+	if (PolygonVertices.Num() > 0)
+	{
+		FVector LastVertex = PolygonVertices[PolygonVertices.Num() - 1];
+		SnapProperties->SegmentLength = FVector::Distance(LastVertex, PreviewVertex);
+	}
+}
+
+void UDrawPolygonTool::AppendVertex(const FVector& Vertex)
+{
+	PolygonVertices.Add(Vertex);
+}
+
+bool UDrawPolygonTool::FindDrawPlaneHitPoint(const FInputDeviceRay& ClickPos, FVector& HitPosOut)
+{
+	bHaveSurfaceHit = false;
+
+	FFrame3d Frame(DrawPlaneOrigin, DrawPlaneOrientation);
+	FVector3d HitPos;
+	bool bHit = Frame.RayPlaneIntersection(ClickPos.WorldRay.Origin, ClickPos.WorldRay.Direction, 2, HitPos);
+	if (bHit == false)
+	{
+		return false;
+	}
+
+	// if we found a scene snap point, add to snap set
+	FVector3d SnapPos;
+	if (bIgnoreSnappingToggle || SnapProperties->bEnableSnapping == false)
+	{
+		SnapEngine.ResetActiveSnap();
+		SnapEngine.UpdatePointHistory(TArray<FVector>());
+	}
+	else 
+	{
+		if (ToolSceneQueriesUtil::FindSceneSnapPoint(this, HitPos, SnapPos, SnapProperties->bSnapToVertices, SnapProperties->bSnapToEdges, 0, &LastSnapGeometry))
+		{
+			SnapEngine.AddPointTarget(SnapPos, CurrentSceneSnapID, SnapEngine.MinInternalPriority()-1 );
+		}
+
+		TArray<FVector>& HistoryPoints = (bInFixedPolygonMode) ? FixedPolygonClickPoints : PolygonVertices;
+		SnapEngine.UpdatePointHistory(HistoryPoints);
+		if (SnapProperties->bSnapToAngles)
+		{
+			SnapEngine.RegenerateTargetLines(true, true);
+		}
+		SnapEngine.bEnableSnapToKnownLengths = SnapProperties->bSnapToLengths;
+	}
+
+	SnapEngine.UpdateSnappedPoint(HitPos);
+
+	// remove scene snap point
+	SnapEngine.RemovePointTargetsByID(CurrentSceneSnapID);
+
+	if (SnapEngine.HaveActiveSnap())
+	{
+		HitPosOut = (FVector)SnapEngine.GetActiveSnapToPoint();
+		return true;
+	}
+
+
+	// if not snap and we want to hit objects, do that
+	if (SnapProperties->bHitSceneObjects)
+	{
+		FCollisionObjectQueryParams QueryParams(FCollisionObjectQueryParams::AllObjects);
+		FHitResult Result;
+		bool bWorldHit = TargetWorld->LineTraceSingleByObjectType(Result, ClickPos.WorldRay.Origin, ClickPos.WorldRay.PointAt(9999), QueryParams);
+		if (bWorldHit)
+		{
+			bHaveSurfaceHit = true;
+			SurfaceHitPoint = Result.ImpactPoint;
+			FVector UseHitPos = Result.ImpactPoint + SnapProperties->HitNormalOffset*Result.Normal;
+			HitPos = Frame.ToPlane((FVector3d)UseHitPos, 2);
+			SurfaceOffsetPoint = UseHitPos;
+		}
+	}
+
+	HitPosOut = (FVector)HitPos;
+	return true;
+}
+
+void UDrawPolygonTool::OnBeginSequencePreview(const FInputDeviceRay& DevicePos)
+{
+	// just update snapped point preview
+	FVector HitPos;
+	if (FindDrawPlaneHitPoint(DevicePos, HitPos))
+	{
+		PreviewVertex = HitPos;
+	}
+	
+}
+
+bool UDrawPolygonTool::CanBeginClickSequence(const FInputDeviceRay& ClickPos)
+{
+	return true;
+}
+
+void UDrawPolygonTool::OnBeginClickSequence(const FInputDeviceRay& ClickPos)
+{
+	ResetPolygon();
+	
+	FVector HitPos;
+	bool bHit = FindDrawPlaneHitPoint(ClickPos, HitPos);
+	if (bHit == false)
+	{
+		bAbortActivePolygonDraw = true;
+		return;
+	}
+	if (ToolSceneQueriesUtil::IsPointVisible(CameraState, HitPos) == false)
+	{
+		bAbortActivePolygonDraw = true;
+		return;		// cannot start a poly an a point that is not visible, this is almost certainly an error due to draw plane
+	}
+
+	AppendVertex(HitPos);
+	UpdatePreviewVertex(HitPos);
+
+	bInFixedPolygonMode = (PolygonProperties->PolygonType != EDrawPolygonDrawMode::Freehand);
+	FixedPolygonClickPoints.Reset();
+	FixedPolygonClickPoints.Add(HitPos);
+
+	// if we are starting a freehand poly, add start point as snap target, but then ignore it until we get 3 verts
+	if (bInFixedPolygonMode == false)
+	{
+		SnapEngine.AddPointTarget(PolygonVertices[0], StartPointSnapID, 1);
+		SnapEngine.AddIgnoreTarget(StartPointSnapID);
+	}
+}
+
+void UDrawPolygonTool::OnNextSequencePreview(const FInputDeviceRay& ClickPos)
+{
+	if (bInInteractiveExtrude)
+	{
+		PolygonProperties->ExtrudeHeight = FindInteractiveHeightDistance(ClickPos);
+		bPreviewUpdatePending = true;
+		return;
+	}
+
+	FVector HitPos;
+	bool bHit = FindDrawPlaneHitPoint(ClickPos, HitPos);
+	if (bHit == false)
+	{
+		return;
+	}
+
+	if (bInFixedPolygonMode)
+	{
+		UpdatePreviewVertex(HitPos);
+		bPreviewUpdatePending = true;
+		return;
+	}
+
+	if (PolygonVertices.Num() > 2)
+	{
+		bPreviewUpdatePending = true;
+	}
+
+	UpdatePreviewVertex(HitPos);
+	UpdateSelfIntersection();
+}
+
+bool UDrawPolygonTool::OnNextSequenceClick(const FInputDeviceRay& ClickPos)
+{
+	if (bInInteractiveExtrude)
+	{
+		EndInteractiveExtrude();
+		return false;
+	}
+
+	FVector HitPos;
+	bool bHit = FindDrawPlaneHitPoint(ClickPos, HitPos);
+	if (bHit == false)
+	{
+		return true;  // ignore click but continue accepting clicks
+	}
+
+	bool bDonePolygon = false;
+	if (bInFixedPolygonMode)
+	{
+		// ignore very close click points
+		if ( ToolSceneQueriesUtil::PointSnapQuery(this, FixedPolygonClickPoints[FixedPolygonClickPoints.Num()-1], HitPos) )
+		{
+			return true;
+		}
+
+		FixedPolygonClickPoints.Add(HitPos);
+		int NumTargetPoints = (PolygonProperties->PolygonType == EDrawPolygonDrawMode::Rectangle) ? 3 : 2;
+		bDonePolygon = (FixedPolygonClickPoints.Num() == NumTargetPoints);
+		if (bDonePolygon)
+		{
+			GenerateFixedPolygon(FixedPolygonClickPoints, PolygonVertices);
+		}
+	} 
+	else
+	{
+		// ignore very close click points
+		if (ToolSceneQueriesUtil::PointSnapQuery(this, PolygonVertices[PolygonVertices.Num()-1], HitPos))
+		{
+			return true;
+		}
+
+		// close polygon if we clicked on start point
+		bDonePolygon = SnapEngine.HaveActiveSnap() && SnapEngine.GetActiveSnapTargetID() == StartPointSnapID;
+
+		if (bHaveSelfIntersection)
+		{
+			// discard vertex in segments before intersection (this is redundant if idx is 0)
+			for (int j = SelfIntersectSegmentIdx; j < PolygonVertices.Num(); ++j)
+			{
+				PolygonVertices[j-SelfIntersectSegmentIdx] = PolygonVertices[j];
+			}
+			PolygonVertices.SetNum(PolygonVertices.Num() - SelfIntersectSegmentIdx);
+			PolygonVertices[0] = PreviewVertex = SelfIntersectionPoint;
+			bDonePolygon = true;
+		}
+	}
+
+	
+	if (bDonePolygon)
+	{
+		SnapEngine.Reset();
+		bHaveSurfaceHit = false;
+		if (PolygonProperties->OutputMode == EDrawPolygonOutputMode::ExtrudedInteractive)
+		{
+			BeginInteractiveExtrude();
+
+			PreviewMesh->ClearPreview();
+			PreviewMesh->SetVisible(true);
+
+			return true;
+		}
+		else 
+		{
+			EmitCurrentPolygon();
+
+			PreviewMesh->ClearPreview();
+			PreviewMesh->SetVisible(false);
+
+			return false;
+		}
+	}
+
+	AppendVertex(HitPos);
+	if (PolygonVertices.Num() > 2)
+	{
+		SnapEngine.RemoveIgnoreTarget(StartPointSnapID);
+	}
+
+	UpdatePreviewVertex(HitPos);
+	return true;
+}
+
+void UDrawPolygonTool::OnTerminateClickSequence()
+{
+	ResetPolygon();
+}
+
+bool UDrawPolygonTool::RequestAbortClickSequence()
+{
+	if (bAbortActivePolygonDraw)
+	{
+		bAbortActivePolygonDraw = false;
+		return true;
+	}
+	return false;
+}
+
+void UDrawPolygonTool::OnUpdateModifierState(int ModifierID, bool bIsOn)
+{
+	if (ModifierID == IgnoreSnappingModifier)
+	{
+		bIgnoreSnappingToggle = bIsOn;
+	}
+	else if (ModifierID == AngleSnapModifier)
+	{
+
+	}
+}
+
+bool UDrawPolygonTool::UpdateSelfIntersection()
+{
+	bHaveSelfIntersection = false;
+	if (bInFixedPolygonMode || PolygonProperties->bAllowSelfIntersections == true)
+	{
+		return false;
+	}
+
+	int NumVertices = PolygonVertices.Num();
+
+	FFrame3f DrawFrame(DrawPlaneOrigin, DrawPlaneOrientation);
+	FSegment2f PreviewSegment(DrawFrame.ToPlaneUV(PolygonVertices[NumVertices - 1],2), DrawFrame.ToPlaneUV(PreviewVertex,2));
+
+	for (int k = 0; k < NumVertices - 2; ++k) 
+	{
+		FSegment2f Segment(DrawFrame.ToPlaneUV(PolygonVertices[k],2), DrawFrame.ToPlaneUV(PolygonVertices[k + 1],2));
+		FIntrSegment2Segment2f Intersection(PreviewSegment, Segment);
+		if (Intersection.Find()) 
+		{
+			bHaveSelfIntersection = true;
+			SelfIntersectSegmentIdx = k;
+			SelfIntersectionPoint = DrawFrame.FromPlaneUV(Intersection.Point0,2);
+			return true;
+		}
+	}
+	return false;
+}
+
+void UDrawPolygonTool::GenerateFixedPolygon(TArray<FVector>& FixedPoints, TArray<FVector>& VerticesOut)
+{
+	FFrame3f DrawFrame(DrawPlaneOrigin, DrawPlaneOrientation);
+	FVector2f CenterPt = DrawFrame.ToPlaneUV(FixedPoints[0], 2 );
+	FVector2f EdgePt = DrawFrame.ToPlaneUV(FixedPoints[1], 2);
+	FVector2f Delta = EdgePt - CenterPt;
+	float AngleRad = FMath::Atan2(Delta.Y, Delta.X);
+	FMatrix2f RotationMat = FMatrix2f::RotationRad(AngleRad);
+	FVector2f RotAxisX = RotationMat * FVector2f::UnitX();
+	float Dist = Delta.Length();
+	float Width = FMath::Abs(Delta.Dot(RotAxisX));
+
+	FPolygon2f Polygon;
+	if (PolygonProperties->PolygonType == EDrawPolygonDrawMode::Square)
+	{
+		Polygon = FPolygon2f::MakeRectangle(FVector2f::Zero(), 2*Width, 2*Width);
+	}
+	else if (PolygonProperties->PolygonType == EDrawPolygonDrawMode::Rectangle)
+	{
+		FVector2f HeightPt = DrawFrame.ToPlaneUV((FixedPoints.Num() == 3) ? FixedPoints[2] : FixedPoints[1], 2);
+		FVector2f HeightDelta = HeightPt - CenterPt;
+		FVector2f RotAxisY = RotationMat * FVector2f::UnitY();
+		float YSign = FMath::Sign(HeightDelta.Dot(RotAxisY));
+		float Height = FMath::Abs(HeightDelta.Dot(RotAxisY));
+		Polygon = FPolygon2f::MakeRectangle(FVector2f(Width/2, YSign*Height/2), Width, Height);
+	}
+	else
+	{
+		Polygon = FPolygon2f::MakeCircle(Dist, PolygonProperties->Steps, 0);
+	}
+	Polygon.Transform([RotationMat](const FVector2f& Pt) { return RotationMat * Pt; });
+
+	VerticesOut.Reset();
+	for (int k = 0; k < Polygon.VertexCount(); ++k)
+	{
+		FVector2f NewPt = CenterPt + Polygon[k];
+		VerticesOut.Add(DrawFrame.FromPlaneUV(NewPt, 2));
+	}
+}
+
+void UDrawPolygonTool::BeginInteractiveExtrude()
+{
+	bInInteractiveExtrude = true;
+
+	GeneratePreviewHeightTarget();
+
+	ShowExtrudeMessage();
+}
+
+void UDrawPolygonTool::EndInteractiveExtrude()
+{
+	EmitCurrentPolygon();
+
+	PreviewMesh->ClearPreview();
+	PreviewMesh->SetVisible(false);
+
+	bInInteractiveExtrude = false;
+
+	ShowStartupMessage();
+}
+
+float UDrawPolygonTool::FindInteractiveHeightDistance(const FInputDeviceRay& ClickPos)
+{
+	float NearestHitDist = TNumericLimits<float>::Max();
+	float NearestHitHeight = 1.0f;
+	FFrame3f NearestHitFrameWorld;
+
+	// cast ray at target object
+	FRay3d LocalRay = PreviewHeightFrame.ToFrame((FRay3d)ClickPos.WorldRay);
+	int HitTID = PreviewHeightTargetAABB.FindNearestHitTriangle(LocalRay);
+	if (HitTID >= 0)
+	{
+		FIntrRay3Triangle3d IntrQuery =
+			TMeshQueries<FDynamicMesh3>::TriangleIntersection(PreviewHeightTarget, HitTID, LocalRay);
+		FVector3d HitPosLocal = LocalRay.PointAt(IntrQuery.RayParameter);
+		FVector3d HitNormalLocal = PreviewHeightTarget.GetTriNormal(HitTID);
+
+		NearestHitFrameWorld = FFrame3f(
+			(FVector3f)PreviewHeightFrame.FromFramePoint(HitPosLocal),
+			(FVector3f)PreviewHeightFrame.FromFrameVector(HitNormalLocal));
+		NearestHitHeight = HitPosLocal.Z;
+		NearestHitDist = ClickPos.WorldRay.GetParameter(NearestHitFrameWorld.Origin);
+	}
+	
+
+	// cast ray into scene
+	FVector RayStart = ClickPos.WorldRay.Origin;
+	FVector RayEnd = ClickPos.WorldRay.PointAt(999999);
+	FCollisionObjectQueryParams QueryParams(FCollisionObjectQueryParams::AllObjects);
+	FHitResult Result;
+	bool bHitWorld = TargetWorld->LineTraceSingleByObjectType(Result, RayStart, RayEnd, QueryParams);
+	if (bHitWorld)
+	{
+		float WorldHitDist = ClickPos.WorldRay.GetParameter(Result.ImpactPoint);
+		if (WorldHitDist < NearestHitDist)
+		{
+			NearestHitFrameWorld = FFrame3f(Result.ImpactPoint, Result.ImpactNormal);
+			FVector3d HitPosWorld = Result.ImpactPoint;
+			FVector3d HitPosLocal = PreviewHeightFrame.ToFramePoint(HitPosWorld);
+			NearestHitHeight = HitPosLocal.Z;
+			NearestHitDist = WorldHitDist;
+		}
+	}
+
+	if (NearestHitDist < TNumericLimits<float>::Max())
+	{
+		this->HitPosFrameWorld = NearestHitFrameWorld;
+		return NearestHitHeight;
+	}
+	else
+	{
+		return PolygonProperties->ExtrudeHeight;
+	}
+	
+}
+
+void UDrawPolygonTool::SetDrawPlaneFromWorldPos(const FVector& Position, const FVector& Normal)
+{
+	DrawPlaneOrigin = Position;
+
+	FFrame3f DrawPlane(Position, DrawPlaneOrientation);
+	if (bIgnoreSnappingToggle == false)
+	{
+		DrawPlane.AlignAxis(2, Normal);
+		DrawPlane.ConstrainedAlignPerpAxes();
+		DrawPlaneOrientation = DrawPlane.Rotation;
+	}
+
+	SnapEngine.Plane = FFrame3d((FVector3d)DrawPlane.Origin, (FQuaterniond)DrawPlane.Rotation);
+
+	PlaneTransformGizmo->SetNewGizmoTransform(FTransform(DrawPlaneOrientation, DrawPlaneOrigin));
+}
+
+
+void UDrawPolygonTool::PlaneTransformChanged(UTransformProxy* Proxy, FTransform Transform)
+{
+	DrawPlaneOrientation = Transform.GetRotation();
+	DrawPlaneOrigin = Transform.GetLocation();
+	SnapEngine.Plane = FFrame3d((FVector3d)DrawPlaneOrigin, (FQuaterniond)DrawPlaneOrientation);
+}
+
+void UDrawPolygonTool::UpdateShowGizmoState(bool bNewVisibility)
+{
+	if (bNewVisibility == false)
+	{
+		GetToolManager()->GetPairedGizmoManager()->DestroyAllGizmosByOwner(this);
+	}
+	else
+	{
+		PlaneTransformGizmo = GetToolManager()->GetPairedGizmoManager()->Create3AxisTransformGizmo(this);
+		PlaneTransformGizmo->SetActiveTarget(PlaneTransformProxy, GetToolManager());
+	}
+}
+
+
+void UDrawPolygonTool::EmitCurrentPolygon()
+{
+	FString BaseName = (PolygonProperties->OutputMode == EDrawPolygonOutputMode::MeshedPolygon) ?
+		TEXT("Polygon") : TEXT("Extrude");
+
+#if WITH_EDITOR
+	// generate new mesh
+	FFrame3d PlaneFrameOut;
+	FDynamicMesh3 Mesh;
+	double ExtrudeDist = (PolygonProperties->OutputMode == EDrawPolygonOutputMode::MeshedPolygon) ?
+		0 : PolygonProperties->ExtrudeHeight;
+	GeneratePolygonMesh(PolygonVertices, &Mesh, PlaneFrameOut, false, ExtrudeDist, false);
+
+	GetToolManager()->BeginUndoTransaction(LOCTEXT("CreatePolygon", "Create Polygon"));
+
+	AActor* NewActor = AssetGenerationUtil::GenerateStaticMeshActor(
+		AssetAPI, TargetWorld,
+		&Mesh, PlaneFrameOut.ToTransform(), BaseName,
+		AssetGenerationUtil::GetDefaultAutoGeneratedAssetPath(),
+		MaterialProperties->Material);
+
+	// select newly-created object
+	ToolSelectionUtil::SetNewActorSelection(GetToolManager(), NewActor);
+
+	GetToolManager()->EndUndoTransaction();
+#else
+	checkNoEntry();
+#endif
+	ResetPolygon();
+}
+
+void UDrawPolygonTool::UpdateLivePreview()
+{
+	int NumVerts = PolygonVertices.Num();
+	if (NumVerts < 2 || PreviewMesh == nullptr || PreviewMesh->IsVisible() == false )
+	{
+		return;
+	}
+
+	FFrame3d PlaneFrame;
+	FDynamicMesh3 Mesh;
+	double ExtrudeDist = (PolygonProperties->OutputMode == EDrawPolygonOutputMode::MeshedPolygon) ?
+		0 : PolygonProperties->ExtrudeHeight;
+	GeneratePolygonMesh(PolygonVertices, &Mesh, PlaneFrame, false, ExtrudeDist, false);
+
+	PreviewMesh->SetTransform(PlaneFrame.ToFTransform());
+	PreviewMesh->SetMaterial(MaterialProperties->Material);
+	PreviewMesh->EnableWireframe(MaterialProperties->bWireframe);
+	PreviewMesh->UpdatePreview(&Mesh);
+}
+
+void UDrawPolygonTool::GeneratePolygonMesh(const TArray<FVector>& Polygon, FDynamicMesh3* ResultMeshOut, FFrame3d& WorldFrameOut, bool bIncludePreviewVtx, double ExtrudeDistance, bool bExtrudeSymmetric)
+{
+	// construct centered frame for polygon
+	WorldFrameOut = FFrame3d(DrawPlaneOrigin, DrawPlaneOrientation);
+
+	int NumVerts = Polygon.Num();
+	FVector3d Centroid(0, 0, 0);
+	for (int k = 0; k < NumVerts; ++k)
+	{
+		Centroid += Polygon[k];
+	}
+	Centroid /= (double)NumVerts;
+	WorldFrameOut.Origin = Centroid;
+
+
+	// triangulate polygon into the MeshDescription
+	FPlanarPolygonMeshGenerator PolyMeshGen;
+	for (int k = 0; k < NumVerts; ++k)
+	{
+		PolyMeshGen.Polygon.AppendVertex(WorldFrameOut.ToPlaneUV(Polygon[k], 2));
+	}
+
+	// add preview vertex
+	if (bIncludePreviewVtx)
+	{
+		if (FVector::Dist(PreviewVertex, Polygon[NumVerts - 1]) > 0.1)
+		{
+			PolyMeshGen.Polygon.AppendVertex(WorldFrameOut.ToPlaneUV(PreviewVertex, 2));
+		}
+	}
+
+	if (PolyMeshGen.Polygon.IsClockwise() == false)
+	{
+		PolyMeshGen.Polygon.Reverse();
+	}
+
+	ResultMeshOut->Copy(&PolyMeshGen.Generate());
+
+	// for symmetric extrude we translate the first poly by -dist along axis
+	if (bExtrudeSymmetric)
+	{
+		FVector3d ShiftNormal = FVector3d::UnitZ();
+		for (int vid : ResultMeshOut->VertexIndicesItr())
+		{
+			FVector3d Pos = ResultMeshOut->GetVertex(vid);
+			ResultMeshOut->SetVertex(vid, Pos - ExtrudeDistance * ShiftNormal);
+		}
+		// double extrude dist
+		ExtrudeDistance *= 2.0;
+	}
+
+	if (ExtrudeDistance != 0)
+	{
+		FExtrudeMesh Extruder(ResultMeshOut);
+		Extruder.DefaultExtrudeDistance = ExtrudeDistance;
+		FAxisAlignedBox2d bounds = PolyMeshGen.Polygon.Bounds();
+		Extruder.UVScaleFactor = 1.0 / bounds.MaxDim();
+		if (ExtrudeDistance < 0)
+		{
+			Extruder.IsPositiveOffset = false;
+		}
+
+		FVector3d ExtrudeNormal = FVector3d::UnitZ();
+		Extruder.ExtrudedPositionFunc = [&ExtrudeDistance, &ExtrudeNormal](const FVector3d& Position, const FVector3f& Normal, int VertexID)
+		{
+			return Position + ExtrudeDistance * ExtrudeNormal;
+		};
+
+		Extruder.Apply();
+	}
+
+	FDynamicMeshEditor Editor(ResultMeshOut);
+	float InitialUVScale = 1.0 / PolyMeshGen.Polygon.Bounds().MaxDim(); // this is the UV scale used by both the polymeshgen and the extruder above
+	// default global rescale -- initial scale doesn't factor in extrude distance; rescale so UVScale of 1.0 fits in the unit square texture
+	float GlobalUVRescale = MaterialProperties->UVScale / FMathf::Max(1.0f, ExtrudeDistance * InitialUVScale);
+	if (MaterialProperties->bWorldSpaceUVScale)
+	{
+		// since we know the initial uv scale, directly compute the global scale (relative to 1 meter as a standard scale)
+		GlobalUVRescale = MaterialProperties->UVScale * .01 / InitialUVScale;
+	}
+	Editor.RescaleAttributeUVs(GlobalUVRescale, false);
+}
+
+void UDrawPolygonTool::GeneratePreviewHeightTarget()
+{
+	GeneratePolygonMesh(PolygonVertices, &PreviewHeightTarget, PreviewHeightFrame, false, 99999, true);
+	PreviewHeightTargetAABB.SetMesh(&PreviewHeightTarget);
+}
+
+
+
+
+void UDrawPolygonTool::ShowStartupMessage()
+{
+	GetToolManager()->DisplayMessage(
+		LOCTEXT("OnStartDraw", "Left-click to place points on the Drawing Plane. Hold Shift to ignore Snapping. Ctrl-click on the scene to reposition the Plane (Shift+Ctrl-click to only Translate). Backspace to discard last vertex. A key toggles Gizmo."),
+		EToolMessageLevel::UserNotification);
+}
+
+void UDrawPolygonTool::ShowExtrudeMessage()
+{
+	GetToolManager()->DisplayMessage(
+		LOCTEXT("OnStartExtrude", "Set the height of the Extrusion by positioning the mouse over the extrusion volume, or over the scene to snap to relative heights."),
+		EToolMessageLevel::UserNotification);
+}
+
+
+#undef LOCTEXT_NAMESPACE
diff --git a/Engine/Plugins/Experimental/MeshModelingToolset/Source/MeshModelingTools/Private/DynamicMeshBrushTool.cpp b/Engine/Plugins/Experimental/MeshModelingToolset/Source/MeshModelingTools/Private/DynamicMeshBrushTool.cpp
new file mode 100644
--- /dev/null
+++ b/Engine/Plugins/Experimental/MeshModelingToolset/Source/MeshModelingTools/Private/DynamicMeshBrushTool.cpp
@@ -0,0 +1,84 @@
+// Copyright 1998-2019 Epic Games, Inc. All Rights Reserved.
+
+#include "DynamicMeshBrushTool.h"
+#include "InteractiveToolManager.h"
+#include "ToolBuilderUtil.h"
+
+// localization namespace
+#define LOCTEXT_NAMESPACE "UDynamicMeshBrushTool"
+
+
+/*
+ * Tool
+ */
+
+UDynamicMeshBrushTool::UDynamicMeshBrushTool()
+{
+}
+
+
+void UDynamicMeshBrushTool::Setup()
+{
+	PreviewMesh = NewObject<UPreviewMesh>(this);
+	PreviewMesh->bBuildSpatialDataStructure = true;
+	PreviewMesh->CreateInWorld(ComponentTarget->GetOwnerActor()->GetWorld(), FTransform::Identity);
+	PreviewMesh->SetTransform(ComponentTarget->GetWorldTransform());
+	if (ComponentTarget->GetMaterial(0) != nullptr)
+	{
+		PreviewMesh->SetMaterial(ComponentTarget->GetMaterial(0));
+	}
+
+	// initialize from LOD-0 MeshDescription
+	PreviewMesh->InitializeMesh(ComponentTarget->GetMesh());
+	OnBaseMeshComponentChangedHandle = PreviewMesh->GetOnMeshChanged().Add(
+		FSimpleMulticastDelegate::FDelegate::CreateUObject(this, &UDynamicMeshBrushTool::OnBaseMeshComponentChanged));
+
+	// call this here so that base tool can estimate target dimension
+	InputMeshBoundsLocal = PreviewMesh->GetPreviewDynamicMesh()->GetBounds();
+	UBaseBrushTool::Setup();
+
+	// hide input StaticMeshComponent
+	ComponentTarget->SetOwnerVisibility(false);
+}
+
+
+
+double UDynamicMeshBrushTool::EstimateMaximumTargetDimension()
+{
+	return InputMeshBoundsLocal.MaxDim();
+}
+
+
+void UDynamicMeshBrushTool::Shutdown(EToolShutdownType ShutdownType)
+{
+	UBaseBrushTool::Shutdown(ShutdownType);
+
+	ComponentTarget->SetOwnerVisibility(true);
+
+	if (PreviewMesh != nullptr)
+	{
+		PreviewMesh->GetOnMeshChanged().Remove(OnBaseMeshComponentChangedHandle);
+
+		OnShutdown(ShutdownType);
+
+		PreviewMesh->SetVisible(false);
+		PreviewMesh->Disconnect();
+		PreviewMesh = nullptr;
+	}
+
+}
+
+
+
+
+bool UDynamicMeshBrushTool::HitTest(const FRay& Ray, FHitResult& OutHit)
+{
+	return PreviewMesh->FindRayIntersection(FRay3d(Ray), OutHit);
+}
+
+
+
+
+
+
+#undef LOCTEXT_NAMESPACE
diff --git a/Engine/Plugins/Experimental/MeshModelingToolset/Source/MeshModelingTools/Private/DynamicMeshSculptTool.cpp b/Engine/Plugins/Experimental/MeshModelingToolset/Source/MeshModelingTools/Private/DynamicMeshSculptTool.cpp
new file mode 100644
--- /dev/null
+++ b/Engine/Plugins/Experimental/MeshModelingToolset/Source/MeshModelingTools/Private/DynamicMeshSculptTool.cpp
@@ -0,0 +1,1684 @@
+// Copyright 1998-2019 Epic Games, Inc. All Rights Reserved.
+
+#include "DynamicMeshSculptTool.h"
+#include "InteractiveToolManager.h"
+#include "InteractiveGizmoManager.h"
+#include "ToolBuilderUtil.h"
+
+#include "SubRegionRemesher.h"
+#include "ProjectionTargets.h"
+#include "MeshConstraints.h"
+#include "MeshConstraintsUtil.h"
+#include "MeshWeights.h"
+#include "MeshNormals.h"
+#include "MeshIndexUtil.h"
+#include "Drawing/MeshDebugDrawing.h"
+#include "PreviewMesh.h"
+#include "ToolSetupUtil.h"
+
+#include "Changes/MeshVertexChange.h"
+#include "Changes/MeshChange.h"
+#include "DynamicMeshChangeTracker.h"
+
+#include "Async/ParallelFor.h"
+
+#define LOCTEXT_NAMESPACE "UDynamicMeshSculptTool"
+
+
+/*
+ * ToolBuilder
+ */
+
+UMeshSurfacePointTool* UDynamicMeshSculptToolBuilder::CreateNewTool(const FToolBuilderState& SceneState) const
+{
+	UDynamicMeshSculptTool* SculptTool = NewObject<UDynamicMeshSculptTool>(SceneState.ToolManager);
+	SculptTool->SetEnableRemeshing(this->bEnableRemeshing);
+	SculptTool->SetWorld(SceneState.World);
+	return SculptTool;
+}
+
+
+
+UBrushSculptProperties::UBrushSculptProperties()
+{
+	SmoothSpeed = 0.25;
+	BrushSpeed = 0.5;
+	PrimaryBrushType = EDynamicMeshSculptBrushType::Move;
+	bPreserveUVFlow = false;
+	BrushDepth = 0;
+	bFreezeTarget = false;
+}
+
+void UBrushSculptProperties::SaveProperties(UInteractiveTool* SaveFromTool)
+{
+	UBrushSculptProperties* PropertyCache = GetPropertyCache<UBrushSculptProperties>();
+	PropertyCache->SmoothSpeed = this->SmoothSpeed;
+	PropertyCache->BrushSpeed = this->BrushSpeed;
+	PropertyCache->PrimaryBrushType = this->PrimaryBrushType;
+	PropertyCache->bPreserveUVFlow = this->bPreserveUVFlow;
+	PropertyCache->BrushDepth = this->BrushDepth;
+}
+
+void UBrushSculptProperties::RestoreProperties(UInteractiveTool* RestoreToTool)
+{
+	UBrushSculptProperties* PropertyCache = GetPropertyCache<UBrushSculptProperties>();
+	this->SmoothSpeed = PropertyCache->SmoothSpeed;
+	this->BrushSpeed = PropertyCache->BrushSpeed;
+	this->PrimaryBrushType = PropertyCache->PrimaryBrushType;
+	this->bPreserveUVFlow = PropertyCache->bPreserveUVFlow;
+	this->BrushDepth = PropertyCache->BrushDepth;
+}
+
+
+
+
+
+
+UBrushRemeshProperties::UBrushRemeshProperties()
+{
+	RelativeSize = 1.0;
+	Smoothing = 0.1;
+}
+
+
+/*
+ * Tool
+ */
+
+UDynamicMeshSculptTool::UDynamicMeshSculptTool()
+{
+	// initialize parameters
+	bEnableRemeshing = true;
+}
+
+
+void UDynamicMeshSculptTool::SetWorld(UWorld* World)
+{
+	this->TargetWorld = World;
+}
+
+
+const FString BrushIndicatorGizmoType = TEXT("BrushIndicatorGizmoType");
+
+void UDynamicMeshSculptTool::Setup()
+{
+	UMeshSurfacePointTool::Setup();
+
+	// create dynamic mesh component to use for live preview
+	DynamicMeshComponent = NewObject<UOctreeDynamicMeshComponent>(ComponentTarget->GetOwnerActor(), "DynamicMeshSculptToolMesh");
+	DynamicMeshComponent->SetupAttachment(ComponentTarget->GetOwnerActor()->GetRootComponent());
+	DynamicMeshComponent->RegisterComponent();
+	DynamicMeshComponent->SetWorldTransform(ComponentTarget->GetWorldTransform());
+
+	// copy material if there is one
+	UMaterialInterface* Material = ComponentTarget->GetMaterial(0);
+	if (Material != nullptr)
+	{
+		DynamicMeshComponent->SetMaterial(0, Material);
+	}
+
+	// initialize from LOD-0 MeshDescription
+	//DynamicMeshComponent->TangentsType = (bEnableRemeshing) ? EDynamicMeshTangentCalcType::NoTangents : EDynamicMeshTangentCalcType::AutoCalculated;
+	DynamicMeshComponent->InitializeMesh(ComponentTarget->GetMesh());
+	OnDynamicMeshComponentChangedHandle = DynamicMeshComponent->OnMeshChanged.Add(
+		FSimpleMulticastDelegate::FDelegate::CreateUObject(this, &UDynamicMeshSculptTool::OnDynamicMeshComponentChanged));
+
+	// initialize AABBTree
+	FDynamicMesh3* Mesh = DynamicMeshComponent->GetMesh();
+
+	// do we always want to keep vertex normals updated? Perhaps should discard vertex normals before baking?
+	FMeshNormals::QuickComputeVertexNormals(*Mesh);
+
+	// switch to vertex normals for testing
+	//DynamicMeshComponent->GetMesh()->DiscardAttributes();
+
+	UpdateTarget();
+	bTargetDirty = false;
+
+	double MaxDimension = DynamicMeshComponent->GetMesh()->GetCachedBounds().MaxDim();
+	BrushRelativeSizeRange = FInterval1d(MaxDimension*0.01, MaxDimension);
+	BrushProperties = NewObject<UBrushBaseProperties>(this, TEXT("Brush"));
+	CalculateBrushRadius();
+
+	SculptProperties = NewObject<UBrushSculptProperties>(this, TEXT("Sculpting"));
+	RemeshProperties = NewObject<UBrushRemeshProperties>(this, TEXT("Remeshing"));
+	InitialEdgeLength = EstimateIntialSafeTargetLength(*Mesh, 5000);
+
+
+	// hide input StaticMeshComponent
+	ComponentTarget->SetOwnerVisibility(false);
+
+	// init state flags flags
+	bInDrag = false;
+	bHaveRemeshed = false;
+	bRemeshPending = false;
+	bStampPending = false;
+	ActiveVertexChange = nullptr;
+
+	// register and spawn brush indicator gizmo
+	GetToolManager()->GetPairedGizmoManager()->RegisterGizmoType(BrushIndicatorGizmoType, NewObject<UBrushStampIndicatorBuilder>());
+	BrushIndicator = GetToolManager()->GetPairedGizmoManager()->CreateGizmo<UBrushStampIndicator>(BrushIndicatorGizmoType, FString(), this);
+	BrushIndicatorMesh = UBrushStampIndicator::MakeDefaultSphereMesh(this, TargetWorld);
+	BrushIndicator->AttachedComponent = BrushIndicatorMesh->GetRootComponent();
+	BrushIndicator->bDrawIndicatorLines = false;
+
+	// initialize our properties
+	AddToolPropertySource(SculptProperties);
+	AddToolPropertySource(BrushProperties);
+	if (this->bEnableRemeshing)
+	{
+		AddToolPropertySource(RemeshProperties);
+	}
+
+	BrushProperties->RestoreProperties(this);
+	CalculateBrushRadius();
+	SculptProperties->RestoreProperties(this);
+
+	ViewProperties = NewObject<UMeshEditingViewProperties>();
+	ViewProperties->RestoreProperties(this);
+	AddToolPropertySource(ViewProperties);
+	ShowWireframeWatcher.Initialize(
+		[this]() { return ViewProperties->bShowWireframe; },
+		[this](bool bNewValue) { DynamicMeshComponent->bExplicitShowWireframe = bNewValue; }, false);
+	MaterialModeWatcher.Initialize(
+		[this]() { return ViewProperties->MaterialMode; },
+		[this](EMeshEditingMaterialModes NewMode) { UpdateMaterialMode(NewMode); }, EMeshEditingMaterialModes::ExistingMaterial);
+
+
+	GetToolManager()->DisplayMessage(
+		LOCTEXT("OnStartSculptTool", "Hold Shift to Smooth, Ctrl to Invert (where applicable). Q/A keys cycle through Brush Types, Shift+Q/A for Brush History. S/D change Size (shift to small-step), W/E change Speed."),
+		EToolMessageLevel::UserNotification);
+
+	if (bEnableRemeshing)
+	{
+		PrecomputeRemeshInfo();
+		if (bHaveUVSeams)
+		{
+			GetToolManager()->DisplayMessage(
+				LOCTEXT("UVSeamWarning", "This mesh has UV seams which may limit remeshing. Consider clearing the UV layers using the Remesh Tool."),
+				EToolMessageLevel::UserWarning);
+		} 
+		else if (bHaveNormalSeams)
+		{
+			GetToolManager()->DisplayMessage(
+				LOCTEXT("NormalSeamWarning", "This mesh has Hard Normal seams which may limit remeshing. Consider clearing Hard Normals using the Remesh Tool."),
+				EToolMessageLevel::UserWarning);
+		}
+	}
+
+}
+
+
+
+
+void UDynamicMeshSculptTool::Shutdown(EToolShutdownType ShutdownType)
+{
+	BrushIndicatorMesh->Disconnect();
+	BrushIndicatorMesh = nullptr;
+
+	GetToolManager()->GetPairedGizmoManager()->DestroyAllGizmosByOwner(this);
+	BrushIndicator = nullptr;
+	GetToolManager()->GetPairedGizmoManager()->DeregisterGizmoType(BrushIndicatorGizmoType);
+
+	if (DynamicMeshComponent != nullptr)
+	{
+		DynamicMeshComponent->OnMeshChanged.Remove(OnDynamicMeshComponentChangedHandle);
+
+		ComponentTarget->SetOwnerVisibility(true);
+
+		if (ShutdownType == EToolShutdownType::Accept)
+		{
+			// this block bakes the modified DynamicMeshComponent back into the StaticMeshComponent inside an undo transaction
+			GetToolManager()->BeginUndoTransaction(LOCTEXT("SculptMeshToolTransactionName", "Sculpt Mesh"));
+			ComponentTarget->CommitMesh([=](FMeshDescription* MeshDescription)
+			{
+				DynamicMeshComponent->Bake(MeshDescription, bHaveRemeshed);
+			});
+			GetToolManager()->EndUndoTransaction();
+		}
+
+		DynamicMeshComponent->UnregisterComponent();
+		DynamicMeshComponent->DestroyComponent();
+		DynamicMeshComponent = nullptr;
+	}
+
+	BrushProperties->SaveProperties(this);
+	SculptProperties->SaveProperties(this);
+	ViewProperties->SaveProperties(this);
+}
+
+
+
+void UDynamicMeshSculptTool::OnDynamicMeshComponentChanged()
+{
+	bNormalUpdatePending = true;
+	bTargetDirty = true;
+}
+
+
+void UDynamicMeshSculptTool::OnPropertyModified(UObject* PropertySet, UProperty* Property)
+{
+	CalculateBrushRadius();
+}
+
+
+
+bool UDynamicMeshSculptTool::HitTest(const FRay& Ray, FHitResult& OutHit)
+{
+	FTransform Transform = ComponentTarget->GetWorldTransform();
+
+	FRay3d LocalRay(Transform.InverseTransformPosition(Ray.Origin),
+		Transform.InverseTransformVector(Ray.Direction));
+	LocalRay.Direction.Normalize();
+
+	FDynamicMesh3* Mesh = DynamicMeshComponent->GetMesh();
+	int HitTID = DynamicMeshComponent->GetOctree()->FindNearestHitObject(LocalRay);
+
+	if (HitTID != IndexConstants::InvalidID)
+	{
+		FTriangle3d Triangle;
+		Mesh->GetTriVertices(HitTID, Triangle.V[0], Triangle.V[1], Triangle.V[2]);
+		FIntrRay3Triangle3d Query(LocalRay, Triangle);
+		Query.Find();
+
+		OutHit.FaceIndex = HitTID;
+		OutHit.Distance = Query.RayParameter;
+		OutHit.Normal = Transform.TransformVectorNoScale(Mesh->GetTriNormal(HitTID));
+		OutHit.ImpactPoint = Transform.TransformPosition(LocalRay.PointAt(Query.RayParameter));
+		return true;
+	}
+
+	return false;
+}
+
+
+
+void UDynamicMeshSculptTool::OnBeginDrag(const FRay& Ray)
+{
+	bSmoothing = GetShiftToggle();
+	bInvert = GetCtrlToggle();
+
+	FHitResult OutHit;
+	if (HitTest(Ray, OutHit))
+	{
+		BrushStartCenterWorld = Ray.PointAt(OutHit.Distance) + SculptProperties->BrushDepth*CurrentBrushRadius*Ray.Direction;
+
+		bInDrag = true;
+
+		ActiveDragPlane = FPlane(BrushStartCenterWorld, -Ray.Direction);
+		LastHitPosWorld = FMath::RayPlaneIntersection(Ray.Origin, Ray.Direction, ActiveDragPlane);
+
+		LastBrushPosWorld = LastHitPosWorld;
+		LastBrushPosNormalWorld = ActiveDragPlane;
+		FTransform Transform = ComponentTarget->GetWorldTransform();
+		LastBrushPosLocal = Transform.InverseTransformPosition(LastHitPosWorld);
+
+		BeginChange(bEnableRemeshing == false);
+
+		UpdateROI(LastBrushPosLocal);
+
+		if (SculptProperties->PrimaryBrushType == EDynamicMeshSculptBrushType::Plane)
+		{
+			ActiveFixedBrushPlane = ComputeROIBrushPlane(LastBrushPosLocal);
+		}
+
+		// apply initial stamp
+		PendingStampRay = Ray;
+		bStampPending = true;
+	}
+}
+
+
+
+void UDynamicMeshSculptTool::UpdateROI(const FVector& BrushPos)
+{
+	SCOPE_CYCLE_COUNTER(SculptTool_UpdateROI);
+
+	// TODO: need dynamic vertex hash table!
+
+	float RadiusSqr = CurrentBrushRadius * CurrentBrushRadius;
+
+
+	FAxisAlignedBox3d BrushBox(
+		BrushPos - CurrentBrushRadius * FVector3d::One(),
+		BrushPos + CurrentBrushRadius * FVector3d::One());
+
+	VertexSetBuffer.Reset();
+	FDynamicMesh3* Mesh = DynamicMeshComponent->GetMesh();
+	FDynamicMeshOctree3* Octree = DynamicMeshComponent->GetOctree();
+	Octree->RangeQuery(BrushBox,
+		[&](int TriIdx) {
+		
+		FIndex3i TriV = Mesh->GetTriangle(TriIdx);
+		for (int j = 0; j < 3; ++j)
+		{
+			FVector3d Position = Mesh->GetVertex(TriV[j]);
+			if ((Position - BrushPos).SquaredLength() < RadiusSqr)
+			{
+				VertexSetBuffer.Add(TriV[j]);
+			}
+		}
+
+	});
+
+	VertexROI.SetNum(0, false);
+	BufferUtil::AppendElements(VertexROI, VertexSetBuffer);
+
+	TriangleROI.Reset();
+	MeshIndexUtil::VertexToTriangleOneRing(Mesh, VertexROI, TriangleROI);
+}
+
+
+void UDynamicMeshSculptTool::OnUpdateDrag(const FRay& WorldRay)
+{
+	if (bInDrag)
+	{
+		PendingStampRay = WorldRay;
+		bStampPending = true;
+	}
+}
+
+
+
+void UDynamicMeshSculptTool::CalculateBrushRadius()
+{
+	CurrentBrushRadius = 0.5 * BrushRelativeSizeRange.Interpolate(BrushProperties->BrushSize);
+	if (BrushProperties->bSpecifyRadius)
+	{
+		CurrentBrushRadius = BrushProperties->BrushRadius;
+	}
+	else
+	{
+		BrushProperties->BrushRadius = CurrentBrushRadius;
+	}
+}
+
+
+
+void UDynamicMeshSculptTool::ApplyStamp(const FRay& WorldRay)
+{
+	SCOPE_CYCLE_COUNTER(STAT_SculptToolApplyStamp);
+
+	// update brush type history. apologies for convoluted logic. 
+	StampTimestamp++;
+	if (LastStampType != PendingStampType)
+	{
+		if (BrushTypeHistoryIndex != BrushTypeHistory.Num() - 1)
+		{
+			if (LastStampType != EDynamicMeshSculptBrushType::LastValue)
+			{
+				BrushTypeHistory.Add(LastStampType);
+			}
+			BrushTypeHistoryIndex = BrushTypeHistory.Num()-1;
+		}
+		LastStampType = PendingStampType;
+		if (BrushTypeHistory.Num() == 0 || BrushTypeHistory[BrushTypeHistory.Num()-1] != PendingStampType)
+		{
+			BrushTypeHistory.Add(PendingStampType);
+			BrushTypeHistoryIndex = BrushTypeHistory.Num() - 1;
+		}
+	}
+
+	CalculateBrushRadius();
+
+	SaveActiveROI();
+
+	if (bSmoothing)
+	{
+		ApplySmoothBrush(WorldRay);
+		return;
+	}
+
+	switch (SculptProperties->PrimaryBrushType)
+	{
+		case EDynamicMeshSculptBrushType::Offset:
+			ApplyOffsetBrush(WorldRay);
+			break;
+		case EDynamicMeshSculptBrushType::SculptMax:
+			ApplySculptMaxBrush(WorldRay);
+			break;
+		case EDynamicMeshSculptBrushType::Move:
+			ApplyMoveBrush(WorldRay);
+			break;
+		case EDynamicMeshSculptBrushType::Pinch:
+			ApplyPinchBrush(WorldRay);
+			break;
+		case EDynamicMeshSculptBrushType::Inflate:
+			ApplyInflateBrush(WorldRay);
+			break;
+		case EDynamicMeshSculptBrushType::Flatten:
+			ApplyFlattenBrush(WorldRay);
+			break;
+		case EDynamicMeshSculptBrushType::Plane:
+			ApplyPlaneBrush(WorldRay);
+			break;
+	}
+}
+
+
+
+
+
+double UDynamicMeshSculptTool::CalculateBrushFalloff(double Distance)
+{
+	double d = Distance / CurrentBrushRadius;
+	double w = 1;
+	if (d > 0.5)
+	{
+		d = VectorUtil::Clamp((d - 0.5) / 0.5, 0.0, 1.0);
+		w = (1 - d * d);
+		w = w * w * w;
+	}
+	return w;
+}
+
+
+
+
+
+void UDynamicMeshSculptTool::ApplySmoothBrush(const FRay& WorldRay)
+{
+	bool bHit = UpdateBrushPositionOnSculptMesh(WorldRay);
+	if (bHit == false)
+	{
+		return;
+	}
+
+
+	FTransform Transform = ComponentTarget->GetWorldTransform();
+	FVector NewBrushPosLocal = Transform.InverseTransformPosition(LastBrushPosWorld);
+
+	FDynamicMesh3* Mesh = DynamicMeshComponent->GetMesh();
+	for (int VertIdx : VertexROI)
+	{
+		FVector3d OrigPos = Mesh->GetVertex(VertIdx);
+
+		double Falloff = CalculateBrushFalloff(OrigPos.Distance(NewBrushPosLocal));
+
+		FVector3d SmoothedPos = (SculptProperties->bPreserveUVFlow) ?
+			FMeshWeights::MeanValueCentroid(*Mesh, VertIdx) : FMeshWeights::UniformCentroid(*Mesh, VertIdx);
+
+		FVector3d NewPos = FVector3d::Lerp(OrigPos, SmoothedPos, Falloff*SculptProperties->SmoothSpeed);
+
+		Mesh->SetVertex(VertIdx, NewPos);
+		UpdateSavedVertex(VertIdx, OrigPos, NewPos);
+	}
+
+	bRemeshPending = bEnableRemeshing;
+	LastBrushPosLocal = NewBrushPosLocal;
+}
+
+
+
+
+
+void UDynamicMeshSculptTool::ApplyMoveBrush(const FRay& WorldRay)
+{
+	bool bHit = UpdateBrushPositionOnActivePlane(WorldRay);
+
+	FTransform Transform = ComponentTarget->GetWorldTransform();
+	FVector NewBrushPosLocal = Transform.InverseTransformPosition(LastBrushPosWorld);
+
+	FVector3d MoveVec = NewBrushPosLocal - LastBrushPosLocal;
+
+	if (MoveVec.SquaredLength() > 0)
+	{
+		FDynamicMesh3* Mesh = DynamicMeshComponent->GetMesh();
+		for (int VertIdx : VertexROI)
+		{
+			FVector3d OrigPos = Mesh->GetVertex(VertIdx);
+
+			double PrevDist = (OrigPos - LastBrushPosLocal).Length();
+			double NewDist = (OrigPos - NewBrushPosLocal).Length();
+			double UseDist = FMath::Min(PrevDist, NewDist);
+
+			double Falloff = CalculateBrushFalloff(UseDist);
+
+			FVector3d NewPos = OrigPos + Falloff * MoveVec;
+
+			Mesh->SetVertex(VertIdx, NewPos);
+			UpdateSavedVertex(VertIdx, OrigPos, NewPos);
+		}
+
+		bRemeshPending = bEnableRemeshing;
+	}
+
+	LastBrushPosLocal = NewBrushPosLocal;
+}
+
+
+
+
+
+
+void UDynamicMeshSculptTool::ApplyOffsetBrush(const FRay& WorldRay)
+{
+	bool bHit = UpdateBrushPositionOnTargetMesh(WorldRay);
+	if (bHit == false)
+	{
+		return;
+	}
+
+	FTransform Transform = ComponentTarget->GetWorldTransform();
+	FVector NewBrushPosLocal = Transform.InverseTransformPosition(LastBrushPosWorld);
+
+	double Direction = (bInvert) ? -1.0 : 1.0;
+	double UseSpeed = Direction * FMathd::Sqrt(CurrentBrushRadius) * (SculptProperties->BrushSpeed*0.25);
+
+	FDynamicMesh3* Mesh = DynamicMeshComponent->GetMesh();
+	for (int VertIdx : VertexROI)
+	{
+		FVector3d OrigPos = Mesh->GetVertex(VertIdx);
+
+		FVector3d BasePos, BaseNormal;
+		GetTargetMeshNearest(OrigPos, (double)(2 * CurrentBrushRadius), BasePos, BaseNormal);
+
+		FVector3d MoveVec = UseSpeed * BaseNormal;
+
+		double Falloff = CalculateBrushFalloff(OrigPos.Distance(NewBrushPosLocal));
+
+		FVector3d NewPos = OrigPos + Falloff * MoveVec;
+
+		Mesh->SetVertex(VertIdx, NewPos);
+		UpdateSavedVertex(VertIdx, OrigPos, NewPos);
+	}
+
+	bRemeshPending = bEnableRemeshing;
+
+	LastBrushPosLocal = NewBrushPosLocal;
+}
+
+
+
+
+
+void UDynamicMeshSculptTool::ApplySculptMaxBrush(const FRay& WorldRay)
+{
+	bool bHit = UpdateBrushPositionOnTargetMesh(WorldRay);
+	if (bHit == false)
+	{
+		return;
+	}
+
+	FTransform Transform = ComponentTarget->GetWorldTransform();
+	FVector NewBrushPosLocal = Transform.InverseTransformPosition(LastBrushPosWorld);
+
+	double Direction = (bInvert) ? -1.0 : 1.0;
+	double UseSpeed = Direction * FMathd::Sqrt(CurrentBrushRadius) * (SculptProperties->BrushSpeed*0.25);
+
+	double MaxOffset = CurrentBrushRadius * 0.5;
+	double MaxOffsetSqr = MaxOffset * MaxOffset;
+
+	FDynamicMesh3* Mesh = DynamicMeshComponent->GetMesh();
+	for (int VertIdx : VertexROI)
+	{
+		FVector3d OrigPos = Mesh->GetVertex(VertIdx);
+
+		FVector3d BasePos, BaseNormal;
+		GetTargetMeshNearest(OrigPos, (double)(2 * CurrentBrushRadius), BasePos, BaseNormal);
+
+		FVector3d MoveVec = UseSpeed * BaseNormal;
+
+		double Falloff = CalculateBrushFalloff(OrigPos.Distance(NewBrushPosLocal));
+
+		FVector3d NewPos = OrigPos + Falloff * MoveVec;
+
+		FVector3d DeltaPos = NewPos - BasePos;
+		if (DeltaPos.SquaredLength() > MaxOffsetSqr)
+		{
+			DeltaPos.Normalize();
+			NewPos = BasePos + MaxOffset * DeltaPos;
+		}
+
+		Mesh->SetVertex(VertIdx, NewPos);
+		UpdateSavedVertex(VertIdx, OrigPos, NewPos);
+	}
+
+	bRemeshPending = bEnableRemeshing;
+
+	LastBrushPosLocal = NewBrushPosLocal;
+}
+
+
+
+void UDynamicMeshSculptTool::ApplyPinchBrush(const FRay& WorldRay)
+{
+	bool bHit = UpdateBrushPositionOnTargetMesh(WorldRay);
+	if (bHit == false)
+	{
+		return;
+	}
+
+	FTransform Transform = ComponentTarget->GetWorldTransform();
+	FVector NewBrushPosLocal = Transform.InverseTransformPosition(LastBrushPosWorld);
+	FVector BrushNormalLocal = Transform.InverseTransformVectorNoScale(LastBrushPosNormalWorld);
+	FVector OffsetBrushPosLocal = NewBrushPosLocal - SculptProperties->BrushDepth * CurrentBrushRadius * BrushNormalLocal;
+
+	double Direction = (bInvert) ? -1.0 : 1.0;
+	double UseSpeed = Direction * FMathd::Sqrt(CurrentBrushRadius) * (SculptProperties->BrushSpeed*0.25);
+
+	FVector3d MotionVec = NewBrushPosLocal - LastBrushPosLocal;
+	MotionVec.Normalize();
+
+	FDynamicMesh3* Mesh = DynamicMeshComponent->GetMesh();
+	for (int VertIdx : VertexROI)
+	{
+		FVector3d OrigPos = Mesh->GetVertex(VertIdx);
+		FVector3d Delta = OffsetBrushPosLocal - OrigPos;
+		FVector3d MoveVec = UseSpeed * Delta;
+
+		double Falloff = CalculateBrushFalloff(OrigPos.Distance(NewBrushPosLocal));
+
+		// make tighter
+		Falloff = Falloff * Falloff * Falloff;
+
+		double AnglePower = 1.0 - FMathd::Abs(MoveVec.Normalized().Dot(MotionVec));
+		Falloff *= AnglePower;
+
+		FVector3d NewPos = OrigPos + Falloff * MoveVec;
+
+		Mesh->SetVertex(VertIdx, NewPos);
+		UpdateSavedVertex(VertIdx, OrigPos, NewPos);
+	}
+
+	bRemeshPending = bEnableRemeshing;
+
+	LastBrushPosLocal = NewBrushPosLocal;
+}
+
+
+
+
+FFrame3d UDynamicMeshSculptTool::ComputeROIBrushPlane(const FVector3d& BrushCenter)
+{
+	FDynamicMesh3* Mesh = DynamicMeshComponent->GetMesh();
+	FVector3d AverageNormal(0, 0, 0);
+	FVector3d AveragePos(0, 0, 0);
+	double WeightSum = 0;
+	for (int TriID : TriangleROI)
+	{
+		FVector3d Centroid = Mesh->GetTriCentroid(TriID);
+		double Weight = CalculateBrushFalloff(BrushCenter.Distance(Centroid));
+
+		AverageNormal += Weight * Mesh->GetTriNormal(TriID);
+		AveragePos += Weight * Centroid;
+		WeightSum += Weight;
+	}
+	AverageNormal.Normalize();
+	AveragePos /= WeightSum;
+
+	FFrame3d Result = FFrame3d(AveragePos, AverageNormal);
+	Result.Origin -= SculptProperties->BrushDepth * CurrentBrushRadius * Result.Z();
+
+	return Result;
+}
+
+void UDynamicMeshSculptTool::ApplyPlaneBrush(const FRay& WorldRay)
+{
+	bool bHit = UpdateBrushPositionOnSculptMesh(WorldRay);
+	if (bHit == false)
+	{
+		return;
+	}
+
+	FTransform Transform = ComponentTarget->GetWorldTransform();
+	FVector NewBrushPosLocal = Transform.InverseTransformPosition(LastBrushPosWorld);
+	FVector BrushNormalLocal = Transform.InverseTransformVectorNoScale(LastBrushPosNormalWorld);
+	FVector OffsetBrushPosLocal = NewBrushPosLocal - SculptProperties->BrushDepth * CurrentBrushRadius * BrushNormalLocal;
+
+	FDynamicMesh3* Mesh = DynamicMeshComponent->GetMesh();
+	double UseSpeed = FMathd::Sqrt(CurrentBrushRadius) * FMathd::Sqrt(SculptProperties->BrushSpeed) * 0.25;
+
+	for (int VertIdx : VertexROI)
+	{
+		FVector3d OrigPos = Mesh->GetVertex(VertIdx);
+		FVector3d PlanePos = ActiveFixedBrushPlane.ToPlane(OrigPos, 2);
+		FVector3d Delta = PlanePos - OrigPos;
+		FVector3d MoveVec = UseSpeed * Delta;
+
+		double Falloff = CalculateBrushFalloff(OrigPos.Distance(NewBrushPosLocal));
+
+		FVector3d NewPos = OrigPos + Falloff * MoveVec;
+
+		Mesh->SetVertex(VertIdx, NewPos);
+		UpdateSavedVertex(VertIdx, OrigPos, NewPos);
+	}
+
+	bRemeshPending = bEnableRemeshing;
+
+	LastBrushPosLocal = NewBrushPosLocal;
+}
+
+
+
+
+
+void UDynamicMeshSculptTool::ApplyFlattenBrush(const FRay& WorldRay)
+{
+	bool bHit = UpdateBrushPositionOnSculptMesh(WorldRay);
+	if (bHit == false)
+	{
+		return;
+	}
+
+	FTransform Transform = ComponentTarget->GetWorldTransform();
+	FVector NewBrushPosLocal = Transform.InverseTransformPosition(LastBrushPosWorld);
+	FVector BrushNormalLocal = Transform.InverseTransformVectorNoScale(LastBrushPosNormalWorld);
+	FVector OffsetBrushPosLocal = NewBrushPosLocal - SculptProperties->BrushDepth * CurrentBrushRadius * BrushNormalLocal;
+
+	FDynamicMesh3* Mesh = DynamicMeshComponent->GetMesh();
+	double UseSpeed = FMathd::Sqrt(CurrentBrushRadius) * FMathd::Sqrt(SculptProperties->BrushSpeed) * 0.25;
+	FFrame3d StampFlattenPlane = ComputeROIBrushPlane(NewBrushPosLocal);
+
+	for (int VertIdx : VertexROI)
+	{
+		FVector3d OrigPos = Mesh->GetVertex(VertIdx);
+		FVector3d PlanePos = StampFlattenPlane.ToPlane(OrigPos, 2);
+		FVector3d Delta = PlanePos - OrigPos;
+		FVector3d MoveVec = UseSpeed * Delta;
+
+		double Falloff = CalculateBrushFalloff(OrigPos.Distance(NewBrushPosLocal));
+
+		FVector3d NewPos = OrigPos + Falloff * MoveVec;
+
+		Mesh->SetVertex(VertIdx, NewPos);
+		UpdateSavedVertex(VertIdx, OrigPos, NewPos);
+	}
+
+	bRemeshPending = bEnableRemeshing;
+
+	LastBrushPosLocal = NewBrushPosLocal;
+}
+
+
+
+void UDynamicMeshSculptTool::ApplyInflateBrush(const FRay& WorldRay)
+{
+	bool bHit = UpdateBrushPositionOnSculptMesh(WorldRay);
+	if (bHit == false)
+	{
+		return;
+	}
+
+	FTransform Transform = ComponentTarget->GetWorldTransform();
+	FVector NewBrushPosLocal = Transform.InverseTransformPosition(LastBrushPosWorld);
+
+	double Direction = (bInvert) ? -1.0 : 1.0;
+	double UseSpeed = Direction * FMathd::Sqrt(CurrentBrushRadius) * SculptProperties->BrushSpeed * 0.25;
+	FDynamicMesh3* Mesh = DynamicMeshComponent->GetMesh();
+
+	// calculate vertex normals
+	ParallelFor(VertexROI.Num(), [this, Mesh](int Index) {
+		int VertIdx = VertexROI[Index];
+		FVector3d Normal = FMeshNormals::ComputeVertexNormal(*Mesh, VertIdx);
+		Mesh->SetVertexNormal(VertIdx, (FVector3f)Normal);
+	});
+
+
+	ParallelFor(VertexROI.Num(), [this, Mesh, UseSpeed, NewBrushPosLocal](int Index) 
+	{
+		int VertIdx = VertexROI[Index];
+		FVector3d OrigPos = Mesh->GetVertex(VertIdx);
+		FVector3d Normal = (FVector3d)Mesh->GetVertexNormal(VertIdx);
+
+		FVector3d MoveVec = UseSpeed * Normal;
+
+		double Falloff = CalculateBrushFalloff(OrigPos.Distance(NewBrushPosLocal));
+
+		FVector3d NewPos = OrigPos + Falloff*MoveVec;
+		Mesh->SetVertex(VertIdx, NewPos);
+
+		UpdateSavedVertex(VertIdx, OrigPos, NewPos);
+	});
+
+	bRemeshPending = bEnableRemeshing;
+
+	LastBrushPosLocal = NewBrushPosLocal;
+}
+
+
+
+
+
+
+bool UDynamicMeshSculptTool::UpdateBrushPositionOnActivePlane(const FRay& WorldRay)
+{
+	FVector NewHitPosWorld = FMath::RayPlaneIntersection(WorldRay.Origin, WorldRay.Direction, ActiveDragPlane);
+	LastBrushPosWorld = NewHitPosWorld;
+	LastBrushPosNormalWorld = ActiveDragPlane;
+	return true;
+}
+
+bool UDynamicMeshSculptTool::UpdateBrushPositionOnTargetMesh(const FRay& WorldRay)
+{
+	FTransform Transform = ComponentTarget->GetWorldTransform();
+
+	FRay3d LocalRay(Transform.InverseTransformPosition(WorldRay.Origin),
+		Transform.InverseTransformVectorNoScale(WorldRay.Direction));
+	LocalRay.Direction.Normalize();
+
+	int HitTID = BrushTargetMeshSpatial.FindNearestHitTriangle(LocalRay);
+	if (HitTID != IndexConstants::InvalidID)
+	{
+		const FDynamicMesh3* TargetMesh = BrushTargetMeshSpatial.GetMesh();
+
+		FTriangle3d Triangle;
+		TargetMesh->GetTriVertices(HitTID, Triangle.V[0], Triangle.V[1], Triangle.V[2]);
+		FIntrRay3Triangle3d Query(LocalRay, Triangle);
+		Query.Find();
+
+		LastBrushPosNormalWorld = Transform.TransformVectorNoScale(TargetMesh->GetTriNormal(HitTID));
+		LastBrushPosWorld = Transform.TransformPosition(LocalRay.PointAt(Query.RayParameter));
+		return true;
+	}
+	return false;
+}
+
+
+bool UDynamicMeshSculptTool::UpdateBrushPositionOnSculptMesh(const FRay& WorldRay)
+{
+	FTransform Transform = ComponentTarget->GetWorldTransform();
+
+	FRay3d LocalRay(Transform.InverseTransformPosition(WorldRay.Origin),
+		Transform.InverseTransformVectorNoScale(WorldRay.Direction));
+	LocalRay.Direction.Normalize();
+
+	int HitTID = DynamicMeshComponent->GetOctree()->FindNearestHitObject(LocalRay);
+	if (HitTID != IndexConstants::InvalidID)
+	{
+		const FDynamicMesh3* SculptMesh = DynamicMeshComponent->GetMesh();
+
+		FTriangle3d Triangle;
+		SculptMesh->GetTriVertices(HitTID, Triangle.V[0], Triangle.V[1], Triangle.V[2]);
+		FIntrRay3Triangle3d Query(LocalRay, Triangle);
+		Query.Find();
+
+		LastBrushPosNormalWorld = Transform.TransformVectorNoScale(SculptMesh->GetTriNormal(HitTID));
+		LastBrushPosWorld = Transform.TransformPosition(LocalRay.PointAt(Query.RayParameter));
+		return true;
+	}
+	return false;
+}
+
+
+
+void UDynamicMeshSculptTool::OnEndDrag(const FRay& Ray)
+{
+	bInDrag = false;
+
+	// cancel these! otherwise change record could become invalid 
+	bStampPending = false;
+	bRemeshPending = false;
+
+	// update spatial
+	bTargetDirty = true;
+
+	// close change record
+	EndChange();
+}
+
+
+bool UDynamicMeshSculptTool::OnUpdateHover(const FInputDeviceRay& DevicePos)
+{
+	PendingStampType = SculptProperties->PrimaryBrushType;
+
+	if (bInDrag)
+	{
+		FVector NewHitPosWorld = FMath::RayPlaneIntersection(DevicePos.WorldRay.Origin, DevicePos.WorldRay.Direction, ActiveDragPlane);
+		LastBrushPosWorld = NewHitPosWorld;
+		LastBrushPosNormalWorld = ActiveDragPlane;
+	}
+	else
+	{
+		FHitResult OutHit;
+		if (HitTest(DevicePos.WorldRay, OutHit))
+		{
+			LastBrushPosWorld = DevicePos.WorldRay.PointAt(OutHit.Distance + SculptProperties->BrushDepth*CurrentBrushRadius);
+			LastBrushPosNormalWorld = OutHit.Normal;
+		}
+	}
+	return true;
+}
+
+
+
+
+void UDynamicMeshSculptTool::Render(IToolsContextRenderAPI* RenderAPI)
+{
+	UMeshSurfacePointTool::Render(RenderAPI);
+
+	BrushIndicator->Update( (float)this->CurrentBrushRadius, this->LastBrushPosWorld, this->LastBrushPosNormalWorld );
+}
+
+
+void UDynamicMeshSculptTool::Tick(float DeltaTime)
+{
+	SCOPE_CYCLE_COUNTER(STAT_SculptToolTick);
+
+	UMeshSurfacePointTool::Tick(DeltaTime);
+
+	ShowWireframeWatcher.CheckAndUpdate();
+	MaterialModeWatcher.CheckAndUpdate();
+
+	// if user changed to not-frozen, we need to update the target
+	if (bCachedFreezeTarget != SculptProperties->bFreezeTarget)
+	{
+		UpdateTarget();
+	}
+
+	bool bMeshModified = false;
+	bool bMeshShapeModified = false;
+
+	FDynamicMesh3* Mesh = DynamicMeshComponent->GetMesh();
+	FDynamicMeshOctree3* Octree = DynamicMeshComponent->GetOctree();
+
+	//Octree->CheckValidity(EValidityCheckFailMode::Check, false, true);
+
+	//
+	// Apply stamp
+	// 
+
+	if (bStampPending)
+	{
+		SCOPE_CYCLE_COUNTER(STAT_SculptTool_Tick_ApplyStampBlock);
+
+		// would this ever be true? does stamp require this?
+		bool bRemoveTrianglesBeforeStamp = false;
+
+		if ( bRemoveTrianglesBeforeStamp )
+		{
+			SCOPE_CYCLE_COUNTER(SculptTool_Tick_ApplyStamp_Remove);
+			Octree->RemoveTriangles(TriangleROI);
+		}
+		else
+		{
+			SCOPE_CYCLE_COUNTER(SculptTool_Tick_ApplyStamp_Remove);
+			Octree->NotifyPendingModification(TriangleROI);		// to mark initial positions
+		}
+
+		ApplyStamp(PendingStampRay);
+		bStampPending = false;
+
+		bNormalUpdatePending = true;
+		bMeshModified = true;
+		bMeshShapeModified = true;
+
+		// flow
+		if (bInDrag)
+		{
+			bStampPending = true;
+		}
+
+		if (bRemoveTrianglesBeforeStamp)
+		{
+			SCOPE_CYCLE_COUNTER(SculptTool_Tick_ApplyStamp_Insert);
+			Octree->InsertTriangles(TriangleROI);
+		}
+		else		
+		{
+			SCOPE_CYCLE_COUNTER(SculptTool_Tick_ApplyStamp_Insert);
+			Octree->ReinsertTriangles(TriangleROI);
+		}
+
+	}
+
+
+	bool bUpdatedROIInRemesh = false;
+	if (bRemeshPending)
+	{
+		SCOPE_CYCLE_COUNTER(STAT_SculptTool_Tick_RemeshBlock);
+
+		check(bInDrag == true);    // this would break undo otherwise!
+
+		RemeshROIPass();
+
+		bMeshModified = true;
+		bMeshShapeModified = true;
+		bRemeshPending = false;
+		bNormalUpdatePending = true;
+		bHaveRemeshed = true;
+		bUpdatedROIInRemesh = true;
+	}
+
+	//Octree->CheckValidity(EValidityCheckFailMode::Check, false, true);
+
+	if (bNormalUpdatePending)
+	{
+		SCOPE_CYCLE_COUNTER(STAT_SculptTool_Tick_NormalsBlock);
+
+		if (Mesh->HasAttributes() && Mesh->Attributes()->PrimaryNormals() != nullptr)
+		{
+			RecalculateNormals_Overlay();
+		}
+		else
+		{
+			RecalculateNormals_PerVertex();
+		}
+		bNormalUpdatePending = false;
+		bMeshModified = true;
+	}
+
+	if (bMeshModified)
+	{
+		SCOPE_CYCLE_COUNTER(STAT_SculptTool_Tick_UpdateMeshBlock);
+
+		if (bMeshModified)
+
+		DynamicMeshComponent->NotifyMeshUpdated();
+		GetToolManager()->PostInvalidation();
+
+		bMeshModified = false;
+
+		if (bUpdatedROIInRemesh == false)
+		{
+			UpdateROI(LastBrushPosLocal);
+		}
+	}
+
+	if (bTargetDirty)
+	{
+		UpdateTarget();
+		bTargetDirty = false;
+	}
+}
+
+
+
+
+
+
+void UDynamicMeshSculptTool::PrecomputeRemeshInfo()
+{
+	FDynamicMesh3* Mesh = DynamicMeshComponent->GetMesh();
+
+	// check if we have any open boundary edges
+	bHaveMeshBoundaries = false;
+	for (int eid : Mesh->EdgeIndicesItr())
+	{
+		if (Mesh->IsBoundaryEdge(eid))
+		{
+			bHaveMeshBoundaries = true;
+			break;
+		}
+	}
+
+	// check if we have any UV seams
+	bHaveUVSeams = false;
+	bHaveNormalSeams = false;
+	if (Mesh->HasAttributes())
+	{
+		FDynamicMeshAttributeSet* Attribs = Mesh->Attributes();
+		for (int k = 0; k < Attribs->NumUVLayers(); ++k)
+		{
+			bHaveUVSeams = bHaveUVSeams || Attribs->GetUVLayer(k)->HasInteriorSeamEdges();
+		}
+
+		bHaveNormalSeams = Attribs->PrimaryNormals()->HasInteriorSeamEdges();
+	}
+}
+
+
+
+
+void UDynamicMeshSculptTool::RemeshROIPass()
+{
+	FDynamicMesh3* Mesh = DynamicMeshComponent->GetMesh();
+	FDynamicMeshOctree3* Octree = DynamicMeshComponent->GetOctree();
+
+	FSubRegionRemesher Remesher(Mesh);
+	Remesher.SetTargetEdgeLength(RemeshProperties->RelativeSize * InitialEdgeLength);
+
+	// this is a temporary tweak for pinch brush. Remesh params should be per-brush!
+	if (SculptProperties->PrimaryBrushType == EDynamicMeshSculptBrushType::Pinch && bSmoothing == false)
+	{
+		Remesher.MinEdgeLength = Remesher.MinEdgeLength / 4.0;
+	}
+
+	Remesher.SmoothSpeedT = RemeshProperties->Smoothing * 0.25;
+	Remesher.SmoothType = (SculptProperties->bPreserveUVFlow) ?
+		FRemesher::ESmoothTypes::MeanValue : FRemesher::ESmoothTypes::Uniform;
+	bool bIsUniformSmooth = (Remesher.SmoothType == FRemesher::ESmoothTypes::Uniform);
+
+	Remesher.bEnableCollapses = RemeshProperties->bCollapses;
+	Remesher.bEnableFlips = RemeshProperties->bFlips;
+	Remesher.bEnableSplits = RemeshProperties->bSplits;
+	Remesher.bPreventNormalFlips = RemeshProperties->bPreventNormalFlips;
+
+	{
+		SCOPE_CYCLE_COUNTER(STAT_SculptTool_Remesh_Setup);
+		for (int VertIdx : VertexROI)
+		{
+			Remesher.VertexROI.Add(VertIdx);
+		}
+		Remesher.InitializeFromVertexROI();
+		Remesher.UpdateROI();	// required to use roi in constraints fn below
+		Octree->RemoveTriangles(Remesher.GetCurrentTriangleROI());
+	}
+
+	FMeshConstraints constraints;
+	bool bConstraintAllowSplits = true;
+	bool bConstraintAllowSmoothing = false;
+	{
+		SCOPE_CYCLE_COUNTER(STAT_SculptTool_Remesh_Constraints);
+
+		// TODO: only constrain in ROI. This is quite difficult to do externally because we need to update based on
+		// the changing triangle set in Remesher. Perhaps FSubRegionRemesher should update the constraints itself?
+
+		if (bHaveUVSeams || bHaveNormalSeams)
+		{
+			FMeshConstraintsUtil::ConstrainAllSeams(constraints, *Mesh, bConstraintAllowSplits, bConstraintAllowSmoothing);
+			Remesher.SetExternalConstraints(&constraints);
+		}
+	}
+
+	for (int k = 0; k < 5; ++k)
+	{
+		if (bIsUniformSmooth == false)
+		{
+			Remesher.bEnableFlips = RemeshProperties->bFlips && (k < 2);
+		}
+
+		{
+			SCOPE_CYCLE_COUNTER(STAT_SculptTool_Remesh_RemeshROIUpdate);
+
+			Remesher.UpdateROI();
+
+			if (ActiveMeshChange != nullptr)
+			{
+				Remesher.SaveActiveROI(ActiveMeshChange);
+				//ActiveMeshChange->VerifySaveState();    // useful for debugging
+			}
+
+			Remesher.BeginTrackRemovedTrisInPass();
+		}
+
+		{
+			SCOPE_CYCLE_COUNTER(STAT_SculptTool_Remesh_RemeshPass);
+			Remesher.BasicRemeshPass();
+		}
+
+		{
+			SCOPE_CYCLE_COUNTER(STAT_SculptTool_Remesh_PassOctreeUpdate);
+			const TSet<int32>& TrisRemovedInPass = Remesher.EndTrackRemovedTrisInPass();
+			Octree->RemoveTriangles(TrisRemovedInPass);
+		}
+	}
+	//UE_LOG(LogTemp, Warning, TEXT("Triangle Count %d after update"), Mesh->TriangleCount());
+
+	{
+		SCOPE_CYCLE_COUNTER(STAT_SculptTool_Remesh_Finish);
+
+		// reinsert new ROI into octree
+		Octree->ReinsertTriangles(Remesher.GetCurrentTriangleROI());
+
+		//Octree->CheckValidity(EValidityCheckFailMode::Check, false, true);
+
+		UpdateROI(LastBrushPosLocal);
+	}
+}
+
+
+
+
+void UDynamicMeshSculptTool::RecalculateNormals_PerVertex()
+{
+	FDynamicMesh3* Mesh = DynamicMeshComponent->GetMesh();
+	FDynamicMeshOctree3* Octree = DynamicMeshComponent->GetOctree();
+
+	int MaxVertexID = Mesh->MaxVertexID();
+	if (NormalsVertexFlags.Num() < MaxVertexID)
+	{
+		NormalsVertexFlags.Init(false, MaxVertexID * 2);
+	}
+
+
+	{
+		SCOPE_CYCLE_COUNTER(SculptTool_Normals_Collect);
+
+		TrianglesBuffer.Reset();
+		NormalsBuffer.Reset();
+		Octree->RangeQuery(Octree->ModifiedBounds, TrianglesBuffer);
+		for (int TriangleID : TrianglesBuffer)
+		{
+			FIndex3i TriV = Mesh->GetTriangle(TriangleID);
+			for (int j = 0; j < 3; ++j)
+			{
+				int vid = TriV[j];
+				if (NormalsVertexFlags[vid] == false)
+				{
+					NormalsBuffer.Add(vid);
+					NormalsVertexFlags[vid] = true;
+				}
+			}
+		}
+	}
+
+	//UE_LOG(LogTemp, Warning, TEXT("Computing %d normals"), NormalsBuffer.Num());
+
+	{
+		SCOPE_CYCLE_COUNTER(SculptTool_Normals_Compute);
+
+		ParallelFor(NormalsBuffer.Num(), [&](int k) {
+			int vid = NormalsBuffer[k];
+			FVector3d NewNormal = FMeshNormals::ComputeVertexNormal(*Mesh, vid);
+			Mesh->SetVertexNormal(vid, (FVector3f)NewNormal);
+			NormalsVertexFlags[vid] = false;
+		});
+	}
+
+}
+
+
+
+
+
+
+
+
+
+void UDynamicMeshSculptTool::RecalculateNormals_Overlay()
+{
+	FDynamicMesh3* Mesh = DynamicMeshComponent->GetMesh();
+	FDynamicMeshNormalOverlay* Normals = Mesh->HasAttributes() ? Mesh->Attributes()->PrimaryNormals() : nullptr;
+	check(Normals != nullptr);
+
+
+	FDynamicMeshOctree3* Octree = DynamicMeshComponent->GetOctree();
+
+	int MaxElementID = Normals->MaxElementID();
+	if (NormalsVertexFlags.Num() < MaxElementID)
+	{
+		NormalsVertexFlags.Init(false, MaxElementID * 2);
+	}
+
+
+	{
+		SCOPE_CYCLE_COUNTER(SculptTool_Normals_Collect);
+
+		TrianglesBuffer.Reset();
+		NormalsBuffer.Reset();
+		Octree->RangeQuery(Octree->ModifiedBounds, TrianglesBuffer);
+		for (int TriangleID : TrianglesBuffer)
+		{
+			FIndex3i TriElems = Normals->GetTriangle(TriangleID);
+			for (int j = 0; j < 3; ++j)
+			{
+				int elemid = TriElems[j];
+				if (NormalsVertexFlags[elemid] == false)
+				{
+					NormalsBuffer.Add(elemid);
+					NormalsVertexFlags[elemid] = true;
+				}
+			}
+		}
+	}
+
+	//UE_LOG(LogTemp, Warning, TEXT("Computing %d normals"), NormalsBuffer.Num());
+
+	{
+		SCOPE_CYCLE_COUNTER(SculptTool_Normals_Compute);
+
+		ParallelFor(NormalsBuffer.Num(), [&](int k) {
+			int elemid = NormalsBuffer[k];
+			FVector3d NewNormal = FMeshNormals::ComputeOverlayNormal(*Mesh, Normals, elemid);
+			Normals->SetElement(elemid, (FVector3f)NewNormal);
+			NormalsVertexFlags[elemid] = false;
+		});
+	}
+
+}
+
+
+
+
+
+
+
+void UDynamicMeshSculptTool::UpdateTarget()
+{
+	if (SculptProperties != nullptr )
+	{
+		bCachedFreezeTarget = SculptProperties->bFreezeTarget;
+		if (SculptProperties->bFreezeTarget)
+		{
+			return;   // do not update frozen target
+		}
+	}
+
+	BrushTargetMesh.Copy(*DynamicMeshComponent->GetMesh(), false, false, false, false);
+	BrushTargetMeshSpatial.SetMesh(&BrushTargetMesh, true);
+
+	BrushTargetNormals.SetMesh(&BrushTargetMesh);
+	BrushTargetNormals.ComputeVertexNormals();
+}
+
+bool UDynamicMeshSculptTool::GetTargetMeshNearest(const FVector3d& Position, double SearchRadius, FVector3d& TargetPosOut, FVector3d& TargetNormalOut)
+{
+	double fDistSqr;
+	int NearTID = BrushTargetMeshSpatial.FindNearestTriangle(Position, fDistSqr, SearchRadius);
+	if (NearTID <= 0)
+	{
+		return false;
+	}
+	FTriangle3d Triangle;
+	BrushTargetMesh.GetTriVertices(NearTID, Triangle.V[0], Triangle.V[1], Triangle.V[2]);
+	FDistPoint3Triangle3d Query(Position, Triangle);
+	Query.Get();
+	FIndex3i Tri = BrushTargetMesh.GetTriangle(NearTID);
+	TargetNormalOut =
+		Query.TriangleBaryCoords.X*BrushTargetNormals[Tri.A]
+		+ Query.TriangleBaryCoords.Y*BrushTargetNormals[Tri.B]
+		+ Query.TriangleBaryCoords.Z*BrushTargetNormals[Tri.C];
+	TargetNormalOut.Normalize();
+	TargetPosOut = Query.ClosestTrianglePoint;
+	return true;
+}
+
+
+double UDynamicMeshSculptTool::EstimateIntialSafeTargetLength(const FDynamicMesh3& Mesh, int MinTargetTriCount)
+{
+	double AreaSum = 0;
+	for (int tid : Mesh.TriangleIndicesItr())
+	{
+		AreaSum += Mesh.GetTriArea(tid);
+	}
+
+	int TriCount = Mesh.TriangleCount();
+	double TargetTriArea = 1.0;
+	if (TriCount < MinTargetTriCount)
+	{
+		TargetTriArea = AreaSum / (double)MinTargetTriCount;
+	}
+	else
+	{
+		TargetTriArea = AreaSum / (double)TriCount;
+	}
+
+	double EdgeLen = TriangleUtil::EquilateralEdgeLengthForArea(TargetTriArea);
+	return (double)FMath::RoundToInt(EdgeLen*100.0) / 100.0;
+}
+
+
+
+
+void UDynamicMeshSculptTool::IncreaseBrushRadiusAction()
+{
+	BrushProperties->BrushSize = FMath::Clamp(BrushProperties->BrushSize + 0.025f, 0.0f, 1.0f);
+	CalculateBrushRadius();
+}
+
+void UDynamicMeshSculptTool::DecreaseBrushRadiusAction()
+{
+	BrushProperties->BrushSize = FMath::Clamp(BrushProperties->BrushSize - 0.025f, 0.0f, 1.0f);
+	CalculateBrushRadius();
+}
+
+void UDynamicMeshSculptTool::IncreaseBrushRadiusSmallStepAction()
+{
+	BrushProperties->BrushSize = FMath::Clamp(BrushProperties->BrushSize + 0.005f, 0.0f, 1.0f);
+	CalculateBrushRadius();
+}
+
+void UDynamicMeshSculptTool::DecreaseBrushRadiusSmallStepAction()
+{
+	BrushProperties->BrushSize = FMath::Clamp(BrushProperties->BrushSize - 0.005f, 0.0f, 1.0f);
+	CalculateBrushRadius();
+}
+
+
+
+void UDynamicMeshSculptTool::IncreaseBrushSpeedAction()
+{
+	SculptProperties->BrushSpeed = FMath::Clamp(SculptProperties->BrushSpeed + 0.05f, 0.0f, 1.0f);
+}
+
+void UDynamicMeshSculptTool::DecreaseBrushSpeedAction()
+{
+	SculptProperties->BrushSpeed = FMath::Clamp(SculptProperties->BrushSpeed - 0.05f, 0.0f, 1.0f);
+}
+
+
+
+void UDynamicMeshSculptTool::NextBrushModeAction()
+{
+	uint8 LastMode = (uint8)EDynamicMeshSculptBrushType::LastValue;
+	SculptProperties->PrimaryBrushType = (EDynamicMeshSculptBrushType)(((uint8)SculptProperties->PrimaryBrushType + 1) % LastMode);
+}
+
+void UDynamicMeshSculptTool::PreviousBrushModeAction()
+{
+	uint8 LastMode = (uint8)EDynamicMeshSculptBrushType::LastValue;
+	uint8 CurMode = (uint8)SculptProperties->PrimaryBrushType;
+	if (CurMode == 0)
+	{
+		SculptProperties->PrimaryBrushType = (EDynamicMeshSculptBrushType)((uint8)LastMode - 1);
+	} 
+	else
+	{
+		SculptProperties->PrimaryBrushType = (EDynamicMeshSculptBrushType)((uint8)CurMode - 1);
+	}
+}
+
+
+void UDynamicMeshSculptTool::NextHistoryBrushModeAction()
+{
+	int MaxHistory = BrushTypeHistory.Num() - 1;
+	if (BrushTypeHistoryIndex < MaxHistory)
+	{
+		BrushTypeHistoryIndex++;
+		SculptProperties->PrimaryBrushType = BrushTypeHistory[BrushTypeHistoryIndex];
+		LastStampType = SculptProperties->PrimaryBrushType;
+	}
+}
+
+void UDynamicMeshSculptTool::PreviousHistoryBrushModeAction()
+{
+	if (BrushTypeHistoryIndex > 0)
+	{
+		BrushTypeHistoryIndex--;
+		SculptProperties->PrimaryBrushType = BrushTypeHistory[BrushTypeHistoryIndex];
+		LastStampType = SculptProperties->PrimaryBrushType;
+	}
+}
+
+
+
+void UDynamicMeshSculptTool::RegisterActions(FInteractiveToolActionSet& ActionSet)
+{
+	ActionSet.RegisterAction(this, (int32)EStandardToolActions::IncreaseBrushSize,
+		TEXT("SculptIncreaseRadius"), 
+		LOCTEXT("SculptIncreaseRadius", "Increase Sculpt Radius"),
+		LOCTEXT("SculptIncreaseRadiusTooltip", "Increase radius of sculpting brush"),
+		EModifierKey::None, EKeys::RightBracket,
+		[this]() { IncreaseBrushRadiusAction(); } );
+
+	ActionSet.RegisterAction(this, (int32)EStandardToolActions::DecreaseBrushSize,
+		TEXT("SculptDecreaseRadius"), 
+		LOCTEXT("SculptDecreaseRadius", "Decrease Sculpt Radius"),
+		LOCTEXT("SculptDecreaseRadiusTooltip", "Decrease radius of sculpting brush"),
+		EModifierKey::None, EKeys::LeftBracket,
+		[this]() { DecreaseBrushRadiusAction(); } );
+
+
+	ActionSet.RegisterAction(this, (int32)EStandardToolActions::BaseClientDefinedActionID+1,
+		TEXT("NextBrushMode"),
+		LOCTEXT("SculptNextBrushMode", "Next Brush Type"),
+		LOCTEXT("SculptNextBrushModeTooltip", "Cycle to next Brush Type"),
+		EModifierKey::None, EKeys::A,
+		[this]() { NextBrushModeAction(); });
+
+	ActionSet.RegisterAction(this, (int32)EStandardToolActions::BaseClientDefinedActionID+2,
+		TEXT("PreviousBrushMode"),
+		LOCTEXT("SculptPreviousBrushMode", "Previous Brush Type"),
+		LOCTEXT("SculptPreviousBrushModeTooltip", "Cycle to previous Brush Type"),
+		EModifierKey::None, EKeys::Q,
+		[this]() { PreviousBrushModeAction(); });
+
+
+	ActionSet.RegisterAction(this, (int32)EStandardToolActions::BaseClientDefinedActionID + 10,
+		TEXT("NextBrushHistoryState"),
+		LOCTEXT("SculptNextBrushHistoryState", "Next Brush History State"),
+		LOCTEXT("SculptSculptNextBrushHistoryStateTooltip", "Cycle to next Brush History state"),
+		EModifierKey::Shift, EKeys::Q,
+		[this]() { NextHistoryBrushModeAction(); });
+
+	ActionSet.RegisterAction(this, (int32)EStandardToolActions::BaseClientDefinedActionID + 11,
+		TEXT("PreviousBrushHistoryState"),
+		LOCTEXT("SculptPreviousBrushHistoryState", "Previous Brush History State"),
+		LOCTEXT("SculptPreviousBrushHistoryStateTooltip", "Cycle to previous Brush History state"),
+		EModifierKey::Shift, EKeys::A,
+		[this]() { PreviousHistoryBrushModeAction(); });
+
+
+
+	ActionSet.RegisterAction(this, (int32)EStandardToolActions::BaseClientDefinedActionID + 50,
+		TEXT("SculptIncreaseSize"),
+		LOCTEXT("SculptIncreaseSize", "Increase Size"),
+		LOCTEXT("SculptIncreaseSizeTooltip", "Increase Brush Size"),
+		EModifierKey::None, EKeys::D,
+		[this]() { IncreaseBrushRadiusAction(); });
+
+	ActionSet.RegisterAction(this, (int32)EStandardToolActions::BaseClientDefinedActionID + 51,
+		TEXT("SculptDecreaseSize"),
+		LOCTEXT("SculptDecreaseSize", "Decrease Size"),
+		LOCTEXT("SculptDecreaseSizeTooltip", "Decrease Brush Size"),
+		EModifierKey::None, EKeys::S,
+		[this]() { DecreaseBrushRadiusAction(); });
+
+	ActionSet.RegisterAction(this, (int32)EStandardToolActions::BaseClientDefinedActionID + 52,
+		TEXT("SculptIncreaseSizeSmallStep"),
+		LOCTEXT("SculptIncreaseSize", "Increase Size"),
+		LOCTEXT("SculptIncreaseSizeTooltip", "Increase Brush Size"),
+		EModifierKey::Shift, EKeys::D,
+		[this]() { IncreaseBrushRadiusSmallStepAction(); });
+
+	ActionSet.RegisterAction(this, (int32)EStandardToolActions::BaseClientDefinedActionID + 53,
+		TEXT("SculptDecreaseSizeSmallStemp"),
+		LOCTEXT("SculptDecreaseSize", "Decrease Size"),
+		LOCTEXT("SculptDecreaseSizeTooltip", "Decrease Brush Size"),
+		EModifierKey::Shift, EKeys::S,
+		[this]() { DecreaseBrushRadiusSmallStepAction(); });
+
+
+
+
+	ActionSet.RegisterAction(this, (int32)EStandardToolActions::BaseClientDefinedActionID + 60,
+		TEXT("SculptIncreaseSpeed"),
+		LOCTEXT("SculptIncreaseSpeed", "Increase Speed"),
+		LOCTEXT("SculptIncreaseSpeedTooltip", "Increase Brush Speed"),
+		EModifierKey::None, EKeys::E,
+		[this]() { IncreaseBrushSpeedAction(); });
+
+	ActionSet.RegisterAction(this, (int32)EStandardToolActions::BaseClientDefinedActionID + 61,
+		TEXT("SculptDecreaseSpeed"),
+		LOCTEXT("SculptDecreaseSpeed", "Decrease Speed"),
+		LOCTEXT("SculptDecreaseSpeedTooltip", "Decrease Brush Speed"),
+		EModifierKey::None, EKeys::W,
+		[this]() { DecreaseBrushSpeedAction(); });
+
+
+
+	ActionSet.RegisterAction(this, (int32)EStandardToolActions::ToggleWireframe,
+		TEXT("ToggleWireframe"),
+		LOCTEXT("ToggleWireframe", "Toggle Wireframe"),
+		LOCTEXT("ToggleWireframeTooltip", "Toggle visibility of wireframe overlay"),
+		EModifierKey::Alt, EKeys::W,
+		[this]() { ViewProperties->bShowWireframe = !ViewProperties->bShowWireframe; });
+
+}
+
+
+
+
+//
+// Change Tracking
+//
+
+
+void UDynamicMeshSculptTool::BeginChange(bool bIsVertexChange)
+{
+	check(ActiveVertexChange == nullptr);
+	check(ActiveMeshChange == nullptr);
+	if (bIsVertexChange) 
+	{
+		ActiveVertexChange = new FMeshVertexChangeBuilder();
+	}
+	else
+	{
+		ActiveMeshChange = new FDynamicMeshChangeTracker(DynamicMeshComponent->GetMesh());
+		ActiveMeshChange->BeginChange();
+	}
+}
+
+
+void UDynamicMeshSculptTool::EndChange()
+{
+	if (ActiveVertexChange != nullptr)
+	{
+		GetToolManager()->EmitObjectChange(DynamicMeshComponent, MoveTemp(ActiveVertexChange->Change), LOCTEXT("VertexSculptChange", "Brush Stroke"));
+
+		delete ActiveVertexChange;
+		ActiveVertexChange = nullptr;
+	}
+
+	if (ActiveMeshChange != nullptr)
+	{
+		FMeshChange* NewMeshChange = new FMeshChange();
+		NewMeshChange->DynamicMeshChange = ActiveMeshChange->EndChange();
+		//NewMeshChange->DynamicMeshChange->CheckValidity();
+		TUniquePtr<FMeshChange> NewChange(NewMeshChange);
+		GetToolManager()->EmitObjectChange(DynamicMeshComponent, MoveTemp(NewChange), LOCTEXT("MeshSculptChange", "Brush Stroke"));
+		delete ActiveMeshChange;
+		ActiveMeshChange = nullptr;
+	}
+}
+
+
+void UDynamicMeshSculptTool::SaveActiveROI()
+{
+	if (ActiveMeshChange != nullptr)
+	{
+		for (int vid : VertexROI)
+		{
+			ActiveMeshChange->SaveVertex(vid);
+		}
+	}
+}
+
+
+void UDynamicMeshSculptTool::UpdateSavedVertex(int vid, const FVector3d& OldPosition, const FVector3d& NewPosition)
+{
+	if (ActiveVertexChange != nullptr)
+	{
+		UpdateSavedVertexLock.Lock();
+		ActiveVertexChange->UpdateVertex(vid, OldPosition, NewPosition);
+		UpdateSavedVertexLock.Unlock();
+	}
+}
+
+
+
+void UDynamicMeshSculptTool::UpdateMaterialMode(EMeshEditingMaterialModes MaterialMode)
+{
+	if (MaterialMode == EMeshEditingMaterialModes::ExistingMaterial)
+	{
+		UMaterialInterface* Material = ComponentTarget->GetMaterial(0);
+		if (Material != nullptr)
+		{
+			DynamicMeshComponent->SetMaterial(0, Material);
+		}
+		DynamicMeshComponent->bCastDynamicShadow = ComponentTarget->GetOwnerComponent()->bCastDynamicShadow;
+	} 
+	else if (MaterialMode == EMeshEditingMaterialModes::MeshFocusMaterial)
+	{ 
+		UMaterialInterface* Material = ToolSetupUtil::GetSculptMaterial1(GetToolManager());
+		if (Material != nullptr)
+		{
+			DynamicMeshComponent->SetMaterial(0, Material);
+		}
+		DynamicMeshComponent->bCastDynamicShadow = false;
+	}
+}
+
+
+#undef LOCTEXT_NAMESPACE
diff --git a/Engine/Plugins/Experimental/MeshModelingToolset/Source/MeshModelingTools/Private/EditNormalsTool.cpp b/Engine/Plugins/Experimental/MeshModelingToolset/Source/MeshModelingTools/Private/EditNormalsTool.cpp
new file mode 100644
--- /dev/null
+++ b/Engine/Plugins/Experimental/MeshModelingToolset/Source/MeshModelingTools/Private/EditNormalsTool.cpp
@@ -0,0 +1,288 @@
+// Copyright 1998-2019 Epic Games, Inc. All Rights Reserved.
+
+#include "EditNormalsTool.h"
+#include "InteractiveToolManager.h"
+#include "ToolBuilderUtil.h"
+
+#include "ToolSetupUtil.h"
+
+#include "DynamicMesh3.h"
+#include "BaseBehaviors/MultiClickSequenceInputBehavior.h"
+#include "Selection/SelectClickedAction.h"
+
+#include "MeshDescriptionToDynamicMesh.h"
+#include "DynamicMeshToMeshDescription.h"
+
+#include "InteractiveGizmoManager.h"
+
+#include "AssetGenerationUtil.h"
+
+
+#define LOCTEXT_NAMESPACE "UEditNormalsTool"
+
+
+/*
+ * ToolBuilder
+ */
+
+
+bool UEditNormalsToolBuilder::CanBuildTool(const FToolBuilderState& SceneState) const
+{
+	return ToolBuilderUtil::CountComponents(SceneState, CanMakeComponentTarget) > 0;
+}
+
+UInteractiveTool* UEditNormalsToolBuilder::BuildTool(const FToolBuilderState& SceneState) const
+{
+	UEditNormalsTool* NewTool = NewObject<UEditNormalsTool>(SceneState.ToolManager);
+
+	TArray<UActorComponent*> Components = ToolBuilderUtil::FindAllComponents(SceneState, CanMakeComponentTarget);
+	check(Components.Num() > 0);
+
+	TArray<TUniquePtr<FPrimitiveComponentTarget>> ComponentTargets;
+	for (UActorComponent* ActorComponent : Components)
+	{
+		auto* MeshComponent = Cast<UPrimitiveComponent>(ActorComponent);
+		if ( MeshComponent )
+		{
+			ComponentTargets.Add(MakeComponentTarget(MeshComponent));
+		}
+	}
+
+	NewTool->SetSelection(MoveTemp(ComponentTargets));
+	NewTool->SetWorld(SceneState.World);
+	NewTool->SetAssetAPI(AssetAPI);
+
+	return NewTool;
+}
+
+
+
+
+/*
+ * Tool
+ */
+
+UEditNormalsToolProperties::UEditNormalsToolProperties()
+{
+	bFixInconsistentNormals = false;
+	bInvertNormals = false;
+	bRecomputeNormals = true;
+	NormalCalculationMethod = ENormalCalculationMethod::AreaAngleWeighting;
+	bRecomputeNormalTopologyAndEdgeSharpness = false;
+	SharpEdgeAngleThreshold = 60;
+	bAllowSharpVertices = false;
+}
+
+UEditNormalsAdvancedProperties::UEditNormalsAdvancedProperties()
+{
+}
+
+
+UEditNormalsTool::UEditNormalsTool()
+{
+}
+
+void UEditNormalsTool::SetWorld(UWorld* World)
+{
+	this->TargetWorld = World;
+}
+
+void UEditNormalsTool::Setup()
+{
+	UInteractiveTool::Setup();
+
+	// hide input StaticMeshComponent
+	for (TUniquePtr<FPrimitiveComponentTarget>& ComponentTarget : ComponentTargets)
+	{
+		ComponentTarget->SetOwnerVisibility(false);
+	}
+
+	BasicProperties = NewObject<UEditNormalsToolProperties>(this, TEXT("Mesh Normals Settings"));
+	AdvancedProperties = NewObject<UEditNormalsAdvancedProperties>(this, TEXT("Advanced Settings"));
+
+	// initialize our properties
+	AddToolPropertySource(BasicProperties);
+	AddToolPropertySource(AdvancedProperties);
+
+	// initialize the PreviewMesh+BackgroundCompute object
+	UpdateNumPreviews();
+
+	
+	
+	for (UMeshOpPreviewWithBackgroundCompute* Preview : Previews)
+	{
+		Preview->InvalidateResult();
+	}
+}
+
+
+void UEditNormalsTool::UpdateNumPreviews()
+{
+	int32 CurrentNumPreview = Previews.Num();
+	int32 TargetNumPreview = ComponentTargets.Num();
+	if (TargetNumPreview < CurrentNumPreview)
+	{
+		for (int32 PreviewIdx = CurrentNumPreview - 1; PreviewIdx >= TargetNumPreview; PreviewIdx--)
+		{
+			Previews[PreviewIdx]->Cancel();
+		}
+		Previews.SetNum(TargetNumPreview);
+		OriginalDynamicMeshes.SetNum(TargetNumPreview);
+	}
+	else
+	{
+		OriginalDynamicMeshes.SetNum(TargetNumPreview);
+		for (int32 PreviewIdx = CurrentNumPreview; PreviewIdx < TargetNumPreview; PreviewIdx++)
+		{
+			UEditNormalsOperatorFactory *OpFactory = NewObject<UEditNormalsOperatorFactory>();
+			OpFactory->Tool = this;
+			OpFactory->ComponentIndex = PreviewIdx;
+			OriginalDynamicMeshes[PreviewIdx] = MakeShared<FDynamicMesh3>();
+			FMeshDescriptionToDynamicMesh Converter;
+			Converter.bPrintDebugMessages = true;
+			Converter.Convert(ComponentTargets[PreviewIdx]->GetMesh(), *OriginalDynamicMeshes[PreviewIdx]);
+
+			UMeshOpPreviewWithBackgroundCompute* Preview = Previews.Add_GetRef(NewObject<UMeshOpPreviewWithBackgroundCompute>(OpFactory, "Preview"));
+			Preview->Setup(this->TargetWorld, OpFactory);
+			Preview->ConfigureMaterials(
+				ToolSetupUtil::GetDefaultMaterial(GetToolManager(), ComponentTargets[PreviewIdx]->GetMaterial(0)),
+				ToolSetupUtil::GetDefaultWorkingMaterial(GetToolManager())
+			);
+			Preview->SetVisibility(true);
+		}
+	}
+}
+
+
+void UEditNormalsTool::Shutdown(EToolShutdownType ShutdownType)
+{
+	// Restore (unhide) the source meshes
+	for (TUniquePtr<FPrimitiveComponentTarget>& ComponentTarget : ComponentTargets)
+	{
+		ComponentTarget->SetOwnerVisibility(true);
+	}
+
+	TArray<TUniquePtr<FDynamicMeshOpResult>> Results;
+	for (UMeshOpPreviewWithBackgroundCompute* Preview : Previews)
+	{
+		Results.Emplace(Preview->Shutdown());
+	}
+	if (ShutdownType == EToolShutdownType::Accept)
+	{
+		GenerateAsset(Results);
+	}
+}
+
+void UEditNormalsTool::SetAssetAPI(IToolsContextAssetAPI* AssetAPIIn)
+{
+	this->AssetAPI = AssetAPIIn;
+}
+
+TSharedPtr<FDynamicMeshOperator> UEditNormalsOperatorFactory::MakeNewOperator()
+{
+	TSharedPtr<FEditNormalsOp> NormalsOp = MakeShared<FEditNormalsOp>();
+	NormalsOp->bFixInconsistentNormals = Tool->BasicProperties->bFixInconsistentNormals;
+	NormalsOp->bInvertNormals = Tool->BasicProperties->bInvertNormals;
+	NormalsOp->bRecomputeNormals = Tool->BasicProperties->bRecomputeNormals;
+	NormalsOp->bSplitNormals = Tool->BasicProperties->bRecomputeNormalTopologyAndEdgeSharpness;
+	NormalsOp->bAllowSharpVertices = Tool->BasicProperties->bAllowSharpVertices;
+	NormalsOp->NormalCalculationMethod = Tool->BasicProperties->NormalCalculationMethod;
+	NormalsOp->NormalSplitThreshold = Tool->BasicProperties->SharpEdgeAngleThreshold;
+
+	FTransform LocalToWorld = Tool->ComponentTargets[ComponentIndex]->GetWorldTransform();
+	NormalsOp->OriginalMesh = Tool->OriginalDynamicMeshes[ComponentIndex];
+	
+	NormalsOp->SetTransform(LocalToWorld);
+
+	return NormalsOp;
+}
+
+
+
+void UEditNormalsTool::Render(IToolsContextRenderAPI* RenderAPI)
+{
+}
+
+void UEditNormalsTool::Tick(float DeltaTime)
+{
+	for (UMeshOpPreviewWithBackgroundCompute* Preview : Previews)
+	{
+		Preview->Tick(DeltaTime);
+	}
+}
+
+
+
+#if WITH_EDITOR
+void UEditNormalsTool::PostEditChangeProperty(FPropertyChangedEvent& PropertyChangedEvent)
+{
+	UpdateNumPreviews();
+	for (UMeshOpPreviewWithBackgroundCompute* Preview : Previews)
+	{
+		Preview->InvalidateResult();
+	}
+}
+#endif
+
+void UEditNormalsTool::OnPropertyModified(UObject* PropertySet, UProperty* Property)
+{
+	UpdateNumPreviews();
+	for (UMeshOpPreviewWithBackgroundCompute* Preview : Previews)
+	{
+		Preview->InvalidateResult();
+	}
+}
+
+
+
+bool UEditNormalsTool::HasAccept() const
+{
+	return true;
+}
+
+bool UEditNormalsTool::CanAccept() const
+{
+	for (UMeshOpPreviewWithBackgroundCompute* Preview : Previews)
+	{
+		if (!Preview->HaveValidResult())
+		{
+			return false;
+		}
+	}
+	return true;
+}
+
+
+void UEditNormalsTool::GenerateAsset(const TArray<TUniquePtr<FDynamicMeshOpResult>>& Results)
+{
+	GetToolManager()->BeginUndoTransaction(LOCTEXT("EditNormalsToolTransactionName", "Edit Normals Tool"));
+
+	check(Results.Num() == ComponentTargets.Num());
+	
+	for (int32 ComponentIdx = 0; ComponentIdx < ComponentTargets.Num(); ComponentIdx++)
+	{
+		check(Results[ComponentIdx]->Mesh.Get() != nullptr);
+		ComponentTargets[ComponentIdx]->CommitMesh([&Results, &ComponentIdx, this](FMeshDescription* MeshDescription)
+		{
+			FDynamicMeshToMeshDescription Converter;
+			
+			if (BasicProperties->WillTopologyChange())
+			{
+				// full conversion if normal topology changed or faces were inverted
+				Converter.Convert(Results[ComponentIdx]->Mesh.Get(), *MeshDescription);
+			}
+			else
+			{
+				// otherwise just copy attributes
+				Converter.UpdateAttributes(Results[ComponentIdx]->Mesh.Get(), *MeshDescription, true, false);
+			}
+		});
+	}
+
+	GetToolManager()->EndUndoTransaction();
+}
+
+
+
+
+#undef LOCTEXT_NAMESPACE
diff --git a/Engine/Plugins/Experimental/MeshModelingToolset/Source/MeshModelingTools/Private/FindPolygonsAlgorithm.cpp b/Engine/Plugins/Experimental/MeshModelingToolset/Source/MeshModelingTools/Private/FindPolygonsAlgorithm.cpp
new file mode 100644
--- /dev/null
+++ b/Engine/Plugins/Experimental/MeshModelingToolset/Source/MeshModelingTools/Private/FindPolygonsAlgorithm.cpp
@@ -0,0 +1,98 @@
+// Copyright 1998-2019 Epic Games, Inc. All Rights Reserved.
+
+#include "FindPolygonsAlgorithm.h"
+#include "MeshNormals.h"
+
+
+
+FFindPolygonsAlgorithm::FFindPolygonsAlgorithm(FDynamicMesh3* MeshIn)
+{
+	Mesh = MeshIn;
+}
+
+
+bool FFindPolygonsAlgorithm::FindPolygons(double DotTolerance)
+{
+	DotTolerance = 1.0 - DotTolerance;
+
+	// compute face normals
+	FMeshNormals Normals(Mesh);
+	Normals.ComputeTriangleNormals();
+
+	TArray<bool> DoneTriangle;
+	DoneTriangle.SetNum(Mesh->MaxTriangleID());
+	PolygonGroupIDs.SetNum(Mesh->MaxTriangleID());
+
+	TArray<int> Stack;
+
+	// grow outward from vertices until we have no more left
+	for (int TriID : Mesh->TriangleIndicesItr())
+	{
+		if (DoneTriangle[TriID] == true)
+		{
+			continue;
+		}
+
+		TArray<int> Polygon;
+		Polygon.Add(TriID);
+		DoneTriangle[TriID] = true;
+
+		Stack.SetNum(0);
+		Stack.Add(TriID);
+		while (Stack.Num() > 0)
+		{
+			int CurTri = Stack.Pop(false);
+			FIndex3i NbrTris = Mesh->GetTriNeighbourTris(CurTri);
+			for (int j = 0; j < 3; ++j)
+			{
+				if (NbrTris[j] >= 0
+					&& DoneTriangle[NbrTris[j]] == false
+					&& PolygonGroupIDs[CurTri] == PolygonGroupIDs[NbrTris[j]])
+				{
+					double Dot = Normals[CurTri].Dot(Normals[NbrTris[j]]);
+					if (Dot > DotTolerance)
+					{
+						Polygon.Add(NbrTris[j]);
+						Stack.Add(NbrTris[j]);
+						DoneTriangle[NbrTris[j]] = true;
+					}
+				}
+			}
+		}
+
+		FoundPolygons.Add(Polygon);
+	}
+
+	Mesh->EnableTriangleGroups(0);
+
+	// set groups from polygons
+	int NumPolygons = FoundPolygons.Num();
+	PolygonTags.SetNum(NumPolygons);
+	PolygonNormals.SetNum(NumPolygons);
+	for (int PolyIdx = 0; PolyIdx < NumPolygons; PolyIdx++)
+	{
+		int Count = FoundPolygons[PolyIdx].Num();
+		for (int k = 0; k < Count; ++k)
+		{
+			Mesh->SetTriangleGroup(FoundPolygons[PolyIdx][k], (PolyIdx + 1));
+		}
+		PolygonTags[PolyIdx] = (PolyIdx + 1);
+		PolygonNormals[PolyIdx] = Normals[FoundPolygons[PolyIdx][0]];
+	}
+
+	return (NumPolygons > 0);
+}
+
+
+
+bool FFindPolygonsAlgorithm::FindPolygonEdges()
+{
+	for (int eid : Mesh->EdgeIndicesItr())
+	{
+		if (Mesh->IsGroupBoundaryEdge(eid))
+		{
+			PolygonEdges.Add(eid);
+		}
+	}
+	return (PolygonEdges.Num() > 0);
+}
