Copying //UE4/Dev-Editor to Dev-Main (//UE4/Dev-Main) for 4.24

diff --git a/Engine/Plugins/Experimental/MeshModelingToolset/Source/MeshModelingTools/Private/MeshInspectorTool.cpp b/Engine/Plugins/Experimental/MeshModelingToolset/Source/MeshModelingTools/Private/MeshInspectorTool.cpp
new file mode 100644
--- /dev/null
+++ b/Engine/Plugins/Experimental/MeshModelingToolset/Source/MeshModelingTools/Private/MeshInspectorTool.cpp
@@ -0,0 +1,347 @@
+// Copyright 1998-2019 Epic Games, Inc. All Rights Reserved.
+
+#include "MeshInspectorTool.h"
+#include "InteractiveToolManager.h"
+#include "ToolBuilderUtil.h"
+
+#include "DynamicMesh3.h"
+#include "DynamicMeshAttributeSet.h"
+#include "MeshNormals.h"
+
+#include "SimpleDynamicMeshComponent.h"
+#include "Materials/MaterialInstanceDynamic.h"
+
+#include "Properties/MeshStatisticsProperties.h"
+#include "Properties/MeshAnalysisProperties.h"
+
+#include "SceneManagement.h" // for FPrimitiveDrawInterface
+
+#define LOCTEXT_NAMESPACE "UMeshInspectorTool"
+
+
+/*
+ * ToolBuilder
+ */
+
+
+bool UMeshInspectorToolBuilder::CanBuildTool(const FToolBuilderState& SceneState) const
+{
+	return ToolBuilderUtil::CountComponents(SceneState, CanMakeComponentTarget) == 1;
+}
+
+UInteractiveTool* UMeshInspectorToolBuilder::BuildTool(const FToolBuilderState& SceneState) const
+{
+	UMeshInspectorTool* NewTool = NewObject<UMeshInspectorTool>(SceneState.ToolManager);
+
+	UActorComponent* ActorComponent = ToolBuilderUtil::FindFirstComponent(SceneState, CanMakeComponentTarget);
+	auto* MeshComponent = Cast<UPrimitiveComponent>(ActorComponent);
+	check(MeshComponent != nullptr);
+	NewTool->SetSelection(MakeComponentTarget(MeshComponent));
+
+	return NewTool;
+}
+
+
+
+/*
+ * Properties
+ */
+
+void UMeshInspectorProperties::SaveProperties(UInteractiveTool* SaveFromTool)
+{
+	UMeshInspectorProperties* PropertyCache = GetPropertyCache<UMeshInspectorProperties>();
+	PropertyCache->bWireframe = this->bWireframe;
+	PropertyCache->bBoundaryEdges = this->bBoundaryEdges;
+	PropertyCache->bPolygonBorders = this->bPolygonBorders;
+	PropertyCache->bUVSeams = this->bUVSeams;
+	PropertyCache->bNormalSeams = this->bNormalSeams;
+	PropertyCache->bNormalVectors = this->bNormalVectors;
+	PropertyCache->bTangentVectors = this->bTangentVectors;
+	PropertyCache->NormalLength = this->NormalLength;
+	PropertyCache->TangentLength = this->TangentLength;
+}
+
+void UMeshInspectorProperties::RestoreProperties(UInteractiveTool* RestoreToTool)
+{
+	UMeshInspectorProperties* PropertyCache = GetPropertyCache<UMeshInspectorProperties>();
+	this->bWireframe = PropertyCache->bWireframe;
+	this->bBoundaryEdges = PropertyCache->bBoundaryEdges;
+	this->bPolygonBorders = PropertyCache->bPolygonBorders;
+	this->bUVSeams = PropertyCache->bUVSeams;
+	this->bNormalSeams = PropertyCache->bNormalSeams;
+	this->bNormalVectors = PropertyCache->bNormalVectors;
+	this->bTangentVectors = PropertyCache->bTangentVectors;
+	this->NormalLength = PropertyCache->NormalLength;
+	this->TangentLength = PropertyCache->TangentLength;
+}
+
+
+
+
+/*
+ * Tool
+ */
+
+UMeshInspectorTool::UMeshInspectorTool()
+{
+}
+
+void UMeshInspectorTool::Setup()
+{
+	UInteractiveTool::Setup();
+
+	// create dynamic mesh component to use for live preview
+	DynamicMeshComponent = NewObject<USimpleDynamicMeshComponent>(ComponentTarget->GetOwnerActor(), "DynamicMesh");
+	DynamicMeshComponent->SetupAttachment(ComponentTarget->GetOwnerActor()->GetRootComponent());
+	DynamicMeshComponent->RegisterComponent();
+	DynamicMeshComponent->SetWorldTransform(ComponentTarget->GetWorldTransform());
+
+	// copy material if there is one
+	DefaultMaterial = ComponentTarget->GetMaterial(0);
+	if (DefaultMaterial != nullptr)
+	{
+		DynamicMeshComponent->SetMaterial(0, DefaultMaterial);
+	}
+
+	DynamicMeshComponent->TangentsType = EDynamicMeshTangentCalcType::ExternallyCalculated;
+	DynamicMeshComponent->InitializeMesh(ComponentTarget->GetMesh());
+
+	Precompute();
+
+	// hide input StaticMeshComponent
+	ComponentTarget->SetOwnerVisibility(false);
+
+	// initialize our properties
+	Settings = NewObject<UMeshInspectorProperties>(this);
+	Settings->RestoreProperties(this);
+	AddToolPropertySource(Settings);
+
+	MaterialSettings = NewObject<UExistingMeshMaterialProperties>(this);
+	MaterialSettings->Setup();
+	AddToolPropertySource(MaterialSettings);
+
+
+	DynamicMeshComponent->bExplicitShowWireframe = Settings->bWireframe;
+
+	UMeshStatisticsProperties* Statistics = NewObject<UMeshStatisticsProperties>(this);
+	Statistics->Update(*DynamicMeshComponent->GetMesh());
+	AddToolPropertySource(Statistics);
+
+	UMeshAnalysisProperties* MeshAnalysis = NewObject<UMeshAnalysisProperties>(this);
+	MeshAnalysis->Update(*DynamicMeshComponent->GetMesh(), ComponentTarget->GetWorldTransform());
+	AddToolPropertySource(MeshAnalysis);
+}
+
+
+void UMeshInspectorTool::Shutdown(EToolShutdownType ShutdownType)
+{
+	if (DynamicMeshComponent != nullptr)
+	{
+		ComponentTarget->SetOwnerVisibility(true);
+
+		DynamicMeshComponent->UnregisterComponent();
+		DynamicMeshComponent->DestroyComponent();
+		DynamicMeshComponent = nullptr;
+	}
+
+	Settings->SaveProperties(this);
+}
+
+
+void UMeshInspectorTool::Precompute()
+{
+	BoundaryEdges.Reset();
+	UVSeamEdges.Reset();
+	NormalSeamEdges.Reset();
+
+	FDynamicMesh3* TargetMesh = DynamicMeshComponent->GetMesh();
+	FDynamicMeshUVOverlay* UVOverlay =
+		TargetMesh->HasAttributes() ? TargetMesh->Attributes()->PrimaryUV() : nullptr;
+	FDynamicMeshNormalOverlay* NormalOverlay =
+		TargetMesh->HasAttributes() ? TargetMesh->Attributes()->PrimaryNormals() : nullptr;
+
+	for (int eid : TargetMesh->EdgeIndicesItr())
+	{
+		if (TargetMesh->IsBoundaryEdge(eid))
+		{
+			BoundaryEdges.Add(eid);
+		}
+		if (UVOverlay != nullptr && UVOverlay->IsSeamEdge(eid))
+		{
+			UVSeamEdges.Add(eid);
+		}
+		if (NormalOverlay != nullptr && NormalOverlay->IsSeamEdge(eid))
+		{
+			NormalSeamEdges.Add(eid);
+		}
+		if (TargetMesh->IsGroupBoundaryEdge(eid))
+		{
+			GroupBoundaryEdges.Add(eid);
+		}
+	}
+}
+
+
+void UMeshInspectorTool::Render(IToolsContextRenderAPI* RenderAPI)
+{
+	FPrimitiveDrawInterface* PDI = RenderAPI->GetPrimitiveDrawInterface();
+	FTransform Transform = ComponentTarget->GetWorldTransform();
+
+	FColor BoundaryEdgeColor(240, 15, 15);
+	float BoundaryEdgeThickness = LineWidthMultiplier * 4.0;
+	FColor UVSeamColor(15, 240, 15);
+	float UVSeamThickness = LineWidthMultiplier * 2.0;
+	FColor NormalSeamColor(15, 240, 240);
+	float NormalSeamThickness = LineWidthMultiplier * 2.0;
+	FColor PolygonBorderColor(240, 15, 240);
+	float PolygonBorderThickness = LineWidthMultiplier * 2.0;
+	FColor NormalColor(15, 15, 240);
+	float NormalThickness = LineWidthMultiplier * 2.0f;
+	FColor TangentColor(240, 15, 15);
+	FColor BinormalColor(15, 240, 15);
+	float TangentThickness = LineWidthMultiplier * 2.0f;
+
+	FDynamicMesh3* TargetMesh = DynamicMeshComponent->GetMesh();
+	FVector3d A, B;
+
+	if (Settings->bBoundaryEdges)
+	{
+		for (int eid : BoundaryEdges)
+		{
+			TargetMesh->GetEdgeV(eid, A, B);
+			PDI->DrawLine(Transform.TransformPosition(A), Transform.TransformPosition(B),
+				BoundaryEdgeColor, 0, BoundaryEdgeThickness, 2.0f, true);
+		}
+	}
+
+	if (Settings->bUVSeams)
+	{
+		for (int eid : UVSeamEdges)
+		{
+			TargetMesh->GetEdgeV(eid, A, B);
+			PDI->DrawLine(Transform.TransformPosition(A), Transform.TransformPosition(B),
+				UVSeamColor, 0, UVSeamThickness, 3.0f, true);
+		}
+	}
+
+	if (Settings->bNormalSeams)
+	{
+		for (int eid : NormalSeamEdges)
+		{
+			TargetMesh->GetEdgeV(eid, A, B);
+			PDI->DrawLine(Transform.TransformPosition(A), Transform.TransformPosition(B),
+				NormalSeamColor, 0, NormalSeamThickness, 3.0f, true);
+		}
+	}
+
+	if (Settings->bPolygonBorders)
+	{
+		for (int eid : GroupBoundaryEdges)
+		{
+			TargetMesh->GetEdgeV(eid, A, B);
+			PDI->DrawLine(Transform.TransformPosition(A), Transform.TransformPosition(B),
+				PolygonBorderColor, 0, PolygonBorderThickness, 2.0f, true);
+		}
+	}
+
+	if (Settings->bNormalVectors && TargetMesh->HasAttributes() && TargetMesh->Attributes()->PrimaryNormals() != nullptr)
+	{
+		const FDynamicMeshNormalOverlay* NormalOverlay = TargetMesh->Attributes()->PrimaryNormals();
+		FVector3d TriV[3];
+		FVector3f TriN[3];
+		for (int tid : TargetMesh->TriangleIndicesItr())
+		{
+			TargetMesh->GetTriVertices(tid, TriV[0], TriV[1], TriV[2]);
+			NormalOverlay->GetTriElements(tid, TriN[0], TriN[1], TriN[2]);
+			for (int j = 0; j < 3; ++j)
+			{
+				TriV[j] = Transform.TransformPosition(TriV[j]);
+				PDI->DrawLine(TriV[j], TriV[j] + Settings->NormalLength * Transform.TransformVectorNoScale(TriN[j]),
+					NormalColor, SDPG_World, NormalThickness, 0.0f, true);
+			}
+		}
+	}
+
+	if (Settings->bTangentVectors)
+	{
+		const FMeshTangentsf* Tangents = DynamicMeshComponent->GetTangents();
+		for (int TID : TargetMesh->TriangleIndicesItr())
+		{
+			FVector3d TriV[3];
+			TargetMesh->GetTriVertices(TID, TriV[0], TriV[1], TriV[2]);
+			for (int SubIdx = 0; SubIdx < 3; SubIdx++)
+			{
+				FVector Vert = Transform.TransformPosition(TriV[SubIdx]);
+				FVector3f Tangent, Bitangent;
+				Tangents->GetPerTriangleTangent(TID, SubIdx, Tangent, Bitangent);
+				PDI->DrawLine(Vert, Vert + Settings->TangentLength * Transform.TransformVectorNoScale(Tangent),
+					TangentColor, SDPG_World, TangentThickness, 3.5f, true);
+				PDI->DrawLine(Vert, Vert + Settings->TangentLength * Transform.TransformVectorNoScale(Bitangent),
+					BinormalColor, SDPG_World, TangentThickness, 3.5f, true);
+			}
+		}
+	}
+
+
+}
+
+
+void UMeshInspectorTool::OnPropertyModified(UObject* PropertySet, UProperty* Property)
+{
+	GetToolManager()->PostInvalidation();
+	DynamicMeshComponent->bExplicitShowWireframe = Settings->bWireframe;
+	MaterialSettings->UpdateMaterials();
+	MaterialSettings->SetMaterialIfChanged(DefaultMaterial, DynamicMeshComponent->GetMaterial(0),
+		[this](UMaterialInterface* Material)
+	{
+		DynamicMeshComponent->SetMaterial(0, Material);
+	});
+}
+
+
+bool UMeshInspectorTool::HasAccept() const
+{
+	return false;
+}
+
+bool UMeshInspectorTool::CanAccept() const
+{
+	return false;
+}
+
+
+
+void UMeshInspectorTool::RegisterActions(FInteractiveToolActionSet& ActionSet)
+{
+	ActionSet.RegisterAction(this, (int32)EStandardToolActions::BaseClientDefinedActionID + 1,
+		TEXT("IncreaseLineWidth"), 
+		LOCTEXT("IncreaseLineWidth", "Increase Line Width"),
+		LOCTEXT("IncreaseLineWidthTooltip", "Increase line width of rendering"),
+		EModifierKey::Shift, EKeys::Equals,
+		[this]() { IncreaseLineWidthAction(); });
+
+	ActionSet.RegisterAction(this, (int32)EStandardToolActions::BaseClientDefinedActionID + 2,
+		TEXT("DecreaseLineWidth"), 
+		LOCTEXT("DecreaseLineWidth", "Decrease Line Width"),
+		LOCTEXT("DecreaseLineWidthTooltip", "Decrease line width of rendering"),
+		EModifierKey::None, EKeys::Equals,
+		[this]() { DecreaseLineWidthAction(); });
+}
+
+
+
+void UMeshInspectorTool::IncreaseLineWidthAction()
+{
+	LineWidthMultiplier = LineWidthMultiplier * 1.25f;
+	GetToolManager()->PostInvalidation();
+}
+
+void UMeshInspectorTool::DecreaseLineWidthAction()
+{
+	LineWidthMultiplier = LineWidthMultiplier * (1.0f / 1.25f);
+	GetToolManager()->PostInvalidation();
+}
+
+
+
+#undef LOCTEXT_NAMESPACE
diff --git a/Engine/Plugins/Experimental/MeshModelingToolset/Source/MeshModelingTools/Private/MeshSelectionTool.cpp b/Engine/Plugins/Experimental/MeshModelingToolset/Source/MeshModelingTools/Private/MeshSelectionTool.cpp
new file mode 100644
--- /dev/null
+++ b/Engine/Plugins/Experimental/MeshModelingToolset/Source/MeshModelingTools/Private/MeshSelectionTool.cpp
@@ -0,0 +1,982 @@
+// Copyright 1998-2019 Epic Games, Inc. All Rights Reserved.
+
+#include "MeshSelectionTool.h"
+#include "InteractiveToolManager.h"
+#include "ToolBuilderUtil.h"
+#include "Drawing/MeshDebugDrawing.h"
+#include "DynamicMeshEditor.h"
+#include "DynamicMeshChangeTracker.h"
+#include "Changes/ToolCommandChangeSequence.h"
+#include "Changes/MeshChange.h"
+#include "MeshIndexUtil.h"
+#include "AssetGenerationUtil.h"
+#include "ToolSetupUtil.h"
+
+#define LOCTEXT_NAMESPACE "UMeshSelectionTool"
+
+/*
+ * ToolBuilder
+ */
+
+UMeshSurfacePointTool* UMeshSelectionToolBuilder::CreateNewTool(const FToolBuilderState& SceneState) const
+{
+	UMeshSelectionTool* SelectionTool = NewObject<UMeshSelectionTool>(SceneState.ToolManager);
+	SelectionTool->SetWorld(SceneState.World);
+	SelectionTool->SetAssetAPI(AssetAPI);
+	return SelectionTool;
+}
+
+
+
+
+void UMeshSelectionToolActionPropertySet::PostAction(EMeshSelectionToolActions Action)
+{
+	if (ParentTool.IsValid())
+	{
+		ParentTool->RequestAction(Action);
+	}
+}
+
+
+void UMeshSelectionToolProperties::SaveProperties(UInteractiveTool* SaveFromTool)
+{
+	UMeshSelectionToolProperties* PropertyCache = GetPropertyCache<UMeshSelectionToolProperties>();
+	PropertyCache->SelectionMode = this->SelectionMode;
+	PropertyCache->AngleTolerance = this->AngleTolerance;
+	PropertyCache->bVolumetricBrush = this->bVolumetricBrush;
+	PropertyCache->bHitBackFaces = this->bHitBackFaces;
+	PropertyCache->bShowWireframe = this->bShowWireframe;
+}
+
+void UMeshSelectionToolProperties::RestoreProperties(UInteractiveTool* RestoreToTool)
+{
+	UMeshSelectionToolProperties* PropertyCache = GetPropertyCache<UMeshSelectionToolProperties>();
+	this->SelectionMode = PropertyCache->SelectionMode;
+	this->AngleTolerance = PropertyCache->AngleTolerance;
+	this->bVolumetricBrush = PropertyCache->bVolumetricBrush;
+	this->bHitBackFaces = PropertyCache->bHitBackFaces;
+	this->bShowWireframe = PropertyCache->bShowWireframe;
+}
+
+
+
+/*
+ * Tool
+ */
+
+UMeshSelectionTool::UMeshSelectionTool()
+{
+}
+
+void UMeshSelectionTool::SetWorld(UWorld* World)
+{
+	this->TargetWorld = World;
+}
+
+void UMeshSelectionTool::SetAssetAPI(IToolsContextAssetAPI* AssetAPIIn)
+{
+	this->AssetAPI = AssetAPIIn;
+}
+
+
+
+void UMeshSelectionTool::Setup()
+{
+	UDynamicMeshBrushTool::Setup();
+
+	SelectionProps = NewObject<UMeshSelectionToolProperties>(this);
+	SelectionProps->RestoreProperties(this);
+	AddToolPropertySource(SelectionProps);
+
+	SelectionActions = NewObject<UMeshSelectionEditActions>(this);
+	SelectionActions->Initialize(this);
+	AddToolPropertySource(SelectionActions);
+
+	EditActions = NewObject<UMeshSelectionMeshEditActions>(this);
+	EditActions->Initialize(this);
+	AddToolPropertySource(EditActions);
+
+	// enable wireframe on component
+	PreviewMesh->EnableWireframe(true);
+
+	// set vertex color material on base component so we can see selection
+	UMaterialInterface* SelectionMat = ToolSetupUtil::GetSelectionMaterial(GetToolManager());
+	if (SelectionMat != nullptr)
+	{
+		PreviewMesh->SetMaterial(SelectionMat);
+	}
+	PreviewMesh->GetRootComponent()->bCastDynamicShadow = false;
+
+	const FDynamicMesh3* Mesh = PreviewMesh->GetPreviewDynamicMesh();
+	SelectedVertices = TBitArray<>(false, Mesh->MaxVertexID());
+	SelectedTriangles = TBitArray<>(false, Mesh->MaxTriangleID());
+
+	this->Selection = NewObject<UMeshSelectionSet>(this);
+	Selection->GetOnModified().AddLambda([this](USelectionSet* SelectionObj)
+	{
+		OnExternalSelectionChange();
+	});
+
+	// rebuild octree if mesh changes
+	PreviewMesh->GetOnMeshChanged().AddLambda([this]() { bOctreeValid = false; });
+
+	ShowWireframeWatcher.Initialize(
+		[this]() { return SelectionProps->bShowWireframe; },
+		[this](bool bNewValue) { PreviewMesh->EnableWireframe(bNewValue); }, false);
+}
+
+
+
+
+void UMeshSelectionTool::OnShutdown(EToolShutdownType ShutdownType)
+{
+	SelectionProps->SaveProperties(this);
+
+	if (bHaveModifiedMesh && ShutdownType == EToolShutdownType::Accept)
+	{
+		// this block bakes the modified DynamicMeshComponent back into the StaticMeshComponent inside an undo transaction
+		GetToolManager()->BeginUndoTransaction(LOCTEXT("MeshSelectionToolTransactionName", "Edit Mesh"));
+
+		ComponentTarget->CommitMesh([=](FMeshDescription* MeshDescription)
+		{
+			PreviewMesh->Bake(MeshDescription, true);
+		});
+		GetToolManager()->EndUndoTransaction();
+	}
+}
+
+
+
+
+void UMeshSelectionTool::RegisterActions(FInteractiveToolActionSet& ActionSet)
+{
+	UDynamicMeshBrushTool::RegisterActions(ActionSet);
+
+	ActionSet.RegisterAction(this, (int32)EStandardToolActions::BaseClientDefinedActionID + 1,
+		TEXT("MeshSelectionToolDelete"),
+		LOCTEXT("MeshSelectionToolDelete", "Delete"),
+		LOCTEXT("MeshSelectionToolDeleteTooltip", "Delete Selected Elements"),
+		EModifierKey::None, EKeys::Delete,
+		[this]() { DeleteSelectedTriangles(); });
+
+
+	ActionSet.RegisterAction(this, (int32)EStandardToolActions::ToggleWireframe,
+		TEXT("ToggleWireframe"),
+		LOCTEXT("ToggleWireframe", "Toggle Wireframe"),
+		LOCTEXT("ToggleWireframeTooltip", "Toggle visibility of wireframe overlay"),
+		EModifierKey::Alt, EKeys::W,
+		[this]() { SelectionProps->bShowWireframe = !SelectionProps->bShowWireframe; });
+}
+
+
+
+void UMeshSelectionTool::OnExternalSelectionChange()
+{
+	SelectedVertices.SetRange(0, SelectedVertices.Num(), false);
+	SelectedTriangles.SetRange(0, SelectedTriangles.Num(), false);
+
+	if (SelectionType == EMeshSelectionElementType::Vertex)
+	{
+		for (int VertIdx : Selection->Vertices)
+		{
+			SelectedVertices[VertIdx] = true;
+		}
+	}
+	else if (SelectionType == EMeshSelectionElementType::Face)
+	{
+		for (int FaceIdx : Selection->Faces)
+		{
+			SelectedTriangles[FaceIdx] = true;
+		}
+	}
+
+	OnSelectionUpdated();
+}
+
+
+
+
+bool UMeshSelectionTool::HitTest(const FRay& Ray, FHitResult& OutHit)
+{
+	bool bHit = UDynamicMeshBrushTool::HitTest(Ray, OutHit);
+	if (bHit && SelectionProps->bHitBackFaces == false)
+	{
+		const FDynamicMesh3* SourceMesh = PreviewMesh->GetPreviewDynamicMesh();
+		FVector3d Normal, Centroid; 
+		double Area;
+		SourceMesh->GetTriInfo(OutHit.FaceIndex, Normal, Area, Centroid);
+		FViewCameraState StateOut;
+		GetToolManager()->GetContextQueriesAPI()->GetCurrentViewState(StateOut);
+		FVector3d LocalEyePosition(ComponentTarget->GetWorldTransform().InverseTransformPosition(StateOut.Position));
+
+		if (Normal.Dot((Centroid - LocalEyePosition)) > 0)
+		{
+			bHit = false;
+		}
+	}
+	return bHit;
+}
+
+
+void UMeshSelectionTool::OnBeginDrag(const FRay& WorldRay)
+{
+	UDynamicMeshBrushTool::OnBeginDrag(WorldRay);
+
+	PreviewBrushROI.Reset();
+	if (IsInBrushStroke())
+	{
+		bInRemoveStroke = GetShiftToggle();
+		BeginChange(bInRemoveStroke == false);
+		StartStamp = UBaseBrushTool::LastBrushStamp;
+		LastStamp = StartStamp;
+		bStampPending = true;
+	}
+}
+
+
+
+void UMeshSelectionTool::OnUpdateDrag(const FRay& WorldRay)
+{
+	UDynamicMeshBrushTool::OnUpdateDrag(WorldRay);
+	if (IsInBrushStroke())
+	{
+		LastStamp = UBaseBrushTool::LastBrushStamp;
+		bStampPending = true;
+	}
+}
+
+
+
+TUniquePtr<FDynamicMeshOctree3>& UMeshSelectionTool::GetOctree()
+{
+	if (bOctreeValid == false)
+	{
+		Octree = MakeUnique<FDynamicMeshOctree3>();
+		Octree->Initialize(PreviewMesh->GetPreviewDynamicMesh());
+		bOctreeValid = true;
+	}
+	return Octree;
+}
+
+
+
+void UMeshSelectionTool::CalculateVertexROI(const FBrushStampData& Stamp, TArray<int>& VertexROI)
+{
+	FTransform Transform = ComponentTarget->GetWorldTransform();
+	FVector StampPosLocal = Transform.InverseTransformPosition(Stamp.WorldPosition);
+
+	// TODO: need dynamic vertex hash table!
+
+	float RadiusSqr = CurrentBrushRadius * CurrentBrushRadius;
+	const FDynamicMesh3* Mesh = PreviewMesh->GetPreviewDynamicMesh();
+	for (int VertIdx : Mesh->VertexIndicesItr())
+	{
+		FVector3d Position = Mesh->GetVertex(VertIdx);
+		if ((Position - StampPosLocal).SquaredLength() < RadiusSqr)
+		{
+			VertexROI.Add(VertIdx);
+		}
+	}
+}
+
+
+
+
+static void GrowToConnectedTriangles(const TArray<int>& TriangleROI, TArray<int>& GrowROI, const FDynamicMesh3* Mesh,
+	TArray<int32>& QueueBuffer, TSet<int32>& DoneBuffer,
+	TFunctionRef<bool(int32, int32)> CanGrowPredicate = [](int32, int32) { return true; }
+	)
+{
+	QueueBuffer.Reset(); QueueBuffer.Insert(TriangleROI, 0);
+	DoneBuffer.Reset(); DoneBuffer.Append(TriangleROI);
+	//TArray<int32> Queue(TriangleROI);
+	//TSet<int32> Done(TriangleROI);
+
+	while (QueueBuffer.Num() > 0)
+	{
+		int32 CurTri = QueueBuffer.Pop(false);
+		GrowROI.Add(CurTri);
+
+		FIndex3i NbrTris = Mesh->GetTriNeighbourTris(CurTri);
+		for (int j = 0; j < 3; ++j)
+		{
+			int32 tid = NbrTris[j];
+			if (tid != FDynamicMesh3::InvalidID && DoneBuffer.Contains(tid) == false && CanGrowPredicate(CurTri, tid))
+			{
+				QueueBuffer.Add(tid);
+				DoneBuffer.Add(tid);
+			}
+		}
+	}
+
+}
+
+
+
+void UMeshSelectionTool::CalculateTriangleROI(const FBrushStampData& Stamp, TArray<int>& TriangleROI)
+{
+	FTransform Transform = ComponentTarget->GetWorldTransform();
+	FVector StampPosLocal = Transform.InverseTransformPosition(Stamp.WorldPosition);
+
+	// always select first triangle
+	const FDynamicMesh3* Mesh = PreviewMesh->GetPreviewDynamicMesh();
+
+	float RadiusSqr = CurrentBrushRadius * CurrentBrushRadius;
+	if (SelectionProps->bVolumetricBrush)
+	{
+		if (Mesh->IsTriangle(Stamp.HitResult.FaceIndex))
+		{
+			TriangleROI.Add(Stamp.HitResult.FaceIndex);
+		}
+
+		FAxisAlignedBox3d Bounds(StampPosLocal-CurrentBrushRadius*FVector3d::One(), StampPosLocal+CurrentBrushRadius*FVector3d::One());
+		TemporaryBuffer.Reset();
+		GetOctree()->RangeQuery(Bounds, TemporaryBuffer);
+
+		for (int32 TriIdx : TemporaryBuffer)
+		{
+			FVector3d Position = Mesh->GetTriCentroid(TriIdx);
+			if ((Position - StampPosLocal).SquaredLength() < RadiusSqr)
+			{
+				TriangleROI.Add(TriIdx);
+			}
+		}
+	}
+	else
+	{
+		TArray<int32> StartROI;
+		StartROI.Add(Stamp.HitResult.FaceIndex);
+		GrowToConnectedTriangles(StartROI, TriangleROI, Mesh, TemporaryBuffer, TemporarySet,
+			[Mesh, RadiusSqr, StampPosLocal](int t1, int t2) { return (Mesh->GetTriCentroid(t2) - StampPosLocal).SquaredLength() < RadiusSqr; });
+
+	}
+
+
+}
+
+
+
+
+static void UpdateList(TArray<int>& List, int Value, bool bAdd)
+{
+	if (bAdd)
+	{
+		List.Add(Value);
+	}
+	else
+	{
+		List.RemoveSwap(Value);
+	}
+}
+
+
+void UMeshSelectionTool::ApplyStamp(const FBrushStampData& Stamp)
+{
+	//const FDynamicMesh3* Mesh = PreviewMesh->GetPreviewDynamicMesh();
+
+	IndexBuf.Reset();
+
+	bool bDesiredValue = bInRemoveStroke ? false : true;
+
+	if (SelectionType == EMeshSelectionElementType::Face)
+	{
+		CalculateTriangleROI(Stamp, IndexBuf);
+		UpdateFaceSelection(Stamp, IndexBuf);
+	}
+	else
+	{
+		CalculateVertexROI(Stamp, IndexBuf);
+		for (int VertIdx : IndexBuf)
+		{
+			if (SelectedVertices[VertIdx] != bDesiredValue)
+			{
+				SelectedVertices[VertIdx] = bDesiredValue;
+				UpdateList(Selection->Vertices, VertIdx, bDesiredValue);
+				if (ActiveSelectionChange != nullptr)
+				{
+					ActiveSelectionChange->Add(VertIdx);
+				}
+			}
+		}
+	}
+
+	OnSelectionUpdated();
+}
+
+
+
+
+
+
+void UMeshSelectionTool::UpdateFaceSelection(const FBrushStampData& Stamp, const TArray<int>& TriangleROI)
+{
+	const FDynamicMesh3* Mesh = PreviewMesh->GetPreviewDynamicMesh();
+	const TArray<int>* UseROI = &TriangleROI;
+
+	TArray<int> LocalROI;
+	if (SelectionProps->SelectionMode == EMeshSelectionToolPrimaryMode::AllConnected)
+	{
+		GrowToConnectedTriangles(TriangleROI, LocalROI, Mesh, TemporaryBuffer, TemporarySet);
+		UseROI = &LocalROI;
+	}
+	else if (SelectionProps->SelectionMode == EMeshSelectionToolPrimaryMode::AllInGroup)
+	{
+		GrowToConnectedTriangles(TriangleROI, LocalROI, Mesh, TemporaryBuffer, TemporarySet,
+			[Mesh](int t1, int t2) { return Mesh->GetTriangleGroup(t1) == Mesh->GetTriangleGroup(t2); } );
+		UseROI = &LocalROI;
+	}
+	else if (SelectionProps->SelectionMode == EMeshSelectionToolPrimaryMode::AllWithinAngle)
+	{
+		TArray<int32> StartROI; 
+		StartROI.Add(Stamp.HitResult.FaceIndex);
+		FVector3d StartNormal = Mesh->GetTriNormal(StartROI[0]);
+		int AngleTol = SelectionProps->AngleTolerance;
+		GrowToConnectedTriangles(StartROI, LocalROI, Mesh, TemporaryBuffer, TemporarySet,
+			[Mesh, AngleTol, StartNormal](int t1, int t2) { return Mesh->GetTriNormal(t2).AngleD(StartNormal) < AngleTol; });
+		UseROI = &LocalROI;
+	}
+	else if (SelectionProps->SelectionMode == EMeshSelectionToolPrimaryMode::AngleFiltered)
+	{
+		TSet<int32> BrushROI(TriangleROI);
+		TArray<int32> StartROI;
+		StartROI.Add(Stamp.HitResult.FaceIndex);
+		FVector3d StartNormal = Mesh->GetTriNormal(StartROI[0]);
+		int AngleTol = SelectionProps->AngleTolerance;
+		GrowToConnectedTriangles(StartROI, LocalROI, Mesh, TemporaryBuffer, TemporarySet,
+			[Mesh, AngleTol, StartNormal, &BrushROI](int t1, int t2) { return BrushROI.Contains(t2) && Mesh->GetTriNormal(t2).AngleD(StartNormal) < AngleTol; });
+		UseROI = &LocalROI;
+	}
+	else if (SelectionProps->SelectionMode == EMeshSelectionToolPrimaryMode::Visible)
+	{
+		FViewCameraState StateOut;
+		GetToolManager()->GetContextQueriesAPI()->GetCurrentViewState(StateOut);
+		FVector3d LocalEyePosition(ComponentTarget->GetWorldTransform().InverseTransformPosition(StateOut.Position));
+
+		for (int tid : TriangleROI)
+		{
+			FVector3d Centroid = Mesh->GetTriCentroid(tid);
+			int HitTID = GetOctree()->FindNearestHitObject(FRay3d(LocalEyePosition, (Centroid - LocalEyePosition).Normalized()));
+			if (HitTID == tid)
+			{
+				LocalROI.Add(HitTID);
+			}
+		}
+		UseROI = &LocalROI;
+	}
+
+	bool bDesiredValue = bInRemoveStroke ? false : true;
+	for (int TriIdx : *UseROI)
+	{
+		if (SelectedTriangles[TriIdx] != bDesiredValue)
+		{
+			SelectedTriangles[TriIdx] = bDesiredValue;
+			UpdateList(Selection->Faces, TriIdx, bDesiredValue);
+			if (ActiveSelectionChange != nullptr)
+			{
+				ActiveSelectionChange->Add(TriIdx);
+			}
+		}
+	}
+
+}
+
+
+
+
+void UMeshSelectionTool::OnEndDrag(const FRay& Ray)
+{
+	UDynamicMeshBrushTool::OnEndDrag(Ray);
+
+	bInRemoveStroke = false;
+	bStampPending = false;
+
+	// close change record
+	TUniquePtr<FMeshSelectionChange> Change = EndChange();
+	GetToolManager()->EmitObjectChange(Selection, MoveTemp(Change), LOCTEXT("MeshSelectionChange", "Mesh Selection"));
+}
+
+
+bool UMeshSelectionTool::OnUpdateHover(const FInputDeviceRay& DevicePos)
+{
+	UDynamicMeshBrushTool::OnUpdateHover(DevicePos);
+
+	// todo get rid of this redundant hit test!
+	FHitResult OutHit;
+	if ( UDynamicMeshBrushTool::HitTest(DevicePos.WorldRay, OutHit) )
+	{
+		PreviewBrushROI.Reset();
+		if (SelectionType == EMeshSelectionElementType::Face)
+		{
+			CalculateTriangleROI(LastBrushStamp, PreviewBrushROI);
+		}
+		else
+		{
+			CalculateVertexROI(LastBrushStamp, PreviewBrushROI);
+		}
+	}
+
+	return true;
+}
+
+
+
+
+
+
+void UMeshSelectionTool::OnSelectionUpdated()
+{
+	UpdateVisualization();
+}
+
+void UMeshSelectionTool::UpdateVisualization()
+{
+	if (SelectionType == EMeshSelectionElementType::Face)
+	{
+		PreviewMesh->SetTriangleColorFunction([this](int TriangleID)
+		{
+			return SelectedTriangles[TriangleID] ? FColor::Red : FColor::White;
+		}, UPreviewMesh::ERenderUpdateMode::FullUpdate);
+
+		//bool bIsFirst = (! BaseMeshComponent->TriangleColorFunc);
+		//BaseMeshComponent->TriangleColorFunc =
+		//if (bIsFirst)
+		//{
+		//	BaseMeshComponent->NotifyMeshUpdated();
+		//}
+		//else
+		//{
+		//	BaseMeshComponent->NotifyMeshUpdated();
+		//	//BaseMeshComponent->FastNotifyColorsUpdated();
+		//}
+	}
+	else
+	{
+		PreviewMesh->ClearTriangleColorFunction(UPreviewMesh::ERenderUpdateMode::FullUpdate);
+	}
+}
+
+
+
+
+void UMeshSelectionTool::Render(IToolsContextRenderAPI* RenderAPI)
+{
+	UDynamicMeshBrushTool::Render(RenderAPI);
+
+	FTransform WorldTransform = ComponentTarget->GetWorldTransform();
+	const FDynamicMesh3* Mesh = PreviewMesh->GetPreviewDynamicMesh();
+
+	if (SelectionType == EMeshSelectionElementType::Vertex)
+	{
+		MeshDebugDraw::DrawVertices(Mesh, Selection->Vertices,
+			12.0f, FColor::Orange, RenderAPI->GetPrimitiveDrawInterface(), WorldTransform);
+		MeshDebugDraw::DrawVertices(Mesh, PreviewBrushROI,
+			8.0f, FColor(40, 200, 40), RenderAPI->GetPrimitiveDrawInterface(), WorldTransform);
+	}
+	else
+	{
+		// drawn via material
+		//MeshDebugDraw::DrawTriCentroids(Mesh, Selection->Faces,
+		//	12.0f, FColor::Green, RenderAPI->GetPrimitiveDrawInterface(), WorldTransform);
+		MeshDebugDraw::DrawTriCentroids(Mesh, PreviewBrushROI,
+			4.0f, FColor(40, 200, 40), RenderAPI->GetPrimitiveDrawInterface(), WorldTransform);
+	}
+}
+
+
+void UMeshSelectionTool::Tick(float DeltaTime)
+{
+	UDynamicMeshBrushTool::Tick(DeltaTime);
+
+	ShowWireframeWatcher.CheckAndUpdate();
+
+	if (bStampPending)
+	{
+		ApplyStamp(LastStamp);
+		bStampPending = false;
+	}
+
+	if (bHavePendingAction)
+	{
+		ApplyAction(PendingAction);
+		bHavePendingAction = false;
+		PendingAction = EMeshSelectionToolActions::NoAction;
+	}
+}
+
+
+
+void UMeshSelectionTool::BeginChange(bool bAdding)
+{
+	check(ActiveSelectionChange == nullptr);
+	ActiveSelectionChange = new FMeshSelectionChangeBuilder(SelectionType, bAdding);
+}
+
+void UMeshSelectionTool::CancelChange()
+{
+	if (ActiveSelectionChange != nullptr)
+	{
+		delete ActiveSelectionChange;
+		ActiveSelectionChange = nullptr;
+	}
+}
+
+TUniquePtr<FMeshSelectionChange> UMeshSelectionTool::EndChange()
+{
+	check(ActiveSelectionChange);
+	if (ActiveSelectionChange != nullptr)
+	{
+		TUniquePtr<FMeshSelectionChange> Result = MoveTemp(ActiveSelectionChange->Change);
+		delete ActiveSelectionChange;
+		ActiveSelectionChange = nullptr;
+
+		return Result;
+	}
+	return TUniquePtr<FMeshSelectionChange>();
+}
+
+
+
+
+
+
+void UMeshSelectionTool::RequestAction(EMeshSelectionToolActions ActionType)
+{
+	if (bHavePendingAction)
+	{
+		return;
+	}
+
+	PendingAction = ActionType;
+	bHavePendingAction = true;
+}
+
+
+void UMeshSelectionTool::ApplyAction(EMeshSelectionToolActions ActionType)
+{
+	switch (ActionType)
+	{
+		case EMeshSelectionToolActions::ClearSelection:
+			ClearSelection();
+			break;
+
+		case EMeshSelectionToolActions::InvertSelection:
+			InvertSelection();
+			break;
+
+
+		case EMeshSelectionToolActions::GrowSelection:
+			GrowShrinkSelection(true);
+			break;
+
+		case EMeshSelectionToolActions::ShrinkSelection:
+			GrowShrinkSelection(false);
+			break;
+
+
+		case EMeshSelectionToolActions::ExpandToConnected:
+			ExpandToConnected();
+			break;
+
+		case EMeshSelectionToolActions::DeleteSelected:
+			DeleteSelectedTriangles();
+			break;
+
+		case EMeshSelectionToolActions::SeparateSelected:
+			SeparateSelectedTriangles();
+			break;
+	}
+}
+
+
+
+
+
+void UMeshSelectionTool::ClearSelection()
+{
+	TArray<int32> SelectedFaces = Selection->GetElements(EMeshSelectionElementType::Face);
+	if (SelectedFaces.Num() == 0)
+	{
+		return;
+	}
+
+	BeginChange(false);
+	ActiveSelectionChange->Add(SelectedFaces);
+	Selection->RemoveIndices(EMeshSelectionElementType::Face, SelectedFaces);
+
+	TUniquePtr<FMeshSelectionChange> SelectionChange = EndChange();
+
+	GetToolManager()->EmitObjectChange(Selection, MoveTemp(SelectionChange), LOCTEXT("ClearSelection", "Clear Selection"));
+
+	OnExternalSelectionChange();
+}
+
+
+
+
+void UMeshSelectionTool::InvertSelection()
+{
+	check(SelectionType == EMeshSelectionElementType::Face);
+	TArray<int32> SelectedFaces = Selection->GetElements(EMeshSelectionElementType::Face);
+	if (SelectedFaces.Num() == 0)
+	{
+		return;
+	}
+
+	TArray<int32> InvertedFaces;
+	const FDynamicMesh3* Mesh = PreviewMesh->GetPreviewDynamicMesh();
+	for (int tid : Mesh->TriangleIndicesItr())
+	{
+		if (SelectedTriangles[tid] == false)
+		{
+			InvertedFaces.Add(tid);
+		}
+	}
+
+	GetToolManager()->BeginUndoTransaction(LOCTEXT("InvertSelection", "Invert Selection"));
+
+	// clear current selection
+	BeginChange(false);
+	ActiveSelectionChange->Add(SelectedFaces);
+	Selection->RemoveIndices(EMeshSelectionElementType::Face, SelectedFaces);
+	TUniquePtr<FMeshSelectionChange> ClearChange = EndChange();
+
+	GetToolManager()->EmitObjectChange(Selection, MoveTemp(ClearChange), LOCTEXT("InvertSelection", "Invert Selection"));
+
+	// add inverted selection
+	BeginChange(true);
+	ActiveSelectionChange->Add(InvertedFaces);
+	Selection->AddIndices(EMeshSelectionElementType::Face, InvertedFaces);
+	TUniquePtr<FMeshSelectionChange> AddChange = EndChange();
+
+	GetToolManager()->EmitObjectChange(Selection, MoveTemp(AddChange), LOCTEXT("InvertSelection", "Invert Selection"));
+
+	GetToolManager()->EndUndoTransaction();
+
+	OnExternalSelectionChange();
+}
+
+
+
+
+
+void UMeshSelectionTool::GrowShrinkSelection(bool bGrow)
+{
+	check(SelectionType == EMeshSelectionElementType::Face);
+	TArray<int32> SelectedFaces = Selection->GetElements(EMeshSelectionElementType::Face);
+	if (SelectedFaces.Num() == 0)
+	{
+		return;
+	}
+
+	const FDynamicMesh3* Mesh = PreviewMesh->GetPreviewDynamicMesh();
+	TArray<int32> Vertices;
+	MeshIndexUtil::TriangleToVertexIDs(Mesh, SelectedFaces, Vertices);
+
+	TSet<int32> ChangeFaces;
+	for (int vid : Vertices)
+	{
+		int OutCount = 0;
+		for (int tid : Mesh->VtxTrianglesItr(vid))
+		{
+			if (SelectedTriangles[tid] == false)
+			{
+				OutCount++;
+			}
+		}
+		if (OutCount == 0)
+		{
+			continue;
+		}
+
+		for (int tid : Mesh->VtxTrianglesItr(vid))
+		{
+			if ( (bGrow && SelectedTriangles[tid] == false) || (bGrow == false && SelectedTriangles[tid]) )
+			{
+				ChangeFaces.Add(tid);
+			}
+		}
+	}
+	if (ChangeFaces.Num() == 0)
+	{
+		return;
+	}
+
+	BeginChange(bGrow);
+	ActiveSelectionChange->Add(ChangeFaces);
+	if (bGrow)
+	{
+		Selection->AddIndices(EMeshSelectionElementType::Face, ChangeFaces);
+		TUniquePtr<FMeshSelectionChange> SelectionChange = EndChange();
+		GetToolManager()->EmitObjectChange(Selection, MoveTemp(SelectionChange), LOCTEXT("GrowSelection", "Grow Selection"));
+	}
+	else
+	{
+		Selection->RemoveIndices(EMeshSelectionElementType::Face, ChangeFaces);
+		TUniquePtr<FMeshSelectionChange> SelectionChange = EndChange();
+		GetToolManager()->EmitObjectChange(Selection, MoveTemp(SelectionChange), LOCTEXT("ShrinkSelection", "Shrink Selection"));
+	}
+	OnExternalSelectionChange();
+}
+
+
+
+
+
+void UMeshSelectionTool::ExpandToConnected()
+{
+	check(SelectionType == EMeshSelectionElementType::Face);
+	TArray<int32> SelectedFaces = Selection->GetElements(EMeshSelectionElementType::Face);
+	if (SelectedFaces.Num() == 0)
+	{
+		return;
+	}
+
+	const FDynamicMesh3* Mesh = PreviewMesh->GetPreviewDynamicMesh();
+
+	TArray<int32> Queue(SelectedFaces);
+	TSet<int32> AddFaces;
+
+	while (Queue.Num() > 0)
+	{
+		int32 CurTri = Queue.Pop(false);
+		FIndex3i NbrTris = Mesh->GetTriNeighbourTris(CurTri);
+
+		for (int j = 0; j < 3; ++j)
+		{
+			int32 tid = NbrTris[j];
+			if (SelectedTriangles[tid] == false && AddFaces.Contains(tid) == false)
+			{
+				AddFaces.Add(tid);
+				Queue.Add(tid);
+			}
+		}
+	}
+	if (AddFaces.Num() == 0)
+	{
+		return;
+	}
+
+	BeginChange(true);
+	ActiveSelectionChange->Add(AddFaces);
+	Selection->AddIndices(EMeshSelectionElementType::Face, AddFaces);
+	TUniquePtr<FMeshSelectionChange> SelectionChange = EndChange();
+	GetToolManager()->EmitObjectChange(Selection, MoveTemp(SelectionChange), LOCTEXT("ExpandToConnected", "Expand Selection"));
+	OnExternalSelectionChange();
+}
+
+
+
+
+void UMeshSelectionTool::DeleteSelectedTriangles()
+{
+	check(SelectionType == EMeshSelectionElementType::Face);
+	TArray<int32> SelectedFaces = Selection->GetElements(EMeshSelectionElementType::Face);
+	if (SelectedFaces.Num() == 0)
+	{
+		return;
+	}
+
+	TUniquePtr<FToolCommandChangeSequence> ChangeSeq = MakeUnique<FToolCommandChangeSequence>();
+
+	// clear current selection
+	BeginChange(false);
+	for (int tid : SelectedFaces)
+	{
+		ActiveSelectionChange->Add(tid);
+	}
+	Selection->RemoveIndices(EMeshSelectionElementType::Face, SelectedFaces);
+	TUniquePtr<FMeshSelectionChange> SelectionChange = EndChange();
+	ChangeSeq->AppendChange(Selection, MoveTemp(SelectionChange));
+
+	// delete triangles and emit delete triangles change
+	TUniquePtr<FMeshChange> MeshChange = PreviewMesh->TrackedEditMesh(
+		[&SelectedFaces](FDynamicMesh3& Mesh, FDynamicMeshChangeTracker& ChangeTracker)
+	{
+		FDynamicMeshEditor Editor(&Mesh);
+		Editor.RemoveTriangles(SelectedFaces, true, [&ChangeTracker](int TriangleID) { ChangeTracker.SaveTriangle(TriangleID, true); });
+	});
+	ChangeSeq->AppendChange(PreviewMesh, MoveTemp(MeshChange));
+
+	// emit combined change sequence
+	GetToolManager()->EmitObjectChange(this, MoveTemp(ChangeSeq), LOCTEXT("MeshSelectionToolDelete", "Delete Faces"));
+
+	OnExternalSelectionChange();
+	bHaveModifiedMesh = true;
+	bOctreeValid = false;
+}
+
+
+
+void AssignMaterial(AActor* ToActor, const TUniquePtr<FPrimitiveComponentTarget>& FromTarget)
+{
+	UMaterialInterface* Material = FromTarget->GetMaterial(0);
+	if (!Material)
+	{
+		return;
+	}
+
+	//if (Cast<AStaticMeshActor>(ToActor) != nullptr)
+	//{
+	//	UStaticMeshComponent* Component = Cast<AStaticMeshActor>(ToActor)->GetStaticMeshComponent();
+	//	if (Component)
+	//	{
+	//		Component->SetMaterial(0, Material);
+	//	}
+	//} 
+	//else
+	//{
+		USceneComponent* Component = ToActor->GetRootComponent();
+		if (Cast<UPrimitiveComponent>(Component) != nullptr)
+		{
+			Cast<UPrimitiveComponent>(Component)->SetMaterial(0, Material);
+		}
+	//}
+}
+
+
+void UMeshSelectionTool::SeparateSelectedTriangles()
+{
+#if WITH_EDITOR
+	// currently AssetGenerationUtil::GenerateStaticMeshActor only defined in editor
+
+	check(SelectionType == EMeshSelectionElementType::Face);
+	TArray<int32> SelectedFaces = Selection->GetElements(EMeshSelectionElementType::Face);
+	if (SelectedFaces.Num() == 0)
+	{
+		return;
+	}
+
+	const FDynamicMesh3* SourceMesh = PreviewMesh->GetPreviewDynamicMesh();
+	if (SelectedFaces.Num() == SourceMesh->TriangleCount())
+	{
+		return;		// don't separate entire mesh
+	}
+
+
+	// extract copy of triangles
+	FDynamicMesh3 SeparatedMesh;
+	SeparatedMesh.EnableAttributes();
+	FDynamicMeshEditor Editor(&SeparatedMesh);
+	FMeshIndexMappings Mappings; FDynamicMeshEditResult EditResult;
+	Editor.AppendTriangles(SourceMesh, SelectedFaces, Mappings, EditResult);
+
+	// emit new asset
+	FTransform3d Transform(PreviewMesh->GetTransform());
+	GetToolManager()->BeginUndoTransaction(LOCTEXT("MeshSelectionToolSeparate", "Separate"));
+	AActor* NewActor = AssetGenerationUtil::GenerateStaticMeshActor(
+		AssetAPI, TargetWorld, &SeparatedMesh, Transform, TEXT("Submesh"),
+		AssetGenerationUtil::GetDefaultAutoGeneratedAssetPath());
+	AssignMaterial(NewActor, ComponentTarget);
+	GetToolManager()->EndUndoTransaction();
+
+	// todo: undo won't remove this asset...
+
+	// delete selected triangles from this mesh
+	DeleteSelectedTriangles();
+#endif
+}
+
+
+
+
+#undef LOCTEXT_NAMESPACE
\ No newline at end of file
diff --git a/Engine/Plugins/Experimental/MeshModelingToolset/Source/MeshModelingTools/Private/MeshSpaceDeformerTool.cpp b/Engine/Plugins/Experimental/MeshModelingToolset/Source/MeshModelingTools/Private/MeshSpaceDeformerTool.cpp
new file mode 100644
--- /dev/null
+++ b/Engine/Plugins/Experimental/MeshModelingToolset/Source/MeshModelingTools/Private/MeshSpaceDeformerTool.cpp
@@ -0,0 +1,523 @@
+// Copyright 1998-2019 Epic Games, Inc. All Rights Reserved.
+
+#include "MeshSpaceDeformerTool.h"
+#include "InteractiveToolManager.h"
+#include "ToolBuilderUtil.h"
+
+#include "SegmentTypes.h"
+#include "DynamicMeshAttributeSet.h"
+#include "MeshNormals.h"
+#include "ToolSceneQueriesUtil.h"
+#include "Intersection/IntersectionUtil.h"
+
+#include "Async/ParallelFor.h"
+#include "Containers/BitArray.h"
+
+///////////////////////////////
+// EIGEN
+
+#if defined(_MSC_VER) && USING_CODE_ANALYSIS
+#pragma warning(push)
+#pragma warning(disable : 6011)
+#pragma warning(disable : 6387)
+#pragma warning(disable : 6313)
+#pragma warning(disable : 6294)
+#endif
+THIRD_PARTY_INCLUDES_START
+#include <Eigen/Core>
+#include <Eigen/Dense> 
+#include <Eigen/Eigenvalues>
+
+#ifndef EIGEN_MPL2_ONLY
+#endif
+THIRD_PARTY_INCLUDES_END
+#if defined(_MSC_VER) && USING_CODE_ANALYSIS
+#pragma warning(pop)
+#endif
+
+///////////////////////////////
+
+
+#define LOCTEXT_NAMESPACE "MeshSpaceDeformerTool"
+
+
+//////////////////////////////
+// DEBUG_SETTINGS
+
+#define DEBUG_ASYNC_TASK
+
+#define DEBUG_DRAW_AXES
+//////////////////////////////
+
+
+#ifdef DEBUG_ASYNC_TASK
+#define print(text,...) UE_LOG(LogTemp, Warning, TEXT(text), __VA_ARGS__)
+#else
+#define print(x)
+#endif
+
+
+
+double RayToSegmentSquareDist(const FRay3d& Ray, const FVector3d& V0, const FVector3d& V1, FVector3d* ClosestPtOnSegment, FVector3d* ClosestPtOnRay)
+{
+	if (V0 == V1)
+	{
+		if (ClosestPtOnRay) 
+		{
+			*ClosestPtOnRay = V0;
+		}
+
+		if (ClosestPtOnSegment)
+		{
+			*ClosestPtOnSegment = V0;
+		}
+		return 0.;
+	}
+	FVector3d SegCenter = (V0 + V1) * 0.5;
+	FVector3d SegDir = (V1 - V0).Normalized();
+	FVector3d Diff = Ray.Origin - SegCenter;
+
+	double SegExtent = V0.Distance(V1) * 0.5;
+	double A01 = -Ray.Direction.Dot(SegDir);
+	double B0 = Diff.Dot(Ray.Direction);
+	double B1 = -Diff.Dot(SegDir);
+	double C = (Diff).SquaredLength();
+	double det = FMath::Abs(1. - A01 * A01);
+	double S0, S1, SqrDist, ExtDet;
+
+	if (det > 0.)
+	{
+		// The ray and segment are not parallel.
+
+		S0 = A01 * B1 - B0;
+		S1 = A01 * B0 - B1;
+		ExtDet = SegExtent * det;
+
+		if (S0 >= 0.)
+		{
+			if (S1 >= -ExtDet)
+			{
+				if (S1 <= ExtDet)
+				{
+					// region 0
+					// Minimum at interior points of ray and segment.
+					float invDet = 1. / det;
+					S0 *= invDet;
+					S1 *= invDet;
+					SqrDist = S0 * (S0 + A01 * S1 + 2. * B0) + S1 * (A01 * S0 + S1 + 2. * B1) + C;
+				}
+				else
+				{
+					// region 1
+					S1 = SegExtent;
+					S0 = std::max(0., -(A01 * S1 + B0));
+					SqrDist = -S0 * S0 + S1 * (S1 + 2. * B1) + C;
+				}
+			}
+			else
+			{
+				// region 5
+				S1 = -SegExtent;
+				S0 = std::max(0., -(A01 * S1 + B0));
+				SqrDist = -S0 * S0 + S1 * (S1 + 2. * B1) + C;
+			}
+		}
+		else
+		{
+			if (S1 <= -ExtDet)
+			{
+				// region 4
+				S0 = std::max(0., -(-A01 * SegExtent + B0));
+				S1 = (S0 > 0.) ? -SegExtent : std::min(std::max(-SegExtent, -B1), SegExtent);
+				SqrDist = -S0 * S0 + S1 * (S1 + 2. * B1) + C;
+			}
+			else if (S1 <= ExtDet)
+			{
+				// region 3
+				S0 = 0;
+				S1 = std::min(std::max(-SegExtent, -B1), SegExtent);
+				SqrDist = S1 * (S1 + 2. * B1) + C;
+			}
+			else
+			{
+				// region 2
+				S0 = std::max(0., -(A01 * SegExtent + B0));
+				S1 = (S0 > 0.) ? SegExtent : std::min(std::max(-SegExtent, -B1), SegExtent);
+				SqrDist = -S0 * S0 + S1 * (S1 + 2. * B1) + C;
+			}
+		}
+	}
+	else
+	{
+		// Ray and segment are parallel.
+
+		S1 = (A01 > 0.) ? -SegExtent : SegExtent;
+		S0 = std::max(0., -(A01 * S1 + B0));
+		SqrDist = -S0 * S0 + S1 * (S1 + 2. * B1) + C;
+	}
+
+	if (ClosestPtOnSegment)
+	{
+		*ClosestPtOnSegment = SegDir * S1 + SegCenter;
+	}
+
+	if (ClosestPtOnRay)
+	{
+		*ClosestPtOnRay = Ray.Direction * S0 + Ray.Origin;
+	}
+	return SqrDist;
+}
+
+FMatrix3d SolveSampleCovariance(double* Array, const int NCols, const int NRows, FVector3d& ExtentsOut, FVector3d& OriginOut)
+{
+	using RowMajorVec3dMatrix = Eigen::Matrix<double, Eigen::Dynamic, 3, Eigen::RowMajor>;
+	using RowMajorVec3d = Eigen::Matrix<double, 1, 3, Eigen::RowMajor>;
+	RowMajorVec3dMatrix VertexMatrix = Eigen::Map<RowMajorVec3dMatrix>(Array, NRows, NCols);
+
+	check(VertexMatrix(0, 0) == *Array); // Ensure the mapping is in expected format
+	check(VertexMatrix(0, 1) == Array[1]);
+	check(VertexMatrix(0, 2) == Array[2]);
+
+	//Using adjoint eigen solver
+	auto Origin = VertexMatrix.colwise().mean();															 // Get the centroid of the point cloud
+	Eigen::MatrixXd CenteredVertices = VertexMatrix.rowwise() - Origin;										 // Translate each point so centroid is (0,0,0)
+	Eigen::MatrixXd CovarianceMatrix = (CenteredVertices.adjoint() * CenteredVertices) / double(NRows - 1);	 // note: adjoint ~= conjugate transpose; formula from https://en.wikipedia.org/wiki/Sample_mean_and_covariance#Sample_covariance
+	
+	// note: SelfAdjointEigenSolver automatically sorts the eigen vectors by the order of corresponding eigenvalues in increasing order
+	Eigen::SelfAdjointEigenSolver<Eigen::MatrixXd> EigenSolver(CovarianceMatrix);							 
+	Eigen::Matrix3d EigenVectors = EigenSolver.eigenvectors();
+
+	//Align object with the standard basis by rotating all the points by the inverse of EigenVectors.
+	//EigenVectors is an orthonormal basis (3x3) so transpose is the inverse; 
+	Eigen::Matrix3Xd RotatedVertices = EigenVectors.transpose() * CenteredVertices.transpose();
+
+	//Since our point cloud is now aligned with the standard axes, the dimensions of its AABB should be the dimensions of its 
+	RowMajorVec3d Extents = RotatedVertices.rowwise().maxCoeff() - RotatedVertices.rowwise().minCoeff();
+
+	std::sort(Extents.data(), Extents.data() +3);						//Sort the extents to match our sorted eigenvectors
+	ExtentsOut = FVector3d{ FMath::Abs(Extents[0]), FMath::Abs(Extents[1]), FMath::Abs(Extents[2]) } / 2;	//Divide by two because extent represents distance from center
+	OriginOut = FVector3d{ Origin(0,0), Origin(0,1), Origin(0,2) };		//Origin in world space
+
+	return FMatrix3d{ FVector3d{ EigenVectors(0,0), EigenVectors(1,0), EigenVectors(2,0) },
+					  FVector3d{ EigenVectors(0,1), EigenVectors(1,1), EigenVectors(2,1) },
+					  FVector3d{ EigenVectors(0,2), EigenVectors(1,2), EigenVectors(2,2) }, 
+					  false };
+}
+
+FMatrix3d UMeshSpaceDeformerTool::CalculateBestAxis(FDynamicMesh3* Mesh, const TArray<FVector3d>& PositionBuffer, FTransform WorldTransform, FVector3d * ExtentsOut, FVector3d * OriginOut)
+{
+	PCAVertexPositionBuffer.Reset(3 * Mesh->VertexCount());
+	double* VectorData = PCAVertexPositionBuffer.GetData();
+	FVector3d* InitializePtr = reinterpret_cast<FVector3d*>(VectorData);
+	for (const int& VertexID : Mesh->VertexIndicesItr())
+	{
+		const FVector3d& Point = PositionBuffer[VertexID];
+		*InitializePtr = WorldTransform.TransformPosition(Point);
+		++InitializePtr;
+	}
+
+	const int NCols = 3; // R^3;
+	const int NRows = Mesh->VertexCount();
+
+	return SolveSampleCovariance(VectorData, NCols, NRows, *ExtentsOut, *OriginOut);
+}
+
+/*
+ * ToolBuilder
+ */
+UMeshSurfacePointTool* UMeshSpaceDeformerToolBuilder::CreateNewTool(const FToolBuilderState& SceneState) const
+{
+	UMeshSpaceDeformerTool* NonlinearTool = NewObject<UMeshSpaceDeformerTool>(SceneState.ToolManager);
+	return NonlinearTool;
+}
+
+UMeshSpaceDeformerTool::UMeshSpaceDeformerTool()
+{
+}
+
+void UMeshSpaceDeformerTool::Setup()
+{
+	UMeshSurfacePointTool::Setup();
+
+	//// create dynamic mesh component to use for live preview
+	DynamicMeshComponent = NewObject<USimpleDynamicMeshComponent>(ComponentTarget->GetOwnerActor(), "Dynamic Mesh");
+	DynamicMeshComponent->SetupAttachment(ComponentTarget->GetOwnerActor()->GetRootComponent());
+	DynamicMeshComponent->RegisterComponent();
+	DynamicMeshComponent->SetWorldTransform(ComponentTarget->GetWorldTransform());
+
+
+	//// copy material if there is one
+	auto Material = ComponentTarget->GetMaterial(0);
+	if (Material != nullptr)
+	{
+		DynamicMeshComponent->SetMaterial(0, Material);
+	}
+	//// dynamic mesh configuration settings
+	ComponentTarget->SetOwnerVisibility(false);
+	DynamicMeshComponent->TangentsType = EDynamicMeshTangentCalcType::AutoCalculated;
+	DynamicMeshComponent->InitializeMesh(ComponentTarget->GetMesh());
+	OnDynamicMeshComponentChangedHandle = DynamicMeshComponent->OnMeshChanged.Add(
+		FSimpleMulticastDelegate::FDelegate::CreateUObject(this, &UMeshSpaceDeformerTool::OnDynamicMeshComponentChanged));
+
+	FDynamicMesh3* TargetMesh = DynamicMeshComponent->GetMesh();
+
+	OriginalPositions.SetNumUninitialized(TargetMesh->MaxVertexID());
+	for (int32 VertexID : TargetMesh->VertexIndicesItr())
+	{
+		OriginalPositions[VertexID] = TargetMesh->GetVertex(VertexID);
+	}
+
+	bInDrag = false;
+
+	//// add properties
+	ToolPropertyObjects.Add(this);
+
+	//// set up visualizers
+	AxisRenderer.LineColor = FLinearColor::Red;
+	AxisRenderer.LineThickness = 2.0;
+	AxisRenderer.bDepthTested = false;
+
+	for (int8 i = 0; i < 3; ++i)
+	{
+		FMeshSpaceDeformerOp* Op = Operators[i];
+		Op->UpdateMesh(DynamicMeshComponent->GetMesh());
+	}
+
+	AutoDetectAxes();
+
+}
+
+void UMeshSpaceDeformerTool::Shutdown(EToolShutdownType ShutdownType)
+{
+
+	if (DynamicMeshComponent != nullptr)
+	{
+		DynamicMeshComponent->OnMeshChanged.Remove(OnDynamicMeshComponentChangedHandle);
+
+		ComponentTarget->SetOwnerVisibility(true);
+
+		if (ShutdownType == EToolShutdownType::Accept)
+		{
+			// this block bakes the modified DynamicMeshComponent back into the StaticMeshComponent inside an undo transaction
+			GetToolManager()->BeginUndoTransaction(LOCTEXT("MeshSpaceDeformerToolTransactionName", "Deform Mesh"));
+			ComponentTarget->CommitMesh([&](FMeshDescription* MeshDescription)
+			{
+				DynamicMeshComponent->Bake(MeshDescription, false);
+			});
+			GetToolManager()->EndUndoTransaction();
+		}
+			DynamicMeshComponent->UnregisterComponent();
+			DynamicMeshComponent->DestroyComponent();
+			DynamicMeshComponent = nullptr;
+	}
+}
+
+
+void UMeshSpaceDeformerTool::RegisterActions(FInteractiveToolActionSet& ActionSet)
+{
+}
+
+void UMeshSpaceDeformerTool::OnDynamicMeshComponentChanged()
+{
+}
+
+bool UMeshSpaceDeformerTool::HitTest(const FRay & Ray, FHitResult & OutHit)
+{
+	return Handle.SelectedVertexID != -1;
+}
+
+void UMeshSpaceDeformerTool::OnBeginDrag(const FRay& WorldRay)
+{
+	if (Handle.SelectedVertexID != -1)
+	{
+		bInDrag = true;
+		Handle.DragUpdateBounds(WorldRay);
+		UpdateIntervalsFromDrag();
+		UpdateOp();
+	}
+}
+
+
+void UMeshSpaceDeformerTool::OnUpdateDrag(const FRay& Ray)
+{
+	if (bInDrag)
+	{
+		Handle.DragUpdateBounds(Ray);
+		UpdateIntervalsFromDrag();
+		UpdateOp();
+	}
+}
+
+void UMeshSpaceDeformerTool::OnEndDrag(const FRay& Ray)
+{
+	bInDrag = false;
+
+}
+
+inline void UMeshSpaceDeformerTool::AutoDetectAxes()
+{
+	FTransform WorldTransform = ComponentTarget->GetWorldTransform();
+	FTransform WorldToObject = WorldTransform.Inverse();
+
+	//Retrieves the best axis in world space via principal component analysis.
+	FMatrix3d Basis = CalculateBestAxis(DynamicMeshComponent->GetMesh(), OriginalPositions, WorldTransform, &PrincipalAxesHalfExtentCoeff, &AxisCentroidWorldSpace);
+
+	PrincipalAxesWorldSpace[0] = FVector3d{ Basis(0,0),Basis(1,0),Basis(2,0) };
+	PrincipalAxesWorldSpace[1] = FVector3d{ Basis(0,1),Basis(1,1),Basis(2,1) };
+	PrincipalAxesWorldSpace[2] = FVector3d{ Basis(0,2),Basis(1,2),Basis(2,2) };
+
+	PrincipalAxesObjectSpace[0] = WorldToObject.TransformVectorNoScale(PrincipalAxesWorldSpace[0]);
+	PrincipalAxesObjectSpace[1] = WorldToObject.TransformVectorNoScale(PrincipalAxesWorldSpace[1]);
+	PrincipalAxesObjectSpace[2] = WorldToObject.TransformVectorNoScale(PrincipalAxesWorldSpace[2]);
+
+	UpdateObjectSpaceAxisCentroid();
+}
+
+void UMeshSpaceDeformerTool::SwapSecondaryAxis()
+{
+	std::swap(SecondaryAxis,ThirdAxis);
+}
+
+inline void UMeshSpaceDeformerTool::UpdateObjectSpaceAxisCentroid()
+{
+	AxisCentroidObjectSpace = ComponentTarget->GetWorldTransform().InverseTransformPositionNoScale(AxisCentroidWorldSpace);
+}
+
+void UMeshSpaceDeformerTool::UpdateIntervalsFromDrag()
+{
+	FHandleWidgetVertex* SelectedVertex = Handle.GetSelectedVertex();
+	double CurrentSegmentLength = (SelectedVertex->WorldCenter - AxisCentroidWorldSpace).Length();
+	double MaxLength = PrincipalAxesHalfExtentCoeff[SelectedAxis];
+	double NewInterval = FMath::Clamp(CurrentSegmentLength / MaxLength, 0.0, 1.5);
+	if (Handle.SelectedVertexID == 0)
+	{
+		LowerBoundsInterval = -NewInterval;
+	}
+	else
+	{
+		UpperBoundsInterval = NewInterval;
+	}
+}
+
+bool UMeshSpaceDeformerTool::OnUpdateHover(const FInputDeviceRay& DevicePos)
+{
+	if (!bInDrag)
+	{
+		const auto ToleranceTest = [this](const FVector3d& Position1, const FVector3d& Position2) {
+			return ToolSceneQueriesUtil::PointSnapQuery(this->CameraState, Position1, Position2, VisualAngleSnapThreshold);
+		};
+		Handle.UpdateHover(DevicePos, ComponentTarget->GetWorldTransform(), TFunction<bool(const FVector3d&, const FVector3d&)>{ToleranceTest});
+	}
+	return true;
+}
+
+#if WITH_EDITOR
+void UMeshSpaceDeformerTool::PostEditChangeProperty(FPropertyChangedEvent & PropertyChangedEvent)
+{
+	//UpdateOp();
+}
+#endif
+
+void UMeshSpaceDeformerTool::UpdateOp()
+{
+	FMeshSpaceDeformerOp* Op = Operators[static_cast<int8>(SelectedOperationType)];
+
+	//Provide the basis so that the selected axis is in the Y direction (this will NEED to be rotated based on the orientation the operator expects)
+	// i.e. the Bend operator expects the bend along the Y axis, the twist expects along the Z axis. Just swap the axes as needed inside the operator
+	FMatrix3d Basis
+	{
+		PrincipalAxesObjectSpace[SecondaryAxis],
+		PrincipalAxesObjectSpace[SelectedAxis],
+		PrincipalAxesObjectSpace[ThirdAxis],
+		false
+	};
+
+	//This allows the object to move, or the handle position to move and have the changes reflected based on their relative world-space positions
+	UpdateObjectSpaceAxisCentroid();
+
+	//Pass the relevant updated data to the operator then run the operation.
+	Op->UpdateAxisData(Basis, AxisCentroidObjectSpace, PrincipalAxesHalfExtentCoeff, LowerBoundsInterval, UpperBoundsInterval, ModifierPercent);
+	Op->CalculateResult(nullptr);
+
+	DynamicMeshComponent->FastNotifyPositionsUpdated();
+}
+
+void UMeshSpaceDeformerTool::Tick(float DeltaTime)
+{
+	UMeshSurfacePointTool::Tick(DeltaTime);
+}
+
+
+void UMeshSpaceDeformerTool::Render(IToolsContextRenderAPI* RenderAPI)
+{
+	GetToolManager()->GetContextQueriesAPI()->GetCurrentViewState(CameraState);
+
+	FDynamicMesh3* TargetMesh = DynamicMeshComponent->GetMesh();
+
+	AxisRenderer.BeginFrame(RenderAPI, CameraState);
+
+	//Update the operator (pass it relevant data)
+	UpdateOp();
+
+	//Update the handle visualization and render it.
+	Handle.SetBasis(PrincipalAxesObjectSpace[0], PrincipalAxesObjectSpace[1], PrincipalAxesObjectSpace[2]);
+	Handle.UpdateDisplayData(SelectedAxis, LowerBoundsInterval, UpperBoundsInterval, PrincipalAxesHalfExtentCoeff);
+	Handle.SetPosition(AxisCentroidWorldSpace);
+	Handle.Render();
+
+
+	AxisRenderer.EndFrame();
+	GetToolManager()->PostInvalidation();
+}
+
+
+
+void FHandleWidget::UpdateHover(const FInputDeviceRay & DevicePos , const FTransform& WorldTransform, TFunction<bool(const FVector3d&, const FVector3d&)> ToleranceFunction)
+{
+	FVector3d LocalPosition, LocalNormal;
+	FGeometrySet3::FNearest Nearest;
+
+	//Reset all handle axes and vertex states
+	for (unsigned k = 0; k < 3; ++k)
+	{
+		Axes[k].State = IHandleWidgetSelectable::Default;
+		Axes[k].Vertices[0].State = IHandleWidgetSelectable::Default;
+		Axes[k].Vertices[1].State = IHandleWidgetSelectable::Default;
+	}
+
+	//If we're hovering over a vertex, change its state
+	bool bFound = GeometrySet.FindNearestPointToRay(DevicePos.WorldRay, Nearest, ToleranceFunction);
+	if (bFound)
+	{ // Found a hit point
+		int j = Nearest.ID % 2;
+		int i = (Nearest.ID) / 2;
+		SelectedVertexID = j;
+		Axes[i].Vertices[j].State = IHandleWidgetSelectable::Hovering;
+		Axes[i].Vertices[(j + 1) % 2].State = IHandleWidgetSelectable::Default;
+	}else
+	{
+		SelectedVertexID = -1;
+	}
+	
+	//TODO: Do the same for the axes. The geometry set should already contain the axes, just need to decide how the UI will work.
+
+
+}
+
+void FHandleWidget::DragUpdateBounds(const FRay3d& WorldRay)
+{
+	FVector3d ClosestPtOnSegment;
+	FHandleWidgetAxis* SelectedAxis = GetSelectedAxis();
+
+	FVector3d V0 = SelectedAxis->WorldCenter + SelectedAxis->Axis * 99999.0;
+	FVector3d V1 = SelectedAxis->WorldCenter - SelectedAxis->Axis * 99999.0;
+
+	double SqDist = RayToSegmentSquareDist(WorldRay, V0, V1, &ClosestPtOnSegment, nullptr);
+
+	FHandleWidgetVertex* SelectedVertex = GetSelectedVertex();
+	SelectedVertex->State = IHandleWidgetSelectable::Selected;
+	SelectedVertex->SetPosition(ClosestPtOnSegment);
+}
+
+#undef LOCTEXT_NAMESPACE
+
diff --git a/Engine/Plugins/Experimental/MeshModelingToolset/Source/MeshModelingTools/Private/PositionPlaneGizmo.cpp b/Engine/Plugins/Experimental/MeshModelingToolset/Source/MeshModelingTools/Private/PositionPlaneGizmo.cpp
new file mode 100644
--- /dev/null
+++ b/Engine/Plugins/Experimental/MeshModelingToolset/Source/MeshModelingTools/Private/PositionPlaneGizmo.cpp
@@ -0,0 +1,253 @@
+// Copyright 1998-2019 Epic Games, Inc. All Rights Reserved.
+
+#include "PositionPlaneGizmo.h"
+#include "InteractiveGizmoManager.h"
+#include "Generators/MinimalBoxMeshGenerator.h"
+#include "Materials/MaterialInstanceDynamic.h"
+#include "Engine/Engine.h"  // for GEngine - @todo remove this?
+
+
+
+
+UInteractiveGizmo* UPositionPlaneGizmoBuilder::BuildGizmo(const FToolBuilderState& SceneState) const
+{
+
+	UPositionPlaneGizmo* NewGizmo = NewObject<UPositionPlaneGizmo>(SceneState.GizmoManager);
+	NewGizmo->SetWorld(SceneState.World);
+	return NewGizmo;
+}
+
+
+
+
+
+
+void UPositionPlaneGizmo::SetWorld(UWorld* World)
+{
+	TargetWorld = World;
+}
+
+
+void UPositionPlaneGizmo::Setup()
+{
+	UInteractiveGizmo::Setup();
+
+	UPositionPlaneOnSceneInputBehavior* MouseBehavior = NewObject<UPositionPlaneOnSceneInputBehavior>();
+	MouseBehavior->Initialize(this);
+	AddInputBehavior(MouseBehavior);
+
+	QuickTransformer.Initialize();
+	bInTransformDrag = false;
+
+	// create temporary mesh object for pivot
+	CenterBallShape = MakeSphereMesh();
+}
+
+
+void UPositionPlaneGizmo::Shutdown()
+{
+	CenterBallShape->Disconnect();
+}
+
+
+
+
+
+UPreviewMesh* UPositionPlaneGizmo::MakeSphereMesh()
+{
+	UPreviewMesh* PreviewMesh = NewObject<UPreviewMesh>(this, TEXT("PreviewMesh"));
+	PreviewMesh->CreateInWorld(this->TargetWorld, FTransform::Identity);
+	PreviewMesh->SetVisible(true);
+	PreviewMesh->bBuildSpatialDataStructure = true;
+	PreviewMesh->bDrawOnTop = true;
+
+	FLinearColor AxisColorX = FLinearColor(0.594f, 0.0197f, 0.0f);
+	UMaterial* AxisMaterialBase = GEngine->ArrowMaterial;
+	UMaterialInstanceDynamic* AxisMaterialX = UMaterialInstanceDynamic::Create(AxisMaterialBase, NULL);
+	AxisMaterialX->SetVectorParameterValue("GizmoColor", AxisColorX);
+	CenterBallMaterial = AxisMaterialX;
+
+	//CenterBallMaterial = LoadObject<UMaterial>(nullptr, TEXT("/MeshModelingToolset/Materials/GizmoDefaultMaterial"));
+
+	if (CenterBallMaterial != nullptr)
+	{
+		PreviewMesh->SetMaterial(CenterBallMaterial);
+	}
+
+	FMinimalBoxMeshGenerator BoxGen;
+	BoxGen.Box = FOrientedBox3d(FVector3d::Zero(), FVector3d(10, 10, 3));
+	FDynamicMesh3 Mesh(&BoxGen.Generate());
+
+	PreviewMesh->UpdatePreview(&Mesh);
+
+	return PreviewMesh;
+}
+
+
+
+void UPositionPlaneGizmo::Render(IToolsContextRenderAPI* RenderAPI)
+{
+	FPrimitiveDrawInterface* PDI = RenderAPI->GetPrimitiveDrawInterface();
+
+	FViewCameraState CameraState;
+	GetGizmoManager()->GetContextQueriesAPI()->GetCurrentViewState(CameraState);
+	QuickTransformer.UpdateCameraState(CameraState);
+
+
+	FTransform Transform = CenterBallShape->GetTransform();
+	FMatrix Matrix = Transform.ToMatrixNoScale();
+
+	DrawBox(PDI, Matrix, FVector(10, 10, 3),
+		CenterBallMaterial->GetRenderProxy(), SDPG_Foreground);
+
+	//DrawSphere(PDI, FVector(0, 0, 100), FRotator(), FVector(50, 50, 50), 16, 16,
+	//	CenterBallMaterial->GetRenderProxy(), SDPG_Foreground);
+	//	
+	
+	if (bInTransformDrag)
+	{
+		QuickTransformer.Render(RenderAPI);
+	}
+}
+
+
+
+bool UPositionPlaneGizmo::HitTest(const FRay& Ray, FHitResult& OutHit)
+{
+	if (CenterBallShape->TestRayIntersection(Ray))
+	{
+		OutHit.Distance = 0.1;
+		return true;
+	}
+	return false;
+}
+
+
+void UPositionPlaneGizmo::OnBeginDrag(const FRay& Ray)
+{
+	bInTransformDrag = true;
+	FTransform CurTransform = CenterBallShape->GetTransform();
+	QuickTransformer.SetActiveFrameFromWorldNormal(CurTransform.GetTranslation(), CurTransform.GetRotation().GetAxisZ(), true);
+}
+
+
+void UPositionPlaneGizmo::OnUpdateDrag(const FRay& Ray)
+{
+	if (bInTransformDrag)
+	{
+		FVector3d SnapPos;
+		if (QuickTransformer.UpdateSnap(Ray, SnapPos))
+		{
+			FTransform CurTransform = CenterBallShape->GetTransform();
+			CurTransform.SetTranslation((FVector)SnapPos);
+			CenterBallShape->SetTransform(CurTransform);
+			PostUpdatedPosition();
+		}
+	}
+	else
+	{
+		FVector RayStart = Ray.Origin;
+		FVector RayEnd = Ray.PointAt(999999);
+		FCollisionObjectQueryParams QueryParams(FCollisionObjectQueryParams::AllObjects);
+		FHitResult Result;
+		bool bHitWorld = TargetWorld->LineTraceSingleByObjectType(Result, RayStart, RayEnd, QueryParams);
+		if (bHitWorld)
+		{
+			FFrame3f UpdatedFrame(CenterBallShape->GetTransform());
+			UpdatedFrame.AlignAxis(2, Result.ImpactNormal);
+			UpdatedFrame.Origin = Result.ImpactPoint;
+			CenterBallShape->SetTransform(UpdatedFrame.ToFTransform());
+			PostUpdatedPosition();
+		}
+	}
+
+}
+
+void UPositionPlaneGizmo::OnEndDrag(const FRay& Ray)
+{
+	bInTransformDrag = false;
+}
+
+
+
+void UPositionPlaneGizmo::ExternalUpdatePosition(const FVector& Position, const FQuat& Orientation, bool bPostUpdate)
+{
+	CenterBallShape->SetTransform(FTransform(Orientation, Position));
+	if (bPostUpdate)
+	{
+		PostUpdatedPosition();
+	}
+}
+
+
+
+void UPositionPlaneGizmo::PostUpdatedPosition()
+{
+	FFrame3d CurPosition(CenterBallShape->GetTransform());
+	if (OnPositionUpdatedFunc != nullptr)
+	{
+		OnPositionUpdatedFunc(CurPosition);
+	}
+}
+
+
+
+/*
+ * Mouse Input Behavior
+ */
+
+
+void UPositionPlaneOnSceneInputBehavior::Initialize(UPositionPlaneGizmo* GizmoIn)
+{
+	this->Gizmo = GizmoIn;
+	bInDragCapture = false;
+}
+
+
+FInputCaptureRequest UPositionPlaneOnSceneInputBehavior::WantsCapture(const FInputDeviceState& input)
+{
+	if (IsPressed(input))
+	{
+		FHitResult OutHit;
+		if (Gizmo->HitTest(input.Mouse.WorldRay, OutHit))
+		{
+			return FInputCaptureRequest::Begin(this, EInputCaptureSide::Any, OutHit.Distance);
+		}
+	}
+	return FInputCaptureRequest::Ignore();
+}
+
+FInputCaptureUpdate UPositionPlaneOnSceneInputBehavior::BeginCapture(const FInputDeviceState& input, EInputCaptureSide eSide)
+{
+	//Gizmo->SetShiftToggle(input.bShiftKeyDown);
+	Gizmo->OnBeginDrag(input.Mouse.WorldRay);
+	LastWorldRay = input.Mouse.WorldRay;
+	bInDragCapture = true;
+	return FInputCaptureUpdate::Begin(this, EInputCaptureSide::Any);
+}
+
+FInputCaptureUpdate UPositionPlaneOnSceneInputBehavior::UpdateCapture(const FInputDeviceState& input, const FInputCaptureData& data)
+{
+	LastWorldRay = input.Mouse.WorldRay;
+
+	if (IsReleased(input))
+	{
+		Gizmo->OnEndDrag(input.Mouse.WorldRay);
+		bInDragCapture = false;
+		return FInputCaptureUpdate::End();
+	}
+
+	Gizmo->OnUpdateDrag(input.Mouse.WorldRay);
+	return FInputCaptureUpdate::Continue();
+}
+
+void UPositionPlaneOnSceneInputBehavior::ForceEndCapture(const FInputCaptureData& data)
+{
+	if (bInDragCapture)
+	{
+		Gizmo->OnEndDrag(LastWorldRay);
+		bInDragCapture = false;
+	}
+
+	// nothing to do
+}
diff --git a/Engine/Plugins/Experimental/MeshModelingToolset/Source/MeshModelingTools/Private/RemeshMeshTool.cpp b/Engine/Plugins/Experimental/MeshModelingToolset/Source/MeshModelingTools/Private/RemeshMeshTool.cpp
new file mode 100644
--- /dev/null
+++ b/Engine/Plugins/Experimental/MeshModelingToolset/Source/MeshModelingTools/Private/RemeshMeshTool.cpp
@@ -0,0 +1,288 @@
+// Copyright 1998-2019 Epic Games, Inc. All Rights Reserved.
+
+#include "RemeshMeshTool.h"
+#include "InteractiveToolManager.h"
+#include "ToolBuilderUtil.h"
+
+#include "DynamicMesh3.h"
+#include "DynamicMeshAttributeSet.h"
+#include "Remesher.h"
+#include "MeshConstraintsUtil.h"
+#include "ProjectionTargets.h"
+#include "MeshNormals.h"
+
+#include "SimpleDynamicMeshComponent.h"
+
+#include "SceneManagement.h" // for FPrimitiveDrawInterface
+#include "Async/ParallelFor.h"
+
+#define LOCTEXT_NAMESPACE "URemeshMeshTool"
+
+
+/*
+ * ToolBuilder
+ */
+
+
+bool URemeshMeshToolBuilder::CanBuildTool(const FToolBuilderState& SceneState) const
+{
+	return ToolBuilderUtil::CountComponents(SceneState, CanMakeComponentTarget) == 1;
+}
+
+UInteractiveTool* URemeshMeshToolBuilder::BuildTool(const FToolBuilderState& SceneState) const
+{
+	URemeshMeshTool* NewTool = NewObject<URemeshMeshTool>(SceneState.ToolManager, "Remesh Tool");
+
+	UActorComponent* ActorComponent = ToolBuilderUtil::FindFirstComponent(SceneState, CanMakeComponentTarget);
+	auto* MeshComponent = Cast<UPrimitiveComponent>(ActorComponent);
+	check(MeshComponent != nullptr);
+	NewTool->SetSelection(MakeComponentTarget(MeshComponent));
+
+	return NewTool;
+}
+
+
+
+/*
+ * Tool
+ */
+
+URemeshMeshToolProperties::URemeshMeshToolProperties()
+{
+	TargetTriangleCount = 5000;
+	SmoothingSpeed = 0.25;
+	RemeshIterations = 20;
+	bDiscardAttributes = false;
+	SmoothingType = ERemeshSmoothingType::MeanValue;
+	bPreserveSharpEdges = true;
+
+	TargetEdgeLength = 5.0;
+	bFlips = true;
+	bSplits = true;
+	bCollapses = true;
+	bReproject = true;
+	bPreventNormalFlips = true;
+	bUseTargetEdgeLength = false;
+}
+
+
+
+URemeshMeshTool::URemeshMeshTool()
+{
+}
+
+
+void URemeshMeshTool::Setup()
+{
+	UInteractiveTool::Setup();
+
+	// create dynamic mesh component to use for live preview
+	DynamicMeshComponent = NewObject<USimpleDynamicMeshComponent>(ComponentTarget->GetOwnerActor(), "DynamicMesh");
+	DynamicMeshComponent->SetupAttachment(ComponentTarget->GetOwnerActor()->GetRootComponent());
+	DynamicMeshComponent->RegisterComponent();
+	DynamicMeshComponent->SetWorldTransform(ComponentTarget->GetWorldTransform());
+
+	// copy material if there is one
+	auto Material = ComponentTarget->GetMaterial(0);
+	if (Material != nullptr)
+	{
+		DynamicMeshComponent->SetMaterial(0, Material);
+	}
+	DynamicMeshComponent->bExplicitShowWireframe = true;
+
+	DynamicMeshComponent->InitializeMesh(ComponentTarget->GetMesh());
+	OriginalMesh.Copy(*DynamicMeshComponent->GetMesh());
+	OriginalMeshSpatial.SetMesh(&OriginalMesh, true);
+
+	// hide input StaticMeshComponent
+	ComponentTarget->SetOwnerVisibility(false);
+
+	// calculate initial mesh area (no utility fn yet)
+	InitialMeshArea = 0;
+	for (int tid : OriginalMesh.TriangleIndicesItr())
+	{
+		InitialMeshArea += OriginalMesh.GetTriArea(tid);
+	}
+
+	BasicProperties = NewObject<URemeshMeshToolProperties>(this);
+	// arbitrary threshold of 5000 tris seems reasonable? 
+	BasicProperties->TargetTriangleCount = (OriginalMesh.TriangleCount() < 5000) ? 5000 : OriginalMesh.TriangleCount();
+	BasicProperties->TargetEdgeLength = CalculateTargetEdgeLength(BasicProperties->TargetTriangleCount);
+
+	// initialize our properties
+	AddToolPropertySource(BasicProperties);
+
+	MeshStatisticsProperties = NewObject<UMeshStatisticsProperties>(this);
+	AddToolPropertySource(MeshStatisticsProperties);
+	MeshStatisticsProperties->Update(*DynamicMeshComponent->GetMesh());
+
+	bResultValid = false;
+}
+
+
+void URemeshMeshTool::Shutdown(EToolShutdownType ShutdownType)
+{
+	if (DynamicMeshComponent != nullptr)
+	{
+		//DynamicMeshComponent->OnMeshChanged.Remove(OnDynamicMeshComponentChangedHandle);
+
+		ComponentTarget->SetOwnerVisibility(true);
+
+		if (ShutdownType == EToolShutdownType::Accept)
+		{
+			// this block bakes the modified DynamicMeshComponent back into the StaticMeshComponent inside an undo transaction
+			GetToolManager()->BeginUndoTransaction(LOCTEXT("RemeshMeshToolTransactionName", "Remesh Mesh"));
+			ComponentTarget->CommitMesh([=](FMeshDescription* MeshDescription)
+			{
+				DynamicMeshComponent->Bake(MeshDescription, true);
+			});
+			GetToolManager()->EndUndoTransaction();
+		}
+
+		DynamicMeshComponent->UnregisterComponent();
+		DynamicMeshComponent->DestroyComponent();
+		DynamicMeshComponent = nullptr;
+	}
+}
+
+
+void URemeshMeshTool::Render(IToolsContextRenderAPI* RenderAPI)
+{
+	UpdateResult();
+
+
+	FPrimitiveDrawInterface* PDI = RenderAPI->GetPrimitiveDrawInterface();
+	FTransform Transform = ComponentTarget->GetWorldTransform(); //Actor->GetTransform();
+
+	FColor LineColor(255, 0, 0);
+	FDynamicMesh3* TargetMesh = DynamicMeshComponent->GetMesh();
+	if (TargetMesh->HasAttributes())
+	{
+		FDynamicMeshUVOverlay* UVOverlay = TargetMesh->Attributes()->PrimaryUV();
+		for (int eid : TargetMesh->EdgeIndicesItr()) 
+		{
+			if (UVOverlay->IsSeamEdge(eid)) 
+			{
+				FVector3d A, B;
+				TargetMesh->GetEdgeV(eid, A, B);
+				PDI->DrawLine(Transform.TransformPosition(A), Transform.TransformPosition(B),
+					LineColor, 0, 2.0, 1.0f, true);
+			}
+		}
+	}
+}
+
+
+void URemeshMeshTool::OnPropertyModified(UObject* PropertySet, UProperty* Property)
+{
+	bResultValid = false;
+}
+
+
+double URemeshMeshTool::CalculateTargetEdgeLength(int TargetTriCount)
+{
+	double TargetTriArea = InitialMeshArea / (double)TargetTriCount;
+	double EdgeLen = TriangleUtil::EquilateralEdgeLengthForArea(TargetTriArea);
+	return (double)FMath::RoundToInt(EdgeLen*100.0) / 100.0;
+}
+
+
+
+void URemeshMeshTool::UpdateResult()
+{
+	if (bResultValid) 
+	{
+		return;
+	}
+
+	FDynamicMesh3* TargetMesh = DynamicMeshComponent->GetMesh();
+	TargetMesh->Copy(OriginalMesh);
+
+	if (BasicProperties->bDiscardAttributes && BasicProperties->bPreserveSharpEdges == false)
+	{
+		TargetMesh->DiscardAttributes();
+	}
+
+	FRemesher Remesher(TargetMesh);
+	Remesher.bEnableSplits = BasicProperties->bSplits;
+	Remesher.bEnableFlips = BasicProperties->bFlips;
+	Remesher.bEnableCollapses = BasicProperties->bCollapses;
+
+	// recalculate target edge length
+	if (BasicProperties->bUseTargetEdgeLength == false)
+	{
+		BasicProperties->TargetEdgeLength = CalculateTargetEdgeLength(BasicProperties->TargetTriangleCount);
+	}
+
+	Remesher.SetTargetEdgeLength(BasicProperties->TargetEdgeLength);
+
+	Remesher.ProjectionMode = (BasicProperties->bReproject) ? 
+		FRemesher::ETargetProjectionMode::AfterRefinement : FRemesher::ETargetProjectionMode::NoProjection;
+
+	Remesher.bEnableSmoothing = (BasicProperties->SmoothingSpeed > 0);
+	Remesher.SmoothSpeedT = BasicProperties->SmoothingSpeed;
+	//Remesher.SmoothType = FRemesher::ESmoothTypes::MeanValue;
+	Remesher.SmoothType = (BasicProperties->bDiscardAttributes) ?
+		FRemesher::ESmoothTypes::Uniform :
+		(FRemesher::ESmoothTypes)(int)BasicProperties->SmoothingType;
+	bool bIsUniformSmooth = (Remesher.SmoothType == FRemesher::ESmoothTypes::Uniform);
+
+	Remesher.bPreventNormalFlips = BasicProperties->bPreventNormalFlips;
+
+	Remesher.DEBUG_CHECK_LEVEL = 0;
+
+	FMeshConstraints constraints;
+	FMeshConstraintsUtil::ConstrainAllSeams(constraints, *TargetMesh, true, !BasicProperties->bPreserveSharpEdges);
+	Remesher.SetExternalConstraints(&constraints);
+
+	FMeshProjectionTarget ProjTarget(&OriginalMesh, &OriginalMeshSpatial);
+	Remesher.SetProjectionTarget(&ProjTarget);
+
+	if (BasicProperties->bDiscardAttributes && BasicProperties->bPreserveSharpEdges == true)
+	{
+		TargetMesh->DiscardAttributes();
+	}
+
+	// run the remesh iterations
+	for (int k = 0; k < BasicProperties->RemeshIterations; ++k)
+	{
+		// If we are not uniform smoothing, then flips seem to often make things worse.
+		// Possibly this is because without the tangential flow, we won't get to the nice tris.
+		// In this case we are better off basically not flipping, and just letting collapses resolve things
+		// regular-valence polygons - things stay "stuck". 
+		// @todo try implementing edge-length flip criteria instead of valence-flip
+		if (bIsUniformSmooth == false)
+		{
+			bool bUseFlipsThisPass = (k % 2 == 0 && k < BasicProperties->RemeshIterations/2);
+			Remesher.bEnableFlips = bUseFlipsThisPass && BasicProperties->bFlips;
+		}
+
+		Remesher.BasicRemeshPass();
+	}
+
+	if (!TargetMesh->HasAttributes() && !TargetMesh->HasVertexNormals())
+	{
+		FMeshNormals::QuickComputeVertexNormals(*TargetMesh);
+	}
+
+	DynamicMeshComponent->NotifyMeshUpdated();
+	GetToolManager()->PostInvalidation();
+	MeshStatisticsProperties->Update(*DynamicMeshComponent->GetMesh());
+
+	bResultValid = true;
+}
+
+bool URemeshMeshTool::HasAccept() const
+{
+	return true;
+}
+
+bool URemeshMeshTool::CanAccept() const
+{
+	return true;
+}
+
+
+
+
+#undef LOCTEXT_NAMESPACE
diff --git a/Engine/Plugins/Experimental/MeshModelingToolset/Source/MeshModelingTools/Private/SmoothMeshTool.cpp b/Engine/Plugins/Experimental/MeshModelingToolset/Source/MeshModelingTools/Private/SmoothMeshTool.cpp
new file mode 100644
--- /dev/null
+++ b/Engine/Plugins/Experimental/MeshModelingToolset/Source/MeshModelingTools/Private/SmoothMeshTool.cpp
@@ -0,0 +1,215 @@
+// Copyright 1998-2019 Epic Games, Inc. All Rights Reserved.
+
+#include "SmoothMeshTool.h"
+#include "InteractiveToolManager.h"
+#include "ToolBuilderUtil.h"
+
+#include "DynamicMesh3.h"
+#include "DynamicMeshToMeshDescription.h"
+#include "MeshDescriptionToDynamicMesh.h"
+
+#include "SimpleDynamicMeshComponent.h"
+
+#include "SceneManagement.h" // for FPrimitiveDrawInterface
+#include "ToolSetupUtil.h"
+
+// Smoothing operators
+#include "SmoothingOps/IterativeSmoothingOp.h"
+#include "SmoothingOps/CotanSmoothingOp.h"
+//#include "SmoothingOps/MeanValueSmoothingOp.h"
+
+
+#define LOCTEXT_NAMESPACE "USmoothMeshTool"
+
+
+/*
+ * ToolBuilder
+ */
+
+
+bool USmoothMeshToolBuilder::CanBuildTool(const FToolBuilderState& SceneState) const
+{
+	return ToolBuilderUtil::CountComponents(SceneState, CanMakeComponentTarget) == 1;
+}
+
+UInteractiveTool* USmoothMeshToolBuilder::BuildTool(const FToolBuilderState& SceneState) const
+{
+	USmoothMeshTool* NewTool = NewObject<USmoothMeshTool>(SceneState.ToolManager);
+
+	UActorComponent* ActorComponent = ToolBuilderUtil::FindFirstComponent(SceneState, CanMakeComponentTarget);
+	auto* MeshComponent = Cast<UPrimitiveComponent>(ActorComponent);
+	check(MeshComponent != nullptr);
+
+	NewTool->SetSelection(MakeComponentTarget(MeshComponent));
+	NewTool->SetWorld(SceneState.World);
+	NewTool->SetAssetAPI(AssetAPI);
+
+	return NewTool;
+}
+
+
+
+/*
+ * Tool
+ */
+
+void USmoothMeshTool::SetWorld(UWorld* World)
+{
+	this->TargetWorld = World;
+}
+
+void USmoothMeshTool::SetAssetAPI(IToolsContextAssetAPI* AssetAPIIn)
+{
+	this->AssetAPI = AssetAPIIn;
+}
+
+
+void USmoothMeshTool::Setup()
+{
+	UInteractiveTool::Setup();
+
+	// hide input StaticMeshComponent
+	ComponentTarget->SetOwnerVisibility(false);
+
+	// initialize our properties
+	ToolPropertyObjects.Add(this);
+
+	// populate the SrcDynamicMesh with a conversion of the input mesh.
+	FMeshDescriptionToDynamicMesh Converter;
+	Converter.bPrintDebugMessages = false;
+	Converter.Convert(ComponentTarget->GetMesh(), SrcDynamicMesh);
+
+	// Initialize the preview mesh with a copy of the source mesh.
+	{
+		// Construct the preview object and set the material on it.
+		Preview = NewObject<UMeshOpPreviewWithBackgroundCompute>(this, "Preview");
+		Preview->Setup(this->TargetWorld, this); // Adds the actual functional tool in the Preview object
+
+		Preview->PreviewMesh->UpdatePreview(&SrcDynamicMesh);
+		Preview->PreviewMesh->SetTransform(ComponentTarget->GetWorldTransform());
+
+	auto Material = ComponentTarget->GetMaterial(0);
+
+		Preview->ConfigureMaterials(
+			ToolSetupUtil::GetDefaultMaterial(GetToolManager(), Material),
+			ToolSetupUtil::GetDefaultWorkingMaterial(GetToolManager())
+		);
+	}
+
+	// show the preview mesh
+	Preview->SetVisibility(true);
+
+	// start the compute
+	Preview->InvalidateResult();    // start compute
+
+	bResultValid = false;
+
+}
+
+
+void USmoothMeshTool::Shutdown(EToolShutdownType ShutdownType)
+{
+	// Restore (unhide) the source meshes
+		ComponentTarget->SetOwnerVisibility(true);
+
+	if (Preview != nullptr)
+	{
+		TUniquePtr<FDynamicMeshOpResult> Result = Preview->Shutdown();
+		if (ShutdownType == EToolShutdownType::Accept)
+		{
+			
+			GetToolManager()->BeginUndoTransaction(LOCTEXT("SmoothMeshToolTransactionName", "Smooth Mesh"));
+
+			FDynamicMesh3* DynamicMeshResult = Result->Mesh.Get();
+			check(DynamicMeshResult != nullptr);
+
+			ComponentTarget->CommitMesh([DynamicMeshResult](FMeshDescription* MeshDescription)
+			{
+				FDynamicMeshToMeshDescription Converter;
+				Converter.Convert(DynamicMeshResult, *MeshDescription);
+			});
+
+
+			GetToolManager()->EndUndoTransaction();
+		}
+	}
+}
+
+void USmoothMeshTool::Render(IToolsContextRenderAPI* RenderAPI)
+{
+	UpdateResult();
+}
+
+#if WITH_EDITOR
+void USmoothMeshTool::PostEditChangeProperty(FPropertyChangedEvent& PropertyChangedEvent)
+{
+	// one of the parameters changed.  Dirty the result for force a recompute.
+	Preview->InvalidateResult();
+	bResultValid = false;
+}
+#endif
+
+void USmoothMeshTool::OnPropertyModified(UObject* PropertySet, UProperty* Property)
+{
+	Preview->InvalidateResult();
+	bResultValid = false;
+}
+
+
+void USmoothMeshTool::Tick(float DeltaTime)
+{
+	Preview->Tick(DeltaTime);
+}
+
+void USmoothMeshTool::UpdateResult()
+{
+	if (bResultValid) 
+	{
+		return;
+	}
+
+	bResultValid = Preview->HaveValidResult();
+}
+
+bool USmoothMeshTool::HasAccept() const
+{
+	return true;
+}
+
+bool USmoothMeshTool::CanAccept() const
+{
+	return bResultValid;
+}
+
+TSharedPtr<FDynamicMeshOperator> USmoothMeshTool::MakeNewOperator()
+{
+	TSharedPtr<FSmoothingOpBase> MeshOp;
+	switch (SmoothType)
+	{
+	default:
+	case ESmoothMeshToolSmoothType::Iterative:
+		MeshOp = MakeShared<FIterativeSmoothingOp>(&SrcDynamicMesh, SmoothSpeed, SmoothIterations);
+		break;
+
+	case ESmoothMeshToolSmoothType::BiHarmonic_Cotan:
+		MeshOp = MakeShared<FCotanSmoothingOp>(&SrcDynamicMesh, SmoothSpeed, SmoothIterations);
+		break;
+/**
+	case ESmoothMeshToolSmoothType::BiHarmonic_MVW:
+		MeshOp = MakeShared<FMeanValueSmoothingOp>(&SrcDynamicMesh, SmoothSpeed, SmoothIterations);
+		break;
+*/
+	}
+
+	const FTransform XForm = ComponentTarget->GetWorldTransform();
+	FTransform3d XForm3d(XForm);
+	MeshOp->SetTransform(XForm3d);
+
+	return MeshOp;
+
+}
+
+
+
+
+#undef LOCTEXT_NAMESPACE
diff --git a/Engine/Plugins/Experimental/MeshModelingToolset/Source/MeshModelingTools/Private/TransformMeshesTool.cpp b/Engine/Plugins/Experimental/MeshModelingToolset/Source/MeshModelingTools/Private/TransformMeshesTool.cpp
new file mode 100644
--- /dev/null
+++ b/Engine/Plugins/Experimental/MeshModelingToolset/Source/MeshModelingTools/Private/TransformMeshesTool.cpp
@@ -0,0 +1,417 @@
+// Copyright 1998-2019 Epic Games, Inc. All Rights Reserved.
+
+#include "TransformMeshesTool.h"
+#include "InteractiveToolManager.h"
+#include "InteractiveGizmoManager.h"
+#include "ToolBuilderUtil.h"
+#include "ToolSetupUtil.h"
+#include "DynamicMesh3.h"
+#include "BaseBehaviors/ClickDragBehavior.h"
+
+#include "BaseGizmos/GizmoComponents.h"
+#include "BaseGizmos/TransformGizmo.h"
+
+#include "Components/PrimitiveComponent.h"
+#include "CollisionQueryParams.h"
+#include "Engine/World.h"
+
+#define LOCTEXT_NAMESPACE "UTransformMeshesTool"
+
+/*
+ * ToolBuilder
+ */
+
+
+bool UTransformMeshesToolBuilder::CanBuildTool(const FToolBuilderState& SceneState) const
+{
+	return ToolBuilderUtil::CountComponents(SceneState, CanMakeComponentTarget) >= 1;
+}
+
+UInteractiveTool* UTransformMeshesToolBuilder::BuildTool(const FToolBuilderState& SceneState) const
+{
+	UTransformMeshesTool* NewTool = NewObject<UTransformMeshesTool>(SceneState.ToolManager);
+
+	TArray<UActorComponent*> Components = ToolBuilderUtil::FindAllComponents(SceneState, CanMakeComponentTarget);
+	check(Components.Num() > 0);
+
+	TArray<TUniquePtr<FPrimitiveComponentTarget>> ComponentTargets;
+	for (UActorComponent* ActorComponent : Components)
+	{
+		auto* MeshComponent = Cast<UPrimitiveComponent>(ActorComponent);
+		if ( MeshComponent )
+		{
+			ComponentTargets.Add(MakeComponentTarget(MeshComponent));
+		}
+	}
+
+	NewTool->SetSelection(MoveTemp(ComponentTargets));
+	NewTool->SetWorld(SceneState.World, SceneState.GizmoManager);
+
+	return NewTool;
+}
+
+
+
+
+/*
+ * Tool
+ */
+
+UTransformMeshesTool::UTransformMeshesTool()
+{
+}
+
+void UTransformMeshesTool::SetWorld(UWorld* World, UInteractiveGizmoManager* GizmoManagerIn)
+{
+	this->TargetWorld = World;
+	this->GizmoManager = GizmoManagerIn;
+}
+
+
+void UTransformMeshesTool::Setup()
+{
+	UInteractiveTool::Setup();
+
+	UClickDragInputBehavior* ClickDragBehavior = NewObject<UClickDragInputBehavior>(this);
+	ClickDragBehavior->Initialize(this);
+	AddInputBehavior(ClickDragBehavior);
+
+	TransformProps = NewObject<UTransformMeshesToolProperties>();
+	AddToolPropertySource(TransformProps);
+
+	UpdateTransformMode(TransformProps->TransformMode);
+
+	GetToolManager()->DisplayMessage(
+		LOCTEXT("OnStartTransformMeshesTool", "To transform Objects around points, reposition the Gizmo using Set Pivot mode (S). To quickly position Objects, enable Snap Drag mode (D). A cycles through Transform modes, W and E through SnapDrag Source and Rotation types."),
+		EToolMessageLevel::UserNotification);
+}
+
+
+
+
+
+
+
+void UTransformMeshesTool::Shutdown(EToolShutdownType ShutdownType)
+{
+	GizmoManager->DestroyAllGizmosByOwner(this);
+}
+
+
+
+void UTransformMeshesTool::Tick(float DeltaTime)
+{
+	// make sure state is up-to-date
+
+	if (CurTransformMode != TransformProps->TransformMode)
+	{
+		UpdateTransformMode(TransformProps->TransformMode);
+	}
+
+	if (bCurSetPivotMode != TransformProps->bSetPivot)
+	{
+		bool bEnableSetPivot = TransformProps->bSetPivot;
+		if (TransformProps->TransformMode == ETransformMeshesTransformMode::SharedGizmoLocal)
+		{
+			bEnableSetPivot = false;
+		}
+		UpdateSetPivotModes(bEnableSetPivot);
+		bCurSetPivotMode = TransformProps->bSetPivot;
+	}
+}
+
+
+void UTransformMeshesTool::Render(IToolsContextRenderAPI* RenderAPI)
+{
+	FPrimitiveDrawInterface* PDI = RenderAPI->GetPrimitiveDrawInterface();
+}
+
+void UTransformMeshesTool::OnPropertyModified(UObject* PropertySet, UProperty* Property)
+{
+}
+
+
+void UTransformMeshesTool::UpdateSetPivotModes(bool bEnableSetPivot)
+{
+	for (FTransformMeshesTarget& Target : ActiveGizmos)
+	{
+		Target.TransformProxy->bSetPivotMode = bEnableSetPivot;
+	}
+}
+
+
+
+void UTransformMeshesTool::RegisterActions(FInteractiveToolActionSet& ActionSet)
+{
+	ActionSet.RegisterAction(this, (int32)EStandardToolActions::BaseClientDefinedActionID + 1,
+		TEXT("ToggleSetPivot"),
+		LOCTEXT("TransformToggleSetPivot", "Toggle Set Pivot"),
+		LOCTEXT("TransformToggleSetPivotTooltip", "Toggle Set Pivot on and off"),
+		EModifierKey::None, EKeys::S,
+		[this]() { this->TransformProps->bSetPivot = !this->TransformProps->bSetPivot; });
+
+	ActionSet.RegisterAction(this, (int32)EStandardToolActions::BaseClientDefinedActionID + 2,
+		TEXT("ToggleSnapDrag"),
+		LOCTEXT("TransformToggleSnapDrag", "Toggle SnapDrag"),
+		LOCTEXT("TransformToggleSnapDragTooltip", "Toggle SnapDrag on and off"),
+		EModifierKey::None, EKeys::D,
+		[this]() { this->TransformProps->bEnableSnapDragging = !this->TransformProps->bEnableSnapDragging; });
+
+	ActionSet.RegisterAction(this, (int32)EStandardToolActions::BaseClientDefinedActionID + 3,
+		TEXT("CycleTransformMode"),
+		LOCTEXT("TransformCycleTransformMode", "Next Transform Mode"),
+		LOCTEXT("TransformCycleTransformModeTooltip", "Cycle through available Transform Modes"),
+		EModifierKey::None, EKeys::A,
+		[this]() { this->TransformProps->TransformMode = (ETransformMeshesTransformMode)(((uint8)TransformProps->TransformMode+1) % (uint8)ETransformMeshesTransformMode::LastValue); });
+
+
+	ActionSet.RegisterAction(this, (int32)EStandardToolActions::BaseClientDefinedActionID + 4,
+		TEXT("CycleSourceMode"),
+		LOCTEXT("TransformCycleSourceMode", "Next SnapDrag Source Mode"),
+		LOCTEXT("TransformCycleSourceModeTooltip", "Cycle through available SnapDrag Source Modes"),
+		EModifierKey::None, EKeys::W,
+		[this]() { this->TransformProps->SnapDragSource = (ETransformMeshesSnapDragSource)(((uint8)TransformProps->SnapDragSource + 1) % (uint8)ETransformMeshesSnapDragSource::LastValue); });
+
+	ActionSet.RegisterAction(this, (int32)EStandardToolActions::BaseClientDefinedActionID + 5,
+		TEXT("CycleRotationMode"),
+		LOCTEXT("TransformCycleRotationMode", "Next SnapDrag Rotation Mode"),
+		LOCTEXT("TransformCycleRotationModeTooltip", "Cycle through available SnapDrag Rotation Modes"),
+		EModifierKey::None, EKeys::E,
+		[this]() { this->TransformProps->RotationMode = (ETransformMeshesSnapDragRotationMode)(((uint8)TransformProps->RotationMode + 1) % (uint8)ETransformMeshesSnapDragRotationMode::LastValue); });
+
+
+}
+
+
+
+
+
+void UTransformMeshesTool::UpdateTransformMode(ETransformMeshesTransformMode NewMode)
+{
+	ResetActiveGizmos();
+
+	switch (NewMode)
+	{
+		default:
+		case ETransformMeshesTransformMode::SharedGizmo:
+			SetActiveGizmos_Single(false);
+			break;
+
+		case ETransformMeshesTransformMode::SharedGizmoLocal:
+			SetActiveGizmos_Single(true);
+			break;
+
+		case ETransformMeshesTransformMode::PerObjectGizmo:
+			SetActiveGizmos_PerObject();
+			break;
+	}
+
+	CurTransformMode = NewMode;
+}
+
+
+
+void UTransformMeshesTool::SetActiveGizmos_Single(bool bLocalRotations)
+{
+	check(ActiveGizmos.Num() == 0);
+
+	FTransformMeshesTarget Transformable;
+	Transformable.TransformProxy = NewObject<UTransformProxy>(this);
+	Transformable.TransformProxy->bRotatePerObject = bLocalRotations;
+
+	for (TUniquePtr<FPrimitiveComponentTarget>& Target : ComponentTargets)
+	{
+		Transformable.TransformProxy->AddComponent(Target->GetOwnerComponent());
+	}
+	Transformable.TransformGizmo = GizmoManager->Create3AxisTransformGizmo(this);
+	Transformable.TransformGizmo->SetActiveTarget(Transformable.TransformProxy);
+
+	ActiveGizmos.Add(Transformable);
+}
+
+void UTransformMeshesTool::SetActiveGizmos_PerObject()
+{
+	check(ActiveGizmos.Num() == 0);
+
+	for (TUniquePtr<FPrimitiveComponentTarget>& Target : ComponentTargets)
+	{
+		FTransformMeshesTarget Transformable;
+		Transformable.TransformProxy = NewObject<UTransformProxy>(this);
+		Transformable.TransformProxy->AddComponent(Target->GetOwnerComponent());
+		Transformable.TransformGizmo = GizmoManager->Create3AxisTransformGizmo(this);
+		Transformable.TransformGizmo->SetActiveTarget(Transformable.TransformProxy);
+		ActiveGizmos.Add(Transformable);
+	}
+}
+
+void UTransformMeshesTool::ResetActiveGizmos()
+{
+	GizmoManager->DestroyAllGizmosByOwner(this);
+	ActiveGizmos.Reset();
+}
+
+
+
+// does not make sense that CanBeginClickDragSequence() returns a RayHit? Needs to be an out-argument...
+FInputRayHit UTransformMeshesTool::CanBeginClickDragSequence(const FInputDeviceRay& PressPos)
+{
+	if (TransformProps->bEnableSnapDragging == false || ActiveGizmos.Num() == 0)
+	{
+		return FInputRayHit();
+	}
+
+	ActiveSnapDragIndex = -1;
+
+	float MinHitDistance = TNumericLimits<float>::Max();
+	FVector HitNormal;
+
+	for ( int k = 0; k < ComponentTargets.Num(); ++k )
+	{
+		TUniquePtr<FPrimitiveComponentTarget>& Target = ComponentTargets[k];
+
+		FHitResult WorldHit;
+		if (Target->HitTest(PressPos.WorldRay, WorldHit))
+		{
+			MinHitDistance = FMath::Min(MinHitDistance, WorldHit.Distance);
+			HitNormal = WorldHit.Normal;
+			ActiveSnapDragIndex = k;
+		}
+	}
+	return (MinHitDistance < TNumericLimits<float>::Max()) ? FInputRayHit(MinHitDistance, HitNormal) : FInputRayHit();
+}
+
+void UTransformMeshesTool::OnClickPress(const FInputDeviceRay& PressPos)
+{
+	FInputRayHit HitPos = CanBeginClickDragSequence(PressPos);
+	check(HitPos.bHit);
+
+	GetToolManager()->BeginUndoTransaction(LOCTEXT("TransformToolTransformTxnName", "SnapDrag"));
+
+	FTransformMeshesTarget& ActiveTarget =
+		(TransformProps->TransformMode == ETransformMeshesTransformMode::PerObjectGizmo) ?
+		ActiveGizmos[ActiveSnapDragIndex] : ActiveGizmos[0];
+	USceneComponent* GizmoComponent = ActiveTarget.TransformGizmo->GetGizmoActor()->GetRootComponent();
+	StartDragTransform = GizmoComponent->GetComponentToWorld();
+
+	if (TransformProps->SnapDragSource == ETransformMeshesSnapDragSource::ClickPoint)
+	{
+		StartDragFrameWorld = FFrame3d(PressPos.WorldRay.PointAt(HitPos.HitDepth), HitPos.HitNormal);
+	}
+	else
+	{
+		StartDragFrameWorld = FFrame3d(StartDragTransform);
+	}
+
+}
+
+
+void UTransformMeshesTool::OnClickDrag(const FInputDeviceRay& DragPos)
+{
+	FCollisionObjectQueryParams ObjectQueryParams(FCollisionObjectQueryParams::AllObjects);
+	FCollisionQueryParams CollisionParams = FCollisionQueryParams::DefaultQueryParam;
+
+	bool bApplyToPivot = TransformProps->bSetPivot;
+
+	if (bApplyToPivot == false)
+	{
+		int IgnoreIndex = (TransformProps->TransformMode == ETransformMeshesTransformMode::PerObjectGizmo) ?
+			ActiveSnapDragIndex : -1;
+		for (int k = 0; k < ComponentTargets.Num(); ++k)
+		{
+			if (IgnoreIndex == -1 || k == IgnoreIndex)
+			{
+				CollisionParams.AddIgnoredComponent(ComponentTargets[k]->GetOwnerComponent());
+			}
+		}
+	}
+
+
+	bool bRotate = (TransformProps->RotationMode != ETransformMeshesSnapDragRotationMode::Ignore);
+	float NormalSign = (TransformProps->RotationMode == ETransformMeshesSnapDragRotationMode::AlignFlipped) ? -1.0f : 1.0f;
+
+	FHitResult Result;
+	bool bWorldHit = TargetWorld->LineTraceSingleByObjectType(Result, DragPos.WorldRay.Origin, DragPos.WorldRay.PointAt(999999), ObjectQueryParams, CollisionParams);
+	if (bWorldHit == false)
+	{
+		return;
+	}
+
+	if (bApplyToPivot)
+	{
+		FVector HitPos = Result.ImpactPoint;
+		FVector TargetNormal = (-NormalSign) * Result.Normal;
+
+		FQuaterniond AlignRotation = (bRotate) ?
+			FQuaterniond(FVector3d::UnitZ(), TargetNormal) : FQuaterniond::Identity();
+
+		FTransform NewTransform = StartDragTransform;
+		NewTransform.SetRotation((FQuat)AlignRotation);
+		NewTransform.SetTranslation(HitPos);
+
+		FTransformMeshesTarget& ActiveTarget =
+			(TransformProps->TransformMode == ETransformMeshesTransformMode::PerObjectGizmo) ?
+			ActiveGizmos[ActiveSnapDragIndex] : ActiveGizmos[0];
+		USceneComponent* GizmoComponent = ActiveTarget.TransformGizmo->GetGizmoActor()->GetRootComponent();
+		GizmoComponent->SetWorldTransform(NewTransform);
+	}
+	else
+	{
+		FVector HitPos = Result.ImpactPoint;
+		FVector TargetNormal = NormalSign * Result.Normal;
+
+		FFrame3d FromFrameWorld = StartDragFrameWorld;
+		FFrame3d ToFrameWorld(HitPos, TargetNormal);
+		FFrame3d ObjectFrameWorld(StartDragTransform);
+
+		FVector3d CenterShift = FromFrameWorld.Origin - ObjectFrameWorld.Origin;
+		FQuaterniond AlignRotation(FromFrameWorld.Z(), ToFrameWorld.Z());
+		if (bRotate == false)
+		{
+			AlignRotation = FQuaterniond::Identity();
+		}
+		FVector3d AlignTranslate = ToFrameWorld.Origin - FromFrameWorld.Origin;
+
+		FTransform NewTransform = StartDragTransform;
+		NewTransform.Accumulate( FTransform(CenterShift) );
+		NewTransform.Accumulate( FTransform(AlignRotation) );
+		NewTransform.Accumulate( FTransform(AlignTranslate) );
+		CenterShift = AlignRotation * CenterShift;
+		NewTransform.Accumulate( FTransform(-CenterShift) );
+
+		FTransformMeshesTarget& ActiveTarget =
+			(TransformProps->TransformMode == ETransformMeshesTransformMode::PerObjectGizmo) ?
+			ActiveGizmos[ActiveSnapDragIndex] : ActiveGizmos[0];
+		USceneComponent* GizmoComponent = ActiveTarget.TransformGizmo->GetGizmoActor()->GetRootComponent();
+		GizmoComponent->SetWorldTransform(NewTransform);
+	}
+
+}
+
+
+void UTransformMeshesTool::OnClickRelease(const FInputDeviceRay& ReleasePos)
+{
+	OnTerminateDragSequence();
+}
+
+void UTransformMeshesTool::OnTerminateDragSequence()
+{
+	FTransformMeshesTarget& ActiveTarget =
+		(TransformProps->TransformMode == ETransformMeshesTransformMode::PerObjectGizmo) ?
+		ActiveGizmos[ActiveSnapDragIndex] : ActiveGizmos[0];
+	USceneComponent* GizmoComponent = ActiveTarget.TransformGizmo->GetGizmoActor()->GetRootComponent();
+	FTransform EndDragtransform = GizmoComponent->GetComponentToWorld();
+
+	TUniquePtr<FComponentWorldTransformChange> Change = MakeUnique<FComponentWorldTransformChange>(StartDragTransform, EndDragtransform);
+	GetToolManager()->EmitObjectChange(GizmoComponent, MoveTemp(Change),
+		LOCTEXT("TransformToolTransformTxnName", "SnapDrag"));
+
+	GetToolManager()->EndUndoTransaction();
+
+	ActiveSnapDragIndex = -1;
+}
+
+
+
+
+
+#undef LOCTEXT_NAMESPACE
