Copying //UE4/Dev-Editor to Dev-Main (//UE4/Dev-Main) for 4.24

diff --git a/Engine/Plugins/Experimental/MeshModelingToolset/Source/MeshModelingTools/Private/UVProjectionTool.cpp b/Engine/Plugins/Experimental/MeshModelingToolset/Source/MeshModelingTools/Private/UVProjectionTool.cpp
new file mode 100644
--- /dev/null
+++ b/Engine/Plugins/Experimental/MeshModelingToolset/Source/MeshModelingTools/Private/UVProjectionTool.cpp
@@ -0,0 +1,366 @@
+// Copyright 1998-2019 Epic Games, Inc. All Rights Reserved.
+
+#include "UVProjectionTool.h"
+#include "InteractiveToolManager.h"
+#include "ToolBuilderUtil.h"
+
+#include "ToolSetupUtil.h"
+
+#include "DynamicMesh3.h"
+#include "BaseBehaviors/MultiClickSequenceInputBehavior.h"
+#include "Selection/SelectClickedAction.h"
+
+#include "MeshDescriptionToDynamicMesh.h"
+#include "DynamicMeshToMeshDescription.h"
+#include "Materials/MaterialInstanceDynamic.h"
+
+#include "InteractiveGizmoManager.h"
+
+#include "AssetGenerationUtil.h"
+
+#include "BaseGizmos/GizmoComponents.h"
+#include "BaseGizmos/TransformGizmo.h"
+
+
+#define LOCTEXT_NAMESPACE "UUVProjectionTool"
+
+
+/*
+ * ToolBuilder
+ */
+
+
+bool UUVProjectionToolBuilder::CanBuildTool(const FToolBuilderState& SceneState) const
+{
+	// note most of the code is written to support working on any number > 0, but that seems maybe confusing UI-wise and is not fully tested, so I've limited it to acting on one for now
+	// TODO: if enable tool working on multiple components, figure out what to do if we have multiple component targets that point to the same underlying mesh data?
+	return ToolBuilderUtil::CountComponents(SceneState, CanMakeComponentTarget) == 1;
+}
+
+UInteractiveTool* UUVProjectionToolBuilder::BuildTool(const FToolBuilderState& SceneState) const
+{
+	UUVProjectionTool* NewTool = NewObject<UUVProjectionTool>(SceneState.ToolManager);
+
+	TArray<UActorComponent*> Components = ToolBuilderUtil::FindAllComponents(SceneState, CanMakeComponentTarget);
+	check(Components.Num() > 0);
+
+	TArray<TUniquePtr<FPrimitiveComponentTarget>> ComponentTargets;
+	for (UActorComponent* ActorComponent : Components)
+	{
+		auto* MeshComponent = Cast<UPrimitiveComponent>(ActorComponent);
+		if ( MeshComponent )
+		{
+			ComponentTargets.Add(MakeComponentTarget(MeshComponent));
+		}
+	}
+
+	NewTool->SetSelection(MoveTemp(ComponentTargets));
+	NewTool->SetWorld(SceneState.World, SceneState.GizmoManager);
+	NewTool->SetAssetAPI(AssetAPI);
+
+	return NewTool;
+}
+
+
+
+
+/*
+ * Tool
+ */
+
+UUVProjectionToolProperties::UUVProjectionToolProperties()
+{
+	UVProjectionMethod = EUVProjectionMethod::Cube;
+	ProjectionPrimitiveScale = FVector::OneVector;
+	UVScale = FVector2D::UnitVector;
+	UVOffset = FVector2D::ZeroVector;
+}
+
+UUVProjectionAdvancedProperties::UUVProjectionAdvancedProperties()
+{
+}
+
+
+UUVProjectionTool::UUVProjectionTool()
+{
+}
+
+void UUVProjectionTool::SetWorld(UWorld* World, UInteractiveGizmoManager* GizmoManagerIn)
+{
+	this->TargetWorld = World;
+	this->GizmoManager = GizmoManagerIn;
+}
+
+void UUVProjectionTool::Setup()
+{
+	UInteractiveTool::Setup();
+
+	// hide input StaticMeshComponent
+	for (TUniquePtr<FPrimitiveComponentTarget>& ComponentTarget : ComponentTargets)
+	{
+		ComponentTarget->SetOwnerVisibility(false);
+	}
+
+	BasicProperties = NewObject<UUVProjectionToolProperties>(this, TEXT("UV Projection Settings"));
+	AdvancedProperties = NewObject<UUVProjectionAdvancedProperties>(this, TEXT("Advanced Settings"));
+
+	// initialize our properties
+	AddToolPropertySource(BasicProperties);
+	AddToolPropertySource(AdvancedProperties);
+
+	MaterialSettings = NewObject<UExistingMeshMaterialProperties>(this);
+	MaterialSettings->Setup();
+	AddToolPropertySource(MaterialSettings);
+
+	// initialize the PreviewMesh+BackgroundCompute object
+	UpdateNumPreviews();
+
+	// set up visualizers
+	ProjectionShapeVisualizer.LineColor = FLinearColor::Red;
+	ProjectionShapeVisualizer.LineThickness = 2.0;
+
+	
+	for (UMeshOpPreviewWithBackgroundCompute* Preview : Previews)
+	{
+		Preview->InvalidateResult();
+	}
+}
+
+
+void UUVProjectionTool::UpdateNumPreviews()
+{
+	int32 CurrentNumPreview = Previews.Num();
+	int32 TargetNumPreview = ComponentTargets.Num();
+	if (TargetNumPreview < CurrentNumPreview)
+	{
+		for (int32 PreviewIdx = CurrentNumPreview - 1; PreviewIdx >= TargetNumPreview; PreviewIdx--)
+		{
+			Previews[PreviewIdx]->Cancel();
+			GizmoManager->DestroyGizmo(TransformGizmos[PreviewIdx]);
+		}
+		Previews.SetNum(TargetNumPreview);
+		TransformGizmos.SetNum(TargetNumPreview);
+		TransformProxies.SetNum(TargetNumPreview);
+		OriginalDynamicMeshes.SetNum(TargetNumPreview);
+	}
+	else
+	{
+		OriginalDynamicMeshes.SetNum(TargetNumPreview);
+		for (int32 PreviewIdx = CurrentNumPreview; PreviewIdx < TargetNumPreview; PreviewIdx++)
+		{
+			UUVProjectionOperatorFactory *OpFactory = NewObject<UUVProjectionOperatorFactory>();
+			OpFactory->Tool = this;
+			OpFactory->ComponentIndex = PreviewIdx;
+			OriginalDynamicMeshes[PreviewIdx] = MakeShared<FDynamicMesh3>();
+			FMeshDescriptionToDynamicMesh Converter;
+			Converter.bPrintDebugMessages = true;
+			Converter.Convert(ComponentTargets[PreviewIdx]->GetMesh(), *OriginalDynamicMeshes[PreviewIdx]);
+
+			FVector Center, Extents;
+			FBoxSphereBounds Bounds = ComponentTargets[PreviewIdx]->GetOwnerComponent()->CalcLocalBounds();
+			
+			FTransform LocalXF(Bounds.Origin);
+			LocalXF.SetScale3D(Bounds.BoxExtent);
+
+			UMeshOpPreviewWithBackgroundCompute* Preview = Previews.Add_GetRef(NewObject<UMeshOpPreviewWithBackgroundCompute>(OpFactory, "Preview"));
+			Preview->Setup(this->TargetWorld, OpFactory);
+			Preview->ConfigureMaterials(
+				ToolSetupUtil::GetDefaultMaterial(GetToolManager(), ComponentTargets[PreviewIdx]->GetMaterial(0)),
+				ToolSetupUtil::GetDefaultWorkingMaterial(GetToolManager())
+			);
+			Preview->SetVisibility(true);
+
+
+			UTransformProxy* TransformProxy = TransformProxies.Add_GetRef(NewObject<UTransformProxy>(this));
+			TransformProxy->SetTransform(LocalXF * ComponentTargets[PreviewIdx]->GetWorldTransform());
+			TransformProxy->OnTransformChanged.AddUObject(this, &UUVProjectionTool::TransformChanged);
+
+			UTransformGizmo* TransformGizmo = TransformGizmos.Add_GetRef(GizmoManager->Create3AxisTransformGizmo(this));
+			TransformGizmo->SetActiveTarget(TransformProxy);
+		}
+		check(TransformProxies.Num() == TargetNumPreview);
+		check(TransformGizmos.Num() == TargetNumPreview);
+	}
+}
+
+
+void UUVProjectionTool::Shutdown(EToolShutdownType ShutdownType)
+{
+	// Restore (unhide) the source meshes
+	for (TUniquePtr<FPrimitiveComponentTarget>& ComponentTarget : ComponentTargets)
+	{
+		ComponentTarget->SetOwnerVisibility(true);
+	}
+
+	TArray<TUniquePtr<FDynamicMeshOpResult>> Results;
+	for (UMeshOpPreviewWithBackgroundCompute* Preview : Previews)
+	{
+		Results.Emplace(Preview->Shutdown());
+	}
+	if (ShutdownType == EToolShutdownType::Accept)
+	{
+		GenerateAsset(Results);
+	}
+
+	GizmoManager->DestroyAllGizmosByOwner(this);
+	TransformGizmos.Empty();
+}
+
+void UUVProjectionTool::SetAssetAPI(IToolsContextAssetAPI* AssetAPIIn)
+{
+	this->AssetAPI = AssetAPIIn;
+}
+
+TSharedPtr<FDynamicMeshOperator> UUVProjectionOperatorFactory::MakeNewOperator()
+{
+	TSharedPtr<FUVProjectionOp> Op = MakeShared<FUVProjectionOp>();
+	Op->ProjectionMethod = Tool->BasicProperties->UVProjectionMethod;
+
+	// TODO: de-dupe this logic (it's also in Render, below)
+	FTransform LocalScale = FTransform::Identity;
+	LocalScale.SetScale3D(Tool->BasicProperties->ProjectionPrimitiveScale);
+	Op->ProjectionTransform = LocalScale * Tool->TransformProxies[ComponentIndex]->GetTransform();
+	Op->CylinderProjectToTopOrBottomAngleThreshold = Tool->BasicProperties->CylinderProjectToTopOrBottomAngleThreshold;
+	Op->UVScale = Tool->BasicProperties->UVScale;
+	Op->UVOffset = Tool->BasicProperties->UVOffset;
+	Op->bWorldSpaceUVScale = Tool->BasicProperties->bWorldSpaceUVScale;
+
+	FTransform LocalToWorld = Tool->ComponentTargets[ComponentIndex]->GetWorldTransform();
+	Op->OriginalMesh = Tool->OriginalDynamicMeshes[ComponentIndex];
+	
+	Op->SetTransform(LocalToWorld);
+
+	return Op;
+}
+
+
+
+void UUVProjectionTool::Render(IToolsContextRenderAPI* RenderAPI)
+{
+	ProjectionShapeVisualizer.bDepthTested = false;
+	ProjectionShapeVisualizer.BeginFrame(RenderAPI, CameraState);
+	FTransform LocalScale = FTransform::Identity;
+	LocalScale.SetScale3D(BasicProperties->ProjectionPrimitiveScale);
+	for (UTransformProxy* TransformProxy : TransformProxies)
+	{
+		ProjectionShapeVisualizer.SetTransform(LocalScale * TransformProxy->GetTransform());
+
+		switch (BasicProperties->UVProjectionMethod)
+		{
+		case EUVProjectionMethod::Cube:
+			ProjectionShapeVisualizer.DrawWireBox(FBox(FVector(-1, -1, -1), FVector(1, 1, 1)));
+			break;
+		case EUVProjectionMethod::Cylinder:
+			ProjectionShapeVisualizer.DrawWireCylinder(FVector(0, 0, -1), FVector(0, 0, 1), 1, 2, 20);
+			break;
+		case EUVProjectionMethod::Plane:
+			ProjectionShapeVisualizer.DrawSquare(FVector(0, 0, 0), FVector(2, 0, 0), FVector(0, 2, 0));
+			break;
+		}
+	}
+	
+	ProjectionShapeVisualizer.EndFrame();
+}
+
+void UUVProjectionTool::Tick(float DeltaTime)
+{
+	for (UMeshOpPreviewWithBackgroundCompute* Preview : Previews)
+	{
+		Preview->Tick(DeltaTime);
+	}
+}
+
+
+
+#if WITH_EDITOR
+void UUVProjectionTool::PostEditChangeProperty(FPropertyChangedEvent& PropertyChangedEvent)
+{
+	UpdateNumPreviews();
+	for (UMeshOpPreviewWithBackgroundCompute* Preview : Previews)
+	{
+		Preview->InvalidateResult();
+	}
+}
+#endif
+
+void UUVProjectionTool::OnPropertyModified(UObject* PropertySet, UProperty* Property)
+{
+	// if we don't know what changed, or we know checker density changed, update checker material
+	MaterialSettings->UpdateMaterials();
+	for (int PreviewIdx = 0; PreviewIdx < Previews.Num(); PreviewIdx++)
+	{
+		UMeshOpPreviewWithBackgroundCompute* Preview = Previews[PreviewIdx];
+		MaterialSettings->SetMaterialIfChanged(ComponentTargets[PreviewIdx]->GetMaterial(0), Preview->StandardMaterial, [&Preview, this](UMaterialInterface* Material)
+		{
+			Preview->ConfigureMaterials(ToolSetupUtil::GetDefaultMaterial(GetToolManager(), Material), Preview->WorkingMaterial);
+		});
+	}
+	
+	UpdateNumPreviews();
+	for (UMeshOpPreviewWithBackgroundCompute* Preview : Previews)
+	{
+		Preview->InvalidateResult();
+	}
+}
+
+
+void UUVProjectionTool::TransformChanged(UTransformProxy* Proxy, FTransform Transform)
+{
+	// TODO: if multi-select is re-enabled, only invalidate the preview that actually needs it?
+	for (UMeshOpPreviewWithBackgroundCompute* Preview : Previews)
+	{
+		Preview->InvalidateResult();
+	}
+}
+
+
+bool UUVProjectionTool::HasAccept() const
+{
+	return true;
+}
+
+bool UUVProjectionTool::CanAccept() const
+{
+	for (UMeshOpPreviewWithBackgroundCompute* Preview : Previews)
+	{
+		if (!Preview->HaveValidResult())
+		{
+			return false;
+		}
+	}
+	return true;
+}
+
+
+void UUVProjectionTool::GenerateAsset(const TArray<TUniquePtr<FDynamicMeshOpResult>>& Results)
+{
+	GetToolManager()->BeginUndoTransaction(LOCTEXT("UVProjectionToolTransactionName", "UV Projection Tool"));
+
+	check(Results.Num() == ComponentTargets.Num());
+	
+	for (int32 ComponentIdx = 0; ComponentIdx < ComponentTargets.Num(); ComponentIdx++)
+	{
+		check(Results[ComponentIdx]->Mesh.Get() != nullptr);
+		ComponentTargets[ComponentIdx]->CommitMesh([&Results, &ComponentIdx, this](FMeshDescription* MeshDescription)
+		{
+			FDynamicMeshToMeshDescription Converter;
+			//if (??) // TODO check if UV topology changed?  or remove all traces of this if statement (may be safe to assume it almost always changes w/ a uv projection op)
+			{
+				// full conversion if topology changed
+				Converter.Convert(Results[ComponentIdx]->Mesh.Get(), *MeshDescription);
+			}
+			//else
+			//{
+			//	// otherwise just copy attributes
+			//	Converter.UpdateAttributes(Results[ComponentIdx]->Mesh.Get(), *MeshDescription, false, true);
+			//}
+		});
+	}
+
+	GetToolManager()->EndUndoTransaction();
+}
+
+
+
+
+#undef LOCTEXT_NAMESPACE
diff --git a/Engine/Plugins/Experimental/MeshModelingToolset/Source/MeshModelingTools/Private/WeldMeshEdgesTool.cpp b/Engine/Plugins/Experimental/MeshModelingToolset/Source/MeshModelingTools/Private/WeldMeshEdgesTool.cpp
new file mode 100644
--- /dev/null
+++ b/Engine/Plugins/Experimental/MeshModelingToolset/Source/MeshModelingTools/Private/WeldMeshEdgesTool.cpp
@@ -0,0 +1,195 @@
+// Copyright 1998-2019 Epic Games, Inc. All Rights Reserved.
+
+#include "WeldMeshEdgesTool.h"
+#include "InteractiveToolManager.h"
+#include "ToolBuilderUtil.h"
+
+#include "DynamicMesh3.h"
+#include "DynamicMeshAttributeSet.h"
+#include "Operations/MergeCoincidentMeshEdges.h"
+
+#include "SimpleDynamicMeshComponent.h"
+
+#include "SceneManagement.h" // for FPrimitiveDrawInterface
+#include "Async/ParallelFor.h"
+
+#define LOCTEXT_NAMESPACE "UWeldMeshEdgesTool"
+
+
+/*
+ * ToolBuilder
+ */
+
+
+bool UWeldMeshEdgesToolBuilder::CanBuildTool(const FToolBuilderState& SceneState) const
+{
+	return ToolBuilderUtil::CountComponents(SceneState, CanMakeComponentTarget) == 1;
+}
+
+UInteractiveTool* UWeldMeshEdgesToolBuilder::BuildTool(const FToolBuilderState& SceneState) const
+{
+	UWeldMeshEdgesTool* NewTool = NewObject<UWeldMeshEdgesTool>(SceneState.ToolManager);
+
+	UActorComponent* ActorComponent = ToolBuilderUtil::FindFirstComponent(SceneState, CanMakeComponentTarget);
+	auto* MeshComponent = Cast<UPrimitiveComponent>(ActorComponent);
+	check(MeshComponent != nullptr);
+
+	NewTool->SetSelection(MakeComponentTarget(MeshComponent));
+
+	return NewTool;
+}
+
+
+
+/*
+ * Tool
+ */
+
+UWeldMeshEdgesTool::UWeldMeshEdgesTool()
+{
+	Tolerance = FMergeCoincidentMeshEdges::DEFAULT_TOLERANCE;
+	bOnlyUnique = false;
+}
+
+void UWeldMeshEdgesTool::Setup()
+{
+	UInteractiveTool::Setup();
+
+	// create dynamic mesh component to use for live preview
+	DynamicMeshComponent = NewObject<USimpleDynamicMeshComponent>(ComponentTarget->GetOwnerActor(), "DynamicMesh");
+	DynamicMeshComponent->SetupAttachment(ComponentTarget->GetOwnerActor()->GetRootComponent());
+	DynamicMeshComponent->RegisterComponent();
+	DynamicMeshComponent->SetWorldTransform(ComponentTarget->GetWorldTransform());
+
+	// copy material if there is one
+	auto Material = ComponentTarget->GetMaterial(0);
+	if (Material != nullptr)
+	{
+		DynamicMeshComponent->SetMaterial(0, Material);
+	}
+
+	DynamicMeshComponent->InitializeMesh(ComponentTarget->GetMesh());
+	OriginalMesh.Copy(*DynamicMeshComponent->GetMesh());
+
+	// hide input StaticMeshComponent
+	ComponentTarget->SetOwnerVisibility(false);
+
+	// initialize our properties
+	ToolPropertyObjects.Add(this);
+
+	bResultValid = false;
+}
+
+
+void UWeldMeshEdgesTool::Shutdown(EToolShutdownType ShutdownType)
+{
+	if (DynamicMeshComponent != nullptr)
+	{
+		ComponentTarget->SetOwnerVisibility(true);
+
+		if (ShutdownType == EToolShutdownType::Accept)
+		{
+			// this block bakes the modified DynamicMeshComponent back into the StaticMeshComponent inside an undo transaction
+			GetToolManager()->BeginUndoTransaction(LOCTEXT("WeldMeshEdgesToolTransactionName", "Remesh Mesh"));
+			ComponentTarget->CommitMesh([=](FMeshDescription* MeshDescription)
+			{
+				DynamicMeshComponent->Bake(MeshDescription, true);
+			});
+			GetToolManager()->EndUndoTransaction();
+		}
+
+		DynamicMeshComponent->UnregisterComponent();
+		DynamicMeshComponent->DestroyComponent();
+		DynamicMeshComponent = nullptr;
+	}
+}
+
+
+void UWeldMeshEdgesTool::Render(IToolsContextRenderAPI* RenderAPI)
+{
+	UpdateResult();
+
+
+	FPrimitiveDrawInterface* PDI = RenderAPI->GetPrimitiveDrawInterface();
+	FTransform Transform = ComponentTarget->GetWorldTransform(); //Actor->GetTransform();
+
+	FColor LineColor(200, 200, 200);
+	FDynamicMesh3* TargetMesh = DynamicMeshComponent->GetMesh();
+	FDynamicMeshUVOverlay* UVOverlay = TargetMesh->Attributes()->PrimaryUV();
+	for (int eid : TargetMesh->EdgeIndicesItr()) 
+	{
+		if (UVOverlay->IsSeamEdge(eid)) 
+		{
+			FVector3d A, B;
+			TargetMesh->GetEdgeV(eid, A, B);
+			PDI->DrawLine(Transform.TransformPosition(A), Transform.TransformPosition(B),
+				LineColor, 0, 1.0, 1.0f, true);
+		}
+	}
+
+
+	FColor LineColor2(255, 0, 0);
+	for (int eid : TargetMesh->BoundaryEdgeIndicesItr()) 
+	{
+		FVector3d A, B;
+		TargetMesh->GetEdgeV(eid, A, B);
+		PDI->DrawLine(Transform.TransformPosition(A), Transform.TransformPosition(B),
+			LineColor2, 0, 2.0, 1.0f, true);
+	}
+
+}
+
+#if WITH_EDITOR
+void UWeldMeshEdgesTool::PostEditChangeProperty(FPropertyChangedEvent& PropertyChangedEvent)
+{
+	UProperty* PropertyThatChanged = PropertyChangedEvent.Property;
+	bResultValid = false;
+}
+#endif
+
+void UWeldMeshEdgesTool::UpdateResult()
+{
+	if (bResultValid) 
+	{
+		return;
+	}
+
+	FDynamicMesh3* TargetMesh = DynamicMeshComponent->GetMesh();
+	TargetMesh->Copy(OriginalMesh);
+
+	FMergeCoincidentMeshEdges Merger(TargetMesh);
+	Merger.MergeSearchTolerance = this->Tolerance;
+	Merger.OnlyUniquePairs = this->bOnlyUnique;
+		
+	if (Merger.Apply() == false)
+	{
+		UE_LOG(LogTemp, Warning, TEXT("WeldMeshEdgesTool : Merger.Apply() returned false"));
+		TargetMesh->Copy(OriginalMesh);
+	}
+
+	if (TargetMesh->CheckValidity(true, EValidityCheckFailMode::ReturnOnly) == false)
+	{
+		UE_LOG(LogTemp, Warning, TEXT("WeldMeshEdgesTool : returned mesh is invalid"));
+		TargetMesh->Copy(OriginalMesh);
+	}
+
+	DynamicMeshComponent->NotifyMeshUpdated();
+	GetToolManager()->PostInvalidation();
+
+	bResultValid = true;
+}
+
+bool UWeldMeshEdgesTool::HasAccept() const
+{
+	return true;
+}
+
+bool UWeldMeshEdgesTool::CanAccept() const
+{
+	return true;
+}
+
+
+
+
+#undef LOCTEXT_NAMESPACE
diff --git a/Engine/Plugins/Experimental/MeshModelingToolset/Source/MeshModelingTools/Public/ConvertToPolygonsTool.h b/Engine/Plugins/Experimental/MeshModelingToolset/Source/MeshModelingTools/Public/ConvertToPolygonsTool.h
new file mode 100644
--- /dev/null
+++ b/Engine/Plugins/Experimental/MeshModelingToolset/Source/MeshModelingTools/Public/ConvertToPolygonsTool.h
@@ -0,0 +1,84 @@
+// Copyright 1998-2019 Epic Games, Inc. All Rights Reserved.
+
+#pragma once
+
+#include "CoreMinimal.h"
+#include "UObject/NoExportTypes.h"
+#include "SingleSelectionTool.h"
+#include "InteractiveToolBuilder.h"
+#include "DynamicMesh3.h"
+#include "FindPolygonsAlgorithm.h"
+#include "PreviewMesh.h"
+#include "ConvertToPolygonsTool.generated.h"
+
+// predeclaration
+struct FMeshDescription;
+
+/**
+ *
+ */
+UCLASS()
+class MESHMODELINGTOOLS_API UConvertToPolygonsToolBuilder : public UInteractiveToolBuilder
+{
+	GENERATED_BODY()
+
+public:
+	virtual bool CanBuildTool(const FToolBuilderState & SceneState) const override;
+	virtual UInteractiveTool* BuildTool(const FToolBuilderState & SceneState) const override;
+};
+
+
+
+
+
+UCLASS()
+class MESHMODELINGTOOLS_API UConvertToPolygonsToolProperties : public UInteractiveToolPropertySet
+{
+	GENERATED_BODY()
+
+public:
+	/** Tolerance for planarity */
+	UPROPERTY(EditAnywhere, Category = PolyGroups, meta = (UIMin = "0.001", UIMax = "20.0", ClampMin = "0.0", ClampMax = "90.0"))
+	float AngleTolerance = 0.1f;
+
+	UPROPERTY(EditAnywhere, Category = PolyGroups)
+	bool bCalculateNormals = true;
+};
+
+/**
+ *
+ */
+UCLASS()
+class MESHMODELINGTOOLS_API UConvertToPolygonsTool : public USingleSelectionTool
+{
+	GENERATED_BODY()
+
+public:
+	virtual void Setup() override;
+	virtual void Shutdown(EToolShutdownType ShutdownType) override;
+
+	virtual void Render(IToolsContextRenderAPI* RenderAPI) override;
+
+	virtual bool HasCancel() const override { return true; }
+	virtual bool HasAccept() const override;
+	virtual bool CanAccept() const override;
+
+	virtual void OnPropertyModified(UObject* PropertySet, UProperty* Property) override;
+
+protected:
+	UPROPERTY()
+	UConvertToPolygonsToolProperties* Settings;
+
+	UPROPERTY()
+	UPreviewMesh* PreviewMesh;
+
+protected:
+	FDynamicMesh3 SearchMesh;
+	FDynamicMeshNormalOverlay InitialNormals;
+	FFindPolygonsAlgorithm Polygons;
+	bool bPolygonsValid = false;
+	void UpdatePolygons();
+
+	void ConvertToPolygons(FMeshDescription* MeshIn);
+
+};
diff --git a/Engine/Plugins/Experimental/MeshModelingToolset/Source/MeshModelingTools/Public/DisplaceMeshTool.h b/Engine/Plugins/Experimental/MeshModelingToolset/Source/MeshModelingTools/Public/DisplaceMeshTool.h
new file mode 100644
--- /dev/null
+++ b/Engine/Plugins/Experimental/MeshModelingToolset/Source/MeshModelingTools/Public/DisplaceMeshTool.h
@@ -0,0 +1,129 @@
+// Copyright 1998-2019 Epic Games, Inc. All Rights Reserved.
+
+#pragma once
+
+#include "CoreMinimal.h"
+#include "UObject/NoExportTypes.h"
+#include "Engine/Classes/Engine/Texture2D.h"
+#include "SingleSelectionTool.h"
+#include "InteractiveToolBuilder.h"
+#include "DynamicMesh3.h"
+#include "MeshNormals.h"
+#include "Spatial/SampledScalarField2.h"
+#include "DisplaceMeshTool.generated.h"
+
+
+// predeclarations
+struct FMeshDescription;
+class USimpleDynamicMeshComponent;
+
+
+UENUM()
+enum class EDisplaceMeshToolDisplaceType : uint8
+{
+	/** Displace with N iterations */
+	Constant UMETA(DisplayName = "Constant"),
+
+	/** Displace with N iterations */
+	RandomNoise UMETA(DisplayName = "Random Noise"),
+
+	/** Displace with N iterations */
+	DisplacementMap UMETA(DisplayName = "Displacement Map"),
+
+};
+
+
+
+/**
+ *
+ */
+UCLASS()
+class MESHMODELINGTOOLS_API UDisplaceMeshToolBuilder : public UInteractiveToolBuilder
+{
+	GENERATED_BODY()
+
+public:
+	virtual bool CanBuildTool(const FToolBuilderState& SceneState) const override;
+	virtual UInteractiveTool* BuildTool(const FToolBuilderState& SceneState) const override;
+};
+
+
+
+
+
+/**
+ * Simple Mesh Displacement Tool
+ */
+UCLASS()
+class MESHMODELINGTOOLS_API UDisplaceMeshTool : public USingleSelectionTool
+{
+	GENERATED_BODY()
+
+public:
+	UDisplaceMeshTool();
+
+	virtual void Setup() override;
+	virtual void Shutdown(EToolShutdownType ShutdownType) override;
+
+	virtual void Render(IToolsContextRenderAPI* RenderAPI) override;
+
+	virtual bool HasCancel() const override { return true; }
+	virtual bool HasAccept() const override;
+	virtual bool CanAccept() const override;
+#if WITH_EDITOR
+	virtual void PostEditChangeProperty(FPropertyChangedEvent& PropertyChangedEvent) override;
+#endif
+protected:
+	// need to update bResultValid if these are modified, so we don't publicly expose them. 
+	// @todo setters/getters for these
+
+	/** primary brush mode */
+	UPROPERTY(EditAnywhere, Category = Options)
+	EDisplaceMeshToolDisplaceType DisplacementType;
+
+	/** Displacement intensity */
+	UPROPERTY(EditAnywhere, Category = Options, meta = (UIMin = "-100.0", UIMax = "100.0", ClampMin = "-10000.0", ClampMax = "100000.0"))
+	float DisplaceIntensity;
+
+	/** Seed for randomization */
+	UPROPERTY(EditAnywhere, Category = Options, meta = (EditCondition = "DisplacementType == EDisplaceMeshToolDisplaceType::RandomNoise"))
+	int RandomSeed;
+
+	/** Subdivision iterations for mesh */
+	UPROPERTY(EditAnywhere, Category = Options, meta = (UIMin = "0", UIMax = "10", ClampMin = "0", ClampMax = "100"))
+	int Subdivisions;
+
+	/** Displacement map */
+	UPROPERTY(EditAnywhere, Category = Options)
+	UTexture2D* DisplacementMap;
+
+
+protected:
+	USimpleDynamicMeshComponent* DynamicMeshComponent;
+	
+	FDynamicMesh3 OriginalMesh;		// don't actually need whole mesh, just positions...
+
+	FDynamicMesh3 SubdividedMesh;
+	int CachedSubdivisionsCount = -1;
+	void UpdateSubdividedMesh();
+
+	FSampledScalarField2f DisplaceField;
+	UTexture2D* CachedMapSource = nullptr;
+	void UpdateMap(bool bForceUpdate = false);
+
+	TArray<FVector3d> PositionBuffer;
+	FMeshNormals NormalsBuffer;
+	TArray<FVector3d> DisplacedBuffer;
+	bool bResultValid;
+	void UpdateResult();
+
+	void ComputeDisplacement_Constant();
+	void ComputeDisplacement_RandomNoise();
+
+	void ComputeDisplacement_Map();
+
+
+private:
+	//TUniquePtr<
+
+};
diff --git a/Engine/Plugins/Experimental/MeshModelingToolset/Source/MeshModelingTools/Public/DrawPolygonTool.h b/Engine/Plugins/Experimental/MeshModelingToolset/Source/MeshModelingTools/Public/DrawPolygonTool.h
new file mode 100644
--- /dev/null
+++ b/Engine/Plugins/Experimental/MeshModelingToolset/Source/MeshModelingTools/Public/DrawPolygonTool.h
@@ -0,0 +1,329 @@
+// Copyright 1998-2019 Epic Games, Inc. All Rights Reserved.
+
+#pragma once
+
+#include "CoreMinimal.h"
+#include "UObject/NoExportTypes.h"
+#include "InteractiveGizmo.h"
+#include "InteractiveToolBuilder.h"
+#include "BaseTools/SingleClickTool.h"
+#include "PreviewMesh.h"
+#include "DynamicMeshAABBTree3.h"
+#include "Snapping/PointPlanarSnapSolver.h"
+#include "ToolSceneQueriesUtil.h"
+#include "Properties/MeshMaterialProperties.h"
+#include "Changes/ValueWatcher.h"
+#include "DrawPolygonTool.generated.h"
+
+
+class UTransformGizmo;
+class UTransformProxy;
+
+/**
+ *
+ */
+UCLASS()
+class MESHMODELINGTOOLS_API UDrawPolygonToolBuilder : public UInteractiveToolBuilder
+{
+	GENERATED_BODY()
+
+public:
+	IToolsContextAssetAPI* AssetAPI;
+
+	UDrawPolygonToolBuilder() 
+	{
+		AssetAPI = nullptr;
+	}
+
+	virtual bool CanBuildTool(const FToolBuilderState& SceneState) const override;
+	virtual UInteractiveTool* BuildTool(const FToolBuilderState& SceneState) const override;
+};
+
+
+
+/** Polygon Tool Draw Type */
+UENUM()
+enum class EDrawPolygonDrawMode : uint8
+{
+	/** Freehand Polygon Drawing */
+	Freehand UMETA(DisplayName = "Freehand"),
+
+	/** Circle */
+	Circle UMETA(DisplayName = "Circle"),
+
+	/** Square */
+	Square UMETA(DisplayName = "Square"),
+
+	/** Rectangle */
+	Rectangle UMETA(DisplayName = "Rectangle")
+
+};
+
+
+/** Output of Draw Polygon Tool */
+UENUM()
+enum class EDrawPolygonOutputMode : uint8
+{
+	/** Generate a meshed planar polygon */
+	MeshedPolygon UMETA(DisplayName = "Flat Mesh"),
+
+	/** Extrude closed polygon to constant height determined by Extrude Height Property */
+	ExtrudedConstant UMETA(DisplayName = "Extrude To Height"),
+
+	/** Extrusion height is set via additional mouse input after closing polygon */
+	ExtrudedInteractive UMETA(DisplayName = "Interactive Extrude"),
+};
+
+
+
+
+
+UCLASS()
+class MESHMODELINGTOOLS_API UDrawPolygonToolStandardProperties : public UInteractiveToolPropertySet
+{
+	GENERATED_BODY()
+
+public:
+	UDrawPolygonToolStandardProperties();
+
+	UPROPERTY(EditAnywhere, NonTransactional, Category = Polygon)
+	EDrawPolygonDrawMode PolygonType = EDrawPolygonDrawMode::Freehand;
+
+	UPROPERTY(EditAnywhere, NonTransactional, Category = Polygon)
+	EDrawPolygonOutputMode OutputMode = EDrawPolygonOutputMode::ExtrudedInteractive;
+
+	/** Extrusion Distance in non-interactive mode */
+	UPROPERTY(EditAnywhere, NonTransactional, Category = Polygon, meta = (UIMin = "-1000", UIMax = "1000", ClampMin = "-10000", ClampMax = "10000"))
+	float ExtrudeHeight = 100.0f;
+
+	/** Extrusion Distance in non-interactive mode */
+	UPROPERTY(EditAnywhere, NonTransactional, Category = Polygon, meta = (UIMin = "3", UIMax = "100", ClampMin = "3", ClampMax = "10000"))
+	int Steps = 16;
+
+	UPROPERTY(EditAnywhere, NonTransactional, Category = Polygon)
+	bool bAllowSelfIntersections = false;
+
+	UPROPERTY(EditAnywhere, NonTransactional, Category = Polygon)
+	bool bShowGizmo = true;
+
+	//
+	// save/restore support
+	//
+	virtual void SaveProperties(UInteractiveTool* SaveFromTool) override;
+	virtual void RestoreProperties(UInteractiveTool* RestoreToTool) override;
+};
+
+
+
+
+UCLASS()
+class MESHMODELINGTOOLS_API UDrawPolygonToolSnapProperties : public UInteractiveToolPropertySet
+{
+	GENERATED_BODY()
+
+public:
+	UPROPERTY(EditAnywhere, NonTransactional, Category = Snapping)
+	bool bEnableSnapping = true;
+
+	UPROPERTY(EditAnywhere, NonTransactional, Category = Snapping, Meta = (EditCondition = "bEnableSnapping"))
+	bool bSnapToVertices = true;
+
+	UPROPERTY(EditAnywhere, NonTransactional, Category = Snapping, Meta = (EditCondition = "bEnableSnapping"))
+	bool bSnapToEdges = false;
+
+	UPROPERTY(EditAnywhere, NonTransactional, Category = Snapping, Meta = (EditCondition = "bEnableSnapping"))
+	bool bSnapToAngles = true;
+
+	UPROPERTY(EditAnywhere, NonTransactional, Category = Snapping, Meta = (EditCondition = "bEnableSnapping"))
+	bool bSnapToLengths = true;
+
+	UPROPERTY(VisibleAnywhere, NonTransactional, Category = Snapping)
+	float SegmentLength = 0.0f;
+
+	UPROPERTY(EditAnywhere, NonTransactional, Category = Snapping)
+	bool bHitSceneObjects = false;
+
+	UPROPERTY(EditAnywhere, NonTransactional, Category = Snapping, Meta = (EditCondition = "bHitSceneObjects"))
+	float HitNormalOffset = 0.0f;
+
+	//
+	// save/restore support
+	//
+	virtual void SaveProperties(UInteractiveTool* SaveFromTool) override;
+	virtual void RestoreProperties(UInteractiveTool* RestoreToTool) override;
+};
+
+
+
+
+
+/**
+ * This tool allows the user to draw and extrude 2D polygons
+ */
+UCLASS()
+class MESHMODELINGTOOLS_API UDrawPolygonTool : public UInteractiveTool, public IClickSequenceBehaviorTarget
+{
+	GENERATED_BODY()
+
+public:
+	UDrawPolygonTool();
+
+	virtual void RegisterActions(FInteractiveToolActionSet& ActionSet) override;
+
+	virtual void SetWorld(UWorld* World);
+	virtual void SetAssetAPI(IToolsContextAssetAPI* AssetAPI);
+
+	virtual void Setup() override;
+	virtual void Shutdown(EToolShutdownType ShutdownType) override;
+
+	virtual void Render(IToolsContextRenderAPI* RenderAPI) override;
+
+	virtual bool HasCancel() const override { return false; }
+	virtual bool HasAccept() const override { return false; }
+	virtual bool CanAccept() const override { return false; }
+
+	// IClickSequenceBehaviorTarget implementation
+
+	virtual void OnBeginSequencePreview(const FInputDeviceRay& ClickPos) override;
+	virtual bool CanBeginClickSequence(const FInputDeviceRay& ClickPos) override;
+	virtual void OnBeginClickSequence(const FInputDeviceRay& ClickPos) override;
+	virtual void OnNextSequencePreview(const FInputDeviceRay& ClickPos) override;
+	virtual bool OnNextSequenceClick(const FInputDeviceRay& ClickPos) override;
+	virtual void OnTerminateClickSequence() override;
+	virtual bool RequestAbortClickSequence() override;
+	virtual void OnUpdateModifierState(int ModifierID, bool bIsOn) override;
+
+	// polygon drawing functions
+	virtual void ResetPolygon();
+	virtual void UpdatePreviewVertex(const FVector& PreviewVertex);
+	virtual void AppendVertex(const FVector& Vertex);
+	virtual bool FindDrawPlaneHitPoint(const FInputDeviceRay& ClickPos, FVector& HitPosOut);
+	virtual void EmitCurrentPolygon();
+
+	virtual void BeginInteractiveExtrude();
+	virtual void EndInteractiveExtrude();
+	virtual float FindInteractiveHeightDistance(const FInputDeviceRay& ClickPos);
+
+
+public:
+	virtual void PopLastVertexAction();
+
+
+protected:
+	// flags used to identify modifier keys/buttons
+	static const int IgnoreSnappingModifier = 1;
+	static const int AngleSnapModifier = 2;
+
+protected:
+
+	/** Properties that control polygon generation exposed to user via detailsview */
+	UPROPERTY()
+	UDrawPolygonToolStandardProperties* PolygonProperties;
+
+	UPROPERTY()
+	UDrawPolygonToolSnapProperties* SnapProperties;
+
+	UPROPERTY()
+	UNewMeshMaterialProperties* MaterialProperties;
+	
+
+	/** Origin of plane we will draw polygon on */
+	UPROPERTY()
+	FVector DrawPlaneOrigin;
+
+	/** Orientation of plane we will draw polygon on */
+	UPROPERTY()
+	FQuat DrawPlaneOrientation;
+	
+	/** Vertices of current preview polygon */
+	UPROPERTY()
+	TArray<FVector> PolygonVertices;
+
+	/** last vertex of polygon that is actively being updated as input device is moved */
+	UPROPERTY()
+	FVector PreviewVertex;
+
+protected:
+	UWorld* TargetWorld;
+	IToolsContextAssetAPI* AssetAPI;
+
+	FViewCameraState CameraState;
+
+	UPROPERTY()
+	UPreviewMesh* PreviewMesh;
+
+	
+	// drawing plane gizmo
+
+	UPROPERTY()
+	UTransformGizmo* PlaneTransformGizmo;
+
+	UPROPERTY()
+	UTransformProxy* PlaneTransformProxy;
+
+	// called on PlaneTransformProxy.OnTransformChanged
+	void PlaneTransformChanged(UTransformProxy* Proxy, FTransform Transform);
+
+	// calls SetDrawPlaneFromWorldPos when user ctrl+clicks on scene
+	IClickBehaviorTarget* SetPointInWorldConnector = nullptr;
+
+	// updates plane and gizmo position
+	virtual void SetDrawPlaneFromWorldPos(const FVector& Position, const FVector& Normal);
+
+	TValueWatcher<bool> ShowGizmoWatcher;
+	void UpdateShowGizmoState(bool bNewVisibility);
+
+
+	// polygon drawing
+
+	bool bAbortActivePolygonDraw;
+
+	bool bInFixedPolygonMode = false;
+	TArray<FVector> FixedPolygonClickPoints;
+
+	// can close poly if current segment intersects existing segment
+	bool UpdateSelfIntersection();
+	bool bHaveSelfIntersection;
+	int SelfIntersectSegmentIdx;
+	FVector3f SelfIntersectionPoint;
+
+	// only used when SnapSettings.bHitSceneObjects = true
+	bool bHaveSurfaceHit;
+	FVector3f SurfaceHitPoint;
+	FVector3f SurfaceOffsetPoint;
+
+	bool bIgnoreSnappingToggle = false;		// toggled by hotkey (shift)
+	FPointPlanarSnapSolver SnapEngine;
+	ToolSceneQueriesUtil::FSnapGeometry LastSnapGeometry;
+
+	void GenerateFixedPolygon(TArray<FVector>& FixedPoints, TArray<FVector>& VerticesOut);
+
+
+	// extrusion control
+
+	bool bInInteractiveExtrude = false;
+
+	void UpdateLivePreview();
+	bool bPreviewUpdatePending;
+
+	FDynamicMesh3 PreviewHeightTarget;
+	FDynamicMeshAABBTree3 PreviewHeightTargetAABB;
+	FFrame3d PreviewHeightFrame;
+	void GeneratePreviewHeightTarget();
+
+
+	FFrame3f HitPosFrameWorld;
+
+	// we use this to generate extruded meshes
+	void GeneratePolygonMesh(const TArray<FVector>& Polygon, FDynamicMesh3* ResultMeshOut, FFrame3d& WorldFrameOut, bool bIncludePreviewVtx, double ExtrudeDistance, bool bExtrudeSymmetric);
+
+
+	// user feedback messages
+	void ShowStartupMessage();
+	void ShowExtrudeMessage();
+};
+
+
+
+
+
diff --git a/Engine/Plugins/Experimental/MeshModelingToolset/Source/MeshModelingTools/Public/DynamicMeshBrushTool.h b/Engine/Plugins/Experimental/MeshModelingToolset/Source/MeshModelingTools/Public/DynamicMeshBrushTool.h
new file mode 100644
--- /dev/null
+++ b/Engine/Plugins/Experimental/MeshModelingToolset/Source/MeshModelingTools/Public/DynamicMeshBrushTool.h
@@ -0,0 +1,56 @@
+// Copyright 1998-2019 Epic Games, Inc. All Rights Reserved.
+
+#pragma once
+
+#include "CoreMinimal.h"
+#include "UObject/NoExportTypes.h"
+#include "InteractiveToolBuilder.h"
+#include "BaseBrushTool.h"
+#include "DynamicMeshAABBTree3.h"
+#include "SimpleDynamicMeshComponent.h"
+#include "PreviewMesh.h"
+#include "DynamicMeshBrushTool.generated.h"
+
+
+/**
+ * UDynamicMeshBrushTool is a base class that specializes UBaseBrushTool
+ * for brushing on an FDynamicMesh3. The input FPrimitiveComponentTarget is hidden
+ * and a UPreviewMesh is created and shown in its place. This UPreviewMesh is
+ * used for hit-testing and dynamic rendering.
+ * 
+ */
+UCLASS(Transient)
+class MESHMODELINGTOOLS_API UDynamicMeshBrushTool : public UBaseBrushTool
+{
+	GENERATED_BODY()
+
+public:
+	UDynamicMeshBrushTool();
+
+	// UInteractiveTool API
+
+	virtual void Setup() override;
+	virtual void Shutdown(EToolShutdownType ShutdownType) override;
+
+	virtual bool HitTest(const FRay& Ray, FHitResult& OutHit) override;
+
+protected:
+	// subclasses can override these to customize behavior
+	virtual void OnShutdown(EToolShutdownType ShutdownType) {}
+
+
+protected:
+	UPROPERTY()
+	UPreviewMesh* PreviewMesh;
+
+	// this function is called when the component inside the PreviewMesh is modified (eg via an undo/redo event)
+	virtual void OnBaseMeshComponentChanged() {}	
+	FDelegateHandle OnBaseMeshComponentChangedHandle;
+
+	FAxisAlignedBox3d InputMeshBoundsLocal;
+
+	//
+	// UBaseBrushTool private interface
+	//
+	virtual double EstimateMaximumTargetDimension() override;
+};
\ No newline at end of file
diff --git a/Engine/Plugins/Experimental/MeshModelingToolset/Source/MeshModelingTools/Public/DynamicMeshSculptTool.h b/Engine/Plugins/Experimental/MeshModelingToolset/Source/MeshModelingTools/Public/DynamicMeshSculptTool.h
new file mode 100644
--- /dev/null
+++ b/Engine/Plugins/Experimental/MeshModelingToolset/Source/MeshModelingTools/Public/DynamicMeshSculptTool.h
@@ -0,0 +1,370 @@
+// Copyright 1998-2019 Epic Games, Inc. All Rights Reserved.
+
+#pragma once
+
+#include "CoreMinimal.h"
+#include "UObject/NoExportTypes.h"
+#include "BaseTools/MeshSurfacePointTool.h"
+#include "OctreeDynamicMeshComponent.h"
+#include "DynamicMeshAABBTree3.h"
+#include "DynamicMeshOctree3.h"
+#include "MeshNormals.h"
+#include "BaseBrushTool.h"
+#include "Drawing/ToolDataVisualizer.h"
+#include "Changes/ValueWatcher.h"
+#include "Gizmos/BrushStampIndicator.h"
+#include "Properties/MeshMaterialProperties.h"
+#include "DynamicMeshSculptTool.generated.h"
+
+
+DECLARE_STATS_GROUP(TEXT("SculptTool"), STATGROUP_SculptTool, STATCAT_Advanced);
+DECLARE_CYCLE_STAT(TEXT("SculptTool_UpdateROI"), SculptTool_UpdateROI, STATGROUP_SculptTool);
+DECLARE_CYCLE_STAT(TEXT("SculptTool_ApplyStamp"), STAT_SculptToolApplyStamp, STATGROUP_SculptTool );
+DECLARE_CYCLE_STAT(TEXT("SculptTool_Tick"), STAT_SculptToolTick, STATGROUP_SculptTool);
+DECLARE_CYCLE_STAT(TEXT("SculptTool_Tick_ApplyStampBlock"), STAT_SculptTool_Tick_ApplyStampBlock, STATGROUP_SculptTool);
+DECLARE_CYCLE_STAT(TEXT("SculptTool_Tick_ApplyStamp_Remove"), SculptTool_Tick_ApplyStamp_Remove, STATGROUP_SculptTool);
+DECLARE_CYCLE_STAT(TEXT("SculptTool_Tick_ApplyStamp_Insert"), SculptTool_Tick_ApplyStamp_Insert, STATGROUP_SculptTool);
+DECLARE_CYCLE_STAT(TEXT("SculptTool_Tick_RemeshBlock"), STAT_SculptTool_Tick_RemeshBlock, STATGROUP_SculptTool);
+DECLARE_CYCLE_STAT(TEXT("SculptTool_Tick_NormalsBlock"), STAT_SculptTool_Tick_NormalsBlock, STATGROUP_SculptTool);
+DECLARE_CYCLE_STAT(TEXT("SculptTool_Tick_UpdateMeshBlock"), STAT_SculptTool_Tick_UpdateMeshBlock, STATGROUP_SculptTool);
+DECLARE_CYCLE_STAT(TEXT("SculptTool_Normals_Collect"), SculptTool_Normals_Collect, STATGROUP_SculptTool);
+DECLARE_CYCLE_STAT(TEXT("SculptTool_Normals_Compute"), SculptTool_Normals_Compute, STATGROUP_SculptTool);
+DECLARE_CYCLE_STAT(TEXT("SculptTool_Remesh_1Setup"), STAT_SculptTool_Remesh_Setup, STATGROUP_SculptTool);
+DECLARE_CYCLE_STAT(TEXT("SculptTool_Remesh_2Constraints"), STAT_SculptTool_Remesh_Constraints, STATGROUP_SculptTool);
+DECLARE_CYCLE_STAT(TEXT("SculptTool_Remesh_3RemeshROIUpdate"), STAT_SculptTool_Remesh_RemeshROIUpdate, STATGROUP_SculptTool);
+DECLARE_CYCLE_STAT(TEXT("SculptTool_Remesh_4RemeshPass"), STAT_SculptTool_Remesh_RemeshPass, STATGROUP_SculptTool);
+DECLARE_CYCLE_STAT(TEXT("SculptTool_Remesh_5PassOctreeUpdate"), STAT_SculptTool_Remesh_PassOctreeUpdate, STATGROUP_SculptTool);
+DECLARE_CYCLE_STAT(TEXT("SculptTool_Remesh_6Finish"), STAT_SculptTool_Remesh_Finish, STATGROUP_SculptTool);
+
+
+class FMeshVertexChangeBuilder;
+class FDynamicMeshChangeTracker;
+class UPreviewMesh;
+
+/** Mesh Sculpting Brush Types */
+UENUM()
+enum class EDynamicMeshSculptBrushType : uint8
+{
+	/** Move Brush moves vertices parallel to the view plane  */
+	Move UMETA(DisplayName = "Move"),
+
+	/** Sculpt Brush displaces vertices from the surface */
+	Offset UMETA(DisplayName = "Sculpt"),
+
+	/** SculptMax Brush displaces vertices to a maximum offset height */
+	SculptMax UMETA(DisplayName = "Sculpt Max"),
+
+	/** Inflate brush displaces vertices along their vertex normals */
+	Inflate UMETA(DisplayName = "Inflate"),
+
+	/** Pinch Brush pulls vertices towards the center of the brush*/
+	Pinch UMETA(DisplayName = "Pinch"),
+
+	/** Flatten Brush pulls vertices towards the average plane of the brush stamp */
+	Flatten UMETA(DisplayName = "Flatten"),
+
+	/** Plane Brush pulls vertices to a plane defined by the initial brush position  */
+	Plane UMETA(DisplayName = "Plane"),
+
+
+	LastValue UMETA(Hidden)
+
+};
+
+
+
+
+/**
+ *
+ */
+UCLASS()
+class MESHMODELINGTOOLS_API UDynamicMeshSculptToolBuilder : public UMeshSurfacePointToolBuilder
+{
+	GENERATED_BODY()
+
+public:
+	bool bEnableRemeshing;
+
+	UDynamicMeshSculptToolBuilder()
+	{
+		bEnableRemeshing = false;
+	}
+
+	virtual UMeshSurfacePointTool* CreateNewTool(const FToolBuilderState& SceneState) const override;
+};
+
+
+
+
+
+
+
+
+
+UCLASS()
+class MESHMODELINGTOOLS_API UBrushSculptProperties : public UInteractiveToolPropertySet
+{
+	GENERATED_BODY()
+
+public:
+	UBrushSculptProperties();
+
+	/** Primary Brush Mode */
+	UPROPERTY(EditAnywhere, Category = Sculpting)
+	EDynamicMeshSculptBrushType PrimaryBrushType;
+
+	/** Power/Speed of Brush */
+	UPROPERTY(EditAnywhere, Category = Sculpting, meta = (UIMin = "0.0", UIMax = "1.0", ClampMin = "0.0", ClampMax = "1.0", EditCondition = "PrimaryBrushType != EDynamicMeshSculptBrushType::Pull"))
+	float BrushSpeed;
+
+	/** Smoothing Speed of Smoothing brush */
+	UPROPERTY(EditAnywhere, Category = Sculpting, meta = (UIMin = "0.0", UIMax = "1.0", ClampMin = "0.0", ClampMax = "1.0"))
+	float SmoothSpeed;
+
+	/** If true, try to preserve the shape of the UV/3D mapping. This will prevent smoothing in some cases */
+	UPROPERTY(EditAnywhere, Category = Sculpting)
+	bool bPreserveUVFlow;
+
+
+	/** Depth of Brush into surface along view ray or surface normal, depending on brush */
+	UPROPERTY(EditAnywhere, Category = Sculpting, meta = (UIMin = "-0.5", UIMax = "0.5", ClampMin = "-1.0", ClampMax = "1.0"))
+	float BrushDepth;
+
+	/** Disable updating of the brush Target Surface after each stroke */
+	UPROPERTY(EditAnywhere, Category = Sculpting, meta = (EditCondition = "PrimaryBrushType == EDynamicMeshSculptBrushType::Sculpt || PrimaryBrushType == EDynamicMeshSculptBrushType::SculptMax || PrimaryBrushType == EDynamicMeshSculptBrushType::Pinch" ))
+	bool bFreezeTarget;
+
+	virtual void SaveProperties(UInteractiveTool* SaveFromTool) override;
+	virtual void RestoreProperties(UInteractiveTool* RestoreToTool) override;
+
+};
+
+
+
+
+UCLASS()
+class MESHMODELINGTOOLS_API UBrushRemeshProperties : public UInteractiveToolPropertySet
+{
+	GENERATED_BODY()
+
+public:
+	UBrushRemeshProperties();
+
+	/** Target Relative Triangle Sizefor Dynamic Meshing */
+	UPROPERTY(EditAnywhere, Category = Remeshing, meta = (UIMin = "0.5", UIMax = "2.0", ClampMin = "0.1", ClampMax = "100.0"))
+	float RelativeSize;
+
+	/** Smoothing speed for dynamic meshing */
+	UPROPERTY(EditAnywhere, Category = Remeshing, meta = (UIMin = "0.0", UIMax = "1.0", ClampMin = "0.0", ClampMax = "1.0"))
+	float Smoothing;
+
+
+	/** Enable edge flips */
+	UPROPERTY(EditAnywhere, Category = Remeshing, AdvancedDisplay)
+	bool bFlips = false;
+
+	/** Enable edge splits */
+	UPROPERTY(EditAnywhere, Category = Remeshing, AdvancedDisplay)
+	bool bSplits = true;
+
+	/** Enable edge collapses */
+	UPROPERTY(EditAnywhere, Category = Remeshing, AdvancedDisplay)
+	bool bCollapses = true;
+
+	/** Prevent normal flips */
+	UPROPERTY(EditAnywhere, Category = Remeshing, AdvancedDisplay)
+	bool bPreventNormalFlips = false;
+
+};
+
+
+
+/**
+ *
+ */
+UCLASS()
+class MESHMODELINGTOOLS_API UDynamicMeshSculptTool : public UMeshSurfacePointTool
+{
+	GENERATED_BODY()
+
+public:
+	UDynamicMeshSculptTool();
+
+	virtual void SetWorld(UWorld* World);
+	virtual void RegisterActions(FInteractiveToolActionSet& ActionSet) override;
+
+	virtual void Setup() override;
+	virtual void Shutdown(EToolShutdownType ShutdownType) override;
+
+	virtual void Tick(float DeltaTime) override;
+	virtual void Render(IToolsContextRenderAPI* RenderAPI) override;
+
+	virtual bool HasCancel() const override { return true; }
+	virtual bool HasAccept() const override { return true; }
+	virtual bool CanAccept() const override { return true; }
+
+	virtual bool HitTest(const FRay& Ray, FHitResult& OutHit) override;
+
+	virtual void OnBeginDrag(const FRay& Ray) override;
+	virtual void OnUpdateDrag(const FRay& Ray) override;
+	virtual void OnEndDrag(const FRay& Ray) override;
+	virtual bool OnUpdateHover(const FInputDeviceRay& DevicePos) override;
+
+	virtual void SetEnableRemeshing(bool bEnable) { bEnableRemeshing = bEnable; }
+	virtual bool GetEnableRemeshing() const { return bEnableRemeshing; }
+
+
+	virtual void OnPropertyModified(UObject* PropertySet, UProperty* Property) override;
+
+public:
+
+	/** Properties that control sculpting*/
+	UPROPERTY()
+	UBrushSculptProperties* SculptProperties;
+
+	/** Properties that control brush size/etc*/
+	UPROPERTY()
+	UBrushBaseProperties* BrushProperties;
+
+	/** Properties that control dynamic remeshing */
+	UPROPERTY()
+	UBrushRemeshProperties* RemeshProperties;
+
+	UPROPERTY()
+	UMeshEditingViewProperties* ViewProperties;
+
+
+public:
+
+	virtual void IncreaseBrushRadiusAction();
+	virtual void DecreaseBrushRadiusAction();
+	virtual void IncreaseBrushRadiusSmallStepAction();
+	virtual void DecreaseBrushRadiusSmallStepAction();
+
+	virtual void IncreaseBrushSpeedAction();
+	virtual void DecreaseBrushSpeedAction();
+
+
+	virtual void NextBrushModeAction();
+	virtual void PreviousBrushModeAction();
+
+	virtual void NextHistoryBrushModeAction();
+	virtual void PreviousHistoryBrushModeAction();
+
+
+protected:
+
+	UWorld* TargetWorld;		// required to spawn UPreviewMesh/etc
+
+	UPROPERTY()
+	UBrushStampIndicator* BrushIndicator;
+
+	UPROPERTY()
+	UPreviewMesh* BrushIndicatorMesh;
+
+	UPROPERTY()
+	UOctreeDynamicMeshComponent* DynamicMeshComponent;
+
+	// realtime visualization
+	void OnDynamicMeshComponentChanged();
+	FDelegateHandle OnDynamicMeshComponentChangedHandle;
+
+	TValueWatcher<bool> ShowWireframeWatcher;
+	TValueWatcher<EMeshEditingMaterialModes> MaterialModeWatcher;
+	void UpdateMaterialMode(EMeshEditingMaterialModes NewMode);
+
+
+	FInterval1d BrushRelativeSizeRange;
+	double CurrentBrushRadius;
+	void CalculateBrushRadius();
+
+	bool bEnableRemeshing;
+	double InitialEdgeLength;
+
+	bool bInDrag;
+
+	FPlane ActiveDragPlane;
+	FVector LastHitPosWorld;
+	FVector BrushStartCenterWorld;
+	FVector BrushStartNormalWorld;
+	FVector LastBrushPosLocal;
+	FVector LastBrushPosWorld;
+	FVector LastBrushPosNormalWorld;
+	
+
+	TArray<int> VertexROI;
+	TSet<int> VertexSetBuffer;
+	TSet<int> TriangleROI;
+	void UpdateROI(const FVector& BrushPos);
+
+	bool bRemeshPending;
+	bool bNormalUpdatePending;
+	bool bTargetDirty;
+
+	bool bSmoothing;
+	bool bInvert;
+
+	bool bHaveRemeshed;
+
+	bool bStampPending;
+	FRay PendingStampRay;
+	int StampTimestamp = 0;
+	EDynamicMeshSculptBrushType LastStampType = EDynamicMeshSculptBrushType::LastValue;
+	EDynamicMeshSculptBrushType PendingStampType = LastStampType;
+	void ApplyStamp(const FRay& WorldRay);
+
+	FDynamicMesh3 BrushTargetMesh;
+	FDynamicMeshAABBTree3 BrushTargetMeshSpatial;
+	FMeshNormals BrushTargetNormals;
+	bool bCachedFreezeTarget = false;
+	void UpdateTarget();
+	bool GetTargetMeshNearest(const FVector3d& Position, double SearchRadius, FVector3d& TargetPosOut, FVector3d& TargetNormalOut);
+
+	bool UpdateBrushPositionOnActivePlane(const FRay& WorldRay);
+	bool UpdateBrushPositionOnTargetMesh(const FRay& WorldRay);
+	bool UpdateBrushPositionOnSculptMesh(const FRay& WorldRay);
+
+	void ApplySmoothBrush(const FRay& WorldRay);
+	void ApplyMoveBrush(const FRay& WorldRay);
+	void ApplyOffsetBrush(const FRay& WorldRay);
+	void ApplySculptMaxBrush(const FRay& WorldRay);
+	void ApplyPinchBrush(const FRay& WorldRay);
+	void ApplyInflateBrush(const FRay& WorldRay);
+	void ApplyPlaneBrush(const FRay& WorldRay);
+	void ApplyFlattenBrush(const FRay& WorldRay);
+
+	double CalculateBrushFalloff(double Distance);
+
+	FFrame3d ActiveFixedBrushPlane;
+	FFrame3d ComputeROIBrushPlane(const FVector3d& BrushCenter);
+
+	TArray<int> TrianglesBuffer;
+	TArray<int> NormalsBuffer;
+	TArray<bool> NormalsVertexFlags;
+	void RecalculateNormals_PerVertex();
+	void RecalculateNormals_Overlay();
+
+	bool bHaveMeshBoundaries;
+	bool bHaveUVSeams;
+	bool bHaveNormalSeams;
+	void PrecomputeRemeshInfo();
+	void RemeshROIPass();
+
+	FMeshVertexChangeBuilder* ActiveVertexChange = nullptr;
+	FDynamicMeshChangeTracker* ActiveMeshChange = nullptr;
+	void BeginChange(bool bIsVertexChange);
+	void EndChange();
+	void SaveActiveROI();
+
+	void UpdateSavedVertex(int vid, const FVector3d& OldPosition, const FVector3d& NewPosition);
+	FCriticalSection UpdateSavedVertexLock;
+
+	double EstimateIntialSafeTargetLength(const FDynamicMesh3& Mesh, int MinTargetTriCount);
+
+	TArray<EDynamicMeshSculptBrushType> BrushTypeHistory;
+	int BrushTypeHistoryIndex = 0;
+
+};
+
+
+
diff --git a/Engine/Plugins/Experimental/MeshModelingToolset/Source/MeshModelingTools/Public/EditNormalsTool.h b/Engine/Plugins/Experimental/MeshModelingToolset/Source/MeshModelingTools/Public/EditNormalsTool.h
new file mode 100644
--- /dev/null
+++ b/Engine/Plugins/Experimental/MeshModelingToolset/Source/MeshModelingTools/Public/EditNormalsTool.h
@@ -0,0 +1,181 @@
+// Copyright 1998-2019 Epic Games, Inc. All Rights Reserved.
+
+#pragma once
+
+#include "CoreMinimal.h"
+#include "UObject/NoExportTypes.h"
+#include "MultiSelectionTool.h"
+#include "InteractiveToolBuilder.h"
+#include "MeshOpPreviewHelpers.h"
+#include "CleaningOps/EditNormalsOp.h"
+#include "DynamicMesh3.h"
+#include "BaseTools/SingleClickTool.h"
+#include "EditNormalsTool.generated.h"
+
+
+// predeclarations
+struct FMeshDescription;
+class USimpleDynamicMeshComponent;
+
+
+
+
+
+/**
+ *
+ */
+UCLASS()
+class MESHMODELINGTOOLS_API UEditNormalsToolBuilder : public UInteractiveToolBuilder
+{
+	GENERATED_BODY()
+
+public:
+	IToolsContextAssetAPI* AssetAPI = nullptr;
+
+	virtual bool CanBuildTool(const FToolBuilderState& SceneState) const override;
+	virtual UInteractiveTool* BuildTool(const FToolBuilderState& SceneState) const override;
+};
+
+
+
+
+
+/**
+ * Standard properties
+ */
+UCLASS()
+class MESHMODELINGTOOLS_API UEditNormalsToolProperties : public UInteractiveToolPropertySet
+{
+	GENERATED_BODY()
+
+public:
+	UEditNormalsToolProperties();
+
+	bool WillTopologyChange()
+	{
+		return bFixInconsistentNormals || bInvertNormals || bRecomputeNormalTopologyAndEdgeSharpness;
+	}
+
+	/** Recompute all mesh normals */
+	UPROPERTY(EditAnywhere, Category = NormalsCalculation, meta = (EditCondition = "!bRecomputeNormalTopologyAndEdgeSharpness"))
+	bool bRecomputeNormals;
+
+	/** Choose the method for computing vertex normals */
+	UPROPERTY(EditAnywhere, Category = NormalsCalculation)
+	ENormalCalculationMethod NormalCalculationMethod;
+
+	/** For meshes with inconsistent triangle orientations/normals, flip as needed to make the normals consistent */
+	UPROPERTY(EditAnywhere, Category = NormalsCalculation)
+	bool bFixInconsistentNormals;
+
+	/** Invert (flip) all mesh normals and associated triangle orientations */
+	UPROPERTY(EditAnywhere, Category = NormalsCalculation)
+	bool bInvertNormals;
+
+	/** Choose where to create sharp edges by having multiple normals on the same vertices, using the change in normals of adjacent faces.  If enabled, normals will always be recomputed. */
+	UPROPERTY(EditAnywhere, Category = NormalsTopology)
+	bool bRecomputeNormalTopologyAndEdgeSharpness;
+
+	/** Threshold on angle of change in face normals across an edge, above which we create a sharp edge if bSplitNormals is true */
+	UPROPERTY(EditAnywhere, Category = NormalsTopology, meta = (UIMin = "0.0", UIMax = "180.0", ClampMin = "0.0", ClampMax = "180.0", EditCondition = "bRecomputeNormalTopologyAndEdgeSharpness"))
+	float SharpEdgeAngleThreshold;
+
+	/** Give separate normals to each triangle at 'sharp' vertices where the coincident triangle normals change (across multiple triangles) by more than the angle threshold (e.g. a cone tip) */
+	UPROPERTY(EditAnywhere, Category = NormalsTopology, meta = (EditCondition = "bRecomputeNormalTopologyAndEdgeSharpness"))
+	bool bAllowSharpVertices;
+};
+
+
+
+
+/**
+ * Advanced properties
+ */
+UCLASS()
+class MESHMODELINGTOOLS_API UEditNormalsAdvancedProperties : public UInteractiveToolPropertySet
+{
+	GENERATED_BODY()
+
+public:
+	UEditNormalsAdvancedProperties();
+};
+
+
+/**
+ * Factory with enough info to spawn the background-thread Operator to do a chunk of work for the tool
+ *  stores a pointer to the tool and enough info to know which specific operator it should spawn
+ */
+UCLASS()
+class MESHMODELINGTOOLS_API UEditNormalsOperatorFactory : public UObject, public IDynamicMeshOperatorFactory
+{
+	GENERATED_BODY()
+
+public:
+	// IDynamicMeshOperatorFactory API
+	virtual TSharedPtr<FDynamicMeshOperator> MakeNewOperator() override;
+
+	UPROPERTY()
+	UEditNormalsTool *Tool;
+
+	int ComponentIndex;
+
+};
+
+/**
+ * Simple Mesh Normal Updating Tool
+ */
+UCLASS()
+class MESHMODELINGTOOLS_API UEditNormalsTool : public UMultiSelectionTool
+{
+	GENERATED_BODY()
+
+public:
+
+	friend UEditNormalsOperatorFactory;
+
+	UEditNormalsTool();
+
+	virtual void Setup() override;
+	virtual void Shutdown(EToolShutdownType ShutdownType) override;
+
+	virtual void SetWorld(UWorld* World);
+	virtual void SetAssetAPI(IToolsContextAssetAPI* AssetAPI);
+
+	virtual void Tick(float DeltaTime) override;
+	virtual void Render(IToolsContextRenderAPI* RenderAPI) override;
+
+	virtual bool HasCancel() const override { return true; }
+	virtual bool HasAccept() const override;
+	virtual bool CanAccept() const override;
+
+#if WITH_EDITOR
+	virtual void PostEditChangeProperty(FPropertyChangedEvent &PropertyChangedEvent) override;
+#endif
+
+	virtual void OnPropertyModified(UObject* PropertySet, UProperty* Property) override;
+
+protected:
+
+	UPROPERTY()
+	UEditNormalsToolProperties* BasicProperties;
+
+	UPROPERTY()
+	UEditNormalsAdvancedProperties* AdvancedProperties;
+
+
+	UPROPERTY()
+	TArray<UMeshOpPreviewWithBackgroundCompute*> Previews;
+
+
+protected:
+	TArray<TSharedPtr<FDynamicMesh3>> OriginalDynamicMeshes;
+
+	UWorld* TargetWorld;
+	IToolsContextAssetAPI* AssetAPI;
+
+	FViewCameraState CameraState;
+
+	void UpdateNumPreviews();
+
+	void GenerateAsset(const TArray<TUniquePtr<FDynamicMeshOpResult>>& Results);
+};
diff --git a/Engine/Plugins/Experimental/MeshModelingToolset/Source/MeshModelingTools/Public/FindPolygonsAlgorithm.h b/Engine/Plugins/Experimental/MeshModelingToolset/Source/MeshModelingTools/Public/FindPolygonsAlgorithm.h
new file mode 100644
--- /dev/null
+++ b/Engine/Plugins/Experimental/MeshModelingToolset/Source/MeshModelingTools/Public/FindPolygonsAlgorithm.h
@@ -0,0 +1,26 @@
+// Copyright 1998-2019 Epic Games, Inc. All Rights Reserved.
+
+#pragma once
+
+#include "CoreMinimal.h"
+#include "DynamicMesh3.h"
+
+class FFindPolygonsAlgorithm
+{
+public:
+
+	FFindPolygonsAlgorithm() {}
+	FFindPolygonsAlgorithm(FDynamicMesh3* MeshIn);
+
+	FDynamicMesh3* Mesh = nullptr;
+	TArray<int> PolygonGroupIDs;
+	TArray<TArray<int>> FoundPolygons;
+	TArray<int> PolygonTags;
+	TArray<FVector> PolygonNormals;
+
+	TArray<int> PolygonEdges;
+
+	bool FindPolygons(double DotTolerance = 0.0001);
+
+	bool FindPolygonEdges();
+};
\ No newline at end of file
diff --git a/Engine/Plugins/Experimental/MeshModelingToolset/Source/MeshModelingTools/Public/MeshInspectorTool.h b/Engine/Plugins/Experimental/MeshModelingToolset/Source/MeshModelingTools/Public/MeshInspectorTool.h
new file mode 100644
--- /dev/null
+++ b/Engine/Plugins/Experimental/MeshModelingToolset/Source/MeshModelingTools/Public/MeshInspectorTool.h
@@ -0,0 +1,140 @@
+// Copyright 1998-2019 Epic Games, Inc. All Rights Reserved.
+
+#pragma once
+
+#include "CoreMinimal.h"
+#include "UObject/NoExportTypes.h"
+#include "SingleSelectionTool.h"
+#include "InteractiveToolBuilder.h"
+#include "DynamicMesh3.h"
+#include "DynamicMeshAABBTree3.h"
+#include "Properties/MeshMaterialProperties.h"
+#include "MeshInspectorTool.generated.h"
+
+
+// predeclarations
+struct FMeshDescription;
+class USimpleDynamicMeshComponent;
+class UMaterialInstanceDynamic;
+
+/**
+ *
+ */
+UCLASS()
+class MESHMODELINGTOOLS_API UMeshInspectorToolBuilder : public UInteractiveToolBuilder
+{
+	GENERATED_BODY()
+
+public:
+	virtual bool CanBuildTool(const FToolBuilderState& SceneState) const override;
+	virtual UInteractiveTool* BuildTool(const FToolBuilderState& SceneState) const override;
+};
+
+
+
+UCLASS()
+class MESHMODELINGTOOLS_API UMeshInspectorProperties : public UInteractiveToolPropertySet
+{
+	GENERATED_BODY()
+
+public:
+	/** Toggle visibility of all mesh edges */
+	UPROPERTY(EditAnywhere, Category = Options)
+	bool bWireframe = true;
+
+	/** Toggle visibility of open boundary edges */
+	UPROPERTY(EditAnywhere, Category = Options)
+	bool bBoundaryEdges = true;
+
+	/** Toggle visibility of polygon borders */
+	UPROPERTY(EditAnywhere, Category = Options)
+	bool bPolygonBorders = false;
+
+	/** Toggle visibility of UV seam edges */
+	UPROPERTY(EditAnywhere, Category = Options)
+	bool bUVSeams = false;
+
+	/** Toggle visibility of Normal seam edges */
+	UPROPERTY(EditAnywhere, Category = Options)
+	bool bNormalSeams = false;
+
+	/** Toggle visibility of normal vectors */
+	UPROPERTY(EditAnywhere, Category = Options)
+	bool bNormalVectors = false;
+
+	/** Toggle visibility of tangent vectors */
+	UPROPERTY(EditAnywhere, Category = Options)
+	bool bTangentVectors = false;
+
+	/** Length of line segments representing normal vectors */
+	UPROPERTY(EditAnywhere, Category = Options, meta = (EditCondition = "bNormalVectors"))
+	float NormalLength = 5.0f;
+
+	/** Length of line segments representing tangent vectors */
+	UPROPERTY(EditAnywhere, Category = Options, meta = (EditCondition = "bTangentVectors"))
+	float TangentLength = 5.0f;
+
+
+	//
+	// save/restore support
+	//
+	virtual void SaveProperties(UInteractiveTool* SaveFromTool) override;
+	virtual void RestoreProperties(UInteractiveTool* RestoreToTool) override;
+};
+
+
+
+
+/**
+ * Mesh Inspector Tool for visualizing mesh information
+ */
+UCLASS()
+class MESHMODELINGTOOLS_API UMeshInspectorTool : public USingleSelectionTool
+{
+	GENERATED_BODY()
+
+public:
+	UMeshInspectorTool();
+
+	virtual void RegisterActions(FInteractiveToolActionSet& ActionSet) override;
+
+	virtual void Setup() override;
+	virtual void Shutdown(EToolShutdownType ShutdownType) override;
+
+	virtual void Render(IToolsContextRenderAPI* RenderAPI) override;
+
+	virtual bool HasCancel() const override { return false; }
+	virtual bool HasAccept() const override;
+	virtual bool CanAccept() const override;
+
+	virtual void OnPropertyModified(UObject* PropertySet, UProperty* Property) override;
+
+public:
+
+	virtual void IncreaseLineWidthAction();
+	virtual void DecreaseLineWidthAction();
+
+protected:
+
+	UPROPERTY()
+	UMeshInspectorProperties* Settings = nullptr;
+
+	UPROPERTY()
+	UExistingMeshMaterialProperties* MaterialSettings = nullptr;
+
+
+	float LineWidthMultiplier = 1.0f;
+
+protected:
+	UPROPERTY()
+	USimpleDynamicMeshComponent* DynamicMeshComponent;
+
+	UPROPERTY()
+	UMaterialInterface* DefaultMaterial = nullptr;
+
+	TArray<int> BoundaryEdges;
+	TArray<int> UVSeamEdges;
+	TArray<int> NormalSeamEdges;
+	TArray<int> GroupBoundaryEdges;
+	void Precompute();
+};
diff --git a/Engine/Plugins/Experimental/MeshModelingToolset/Source/MeshModelingTools/Public/MeshSelectionTool.h b/Engine/Plugins/Experimental/MeshModelingToolset/Source/MeshModelingTools/Public/MeshSelectionTool.h
new file mode 100644
--- /dev/null
+++ b/Engine/Plugins/Experimental/MeshModelingToolset/Source/MeshModelingTools/Public/MeshSelectionTool.h
@@ -0,0 +1,300 @@
+// Copyright 1998-2019 Epic Games, Inc. All Rights Reserved.
+
+#pragma once
+
+#include "CoreMinimal.h"
+#include "DynamicMeshBrushTool.h"
+#include "SelectionSet.h"
+#include "Changes/MeshSelectionChange.h"
+#include "Changes/ValueWatcher.h"
+#include "DynamicMeshOctree3.h"
+#include "MeshSelectionTool.generated.h"
+
+/**
+ *
+ */
+UCLASS()
+class MESHMODELINGTOOLS_API UMeshSelectionToolBuilder : public UMeshSurfacePointToolBuilder
+{
+	GENERATED_BODY()
+
+public:
+	IToolsContextAssetAPI* AssetAPI = nullptr;
+
+	virtual UMeshSurfacePointTool* CreateNewTool(const FToolBuilderState& SceneState) const override;
+};
+
+
+
+
+UENUM()
+enum class EMeshSelectionToolActions
+{
+	NoAction,
+
+	ClearSelection,
+	InvertSelection,
+	GrowSelection,
+	ShrinkSelection,
+	ExpandToConnected,
+
+	DeleteSelected,
+	SeparateSelected
+};
+
+
+
+UCLASS()
+class MESHMODELINGTOOLS_API UMeshSelectionToolActionPropertySet : public UInteractiveToolPropertySet
+{
+	GENERATED_BODY()
+
+public:
+	TWeakObjectPtr<UMeshSelectionTool> ParentTool;
+
+	void Initialize(UMeshSelectionTool* ParentToolIn) { ParentTool = ParentToolIn; }
+
+	void PostAction(EMeshSelectionToolActions Action);
+};
+
+
+
+
+UCLASS()
+class MESHMODELINGTOOLS_API UMeshSelectionEditActions : public UMeshSelectionToolActionPropertySet
+{
+	GENERATED_BODY()
+
+public:
+	UFUNCTION(CallInEditor, Category = Selection, meta = (DisplayName = "Clear", DisplayPriority = 1))
+	void Clear()
+	{
+		PostAction(EMeshSelectionToolActions::ClearSelection);
+	}
+
+	UFUNCTION(CallInEditor, Category = Selection, meta = (DisplayName = "Invert", DisplayPriority = 2))
+	void Invert()
+	{
+		PostAction(EMeshSelectionToolActions::InvertSelection);
+	}
+
+
+	UFUNCTION(CallInEditor, Category = Selection, meta = (DisplayName = "Grow", DisplayPriority = 3))
+	void Grow()
+	{
+		PostAction(EMeshSelectionToolActions::GrowSelection);
+	}
+
+	UFUNCTION(CallInEditor, Category = Selection, meta = (DisplayName = "Shrink", DisplayPriority = 4))
+	void Shrink()
+	{
+		PostAction(EMeshSelectionToolActions::ShrinkSelection);
+	}
+
+	UFUNCTION(CallInEditor, Category = Selection, meta = (DisplayName = "ExpandToConnected", DisplayPriority = 5))
+	void ExpandToConnected()
+	{
+		PostAction(EMeshSelectionToolActions::ExpandToConnected);
+	}
+
+};
+
+
+
+
+UCLASS()
+class MESHMODELINGTOOLS_API UMeshSelectionMeshEditActions : public UMeshSelectionToolActionPropertySet
+{
+	GENERATED_BODY()
+
+public:
+	UFUNCTION(CallInEditor, Category = MeshEdits, meta = (DisplayName = "Delete"))
+	void DeleteTriangles()
+	{
+		PostAction(EMeshSelectionToolActions::DeleteSelected);
+	}
+
+	UFUNCTION(CallInEditor, Category = MeshEdits, meta = (DisplayName = "Separate"))
+	void SeparateTriangles() 
+	{
+		PostAction(EMeshSelectionToolActions::SeparateSelected);
+	}
+};
+
+
+
+
+
+
+UENUM()
+enum class EMeshSelectionToolPrimaryMode
+{
+	/** Select all triangles inside the brush */
+	Brush,
+
+	/** Select all triangles inside brush with normal within angular tolerance of hit triangle */
+	AngleFiltered,
+
+	/** Select all triangles inside brush that are visible from current view */
+	Visible,
+
+	/** Select all triangles connected to any triangle inside the brush */
+	AllConnected,
+
+	/** Select all triangles in groups connected to any triangle inside the brush */
+	AllInGroup,
+
+	/** Select all triangles with normal within angular tolerance of hit triangle */
+	AllWithinAngle
+};
+
+
+
+UCLASS()
+class MESHMODELINGTOOLS_API UMeshSelectionToolProperties : public UInteractiveToolPropertySet
+{
+	GENERATED_BODY()
+public:
+
+	/** The Selection Mode defines the behavior of the selection brush */
+	UPROPERTY(EditAnywhere, Category = Selection)
+	EMeshSelectionToolPrimaryMode SelectionMode = EMeshSelectionToolPrimaryMode::Brush;
+
+	/** Angle in Degrees used for Angle-based Selection Modes */
+	UPROPERTY(EditAnywhere, Category = Selection, meta = (UIMin = "0.0", UIMax = "90.0") )
+	float AngleTolerance = 1.0;
+
+	/** A Volumetric brush selects anything within a 3D sphere, rather than only elements connected to the Brush Position  */
+	UPROPERTY(EditAnywhere, Category = Selection)
+	bool bVolumetricBrush = false;
+
+	/** Allow the brush to hit back-facing parts of the surface  */
+	UPROPERTY(EditAnywhere, Category = Selection)
+	bool bHitBackFaces = true;
+
+	/** Toggle drawing of wireframe overlay on/off [Alt+W] */
+	UPROPERTY(EditAnywhere, Category = Selection)
+	bool bShowWireframe = true;
+
+	virtual void SaveProperties(UInteractiveTool* SaveFromTool) override;
+	virtual void RestoreProperties(UInteractiveTool* RestoreToTool) override;
+};
+
+
+
+/**
+ *
+ */
+UCLASS()
+class MESHMODELINGTOOLS_API UMeshSelectionTool : public UDynamicMeshBrushTool
+{
+	GENERATED_BODY()
+
+public:
+	UMeshSelectionTool();
+
+	virtual void SetWorld(UWorld* World);
+	virtual void SetAssetAPI(IToolsContextAssetAPI* AssetAPI);
+
+	virtual void RegisterActions(FInteractiveToolActionSet& ActionSet) override;
+
+	virtual void Setup() override;
+
+	virtual void Tick(float DeltaTime) override;
+	virtual void Render(IToolsContextRenderAPI* RenderAPI) override;
+
+	virtual bool HasCancel() const override { return true; }
+	virtual bool HasAccept() const override { return true; }
+	virtual bool CanAccept() const override { return bHaveModifiedMesh; }
+
+	// UBaseBrushTool overrides
+	virtual bool HitTest(const FRay& Ray, FHitResult& OutHit) override;
+	virtual void OnBeginDrag(const FRay& Ray) override;
+	virtual void OnUpdateDrag(const FRay& Ray) override;
+	virtual void OnEndDrag(const FRay& Ray) override;
+	virtual bool OnUpdateHover(const FInputDeviceRay& DevicePos) override;
+
+public:
+
+	virtual void RequestAction(EMeshSelectionToolActions ActionType);
+
+	UPROPERTY()
+	UMeshSelectionToolProperties* SelectionProps;
+
+	UPROPERTY()
+	UMeshSelectionEditActions* SelectionActions;
+
+	UPROPERTY()
+	UMeshSelectionMeshEditActions* EditActions;
+
+protected:
+
+	virtual void ApplyStamp(const FBrushStampData& Stamp);
+
+	virtual void OnShutdown(EToolShutdownType ShutdownType) override;
+
+protected:
+	UPROPERTY()
+	UMeshSelectionSet* Selection;
+
+	UWorld* TargetWorld;
+	IToolsContextAssetAPI* AssetAPI;
+
+	// note: ideally this octree would be part of PreviewMesh!
+	TUniquePtr<FDynamicMeshOctree3> Octree;
+	bool bOctreeValid = false;
+	TUniquePtr<FDynamicMeshOctree3>& GetOctree();
+
+	EMeshSelectionElementType SelectionType = EMeshSelectionElementType::Face;
+
+	TValueWatcher<bool> ShowWireframeWatcher;
+
+
+	bool bInRemoveStroke = false;
+	FBrushStampData StartStamp;
+	FBrushStampData LastStamp;
+	bool bStampPending;
+
+	void UpdateFaceSelection(const FBrushStampData& Stamp, const TArray<int>& BrushROI);
+
+
+	// temp
+	TArray<int> IndexBuf;
+	TArray<int32> TemporaryBuffer;
+	TSet<int32> TemporarySet;
+	void CalculateVertexROI(const FBrushStampData& Stamp, TArray<int>& VertexROI);
+	void CalculateTriangleROI(const FBrushStampData& Stamp, TArray<int>& TriangleROI);
+	TArray<int> PreviewBrushROI;
+	TBitArray<> SelectedVertices;
+	TBitArray<> SelectedTriangles;
+	void OnExternalSelectionChange();
+
+	void OnSelectionUpdated();
+	void UpdateVisualization();
+
+	// selection change
+	FMeshSelectionChangeBuilder* ActiveSelectionChange = nullptr;
+	void BeginChange(bool bAdding);
+	TUniquePtr<FMeshSelectionChange> EndChange();
+	void CancelChange();
+
+
+	// actions
+
+	bool bHavePendingAction = false;
+	EMeshSelectionToolActions PendingAction;
+	virtual void ApplyAction(EMeshSelectionToolActions ActionType);
+
+	void ClearSelection();
+	void InvertSelection();
+	void GrowShrinkSelection(bool bGrow);
+	void ExpandToConnected();
+
+	void DeleteSelectedTriangles();
+	void SeparateSelectedTriangles();
+
+	bool bHaveModifiedMesh = false;
+};
+
+
+
diff --git a/Engine/Plugins/Experimental/MeshModelingToolset/Source/MeshModelingTools/Public/MeshSpaceDeformerTool.h b/Engine/Plugins/Experimental/MeshModelingToolset/Source/MeshModelingTools/Public/MeshSpaceDeformerTool.h
new file mode 100644
--- /dev/null
+++ b/Engine/Plugins/Experimental/MeshModelingToolset/Source/MeshModelingTools/Public/MeshSpaceDeformerTool.h
@@ -0,0 +1,464 @@
+// Copyright 1998-2019 Epic Games, Inc. All Rights Reserved.
+
+#pragma once
+
+#include "CoreMinimal.h"
+#include "UObject/NoExportTypes.h"
+#include "BaseTools/MeshSurfacePointTool.h"
+#include "DynamicMeshAABBTree3.h"
+#include "Drawing/ToolDataVisualizer.h"
+#include "Transforms/QuickAxisTranslater.h"
+#include "Transforms/QuickAxisRotator.h"
+#include "Changes/MeshVertexChange.h"
+#include "GroupTopology.h"
+#include "Spatial/GeometrySet3.h"
+#include "Selection/GroupTopologySelector.h"
+#include "Operations/GroupTopologyDeformer.h"
+#include "Curves/CurveFloat.h"
+#include "SpaceDeformerOps/MeshSpaceDeformerOp.h"
+#include "SpaceDeformerOps/BendMeshOp.h"
+#include "SpaceDeformerOps/TwistMeshOp.h"
+#include "SpaceDeformerOps/FlareMeshOp.h"
+#include "SimpleDynamicMeshComponent.h"
+
+#include "ModelingOperators/Public/ModelingTaskTypes.h"
+
+
+#include "MeshSpaceDeformerTool.generated.h"
+
+class FMeshVertexChangeBuilder;
+
+/**
+ * ToolBuilder
+ */
+UCLASS()
+class MESHMODELINGTOOLS_API UMeshSpaceDeformerToolBuilder : public UMeshSurfacePointToolBuilder
+{
+	GENERATED_BODY()
+
+public:
+	UMeshSpaceDeformerToolBuilder()
+	{
+	}
+
+	virtual UMeshSurfacePointTool* CreateNewTool(const FToolBuilderState& SceneState) const override;
+};
+
+/** ENonlinearOperation determines which type of nonlinear deformation will be applied*/
+UENUM()
+enum class  ENonlinearOperationType : int8
+{
+	Bend		UMETA(DisplayName = "Bend"),
+	Flare		UMETA(DisplayName = "Flare"),
+	Twist		UMETA(DisplayName = "Twist"),
+	//Sinusoid	UMETA(DisplayName = "Sinusoid"),
+	//Wave		UMETA(DisplayName = "Wave"),
+	//Squish	UMETA(DisplayName = "Squish")
+};
+
+/**
+* \brief IHandleWidgetSelectable is a pure virtual base class for elements of the HandleWidget which may be selectable (i.e. the end points of line segments or the axes themselves)
+*/
+class IHandleWidgetSelectable
+{
+public:
+	IHandleWidgetSelectable() = delete;
+	IHandleWidgetSelectable(const FLinearColor& DefaultColor, const FLinearColor& HoveringColor, const FLinearColor& SelectedColor) : StateColors{ DefaultColor,HoveringColor,SelectedColor } 
+	{
+	};
+	virtual ~IHandleWidgetSelectable() 
+	{
+	};
+
+	/**
+	* \brief Sets the renderer to be used by the tool
+	*/
+	virtual inline void SetRenderer(FToolDataVisualizer* Target) 
+	{ 
+		Renderer = Target;
+	};
+
+	/**
+	* \brief Sets the world position of the tool
+	*/ 
+	virtual inline void SetPosition(const FVector3d& Pos) 
+	{
+		WorldCenter = Pos;
+	};
+
+	/**
+	* \brief Gets the color depnding on the state of the selectable element
+	*/
+	virtual inline const FLinearColor& GetColor() const 
+	{
+		return StateColors[static_cast<int8>(State)];
+	}
+
+	/**
+	* \brief Sets the visibility of the specified widget element
+	*/
+	virtual inline void SetVisible(bool Enable) 
+	{
+		Visible = Enable;
+	};
+
+	FToolDataVisualizer* Renderer{nullptr};
+
+	enum EDisplayState : int8
+	{
+		Default,  // Unselected, the default display color
+		Hovering, // The element is being hovered over
+		Selected  // The element is selected
+	};
+
+	EDisplayState State = EDisplayState::Default;
+	FLinearColor StateColors[3];
+
+	FVector3d WorldCenter;
+	bool Visible{false};
+	virtual void Render() = 0;
+};
+
+/**
+* \brief This object is located on the ends of the line segments which extend from the center of the widget. 
+*/
+class FHandleWidgetVertex : public IHandleWidgetSelectable
+{
+public:
+	FHandleWidgetVertex() : IHandleWidgetSelectable
+	{
+		FLinearColor{1.f,1.f,1.f,1.f},	// White
+		FLinearColor{0.f,1.f,1.f,1.f},	// Yellow
+		FLinearColor{1.f,0.0f,1.f,1.f}	// Pink
+	} {};
+	virtual void Render() override
+	{
+		if (Visible)
+		{
+			float Width = State != Default ? 4.0 : 2.0;
+			Renderer->DrawViewFacingCircle(WorldCenter,4.0, 16, GetColor(), Width,false);
+		}
+	};
+};
+
+
+/**
+* \brief Represents an axis of the handle widget, each axis is selectable and has two vertices located at the end points
+*/
+class FHandleWidgetAxis : public IHandleWidgetSelectable
+{
+public:
+	FHandleWidgetAxis() : IHandleWidgetSelectable
+	{
+		FLinearColor{0.2f,0.2f,1.f,1.f}, // Light-blue
+		FLinearColor{0.f,1.f,1.f,1.f},	 // Yellow
+		FLinearColor{1.f,0.f,1.f,1.f}	 // Pink
+	} {};
+
+	double LowerExtentLength;
+	double UpperExtentLength;
+	
+	//Identity
+	FMatrix3d Basis{ 
+		FVector3d{1.0,0.0,0.0}, 
+		FVector3d{0.0,1.0,0.0}, 
+		FVector3d{0.0,0.0,1.0}, 
+		false };
+
+	//The direction vector representing the local space direction of the line segment this handle axis represents
+	FVector3d Axis;
+
+	void SetAxis(const FVector3d& TargetAxis) {Axis = TargetAxis;};
+
+	void SetBasis(const FVector3d& V0, const FVector3d& V1, const FVector3d& V2 ) { Basis = FMatrix3d{ V0, V1, V2, false }; };
+
+	FHandleWidgetVertex Vertices[2];
+	double SelectedWidth = 2.5;
+	double DefaultWidth = 1.5;
+
+	/**
+	* \brief Sets the center position from the center of the handle. This, in turn, sets the position of its associated vertices
+	*/
+	virtual inline void SetPosition(const FVector3d& Center) override
+	{
+		WorldCenter = Center;
+
+		const FVector3d Positions[2]
+		{
+			WorldCenter + Axis * LowerExtentLength,
+			WorldCenter + Axis * UpperExtentLength
+		};
+
+		for (uint8 i = 0; i < 2; ++i)
+		{
+			Vertices[i].SetPosition(Positions[i]);
+		}
+
+	};
+
+	virtual void Render() override
+	{
+		double Width = State == EDisplayState::Default ? DefaultWidth : SelectedWidth;
+
+		for (uint8 i = 0; i < 2; ++i)
+		{
+			Vertices[i].Render();
+		}
+
+		if (Visible)
+		{
+			Renderer->DrawLine(Vertices[0].WorldCenter, Vertices[1].WorldCenter, GetColor(), Width, false);
+		}
+	};
+
+};
+
+/**
+* \brief This purpose of this "widget" is to visualize the world space center of the transformation being applied by a Mesh Space Deformer operator. 
+  \par This is meant to be an analogous visualizer to the "handle" widgets used by popular modeling tools i.e. maya
+*/
+class FHandleWidget
+{
+public:
+	FHandleWidget(FToolDataVisualizer* Renderer) 
+	{
+		for (unsigned i=0;i<3;++i)
+		{
+			Axes[i].SetRenderer(Renderer);
+
+
+			for (unsigned j=0;j<2;++j)
+			{
+				Axes[i].Vertices[j].SetRenderer(Renderer);
+			}
+			const FVector3d& A = Axes[i].Vertices[0].WorldCenter;
+			const FVector3d& B = Axes[i].Vertices[1].WorldCenter;
+			GeometrySet.AddPoint(i * 2 , A);
+			GeometrySet.AddPoint(i * 2 + 1, B);
+			GeometrySet.AddCurve(i, FPolyline3d{ TArray<FVector3d>{A,B} });
+
+		}
+
+
+	};
+	virtual ~FHandleWidget() {};
+
+	/**
+	* \brief Returns a pointer to the Handle Widget Axis related object as specified by the Selection
+	*/
+	FHandleWidgetAxis* GetSelectedAxis() { return &Axes[SelectedAxisID];};
+
+	/**
+	* \brief Returns a pointer to the Handle Widget Axis related object as specified by the Selection
+	*/
+	FHandleWidgetVertex* GetSelectedVertex() { return SelectedVertexID == -1 ? nullptr : &Axes[SelectedAxisID].Vertices[SelectedVertexID]; };
+	int SelectedAxisID;
+	int SelectedVertexID;
+
+	FHandleWidgetAxis Axes[3];
+	FVector3d WorldCenter;
+
+	/** Given a ray, tests the geometry collection and updates the corresponding selectable handle elements */
+	void UpdateHover(const FInputDeviceRay& DevicePos, const FTransform& WorldTransform, TFunction<bool(const FVector3d&, const FVector3d&)> ToleranceFunction);
+	void SetPosition(const FVector3d& Center)
+	{
+		WorldCenter = Center;
+
+
+		for (unsigned i = 0; i < 3; ++i)
+		{
+			Axes[i].SetPosition(Center);
+			const FVector3d& A = Axes[i].Vertices[0].WorldCenter;
+			const FVector3d& B = Axes[i].Vertices[1].WorldCenter;
+			GeometrySet.UpdatePoint(i * 2 , A);
+			GeometrySet.UpdatePoint(i * 2 + 1, B);
+			GeometrySet.UpdateCurve(i, FPolyline3d{ TArray<FVector3d>{A,B} });
+		}
+	};
+
+	void SetBasis(const FVector3d& Axis0, const FVector3d& Axis1, const FVector3d& Axis2 )
+	{
+		Axes[0].SetAxis(Axis0);
+		Axes[1].SetAxis(Axis1);
+		Axes[2].SetAxis(Axis2);
+	};
+	void UpdateDisplayData(int SelectedAxis, double LowerBoundsInterval, double UpperBoundsInterval, const FVector3d& HalfExtents)
+	{
+
+		SelectedAxisID = SelectedAxis;
+		for (unsigned i = 0; i < 3; ++i)
+		{
+			if ( i == SelectedAxis )
+			{
+				Axes[i].State = IHandleWidgetSelectable::EDisplayState::Selected;
+				Axes[i].SetVisible(true);
+				Axes[i].Vertices[0].SetVisible(true);
+				Axes[i].Vertices[1].SetVisible(true);
+				Axes[i].LowerExtentLength = HalfExtents[i] * LowerBoundsInterval;
+				Axes[i].UpperExtentLength = HalfExtents[i] * UpperBoundsInterval;
+			}else
+			{
+				Axes[i].SetVisible(true);
+				Axes[i].Vertices[0].SetVisible(false);
+				Axes[i].Vertices[1].SetVisible(false);
+				Axes[i].LowerExtentLength = -HalfExtents[i] ;
+				Axes[i].UpperExtentLength = HalfExtents[i] ;
+			}
+
+		}
+	}
+
+	/** Called when updating a click-and-drag transaction, updates the positions of the selected vertices based on the new mouse input position*/
+	void DragUpdateBounds(const FRay3d& WorldRay);
+
+	void Render() 
+	{
+		for (unsigned i=0;i<3;++i)
+		{
+			Axes[i].Render();
+		}
+	};
+
+	/** The set of geometry containing the vertices which are draggable and the axes objects which are selectable */
+	FGeometrySet3 GeometrySet;
+
+};
+
+
+/**
+ * Applies non-linear deformations to a mesh 
+ */
+UCLASS()
+class MESHMODELINGTOOLS_API UMeshSpaceDeformerTool : public UMeshSurfacePointTool
+{
+	GENERATED_BODY()
+
+public:
+	UMeshSpaceDeformerTool();
+
+	virtual void RegisterActions(FInteractiveToolActionSet& ActionSet) override;
+
+	FMatrix3d CalculateBestAxis(FDynamicMesh3 * Mesh, const TArray<FVector3d>& PositionBuffer, FTransform WorldTransform, FVector3d * ExtentsOut, FVector3d * OriginOut );
+
+	virtual void Setup() override;
+	virtual void Shutdown(EToolShutdownType ShutdownType) override;
+
+	virtual void Tick(float DeltaTime) override;
+	virtual void Render(IToolsContextRenderAPI* RenderAPI) override;
+
+	virtual bool HasCancel() const override { return true; }
+	virtual bool HasAccept() const override { return true; }
+	virtual bool CanAccept() const override { return true; }
+
+	virtual bool HitTest(const FRay& Ray, FHitResult& OutHit) override;
+	
+	//
+	// Input behavior
+	//
+	virtual void OnBeginDrag(const FRay& Ray) override;
+	virtual void OnUpdateDrag(const FRay& Ray) override;
+	virtual void OnEndDrag(const FRay& Ray) override;
+
+	virtual bool OnUpdateHover(const FInputDeviceRay& DevicePos) override;
+
+#if WITH_EDITOR
+	virtual void PostEditChangeProperty(FPropertyChangedEvent& PropertyChangedEvent) override;
+#endif
+
+	void UpdateOp();
+public:
+
+	float VisualAngleSnapThreshold = 1.5f;
+
+protected:
+
+	//Options
+	UPROPERTY(EditAnywhere, Category = Options, meta = (DisplayName = "Operation Type"))
+	ENonlinearOperationType SelectedOperationType;
+
+	/** The upper bounds interval corresponds to the region of space which the selected operator will affect. A setting of -1.0 should envelope all points in the "lower" half of the mesh given the axis has been auto-detected. The corresponding upper value of 1 will cover the entire mesh. */
+	UPROPERTY(EditAnywhere, Category = Options, meta = (DisplayName = "Lower Bound", UIMin = "-1.5", UIMax = "0.0", ClampMin = "-1.5", ClampMax = "0.0"))
+	double LowerBoundsInterval{ -1.0 };
+
+	/** The upper bounds interval corresponds to the region of space which the selected operator will affect. A setting of 1.0 should envelope all points in the "upper" half of the mesh given the axis has been auto-detected. The corresponding lower value of -1 will cover the entire mesh. */
+	UPROPERTY(EditAnywhere, Category = Options, meta = (DisplayName = "Upper Bound", UIMin = "0.0", UIMax = "1.5", ClampMin = "0.0", ClampMax = "1.5"))
+	double UpperBoundsInterval{ 1.0 };
+
+	/** As each operator has a range of values (i.e. curvature, angle of twist, scale), this represents the percentage passed to the operator as a parameter. In the future, for more control, this should be separated into individual settings for each operator for more precise control */
+	UPROPERTY(EditAnywhere, Category = Options, meta = (DisplayName = "Modifier Percent", UIMin = "-100.0", UIMax = "100.0", ClampMin = "-100.0", ClampMax = "100.0"))
+	double ModifierPercent{ 0.0 };
+
+	/** Automatically detects the principal axes of the point cloud of the mesh, as well as determining the lengths of the maximal elements in those directions, setting the intervals accordingly. */
+	UFUNCTION(CallInEditor, Category = Options, meta = (DisplayName = "Auto-detect Axis", ToolTip = "Attempts to automatically align the handle with principal axis."))
+	void AutoDetectAxes();
+
+	/** Swaps the secondary axis upon which the operator will work. I.e. bending */
+	UFUNCTION(CallInEditor, Category = Options, meta = (DisplayName = "Swap secondary axis"))
+	void SwapSecondaryAxis();
+
+	//Calculate the updated object-space position which gets passed to the handle widget.
+	void UpdateObjectSpaceAxisCentroid();
+
+	//Updates the upper and lower intervals from the drag update
+	void UpdateIntervalsFromDrag();
+
+
+	UPROPERTY()
+	USimpleDynamicMeshComponent* DynamicMeshComponent;
+
+	FViewCameraState CameraState;
+
+	//For buffer re-use and avoiding recurrent dynamic allocation
+	TArray<double> PCAVertexPositionBuffer;
+
+	//Array of the original positions of each vertex
+	TArray<FVector3d> OriginalPositions;
+
+	//Index of the axes being used. 
+	uint8 SelectedAxis{ 2 };
+	uint8 SecondaryAxis{ 1 };
+	uint8 ThirdAxis{0};
+
+	//Provided PCA, these are the direction vectors representing the principal axes to best align the mesh
+	FVector3d PrincipalAxesWorldSpace[3];
+	FVector3d PrincipalAxesObjectSpace[3];
+
+	// These represent the distance from the centroid to the farthest elements in the direction (respectively) of the PCA bases. 
+	// Specifically, the X element is the shortest of the PCA axes, the Y element is the middle axis, and the Z element represents the axis with the largest degree of variance, meaning Z is the greatest half extent.
+	FVector3d PrincipalAxesHalfExtentCoeff;
+
+	//Where the handle widget will be rendered
+	FVector3d AxisCentroidWorldSpace;
+	FVector3d AxisCentroidObjectSpace;
+
+	FToolDataVisualizer AxisRenderer;
+	FHandleWidget Handle{&AxisRenderer};
+
+	FBendMeshOp BendOp;
+	FTwistMeshOp TwistOp;
+	FFlareMeshOp FlareOp;
+
+	TArray<FMeshSpaceDeformerOp*> Operators
+	{
+		&BendOp,
+		&FlareOp,
+		&TwistOp
+		//Sinusoid
+		//Wave
+		//Squish
+	};
+
+
+	// realtime visualization
+	void OnDynamicMeshComponentChanged();
+	FDelegateHandle OnDynamicMeshComponentChangedHandle;
+	
+	//
+	// Update	
+	//
+	/** True for the duration of UI click+drag */
+	bool bInDrag;
+	
+private:
+};
+
diff --git a/Engine/Plugins/Experimental/MeshModelingToolset/Source/MeshModelingTools/Public/PositionPlaneGizmo.h b/Engine/Plugins/Experimental/MeshModelingToolset/Source/MeshModelingTools/Public/PositionPlaneGizmo.h
new file mode 100644
--- /dev/null
+++ b/Engine/Plugins/Experimental/MeshModelingToolset/Source/MeshModelingTools/Public/PositionPlaneGizmo.h
@@ -0,0 +1,122 @@
+// Copyright 1998-2019 Epic Games, Inc. All Rights Reserved.
+
+#pragma once
+
+#include "CoreMinimal.h"
+#include "UObject/NoExportTypes.h"
+#include "InteractiveGizmo.h"
+#include "InteractiveGizmoBuilder.h"
+#include "PreviewMesh.h"
+#include "BaseBehaviors/AnyButtonInputBehavior.h"
+#include "Transforms/QuickAxisTranslater.h"
+#include "PositionPlaneGizmo.generated.h"
+
+/**
+ *
+ */
+UCLASS()
+class MESHMODELINGTOOLS_API UPositionPlaneGizmoBuilder : public UInteractiveGizmoBuilder
+{
+	GENERATED_BODY()
+
+public:
+	virtual UInteractiveGizmo* BuildGizmo(const FToolBuilderState& SceneState) const override;
+};
+
+
+
+
+/**
+ * This is a simple gizmo you can use to position a 3D plane in the world, based on QuickAxisTransformer
+ */
+UCLASS()
+class MESHMODELINGTOOLS_API UPositionPlaneGizmo : public UInteractiveGizmo
+{
+	GENERATED_BODY()
+
+public:
+	virtual void SetWorld(UWorld* World);
+
+	virtual void Setup() override;
+	virtual void Shutdown() override;
+
+	virtual void Render(IToolsContextRenderAPI* RenderAPI) override;
+
+	virtual bool HitTest(const FRay& Ray, FHitResult& OutHit);
+
+
+	/**
+	 * This function is called by registered InputBehaviors when the user begins a click-drag-release interaction
+	 */
+	virtual void OnBeginDrag(const FRay& Ray);
+
+	/**
+	 * This function is called by registered InputBehaviorseach frame that the user is in a click-drag-release interaction
+	 */
+	virtual void OnUpdateDrag(const FRay& Ray);
+
+	/**
+	 * This function is called by registered InputBehaviors when the user releases the button driving a click-drag-release interaction
+	 */
+	virtual void OnEndDrag(const FRay& Ray);
+
+
+	TFunction<void(const FFrame3d&)> OnPositionUpdatedFunc = nullptr;
+
+
+	virtual void ExternalUpdatePosition(const FVector& Position, const FQuat& Orientation, bool bPostUpdate);
+
+
+public:
+	UWorld* TargetWorld;
+
+
+	UPROPERTY()
+	UPreviewMesh* CenterBallShape;
+
+
+	UPROPERTY()
+	UMaterialInstance* CenterBallMaterial;
+
+	UPreviewMesh* MakeSphereMesh();
+
+	bool bInTransformDrag;
+	FQuickAxisTranslater QuickTransformer;
+
+	void PostUpdatedPosition();
+
+};
+
+
+
+
+
+
+
+
+/**
+ * UMeshSurfacePointToolMouseBehavior implements mouse press-drag-release interaction behavior for Mouse devices.
+ * You can configure the base UAnyButtonInputBehavior to change the mouse button in use (default = left mouse)
+ */
+UCLASS()
+class MESHMODELINGTOOLS_API UPositionPlaneOnSceneInputBehavior : public UAnyButtonInputBehavior
+{
+	GENERATED_BODY()
+
+public:
+	virtual FInputCapturePriority GetPriority() override { return FInputCapturePriority(FInputCapturePriority::DEFAULT_GIZMO_PRIORITY); }
+
+	virtual void Initialize(UPositionPlaneGizmo* Gizmo);
+
+	virtual FInputCaptureRequest WantsCapture(const FInputDeviceState& input) override;
+	virtual FInputCaptureUpdate BeginCapture(const FInputDeviceState& input, EInputCaptureSide eSide) override;
+	virtual FInputCaptureUpdate UpdateCapture(const FInputDeviceState& input, const FInputCaptureData& data) override;
+	virtual void ForceEndCapture(const FInputCaptureData& data) override;
+
+protected:
+	UPositionPlaneGizmo* Gizmo;
+	FRay LastWorldRay;
+	bool bInDragCapture;
+};
+
+
diff --git a/Engine/Plugins/Experimental/MeshModelingToolset/Source/MeshModelingTools/Public/Properties/MeshStatisticsProperties.h b/Engine/Plugins/Experimental/MeshModelingToolset/Source/MeshModelingTools/Public/Properties/MeshStatisticsProperties.h
new file mode 100644
--- /dev/null
+++ b/Engine/Plugins/Experimental/MeshModelingToolset/Source/MeshModelingTools/Public/Properties/MeshStatisticsProperties.h
@@ -0,0 +1,34 @@
+// Copyright 1998-2019 Epic Games, Inc. All Rights Reserved.
+
+#pragma once
+
+#include "CoreMinimal.h"
+
+#include "InteractiveToolBuilder.h"
+
+#include "MeshStatisticsProperties.generated.h"
+
+
+// predeclarations
+class FDynamicMesh3;
+
+
+
+UCLASS()
+class MESHMODELINGTOOLS_API UMeshStatisticsProperties : public UInteractiveToolPropertySet
+{
+	GENERATED_BODY()
+public:
+
+	UPROPERTY(VisibleAnywhere, Category = MeshStatistics)
+	FString Mesh;
+
+	UPROPERTY(VisibleAnywhere, Category = MeshStatistics, AdvancedDisplay)
+	FString UV;
+
+	UPROPERTY(VisibleAnywhere, Category = MeshStatistics, AdvancedDisplay)
+	FString Attributes;
+
+	void Update(const FDynamicMesh3& Mesh);
+};
+
diff --git a/Engine/Plugins/Experimental/MeshModelingToolset/Source/MeshModelingTools/Public/RemeshMeshTool.h b/Engine/Plugins/Experimental/MeshModelingToolset/Source/MeshModelingTools/Public/RemeshMeshTool.h
new file mode 100644
--- /dev/null
+++ b/Engine/Plugins/Experimental/MeshModelingToolset/Source/MeshModelingTools/Public/RemeshMeshTool.h
@@ -0,0 +1,158 @@
+// Copyright 1998-2019 Epic Games, Inc. All Rights Reserved.
+
+#pragma once
+
+#include "CoreMinimal.h"
+#include "UObject/NoExportTypes.h"
+#include "SingleSelectionTool.h"
+#include "InteractiveToolBuilder.h"
+#include "DynamicMesh3.h"
+#include "DynamicMeshAABBTree3.h"
+#include "Properties/MeshStatisticsProperties.h"
+#include "RemeshMeshTool.generated.h"
+
+// predeclarations
+struct FMeshDescription;
+class USimpleDynamicMeshComponent;
+
+/** Mesh paint color view modes (somewhat maps to EVertexColorViewMode engine enum.) */
+UENUM()
+enum class ERemeshSmoothingType : uint8
+{
+	/** Uniform Smoothing */
+	Uniform = 0 UMETA(DisplayName = "Uniform"),
+
+	/** Cotangent Smoothing */
+	Cotangent = 1 UMETA(DisplayName = "Shape Preserving"),
+
+	/** Mean Value Smoothing */
+	MeanValue = 2 UMETA(DisplayName = "Mixed")
+};
+
+/**
+ *
+ */
+UCLASS()
+class MESHMODELINGTOOLS_API URemeshMeshToolBuilder : public UInteractiveToolBuilder
+{
+	GENERATED_BODY()
+
+public:
+	virtual bool CanBuildTool(const FToolBuilderState& SceneState) const override;
+	virtual UInteractiveTool* BuildTool(const FToolBuilderState& SceneState) const override;
+};
+
+/**
+ * Standard properties of the Remesh operation
+ */
+UCLASS()
+class MESHMODELINGTOOLS_API URemeshMeshToolProperties : public UInteractiveToolPropertySet
+{
+	GENERATED_BODY()
+
+public:
+	URemeshMeshToolProperties();
+
+	/** Target triangle count */
+	UPROPERTY(EditAnywhere, Category = Options, meta = (EditCondition = "bUseTargetEdgeLength == false"))
+	int TargetTriangleCount;
+
+
+	/** Smoothing speed */
+	UPROPERTY(EditAnywhere, Category = Options, meta = (UIMin = "0.0", UIMax = "1.0", ClampMin = "0.0", ClampMax = "1.0"))
+	float SmoothingSpeed;
+
+	/** Smoothing type */
+	UPROPERTY(EditAnywhere, Category = Options)
+	ERemeshSmoothingType SmoothingType;
+
+	/** Number of Remeshing passes */
+	UPROPERTY(EditAnywhere, Category = Options, meta = (UIMin = "0", UIMax = "50", ClampMin = "0", ClampMax = "1000"))
+	int RemeshIterations;
+
+
+	/** If true, UVs and Normals are discarded  */
+	UPROPERTY(EditAnywhere, Category = Options)
+	bool bDiscardAttributes;
+
+	/** If true, sharp edges are preserved  */
+	UPROPERTY(EditAnywhere, Category = Options)
+	bool bPreserveSharpEdges;
+
+
+	/** If true, the target count is ignored and the target edge length is used directly */
+	UPROPERTY(EditAnywhere, Category = Options, AdvancedDisplay)
+	bool bUseTargetEdgeLength;
+
+	/** Target edge length */
+	UPROPERTY(EditAnywhere, Category = Options, AdvancedDisplay, meta = (NoSpinbox = "true", EditCondition = "bUseTargetEdgeLength == true"))
+	float TargetEdgeLength;
+
+
+	/** Enable edge flips */
+	UPROPERTY(EditAnywhere, Category = Options, AdvancedDisplay)
+	bool bFlips;
+
+	/** Enable edge splits */
+	UPROPERTY(EditAnywhere, Category = Options, AdvancedDisplay)
+	bool bSplits;
+
+	/** Enable edge collapses */
+	UPROPERTY(EditAnywhere, Category = Options, AdvancedDisplay)
+	bool bCollapses;
+
+	/** Enable projection back to input mesh */
+	UPROPERTY(EditAnywhere, Category = Options, AdvancedDisplay)
+	bool bReproject;
+
+	/** Prevent normal flips */
+	UPROPERTY(EditAnywhere, Category = Options, AdvancedDisplay)
+	bool bPreventNormalFlips;
+
+};
+
+
+/**
+ * Simple Mesh Remeshing Tool
+ */
+UCLASS()
+class MESHMODELINGTOOLS_API URemeshMeshTool : public USingleSelectionTool
+{
+	GENERATED_BODY()
+
+public:
+	URemeshMeshTool();
+
+	virtual void Setup() override;
+	virtual void Shutdown(EToolShutdownType ShutdownType) override;
+
+	virtual void Render(IToolsContextRenderAPI* RenderAPI) override;
+
+	virtual bool HasCancel() const override { return true; }
+	virtual bool HasAccept() const override;
+	virtual bool CanAccept() const override;
+
+	virtual void OnPropertyModified(UObject* PropertySet, UProperty* Property) override;
+
+public:
+
+	UPROPERTY()
+	URemeshMeshToolProperties* BasicProperties;
+
+	UPROPERTY()
+	UMeshStatisticsProperties* MeshStatisticsProperties;
+
+
+protected:
+	USimpleDynamicMeshComponent* DynamicMeshComponent;
+	FDynamicMesh3 OriginalMesh;
+	FDynamicMeshAABBTree3 OriginalMeshSpatial;
+	bool bResultValid;
+	void UpdateResult();
+
+	double InitialMeshArea;
+
+	double CalculateTargetEdgeLength(int TargetTriCount);
+
+	void ComputeRemeshing();
+};
diff --git a/Engine/Plugins/Experimental/MeshModelingToolset/Source/MeshModelingTools/Public/SmoothMeshTool.h b/Engine/Plugins/Experimental/MeshModelingToolset/Source/MeshModelingTools/Public/SmoothMeshTool.h
new file mode 100644
--- /dev/null
+++ b/Engine/Plugins/Experimental/MeshModelingToolset/Source/MeshModelingTools/Public/SmoothMeshTool.h
@@ -0,0 +1,119 @@
+// Copyright 1998-2019 Epic Games, Inc. All Rights Reserved.
+
+#pragma once
+
+#include "CoreMinimal.h"
+#include "UObject/NoExportTypes.h"
+#include "SingleSelectionTool.h"
+#include "InteractiveToolBuilder.h"
+#include "DynamicMesh3.h"
+#include "MeshOpPreviewHelpers.h"
+#include "SmoothMeshTool.generated.h"
+
+
+// predeclarations
+struct FMeshDescription;
+class USimpleDynamicMeshComponent;
+
+
+UENUM()
+enum class ESmoothMeshToolSmoothType : uint8
+{
+	/** Fast smooth with N iterations */
+	Iterative UMETA(DisplayName = "Fast Smoothing"),
+
+	/** Implicit smoothing, usually better with UV preservation */
+	BiHarmonic_Cotan UMETA(DisplayName = "Implicit Smoothing"),
+
+};
+
+
+/**
+ *
+ */
+UCLASS()
+class MESHMODELINGTOOLS_API USmoothMeshToolBuilder : public UInteractiveToolBuilder
+{
+	GENERATED_BODY()
+
+public:
+
+	IToolsContextAssetAPI* AssetAPI;
+
+	USmoothMeshToolBuilder()
+	{
+		AssetAPI = nullptr;
+	}
+
+	virtual bool CanBuildTool(const FToolBuilderState& SceneState) const override;
+	virtual UInteractiveTool* BuildTool(const FToolBuilderState& SceneState) const override;
+};
+
+
+
+
+
+/**
+ * Simple Mesh Smoothing Tool
+ */
+UCLASS()
+class MESHMODELINGTOOLS_API USmoothMeshTool : public USingleSelectionTool, public IDynamicMeshOperatorFactory
+{
+	GENERATED_BODY()
+
+public:
+	USmoothMeshTool() = default;
+
+	virtual void SetWorld(UWorld* World);
+	virtual void SetAssetAPI(IToolsContextAssetAPI* AssetAPIIn);
+
+	virtual void Setup() override;
+	virtual void Shutdown(EToolShutdownType ShutdownType) override;
+
+	virtual void Tick(float DeltaTime) override;
+	virtual void Render(IToolsContextRenderAPI* RenderAPI) override;
+
+	virtual bool HasCancel() const override { return true; }
+	virtual bool HasAccept() const override;
+	virtual bool CanAccept() const override;
+
+#if WITH_EDITOR
+	virtual void PostEditChangeProperty(FPropertyChangedEvent& PropertyChangedEvent) override;
+#endif
+
+	virtual void OnPropertyModified(UObject* PropertySet, UProperty* Property) override;
+
+	virtual TSharedPtr<FDynamicMeshOperator> MakeNewOperator() override;
+
+
+	
+protected:
+	// need to update bResultValid if these are modified, so we don't publicly expose them. 
+	// @todo setters/getters for these
+
+	/** primary brush mode */
+	UPROPERTY(EditAnywhere, Category = Options)
+	ESmoothMeshToolSmoothType SmoothType = ESmoothMeshToolSmoothType::Iterative;
+
+	/** Smoothing speed */
+	UPROPERTY(EditAnywhere, Category = Options, meta = (UIMin = "0.0", UIMax = "1.0", ClampMin = "0.0", ClampMax = "1.0"))
+	float SmoothSpeed = 0.1f;
+
+	/** Smoothing speed */
+	UPROPERTY(EditAnywhere, Category = Options, meta = (UIMin = "0", UIMax = "100", ClampMin = "0", ClampMax = "1000"))
+	int32 SmoothIterations = 1;
+
+
+protected:
+	void UpdateResult();
+
+	bool bResultValid = false;
+
+	UPROPERTY()
+	UMeshOpPreviewWithBackgroundCompute* Preview = nullptr;
+	FDynamicMesh3 SrcDynamicMesh;
+
+
+	UWorld* TargetWorld = nullptr;
+	IToolsContextAssetAPI* AssetAPI = nullptr;
+};
