Copying //UE4/Dev-Editor to Dev-Main (//UE4/Dev-Main) for 4.24

diff --git a/Engine/Plugins/Experimental/MeshModelingToolset/Source/MeshModelingTools/Public/TransformMeshesTool.h b/Engine/Plugins/Experimental/MeshModelingToolset/Source/MeshModelingTools/Public/TransformMeshesTool.h
new file mode 100644
--- /dev/null
+++ b/Engine/Plugins/Experimental/MeshModelingToolset/Source/MeshModelingTools/Public/TransformMeshesTool.h
@@ -0,0 +1,192 @@
+// Copyright 1998-2019 Epic Games, Inc. All Rights Reserved.
+
+#pragma once
+
+#include "CoreMinimal.h"
+#include "MultiSelectionTool.h"
+#include "InteractiveToolBuilder.h"
+#include "BaseBehaviors//BehaviorTargetInterfaces.h"
+#include "Changes/TransformChange.h"
+#include "FrameTypes.h"
+#include "TransformMeshesTool.generated.h"
+
+class UBaseAxisTranslationGizmo;
+class UAxisAngleGizmo;
+class UTransformGizmo;
+class UTransformProxy;
+
+/**
+ *
+ */
+UCLASS()
+class MESHMODELINGTOOLS_API UTransformMeshesToolBuilder : public UInteractiveToolBuilder
+{
+	GENERATED_BODY()
+
+public:
+	virtual bool CanBuildTool(const FToolBuilderState& SceneState) const override;
+	virtual UInteractiveTool* BuildTool(const FToolBuilderState& SceneState) const override;
+};
+
+
+
+/** Mesh Transform modes */
+UENUM()
+enum class ETransformMeshesTransformMode : uint8
+{
+	/** Single Gizmo for all Objects */
+	SharedGizmo = 0 UMETA(DisplayName = "Shared Gizmo"),
+
+	/** Single Gizmo for all Objects, Rotations applied per-Object */
+	SharedGizmoLocal = 1 UMETA(DisplayName = "Shared Gizmo (Local)"),
+
+	/** Separate Gizmo for each Object */
+	PerObjectGizmo = 2 UMETA(DisplayName = "Multi-Gizmo"),
+
+	LastValue UMETA(Hidden)
+};
+
+
+
+/** Snap-Drag Source Point */
+UENUM()
+enum class ETransformMeshesSnapDragSource : uint8
+{
+	/** Snap-Drag moves the Clicked Point to the Target Location */
+	ClickPoint = 0 UMETA(DisplayName = "Click Point"),
+
+	/** Snap-Drag moves the Gizmo/Pivot to the Target Location */
+	Pivot = 1 UMETA(DisplayName = "Pivot"),
+
+
+	LastValue UMETA(Hidden)
+
+};
+
+
+
+/** Snap-Drag Rotation Mode */
+UENUM()
+enum class ETransformMeshesSnapDragRotationMode : uint8
+{
+	/** Snap-Drag only translates, ignoring Normals */
+	Ignore = 0 UMETA(DisplayName = "Ignore"),
+
+	/** Snap-Drag aligns the Source and Target Normals to point in the same direction */
+	Align = 1 UMETA(DisplayName = "Align"),
+
+	/** Snap-Drag aligns the Source Normal to the opposite of the Target Normal direction */
+	AlignFlipped = 2 UMETA(DisplayName = "Align Flipped"),
+
+	LastValue UMETA(Hidden)
+};
+
+
+
+/**
+ * Standard properties of the Transform Meshes operation
+ */
+UCLASS()
+class MESHMODELINGTOOLS_API UTransformMeshesToolProperties : public UInteractiveToolPropertySet
+{
+	GENERATED_BODY()
+public:
+
+	UPROPERTY(EditAnywhere, Category = Options)
+	ETransformMeshesTransformMode TransformMode = ETransformMeshesTransformMode::SharedGizmo;
+
+
+	UPROPERTY(EditAnywhere, Category = Options, meta = (EditCondition = "TransformMode == ETransformMeshesTransformMode::SharedGizmo || TransformMode == ETransformMeshesTransformMode::PerObjectGizmo"))
+	bool bSetPivot = false;
+
+
+	/** Click-drag starting on the target objects to reposition them on the rest of the scene */
+	UPROPERTY(EditAnywhere, Category = Options)
+	bool bEnableSnapDragging = false;
+
+
+	UPROPERTY(EditAnywhere, Category = Options, meta = (EditCondition = "bEnableSnapDragging == true"))
+	ETransformMeshesSnapDragSource SnapDragSource = ETransformMeshesSnapDragSource::ClickPoint;
+
+	/** When Snap-Dragging, align source and target normals */
+	UPROPERTY(EditAnywhere, Category = Options, meta = (EditCondition = "bEnableSnapDragging == true"))
+	ETransformMeshesSnapDragRotationMode RotationMode = ETransformMeshesSnapDragRotationMode::AlignFlipped;
+};
+
+
+USTRUCT()
+struct FTransformMeshesTarget
+{
+	GENERATED_BODY()
+
+	UPROPERTY()
+	UTransformProxy* TransformProxy;
+
+	UPROPERTY()
+	UTransformGizmo* TransformGizmo;
+};
+
+
+/**
+ *
+ */
+UCLASS()
+class MESHMODELINGTOOLS_API UTransformMeshesTool : public UMultiSelectionTool, public IClickDragBehaviorTarget
+{
+	GENERATED_BODY()
+
+public:
+	UTransformMeshesTool();
+
+	virtual void RegisterActions(FInteractiveToolActionSet& ActionSet) override;
+
+	virtual void SetWorld(UWorld* World, UInteractiveGizmoManager* GizmoManager);
+
+	virtual void Setup() override;
+	virtual void Shutdown(EToolShutdownType ShutdownType) override;
+
+	virtual void Tick(float DeltaTime) override;
+	virtual void Render(IToolsContextRenderAPI* RenderAPI) override;
+
+	virtual bool HasCancel() const override { return false; }
+	virtual bool HasAccept() const override { return false; }
+	virtual bool CanAccept() const override { return false; }
+
+	virtual void OnPropertyModified(UObject* PropertySet, UProperty* Property) override;
+
+
+	// ICLickDragBehaviorTarget interface
+	virtual FInputRayHit CanBeginClickDragSequence(const FInputDeviceRay& PressPos) override;
+	virtual void OnClickPress(const FInputDeviceRay& PressPos) override;
+	virtual void OnClickDrag(const FInputDeviceRay& DragPos) override;
+	virtual void OnClickRelease(const FInputDeviceRay& ReleasePos) override;
+	virtual void OnTerminateDragSequence() override;
+
+protected:
+	UPROPERTY()
+	UTransformMeshesToolProperties* TransformProps;
+
+protected:
+	UWorld* TargetWorld;
+	UInteractiveGizmoManager* GizmoManager;
+
+	UPROPERTY()
+	TArray<FTransformMeshesTarget> ActiveGizmos;
+
+
+	ETransformMeshesTransformMode CurTransformMode;
+	void UpdateTransformMode(ETransformMeshesTransformMode NewMode);
+	bool bCurSetPivotMode;
+	void UpdateSetPivotModes(bool bEnableSetPivot);
+
+	void SetActiveGizmos_Single(bool bLocalRotations);
+	void SetActiveGizmos_PerObject();
+	void ResetActiveGizmos();
+
+
+	FFrame3d StartDragFrameWorld;
+	FTransform StartDragTransform;
+	int ActiveSnapDragIndex = -1;
+
+	void OnParametersUpdated();
+};
diff --git a/Engine/Plugins/Experimental/MeshModelingToolset/Source/MeshModelingTools/Public/UVProjectionTool.h b/Engine/Plugins/Experimental/MeshModelingToolset/Source/MeshModelingTools/Public/UVProjectionTool.h
new file mode 100644
--- /dev/null
+++ b/Engine/Plugins/Experimental/MeshModelingToolset/Source/MeshModelingTools/Public/UVProjectionTool.h
@@ -0,0 +1,196 @@
+// Copyright 1998-2019 Epic Games, Inc. All Rights Reserved.
+
+#pragma once
+
+#include "CoreMinimal.h"
+#include "UObject/NoExportTypes.h"
+#include "MultiSelectionTool.h"
+#include "InteractiveToolBuilder.h"
+#include "MeshOpPreviewHelpers.h"
+#include "Drawing/ToolDataVisualizer.h"
+#include "ParameterizationOps/UVProjectionOp.h"
+#include "DynamicMesh3.h"
+#include "BaseTools/SingleClickTool.h"
+#include "Properties/MeshMaterialProperties.h"
+
+#include "UVProjectionTool.generated.h"
+
+
+// predeclarations
+struct FMeshDescription;
+class USimpleDynamicMeshComponent;
+class UTransformGizmo;
+class UTransformProxy;
+
+
+/**
+ *
+ */
+UCLASS()
+class MESHMODELINGTOOLS_API UUVProjectionToolBuilder : public UInteractiveToolBuilder
+{
+	GENERATED_BODY()
+
+public:
+	IToolsContextAssetAPI* AssetAPI = nullptr;
+
+	virtual bool CanBuildTool(const FToolBuilderState& SceneState) const override;
+	virtual UInteractiveTool* BuildTool(const FToolBuilderState& SceneState) const override;
+};
+
+
+
+
+
+/**
+ * Standard properties
+ */
+UCLASS()
+class MESHMODELINGTOOLS_API UUVProjectionToolProperties : public UInteractiveToolPropertySet
+{
+	GENERATED_BODY()
+
+public:
+	UUVProjectionToolProperties();
+
+
+
+	/** Choose the UV projection method (cube, cylinder, plane) */
+	UPROPERTY(EditAnywhere, Category = ProjectionSettings)
+	EUVProjectionMethod UVProjectionMethod;
+
+	/** Per-axis scaling of projection primitive */
+	UPROPERTY(EditAnywhere, Category = ProjectionSettings)
+	FVector ProjectionPrimitiveScale;
+
+	/** If triangle normal direction is within this threshold degrees of the cylinder top/bottom plane direction, project UVs to the top/bottom plane instead of the sides */
+	UPROPERTY(EditAnywhere, Category = ProjectionSettings, meta = (Default = "1", UIMin = "0", UIMax = "20", ClampMin = "0.0", ClampMax = "90.0", EditCondition = "UVProjectionMethod == EUVProjectionMethod::Cylinder"))
+	float CylinderProjectToTopOrBottomAngleThreshold;
+
+	/** Choose the UV scale factors */
+	UPROPERTY(EditAnywhere, Category = ProjectionSettings)
+	FVector2D UVScale;
+
+	/** Choose the UV offsets */
+	UPROPERTY(EditAnywhere, Category = ProjectionSettings)
+	FVector2D UVOffset;
+
+	/** If set, UV scales will be relative to world space so different objects created with the same UV scale should have the same average texel size */
+	UPROPERTY(EditAnywhere, Category = ProjectionSettings, meta = (DisplayName = "UV Scale Relative to World Space"))
+	bool bWorldSpaceUVScale = false;
+
+};
+
+
+
+
+/**
+ * Advanced properties
+ */
+UCLASS()
+class MESHMODELINGTOOLS_API UUVProjectionAdvancedProperties : public UInteractiveToolPropertySet
+{
+	GENERATED_BODY()
+
+public:
+	UUVProjectionAdvancedProperties();
+};
+
+
+/**
+ * Factory with enough info to spawn the background-thread Operator to do a chunk of work for the tool
+ *  stores a pointer to the tool and enough info to know which specific operator it should spawn
+ */
+UCLASS()
+class MESHMODELINGTOOLS_API UUVProjectionOperatorFactory : public UObject, public IDynamicMeshOperatorFactory
+{
+	GENERATED_BODY()
+
+public:
+	// IDynamicMeshOperatorFactory API
+	virtual TSharedPtr<FDynamicMeshOperator> MakeNewOperator() override;
+
+	UPROPERTY()
+	UUVProjectionTool *Tool;
+
+	int ComponentIndex;
+
+};
+
+/**
+ * Simple Mesh Normal Updating Tool
+ */
+UCLASS()
+class MESHMODELINGTOOLS_API UUVProjectionTool : public UMultiSelectionTool
+{
+	GENERATED_BODY()
+
+public:
+
+	friend UUVProjectionOperatorFactory;
+
+	UUVProjectionTool();
+
+	virtual void Setup() override;
+	virtual void Shutdown(EToolShutdownType ShutdownType) override;
+
+	virtual void SetWorld(UWorld* World, UInteractiveGizmoManager* GizmoManagerIn);
+	virtual void SetAssetAPI(IToolsContextAssetAPI* AssetAPI);
+
+	virtual void Tick(float DeltaTime) override;
+	virtual void Render(IToolsContextRenderAPI* RenderAPI) override;
+
+	virtual bool HasCancel() const override { return true; }
+	virtual bool HasAccept() const override;
+	virtual bool CanAccept() const override;
+
+#if WITH_EDITOR
+	virtual void PostEditChangeProperty(FPropertyChangedEvent &PropertyChangedEvent) override;
+#endif
+
+	virtual void OnPropertyModified(UObject* PropertySet, UProperty* Property) override;
+
+protected:
+
+	UPROPERTY()
+	UUVProjectionToolProperties* BasicProperties;
+
+	UPROPERTY()
+	UUVProjectionAdvancedProperties* AdvancedProperties;
+
+	UPROPERTY()
+	UExistingMeshMaterialProperties* MaterialSettings = nullptr;
+
+
+
+	UPROPERTY()
+	TArray<UMeshOpPreviewWithBackgroundCompute*> Previews;
+
+
+	UPROPERTY()
+	UMaterialInstanceDynamic* CheckerMaterial = nullptr;
+
+	UPROPERTY()
+	TArray<UTransformGizmo*> TransformGizmos;
+	
+	UPROPERTY()
+	TArray<UTransformProxy*> TransformProxies;
+
+protected:
+	TArray<TSharedPtr<FDynamicMesh3>> OriginalDynamicMeshes;
+	TArray<FDynamicMesh3> ReferencePrimitives;
+
+	UWorld* TargetWorld;
+	IToolsContextAssetAPI* AssetAPI;
+	UInteractiveGizmoManager* GizmoManager;
+
+	FViewCameraState CameraState;
+
+	FToolDataVisualizer ProjectionShapeVisualizer;
+
+	void UpdateNumPreviews();
+
+	void GenerateAsset(const TArray<TUniquePtr<FDynamicMeshOpResult>>& Results);
+
+	void TransformChanged(UTransformProxy* Proxy, FTransform Transform);
+};
diff --git a/Engine/Plugins/Experimental/MeshModelingToolset/Source/MeshModelingTools/Public/WeldMeshEdgesTool.h b/Engine/Plugins/Experimental/MeshModelingToolset/Source/MeshModelingTools/Public/WeldMeshEdgesTool.h
new file mode 100644
--- /dev/null
+++ b/Engine/Plugins/Experimental/MeshModelingToolset/Source/MeshModelingTools/Public/WeldMeshEdgesTool.h
@@ -0,0 +1,71 @@
+// Copyright 1998-2019 Epic Games, Inc. All Rights Reserved.
+
+#pragma once
+
+#include "CoreMinimal.h"
+#include "UObject/NoExportTypes.h"
+#include "SingleSelectionTool.h"
+#include "InteractiveToolBuilder.h"
+#include "DynamicMesh3.h"
+#include "DynamicMeshAABBTree3.h"
+#include "WeldMeshEdgesTool.generated.h"
+
+
+// predeclarations
+struct FMeshDescription;
+class USimpleDynamicMeshComponent;
+
+/**
+ *
+ */
+UCLASS()
+class MESHMODELINGTOOLS_API UWeldMeshEdgesToolBuilder : public UInteractiveToolBuilder
+{
+	GENERATED_BODY()
+
+public:
+	virtual bool CanBuildTool(const FToolBuilderState& SceneState) const override;
+	virtual UInteractiveTool* BuildTool(const FToolBuilderState& SceneState) const override;
+};
+
+/**
+ * Mesh Weld Edges Tool
+ */
+UCLASS()
+class MESHMODELINGTOOLS_API UWeldMeshEdgesTool : public USingleSelectionTool
+{
+	GENERATED_BODY()
+
+public:
+	UWeldMeshEdgesTool();
+
+	virtual void Setup() override;
+	virtual void Shutdown(EToolShutdownType ShutdownType) override;
+
+	virtual void Render(IToolsContextRenderAPI* RenderAPI) override;
+
+	virtual bool HasCancel() const override { return true; }
+	virtual bool HasAccept() const override;
+	virtual bool CanAccept() const override;
+
+#if WITH_EDITOR
+	virtual void PostEditChangeProperty(FPropertyChangedEvent& PropertyChangedEvent) override;
+#endif
+protected:
+	// need to update bResultValid if these are modified, so we don't publicly expose them. 
+	// @todo setters/getters for these
+
+	/** Target edge length */
+	UPROPERTY(EditAnywhere, Category = Options, meta = (UIMin = "0.000001", UIMax = "0.01", ClampMin = "0.00000001", ClampMax = "1000.0"))
+	float Tolerance;
+
+	/** Enable edge flips */
+	UPROPERTY(EditAnywhere, Category = Options, AdvancedDisplay)
+	bool bOnlyUnique;
+
+protected:
+	USimpleDynamicMeshComponent* DynamicMeshComponent;
+	FDynamicMesh3 OriginalMesh;
+	bool bResultValid;
+	void UpdateResult();
+};
diff --git a/Engine/Plugins/Experimental/MeshModelingToolset/Source/MeshModelingToolsEditorOnly/MeshModelingToolsEditorOnly.Build.cs b/Engine/Plugins/Experimental/MeshModelingToolset/Source/MeshModelingToolsEditorOnly/MeshModelingToolsEditorOnly.Build.cs
new file mode 100644
--- /dev/null
+++ b/Engine/Plugins/Experimental/MeshModelingToolset/Source/MeshModelingToolsEditorOnly/MeshModelingToolsEditorOnly.Build.cs
@@ -0,0 +1,76 @@
+// Copyright 1998-2019 Epic Games, Inc. All Rights Reserved.
+
+using UnrealBuildTool;
+
+public class MeshModelingToolsEditorOnly : ModuleRules
+{
+	public MeshModelingToolsEditorOnly(ReadOnlyTargetRules Target) : base(Target)
+	{
+		PCHUsage = ModuleRules.PCHUsageMode.UseExplicitOrSharedPCHs;
+		
+		PublicIncludePaths.AddRange(
+			new string[] {
+				// ... add public include paths required here ...
+			}
+			);
+				
+		
+		PrivateIncludePaths.AddRange(
+			new string[] {
+				// ... add other private include paths required here ...
+			}
+			);
+			
+		
+		PublicDependencyModuleNames.AddRange(
+			new string[]
+			{
+				"Core",
+                "InteractiveToolsFramework",
+                "MeshDescription",
+				"StaticMeshDescription",
+				"GeometricObjects",
+				"DynamicMesh",
+				"MeshConversion",
+				"MeshModelingTools",
+                "MeshSolverUtilities",
+				"ModelingComponents",
+				"ModelingOperators",
+				"ModelingOperatorsEditorOnly",
+                "ProxyLODMeshReduction", // currently required to be public due to IVoxelBasedCSG API
+				// ... add other public dependencies that you statically link with here ...
+			}
+            );
+			
+		
+		PrivateDependencyModuleNames.AddRange(
+			new string[]
+			{
+				"CoreUObject",
+				"Engine",
+                "RenderCore",
+                "RHI",
+				"InputCore",
+
+				"MeshUtilities",    // temp for saving mesh asset
+				"UnrealEd",
+				"MeshBuilder",
+                "MeshDescriptionOperations",
+				"MeshUtilitiesCommon",  
+				"MeshReductionInterface", // for UE4 standard simplification 
+                "ProxyLODMeshReduction", // for mesh merging voxel-based csg
+				//"Slate",
+				//"SlateCore",
+				// ... add private dependencies that you statically link with here ...	
+			}
+			);
+		
+		
+		DynamicallyLoadedModuleNames.AddRange(
+			new string[]
+			{
+				// ... add any modules that your module loads dynamically here ...
+			}
+			);
+	}
+}
diff --git a/Engine/Plugins/Experimental/MeshModelingToolset/Source/MeshModelingToolsEditorOnly/Private/MergeMeshesTool.cpp b/Engine/Plugins/Experimental/MeshModelingToolset/Source/MeshModelingToolsEditorOnly/Private/MergeMeshesTool.cpp
new file mode 100644
--- /dev/null
+++ b/Engine/Plugins/Experimental/MeshModelingToolset/Source/MeshModelingToolsEditorOnly/Private/MergeMeshesTool.cpp
@@ -0,0 +1,250 @@
+// Copyright 1998-2019 Epic Games, Inc. All Rights Reserved.
+
+#include "MergeMeshesTool.h"
+#include "InteractiveToolManager.h"
+#include "ToolBuilderUtil.h"
+#include "ToolSetupUtil.h"
+#include "AssetGenerationUtil.h"
+#include "Selection/ToolSelectionUtil.h"
+
+#include "DynamicMesh3.h"
+#include "MeshDescriptionToDynamicMesh.h"
+
+#include "CompositionOps/VoxelMergeMeshesOp.h"
+
+#define LOCTEXT_NAMESPACE "UMergeMeshesTool"
+
+/*
+ * ToolBuilder
+ */
+
+
+bool UMergeMeshesToolBuilder::CanBuildTool(const FToolBuilderState& SceneState) const
+{
+	bool bHasBuildAPI = (this->AssetAPI != nullptr);
+	bool bHasComponents = ToolBuilderUtil::CountComponents(SceneState, CanMakeComponentTarget) > 1;
+	return ( bHasBuildAPI && bHasComponents );
+}
+
+UInteractiveTool* UMergeMeshesToolBuilder::BuildTool(const FToolBuilderState& SceneState) const
+{
+	UMergeMeshesTool* NewTool = NewObject<UMergeMeshesTool>(SceneState.ToolManager);
+
+	TArray<UActorComponent*> Components = ToolBuilderUtil::FindAllComponents(SceneState, CanMakeComponentTarget);
+	check(Components.Num() > 0);
+
+	TArray<TUniquePtr<FPrimitiveComponentTarget>> ComponentTargets;
+	for (UActorComponent* ActorComponent : Components)
+	{
+		auto* MeshComponent = Cast<UPrimitiveComponent>(ActorComponent);
+		if ( MeshComponent )
+		{
+			ComponentTargets.Add(MakeComponentTarget(MeshComponent));
+		}
+	}
+
+	NewTool->SetSelection(MoveTemp(ComponentTargets));
+	NewTool->SetWorld(SceneState.World);
+	NewTool->SetAssetAPI(AssetAPI);
+
+	return NewTool;
+}
+
+
+
+
+/*
+ * Tool
+ */
+
+UMergeMeshesTool::UMergeMeshesTool()
+{
+}
+
+void UMergeMeshesTool::SetWorld(UWorld* World)
+{
+	this->TargetWorld = World;
+}
+
+void UMergeMeshesTool::SetAssetAPI(IToolsContextAssetAPI* AssetAPIIn)
+{
+	this->AssetAPI = AssetAPIIn;
+}
+
+
+
+void UMergeMeshesTool::Setup()
+{
+	UInteractiveTool::Setup();
+
+	MergeProps = NewObject<UMergeMeshesToolProperties>();
+	AddToolPropertySource(MergeProps);
+
+	MeshStatisticsProperties = NewObject<UMeshStatisticsProperties>(this);
+	AddToolPropertySource(MeshStatisticsProperties);
+
+	// Hide the source meshes
+	for (auto& ComponentTarget : ComponentTargets)
+	{
+		ComponentTarget->SetOwnerVisibility(false);
+	}
+
+	// save transformed version of input meshes (maybe this could happen in the Operator?)
+	CacheInputMeshes();
+
+	// initialize the PreviewMesh+BackgroundCompute object
+	Preview = NewObject<UMeshOpPreviewWithBackgroundCompute>(this, "Preview");
+	Preview->Setup(this->TargetWorld, this);
+	Preview->OnMeshUpdated.AddLambda([this](UMeshOpPreviewWithBackgroundCompute* Compute) {
+		MeshStatisticsProperties->Update(*Compute->PreviewMesh->GetPreviewDynamicMesh());
+	});
+
+	CreateLowQualityPreview(); // update the preview with a low-quality result
+	
+	Preview->ConfigureMaterials(
+		// TODO: This is is very likely a bug waiting to happen
+		ToolSetupUtil::GetDefaultMaterial(GetToolManager(), ComponentTargets[0]->GetMaterial(0)),
+		ToolSetupUtil::GetDefaultWorkingMaterial(GetToolManager())
+	);
+
+	Preview->InvalidateResult();    // start compute
+
+}
+
+
+void UMergeMeshesTool::Shutdown(EToolShutdownType ShutdownType)
+{
+	TUniquePtr<FDynamicMeshOpResult> Result = Preview->Shutdown();
+	if (ShutdownType == EToolShutdownType::Accept)
+	{
+		GetToolManager()->BeginUndoTransaction(LOCTEXT("MergeMeshes", "Merge Meshes"));
+
+		GenerateAsset(Result);
+
+		for (auto& ComponentTarget : ComponentTargets)
+		{
+			ComponentTarget->SetOwnerVisibility(true);
+			AActor* Actor = ComponentTarget->GetOwnerActor();
+			if (MergeProps->bDeleteInputActors)
+			{
+				Actor->Destroy();
+			}
+			else
+			{
+				// just hide the result.
+				Actor->SetIsTemporarilyHiddenInEditor(true);
+			}
+		}
+
+		GetToolManager()->EndUndoTransaction();
+
+	}
+	else
+	{
+		// Restore (unhide) the source meshes
+		for (auto& ComponentTarget : ComponentTargets)
+		{
+			ComponentTarget->SetOwnerVisibility(true);
+		}
+	}
+}
+
+
+
+void UMergeMeshesTool::Tick(float DeltaTime)
+{
+	Preview->Tick(DeltaTime);
+}
+
+
+void UMergeMeshesTool::Render(IToolsContextRenderAPI* RenderAPI)
+{
+}
+
+
+bool UMergeMeshesTool::HasAccept() const
+{
+	return true;
+}
+
+bool UMergeMeshesTool::CanAccept() const
+{
+	return Preview->HaveValidResult();
+}
+
+void UMergeMeshesTool::OnPropertyModified(UObject* PropertySet, UProperty* Property)
+{
+	Preview->InvalidateResult();
+}
+
+
+
+
+TSharedPtr<FDynamicMeshOperator> UMergeMeshesTool::MakeNewOperator()
+{
+	TSharedPtr<FVoxelMergeMeshesOp> MergeOp = MakeShared<FVoxelMergeMeshesOp>();
+	MergeOp->VoxelCount     = MergeProps->VoxelCount;
+	MergeOp->AdaptivityD    = MergeProps->MeshAdaptivity;
+	MergeOp->IsoSurfaceD    = MergeProps->OffsetDistance;
+	MergeOp->bAutoSimplify  = MergeProps->bAutoSimplify;
+	MergeOp->InputMeshArray = InputMeshes;
+	return MergeOp;
+}
+
+
+
+void UMergeMeshesTool::CacheInputMeshes()
+{
+	InputMeshes = MakeShared<TArray<IVoxelBasedCSG::FPlacedMesh>>();
+
+	// Package the selected meshes and transforms for consumption by the CSGTool
+	for (auto& ComponentTarget : ComponentTargets)
+	{
+		IVoxelBasedCSG::FPlacedMesh PlacedMesh;
+		PlacedMesh.Mesh = ComponentTarget->GetMesh();
+		PlacedMesh.Transform = ComponentTarget->GetWorldTransform();
+		InputMeshes->Add(PlacedMesh);
+	}
+}
+
+void UMergeMeshesTool::CreateLowQualityPreview()
+{
+
+	FProgressCancel NullInterrupter;
+	FVoxelMergeMeshesOp MergeMeshesOp;
+
+	MergeMeshesOp.VoxelCount = 12;
+	MergeMeshesOp.AdaptivityD = 0.001;
+	MergeMeshesOp.bAutoSimplify = true;
+	MergeMeshesOp.InputMeshArray = InputMeshes;
+	
+	MergeMeshesOp.CalculateResult(&NullInterrupter);
+	TUniquePtr<FDynamicMesh3> FastPreviewMesh = MergeMeshesOp.ExtractResult();
+
+
+	Preview->PreviewMesh->SetTransform((FTransform)MergeMeshesOp.GetResultTransform());
+	Preview->PreviewMesh->UpdatePreview(FastPreviewMesh.Get());  // copies the mesh @todo we could just give ownership to the Preview!
+	Preview->SetVisibility(true);
+}
+
+
+void UMergeMeshesTool::GenerateAsset(const TUniquePtr<FDynamicMeshOpResult>& Result)
+{
+	check(Result->Mesh.Get() != nullptr);
+
+	
+
+	AActor* NewActor = AssetGenerationUtil::GenerateStaticMeshActor(
+		AssetAPI, TargetWorld,
+		Result->Mesh.Get(), Result->Transform, TEXT("MergedMesh"),
+		AssetGenerationUtil::GetDefaultAutoGeneratedAssetPath());
+
+	// select newly-created object
+	ToolSelectionUtil::SetNewActorSelection(GetToolManager(), NewActor);
+
+	
+}
+
+
+
+#undef LOCTEXT_NAMESPACE
diff --git a/Engine/Plugins/Experimental/MeshModelingToolset/Source/MeshModelingToolsEditorOnly/Private/ParameterizeMeshTool.cpp b/Engine/Plugins/Experimental/MeshModelingToolset/Source/MeshModelingToolsEditorOnly/Private/ParameterizeMeshTool.cpp
new file mode 100644
--- /dev/null
+++ b/Engine/Plugins/Experimental/MeshModelingToolset/Source/MeshModelingToolsEditorOnly/Private/ParameterizeMeshTool.cpp
@@ -0,0 +1,198 @@
+// Copyright 1998-2019 Epic Games, Inc. All Rights Reserved.
+
+#include "ParameterizeMeshTool.h"
+#include "InteractiveToolManager.h"
+#include "ToolBuilderUtil.h"
+
+#include "DynamicMesh3.h"
+#include "DynamicMeshToMeshDescription.h"
+
+
+
+#include "SimpleDynamicMeshComponent.h"
+#include "Materials/MaterialInstanceDynamic.h"
+#include "MeshDescription.h"
+#include "MeshDescriptionToDynamicMesh.h"
+
+#include "ParameterizationOps/ParameterizeMeshOp.h"
+#include "ToolSetupUtil.h"
+
+#define LOCTEXT_NAMESPACE "UParameterizeMeshTool"
+
+
+DEFINE_LOG_CATEGORY_STATIC(LogParameterizeMeshTool, Log, All);
+
+/*
+ * ToolBuilder
+ */
+
+bool UParameterizeMeshToolBuilder::CanBuildTool(const FToolBuilderState& SceneState) const
+{
+	return ToolBuilderUtil::CountComponents(SceneState, CanMakeComponentTarget) == 1;
+}
+
+UInteractiveTool* UParameterizeMeshToolBuilder::BuildTool(const FToolBuilderState& SceneState) const
+{
+	UParameterizeMeshTool* NewTool = NewObject<UParameterizeMeshTool>(SceneState.ToolManager);
+
+	UActorComponent* ActorComponent = ToolBuilderUtil::FindFirstComponent(SceneState, CanMakeComponentTarget);
+	auto* MeshComponent = Cast<UPrimitiveComponent>(ActorComponent);
+	check(MeshComponent != nullptr);
+	NewTool->SetSelection(MakeComponentTarget(MeshComponent));
+	NewTool->SetWorld(SceneState.World);
+	NewTool->SetAssetAPI(AssetAPI);
+
+	return NewTool;
+}
+
+
+/*
+ * Tool
+ */
+
+UParameterizeMeshTool::UParameterizeMeshTool()
+{
+	ChartStretch = 0.11f;
+}
+
+void UParameterizeMeshTool::SetWorld(UWorld* World)
+{
+	this->TargetWorld = World;
+}
+
+void UParameterizeMeshTool::SetAssetAPI(IToolsContextAssetAPI* AssetAPIIn)
+{
+	this->AssetAPI = AssetAPIIn;
+}
+
+void UParameterizeMeshTool::Setup()
+{
+	UInteractiveTool::Setup();
+
+	// Deep copy of input mesh to be shared with the UV generation tool.
+	InputMesh = MakeShared<FMeshDescription>(*ComponentTarget->GetMesh());
+
+
+	// Copy existing material if there is one	
+	DefaultMaterial = ComponentTarget->GetMaterial(0);
+	if (DefaultMaterial == nullptr)
+	{
+		DefaultMaterial = LoadObject<UMaterial>(nullptr, TEXT("/Engine/EngineMaterials/DefaultMaterial"));  
+	}
+
+	
+
+	// hide input StaticMeshComponent
+	ComponentTarget->SetOwnerVisibility(false);
+
+	// initialize our properties
+	ToolPropertyObjects.Add(this);
+
+	
+	// Construct the preview object and set the material on it
+	Preview = NewObject<UMeshOpPreviewWithBackgroundCompute>(this, "Preview");
+	Preview->Setup(this->TargetWorld, this);
+
+	// Initialize the preview mesh with a copy of the source mesh.
+	{
+		FDynamicMesh3 Mesh;
+		FMeshDescriptionToDynamicMesh Converter;
+		Converter.bPrintDebugMessages = false;
+		Converter.Convert(InputMesh.Get(), Mesh);
+
+		Preview->PreviewMesh->UpdatePreview(&Mesh);
+		Preview->PreviewMesh->SetTransform(ComponentTarget->GetWorldTransform());
+	}
+	Preview->ConfigureMaterials(
+		ToolSetupUtil::GetDefaultMaterial(GetToolManager(), DefaultMaterial),
+		ToolSetupUtil::GetDefaultWorkingMaterial(GetToolManager())
+	);
+
+
+	MaterialSettings = NewObject<UExistingMeshMaterialProperties>(this);
+	MaterialSettings->Setup();
+	AddToolPropertySource(MaterialSettings);
+
+
+	Preview->SetVisibility(true);
+	Preview->InvalidateResult();    // start compute
+}
+
+void UParameterizeMeshTool::OnPropertyModified(UObject* PropertySet, UProperty* Property)
+{
+	if (PropertySet == MaterialSettings)
+	{
+		MaterialSettings->UpdateMaterials();
+		MaterialSettings->SetMaterialIfChanged(ComponentTarget->GetMaterial(0), Preview->StandardMaterial, [this](UMaterialInterface* Material)
+		{
+			Preview->ConfigureMaterials(ToolSetupUtil::GetDefaultMaterial(GetToolManager(), Material), Preview->WorkingMaterial);
+		});
+	}
+}
+
+
+void UParameterizeMeshTool::Shutdown(EToolShutdownType ShutdownType)
+{
+
+	TUniquePtr<FDynamicMeshOpResult> Result = Preview->Shutdown();
+	if (ShutdownType == EToolShutdownType::Accept)
+	{
+		FDynamicMesh3* DynamicMeshResult = Result->Mesh.Get();
+		check(DynamicMeshResult != nullptr);
+		GetToolManager()->BeginUndoTransaction(LOCTEXT("ParameterizeMesh", "Parameterize Mesh"));
+
+		ComponentTarget->CommitMesh([DynamicMeshResult](FMeshDescription* MeshDescription)
+		{
+			FDynamicMeshToMeshDescription Converter;
+			Converter.Convert(DynamicMeshResult, *MeshDescription);
+		});
+		GetToolManager()->EndUndoTransaction();
+	}
+
+	// Restore (unhide) the source meshes
+	ComponentTarget->SetOwnerVisibility(true);
+
+}
+
+void UParameterizeMeshTool::Tick(float DeltaTime)
+{
+	Preview->Tick(DeltaTime);
+}
+
+void UParameterizeMeshTool::Render(IToolsContextRenderAPI* RenderAPI)
+{
+}
+
+bool UParameterizeMeshTool::HasAccept() const
+{
+	return true;
+}
+
+bool UParameterizeMeshTool::CanAccept() const
+{
+	return Preview->HaveValidResult();
+}
+
+void UParameterizeMeshTool::PostEditChangeProperty(FPropertyChangedEvent& PropertyChangedEvent)
+{
+	// One of the UV generation properties must have changed.  Dirty the result to force a recompute
+	Preview->InvalidateResult();
+}
+
+TSharedPtr<FDynamicMeshOperator> UParameterizeMeshTool::MakeNewOperator()
+{
+	TSharedPtr<FParameterizeMeshOp> ParamertizeMeshOp = MakeShared<FParameterizeMeshOp>();
+	ParamertizeMeshOp->Stretch   = ChartStretch;
+	ParamertizeMeshOp->NumCharts = 0;
+	ParamertizeMeshOp->InputMesh = InputMesh;
+
+	const FTransform XForm = ComponentTarget->GetWorldTransform();
+	FTransform3d XForm3d(XForm);
+	ParamertizeMeshOp->SetTransform(XForm3d);
+
+	return ParamertizeMeshOp;
+}
+
+
+
+#undef LOCTEXT_NAMESPACE
diff --git a/Engine/Plugins/Experimental/MeshModelingToolset/Source/MeshModelingToolsEditorOnly/Private/PolygonOnMeshTool.cpp b/Engine/Plugins/Experimental/MeshModelingToolset/Source/MeshModelingToolsEditorOnly/Private/PolygonOnMeshTool.cpp
new file mode 100644
--- /dev/null
+++ b/Engine/Plugins/Experimental/MeshModelingToolset/Source/MeshModelingToolsEditorOnly/Private/PolygonOnMeshTool.cpp
@@ -0,0 +1,347 @@
+// Copyright 1998-2019 Epic Games, Inc. All Rights Reserved.
+
+#include "PolygonOnMeshTool.h"
+#include "InteractiveToolManager.h"
+#include "ToolBuilderUtil.h"
+
+#include "ToolSetupUtil.h"
+
+#include "DynamicMesh3.h"
+#include "BaseBehaviors/MultiClickSequenceInputBehavior.h"
+#include "Selection/SelectClickedAction.h"
+
+#include "MeshDescriptionToDynamicMesh.h"
+#include "DynamicMeshToMeshDescription.h"
+
+#include "InteractiveGizmoManager.h"
+
+#include "BaseGizmos/GizmoComponents.h"
+#include "BaseGizmos/TransformGizmo.h"
+
+#include "AssetGenerationUtil.h"
+
+
+#define LOCTEXT_NAMESPACE "UPolygonOnMeshTool"
+
+
+/*
+ * ToolBuilder
+ */
+
+
+bool UPolygonOnMeshToolBuilder::CanBuildTool(const FToolBuilderState& SceneState) const
+{
+	return ToolBuilderUtil::CountComponents(SceneState, CanMakeComponentTarget) == 1;
+}
+
+UInteractiveTool* UPolygonOnMeshToolBuilder::BuildTool(const FToolBuilderState& SceneState) const
+{
+	UPolygonOnMeshTool* NewTool = NewObject<UPolygonOnMeshTool>(SceneState.ToolManager);
+
+	UActorComponent* ActorComponent = ToolBuilderUtil::FindFirstComponent(SceneState, CanMakeComponentTarget);
+	auto* MeshComponent = Cast<UPrimitiveComponent>(ActorComponent);
+	check(MeshComponent != nullptr);
+	NewTool->SetSelection(MakeComponentTarget(MeshComponent));
+
+	NewTool->SetWorld(SceneState.World);
+	NewTool->SetAssetAPI(AssetAPI);
+	return NewTool;
+}
+
+
+
+
+/*
+ * Tool
+ */
+
+UPolygonOnMeshToolProperties::UPolygonOnMeshToolProperties()
+{
+	PolygonOperation = EEmbeddedPolygonOpMethod::CutThrough;
+	PolygonScale = 10;
+	// ExtrudeDistance = 10;
+	bDiscardAttributes = false;
+}
+
+UPolygonOnMeshAdvancedProperties::UPolygonOnMeshAdvancedProperties()
+{
+}
+
+
+UPolygonOnMeshTool::UPolygonOnMeshTool()
+{
+	EmbedPolygonOrigin = FVector::ZeroVector;
+	EmbedPolygonOrientation = FQuat::Identity;
+}
+
+void UPolygonOnMeshTool::SetWorld(UWorld* World)
+{
+	this->TargetWorld = World;
+}
+
+void UPolygonOnMeshTool::Setup()
+{
+	UInteractiveTool::Setup();
+
+	// hide input StaticMeshComponent
+	ComponentTarget->SetOwnerVisibility(false);
+
+	// click to set plane behavior
+	FSelectClickedAction* SetPlaneAction = new FSelectClickedAction();
+	SetPlaneAction->World = this->TargetWorld;
+	SetPlaneAction->OnClickedPositionFunc = [this](const FHitResult& Hit)
+	{
+		SetPlaneFromWorldPos(Hit.ImpactPoint, Hit.ImpactNormal);
+		for (UMeshOpPreviewWithBackgroundCompute* Preview : Previews)
+		{
+			Preview->InvalidateResult();
+		}
+	};
+	SetPointInWorldConnector = SetPlaneAction;
+
+	USingleClickInputBehavior* ClickToSetPlaneBehavior = NewObject<USingleClickInputBehavior>();
+	ClickToSetPlaneBehavior->ModifierCheckFunc = FInputDeviceState::IsCtrlKeyDown;
+	ClickToSetPlaneBehavior->Initialize(SetPointInWorldConnector);
+	AddInputBehavior(ClickToSetPlaneBehavior);
+
+
+	// create proxy and gizmo (but don't attach yet)
+	UInteractiveGizmoManager* GizmoManager = GetToolManager()->GetPairedGizmoManager();
+	PlaneTransformProxy = NewObject<UTransformProxy>(this);
+	PlaneTransformGizmo = GizmoManager->Create3AxisTransformGizmo(this);
+
+	BasicProperties = NewObject<UPolygonOnMeshToolProperties>(this, TEXT("Polygon On Mesh Settings"));
+	AdvancedProperties = NewObject<UPolygonOnMeshAdvancedProperties>(this, TEXT("Advanced Settings"));
+
+	// initialize our properties
+	AddToolPropertySource(BasicProperties);
+	AddToolPropertySource(AdvancedProperties);
+
+
+	// initialize the PreviewMesh+BackgroundCompute object
+	UpdateNumPreviews();
+
+	// set initial cut plane (also attaches gizmo/proxy)
+	FVector DefaultOrigin, Extents;
+	ComponentTarget->GetOwnerActor()->GetActorBounds(false, DefaultOrigin, Extents);
+	SetPlaneFromWorldPos(DefaultOrigin, FVector::UpVector);
+	// hook up callback so further changes trigger recut
+	PlaneTransformProxy->OnTransformChanged.AddUObject(this, &UPolygonOnMeshTool::TransformChanged);
+
+
+	// Convert input mesh description to dynamic mesh
+	OriginalDynamicMesh = MakeShared<FDynamicMesh3>();
+	FMeshDescriptionToDynamicMesh Converter;
+	Converter.bPrintDebugMessages = true;
+	Converter.Convert(ComponentTarget->GetMesh(), *OriginalDynamicMesh);
+	// TODO: consider adding an AABB tree construction here?  tradeoff vs doing a raycast against full every time a param change happens ...
+
+	for (UMeshOpPreviewWithBackgroundCompute* Preview : Previews)
+	{
+		Preview->InvalidateResult();
+	}
+}
+
+
+void UPolygonOnMeshTool::UpdateNumPreviews()
+{
+	int32 CurrentNumPreview = Previews.Num();
+	int32 TargetNumPreview = 1; // currently only have single preview use case; TODO consider cases where multiple meshes are generated with separate preview objects
+	if (TargetNumPreview < CurrentNumPreview)
+	{
+		for (int32 PreviewIdx = CurrentNumPreview - 1; PreviewIdx >= TargetNumPreview; PreviewIdx--)
+		{
+			Previews[PreviewIdx]->Cancel();
+		}
+		Previews.SetNum(TargetNumPreview);
+	}
+	else
+	{
+		for (int32 PreviewIdx = CurrentNumPreview; PreviewIdx < TargetNumPreview; PreviewIdx++)
+		{
+			UPolygonOnMeshOperatorFactory *OpFactory = NewObject<UPolygonOnMeshOperatorFactory>();
+			OpFactory->Tool = this;
+			// TODO: give the OpFactory more info about what specifically it is doing differently vs the other previews
+			UMeshOpPreviewWithBackgroundCompute* Preview = Previews.Add_GetRef(NewObject<UMeshOpPreviewWithBackgroundCompute>(OpFactory, "Preview"));
+			Preview->Setup(this->TargetWorld, OpFactory);
+			Preview->ConfigureMaterials(
+				ToolSetupUtil::GetDefaultMaterial(GetToolManager(), ComponentTarget->GetMaterial(0)),
+				ToolSetupUtil::GetDefaultWorkingMaterial(GetToolManager())
+			);
+			Preview->SetVisibility(true);
+		}
+	}
+}
+
+
+void UPolygonOnMeshTool::Shutdown(EToolShutdownType ShutdownType)
+{
+	// Restore (unhide) the source meshes
+	ComponentTarget->SetOwnerVisibility(true);
+
+	TArray<TUniquePtr<FDynamicMeshOpResult>> Results;
+	for (UMeshOpPreviewWithBackgroundCompute* Preview : Previews)
+	{
+		Results.Emplace(Preview->Shutdown());
+	}
+	if (ShutdownType == EToolShutdownType::Accept)
+	{
+		GenerateAsset(Results);
+	}
+
+	if (SetPointInWorldConnector != nullptr)
+	{
+		delete SetPointInWorldConnector;
+	}
+	UInteractiveGizmoManager* GizmoManager = GetToolManager()->GetPairedGizmoManager();
+	GizmoManager->DestroyAllGizmosByOwner(this);
+}
+
+void UPolygonOnMeshTool::SetAssetAPI(IToolsContextAssetAPI* AssetAPIIn)
+{
+	this->AssetAPI = AssetAPIIn;
+}
+
+TSharedPtr<FDynamicMeshOperator> UPolygonOnMeshOperatorFactory::MakeNewOperator()
+{
+	TSharedPtr<FEmbedPolygonsOp> EmbedOp = MakeShared<FEmbedPolygonsOp>();
+	EmbedOp->bDiscardAttributes = Tool->BasicProperties->bDiscardAttributes;
+	EmbedOp->Operation = Tool->BasicProperties->PolygonOperation;
+	EmbedOp->PolygonScale = Tool->BasicProperties->PolygonScale;
+	// EmbedOp->ExtrudeDistance = Tool->BasicProperties->ExtrudeDistance;
+	// TODO: also put the polygon info, etc, into the EmbedOp
+
+	FTransform LocalToWorld = Tool->ComponentTarget->GetWorldTransform();
+	FTransform WorldToLocal = LocalToWorld.Inverse();
+	FVector LocalOrigin = WorldToLocal.TransformPosition(Tool->EmbedPolygonOrigin);
+	FVector WorldNormal = Tool->EmbedPolygonOrientation.GetAxisZ();
+	FVector LocalNormal = WorldToLocal.TransformVectorNoScale(WorldNormal);
+	EmbedOp->LocalPlaneOrigin = LocalOrigin;
+	EmbedOp->LocalPlaneNormal = LocalNormal;
+	EmbedOp->OriginalMesh = Tool->OriginalDynamicMesh;
+	
+	EmbedOp->SetTransform(LocalToWorld);
+
+	return EmbedOp;
+}
+
+
+
+void UPolygonOnMeshTool::Render(IToolsContextRenderAPI* RenderAPI)
+{
+}
+
+void UPolygonOnMeshTool::Tick(float DeltaTime)
+{
+	for (UMeshOpPreviewWithBackgroundCompute* Preview : Previews)
+	{
+		Preview->Tick(DeltaTime);
+	}
+}
+
+
+
+void UPolygonOnMeshTool::PostEditChangeProperty(FPropertyChangedEvent& PropertyChangedEvent)
+{
+	UpdateNumPreviews();
+	for (UMeshOpPreviewWithBackgroundCompute* Preview : Previews)
+	{
+		Preview->InvalidateResult();
+	}
+}
+
+void UPolygonOnMeshTool::OnPropertyModified(UObject* PropertySet, UProperty* Property)
+{
+	UpdateNumPreviews();
+	for (UMeshOpPreviewWithBackgroundCompute* Preview : Previews)
+	{
+		Preview->InvalidateResult();
+	}
+}
+
+
+
+
+
+void UPolygonOnMeshTool::TransformChanged(UTransformProxy* Proxy, FTransform Transform)
+{
+	// TODO: if multi-select is re-enabled, only invalidate the preview that actually needs it?
+	EmbedPolygonOrientation = Transform.GetRotation();
+	EmbedPolygonOrigin = (FVector)Transform.GetTranslation();
+	for (UMeshOpPreviewWithBackgroundCompute* Preview : Previews)
+	{
+		Preview->InvalidateResult();
+	}
+}
+
+
+void UPolygonOnMeshTool::SetPlaneFromWorldPos(const FVector& Position, const FVector& Normal)
+{
+	EmbedPolygonOrigin = Position;
+
+	FFrame3f CutPlane(Position, Normal);
+	EmbedPolygonOrientation = CutPlane.Rotation;
+
+	PlaneTransformGizmo->SetActiveTarget(PlaneTransformProxy);
+	PlaneTransformGizmo->SetNewGizmoTransform(CutPlane.ToFTransform());
+}
+
+
+bool UPolygonOnMeshTool::HasAccept() const
+{
+	return true;
+}
+
+bool UPolygonOnMeshTool::CanAccept() const
+{
+	for (UMeshOpPreviewWithBackgroundCompute* Preview : Previews)
+	{
+		if (!Preview->HaveValidResult())
+		{
+			return false;
+		}
+	}
+	return true;
+}
+
+
+void UPolygonOnMeshTool::GenerateAsset(const TArray<TUniquePtr<FDynamicMeshOpResult>>& Results)
+{
+	GetToolManager()->BeginUndoTransaction(LOCTEXT("PolygonOnMeshToolTransactionName", "Polygon On Mesh Tool"));
+	
+
+	// first preview always replaces, and any subsequent previews can add new actors
+	// TODO: options to support other choices re what should be a new actor
+	check(Results.Num() > 0);
+	check(Results[0]->Mesh.Get() != nullptr);
+	ComponentTarget->CommitMesh([&Results](FMeshDescription* MeshDescription)
+	{
+		FDynamicMeshToMeshDescription Converter;
+		Converter.Convert(Results[0]->Mesh.Get(), *MeshDescription);
+	});
+
+	// currently nothing generates more than the single result for this tool
+	if (Results.Num() > 1)
+	{
+		unimplemented();
+		//FSelectedOjectsChangeList NewSelection;
+		//NewSelection.ModificationType = ESelectedObjectsModificationType::Replace;
+		//NewSelection.Actors.Add(ComponentTarget->GetOwnerActor());
+
+		//check(Results[1]->Mesh.Get() != nullptr);
+		//// TODO: copy over material?
+		//AActor* NewActor = AssetGenerationUtil::GenerateStaticMeshActor(
+		//	AssetAPI, TargetWorld,
+		//	Results[1]->Mesh.Get(), Results[1]->Transform, TEXT("Additional Polygon On Mesh Result"),
+		//	AssetGenerationUtil::GetDefaultAutoGeneratedAssetPath());
+		//NewSelection.Actors.Add(NewActor);
+		//GetToolManager()->RequestSelectionChange(NewSelection);
+	}
+
+	GetToolManager()->EndUndoTransaction();
+}
+
+
+
+
+#undef LOCTEXT_NAMESPACE
diff --git a/Engine/Plugins/Experimental/MeshModelingToolset/Source/MeshModelingToolsEditorOnly/Private/SimplifyMeshTool.cpp b/Engine/Plugins/Experimental/MeshModelingToolset/Source/MeshModelingToolsEditorOnly/Private/SimplifyMeshTool.cpp
new file mode 100644
--- /dev/null
+++ b/Engine/Plugins/Experimental/MeshModelingToolset/Source/MeshModelingToolsEditorOnly/Private/SimplifyMeshTool.cpp
@@ -0,0 +1,315 @@
+// Copyright 1998-2019 Epic Games, Inc. All Rights Reserved.
+
+#include "SimplifyMeshTool.h"
+#include "InteractiveToolManager.h"
+#include "ToolBuilderUtil.h"
+
+#include "DynamicMesh3.h"
+#include "DynamicMeshAttributeSet.h"
+#include "MeshSimplification.h"
+#include "MeshConstraintsUtil.h"
+#include "ProjectionTargets.h"
+
+#include "SimpleDynamicMeshComponent.h"
+
+#include "SceneManagement.h" // for FPrimitiveDrawInterface
+
+//#include "ProfilingDebugging/ScopedTimers.h" // enable this to use the timer.
+#include "MeshDescription.h"
+#include "MeshDescriptionOperations.h"
+#include "OverlappingCorners.h"
+#include "IMeshReductionManagerModule.h"
+#include "IMeshReductionInterfaces.h"
+#include "Modules/ModuleManager.h"
+#include "Operations/MergeCoincidentMeshEdges.h"
+
+#define LOCTEXT_NAMESPACE "USimplifyMeshTool"
+
+
+DEFINE_LOG_CATEGORY_STATIC(LogMeshSimplification, Log, All);
+
+/*
+ * ToolBuilder
+ */
+
+
+bool USimplifyMeshToolBuilder::CanBuildTool(const FToolBuilderState& SceneState) const
+{
+	return ToolBuilderUtil::CountComponents(SceneState, CanMakeComponentTarget) == 1;
+}
+
+UInteractiveTool* USimplifyMeshToolBuilder::BuildTool(const FToolBuilderState& SceneState) const
+{
+	USimplifyMeshTool* NewTool = NewObject<USimplifyMeshTool>(SceneState.ToolManager);
+
+	UActorComponent* ActorComponent = ToolBuilderUtil::FindFirstComponent(SceneState, CanMakeComponentTarget);
+	auto* MeshComponent = Cast<UPrimitiveComponent>(ActorComponent);
+	check(MeshComponent != nullptr);
+	NewTool->SetSelection(MakeComponentTarget(MeshComponent));
+
+	return NewTool;
+}
+
+
+
+/*
+ * Tool
+ */
+
+USimplifyMeshToolProperties::USimplifyMeshToolProperties()
+{
+	SimplifierType = ESimplifyType::QEM;
+	TargetMode = ESimplifyTargetType::Percentage;
+	TargetPercentage = 50;
+	TargetCount = 1000;
+	TargetEdgeLength = 5.0;
+	bReproject = false;
+	bPreventNormalFlips = true;
+	bDiscardAttributes = false;
+}
+
+
+
+void USimplifyMeshTool::Setup()
+{
+	UInteractiveTool::Setup();
+
+	// create dynamic mesh component to use for live preview
+	DynamicMeshComponent = NewObject<USimpleDynamicMeshComponent>(ComponentTarget->GetOwnerActor(), "DynamicMesh");
+	DynamicMeshComponent->SetupAttachment(ComponentTarget->GetOwnerActor()->GetRootComponent());
+	DynamicMeshComponent->RegisterComponent();
+	DynamicMeshComponent->SetWorldTransform(ComponentTarget->GetWorldTransform());
+
+	// copy material if there is one
+	auto Material = ComponentTarget->GetMaterial(0);
+	if (Material != nullptr)
+	{
+		DynamicMeshComponent->SetMaterial(0, Material);
+	}
+	DynamicMeshComponent->bExplicitShowWireframe = true;
+
+	DynamicMeshComponent->InitializeMesh(ComponentTarget->GetMesh());
+	OriginalMesh.Copy(*DynamicMeshComponent->GetMesh());
+	OriginalMeshSpatial.SetMesh(&OriginalMesh, true);
+
+	// hide input StaticMeshComponent
+	ComponentTarget->SetOwnerVisibility(false);
+
+	// initialize our properties
+	SimplifyProperties = NewObject<USimplifyMeshToolProperties>(this);
+	AddToolPropertySource(SimplifyProperties);
+
+	MeshStatisticsProperties = NewObject<UMeshStatisticsProperties>(this);
+	AddToolPropertySource(MeshStatisticsProperties);
+	MeshStatisticsProperties->Update(*DynamicMeshComponent->GetMesh());
+
+	bResultValid = false;
+}
+
+
+void USimplifyMeshTool::Shutdown(EToolShutdownType ShutdownType)
+{
+	if (DynamicMeshComponent != nullptr)
+	{
+		//DynamicMeshComponent->OnMeshChanged.Remove(OnDynamicMeshComponentChangedHandle);
+
+		ComponentTarget->SetOwnerVisibility(true);
+
+		if (ShutdownType == EToolShutdownType::Accept)
+		{
+			// this block bakes the modified DynamicMeshComponent back into the StaticMeshComponent inside an undo transaction
+			GetToolManager()->BeginUndoTransaction(LOCTEXT("SimplifyMeshToolTransactionName", "Simplify Mesh"));
+			ComponentTarget->CommitMesh([=](FMeshDescription* MeshDescription)
+			{
+				DynamicMeshComponent->Bake(MeshDescription, true);
+			});
+			GetToolManager()->EndUndoTransaction();
+		}
+
+		DynamicMeshComponent->UnregisterComponent();
+		DynamicMeshComponent->DestroyComponent();
+		DynamicMeshComponent = nullptr;
+	}
+}
+
+
+void USimplifyMeshTool::Render(IToolsContextRenderAPI* RenderAPI)
+{
+	UpdateResult();
+
+
+	FPrimitiveDrawInterface* PDI = RenderAPI->GetPrimitiveDrawInterface();
+	FTransform Transform = ComponentTarget->GetWorldTransform(); //Actor->GetTransform();
+
+	FColor LineColor(255, 0, 0);
+	FDynamicMesh3* TargetMesh = DynamicMeshComponent->GetMesh();
+	if (TargetMesh->HasAttributes())
+	{
+		FDynamicMeshUVOverlay* UVOverlay = TargetMesh->Attributes()->PrimaryUV();
+		for (int eid : TargetMesh->EdgeIndicesItr()) 
+		{
+			if (UVOverlay->IsSeamEdge(eid)) 
+			{
+				FVector3d A, B;
+				TargetMesh->GetEdgeV(eid, A, B);
+				PDI->DrawLine(Transform.TransformPosition(A), Transform.TransformPosition(B),
+					LineColor, 0, 2.0, 1.0f, true);
+			}
+		}
+	}
+}
+
+
+void USimplifyMeshTool::OnPropertyModified(UObject* PropertySet, UProperty* Property)
+{
+	bResultValid = false;
+}
+
+
+template <typename SimplificationType>
+void ComputeResult(FDynamicMesh3* TargetMesh, const bool bReproject, int OriginalTriCount, FDynamicMesh3& OriginalMesh, FDynamicMeshAABBTree3& OriginalMeshSpatial, const ESimplifyTargetType TargetMode, const float TargetPercentage, const int TargetCount, const float TargetEdgeLength)
+{
+	SimplificationType Reducer(TargetMesh);
+
+	Reducer.ProjectionMode = (bReproject) ? 
+		SimplificationType::ETargetProjectionMode::AfterRefinement : SimplificationType::ETargetProjectionMode::NoProjection;
+
+	Reducer.DEBUG_CHECK_LEVEL = 0;
+
+	FMeshConstraints constraints;
+	FMeshConstraintsUtil::ConstrainAllSeams(constraints, *TargetMesh, true, false);
+	//FMeshConstraintsUtil::ConstrainAllSeamJunctions(constraints, *TargetMesh, true, false);
+	Reducer.SetExternalConstraints(&constraints);
+
+	FMeshProjectionTarget ProjTarget(&OriginalMesh, &OriginalMeshSpatial);
+	Reducer.SetProjectionTarget(&ProjTarget);
+
+	if (TargetMode == ESimplifyTargetType::Percentage)
+	{
+		double Ratio = (double)TargetPercentage / 100.0;
+		int UseTarget = FMath::Max(4, (int)(Ratio * (double)OriginalTriCount));
+		Reducer.SimplifyToTriangleCount(UseTarget);
+	} 
+	else if (TargetMode == ESimplifyTargetType::TriangleCount)
+	{
+		Reducer.SimplifyToTriangleCount(TargetCount);
+	}
+	else if (TargetMode == ESimplifyTargetType::EdgeLength)
+	{
+		Reducer.SimplifyToEdgeLength(TargetEdgeLength);
+	}
+}
+
+
+void USimplifyMeshTool::UpdateResult()
+{
+	
+	if (bResultValid) 
+	{
+		return;
+	}
+
+	//FScopedDurationTimeLogger Timer(TEXT("Simplification Time"));
+
+	FDynamicMesh3* TargetMesh = DynamicMeshComponent->GetMesh();
+	TargetMesh->Copy(OriginalMesh);
+	int OriginalTriCount = OriginalMesh.TriangleCount();
+
+	if (SimplifyProperties->bDiscardAttributes)
+	{
+		TargetMesh->DiscardAttributes();
+	}
+	if (SimplifyProperties->SimplifierType == ESimplifyType::QEM)
+	{
+		ComputeResult<FQEMSimplification>(TargetMesh, SimplifyProperties->bReproject, OriginalTriCount, OriginalMesh, OriginalMeshSpatial, 
+			SimplifyProperties->TargetMode, SimplifyProperties->TargetPercentage, SimplifyProperties->TargetCount, SimplifyProperties->TargetEdgeLength);
+	}
+	else if (SimplifyProperties->SimplifierType == ESimplifyType::Attribute)
+	{
+		ComputeResult<FAttrMeshSimplification>(TargetMesh, SimplifyProperties->bReproject, OriginalTriCount, OriginalMesh, OriginalMeshSpatial,
+			SimplifyProperties->TargetMode, SimplifyProperties->TargetPercentage, SimplifyProperties->TargetCount, SimplifyProperties->TargetEdgeLength);
+	}
+	else //if (SimplifierType == ESimplifyType::UE4Standard)
+	{
+
+		const FMeshDescription* SrcMeshDescription = ComponentTarget->GetMesh();
+		FMeshDescription DstMeshDescription(*SrcMeshDescription);
+
+
+		FOverlappingCorners OverlappingCorners;
+		FMeshDescriptionOperations::FindOverlappingCorners(OverlappingCorners, *SrcMeshDescription, 1.e-5);
+
+		FMeshReductionSettings ReductionSettings;
+		if (SimplifyProperties->TargetMode == ESimplifyTargetType::Percentage)
+		{
+			ReductionSettings.PercentTriangles = FMath::Max(SimplifyProperties->TargetPercentage / 100., .001);  // Only support triangle percentage and count, but not edge length
+		}
+		else if (SimplifyProperties->TargetMode == ESimplifyTargetType::TriangleCount)
+		{
+			int32 NumTris = SrcMeshDescription->Polygons().Num();
+			ReductionSettings.PercentTriangles = (float)SimplifyProperties->TargetCount / (float)NumTris;
+		}
+
+		float Error;
+		{
+			IMeshReductionManagerModule& MeshReductionModule = FModuleManager::Get().LoadModuleChecked<IMeshReductionManagerModule>("MeshReductionInterface");
+			IMeshReduction* MeshReduction = MeshReductionModule.GetStaticMeshReductionInterface();
+
+
+			if (!MeshReduction)
+			{
+				// no reduction possible
+				Error = 0.f;
+				return;
+			}
+
+			Error = ReductionSettings.MaxDeviation;
+			MeshReduction->ReduceMeshDescription(DstMeshDescription, Error, *SrcMeshDescription, OverlappingCorners, ReductionSettings);
+		}
+
+		// Put the reduced mesh into the target...
+		DynamicMeshComponent->InitializeMesh(&DstMeshDescription);
+
+		// The UE4 tool will split the UV boundaries.  Need to weld this.
+		{
+			FDynamicMesh3* ComponentMesh = DynamicMeshComponent->GetMesh();
+
+			FMergeCoincidentMeshEdges Merger(ComponentMesh);
+			Merger.MergeSearchTolerance = 10.0f * FMathf::ZeroTolerance;
+			Merger.OnlyUniquePairs = false;
+			if (Merger.Apply() == false)
+			{
+				DynamicMeshComponent->InitializeMesh(&DstMeshDescription);
+			}
+			if (ComponentMesh->CheckValidity(true, EValidityCheckFailMode::ReturnOnly) == false)
+			{
+				DynamicMeshComponent->InitializeMesh(&DstMeshDescription);
+			}
+		
+			DynamicMeshComponent->NotifyMeshUpdated();
+		}
+	}
+
+	//UE_LOG(LogMeshSimplification, Log, TEXT("Mesh simplified to %d triangles"), TargetMesh->TriangleCount());
+
+	DynamicMeshComponent->NotifyMeshUpdated();
+	GetToolManager()->PostInvalidation();
+	MeshStatisticsProperties->Update(*DynamicMeshComponent->GetMesh());
+
+	bResultValid = true;
+}
+
+bool USimplifyMeshTool::HasAccept() const
+{
+	return true;
+}
+
+bool USimplifyMeshTool::CanAccept() const
+{
+	return true;
+}
+
+
+
+
+#undef LOCTEXT_NAMESPACE
diff --git a/Engine/Plugins/Experimental/MeshModelingToolset/Source/MeshModelingToolsEditorOnly/Private/VoxelCSGMeshesTool.cpp b/Engine/Plugins/Experimental/MeshModelingToolset/Source/MeshModelingToolsEditorOnly/Private/VoxelCSGMeshesTool.cpp
new file mode 100644
--- /dev/null
+++ b/Engine/Plugins/Experimental/MeshModelingToolset/Source/MeshModelingToolsEditorOnly/Private/VoxelCSGMeshesTool.cpp
@@ -0,0 +1,238 @@
+// Copyright 1998-2019 Epic Games, Inc. All Rights Reserved.
+
+#include "VoxelCSGMeshesTool.h"
+#include "InteractiveToolManager.h"
+#include "ToolBuilderUtil.h"
+#include "ToolSetupUtil.h"
+#include "AssetGenerationUtil.h"
+#include "Selection/ToolSelectionUtil.h"
+
+#include "DynamicMesh3.h"
+#include "MeshDescriptionToDynamicMesh.h"
+
+#include "CompositionOps/VoxelBooleanMeshesOp.h"
+
+#define LOCTEXT_NAMESPACE "UVoxelCSGMeshesTool"
+
+
+/*
+ * ToolBuilder
+ */
+bool UVoxelCSGMeshesToolBuilder::CanBuildTool(const FToolBuilderState& SceneState) const
+{
+	bool bHasBuildAPI = (this->AssetAPI != nullptr);
+	bool bHasComponents = ToolBuilderUtil::CountComponents(SceneState, CanMakeComponentTarget) == 2;
+	return (bHasBuildAPI && bHasComponents);
+}
+
+UInteractiveTool* UVoxelCSGMeshesToolBuilder::BuildTool(const FToolBuilderState& SceneState) const
+{
+	UVoxelCSGMeshesTool* NewTool = NewObject<UVoxelCSGMeshesTool>(SceneState.ToolManager);
+
+	TArray<UActorComponent*> Components = ToolBuilderUtil::FindAllComponents(SceneState, CanMakeComponentTarget);
+	check(Components.Num() > 0);
+
+	TArray<TUniquePtr<FPrimitiveComponentTarget>> ComponentTargets;
+	for (UActorComponent* ActorComponent : Components)
+	{
+		auto* MeshComponent = Cast<UPrimitiveComponent>(ActorComponent);
+		if ( MeshComponent )
+		{
+			ComponentTargets.Add(MakeComponentTarget(MeshComponent));
+		}
+	}
+
+	NewTool->SetSelection(MoveTemp(ComponentTargets));
+	NewTool->SetWorld(SceneState.World);
+	NewTool->SetAssetAPI(AssetAPI);
+
+	return NewTool;
+}
+
+/*
+ * Tool
+ */
+
+UVoxelCSGMeshesTool::UVoxelCSGMeshesTool()
+{
+}
+
+void UVoxelCSGMeshesTool::SetWorld(UWorld* World)
+{
+	this->TargetWorld = World;
+}
+
+void UVoxelCSGMeshesTool::SetAssetAPI(IToolsContextAssetAPI* AssetAPIIn)
+{
+	this->AssetAPI = AssetAPIIn;
+}
+
+
+
+void UVoxelCSGMeshesTool::Setup()
+{
+	UInteractiveTool::Setup();
+
+	CSGProps = NewObject<UVoxelCSGMeshesToolProperties>();
+	CSGProps->VoxelCount = 128;
+	CSGProps->MeshAdaptivity = 0.01f;
+	CSGProps->OffsetDistance = 0.0f;
+	AddToolPropertySource(CSGProps);
+
+	MeshStatisticsProperties = NewObject<UMeshStatisticsProperties>(this);
+	AddToolPropertySource(MeshStatisticsProperties);
+
+	// Hide the source meshes
+	for (auto& ComponentTarget : ComponentTargets)
+	{
+		ComponentTarget->SetOwnerVisibility(false);
+	}
+
+	// save transformed version of input meshes (maybe this could happen in the Operator?)
+	CacheInputMeshes();
+
+	// initialize the PreviewMesh+BackgroundCompute object
+	Preview = NewObject<UMeshOpPreviewWithBackgroundCompute>(this, "Preview");
+	Preview->Setup(this->TargetWorld, this);
+	Preview->OnMeshUpdated.AddLambda([this](UMeshOpPreviewWithBackgroundCompute* Compute) {
+		MeshStatisticsProperties->Update(*Compute->PreviewMesh->GetPreviewDynamicMesh());
+	});
+
+	CreateLowQualityPreview();
+
+	Preview->ConfigureMaterials(
+		// TODO : Probably multi-select bug - MCD
+		ToolSetupUtil::GetDefaultMaterial(GetToolManager(), ComponentTargets[0]->GetMaterial(0)),
+		ToolSetupUtil::GetDefaultWorkingMaterial(GetToolManager())
+	);
+	
+	Preview->InvalidateResult();    // start compute
+}
+
+
+void UVoxelCSGMeshesTool::Shutdown(EToolShutdownType ShutdownType)
+{
+	TUniquePtr<FDynamicMeshOpResult> Result = Preview->Shutdown();
+	if (ShutdownType == EToolShutdownType::Accept)
+	{
+		GetToolManager()->BeginUndoTransaction(LOCTEXT("VoxelCSGMeshes", "Boolean Meshes"));
+
+		GenerateAsset(Result);
+
+		for (auto& ComponentTarget : ComponentTargets)
+		{
+			ComponentTarget->SetOwnerVisibility(true);
+			AActor* Actor = ComponentTarget->GetOwnerActor();
+			if (CSGProps->bDeleteInputActors)
+			{
+				Actor->Destroy();
+			}
+			else
+			{
+				Actor->SetIsTemporarilyHiddenInEditor(true);
+			}
+		}
+
+		GetToolManager()->EndUndoTransaction();
+	}
+	else
+	{
+		// Restore (unhide) the source meshes
+		for (auto& ComponentTarget : ComponentTargets)
+		{
+			ComponentTarget->SetOwnerVisibility(true);
+		}
+	}
+}
+
+
+
+void UVoxelCSGMeshesTool::Tick(float DeltaTime)
+{
+	Preview->Tick(DeltaTime);
+}
+
+void UVoxelCSGMeshesTool::Render(IToolsContextRenderAPI* RenderAPI)
+{
+}
+
+bool UVoxelCSGMeshesTool::HasAccept() const
+{
+	return true;
+}
+
+bool UVoxelCSGMeshesTool::CanAccept() const
+{
+	return Preview->HaveValidResult();
+}
+
+void UVoxelCSGMeshesTool::OnPropertyModified(UObject* PropertySet, UProperty* Property)
+{
+	Preview->InvalidateResult();
+}
+
+TSharedPtr<FDynamicMeshOperator> UVoxelCSGMeshesTool::MakeNewOperator()
+{
+	TSharedPtr<FVoxelBooleanMeshesOp> CSGOp = MakeShared<FVoxelBooleanMeshesOp>();
+	CSGOp->Operation      = (FVoxelBooleanMeshesOp::EBooleanOperation)(int)CSGProps->Operation;
+	CSGOp->VoxelCount     = CSGProps->VoxelCount;
+	CSGOp->AdaptivityD    = CSGProps->MeshAdaptivity;
+	CSGOp->IsoSurfaceD    = CSGProps->OffsetDistance;
+	CSGOp->bAutoSimplify  = CSGProps->bAutoSimplify;
+	CSGOp->InputMeshArray = InputMeshes;
+	return CSGOp;
+}
+
+void UVoxelCSGMeshesTool::CacheInputMeshes()
+{
+	InputMeshes = MakeShared<TArray<IVoxelBasedCSG::FPlacedMesh>>();
+
+	// Package the selected meshes and transforms for consumption by the CSGTool
+	for (auto& ComponentTarget : ComponentTargets)
+	{
+		IVoxelBasedCSG::FPlacedMesh PlacedMesh;
+		PlacedMesh.Mesh = ComponentTarget->GetMesh();
+		PlacedMesh.Transform = ComponentTarget->GetWorldTransform();
+		InputMeshes->Add(PlacedMesh);
+	}
+}
+
+void UVoxelCSGMeshesTool::CreateLowQualityPreview()
+{
+
+	FProgressCancel NullInterrupter;
+	FVoxelBooleanMeshesOp BooleanOp;
+
+	BooleanOp.Operation = (FVoxelBooleanMeshesOp::EBooleanOperation)(int)CSGProps->Operation;
+	BooleanOp.VoxelCount = 12;
+	BooleanOp.AdaptivityD = 0.01;
+	BooleanOp.bAutoSimplify = true;
+	BooleanOp.InputMeshArray = InputMeshes;
+	
+	BooleanOp.CalculateResult(&NullInterrupter);
+	TUniquePtr<FDynamicMesh3> FastPreviewMesh = BooleanOp.ExtractResult();
+
+
+	Preview->PreviewMesh->SetTransform((FTransform)BooleanOp.GetResultTransform());
+	Preview->PreviewMesh->UpdatePreview(FastPreviewMesh.Get());  // copies the mesh @todo we could just give ownership to the Preview!
+	Preview->SetVisibility(true);
+}
+
+void UVoxelCSGMeshesTool::GenerateAsset(const TUniquePtr<FDynamicMeshOpResult>& Result)
+{
+	check(Result->Mesh.Get() != nullptr);
+
+	//GetToolManager()->BeginUndoTransaction(LOCTEXT("VoxelCSGMeshes", "Boolean Meshes"));
+
+	AActor* NewActor = AssetGenerationUtil::GenerateStaticMeshActor(
+		AssetAPI, TargetWorld,
+		Result->Mesh.Get(), Result->Transform, TEXT("CSGMesh"),
+		AssetGenerationUtil::GetDefaultAutoGeneratedAssetPath());
+
+	// select newly-created object
+	ToolSelectionUtil::SetNewActorSelection(GetToolManager(), NewActor);
+
+	//GetToolManager()->EndUndoTransaction();
+}
+
+#undef LOCTEXT_NAMESPACE
diff --git a/Engine/Plugins/Experimental/MeshModelingToolset/Source/MeshModelingToolsEditorOnly/Public/MergeMeshesTool.h b/Engine/Plugins/Experimental/MeshModelingToolset/Source/MeshModelingToolsEditorOnly/Public/MergeMeshesTool.h
new file mode 100644
--- /dev/null
+++ b/Engine/Plugins/Experimental/MeshModelingToolset/Source/MeshModelingToolsEditorOnly/Public/MergeMeshesTool.h
@@ -0,0 +1,118 @@
+// Copyright 1998-2019 Epic Games, Inc. All Rights Reserved.
+
+#pragma once
+
+#include "CoreMinimal.h"
+#include "MultiSelectionTool.h"
+#include "InteractiveToolBuilder.h"
+#include "MeshOpPreviewHelpers.h"
+#include "ProxyLODVolume.h"
+#include "Properties/MeshStatisticsProperties.h"
+#include "MergeMeshesTool.generated.h"
+
+
+/**
+ *
+ */
+UCLASS()
+class MESHMODELINGTOOLSEDITORONLY_API UMergeMeshesToolBuilder : public UInteractiveToolBuilder
+{
+	GENERATED_BODY()
+
+public:
+	IToolsContextAssetAPI* AssetAPI;
+
+	UMergeMeshesToolBuilder()
+	{
+		AssetAPI = nullptr;
+	}
+
+	virtual bool CanBuildTool(const FToolBuilderState& SceneState) const override;
+	virtual UInteractiveTool* BuildTool(const FToolBuilderState& SceneState) const override;
+};
+
+
+
+/**
+ * Standard properties of the Merge Meshes operation
+ */
+UCLASS()
+class MESHMODELINGTOOLSEDITORONLY_API UMergeMeshesToolProperties : public UInteractiveToolPropertySet
+{
+	GENERATED_BODY()
+public:
+
+	/** The size of the geometry bounding box major axis measured in voxels.*/
+	UPROPERTY(EditAnywhere, Category = Options, meta = (UIMin = "8", UIMax = "512", ClampMin = "8", ClampMax = "512"))
+	int32 VoxelCount = 128;
+
+	/** Remeshing adaptivity */
+	UPROPERTY(EditAnywhere, Category = Options, meta = (UIMin = "0", UIMax = "1", ClampMin = "0", ClampMax = "1"))
+	float MeshAdaptivity = 0.001f;
+
+	/** Offset when remeshing, measured in voxels units */
+	UPROPERTY(EditAnywhere, Category = Options, meta = (UIMin = "-2", UIMax = "2", ClampMin = "-2", ClampMax = "2"))
+	float OffsetDistance = 0;
+
+	/** Automatically simplify the result of voxel-based merge.*/
+	UPROPERTY(EditAnywhere, Category = Options)
+	bool bAutoSimplify = false;
+
+	/** Delete the source Actors/Components when accepting results of tool.*/
+	UPROPERTY(EditAnywhere, Category = Options)
+	bool bDeleteInputActors = true;
+};
+
+/**
+ *
+ */
+UCLASS()
+class MESHMODELINGTOOLSEDITORONLY_API UMergeMeshesTool : public UMultiSelectionTool, public IDynamicMeshOperatorFactory
+{
+	GENERATED_BODY()
+
+public:
+	UMergeMeshesTool();
+
+	virtual void SetWorld(UWorld* World);
+	virtual void SetAssetAPI(IToolsContextAssetAPI* AssetAPI);
+
+	virtual void Setup() override;
+	virtual void Shutdown(EToolShutdownType ShutdownType) override;
+
+	virtual void Tick(float DeltaTime) override;
+	virtual void Render(IToolsContextRenderAPI* RenderAPI) override;
+
+	virtual bool HasCancel() const override { return true; }
+	virtual bool HasAccept() const override;
+	virtual bool CanAccept() const override;
+
+	virtual void OnPropertyModified(UObject* PropertySet, UProperty* Property) override;
+
+	// IDynamicMeshOperatorFactory API
+	virtual TSharedPtr<FDynamicMeshOperator> MakeNewOperator() override;
+
+protected:
+	UPROPERTY()
+	UMergeMeshesToolProperties* MergeProps;
+
+	UPROPERTY()
+	UMeshStatisticsProperties* MeshStatisticsProperties;
+
+
+	UPROPERTY()
+	UMeshOpPreviewWithBackgroundCompute* Preview;
+
+protected:
+	UWorld* TargetWorld;
+	IToolsContextAssetAPI* AssetAPI;
+
+	TSharedPtr<TArray<IVoxelBasedCSG::FPlacedMesh>> InputMeshes;
+	/** stash copies of the transforms and pointers to the meshes for consumption by merge Op*/
+	void CacheInputMeshes();
+
+	/** quickly generate a low-quality result for display while the actual result is being computed. */
+	void CreateLowQualityPreview();
+
+	void GenerateAsset(const TUniquePtr<FDynamicMeshOpResult>& Result);
+};
diff --git a/Engine/Plugins/Experimental/MeshModelingToolset/Source/MeshModelingToolsEditorOnly/Public/ParameterizeMeshTool.h b/Engine/Plugins/Experimental/MeshModelingToolset/Source/MeshModelingToolsEditorOnly/Public/ParameterizeMeshTool.h
new file mode 100644
--- /dev/null
+++ b/Engine/Plugins/Experimental/MeshModelingToolset/Source/MeshModelingToolsEditorOnly/Public/ParameterizeMeshTool.h
@@ -0,0 +1,127 @@
+// Copyright 1998-2019 Epic Games, Inc. All Rights Reserved.
+
+#pragma once
+
+#include "CoreMinimal.h"
+#include "UObject/NoExportTypes.h"
+#include "SingleSelectionTool.h"
+#include "InteractiveToolBuilder.h"
+#include "DynamicMesh3.h"
+#include "MeshOpPreviewHelpers.h"
+#include "Properties/MeshMaterialProperties.h"
+
+#include "ParameterizeMeshTool.generated.h"
+
+
+// predeclarations
+struct FMeshDescription;
+class USimpleDynamicMeshComponent;
+class UMaterialInterface;
+class UMaterialInstanceDynamic;
+
+
+
+/** Material modes for MeshInspectorTool */
+UENUM()
+enum class EParameterizeMeshMaterialMode : uint8
+{
+	/** Input material */
+	Default,
+
+	/** Checkerboard material */
+	Checkerboard,
+
+	/** Override material */
+	Override
+};
+
+/**
+ *
+ */
+UCLASS()
+class MESHMODELINGTOOLSEDITORONLY_API UParameterizeMeshToolBuilder : public UInteractiveToolBuilder
+{
+	GENERATED_BODY()
+
+public:
+
+	IToolsContextAssetAPI* AssetAPI;
+
+	UParameterizeMeshToolBuilder()
+	{
+		AssetAPI = nullptr;
+	}
+
+	virtual bool CanBuildTool(const FToolBuilderState& SceneState) const override;
+	virtual UInteractiveTool* BuildTool(const FToolBuilderState& SceneState) const override;
+};
+
+
+
+
+
+/**
+ * Simple Mesh Simplifying Tool
+ */
+UCLASS()
+class MESHMODELINGTOOLSEDITORONLY_API UParameterizeMeshTool : public USingleSelectionTool, public IDynamicMeshOperatorFactory
+{
+	GENERATED_BODY()
+
+public:
+	UParameterizeMeshTool();
+
+	virtual void SetWorld(UWorld* World);
+	virtual void SetAssetAPI(IToolsContextAssetAPI* AssetAPI);
+
+	virtual void Setup() override;
+	virtual void Shutdown(EToolShutdownType ShutdownType) override;
+
+	virtual void Tick(float DeltaTime) override;
+	virtual void Render(IToolsContextRenderAPI* RenderAPI) override;
+
+	virtual bool HasCancel() const override { return true; }
+	virtual bool HasAccept() const override;
+	virtual bool CanAccept() const override;
+
+#if WITH_EDITOR
+	virtual void PostEditChangeProperty(FPropertyChangedEvent& PropertyChangedEvent) override;
+#endif
+
+	virtual void OnPropertyModified(UObject* PropertySet, UProperty* Property) override;
+
+	// IDynamicMeshOperatorFactory API
+	virtual TSharedPtr<FDynamicMeshOperator> MakeNewOperator() override;
+
+protected:
+	// need to update bResultValid if these are modified, so we don't publicly expose them. 
+	// @todo setters/getters for these
+
+	/** Maximum amount of stretch, from none to any.  If zero stretch is specified each triangle will likey be its own chart */
+	UPROPERTY(EditAnywhere, Category = Options, meta = (Default = "0.166", UIMin = "0", UIMax = "1", ClampMin = "0", ClampMax= "1"))
+	float ChartStretch = 0.11f;
+
+	UPROPERTY()
+	UExistingMeshMaterialProperties* MaterialSettings = nullptr;
+
+protected:
+	UPROPERTY()
+	UMaterialInterface* DefaultMaterial = nullptr;
+
+	UPROPERTY()
+	UMaterialInstanceDynamic* DisplayMaterial = nullptr;
+
+	UPROPERTY()
+	UMaterialInstanceDynamic* CheckerMaterial = nullptr;
+	
+	UPROPERTY()
+	UMeshOpPreviewWithBackgroundCompute* Preview = nullptr;
+
+
+protected:
+	UWorld* TargetWorld;
+	IToolsContextAssetAPI* AssetAPI;
+
+	TSharedPtr<FMeshDescription> InputMesh;
+
+};
diff --git a/Engine/Plugins/Experimental/MeshModelingToolset/Source/MeshModelingToolsEditorOnly/Public/PolygonOnMeshTool.h b/Engine/Plugins/Experimental/MeshModelingToolset/Source/MeshModelingToolsEditorOnly/Public/PolygonOnMeshTool.h
new file mode 100644
--- /dev/null
+++ b/Engine/Plugins/Experimental/MeshModelingToolset/Source/MeshModelingToolsEditorOnly/Public/PolygonOnMeshTool.h
@@ -0,0 +1,183 @@
+// Copyright 1998-2019 Epic Games, Inc. All Rights Reserved.
+
+#pragma once
+
+#include "CoreMinimal.h"
+#include "UObject/NoExportTypes.h"
+#include "InteractiveGizmo.h"
+#include "SingleSelectionTool.h"
+#include "InteractiveToolBuilder.h"
+#include "MeshOpPreviewHelpers.h"
+#include "CuttingOps/EmbedPolygonsOp.h"
+#include "DynamicMesh3.h"
+#include "BaseTools/SingleClickTool.h"
+#include "PolygonOnMeshTool.generated.h"
+
+
+// predeclarations
+struct FMeshDescription;
+class USimpleDynamicMeshComponent;
+class UTransformGizmo;
+class UTransformProxy;
+
+
+
+
+/**
+ *
+ */
+UCLASS()
+class MESHMODELINGTOOLSEDITORONLY_API UPolygonOnMeshToolBuilder : public UInteractiveToolBuilder
+{
+	GENERATED_BODY()
+
+public:
+	IToolsContextAssetAPI* AssetAPI = nullptr;
+
+	virtual bool CanBuildTool(const FToolBuilderState& SceneState) const override;
+	virtual UInteractiveTool* BuildTool(const FToolBuilderState& SceneState) const override;
+};
+
+
+
+
+
+/**
+ * Standard properties of the polygon-on-mesh operations
+ */
+UCLASS()
+class MESHMODELINGTOOLSEDITORONLY_API UPolygonOnMeshToolProperties : public UInteractiveToolPropertySet
+{
+	GENERATED_BODY()
+
+public:
+	UPolygonOnMeshToolProperties();
+
+	/** If true, UVs and Normals are discarded  */
+	UPROPERTY(EditAnywhere, Category = Options)
+	bool bDiscardAttributes;
+
+	/** What operation to do with the polygon -- e.g. extrude it, cut through with it, or delete inside it */
+	UPROPERTY(EditAnywhere, Category = Options)
+	EEmbeddedPolygonOpMethod PolygonOperation;
+
+	// TODO: re-add if/when extrude is added as a supported operation
+	///** Amount to extrude, if extrude is enabled */
+	//UPROPERTY(EditAnywhere, Category = Options)
+	//float ExtrudeDistance;
+
+	/** Scale of polygon to embed */
+	UPROPERTY(EditAnywhere, Category = Options)
+	float PolygonScale;
+};
+
+
+
+
+/**
+ * Advanced properties of the plane cut operation
+ */
+UCLASS()
+class MESHMODELINGTOOLSEDITORONLY_API UPolygonOnMeshAdvancedProperties : public UInteractiveToolPropertySet
+{
+	GENERATED_BODY()
+
+public:
+	UPolygonOnMeshAdvancedProperties();
+
+	// TODO: add additional properties, e.g. perhaps more advanced triangulation options, perhaps post-process remeshing options (because the polygon-in-mesh embedding can create slivers)
+
+};
+
+
+UCLASS()
+class MESHMODELINGTOOLSEDITORONLY_API UPolygonOnMeshOperatorFactory : public UObject, public IDynamicMeshOperatorFactory
+{
+	GENERATED_BODY()
+
+public:
+	// IDynamicMeshOperatorFactory API
+	virtual TSharedPtr<FDynamicMeshOperator> MakeNewOperator() override;
+
+	UPROPERTY()
+	UPolygonOnMeshTool *Tool;
+
+	// TODO: add info about what this factory is specifically meant to do with the tool (e.g. cut inside vs outside)
+	// currently not needed as we only do one thing at a time, as configured by the tool properties
+};
+
+/**
+ * Simple Mesh Plane Cutting Tool
+ */
+UCLASS()
+class MESHMODELINGTOOLSEDITORONLY_API UPolygonOnMeshTool : public USingleSelectionTool
+{
+	GENERATED_BODY()
+
+public:
+
+	friend UPolygonOnMeshOperatorFactory;
+
+	UPolygonOnMeshTool();
+
+	virtual void Setup() override;
+	virtual void Shutdown(EToolShutdownType ShutdownType) override;
+
+	virtual void SetWorld(UWorld* World);
+	virtual void SetAssetAPI(IToolsContextAssetAPI* AssetAPI);
+
+	virtual void Tick(float DeltaTime) override;
+	virtual void Render(IToolsContextRenderAPI* RenderAPI) override;
+
+	virtual bool HasCancel() const override { return true; }
+	virtual bool HasAccept() const override;
+	virtual bool CanAccept() const override;
+
+	virtual void PostEditChangeProperty(FPropertyChangedEvent& PropertyChangedEvent) override;
+
+	virtual void OnPropertyModified(UObject* PropertySet, UProperty* Property) override;
+
+protected:
+
+	UPROPERTY()
+	UPolygonOnMeshToolProperties* BasicProperties;
+
+	UPROPERTY()
+	UPolygonOnMeshAdvancedProperties* AdvancedProperties;
+
+	/** Origin of polygon projection plane */
+	UPROPERTY()
+	FVector EmbedPolygonOrigin;
+
+	/** Orientation of polygon projection plane */
+	UPROPERTY()
+	FQuat EmbedPolygonOrientation;
+
+	UPROPERTY()
+	TArray<UMeshOpPreviewWithBackgroundCompute*> Previews;
+
+
+protected:
+	TSharedPtr<FDynamicMesh3> OriginalDynamicMesh;
+
+	UWorld* TargetWorld;
+	IToolsContextAssetAPI* AssetAPI;
+
+	FViewCameraState CameraState;
+
+	UPROPERTY()
+	UTransformGizmo* PlaneTransformGizmo;
+
+	UPROPERTY()
+	UTransformProxy* PlaneTransformProxy;
+
+	void TransformChanged(UTransformProxy* Proxy, FTransform Transform);
+
+	void UpdateNumPreviews();
+
+	IClickBehaviorTarget* SetPointInWorldConnector = nullptr;
+
+	virtual void SetPlaneFromWorldPos(const FVector& Position, const FVector& Normal);
+
+	void GenerateAsset(const TArray<TUniquePtr<FDynamicMeshOpResult>>& Results);
+};
diff --git a/Engine/Plugins/Experimental/MeshModelingToolset/Source/MeshModelingToolsEditorOnly/Public/ShapeSprayTool.h b/Engine/Plugins/Experimental/MeshModelingToolset/Source/MeshModelingToolsEditorOnly/Public/ShapeSprayTool.h
new file mode 100644
--- /dev/null
+++ b/Engine/Plugins/Experimental/MeshModelingToolset/Source/MeshModelingToolsEditorOnly/Public/ShapeSprayTool.h
@@ -0,0 +1,112 @@
+// Copyright 1998-2019 Epic Games, Inc. All Rights Reserved.
+
+#pragma once
+
+#include "CoreMinimal.h"
+#include "UObject/NoExportTypes.h"
+#include "InteractiveToolBuilder.h"
+#include "DynamicMeshBrushTool.h"
+#include "ShapeSprayTool.generated.h"
+
+
+/**
+ * UMeshSurfacePointToolBuilder override for UShapeSprayTool
+ */
+UCLASS(Transient)
+class MESHMODELINGTOOLSEDITORONLY_API UShapeSprayToolBuilder : public UMeshSurfacePointToolBuilder
+{
+	GENERATED_BODY()
+
+public:
+	IToolsContextAssetAPI* AssetAPI = nullptr;
+
+	virtual UMeshSurfacePointTool* CreateNewTool(const FToolBuilderState& SceneState) const override;
+};
+
+
+/**
+ * Settings UObject for UShapeSprayTool. 
+ */
+UCLASS(Transient)
+class MESHMODELINGTOOLSEDITORONLY_API UShapeSprayToolProperties : public UInteractiveToolPropertySet
+{
+	GENERATED_BODY()
+public:
+	UShapeSprayToolProperties();
+
+	UPROPERTY(EditAnywhere, Category = Options)
+	FLinearColor Color;
+
+	UPROPERTY(EditAnywhere, Category = Options)
+	bool bRandomColor;
+
+	UPROPERTY(EditAnywhere, Category = Options, meta = (DisplayName = "Speed", UIMin = "0.01", UIMax = "1.0"))
+	float DropSpeed;
+
+	UPROPERTY(EditAnywhere, Category = Options, meta = (DisplayName = "Shape Size", UIMin = "1.0", UIMax = "30.0"))
+	float ObjectSize;
+
+	UPROPERTY(EditAnywhere, Category = Options, meta = (DisplayName = "Repeat Per Stamp"))
+	int NumSplats;
+
+	UPROPERTY(EditAnywhere, Category = Options)
+	UMaterialInterface* Material = nullptr;
+};
+
+
+
+/**
+ * UShapeSprayTool is a brush-based tool that generates random points on the
+ * target surface within the brush radius, and then creates small meshes
+ * at those points. The accumulated meshes are appended and can
+ * be emitted as a new StaticMeshComponent on Accept.
+ */
+UCLASS(Transient)
+class MESHMODELINGTOOLSEDITORONLY_API UShapeSprayTool : public UDynamicMeshBrushTool
+{
+	GENERATED_BODY()
+
+public:
+	UShapeSprayTool();
+
+	virtual void SetWorld(UWorld* World);
+	virtual void SetAssetAPI(IToolsContextAssetAPI* AssetAPI);
+
+	// UInteractiveTool API
+
+	virtual void Setup() override;
+	virtual void Shutdown(EToolShutdownType ShutdownType) override;
+
+	virtual bool HasCancel() const override { return true; }
+	virtual bool HasAccept() const override { return true; }
+	virtual bool CanAccept() const override;
+
+	virtual void OnPropertyModified(UObject* PropertySet, UProperty* Property) override;
+
+	// UMeshSurfacePointTool API
+	virtual void OnBeginDrag(const FRay& Ray) override;
+	virtual void OnUpdateDrag(const FRay& Ray) override;
+	virtual void OnEndDrag(const FRay& Ray) override;
+
+protected:
+	UPROPERTY()
+	UShapeSprayToolProperties* Settings;
+
+	// small meshes are accumulated here
+	UPROPERTY()
+	USimpleDynamicMeshComponent* AccumMeshComponent;
+
+
+protected:
+	UWorld* TargetWorld;
+	IToolsContextAssetAPI* AssetAPI;
+
+	FDynamicMesh3 ShapeMesh;
+	void UpdateShapeMesh();
+	void SplatShape(const FFrame3d& LocalFrame, double Scale, FDynamicMesh3* TargetMesh);
+	TArray<int> VertexMap;
+
+	FRandomStream Random;
+
+	virtual void EmitResult();
+};
diff --git a/Engine/Plugins/Experimental/MeshModelingToolset/Source/MeshModelingToolsEditorOnly/Public/SimplifyMeshTool.h b/Engine/Plugins/Experimental/MeshModelingToolset/Source/MeshModelingToolsEditorOnly/Public/SimplifyMeshTool.h
new file mode 100644
--- /dev/null
+++ b/Engine/Plugins/Experimental/MeshModelingToolset/Source/MeshModelingToolsEditorOnly/Public/SimplifyMeshTool.h
@@ -0,0 +1,152 @@
+// Copyright 1998-2019 Epic Games, Inc. All Rights Reserved.
+
+#pragma once
+
+#include "CoreMinimal.h"
+#include "UObject/NoExportTypes.h"
+#include "SingleSelectionTool.h"
+#include "InteractiveToolBuilder.h"
+#include "DynamicMesh3.h"
+#include "DynamicMeshAABBTree3.h"
+#include "Properties/MeshStatisticsProperties.h"
+#include "SimplifyMeshTool.generated.h"
+
+
+// predeclarations
+struct FMeshDescription;
+class USimpleDynamicMeshComponent;
+
+
+
+
+/**  */
+UENUM()
+enum class ESimplifyTargetType : uint8
+{
+	/** Percentage of input triangles */
+	Percentage = 0 UMETA(DisplayName = "Percentage"),
+
+	/** Target triangle count */
+	TriangleCount = 1 UMETA(DisplayName = "Triangle Count"),
+
+	/** Target edge length */
+	EdgeLength = 2 UMETA(DisplayName = "Edge Length")
+};
+
+/**  */
+UENUM()
+enum class ESimplifyType : uint8
+{
+	/** Fastest. Standard quadric error metric. Will not simplify UV boundaries.*/
+	QEM = 0 UMETA(DisplayName = "QEM"),
+
+	/** Potentially higher quality. Takes the normal into account. Will not simplify UV bounaries. */
+	Attribute = 1 UMETA(DisplayName = "Normal Aware"),
+
+	/** Highest quality reduction.  Will simplify UV boundaries. */
+	UE4Standard = 2 UMETA(DisplayName = "UE4 Standard"),
+};
+
+
+
+
+/**
+ *
+ */
+UCLASS()
+class MESHMODELINGTOOLSEDITORONLY_API USimplifyMeshToolBuilder : public UInteractiveToolBuilder
+{
+	GENERATED_BODY()
+
+public:
+	virtual bool CanBuildTool(const FToolBuilderState& SceneState) const override;
+	virtual UInteractiveTool* BuildTool(const FToolBuilderState& SceneState) const override;
+};
+
+
+
+/**
+ * Standard properties of the Simplify operation
+ */
+UCLASS()
+class MESHMODELINGTOOLSEDITORONLY_API USimplifyMeshToolProperties : public UInteractiveToolPropertySet
+{
+	GENERATED_BODY()
+public:
+	USimplifyMeshToolProperties();
+
+	/** Simplification Target Type  */
+	UPROPERTY(EditAnywhere, Category = Options)
+	ESimplifyTargetType TargetMode;
+
+	/** Simplification Scheme  */
+	UPROPERTY(EditAnywhere, Category = Options)
+	ESimplifyType SimplifierType;
+
+	/** Target percentage of original triangle count */
+	UPROPERTY(EditAnywhere, Category = Options, meta = (UIMin = "0", UIMax = "100", EditCondition = "TargetMode == ESimplifyTargetType::Percentage"))
+	int TargetPercentage;
+
+	/** Target edge length */
+	UPROPERTY(EditAnywhere, Category = Options, meta = (UIMin = "3.0", UIMax = "10.0", ClampMin = "0.001", ClampMax = "1000.0", EditCondition = "TargetMode == ESimplifyTargetType::EdgeLength && SimplifierType != ESimplifyType::UE4Standard"))
+	float TargetEdgeLength;
+
+	/** Target triangle count */
+	UPROPERTY(EditAnywhere, Category = Options, meta = (UIMin = "4", UIMax = "10000", ClampMin = "1", ClampMax = "9999999999", EditCondition = "TargetMode == ESimplifyTargetType::TriangleCount"))
+	int TargetCount;
+
+	/** If true, UVs and Normals are discarded  */
+	UPROPERTY(EditAnywhere, Category = Options)
+	bool bDiscardAttributes;
+
+	/** Enable projection back to input mesh */
+	UPROPERTY(EditAnywhere, Category = Options, AdvancedDisplay)
+	bool bReproject;
+
+	/** Prevent normal flips */
+	UPROPERTY(EditAnywhere, Category = Options, AdvancedDisplay)
+	bool bPreventNormalFlips;
+
+};
+
+
+
+
+/**
+ * Simple Mesh Simplifying Tool
+ */
+UCLASS()
+class MESHMODELINGTOOLSEDITORONLY_API USimplifyMeshTool : public USingleSelectionTool
+{
+	GENERATED_BODY()
+
+public:
+	virtual void Setup() override;
+	virtual void Shutdown(EToolShutdownType ShutdownType) override;
+
+	virtual void Render(IToolsContextRenderAPI* RenderAPI) override;
+
+	virtual bool HasCancel() const override { return true; }
+	virtual bool HasAccept() const override;
+	virtual bool CanAccept() const override;
+
+	virtual void OnPropertyModified(UObject* PropertySet, UProperty* Property) override;
+
+protected:
+	// need to update bResultValid if these are modified, so we don't publicly expose them. 
+	// @todo setters/getters for these
+
+	UPROPERTY()
+	USimplifyMeshToolProperties* SimplifyProperties;
+
+	UPROPERTY()
+	UMeshStatisticsProperties* MeshStatisticsProperties;
+
+protected:
+	USimpleDynamicMeshComponent* DynamicMeshComponent;
+	FDynamicMesh3 OriginalMesh;
+	FDynamicMeshAABBTree3 OriginalMeshSpatial;
+	bool bResultValid;
+	void UpdateResult();
+
+};
diff --git a/Engine/Plugins/Experimental/MeshModelingToolset/Source/MeshModelingToolsEditorOnly/Public/VoxelCSGMeshesTool.h b/Engine/Plugins/Experimental/MeshModelingToolset/Source/MeshModelingToolsEditorOnly/Public/VoxelCSGMeshesTool.h
new file mode 100644
--- /dev/null
+++ b/Engine/Plugins/Experimental/MeshModelingToolset/Source/MeshModelingToolsEditorOnly/Public/VoxelCSGMeshesTool.h
@@ -0,0 +1,143 @@
+// Copyright 1998-2019 Epic Games, Inc. All Rights Reserved.
+
+#pragma once
+
+#include "CoreMinimal.h"
+#include "UObject/NoExportTypes.h"
+#include "InteractiveToolBuilder.h"
+#include "MeshOpPreviewHelpers.h"
+#include "MultiSelectionTool.h"
+#include "ProxyLODVolume.h"
+#include "Properties/MeshStatisticsProperties.h"
+#include "VoxelCSGMeshesTool.generated.h"
+
+
+/**
+ *
+ */
+UCLASS()
+class MESHMODELINGTOOLSEDITORONLY_API UVoxelCSGMeshesToolBuilder : public UInteractiveToolBuilder
+{
+	GENERATED_BODY()
+
+public:
+	IToolsContextAssetAPI* AssetAPI;
+
+	UVoxelCSGMeshesToolBuilder()
+	{
+		AssetAPI = nullptr;
+	}
+
+	virtual bool CanBuildTool(const FToolBuilderState& SceneState) const override;
+	virtual UInteractiveTool* BuildTool(const FToolBuilderState& SceneState) const override;
+};
+
+
+/**  */
+UENUM()
+enum class EVoxelCSGOperation : uint8
+{
+	/** Subtracts the first object from the second */
+	DifferenceAB = 0 UMETA(DisplayName = "A - B"),
+
+	/** Subtracts the second object from the first */
+	DifferenceBA = 1 UMETA(DisplayName = "B - A"),
+
+	/** intersection of two objects */
+	Intersect = 2 UMETA(DisplayName = "Intersect"),
+
+	/** union of two objects */
+	Union = 3 UMETA(DisplayName = "Union"),
+
+};
+
+
+/**
+ * Standard properties of the Voxel CSG operation
+ */
+UCLASS()
+class MESHMODELINGTOOLSEDITORONLY_API UVoxelCSGMeshesToolProperties : public UInteractiveToolPropertySet
+{
+	GENERATED_BODY()
+public:
+	/** The type of operation  */
+	UPROPERTY(EditAnywhere, Category = Options)
+	EVoxelCSGOperation Operation;
+
+	/** The size of the geometry bounding box major axis measured in voxels.*/
+	UPROPERTY(EditAnywhere, Category = Options, meta = (UIMin = "8", UIMax = "512", ClampMin = "8", ClampMax = "512"))
+	int32 VoxelCount = 128;
+
+	/** Remeshing adaptivity */
+	UPROPERTY(EditAnywhere, Category = Options, meta = (UIMin = "0", UIMax = "1", ClampMin = "0", ClampMax = "1"))
+	float MeshAdaptivity = 0.001f;
+
+	/** Offset when remeshing, measured in voxels units */
+	UPROPERTY(EditAnywhere, Category = Options, meta = (UIMin = "-2", UIMax = "2", ClampMin = "-2", ClampMax = "2"))
+	float OffsetDistance = 0;
+
+	/** Automatically simplify the result of voxel-based merge.*/
+	UPROPERTY(EditAnywhere, Category = Options)
+	bool bAutoSimplify = false;
+
+	/** Remove the source Actors/Components when accepting results of tool.*/
+	UPROPERTY(EditAnywhere, Category = Options)
+	bool bDeleteInputActors = true;
+};
+
+
+
+
+/**
+ *
+ */
+UCLASS()
+class MESHMODELINGTOOLSEDITORONLY_API UVoxelCSGMeshesTool : public UMultiSelectionTool, public IDynamicMeshOperatorFactory
+{
+	GENERATED_BODY()
+
+public:
+	UVoxelCSGMeshesTool();
+
+	virtual void SetWorld(UWorld* World);
+	virtual void SetAssetAPI(IToolsContextAssetAPI* AssetAPI);
+
+	virtual void Setup() override;
+	virtual void Shutdown(EToolShutdownType ShutdownType) override;
+
+	virtual void Tick(float DeltaTime) override;
+	virtual void Render(IToolsContextRenderAPI* RenderAPI) override;
+
+	virtual bool HasCancel() const override { return true; }
+	virtual bool HasAccept() const override;
+	virtual bool CanAccept() const override;
+
+	virtual void OnPropertyModified(UObject* PropertySet, UProperty* Property) override;
+
+	// IDynamicMeshOperatorFactory API
+	virtual TSharedPtr<FDynamicMeshOperator> MakeNewOperator() override;
+
+
+protected:
+	UPROPERTY()
+	UVoxelCSGMeshesToolProperties* CSGProps;
+
+	UPROPERTY()
+	UMeshStatisticsProperties* MeshStatisticsProperties;
+
+	UPROPERTY()
+	UMeshOpPreviewWithBackgroundCompute* Preview;
+
+protected:
+	UWorld* TargetWorld;
+	IToolsContextAssetAPI* AssetAPI;
+
+	TSharedPtr<TArray<IVoxelBasedCSG::FPlacedMesh>> InputMeshes;
+	/** stash copies of the transforms and pointers to the meshes for consumption by the CSG Op*/
+	void CacheInputMeshes();
+
+	/** quickly generate a low-quality result for display while the actual result is being computed. */
+	void CreateLowQualityPreview();
+
+	void GenerateAsset(const TUniquePtr<FDynamicMeshOpResult>& Result);
+};
diff --git a/Engine/Plugins/Experimental/MeshModelingToolset/Source/ModelingComponents/ModelingComponents.Build.cs b/Engine/Plugins/Experimental/MeshModelingToolset/Source/ModelingComponents/ModelingComponents.Build.cs
new file mode 100644
--- /dev/null
+++ b/Engine/Plugins/Experimental/MeshModelingToolset/Source/ModelingComponents/ModelingComponents.Build.cs
@@ -0,0 +1,66 @@
+// Copyright 1998-2019 Epic Games, Inc. All Rights Reserved.
+
+using UnrealBuildTool;
+
+public class ModelingComponents : ModuleRules
+{
+	public ModelingComponents(ReadOnlyTargetRules Target) : base(Target)
+	{
+		PCHUsage = ModuleRules.PCHUsageMode.UseExplicitOrSharedPCHs;
+		
+		PublicIncludePaths.AddRange(
+			new string[] {
+				// ... add public include paths required here ...
+			}
+			);
+				
+		
+		PrivateIncludePaths.AddRange(
+			new string[] {
+				// ... add other private include paths required here ...
+			}
+			);
+			
+		
+		PublicDependencyModuleNames.AddRange(
+			new string[]
+			{
+				"Core",
+                "InteractiveToolsFramework",
+                "MeshDescription",
+				"StaticMeshDescription",
+				"GeometricObjects",
+				"GeometryAlgorithms",
+				"DynamicMesh",
+				"MeshConversion",
+				"ModelingOperators"
+				// ... add other public dependencies that you statically link with here ...
+			}
+            );
+			
+		
+		PrivateDependencyModuleNames.AddRange(
+			new string[]
+			{
+				"CoreUObject",
+				"Engine",
+                "RenderCore",
+                "RHI",
+
+				//"MeshUtilities",    // temp for saving mesh asset
+				//"UnrealEd"
+				//"Slate",
+				//"SlateCore",
+				// ... add private dependencies that you statically link with here ...	
+			}
+			);
+		
+		
+		DynamicallyLoadedModuleNames.AddRange(
+			new string[]
+			{
+				// ... add any modules that your module loads dynamically here ...
+			}
+			);
+	}
+}
diff --git a/Engine/Plugins/Experimental/MeshModelingToolset/Source/ModelingComponents/Private/AssetGenerationUtil.cpp b/Engine/Plugins/Experimental/MeshModelingToolset/Source/ModelingComponents/Private/AssetGenerationUtil.cpp
new file mode 100644
--- /dev/null
+++ b/Engine/Plugins/Experimental/MeshModelingToolset/Source/ModelingComponents/Private/AssetGenerationUtil.cpp
@@ -0,0 +1,88 @@
+// Copyright 1998-2019 Epic Games, Inc. All Rights Reserved.
+
+#include "AssetGenerationUtil.h"
+#include "InteractiveTool.h"
+#include "InteractiveToolManager.h"
+#include "Materials/Material.h"
+
+#include "Components/StaticMeshComponent.h"
+#include "Engine/StaticMesh.h"
+#include "Engine/StaticMeshActor.h"
+
+#include "StaticMeshComponentBuilder.h"
+#include "MeshDescriptionToDynamicMesh.h"
+#include "MeshDescriptionBuilder.h"
+
+
+FString AssetGenerationUtil::GetDefaultAutoGeneratedAssetPath()
+{
+	return FString("/Game/_GENERATED");
+}
+
+
+#if WITH_EDITOR
+AActor* AssetGenerationUtil::GenerateStaticMeshActor(
+	IToolsContextAssetAPI* AssetAPI,
+	UWorld* TargetWorld,
+	const FDynamicMesh3* Mesh,
+	const FTransform3d& Transform,
+	FString ObjectName,
+	FString PackagePath,
+	UMaterialInterface* Material
+)
+{
+	check(AssetAPI);
+	check(TargetWorld);
+	check(Mesh);
+
+	// asset name and location (temp)
+	FString PackageFolderPath = (PackagePath.IsEmpty()) ? GetDefaultAutoGeneratedAssetPath() : PackagePath;
+	//FString PackageFolderPath = AssetAPI->GetActiveAssetFolderPath();
+	FString MeshName = ObjectName;
+
+	// create new package
+	FString UniqueAssetName;
+	UPackage* AssetPackage = AssetAPI->MakeNewAssetPackage(PackageFolderPath, ObjectName, UniqueAssetName);
+
+	// create new actor
+	FRotator Rotation(0.0f, 0.0f, 0.0f);
+	FActorSpawnParameters SpawnInfo;
+	// @todo nothing here is specific to AStaticMeshActor...could we pass in a CDO and clone it instead of using SpawnActor?
+	AStaticMeshActor* NewActor = TargetWorld->SpawnActor<AStaticMeshActor>(FVector::ZeroVector, Rotation, SpawnInfo);
+	NewActor->SetActorLabel(*UniqueAssetName);
+
+	// construct new static mesh
+	FStaticMeshComponentBuilder Builder;
+	Builder.Initialize(AssetPackage, FName(*UniqueAssetName));
+
+	// Populate the MeshDescription associated with the asset.
+	FMeshDescriptionBuilder MeshDescripBuilder;
+	MeshDescripBuilder.SetMeshDescription(Builder.MeshDescription);
+	if (Mesh != nullptr)
+	{
+		MeshDescripBuilder.AppendMesh(Mesh, true);
+	}
+	else
+	{
+		// should generate default sphere here or something...
+	}
+
+	// create new mesh component and set as root of NewActor.
+	// NB: This should to be called only after the mesh has been created / populated,
+	// internally it calls  NewStaticMesh->CommitMeshDescription(0);
+	Builder.CreateAndSetAsRootComponent(NewActor);
+
+	// transform new component to origin
+	Builder.NewMeshComponent->SetWorldTransform((FTransform)Transform);
+
+	if (Material)
+	{
+		Builder.NewMeshComponent->SetMaterial(0, Material);
+	}
+	// force save of asset
+	//AssetAPI->InteractiveSaveGeneratedAsset(Builder.NewStaticMesh, AssetPackage);
+	AssetAPI->AutoSaveGeneratedAsset(Builder.NewStaticMesh, AssetPackage);
+
+	return NewActor;
+}
+#endif
diff --git a/Engine/Plugins/Experimental/MeshModelingToolset/Source/ModelingComponents/Private/BaseDynamicMeshComponent.cpp b/Engine/Plugins/Experimental/MeshModelingToolset/Source/ModelingComponents/Private/BaseDynamicMeshComponent.cpp
new file mode 100644
--- /dev/null
+++ b/Engine/Plugins/Experimental/MeshModelingToolset/Source/ModelingComponents/Private/BaseDynamicMeshComponent.cpp
@@ -0,0 +1,10 @@
+// Copyright 1998-2019 Epic Games, Inc. All Rights Reserved. 
+
+#include "BaseDynamicMeshComponent.h"
+
+
+
+UBaseDynamicMeshComponent::UBaseDynamicMeshComponent(const FObjectInitializer& ObjectInitializer)
+	: Super(ObjectInitializer)
+{
+}
diff --git a/Engine/Plugins/Experimental/MeshModelingToolset/Source/ModelingComponents/Private/Drawing/MeshDebugDrawing.cpp b/Engine/Plugins/Experimental/MeshModelingToolset/Source/ModelingComponents/Private/Drawing/MeshDebugDrawing.cpp
new file mode 100644
--- /dev/null
+++ b/Engine/Plugins/Experimental/MeshModelingToolset/Source/ModelingComponents/Private/Drawing/MeshDebugDrawing.cpp
@@ -0,0 +1,109 @@
+// Copyright 1998-2019 Epic Games, Inc. All Rights Reserved.
+
+#include "Drawing/MeshDebugDrawing.h"
+#include "DynamicMesh3.h"
+#include "FrameTypes.h"
+
+#include "SceneManagement.h" // FPrimitiveDrawInterface
+
+
+void MeshDebugDraw::DrawNormals(
+	const FDynamicMeshNormalOverlay* Overlay,
+	float Length, FColor Color, float Thickness, bool bScreenSpace,
+	FPrimitiveDrawInterface* PDI, const FTransform& Transform)
+{
+	const FDynamicMesh3* Mesh = Overlay->GetParentMesh();
+	for (int ElementID : Overlay->ElementIndicesItr())
+	{
+		FVector3f Normal = Overlay->GetElement(ElementID);
+		int ParentVID = Overlay->GetParentVertex(ElementID);
+		FVector3f ParentPos = (FVector3f)Mesh->GetVertex(ParentVID);
+
+		FVector A = ParentPos, B = ParentPos + Length * Normal;
+		PDI->DrawLine(Transform.TransformPosition(A), Transform.TransformPosition(B),
+			Color, 0, Thickness, 0, bScreenSpace);
+	}
+}
+
+
+
+
+
+void MeshDebugDraw::DrawVertices(
+	const FDynamicMesh3* Mesh, const TArray<int>& Indices,
+	float PointSize, FColor Color,
+	FPrimitiveDrawInterface* PDI, const FTransform& Transform)
+{
+	for (int VertID : Indices)
+	{
+		FVector3d Pos = Mesh->GetVertex(VertID);
+		PDI->DrawPoint(Transform.TransformPosition(Pos), Color, PointSize, SDPG_World);
+	}
+}
+
+void MeshDebugDraw::DrawVertices(
+	const FDynamicMesh3* Mesh, const TSet<int>& Indices,
+	float PointSize, FColor Color,
+	FPrimitiveDrawInterface* PDI, const FTransform& Transform)
+{
+	for (int VertID : Indices)
+	{
+		FVector3d Pos = Mesh->GetVertex(VertID);
+		PDI->DrawPoint(Transform.TransformPosition(Pos), Color, PointSize, SDPG_World);
+	}
+}
+
+
+
+void MeshDebugDraw::DrawTriCentroids(
+	const FDynamicMesh3* Mesh, const TArray<int>& Indices,
+	float PointSize, FColor Color,
+	FPrimitiveDrawInterface* PDI, const FTransform& Transform)
+{
+	for (int TriID : Indices)
+	{
+		FVector3d Pos = Mesh->GetTriCentroid(TriID);
+		PDI->DrawPoint(Transform.TransformPosition(Pos), Color, PointSize, SDPG_World);
+	}
+}
+
+
+
+void MeshDebugDraw::DrawSimpleGrid(	
+	const FFrame3f& LocalFrame, int GridLines, float GridLineSpacing,
+	float LineWidth, FColor Color, bool bDepthTested,
+	FPrimitiveDrawInterface* PDI, const FTransform& Transform)
+{
+	ESceneDepthPriorityGroup DepthPriority = (bDepthTested) ? SDPG_World : SDPG_Foreground;
+
+	FFrame3f WorldFrame = LocalFrame;
+	WorldFrame.Transform(Transform);
+
+	float Width = (float)(GridLines-1) * GridLineSpacing;
+	float Extent = Width * 0.5;
+
+	FVector3f Origin = WorldFrame.Origin;
+	FVector3f X = WorldFrame.X();
+	FVector3f Y = WorldFrame.Y();
+	FVector A, B;
+
+	int LineSteps = GridLines / 2;
+	for (int i = 0; i < LineSteps; i++)
+	{
+		float dx = (float)i * GridLineSpacing;
+		A = Origin - Extent * Y - dx * X;
+		B = Origin + Extent * Y - dx * X;
+		PDI->DrawLine(A, B, Color, DepthPriority, LineWidth, 0, true);
+		A = Origin - Extent * Y + dx * X;
+		B = Origin + Extent * Y + dx * X;
+		PDI->DrawLine(A, B, Color, DepthPriority, LineWidth, 0, true);
+
+		A = Origin - Extent * X - dx * Y;
+		B = Origin + Extent * X - dx * Y;
+		PDI->DrawLine(A, B, Color, DepthPriority, LineWidth, 0, true);
+		A = Origin - Extent * X + dx * Y;
+		B = Origin + Extent * X + dx * Y;
+		PDI->DrawLine(A, B, Color, DepthPriority, LineWidth, 0, true);
+	}
+
+}
diff --git a/Engine/Plugins/Experimental/MeshModelingToolset/Source/ModelingComponents/Private/MeshOpPreviewHelpers.cpp b/Engine/Plugins/Experimental/MeshModelingToolset/Source/ModelingComponents/Private/MeshOpPreviewHelpers.cpp
new file mode 100644
--- /dev/null
+++ b/Engine/Plugins/Experimental/MeshModelingToolset/Source/ModelingComponents/Private/MeshOpPreviewHelpers.cpp
@@ -0,0 +1,106 @@
+// Copyright 1998-2019 Epic Games, Inc. All Rights Reserved.
+
+#include "MeshOpPreviewHelpers.h"
+
+
+void UMeshOpPreviewWithBackgroundCompute::Setup(UWorld* InWorld, IDynamicMeshOperatorFactory* OpGenerator)
+{
+	PreviewMesh = NewObject<UPreviewMesh>(this, TEXT("PreviewMesh"));
+	PreviewMesh->CreateInWorld(InWorld, FTransform::Identity);
+
+	BackgroundCompute = MakeUnique<FBackgroundDynamicMeshComputeSource>(OpGenerator);
+	bResultValid = false;
+}
+
+TUniquePtr<FDynamicMeshOpResult> UMeshOpPreviewWithBackgroundCompute::Shutdown()
+{
+	BackgroundCompute->CancelActiveCompute();
+
+
+	TUniquePtr<FDynamicMeshOpResult> Result = MakeUnique<FDynamicMeshOpResult>();
+	Result->Mesh = PreviewMesh->ExtractPreviewMesh();
+	Result->Transform = FTransform3d(PreviewMesh->GetTransform());
+
+	PreviewMesh->SetVisible(false);
+	PreviewMesh->Disconnect();
+	PreviewMesh = nullptr;
+
+	return MoveTemp(Result);
+}
+
+
+void UMeshOpPreviewWithBackgroundCompute::Cancel()
+{
+	BackgroundCompute->CancelActiveCompute();
+
+	PreviewMesh->SetVisible(false);
+	PreviewMesh->Disconnect();
+	PreviewMesh = nullptr;
+}
+
+
+void UMeshOpPreviewWithBackgroundCompute::Tick(float DeltaTime)
+{
+	bool bIsLongDelay = false;
+	if (BackgroundCompute)
+	{
+		BackgroundCompute->Tick(DeltaTime);
+		bIsLongDelay = BackgroundCompute->GetElapsedComputeTime() > 2.0;
+	}
+
+	UpdateResults();
+
+	if (bResultValid || WorkingMaterial == nullptr || bIsLongDelay == false)
+	{
+		PreviewMesh->SetMaterial(StandardMaterial);
+	}
+	else
+	{
+		PreviewMesh->SetMaterial(WorkingMaterial);
+	}
+}
+
+
+void UMeshOpPreviewWithBackgroundCompute::UpdateResults()
+{
+	EBackgroundComputeTaskStatus Status = BackgroundCompute->CheckStatus();
+	if (Status == EBackgroundComputeTaskStatus::NewResultAvailable)
+	{
+		TSharedPtr<FDynamicMeshOperator> MeshOp =
+			StaticCastSharedPtr<FDynamicMeshOperator>(BackgroundCompute->ExtractResult());
+
+		TUniquePtr<FDynamicMesh3> ResultMesh = MeshOp->ExtractResult();
+		PreviewMesh->SetTransform((FTransform)MeshOp->GetResultTransform());
+		PreviewMesh->UpdatePreview(ResultMesh.Get());  // copies the mesh @todo we could just give ownership to the Preview!
+		PreviewMesh->SetVisible(bVisible);
+		bResultValid = true;
+
+		OnMeshUpdated.Broadcast(this);
+	}
+}
+
+
+void UMeshOpPreviewWithBackgroundCompute::InvalidateResult()
+{
+	BackgroundCompute->NotifyActiveComputeInvalidated();
+	bResultValid = false;
+}
+
+
+void UMeshOpPreviewWithBackgroundCompute::ConfigureMaterials(UMaterialInterface* StandardMaterialIn, UMaterialInterface* WorkingMaterialIn)
+{
+	this->StandardMaterial = StandardMaterialIn;
+	this->WorkingMaterial = WorkingMaterialIn;
+
+	if (PreviewMesh != nullptr)
+	{
+		PreviewMesh->SetMaterial(this->StandardMaterial);
+	}
+}
+
+
+void UMeshOpPreviewWithBackgroundCompute::SetVisibility(bool bVisibleIn)
+{
+	bVisible = bVisibleIn;
+	PreviewMesh->SetVisible(bVisible);
+}
diff --git a/Engine/Plugins/Experimental/MeshModelingToolset/Source/ModelingComponents/Private/ModelingComponentsModule.cpp b/Engine/Plugins/Experimental/MeshModelingToolset/Source/ModelingComponents/Private/ModelingComponentsModule.cpp
new file mode 100644
--- /dev/null
+++ b/Engine/Plugins/Experimental/MeshModelingToolset/Source/ModelingComponents/Private/ModelingComponentsModule.cpp
@@ -0,0 +1,20 @@
+// Copyright 1998-2019 Epic Games, Inc. All Rights Reserved.
+
+#include "ModelingComponentsModule.h"
+
+#define LOCTEXT_NAMESPACE "FModelingComponentsModule"
+
+void FModelingComponentsModule::StartupModule()
+{
+	// This code will execute after your module is loaded into memory; the exact timing is specified in the .uplugin file per-module
+}
+
+void FModelingComponentsModule::ShutdownModule()
+{
+	// This function may be called during shutdown to clean up your module.  For modules that support dynamic reloading,
+	// we call this function before unloading the module.
+}
+
+#undef LOCTEXT_NAMESPACE
+	
+IMPLEMENT_MODULE(FModelingComponentsModule, ModelingComponents)
\ No newline at end of file
