Copying //UE4/Dev-Editor to Dev-Main (//UE4/Dev-Main) for 4.24

diff --git a/Engine/Plugins/Experimental/MeshModelingToolset/Source/ModelingComponents/Private/OctreeDynamicMeshComponent.cpp b/Engine/Plugins/Experimental/MeshModelingToolset/Source/ModelingComponents/Private/OctreeDynamicMeshComponent.cpp
new file mode 100644
--- /dev/null
+++ b/Engine/Plugins/Experimental/MeshModelingToolset/Source/ModelingComponents/Private/OctreeDynamicMeshComponent.cpp
@@ -0,0 +1,378 @@
+// Copyright 1998-2019 Epic Games, Inc. All Rights Reserved. 
+
+#include "OctreeDynamicMeshComponent.h"
+#include "RenderingThread.h"
+#include "RenderResource.h"
+#include "PrimitiveViewRelevance.h"
+#include "PrimitiveSceneProxy.h"
+#include "VertexFactory.h"
+#include "MaterialShared.h"
+#include "Engine/CollisionProfile.h"
+#include "Materials/Material.h"
+#include "LocalVertexFactory.h"
+#include "SceneManagement.h"
+#include "DynamicMeshBuilder.h"
+#include "EngineGlobals.h"
+#include "Engine/Engine.h"
+#include "StaticMeshResources.h"
+#include "StaticMeshAttributes.h"
+
+#include "DynamicMeshAttributeSet.h"
+#include "MeshNormals.h"
+#include "MeshDescriptionToDynamicMesh.h"
+#include "DynamicMeshToMeshDescription.h"
+
+#include "Changes/MeshVertexChange.h"
+#include "Changes/MeshChange.h"
+#include "DynamicMeshChangeTracker.h"
+
+
+
+// default proxy for this component
+#include "OctreeDynamicMeshSceneProxy.h"
+
+
+
+
+UOctreeDynamicMeshComponent::UOctreeDynamicMeshComponent(const FObjectInitializer& ObjectInitializer)
+	: Super(ObjectInitializer)
+{
+	PrimaryComponentTick.bCanEverTick = false;
+
+	SetCollisionProfileName(UCollisionProfile::NoCollision_ProfileName);
+
+	InitializeNewMesh();
+}
+
+
+
+void UOctreeDynamicMeshComponent::InitializeMesh(FMeshDescription* MeshDescription)
+{
+	FMeshDescriptionToDynamicMesh Converter;
+	Converter.bPrintDebugMessages = true;
+	Mesh->Clear();
+	Converter.Convert(MeshDescription, *Mesh);
+
+	FAxisAlignedBox3d MeshBounds = Mesh->GetCachedBounds();
+	Octree = MakeUnique<FDynamicMeshOctree3>();
+	Octree->RootDimension = MeshBounds.MaxDim() * 0.5;
+
+	Octree->Initialize(GetMesh());
+	
+	FDynamicMeshOctree3::FStatistics Stats;
+	Octree->ComputeStatistics(Stats);
+	//UE_LOG(LogTemp, Warning, TEXT("OctreeStats %s"), *Stats.ToString());
+
+	OctreeCut = MakeUnique<FDynamicMeshOctree3::FTreeCutSet>();
+
+	NotifyMeshUpdated();
+}
+
+
+TUniquePtr<FDynamicMesh3> UOctreeDynamicMeshComponent::ExtractMesh(bool bNotifyUpdate)
+{
+	TUniquePtr<FDynamicMesh3> CurMesh = MoveTemp(Mesh);
+	InitializeNewMesh();
+	if (bNotifyUpdate)
+	{
+		NotifyMeshUpdated();
+	}
+	return CurMesh;
+}
+
+
+void UOctreeDynamicMeshComponent::InitializeNewMesh()
+{
+	Mesh = MakeUnique<FDynamicMesh3>();
+	// discard any attributes/etc initialized by default
+	Mesh->Clear();
+	Octree = MakeUnique<FDynamicMeshOctree3>();
+	Octree->Initialize(GetMesh());
+	OctreeCut = MakeUnique<FDynamicMeshOctree3::FTreeCutSet>();
+}
+
+
+void UOctreeDynamicMeshComponent::Bake(FMeshDescription* MeshDescription, bool bHaveModifiedTopology, const FConversionToMeshDescriptionOptions& ConversionOptions)
+{
+	if (bHaveModifiedTopology == false)
+	{
+		FDynamicMeshToMeshDescription Converter(ConversionOptions);
+		Converter.Update(Mesh.Get(), *MeshDescription);
+	}
+	else
+	{
+		FDynamicMeshToMeshDescription Converter(ConversionOptions);
+		Converter.Convert(Mesh.Get(), *MeshDescription);
+
+		//UE_LOG(LogTemp, Warning, TEXT("MeshDescription has %d instances"), MeshDescription->VertexInstances().Num());
+	}
+}
+
+
+
+void UOctreeDynamicMeshComponent::NotifyMeshUpdated()
+{
+	if (CurrentProxy != nullptr)
+	{
+		FAxisAlignedBox3d DirtyBox = Octree->ModifiedBounds;
+		Octree->ResetModifiedBounds();
+
+		// update existing cells
+		int TotalTriangleCount = 0;
+		int SpillTriangleCount = 0;
+		TArray<int32> SetsToUpdate;
+
+		{
+			SCOPE_CYCLE_COUNTER(STAT_SculptToolOctree_UpdateExisting);
+			FCriticalSection SetsLock;
+			int NumCutCells = CutCellSetMap.Num();
+			ParallelFor(NumCutCells, [&](int i) 
+			{
+				const FCutCellIndexSet& CutCellSet = CutCellSetMap[i];
+
+				if (Octree->TestCellIntersection(CutCellSet.CellRef, DirtyBox) == false)
+				{
+					return;
+				}
+
+				TArray<int32>& TriangleSet = TriangleDecomposition.GetIndexSetArray(CutCellSet.DecompSetID);
+				TriangleSet.Reset();
+				Octree->CollectTriangles(CutCellSet.CellRef, [&TriangleSet](int TriangleID) {
+					TriangleSet.Add(TriangleID);
+				});
+
+				SetsLock.Lock();
+				TotalTriangleCount += TriangleSet.Num();
+				SetsToUpdate.Add(CutCellSet.DecompSetID);
+				SetsLock.Unlock();
+			}, false);
+		}
+
+		// update cut set to find new cells
+		TArray<FDynamicMeshOctree3::FCellReference> NewCutCells;
+		{
+			SCOPE_CYCLE_COUNTER(STAT_SculptToolOctree_UpdateCutSet);
+			Octree->UpdateLevelCutSet(*OctreeCut, NewCutCells);
+		}
+
+		// add new ones
+		{
+			SCOPE_CYCLE_COUNTER(STAT_SculptToolOctree_CreateNew);
+			for (const FDynamicMeshOctree3::FCellReference& CellRef : NewCutCells)
+			{
+
+				int32 IndexSetID = TriangleDecomposition.CreateNewIndexSet();
+				TArray<int32>& TriangleSet = TriangleDecomposition.GetIndexSetArray(IndexSetID);
+				Octree->CollectTriangles(CellRef, [&TriangleSet](int TriangleID) {
+					TriangleSet.Add(TriangleID);
+				});
+				TotalTriangleCount += TriangleSet.Num();
+				FCutCellIndexSet SetMapEntry = { CellRef, IndexSetID };
+				CutCellSetMap.Add(SetMapEntry);
+				SetsToUpdate.Add(IndexSetID);
+			}
+		}
+
+		// rebuild spill set (always for now...should track bounds? and do separately for each root cell?)
+		{
+			SCOPE_CYCLE_COUNTER(STAT_SculptToolOctree_UpdateSpill);
+			TArray<int32>& SpillTriangleSet = TriangleDecomposition.GetIndexSetArray(SpillDecompSetID);
+			SpillTriangleSet.Reset();
+			Octree->CollectRootTriangles(*OctreeCut,
+				[&SpillTriangleSet](int TriangleID) {
+				SpillTriangleSet.Add(TriangleID);
+			});
+			Octree->CollectSpillTriangles(
+				[&SpillTriangleSet](int TriangleID) {
+				SpillTriangleSet.Add(TriangleID);
+			});
+			TotalTriangleCount += SpillTriangleSet.Num();
+			SpillTriangleCount += SpillTriangleSet.Num();
+			SetsToUpdate.Add(SpillDecompSetID);
+		}
+
+
+		//UE_LOG(LogTemp, Warning, TEXT("Updating %d of %d decomposition sets, %d tris total, spillcount %d"), SetsToUpdate.Num(), CutCellSetMap.Num(), TotalTriangleCount, SpillTriangleCount);
+		{
+			SCOPE_CYCLE_COUNTER(STAT_SculptToolOctree_UpdateFromDecomp);
+			CurrentProxy->UpdateFromDecomposition(TriangleDecomposition, SetsToUpdate);
+		}
+
+	}
+
+	//MarkRenderStateDirty();
+	//UpdateBounds();
+	//CurrentProxy = nullptr;
+}
+
+
+
+
+static void InitializeOctreeCutSet(const FDynamicMesh3& Mesh, const FDynamicMeshOctree3& Octree, TUniquePtr<FDynamicMeshOctree3::FTreeCutSet>& CutSet)
+{
+	int TriangleCount = Mesh.TriangleCount();
+	if (TriangleCount < 50000)
+	{
+		*CutSet = Octree.BuildLevelCutSet(1);
+		return;
+	}
+
+	FDynamicMeshOctree3::FStatistics Stats;
+	Octree.ComputeStatistics(Stats);
+	int MaxLevel = Stats.Levels;
+
+	int CutLevel = 0;
+	while ( CutLevel < MaxLevel-1 && Stats.LevelBoxCounts[CutLevel] < 200 && Stats.LevelBoxCounts[CutLevel+1] < 300)
+	{
+		CutLevel++;
+	}
+	*CutSet = Octree.BuildLevelCutSet(CutLevel);
+	//UE_LOG(LogTemp, Warning, TEXT("InitializeOctreeCutSet - level %d cells %d"), CutLevel, CutSet->CutCells.Num());
+}
+
+
+
+
+FPrimitiveSceneProxy* UOctreeDynamicMeshComponent::CreateSceneProxy()
+{
+	CurrentProxy = nullptr;
+	if (Mesh->TriangleCount() > 0)
+	{
+		CurrentProxy = new FOctreeDynamicMeshSceneProxy(this);
+
+		if (TriangleColorFunc != nullptr)
+		{
+			CurrentProxy->bUsePerTriangleColor = true;
+			CurrentProxy->PerTriangleColorFunc = [this](int TriangleID) { return GetTriangleColor(TriangleID); };
+		}
+
+
+		OctreeCut = MakeUnique<FDynamicMeshOctree3::FTreeCutSet>();
+		InitializeOctreeCutSet(*Mesh, *Octree, OctreeCut);
+
+		TriangleDecomposition = FArrayIndexSetsDecomposition();
+
+		SpillDecompSetID = TriangleDecomposition.CreateNewIndexSet();
+
+		CutCellSetMap.Reset();
+		for (auto CellRef : OctreeCut->CutCells)
+		{
+			int32 IndexSetID = TriangleDecomposition.CreateNewIndexSet();
+			TArray<int32>& TriangleSet = TriangleDecomposition.GetIndexSetArray(IndexSetID);
+			Octree->CollectTriangles(CellRef, [&TriangleSet](int TriangleID) {
+				TriangleSet.Add(TriangleID);
+			});
+
+			FCutCellIndexSet SetMapEntry = { CellRef, IndexSetID };
+			CutCellSetMap.Add(SetMapEntry);
+		}
+
+		// collect spill triangles
+		{
+			TArray<int32>& SpillTriangleSet = TriangleDecomposition.GetIndexSetArray(SpillDecompSetID);
+			Octree->CollectRootTriangles(*OctreeCut, 
+				[&SpillTriangleSet](int TriangleID) {
+					SpillTriangleSet.Add(TriangleID);
+			});
+			Octree->CollectSpillTriangles(
+				[&SpillTriangleSet](int TriangleID) {
+				SpillTriangleSet.Add(TriangleID);
+			});
+		}
+
+		CurrentProxy->InitializeFromDecomposition(TriangleDecomposition);
+	}
+	return CurrentProxy;
+}
+
+int32 UOctreeDynamicMeshComponent::GetNumMaterials() const
+{
+	return 1;
+}
+
+
+FColor UOctreeDynamicMeshComponent::GetTriangleColor(int TriangleID)
+{
+	if (TriangleColorFunc != nullptr)
+	{
+		return TriangleColorFunc(TriangleID);
+	}
+	else
+	{
+		return (TriangleID % 2 == 0) ? FColor::Red : FColor::White;
+	}
+}
+
+
+
+FBoxSphereBounds UOctreeDynamicMeshComponent::CalcBounds(const FTransform& LocalToWorld) const
+{
+	// Bounds are tighter if the box is generated from pre-transformed vertices.
+	FBox BoundingBox(ForceInit);
+	for ( FVector3d Vertex : Mesh->VerticesItr() ) 
+	{
+		BoundingBox += LocalToWorld.TransformPosition(Vertex);
+	}
+
+	FBoxSphereBounds NewBounds;
+	NewBounds.BoxExtent = BoundingBox.GetExtent();
+	NewBounds.Origin = BoundingBox.GetCenter();
+	NewBounds.SphereRadius = NewBounds.BoxExtent.Size();
+
+	return NewBounds;
+}
+
+
+
+void UOctreeDynamicMeshComponent::ApplyChange(const FMeshVertexChange* Change, bool bRevert)
+{
+	int NV = Change->Vertices.Num();
+	const TArray<FVector3d>& Positions = (bRevert) ? Change->OldPositions : Change->NewPositions;
+	
+	Octree->ResetModifiedBounds();
+	TSet<int> TrianglesToUpdate;
+
+	for (int k = 0; k < NV; ++k)
+	{
+		int vid = Change->Vertices[k];
+
+		for (int tid : Mesh->VtxTrianglesItr(vid))
+		{
+			if (TrianglesToUpdate.Contains(tid) == false)
+			{
+				Octree->NotifyPendingModification(tid);
+				TrianglesToUpdate.Add(tid);
+			}
+		}
+
+		Mesh->SetVertex(vid, Positions[k]);
+	}
+
+	Octree->ReinsertTriangles(TrianglesToUpdate);
+
+	//NotifyMeshUpdated();
+
+	OnMeshChanged.Broadcast();
+}
+
+
+
+
+void UOctreeDynamicMeshComponent::ApplyChange(const FMeshChange* Change, bool bRevert)
+{
+	TArray<int> RemoveTriangles, AddTriangles;
+	bool bRemoveOld = ! bRevert;
+	Change->DynamicMeshChange->GetSavedTriangleList(RemoveTriangles, bRemoveOld);
+	Change->DynamicMeshChange->GetSavedTriangleList(AddTriangles, !bRemoveOld);
+
+	Octree->ResetModifiedBounds();
+	Octree->RemoveTriangles(RemoveTriangles);
+
+	Change->DynamicMeshChange->Apply(Mesh.Get(), bRevert);
+
+	Octree->InsertTriangles(AddTriangles);
+
+	//NotifyMeshUpdated();
+
+	OnMeshChanged.Broadcast();
+}
\ No newline at end of file
diff --git a/Engine/Plugins/Experimental/MeshModelingToolset/Source/ModelingComponents/Private/OctreeDynamicMeshSceneProxy.h b/Engine/Plugins/Experimental/MeshModelingToolset/Source/ModelingComponents/Private/OctreeDynamicMeshSceneProxy.h
new file mode 100644
--- /dev/null
+++ b/Engine/Plugins/Experimental/MeshModelingToolset/Source/ModelingComponents/Private/OctreeDynamicMeshSceneProxy.h
@@ -0,0 +1,540 @@
+// Copyright 1998-2019 Epic Games, Inc. All Rights Reserved. 
+
+#pragma once
+
+#include "OctreeDynamicMeshComponent.h"
+#include "Util/IndexSetDecompositions.h"
+
+
+DECLARE_STATS_GROUP(TEXT("SculptToolOctree"), STATGROUP_SculptToolOctree, STATCAT_Advanced);
+DECLARE_CYCLE_STAT(TEXT("SculptToolOctree_UpdateExisting"), STAT_SculptToolOctree_UpdateExisting, STATGROUP_SculptToolOctree);
+DECLARE_CYCLE_STAT(TEXT("SculptToolOctree_UpdateCutSet"), STAT_SculptToolOctree_UpdateCutSet, STATGROUP_SculptToolOctree);
+DECLARE_CYCLE_STAT(TEXT("SculptToolOctree_CreateNew"), STAT_SculptToolOctree_CreateNew, STATGROUP_SculptToolOctree);
+DECLARE_CYCLE_STAT(TEXT("SculptToolOctree_UpdateSpill"), STAT_SculptToolOctree_UpdateSpill, STATGROUP_SculptToolOctree);
+DECLARE_CYCLE_STAT(TEXT("SculptToolOctree_UpdateFromDecomp"), STAT_SculptToolOctree_UpdateFromDecomp, STATGROUP_SculptToolOctree);
+DECLARE_CYCLE_STAT(TEXT("SculptToolOctree_UpdateDecompDestroy"), STAT_SculptToolOctree_UpdateDecompDestroy, STATGROUP_SculptToolOctree);
+DECLARE_CYCLE_STAT(TEXT("SculptToolOctree_UpdateDecompCreate"), STAT_SculptToolOctree_UpdateDecompCreate, STATGROUP_SculptToolOctree);
+DECLARE_CYCLE_STAT(TEXT("SculptToolOctree_InitializeBufferFromOverlay"), STAT_SculptToolOctree_InitializeBufferFromOverlay, STATGROUP_SculptToolOctree);
+DECLARE_CYCLE_STAT(TEXT("SculptToolOctree_BufferUpload"), STAT_SculptToolOctree_BufferUpload, STATGROUP_SculptToolOctree);
+
+class FMeshRenderBufferSet
+{
+public:
+	int TriangleCount = 0;
+
+	/** The buffer containing vertex data. */
+	FStaticMeshVertexBuffer StaticMeshVertexBuffer;
+	/** The buffer containing the position vertex data. */
+	FPositionVertexBuffer PositionVertexBuffer;
+	/** The buffer containing the vertex color data. */
+	FColorVertexBuffer ColorVertexBuffer;
+
+	FDynamicMeshIndexBuffer32 IndexBuffer;
+	FLocalVertexFactory VertexFactory;
+
+	FMeshRenderBufferSet(ERHIFeatureLevel::Type FeatureLevelType)
+		: VertexFactory(FeatureLevelType, "FMeshRenderBufferSet")
+	{
+
+	}
+
+
+	virtual ~FMeshRenderBufferSet()
+	{
+		check(IsInRenderingThread());
+
+		if (TriangleCount > 0)
+		{
+			PositionVertexBuffer.ReleaseResource();
+			StaticMeshVertexBuffer.ReleaseResource();
+			ColorVertexBuffer.ReleaseResource();
+			IndexBuffer.ReleaseResource();
+			VertexFactory.ReleaseResource();
+		}
+	}
+
+
+	static void DestroyRenderBufferSet(FMeshRenderBufferSet* BufferSet)
+	{
+		if (BufferSet->TriangleCount == 0)
+		{
+			return;
+		}
+
+		ENQUEUE_RENDER_COMMAND(FMeshRenderBufferSetDestroy)(
+			[BufferSet](FRHICommandListImmediate& RHICmdList)
+		{
+			delete BufferSet;
+		});
+	}
+
+	void Upload()
+	{
+		check(IsInRenderingThread());
+
+		if (TriangleCount == 0)
+		{
+			return;
+		}
+
+		InitOrUpdateResource(&this->PositionVertexBuffer);
+		InitOrUpdateResource(&this->StaticMeshVertexBuffer);
+		InitOrUpdateResource(&this->ColorVertexBuffer);
+
+		FLocalVertexFactory::FDataType Data;
+		this->PositionVertexBuffer.BindPositionVertexBuffer(&this->VertexFactory, Data);
+		this->StaticMeshVertexBuffer.BindTangentVertexBuffer(&this->VertexFactory, Data);
+		this->StaticMeshVertexBuffer.BindPackedTexCoordVertexBuffer(&this->VertexFactory, Data);
+		//this->StaticMeshVertexBuffer.BindLightMapVertexBuffer(&this->VertexFactory, Data, LightMapIndex);
+		this->ColorVertexBuffer.BindColorVertexBuffer(&this->VertexFactory, Data);
+		this->VertexFactory.SetData(Data);
+
+		InitOrUpdateResource(&this->VertexFactory);
+		PositionVertexBuffer.InitResource();
+		StaticMeshVertexBuffer.InitResource();
+		ColorVertexBuffer.InitResource();
+		IndexBuffer.InitResource();
+		VertexFactory.InitResource();
+	}
+
+	// copied from StaticMesh.cpp
+	void InitOrUpdateResource(FRenderResource* Resource)
+	{
+		check(IsInRenderingThread());
+
+		if (!Resource->IsInitialized())
+		{
+			Resource->InitResource();
+		}
+		else
+		{
+			Resource->UpdateRHI();
+		}
+	}
+
+};
+
+/**
+ * Scene Proxy for a mesh buffer.
+ * 
+ * Based on FProceduralMeshSceneProxy but simplified in various ways.
+ * 
+ * Supports wireframe-on-shaded rendering.
+ * 
+ */
+class FOctreeDynamicMeshSceneProxy final : public FPrimitiveSceneProxy
+{
+private:
+	UMaterialInterface* Material;
+
+	FMaterialRelevance MaterialRelevance;
+
+	TMap<int32, FMeshRenderBufferSet*> RenderBufferSets;
+
+public:
+	/** Component that created this proxy (is there a way to look this up?) */
+	UOctreeDynamicMeshComponent* ParentComponent;
+
+
+	FColor ConstantVertexColor = FColor::White;
+	bool bIgnoreVertexColors = false;
+	bool bIgnoreVertexNormals = false;
+
+
+	bool bUsePerTriangleColor = false;
+	TFunction<FColor(int)> PerTriangleColorFunc = nullptr;
+
+
+	FOctreeDynamicMeshSceneProxy(UOctreeDynamicMeshComponent* Component)
+		: FPrimitiveSceneProxy(Component)
+		, MaterialRelevance(Component->GetMaterialRelevance(GetScene().GetFeatureLevel()))
+	{
+		check(!IsInRenderingThread());
+
+		ParentComponent = Component;
+
+		// Grab material
+		Material = Component->GetMaterial(0);
+		if (Material == NULL)
+		{
+			Material = UMaterial::GetDefaultMaterial(MD_Surface);
+		}
+	}
+
+
+	virtual ~FOctreeDynamicMeshSceneProxy()
+	{
+		check(IsInRenderingThread());
+
+		for (auto MapPair : RenderBufferSets)
+		{
+			FMeshRenderBufferSet* BufferSet = MapPair.Value;
+			FMeshRenderBufferSet::DestroyRenderBufferSet(BufferSet);
+		}
+		RenderBufferSets.Reset();
+	}
+
+
+	FMeshRenderBufferSet* AllocateNewRenderBufferSet()
+	{
+		FMeshRenderBufferSet* RenderBufferSet = new FMeshRenderBufferSet(GetScene().GetFeatureLevel());
+		return RenderBufferSet;
+	}
+
+
+
+	virtual void InitializeSingleBuffer()
+	{
+			check(RenderBufferSets.Num() == 0);
+
+			FDynamicMesh3* Mesh = ParentComponent->GetMesh();
+
+			FMeshRenderBufferSet* RenderBuffers = AllocateNewRenderBufferSet();
+
+			// find suitable overlays
+			FDynamicMeshUVOverlay* UVOverlay = nullptr;
+			FDynamicMeshNormalOverlay* NormalOverlay = nullptr;
+			if (Mesh->HasAttributes())
+			{
+				UVOverlay = Mesh->Attributes()->PrimaryUV();
+				NormalOverlay = Mesh->Attributes()->PrimaryNormals();
+			}
+
+			InitializeBuffersFromOverlays(Mesh,
+				Mesh->TriangleCount(), Mesh->TriangleIndicesItr(),
+				UVOverlay, NormalOverlay, RenderBuffers);
+
+			ENQUEUE_RENDER_COMMAND(FOctreeDynamicMeshSceneProxyInitializeSingle)(
+				[this, RenderBuffers](FRHICommandListImmediate& RHICmdList)
+			{
+				RenderBuffers->Upload();
+				RenderBufferSets.Add(0, RenderBuffers);
+			});
+	}
+
+
+
+	virtual void InitializeFromDecomposition(const FArrayIndexSetsDecomposition& Decomposition)
+	{
+		check(RenderBufferSets.Num() == 0);
+
+		FDynamicMesh3* Mesh = ParentComponent->GetMesh();
+
+		// find suitable overlays
+		FDynamicMeshUVOverlay* UVOverlay = nullptr;
+		FDynamicMeshNormalOverlay* NormalOverlay = nullptr;
+		if (Mesh->HasAttributes())
+		{
+			UVOverlay = Mesh->Attributes()->PrimaryUV();
+			NormalOverlay = Mesh->Attributes()->PrimaryNormals();
+		}
+
+		const TArray<int32>& SetIDs = Decomposition.GetIndexSetIDs();
+		for (int32 SetID : SetIDs)
+		{
+			const TArray<int32>& Tris = Decomposition.GetIndexSetArray(SetID);
+			
+			FMeshRenderBufferSet* RenderBuffers = AllocateNewRenderBufferSet();
+
+			InitializeBuffersFromOverlays(Mesh,
+				Tris.Num(), Tris,
+				UVOverlay, NormalOverlay, RenderBuffers);
+
+			ENQUEUE_RENDER_COMMAND(FOctreeDynamicMeshSceneProxyInitializeFromDecomposition)(
+				[this, SetID, RenderBuffers](FRHICommandListImmediate& RHICmdList)
+			{
+				RenderBuffers->Upload();
+				RenderBufferSets.Add(SetID, RenderBuffers);
+			});
+		}
+	}
+
+
+
+	virtual void UpdateFromDecomposition(const FArrayIndexSetsDecomposition& Decomposition, const TArray<int32>& SetsToUpdate )
+	{
+		// CAN WE REUSE EXISTING BUFFER SETS??
+		//   - could have timestamp for each decomposition set array...if tris don't change we only have to update vertices
+		//   - can re-use allocated memory if new data is smaller
+
+		SCOPE_CYCLE_COUNTER(STAT_SculptToolOctree_UpdateFromDecomp);
+
+		// remove sets to update
+		ENQUEUE_RENDER_COMMAND(FOctreeDynamicMeshSceneProxyUpdatePreClean)(
+			[this, SetsToUpdate](FRHICommandListImmediate& RHICmdList)
+		{
+			SCOPE_CYCLE_COUNTER(STAT_SculptToolOctree_UpdateDecompDestroy);
+			for (int32 SetID : SetsToUpdate)
+			{
+				if (RenderBufferSets.Contains(SetID))
+				{
+					FMeshRenderBufferSet* BufferSet = RenderBufferSets.FindAndRemoveChecked(SetID);
+					FMeshRenderBufferSet::DestroyRenderBufferSet(BufferSet);
+				}
+			}
+		});
+
+		FDynamicMesh3* Mesh = ParentComponent->GetMesh();
+
+		// find suitable overlays
+		FDynamicMeshUVOverlay* UVOverlay = nullptr;
+		FDynamicMeshNormalOverlay* NormalOverlay = nullptr;
+		if (Mesh->HasAttributes())
+		{
+			UVOverlay = Mesh->Attributes()->PrimaryUV();
+			NormalOverlay = Mesh->Attributes()->PrimaryNormals();
+		}
+
+		{
+			FCriticalSection SectionLock;
+
+			SCOPE_CYCLE_COUNTER(STAT_SculptToolOctree_UpdateDecompCreate);
+			int NumSets = SetsToUpdate.Num();
+			ParallelFor(NumSets, [&](int k)
+			{
+				int32 SetID = SetsToUpdate[k];
+				const TArray<int32>& Tris = Decomposition.GetIndexSetArray(SetID);
+
+				FMeshRenderBufferSet* RenderBuffers = AllocateNewRenderBufferSet();
+
+				InitializeBuffersFromOverlays(Mesh,
+					Tris.Num(), Tris,
+					UVOverlay, NormalOverlay, RenderBuffers);
+
+				ENQUEUE_RENDER_COMMAND(FOctreeDynamicMeshSceneProxyUpdateAddOne)(
+					[this, SetID, RenderBuffers](FRHICommandListImmediate& RHICmdList)
+				{
+					SCOPE_CYCLE_COUNTER(STAT_SculptToolOctree_BufferUpload);
+					RenderBuffers->Upload();
+					RenderBufferSets.Add(SetID, RenderBuffers);
+				});
+			}, false);
+
+
+		}
+
+		//UE_LOG(LogTemp, Warning, TEXT("Have %d renderbuffers"), RenderBufferSets.Num());
+	}
+
+
+
+protected:
+
+
+	/**
+	 * Initialize rendering buffers from given attribute overlays.
+	 * Creates three vertices per triangle, IE no shared vertices in buffers.
+	 */
+	template<typename TriangleEnumerable>
+	void InitializeBuffersFromOverlays(const FDynamicMesh3* Mesh, 
+		int NumTriangles, TriangleEnumerable Enumerable,
+		FDynamicMeshUVOverlay* UVOverlay, 
+		FDynamicMeshNormalOverlay* NormalOverlay,
+		FMeshRenderBufferSet* RenderBuffers)
+	{
+		SCOPE_CYCLE_COUNTER(STAT_SculptToolOctree_InitializeBufferFromOverlay);
+
+		RenderBuffers->TriangleCount = NumTriangles;
+		if (NumTriangles == 0)
+		{
+			return;
+		}
+
+		bool bHaveColors = Mesh->HasVertexColors() && (bIgnoreVertexColors == false);
+
+		int NumVertices = NumTriangles * 3;
+		int NumTexCoords = 1;		// no! zero!
+
+		{
+			RenderBuffers->PositionVertexBuffer.Init(NumVertices);
+			RenderBuffers->StaticMeshVertexBuffer.Init(NumVertices, NumTexCoords);
+			RenderBuffers->ColorVertexBuffer.Init(NumVertices);
+			RenderBuffers->IndexBuffer.Indices.AddUninitialized(NumTriangles * 3);
+		}
+
+		int TriIdx = 0, VertIdx = 0;
+		FVector3f TangentX, TangentY;
+		for (int TriangleID : Enumerable)
+		{
+			FIndex3i Tri = Mesh->GetTriangle(TriangleID);
+			FIndex3i TriUV = (UVOverlay != nullptr) ? UVOverlay->GetTriangle(TriangleID) : FIndex3i::Zero();
+			FIndex3i TriNormal = (NormalOverlay != nullptr) ? NormalOverlay->GetTriangle(TriangleID) : FIndex3i::Zero();
+
+			FColor TriColor = ConstantVertexColor;
+			if (bUsePerTriangleColor && PerTriangleColorFunc != nullptr)
+			{
+				TriColor = PerTriangleColorFunc(TriangleID);
+				bHaveColors = false;
+			}
+
+			for (int j = 0; j < 3; ++j)
+			{
+				RenderBuffers->PositionVertexBuffer.VertexPosition(VertIdx) = (FVector)Mesh->GetVertex(Tri[j]);
+
+				FVector3f Normal = (NormalOverlay != nullptr && TriNormal[j] != FDynamicMesh3::InvalidID) ? 
+					NormalOverlay->GetElement(TriNormal[j]) : Mesh->GetVertexNormal(Tri[j]);
+
+				// calculate a nonsense tangent
+				VectorUtil::MakePerpVectors(Normal, TangentX, TangentY);
+				RenderBuffers->StaticMeshVertexBuffer.SetVertexTangents(VertIdx, TangentX, TangentY, Normal);
+
+				FVector2f UV = (UVOverlay != nullptr && TriUV[j] != FDynamicMesh3::InvalidID) ? 
+					UVOverlay->GetElement(TriUV[j]) : FVector2f::Zero();
+				RenderBuffers->StaticMeshVertexBuffer.SetVertexUV(VertIdx, 0, UV);
+
+				RenderBuffers->ColorVertexBuffer.VertexColor(VertIdx) = (bHaveColors) ?
+					(FColor)Mesh->GetVertexColor(Tri[j]) : TriColor;
+
+				RenderBuffers->IndexBuffer.Indices[TriIdx++] = VertIdx;
+				VertIdx++;
+			}
+		}
+	}
+
+
+
+
+
+public:
+
+	virtual void GetDynamicMeshElements(const TArray<const FSceneView*>& Views, const FSceneViewFamily& ViewFamily, uint32 VisibilityMap, FMeshElementCollector& Collector) const override
+	{
+		QUICK_SCOPE_CYCLE_COUNTER(STAT_OctreeDynamicMeshSceneProxy_GetDynamicMeshElements);
+
+		const bool bWireframe = (AllowDebugViewmodes() && ViewFamily.EngineShowFlags.Wireframe)
+			|| ParentComponent->bExplicitShowWireframe;
+
+		FColoredMaterialRenderProxy* WireframeMaterialInstance = new FColoredMaterialRenderProxy(
+			GEngine->WireframeMaterial ? GEngine->WireframeMaterial->GetRenderProxy() : nullptr,
+			FLinearColor(0, 0.5f, 1.f)
+		);
+
+		Collector.RegisterOneFrameMaterialProxy(WireframeMaterialInstance);
+
+		FMaterialRenderProxy* MaterialProxy = Material->GetRenderProxy();
+		FMaterialRenderProxy* WireframeMaterialProxy = WireframeMaterialInstance;
+
+		//ESceneDepthPriorityGroup DepthPriority = (ParentComponent->bDrawOnTop) ? SDPG_Foreground : SDPG_World;
+		ESceneDepthPriorityGroup DepthPriority = SDPG_World;
+
+
+		for (int32 ViewIndex = 0; ViewIndex < Views.Num(); ViewIndex++)
+		{
+			if (VisibilityMap & (1 << ViewIndex))
+			{
+				const FSceneView* View = Views[ViewIndex];
+
+				bool bHasPrecomputedVolumetricLightmap;
+				FMatrix PreviousLocalToWorld;
+				int32 SingleCaptureIndex;
+				bool bOutputVelocity;
+				GetScene().GetPrimitiveUniformShaderParameters_RenderThread(GetPrimitiveSceneInfo(), bHasPrecomputedVolumetricLightmap, PreviousLocalToWorld, SingleCaptureIndex, bOutputVelocity);
+
+				// Draw the mesh.
+				for (auto MapPair : RenderBufferSets)
+				{
+					const FMeshRenderBufferSet& RenderBuffers = *MapPair.Value;
+					if (RenderBuffers.TriangleCount == 0)
+					{
+						continue;
+					}
+
+					// do we need separate one of these for each MeshRenderBufferSet?
+					FDynamicPrimitiveUniformBuffer& DynamicPrimitiveUniformBuffer = Collector.AllocateOneFrameResource<FDynamicPrimitiveUniformBuffer>();
+					DynamicPrimitiveUniformBuffer.Set(GetLocalToWorld(), PreviousLocalToWorld, GetBounds(), GetLocalBounds(), true, bHasPrecomputedVolumetricLightmap, DrawsVelocity(), bOutputVelocity);
+
+					DrawBatch(Collector, RenderBuffers, MaterialProxy, false, DepthPriority, ViewIndex, DynamicPrimitiveUniformBuffer);
+					if (bWireframe)
+					{
+						DrawBatch(Collector, RenderBuffers, WireframeMaterialProxy, true, DepthPriority, ViewIndex, DynamicPrimitiveUniformBuffer);
+					}
+				}
+			}
+		}
+	}
+
+
+
+	virtual void DrawBatch(FMeshElementCollector& Collector, 
+		const FMeshRenderBufferSet& RenderBuffers,
+		FMaterialRenderProxy* UseMaterial, 
+		bool bWireframe,
+		ESceneDepthPriorityGroup DepthPriority,
+		int ViewIndex,
+		FDynamicPrimitiveUniformBuffer& DynamicPrimitiveUniformBuffer) const
+	{
+		FMeshBatch& Mesh = Collector.AllocateMesh();
+		FMeshBatchElement& BatchElement = Mesh.Elements[0];
+		BatchElement.IndexBuffer = &RenderBuffers.IndexBuffer;
+		Mesh.bWireframe = bWireframe;
+		Mesh.VertexFactory = &RenderBuffers.VertexFactory;
+		Mesh.MaterialRenderProxy = UseMaterial;
+
+		BatchElement.PrimitiveUniformBufferResource = &DynamicPrimitiveUniformBuffer.UniformBuffer;
+
+		BatchElement.FirstIndex = 0;
+		BatchElement.NumPrimitives = RenderBuffers.IndexBuffer.Indices.Num() / 3;
+		BatchElement.MinVertexIndex = 0;
+		BatchElement.MaxVertexIndex = RenderBuffers.PositionVertexBuffer.GetNumVertices() - 1;
+		Mesh.ReverseCulling = IsLocalToWorldDeterminantNegative();
+		Mesh.Type = PT_TriangleList;
+		Mesh.DepthPriorityGroup = DepthPriority;
+		Mesh.bCanApplyViewModeOverrides = false;
+		Collector.AddMesh(ViewIndex, Mesh);
+	}
+
+
+
+	virtual FPrimitiveViewRelevance GetViewRelevance(const FSceneView* View) const override
+	{
+		FPrimitiveViewRelevance Result;
+
+		Result.bDrawRelevance = IsShown(View);
+		Result.bShadowRelevance = IsShadowCast(View);
+		Result.bDynamicRelevance = true;
+		Result.bRenderInMainPass = ShouldRenderInMainPass();
+		Result.bUsesLightingChannels = GetLightingChannelMask() != GetDefaultLightingChannelMask();
+		Result.bTranslucentSelfShadow = bCastVolumetricTranslucentShadow;
+		Result.bRenderCustomDepth = ShouldRenderCustomDepth();
+		MaterialRelevance.SetPrimitiveViewRelevance(Result);
+		Result.bVelocityRelevance = IsMovable() && Result.bOpaqueRelevance && Result.bRenderInMainPass;
+
+		return Result;
+	}
+
+
+	virtual void GetLightRelevance(const FLightSceneProxy* LightSceneProxy, bool& bDynamic, bool& bRelevant, bool& bLightMapped, bool& bShadowMapped) const override
+	{
+		FPrimitiveSceneProxy::GetLightRelevance(LightSceneProxy, bDynamic, bRelevant, bLightMapped, bShadowMapped);
+	}
+
+
+	virtual bool CanBeOccluded() const override
+	{
+		return !MaterialRelevance.bDisableDepthTest;
+	}
+
+	virtual uint32 GetMemoryFootprint(void) const override { return(sizeof(*this) + GetAllocatedSize()); }
+
+	uint32 GetAllocatedSize(void) const { return(FPrimitiveSceneProxy::GetAllocatedSize()); }
+
+
+	virtual void SetMaterial(UMaterialInterface* MaterialIn)
+	{
+		this->Material = MaterialIn;
+	}
+
+
+	SIZE_T GetTypeHash() const override
+	{
+		static size_t UniquePointer;
+		return reinterpret_cast<size_t>(&UniquePointer);
+	}
+};
+
+
+
+
+
+
+
diff --git a/Engine/Plugins/Experimental/MeshModelingToolset/Source/ModelingComponents/Private/Selection/GroupTopologySelector.cpp b/Engine/Plugins/Experimental/MeshModelingToolset/Source/ModelingComponents/Private/Selection/GroupTopologySelector.cpp
new file mode 100644
--- /dev/null
+++ b/Engine/Plugins/Experimental/MeshModelingToolset/Source/ModelingComponents/Private/Selection/GroupTopologySelector.cpp
@@ -0,0 +1,255 @@
+// Copyright 1998-2019 Epic Games, Inc. All Rights Reserved.
+
+
+#include "Selection/GroupTopologySelector.h"
+#include "MeshQueries.h"
+#include "Drawing/ToolDataVisualizer.h"
+#include "ToolSceneQueriesUtil.h"
+
+
+
+void FGroupTopologySelection::Clear()
+{
+	SelectedGroupIDs.Reset();
+	SelectedCornerIDs.Reset();
+	SelectedEdgeIDs.Reset();
+}
+
+
+
+FGroupTopologySelector::FGroupTopologySelector()
+{
+	// initialize to sane values
+	PointsWithinToleranceTest =
+		[](const FVector3d& A, const FVector3d& B) { return A.Distance(B) < 1.0; };
+	GetSpatial =
+		[]() { return nullptr; };
+}
+
+
+
+void FGroupTopologySelector::Initialize(const FDynamicMesh3* MeshIn, const FGroupTopology* TopologyIn)
+{
+	Mesh = MeshIn;
+	Topology = TopologyIn;
+	bGeometryInitialized = false;
+	bGeometryUpToDate = false;
+}
+
+
+void FGroupTopologySelector::Invalidate(bool bTopologyDeformed, bool bTopologyModified)
+{
+	if (bTopologyDeformed)
+	{
+		bGeometryUpToDate = false;
+	}
+	if (bTopologyModified)
+	{
+		bGeometryUpToDate = bGeometryInitialized = false;
+	}
+}
+
+
+const FGeometrySet3& FGroupTopologySelector::GetGeometrySet()
+{
+	if (bGeometryInitialized == false)
+	{
+		GeometrySet.Reset();
+		int NumCorners = Topology->Corners.Num();
+		for (int CornerID = 0; CornerID < NumCorners; ++CornerID)
+		{
+			FVector3d Position = Mesh->GetVertex(Topology->Corners[CornerID].VertexID);
+			GeometrySet.AddPoint(CornerID, Position);
+		}
+		int NumEdges = Topology->Edges.Num();
+		for (int EdgeID = 0; EdgeID < NumEdges; ++EdgeID)
+		{
+			FPolyline3d Polyline;
+			Topology->Edges[EdgeID].Span.GetPolyline(Polyline);
+			GeometrySet.AddCurve(EdgeID, Polyline);
+		}
+
+		bGeometryInitialized = true;
+		bGeometryUpToDate = true;
+	}
+
+	if (bGeometryUpToDate == false)
+	{
+		int NumCorners = Topology->Corners.Num();
+		for (int CornerID = 0; CornerID < NumCorners; ++CornerID)
+		{
+			FVector3d Position = Mesh->GetVertex(Topology->Corners[CornerID].VertexID);
+			GeometrySet.UpdatePoint(CornerID, Position);
+		}
+		int NumEdges = Topology->Edges.Num();
+		for (int EdgeID = 0; EdgeID < NumEdges; ++EdgeID)
+		{
+			FPolyline3d Polyline;
+			Topology->Edges[EdgeID].Span.GetPolyline(Polyline);
+			GeometrySet.UpdateCurve(EdgeID, Polyline);
+		}
+		bGeometryUpToDate = true;
+	}
+
+	return GeometrySet;
+}
+
+
+void FGroupTopologySelector::UpdateEnableFlags(bool bFaceHits, bool bEdgeHits, bool bCornerHits)
+{
+	bEnableFaceHits = bFaceHits;
+	bEnableEdgeHits = bEdgeHits;
+	bEnableCornerHits = bCornerHits;
+}
+
+
+bool FGroupTopologySelector::FindSelectedElement(const FRay3d& Ray, FGroupTopologySelection& ResultOut,
+	FVector3d& SelectedPositionOut, FVector3d& SelectedNormalOut)
+{
+	FDynamicMeshAABBTree3* Spatial = GetSpatial();
+	const FGeometrySet3& TopoSpatial = GetGeometrySet();
+	
+	double TriangleHitRayParam = TNumericLimits<double>::Max();
+	FVector3d TriangleHitPos;
+
+	int HitTriangleID = (Spatial != nullptr) ? Spatial->FindNearestHitTriangle(Ray) : IndexConstants::InvalidID;
+	if (HitTriangleID != IndexConstants::InvalidID)
+	{
+		FTriangle3d Triangle;
+		Mesh->GetTriVertices(HitTriangleID, Triangle.V[0], Triangle.V[1], Triangle.V[2]);
+		FIntrRay3Triangle3d Query(Ray, Triangle);
+		Query.Find();
+
+		TriangleHitRayParam = Query.RayParameter;
+		TriangleHitPos = Ray.PointAt(Query.RayParameter);
+	}
+	bool bHitSurface = bEnableFaceHits && (HitTriangleID >= 0);
+	bool bActuallyHitSurface = (HitTriangleID >= 0);
+
+
+	int HitCornerID = -1;
+	double HitCornerRayParam = TNumericLimits<double>::Max();
+	FVector3d HitCornerPt, OnCornerPt;
+	if (bEnableCornerHits)
+	{
+		FGeometrySet3::FNearest Nearest;
+		bool bFound = TopoSpatial.FindNearestPointToRay(Ray, Nearest, PointsWithinToleranceTest);
+		if (bFound)
+		{
+			HitCornerID = Nearest.ID;
+			HitCornerRayParam = Nearest.RayParam;
+			HitCornerPt = Nearest.NearestRayPoint;
+			OnCornerPt = Nearest.NearestGeoPoint;
+		}
+	}
+	bool bHitCorner = (HitCornerID >= 0);
+
+	// @todo currently implicitly assuming that corner is better than edge, 
+	// but we should be checking the ray-param...edge could be much closer!
+
+	int HitEdgeID = -1;
+	double HitEdgeRayParam = TNumericLimits<double>::Max();
+	FVector3d HitEdgePt, OnEdgePt;
+	if (bEnableEdgeHits && bHitCorner == false)		// EDGE HIT DISABLED IF CORNER HIT
+	{
+		FGeometrySet3::FNearest Nearest;
+		bool bFound = TopoSpatial.FindNearestCurveToRay(Ray, Nearest, PointsWithinToleranceTest);
+		if (bFound)
+		{
+			HitEdgeID = Nearest.ID;
+			HitEdgeRayParam = Nearest.RayParam;
+			HitEdgePt = Nearest.NearestRayPoint;
+			OnEdgePt = Nearest.NearestGeoPoint;
+		}
+	}
+	bool bHitEdge = (HitEdgeID >= 0);
+
+
+	if (bHitCorner == false && bHitEdge == false && bHitSurface == false)
+	{
+		return false;
+	}
+
+
+	// can only have one hit
+	if ((bHitCorner || bHitEdge) && bHitSurface)
+	{
+		// to determine if the edge/corner we "hit" is visible, we cast a ray back
+		// towards the eye from that hit point. If it hits anything we can't see it
+		// and so should take the face selection instead.
+
+		FVector3d ElementHitPt = (bHitCorner) ? OnCornerPt : OnEdgePt;
+		FRay3d ToEyeRay(ElementHitPt, (Ray.Origin - ElementHitPt).Normalized(), true);
+		ToEyeRay.Origin += (double)(100*FMathf::ZeroTolerance) * ToEyeRay.Direction;
+		int ToEyeHitTID = Spatial->FindNearestHitTriangle(ToEyeRay);
+		if (ToEyeHitTID >= 0)
+		{
+			bHitCorner = bHitEdge = false;
+		}
+		else
+		{
+			bHitSurface = false;
+		}
+	}
+
+
+	if (bHitCorner)
+	{
+		ResultOut.SelectedCornerIDs.Add(HitCornerID);
+		SelectedPositionOut = HitCornerPt;
+		SelectedNormalOut = (bActuallyHitSurface) ?
+			Mesh->GetTriNormal(HitTriangleID) : FVector3d::UnitZ();
+		return true;
+	}
+	else if (bHitEdge)
+	{
+		ResultOut.SelectedEdgeIDs.Add(HitEdgeID);
+		SelectedPositionOut = HitEdgePt;
+		SelectedNormalOut = (bActuallyHitSurface) ?
+			Mesh->GetTriNormal(HitTriangleID) : FVector3d::UnitZ();
+		return true;
+	}
+	else if (bHitSurface)
+	{
+		ResultOut.SelectedGroupIDs.Add(Mesh->GetTriangleGroup(HitTriangleID));
+		SelectedPositionOut = TriangleHitPos;
+		SelectedNormalOut = Mesh->GetTriNormal(HitTriangleID);
+		return true;
+	}
+
+	return false;
+}
+
+
+
+void FGroupTopologySelector::DrawSelection(const FGroupTopologySelection& Selection, FToolDataVisualizer* Renderer, const FViewCameraState* CameraState)
+{
+	for (int CornerID : Selection.SelectedCornerIDs)
+	{
+		FLinearColor UseColor = Renderer->LineColor;
+		float LineWidth = Renderer->LineThickness;
+
+		int VertexID = Topology->GetCornerVertexID(CornerID);
+		FVector Position = (FVector)Mesh->GetVertex(VertexID);
+		FVector WorldPosition = Renderer->TransformP(Position);
+		float Radius = (float)ToolSceneQueriesUtil::CalculateDimensionFromVisualAngleD(*CameraState, WorldPosition, 0.5);
+		Renderer->DrawViewFacingCircle(Position, Radius, 16, UseColor, LineWidth, false);
+	}
+
+	for (int EdgeID : Selection.SelectedEdgeIDs)
+	{
+		FLinearColor UseColor = Renderer->LineColor;
+		float LineWidth = Renderer->LineThickness;
+		const TArray<int>& Vertices = Topology->GetGroupEdgeVertices(EdgeID);
+		int NV = Vertices.Num() - 1;
+		FVector A = (FVector)Mesh->GetVertex(Vertices[0]);
+		for (int k = 0; k < NV; ++k)
+		{
+			FVector B = (FVector)Mesh->GetVertex(Vertices[k+1]);
+			Renderer->DrawLine(A, B, UseColor, LineWidth, false);
+			A = B;
+		}
+	}
+
+
+}
\ No newline at end of file
diff --git a/Engine/Plugins/Experimental/MeshModelingToolset/Source/ModelingComponents/Private/SimpleDynamicMeshComponent.cpp b/Engine/Plugins/Experimental/MeshModelingToolset/Source/ModelingComponents/Private/SimpleDynamicMeshComponent.cpp
new file mode 100644
--- /dev/null
+++ b/Engine/Plugins/Experimental/MeshModelingToolset/Source/ModelingComponents/Private/SimpleDynamicMeshComponent.cpp
@@ -0,0 +1,269 @@
+// Copyright 1998-2019 Epic Games, Inc. All Rights Reserved. 
+
+#include "SimpleDynamicMeshComponent.h"
+#include "RenderingThread.h"
+#include "RenderResource.h"
+#include "PrimitiveViewRelevance.h"
+#include "PrimitiveSceneProxy.h"
+#include "VertexFactory.h"
+#include "MaterialShared.h"
+#include "Engine/CollisionProfile.h"
+#include "Materials/Material.h"
+#include "LocalVertexFactory.h"
+#include "SceneManagement.h"
+#include "DynamicMeshBuilder.h"
+#include "EngineGlobals.h"
+#include "Engine/Engine.h"
+#include "StaticMeshResources.h"
+
+#include "StaticMeshAttributes.h"
+
+#include "DynamicMeshAttributeSet.h"
+#include "MeshNormals.h"
+#include "MeshDescriptionToDynamicMesh.h"
+#include "DynamicMeshToMeshDescription.h"
+
+#include "Changes/MeshVertexChange.h"
+#include "Changes/MeshChange.h"
+#include "DynamicMeshChangeTracker.h"
+
+
+
+// default proxy for this component
+#include "SimpleDynamicMeshSceneProxy.h"
+
+
+
+
+USimpleDynamicMeshComponent::USimpleDynamicMeshComponent(const FObjectInitializer& ObjectInitializer)
+	: Super(ObjectInitializer)
+{
+	PrimaryComponentTick.bCanEverTick = false;
+
+	SetCollisionProfileName(UCollisionProfile::NoCollision_ProfileName);
+
+	InitializeNewMesh();
+}
+
+
+
+void USimpleDynamicMeshComponent::InitializeMesh(FMeshDescription* MeshDescription)
+{
+	FMeshDescriptionToDynamicMesh Converter;
+	Converter.bPrintDebugMessages = true;
+	Mesh->Clear();
+	Converter.Convert(MeshDescription, *Mesh);
+	if (TangentsType == EDynamicMeshTangentCalcType::ExternallyCalculated)
+	{
+		Converter.CopyTangents(MeshDescription, Mesh.Get(), Tangents);
+	}
+
+	NotifyMeshUpdated();
+}
+
+
+TUniquePtr<FDynamicMesh3> USimpleDynamicMeshComponent::ExtractMesh(bool bNotifyUpdate)
+{
+	TUniquePtr<FDynamicMesh3> CurMesh = MoveTemp(Mesh);
+	InitializeNewMesh();
+	if (bNotifyUpdate)
+	{
+		NotifyMeshUpdated();
+	}
+	return CurMesh;
+}
+
+
+void USimpleDynamicMeshComponent::InitializeNewMesh()
+{
+	Mesh = MakeUnique<FDynamicMesh3>();
+	// discard any attributes/etc initialized by default
+	Mesh->Clear();
+
+	Tangents.SetMesh(Mesh.Get());
+}
+
+
+void USimpleDynamicMeshComponent::Bake(FMeshDescription* MeshDescription, bool bHaveModifiedTopology, const FConversionToMeshDescriptionOptions& ConversionOptions)
+{
+	if (bHaveModifiedTopology == false)
+	{
+		FDynamicMeshToMeshDescription Converter(ConversionOptions);
+		Converter.Update(Mesh.Get(), *MeshDescription);
+	}
+	else
+	{
+		FDynamicMeshToMeshDescription Converter(ConversionOptions);
+		Converter.Convert(Mesh.Get(), *MeshDescription);
+
+		//UE_LOG(LogTemp, Warning, TEXT("MeshDescription has %d instances"), MeshDescription->VertexInstances().Num());
+	}
+}
+
+
+
+
+
+
+FMeshTangentsf* USimpleDynamicMeshComponent::GetTangents()
+{
+	if (TangentsType == EDynamicMeshTangentCalcType::NoTangents)
+	{
+		return nullptr;
+	}
+	
+	if (TangentsType == EDynamicMeshTangentCalcType::AutoCalculated)
+	{
+		if (bTangentsValid == false && Mesh->HasAttributes())
+		{
+			FDynamicMeshUVOverlay* UVOverlay = Mesh->Attributes()->PrimaryUV();
+			FDynamicMeshNormalOverlay* NormalOverlay = Mesh->Attributes()->PrimaryNormals();
+			if (UVOverlay != nullptr && NormalOverlay != nullptr)
+			{
+				Tangents.ComputePerTriangleTangents(NormalOverlay, UVOverlay);
+				bTangentsValid = true;
+			}
+		}
+		return (bTangentsValid) ? &Tangents : nullptr;
+	}
+
+	// in this mode we assume the tangents are valid
+	check(TangentsType == EDynamicMeshTangentCalcType::ExternallyCalculated);
+	return &Tangents;
+}
+
+
+
+void USimpleDynamicMeshComponent::SetDrawOnTop(bool bSet)
+{
+	bDrawOnTop = bSet;
+	bUseEditorCompositing = bSet;
+}
+
+
+void USimpleDynamicMeshComponent::NotifyMeshUpdated()
+{
+	// Need to recreate scene proxy to send it over
+	MarkRenderStateDirty();
+	UpdateBounds();
+	CurrentProxy = nullptr;
+
+	if (TangentsType != EDynamicMeshTangentCalcType::ExternallyCalculated)
+	{
+		bTangentsValid = false;
+	}
+}
+
+
+void USimpleDynamicMeshComponent::FastNotifyColorsUpdated()
+{
+	if (CurrentProxy != nullptr)
+	{
+		CurrentProxy->FastUpdateColors();
+	}
+	else
+	{
+		NotifyMeshUpdated();
+	}
+}
+
+
+
+void USimpleDynamicMeshComponent::FastNotifyPositionsUpdated()
+{
+	if (CurrentProxy != nullptr)
+	{
+		bTangentsValid = false;
+		CurrentProxy->FastUpdatePositions(false);
+	}
+	else
+	{
+		NotifyMeshUpdated();
+	}
+}
+
+
+
+
+FPrimitiveSceneProxy* USimpleDynamicMeshComponent::CreateSceneProxy()
+{
+	CurrentProxy = nullptr;
+	if (Mesh->TriangleCount() > 0)
+	{
+		CurrentProxy = new FSimpleDynamicMeshSceneProxy(this);
+
+		if (TriangleColorFunc)
+		{
+			CurrentProxy->bUsePerTriangleColor = true;
+			CurrentProxy->PerTriangleColorFunc = [this](int TriangleID) { return GetTriangleColor(TriangleID); };
+		}
+
+		CurrentProxy->Initialize();
+	}
+	return CurrentProxy;
+}
+
+int32 USimpleDynamicMeshComponent::GetNumMaterials() const
+{
+	return 1;
+}
+
+
+FColor USimpleDynamicMeshComponent::GetTriangleColor(int TriangleID)
+{
+	if (TriangleColorFunc)
+	{
+		return TriangleColorFunc(TriangleID);
+	}
+	else
+	{
+		return (TriangleID % 2 == 0) ? FColor::Red : FColor::White;
+	}
+}
+
+
+
+FBoxSphereBounds USimpleDynamicMeshComponent::CalcBounds(const FTransform& LocalToWorld) const
+{
+	// Bounds are tighter if the box is generated from pre-transformed vertices.
+	FBox BoundingBox(ForceInit);
+	for ( FVector3d Vertex : Mesh->VerticesItr() ) 
+	{
+		BoundingBox += LocalToWorld.TransformPosition(Vertex);
+	}
+
+	FBoxSphereBounds NewBounds;
+	NewBounds.BoxExtent = BoundingBox.GetExtent();
+	NewBounds.Origin = BoundingBox.GetCenter();
+	NewBounds.SphereRadius = NewBounds.BoxExtent.Size();
+
+	return NewBounds;
+}
+
+
+
+void USimpleDynamicMeshComponent::ApplyChange(const FMeshVertexChange* Change, bool bRevert)
+{
+	int NV = Change->Vertices.Num();
+	const TArray<FVector3d>& Positions = (bRevert) ? Change->OldPositions : Change->NewPositions;
+	
+	for (int k = 0; k < NV; ++k)
+	{
+		int vid = Change->Vertices[k];
+		Mesh->SetVertex(vid, Positions[k]);
+	}
+
+	NotifyMeshUpdated();
+	OnMeshChanged.Broadcast();
+}
+
+
+
+
+void USimpleDynamicMeshComponent::ApplyChange(const FMeshChange* Change, bool bRevert)
+{
+	Change->DynamicMeshChange->Apply(Mesh.Get(), bRevert);
+
+	NotifyMeshUpdated();
+	OnMeshChanged.Broadcast();
+}
\ No newline at end of file
diff --git a/Engine/Plugins/Experimental/MeshModelingToolset/Source/ModelingComponents/Private/SimpleDynamicMeshSceneProxy.h b/Engine/Plugins/Experimental/MeshModelingToolset/Source/ModelingComponents/Private/SimpleDynamicMeshSceneProxy.h
new file mode 100644
--- /dev/null
+++ b/Engine/Plugins/Experimental/MeshModelingToolset/Source/ModelingComponents/Private/SimpleDynamicMeshSceneProxy.h
@@ -0,0 +1,539 @@
+// Copyright 1998-2019 Epic Games, Inc. All Rights Reserved. 
+
+#pragma once
+
+#include "SimpleDynamicMeshComponent.h"
+
+
+/**
+ * Scene Proxy for USimpleDynamicMeshComponent.
+ * 
+ * Based on FProceduralMeshSceneProxy but simplified in various ways.
+ * 
+ * Supports wireframe-on-shaded rendering.
+ * 
+ */
+class FSimpleDynamicMeshSceneProxy final : public FPrimitiveSceneProxy
+{
+private:
+	UMaterialInterface* Material;
+
+	/** The buffer containing vertex data. */
+	FStaticMeshVertexBuffer StaticMeshVertexBuffer;
+	/** The buffer containing the position vertex data. */
+	FPositionVertexBuffer PositionVertexBuffer;
+	/** The buffer containing the vertex color data. */
+	FColorVertexBuffer ColorVertexBuffer;
+
+	FDynamicMeshIndexBuffer32 IndexBuffer;
+	FLocalVertexFactory VertexFactory;
+
+	FMaterialRelevance MaterialRelevance;
+
+
+
+public:
+	/** Component that created this proxy (is there a way to look this up?) */
+	USimpleDynamicMeshComponent* ParentComponent;
+
+
+	FColor ConstantVertexColor = FColor::White;
+	bool bIgnoreVertexColors = false;
+	bool bIgnoreVertexNormals = false;
+
+
+	bool bUsePerTriangleColor = false;
+	TFunction<FColor(int)> PerTriangleColorFunc;
+
+
+	FSimpleDynamicMeshSceneProxy(USimpleDynamicMeshComponent* Component)
+		: FPrimitiveSceneProxy(Component)
+		, VertexFactory(GetScene().GetFeatureLevel(), "FSimpleDynamicMeshSceneProxy")
+		, MaterialRelevance(Component->GetMaterialRelevance(GetScene().GetFeatureLevel()))
+	{
+		ParentComponent = Component;
+
+		// Grab material
+		Material = Component->GetMaterial(0);
+		if (Material == NULL)
+		{
+			Material = UMaterial::GetDefaultMaterial(MD_Surface);
+		}
+	}
+
+
+	virtual ~FSimpleDynamicMeshSceneProxy()
+	{
+		PositionVertexBuffer.ReleaseResource();
+		StaticMeshVertexBuffer.ReleaseResource();
+		ColorVertexBuffer.ReleaseResource();
+		IndexBuffer.ReleaseResource();
+		VertexFactory.ReleaseResource();
+	}
+
+
+
+
+
+	virtual void Initialize()
+	{
+		FDynamicMesh3* Mesh = ParentComponent->GetMesh();
+
+		// find suitable overlays
+		FDynamicMeshUVOverlay* UVOverlay = nullptr;
+		FDynamicMeshNormalOverlay* NormalOverlay = nullptr;
+		if (Mesh->HasAttributes())
+		{
+			UVOverlay = Mesh->Attributes()->PrimaryUV();
+			NormalOverlay = Mesh->Attributes()->PrimaryNormals();
+		}
+
+		// fill our buffers
+		if (UVOverlay != nullptr || NormalOverlay != nullptr)
+		{
+			InitializeBuffersFromOverlays(Mesh, UVOverlay, NormalOverlay);
+		}
+		else
+		{
+			InitializeBuffersFromMesh_SharedVertices(Mesh);
+		}
+
+		// transfer to render thread?
+		// copied from FStaticMeshVertexBuffers::InitFromDynamicVertex
+		int LightMapIndex = 0;
+		ENQUEUE_RENDER_COMMAND(DynamicMeshComponentVertexBuffersInit)(
+			[this, LightMapIndex](FRHICommandListImmediate& RHICmdList)
+		{
+			InitOrUpdateResource(&this->PositionVertexBuffer);
+			InitOrUpdateResource(&this->StaticMeshVertexBuffer);
+			InitOrUpdateResource(&this->ColorVertexBuffer);
+
+			FLocalVertexFactory::FDataType Data;
+			this->PositionVertexBuffer.BindPositionVertexBuffer(&this->VertexFactory, Data);
+			this->StaticMeshVertexBuffer.BindTangentVertexBuffer(&this->VertexFactory, Data);
+			this->StaticMeshVertexBuffer.BindPackedTexCoordVertexBuffer(&this->VertexFactory, Data);
+			this->StaticMeshVertexBuffer.BindLightMapVertexBuffer(&this->VertexFactory, Data, LightMapIndex);
+			this->ColorVertexBuffer.BindColorVertexBuffer(&this->VertexFactory, Data);
+			this->VertexFactory.SetData(Data);
+
+			InitOrUpdateResource(&this->VertexFactory);
+		});
+
+		// Enqueue initialization of render resource
+		BeginInitResource(&PositionVertexBuffer);
+		BeginInitResource(&StaticMeshVertexBuffer);
+		BeginInitResource(&ColorVertexBuffer);
+		BeginInitResource(&IndexBuffer);
+		BeginInitResource(&VertexFactory);
+	}
+
+
+
+
+
+
+	virtual void FastUpdatePositions(bool bNormals)
+	{
+		check(bNormals == false);
+
+		FDynamicMesh3* Mesh = ParentComponent->GetMesh();
+		int NumVertices = Mesh->TriangleCount() * 3;
+		check(PositionVertexBuffer.GetNumVertices() == NumVertices);
+
+		int VertIdx = 0;
+		for (int TriangleID : Mesh->TriangleIndicesItr())
+		{
+			FIndex3i Tri = Mesh->GetTriangle(TriangleID);
+			for (int j = 0; j < 3; ++j)
+			{
+				PositionVertexBuffer.VertexPosition(VertIdx++) = Mesh->GetVertex(Tri[j]);
+			}
+		}
+
+		// @todo is it necessary to rebind everything? does this re-upload the buffers?
+		int LightMapIndex = 0;
+		ENQUEUE_RENDER_COMMAND(DynamicMeshComponentFastUpdatePositions)(
+			[this, LightMapIndex](FRHICommandListImmediate& RHICmdList)
+		{
+			InitOrUpdateResource(&this->PositionVertexBuffer);
+
+			FLocalVertexFactory::FDataType Data;
+			this->PositionVertexBuffer.BindPositionVertexBuffer(&this->VertexFactory, Data);
+			this->StaticMeshVertexBuffer.BindTangentVertexBuffer(&this->VertexFactory, Data);
+			this->StaticMeshVertexBuffer.BindPackedTexCoordVertexBuffer(&this->VertexFactory, Data);
+			this->StaticMeshVertexBuffer.BindLightMapVertexBuffer(&this->VertexFactory, Data, LightMapIndex);
+			this->ColorVertexBuffer.BindColorVertexBuffer(&this->VertexFactory, Data);
+			this->VertexFactory.SetData(Data);
+
+			InitOrUpdateResource(&this->VertexFactory);
+		});
+
+		// don't need to do this because they aren't new buffers??
+		//BeginInitResource(&PositionVertexBuffer);
+		//BeginInitResource(&VertexFactory);
+	}
+
+
+
+
+
+
+	virtual void FastUpdateColors()
+	{
+		FDynamicMesh3* Mesh = ParentComponent->GetMesh();
+		int NumVertices = Mesh->TriangleCount() * 3;
+		check(ColorVertexBuffer.GetNumVertices() == NumVertices);
+
+		bool bHaveColors = Mesh->HasVertexColors() && (bIgnoreVertexColors == false);
+
+		int VertIdx = 0;
+		for (int TriangleID : Mesh->TriangleIndicesItr())
+		{
+			FIndex3i Tri = Mesh->GetTriangle(TriangleID);
+
+			FColor TriColor = ConstantVertexColor;
+			if (bUsePerTriangleColor && PerTriangleColorFunc)
+			{
+				TriColor = PerTriangleColorFunc(TriangleID);
+				bHaveColors = false;
+			}
+
+			for (int j = 0; j < 3; ++j)
+			{
+				ColorVertexBuffer.VertexColor(VertIdx) = (bHaveColors) ?
+					(FColor)Mesh->GetVertexColor(Tri[j]) : TriColor;
+				VertIdx++;
+			}
+		}
+
+		// @todo is it necessary to rebind everything? does this re-upload the buffers?
+		int LightMapIndex = 0;
+		ENQUEUE_RENDER_COMMAND(DynamicMeshComponentFastUpdateColors)(
+			[this, LightMapIndex](FRHICommandListImmediate& RHICmdList)
+		{
+			InitOrUpdateResource(&this->ColorVertexBuffer);
+
+			FLocalVertexFactory::FDataType Data;
+			this->PositionVertexBuffer.BindPositionVertexBuffer(&this->VertexFactory, Data);
+			this->StaticMeshVertexBuffer.BindTangentVertexBuffer(&this->VertexFactory, Data);
+			this->StaticMeshVertexBuffer.BindPackedTexCoordVertexBuffer(&this->VertexFactory, Data);
+			this->StaticMeshVertexBuffer.BindLightMapVertexBuffer(&this->VertexFactory, Data, LightMapIndex);
+			this->ColorVertexBuffer.BindColorVertexBuffer(&this->VertexFactory, Data);
+			this->VertexFactory.SetData(Data);
+
+			InitOrUpdateResource(&this->VertexFactory);
+		});
+
+		BeginInitResource(&ColorVertexBuffer);		// calls ENQUEUE_RENDER_COMMAND
+		//BeginInitResource(&VertexFactory);
+	}
+
+
+
+protected:
+
+
+	/**
+	 * Initialize rendering buffers from given attribute overlays.
+	 * Creates three vertices per triangle, IE no shared vertices in buffers.
+	 */
+	void InitializeBuffersFromOverlays(FDynamicMesh3* Mesh, FDynamicMeshUVOverlay* UVOverlay, FDynamicMeshNormalOverlay* NormalOverlay)
+	{
+		// compute normals if we need to
+		FMeshNormals Normals(Mesh);
+		bool bUseComputedNormals = false;
+		if (NormalOverlay == nullptr && (Mesh->HasVertexNormals() == false || bIgnoreVertexNormals) )
+		{
+			Normals.ComputeVertexNormals();
+			bUseComputedNormals = true;
+		}
+
+		bool bHaveColors = Mesh->HasVertexColors() && (bIgnoreVertexColors == false);
+
+		// this may be null
+		FMeshTangentsf* Tangents = ParentComponent->GetTangents();
+
+
+		int NumVertices = Mesh->TriangleCount() * 3;
+		int NumTexCoords = 1;		// no! zero!
+		PositionVertexBuffer.Init(NumVertices);
+		StaticMeshVertexBuffer.Init(NumVertices, NumTexCoords);
+		ColorVertexBuffer.Init(NumVertices);
+
+		IndexBuffer.Indices.AddUninitialized(Mesh->TriangleCount() * 3);
+		int TriIdx = 0, VertIdx = 0;
+		FVector3f TangentX, TangentY;
+		for (int TriangleID : Mesh->TriangleIndicesItr())
+		{
+			FIndex3i Tri = Mesh->GetTriangle(TriangleID);
+			FIndex3i TriUV = (UVOverlay != nullptr) ? UVOverlay->GetTriangle(TriangleID) : FIndex3i::Zero();
+			FIndex3i TriNormal = (NormalOverlay != nullptr) ? NormalOverlay->GetTriangle(TriangleID) : FIndex3i::Zero();
+
+			FColor TriColor = ConstantVertexColor;
+			if (bUsePerTriangleColor && PerTriangleColorFunc)
+			{
+				TriColor = PerTriangleColorFunc(TriangleID);
+				bHaveColors = false;
+			}
+
+			for (int j = 0; j < 3; ++j)
+			{
+				PositionVertexBuffer.VertexPosition(VertIdx) = Mesh->GetVertex(Tri[j]);
+
+				FVector3f Normal = (NormalOverlay != nullptr && TriNormal[j] != FDynamicMesh3::InvalidID) ? NormalOverlay->GetElement(TriNormal[j])
+					: ((bUseComputedNormals) ? (FVector3f)Normals[Tri[j]] : Mesh->GetVertexNormal(Tri[j]));
+
+				if (Tangents == nullptr)
+				{
+					// calculate a nonsense tangent
+					VectorUtil::MakePerpVectors(Normal, TangentX, TangentY);
+					StaticMeshVertexBuffer.SetVertexTangents(VertIdx, TangentX, TangentY, Normal);
+				}
+				else
+				{
+					Tangents->GetPerTriangleTangent(TriangleID, j, TangentX, TangentY);
+				}
+
+				StaticMeshVertexBuffer.SetVertexTangents(VertIdx, TangentX, TangentY, Normal);
+
+				FVector2f UV = (UVOverlay != nullptr && TriUV[j] != FDynamicMesh3::InvalidID) ? UVOverlay->GetElement(TriUV[j]) : FVector2f::Zero();
+				StaticMeshVertexBuffer.SetVertexUV(VertIdx, 0, UV);
+
+				ColorVertexBuffer.VertexColor(VertIdx) = (bHaveColors) ?
+					(FColor)Mesh->GetVertexColor(Tri[j]) : TriColor;
+
+				IndexBuffer.Indices[TriIdx++] = VertIdx;
+				VertIdx++;
+			}
+		}
+	}
+
+
+
+
+
+	/**
+	 * Initialize rendering buffers from DMesh using shared vertices
+	 * (ie no attribute overlays)
+	 */
+	void InitializeBuffersFromMesh_SharedVertices(FDynamicMesh3* Mesh)
+	{
+		// compute normals if we need to
+		FMeshNormals Normals(Mesh);
+		bool bUseComputedNormals = false;
+		if (Mesh->HasVertexNormals() == false || bIgnoreVertexNormals)
+		{
+			Normals.ComputeVertexNormals();
+			bUseComputedNormals = true;
+		}
+
+		bool bHaveColors = Mesh->HasVertexColors() && (bIgnoreVertexColors == false);
+
+		int NumVertices = Mesh->MaxVertexID();
+		int NumTexCoords = 1;		// no! zero!
+		PositionVertexBuffer.Init(NumVertices);
+		StaticMeshVertexBuffer.Init(NumVertices, NumTexCoords);
+		ColorVertexBuffer.Init(NumVertices);
+
+		for (int VertIdx : Mesh->VertexIndicesItr())
+		{
+			PositionVertexBuffer.VertexPosition(VertIdx) = Mesh->GetVertex(VertIdx);
+
+			FVector3d Normal = (bUseComputedNormals) ? Normals[VertIdx] : (FVector3d)Mesh->GetVertexNormal(VertIdx);
+			FVector3d TangentX, TangentY;
+			VectorUtil::MakePerpVectors(Normal, TangentX, TangentY);
+
+			StaticMeshVertexBuffer.SetVertexTangents(VertIdx, TangentX, TangentY, Normal);
+			StaticMeshVertexBuffer.SetVertexUV(VertIdx, 0, FVector2D::ZeroVector);
+
+			ColorVertexBuffer.VertexColor(VertIdx) = (bHaveColors) ?
+				(FColor)Mesh->GetVertexColor(VertIdx) : ConstantVertexColor;
+		}
+
+		IndexBuffer.Indices.AddUninitialized(Mesh->TriangleCount() * 3);
+		int TriIndex = 0;
+		for (FIndex3i Tri : Mesh->TrianglesItr())
+		{
+			IndexBuffer.Indices[TriIndex++] = Tri.A;
+			IndexBuffer.Indices[TriIndex++] = Tri.B;
+			IndexBuffer.Indices[TriIndex++] = Tri.C;
+		}
+	}
+
+
+
+
+public:
+
+
+
+	virtual void GetDynamicMeshElements(const TArray<const FSceneView*>& Views, const FSceneViewFamily& ViewFamily, uint32 VisibilityMap, FMeshElementCollector& Collector) const override
+	{
+		QUICK_SCOPE_CYCLE_COUNTER(STAT_SimpleDynamicMeshSceneProxy_GetDynamicMeshElements);
+
+		const bool bWireframe = (AllowDebugViewmodes() && ViewFamily.EngineShowFlags.Wireframe)
+			|| ParentComponent->bExplicitShowWireframe;
+
+		FColoredMaterialRenderProxy* WireframeMaterialInstance = new FColoredMaterialRenderProxy(
+			GEngine->WireframeMaterial ? GEngine->WireframeMaterial->GetRenderProxy() : nullptr,
+			FLinearColor(0, 0.5f, 1.f)
+		);
+
+		Collector.RegisterOneFrameMaterialProxy(WireframeMaterialInstance);
+
+		FMaterialRenderProxy* MaterialProxy = Material->GetRenderProxy();
+		FMaterialRenderProxy* WireframeMaterialProxy = WireframeMaterialInstance;
+
+		ESceneDepthPriorityGroup DepthPriority = (ParentComponent->bDrawOnTop) ? SDPG_Foreground : SDPG_World;
+
+
+		for (int32 ViewIndex = 0; ViewIndex < Views.Num(); ViewIndex++)
+		{
+			if (VisibilityMap & (1 << ViewIndex))
+			{
+				const FSceneView* View = Views[ViewIndex];
+
+				bool bHasPrecomputedVolumetricLightmap;
+				FMatrix PreviousLocalToWorld;
+				int32 SingleCaptureIndex;
+				bool bOutputVelocity;
+				GetScene().GetPrimitiveUniformShaderParameters_RenderThread(GetPrimitiveSceneInfo(), bHasPrecomputedVolumetricLightmap, PreviousLocalToWorld, SingleCaptureIndex, bOutputVelocity);
+
+				FDynamicPrimitiveUniformBuffer& DynamicPrimitiveUniformBuffer = Collector.AllocateOneFrameResource<FDynamicPrimitiveUniformBuffer>();
+				DynamicPrimitiveUniformBuffer.Set(GetLocalToWorld(), PreviousLocalToWorld, GetBounds(), GetLocalBounds(), true, bHasPrecomputedVolumetricLightmap, DrawsVelocity(), bOutputVelocity);
+
+				// Draw the mesh.
+				DrawBatch(Collector, MaterialProxy, false, DepthPriority, ViewIndex, DynamicPrimitiveUniformBuffer);
+				if (bWireframe)
+				{
+					DrawBatch(Collector, WireframeMaterialProxy, true, DepthPriority, ViewIndex, DynamicPrimitiveUniformBuffer);
+				}
+			}
+		}
+	}
+
+
+
+	virtual void DrawBatch(FMeshElementCollector& Collector, 
+		FMaterialRenderProxy* UseMaterial, 
+		bool bWireframe,
+		ESceneDepthPriorityGroup DepthPriority,
+		int ViewIndex,
+		FDynamicPrimitiveUniformBuffer& DynamicPrimitiveUniformBuffer) const
+	{
+		FMeshBatch& Mesh = Collector.AllocateMesh();
+		FMeshBatchElement& BatchElement = Mesh.Elements[0];
+		BatchElement.IndexBuffer = &IndexBuffer;
+		Mesh.bWireframe = bWireframe;
+		Mesh.VertexFactory = &VertexFactory;
+		Mesh.MaterialRenderProxy = UseMaterial;
+
+		BatchElement.PrimitiveUniformBufferResource = &DynamicPrimitiveUniformBuffer.UniformBuffer;
+
+		BatchElement.FirstIndex = 0;
+		BatchElement.NumPrimitives = IndexBuffer.Indices.Num() / 3;
+		BatchElement.MinVertexIndex = 0;
+		BatchElement.MaxVertexIndex = PositionVertexBuffer.GetNumVertices() - 1;
+		Mesh.ReverseCulling = IsLocalToWorldDeterminantNegative();
+		Mesh.Type = PT_TriangleList;
+		Mesh.DepthPriorityGroup = DepthPriority;
+		Mesh.bCanApplyViewModeOverrides = false;
+		Collector.AddMesh(ViewIndex, Mesh);
+	}
+
+
+
+	virtual FPrimitiveViewRelevance GetViewRelevance(const FSceneView* View) const override
+	{
+		FPrimitiveViewRelevance Result;
+
+
+		if (ParentComponent->bDrawOnTop)
+		{
+			Result.bDrawRelevance = IsShown(View);
+			Result.bDynamicRelevance = true;
+			Result.bShadowRelevance = false;
+			Result.bEditorPrimitiveRelevance = UseEditorCompositing(View);
+			Result.bEditorNoDepthTestPrimitiveRelevance = true;
+		}
+		else
+		{
+			Result.bDrawRelevance = IsShown(View);
+			Result.bShadowRelevance = IsShadowCast(View);
+			Result.bDynamicRelevance = true;
+			Result.bRenderInMainPass = ShouldRenderInMainPass();
+			Result.bUsesLightingChannels = GetLightingChannelMask() != GetDefaultLightingChannelMask();
+			Result.bTranslucentSelfShadow = bCastVolumetricTranslucentShadow;
+			Result.bRenderCustomDepth = ShouldRenderCustomDepth();
+			MaterialRelevance.SetPrimitiveViewRelevance(Result);
+			Result.bVelocityRelevance = IsMovable() && Result.bOpaqueRelevance && Result.bRenderInMainPass;
+		}
+
+		return Result;
+	}
+
+
+	virtual void GetLightRelevance(const FLightSceneProxy* LightSceneProxy, bool& bDynamic, bool& bRelevant, bool& bLightMapped, bool& bShadowMapped) const override
+	{
+		if (ParentComponent->bDrawOnTop)
+		{
+			bDynamic = false;
+			bRelevant = false;
+			bLightMapped = false;
+			bShadowMapped = false;
+		}
+		else
+		{
+			FPrimitiveSceneProxy::GetLightRelevance(LightSceneProxy, bDynamic, bRelevant, bLightMapped, bShadowMapped);
+		}
+	}
+
+
+	virtual bool CanBeOccluded() const override
+	{
+		//if (ParentComponent->bDrawOnTop)
+		//{
+		//	return false;
+		//}
+		return !MaterialRelevance.bDisableDepthTest;
+	}
+
+	virtual uint32 GetMemoryFootprint(void) const override { return(sizeof(*this) + GetAllocatedSize()); }
+
+	uint32 GetAllocatedSize(void) const { return(FPrimitiveSceneProxy::GetAllocatedSize()); }
+
+
+
+	virtual void SetMaterial(UMaterialInterface* MaterialIn)
+	{
+		this->Material = MaterialIn;
+	}
+
+
+	// copied from StaticMesh.cpp
+	void InitOrUpdateResource(FRenderResource* Resource)
+	{
+		if (!Resource->IsInitialized())
+		{
+			Resource->InitResource();
+		}
+		else
+		{
+			Resource->UpdateRHI();
+		}
+	}
+
+
+	SIZE_T GetTypeHash() const override
+	{
+		static size_t UniquePointer;
+		return reinterpret_cast<size_t>(&UniquePointer);
+	}
+};
+
+
+
+
+
+
+
diff --git a/Engine/Plugins/Experimental/MeshModelingToolset/Source/ModelingComponents/Private/StaticMeshComponentBuilder.cpp b/Engine/Plugins/Experimental/MeshModelingToolset/Source/ModelingComponents/Private/StaticMeshComponentBuilder.cpp
new file mode 100644
--- /dev/null
+++ b/Engine/Plugins/Experimental/MeshModelingToolset/Source/ModelingComponents/Private/StaticMeshComponentBuilder.cpp
@@ -0,0 +1,105 @@
+// Copyright 1998-2019 Epic Games, Inc. All Rights Reserved.
+
+#include "StaticMeshComponentBuilder.h"
+#include "Engine/StaticMesh.h"
+#include "Engine/StaticMeshActor.h"
+#include "Components/StaticMeshComponent.h"
+#include "Engine/CollisionProfile.h"
+#include "PhysicsEngine/BodySetup.h"
+#include "MeshDescription.h"
+
+
+
+FStaticMeshComponentBuilder::FStaticMeshComponentBuilder()
+{
+	NewStaticMesh = nullptr;
+}
+
+void FStaticMeshComponentBuilder::Initialize(UPackage* AssetPackage, FName MeshName)
+{
+	// create new UStaticMesh object
+	EObjectFlags flags = EObjectFlags::RF_Public | EObjectFlags::RF_Standalone;
+	this->NewStaticMesh = NewObject<UStaticMesh>(AssetPackage, MeshName, flags);
+
+#if WITH_EDITOR
+	// initialize the LOD 0 MeshDescription
+	NewStaticMesh->SetNumSourceModels(1);
+	NewStaticMesh->GetSourceModel(0).BuildSettings.bRecomputeNormals = false;
+	NewStaticMesh->GetSourceModel(0).BuildSettings.bRecomputeTangents = true;
+	this->MeshDescription = NewStaticMesh->CreateMeshDescription(0);
+#endif
+
+	if (NewStaticMesh->BodySetup == nullptr)
+	{
+		NewStaticMesh->CreateBodySetup();
+	}
+	if (NewStaticMesh->BodySetup != nullptr)
+	{
+		// enable complex as simple collision to use mesh directly
+		NewStaticMesh->BodySetup->CollisionTraceFlag = ECollisionTraceFlag::CTF_UseComplexAsSimple;
+	}
+
+	// add a material slot
+	NewStaticMesh->StaticMaterials.Add(FStaticMaterial());
+}
+
+
+
+void FStaticMeshComponentBuilder::CreateAndSetAsRootComponent(AActor* Actor)
+{
+#if WITH_EDITOR
+	// assuming we have updated the LOD 0 MeshDescription, tell UStaticMesh about this
+	NewStaticMesh->CommitMeshDescription(0);
+#endif
+	// if we have a StaticMeshActor we already have a StaticMeshComponent, otherwise we 
+	// need to make a new one. Note that if we make a new one it will not be editable in the 
+	// Editor because it is not a UPROPERTY...
+	AStaticMeshActor* StaticMeshActor = Cast<AStaticMeshActor>(Actor);
+	if (StaticMeshActor != nullptr)
+	{
+		NewMeshComponent = StaticMeshActor->GetStaticMeshComponent();
+	}
+	else 
+	{
+		// create component
+		NewMeshComponent = NewObject<UStaticMeshComponent>(Actor);
+		Actor->SetRootComponent(NewMeshComponent);
+	}
+
+	// this disconnects the component from various events
+	NewMeshComponent->UnregisterComponent();
+
+	// Configure flags of the component. Is this necessary?
+	NewMeshComponent->SetMobility(EComponentMobility::Movable);
+	NewMeshComponent->bSelectable = true;
+
+	// replace the UStaticMesh in the component
+	NewMeshComponent->SetStaticMesh(NewStaticMesh);
+
+	// UActorFactoryBasicShape::PostSpawnActor does this...what for??
+	//NewMeshComponent->StaticMeshDerivedDataKey = StaticMesh->RenderData->DerivedDataKey;
+
+	// call this to set a different default material. WorldMaterial shows UV so is more useful...
+	//UMaterialInterface* Material = LoadObject<UMaterial>(nullptr, TEXT("/Engine/BasicShapes/BasicShapeMaterial.BasicShapeMaterial"));
+	//check(Material != nullptr);
+	//NewMeshComponent->SetMaterial(0, Material);
+
+	// re-connect the component (?)
+	NewMeshComponent->RegisterComponent();
+
+#if WITH_EDITOR
+	// not sure what this does...marks things dirty? updates stuff after modification??
+	NewStaticMesh->PostEditChange();
+#endif
+
+	// do we need to do this?
+	//NewMeshComponent->PostEditChange();
+
+	// do we need to do any of these? or is default collision profile sufficient?
+	//NewStaticMesh->BodySetup->DefaultInstance.SetCollisionProfileName(UCollisionProfile::BlockAll_ProfileName);
+	//NewMeshComponent->SetCollisionResponseToAllChannels(ECollisionResponse::ECR_Block);
+	//NewMeshComponent->SetCollisionProfileName(UCollisionProfile::BlockAll_ProfileName);
+
+	// if we don't do this, world traces don't hit the mesh
+	NewMeshComponent->RecreatePhysicsState();
+}
diff --git a/Engine/Plugins/Experimental/MeshModelingToolset/Source/ModelingComponents/Private/ToolSetupUtil.cpp b/Engine/Plugins/Experimental/MeshModelingToolset/Source/ModelingComponents/Private/ToolSetupUtil.cpp
new file mode 100644
--- /dev/null
+++ b/Engine/Plugins/Experimental/MeshModelingToolset/Source/ModelingComponents/Private/ToolSetupUtil.cpp
@@ -0,0 +1,64 @@
+// Copyright 1998-2019 Epic Games, Inc. All Rights Reserved.
+
+
+#include "ToolSetupUtil.h"
+#include "InteractiveTool.h"
+#include "InteractiveToolManager.h"
+#include "Materials/Material.h"
+
+
+
+UMaterialInterface* ToolSetupUtil::GetDefaultMaterial(UInteractiveToolManager* ToolManager, UMaterialInterface* SourceMaterial)
+{
+	if (SourceMaterial == nullptr && ToolManager != nullptr)
+	{
+		return ToolManager->GetContextQueriesAPI()->GetStandardMaterial(EStandardToolContextMaterials::VertexColorMaterial);
+	}
+	return SourceMaterial;
+}
+
+
+
+UMaterialInterface* ToolSetupUtil::GetDefaultWorkingMaterial(UInteractiveToolManager* ToolManager)
+{
+	UMaterialInterface* Material = LoadObject<UMaterial>(nullptr, TEXT("/MeshModelingToolset/Materials/InProgressMaterial"));
+	if (Material == nullptr && ToolManager != nullptr)
+	{
+		return ToolManager->GetContextQueriesAPI()->GetStandardMaterial(EStandardToolContextMaterials::VertexColorMaterial);
+	}
+	return Material;
+}
+
+
+UMaterialInterface* ToolSetupUtil::GetDefaultBrushVolumeMaterial(UInteractiveToolManager* ToolManager)
+{
+	UMaterialInterface* Material = LoadObject<UMaterial>(nullptr, TEXT("/MeshModelingToolset/Materials/BrushIndicatorMaterial"));
+	if (Material == nullptr && ToolManager != nullptr)
+	{
+		return ToolManager->GetContextQueriesAPI()->GetStandardMaterial(EStandardToolContextMaterials::VertexColorMaterial);
+	}
+	return Material;
+}
+
+
+
+UMaterialInterface* ToolSetupUtil::GetSculptMaterial1(UInteractiveToolManager* ToolManager)
+{
+	UMaterialInterface* Material = LoadObject<UMaterial>(nullptr, TEXT("/MeshModelingToolset/Materials/SculptMaterial"));
+	if (Material == nullptr && ToolManager != nullptr)
+	{
+		return ToolManager->GetContextQueriesAPI()->GetStandardMaterial(EStandardToolContextMaterials::VertexColorMaterial);
+	}
+	return Material;
+}
+
+
+UMaterialInterface* ToolSetupUtil::GetSelectionMaterial(UInteractiveToolManager* ToolManager)
+{
+	UMaterialInterface* Material = LoadObject<UMaterial>(nullptr, TEXT("/MeshModelingToolset/Materials/SelectionMaterial"));
+	if (Material == nullptr && ToolManager != nullptr)
+	{
+		return ToolManager->GetContextQueriesAPI()->GetStandardMaterial(EStandardToolContextMaterials::VertexColorMaterial);
+	}
+	return Material;
+}
diff --git a/Engine/Plugins/Experimental/MeshModelingToolset/Source/ModelingComponents/Public/AssetGenerationUtil.h b/Engine/Plugins/Experimental/MeshModelingToolset/Source/ModelingComponents/Public/AssetGenerationUtil.h
new file mode 100644
--- /dev/null
+++ b/Engine/Plugins/Experimental/MeshModelingToolset/Source/ModelingComponents/Public/AssetGenerationUtil.h
@@ -0,0 +1,48 @@
+// Copyright 1998-2019 Epic Games, Inc. All Rights Reserved.
+
+#pragma once
+
+#include "CoreMinimal.h"
+#include "TransformTypes.h"
+
+class UMaterialInterface;
+class UInteractiveToolManager;
+class IToolsContextAssetAPI;
+struct FMeshDescription;
+class FDynamicMesh3;
+class AActor;
+
+/**
+ * Utility functions for generating Assets in InteractiveTool implementations
+ */
+namespace AssetGenerationUtil
+{
+	/**
+	 * @return default path to use for auto-generated assets
+	 * @warning currently hardcoded to "/Game/_GENERATED"
+	 */
+	MODELINGCOMPONENTS_API FString GetDefaultAutoGeneratedAssetPath();
+
+	/**
+	 * Creates a new StaticMesh actor/component, with a new mesh asset stored at the given PackagePath created via the AssetAPI
+	 * @param AssetAPI pointer to context asset API that will be used to create new asset
+	 * @param TargetWorld world that Actor will be created in
+	 * @param Mesh geometry for the mesh
+	 * @param Transform transformation for the new actor
+	 * @param ObjectName name of the new asset
+	 * @param PackagePath path of the new asset
+	 * @return new mesh actor
+	 */
+#if WITH_EDITOR
+	MODELINGCOMPONENTS_API AActor* GenerateStaticMeshActor(
+		IToolsContextAssetAPI* AssetAPI,
+		UWorld* TargetWorld,
+		const FDynamicMesh3* Mesh,
+		const FTransform3d& Transform,
+		FString ObjectName,
+		FString PackagePath,
+		UMaterialInterface* Material = nullptr
+	);
+#endif
+
+}
diff --git a/Engine/Plugins/Experimental/MeshModelingToolset/Source/ModelingComponents/Public/BaseDynamicMeshComponent.h b/Engine/Plugins/Experimental/MeshModelingToolset/Source/ModelingComponents/Public/BaseDynamicMeshComponent.h
new file mode 100644
--- /dev/null
+++ b/Engine/Plugins/Experimental/MeshModelingToolset/Source/ModelingComponents/Public/BaseDynamicMeshComponent.h
@@ -0,0 +1,74 @@
+// Copyright 1998-2019 Epic Games, Inc. All Rights Reserved.
+
+#pragma once
+
+#include "CoreMinimal.h"
+#include "UObject/ObjectMacros.h"
+#include "Components/MeshComponent.h"
+#include "Changes/MeshVertexChange.h"
+#include "Changes/MeshChange.h"
+
+#include "DynamicMesh3.h"
+
+#include "BaseDynamicMeshComponent.generated.h"
+
+
+// predecl
+class FMeshVertexChange;
+class FMeshChange;
+
+
+
+
+/**
+ * Tangent calculation modes
+ */
+UENUM()
+enum class EDynamicMeshTangentCalcType : uint8
+{
+	/** Tangents are not used/available, proceed accordingly (eg generate arbitrary orthogonal basis) */
+	NoTangents,
+	/** Tangents should be automatically calculated on demand */
+	AutoCalculated,
+	/** Tangents are externally calculated (behavior undefined if they are not actually externally calculated!) */
+	ExternallyCalculated
+};
+
+
+/**
+ * UBaseDynamicMeshComponent is a base interface for a UMeshComponent based on a FDynamicMesh.
+ * Currently no functionality lives here, only some interface functions are defined that various subclasses implement.
+ */
+UCLASS(hidecategories = (LOD, Physics, Collision), editinlinenew, ClassGroup = Rendering)
+class MODELINGCOMPONENTS_API UBaseDynamicMeshComponent : public UMeshComponent, public IMeshVertexCommandChangeTarget, public IMeshCommandChangeTarget
+{
+	GENERATED_UCLASS_BODY()
+
+public:
+
+	/**
+	 * Call this if you update the mesh via GetMesh()
+	 * @todo should provide a function that calls a lambda to modify the mesh, and only return const mesh pointer
+	 */
+	virtual void NotifyMeshUpdated()
+	{
+		unimplemented();
+	}
+
+	/**
+	 * Apply a vertex deformation change to the internal mesh  (implements IMeshVertexCommandChangeTarget)
+	 */
+	virtual void ApplyChange(const FMeshVertexChange* Change, bool bRevert) override
+	{
+		unimplemented();
+	}
+
+	/**
+	 * Apply a general mesh change to the internal mesh  (implements IMeshCommandChangeTarget)
+	 */
+	virtual void ApplyChange(const FMeshChange* Change, bool bRevert) override
+	{
+		unimplemented();
+	}
+
+};
diff --git a/Engine/Plugins/Experimental/MeshModelingToolset/Source/ModelingComponents/Public/Changes/MeshChange.h b/Engine/Plugins/Experimental/MeshModelingToolset/Source/ModelingComponents/Public/Changes/MeshChange.h
new file mode 100644
--- /dev/null
+++ b/Engine/Plugins/Experimental/MeshModelingToolset/Source/ModelingComponents/Public/Changes/MeshChange.h
@@ -0,0 +1,61 @@
+// Copyright 1998-2019 Epic Games, Inc. All Rights Reserved.
+
+#pragma once
+
+#include "CoreMinimal.h"
+#include "UObject/Interface.h"
+#include "InteractiveToolChange.h"
+#include "VectorTypes.h"
+#include "DynamicMeshChangeTracker.h"
+#include "MeshChange.generated.h"
+
+//class FDynamicMeshChange;		// need to refactor this out of DynamicMeshChangeTracker
+
+
+
+
+/**
+ * FMeshChange represents an undoable change to a FDynamicMesh3.
+ * Currently only valid to call Apply/Revert when the Object is a USimpleDynamicMeshComponent
+ */
+class MODELINGCOMPONENTS_API FMeshChange : public FToolCommandChange
+{
+public:
+	FMeshChange();
+	FMeshChange(TUniquePtr<FDynamicMeshChange> DynamicMeshChangeIn);
+
+	TUniquePtr<FDynamicMeshChange> DynamicMeshChange;
+
+	/** This function is called on Apply and Revert (last argument is true on Apply)*/
+	TFunction<void(FMeshChange*, UObject*, bool)> OnChangeAppliedFunc;
+
+	/** Makes the change to the object */
+	virtual void Apply(UObject* Object) override;
+
+	/** Reverts change to the object */
+	virtual void Revert(UObject* Object) override;
+
+	/** Describes this change (for debugging) */
+	virtual FString ToString() const override;
+};
+
+
+
+
+
+UINTERFACE()
+class MODELINGCOMPONENTS_API UMeshCommandChangeTarget : public UInterface
+{
+	GENERATED_BODY()
+};
+/**
+ * IMeshCommandChangeTarget is an interface which is used to apply a mesh change
+ */
+class MODELINGCOMPONENTS_API IMeshCommandChangeTarget
+{
+	GENERATED_BODY()
+public:
+	virtual void ApplyChange(const FMeshChange* Change, bool bRevert) = 0;
+};
+
+
diff --git a/Engine/Plugins/Experimental/MeshModelingToolset/Source/ModelingComponents/Public/OctreeDynamicMeshComponent.h b/Engine/Plugins/Experimental/MeshModelingToolset/Source/ModelingComponents/Public/OctreeDynamicMeshComponent.h
new file mode 100644
--- /dev/null
+++ b/Engine/Plugins/Experimental/MeshModelingToolset/Source/ModelingComponents/Public/OctreeDynamicMeshComponent.h
@@ -0,0 +1,146 @@
+// Copyright 1998-2019 Epic Games, Inc. All Rights Reserved.
+
+#pragma once
+
+#include "CoreMinimal.h"
+#include "BaseDynamicMeshComponent.h"
+#include "MeshConversionOptions.h"
+
+#include "DynamicMeshOctree3.h"
+#include "Util/IndexSetDecompositions.h"
+
+#include "OctreeDynamicMeshComponent.generated.h"
+
+// predecl
+struct FMeshDescription;
+
+/** internal FPrimitiveSceneProxy defined in OctreeDynamicMeshSceneProxy.h */
+class FOctreeDynamicMeshSceneProxy;
+
+
+/** 
+ * UOctreeDynamicMeshComponent is a mesh component similar to UProceduralMeshComponent,
+ * except it bases the renderable geometry off an internal FDynamicMesh3 instance.
+ * The class generally has the same capabilities as USimpleDynamicMeshComponent.
+ * 
+ * A FDynamicMeshOctree3 is available to dynamically track the triangles of the mesh
+ * (however the client is responsible for updating this octree).
+ * Based on the Octree, the mesh is partitioned into chunks that are stored in separate
+ * RenderBuffers in the FOctreeDynamicMeshSceneProxy.
+ * Calling NotifyMeshUpdated() will result in only the "dirty" chunks being updated,
+ * rather than the entire mesh.
+ * 
+ * (So, if you don't need this capability, and don't want to update an Octree, use USimpleDynamicMeshComponent!)
+ */
+UCLASS(hidecategories = (LOD, Physics, Collision), editinlinenew, ClassGroup = Rendering)
+class MODELINGCOMPONENTS_API UOctreeDynamicMeshComponent : public UBaseDynamicMeshComponent
+{
+	GENERATED_UCLASS_BODY()
+
+public:
+	/**
+	 * initialize the internal mesh from a MeshDescription
+	 */
+	void InitializeMesh(FMeshDescription* MeshDescription);
+
+	/**
+	 * @return pointer to internal mesh
+	 */
+	FDynamicMesh3* GetMesh() { return Mesh.Get(); }
+
+	FDynamicMeshOctree3* GetOctree() { return Octree.Get(); }
+
+	/**
+	 * @return the current internal mesh, which is replaced with an empty mesh
+	 */
+	TUniquePtr<FDynamicMesh3> ExtractMesh(bool bNotifyUpdate);
+
+	/**
+	 * Write the internal mesh to a MeshDescription
+	 * @param bHaveModifiedTopology if false, we only update the vertex positions in the MeshDescription, otherwise it is Empty()'d and regenerated entirely
+	 * @param ConversionOptions struct of additional options for the conversion
+	 */
+	void Bake(FMeshDescription* MeshDescription, bool bHaveModifiedTopology, const FConversionToMeshDescriptionOptions& ConversionOptions);
+
+	/**
+	* Write the internal mesh to a MeshDescription with default conversion options
+	* @param bHaveModifiedTopology if false, we only update the vertex positions in the MeshDescription, otherwise it is Empty()'d and regenerated entirely
+	*/
+	void Bake(FMeshDescription* MeshDescription, bool bHaveModifiedTopology)
+	{
+		FConversionToMeshDescriptionOptions ConversionOptions;
+		Bake(MeshDescription, bHaveModifiedTopology, ConversionOptions);
+	}
+
+
+	//
+	// change tracking/etc
+	//
+
+	/**
+	 * Call this if you update the mesh via GetMesh()
+	 * @todo should provide a function that calls a lambda to modify the mesh, and only return const mesh pointer
+	 */
+	virtual void NotifyMeshUpdated() override;
+
+	/**
+	 * Apply a vertex deformation change to the internal mesh
+	 */
+	virtual void ApplyChange(const FMeshVertexChange* Change, bool bRevert) override;
+
+	/**
+	 * Apply a general mesh change to the internal mesh
+	 */
+	virtual void ApplyChange(const FMeshChange* Change, bool bRevert) override;
+
+
+	/**
+	 * This delegate fires when a FCommandChange is applied to this component, so that
+	 * parent objects know the mesh has changed.
+	 */
+	FSimpleMulticastDelegate OnMeshChanged;
+
+	/**
+	 * if true, we always show the wireframe on top of the shaded mesh, even when not in wireframe mode
+	 */
+	UPROPERTY()
+	bool bExplicitShowWireframe = false;
+
+	TFunction<FColor(int)> TriangleColorFunc = nullptr;
+
+private:
+
+	FOctreeDynamicMeshSceneProxy* CurrentProxy = nullptr;
+
+	//~ Begin UPrimitiveComponent Interface.
+	virtual FPrimitiveSceneProxy* CreateSceneProxy() override;
+	//~ End UPrimitiveComponent Interface.
+
+	//~ Begin UMeshComponent Interface.
+	virtual int32 GetNumMaterials() const override;
+	//~ End UMeshComponent Interface.
+
+	//~ Begin USceneComponent Interface.
+	virtual FBoxSphereBounds CalcBounds(const FTransform& LocalToWorld) const override;
+	//~ Begin USceneComponent Interface.
+
+	TUniquePtr<FDynamicMesh3> Mesh;
+	TUniquePtr<FDynamicMeshOctree3> Octree;
+	TUniquePtr<FDynamicMeshOctree3::FTreeCutSet> OctreeCut;
+	FArrayIndexSetsDecomposition TriangleDecomposition;
+	struct FCutCellIndexSet
+	{
+		FDynamicMeshOctree3::FCellReference CellRef;
+		int32 DecompSetID;
+	};
+	TArray<FCutCellIndexSet> CutCellSetMap;
+	int32 SpillDecompSetID;
+
+	void InitializeNewMesh();
+
+
+
+	FColor GetTriangleColor(int TriangleID);
+
+	//friend class FCustomMeshSceneProxy;
+};
diff --git a/Engine/Plugins/Experimental/MeshModelingToolset/Source/ModelingComponents/Public/PreviewMesh.h b/Engine/Plugins/Experimental/MeshModelingToolset/Source/ModelingComponents/Public/PreviewMesh.h
new file mode 100644
--- /dev/null
+++ b/Engine/Plugins/Experimental/MeshModelingToolset/Source/ModelingComponents/Public/PreviewMesh.h
@@ -0,0 +1,262 @@
+// Copyright 1998-2019 Epic Games, Inc. All Rights Reserved.
+
+#pragma once
+
+#include "CoreMinimal.h"
+#include "SimpleDynamicMeshComponent.h"
+#include "DynamicMeshAABBTree3.h"
+#include "InteractiveToolObjects.h"
+#include "Changes/MeshVertexChange.h"
+#include "Changes/MeshChange.h"
+#include "PreviewMesh.generated.h"
+
+class FDynamicMesh3;
+struct FMeshDescription;
+
+
+
+/**
+ * UPreviewMesh internally spawns a APreviewMeshActor to hold the preview mesh object.
+ * We use this AInternalToolFrameworkActor subclass so that we can identify such objects
+ * at higher levels (for example to prevent them from being deleted in the Editor)
+ */
+UCLASS()
+class MODELINGCOMPONENTS_API APreviewMeshActor : public AInternalToolFrameworkActor
+{
+	GENERATED_BODY()
+private:
+
+public:
+
+};
+
+
+
+/** 
+ * UPreviewMesh is a utility object that spawns and owns a transient mesh object in the World.
+ * This can be used to show live preview geometry during modeling operations.
+ * Call CreateInWorld() to set it up, and Disconnect() to shut it down.
+ * 
+ * Currently implemented via an internal Actor that has a USimpleDynamicMeshComponent root component,
+ * with an AABBTree created/updated if UProperty bBuildSpatialDataStructure=true.
+ * The Actor is destroyed on Disconnect().
+ * 
+ * The intention with UPreviewMesh is to provide a higher-level interface than the Component.
+ * In future the internal Component may be replaced with another class (eg OctreeDynamicMeshComponent),
+ * or automatically swap between the two, etc.
+ * 
+ * As a result direct access to the Actor/Component, or a non-const FDynamicMesh3, is intentionally not provided.
+ * Wrapper functions are provided (or should be added) for necessary Actor/Component parameters.
+ * To edit the mesh either a copy is done, or EditMesh()/ApplyChange() must be used.
+ * These functions automatically update necessary internal data structures.
+ * 
+ */
+UCLASS(Transient)
+class MODELINGCOMPONENTS_API UPreviewMesh : public UObject, public IMeshVertexCommandChangeTarget, public IMeshCommandChangeTarget
+{
+	GENERATED_BODY()
+
+public:
+	UPreviewMesh();
+	virtual ~UPreviewMesh();
+
+	//
+	// construction / destruction
+	// 
+
+	/**
+	 * Create preview mesh in the World with the given transform
+	 */
+	void CreateInWorld(UWorld* World, const FTransform& WithTransform);
+
+	/**
+	 * Remove and destroy preview mesh
+	 */
+	void Disconnect();
+
+
+	/**
+	 * @return internal Root Component of internal Actor
+	 */
+	UPrimitiveComponent* GetRootComponent() { return DynamicMeshComponent; }
+
+
+	//
+	// visualization parameters
+	// 
+
+	/**
+	 * Enable/disable wireframe overlay rendering
+	 */
+	void EnableWireframe(bool bEnable);
+
+	/**
+	 * Set material on the preview mesh
+	 */
+	void SetMaterial(UMaterialInterface* Material);
+
+	/**
+	* Get material from the preview mesh
+	*/
+	UMaterialInterface* GetMaterial() const;
+	/**
+	 * Get the current transform on the preview mesh
+	 */
+	FTransform GetTransform() const;
+
+	/**
+	 * Set the transform on the preview mesh
+	 */
+	void SetTransform(const FTransform& UseTransform);
+
+	/**
+	 * @return true if the preview mesh is visible
+	 */
+	bool IsVisible() const;
+
+	/**
+	 * Set visibility state of the preview mesh
+	 */
+	void SetVisible(bool bVisible);
+
+
+	/** Render data update hint */
+	enum class ERenderUpdateMode
+	{
+		/** Do not update render data */
+		NoUpdate,	
+		/** Invalidate overlay of internal component, rebuilding all render data */
+		FullUpdate,
+		/** Attempt to do partial update of render data if possible */
+		FastUpdate
+	};
+
+	/**
+	 * Set the triangle color function for rendering / render data construction
+	 */
+	void SetTriangleColorFunction(TFunction<FColor(int)> TriangleColorFunc, ERenderUpdateMode UpdateMode = ERenderUpdateMode::FullUpdate);
+
+	/**
+	 * Clear the triangle color function for rendering / render data construction
+	 */
+	void ClearTriangleColorFunction(ERenderUpdateMode UpdateMode = ERenderUpdateMode::FullUpdate);
+
+
+
+
+
+	//
+	// Queries
+	// 
+
+	/**
+	 * Test for ray intersection with the preview mesh.
+	 * Requires that bBuildSpatialDataStructure = true
+	 * @return true if ray intersections mesh
+	 */
+	bool TestRayIntersection(const FRay3d& WorldRay);
+
+	/**
+	 * Find ray intersection with the preview mesh.
+	 * Requires that bBuildSpatialDataStructure = true
+	 * @param WorldRay ray in world space
+	 * @param HitOut output hit data (only certain members are initialized, see implementation)
+	 * @return true if ray intersections mesh
+	 */
+	bool FindRayIntersection(const FRay3d& WorldRay, FHitResult& HitOut);
+
+
+
+
+
+
+	//
+	// Read access to internal mesh
+	// 
+
+	/**
+	 * Clear the preview mesh
+	 */
+	void ClearPreview();
+
+	/**
+	 * Update the internal mesh by copying the given Mesh
+	 */
+	void UpdatePreview(const FDynamicMesh3* Mesh);
+
+	/**
+	 * Initialize the internal mesh based on the given MeshDescription
+	 */
+	void InitializeMesh(FMeshDescription* MeshDescription);
+
+	/**
+	* @return pointer to the current FDynamicMesh used for preview
+	*/
+	const FDynamicMesh3* GetPreviewDynamicMesh() const;
+
+	/**
+	 * @return the current preview FDynamicMesh, and replace with a new empty mesh
+	 */
+	TUniquePtr<FDynamicMesh3> ExtractPreviewMesh() const;
+
+	/**
+	 * Write the internal mesh to a MeshDescription
+	 * @param bHaveModifiedToplogy if false, we only update the vertex positions in the MeshDescription, otherwise it is Empty()'d and regenerated entirely
+	 */
+	void Bake(FMeshDescription* MeshDescription, bool bHaveModifiedToplogy);
+
+
+
+
+	//
+	// Edit access to internal mesh, and change-tracking/notification
+	// 
+
+	/**
+	 * Apply EditFunc to the internal mesh and update internal data structures as necessary
+	 */
+	void EditMesh(TFunctionRef<void(FDynamicMesh3&)> EditFunc);
+
+	/**
+	 * Apply EditFunc to the internal mesh and update internal data structures as necessary.
+	 * EditFunc is required to notify the given FDynamicMeshChangeTracker about all mesh changes
+	 * @return FMeshChange extracted from FDynamicMeshChangeTracker that represents mesh edit
+	 */
+	TUniquePtr<FMeshChange> TrackedEditMesh(TFunctionRef<void(FDynamicMesh3&, FDynamicMeshChangeTracker&)> EditFunc);
+
+	/**
+	 * Apply/Revert a vertex deformation change to the internal mesh (implements IMeshVertexCommandChangeTarget)
+	 */
+	virtual void ApplyChange(const FMeshVertexChange* Change, bool bRevert) override;
+
+	/**
+	 * Apply/Revert a general mesh change to the internal mesh   (implements IMeshCommandChangeTarget)
+	 */
+	virtual void ApplyChange(const FMeshChange* Change, bool bRevert) override;
+
+	/** @return delegate that is broadcast whenever the internal mesh component is changed */
+	FSimpleMulticastDelegate& GetOnMeshChanged();
+
+
+
+
+public:
+	/** If true, we build a spatial data structure internally for the preview mesh, which allows for hit-testing */
+	UPROPERTY()
+	bool bBuildSpatialDataStructure;
+
+	// results in component drawing w/o z-testing and with editor compositing. Do not recommend using this flag, will be deprecated/removed.
+	UPROPERTY()
+	bool bDrawOnTop;
+
+protected:
+	/** The temporary actor we create internally to own the preview mesh component */
+	APreviewMeshActor* TemporaryParentActor = nullptr;
+
+	/** This component is set as the root component of TemporaryParentActor */
+	UPROPERTY()
+	USimpleDynamicMeshComponent* DynamicMeshComponent = nullptr;
+
+	/** Spatial data structure that is initialized if bBuildSpatialDataStructure = true when UpdatePreview() is called */
+	FDynamicMeshAABBTree3 MeshAABBTree;
+};
diff --git a/Engine/Plugins/Experimental/MeshModelingToolset/Source/ModelingComponents/Public/Selection/GroupTopologySelector.h b/Engine/Plugins/Experimental/MeshModelingToolset/Source/ModelingComponents/Public/Selection/GroupTopologySelector.h
new file mode 100644
--- /dev/null
+++ b/Engine/Plugins/Experimental/MeshModelingToolset/Source/ModelingComponents/Public/Selection/GroupTopologySelector.h
@@ -0,0 +1,130 @@
+// Copyright 1998-2019 Epic Games, Inc. All Rights Reserved.
+
+#pragma once
+
+#include "GroupTopology.h"
+#include "Spatial/GeometrySet3.h"
+#include "DynamicMeshAABBTree3.h"
+#include "RayTypes.h"
+
+class FToolDataVisualizer;
+struct FViewCameraState;
+
+/**
+ * FGroupTopologySelection represents a set of selected elements of a FGroupTopology
+ */
+struct MODELINGCOMPONENTS_API FGroupTopologySelection
+{
+	TArray<int> SelectedGroupIDs;
+	TArray<int> SelectedCornerIDs;
+	TArray<int> SelectedEdgeIDs;
+
+	FGroupTopologySelection() { Clear(); }
+	void Clear();
+};
+
+
+/**
+ * FGroupTopologySelector implements selection behavior for a FGroupTopology mesh.
+ * Groups, Group Edges, and Corners can be selected (optionally configurable via UpdateEnableFlags).
+ * Internally a FGeometrySet3 is constructed to support ray-hit testing against the edges and corners.
+ * 
+ * Note that to hit-test against the mesh you have to provide a your own FDynamicMeshAABBTree3.
+ * You do this by providing a callback via SetSpatialSource(). The reason for this is that
+ * (1) owners of FGroupTopologySelector likely already have a BVTree and (2) if the use case
+ * is deformation, we need to make sure the owner has recomputed the BVTree before we call functions 
+ * on it. The callback you provide should do that!
+ * 
+ * DrawSelection() can be used to visualize a selection via line/circle drawing.
+ * 
+ * @todo optionally have an internal mesh AABBTree that can be used when owner does not provide one?
+ */
+class MODELINGCOMPONENTS_API FGroupTopologySelector
+{
+public:
+
+	//
+	// Configuration variables
+	// 
+
+	/** 
+	 * This is the function we use to determine if a point on a corner/edge is close enough
+	 * to the hit-test ray to treat as a "hit". By default this is Euclidean distance with
+	 * a tolerance of 1.0. You probably need to replace this with your own function.
+	 */
+	TFunction<bool(const FVector3d&, const FVector3d&)> PointsWithinToleranceTest;
+
+public:
+	FGroupTopologySelector();
+
+	/**
+	 * Initialize the selector with the given Mesh and Topology.
+	 * This does not create the internal data structures, this happens lazily on GetGeometrySet()
+	 */
+	void Initialize(const FDynamicMesh3* Mesh, const FGroupTopology* Topology);
+
+	/**
+	 * Provide a function that will return an AABBTree for the Mesh.
+	 * See class comment for why this is necessary.
+	 */
+	void SetSpatialSource(TFunction<FDynamicMeshAABBTree3*(void)> GetSpatialFunc)
+	{
+		GetSpatial = GetSpatialFunc;
+	}
+
+	/**
+	 * Notify the Selector that the mesh has changed. 
+	 * @param bTopologyDeformed if this is true, the mesh vertices have been moved so we need to update bounding boxes/etc
+	 * @param bTopologyModified if this is true, topology has changed and we need to rebuild spatial data structures from scratch
+	 */
+	void Invalidate(bool bTopologyDeformed, bool bTopologyModified);
+
+	/**
+	 * @return the internal GeometrySet. This does lazy updating of the GeometrySet, so this function may take some time.
+	 */
+	const FGeometrySet3& GetGeometrySet();
+
+	/**
+	 * Configure whether faces, edges, and corners will be returned by hit-tests
+	 */
+	void UpdateEnableFlags(bool bFaceHits, bool bEdgeHits, bool bCornerHits);
+
+	/**
+	 * Find which element was selected for a given ray
+	 * @param Ray hit-test ray
+	 * @param ResultOut resulting selection. At most one of the Groups/Corners/Edges members will contain one element.
+	 * @param SelectedPositionOut The point on the ray nearest to the selected element
+	 * @param SelectedNormalOut the normal at that point if ResultOut contains a selected face, otherwise uninitialized
+	 */
+	bool FindSelectedElement(const FRay3d& Ray, FGroupTopologySelection& ResultOut, 
+		FVector3d& SelectedPositionOut, FVector3d& SelectedNormalOut);
+
+
+	/**
+	 * Render the given selection with the default settings of the FToolDataVisualizer.
+	 * Selected edges are drawn as lines, and selected corners are drawn as small view-facing circles.
+	 * (Currently seleced faces are not draw)
+	 */
+	void DrawSelection(const FGroupTopologySelection& Selection, FToolDataVisualizer* Renderer, const FViewCameraState* CameraState);
+
+
+public:
+	// internal rendering parameters
+	float VisualAngleSnapThreshold = 0.5;
+
+
+protected:
+	const FDynamicMesh3* Mesh = nullptr;
+	const FGroupTopology* Topology = nullptr;
+
+	TFunction<FDynamicMeshAABBTree3*(void)> GetSpatial = nullptr;
+
+	bool bGeometryInitialized = false;
+	bool bGeometryUpToDate = false;
+	FGeometrySet3 GeometrySet;
+
+
+	bool bEnableFaceHits = true;
+	bool bEnableEdgeHits = true;
+	bool bEnableCornerHits = true;
+};
\ No newline at end of file
diff --git a/Engine/Plugins/Experimental/MeshModelingToolset/Source/ModelingComponents/Public/Selection/SelectClickedAction.h b/Engine/Plugins/Experimental/MeshModelingToolset/Source/ModelingComponents/Public/Selection/SelectClickedAction.h
new file mode 100644
--- /dev/null
+++ b/Engine/Plugins/Experimental/MeshModelingToolset/Source/ModelingComponents/Public/Selection/SelectClickedAction.h
@@ -0,0 +1,42 @@
+// Copyright 1998-2019 Epic Games, Inc. All Rights Reserved.
+
+
+#include "BaseBehaviors/BehaviorTargetInterfaces.h"
+
+#pragma once
+
+/**
+ * BehaviorTarget to do world raycast selection from a click
+ * Currently used to click-select reference planes in the world
+ */
+class FSelectClickedAction : public IClickBehaviorTarget
+{
+	bool DoRayCast(const FInputDeviceRay& ClickPos, bool callbackOnHit)
+	{
+		FVector RayStart = ClickPos.WorldRay.Origin;
+		FVector RayEnd = ClickPos.WorldRay.PointAt(999999);
+		FCollisionObjectQueryParams QueryParams(FCollisionObjectQueryParams::AllObjects);
+		FHitResult Result;
+		bool bHitWorld = World->LineTraceSingleByObjectType(Result, RayStart, RayEnd, QueryParams);
+		if (callbackOnHit && bHitWorld && OnClickedPositionFunc != nullptr)
+		{
+			OnClickedPositionFunc(Result);
+		}
+		return bHitWorld;
+	}
+
+public:
+	UWorld* World;
+	TFunction<void(const FHitResult&)> OnClickedPositionFunc = nullptr;
+
+	virtual bool IsHitByClick(const FInputDeviceRay& ClickPos) override
+	{
+		return DoRayCast(ClickPos, false);
+	}
+
+	virtual void OnClicked(const FInputDeviceRay& ClickPos) override
+	{
+		DoRayCast(ClickPos, true);
+	}
+};
+
diff --git a/Engine/Plugins/Experimental/MeshModelingToolset/Source/ModelingComponents/Public/SimpleDynamicMeshComponent.h b/Engine/Plugins/Experimental/MeshModelingToolset/Source/ModelingComponents/Public/SimpleDynamicMeshComponent.h
new file mode 100644
--- /dev/null
+++ b/Engine/Plugins/Experimental/MeshModelingToolset/Source/ModelingComponents/Public/SimpleDynamicMeshComponent.h
@@ -0,0 +1,154 @@
+// Copyright 1998-2019 Epic Games, Inc. All Rights Reserved.
+
+#pragma once
+
+#include "CoreMinimal.h"
+#include "BaseDynamicMeshComponent.h"
+#include "MeshConversionOptions.h"
+
+#include "MeshTangents.h"
+
+#include "SimpleDynamicMeshComponent.generated.h"
+
+// predecl
+struct FMeshDescription;
+
+/** internal FPrimitiveSceneProxy defined in SimpleDynamicMeshSceneProxy.h */
+class FSimpleDynamicMeshSceneProxy;
+
+
+
+/** 
+ * USimpleDynamicMeshComponent is a mesh component similar to UProceduralMeshComponent,
+ * except it bases the renderable geometry off an internal FDynamicMesh3 instance.
+ * 
+ * There is some support for undo/redo on the component (@todo is this the right place?)
+ * 
+ * This component draws wireframe-on-shaded when Wireframe is enabled, or when bExplicitShowWireframe = true
+ *
+ */
+UCLASS(hidecategories = (LOD, Physics, Collision), editinlinenew, ClassGroup = Rendering)
+class MODELINGCOMPONENTS_API USimpleDynamicMeshComponent : public UBaseDynamicMeshComponent
+{
+	GENERATED_UCLASS_BODY()
+
+
+public:
+	/** How should Tangents be calculated/handled */
+	UPROPERTY()
+	EDynamicMeshTangentCalcType TangentsType = EDynamicMeshTangentCalcType::NoTangents;
+
+public:
+	/**
+	 * initialize the internal mesh from a MeshDescription
+	 */
+	void InitializeMesh(FMeshDescription* MeshDescription);
+
+	/**
+	 * @return pointer to internal mesh
+	 */
+	FDynamicMesh3* GetMesh() { return Mesh.Get(); }
+
+	/**
+	 * @return the current internal mesh, which is replaced with an empty mesh
+	 */
+	TUniquePtr<FDynamicMesh3> ExtractMesh(bool bNotifyUpdate);
+
+	/**
+	 * @return pointer to internal tangents object. 
+	 * @warning calling this with TangentsType = AutoCalculated will result in possibly-expensive Tangents calculation
+	 */
+	FMeshTangentsf* GetTangents();
+
+
+	/**
+	 * Write the internal mesh to a MeshDescription
+	 * @param bHaveModifiedTopology if false, we only update the vertex positions in the MeshDescription, otherwise it is Empty()'d and regenerated entirely
+	 * @param ConversionOptions struct of additional options for the conversion
+	 */
+	void Bake(FMeshDescription* MeshDescription, bool bHaveModifiedTopology, const FConversionToMeshDescriptionOptions& ConversionOptions);
+
+	/**
+	* Write the internal mesh to a MeshDescription with default conversion options
+	* @param bHaveModifiedTopology if false, we only update the vertex positions in the MeshDescription, otherwise it is Empty()'d and regenerated entirely
+	*/
+	void Bake(FMeshDescription* MeshDescription, bool bHaveModifiedTopology)
+	{
+		FConversionToMeshDescriptionOptions ConversionOptions;
+		Bake(MeshDescription, bHaveModifiedTopology, ConversionOptions);
+	}
+
+
+	//
+	// change tracking/etc
+	//
+
+	/**
+	 * Call this if you update the mesh via GetMesh()
+	 * @todo should provide a function that calls a lambda to modify the mesh, and only return const mesh pointer
+	 */
+	virtual void NotifyMeshUpdated() override;
+
+	/**
+	 * Apply a vertex deformation change to the internal mesh
+	 */
+	virtual void ApplyChange(const FMeshVertexChange* Change, bool bRevert) override;
+
+	/**
+	 * Apply a general mesh change to the internal mesh
+	 */
+	virtual void ApplyChange(const FMeshChange* Change, bool bRevert) override;
+
+
+	/**
+	 * This delegate fires when a FCommandChange is applied to this component, so that
+	 * parent objects know the mesh has changed.
+	 */
+	FSimpleMulticastDelegate OnMeshChanged;
+
+
+	void SetDrawOnTop(bool bSet);
+
+	/**
+	 * if true, we always show the wireframe on top of the shaded mesh, even when not in wireframe mode
+	 */
+	UPROPERTY()
+	bool bExplicitShowWireframe = false;
+
+
+	UPROPERTY()
+	bool bDrawOnTop = false;
+
+
+
+	TFunction<FColor(int)> TriangleColorFunc = nullptr;
+	void FastNotifyColorsUpdated();
+
+	void FastNotifyPositionsUpdated();
+
+private:
+
+	FSimpleDynamicMeshSceneProxy* CurrentProxy = nullptr;
+
+	//~ Begin UPrimitiveComponent Interface.
+	virtual FPrimitiveSceneProxy* CreateSceneProxy() override;
+	//~ End UPrimitiveComponent Interface.
+
+	//~ Begin UMeshComponent Interface.
+	virtual int32 GetNumMaterials() const override;
+	//~ End UMeshComponent Interface.
+
+	//~ Begin USceneComponent Interface.
+	virtual FBoxSphereBounds CalcBounds(const FTransform& LocalToWorld) const override;
+	//~ Begin USceneComponent Interface.
+
+	TUniquePtr<FDynamicMesh3> Mesh;
+	void InitializeNewMesh();
+
+	bool bTangentsValid = false;
+	FMeshTangentsf Tangents;
+	
+	FColor GetTriangleColor(int TriangleID);
+
+	//friend class FCustomMeshSceneProxy;
+};
