fix handling of open boundaries for mesh selection tool's region growth / connected component selection

#rnx

diff --git a/Engine/Plugins/Experimental/MeshModelingToolset/Source/MeshModelingTools/Private/MeshSelectionTool.cpp b/Engine/Plugins/Experimental/MeshModelingToolset/Source/MeshModelingTools/Private/MeshSelectionTool.cpp
--- a/Engine/Plugins/Experimental/MeshModelingToolset/Source/MeshModelingTools/Private/MeshSelectionTool.cpp
+++ b/Engine/Plugins/Experimental/MeshModelingToolset/Source/MeshModelingTools/Private/MeshSelectionTool.cpp
@@ -844,7 +844,7 @@ void UMeshSelectionTool::ExpandToConnected()
 		for (int j = 0; j < 3; ++j)
 		{
 			int32 tid = NbrTris[j];
-			if (SelectedTriangles[tid] == false && AddFaces.Contains(tid) == false)
+			if (tid != FDynamicMesh3::InvalidID && SelectedTriangles[tid] == false && AddFaces.Contains(tid) == false)
 			{
 				AddFaces.Add(tid);
 				Queue.Add(tid);
