make simplify mesh tool use background compute op stuff so it doesn't hang on opening

#rnx

diff --git a/Engine/Plugins/Experimental/MeshModelingToolset/Source/MeshModelingToolsEditorOnly/Private/SimplifyMeshTool.cpp b/Engine/Plugins/Experimental/MeshModelingToolset/Source/MeshModelingToolsEditorOnly/Private/SimplifyMeshTool.cpp
--- a/Engine/Plugins/Experimental/MeshModelingToolset/Source/MeshModelingToolsEditorOnly/Private/SimplifyMeshTool.cpp
+++ b/Engine/Plugins/Experimental/MeshModelingToolset/Source/MeshModelingToolsEditorOnly/Private/SimplifyMeshTool.cpp
@@ -4,24 +4,30 @@
 #include "InteractiveToolManager.h"
 #include "ToolBuilderUtil.h"
 
+#include "ToolSetupUtil.h"
+
 #include "DynamicMesh3.h"
-#include "DynamicMeshAttributeSet.h"
-#include "MeshSimplification.h"
-#include "MeshConstraintsUtil.h"
-#include "ProjectionTargets.h"
 
-#include "SimpleDynamicMeshComponent.h"
+#include "MeshDescriptionToDynamicMesh.h"
+#include "DynamicMeshToMeshDescription.h"
+
+#include "AssetGenerationUtil.h"
 
 #include "SceneManagement.h" // for FPrimitiveDrawInterface
 
 //#include "ProfilingDebugging/ScopedTimers.h" // enable this to use the timer.
-#include "MeshDescription.h"
-#include "MeshDescriptionOperations.h"
-#include "OverlappingCorners.h"
+#include "Modules/ModuleManager.h"
+
+
 #include "IMeshReductionManagerModule.h"
 #include "IMeshReductionInterfaces.h"
-#include "Modules/ModuleManager.h"
-#include "Operations/MergeCoincidentMeshEdges.h"
+
+
+#if WITH_EDITOR
+#include "Misc/ScopedSlowTask.h"
+#endif
+
+
 
 #define LOCTEXT_NAMESPACE "USimplifyMeshTool"
 
@@ -45,7 +51,10 @@ UInteractiveTool* USimplifyMeshToolBuilder::BuildTool(const FToolBuilderState& S
 	UActorComponent* ActorComponent = ToolBuilderUtil::FindFirstComponent(SceneState, CanMakeComponentTarget);
 	auto* MeshComponent = Cast<UPrimitiveComponent>(ActorComponent);
 	check(MeshComponent != nullptr);
+
 	NewTool->SetSelection(MakeComponentTarget(MeshComponent));
+	NewTool->SetWorld(SceneState.World);
+	NewTool->SetAssetAPI(AssetAPI);
 
 	return NewTool;
 }
@@ -69,31 +78,62 @@ USimplifyMeshToolProperties::USimplifyMeshToolProperties()
 }
 
 
+void USimplifyMeshTool::SetWorld(UWorld* World)
+{
+	this->TargetWorld = World;
+}
+
+void USimplifyMeshTool::SetAssetAPI(IToolsContextAssetAPI* AssetAPIIn)
+{
+	this->AssetAPI = AssetAPIIn;
+}
+
+
 
 void USimplifyMeshTool::Setup()
 {
 	UInteractiveTool::Setup();
 
-	// create dynamic mesh component to use for live preview
-	DynamicMeshComponent = NewObject<USimpleDynamicMeshComponent>(ComponentTarget->GetOwnerActor(), "DynamicMesh");
-	DynamicMeshComponent->SetupAttachment(ComponentTarget->GetOwnerActor()->GetRootComponent());
-	DynamicMeshComponent->RegisterComponent();
-	DynamicMeshComponent->SetWorldTransform(ComponentTarget->GetWorldTransform());
+	
+	// hide component and create + show preview
+	ComponentTarget->SetOwnerVisibility(false);
+	Preview = NewObject<UMeshOpPreviewWithBackgroundCompute>(this, "Preview");
+	Preview->Setup(this->TargetWorld, this);
+	Preview->ConfigureMaterials(
+		ToolSetupUtil::GetDefaultMaterial(GetToolManager(), ComponentTarget->GetMaterial(0)),
+		ToolSetupUtil::GetDefaultWorkingMaterial(GetToolManager())
+	);
+	Preview->PreviewMesh->EnableWireframe(true);
 
-	// copy material if there is one
-	auto Material = ComponentTarget->GetMaterial(0);
-	if (Material != nullptr)
 	{
-		DynamicMeshComponent->SetMaterial(0, Material);
-	}
-	DynamicMeshComponent->bExplicitShowWireframe = true;
+		// if in editor, create progress indicator dialog because building mesh copies can be slow (for very large meshes)
+		// this is especially needed because of the copy we make of the meshdescription; for Reasons, copying meshdescription is pretty slow
+#if WITH_EDITOR
+		static const FText SlowTaskText = LOCTEXT("SimplifyMeshInit", "Building mesh simplification data...");
 
-	DynamicMeshComponent->InitializeMesh(ComponentTarget->GetMesh());
-	OriginalMesh.Copy(*DynamicMeshComponent->GetMesh());
-	OriginalMeshSpatial.SetMesh(&OriginalMesh, true);
+		FScopedSlowTask SlowTask(3.0f, SlowTaskText);
+		SlowTask.MakeDialog();
 
-	// hide input StaticMeshComponent
-	ComponentTarget->SetOwnerVisibility(false);
+		// Declare progress shortcut lambdas
+		auto EnterProgressFrame = [&SlowTask](int Progress)
+		{
+			SlowTask.EnterProgressFrame((float)Progress);
+		};
+#else
+		auto EnterProgressFrame = [](int Progress) {};
+#endif
+		OriginalMeshDescription = MakeShared<FMeshDescription>(*ComponentTarget->GetMesh());
+		EnterProgressFrame(1);
+		OriginalMesh = MakeShared<FDynamicMesh3>();
+		FMeshDescriptionToDynamicMesh Converter;
+		Converter.bPrintDebugMessages = true;
+		Converter.Convert(ComponentTarget->GetMesh(), *OriginalMesh);
+		EnterProgressFrame(2);
+		OriginalMeshSpatial = MakeShared<FDynamicMeshAABBTree3>(OriginalMesh.Get(), true);
+	}
+
+	Preview->PreviewMesh->SetTransform(ComponentTarget->GetWorldTransform());
+	Preview->PreviewMesh->UpdatePreview(OriginalMesh.Get());
 
 	// initialize our properties
 	SimplifyProperties = NewObject<USimplifyMeshToolProperties>(this);
@@ -101,51 +141,71 @@ void USimplifyMeshTool::Setup()
 
 	MeshStatisticsProperties = NewObject<UMeshStatisticsProperties>(this);
 	AddToolPropertySource(MeshStatisticsProperties);
-	MeshStatisticsProperties->Update(*DynamicMeshComponent->GetMesh());
 
-	bResultValid = false;
+	Preview->OnMeshUpdated.AddLambda([this](UMeshOpPreviewWithBackgroundCompute* Compute)
+	{
+		MeshStatisticsProperties->Update(*Compute->PreviewMesh->GetPreviewDynamicMesh());
+	});
+
+	Preview->InvalidateResult();
 }
 
 
 void USimplifyMeshTool::Shutdown(EToolShutdownType ShutdownType)
 {
-	if (DynamicMeshComponent != nullptr)
+	ComponentTarget->SetOwnerVisibility(true);
+	TUniquePtr<FDynamicMeshOpResult> Result = Preview->Shutdown();
+	if (ShutdownType == EToolShutdownType::Accept)
 	{
-		//DynamicMeshComponent->OnMeshChanged.Remove(OnDynamicMeshComponentChangedHandle);
+		GenerateAsset(*Result);
+	}
+}
 
-		ComponentTarget->SetOwnerVisibility(true);
 
-		if (ShutdownType == EToolShutdownType::Accept)
-		{
-			// this block bakes the modified DynamicMeshComponent back into the StaticMeshComponent inside an undo transaction
-			GetToolManager()->BeginUndoTransaction(LOCTEXT("SimplifyMeshToolTransactionName", "Simplify Mesh"));
-			ComponentTarget->CommitMesh([=](FMeshDescription* MeshDescription)
-			{
-				DynamicMeshComponent->Bake(MeshDescription, true);
-			});
-			GetToolManager()->EndUndoTransaction();
-		}
+void USimplifyMeshTool::Tick(float DeltaTime)
+{
+	Preview->Tick(DeltaTime);
+}
 
-		DynamicMeshComponent->UnregisterComponent();
-		DynamicMeshComponent->DestroyComponent();
-		DynamicMeshComponent = nullptr;
-	}
+TSharedPtr<FDynamicMeshOperator> USimplifyMeshTool::MakeNewOperator()
+{
+	TSharedPtr<FSimplifyMeshOp> Op = MakeShared<FSimplifyMeshOp>();
+
+	Op->bDiscardAttributes = SimplifyProperties->bDiscardAttributes;
+	Op->bPreventNormalFlips = SimplifyProperties->bPreventNormalFlips;
+	Op->bReproject = SimplifyProperties->bReproject;
+	Op->SimplifierType = SimplifyProperties->SimplifierType;
+	Op->TargetCount = SimplifyProperties->TargetCount;
+	Op->TargetEdgeLength = SimplifyProperties->TargetEdgeLength;
+	Op->TargetMode = SimplifyProperties->TargetMode;
+	Op->TargetPercentage = SimplifyProperties->TargetPercentage;
+
+	FTransform LocalToWorld = ComponentTarget->GetWorldTransform();
+	Op->SetTransform(LocalToWorld);
+
+	Op->OriginalMeshDescription = OriginalMeshDescription;
+	Op->OriginalMesh = OriginalMesh;
+	Op->OriginalMeshSpatial = OriginalMeshSpatial;
+
+	IMeshReductionManagerModule& MeshReductionModule = FModuleManager::Get().LoadModuleChecked<IMeshReductionManagerModule>("MeshReductionInterface");
+	Op->MeshReduction = MeshReductionModule.GetStaticMeshReductionInterface();
+
+	return Op;
 }
 
 
+
 void USimplifyMeshTool::Render(IToolsContextRenderAPI* RenderAPI)
 {
-	UpdateResult();
-
-
+	
 	FPrimitiveDrawInterface* PDI = RenderAPI->GetPrimitiveDrawInterface();
 	FTransform Transform = ComponentTarget->GetWorldTransform(); //Actor->GetTransform();
 
 	FColor LineColor(255, 0, 0);
-	FDynamicMesh3* TargetMesh = DynamicMeshComponent->GetMesh();
+	const FDynamicMesh3* TargetMesh = Preview->PreviewMesh->GetPreviewDynamicMesh();
 	if (TargetMesh->HasAttributes())
 	{
-		FDynamicMeshUVOverlay* UVOverlay = TargetMesh->Attributes()->PrimaryUV();
+		const FDynamicMeshUVOverlay* UVOverlay = TargetMesh->Attributes()->PrimaryUV();
 		for (int eid : TargetMesh->EdgeIndicesItr()) 
 		{
 			if (UVOverlay->IsSeamEdge(eid)) 
@@ -162,151 +222,33 @@ void USimplifyMeshTool::Render(IToolsContextRenderAPI* RenderAPI)
 
 void USimplifyMeshTool::OnPropertyModified(UObject* PropertySet, UProperty* Property)
 {
-	bResultValid = false;
+	Preview->InvalidateResult();
 }
 
-
-template <typename SimplificationType>
-void ComputeResult(FDynamicMesh3* TargetMesh, const bool bReproject, int OriginalTriCount, FDynamicMesh3& OriginalMesh, FDynamicMeshAABBTree3& OriginalMeshSpatial, const ESimplifyTargetType TargetMode, const float TargetPercentage, const int TargetCount, const float TargetEdgeLength)
+bool USimplifyMeshTool::HasAccept() const
 {
-	SimplificationType Reducer(TargetMesh);
-
-	Reducer.ProjectionMode = (bReproject) ? 
-		SimplificationType::ETargetProjectionMode::AfterRefinement : SimplificationType::ETargetProjectionMode::NoProjection;
-
-	Reducer.DEBUG_CHECK_LEVEL = 0;
-
-	FMeshConstraints constraints;
-	FMeshConstraintsUtil::ConstrainAllSeams(constraints, *TargetMesh, true, false);
-	//FMeshConstraintsUtil::ConstrainAllSeamJunctions(constraints, *TargetMesh, true, false);
-	Reducer.SetExternalConstraints(&constraints);
-
-	FMeshProjectionTarget ProjTarget(&OriginalMesh, &OriginalMeshSpatial);
-	Reducer.SetProjectionTarget(&ProjTarget);
-
-	if (TargetMode == ESimplifyTargetType::Percentage)
-	{
-		double Ratio = (double)TargetPercentage / 100.0;
-		int UseTarget = FMath::Max(4, (int)(Ratio * (double)OriginalTriCount));
-		Reducer.SimplifyToTriangleCount(UseTarget);
-	} 
-	else if (TargetMode == ESimplifyTargetType::TriangleCount)
-	{
-		Reducer.SimplifyToTriangleCount(TargetCount);
-	}
-	else if (TargetMode == ESimplifyTargetType::EdgeLength)
-	{
-		Reducer.SimplifyToEdgeLength(TargetEdgeLength);
-	}
+	return true;
 }
 
-
-void USimplifyMeshTool::UpdateResult()
+bool USimplifyMeshTool::CanAccept() const
 {
-	
-	if (bResultValid) 
-	{
-		return;
-	}
-
-	//FScopedDurationTimeLogger Timer(TEXT("Simplification Time"));
+	return true;
+}
 
-	FDynamicMesh3* TargetMesh = DynamicMeshComponent->GetMesh();
-	TargetMesh->Copy(OriginalMesh);
-	int OriginalTriCount = OriginalMesh.TriangleCount();
+void USimplifyMeshTool::GenerateAsset(const FDynamicMeshOpResult& Result)
+{
+	GetToolManager()->BeginUndoTransaction(LOCTEXT("SimplifyMeshToolTransactionName", "Simplify Mesh"));
 
-	if (SimplifyProperties->bDiscardAttributes)
-	{
-		TargetMesh->DiscardAttributes();
-	}
-	if (SimplifyProperties->SimplifierType == ESimplifyType::QEM)
+	check(Result.Mesh.Get() != nullptr);
+	ComponentTarget->CommitMesh([&Result](FMeshDescription* MeshDescription)
 	{
-		ComputeResult<FQEMSimplification>(TargetMesh, SimplifyProperties->bReproject, OriginalTriCount, OriginalMesh, OriginalMeshSpatial, 
-			SimplifyProperties->TargetMode, SimplifyProperties->TargetPercentage, SimplifyProperties->TargetCount, SimplifyProperties->TargetEdgeLength);
-	}
-	else if (SimplifyProperties->SimplifierType == ESimplifyType::Attribute)
-	{
-		ComputeResult<FAttrMeshSimplification>(TargetMesh, SimplifyProperties->bReproject, OriginalTriCount, OriginalMesh, OriginalMeshSpatial,
-			SimplifyProperties->TargetMode, SimplifyProperties->TargetPercentage, SimplifyProperties->TargetCount, SimplifyProperties->TargetEdgeLength);
-	}
-	else //if (SimplifierType == ESimplifyType::UE4Standard)
-	{
-
-		const FMeshDescription* SrcMeshDescription = ComponentTarget->GetMesh();
-		FMeshDescription DstMeshDescription(*SrcMeshDescription);
-
-
-		FOverlappingCorners OverlappingCorners;
-		FMeshDescriptionOperations::FindOverlappingCorners(OverlappingCorners, *SrcMeshDescription, 1.e-5);
-
-		FMeshReductionSettings ReductionSettings;
-		if (SimplifyProperties->TargetMode == ESimplifyTargetType::Percentage)
-		{
-			ReductionSettings.PercentTriangles = FMath::Max(SimplifyProperties->TargetPercentage / 100., .001);  // Only support triangle percentage and count, but not edge length
-		}
-		else if (SimplifyProperties->TargetMode == ESimplifyTargetType::TriangleCount)
-		{
-			int32 NumTris = SrcMeshDescription->Polygons().Num();
-			ReductionSettings.PercentTriangles = (float)SimplifyProperties->TargetCount / (float)NumTris;
-		}
-
-		float Error;
-		{
-			IMeshReductionManagerModule& MeshReductionModule = FModuleManager::Get().LoadModuleChecked<IMeshReductionManagerModule>("MeshReductionInterface");
-			IMeshReduction* MeshReduction = MeshReductionModule.GetStaticMeshReductionInterface();
+		FDynamicMeshToMeshDescription Converter;
 
+		// full conversion if normal topology changed or faces were inverted
+		Converter.Convert(Result.Mesh.Get(), *MeshDescription);
+	});
 
-			if (!MeshReduction)
-			{
-				// no reduction possible
-				Error = 0.f;
-				return;
-			}
-
-			Error = ReductionSettings.MaxDeviation;
-			MeshReduction->ReduceMeshDescription(DstMeshDescription, Error, *SrcMeshDescription, OverlappingCorners, ReductionSettings);
-		}
-
-		// Put the reduced mesh into the target...
-		DynamicMeshComponent->InitializeMesh(&DstMeshDescription);
-
-		// The UE4 tool will split the UV boundaries.  Need to weld this.
-		{
-			FDynamicMesh3* ComponentMesh = DynamicMeshComponent->GetMesh();
-
-			FMergeCoincidentMeshEdges Merger(ComponentMesh);
-			Merger.MergeSearchTolerance = 10.0f * FMathf::ZeroTolerance;
-			Merger.OnlyUniquePairs = false;
-			if (Merger.Apply() == false)
-			{
-				DynamicMeshComponent->InitializeMesh(&DstMeshDescription);
-			}
-			if (ComponentMesh->CheckValidity(true, EValidityCheckFailMode::ReturnOnly) == false)
-			{
-				DynamicMeshComponent->InitializeMesh(&DstMeshDescription);
-			}
-		
-			DynamicMeshComponent->NotifyMeshUpdated();
-		}
-	}
-
-	//UE_LOG(LogMeshSimplification, Log, TEXT("Mesh simplified to %d triangles"), TargetMesh->TriangleCount());
-
-	DynamicMeshComponent->NotifyMeshUpdated();
-	GetToolManager()->PostInvalidation();
-	MeshStatisticsProperties->Update(*DynamicMeshComponent->GetMesh());
-
-	bResultValid = true;
-}
-
-bool USimplifyMeshTool::HasAccept() const
-{
-	return true;
-}
-
-bool USimplifyMeshTool::CanAccept() const
-{
-	return true;
+	GetToolManager()->EndUndoTransaction();
 }
 
 
diff --git a/Engine/Plugins/Experimental/MeshModelingToolset/Source/MeshModelingToolsEditorOnly/Public/SimplifyMeshTool.h b/Engine/Plugins/Experimental/MeshModelingToolset/Source/MeshModelingToolsEditorOnly/Public/SimplifyMeshTool.h
--- a/Engine/Plugins/Experimental/MeshModelingToolset/Source/MeshModelingToolsEditorOnly/Public/SimplifyMeshTool.h
+++ b/Engine/Plugins/Experimental/MeshModelingToolset/Source/MeshModelingToolsEditorOnly/Public/SimplifyMeshTool.h
@@ -8,44 +8,12 @@
 #include "InteractiveToolBuilder.h"
 #include "DynamicMesh3.h"
 #include "DynamicMeshAABBTree3.h"
+#include "MeshOpPreviewHelpers.h"
+#include "CleaningOps/SimplifyMeshOp.h"
 #include "Properties/MeshStatisticsProperties.h"
-#include "SimplifyMeshTool.generated.h"
-
-
-// predeclarations
-struct FMeshDescription;
-class USimpleDynamicMeshComponent;
-
-
-
-
-/**  */
-UENUM()
-enum class ESimplifyTargetType : uint8
-{
-	/** Percentage of input triangles */
-	Percentage = 0 UMETA(DisplayName = "Percentage"),
-
-	/** Target triangle count */
-	TriangleCount = 1 UMETA(DisplayName = "Triangle Count"),
-
-	/** Target edge length */
-	EdgeLength = 2 UMETA(DisplayName = "Edge Length")
-};
-
-/**  */
-UENUM()
-enum class ESimplifyType : uint8
-{
-	/** Fastest. Standard quadric error metric. Will not simplify UV boundaries.*/
-	QEM = 0 UMETA(DisplayName = "QEM"),
 
-	/** Potentially higher quality. Takes the normal into account. Will not simplify UV bounaries. */
-	Attribute = 1 UMETA(DisplayName = "Normal Aware"),
+#include "SimplifyMeshTool.generated.h"
 
-	/** Highest quality reduction.  Will simplify UV boundaries. */
-	UE4Standard = 2 UMETA(DisplayName = "UE4 Standard"),
-};
 
 
 
@@ -59,6 +27,8 @@ class MESHMODELINGTOOLSEDITORONLY_API USimplifyMeshToolBuilder : public UInterac
 	GENERATED_BODY()
 
 public:
+	IToolsContextAssetAPI* AssetAPI = nullptr;
+
 	virtual bool CanBuildTool(const FToolBuilderState& SceneState) const override;
 	virtual UInteractiveTool* BuildTool(const FToolBuilderState& SceneState) const override;
 };
@@ -116,14 +86,18 @@ public:
  * Simple Mesh Simplifying Tool
  */
 UCLASS()
-class MESHMODELINGTOOLSEDITORONLY_API USimplifyMeshTool : public USingleSelectionTool
+class MESHMODELINGTOOLSEDITORONLY_API USimplifyMeshTool : public USingleSelectionTool, public IDynamicMeshOperatorFactory
 {
 	GENERATED_BODY()
 
 public:
+	virtual void SetWorld(UWorld* World);
+	virtual void SetAssetAPI(IToolsContextAssetAPI* AssetAPI);
+
 	virtual void Setup() override;
 	virtual void Shutdown(EToolShutdownType ShutdownType) override;
 
+	virtual void Tick(float DeltaTime) override;
 	virtual void Render(IToolsContextRenderAPI* RenderAPI) override;
 
 	virtual bool HasCancel() const override { return true; }
@@ -132,21 +106,27 @@ public:
 
 	virtual void OnPropertyModified(UObject* PropertySet, UProperty* Property) override;
 
-protected:
-	// need to update bResultValid if these are modified, so we don't publicly expose them. 
-	// @todo setters/getters for these
+	// IDynamicMeshOperatorFactory API
+	virtual TSharedPtr<FDynamicMeshOperator> MakeNewOperator() override;
 
+protected:
 	UPROPERTY()
 	USimplifyMeshToolProperties* SimplifyProperties;
 
 	UPROPERTY()
 	UMeshStatisticsProperties* MeshStatisticsProperties;
 
+	UPROPERTY()
+	UMeshOpPreviewWithBackgroundCompute* Preview;
+
 protected:
-	USimpleDynamicMeshComponent* DynamicMeshComponent;
-	FDynamicMesh3 OriginalMesh;
-	FDynamicMeshAABBTree3 OriginalMeshSpatial;
-	bool bResultValid;
-	void UpdateResult();
+	UWorld* TargetWorld;
+	IToolsContextAssetAPI* AssetAPI;
+
+	TSharedPtr<FMeshDescription> OriginalMeshDescription;
+	// Dynamic Mesh versions precomputed in Setup (rather than recomputed for every simplify op)
+	TSharedPtr<FDynamicMesh3> OriginalMesh;
+	TSharedPtr<FDynamicMeshAABBTree3> OriginalMeshSpatial;
 
+	void GenerateAsset(const FDynamicMeshOpResult& Result);
 };
diff --git a/Engine/Plugins/Experimental/MeshModelingToolset/Source/ModelingOperatorsEditorOnly/Private/CleaningOps/SimplifyMeshOp.cpp b/Engine/Plugins/Experimental/MeshModelingToolset/Source/ModelingOperatorsEditorOnly/Private/CleaningOps/SimplifyMeshOp.cpp
new file mode 100644
--- /dev/null
+++ b/Engine/Plugins/Experimental/MeshModelingToolset/Source/ModelingOperatorsEditorOnly/Private/CleaningOps/SimplifyMeshOp.cpp
@@ -0,0 +1,205 @@
+// Copyright 1998-2019 Epic Games, Inc. All Rights Reserved.
+
+#pragma once
+
+#include "CleaningOps/SimplifyMeshOp.h"
+
+#include "DynamicMeshAttributeSet.h"
+#include "MeshDescriptionToDynamicMesh.h"
+
+
+#include "DynamicMesh3.h"
+#include "DynamicMeshAttributeSet.h"
+#include "MeshSimplification.h"
+#include "MeshConstraintsUtil.h"
+#include "ProjectionTargets.h"
+
+#include "Operations/MergeCoincidentMeshEdges.h"
+#include "MeshDescription.h"
+#include "MeshDescriptionOperations.h"
+#include "OverlappingCorners.h"
+#include "IMeshReductionInterfaces.h"
+
+#include "MeshNormals.h"
+
+
+template <typename SimplificationType>
+void ComputeSimplify(FDynamicMesh3* TargetMesh, const bool bReproject, int OriginalTriCount, FDynamicMesh3& OriginalMesh, FDynamicMeshAABBTree3& OriginalMeshSpatial, const ESimplifyTargetType TargetMode, const float TargetPercentage, const int TargetCount, const float TargetEdgeLength)
+{
+	SimplificationType Reducer(TargetMesh);
+
+	Reducer.ProjectionMode = (bReproject) ? 
+		SimplificationType::ETargetProjectionMode::AfterRefinement : SimplificationType::ETargetProjectionMode::NoProjection;
+
+	Reducer.DEBUG_CHECK_LEVEL = 0;
+
+	FMeshConstraints constraints;
+	FMeshConstraintsUtil::ConstrainAllSeams(constraints, *TargetMesh, true, false);
+	//FMeshConstraintsUtil::ConstrainAllSeamJunctions(constraints, *TargetMesh, true, false);
+	Reducer.SetExternalConstraints(&constraints);
+
+	FMeshProjectionTarget ProjTarget(&OriginalMesh, &OriginalMeshSpatial);
+	Reducer.SetProjectionTarget(&ProjTarget);
+
+	if (TargetMode == ESimplifyTargetType::Percentage)
+	{
+		double Ratio = (double)TargetPercentage / 100.0;
+		int UseTarget = FMath::Max(4, (int)(Ratio * (double)OriginalTriCount));
+		Reducer.SimplifyToTriangleCount(UseTarget);
+	} 
+	else if (TargetMode == ESimplifyTargetType::TriangleCount)
+	{
+		Reducer.SimplifyToTriangleCount(TargetCount);
+	}
+	else if (TargetMode == ESimplifyTargetType::EdgeLength)
+	{
+		Reducer.SimplifyToEdgeLength(TargetEdgeLength);
+	}
+}
+
+
+
+void FSimplifyMeshOp::CalculateResult(FProgressCancel* Progress)
+{
+	if (Progress->Cancelled())
+	{
+		return;
+	}
+
+	// Need access to the source mesh:
+	const FMeshDescription* MeshDescription = OriginalMeshDescription.Get();
+	FDynamicMesh3* TargetMesh = ResultMesh.Get();
+
+	if (Progress->Cancelled())
+	{
+		return;
+	}
+
+	int OriginalTriCount = OriginalMesh->TriangleCount();
+	if (SimplifierType != ESimplifyType::UE4Standard)
+	{
+		// GeometryProcessing-specific methods that operate directly on a DynamicMesh3
+		ResultMesh->Copy(*OriginalMesh, true, true, true, !bDiscardAttributes);
+
+		if (SimplifierType == ESimplifyType::QEM)
+		{
+			ComputeSimplify<FQEMSimplification>(TargetMesh, bReproject, OriginalTriCount, *OriginalMesh, *OriginalMeshSpatial, 
+				TargetMode, TargetPercentage, TargetCount, TargetEdgeLength);
+		}
+		else if (SimplifierType == ESimplifyType::Attribute)
+		{
+			ComputeSimplify<FAttrMeshSimplification>(TargetMesh, bReproject, OriginalTriCount, *OriginalMesh, *OriginalMeshSpatial,
+				TargetMode, TargetPercentage, TargetCount, TargetEdgeLength);
+		}
+	}
+	else
+	{
+		ensure(SimplifierType == ESimplifyType::UE4Standard);
+		const FMeshDescription* SrcMeshDescription = OriginalMeshDescription.Get();
+		FMeshDescription DstMeshDescription(*SrcMeshDescription);
+
+		if (Progress->Cancelled())
+		{
+			return;
+		}
+
+		FOverlappingCorners OverlappingCorners;
+		FMeshDescriptionOperations::FindOverlappingCorners(OverlappingCorners, *SrcMeshDescription, 1.e-5);
+
+		if (Progress->Cancelled())
+		{
+			return;
+		}
+
+		FMeshReductionSettings ReductionSettings;
+		if (TargetMode == ESimplifyTargetType::Percentage)
+		{
+			ReductionSettings.PercentTriangles = FMath::Max(TargetPercentage / 100., .001);  // Only support triangle percentage and count, but not edge length
+		}
+		else if (TargetMode == ESimplifyTargetType::TriangleCount)
+		{
+			int32 NumTris = SrcMeshDescription->Polygons().Num();
+			ReductionSettings.PercentTriangles = (float)TargetCount / (float)NumTris;
+		}
+
+		float Error;
+		{
+			if (!MeshReduction)
+			{
+				// no reduction possible, failed to load the required interface
+				Error = 0.f;
+				ResultMesh->Copy(*OriginalMesh, true, true, true, !bDiscardAttributes);
+				return;
+			}
+
+			Error = ReductionSettings.MaxDeviation;
+			MeshReduction->ReduceMeshDescription(DstMeshDescription, Error, *SrcMeshDescription, OverlappingCorners, ReductionSettings);
+		}
+
+		if (Progress->Cancelled())
+		{
+			return;
+		}
+
+		// Put the reduced mesh into the target...
+		FMeshDescriptionToDynamicMesh Converter;
+		Converter.Convert(&DstMeshDescription, *ResultMesh);
+		if (bDiscardAttributes)
+		{
+			ResultMesh->DiscardAttributes();
+		}
+
+
+		bool bFailedModifyNeedsRegen = false;
+		// The UE4 tool will split the UV boundaries.  Need to weld this.
+		{
+			FDynamicMesh3* ComponentMesh = ResultMesh.Get();
+
+			FMergeCoincidentMeshEdges Merger(ComponentMesh);
+			Merger.MergeSearchTolerance = 10.0f * FMathf::ZeroTolerance;
+			Merger.OnlyUniquePairs = false;
+			if (Merger.Apply() == false)
+			{
+				bFailedModifyNeedsRegen = true;
+			}
+
+			if (Progress->Cancelled())
+			{
+				return;
+			}
+
+			if (ResultMesh->CheckValidity(true, EValidityCheckFailMode::ReturnOnly) == false)
+			{
+				bFailedModifyNeedsRegen = true;
+			}
+
+			if (Progress->Cancelled())
+			{
+				return;
+			}
+
+			// in the fallback case where merge edges failed, give up and reset it to what it was before the attempted merger (w/ split UV boundaries everywhere, oh well)
+			if (bFailedModifyNeedsRegen)
+			{
+				ResultMesh->Clear();
+				Converter.Convert(&DstMeshDescription, *ResultMesh);
+				if (bDiscardAttributes)
+				{
+					ResultMesh->DiscardAttributes();
+				}
+			}
+		}
+	}
+
+
+	if (Progress->Cancelled())
+	{
+		return;
+	}
+
+	if (!ResultMesh->HasAttributes())
+	{
+		FMeshNormals::QuickComputeVertexNormals(*ResultMesh);
+	}
+
+}
\ No newline at end of file
diff --git a/Engine/Plugins/Experimental/MeshModelingToolset/Source/ModelingOperatorsEditorOnly/Public/CleaningOps/SimplifyMeshOp.h b/Engine/Plugins/Experimental/MeshModelingToolset/Source/ModelingOperatorsEditorOnly/Public/CleaningOps/SimplifyMeshOp.h
new file mode 100644
--- /dev/null
+++ b/Engine/Plugins/Experimental/MeshModelingToolset/Source/ModelingOperatorsEditorOnly/Public/CleaningOps/SimplifyMeshOp.h
@@ -0,0 +1,82 @@
+// Copyright 1998-2019 Epic Games, Inc. All Rights Reserved.
+
+#pragma once
+
+#include "CoreMinimal.h"
+#include "Util/ProgressCancel.h"
+#include "ModelingOperators.h"
+
+#include "SimplifyMeshOp.generated.h"
+
+class FDynamicMesh3;
+
+template <class TriangleMeshType>
+class TMeshAABBTree3;
+typedef TMeshAABBTree3<FDynamicMesh3> FDynamicMeshAABBTree3;
+
+
+UENUM()
+enum class ESimplifyTargetType : uint8
+{
+	/** Percentage of input triangles */
+	Percentage = 0 UMETA(DisplayName = "Percentage"),
+
+	/** Target triangle count */
+	TriangleCount = 1 UMETA(DisplayName = "Triangle Count"),
+
+	/** Target edge length */
+	EdgeLength = 2 UMETA(DisplayName = "Edge Length")
+};
+
+UENUM()
+enum class ESimplifyType : uint8
+{
+	/** Fastest. Standard quadric error metric. Will not simplify UV boundaries.*/
+	QEM = 0 UMETA(DisplayName = "QEM"),
+
+	/** Potentially higher quality. Takes the normal into account. Will not simplify UV bounaries. */
+	Attribute = 1 UMETA(DisplayName = "Normal Aware"),
+
+	/** Highest quality reduction.  Will simplify UV boundaries. */
+	UE4Standard = 2 UMETA(DisplayName = "UE4 Standard"),
+};
+
+
+
+
+
+class MODELINGOPERATORSEDITORONLY_API FSimplifyMeshOp : public FDynamicMeshOperator
+{
+public:
+	virtual ~FSimplifyMeshOp() {}
+
+	//
+	// Inputs
+	// 
+	ESimplifyTargetType TargetMode;
+	ESimplifyType SimplifierType;
+	int TargetPercentage, TargetCount;
+	float TargetEdgeLength;
+	bool bDiscardAttributes, bReproject, bPreventNormalFlips;
+
+	// stored for the UE4 Standard path
+	TSharedPtr<FMeshDescription> OriginalMeshDescription;
+	// stored for the GeometryProcessing custom simplifier paths (currently precomputed once in tool setup)
+	TSharedPtr<FDynamicMesh3> OriginalMesh;
+	TSharedPtr<FDynamicMeshAABBTree3> OriginalMeshSpatial;
+
+	class IMeshReduction* MeshReduction;
+
+
+	// set ability on protected transform.
+	void SetTransform(const FTransform& Transform)
+	{
+		ResultTransform = (FTransform3d)Transform;
+	}
+
+	//
+	// FDynamicMeshOperator implementation
+	// 
+
+	virtual void CalculateResult(FProgressCancel* Progress) override;
+};
