fix: Editor freezes while Remesh Tool is computing

#rnx

diff --git a/Engine/Plugins/Experimental/MeshModelingToolset/Source/MeshModelingTools/Private/RemeshMeshTool.cpp b/Engine/Plugins/Experimental/MeshModelingToolset/Source/MeshModelingTools/Private/RemeshMeshTool.cpp
--- a/Engine/Plugins/Experimental/MeshModelingToolset/Source/MeshModelingTools/Private/RemeshMeshTool.cpp
+++ b/Engine/Plugins/Experimental/MeshModelingToolset/Source/MeshModelingTools/Private/RemeshMeshTool.cpp
@@ -4,17 +4,16 @@
 #include "InteractiveToolManager.h"
 #include "ToolBuilderUtil.h"
 
+#include "ToolSetupUtil.h"
+
 #include "DynamicMesh3.h"
-#include "DynamicMeshAttributeSet.h"
-#include "Remesher.h"
-#include "MeshConstraintsUtil.h"
-#include "ProjectionTargets.h"
-#include "MeshNormals.h"
 
-#include "SimpleDynamicMeshComponent.h"
+#include "MeshDescriptionToDynamicMesh.h"
+#include "DynamicMeshToMeshDescription.h"
+
+#include "AssetGenerationUtil.h"
 
 #include "SceneManagement.h" // for FPrimitiveDrawInterface
-#include "Async/ParallelFor.h"
 
 #define LOCTEXT_NAMESPACE "URemeshMeshTool"
 
@@ -31,12 +30,15 @@ bool URemeshMeshToolBuilder::CanBuildTool(const FToolBuilderState& SceneState) c
 
 UInteractiveTool* URemeshMeshToolBuilder::BuildTool(const FToolBuilderState& SceneState) const
 {
-	URemeshMeshTool* NewTool = NewObject<URemeshMeshTool>(SceneState.ToolManager, "Remesh Tool");
+	URemeshMeshTool* NewTool = NewObject<URemeshMeshTool>(SceneState.ToolManager);
 
 	UActorComponent* ActorComponent = ToolBuilderUtil::FindFirstComponent(SceneState, CanMakeComponentTarget);
 	auto* MeshComponent = Cast<UPrimitiveComponent>(ActorComponent);
 	check(MeshComponent != nullptr);
+
 	NewTool->SetSelection(MakeComponentTarget(MeshComponent));
+	NewTool->SetWorld(SceneState.World);
+	NewTool->SetAssetAPI(AssetAPI);
 
 	return NewTool;
 }
@@ -72,93 +74,132 @@ URemeshMeshTool::URemeshMeshTool()
 }
 
 
+void URemeshMeshTool::SetWorld(UWorld* World)
+{
+	this->TargetWorld = World;
+}
+
+void URemeshMeshTool::SetAssetAPI(IToolsContextAssetAPI* AssetAPIIn)
+{
+	this->AssetAPI = AssetAPIIn;
+}
+
+
 void URemeshMeshTool::Setup()
 {
 	UInteractiveTool::Setup();
 
-	// create dynamic mesh component to use for live preview
-	DynamicMeshComponent = NewObject<USimpleDynamicMeshComponent>(ComponentTarget->GetOwnerActor(), "DynamicMesh");
-	DynamicMeshComponent->SetupAttachment(ComponentTarget->GetOwnerActor()->GetRootComponent());
-	DynamicMeshComponent->RegisterComponent();
-	DynamicMeshComponent->SetWorldTransform(ComponentTarget->GetWorldTransform());
+	BasicProperties = NewObject<URemeshMeshToolProperties>(this);
+	MeshStatisticsProperties = NewObject<UMeshStatisticsProperties>(this);
 
-	// copy material if there is one
-	auto Material = ComponentTarget->GetMaterial(0);
-	if (Material != nullptr)
+	// hide component and create + show preview
+	ComponentTarget->SetOwnerVisibility(false);
+	Preview = NewObject<UMeshOpPreviewWithBackgroundCompute>(this, "Preview");
+	Preview->Setup(this->TargetWorld, this);
+	Preview->OnMeshUpdated.AddLambda([this](UMeshOpPreviewWithBackgroundCompute* Compute)
 	{
-		DynamicMeshComponent->SetMaterial(0, Material);
-	}
-	DynamicMeshComponent->bExplicitShowWireframe = true;
+		MeshStatisticsProperties->Update(*Compute->PreviewMesh->GetPreviewDynamicMesh());
+	});
+	Preview->ConfigureMaterials(
+		ToolSetupUtil::GetDefaultMaterial(GetToolManager(), ComponentTarget->GetMaterial(0)),
+		ToolSetupUtil::GetDefaultWorkingMaterial(GetToolManager())
+	);
+	Preview->PreviewMesh->EnableWireframe(true);
 
-	DynamicMeshComponent->InitializeMesh(ComponentTarget->GetMesh());
-	OriginalMesh.Copy(*DynamicMeshComponent->GetMesh());
-	OriginalMeshSpatial.SetMesh(&OriginalMesh, true);
+	OriginalMesh = MakeShared<FDynamicMesh3>();
+	FMeshDescriptionToDynamicMesh Converter;
+	Converter.bPrintDebugMessages = true;
+	Converter.Convert(ComponentTarget->GetMesh(), *OriginalMesh);
 
-	// hide input StaticMeshComponent
-	ComponentTarget->SetOwnerVisibility(false);
+	Preview->PreviewMesh->SetTransform(ComponentTarget->GetWorldTransform());
+	Preview->PreviewMesh->UpdatePreview(OriginalMesh.Get());
+
+	OriginalMeshSpatial = MakeShared<FDynamicMeshAABBTree3>(OriginalMesh.Get(), true);
 
 	// calculate initial mesh area (no utility fn yet)
+	// TODO: will need to change to account for component transform's Scale3D
 	InitialMeshArea = 0;
-	for (int tid : OriginalMesh.TriangleIndicesItr())
+	for (int tid : OriginalMesh->TriangleIndicesItr())
 	{
-		InitialMeshArea += OriginalMesh.GetTriArea(tid);
+		InitialMeshArea += OriginalMesh->GetTriArea(tid);
 	}
-
-	BasicProperties = NewObject<URemeshMeshToolProperties>(this);
+	
+	// set properties defaults
+	
 	// arbitrary threshold of 5000 tris seems reasonable? 
-	BasicProperties->TargetTriangleCount = (OriginalMesh.TriangleCount() < 5000) ? 5000 : OriginalMesh.TriangleCount();
+	BasicProperties->TargetTriangleCount = (OriginalMesh->TriangleCount() < 5000) ? 5000 : OriginalMesh->TriangleCount();
 	BasicProperties->TargetEdgeLength = CalculateTargetEdgeLength(BasicProperties->TargetTriangleCount);
 
-	// initialize our properties
+	// add properties to GUI
 	AddToolPropertySource(BasicProperties);
-
-	MeshStatisticsProperties = NewObject<UMeshStatisticsProperties>(this);
 	AddToolPropertySource(MeshStatisticsProperties);
-	MeshStatisticsProperties->Update(*DynamicMeshComponent->GetMesh());
 
-	bResultValid = false;
+	Preview->InvalidateResult();
 }
 
 
 void URemeshMeshTool::Shutdown(EToolShutdownType ShutdownType)
 {
-	if (DynamicMeshComponent != nullptr)
+	ComponentTarget->SetOwnerVisibility(true);
+	TUniquePtr<FDynamicMeshOpResult> Result = Preview->Shutdown();
+	if (ShutdownType == EToolShutdownType::Accept)
 	{
-		//DynamicMeshComponent->OnMeshChanged.Remove(OnDynamicMeshComponentChangedHandle);
+		GenerateAsset(*Result);
+	}
+}
 
-		ComponentTarget->SetOwnerVisibility(true);
 
-		if (ShutdownType == EToolShutdownType::Accept)
-		{
-			// this block bakes the modified DynamicMeshComponent back into the StaticMeshComponent inside an undo transaction
-			GetToolManager()->BeginUndoTransaction(LOCTEXT("RemeshMeshToolTransactionName", "Remesh Mesh"));
-			ComponentTarget->CommitMesh([=](FMeshDescription* MeshDescription)
-			{
-				DynamicMeshComponent->Bake(MeshDescription, true);
-			});
-			GetToolManager()->EndUndoTransaction();
-		}
+void URemeshMeshTool::Tick(float DeltaTime)
+{
+	Preview->Tick(DeltaTime);
+}
+
+
+TSharedPtr<FDynamicMeshOperator> URemeshMeshTool::MakeNewOperator()
+{
+	TSharedPtr<FRemeshMeshOp> Op = MakeShared<FRemeshMeshOp>();
 
-		DynamicMeshComponent->UnregisterComponent();
-		DynamicMeshComponent->DestroyComponent();
-		DynamicMeshComponent = nullptr;
+	if (!BasicProperties->bUseTargetEdgeLength)
+	{
+		Op->TargetEdgeLength = CalculateTargetEdgeLength(BasicProperties->TargetTriangleCount);
+	}
+	else
+	{
+		Op->TargetEdgeLength = BasicProperties->TargetEdgeLength;
 	}
+
+	Op->bCollapses = BasicProperties->bCollapses;
+	Op->bDiscardAttributes = BasicProperties->bDiscardAttributes;
+	Op->bFlips = BasicProperties->bFlips;
+	Op->bPreserveSharpEdges = BasicProperties->bPreserveSharpEdges;
+	Op->bPreventNormalFlips = BasicProperties->bPreventNormalFlips;
+	Op->bReproject = BasicProperties->bReproject;
+	Op->bSplits = BasicProperties->bSplits;
+	Op->RemeshIterations = BasicProperties->RemeshIterations;
+	Op->SmoothingSpeed = BasicProperties->SmoothingSpeed;
+	Op->SmoothingType = BasicProperties->SmoothingType;
+	Op->TargetTriangleCount = BasicProperties->TargetTriangleCount;
+
+	FTransform LocalToWorld = ComponentTarget->GetWorldTransform();
+	Op->SetTransform(LocalToWorld);
+
+	Op->OriginalMesh = OriginalMesh;
+	Op->OriginalMeshSpatial = OriginalMeshSpatial;
+
+	return Op;
 }
 
 
 void URemeshMeshTool::Render(IToolsContextRenderAPI* RenderAPI)
 {
-	UpdateResult();
-
-
 	FPrimitiveDrawInterface* PDI = RenderAPI->GetPrimitiveDrawInterface();
-	FTransform Transform = ComponentTarget->GetWorldTransform(); //Actor->GetTransform();
+	FTransform Transform = ComponentTarget->GetWorldTransform();
 
 	FColor LineColor(255, 0, 0);
-	FDynamicMesh3* TargetMesh = DynamicMeshComponent->GetMesh();
-	if (TargetMesh->HasAttributes())
+	const FDynamicMesh3* TargetMesh = Preview->PreviewMesh->GetPreviewDynamicMesh();
+	if (TargetMesh && TargetMesh->HasAttributes())
 	{
-		FDynamicMeshUVOverlay* UVOverlay = TargetMesh->Attributes()->PrimaryUV();
+		const FDynamicMeshUVOverlay* UVOverlay = TargetMesh->Attributes()->PrimaryUV();
 		for (int eid : TargetMesh->EdgeIndicesItr()) 
 		{
 			if (UVOverlay->IsSeamEdge(eid)) 
@@ -175,7 +216,7 @@ void URemeshMeshTool::Render(IToolsContextRenderAPI* RenderAPI)
 
 void URemeshMeshTool::OnPropertyModified(UObject* PropertySet, UProperty* Property)
 {
-	bResultValid = false;
+	Preview->InvalidateResult();
 }
 
 
@@ -186,92 +227,6 @@ double URemeshMeshTool::CalculateTargetEdgeLength(int TargetTriCount)
 	return (double)FMath::RoundToInt(EdgeLen*100.0) / 100.0;
 }
 
-
-
-void URemeshMeshTool::UpdateResult()
-{
-	if (bResultValid) 
-	{
-		return;
-	}
-
-	FDynamicMesh3* TargetMesh = DynamicMeshComponent->GetMesh();
-	TargetMesh->Copy(OriginalMesh);
-
-	if (BasicProperties->bDiscardAttributes && BasicProperties->bPreserveSharpEdges == false)
-	{
-		TargetMesh->DiscardAttributes();
-	}
-
-	FRemesher Remesher(TargetMesh);
-	Remesher.bEnableSplits = BasicProperties->bSplits;
-	Remesher.bEnableFlips = BasicProperties->bFlips;
-	Remesher.bEnableCollapses = BasicProperties->bCollapses;
-
-	// recalculate target edge length
-	if (BasicProperties->bUseTargetEdgeLength == false)
-	{
-		BasicProperties->TargetEdgeLength = CalculateTargetEdgeLength(BasicProperties->TargetTriangleCount);
-	}
-
-	Remesher.SetTargetEdgeLength(BasicProperties->TargetEdgeLength);
-
-	Remesher.ProjectionMode = (BasicProperties->bReproject) ? 
-		FRemesher::ETargetProjectionMode::AfterRefinement : FRemesher::ETargetProjectionMode::NoProjection;
-
-	Remesher.bEnableSmoothing = (BasicProperties->SmoothingSpeed > 0);
-	Remesher.SmoothSpeedT = BasicProperties->SmoothingSpeed;
-	//Remesher.SmoothType = FRemesher::ESmoothTypes::MeanValue;
-	Remesher.SmoothType = (BasicProperties->bDiscardAttributes) ?
-		FRemesher::ESmoothTypes::Uniform :
-		(FRemesher::ESmoothTypes)(int)BasicProperties->SmoothingType;
-	bool bIsUniformSmooth = (Remesher.SmoothType == FRemesher::ESmoothTypes::Uniform);
-
-	Remesher.bPreventNormalFlips = BasicProperties->bPreventNormalFlips;
-
-	Remesher.DEBUG_CHECK_LEVEL = 0;
-
-	FMeshConstraints constraints;
-	FMeshConstraintsUtil::ConstrainAllSeams(constraints, *TargetMesh, true, !BasicProperties->bPreserveSharpEdges);
-	Remesher.SetExternalConstraints(&constraints);
-
-	FMeshProjectionTarget ProjTarget(&OriginalMesh, &OriginalMeshSpatial);
-	Remesher.SetProjectionTarget(&ProjTarget);
-
-	if (BasicProperties->bDiscardAttributes && BasicProperties->bPreserveSharpEdges == true)
-	{
-		TargetMesh->DiscardAttributes();
-	}
-
-	// run the remesh iterations
-	for (int k = 0; k < BasicProperties->RemeshIterations; ++k)
-	{
-		// If we are not uniform smoothing, then flips seem to often make things worse.
-		// Possibly this is because without the tangential flow, we won't get to the nice tris.
-		// In this case we are better off basically not flipping, and just letting collapses resolve things
-		// regular-valence polygons - things stay "stuck". 
-		// @todo try implementing edge-length flip criteria instead of valence-flip
-		if (bIsUniformSmooth == false)
-		{
-			bool bUseFlipsThisPass = (k % 2 == 0 && k < BasicProperties->RemeshIterations/2);
-			Remesher.bEnableFlips = bUseFlipsThisPass && BasicProperties->bFlips;
-		}
-
-		Remesher.BasicRemeshPass();
-	}
-
-	if (!TargetMesh->HasAttributes() && !TargetMesh->HasVertexNormals())
-	{
-		FMeshNormals::QuickComputeVertexNormals(*TargetMesh);
-	}
-
-	DynamicMeshComponent->NotifyMeshUpdated();
-	GetToolManager()->PostInvalidation();
-	MeshStatisticsProperties->Update(*DynamicMeshComponent->GetMesh());
-
-	bResultValid = true;
-}
-
 bool URemeshMeshTool::HasAccept() const
 {
 	return true;
@@ -279,10 +234,26 @@ bool URemeshMeshTool::HasAccept() const
 
 bool URemeshMeshTool::CanAccept() const
 {
-	return true;
+	return Preview->HaveValidResult();
 }
 
 
+void URemeshMeshTool::GenerateAsset(const FDynamicMeshOpResult& Result)
+{
+	GetToolManager()->BeginUndoTransaction(LOCTEXT("RemeshMeshToolTransactionName", "Remesh Mesh"));
+
+	check(Result.Mesh.Get() != nullptr);
+	ComponentTarget->CommitMesh([&Result](FMeshDescription* MeshDescription)
+	{
+		FDynamicMeshToMeshDescription Converter;
+
+		// full conversion if normal topology changed or faces were inverted
+		Converter.Convert(Result.Mesh.Get(), *MeshDescription);
+	});
+
+	GetToolManager()->EndUndoTransaction();
+}
+
 
 
 #undef LOCTEXT_NAMESPACE
diff --git a/Engine/Plugins/Experimental/MeshModelingToolset/Source/MeshModelingTools/Public/RemeshMeshTool.h b/Engine/Plugins/Experimental/MeshModelingToolset/Source/MeshModelingTools/Public/RemeshMeshTool.h
--- a/Engine/Plugins/Experimental/MeshModelingToolset/Source/MeshModelingTools/Public/RemeshMeshTool.h
+++ b/Engine/Plugins/Experimental/MeshModelingToolset/Source/MeshModelingTools/Public/RemeshMeshTool.h
@@ -8,26 +8,13 @@
 #include "InteractiveToolBuilder.h"
 #include "DynamicMesh3.h"
 #include "DynamicMeshAABBTree3.h"
+#include "MeshOpPreviewHelpers.h"
+#include "CleaningOps/RemeshMeshOp.h"
 #include "Properties/MeshStatisticsProperties.h"
-#include "RemeshMeshTool.generated.h"
-
-// predeclarations
-struct FMeshDescription;
-class USimpleDynamicMeshComponent;
 
-/** Mesh paint color view modes (somewhat maps to EVertexColorViewMode engine enum.) */
-UENUM()
-enum class ERemeshSmoothingType : uint8
-{
-	/** Uniform Smoothing */
-	Uniform = 0 UMETA(DisplayName = "Uniform"),
+#include "RemeshMeshTool.generated.h"
 
-	/** Cotangent Smoothing */
-	Cotangent = 1 UMETA(DisplayName = "Shape Preserving"),
 
-	/** Mean Value Smoothing */
-	MeanValue = 2 UMETA(DisplayName = "Mixed")
-};
 
 /**
  *
@@ -38,6 +25,8 @@ class MESHMODELINGTOOLS_API URemeshMeshToolBuilder : public UInteractiveToolBuil
 	GENERATED_BODY()
 
 public:
+	IToolsContextAssetAPI* AssetAPI = nullptr;
+
 	virtual bool CanBuildTool(const FToolBuilderState& SceneState) const override;
 	virtual UInteractiveTool* BuildTool(const FToolBuilderState& SceneState) const override;
 };
@@ -116,16 +105,20 @@ public:
  * Simple Mesh Remeshing Tool
  */
 UCLASS()
-class MESHMODELINGTOOLS_API URemeshMeshTool : public USingleSelectionTool
+class MESHMODELINGTOOLS_API URemeshMeshTool : public USingleSelectionTool, public IDynamicMeshOperatorFactory
 {
 	GENERATED_BODY()
 
 public:
 	URemeshMeshTool();
 
+	virtual void SetWorld(UWorld* World);
+	virtual void SetAssetAPI(IToolsContextAssetAPI* AssetAPI);
+
 	virtual void Setup() override;
 	virtual void Shutdown(EToolShutdownType ShutdownType) override;
 
+	virtual void Tick(float DeltaTime) override;
 	virtual void Render(IToolsContextRenderAPI* RenderAPI) override;
 
 	virtual bool HasCancel() const override { return true; }
@@ -134,25 +127,27 @@ public:
 
 	virtual void OnPropertyModified(UObject* PropertySet, UProperty* Property) override;
 
-public:
+	// IDynamicMeshOperatorFactory API
+	virtual TSharedPtr<FDynamicMeshOperator> MakeNewOperator() override;
 
+protected:
 	UPROPERTY()
 	URemeshMeshToolProperties* BasicProperties;
 
 	UPROPERTY()
 	UMeshStatisticsProperties* MeshStatisticsProperties;
 
+	UPROPERTY()
+	UMeshOpPreviewWithBackgroundCompute* Preview;
 
 protected:
-	USimpleDynamicMeshComponent* DynamicMeshComponent;
-	FDynamicMesh3 OriginalMesh;
-	FDynamicMeshAABBTree3 OriginalMeshSpatial;
-	bool bResultValid;
-	void UpdateResult();
+	UWorld* TargetWorld;
+	IToolsContextAssetAPI* AssetAPI;
 
+	TSharedPtr<FDynamicMesh3> OriginalMesh;
+	TSharedPtr<FDynamicMeshAABBTree3> OriginalMeshSpatial;
 	double InitialMeshArea;
-
 	double CalculateTargetEdgeLength(int TargetTriCount);
 
-	void ComputeRemeshing();
+	void GenerateAsset(const FDynamicMeshOpResult& Result);
 };
diff --git a/Engine/Plugins/Experimental/MeshModelingToolset/Source/ModelingOperators/Private/CleaningOps/RemeshMeshOp.cpp b/Engine/Plugins/Experimental/MeshModelingToolset/Source/ModelingOperators/Private/CleaningOps/RemeshMeshOp.cpp
new file mode 100644
--- /dev/null
+++ b/Engine/Plugins/Experimental/MeshModelingToolset/Source/ModelingOperators/Private/CleaningOps/RemeshMeshOp.cpp
@@ -0,0 +1,110 @@
+// Copyright 1998-2019 Epic Games, Inc. All Rights Reserved.
+
+#pragma once
+
+#include "CleaningOps/RemeshMeshOp.h"
+
+#include "DynamicMesh3.h"
+#include "DynamicMeshAABBTree3.h"
+#include "DynamicMeshAttributeSet.h"
+#include "Remesher.h"
+#include "MeshConstraintsUtil.h"
+#include "ProjectionTargets.h"
+#include "MeshNormals.h"
+
+void FRemeshMeshOp::SetTransform(const FTransform& Transform)
+{
+	ResultTransform = (FTransform3d)Transform;
+}
+
+void FRemeshMeshOp::CalculateResult(FProgressCancel* Progress)
+{
+	if (Progress->Cancelled())
+	{
+		return;
+	}
+
+	bool bDiscardAttributesImmediately = bDiscardAttributes && !bPreserveSharpEdges;
+	ResultMesh->Copy(*OriginalMesh, true, true, true, !bDiscardAttributesImmediately);
+
+	if (Progress->Cancelled())
+	{
+		return;
+	}
+
+	FDynamicMesh3* TargetMesh = ResultMesh.Get();
+
+	FRemesher Remesher(TargetMesh);
+	Remesher.bEnableSplits = bSplits;
+	Remesher.bEnableFlips = bFlips;
+	Remesher.bEnableCollapses = bCollapses;
+
+	Remesher.SetTargetEdgeLength(TargetEdgeLength);
+
+	Remesher.ProjectionMode = (bReproject) ? 
+		FRemesher::ETargetProjectionMode::AfterRefinement : FRemesher::ETargetProjectionMode::NoProjection;
+
+	Remesher.bEnableSmoothing = (SmoothingSpeed > 0);
+	Remesher.SmoothSpeedT = SmoothingSpeed;
+	// convert smooth type from UI enum to (currently 1:1) FRemesher enum
+	Remesher.SmoothType = FRemesher::ESmoothTypes::Uniform;
+	if (!bDiscardAttributes)
+	{
+		switch (SmoothingType)
+		{
+		case ERemeshSmoothingType::Uniform:
+			Remesher.SmoothType = FRemesher::ESmoothTypes::Uniform;
+			break;
+		case ERemeshSmoothingType::Cotangent:
+			Remesher.SmoothType = FRemesher::ESmoothTypes::Cotan;
+			break;
+		case ERemeshSmoothingType::MeanValue:
+			Remesher.SmoothType = FRemesher::ESmoothTypes::MeanValue;
+			break;
+		default:
+			ensure(false);
+		}
+	}
+	bool bIsUniformSmooth = (Remesher.SmoothType == FRemesher::ESmoothTypes::Uniform);
+
+	Remesher.bPreventNormalFlips = bPreventNormalFlips;
+
+	Remesher.DEBUG_CHECK_LEVEL = 0;
+
+	FMeshConstraints constraints;
+	FMeshConstraintsUtil::ConstrainAllSeams(constraints, *TargetMesh, true, !bPreserveSharpEdges);
+	Remesher.SetExternalConstraints(&constraints);
+
+	FMeshProjectionTarget ProjTarget(OriginalMesh.Get(), OriginalMeshSpatial.Get());
+	Remesher.SetProjectionTarget(&ProjTarget);
+
+	Remesher.Progress = Progress;
+
+	if (bDiscardAttributes && !bDiscardAttributesImmediately)
+	{
+		TargetMesh->DiscardAttributes();
+	}
+
+	// run the remesh iterations
+	for (int k = 0; k < RemeshIterations; ++k)
+	{
+		// If we are not uniform smoothing, then flips seem to often make things worse.
+		// Possibly this is because without the tangential flow, we won't get to the nice tris.
+		// In this case we are better off basically not flipping, and just letting collapses resolve things
+		// regular-valence polygons - things stay "stuck". 
+		// @todo try implementing edge-length flip criteria instead of valence-flip
+		if (bIsUniformSmooth == false)
+		{
+			bool bUseFlipsThisPass = (k % 2 == 0 && k < RemeshIterations/2);
+			Remesher.bEnableFlips = bUseFlipsThisPass && bFlips;
+		}
+
+		Remesher.BasicRemeshPass();
+	}
+
+	if (!TargetMesh->HasAttributes())
+	{
+		FMeshNormals::QuickComputeVertexNormals(*TargetMesh);
+	}
+
+}
\ No newline at end of file
diff --git a/Engine/Plugins/Experimental/MeshModelingToolset/Source/ModelingOperators/Public/CleaningOps/RemeshMeshOp.h b/Engine/Plugins/Experimental/MeshModelingToolset/Source/ModelingOperators/Public/CleaningOps/RemeshMeshOp.h
new file mode 100644
--- /dev/null
+++ b/Engine/Plugins/Experimental/MeshModelingToolset/Source/ModelingOperators/Public/CleaningOps/RemeshMeshOp.h
@@ -0,0 +1,58 @@
+// Copyright 1998-2019 Epic Games, Inc. All Rights Reserved.
+
+#pragma once
+
+#include "CoreMinimal.h"
+#include "Util/ProgressCancel.h"
+#include "ModelingOperators.h"
+
+#include "RemeshMeshOp.generated.h"
+
+class FDynamicMesh3;
+
+template <class TriangleMeshType>
+class TMeshAABBTree3;
+typedef TMeshAABBTree3<FDynamicMesh3> FDynamicMeshAABBTree3;
+
+
+/** Remeshing modes */
+UENUM()
+enum class ERemeshSmoothingType : uint8
+{
+	/** Uniform Smoothing */
+	Uniform = 0 UMETA(DisplayName = "Uniform"),
+
+	/** Cotangent Smoothing */
+	Cotangent = 1 UMETA(DisplayName = "Shape Preserving"),
+
+	/** Mean Value Smoothing */
+	MeanValue = 2 UMETA(DisplayName = "Mixed")
+};
+
+
+
+
+
+class MODELINGOPERATORS_API FRemeshMeshOp : public FDynamicMeshOperator
+{
+public:
+	virtual ~FRemeshMeshOp() {}
+
+	// inputs
+	TSharedPtr<FDynamicMesh3> OriginalMesh;
+	TSharedPtr<FDynamicMeshAABBTree3> OriginalMeshSpatial;
+
+	int TargetTriangleCount, RemeshIterations;
+	float SmoothingSpeed, TargetEdgeLength;
+	ERemeshSmoothingType SmoothingType;
+	bool bDiscardAttributes, bPreserveSharpEdges, bFlips, bSplits, bCollapses, bReproject, bPreventNormalFlips;
+
+	void SetTransform(const FTransform& Transform);
+
+	//
+	// FDynamicMeshOperator implementation
+	// 
+	virtual void CalculateResult(FProgressCancel* Progress) override;
+};
+
+
