CIS fix missed forward declare of FMeshDescription
#rnx

diff --git a/Engine/Plugins/Experimental/MeshModelingToolset/Source/ModelingOperatorsEditorOnly/Public/CleaningOps/SimplifyMeshOp.h b/Engine/Plugins/Experimental/MeshModelingToolset/Source/ModelingOperatorsEditorOnly/Public/CleaningOps/SimplifyMeshOp.h
--- a/Engine/Plugins/Experimental/MeshModelingToolset/Source/ModelingOperatorsEditorOnly/Public/CleaningOps/SimplifyMeshOp.h
+++ b/Engine/Plugins/Experimental/MeshModelingToolset/Source/ModelingOperatorsEditorOnly/Public/CleaningOps/SimplifyMeshOp.h
@@ -6,9 +6,11 @@
 #include "Util/ProgressCancel.h"
 #include "ModelingOperators.h"
 
+
 #include "SimplifyMeshOp.generated.h"
 
 class FDynamicMesh3;
+struct FMeshDescription;
 
 template <class TriangleMeshType>
 class TMeshAABBTree3;
