fix CIS static analysis warning due to always-true comparison

#rnx

diff --git a/Engine/Plugins/Experimental/MeshModelingToolset/Source/ModelingOperatorsEditorOnly/Private/CleaningOps/SimplifyMeshOp.cpp b/Engine/Plugins/Experimental/MeshModelingToolset/Source/ModelingOperatorsEditorOnly/Private/CleaningOps/SimplifyMeshOp.cpp
--- a/Engine/Plugins/Experimental/MeshModelingToolset/Source/ModelingOperatorsEditorOnly/Private/CleaningOps/SimplifyMeshOp.cpp
+++ b/Engine/Plugins/Experimental/MeshModelingToolset/Source/ModelingOperatorsEditorOnly/Private/CleaningOps/SimplifyMeshOp.cpp
@@ -92,9 +92,8 @@ void FSimplifyMeshOp::CalculateResult(FProgressCancel* Progress)
 				TargetMode, TargetPercentage, TargetCount, TargetEdgeLength);
 		}
 	}
-	else
+	else // SimplifierType == ESimplifyType::UE4Standard
 	{
-		ensure(SimplifierType == ESimplifyType::UE4Standard);
 		const FMeshDescription* SrcMeshDescription = OriginalMeshDescription.Get();
 		FMeshDescription DstMeshDescription(*SrcMeshDescription);
 
