Merging CL# 9542043 from UE4/Dev-Editor
ModelingTools: Fix Pinch and Flatten brushes. These brushes were being applied much too quickly and produced unstable/undesirable results.
#rnx

diff --git a/Engine/Plugins/Experimental/GeometryProcessing/Source/DynamicMesh/Private/Remesher.cpp b/Engine/Plugins/Experimental/GeometryProcessing/Source/DynamicMesh/Private/Remesher.cpp
--- a/Engine/Plugins/Experimental/GeometryProcessing/Source/DynamicMesh/Private/Remesher.cpp
+++ b/Engine/Plugins/Experimental/GeometryProcessing/Source/DynamicMesh/Private/Remesher.cpp
@@ -81,7 +81,7 @@ void FRemesher::BasicRemeshPass()
 	}
 
 	ProfileBeginSmooth();
-	if (bEnableSmoothing && SmoothSpeedT > 0)
+	if (bEnableSmoothing && (SmoothSpeedT > 0 || CustomSmoothSpeedF) )
 	{
 		if (bEnableSmoothInPlace) 
 		{
@@ -562,6 +562,12 @@ FVector3d FRemesher::ComputeSmoothedVertexPos(int vID,
 		return Mesh->GetVertex(vID);
 	}
 
+	double UseSmoothSpeed = (CustomSmoothSpeedF) ? CustomSmoothSpeedF(*Mesh, vID) : SmoothSpeedT;
+	if (UseSmoothSpeed <= 0)
+	{
+		return Mesh->GetVertex(vID);
+	}
+
 	FVector3d vSmoothed = smoothFunc(*Mesh, vID, SmoothSpeedT);
 	check(VectorUtil::IsFinite(vSmoothed));     // this will really catch a lot of bugs...
 
diff --git a/Engine/Plugins/Experimental/GeometryProcessing/Source/DynamicMesh/Public/Remesher.h b/Engine/Plugins/Experimental/GeometryProcessing/Source/DynamicMesh/Public/Remesher.h
--- a/Engine/Plugins/Experimental/GeometryProcessing/Source/DynamicMesh/Public/Remesher.h
+++ b/Engine/Plugins/Experimental/GeometryProcessing/Source/DynamicMesh/Public/Remesher.h
@@ -76,6 +76,9 @@ public:
 	/** Override default smoothing function */
 	TFunction<FVector3d(const FDynamicMesh3&, int, double)> CustomSmoothF;
 
+	/** Override constant SmoothSpeedT with function */
+	TFunction<double(const FDynamicMesh3&, int)> CustomSmoothSpeedF;
+
 
 
 	/** enable parallel projection. Only applied in AfterRefinement mode */
diff --git a/Engine/Plugins/Experimental/MeshModelingToolset/Source/MeshModelingTools/Private/DynamicMeshSculptTool.cpp b/Engine/Plugins/Experimental/MeshModelingToolset/Source/MeshModelingTools/Private/DynamicMeshSculptTool.cpp
--- a/Engine/Plugins/Experimental/MeshModelingToolset/Source/MeshModelingTools/Private/DynamicMeshSculptTool.cpp
+++ b/Engine/Plugins/Experimental/MeshModelingToolset/Source/MeshModelingTools/Private/DynamicMeshSculptTool.cpp
@@ -312,6 +312,7 @@ void UDynamicMeshSculptTool::OnBeginDrag(const FRay& Ray)
 		LastBrushPosNormalWorld = ActiveDragPlane;
 		FTransform Transform = ComponentTarget->GetWorldTransform();
 		LastBrushPosLocal = Transform.InverseTransformPosition(LastHitPosWorld);
+		LastSmoothBrushPosLocal = LastBrushPosLocal;
 
 		BeginChange(bEnableRemeshing == false);
 
@@ -319,7 +320,7 @@ void UDynamicMeshSculptTool::OnBeginDrag(const FRay& Ray)
 
 		if (SculptProperties->PrimaryBrushType == EDynamicMeshSculptBrushType::Plane)
 		{
-			ActiveFixedBrushPlane = ComputeROIBrushPlane(LastBrushPosLocal);
+			ActiveFixedBrushPlane = ComputeROIBrushPlane(LastBrushPosLocal, false);
 		}
 
 		// apply initial stamp
@@ -440,6 +441,9 @@ void UDynamicMeshSculptTool::ApplyStamp(const FRay& WorldRay)
 		case EDynamicMeshSculptBrushType::Move:
 			ApplyMoveBrush(WorldRay);
 			break;
+		case EDynamicMeshSculptBrushType::Smooth:
+			ApplySmoothBrush(WorldRay);
+			break;
 		case EDynamicMeshSculptBrushType::Pinch:
 			ApplyPinchBrush(WorldRay);
 			break;
@@ -484,13 +488,17 @@ void UDynamicMeshSculptTool::ApplySmoothBrush(const FRay& WorldRay)
 		return;
 	}
 
-
 	FTransform Transform = ComponentTarget->GetWorldTransform();
 	FVector NewBrushPosLocal = Transform.InverseTransformPosition(LastBrushPosWorld);
 
 	FDynamicMesh3* Mesh = DynamicMeshComponent->GetMesh();
-	for (int VertIdx : VertexROI)
+	int NumV = VertexROI.Num();
+	ROIPositionBuffer.SetNum(NumV, false);
+
+	ParallelFor(NumV, [this, Mesh, NewBrushPosLocal](int k)
 	{
+		int VertIdx = VertexROI[k];
+
 		FVector3d OrigPos = Mesh->GetVertex(VertIdx);
 
 		double Falloff = CalculateBrushFalloff(OrigPos.Distance(NewBrushPosLocal));
@@ -500,6 +508,16 @@ void UDynamicMeshSculptTool::ApplySmoothBrush(const FRay& WorldRay)
 
 		FVector3d NewPos = FVector3d::Lerp(OrigPos, SmoothedPos, Falloff*SculptProperties->SmoothSpeed);
 
+		ROIPositionBuffer[k] = NewPos;
+
+	});
+
+
+	for (int k = 0; k < NumV; ++k)
+	{
+		int VertIdx = VertexROI[k];
+		const FVector3d& NewPos = ROIPositionBuffer[k];
+		FVector3d OrigPos = Mesh->GetVertex(VertIdx);
 		Mesh->SetVertex(VertIdx, NewPos);
 		UpdateSavedVertex(VertIdx, OrigPos, NewPos);
 	}
@@ -521,28 +539,42 @@ void UDynamicMeshSculptTool::ApplyMoveBrush(const FRay& WorldRay)
 
 	FVector3d MoveVec = NewBrushPosLocal - LastBrushPosLocal;
 
-	if (MoveVec.SquaredLength() > 0)
+	if (MoveVec.SquaredLength() <= 0)
 	{
-		FDynamicMesh3* Mesh = DynamicMeshComponent->GetMesh();
-		for (int VertIdx : VertexROI)
-		{
-			FVector3d OrigPos = Mesh->GetVertex(VertIdx);
+		goto skip_brush;
+	}
+
+	FDynamicMesh3* Mesh = DynamicMeshComponent->GetMesh();
+	int NumV = VertexROI.Num();
+	ROIPositionBuffer.SetNum(NumV, false);
 
-			double PrevDist = (OrigPos - LastBrushPosLocal).Length();
-			double NewDist = (OrigPos - NewBrushPosLocal).Length();
-			double UseDist = FMath::Min(PrevDist, NewDist);
+	ParallelFor(NumV, [this, Mesh, NewBrushPosLocal, MoveVec](int k)
+	{
+		int VertIdx = VertexROI[k];
+		FVector3d OrigPos = Mesh->GetVertex(VertIdx);
 
-			double Falloff = CalculateBrushFalloff(UseDist);
+		double PrevDist = (OrigPos - LastBrushPosLocal).Length();
+		double NewDist = (OrigPos - NewBrushPosLocal).Length();
+		double UseDist = FMath::Min(PrevDist, NewDist);
 
-			FVector3d NewPos = OrigPos + Falloff * MoveVec;
+		double Falloff = CalculateBrushFalloff(UseDist);
 
-			Mesh->SetVertex(VertIdx, NewPos);
-			UpdateSavedVertex(VertIdx, OrigPos, NewPos);
-		}
+		FVector3d NewPos = OrigPos + Falloff * MoveVec;
+		ROIPositionBuffer[k] = NewPos;
+	});
 
-		bRemeshPending = bEnableRemeshing;
+	for (int k = 0; k < NumV; ++k)
+	{
+		int VertIdx = VertexROI[k];
+		const FVector3d& NewPos = ROIPositionBuffer[k];
+		FVector3d OrigPos = Mesh->GetVertex(VertIdx);
+		Mesh->SetVertex(VertIdx, NewPos);
+		UpdateSavedVertex(VertIdx, OrigPos, NewPos);
 	}
 
+	bRemeshPending = bEnableRemeshing;
+
+skip_brush:
 	LastBrushPosLocal = NewBrushPosLocal;
 }
 
@@ -563,11 +595,15 @@ void UDynamicMeshSculptTool::ApplyOffsetBrush(const FRay& WorldRay)
 	FVector NewBrushPosLocal = Transform.InverseTransformPosition(LastBrushPosWorld);
 
 	double Direction = (bInvert) ? -1.0 : 1.0;
-	double UseSpeed = Direction * FMathd::Sqrt(CurrentBrushRadius) * (SculptProperties->BrushSpeed*0.25);
+	double UseSpeed = Direction * FMathd::Sqrt(CurrentBrushRadius) * (SculptProperties->BrushSpeed);
 
 	FDynamicMesh3* Mesh = DynamicMeshComponent->GetMesh();
-	for (int VertIdx : VertexROI)
+	int NumV = VertexROI.Num();
+	ROIPositionBuffer.SetNum(NumV, false);
+
+	ParallelFor(NumV, [this, Mesh, NewBrushPosLocal, UseSpeed](int k)
 	{
+		int VertIdx = VertexROI[k];
 		FVector3d OrigPos = Mesh->GetVertex(VertIdx);
 
 		FVector3d BasePos, BaseNormal;
@@ -578,7 +614,14 @@ void UDynamicMeshSculptTool::ApplyOffsetBrush(const FRay& WorldRay)
 		double Falloff = CalculateBrushFalloff(OrigPos.Distance(NewBrushPosLocal));
 
 		FVector3d NewPos = OrigPos + Falloff * MoveVec;
+		ROIPositionBuffer[k] = NewPos;
+	});
 
+	for (int k = 0; k < NumV; ++k)
+	{
+		int VertIdx = VertexROI[k];
+		const FVector3d& NewPos = ROIPositionBuffer[k];
+		FVector3d OrigPos = Mesh->GetVertex(VertIdx);
 		Mesh->SetVertex(VertIdx, NewPos);
 		UpdateSavedVertex(VertIdx, OrigPos, NewPos);
 	}
@@ -604,14 +647,17 @@ void UDynamicMeshSculptTool::ApplySculptMaxBrush(const FRay& WorldRay)
 	FVector NewBrushPosLocal = Transform.InverseTransformPosition(LastBrushPosWorld);
 
 	double Direction = (bInvert) ? -1.0 : 1.0;
-	double UseSpeed = Direction * FMathd::Sqrt(CurrentBrushRadius) * (SculptProperties->BrushSpeed*0.25);
-
+	double UseSpeed = Direction * FMathd::Sqrt(CurrentBrushRadius) * (SculptProperties->BrushSpeed);
 	double MaxOffset = CurrentBrushRadius * 0.5;
-	double MaxOffsetSqr = MaxOffset * MaxOffset;
 
 	FDynamicMesh3* Mesh = DynamicMeshComponent->GetMesh();
-	for (int VertIdx : VertexROI)
+	int NumV = VertexROI.Num();
+	ROIPositionBuffer.SetNum(NumV, false);
+
+	ParallelFor(NumV, [this, Mesh, NewBrushPosLocal, UseSpeed, MaxOffset](int k)
 	{
+		int VertIdx = VertexROI[k];
+
 		FVector3d OrigPos = Mesh->GetVertex(VertIdx);
 
 		FVector3d BasePos, BaseNormal;
@@ -624,12 +670,20 @@ void UDynamicMeshSculptTool::ApplySculptMaxBrush(const FRay& WorldRay)
 		FVector3d NewPos = OrigPos + Falloff * MoveVec;
 
 		FVector3d DeltaPos = NewPos - BasePos;
-		if (DeltaPos.SquaredLength() > MaxOffsetSqr)
+		if (DeltaPos.SquaredLength() > MaxOffset*MaxOffset)
 		{
 			DeltaPos.Normalize();
 			NewPos = BasePos + MaxOffset * DeltaPos;
 		}
 
+		ROIPositionBuffer[k] = NewPos;
+	});
+
+	for (int k = 0; k < NumV; ++k)
+	{
+		int VertIdx = VertexROI[k];
+		const FVector3d& NewPos = ROIPositionBuffer[k];
+		FVector3d OrigPos = Mesh->GetVertex(VertIdx);
 		Mesh->SetVertex(VertIdx, NewPos);
 		UpdateSavedVertex(VertIdx, OrigPos, NewPos);
 	}
@@ -654,42 +708,65 @@ void UDynamicMeshSculptTool::ApplyPinchBrush(const FRay& WorldRay)
 	FVector BrushNormalLocal = Transform.InverseTransformVectorNoScale(LastBrushPosNormalWorld);
 	FVector OffsetBrushPosLocal = NewBrushPosLocal - SculptProperties->BrushDepth * CurrentBrushRadius * BrushNormalLocal;
 
+	// hardcoded lazybrush...
+	FVector NewSmoothBrushPosLocal = (0.75f)*LastSmoothBrushPosLocal + (0.25f)*NewBrushPosLocal;
+
 	double Direction = (bInvert) ? -1.0 : 1.0;
-	double UseSpeed = Direction * FMathd::Sqrt(CurrentBrushRadius) * (SculptProperties->BrushSpeed*0.25);
+	double UseSpeed = Direction * FMathd::Sqrt(CurrentBrushRadius) * (SculptProperties->BrushSpeed*0.05);
 
-	FVector3d MotionVec = NewBrushPosLocal - LastBrushPosLocal;
+	FVector3d MotionVec = NewSmoothBrushPosLocal - LastSmoothBrushPosLocal;
+	bool bHaveMotion = (MotionVec.Length() > FMathf::ZeroTolerance);
 	MotionVec.Normalize();
+	FLine3d MoveLine(LastSmoothBrushPosLocal, MotionVec);
 
 	FDynamicMesh3* Mesh = DynamicMeshComponent->GetMesh();
-	for (int VertIdx : VertexROI)
+	int NumV = VertexROI.Num();
+	ROIPositionBuffer.SetNum(NumV, false);
+
+	ParallelFor(NumV, [this, Mesh, NewBrushPosLocal, OffsetBrushPosLocal, bHaveMotion, MotionVec, UseSpeed](int k)
 	{
+		int VertIdx = VertexROI[k];
 		FVector3d OrigPos = Mesh->GetVertex(VertIdx);
 		FVector3d Delta = OffsetBrushPosLocal - OrigPos;
-		FVector3d MoveVec = UseSpeed * Delta;
 
-		double Falloff = CalculateBrushFalloff(OrigPos.Distance(NewBrushPosLocal));
+		FVector3d MoveVec = UseSpeed * Delta;
 
-		// make tighter
-		Falloff = Falloff * Falloff * Falloff;
+		// pinch uses 1/x falloff, shifted so that 
+		double Distance = OrigPos.Distance(NewBrushPosLocal);
+		double NormalizedDistance = Distance / CurrentBrushRadius + FMathf::ZeroTolerance;
+		double Falloff = (1.0/NormalizedDistance) - 1.0;
+		Falloff = FMathd::Clamp(Falloff, 0.0, 1.0);
 
-		double AnglePower = 1.0 - FMathd::Abs(MoveVec.Normalized().Dot(MotionVec));
-		Falloff *= AnglePower;
+		if (bHaveMotion && Falloff < 0.8f)
+		{
+			double AnglePower = 1.0 - FMathd::Abs(MoveVec.Normalized().Dot(MotionVec));
+			Falloff *= AnglePower;
+		}
 
 		FVector3d NewPos = OrigPos + Falloff * MoveVec;
+		ROIPositionBuffer[k] = NewPos;
+	});
 
+	for (int k = 0; k < NumV; ++k)
+	{
+		int VertIdx = VertexROI[k];
+		const FVector3d& NewPos = ROIPositionBuffer[k];
+		FVector3d OrigPos = Mesh->GetVertex(VertIdx);
 		Mesh->SetVertex(VertIdx, NewPos);
 		UpdateSavedVertex(VertIdx, OrigPos, NewPos);
 	}
 
+
 	bRemeshPending = bEnableRemeshing;
 
 	LastBrushPosLocal = NewBrushPosLocal;
+	LastSmoothBrushPosLocal = NewSmoothBrushPosLocal;
 }
 
 
 
 
-FFrame3d UDynamicMeshSculptTool::ComputeROIBrushPlane(const FVector3d& BrushCenter)
+FFrame3d UDynamicMeshSculptTool::ComputeROIBrushPlane(const FVector3d& BrushCenter, bool bIgnoreDepth)
 {
 	FDynamicMesh3* Mesh = DynamicMeshComponent->GetMesh();
 	FVector3d AverageNormal(0, 0, 0);
@@ -708,7 +785,10 @@ FFrame3d UDynamicMeshSculptTool::ComputeROIBrushPlane(const FVector3d& BrushCent
 	AveragePos /= WeightSum;
 
 	FFrame3d Result = FFrame3d(AveragePos, AverageNormal);
-	Result.Origin -= SculptProperties->BrushDepth * CurrentBrushRadius * Result.Z();
+	if (bIgnoreDepth == false)
+	{
+		Result.Origin -= SculptProperties->BrushDepth * CurrentBrushRadius * Result.Z();
+	}
 
 	return Result;
 }
@@ -724,13 +804,16 @@ void UDynamicMeshSculptTool::ApplyPlaneBrush(const FRay& WorldRay)
 	FTransform Transform = ComponentTarget->GetWorldTransform();
 	FVector NewBrushPosLocal = Transform.InverseTransformPosition(LastBrushPosWorld);
 	FVector BrushNormalLocal = Transform.InverseTransformVectorNoScale(LastBrushPosNormalWorld);
-	FVector OffsetBrushPosLocal = NewBrushPosLocal - SculptProperties->BrushDepth * CurrentBrushRadius * BrushNormalLocal;
+	double UseSpeed = FMathd::Sqrt(CurrentBrushRadius) * FMathd::Sqrt(SculptProperties->BrushSpeed) * 0.05;
+
 
 	FDynamicMesh3* Mesh = DynamicMeshComponent->GetMesh();
-	double UseSpeed = FMathd::Sqrt(CurrentBrushRadius) * FMathd::Sqrt(SculptProperties->BrushSpeed) * 0.25;
+	int NumV = VertexROI.Num();
+	ROIPositionBuffer.SetNum(NumV, false);
 
-	for (int VertIdx : VertexROI)
+	ParallelFor(NumV, [this, Mesh, NewBrushPosLocal, UseSpeed](int k)
 	{
+		int VertIdx = VertexROI[k];
 		FVector3d OrigPos = Mesh->GetVertex(VertIdx);
 		FVector3d PlanePos = ActiveFixedBrushPlane.ToPlane(OrigPos, 2);
 		FVector3d Delta = PlanePos - OrigPos;
@@ -739,7 +822,14 @@ void UDynamicMeshSculptTool::ApplyPlaneBrush(const FRay& WorldRay)
 		double Falloff = CalculateBrushFalloff(OrigPos.Distance(NewBrushPosLocal));
 
 		FVector3d NewPos = OrigPos + Falloff * MoveVec;
+		ROIPositionBuffer[k] = NewPos;
+	});
 
+	for (int k = 0; k < NumV; ++k)
+	{
+		int VertIdx = VertexROI[k];
+		const FVector3d& NewPos = ROIPositionBuffer[k];
+		FVector3d OrigPos = Mesh->GetVertex(VertIdx);
 		Mesh->SetVertex(VertIdx, NewPos);
 		UpdateSavedVertex(VertIdx, OrigPos, NewPos);
 	}
@@ -764,14 +854,18 @@ void UDynamicMeshSculptTool::ApplyFlattenBrush(const FRay& WorldRay)
 	FTransform Transform = ComponentTarget->GetWorldTransform();
 	FVector NewBrushPosLocal = Transform.InverseTransformPosition(LastBrushPosWorld);
 	FVector BrushNormalLocal = Transform.InverseTransformVectorNoScale(LastBrushPosNormalWorld);
-	FVector OffsetBrushPosLocal = NewBrushPosLocal - SculptProperties->BrushDepth * CurrentBrushRadius * BrushNormalLocal;
+
+	double UseSpeed = FMathd::Sqrt(CurrentBrushRadius) * FMathd::Sqrt(SculptProperties->BrushSpeed) * 0.05;
+	FFrame3d StampFlattenPlane = ComputeROIBrushPlane(NewBrushPosLocal, true);
+	//StampFlattenPlane.Origin -= SculptProperties->BrushDepth * CurrentBrushRadius * StampFlattenPlane.Z();
 
 	FDynamicMesh3* Mesh = DynamicMeshComponent->GetMesh();
-	double UseSpeed = FMathd::Sqrt(CurrentBrushRadius) * FMathd::Sqrt(SculptProperties->BrushSpeed) * 0.25;
-	FFrame3d StampFlattenPlane = ComputeROIBrushPlane(NewBrushPosLocal);
+	int NumV = VertexROI.Num();
+	ROIPositionBuffer.SetNum(NumV, false);
 
-	for (int VertIdx : VertexROI)
+	ParallelFor(NumV, [this, Mesh, NewBrushPosLocal, UseSpeed, StampFlattenPlane](int k)
 	{
+		int VertIdx = VertexROI[k];
 		FVector3d OrigPos = Mesh->GetVertex(VertIdx);
 		FVector3d PlanePos = StampFlattenPlane.ToPlane(OrigPos, 2);
 		FVector3d Delta = PlanePos - OrigPos;
@@ -780,7 +874,14 @@ void UDynamicMeshSculptTool::ApplyFlattenBrush(const FRay& WorldRay)
 		double Falloff = CalculateBrushFalloff(OrigPos.Distance(NewBrushPosLocal));
 
 		FVector3d NewPos = OrigPos + Falloff * MoveVec;
+		ROIPositionBuffer[k] = NewPos;
+	});
 
+	for (int k = 0; k < NumV; ++k)
+	{
+		int VertIdx = VertexROI[k];
+		const FVector3d& NewPos = ROIPositionBuffer[k];
+		FVector3d OrigPos = Mesh->GetVertex(VertIdx);
 		Mesh->SetVertex(VertIdx, NewPos);
 		UpdateSavedVertex(VertIdx, OrigPos, NewPos);
 	}
@@ -805,7 +906,10 @@ void UDynamicMeshSculptTool::ApplyInflateBrush(const FRay& WorldRay)
 
 	double Direction = (bInvert) ? -1.0 : 1.0;
 	double UseSpeed = Direction * FMathd::Sqrt(CurrentBrushRadius) * SculptProperties->BrushSpeed * 0.25;
+
 	FDynamicMesh3* Mesh = DynamicMeshComponent->GetMesh();
+	int NumV = VertexROI.Num();
+	ROIPositionBuffer.SetNum(NumV, false);
 
 	// calculate vertex normals
 	ParallelFor(VertexROI.Num(), [this, Mesh](int Index) {
@@ -815,9 +919,9 @@ void UDynamicMeshSculptTool::ApplyInflateBrush(const FRay& WorldRay)
 	});
 
 
-	ParallelFor(VertexROI.Num(), [this, Mesh, UseSpeed, NewBrushPosLocal](int Index) 
+	ParallelFor(VertexROI.Num(), [this, Mesh, UseSpeed, NewBrushPosLocal](int k) 
 	{
-		int VertIdx = VertexROI[Index];
+		int VertIdx = VertexROI[k];
 		FVector3d OrigPos = Mesh->GetVertex(VertIdx);
 		FVector3d Normal = (FVector3d)Mesh->GetVertexNormal(VertIdx);
 
@@ -826,10 +930,17 @@ void UDynamicMeshSculptTool::ApplyInflateBrush(const FRay& WorldRay)
 		double Falloff = CalculateBrushFalloff(OrigPos.Distance(NewBrushPosLocal));
 
 		FVector3d NewPos = OrigPos + Falloff*MoveVec;
-		Mesh->SetVertex(VertIdx, NewPos);
+		ROIPositionBuffer[k] = NewPos;
+	});
 
+	for (int k = 0; k < NumV; ++k)
+	{
+		int VertIdx = VertexROI[k];
+		const FVector3d& NewPos = ROIPositionBuffer[k];
+		FVector3d OrigPos = Mesh->GetVertex(VertIdx);
+		Mesh->SetVertex(VertIdx, NewPos);
 		UpdateSavedVertex(VertIdx, OrigPos, NewPos);
-	});
+	}
 
 	bRemeshPending = bEnableRemeshing;
 
@@ -957,6 +1068,14 @@ void UDynamicMeshSculptTool::Tick(float DeltaTime)
 
 	UMeshSurfacePointTool::Tick(DeltaTime);
 
+	// Allow a tick to pass between application of brush stamps. Bizarrely this
+	// improves responsiveness in the Editor...
+	static int TICK_SKIP_HACK = 0;
+	if (TICK_SKIP_HACK++ % 2 == 0)
+	{
+		return;
+	}
+
 	ShowWireframeWatcher.CheckAndUpdate();
 	MaterialModeWatcher.CheckAndUpdate();
 
@@ -1126,15 +1245,32 @@ void UDynamicMeshSculptTool::RemeshROIPass()
 	FDynamicMeshOctree3* Octree = DynamicMeshComponent->GetOctree();
 
 	FSubRegionRemesher Remesher(Mesh);
-	Remesher.SetTargetEdgeLength(RemeshProperties->RelativeSize * InitialEdgeLength);
+	double TargetEdgeLength = RemeshProperties->RelativeSize * InitialEdgeLength;
+	Remesher.SetTargetEdgeLength(TargetEdgeLength);
 
-	// this is a temporary tweak for pinch brush. Remesh params should be per-brush!
+	double UseSmoothing = RemeshProperties->Smoothing * 0.25;
+	Remesher.SmoothSpeedT = UseSmoothing;
+
+	// this is a temporary tweak for Pinch brush. Remesh params should be per-brush!
 	if (SculptProperties->PrimaryBrushType == EDynamicMeshSculptBrushType::Pinch && bSmoothing == false)
 	{
-		Remesher.MinEdgeLength = Remesher.MinEdgeLength / 4.0;
+		Remesher.MinEdgeLength = TargetEdgeLength * 0.1;
+
+		Remesher.CustomSmoothSpeedF = [this, &UseSmoothing](const FDynamicMesh3& Mesh, int vID)
+		{
+			FVector3d Pos = Mesh.GetVertex(vID);
+			double Falloff = CalculateBrushFalloff(Pos.Distance((FVector3d)LastBrushPosLocal));
+			return (1.0f - Falloff) * UseSmoothing;
+		};
+	}
+
+	// tweak remesh params for Smooth brush
+	if (bSmoothing && RemeshProperties->bRemeshSmooth == false)
+	{
+		Remesher.MaxEdgeLength = 2*InitialEdgeLength;
+		Remesher.MinEdgeLength = InitialEdgeLength * 0.01;
 	}
 
-	Remesher.SmoothSpeedT = RemeshProperties->Smoothing * 0.25;
 	Remesher.SmoothType = (SculptProperties->bPreserveUVFlow) ?
 		FRemesher::ESmoothTypes::MeanValue : FRemesher::ESmoothTypes::Uniform;
 	bool bIsUniformSmooth = (Remesher.SmoothType == FRemesher::ESmoothTypes::Uniform);
diff --git a/Engine/Plugins/Experimental/MeshModelingToolset/Source/MeshModelingTools/Public/DynamicMeshSculptTool.h b/Engine/Plugins/Experimental/MeshModelingToolset/Source/MeshModelingTools/Public/DynamicMeshSculptTool.h
--- a/Engine/Plugins/Experimental/MeshModelingToolset/Source/MeshModelingTools/Public/DynamicMeshSculptTool.h
+++ b/Engine/Plugins/Experimental/MeshModelingToolset/Source/MeshModelingTools/Public/DynamicMeshSculptTool.h
@@ -48,6 +48,9 @@ enum class EDynamicMeshSculptBrushType : uint8
 	/** Move Brush moves vertices parallel to the view plane  */
 	Move UMETA(DisplayName = "Move"),
 
+	/** Smooth brush smooths mesh vertices  */
+	Smooth UMETA(DisplayName = "Smooth"),
+
 	/** Sculpt Brush displaces vertices from the surface */
 	Offset UMETA(DisplayName = "Sculpt"),
 
@@ -158,6 +161,10 @@ public:
 	UPROPERTY(EditAnywhere, Category = Remeshing, meta = (UIMin = "0.0", UIMax = "1.0", ClampMin = "0.0", ClampMax = "1.0"))
 	float Smoothing;
 
+	/** If enabled, Full Remeshing is applied during smoothing, which will wipe out fine details */
+	UPROPERTY(EditAnywhere, Category = Remeshing)
+	bool bRemeshSmooth = false;
+
 
 	/** Enable edge flips */
 	UPROPERTY(EditAnywhere, Category = Remeshing, AdvancedDisplay)
@@ -290,6 +297,7 @@ protected:
 	FVector LastBrushPosLocal;
 	FVector LastBrushPosWorld;
 	FVector LastBrushPosNormalWorld;
+	FVector LastSmoothBrushPosLocal;
 	
 
 	TArray<int> VertexROI;
@@ -334,9 +342,10 @@ protected:
 	void ApplyFlattenBrush(const FRay& WorldRay);
 
 	double CalculateBrushFalloff(double Distance);
+	TArray<FVector3d> ROIPositionBuffer;
 
 	FFrame3d ActiveFixedBrushPlane;
-	FFrame3d ComputeROIBrushPlane(const FVector3d& BrushCenter);
+	FFrame3d ComputeROIBrushPlane(const FVector3d& BrushCenter, bool bIgnoreDepth);
 
 	TArray<int> TrianglesBuffer;
 	TArray<int> NormalsBuffer;
