Merging CL# 9552130 from UE4/Dev-Editor
ModelingTools: handle scaling (uniform and non-uniform) in Sculpt tool. This is done by transforming mesh to world space for sculpting, and transforming back on Bake. Not an ideal solution, but usable numeric range in the World is bounded by float precision and we are working in doubles.

Added functions to TTransform3 to clamp scale to valid range and transform Normals.
Added MeshTransforms, has functions to transform and inverse-transform a mesh.
Added OctreeDynamicMeshComponent::ApplyTransform() to apply these functions to internal mesh.

#rnx

diff --git a/Engine/Plugins/Experimental/GeometryProcessing/Source/DynamicMesh/Private/MeshTransforms.cpp b/Engine/Plugins/Experimental/GeometryProcessing/Source/DynamicMesh/Private/MeshTransforms.cpp
new file mode 100644
--- /dev/null
+++ b/Engine/Plugins/Experimental/GeometryProcessing/Source/DynamicMesh/Private/MeshTransforms.cpp
@@ -0,0 +1,67 @@
+// Copyright 1998-2019 Epic Games, Inc. All Rights Reserved.
+
+
+#include "MeshTransforms.h"
+#include "DynamicMeshAttributeSet.h"
+#include "DynamicMeshOverlay.h"
+
+
+void MeshTransforms::ApplyTransform(FDynamicMesh3& Mesh, const FTransform3d& Transform)
+{
+	bool bVertexNormals = Mesh.HasVertexNormals();
+
+	for (int vid : Mesh.VertexIndicesItr())
+	{
+		FVector3d Position = Mesh.GetVertex(vid);
+		Position = Transform.TransformPosition(Position);
+		Mesh.SetVertex(vid, Position);
+
+		if (bVertexNormals)
+		{
+			FVector3f Normal = Mesh.GetVertexNormal(vid);
+			Normal = (FVector3f)Transform.TransformNormal((FVector3d)Normal);
+			Mesh.SetVertexNormal(vid, Normal.Normalized());
+		}
+	}
+	if (Mesh.HasAttributes())
+	{
+		FDynamicMeshNormalOverlay* Normals = Mesh.Attributes()->PrimaryNormals();
+		for (int elemid : Normals->ElementIndicesItr())
+		{
+			FVector3f Normal = Normals->GetElement(elemid);
+			Normal = (FVector3f)Transform.TransformNormal((FVector3d)Normal);
+			Normals->SetElement(elemid, Normal.Normalized());
+		}
+	}
+}
+
+
+
+void MeshTransforms::ApplyTransformInverse(FDynamicMesh3& Mesh, const FTransform3d& Transform)
+{
+	bool bVertexNormals = Mesh.HasVertexNormals();
+
+	for (int vid : Mesh.VertexIndicesItr())
+	{
+		FVector3d Position = Mesh.GetVertex(vid);
+		Position = Transform.InverseTransformPosition(Position);
+		Mesh.SetVertex(vid, Position);
+
+		if (bVertexNormals)
+		{
+			FVector3f Normal = Mesh.GetVertexNormal(vid);
+			Normal = (FVector3f)Transform.InverseTransformNormal((FVector3d)Normal);
+			Mesh.SetVertexNormal(vid, Normal.Normalized());
+		}
+	}
+	if (Mesh.HasAttributes())
+	{
+		FDynamicMeshNormalOverlay* Normals = Mesh.Attributes()->PrimaryNormals();
+		for (int elemid : Normals->ElementIndicesItr())
+		{
+			FVector3f Normal = Normals->GetElement(elemid);
+			Normal = (FVector3f)Transform.InverseTransformNormal((FVector3d)Normal);
+			Normals->SetElement(elemid, Normal.Normalized());
+		}
+	}
+}
\ No newline at end of file
diff --git a/Engine/Plugins/Experimental/GeometryProcessing/Source/DynamicMesh/Public/MeshTransforms.h b/Engine/Plugins/Experimental/GeometryProcessing/Source/DynamicMesh/Public/MeshTransforms.h
new file mode 100644
--- /dev/null
+++ b/Engine/Plugins/Experimental/GeometryProcessing/Source/DynamicMesh/Public/MeshTransforms.h
@@ -0,0 +1,32 @@
+// Copyright 1998-2019 Epic Games, Inc. All Rights Reserved.
+
+// Port of geometry3cpp MeshWeights
+
+#pragma once
+
+#include "DynamicMesh3.h"
+
+
+/**
+ * Utility functions for applying transformations to meshes
+ */
+namespace MeshTransforms
+{
+	
+	/**
+	 * Apply given Transform to a Mesh.
+	 * Modifies Vertex Positions and Normals, and any Per-Triangle Normal Overlays
+	 */
+	DYNAMICMESH_API void ApplyTransform(FDynamicMesh3& Mesh, const FTransform3d& Transform);
+
+
+	/**
+	 * Apply inverse of given Transform to a Mesh.
+	 * Modifies Vertex Positions and Normals, and any Per-Triangle Normal Overlays
+	 */
+	DYNAMICMESH_API void ApplyTransformInverse(FDynamicMesh3& Mesh, const FTransform3d& Transform);
+
+
+
+
+};
\ No newline at end of file
diff --git a/Engine/Plugins/Experimental/GeometryProcessing/Source/GeometricObjects/Public/TransformTypes.h b/Engine/Plugins/Experimental/GeometryProcessing/Source/GeometricObjects/Public/TransformTypes.h
--- a/Engine/Plugins/Experimental/GeometryProcessing/Source/GeometricObjects/Public/TransformTypes.h
+++ b/Engine/Plugins/Experimental/GeometryProcessing/Source/GeometricObjects/Public/TransformTypes.h
@@ -71,6 +71,55 @@ public:
 	}
 
 
+	/**
+	 * @return Rotation portion of Transform, as Quaternion
+	 */
+	const TQuaternion<RealType>& GetRotation() const 
+	{ 
+		return Rotation; 
+	}
+
+	/** 
+	 * Set Rotation portion of Transform 
+	 */
+	void SetRotation(const TQuaternion<RealType>& RotationIn)
+	{
+		Rotation = RotationIn;
+	}
+
+	/**
+	 * @return Translation portion of transform
+	 */
+	const FVector3<RealType>& GetTranslation() const
+	{
+		return Translation;
+	}
+
+	/**
+	 * set Translation portion of transform
+	 */
+	void SetTranslation(const FVector3<RealType>& TranslationIn)
+	{
+		Translation = TranslationIn;
+	}
+
+	/**
+	 * @return Scale portion of transform
+	 */
+	const FVector3<RealType>& GetScale() const
+	{
+		return Scale3D;
+	}
+
+	/**
+	 * set Scale portion of transform
+	 */
+	void SetScale(const FVector3<RealType>& ScaleIn)
+	{
+		Scale3D = ScaleIn;
+	}
+
+
 	/**
 	 * @return input point with QST transformation applied, ie QST(P) = Rotate(Scale*P) + Translate
 	 */
@@ -97,6 +146,22 @@ public:
 		return Rotation * V;
 	}
 
+
+	/**
+	 * Surface Normals are special, their transform is Rotate( Normalize( (1/Scale) * Normal) ) ).
+	 * However 1/Scale requires special handling in case any component is near-zero.
+	 * @return input surface normal with transform applied.
+	 */
+	FVector3<RealType> TransformNormal(const FVector3<RealType>& Normal) const
+	{
+		// transform normal by a safe inverse scale + normalize, and a standard rotation
+		const FVector3<RealType>& S = Scale3D;
+		RealType DetSign = FMathd::SignNonZero(S.X * S.Y * S.Z); // we only need to multiply by the sign of the determinant, rather than divide by it, since we normalize later anyway
+		FVector3<RealType> SafeInvS(S.Y*S.Z*DetSign, S.X*S.Z*DetSign, S.X*S.Y*DetSign);
+		return TransformVectorNoScale( (SafeInvS*Normal).Normalized() );
+	}
+
+
 	/**
 	 * @return input vector with inverse-QST transformation applied, ie QSTinv(P) = InverseScale(InverseRotate(P - Translate))
 	 */
@@ -122,6 +187,35 @@ public:
 		return Rotation.InverseMultiply(V);
 	}
 
+
+	/**
+	 * Surface Normals are special, their inverse transform is InverseRotate( Normalize(Scale * Normal) ) )
+	 * @return input surface normal with inverse transform applied.
+	 */
+	FVector3<RealType> InverseTransformNormal(const FVector3<RealType>& Normal) const
+	{
+		return InverseTransformVectorNoScale( (Scale3D*Normal).Normalized() );
+	}
+
+
+
+	/**
+	 * Clamp all scale components to a minimum value. Sign of scale components is preserved.
+	 * This is used to remove uninvertible zero/near-zero scaling.
+	 */
+	void ClampMinimumScale(RealType MinimumScale = TMathUtil<RealType>::ZeroTolerance)
+	{
+		for (int j = 0; j < 3; ++j)
+		{
+			RealType Value = Scale3D[j];
+			if (TMathUtil<RealType>::Abs(Value) < MinimumScale)
+			{
+				Value = MinimumScale * TMathUtil<RealType>::SignNonZero(Value);
+				Scale3D[j] = Value;
+			}
+		}
+	}
+
 	
 
 	static FVector3<RealType> GetSafeScaleReciprocal(const FVector3<RealType>& InScale, RealType Tolerance = TMathUtil<RealType>::ZeroTolerance)
diff --git a/Engine/Plugins/Experimental/MeshModelingToolset/Source/MeshModelingTools/Private/DynamicMeshSculptTool.cpp b/Engine/Plugins/Experimental/MeshModelingToolset/Source/MeshModelingTools/Private/DynamicMeshSculptTool.cpp
--- a/Engine/Plugins/Experimental/MeshModelingToolset/Source/MeshModelingTools/Private/DynamicMeshSculptTool.cpp
+++ b/Engine/Plugins/Experimental/MeshModelingToolset/Source/MeshModelingTools/Private/DynamicMeshSculptTool.cpp
@@ -108,7 +108,19 @@ void UDynamicMeshSculptTool::Setup()
 	DynamicMeshComponent = NewObject<UOctreeDynamicMeshComponent>(ComponentTarget->GetOwnerActor(), "DynamicMeshSculptToolMesh");
 	DynamicMeshComponent->SetupAttachment(ComponentTarget->GetOwnerActor()->GetRootComponent());
 	DynamicMeshComponent->RegisterComponent();
-	DynamicMeshComponent->SetWorldTransform(ComponentTarget->GetWorldTransform());
+
+	// initialize from LOD-0 MeshDescription
+	DynamicMeshComponent->InitializeMesh(ComponentTarget->GetMesh());
+
+	// transform mesh to world space because handling scaling inside brush is a mess
+	InitialTargetTransform = FTransform3d(ComponentTarget->GetWorldTransform());
+	// clamp scaling because if we allow zero-scale we cannot invert this transform on Accept
+	InitialTargetTransform.ClampMinimumScale(0.01);
+	DynamicMeshComponent->ApplyTransform(InitialTargetTransform, false);
+	// since we moved to World coords there is not a current transform anymore.
+	// @todo: only remove scaling, keep rotation and translation in CurTargetTransform. This will improve numerical precision.
+	CurTargetTransform = FTransform3d::Identity();
+	DynamicMeshComponent->SetWorldTransform((FTransform)CurTargetTransform);
 
 	// copy material if there is one
 	UMaterialInterface* Material = ComponentTarget->GetMaterial(0);
@@ -117,35 +129,32 @@ void UDynamicMeshSculptTool::Setup()
 		DynamicMeshComponent->SetMaterial(0, Material);
 	}
 
-	// initialize from LOD-0 MeshDescription
-	//DynamicMeshComponent->TangentsType = (bEnableRemeshing) ? EDynamicMeshTangentCalcType::NoTangents : EDynamicMeshTangentCalcType::AutoCalculated;
-	DynamicMeshComponent->InitializeMesh(ComponentTarget->GetMesh());
 	OnDynamicMeshComponentChangedHandle = DynamicMeshComponent->OnMeshChanged.Add(
 		FSimpleMulticastDelegate::FDelegate::CreateUObject(this, &UDynamicMeshSculptTool::OnDynamicMeshComponentChanged));
 
-	// initialize AABBTree
-	FDynamicMesh3* Mesh = DynamicMeshComponent->GetMesh();
-
 	// do we always want to keep vertex normals updated? Perhaps should discard vertex normals before baking?
+	FDynamicMesh3* Mesh = DynamicMeshComponent->GetMesh();
 	FMeshNormals::QuickComputeVertexNormals(*Mesh);
 
 	// switch to vertex normals for testing
 	//DynamicMeshComponent->GetMesh()->DiscardAttributes();
 
+	// initialize target mesh
 	UpdateTarget();
 	bTargetDirty = false;
 
+	// initialize brush radius range interval, brush properties
 	double MaxDimension = DynamicMeshComponent->GetMesh()->GetCachedBounds().MaxDim();
 	BrushRelativeSizeRange = FInterval1d(MaxDimension*0.01, MaxDimension);
 	BrushProperties = NewObject<UBrushBaseProperties>(this, TEXT("Brush"));
 	CalculateBrushRadius();
 
+	// initialize other properties
 	SculptProperties = NewObject<UBrushSculptProperties>(this, TEXT("Sculpting"));
 	RemeshProperties = NewObject<UBrushRemeshProperties>(this, TEXT("Remeshing"));
 	InitialEdgeLength = EstimateIntialSafeTargetLength(*Mesh, 5000);
 
-
-	// hide input StaticMeshComponent
+	// hide input Component
 	ComponentTarget->SetOwnerVisibility(false);
 
 	// init state flags flags
@@ -228,6 +237,9 @@ void UDynamicMeshSculptTool::Shutdown(EToolShutdownType ShutdownType)
 
 		if (ShutdownType == EToolShutdownType::Accept)
 		{
+			// safe to do this here because we are about to destroy componeont
+			DynamicMeshComponent->ApplyTransform(InitialTargetTransform, true);
+
 			// this block bakes the modified DynamicMeshComponent back into the StaticMeshComponent inside an undo transaction
 			GetToolManager()->BeginUndoTransaction(LOCTEXT("SculptMeshToolTransactionName", "Sculpt Mesh"));
 			ComponentTarget->CommitMesh([=](FMeshDescription* MeshDescription)
@@ -265,10 +277,8 @@ void UDynamicMeshSculptTool::OnPropertyModified(UObject* PropertySet, UProperty*
 
 bool UDynamicMeshSculptTool::HitTest(const FRay& Ray, FHitResult& OutHit)
 {
-	FTransform Transform = ComponentTarget->GetWorldTransform();
-
-	FRay3d LocalRay(Transform.InverseTransformPosition(Ray.Origin),
-		Transform.InverseTransformVector(Ray.Direction));
+	FRay3d LocalRay(CurTargetTransform.InverseTransformPosition(Ray.Origin),
+		CurTargetTransform.InverseTransformVector(Ray.Direction));
 	LocalRay.Direction.Normalize();
 
 	FDynamicMesh3* Mesh = DynamicMeshComponent->GetMesh();
@@ -283,8 +293,8 @@ bool UDynamicMeshSculptTool::HitTest(const FRay& Ray, FHitResult& OutHit)
 
 		OutHit.FaceIndex = HitTID;
 		OutHit.Distance = Query.RayParameter;
-		OutHit.Normal = Transform.TransformVectorNoScale(Mesh->GetTriNormal(HitTID));
-		OutHit.ImpactPoint = Transform.TransformPosition(LocalRay.PointAt(Query.RayParameter));
+		OutHit.Normal = (FVector)CurTargetTransform.TransformNormal(Mesh->GetTriNormal(HitTID));
+		OutHit.ImpactPoint = (FVector)CurTargetTransform.TransformPosition(LocalRay.PointAt(Query.RayParameter));
 		return true;
 	}
 
@@ -305,13 +315,12 @@ void UDynamicMeshSculptTool::OnBeginDrag(const FRay& Ray)
 
 		bInDrag = true;
 
-		ActiveDragPlane = FPlane(BrushStartCenterWorld, -Ray.Direction);
-		LastHitPosWorld = FMath::RayPlaneIntersection(Ray.Origin, Ray.Direction, ActiveDragPlane);
+		ActiveDragPlane = FFrame3d(BrushStartCenterWorld, -Ray.Direction);
+		ActiveDragPlane.RayPlaneIntersection(Ray.Origin, Ray.Direction, 2, LastHitPosWorld);
 
 		LastBrushPosWorld = LastHitPosWorld;
-		LastBrushPosNormalWorld = ActiveDragPlane;
-		FTransform Transform = ComponentTarget->GetWorldTransform();
-		LastBrushPosLocal = Transform.InverseTransformPosition(LastHitPosWorld);
+		LastBrushPosNormalWorld = ActiveDragPlane.Z();
+		LastBrushPosLocal = CurTargetTransform.InverseTransformPosition(LastHitPosWorld);
 		LastSmoothBrushPosLocal = LastBrushPosLocal;
 
 		BeginChange(bEnableRemeshing == false);
@@ -488,8 +497,7 @@ void UDynamicMeshSculptTool::ApplySmoothBrush(const FRay& WorldRay)
 		return;
 	}
 
-	FTransform Transform = ComponentTarget->GetWorldTransform();
-	FVector NewBrushPosLocal = Transform.InverseTransformPosition(LastBrushPosWorld);
+	FVector3d NewBrushPosLocal = CurTargetTransform.InverseTransformPosition(LastBrushPosWorld);
 
 	FDynamicMesh3* Mesh = DynamicMeshComponent->GetMesh();
 	int NumV = VertexROI.Num();
@@ -534,14 +542,13 @@ void UDynamicMeshSculptTool::ApplyMoveBrush(const FRay& WorldRay)
 {
 	bool bHit = UpdateBrushPositionOnActivePlane(WorldRay);
 
-	FTransform Transform = ComponentTarget->GetWorldTransform();
-	FVector NewBrushPosLocal = Transform.InverseTransformPosition(LastBrushPosWorld);
-
+	FVector3d NewBrushPosLocal = CurTargetTransform.InverseTransformPosition(LastBrushPosWorld);
 	FVector3d MoveVec = NewBrushPosLocal - LastBrushPosLocal;
 
 	if (MoveVec.SquaredLength() <= 0)
 	{
-		goto skip_brush;
+		LastBrushPosLocal = NewBrushPosLocal;
+		return;
 	}
 
 	FDynamicMesh3* Mesh = DynamicMeshComponent->GetMesh();
@@ -573,8 +580,6 @@ void UDynamicMeshSculptTool::ApplyMoveBrush(const FRay& WorldRay)
 	}
 
 	bRemeshPending = bEnableRemeshing;
-
-skip_brush:
 	LastBrushPosLocal = NewBrushPosLocal;
 }
 
@@ -591,8 +596,7 @@ void UDynamicMeshSculptTool::ApplyOffsetBrush(const FRay& WorldRay)
 		return;
 	}
 
-	FTransform Transform = ComponentTarget->GetWorldTransform();
-	FVector NewBrushPosLocal = Transform.InverseTransformPosition(LastBrushPosWorld);
+	FVector3d NewBrushPosLocal = CurTargetTransform.InverseTransformPosition(LastBrushPosWorld);
 
 	double Direction = (bInvert) ? -1.0 : 1.0;
 	double UseSpeed = Direction * FMathd::Sqrt(CurrentBrushRadius) * (SculptProperties->BrushSpeed);
@@ -643,8 +647,7 @@ void UDynamicMeshSculptTool::ApplySculptMaxBrush(const FRay& WorldRay)
 		return;
 	}
 
-	FTransform Transform = ComponentTarget->GetWorldTransform();
-	FVector NewBrushPosLocal = Transform.InverseTransformPosition(LastBrushPosWorld);
+	FVector3d NewBrushPosLocal = CurTargetTransform.InverseTransformPosition(LastBrushPosWorld);
 
 	double Direction = (bInvert) ? -1.0 : 1.0;
 	double UseSpeed = Direction * FMathd::Sqrt(CurrentBrushRadius) * (SculptProperties->BrushSpeed);
@@ -703,13 +706,12 @@ void UDynamicMeshSculptTool::ApplyPinchBrush(const FRay& WorldRay)
 		return;
 	}
 
-	FTransform Transform = ComponentTarget->GetWorldTransform();
-	FVector NewBrushPosLocal = Transform.InverseTransformPosition(LastBrushPosWorld);
-	FVector BrushNormalLocal = Transform.InverseTransformVectorNoScale(LastBrushPosNormalWorld);
-	FVector OffsetBrushPosLocal = NewBrushPosLocal - SculptProperties->BrushDepth * CurrentBrushRadius * BrushNormalLocal;
+	FVector3d NewBrushPosLocal = CurTargetTransform.InverseTransformPosition(LastBrushPosWorld);
+	FVector3d BrushNormalLocal = CurTargetTransform.InverseTransformNormal(LastBrushPosNormalWorld);
+	FVector3d OffsetBrushPosLocal = NewBrushPosLocal - SculptProperties->BrushDepth * CurrentBrushRadius * BrushNormalLocal;
 
 	// hardcoded lazybrush...
-	FVector NewSmoothBrushPosLocal = (0.75f)*LastSmoothBrushPosLocal + (0.25f)*NewBrushPosLocal;
+	FVector3d NewSmoothBrushPosLocal = (0.75f)*LastSmoothBrushPosLocal + (0.25f)*NewBrushPosLocal;
 
 	double Direction = (bInvert) ? -1.0 : 1.0;
 	double UseSpeed = Direction * FMathd::Sqrt(CurrentBrushRadius) * (SculptProperties->BrushSpeed*0.05);
@@ -801,9 +803,8 @@ void UDynamicMeshSculptTool::ApplyPlaneBrush(const FRay& WorldRay)
 		return;
 	}
 
-	FTransform Transform = ComponentTarget->GetWorldTransform();
-	FVector NewBrushPosLocal = Transform.InverseTransformPosition(LastBrushPosWorld);
-	FVector BrushNormalLocal = Transform.InverseTransformVectorNoScale(LastBrushPosNormalWorld);
+	FVector3d NewBrushPosLocal = CurTargetTransform.InverseTransformPosition(LastBrushPosWorld);
+	FVector3d BrushNormalLocal = CurTargetTransform.InverseTransformNormal(LastBrushPosNormalWorld);
 	double UseSpeed = FMathd::Sqrt(CurrentBrushRadius) * FMathd::Sqrt(SculptProperties->BrushSpeed) * 0.05;
 
 
@@ -851,9 +852,8 @@ void UDynamicMeshSculptTool::ApplyFlattenBrush(const FRay& WorldRay)
 		return;
 	}
 
-	FTransform Transform = ComponentTarget->GetWorldTransform();
-	FVector NewBrushPosLocal = Transform.InverseTransformPosition(LastBrushPosWorld);
-	FVector BrushNormalLocal = Transform.InverseTransformVectorNoScale(LastBrushPosNormalWorld);
+	FVector3d NewBrushPosLocal = CurTargetTransform.InverseTransformPosition(LastBrushPosWorld);
+	FVector3d BrushNormalLocal = CurTargetTransform.InverseTransformNormal(LastBrushPosNormalWorld);
 
 	double UseSpeed = FMathd::Sqrt(CurrentBrushRadius) * FMathd::Sqrt(SculptProperties->BrushSpeed) * 0.05;
 	FFrame3d StampFlattenPlane = ComputeROIBrushPlane(NewBrushPosLocal, true);
@@ -901,8 +901,7 @@ void UDynamicMeshSculptTool::ApplyInflateBrush(const FRay& WorldRay)
 		return;
 	}
 
-	FTransform Transform = ComponentTarget->GetWorldTransform();
-	FVector NewBrushPosLocal = Transform.InverseTransformPosition(LastBrushPosWorld);
+	FVector3d NewBrushPosLocal = CurTargetTransform.InverseTransformPosition(LastBrushPosWorld);
 
 	double Direction = (bInvert) ? -1.0 : 1.0;
 	double UseSpeed = Direction * FMathd::Sqrt(CurrentBrushRadius) * SculptProperties->BrushSpeed * 0.25;
@@ -954,18 +953,17 @@ void UDynamicMeshSculptTool::ApplyInflateBrush(const FRay& WorldRay)
 
 bool UDynamicMeshSculptTool::UpdateBrushPositionOnActivePlane(const FRay& WorldRay)
 {
-	FVector NewHitPosWorld = FMath::RayPlaneIntersection(WorldRay.Origin, WorldRay.Direction, ActiveDragPlane);
+	FVector3d NewHitPosWorld;
+	ActiveDragPlane.RayPlaneIntersection(WorldRay.Origin, WorldRay.Direction, 2, NewHitPosWorld);
 	LastBrushPosWorld = NewHitPosWorld;
-	LastBrushPosNormalWorld = ActiveDragPlane;
+	LastBrushPosNormalWorld = ActiveDragPlane.Z();
 	return true;
 }
 
 bool UDynamicMeshSculptTool::UpdateBrushPositionOnTargetMesh(const FRay& WorldRay)
 {
-	FTransform Transform = ComponentTarget->GetWorldTransform();
-
-	FRay3d LocalRay(Transform.InverseTransformPosition(WorldRay.Origin),
-		Transform.InverseTransformVectorNoScale(WorldRay.Direction));
+	FRay3d LocalRay(CurTargetTransform.InverseTransformPosition(WorldRay.Origin),
+		CurTargetTransform.InverseTransformVector(WorldRay.Direction));
 	LocalRay.Direction.Normalize();
 
 	int HitTID = BrushTargetMeshSpatial.FindNearestHitTriangle(LocalRay);
@@ -978,8 +976,8 @@ bool UDynamicMeshSculptTool::UpdateBrushPositionOnTargetMesh(const FRay& WorldRa
 		FIntrRay3Triangle3d Query(LocalRay, Triangle);
 		Query.Find();
 
-		LastBrushPosNormalWorld = Transform.TransformVectorNoScale(TargetMesh->GetTriNormal(HitTID));
-		LastBrushPosWorld = Transform.TransformPosition(LocalRay.PointAt(Query.RayParameter));
+		LastBrushPosNormalWorld = CurTargetTransform.TransformNormal(TargetMesh->GetTriNormal(HitTID));
+		LastBrushPosWorld = CurTargetTransform.TransformPosition(LocalRay.PointAt(Query.RayParameter));
 		return true;
 	}
 	return false;
@@ -988,10 +986,8 @@ bool UDynamicMeshSculptTool::UpdateBrushPositionOnTargetMesh(const FRay& WorldRa
 
 bool UDynamicMeshSculptTool::UpdateBrushPositionOnSculptMesh(const FRay& WorldRay)
 {
-	FTransform Transform = ComponentTarget->GetWorldTransform();
-
-	FRay3d LocalRay(Transform.InverseTransformPosition(WorldRay.Origin),
-		Transform.InverseTransformVectorNoScale(WorldRay.Direction));
+	FRay3d LocalRay(CurTargetTransform.InverseTransformPosition(WorldRay.Origin),
+		CurTargetTransform.InverseTransformVector(WorldRay.Direction));
 	LocalRay.Direction.Normalize();
 
 	int HitTID = DynamicMeshComponent->GetOctree()->FindNearestHitObject(LocalRay);
@@ -1004,8 +1000,8 @@ bool UDynamicMeshSculptTool::UpdateBrushPositionOnSculptMesh(const FRay& WorldRa
 		FIntrRay3Triangle3d Query(LocalRay, Triangle);
 		Query.Find();
 
-		LastBrushPosNormalWorld = Transform.TransformVectorNoScale(SculptMesh->GetTriNormal(HitTID));
-		LastBrushPosWorld = Transform.TransformPosition(LocalRay.PointAt(Query.RayParameter));
+		LastBrushPosNormalWorld = CurTargetTransform.TransformNormal(SculptMesh->GetTriNormal(HitTID));
+		LastBrushPosWorld = CurTargetTransform.TransformPosition(LocalRay.PointAt(Query.RayParameter));
 		return true;
 	}
 	return false;
@@ -1035,9 +1031,10 @@ bool UDynamicMeshSculptTool::OnUpdateHover(const FInputDeviceRay& DevicePos)
 
 	if (bInDrag)
 	{
-		FVector NewHitPosWorld = FMath::RayPlaneIntersection(DevicePos.WorldRay.Origin, DevicePos.WorldRay.Direction, ActiveDragPlane);
+		FVector3d NewHitPosWorld;
+		ActiveDragPlane.RayPlaneIntersection(DevicePos.WorldRay.Origin, DevicePos.WorldRay.Direction, 2, NewHitPosWorld);
 		LastBrushPosWorld = NewHitPosWorld;
-		LastBrushPosNormalWorld = ActiveDragPlane;
+		LastBrushPosNormalWorld = ActiveDragPlane.Z();
 	}
 	else
 	{
diff --git a/Engine/Plugins/Experimental/MeshModelingToolset/Source/MeshModelingTools/Public/DynamicMeshSculptTool.h b/Engine/Plugins/Experimental/MeshModelingToolset/Source/MeshModelingTools/Public/DynamicMeshSculptTool.h
--- a/Engine/Plugins/Experimental/MeshModelingToolset/Source/MeshModelingTools/Public/DynamicMeshSculptTool.h
+++ b/Engine/Plugins/Experimental/MeshModelingToolset/Source/MeshModelingTools/Public/DynamicMeshSculptTool.h
@@ -14,6 +14,7 @@
 #include "Changes/ValueWatcher.h"
 #include "Gizmos/BrushStampIndicator.h"
 #include "Properties/MeshMaterialProperties.h"
+#include "TransformTypes.h"
 #include "DynamicMeshSculptTool.generated.h"
 
 
@@ -272,6 +273,9 @@ protected:
 	UPROPERTY()
 	UOctreeDynamicMeshComponent* DynamicMeshComponent;
 
+	FTransform3d InitialTargetTransform;
+	FTransform3d CurTargetTransform;
+
 	// realtime visualization
 	void OnDynamicMeshComponentChanged();
 	FDelegateHandle OnDynamicMeshComponentChangedHandle;
@@ -290,14 +294,14 @@ protected:
 
 	bool bInDrag;
 
-	FPlane ActiveDragPlane;
-	FVector LastHitPosWorld;
-	FVector BrushStartCenterWorld;
-	FVector BrushStartNormalWorld;
-	FVector LastBrushPosLocal;
-	FVector LastBrushPosWorld;
-	FVector LastBrushPosNormalWorld;
-	FVector LastSmoothBrushPosLocal;
+	FFrame3d ActiveDragPlane;
+	FVector3d LastHitPosWorld;
+	FVector3d BrushStartCenterWorld;
+	FVector3d BrushStartNormalWorld;
+	FVector3d LastBrushPosLocal;
+	FVector3d LastBrushPosWorld;
+	FVector3d LastBrushPosNormalWorld;
+	FVector3d LastSmoothBrushPosLocal;
 	
 
 	TArray<int> VertexROI;
diff --git a/Engine/Plugins/Experimental/MeshModelingToolset/Source/ModelingComponents/Private/OctreeDynamicMeshComponent.cpp b/Engine/Plugins/Experimental/MeshModelingToolset/Source/ModelingComponents/Private/OctreeDynamicMeshComponent.cpp
--- a/Engine/Plugins/Experimental/MeshModelingToolset/Source/ModelingComponents/Private/OctreeDynamicMeshComponent.cpp
+++ b/Engine/Plugins/Experimental/MeshModelingToolset/Source/ModelingComponents/Private/OctreeDynamicMeshComponent.cpp
@@ -19,6 +19,7 @@
 
 #include "DynamicMeshAttributeSet.h"
 #include "MeshNormals.h"
+#include "MeshTransforms.h"
 #include "MeshDescriptionToDynamicMesh.h"
 #include "DynamicMeshToMeshDescription.h"
 
@@ -92,6 +93,37 @@ void UOctreeDynamicMeshComponent::InitializeNewMesh()
 }
 
 
+
+void UOctreeDynamicMeshComponent::ApplyTransform(const FTransform3d& Transform, bool bInvert)
+{
+	if (bInvert)
+	{
+		MeshTransforms::ApplyTransformInverse(*GetMesh(), Transform);
+	}
+	else
+	{
+		MeshTransforms::ApplyTransform(*GetMesh(), Transform);
+	}
+
+	if (CurrentProxy != nullptr)
+	{
+		Octree->ModifiedBounds = FAxisAlignedBox3d(
+			-TNumericLimits<float>::Max()*FVector3d::One(),
+			TNumericLimits<float>::Max()*FVector3d::One());
+		NotifyMeshUpdated();
+	}
+	else
+	{
+		FAxisAlignedBox3d MeshBounds = Mesh->GetCachedBounds();
+		Octree = MakeUnique<FDynamicMeshOctree3>();
+		Octree->RootDimension = MeshBounds.MaxDim() * 0.5;
+		Octree->Initialize(GetMesh());
+		OctreeCut = MakeUnique<FDynamicMeshOctree3::FTreeCutSet>();
+	}
+
+}
+
+
 void UOctreeDynamicMeshComponent::Bake(FMeshDescription* MeshDescription, bool bHaveModifiedTopology, const FConversionToMeshDescriptionOptions& ConversionOptions)
 {
 	if (bHaveModifiedTopology == false)
diff --git a/Engine/Plugins/Experimental/MeshModelingToolset/Source/ModelingComponents/Public/OctreeDynamicMeshComponent.h b/Engine/Plugins/Experimental/MeshModelingToolset/Source/ModelingComponents/Public/OctreeDynamicMeshComponent.h
--- a/Engine/Plugins/Experimental/MeshModelingToolset/Source/ModelingComponents/Public/OctreeDynamicMeshComponent.h
+++ b/Engine/Plugins/Experimental/MeshModelingToolset/Source/ModelingComponents/Public/OctreeDynamicMeshComponent.h
@@ -72,6 +72,11 @@ public:
 		Bake(MeshDescription, bHaveModifiedTopology, ConversionOptions);
 	}
 
+	/**
+	 * Apply transform to internal mesh. Updates Octree and RenderProxy if available.
+	 * @param bInvert if true, inverse tranform is applied instead of forward transform
+	 */
+	void ApplyTransform(const FTransform3d& Transform, bool bInvert);
 
 	//
 	// change tracking/etc
