fix comment associated with sharp normals fix
#rnx

diff --git a/Engine/Plugins/Experimental/MeshModelingToolset/Source/MeshModelingTools/Public/EditNormalsTool.h b/Engine/Plugins/Experimental/MeshModelingToolset/Source/MeshModelingTools/Public/EditNormalsTool.h
--- a/Engine/Plugins/Experimental/MeshModelingToolset/Source/MeshModelingTools/Public/EditNormalsTool.h
+++ b/Engine/Plugins/Experimental/MeshModelingToolset/Source/MeshModelingTools/Public/EditNormalsTool.h
@@ -80,7 +80,7 @@ public:
 	UPROPERTY(EditAnywhere, Category = NormalsTopology, meta = (UIMin = "0.0", UIMax = "180.0", ClampMin = "0.0", ClampMax = "180.0", EditCondition = "bRecomputeNormalTopologyAndEdgeSharpness"))
 	float SharpEdgeAngleThreshold;
 
-	/** Give separate normals to each triangle at 'sharp' vertices where the coincident triangle normals change (across multiple triangles) by more than the angle threshold (e.g. a cone tip) */
+	/** Assign separate normals at 'sharp' vertices -- for example, at the tip of a cone */
 	UPROPERTY(EditAnywhere, Category = NormalsTopology, meta = (EditCondition = "bRecomputeNormalTopologyAndEdgeSharpness"))
 	bool bAllowSharpVertices;
 };
