fix: Translate debug line is hanging at the end of the transform
reset quick transforms on drag end

#rnx

diff --git a/Engine/Plugins/Experimental/MeshModelingToolset/Source/MeshModelingTools/Private/EditMeshPolygonsTool.cpp b/Engine/Plugins/Experimental/MeshModelingToolset/Source/MeshModelingTools/Private/EditMeshPolygonsTool.cpp
--- a/Engine/Plugins/Experimental/MeshModelingToolset/Source/MeshModelingTools/Private/EditMeshPolygonsTool.cpp
+++ b/Engine/Plugins/Experimental/MeshModelingToolset/Source/MeshModelingTools/Private/EditMeshPolygonsTool.cpp
@@ -1049,7 +1049,8 @@ void UEditMeshPolygonsTool::OnEndDrag(const FRay& Ray)
 
 	HilightSelection.Clear(); 
 	TopoSelector.Invalidate(true, false);
-	QuickAxisRotator.ClearAxisLock();
+	QuickAxisRotator.Reset();
+	QuickAxisTranslater.Reset();
 
 	//If it's linear, it's computed real time with no delay. This may need to be restructured for clarity by using the background task for this as well.
 	if (DeformationStrategy == EGroupTopologyDeformationStrategy::Linear)
diff --git a/Engine/Plugins/Experimental/MeshModelingToolset/Source/ModelingComponents/Private/Transforms/QuickAxisRotator.cpp b/Engine/Plugins/Experimental/MeshModelingToolset/Source/ModelingComponents/Private/Transforms/QuickAxisRotator.cpp
--- a/Engine/Plugins/Experimental/MeshModelingToolset/Source/ModelingComponents/Private/Transforms/QuickAxisRotator.cpp
+++ b/Engine/Plugins/Experimental/MeshModelingToolset/Source/ModelingComponents/Private/Transforms/QuickAxisRotator.cpp
@@ -193,3 +193,9 @@ void FQuickAxisRotator::PreviewRender(IToolsContextRenderAPI* RenderAPI)
 	MoveAxisSolver.Draw(&QuickAxisPreviewRenderer, 2 * CurViewSizeFactor);
 	QuickAxisPreviewRenderer.EndFrame();
 }
+
+void FQuickAxisRotator::Reset()
+{
+	MoveAxisSolver.Reset();
+	ClearAxisLock();
+}
diff --git a/Engine/Plugins/Experimental/MeshModelingToolset/Source/ModelingComponents/Private/Transforms/QuickAxisTranslater.cpp b/Engine/Plugins/Experimental/MeshModelingToolset/Source/ModelingComponents/Private/Transforms/QuickAxisTranslater.cpp
--- a/Engine/Plugins/Experimental/MeshModelingToolset/Source/ModelingComponents/Private/Transforms/QuickAxisTranslater.cpp
+++ b/Engine/Plugins/Experimental/MeshModelingToolset/Source/ModelingComponents/Private/Transforms/QuickAxisTranslater.cpp
@@ -226,3 +226,9 @@ void FQuickAxisTranslater::PreviewRender(IToolsContextRenderAPI* RenderAPI)
 	MoveAxisSolver.Draw(&QuickAxisPreviewRenderer, 2 * CurViewSizeFactor);
 	QuickAxisPreviewRenderer.EndFrame();
 }
+
+
+void FQuickAxisTranslater::Reset()
+{
+	MoveAxisSolver.Reset();
+}
diff --git a/Engine/Plugins/Experimental/MeshModelingToolset/Source/ModelingComponents/Public/Transforms/QuickAxisRotator.h b/Engine/Plugins/Experimental/MeshModelingToolset/Source/ModelingComponents/Public/Transforms/QuickAxisRotator.h
--- a/Engine/Plugins/Experimental/MeshModelingToolset/Source/ModelingComponents/Public/Transforms/QuickAxisRotator.h
+++ b/Engine/Plugins/Experimental/MeshModelingToolset/Source/ModelingComponents/Public/Transforms/QuickAxisRotator.h
@@ -64,6 +64,9 @@ public:
 
 	virtual FFrame3d GetActiveRotationFrame() const;
 
+	/** Reset transformer state */
+	virtual void Reset() override;
+
 
 	void ClearAxisLock();
 	void SetAxisLock();
diff --git a/Engine/Plugins/Experimental/MeshModelingToolset/Source/ModelingComponents/Public/Transforms/QuickAxisTranslater.h b/Engine/Plugins/Experimental/MeshModelingToolset/Source/ModelingComponents/Public/Transforms/QuickAxisTranslater.h
--- a/Engine/Plugins/Experimental/MeshModelingToolset/Source/ModelingComponents/Public/Transforms/QuickAxisTranslater.h
+++ b/Engine/Plugins/Experimental/MeshModelingToolset/Source/ModelingComponents/Public/Transforms/QuickAxisTranslater.h
@@ -81,6 +81,9 @@ public:
 	 */
 	virtual void PreviewRender(IToolsContextRenderAPI* RenderAPI) override;
 
+	/** Reset transformer state */
+	virtual void Reset() override;
+
 
 protected:
 	// camera state saved at last Render()
diff --git a/Engine/Plugins/Experimental/MeshModelingToolset/Source/ModelingComponents/Public/Transforms/QuickTransformer.h b/Engine/Plugins/Experimental/MeshModelingToolset/Source/ModelingComponents/Public/Transforms/QuickTransformer.h
--- a/Engine/Plugins/Experimental/MeshModelingToolset/Source/ModelingComponents/Public/Transforms/QuickTransformer.h
+++ b/Engine/Plugins/Experimental/MeshModelingToolset/Source/ModelingComponents/Public/Transforms/QuickTransformer.h
@@ -37,9 +37,10 @@ public:
 	 */
 	virtual void UpdateActiveFrameOrigin(const FVector3d& NewOrigin) = 0;
 
-
-
-
+	/**
+	 * Reset transformer state
+	 */
+	virtual void Reset() = 0;
 
 
 	/**
