Fix Jira UE-79997
Fix Jira UE-81800
Changed computation logic in DisplaceMeshTool to be asynchronous

diff --git a/Engine/Plugins/Experimental/MeshModelingToolset/Source/MeshModelingTools/Private/DisplaceMeshTool.cpp b/Engine/Plugins/Experimental/MeshModelingToolset/Source/MeshModelingTools/Private/DisplaceMeshTool.cpp
--- a/Engine/Plugins/Experimental/MeshModelingToolset/Source/MeshModelingTools/Private/DisplaceMeshTool.cpp
+++ b/Engine/Plugins/Experimental/MeshModelingToolset/Source/MeshModelingTools/Private/DisplaceMeshTool.cpp
@@ -3,20 +3,402 @@
 #include "DisplaceMeshTool.h"
 #include "InteractiveToolManager.h"
 #include "ToolBuilderUtil.h"
-
+#include "ToolSetupUtil.h"
 #include "DynamicMesh3.h"
-
 #include "SimpleDynamicMeshComponent.h"
-
 #include "MeshNormals.h"
-
+#include "ModelingOperators.h"
 #include "SceneManagement.h" // for FPrimitiveDrawInterface
 #include "Async/ParallelFor.h"
+#include "ProfilingDebugging/ScopedTimers.h"
+#define LOCTEXT_NAMESPACE "UDisplaceMeshTool"
 
+namespace {
+	void SubdivideMesh(FDynamicMesh3& Mesh)
+	{
+		TArray<int> EdgesToProcess;
+		for (int tid : Mesh.EdgeIndicesItr())
+		{
+			EdgesToProcess.Add(tid);
+		}
+		int MaxTriangleID = Mesh.MaxTriangleID();
 
-#include "ProfilingDebugging/ScopedTimers.h"
+		TArray<int> TriSplitEdges;
+		TriSplitEdges.Init(-1, Mesh.MaxTriangleID());
 
-#define LOCTEXT_NAMESPACE "UDisplaceMeshTool"
+		for (int eid : EdgesToProcess)
+		{
+			FIndex2i EdgeTris = Mesh.GetEdgeT(eid);
+
+			FDynamicMesh3::FEdgeSplitInfo SplitInfo;
+			EMeshResult result = Mesh.SplitEdge(eid, SplitInfo);
+			if (result == EMeshResult::Ok)
+			{
+				if (EdgeTris.A < MaxTriangleID && TriSplitEdges[EdgeTris.A] == -1)
+				{
+					TriSplitEdges[EdgeTris.A] = SplitInfo.NewEdges.B;
+				}
+				if (EdgeTris.B != FDynamicMesh3::InvalidID)
+				{
+					if (EdgeTris.B < MaxTriangleID && TriSplitEdges[EdgeTris.B] == -1)
+					{
+						TriSplitEdges[EdgeTris.B] = SplitInfo.NewEdges.C;
+					}
+				}
+			}
+		}
+
+		for (int eid : TriSplitEdges)
+		{
+			if (eid != -1)
+			{
+				FDynamicMesh3::FEdgeFlipInfo FlipInfo;
+				Mesh.FlipEdge(eid, FlipInfo);
+			}
+		}
+	}
+
+	namespace ComputeDisplacement {
+		void Constant(const FDynamicMesh3& Mesh,
+				const TArray<FVector3d>& Positions, const FMeshNormals& Normals, double Intensity,
+				TArray<FVector3d>& DisplacedPositions)
+		{
+			DisplacedPositions.SetNumUninitialized(Positions.Num());
+			for (int vid : Mesh.VertexIndicesItr())
+			{
+				DisplacedPositions[vid] = Positions[vid] + (Intensity * Normals[vid]);
+			}
+		}
+
+		void RandomNoise(const FDynamicMesh3& Mesh,
+				const TArray<FVector3d>& Positions, const FMeshNormals& Normals,
+				double Intensity, int RandomSeed,
+				TArray<FVector3d>& DisplacedPositions)
+		{
+			FMath::SRandInit(RandomSeed);
+			for (int vid : Mesh.VertexIndicesItr())
+			{
+				double RandVal = 2.0 * (FMath::SRand() - 0.5);
+				DisplacedPositions[vid] = Positions[vid] + (Normals[vid] * RandVal * Intensity);
+			}
+		}
+
+		void Map(const FDynamicMesh3& Mesh,
+				const TArray<FVector3d>& Positions, const FMeshNormals& Normals,
+				double Intensity, const FSampledScalarField2f& DisplaceField,
+				TArray<FVector3d>& DisplacedPositions)
+		{
+			const FDynamicMeshUVOverlay* UVOverlay = Mesh.Attributes()->GetUVLayer(0);
+			for (int tid : Mesh.TriangleIndicesItr())
+			{
+				FIndex3i Tri = Mesh.GetTriangle(tid);
+				FIndex3i UVTri = UVOverlay->GetTriangle(tid);
+				for (int j = 0; j < 3; ++j)
+				{
+					int vid = Tri[j];
+					FVector2f UV = UVOverlay->GetElement(UVTri[j]);
+					double Offset = DisplaceField.BilinearSampleClamped(UV);
+					DisplacedPositions[vid] = Positions[vid] + (Offset * Intensity * Normals[vid]);
+				}
+			}
+
+		}
+	}
+	class FTextureAccess
+	{
+	public:
+		FTextureAccess(UTexture2D* DisplacementMap)
+			:DisplacementMap(DisplacementMap)
+		{
+			check(DisplacementMap);
+			OldCompressionSettings = DisplacementMap->CompressionSettings;
+			bOldSRGB = DisplacementMap->SRGB;
+#if WITH_EDITOR
+			OldMipGenSettings = DisplacementMap->MipGenSettings;
+#endif
+			DisplacementMap->CompressionSettings = TextureCompressionSettings::TC_VectorDisplacementmap;
+			DisplacementMap->SRGB = false;
+#if WITH_EDITOR
+			DisplacementMap->MipGenSettings = TextureMipGenSettings::TMGS_NoMipmaps;
+#endif
+			DisplacementMap->UpdateResource();
+
+			FormattedImageData = static_cast<const FColor*>(DisplacementMap->PlatformData->Mips[0].BulkData.LockReadOnly());
+		}
+		FTextureAccess(const FTextureAccess&) = delete;
+		FTextureAccess(FTextureAccess&&) = delete;
+		void operator=(const FTextureAccess&) = delete;
+		void operator=(FTextureAccess&&) = delete;
+
+		~FTextureAccess()
+		{
+			DisplacementMap->PlatformData->Mips[0].BulkData.Unlock();
+
+			DisplacementMap->CompressionSettings = OldCompressionSettings;
+			DisplacementMap->SRGB = bOldSRGB;
+#if WITH_EDITOR
+			DisplacementMap->MipGenSettings = OldMipGenSettings;
+#endif
+
+			DisplacementMap->UpdateResource();
+		}
+
+		bool HasData() const
+		{
+			return FormattedImageData != nullptr;
+		}
+		const FColor* GetData() const
+		{
+			return FormattedImageData;
+		}
+
+	private:
+		UTexture2D* DisplacementMap{ nullptr };
+		TextureCompressionSettings OldCompressionSettings{};
+		TextureMipGenSettings OldMipGenSettings{};
+		bool bOldSRGB{ false };
+		const FColor* FormattedImageData{ nullptr };
+	};
+
+	class FSubdivideMeshOp : public FDynamicMeshOperator
+	{
+	public:
+		FSubdivideMeshOp(const FDynamicMesh3& SourceMesh, int SubdivisionsCountIn);
+		void CalculateResult(FProgressCancel* Progress) final;
+	private:
+		int SubdivisionsCount;
+	};
+
+	FSubdivideMeshOp::FSubdivideMeshOp(const FDynamicMesh3& SourceMesh, int SubdivisionsCountIn)
+		: SubdivisionsCount(SubdivisionsCountIn)
+	{
+		ResultMesh->Copy(SourceMesh);
+	}
+
+	void FSubdivideMeshOp::CalculateResult(FProgressCancel* ProgressCancel)
+	{
+		// calculate subdivisions (todo: move to elsewhere)
+		for (int ri = 0; ri < SubdivisionsCount; ri++)
+		{
+			if (ProgressCancel->Cancelled()) return;
+			SubdivideMesh(*ResultMesh);
+		}
+	}
+
+	class FSubdivideMeshOpFactory : public IDynamicMeshOperatorFactory
+	{
+	public:
+		FSubdivideMeshOpFactory(FDynamicMesh3& SourceMeshIn,
+			int SubdivisionsCountIn)
+			: SourceMesh(SourceMeshIn), SubdivisionsCount(SubdivisionsCountIn)
+		{
+		}
+		void SetSubdivisionsCount(int SubdivisionsCountIn);
+		int  GetSubdivisionsCount();
+
+		TSharedPtr<FDynamicMeshOperator> MakeNewOperator() final
+		{
+			return MakeShared<FSubdivideMeshOp>(SourceMesh, SubdivisionsCount);
+		}
+	private:
+		const FDynamicMesh3& SourceMesh;
+		int SubdivisionsCount;
+	};
+
+	void FSubdivideMeshOpFactory::SetSubdivisionsCount(int SubdivisionsCountIn)
+	{
+		SubdivisionsCount = SubdivisionsCountIn;
+	}
+
+	int FSubdivideMeshOpFactory::GetSubdivisionsCount()
+	{
+		return SubdivisionsCount;
+	}
+
+	class FDisplaceMeshOp : public FDynamicMeshOperator
+	{
+	public:
+		FDisplaceMeshOp(TSharedPtr<FDynamicMesh3> SourceMeshIn, const FSampledScalarField2f& DisplaceFieldIn,
+			float DisplaceIntensityIn, int RandomSeedIn, EDisplaceMeshToolDisplaceType DisplacementTypeIn);
+		void CalculateResult(FProgressCancel* Progress) final;
+
+	private:
+		TSharedPtr<FDynamicMesh3> SourceMesh;
+		float DisplaceIntensity;
+		int RandomSeed;
+		EDisplaceMeshToolDisplaceType DisplacementType;
+		const FSampledScalarField2f& DisplaceField;
+
+		TArray<FVector3d> SourcePositions;
+		FMeshNormals SourceNormals;
+		TArray<FVector3d> DisplacedPositions;
+	};
+
+	FDisplaceMeshOp::FDisplaceMeshOp(TSharedPtr<FDynamicMesh3> SourceMeshIn, const FSampledScalarField2f& DisplaceFieldIn,
+		float DisplaceIntensityIn, int RandomSeedIn, EDisplaceMeshToolDisplaceType DisplacementTypeIn)
+		: SourceMesh(MoveTemp(SourceMeshIn)), DisplaceIntensity(DisplaceIntensityIn), RandomSeed(RandomSeedIn),
+		DisplacementType(DisplacementTypeIn),DisplaceField(DisplaceFieldIn)
+	{
+	}
+
+	void FDisplaceMeshOp::CalculateResult(FProgressCancel* Progress)
+	{
+		if (Progress->Cancelled()) return;
+		ResultMesh->Copy(*SourceMesh);
+
+		if (Progress->Cancelled()) return;
+		SourceNormals = FMeshNormals(SourceMesh.Get());
+		SourceNormals.ComputeVertexNormals();
+
+		if (Progress->Cancelled()) return;
+		// cache initial positions
+		SourcePositions.SetNum(SourceMesh->MaxVertexID());
+		for (int vid : SourceMesh->VertexIndicesItr())
+		{
+			SourcePositions[vid] = SourceMesh->GetVertex(vid);
+		}
+
+		if (Progress->Cancelled()) return;
+		DisplacedPositions.SetNum(SourceMesh->MaxVertexID());
+
+		if (Progress->Cancelled()) return;
+		// compute Displaced positions in PositionBuffer
+		switch (DisplacementType)
+		{
+		default:
+		case EDisplaceMeshToolDisplaceType::Constant:
+			ComputeDisplacement::Constant(*SourceMesh, SourcePositions, SourceNormals,
+				DisplaceIntensity,
+				DisplacedPositions);
+			break;
+
+		case EDisplaceMeshToolDisplaceType::RandomNoise:
+			ComputeDisplacement::RandomNoise(*SourceMesh, SourcePositions, SourceNormals,
+				DisplaceIntensity, RandomSeed,
+				DisplacedPositions);
+			break;
+
+		case EDisplaceMeshToolDisplaceType::DisplacementMap:
+			ComputeDisplacement::Map(*SourceMesh, SourcePositions, SourceNormals,
+				DisplaceIntensity, DisplaceField,
+				DisplacedPositions);
+			break;
+		}
+
+		// update preview vertex positions
+		for (int vid : ResultMesh->VertexIndicesItr())
+		{
+			ResultMesh->SetVertex(vid, DisplacedPositions[vid]);
+		}
+
+		// recalculate normals
+		if (ResultMesh->HasAttributes())
+		{
+			FMeshNormals Normals(ResultMesh.Get());
+			FDynamicMeshNormalOverlay* NormalOverlay = ResultMesh->Attributes()->PrimaryNormals();
+			Normals.RecomputeOverlayNormals(NormalOverlay);
+			Normals.CopyToOverlay(NormalOverlay);
+		}
+		else
+		{
+			FMeshNormals::QuickComputeVertexNormals(*ResultMesh);
+		}
+	}
+
+	class FDisplaceMeshOpFactory : public IDynamicMeshOperatorFactory
+	{
+	public:
+		FDisplaceMeshOpFactory(TSharedPtr<FDynamicMesh3>& SourceMeshIn,
+			float DisplaceIntensityIn, int RandomSeedIn, UTexture2D* DisplacementMapIn,
+			EDisplaceMeshToolDisplaceType DisplacementTypeIn)
+			: SourceMesh(SourceMeshIn)
+		{
+			SetIntensity(DisplaceIntensityIn);
+			SetRandomSeed(RandomSeed);
+			SetDisplacementMap(DisplacementMapIn);
+			SetDisplacementType(DisplacementTypeIn);
+
+		}
+		void SetIntensity(float IntensityIn);
+		void SetRandomSeed(int RandomSeedIn);
+		void SetDisplacementMap(UTexture2D* DisplacementMapIn);
+		void SetDisplacementType(EDisplaceMeshToolDisplaceType TypeIn);
+
+		TSharedPtr<FDynamicMeshOperator> MakeNewOperator() final
+		{
+			return MakeShared<FDisplaceMeshOp>(SourceMesh, DisplaceField, DisplaceIntensity, RandomSeed, DisplacementType);
+		}
+	private:
+		void UpdateMap();
+
+		float DisplaceIntensity;
+		int RandomSeed;
+		UTexture2D* DisplacementMap;
+		EDisplaceMeshToolDisplaceType DisplacementType;
+
+		TSharedPtr<FDynamicMesh3>& SourceMesh;
+		FSampledScalarField2f DisplaceField;
+	};
+
+	void FDisplaceMeshOpFactory::SetIntensity(float IntensityIn)
+	{
+		DisplaceIntensity = IntensityIn;
+	}
+
+	void FDisplaceMeshOpFactory::SetRandomSeed(int RandomSeedIn)
+	{
+		RandomSeed = RandomSeedIn;
+	}
+
+	void FDisplaceMeshOpFactory::SetDisplacementMap(UTexture2D* DisplacementMapIn)
+	{
+		if (DisplacementMap != DisplacementMapIn)
+		{
+			DisplacementMap = DisplacementMapIn;
+			UpdateMap();
+		}
+	}
+
+	void FDisplaceMeshOpFactory::SetDisplacementType(EDisplaceMeshToolDisplaceType TypeIn)
+	{
+		DisplacementType = TypeIn;
+	}
+
+	void FDisplaceMeshOpFactory::UpdateMap()
+	{
+		if (DisplacementMap == nullptr ||
+			DisplacementMap->PlatformData == nullptr ||
+			DisplacementMap->PlatformData->Mips.Num() < 1)
+		{
+			DisplaceField = FSampledScalarField2f();
+			return;
+		}
+
+		FTextureAccess TextureAccess(DisplacementMap);
+		if (!TextureAccess.HasData())
+		{
+			DisplaceField = FSampledScalarField2f();
+		}
+		else
+		{
+			int TextureWidth = DisplacementMap->GetSizeX();
+			int TextureHeight = DisplacementMap->GetSizeY();
+			DisplaceField.Resize(TextureWidth, TextureHeight, 0.0f);
+			DisplaceField.SetCellSize(1.0f / (float)TextureWidth);
+
+			const FColor* FormattedData = TextureAccess.GetData();
+			for (int y = 0; y < TextureHeight; ++y)
+			{
+				for (int x = 0; x < TextureWidth; ++x)
+				{
+					FColor PixelColor = FormattedData[y*TextureWidth + x];
+					float Value = PixelColor.R / 255.0;
+					DisplaceField.GridValues[y*TextureWidth + x] = Value;
+				}
+			}
+		}
+	}
+} // namespace
 
 /*
  * ToolBuilder
@@ -44,7 +426,7 @@ UInteractiveTool* UDisplaceMeshToolBuilder::BuildTool(const FToolBuilderState& S
  */
 UDisplaceMeshTool::UDisplaceMeshTool()
 {
-	DisplacementType = EDisplaceMeshToolDisplaceType::DisplacementMap;
+	DisplacementType = EDisplaceMeshToolDisplaceType::Constant;
 	DisplaceIntensity = 10.0f;
 	RandomSeed = 31337;
 	Subdivisions = 4;
@@ -70,23 +452,23 @@ void UDisplaceMeshTool::Setup()
 	DynamicMeshComponent->InitializeMesh(ComponentTarget->GetMesh());
 	OriginalMesh.Copy(*DynamicMeshComponent->GetMesh());
 
-	UpdateMap(true);
-
+	Subdivider = MakeUnique<FSubdivideMeshOpFactory>(OriginalMesh, Subdivisions);
+	Displacer = MakeUnique<FDisplaceMeshOpFactory>(SubdividedMesh,
+		DisplaceIntensity, RandomSeed, DisplacementMap, DisplacementType);
+		
 	// hide input StaticMeshComponent
 	ComponentTarget->SetOwnerVisibility(false);
 
 	// initialize our properties
 	ToolPropertyObjects.Add(this);
-
-	bResultValid = false;
+	ValidateSubdivisions();
+	StartComputation();
 }
 
 void UDisplaceMeshTool::Shutdown(EToolShutdownType ShutdownType)
 {
 	if (DynamicMeshComponent != nullptr)
 	{
-		//DynamicMeshComponent->OnMeshChanged.Remove(OnDynamicMeshComponentChangedHandle);
-
 		ComponentTarget->SetOwnerVisibility(true);
 
 		if (ShutdownType == EToolShutdownType::Accept)
@@ -106,329 +488,126 @@ void UDisplaceMeshTool::Shutdown(EToolShutdownType ShutdownType)
 	}
 }
 
-void UDisplaceMeshTool::Render(IToolsContextRenderAPI* RenderAPI)
+void UDisplaceMeshTool::ValidateSubdivisions()
 {
-	UpdateResult();
-}
-
-#if WITH_EDITOR
-void UDisplaceMeshTool::PostEditChangeProperty(FPropertyChangedEvent& PropertyChangedEvent)
-{
-	UProperty* PropertyThatChanged = PropertyChangedEvent.Property;
-
-	if (PropertyThatChanged)
-	{
-		//if (PropertyThatChanged->GetFName() == GET_MEMBER_NAME_CHECKED(UDisplaceMeshTool, DisplaceIterations))
-		//{
-		//}
-	}
-
-	bResultValid = false;
-}
-#endif
-
-void UDisplaceMeshTool::UpdateResult()
-{
-	if (bResultValid) 
-	{
-		return;
-	}
-
-	UpdateSubdividedMesh();
-
-	// compute Displaceed positions in PositionBuffer
-	switch (DisplacementType)
-	{
-		default:
-		case EDisplaceMeshToolDisplaceType::Constant:
-			ComputeDisplacement_Constant();
-			break;
-
-		case EDisplaceMeshToolDisplaceType::RandomNoise:
-			ComputeDisplacement_RandomNoise();
-			break;
-
-		case EDisplaceMeshToolDisplaceType::DisplacementMap:
-			ComputeDisplacement_Map();
-			break;
-	}
-
-	// update preview vertex positions
-	FDynamicMesh3* TargetMesh = DynamicMeshComponent->GetMesh();
-	for (int vid : SubdividedMesh.VertexIndicesItr())
+	constexpr int MaxTriangles = 3000000;
+	double NumTriangles = OriginalMesh.MaxTriangleID();
+	int MaxSubdivisions = (int)floor(log2(MaxTriangles / NumTriangles) / 2.0);
+	if (Subdivisions > MaxSubdivisions)
 	{
-		TargetMesh->SetVertex(vid, DisplacedBuffer[vid]);
-	}
-
-	// recalculate normals
-	if (TargetMesh->HasAttributes())
-	{
-		FMeshNormals Normals(TargetMesh);
-		FDynamicMeshNormalOverlay* NormalOverlay = TargetMesh->Attributes()->PrimaryNormals();
-		Normals.RecomputeOverlayNormals(NormalOverlay);
-		Normals.CopyToOverlay(NormalOverlay);
+		FText WarningText = FText::Format(LOCTEXT("SubdivisionsTooHigh", "Desired number of Subdivisions ({0}) exceeds maximum number of {1}"), FText::AsNumber(Subdivisions), FText::AsNumber(MaxSubdivisions));
+		GetToolManager()->DisplayMessage(WarningText, EToolMessageLevel::UserWarning);
+		Subdivisions = MaxSubdivisions;
 	}
 	else
 	{
-		FMeshNormals::QuickComputeVertexNormals(*TargetMesh);
+		FText ClearWarningText;
+		GetToolManager()->DisplayMessage(ClearWarningText, EToolMessageLevel::UserWarning);
 	}
-
-	DynamicMeshComponent->NotifyMeshUpdated();
-	//DynamicMeshComponent->FastNotifyPositionsUpdated(true);
-
-	GetToolManager()->PostInvalidation();
-
-	bResultValid = true;
-}
-
-void UDisplaceMeshTool::ComputeDisplacement_Constant()
-{
-	double Intensity = DisplaceIntensity;
-
-	for (int vid : SubdividedMesh.VertexIndicesItr())
-	{
-		FVector3d Position = PositionBuffer[vid];
-		FVector3d Normal = NormalsBuffer[vid];
-		FVector3d Displacement = Intensity * Normal;
-		DisplacedBuffer[vid] = Position + Displacement;
-	}
-}
-
-void UDisplaceMeshTool::ComputeDisplacement_RandomNoise()
-{
-	double Intensity = DisplaceIntensity;
-
-	FMath::SRandInit(RandomSeed);
-
-	for (int vid : SubdividedMesh.VertexIndicesItr())
+	if (Subdivisions < 0)
 	{
-		FVector3d Position = PositionBuffer[vid];
-		FVector3d Normal = NormalsBuffer[vid];
-
-		double RandVal = 2.0 * (FMath::SRand() - 0.5);
-		
-		FVector3d Displacement = RandVal * Intensity * Normal;
-
-		DisplacedBuffer[vid] = Position + Displacement;
+		Subdivisions = 0;
 	}
 }
 
-void UDisplaceMeshTool::ComputeDisplacement_Map()
+#if WITH_EDITOR
+void UDisplaceMeshTool::PostEditChangeProperty(FPropertyChangedEvent& PropertyChangedEvent)
 {
-	UpdateMap();
-
-	const FDynamicMeshUVOverlay* UVOverlay = SubdividedMesh.Attributes()->GetUVLayer(0);
-
-	double Intensity = DisplaceIntensity;
-	FMath::SRandInit(RandomSeed);
-
-	for (int tid : SubdividedMesh.TriangleIndicesItr())
+	UProperty* PropertyThatChanged = PropertyChangedEvent.Property;
+	if (PropertyThatChanged)
 	{
-		FIndex3i Tri = SubdividedMesh.GetTriangle(tid);
-		FIndex3i UVTri = UVOverlay->GetTriangle(tid);
-
-		for (int j = 0; j < 3; ++j)
+		FSubdivideMeshOpFactory*  SubdividerDownCast = static_cast<FSubdivideMeshOpFactory*>(Subdivider.Get());
+		FDisplaceMeshOpFactory*  DisplacerDownCast = static_cast<FDisplaceMeshOpFactory*>(Displacer.Get());
+		const FName PropName = PropertyThatChanged->GetFName();
+		bNeedsDisplaced = true;
+		if (PropName == GET_MEMBER_NAME_CHECKED(UDisplaceMeshTool, Subdivisions))
 		{
-			int vid = Tri[j];
-
-			FVector3d Position = PositionBuffer[vid];
-			FVector3d Normal = NormalsBuffer[vid];
-
-			FVector2f UV = UVOverlay->GetElement(UVTri[j]);
-			double Offset = DisplaceField.BilinearSampleClamped(UV);
-
-			FVector3d Displacement = Offset * Intensity * Normal;
-
-			DisplacedBuffer[vid] = Position + Displacement;
+			ValidateSubdivisions();
+			if (Subdivisions != SubdividerDownCast->GetSubdivisionsCount())
+			{
+				SubdividerDownCast->SetSubdivisionsCount(Subdivisions);
+				bNeedsSubdivided = true;
+			}
+			else
+			{
+				return;
+			}
 		}
-	}
-
-}
-
-namespace {
-	class FTextureAccess
-	{
-	public:
-		FTextureAccess(UTexture2D* DisplacementMap)
-			:DisplacementMap(DisplacementMap)
+		else if (PropName == GET_MEMBER_NAME_CHECKED(UDisplaceMeshTool, RandomSeed))
 		{
-			check(DisplacementMap);
-			OldCompressionSettings = DisplacementMap->CompressionSettings;
-			bOldSRGB = DisplacementMap->SRGB; 
-#if WITH_EDITOR
-			OldMipGenSettings = DisplacementMap->MipGenSettings;
-#endif
-
-
-			DisplacementMap->CompressionSettings = TextureCompressionSettings::TC_VectorDisplacementmap;
-			DisplacementMap->SRGB = false;
-#if WITH_EDITOR
-			DisplacementMap->MipGenSettings = TextureMipGenSettings::TMGS_NoMipmaps;
-#endif
-			DisplacementMap->UpdateResource();
-
-			FormattedImageData = static_cast<const FColor*>(DisplacementMap->PlatformData->Mips[0].BulkData.LockReadOnly());
+			DisplacerDownCast->SetRandomSeed(RandomSeed);
 		}
-		FTextureAccess(const FTextureAccess&) = delete;
-		FTextureAccess(FTextureAccess&&) = delete;
-		void operator=(const FTextureAccess&) = delete;
-		void operator=(FTextureAccess&&) = delete;
-
-		~FTextureAccess()
+		else if (PropName == GET_MEMBER_NAME_CHECKED(UDisplaceMeshTool, DisplacementType))
 		{
-			DisplacementMap->PlatformData->Mips[0].BulkData.Unlock();
-
-			DisplacementMap->CompressionSettings = OldCompressionSettings;
-			DisplacementMap->SRGB = bOldSRGB;
-#if WITH_EDITOR
-			DisplacementMap->MipGenSettings = OldMipGenSettings;
-#endif
-
-			DisplacementMap->UpdateResource();
+			DisplacerDownCast->SetDisplacementType(DisplacementType);
 		}
-
-		bool HasData() const
+		else if (PropName == GET_MEMBER_NAME_CHECKED(UDisplaceMeshTool, DisplaceIntensity))
 		{
-			return FormattedImageData != nullptr;
+			DisplacerDownCast->SetIntensity(DisplaceIntensity);
 		}
-		const FColor* GetData() const
+		else if (PropName == GET_MEMBER_NAME_CHECKED(UDisplaceMeshTool, DisplacementMap))
 		{
-			return FormattedImageData;
+			DisplacerDownCast->SetDisplacementMap(DisplacementMap);
 		}
-
-	private:
-		UTexture2D* DisplacementMap{nullptr};
-		TextureCompressionSettings OldCompressionSettings{};
-		TextureMipGenSettings OldMipGenSettings{};
-		bool bOldSRGB{false};
-		const FColor* FormattedImageData{nullptr};
-	};
+		StartComputation();
+	}
 }
+#endif
 
-void UDisplaceMeshTool::UpdateMap(bool bForceUpdate)
+void UDisplaceMeshTool::Tick(float DeltaTime)
 {
-	if (CachedMapSource == DisplacementMap && bForceUpdate == false)
-	{
-		return;
-	}
-	CachedMapSource = DisplacementMap;
-
-	if (DisplacementMap == nullptr || DisplacementMap->PlatformData == nullptr || DisplacementMap->PlatformData->Mips.Num() < 1)
-	{
-		DisplaceField = FSampledScalarField2f();
-		return;
-	}
+	Super::Tick(DeltaTime); 
+	AdvanceComputation();
+}
 
-	FTextureAccess TextureAccess( DisplacementMap );
-	if (!TextureAccess.HasData())
-	{
-		DisplaceField = FSampledScalarField2f();
-	}
-	else
+void UDisplaceMeshTool::StartComputation()
+{
+	if ( bNeedsSubdivided )
 	{
-		int TextureWidth = DisplacementMap->GetSizeX();
-		int TextureHeight = DisplacementMap->GetSizeY();
-		DisplaceField.Resize(TextureWidth, TextureHeight, 0.0f);
-		DisplaceField.SetCellSize(1.0f / (float)TextureWidth);
-
-		const FColor* FormattedData = TextureAccess.GetData();
-		for (int y = 0; y < TextureHeight; ++y)
+		if (SubdivideTask)
 		{
-			for (int x = 0; x < TextureWidth; ++x)
-			{
-				FColor PixelColor = FormattedData[y*TextureWidth + x];
-				float Value = PixelColor.R / 255.0;
-				DisplaceField.GridValues[y*TextureWidth + x] = Value;
-			}
+			SubdivideTask->CancelAndDelete();
+			SubdividedMesh = nullptr;
 		}
+		auto SubdivideMeshOp = Subdivider->MakeNewOperator();
+		SubdivideTask = new FAsyncTaskExecuterWithAbort<TModelingOpTask<FDynamicMeshOperator>>(Subdivider->MakeNewOperator());
+		SubdivideTask->StartBackgroundTask();
+		bNeedsSubdivided = false;
+		DynamicMeshComponent->SetMaterial(0, ToolSetupUtil::GetDefaultWorkingMaterial(GetToolManager()));
+	}
+	if (bNeedsDisplaced && DisplaceTask)
+	{
+		DisplaceTask->CancelAndDelete();
+		DisplaceTask = nullptr;
+		DynamicMeshComponent->SetMaterial(0, ToolSetupUtil::GetDefaultWorkingMaterial(GetToolManager()));
 	}
+	AdvanceComputation();
 }
 
-void UDisplaceMeshTool::UpdateSubdividedMesh()
+void UDisplaceMeshTool::AdvanceComputation()
 {
-	if (Subdivisions == CachedSubdivisionsCount)
+	if (SubdivideTask && SubdivideTask->IsDone())
 	{
-		return;
+		SubdividedMesh = TSharedPtr<FDynamicMesh3>(SubdivideTask->GetTask().ExtractOperator()->ExtractResult().Release());
+		delete SubdivideTask;
+		SubdivideTask = nullptr;
 	}
-
-	SubdividedMesh.Copy(OriginalMesh);
-	int Rounds = Subdivisions;
-
-	// calculate subdivisions (todo: move to elsewhere)
-	for (int ri = 0; ri < Rounds; ri++)
+	if (SubdividedMesh && bNeedsDisplaced)
 	{
-		TArray<int> EdgesToProcess;
-		for (int tid : SubdividedMesh.EdgeIndicesItr())
-		{
-			EdgesToProcess.Add(tid);
-		}
-		int MaxTriangleID = SubdividedMesh.MaxTriangleID();
-
-		TArray<int> TriSplitEdges;
-		TriSplitEdges.Init(-1, SubdividedMesh.MaxTriangleID());
-
-		for (int eid : EdgesToProcess)
-		{
-			FIndex2i EdgeTris = SubdividedMesh.GetEdgeT(eid);
-
-			FDynamicMesh3::FEdgeSplitInfo SplitInfo;
-			EMeshResult result = SubdividedMesh.SplitEdge(eid, SplitInfo);
-			if (result == EMeshResult::Ok)
-			{
-				if (EdgeTris.A < MaxTriangleID && TriSplitEdges[EdgeTris.A] == -1)
-				{
-					TriSplitEdges[EdgeTris.A] = SplitInfo.NewEdges.B;
-				}
-				if (EdgeTris.B != FDynamicMesh3::InvalidID)
-				{
-					if (EdgeTris.B < MaxTriangleID && TriSplitEdges[EdgeTris.B] == -1)
-					{
-						TriSplitEdges[EdgeTris.B] = SplitInfo.NewEdges.C;
-					}
-				}
-			}
-		}
-
-		for (int eid : TriSplitEdges)
-		{
-			if (eid != -1)
-			{
-				FDynamicMesh3::FEdgeFlipInfo FlipInfo;
-				SubdividedMesh.FlipEdge(eid, FlipInfo);
-			}
-		}
+		DisplaceTask = new FAsyncTaskExecuterWithAbort<TModelingOpTask<FDynamicMeshOperator>>(Displacer->MakeNewOperator());
+		DisplaceTask->StartBackgroundTask();
+		bNeedsDisplaced = false;
 	}
-
-	CachedSubdivisionsCount = Subdivisions;
-
-	// update dynamic mesh component
-	DynamicMeshComponent->GetMesh()->Copy(SubdividedMesh);
-	DynamicMeshComponent->NotifyMeshUpdated();
-
-	NormalsBuffer = FMeshNormals(&SubdividedMesh);
-	NormalsBuffer.ComputeVertexNormals();
-
-	// cache initial positions
-	PositionBuffer.SetNum(SubdividedMesh.MaxVertexID());
-	for (int vid : SubdividedMesh.VertexIndicesItr())
+	if (DisplaceTask && DisplaceTask->IsDone())
 	{
-		PositionBuffer[vid] = SubdividedMesh.GetVertex(vid);
+		TUniquePtr<FDynamicMesh3> DisplacedMesh = DisplaceTask->GetTask().ExtractOperator()->ExtractResult();
+		delete DisplaceTask;
+		DisplaceTask = nullptr;
+		DynamicMeshComponent->SetMaterial(0,ToolSetupUtil::GetDefaultMaterial(GetToolManager(), ComponentTarget->GetMaterial(0)));
+		DynamicMeshComponent->GetMesh()->Copy(*DisplacedMesh);
+		DynamicMeshComponent->NotifyMeshUpdated();
+		GetToolManager()->PostInvalidation();
 	}
-	DisplacedBuffer.SetNum(SubdividedMesh.MaxVertexID());
-}
-
-
-bool UDisplaceMeshTool::HasAccept() const
-{
-	return true;
 }
 
-bool UDisplaceMeshTool::CanAccept() const
-{
-	return true;
-}
-
-
-
-
 #undef LOCTEXT_NAMESPACE
diff --git a/Engine/Plugins/Experimental/MeshModelingToolset/Source/MeshModelingTools/Public/DisplaceMeshTool.h b/Engine/Plugins/Experimental/MeshModelingToolset/Source/MeshModelingTools/Public/DisplaceMeshTool.h
--- a/Engine/Plugins/Experimental/MeshModelingToolset/Source/MeshModelingTools/Public/DisplaceMeshTool.h
+++ b/Engine/Plugins/Experimental/MeshModelingToolset/Source/MeshModelingTools/Public/DisplaceMeshTool.h
@@ -9,15 +9,13 @@
 #include "InteractiveToolBuilder.h"
 #include "DynamicMesh3.h"
 #include "MeshNormals.h"
+#include "MeshOpPreviewHelpers.h"
 #include "Spatial/SampledScalarField2.h"
 #include "DisplaceMeshTool.generated.h"
 
-
-// predeclarations
 struct FMeshDescription;
 class USimpleDynamicMeshComponent;
 
-
 UENUM()
 enum class EDisplaceMeshToolDisplaceType : uint8
 {
@@ -29,28 +27,20 @@ enum class EDisplaceMeshToolDisplaceType : uint8
 
 	/** Displace with N iterations */
 	DisplacementMap UMETA(DisplayName = "Displacement Map"),
-
 };
 
-
-
 /**
- *
+ * Builder for Simple Mesh Displacement Tool
  */
 UCLASS()
 class MESHMODELINGTOOLS_API UDisplaceMeshToolBuilder : public UInteractiveToolBuilder
 {
 	GENERATED_BODY()
-
 public:
-	virtual bool CanBuildTool(const FToolBuilderState& SceneState) const override;
-	virtual UInteractiveTool* BuildTool(const FToolBuilderState& SceneState) const override;
+	bool CanBuildTool(const FToolBuilderState& SceneState) const final;
+	UInteractiveTool* BuildTool(const FToolBuilderState& SceneState) const final;
 };
 
-
-
-
-
 /**
  * Simple Mesh Displacement Tool
  */
@@ -62,68 +52,55 @@ class MESHMODELINGTOOLS_API UDisplaceMeshTool : public USingleSelectionTool
 public:
 	UDisplaceMeshTool();
 
-	virtual void Setup() override;
-	virtual void Shutdown(EToolShutdownType ShutdownType) override;
+	void Setup() final;
+	void Shutdown(EToolShutdownType ShutdownType) final;
+	void Tick(float DeltaTime) final;
 
-	virtual void Render(IToolsContextRenderAPI* RenderAPI) override;
-
-	virtual bool HasCancel() const override { return true; }
-	virtual bool HasAccept() const override;
-	virtual bool CanAccept() const override;
+	bool HasCancel() const final { return true; }
+	bool HasAccept() const final { return true; }
+	bool CanAccept() const final { return true; }
 #if WITH_EDITOR
-	virtual void PostEditChangeProperty(FPropertyChangedEvent& PropertyChangedEvent) override;
+	void PostEditChangeProperty(FPropertyChangedEvent& PropertyChangedEvent) final;
 #endif
-protected:
-	// need to update bResultValid if these are modified, so we don't publicly expose them. 
-	// @todo setters/getters for these
 
 	/** primary brush mode */
 	UPROPERTY(EditAnywhere, Category = Options)
 	EDisplaceMeshToolDisplaceType DisplacementType;
 
 	/** Displacement intensity */
-	UPROPERTY(EditAnywhere, Category = Options, meta = (UIMin = "-100.0", UIMax = "100.0", ClampMin = "-10000.0", ClampMax = "100000.0"))
+	UPROPERTY(EditAnywhere, Category = Options,
+		meta = (UIMin = "-100.0", UIMax = "100.0", ClampMin = "-10000.0", ClampMax = "100000.0"))
 	float DisplaceIntensity;
 
 	/** Seed for randomization */
-	UPROPERTY(EditAnywhere, Category = Options, meta = (EditCondition = "DisplacementType == EDisplaceMeshToolDisplaceType::RandomNoise"))
+	UPROPERTY(EditAnywhere, Category = Options,
+		meta = (EditCondition = "DisplacementType == EDisplaceMeshToolDisplaceType::RandomNoise"))
 	int RandomSeed;
 
 	/** Subdivision iterations for mesh */
-	UPROPERTY(EditAnywhere, Category = Options, meta = (UIMin = "0", UIMax = "10", ClampMin = "0", ClampMax = "100"))
+	UPROPERTY(EditAnywhere, Category = Options,
+		meta = (UIMin = "0", UIMax = "10", ClampMin = "0", ClampMax = "100"))
 	int Subdivisions;
 
 	/** Displacement map */
-	UPROPERTY(EditAnywhere, Category = Options)
+	UPROPERTY(EditAnywhere, Category = Options,
+		meta = (EditCondition = "DisplacementType == EDisplaceMeshToolDisplaceType::DisplacementMap"))
 	UTexture2D* DisplacementMap;
 
+private:
+	void StartComputation();
+	void AdvanceComputation();
+	void ValidateSubdivisions();
 
-protected:
-	USimpleDynamicMeshComponent* DynamicMeshComponent;
-	
-	FDynamicMesh3 OriginalMesh;		// don't actually need whole mesh, just positions...
-
-	FDynamicMesh3 SubdividedMesh;
-	int CachedSubdivisionsCount = -1;
-	void UpdateSubdividedMesh();
-
-	FSampledScalarField2f DisplaceField;
-	UTexture2D* CachedMapSource = nullptr;
-	void UpdateMap(bool bForceUpdate = false);
-
-	TArray<FVector3d> PositionBuffer;
-	FMeshNormals NormalsBuffer;
-	TArray<FVector3d> DisplacedBuffer;
-	bool bResultValid;
-	void UpdateResult();
-
-	void ComputeDisplacement_Constant();
-	void ComputeDisplacement_RandomNoise();
-
-	void ComputeDisplacement_Map();
-
+	FAsyncTaskExecuterWithAbort<TModelingOpTask<FDynamicMeshOperator>>* SubdivideTask = nullptr;
+	bool bNeedsSubdivided = true;
+	FAsyncTaskExecuterWithAbort<TModelingOpTask<FDynamicMeshOperator>>* DisplaceTask = nullptr;
+	bool bNeedsDisplaced = true;
 
-private:
-	//TUniquePtr<
+	FDynamicMesh3 OriginalMesh;
+	TSharedPtr<FDynamicMesh3> SubdividedMesh = nullptr;
+	USimpleDynamicMeshComponent* DynamicMeshComponent = nullptr;
 
+	TUniquePtr<IDynamicMeshOperatorFactory> Subdivider = nullptr;
+	TUniquePtr<IDynamicMeshOperatorFactory> Displacer = nullptr;
 };
