ModelingTools: if nearest query fails, keep existing position instead of using garbage position/normal. Increased query range for Offset brush, as it can get arbitrarily far from base mesh (behavior really degenerates though)
#rnx

diff --git a/Engine/Plugins/Experimental/MeshModelingToolset/Source/MeshModelingTools/Private/DynamicMeshSculptTool.cpp b/Engine/Plugins/Experimental/MeshModelingToolset/Source/MeshModelingTools/Private/DynamicMeshSculptTool.cpp
--- a/Engine/Plugins/Experimental/MeshModelingToolset/Source/MeshModelingTools/Private/DynamicMeshSculptTool.cpp
+++ b/Engine/Plugins/Experimental/MeshModelingToolset/Source/MeshModelingTools/Private/DynamicMeshSculptTool.cpp
@@ -587,7 +587,6 @@ void UDynamicMeshSculptTool::ApplyMoveBrush(const FRay& WorldRay)
 
 
 
-
 void UDynamicMeshSculptTool::ApplyOffsetBrush(const FRay& WorldRay)
 {
 	bool bHit = UpdateBrushPositionOnTargetMesh(WorldRay);
@@ -611,14 +610,17 @@ void UDynamicMeshSculptTool::ApplyOffsetBrush(const FRay& WorldRay)
 		FVector3d OrigPos = Mesh->GetVertex(VertIdx);
 
 		FVector3d BasePos, BaseNormal;
-		GetTargetMeshNearest(OrigPos, (double)(2 * CurrentBrushRadius), BasePos, BaseNormal);
-
-		FVector3d MoveVec = UseSpeed * BaseNormal;
-
-		double Falloff = CalculateBrushFalloff(OrigPos.Distance(NewBrushPosLocal));
-
-		FVector3d NewPos = OrigPos + Falloff * MoveVec;
-		ROIPositionBuffer[k] = NewPos;
+		if (GetTargetMeshNearest(OrigPos, (double)(4 * CurrentBrushRadius), BasePos, BaseNormal) == false)
+		{
+			ROIPositionBuffer[k] = OrigPos;
+		}
+		else
+		{
+			FVector3d MoveVec = UseSpeed * BaseNormal;
+			double Falloff = CalculateBrushFalloff(OrigPos.Distance(NewBrushPosLocal));
+			FVector3d NewPos = OrigPos + Falloff * MoveVec;
+			ROIPositionBuffer[k] = NewPos;
+		}
 	});
 
 	for (int k = 0; k < NumV; ++k)
@@ -638,7 +640,6 @@ void UDynamicMeshSculptTool::ApplyOffsetBrush(const FRay& WorldRay)
 
 
 
-
 void UDynamicMeshSculptTool::ApplySculptMaxBrush(const FRay& WorldRay)
 {
 	bool bHit = UpdateBrushPositionOnTargetMesh(WorldRay);
@@ -664,22 +665,25 @@ void UDynamicMeshSculptTool::ApplySculptMaxBrush(const FRay& WorldRay)
 		FVector3d OrigPos = Mesh->GetVertex(VertIdx);
 
 		FVector3d BasePos, BaseNormal;
-		GetTargetMeshNearest(OrigPos, (double)(2 * CurrentBrushRadius), BasePos, BaseNormal);
-
-		FVector3d MoveVec = UseSpeed * BaseNormal;
-
-		double Falloff = CalculateBrushFalloff(OrigPos.Distance(NewBrushPosLocal));
-
-		FVector3d NewPos = OrigPos + Falloff * MoveVec;
-
-		FVector3d DeltaPos = NewPos - BasePos;
-		if (DeltaPos.SquaredLength() > MaxOffset*MaxOffset)
+		if (GetTargetMeshNearest(OrigPos, (double)(2 * CurrentBrushRadius), BasePos, BaseNormal) == false)
 		{
-			DeltaPos.Normalize();
-			NewPos = BasePos + MaxOffset * DeltaPos;
+			ROIPositionBuffer[k] = OrigPos;
 		}
+		else
+		{
+			FVector3d MoveVec = UseSpeed * BaseNormal;
+			double Falloff = CalculateBrushFalloff(OrigPos.Distance(NewBrushPosLocal));
+			FVector3d NewPos = OrigPos + Falloff * MoveVec;
 
-		ROIPositionBuffer[k] = NewPos;
+			FVector3d DeltaPos = NewPos - BasePos;
+			if (DeltaPos.SquaredLength() > MaxOffset*MaxOffset)
+			{
+				DeltaPos.Normalize();
+				NewPos = BasePos + MaxOffset * DeltaPos;
+			}
+
+			ROIPositionBuffer[k] = NewPos;
+		}
 	});
 
 	for (int k = 0; k < NumV; ++k)
