ToolsFramework: hide transient actors for gizmos and preview meshes in scene outliner
#rnx

diff --git a/Engine/Plugins/Experimental/MeshModelingToolset/Source/ModelingComponents/Private/PreviewMesh.cpp b/Engine/Plugins/Experimental/MeshModelingToolset/Source/ModelingComponents/Private/PreviewMesh.cpp
--- a/Engine/Plugins/Experimental/MeshModelingToolset/Source/ModelingComponents/Private/PreviewMesh.cpp
+++ b/Engine/Plugins/Experimental/MeshModelingToolset/Source/ModelingComponents/Private/PreviewMesh.cpp
@@ -4,6 +4,16 @@
 
 #include "Containers/StaticArray.h"
 
+
+APreviewMeshActor::APreviewMeshActor()
+{
+#if WITH_EDITORONLY_DATA
+	// hide this actor in the scene outliner
+	bListedInSceneOutliner = false;
+#endif
+}
+
+
 UPreviewMesh::UPreviewMesh()
 {
 	bBuildSpatialDataStructure = false;
diff --git a/Engine/Plugins/Experimental/MeshModelingToolset/Source/ModelingComponents/Public/PreviewMesh.h b/Engine/Plugins/Experimental/MeshModelingToolset/Source/ModelingComponents/Public/PreviewMesh.h
--- a/Engine/Plugins/Experimental/MeshModelingToolset/Source/ModelingComponents/Public/PreviewMesh.h
+++ b/Engine/Plugins/Experimental/MeshModelingToolset/Source/ModelingComponents/Public/PreviewMesh.h
@@ -25,7 +25,7 @@ class MODELINGCOMPONENTS_API APreviewMeshActor : public AInternalToolFrameworkAc
 {
 	GENERATED_BODY()
 private:
-
+	APreviewMeshActor();
 public:
 
 };
diff --git a/Engine/Source/Runtime/Experimental/InteractiveToolsFramework/Private/BaseGizmos/GizmoActor.cpp b/Engine/Source/Runtime/Experimental/InteractiveToolsFramework/Private/BaseGizmos/GizmoActor.cpp
--- a/Engine/Source/Runtime/Experimental/InteractiveToolsFramework/Private/BaseGizmos/GizmoActor.cpp
+++ b/Engine/Source/Runtime/Experimental/InteractiveToolsFramework/Private/BaseGizmos/GizmoActor.cpp
@@ -17,6 +17,11 @@ AGizmoActor::AGizmoActor()
 {
 	// generally gizmo actor creation/destruction should not be transacted
 	ClearFlags(RF_Transactional);
+
+#if WITH_EDITORONLY_DATA
+	// hide this actor in the scene outliner
+	bListedInSceneOutliner = false;
+#endif
 }
 
 
