update check condition because IsInRenderingThread() is in fact sometimes true on game thread during lighting build.
#rnx

diff --git a/Engine/Plugins/Experimental/MeshModelingToolset/Source/ModelingComponents/Private/OctreeDynamicMeshSceneProxy.h b/Engine/Plugins/Experimental/MeshModelingToolset/Source/ModelingComponents/Private/OctreeDynamicMeshSceneProxy.h
--- a/Engine/Plugins/Experimental/MeshModelingToolset/Source/ModelingComponents/Private/OctreeDynamicMeshSceneProxy.h
+++ b/Engine/Plugins/Experimental/MeshModelingToolset/Source/ModelingComponents/Private/OctreeDynamicMeshSceneProxy.h
@@ -149,7 +149,9 @@ public:
 		: FPrimitiveSceneProxy(Component)
 		, MaterialRelevance(Component->GetMaterialRelevance(GetScene().GetFeatureLevel()))
 	{
-		check(!IsInRenderingThread());
+		// This is an assumption we are currently making. We do not necessarily require this
+		// but if this check is hit then possibly an assumption is wrong
+		check(IsInGameThread());
 
 		ParentComponent = Component;
 
@@ -164,6 +166,7 @@ public:
 
 	virtual ~FOctreeDynamicMeshSceneProxy()
 	{
+		// we are assuming in code below that this is always called from the rendering thread
 		check(IsInRenderingThread());
 
 		for (auto MapPair : RenderBufferSets)
