Catch cases in Plane Cut tool where resulting cut meshes have zero triangles, and don't emit new meshes in those cases
#rnx

diff --git a/Engine/Plugins/Experimental/MeshModelingToolset/Source/MeshModelingTools/Private/PlaneCutTool.cpp b/Engine/Plugins/Experimental/MeshModelingToolset/Source/MeshModelingTools/Private/PlaneCutTool.cpp
--- a/Engine/Plugins/Experimental/MeshModelingToolset/Source/MeshModelingTools/Private/PlaneCutTool.cpp
+++ b/Engine/Plugins/Experimental/MeshModelingToolset/Source/MeshModelingTools/Private/PlaneCutTool.cpp
@@ -323,13 +323,18 @@ bool UPlaneCutTool::CanAccept() const
 
 void UPlaneCutTool::GenerateAsset(const TArray<TUniquePtr<FDynamicMeshOpResult>>& Results)
 {
+	if (Results.Num() == 0 
+		|| Results[0]->Mesh.IsValid() == false
+		|| Results[0]->Mesh->TriangleCount() == 0 )
+	{
+		return;
+	}
+
 	GetToolManager()->BeginUndoTransaction(LOCTEXT("PlaneCutToolTransactionName", "Plane Cut Tool"));
 	
 
 	// currently in-place replaces the first half, and adds a new actor for the second half (if it was generated)
 	// TODO: options to support other choices re what should be a new actor
-	check(Results.Num() > 0);
-	check(Results[0]->Mesh.Get() != nullptr);
 	ComponentTarget->CommitMesh([&Results](FMeshDescription* MeshDescription)
 	{
 		FDynamicMeshToMeshDescription Converter;
@@ -339,14 +344,12 @@ void UPlaneCutTool::GenerateAsset(const TArray<TUniquePtr<FDynamicMeshOpResult>>
 
 	// The method for creating a new mesh (AssetGenerationUtil::GenerateStaticMeshActor) is editor-only; just creating the other half if not in editor
 #if WITH_EDITOR
-	if (Results.Num() == 2)
+	if (Results.Num() == 2 && Results[1]->Mesh.IsValid() && Results[1]->Mesh->TriangleCount() > 0)
 	{
 		FSelectedOjectsChangeList NewSelection;
 		NewSelection.ModificationType = ESelectedObjectsModificationType::Replace;
 		NewSelection.Actors.Add(ComponentTarget->GetOwnerActor());
 
-		check(Results[1]->Mesh.Get() != nullptr);
-
 		TArray<UMaterialInterface*> Materials;
 		for (int MaterialIdx = 0, NumMaterials = ComponentTarget->GetNumMaterials(); MaterialIdx < NumMaterials; MaterialIdx++)
 		{
