Always create at least one material slot in generated StaticMeshes - copy CL 9965758 from Dev-Editor
#rnx

diff --git a/Engine/Plugins/Experimental/MeshModelingToolset/Source/ModelingComponents/Private/StaticMeshComponentBuilder.cpp b/Engine/Plugins/Experimental/MeshModelingToolset/Source/ModelingComponents/Private/StaticMeshComponentBuilder.cpp
--- a/Engine/Plugins/Experimental/MeshModelingToolset/Source/ModelingComponents/Private/StaticMeshComponentBuilder.cpp
+++ b/Engine/Plugins/Experimental/MeshModelingToolset/Source/ModelingComponents/Private/StaticMeshComponentBuilder.cpp
@@ -39,8 +39,9 @@ void FStaticMeshComponentBuilder::Initialize(UPackage* AssetPackage, FName MeshN
 		NewStaticMesh->BodySetup->CollisionTraceFlag = ECollisionTraceFlag::CTF_UseComplexAsSimple;
 	}
 
-	// add a material slot
-	for (int MatIdx = 0; MatIdx < NumMaterialSlots; MatIdx++)
+	// add a material slot. Must always have one material slot.
+	int AddMaterialCount = FMath::Max(1, NumMaterialSlots);
+	for (int MatIdx = 0; MatIdx < AddMaterialCount; MatIdx++)
 	{
 		NewStaticMesh->StaticMaterials.Add(FStaticMaterial());
 	}
