Modify Sculpt hotkeys so they don't conflict with RightMouse+WASD navigation
#rnx

diff --git a/Engine/Plugins/Experimental/MeshModelingToolset/Source/MeshModelingTools/Private/DynamicMeshSculptTool.cpp b/Engine/Plugins/Experimental/MeshModelingToolset/Source/MeshModelingTools/Private/DynamicMeshSculptTool.cpp
--- a/Engine/Plugins/Experimental/MeshModelingToolset/Source/MeshModelingTools/Private/DynamicMeshSculptTool.cpp
+++ b/Engine/Plugins/Experimental/MeshModelingToolset/Source/MeshModelingTools/Private/DynamicMeshSculptTool.cpp
@@ -196,7 +196,7 @@ void UDynamicMeshSculptTool::Setup()
 
 
 	GetToolManager()->DisplayMessage(
-		LOCTEXT("OnStartSculptTool", "Hold Shift to Smooth, Ctrl to Invert (where applicable). Q/A keys cycle through Brush Types, Shift+Q/A for Brush History. S/D change Size (shift to small-step), W/E change Speed."),
+		LOCTEXT("OnStartSculptTool", "Hold Shift to Smooth, Ctrl to Invert (where applicable). Shift+Q/A keys cycle through Brush Types. Shift+S/D change Size (Ctrl+Shift to small-step), Shift+W/E change Speed."),
 		EToolMessageLevel::UserNotification);
 
 	if (bEnableRemeshing)
@@ -909,7 +909,7 @@ void UDynamicMeshSculptTool::ApplyInflateBrush(const FRay& WorldRay)
 	FVector3d NewBrushPosLocal = CurTargetTransform.InverseTransformPosition(LastBrushPosWorld);
 
 	double Direction = (bInvert) ? -1.0 : 1.0;
-	double UseSpeed = Direction * FMathd::Sqrt(CurrentBrushRadius) * SculptProperties->BrushSpeed * 0.25;
+	double UseSpeed = Direction * CurrentBrushRadius * SculptProperties->BrushSpeed * 0.05;
 
 	FDynamicMesh3* Mesh = DynamicMeshComponent->GetMesh();
 	int NumV = VertexROI.Num();
@@ -1642,59 +1642,57 @@ void UDynamicMeshSculptTool::RegisterActions(FInteractiveToolActionSet& ActionSe
 		TEXT("NextBrushMode"),
 		LOCTEXT("SculptNextBrushMode", "Next Brush Type"),
 		LOCTEXT("SculptNextBrushModeTooltip", "Cycle to next Brush Type"),
-		EModifierKey::None, EKeys::A,
+		EModifierKey::Shift, EKeys::A,
 		[this]() { NextBrushModeAction(); });
 
 	ActionSet.RegisterAction(this, (int32)EStandardToolActions::BaseClientDefinedActionID+2,
 		TEXT("PreviousBrushMode"),
 		LOCTEXT("SculptPreviousBrushMode", "Previous Brush Type"),
 		LOCTEXT("SculptPreviousBrushModeTooltip", "Cycle to previous Brush Type"),
-		EModifierKey::None, EKeys::Q,
-		[this]() { PreviousBrushModeAction(); });
-
-
-	ActionSet.RegisterAction(this, (int32)EStandardToolActions::BaseClientDefinedActionID + 10,
-		TEXT("NextBrushHistoryState"),
-		LOCTEXT("SculptNextBrushHistoryState", "Next Brush History State"),
-		LOCTEXT("SculptSculptNextBrushHistoryStateTooltip", "Cycle to next Brush History state"),
 		EModifierKey::Shift, EKeys::Q,
-		[this]() { NextHistoryBrushModeAction(); });
+		[this]() { PreviousBrushModeAction(); });
 
-	ActionSet.RegisterAction(this, (int32)EStandardToolActions::BaseClientDefinedActionID + 11,
-		TEXT("PreviousBrushHistoryState"),
-		LOCTEXT("SculptPreviousBrushHistoryState", "Previous Brush History State"),
-		LOCTEXT("SculptPreviousBrushHistoryStateTooltip", "Cycle to previous Brush History state"),
-		EModifierKey::Shift, EKeys::A,
-		[this]() { PreviousHistoryBrushModeAction(); });
 
+	//ActionSet.RegisterAction(this, (int32)EStandardToolActions::BaseClientDefinedActionID + 10,
+	//	TEXT("NextBrushHistoryState"),
+	//	LOCTEXT("SculptNextBrushHistoryState", "Next Brush History State"),
+	//	LOCTEXT("SculptSculptNextBrushHistoryStateTooltip", "Cycle to next Brush History state"),
+	//	EModifierKey::Shift, EKeys::Q,
+	//	[this]() { NextHistoryBrushModeAction(); });
 
+	//ActionSet.RegisterAction(this, (int32)EStandardToolActions::BaseClientDefinedActionID + 11,
+	//	TEXT("PreviousBrushHistoryState"),
+	//	LOCTEXT("SculptPreviousBrushHistoryState", "Previous Brush History State"),
+	//	LOCTEXT("SculptPreviousBrushHistoryStateTooltip", "Cycle to previous Brush History state"),
+	//	EModifierKey::Shift, EKeys::A,
+	//	[this]() { PreviousHistoryBrushModeAction(); });
 
 	ActionSet.RegisterAction(this, (int32)EStandardToolActions::BaseClientDefinedActionID + 50,
 		TEXT("SculptIncreaseSize"),
 		LOCTEXT("SculptIncreaseSize", "Increase Size"),
 		LOCTEXT("SculptIncreaseSizeTooltip", "Increase Brush Size"),
-		EModifierKey::None, EKeys::D,
+		EModifierKey::Shift, EKeys::D,
 		[this]() { IncreaseBrushRadiusAction(); });
 
 	ActionSet.RegisterAction(this, (int32)EStandardToolActions::BaseClientDefinedActionID + 51,
 		TEXT("SculptDecreaseSize"),
 		LOCTEXT("SculptDecreaseSize", "Decrease Size"),
 		LOCTEXT("SculptDecreaseSizeTooltip", "Decrease Brush Size"),
-		EModifierKey::None, EKeys::S,
+		EModifierKey::Shift, EKeys::S,
 		[this]() { DecreaseBrushRadiusAction(); });
 
 	ActionSet.RegisterAction(this, (int32)EStandardToolActions::BaseClientDefinedActionID + 52,
 		TEXT("SculptIncreaseSizeSmallStep"),
 		LOCTEXT("SculptIncreaseSize", "Increase Size"),
 		LOCTEXT("SculptIncreaseSizeTooltip", "Increase Brush Size"),
-		EModifierKey::Shift, EKeys::D,
+		EModifierKey::Shift | EModifierKey::Control, EKeys::D,
 		[this]() { IncreaseBrushRadiusSmallStepAction(); });
 
 	ActionSet.RegisterAction(this, (int32)EStandardToolActions::BaseClientDefinedActionID + 53,
 		TEXT("SculptDecreaseSizeSmallStemp"),
 		LOCTEXT("SculptDecreaseSize", "Decrease Size"),
 		LOCTEXT("SculptDecreaseSizeTooltip", "Decrease Brush Size"),
-		EModifierKey::Shift, EKeys::S,
+		EModifierKey::Shift | EModifierKey::Control, EKeys::S,
 		[this]() { DecreaseBrushRadiusSmallStepAction(); });
 
 
@@ -1704,14 +1702,14 @@ void UDynamicMeshSculptTool::RegisterActions(FInteractiveToolActionSet& ActionSe
 		TEXT("SculptIncreaseSpeed"),
 		LOCTEXT("SculptIncreaseSpeed", "Increase Speed"),
 		LOCTEXT("SculptIncreaseSpeedTooltip", "Increase Brush Speed"),
-		EModifierKey::None, EKeys::E,
+		EModifierKey::Shift, EKeys::E,
 		[this]() { IncreaseBrushSpeedAction(); });
 
 	ActionSet.RegisterAction(this, (int32)EStandardToolActions::BaseClientDefinedActionID + 61,
 		TEXT("SculptDecreaseSpeed"),
 		LOCTEXT("SculptDecreaseSpeed", "Decrease Speed"),
 		LOCTEXT("SculptDecreaseSpeedTooltip", "Decrease Brush Speed"),
-		EModifierKey::None, EKeys::W,
+		EModifierKey::Shift, EKeys::W,
 		[this]() { DecreaseBrushSpeedAction(); });
 
 
