Mesh Modeling Tools moved from NotForLicensees to Experimental

diff --git a/Engine/Plugins/Experimental/MeshModelingToolset/MeshModelingToolset.uplugin b/Engine/Plugins/Experimental/MeshModelingToolset/MeshModelingToolset.uplugin
new file mode 100644
--- /dev/null
+++ b/Engine/Plugins/Experimental/MeshModelingToolset/MeshModelingToolset.uplugin
@@ -0,0 +1,58 @@
+{
+	"FileVersion": 3,
+	"Version": 1,
+	"VersionName": "0.1",
+	"FriendlyName": "Mesh Modeling Toolset",
+	"Description": "Adds Mesh Modeling Modules",
+	"Category": "Other",
+	"CreatedBy": "Epic Games, Inc.",
+	"CreatedByURL": "http://epicgames.com",
+	"DocsURL": "",
+	"MarketplaceURL": "",
+	"SupportURL": "",
+	"CanContainContent": true,
+	"IsBetaVersion": true,
+	"Installed": false,
+	"Modules": [
+		{
+			"Name": "ModelingComponents",
+			"Type": "Runtime",
+			"LoadingPhase": "Default",
+			"WhitelistPlatforms": [ "Win64" ]
+		},
+		{
+			"Name": "ModelingOperators",
+			"Type": "Runtime",
+			"LoadingPhase": "Default",
+			"WhitelistPlatforms": [ "Win64" ]
+		},
+		{
+			"Name": "ModelingOperatorsEditorOnly",
+			"Type": "Editor",
+			"LoadingPhase": "Default",
+			"WhitelistPlatforms": [ "Win64" ]
+		},
+		{
+			"Name": "MeshModelingTools",
+			"Type": "Runtime",
+			"LoadingPhase": "Default",
+			"WhitelistPlatforms": [ "Win64" ]
+		},
+		{
+			"Name": "MeshModelingToolsEditorOnly",
+			"Type": "Editor",
+			"LoadingPhase": "Default",
+			"WhitelistPlatforms": [ "Win64" ]
+		},
+	],
+	"Plugins": [
+		{
+			"Name": "GeometryProcessing",
+			"Enabled": true
+		},
+		{
+			"Name": "ProxyLODPlugin",
+			"Enabled" : true
+		}
+	]
+}
diff --git a/Engine/Plugins/Experimental/MeshModelingToolset/Source/MeshModelingTools/Private/BaseBrushTool.cpp b/Engine/Plugins/Experimental/MeshModelingToolset/Source/MeshModelingTools/Private/BaseBrushTool.cpp
new file mode 100644
--- /dev/null
+++ b/Engine/Plugins/Experimental/MeshModelingToolset/Source/MeshModelingTools/Private/BaseBrushTool.cpp
@@ -0,0 +1,183 @@
+// Copyright 1998-2019 Epic Games, Inc. All Rights Reserved.
+
+#include "BaseBrushTool.h"
+#include "InteractiveToolManager.h"
+#include "ToolBuilderUtil.h"
+#include "BrushToolIndicator.h"
+
+#define LOCTEXT_NAMESPACE "UBaseBrushTool"
+
+
+
+UBrushBaseProperties::UBrushBaseProperties()
+{
+	BrushSize = 0.25f;
+	bSpecifyRadius = false;
+	BrushRadius = 0.0f;
+}
+
+
+
+UBaseBrushTool::UBaseBrushTool()
+{
+}
+
+
+
+void UBaseBrushTool::Setup()
+{
+	UMeshSurfacePointTool::Setup();
+
+	double MaxDimension = EstimateMaximumTargetDimension();
+	BrushRelativeSizeRange = FInterval1d(MaxDimension*0.01, MaxDimension);
+	BrushProperties = NewObject<UBrushBaseProperties>(this, TEXT("Brush"));
+	RecalculateBrushRadius();
+
+
+	// create brush stamp indicator
+	Indicators = NewObject<UToolIndicatorSet>(this, "Indicators");
+	Indicators->Connect(this);
+	Indicators->AddIndicator(MakeBrushIndicator());
+
+
+	// initialize our properties
+	AddToolPropertySource(BrushProperties);
+}
+
+
+void UBaseBrushTool::Shutdown(EToolShutdownType ShutdownType)
+{
+	Indicators->Disconnect();
+}
+
+
+void UBaseBrushTool::OnPropertyModified(UObject* PropertySet, UProperty* Property)
+{
+	if (PropertySet == BrushProperties)
+	{
+		RecalculateBrushRadius();
+	}
+}
+
+
+IToolIndicator* UBaseBrushTool::MakeBrushIndicator()
+{
+	UBrushStampSizeIndicator* StampIndicator = NewObject<UBrushStampSizeIndicator>(this, "Brush Circle");
+	StampIndicator->bDrawSecondaryLines = true;
+	StampIndicator->DepthLayer = 1;
+	StampIndicator->BrushRadius = MakeAttributeLambda([this] { return this->LastBrushStamp.Radius; });
+	StampIndicator->BrushPosition = MakeAttributeLambda([this] { return this->LastBrushStamp.WorldPosition; });
+	StampIndicator->BrushNormal = MakeAttributeLambda([this] { return this->LastBrushStamp.WorldNormal; });
+	return StampIndicator;
+}
+
+
+
+
+void UBaseBrushTool::IncreaseBrushSizeAction()
+{
+	BrushProperties->BrushSize = FMath::Clamp(BrushProperties->BrushSize + 0.025f, 0.0f, 1.0f);
+	RecalculateBrushRadius();
+}
+
+void UBaseBrushTool::DecreaseBrushSizeAction()
+{
+	BrushProperties->BrushSize = FMath::Clamp(BrushProperties->BrushSize - 0.025f, 0.0f, 1.0f);
+	RecalculateBrushRadius();
+}
+
+
+void UBaseBrushTool::RegisterActions(FInteractiveToolActionSet& ActionSet)
+{
+	ActionSet.RegisterAction(this, (int32)EStandardToolActions::IncreaseBrushSize,
+		TEXT("BrushIncreaseSize"), 
+		LOCTEXT("BrushIncreaseSize", "Increase Brush Size"),
+		LOCTEXT("BrushIncreaseSizeTooltip", "Increase size of brush"),
+		EModifierKey::None, EKeys::RightBracket,
+		[this]() { IncreaseBrushSizeAction(); });
+
+	ActionSet.RegisterAction(this, (int32)EStandardToolActions::DecreaseBrushSize,
+		TEXT("BrushDecreaseSize"), 
+		LOCTEXT("BrushDecreaseSize", "Decrease Brush Size"),
+		LOCTEXT("BrushDecreaseSizeTooltip", "Decrease size of brush"),
+		EModifierKey::None, EKeys::LeftBracket,
+		[this]() { DecreaseBrushSizeAction(); });
+}
+
+
+
+
+void UBaseBrushTool::RecalculateBrushRadius()
+{
+	CurrentBrushRadius = 0.5 * BrushRelativeSizeRange.Interpolate(BrushProperties->BrushSize);
+	if (BrushProperties->bSpecifyRadius)
+	{
+		CurrentBrushRadius = BrushProperties->BrushRadius;
+	}
+	else
+	{
+		BrushProperties->BrushRadius = CurrentBrushRadius;
+	}
+}
+
+
+void UBaseBrushTool::OnBeginDrag(const FRay& Ray)
+{
+	FHitResult OutHit;
+	if (HitTest(Ray, OutHit))
+	{
+		LastBrushStamp.Radius = BrushProperties->BrushRadius;
+		LastBrushStamp.WorldPosition = OutHit.ImpactPoint;
+		LastBrushStamp.WorldNormal = OutHit.Normal;
+		LastBrushStamp.HitResult = OutHit;
+	}
+	bInBrushStroke = true;
+}
+
+void UBaseBrushTool::OnUpdateDrag(const FRay& Ray)
+{
+	FHitResult OutHit;
+	if (HitTest(Ray, OutHit))
+	{
+		LastBrushStamp.Radius = BrushProperties->BrushRadius;
+		LastBrushStamp.WorldPosition = OutHit.ImpactPoint;
+		LastBrushStamp.WorldNormal = OutHit.Normal;
+		LastBrushStamp.HitResult = OutHit;
+	}
+}
+
+void UBaseBrushTool::OnEndDrag(const FRay& Ray)
+{
+	bInBrushStroke = false;
+}
+
+
+void UBaseBrushTool::OnUpdateHover(const FInputDeviceRay& DevicePos)
+{
+	FHitResult OutHit;
+	if (HitTest(DevicePos.WorldRay, OutHit))
+	{
+		LastBrushStamp.Radius = BrushProperties->BrushRadius;
+		LastBrushStamp.WorldPosition = OutHit.ImpactPoint;
+		LastBrushStamp.WorldNormal = OutHit.Normal;
+		LastBrushStamp.HitResult = OutHit;
+	}
+}
+
+
+
+void UBaseBrushTool::Render(IToolsContextRenderAPI* RenderAPI)
+{
+	UMeshSurfacePointTool::Render(RenderAPI);
+	Indicators->Render(RenderAPI);
+}
+
+void UBaseBrushTool::Tick(float DeltaTime)
+{
+	UMeshSurfacePointTool::Tick(DeltaTime);
+	Indicators->Tick(DeltaTime);
+}
+
+
+
+#undef LOCTEXT_NAMESPACE
diff --git a/Engine/Plugins/Experimental/MeshModelingToolset/Source/MeshModelingTools/Private/BrushToolIndicator.cpp b/Engine/Plugins/Experimental/MeshModelingToolset/Source/MeshModelingTools/Private/BrushToolIndicator.cpp
new file mode 100644
--- /dev/null
+++ b/Engine/Plugins/Experimental/MeshModelingToolset/Source/MeshModelingTools/Private/BrushToolIndicator.cpp
@@ -0,0 +1,100 @@
+// Copyright 1998-2019 Epic Games, Inc. All Rights Reserved.
+
+#include "BrushToolIndicator.h"
+#include "SceneManagement.h" // for FPrimitiveDrawInterface
+#include "VectorUtil.h"
+
+
+UBrushStampSizeIndicator::UBrushStampSizeIndicator()
+{
+	ParentTool = nullptr;
+	Radius = 1.0;
+	Center = FVector::ZeroVector;
+	Normal = FVector(0, 1, 0);
+
+	SampleStepCount = 32;
+	LineColor = FColor(255, 0, 0);
+	LineThickness = 2.0f;
+	IsPixelThickness = true;
+	DepthLayer = 0;
+
+	SecondaryLineThickness = 0.5f;
+	SecondaryLineColor = FColor(128, 128, 128, 128);
+	bDrawSecondaryLines = true;
+}
+
+UBrushStampSizeIndicator::~UBrushStampSizeIndicator()
+{
+	check(ParentTool == nullptr);
+}
+
+void UBrushStampSizeIndicator::Connect(UInteractiveTool* Tool)
+{
+	ParentTool = Tool;
+}
+
+void UBrushStampSizeIndicator::Disconnect()
+{
+	ParentTool = nullptr;
+}
+
+void UBrushStampSizeIndicator::Render(IToolsContextRenderAPI* RenderAPI)
+{
+	if (BrushRadius.IsBound())
+	{
+		Radius = BrushRadius.Get();
+	}
+	if (BrushPosition.IsBound())
+	{
+		Center = BrushPosition.Get();
+	}
+	if (BrushNormal.IsBound())
+	{
+		Normal = BrushNormal.Get();
+	}
+
+	FPrimitiveDrawInterface* PDI = RenderAPI->GetPrimitiveDrawInterface();
+
+	FVector3f Perp1, Perp2;
+	VectorUtil::MakePerpVectors( (FVector3f)Normal, Perp1, Perp2);
+
+	DrawCircle(PDI, Perp1, Perp2, Radius, false);
+	if (bDrawSecondaryLines)
+	{
+		DrawCircle(PDI, Perp1, Perp2, Radius/2, true);
+		DrawLine(PDI, Center, Center + Radius*Normal, true);
+		//DrawCircle(PDI, Normal, Perp1, true);
+		//DrawCircle(PDI, Normal, Perp2, true);
+	}
+}
+
+void UBrushStampSizeIndicator::Tick(float DeltaTime)
+{
+}
+
+void UBrushStampSizeIndicator::DrawCircle(FPrimitiveDrawInterface* PDI, const FVector3f& AxisX, const FVector3f& AxisY, float UseRadius, bool bIsSecondary)
+{
+	FVector3f Prev = Center + UseRadius * AxisX;
+	for (int k = 1; k <= SampleStepCount; ++k)
+	{
+		float angle = (float)k / (float)SampleStepCount * FMathd::TwoPi;
+		FVector2f Next2D = UseRadius * FVector2f((float)cos(angle), (float)sin(angle));
+		FVector3f Next3D = Center + Next2D.X*AxisX + Next2D.Y*AxisY;
+		PDI->DrawLine(Prev, Next3D, 
+			(bIsSecondary) ? SecondaryLineColor : LineColor,
+			DepthLayer, 
+			(bIsSecondary)? SecondaryLineThickness :LineThickness,
+			0.0f, IsPixelThickness);
+		Prev = Next3D;
+	}
+}
+
+
+void UBrushStampSizeIndicator::DrawLine(FPrimitiveDrawInterface* PDI, const FVector3f& Start, const FVector3f& End, bool bIsSecondary)
+{
+	PDI->DrawLine(Start, End,
+		(bIsSecondary) ? SecondaryLineColor : LineColor,
+		DepthLayer,
+		(bIsSecondary) ? SecondaryLineThickness : LineThickness,
+		0.0f, IsPixelThickness);
+}
\ No newline at end of file
diff --git a/Engine/Plugins/Experimental/MeshModelingToolset/Source/MeshModelingTools/Private/ConvertToPolygonsTool.cpp b/Engine/Plugins/Experimental/MeshModelingToolset/Source/MeshModelingTools/Private/ConvertToPolygonsTool.cpp
new file mode 100644
--- /dev/null
+++ b/Engine/Plugins/Experimental/MeshModelingToolset/Source/MeshModelingTools/Private/ConvertToPolygonsTool.cpp
@@ -0,0 +1,207 @@
+// Copyright 1998-2019 Epic Games, Inc. All Rights Reserved.
+
+#include "ConvertToPolygonsTool.h"
+#include "InteractiveToolManager.h"
+#include "ToolBuilderUtil.h"
+
+#include "DynamicMesh3.h"
+#include "MeshNormals.h"
+#include "DynamicMeshEditor.h"
+#include "MeshRegionBoundaryLoops.h"
+#include "MeshDescriptionToDynamicMesh.h"
+#include "DynamicMeshToMeshDescription.h"
+
+
+#include "SceneManagement.h" // for FPrimitiveDrawInterface
+
+#define LOCTEXT_NAMESPACE "UConvertToPolygonsTool"
+
+
+/*
+ * ToolBuilder
+ */
+
+
+bool UConvertToPolygonsToolBuilder::CanBuildTool(const FToolBuilderState& SceneState) const
+{
+	return ToolBuilderUtil::CountComponents(SceneState, CanMakeComponentTarget) == 1;
+}
+
+UInteractiveTool* UConvertToPolygonsToolBuilder::BuildTool(const FToolBuilderState& SceneState) const
+{
+	UConvertToPolygonsTool* NewTool = NewObject<UConvertToPolygonsTool>(SceneState.ToolManager);
+
+	UActorComponent* ActorComponent = ToolBuilderUtil::FindFirstComponent(SceneState, CanMakeComponentTarget);
+	UPrimitiveComponent* MeshComponent = Cast<UPrimitiveComponent>(ActorComponent);
+	check(MeshComponent != nullptr);
+	NewTool->SetSelection(MakeComponentTarget(MeshComponent));
+
+	return NewTool;
+}
+
+
+
+/*
+ * Tool
+ */
+
+
+void UConvertToPolygonsTool::Setup()
+{
+	UInteractiveTool::Setup();
+
+	FMeshDescription* MeshDescription = nullptr;
+	MeshDescription = ComponentTarget->GetMesh();
+
+	FMeshDescriptionToDynamicMesh Converter;
+	Converter.Convert(MeshDescription, SearchMesh);
+
+	if (SearchMesh.HasAttributes())
+	{
+		InitialNormals.Copy( *SearchMesh.Attributes()->PrimaryNormals() );
+	}
+
+	Settings = NewObject<UConvertToPolygonsToolProperties>(this);
+	AddToolPropertySource(Settings);
+
+	// create preview mesh object
+	PreviewMesh = NewObject<UPreviewMesh>(this, TEXT("PreviewMesh"));
+	PreviewMesh->CreateInWorld(ComponentTarget->GetOwnerActor()->GetWorld(), FTransform::Identity);
+	PreviewMesh->SetVisible(false);
+	PreviewMesh->SetTransform(ComponentTarget->GetWorldTransform());
+
+	if (ComponentTarget->GetMaterial(0) != nullptr)
+	{
+		PreviewMesh->SetMaterial(ComponentTarget->GetMaterial(0));
+	}
+
+	UpdatePolygons();
+}
+
+
+void UConvertToPolygonsTool::Shutdown(EToolShutdownType ShutdownType)
+{
+	PreviewMesh->SetVisible(false);
+	PreviewMesh->Disconnect();
+	PreviewMesh = nullptr;
+
+	ComponentTarget->SetOwnerVisibility(true);
+
+	if (ShutdownType == EToolShutdownType::Accept)
+	{
+		GetToolManager()->BeginUndoTransaction(LOCTEXT("DeformMeshToolTransactionName", "Convert to Polygons"));
+		ComponentTarget->CommitMesh([=](FMeshDescription* MeshDescription)
+		{
+			ConvertToPolygons(MeshDescription);
+		});
+		GetToolManager()->EndUndoTransaction();
+	}
+}
+
+
+void UConvertToPolygonsTool::OnPropertyModified(UObject* PropertySet, UProperty* Property)
+{
+	bPolygonsValid = false;
+	GetToolManager()->PostInvalidation();
+}
+
+
+void UConvertToPolygonsTool::Render(IToolsContextRenderAPI* RenderAPI)
+{
+	FColor LineColor(255, 0, 0);
+
+	if (bPolygonsValid == false)
+	{
+		UpdatePolygons();
+	}
+
+	if (true)
+	{
+		FPrimitiveDrawInterface* PDI = RenderAPI->GetPrimitiveDrawInterface();
+		FTransform Transform = ComponentTarget->GetWorldTransform(); //Actor->GetTransform();
+
+		TArray<int>& Edges = Polygons.PolygonEdges;
+		for (int eid : Edges)
+		{
+			FVector3d A, B;
+			Polygons.Mesh->GetEdgeV(eid, A, B);
+			PDI->DrawLine(Transform.TransformPosition(A), Transform.TransformPosition(B),
+				LineColor, 0, 2.0, 1.0f, true);
+		}
+	}
+}
+
+
+bool UConvertToPolygonsTool::HasAccept() const
+{
+	return true;
+}
+
+bool UConvertToPolygonsTool::CanAccept() const
+{
+	return true;
+}
+
+
+void UConvertToPolygonsTool::UpdatePolygons()
+{
+	Polygons = FFindPolygonsAlgorithm(&SearchMesh);
+	double DotTolerance = 1.0 - FMathd::Cos(Settings->AngleTolerance * FMathd::DegToRad);
+	Polygons.FindPolygons(DotTolerance);
+	Polygons.FindPolygonEdges();
+
+	GetToolManager()->PostMessage(FString::Printf(TEXT("ConvertToPolygons - found %d polys in %d triangles"), Polygons.FoundPolygons.Num(), SearchMesh.TriangleCount()), EToolMessageLevel::Internal);
+
+	if (Settings->bCalculateNormals)
+	{
+		if (SearchMesh.HasAttributes() == false)
+		{
+			SearchMesh.EnableAttributes();
+		}
+
+		FDynamicMeshNormalOverlay* NormalOverlay = SearchMesh.Attributes()->PrimaryNormals();
+		NormalOverlay->ClearElements();
+
+		FDynamicMeshEditor Editor(&SearchMesh);
+		for (const TArray<int>& Polygon : Polygons.FoundPolygons)
+		{
+			FVector3f Normal = (FVector3f)SearchMesh.GetTriNormal(Polygon[0]);
+			Editor.SetTriangleNormals(Polygon, Normal);
+		}
+
+		FMeshNormals Normals(&SearchMesh);
+		Normals.RecomputeOverlayNormals(SearchMesh.Attributes()->PrimaryNormals());
+		Normals.CopyToOverlay(NormalOverlay, false);
+
+		PreviewMesh->UpdatePreview(&SearchMesh);
+		PreviewMesh->SetVisible(true);
+		ComponentTarget->SetOwnerVisibility(false);
+	}
+	else
+	{
+		PreviewMesh->SetVisible(false);
+		ComponentTarget->SetOwnerVisibility(true);
+	}
+
+	bPolygonsValid = true;
+}
+
+
+
+void UConvertToPolygonsTool::ConvertToPolygons(FMeshDescription* MeshIn)
+{
+	// restore input normals
+	if (Settings->bCalculateNormals == false)
+	{
+		SearchMesh.Attributes()->PrimaryNormals()->Copy(InitialNormals);
+	}
+
+	FDynamicMeshToMeshDescription Converter;
+	Converter.bSetPolyGroups = true;
+	Converter.Convert(&SearchMesh, *MeshIn);
+}
+
+
+
+
+#undef LOCTEXT_NAMESPACE
diff --git a/Engine/Plugins/Experimental/MeshModelingToolset/Source/MeshModelingTools/Private/DisplaceMeshTool.cpp b/Engine/Plugins/Experimental/MeshModelingToolset/Source/MeshModelingTools/Private/DisplaceMeshTool.cpp
new file mode 100644
--- /dev/null
+++ b/Engine/Plugins/Experimental/MeshModelingToolset/Source/MeshModelingTools/Private/DisplaceMeshTool.cpp
@@ -0,0 +1,396 @@
+// Copyright 1998-2019 Epic Games, Inc. All Rights Reserved.
+
+#include "DisplaceMeshTool.h"
+#include "InteractiveToolManager.h"
+#include "ToolBuilderUtil.h"
+
+#include "DynamicMesh3.h"
+
+#include "SimpleDynamicMeshComponent.h"
+
+#include "MeshNormals.h"
+
+#include "SceneManagement.h" // for FPrimitiveDrawInterface
+#include "Async/ParallelFor.h"
+
+
+#include "ProfilingDebugging/ScopedTimers.h"
+
+#define LOCTEXT_NAMESPACE "UDisplaceMeshTool"
+
+
+/*
+ * ToolBuilder
+ */
+
+
+bool UDisplaceMeshToolBuilder::CanBuildTool(const FToolBuilderState& SceneState) const
+{
+	return ToolBuilderUtil::CountComponents(SceneState, CanMakeComponentTarget) == 1;
+}
+
+UInteractiveTool* UDisplaceMeshToolBuilder::BuildTool(const FToolBuilderState& SceneState) const
+{
+	UDisplaceMeshTool* NewTool = NewObject<UDisplaceMeshTool>(SceneState.ToolManager);
+
+	UActorComponent* ActorComponent = ToolBuilderUtil::FindFirstComponent(SceneState, CanMakeComponentTarget);
+	auto* MeshComponent = Cast<UPrimitiveComponent>(ActorComponent);
+	check(MeshComponent != nullptr);
+
+	NewTool->SetSelection(MakeComponentTarget(MeshComponent));
+
+	return NewTool;
+}
+
+
+
+/*
+ * Tool
+ */
+
+UDisplaceMeshTool::UDisplaceMeshTool()
+{
+	DisplacementType = EDisplaceMeshToolDisplaceType::DisplacementMap;
+	DisplaceIntensity = 10.0f;
+	RandomSeed = 31337;
+	Subdivisions = 4;
+}
+
+void UDisplaceMeshTool::Setup()
+{
+	UInteractiveTool::Setup();
+
+	// create dynamic mesh component to use for live preview
+	DynamicMeshComponent = NewObject<USimpleDynamicMeshComponent>(ComponentTarget->GetOwnerActor(), "DynamicMesh");
+	DynamicMeshComponent->SetupAttachment(ComponentTarget->GetOwnerActor()->GetRootComponent());
+	DynamicMeshComponent->RegisterComponent();
+	DynamicMeshComponent->SetWorldTransform(ComponentTarget->GetWorldTransform());
+
+	// copy material if there is one
+	auto Material = ComponentTarget->GetMaterial(0);
+	if (Material != nullptr)
+	{
+		DynamicMeshComponent->SetMaterial(0, Material);
+	}
+
+	DynamicMeshComponent->InitializeMesh(ComponentTarget->GetMesh());
+	OriginalMesh.Copy(*DynamicMeshComponent->GetMesh());
+
+	UpdateMap(true);
+
+	// hide input StaticMeshComponent
+	ComponentTarget->SetOwnerVisibility(false);
+
+	// initialize our properties
+	ToolPropertyObjects.Add(this);
+
+	bResultValid = false;
+}
+
+
+void UDisplaceMeshTool::Shutdown(EToolShutdownType ShutdownType)
+{
+	if (DynamicMeshComponent != nullptr)
+	{
+		//DynamicMeshComponent->OnMeshChanged.Remove(OnDynamicMeshComponentChangedHandle);
+
+		ComponentTarget->SetOwnerVisibility(true);
+
+		if (ShutdownType == EToolShutdownType::Accept)
+		{
+			// this block bakes the modified DynamicMeshComponent back into the StaticMeshComponent inside an undo transaction
+			GetToolManager()->BeginUndoTransaction(LOCTEXT("DisplaceMeshToolTransactionName", "Displace Mesh"));
+			ComponentTarget->CommitMesh([=](FMeshDescription* MeshDescription)
+			{
+				DynamicMeshComponent->Bake(MeshDescription, Subdivisions > 0);
+			});
+			GetToolManager()->EndUndoTransaction();
+		}
+
+		DynamicMeshComponent->UnregisterComponent();
+		DynamicMeshComponent->DestroyComponent();
+		DynamicMeshComponent = nullptr;
+	}
+}
+
+
+void UDisplaceMeshTool::Render(IToolsContextRenderAPI* RenderAPI)
+{
+	UpdateResult();
+}
+
+#if WITH_EDITOR
+void UDisplaceMeshTool::PostEditChangeProperty(FPropertyChangedEvent& PropertyChangedEvent)
+{
+	UProperty* PropertyThatChanged = PropertyChangedEvent.Property;
+
+	if (PropertyThatChanged)
+	{
+		//if (PropertyThatChanged->GetFName() == GET_MEMBER_NAME_CHECKED(UDisplaceMeshTool, DisplaceIterations))
+		//{
+		//}
+	}
+
+	bResultValid = false;
+}
+#endif
+
+void UDisplaceMeshTool::UpdateResult()
+{
+	if (bResultValid) 
+	{
+		return;
+	}
+
+	UpdateSubdividedMesh();
+
+	// compute Displaceed positions in PositionBuffer
+	switch (DisplacementType)
+	{
+		default:
+		case EDisplaceMeshToolDisplaceType::Constant:
+			ComputeDisplacement_Constant();
+			break;
+
+		case EDisplaceMeshToolDisplaceType::RandomNoise:
+			ComputeDisplacement_RandomNoise();
+			break;
+
+		case EDisplaceMeshToolDisplaceType::DisplacementMap:
+			ComputeDisplacement_Map();
+			break;
+	}
+
+	// update preview vertex positions
+	FDynamicMesh3* TargetMesh = DynamicMeshComponent->GetMesh();
+	for (int vid : SubdividedMesh.VertexIndicesItr())
+	{
+		TargetMesh->SetVertex(vid, DisplacedBuffer[vid]);
+	}
+
+	// recalculate normals
+	if (TargetMesh->HasAttributes())
+	{
+		FMeshNormals Normals(TargetMesh);
+		FDynamicMeshNormalOverlay* NormalOverlay = TargetMesh->Attributes()->PrimaryNormals();
+		Normals.RecomputeOverlayNormals(NormalOverlay);
+		Normals.CopyToOverlay(NormalOverlay);
+	}
+	else
+	{
+		FMeshNormals::QuickComputeVertexNormals(*TargetMesh);
+	}
+
+	DynamicMeshComponent->NotifyMeshUpdated();
+	//DynamicMeshComponent->FastNotifyPositionsUpdated(true);
+
+	GetToolManager()->PostInvalidation();
+
+	bResultValid = true;
+}
+
+
+
+void UDisplaceMeshTool::ComputeDisplacement_Constant()
+{
+	double Intensity = DisplaceIntensity;
+
+	for (int vid : SubdividedMesh.VertexIndicesItr())
+	{
+		FVector3d Position = PositionBuffer[vid];
+		FVector3d Normal = NormalsBuffer[vid];
+		FVector3d Displacement = Intensity * Normal;
+		DisplacedBuffer[vid] = Position + Displacement;
+	}
+}
+
+
+void UDisplaceMeshTool::ComputeDisplacement_RandomNoise()
+{
+	double Intensity = DisplaceIntensity;
+
+	FMath::SRandInit(RandomSeed);
+
+	for (int vid : SubdividedMesh.VertexIndicesItr())
+	{
+		FVector3d Position = PositionBuffer[vid];
+		FVector3d Normal = NormalsBuffer[vid];
+
+		double RandVal = 2.0 * (FMath::SRand() - 0.5);
+		
+		FVector3d Displacement = RandVal * Intensity * Normal;
+
+		DisplacedBuffer[vid] = Position + Displacement;
+	}
+}
+
+
+
+
+
+void UDisplaceMeshTool::ComputeDisplacement_Map()
+{
+	UpdateMap();
+
+	const FDynamicMeshUVOverlay* UVOverlay = SubdividedMesh.Attributes()->GetUVLayer(0);
+
+	double Intensity = DisplaceIntensity;
+	FMath::SRandInit(RandomSeed);
+
+	for (int tid : SubdividedMesh.TriangleIndicesItr())
+	{
+		FIndex3i Tri = SubdividedMesh.GetTriangle(tid);
+		FIndex3i UVTri = UVOverlay->GetTriangle(tid);
+
+		for (int j = 0; j < 3; ++j)
+		{
+			int vid = Tri[j];
+
+			FVector3d Position = PositionBuffer[vid];
+			FVector3d Normal = NormalsBuffer[vid];
+
+			FVector2f UV = UVOverlay->GetElement(UVTri[j]);
+			double Offset = DisplaceField.BilinearSampleClamped(UV);
+
+			FVector3d Displacement = Offset * Intensity * Normal;
+
+			DisplacedBuffer[vid] = Position + Displacement;
+		}
+	}
+
+}
+
+
+
+
+
+void UDisplaceMeshTool::UpdateMap(bool bForceUpdate)
+{
+	if (CachedMapSource == DisplacementMap && bForceUpdate == false)
+	{
+		return;
+	}
+	CachedMapSource = DisplacementMap;
+
+	if (DisplacementMap == nullptr || DisplacementMap->PlatformData == nullptr || DisplacementMap->PlatformData->Mips.Num() < 1)
+	{
+		DisplaceField = FSampledScalarField2f();
+		return;
+	}
+
+	const FColor* FormatedImageData = static_cast<const FColor*>(DisplacementMap->PlatformData->Mips[0].BulkData.LockReadOnly());
+	if (FormatedImageData == nullptr)
+	{
+		DisplaceField = FSampledScalarField2f();
+		return;
+	}
+
+	int TextureWidth = DisplacementMap->GetSizeX();
+	int TextureHeight = DisplacementMap->GetSizeY();
+
+	DisplaceField.Resize(TextureWidth, TextureHeight, 0.0f);
+	DisplaceField.SetCellSize(1.0f / (float)TextureWidth);
+
+	for (int y = 0; y < TextureHeight; ++y)
+	{
+		for (int x = 0; x < TextureWidth; ++x)
+		{
+			FColor PixelColor = FormatedImageData[y*TextureWidth + x];
+			float Value = PixelColor.R / 255.0;
+			DisplaceField.GridValues[y*TextureWidth + x] = Value;
+		}
+	}
+
+	DisplacementMap->PlatformData->Mips[0].BulkData.Unlock();
+}
+
+
+
+void UDisplaceMeshTool::UpdateSubdividedMesh()
+{
+	if (Subdivisions == CachedSubdivisionsCount)
+	{
+		return;
+	}
+
+	SubdividedMesh.Copy(OriginalMesh);
+	int Rounds = Subdivisions;
+
+	// calculate subdivisions (todo: move to elsewhere)
+	for (int ri = 0; ri < Rounds; ri++)
+	{
+		TArray<int> EdgesToProcess;
+		for (int tid : SubdividedMesh.EdgeIndicesItr())
+		{
+			EdgesToProcess.Add(tid);
+		}
+		int MaxTriangleID = SubdividedMesh.MaxTriangleID();
+
+		TArray<int> TriSplitEdges;
+		TriSplitEdges.Init(-1, SubdividedMesh.MaxTriangleID());
+
+		for (int eid : EdgesToProcess)
+		{
+			FIndex2i EdgeTris = SubdividedMesh.GetEdgeT(eid);
+
+			FDynamicMesh3::FEdgeSplitInfo SplitInfo;
+			EMeshResult result = SubdividedMesh.SplitEdge(eid, SplitInfo);
+			if (result == EMeshResult::Ok)
+			{
+				if (EdgeTris.A < MaxTriangleID && TriSplitEdges[EdgeTris.A] == -1)
+				{
+					TriSplitEdges[EdgeTris.A] = SplitInfo.NewEdges.B;
+				}
+				if (EdgeTris.B != FDynamicMesh3::InvalidID)
+				{
+					if (EdgeTris.B < MaxTriangleID && TriSplitEdges[EdgeTris.B] == -1)
+					{
+						TriSplitEdges[EdgeTris.B] = SplitInfo.NewEdges.C;
+					}
+				}
+			}
+		}
+
+		for (int eid : TriSplitEdges)
+		{
+			if (eid != -1)
+			{
+				FDynamicMesh3::FEdgeFlipInfo FlipInfo;
+				SubdividedMesh.FlipEdge(eid, FlipInfo);
+			}
+		}
+	}
+
+	CachedSubdivisionsCount = Subdivisions;
+
+	// update dynamic mesh component
+	DynamicMeshComponent->GetMesh()->Copy(SubdividedMesh);
+	DynamicMeshComponent->NotifyMeshUpdated();
+
+	NormalsBuffer = FMeshNormals(&SubdividedMesh);
+	NormalsBuffer.ComputeVertexNormals();
+
+	// cache initial positions
+	PositionBuffer.SetNum(SubdividedMesh.MaxVertexID());
+	for (int vid : SubdividedMesh.VertexIndicesItr())
+	{
+		PositionBuffer[vid] = SubdividedMesh.GetVertex(vid);
+	}
+	DisplacedBuffer.SetNum(SubdividedMesh.MaxVertexID());
+}
+
+
+bool UDisplaceMeshTool::HasAccept() const
+{
+	return true;
+}
+
+bool UDisplaceMeshTool::CanAccept() const
+{
+	return true;
+}
+
+
+
+
+#undef LOCTEXT_NAMESPACE
diff --git a/Engine/Plugins/Experimental/MeshModelingToolset/Source/MeshModelingTools/Private/DrawPolygonTool.cpp b/Engine/Plugins/Experimental/MeshModelingToolset/Source/MeshModelingTools/Private/DrawPolygonTool.cpp
new file mode 100644
--- /dev/null
+++ b/Engine/Plugins/Experimental/MeshModelingToolset/Source/MeshModelingTools/Private/DrawPolygonTool.cpp
@@ -0,0 +1,1041 @@
+// Copyright 1998-2019 Epic Games, Inc. All Rights Reserved.
+
+#include "DrawPolygonTool.h"
+#include "InteractiveToolManager.h"
+#include "InteractiveGizmoManager.h"
+#include "ToolBuilderUtil.h"
+#include "BaseBehaviors/MultiClickSequenceInputBehavior.h"
+#include "BaseBehaviors/KeyAsModifierInputBehavior.h"
+
+#include "Polygon2.h"
+#include "FrameTypes.h"
+#include "MatrixTypes.h"
+#include "DynamicMeshAttributeSet.h"
+
+#include "MeshDescriptionBuilder.h"
+#include "Generators/PlanarPolygonMeshGenerator.h"
+#include "Operations/ExtrudeMesh.h"
+#include "Distance/DistLine3Ray3.h"
+#include "Intersection/IntrSegment2Segment2.h"
+#include "MeshQueries.h"
+#include "ToolSceneQueriesUtil.h"
+
+#include "PositionPlaneGizmo.h"
+#include "Drawing/MeshDebugDrawing.h"
+
+#include "Selection/SelectClickedAction.h"
+#include "Selection/ToolSelectionUtil.h"
+#include "AssetGenerationUtil.h"
+
+#include "StaticMeshComponentBuilder.h"
+
+// [RMS] need these for now
+#include "Components/StaticMeshComponent.h"
+#include "Engine/StaticMesh.h"
+#include "Engine/StaticMeshActor.h"
+
+
+#define LOCTEXT_NAMESPACE "UDrawPolygonTool"
+
+
+/*
+ * ToolBuilder
+ */
+
+
+constexpr int StartPointSnapID = FPointPlanarSnapSolver::BaseExternalPointID + 1;
+constexpr int CurrentSceneSnapID = FPointPlanarSnapSolver::BaseExternalPointID + 2;
+
+
+bool UDrawPolygonToolBuilder::CanBuildTool(const FToolBuilderState& SceneState) const
+{
+	return (this->AssetAPI != nullptr);
+}
+
+UInteractiveTool* UDrawPolygonToolBuilder::BuildTool(const FToolBuilderState& SceneState) const
+{
+	UDrawPolygonTool* NewTool = NewObject<UDrawPolygonTool>(SceneState.ToolManager);
+	NewTool->SetWorld(SceneState.World);
+	NewTool->SetAssetAPI(AssetAPI);
+	return NewTool;
+}
+
+
+
+/*
+ * Tool
+ */
+
+UDrawPolygonTool::UDrawPolygonTool()
+{
+	DrawPlaneOrigin = FVector::ZeroVector;
+	DrawPlaneOrientation = FQuat::Identity;
+	bInInteractiveExtrude = false;
+}
+
+
+void UDrawPolygonTool::SetWorld(UWorld* World)
+{
+	this->TargetWorld = World;
+}
+
+void UDrawPolygonTool::SetAssetAPI(IToolsContextAssetAPI* AssetAPIIn)
+{
+	this->AssetAPI = AssetAPIIn;
+}
+
+
+
+
+void UDrawPolygonTool::Setup()
+{
+	UInteractiveTool::Setup();
+
+	// add default button input behaviors for devices
+	UMultiClickSequenceInputBehavior* MouseBehavior = NewObject<UMultiClickSequenceInputBehavior>();
+	MouseBehavior->Initialize(this);
+	MouseBehavior->Modifiers.RegisterModifier(IgnoreSnappingModifier, FInputDeviceState::IsShiftKeyDown);
+	AddInputBehavior(MouseBehavior);
+
+	// click to set plane behavior
+	FSelectClickedAction* SetPlaneAction = new FSelectClickedAction();
+	SetPlaneAction->World = this->TargetWorld;
+	SetPlaneAction->OnClickedPositionFunc = [this](const FHitResult& Hit) {
+		SetDrawPlaneFromWorldPos(Hit.ImpactPoint, Hit.ImpactNormal);
+	};
+	SetPointInWorldConnector = SetPlaneAction;
+
+	USingleClickInputBehavior* ClickToSetPlaneBehavior = NewObject<USingleClickInputBehavior>();
+	ClickToSetPlaneBehavior->ModifierCheckFunc = FInputDeviceState::IsCtrlKeyDown;
+	ClickToSetPlaneBehavior->Initialize(SetPointInWorldConnector);
+	ClickToSetPlaneBehavior->SetDefaultPriority(MouseBehavior->GetPriority().MakeHigher());
+	AddInputBehavior(ClickToSetPlaneBehavior);
+
+	UKeyAsModifierInputBehavior* AKeyBehavior = NewObject<UKeyAsModifierInputBehavior>();
+	AKeyBehavior->Initialize(this, AngleSnapModifier, EKeys::A);
+	AddInputBehavior(AKeyBehavior);
+
+	PolygonProperties = NewObject<UDrawPolygonToolStandardProperties>(this, TEXT("Polygon Settings"));
+	PolygonProperties->OutputMode = EDrawPolygonOutputMode::ExtrudedInteractive;
+	PolygonProperties->ExtrudeHeight = 100.0f;
+
+
+	UPositionPlaneGizmoBuilder* PositionPlaneGizmoBuilder = NewObject<UPositionPlaneGizmoBuilder>();
+	UInteractiveGizmoManager* GizmoManager = GetToolManager()->GetPairedGizmoManager();
+	GizmoManager->RegisterGizmoType(TEXT("DrawPolygonPlaneGizmo"), PositionPlaneGizmoBuilder);
+	PositionPlaneGizmo = GizmoManager->CreateGizmo(TEXT("DrawPolygonPlaneGizmo"), TEXT("TestGizmo1"));
+	Cast<UPositionPlaneGizmo>(PositionPlaneGizmo)->OnPositionUpdatedFunc = [this](const FFrame3d& WorldFrame)
+	{
+		UpdateDrawPlaneFromGizmo(WorldFrame);
+	};
+
+	// create preview mesh object
+	PreviewMesh = NewObject<UPreviewMesh>(this, TEXT("PreviewMesh"));
+	PreviewMesh->CreateInWorld(this->TargetWorld, FTransform::Identity);
+	PreviewMesh->SetVisible(false);
+	bPreviewUpdatePending = false;
+
+	SnapEngine.SnapMetricTolerance = ToolSceneQueriesUtil::GetDefaultVisualAngleSnapThreshD();
+	SnapEngine.SnapMetricFunc = [this](const FVector3d& Position1, const FVector3d& Position2) {
+		return ToolSceneQueriesUtil::CalculateViewVisualAngleD(this->CameraState, Position1, Position2);
+	};
+	SnapEngine.Plane = FFrame3d(DrawPlaneOrigin, DrawPlaneOrientation);
+
+
+	SnapProperties = NewObject<UDrawPolygonToolSnapProperties>(this, TEXT("Snapping"));
+
+	// add self as tool props
+	AddToolPropertySource(PolygonProperties);
+	AddToolPropertySource(SnapProperties);
+}
+
+
+void UDrawPolygonTool::Shutdown(EToolShutdownType ShutdownType)
+{
+	PreviewMesh->Disconnect();
+	PreviewMesh = nullptr;
+
+	if (SetPointInWorldConnector != nullptr)
+	{
+		delete SetPointInWorldConnector;
+	}
+
+	UInteractiveGizmoManager* GizmoManager = GetToolManager()->GetPairedGizmoManager();
+	GizmoManager->DestroyGizmo(PositionPlaneGizmo);
+	PositionPlaneGizmo = nullptr;
+	GizmoManager->DeregisterGizmoType(TEXT("DrawPolygonPlaneGizmo"));
+}
+
+
+
+
+void UDrawPolygonTool::RegisterActions(FInteractiveToolActionSet& ActionSet)
+{
+	ActionSet.RegisterAction(this, (int32)EStandardToolActions::BaseClientDefinedActionID + 1,
+		TEXT("PopLastVertex"), 
+		LOCTEXT("PopLastVertex", "Pop Last Vertex"),
+		LOCTEXT("PopLastVertexTooltip", "Pop last vertex added to polygon"),
+		EModifierKey::None, EKeys::BackSpace,
+		[this]() { PopLastVertexAction(); });
+}
+
+
+void UDrawPolygonTool::PopLastVertexAction()
+{
+	if (bInInteractiveExtrude || PolygonVertices.Num() == 0)
+	{
+		return;
+	}
+
+	if (bInFixedPolygonMode == false)
+	{
+		int NumVertices = PolygonVertices.Num();
+		if (NumVertices > 1)
+		{
+			PolygonVertices.RemoveAt(NumVertices-1);
+		}
+		else
+		{
+			PolygonVertices.RemoveAt(0);
+			bAbortActivePolygonDraw = true;
+		}
+	}
+	else
+	{
+		int NumVertices = FixedPolygonClickPoints.Num();
+		if (NumVertices > 1)
+		{
+			FixedPolygonClickPoints.RemoveAt(NumVertices - 1);
+		}
+		else
+		{
+			FixedPolygonClickPoints.RemoveAt(0);
+			bAbortActivePolygonDraw = true;
+		}
+	}
+}
+
+
+
+
+void DrawEdgeTicks(FPrimitiveDrawInterface* PDI, 
+	const FSegment3f& Segment, float Height,
+	const FVector3f& PlaneNormal, 
+	const FLinearColor& Color, uint8 DepthPriorityGroup, float LineThickness, bool bIsScreenSpace)
+{
+	FVector3f Center = Segment.Center;
+	FVector3f X = Segment.Direction;
+	FVector3f Y = X.Cross(PlaneNormal);
+	Y.Normalize();
+	FVector A = Center - Height * 0.25f*X - Height * Y;
+	FVector B = Center + Height * 0.25f*X + Height * Y;
+	PDI->DrawLine(A, B, Color, DepthPriorityGroup, LineThickness, 0.0f, bIsScreenSpace);
+	A += Height * 0.5f*X;
+	B += Height * 0.5f*X;
+	PDI->DrawLine(A, B, Color, DepthPriorityGroup, LineThickness, 0.0f, bIsScreenSpace);
+}
+
+
+
+void UDrawPolygonTool::Render(IToolsContextRenderAPI* RenderAPI)
+{
+	FPrimitiveDrawInterface* PDI = RenderAPI->GetPrimitiveDrawInterface();
+	GetToolManager()->GetContextQueriesAPI()->GetCurrentViewState(CameraState);
+
+	if (bPreviewUpdatePending)
+	{
+		UpdateLivePreview();
+		bPreviewUpdatePending = false;
+	}
+
+	double CurViewSizeFactor = ToolSceneQueriesUtil::CalculateDimensionFromVisualAngleD(CameraState, PreviewVertex, 1.0);
+
+	FColor OpenPolygonColor(240, 16, 240);
+	FColor ClosedPolygonColor(16, 240, 16);
+	FColor ErrorColor(240, 16, 16);
+	float HiddenLineThickness = 1.0f;
+	float LineThickness = 4.0f;
+	float SelfIntersectThickness = 8.0f;
+	FColor GridColor(128, 128, 128, 32);
+	float GridThickness = 0.5f;
+	float GridLineSpacing = 25.0f;   // @todo should be relative to view
+	int NumGridLines = 21;
+	FColor SnapHighlightColor(240, 200, 16);
+	float ElementSize = CurViewSizeFactor;
+
+	bool bIsClosed = SnapEngine.HaveActiveSnap() && SnapEngine.GetActiveSnapTargetID() == StartPointSnapID;
+
+	if (bInInteractiveExtrude == false)
+	{
+		FFrame3f DrawFrame(DrawPlaneOrigin, DrawPlaneOrientation);
+		MeshDebugDraw::DrawSimpleGrid(DrawFrame, NumGridLines, GridLineSpacing, GridThickness, GridColor, false, PDI, FTransform::Identity);
+	}
+
+	if (bInFixedPolygonMode)
+	{
+		if (bInInteractiveExtrude == false)		// once we are in extrude, polygon is done
+		{
+			FixedPolygonClickPoints.Add(PreviewVertex);
+			GenerateFixedPolygon(FixedPolygonClickPoints, PolygonVertices);
+			FixedPolygonClickPoints.Pop(false);
+		}
+		bIsClosed = true;
+	}
+
+	int NumVerts = PolygonVertices.Num();
+
+	if (SnapEngine.HaveActiveSnap())
+	{
+		PDI->DrawPoint((FVector)SnapEngine.GetActiveSnapToPoint(), ClosedPolygonColor, 10, SDPG_Foreground);
+		
+		PDI->DrawPoint((FVector)SnapEngine.GetActiveSnapFromPoint(), OpenPolygonColor, 15, SDPG_Foreground);
+		PDI->DrawLine((FVector)SnapEngine.GetActiveSnapToPoint(), (FVector)SnapEngine.GetActiveSnapFromPoint(),
+			ClosedPolygonColor, SDPG_Foreground, 0.5f, 0.0f, true);
+		if (SnapEngine.GetActiveSnapTargetID() == CurrentSceneSnapID)
+		{
+			if (LastSnapGeometry.PointCount == 1) {
+				DrawCircle(PDI, (FVector)LastSnapGeometry.Points[0], CameraState.Right(), CameraState.Up(),
+					SnapHighlightColor, ElementSize, 32, SDPG_Foreground, 1.0f, 0.0f, true);
+			} 
+			else
+			{
+				PDI->DrawLine((FVector)LastSnapGeometry.Points[0], (FVector)LastSnapGeometry.Points[1],
+					SnapHighlightColor, SDPG_Foreground, 1.0f, 0.0f, true);
+			}
+		}
+
+		if (SnapEngine.HaveActiveSnapLine())
+		{
+			FLine3d SnapLine = SnapEngine.GetActiveSnapLine();
+			PDI->DrawLine((FVector)SnapLine.PointAt(-9999), (FVector)SnapLine.PointAt(9999),
+				ClosedPolygonColor, SDPG_Foreground, 0.5, 0.0f, true);
+
+			if (SnapEngine.HaveActiveSnapDistance())
+			{
+				int iSegment = SnapEngine.GetActiveSnapDistanceID();
+				TArray<FVector>& HistoryPoints = (bInFixedPolygonMode) ? FixedPolygonClickPoints : PolygonVertices;
+				FVector UseNormal = DrawPlaneOrientation.GetAxisZ();
+				DrawEdgeTicks(PDI, FSegment3f(HistoryPoints[iSegment], HistoryPoints[iSegment+1]),
+					0.75f*ElementSize, UseNormal, SnapHighlightColor, SDPG_Foreground, 1.0f, true);
+				DrawEdgeTicks(PDI, FSegment3f(HistoryPoints[HistoryPoints.Num()-1], PreviewVertex),
+					0.75f*ElementSize, UseNormal, SnapHighlightColor, SDPG_Foreground, 1.0f, true);
+				PDI->DrawLine(HistoryPoints[iSegment], HistoryPoints[iSegment + 1],
+					SnapHighlightColor, SDPG_Foreground, 2.0f, 0.0f, true);
+			}
+		}
+	}
+
+
+	if (bHaveSurfaceHit)
+	{
+		PDI->DrawPoint((FVector)SurfaceHitPoint, ClosedPolygonColor, 10, SDPG_Foreground);
+		if (SnapProperties->HitNormalOffset != 0)
+		{
+			PDI->DrawPoint((FVector)SurfaceOffsetPoint, OpenPolygonColor, 15, SDPG_Foreground);
+			PDI->DrawLine((FVector)SurfaceOffsetPoint, (FVector)SurfaceHitPoint,
+				ClosedPolygonColor, SDPG_Foreground, 0.5f, 0.0f, true);
+		}
+		PDI->DrawLine((FVector)SurfaceOffsetPoint, (FVector)PreviewVertex,
+			ClosedPolygonColor, SDPG_Foreground, 0.5f, 0.0f, true);
+	}
+
+
+	if (PolygonVertices.Num() > 0)
+	{
+		FColor UseColor = (bIsClosed) ? ClosedPolygonColor : OpenPolygonColor;
+		FVector UseLastVertex = (bIsClosed) ? PolygonVertices[0] : PreviewVertex;
+		float UseThickness = LineThickness;
+		if (bHaveSelfIntersection)
+		{
+			UseColor = ErrorColor;
+			UseThickness = SelfIntersectThickness;
+		}
+
+		// draw thin no-depth
+		for (int i = 0; i < NumVerts - 1; ++i)
+		{
+			PDI->DrawLine(PolygonVertices[i], PolygonVertices[i + 1],
+				UseColor, SDPG_Foreground, HiddenLineThickness, 0.0f, true);
+		}
+		PDI->DrawLine(PolygonVertices[NumVerts - 1], UseLastVertex,
+			UseColor, SDPG_Foreground, HiddenLineThickness, 0.0f, true);
+
+
+		// draw thick depth-tested
+		for (int i = 0; i < NumVerts-1; ++i)
+		{
+			PDI->DrawLine(PolygonVertices[i], PolygonVertices[i+1],
+				UseColor, SDPG_World, LineThickness, 0.0f, true);
+		}
+		PDI->DrawLine(PolygonVertices[NumVerts-1], UseLastVertex,
+			UseColor, SDPG_World, LineThickness, 0.0f, true);
+
+		if (bHaveSelfIntersection)
+		{
+			PDI->DrawPoint(SelfIntersectionPoint, ErrorColor, 10, SDPG_Foreground);
+		}
+	}
+
+	// draw preview vertex
+	PDI->DrawPoint(PreviewVertex, ClosedPolygonColor, 10, SDPG_Foreground);
+
+
+	// todo should be an indicator
+	if (bInInteractiveExtrude)
+	{
+		float Length = 10; float Thickness = 2;
+		FColor HitFrameColor(0, 128, 128);
+		PDI->DrawLine(
+			HitPosFrameWorld.PointAt(-Length, -Length, 0), HitPosFrameWorld.PointAt(Length, Length, 0),
+			HitFrameColor, 1, Thickness, 0.0f, true);
+		PDI->DrawLine(
+			HitPosFrameWorld.PointAt(-Length, Length, 0), HitPosFrameWorld.PointAt(Length, -Length, 0),
+			HitFrameColor, 1, Thickness, 0.0f, true);
+
+		FVector PreviewOrigin = (FVector)PreviewHeightFrame.Origin;
+
+		FVector DrawPlaneNormal = DrawPlaneOrientation.GetAxisZ();
+
+		FColor AxisColor(128, 128, 0);
+		PDI->DrawLine(
+			PreviewOrigin -1000*DrawPlaneNormal, PreviewOrigin +1000*DrawPlaneNormal,
+			AxisColor, 1, 1.0f, 0.0f, true);
+
+		FColor HeightPosColor(128, 0, 128);
+		PDI->DrawLine(
+			PreviewOrigin + PolygonProperties->ExtrudeHeight*DrawPlaneNormal, HitPosFrameWorld.Origin,
+			HeightPosColor, 1, 1.0f, 0.0f, true);
+	}
+}
+
+
+
+void UDrawPolygonTool::ResetPolygon()
+{
+	PolygonVertices.Reset();
+	SnapEngine.Reset();
+	bHaveSurfaceHit = false;
+	bInFixedPolygonMode = false;
+}
+
+void UDrawPolygonTool::UpdatePreviewVertex(const FVector& PreviewVertexIn)
+{
+	PreviewVertex = PreviewVertexIn;
+
+	// update length and angle
+	if (PolygonVertices.Num() > 0)
+	{
+		FVector LastVertex = PolygonVertices[PolygonVertices.Num() - 1];
+		SnapProperties->Length = FVector::Distance(LastVertex, PreviewVertex);
+	}
+}
+
+
+void UDrawPolygonTool::AppendVertex(const FVector& Vertex)
+{
+	PolygonVertices.Add(Vertex);
+}
+
+
+bool UDrawPolygonTool::FindDrawPlaneHitPoint(const FInputDeviceRay& ClickPos, FVector& HitPosOut)
+{
+	bHaveSurfaceHit = false;
+
+	FFrame3d Frame(DrawPlaneOrigin, DrawPlaneOrientation);
+	FVector3d HitPos;
+	bool bHit = Frame.RayPlaneIntersection(ClickPos.WorldRay.Origin, ClickPos.WorldRay.Direction, 2, HitPos);
+	if (bHit == false)
+	{
+		return false;
+	}
+
+	// if we found a scene snap point, add to snap set
+	FVector3d SnapPos;
+	if (bIgnoreSnappingToggle || SnapProperties->bEnableSnapping == false)
+	{
+		SnapEngine.ResetActiveSnap();
+		SnapEngine.UpdatePointHistory(TArray<FVector>());
+	}
+	else 
+	{
+		if (ToolSceneQueriesUtil::FindSceneSnapPoint(this, HitPos, SnapPos, SnapProperties->bSnapToVertices, SnapProperties->bSnapToEdges, 0, &LastSnapGeometry))
+		{
+			SnapEngine.AddPointTarget(SnapPos, CurrentSceneSnapID, SnapEngine.MinInternalPriority()-1 );
+		}
+
+		TArray<FVector>& HistoryPoints = (bInFixedPolygonMode) ? FixedPolygonClickPoints : PolygonVertices;
+		SnapEngine.UpdatePointHistory(HistoryPoints);
+		if (SnapProperties->bSnapToAngles)
+		{
+			SnapEngine.RegenerateTargetLines(true, true);
+		}
+		SnapEngine.bEnableSnapToKnownLengths = SnapProperties->bSnapToLengths;
+	}
+
+	SnapEngine.UpdateSnappedPoint(HitPos);
+
+	// remove scene snap point
+	SnapEngine.RemovePointTargetsByID(CurrentSceneSnapID);
+
+	if (SnapEngine.HaveActiveSnap())
+	{
+		HitPosOut = (FVector)SnapEngine.GetActiveSnapToPoint();
+		return true;
+	}
+
+
+	// if not snap and we want to hit objects, do that
+	if (SnapProperties->bHitSceneObjects)
+	{
+		FCollisionObjectQueryParams QueryParams(FCollisionObjectQueryParams::AllObjects);
+		FHitResult Result;
+		bool bWorldHit = TargetWorld->LineTraceSingleByObjectType(Result, ClickPos.WorldRay.Origin, ClickPos.WorldRay.PointAt(9999), QueryParams);
+		if (bWorldHit)
+		{
+			bHaveSurfaceHit = true;
+			SurfaceHitPoint = Result.ImpactPoint;
+			FVector UseHitPos = Result.ImpactPoint + SnapProperties->HitNormalOffset*Result.Normal;
+			HitPos = Frame.ToPlane((FVector3d)UseHitPos, 2);
+			SurfaceOffsetPoint = UseHitPos;
+		}
+	}
+
+	HitPosOut = (FVector)HitPos;
+	return true;
+}
+
+
+void UDrawPolygonTool::OnBeginSequencePreview(const FInputDeviceRay& DevicePos)
+{
+	// just update snapped point preview
+	FVector HitPos;
+	if (FindDrawPlaneHitPoint(DevicePos, HitPos))
+	{
+		PreviewVertex = HitPos;
+	}
+	
+}
+
+
+bool UDrawPolygonTool::CanBeginClickSequence(const FInputDeviceRay& ClickPos)
+{
+	return true;
+}
+
+void UDrawPolygonTool::OnBeginClickSequence(const FInputDeviceRay& ClickPos)
+{
+	ResetPolygon();
+	
+	FVector HitPos;
+	bool bHit = FindDrawPlaneHitPoint(ClickPos, HitPos);
+	if (bHit == false)
+	{
+		bAbortActivePolygonDraw = true;
+		return;
+	}
+	if (ToolSceneQueriesUtil::IsPointVisible(CameraState, HitPos) == false)
+	{
+		bAbortActivePolygonDraw = true;
+		return;		// cannot start a poly an a point that is not visible, this is almost certainly an error due to draw plane
+	}
+
+	AppendVertex(HitPos);
+	UpdatePreviewVertex(HitPos);
+
+	bInFixedPolygonMode = (PolygonProperties->PolygonType != EDrawPolygonDrawMode::Freehand);
+	FixedPolygonClickPoints.Reset();
+	FixedPolygonClickPoints.Add(HitPos);
+
+	// if we are starting a freehand poly, add start point as snap target, but then ignore it until we get 3 verts
+	if (bInFixedPolygonMode == false)
+	{
+		SnapEngine.AddPointTarget(PolygonVertices[0], StartPointSnapID, 1);
+		SnapEngine.AddIgnoreTarget(StartPointSnapID);
+	}
+}
+
+void UDrawPolygonTool::OnNextSequencePreview(const FInputDeviceRay& ClickPos)
+{
+	if (bInInteractiveExtrude)
+	{
+		PolygonProperties->ExtrudeHeight = FindInteractiveHeightDistance(ClickPos);
+		bPreviewUpdatePending = true;
+		return;
+	}
+
+	FVector HitPos;
+	bool bHit = FindDrawPlaneHitPoint(ClickPos, HitPos);
+	if (bHit == false)
+	{
+		return;
+	}
+
+	if (bInFixedPolygonMode)
+	{
+		UpdatePreviewVertex(HitPos);
+		bPreviewUpdatePending = true;
+		return;
+	}
+
+	if (PolygonVertices.Num() > 2)
+	{
+		bPreviewUpdatePending = true;
+	}
+
+	UpdatePreviewVertex(HitPos);
+	UpdateSelfIntersection();
+}
+
+bool UDrawPolygonTool::OnNextSequenceClick(const FInputDeviceRay& ClickPos)
+{
+	if (bInInteractiveExtrude)
+	{
+		EndInteractiveExtrude();
+		return false;
+	}
+
+	FVector HitPos;
+	bool bHit = FindDrawPlaneHitPoint(ClickPos, HitPos);
+	if (bHit == false)
+	{
+		return true;  // ignore click but continue accepting clicks
+	}
+
+	bool bDonePolygon = false;
+	if (bInFixedPolygonMode)
+	{
+		// ignore very close click points
+		if ( ToolSceneQueriesUtil::PointSnapQuery(this, FixedPolygonClickPoints[FixedPolygonClickPoints.Num()-1], HitPos) )
+		{
+			return true;
+		}
+
+		FixedPolygonClickPoints.Add(HitPos);
+		int NumTargetPoints = (PolygonProperties->PolygonType == EDrawPolygonDrawMode::Rectangle) ? 3 : 2;
+		bDonePolygon = (FixedPolygonClickPoints.Num() == NumTargetPoints);
+		if (bDonePolygon)
+		{
+			GenerateFixedPolygon(FixedPolygonClickPoints, PolygonVertices);
+		}
+	} 
+	else
+	{
+		// ignore very close click points
+		if (ToolSceneQueriesUtil::PointSnapQuery(this, PolygonVertices[PolygonVertices.Num()-1], HitPos))
+		{
+			return true;
+		}
+
+		// close polygon if we clicked on start point
+		bDonePolygon = SnapEngine.HaveActiveSnap() && SnapEngine.GetActiveSnapTargetID() == StartPointSnapID;
+
+		if (bHaveSelfIntersection)
+		{
+			// discard vertex in segments before intersection (this is redundant if idx is 0)
+			for (int j = SelfIntersectSegmentIdx; j < PolygonVertices.Num(); ++j)
+			{
+				PolygonVertices[j-SelfIntersectSegmentIdx] = PolygonVertices[j];
+			}
+			PolygonVertices.SetNum(PolygonVertices.Num() - SelfIntersectSegmentIdx);
+			PolygonVertices[0] = PreviewVertex = SelfIntersectionPoint;
+			bDonePolygon = true;
+		}
+	}
+
+	
+	if (bDonePolygon)
+	{
+		SnapEngine.Reset();
+		bHaveSurfaceHit = false;
+		if (PolygonProperties->OutputMode == EDrawPolygonOutputMode::ExtrudedInteractive)
+		{
+			BeginInteractiveExtrude();
+
+			PreviewMesh->ClearPreview();
+			PreviewMesh->SetVisible(true);
+
+			return true;
+		}
+		else 
+		{
+			EmitCurrentPolygon();
+
+			PreviewMesh->ClearPreview();
+			PreviewMesh->SetVisible(false);
+
+			return false;
+		}
+	}
+
+	AppendVertex(HitPos);
+	if (PolygonVertices.Num() > 2)
+	{
+		SnapEngine.RemoveIgnoreTarget(StartPointSnapID);
+	}
+
+	UpdatePreviewVertex(HitPos);
+	return true;
+}
+
+void UDrawPolygonTool::OnTerminateClickSequence()
+{
+	ResetPolygon();
+}
+
+bool UDrawPolygonTool::RequestAbortClickSequence()
+{
+	if (bAbortActivePolygonDraw)
+	{
+		bAbortActivePolygonDraw = false;
+		return true;
+	}
+	return false;
+}
+
+
+
+void UDrawPolygonTool::OnUpdateModifierState(int ModifierID, bool bIsOn)
+{
+	if (ModifierID == IgnoreSnappingModifier)
+	{
+		bIgnoreSnappingToggle = bIsOn;
+	}
+	else if (ModifierID == AngleSnapModifier)
+	{
+
+	}
+}
+
+
+
+
+bool UDrawPolygonTool::UpdateSelfIntersection()
+{
+	bHaveSelfIntersection = false;
+	if (bInFixedPolygonMode || PolygonProperties->bAllowSelfIntersections == true)
+	{
+		return false;
+	}
+
+	int NumVertices = PolygonVertices.Num();
+
+	FFrame3f DrawFrame(DrawPlaneOrigin, DrawPlaneOrientation);
+	FSegment2f PreviewSegment(DrawFrame.ToPlaneUV(PolygonVertices[NumVertices - 1],2), DrawFrame.ToPlaneUV(PreviewVertex,2));
+
+	for (int k = 0; k < NumVertices - 2; ++k) 
+	{
+		FSegment2f Segment(DrawFrame.ToPlaneUV(PolygonVertices[k],2), DrawFrame.ToPlaneUV(PolygonVertices[k + 1],2));
+		FIntrSegment2Segment2f Intersection(PreviewSegment, Segment);
+		if (Intersection.Find()) 
+		{
+			bHaveSelfIntersection = true;
+			SelfIntersectSegmentIdx = k;
+			SelfIntersectionPoint = DrawFrame.FromPlaneUV(Intersection.Point0,2);
+			return true;
+		}
+	}
+	return false;
+}
+
+
+
+void UDrawPolygonTool::GenerateFixedPolygon(TArray<FVector>& FixedPoints, TArray<FVector>& VerticesOut)
+{
+	FFrame3f DrawFrame(DrawPlaneOrigin, DrawPlaneOrientation);
+	FVector2f CenterPt = DrawFrame.ToPlaneUV(FixedPoints[0], 2 );
+	FVector2f EdgePt = DrawFrame.ToPlaneUV(FixedPoints[1], 2);
+	FVector2f Delta = EdgePt - CenterPt;
+	float AngleRad = FMath::Atan2(Delta.Y, Delta.X);
+	FMatrix2f RotationMat = FMatrix2f::RotationRad(AngleRad);
+	FVector2f RotAxisX = RotationMat * FVector2f::UnitX();
+	float Dist = Delta.Length();
+	float Width = FMath::Abs(Delta.Dot(RotAxisX));
+
+	FPolygon2f Polygon;
+	if (PolygonProperties->PolygonType == EDrawPolygonDrawMode::Square)
+	{
+		Polygon = FPolygon2f::MakeRectangle(FVector2f::Zero(), 2*Width, 2*Width);
+	}
+	else if (PolygonProperties->PolygonType == EDrawPolygonDrawMode::Rectangle)
+	{
+		FVector2f HeightPt = DrawFrame.ToPlaneUV((FixedPoints.Num() == 3) ? FixedPoints[2] : FixedPoints[1], 2);
+		FVector2f HeightDelta = HeightPt - CenterPt;
+		FVector2f RotAxisY = RotationMat * FVector2f::UnitY();
+		float YSign = FMath::Sign(HeightDelta.Dot(RotAxisY));
+		float Height = FMath::Abs(HeightDelta.Dot(RotAxisY));
+		Polygon = FPolygon2f::MakeRectangle(FVector2f(Width/2, YSign*Height/2), Width, Height);
+	}
+	else
+	{
+		Polygon = FPolygon2f::MakeCircle(Dist, PolygonProperties->Steps, 0);
+	}
+	Polygon.Transform([RotationMat](const FVector2f& Pt) { return RotationMat * Pt; });
+
+	VerticesOut.Reset();
+	for (int k = 0; k < Polygon.VertexCount(); ++k)
+	{
+		FVector2f NewPt = CenterPt + Polygon[k];
+		VerticesOut.Add(DrawFrame.FromPlaneUV(NewPt, 2));
+	}
+}
+
+
+
+void UDrawPolygonTool::BeginInteractiveExtrude()
+{
+	bInInteractiveExtrude = true;
+
+	GeneratePreviewHeightTarget();
+}
+
+
+void UDrawPolygonTool::EndInteractiveExtrude()
+{
+	EmitCurrentPolygon();
+
+	PreviewMesh->ClearPreview();
+	PreviewMesh->SetVisible(false);
+
+	bInInteractiveExtrude = false;
+}
+
+
+float UDrawPolygonTool::FindInteractiveHeightDistance(const FInputDeviceRay& ClickPos)
+{
+	float NearestHitDist = TNumericLimits<float>::Max();
+	float NearestHitHeight = 1.0f;
+	FFrame3f NearestHitFrameWorld;
+
+	// cast ray at target object
+	FRay3d LocalRay = PreviewHeightFrame.ToFrame((FRay3d)ClickPos.WorldRay);
+	int HitTID = PreviewHeightTargetAABB.FindNearestHitTriangle(LocalRay);
+	if (HitTID >= 0)
+	{
+		FIntrRay3Triangle3d IntrQuery =
+			TMeshQueries<FDynamicMesh3>::TriangleIntersection(PreviewHeightTarget, HitTID, LocalRay);
+		FVector3d HitPosLocal = LocalRay.PointAt(IntrQuery.RayParameter);
+		FVector3d HitNormalLocal = PreviewHeightTarget.GetTriNormal(HitTID);
+
+		NearestHitFrameWorld = FFrame3f(
+			(FVector3f)PreviewHeightFrame.FromFramePoint(HitPosLocal),
+			(FVector3f)PreviewHeightFrame.FromFrameVector(HitNormalLocal));
+		NearestHitHeight = HitPosLocal.Z;
+		NearestHitDist = ClickPos.WorldRay.GetParameter(NearestHitFrameWorld.Origin);
+	}
+	
+
+	// cast ray into scene
+	FVector RayStart = ClickPos.WorldRay.Origin;
+	FVector RayEnd = ClickPos.WorldRay.PointAt(999999);
+	FCollisionObjectQueryParams QueryParams(FCollisionObjectQueryParams::AllObjects);
+	FHitResult Result;
+	bool bHitWorld = TargetWorld->LineTraceSingleByObjectType(Result, RayStart, RayEnd, QueryParams);
+	if (bHitWorld)
+	{
+		float WorldHitDist = ClickPos.WorldRay.GetParameter(Result.ImpactPoint);
+		if (WorldHitDist < NearestHitDist)
+		{
+			NearestHitFrameWorld = FFrame3f(Result.ImpactPoint, Result.ImpactNormal);
+			FVector3d HitPosWorld = Result.ImpactPoint;
+			FVector3d HitPosLocal = PreviewHeightFrame.ToFramePoint(HitPosWorld);
+			NearestHitHeight = HitPosLocal.Z;
+			NearestHitDist = WorldHitDist;
+		}
+	}
+
+	if (NearestHitDist < TNumericLimits<float>::Max())
+	{
+		this->HitPosFrameWorld = NearestHitFrameWorld;
+		return NearestHitHeight;
+	}
+	else
+	{
+		return PolygonProperties->ExtrudeHeight;
+	}
+	
+}
+
+
+
+
+
+void UDrawPolygonTool::SetDrawPlaneFromWorldPos(const FVector& Position, const FVector& Normal)
+{
+	DrawPlaneOrigin = Position;
+
+	FFrame3f DrawPlane(Position, DrawPlaneOrientation);
+	if (bIgnoreSnappingToggle == false)
+	{
+		DrawPlane.AlignAxis(2, Normal);
+		DrawPlane.ConstrainedAlignPerpAxes();
+		DrawPlaneOrientation = DrawPlane.Rotation;
+	}
+
+	SnapEngine.Plane = FFrame3d((FVector3d)DrawPlane.Origin, (FQuaterniond)DrawPlane.Rotation);
+
+	UPositionPlaneGizmo* Gizmo = Cast<UPositionPlaneGizmo>(PositionPlaneGizmo);
+	Gizmo->ExternalUpdatePosition(DrawPlaneOrigin, DrawPlaneOrientation, false);
+}
+
+
+void UDrawPolygonTool::UpdateDrawPlaneFromGizmo(const FFrame3d& WorldPosition)
+{
+	if (bIgnoreSnappingToggle == false)
+	{
+		DrawPlaneOrientation = WorldPosition.Rotation;
+	}
+	DrawPlaneOrigin = (FVector)WorldPosition.Origin;
+
+	SnapEngine.Plane = FFrame3d((FVector3d)DrawPlaneOrigin, (FQuaterniond)DrawPlaneOrientation);
+
+	if (bIgnoreSnappingToggle) 
+	{
+		UPositionPlaneGizmo* Gizmo = Cast<UPositionPlaneGizmo>(PositionPlaneGizmo);
+		Gizmo->ExternalUpdatePosition(DrawPlaneOrigin, DrawPlaneOrientation, false);
+	}
+}
+
+
+
+
+
+void UDrawPolygonTool::EmitCurrentPolygon()
+{
+#if WITH_EDITOR
+	// generate new mesh
+	FFrame3d PlaneFrameOut;
+	FDynamicMesh3 Mesh;
+	double ExtrudeDist = (PolygonProperties->OutputMode == EDrawPolygonOutputMode::MeshedPolygon) ?
+		0 : PolygonProperties->ExtrudeHeight;
+	GeneratePolygonMesh(PolygonVertices, &Mesh, PlaneFrameOut, false, ExtrudeDist, false);
+
+
+	GetToolManager()->BeginUndoTransaction(LOCTEXT("CreatePolygon", "Create Polygon"));
+
+	AActor* NewActor = AssetGenerationUtil::GenerateStaticMeshActor(
+		AssetAPI, TargetWorld,
+		&Mesh, PlaneFrameOut.ToTransform(), TEXT("Polygon"),
+		AssetGenerationUtil::GetDefaultAutoGeneratedAssetPath());
+
+	// select newly-created object
+	ToolSelectionUtil::SetNewActorSelection(GetToolManager(), NewActor);
+
+	GetToolManager()->EndUndoTransaction();
+#else
+	checkNoEntry();
+#endif
+	ResetPolygon();
+}
+
+
+
+
+
+
+void UDrawPolygonTool::UpdateLivePreview()
+{
+	int NumVerts = PolygonVertices.Num();
+	if (NumVerts < 2 || PreviewMesh == nullptr || PreviewMesh->IsVisible() == false )
+	{
+		return;
+	}
+
+	FFrame3d PlaneFrame;
+	FDynamicMesh3 Mesh;
+	double ExtrudeDist = (PolygonProperties->OutputMode == EDrawPolygonOutputMode::MeshedPolygon) ?
+		0 : PolygonProperties->ExtrudeHeight;
+	GeneratePolygonMesh(PolygonVertices, &Mesh, PlaneFrame, false, ExtrudeDist, false);
+
+	PreviewMesh->SetTransform(PlaneFrame.ToFTransform());
+	PreviewMesh->UpdatePreview(&Mesh);
+}
+
+
+
+
+void UDrawPolygonTool::GeneratePolygonMesh(const TArray<FVector>& Polygon, FDynamicMesh3* ResultMeshOut, FFrame3d& WorldFrameOut, bool bIncludePreviewVtx, double ExtrudeDistance, bool bExtrudeSymmetric)
+{
+	// construct centered frame for polygon
+	WorldFrameOut = FFrame3d(DrawPlaneOrigin, DrawPlaneOrientation);
+
+	int NumVerts = Polygon.Num();
+	FVector3d Centroid(0, 0, 0);
+	for (int k = 0; k < NumVerts; ++k)
+	{
+		Centroid += Polygon[k];
+	}
+	Centroid /= (double)NumVerts;
+	WorldFrameOut.Origin = Centroid;
+
+
+	// triangulate polygon into the MeshDescription
+	FPlanarPolygonMeshGenerator PolyMeshGen;
+	for (int k = 0; k < NumVerts; ++k)
+	{
+		PolyMeshGen.Polygon.AppendVertex(WorldFrameOut.ToPlaneUV(Polygon[k], 2));
+	}
+
+	// add preview vertex
+	if (bIncludePreviewVtx)
+	{
+		if (FVector::Dist(PreviewVertex, Polygon[NumVerts - 1]) > 0.1)
+		{
+			PolyMeshGen.Polygon.AppendVertex(WorldFrameOut.ToPlaneUV(PreviewVertex, 2));
+		}
+	}
+
+	if (PolyMeshGen.Polygon.IsClockwise() == false)
+	{
+		PolyMeshGen.Polygon.Reverse();
+	}
+
+	ResultMeshOut->Copy(&PolyMeshGen.Generate());
+
+	// for symmetric extrude we translate the first poly by -dist along axis
+	if (bExtrudeSymmetric)
+	{
+		FVector3d ShiftNormal = FVector3d::UnitZ();
+		for (int vid : ResultMeshOut->VertexIndicesItr())
+		{
+			FVector3d Pos = ResultMeshOut->GetVertex(vid);
+			ResultMeshOut->SetVertex(vid, Pos - ExtrudeDistance * ShiftNormal);
+		}
+		// double extrude dist
+		ExtrudeDistance *= 2.0;
+	}
+
+	if (ExtrudeDistance != 0)
+	{
+		FExtrudeMesh Extruder(ResultMeshOut);
+		Extruder.DefaultExtrudeDistance = ExtrudeDistance;
+		FAxisAlignedBox2d bounds = PolyMeshGen.Polygon.Bounds();
+		Extruder.UVScaleFactor = 1.0 / bounds.MaxDim();
+		if (ExtrudeDistance < 0)
+		{
+			Extruder.IsPositiveOffset = false;
+		}
+
+		FVector3d ExtrudeNormal = FVector3d::UnitZ();
+		Extruder.ExtrudedPositionFunc = [&ExtrudeDistance, &ExtrudeNormal](const FVector3d& Position, const FVector3f& Normal, int VertexID)
+		{
+			return Position + ExtrudeDistance * ExtrudeNormal;
+		};
+
+		Extruder.Apply();
+	}
+}
+
+
+
+
+void UDrawPolygonTool::GeneratePreviewHeightTarget()
+{
+	GeneratePolygonMesh(PolygonVertices, &PreviewHeightTarget, PreviewHeightFrame, false, 99999, true);
+	PreviewHeightTargetAABB.SetMesh(&PreviewHeightTarget);
+}
+
+
+
+
+
+#undef LOCTEXT_NAMESPACE
diff --git a/Engine/Plugins/Experimental/MeshModelingToolset/Source/MeshModelingTools/Private/DynamicMeshSculptTool.cpp b/Engine/Plugins/Experimental/MeshModelingToolset/Source/MeshModelingTools/Private/DynamicMeshSculptTool.cpp
new file mode 100644
--- /dev/null
+++ b/Engine/Plugins/Experimental/MeshModelingToolset/Source/MeshModelingTools/Private/DynamicMeshSculptTool.cpp
@@ -0,0 +1,1185 @@
+// Copyright 1998-2019 Epic Games, Inc. All Rights Reserved.
+
+#include "DynamicMeshSculptTool.h"
+#include "InteractiveToolManager.h"
+#include "ToolBuilderUtil.h"
+
+#include "SubRegionRemesher.h"
+#include "ProjectionTargets.h"
+#include "MeshConstraints.h"
+#include "MeshConstraintsUtil.h"
+#include "MeshVertexChange.h"
+#include "MeshWeights.h"
+#include "MeshNormals.h"
+#include "MeshIndexUtil.h"
+#include "Drawing/MeshDebugDrawing.h"
+
+#include "MeshVertexChange.h"
+#include "Changes/MeshChange.h"
+#include "DynamicMeshChangeTracker.h"
+
+#include "Async/ParallelFor.h"
+
+#define LOCTEXT_NAMESPACE "UDynamicMeshSculptTool"
+
+
+/*
+ * ToolBuilder
+ */
+
+UMeshSurfacePointTool* UDynamicMeshSculptToolBuilder::CreateNewTool(const FToolBuilderState& SceneState) const
+{
+	UDynamicMeshSculptTool* SculptTool = NewObject<UDynamicMeshSculptTool>(SceneState.ToolManager);
+	SculptTool->SetEnableRemeshing(this->bEnableRemeshing);
+	return SculptTool;
+}
+
+
+
+UBrushSculptProperties::UBrushSculptProperties()
+{
+	SmoothPower = 0.25;
+	OffsetPower = 0.5;
+	PrimaryBrushType = EDynamicMeshSculptBrushType::Pull;
+	SmoothingType = EMeshSculptToolSmoothType::TexturePreserving;
+	Depth = 0;
+}
+
+
+
+UBrushRemeshProperties::UBrushRemeshProperties()
+{
+	RelativeSize = 1.0;
+	Smoothing = 0.1;
+}
+
+
+/*
+ * Tool
+ */
+
+UDynamicMeshSculptTool::UDynamicMeshSculptTool()
+{
+	// initialize parameters
+	bEnableRemeshing = true;
+}
+
+
+void UDynamicMeshSculptTool::Setup()
+{
+	UMeshSurfacePointTool::Setup();
+
+	// create dynamic mesh component to use for live preview
+	DynamicMeshComponent = NewObject<UOctreeDynamicMeshComponent>(ComponentTarget->GetOwnerActor(), "DynamicMeshSculptToolMesh");
+	DynamicMeshComponent->SetupAttachment(ComponentTarget->GetOwnerActor()->GetRootComponent());
+	DynamicMeshComponent->RegisterComponent();
+	DynamicMeshComponent->SetWorldTransform(ComponentTarget->GetWorldTransform());
+
+	// copy material if there is one
+	auto Material = ComponentTarget->GetMaterial(0);
+	if (Material != nullptr)
+	{
+		DynamicMeshComponent->SetMaterial(0, Material);
+	}
+
+	// initialize from LOD-0 MeshDescription
+	//DynamicMeshComponent->TangentsType = (bEnableRemeshing) ? EDynamicMeshTangentCalcType::NoTangents : EDynamicMeshTangentCalcType::AutoCalculated;
+	DynamicMeshComponent->InitializeMesh(ComponentTarget->GetMesh());
+	OnDynamicMeshComponentChangedHandle = DynamicMeshComponent->OnMeshChanged.Add(
+		FSimpleMulticastDelegate::FDelegate::CreateUObject(this, &UDynamicMeshSculptTool::OnDynamicMeshComponentChanged));
+
+	// initialize AABBTree
+	FDynamicMesh3* Mesh = DynamicMeshComponent->GetMesh();
+
+	// do we always want to keep vertex normals updated?
+	FMeshNormals::QuickComputeVertexNormals(*Mesh);
+
+	// switch to vertex normals for testing
+	//DynamicMeshComponent->GetMesh()->DiscardAttributes();
+
+	UpdateTarget();
+	bTargetDirty = false;
+
+	double MaxDimension = DynamicMeshComponent->GetMesh()->GetCachedBounds().MaxDim();
+	BrushRelativeSizeRange = FInterval1d(MaxDimension*0.01, 2*MaxDimension);
+	BrushProperties = NewObject<UBrushBaseProperties>(this, TEXT("Brush"));
+	CalculateBrushRadius();
+
+	SculptProperties = NewObject<UBrushSculptProperties>(this, TEXT("Sculpting"));
+
+	RemeshProperties = NewObject<UBrushRemeshProperties>(this, TEXT("Remeshing"));
+	InitialEdgeLength = EstimateIntialSafeTargetLength(*Mesh, 5000);
+
+	// hide input StaticMeshComponent
+	ComponentTarget->SetOwnerVisibility(false);
+
+	// init state flags flags
+	bInDrag = false;
+	bHaveRemeshed = false;
+	bRemeshPending = false;
+	bStampPending = false;
+	ActiveVertexChange = nullptr;
+
+	// create indicators
+	Indicators = NewObject<UToolIndicatorSet>(this, "Indicators");
+	Indicators->Connect(this);
+
+	UBrushStampSizeIndicator* StampIndicator = NewObject<UBrushStampSizeIndicator>(this, "Brush Circle");
+	StampIndicator->bDrawSecondaryLines = true;
+	StampIndicator->DepthLayer = 1;
+	StampIndicator->BrushRadius = MakeAttributeLambda([this] { return (float)this->CurrentBrushRadius; });
+	StampIndicator->BrushPosition = MakeAttributeLambda([this] { return this->LastBrushPosWorld; });
+	StampIndicator->BrushNormal = MakeAttributeLambda([this] { return this->LastBrushPosNormalWorld; });
+	Indicators->AddIndicator(StampIndicator);
+
+	// initialize our properties
+	AddToolPropertySource(SculptProperties);
+	AddToolPropertySource(BrushProperties);
+	if (this->bEnableRemeshing)
+	{
+		AddToolPropertySource(RemeshProperties);
+	}
+}
+
+
+
+
+void UDynamicMeshSculptTool::Shutdown(EToolShutdownType ShutdownType)
+{
+	Indicators->Disconnect();
+
+	if (DynamicMeshComponent != nullptr)
+	{
+		DynamicMeshComponent->OnMeshChanged.Remove(OnDynamicMeshComponentChangedHandle);
+
+		ComponentTarget->SetOwnerVisibility(true);
+
+		if (ShutdownType == EToolShutdownType::Accept)
+		{
+			// this block bakes the modified DynamicMeshComponent back into the StaticMeshComponent inside an undo transaction
+			GetToolManager()->BeginUndoTransaction(LOCTEXT("SculptMeshToolTransactionName", "Sculpt Mesh"));
+			ComponentTarget->CommitMesh([=](FMeshDescription* MeshDescription)
+			{
+				DynamicMeshComponent->Bake(MeshDescription, bHaveRemeshed);
+			});
+			GetToolManager()->EndUndoTransaction();
+		}
+
+		DynamicMeshComponent->UnregisterComponent();
+		DynamicMeshComponent->DestroyComponent();
+		DynamicMeshComponent = nullptr;
+	}
+
+}
+
+
+
+void UDynamicMeshSculptTool::OnDynamicMeshComponentChanged()
+{
+	bNormalUpdatePending = true;
+	bTargetDirty = true;
+}
+
+
+void UDynamicMeshSculptTool::OnPropertyModified(UObject* PropertySet, UProperty* Property)
+{
+	CalculateBrushRadius();
+}
+
+
+
+bool UDynamicMeshSculptTool::HitTest(const FRay& Ray, FHitResult& OutHit)
+{
+	FTransform Transform = ComponentTarget->GetWorldTransform();
+
+	FRay3d LocalRay(Transform.InverseTransformPosition(Ray.Origin),
+		Transform.InverseTransformVector(Ray.Direction));
+	LocalRay.Direction.Normalize();
+
+	FDynamicMesh3* Mesh = DynamicMeshComponent->GetMesh();
+	int HitTID = DynamicMeshComponent->GetOctree()->FindNearestHitObject(LocalRay);
+
+	if (HitTID != IndexConstants::InvalidID)
+	{
+		FTriangle3d Triangle;
+		Mesh->GetTriVertices(HitTID, Triangle.V[0], Triangle.V[1], Triangle.V[2]);
+		FIntrRay3Triangle3d Query(LocalRay, Triangle);
+		Query.Find();
+
+		OutHit.FaceIndex = HitTID;
+		OutHit.Distance = Query.RayParameter;
+		OutHit.Normal = Transform.TransformVectorNoScale(Mesh->GetTriNormal(HitTID));
+		OutHit.ImpactPoint = Transform.TransformPosition(LocalRay.PointAt(Query.RayParameter));
+		return true;
+	}
+
+	return false;
+}
+
+
+
+void UDynamicMeshSculptTool::OnBeginDrag(const FRay& Ray)
+{
+	bSmoothing = GetShiftToggle();
+
+	FHitResult OutHit;
+	if (HitTest(Ray, OutHit))
+	{
+		BrushStartCenterWorld = Ray.PointAt(OutHit.Distance) + SculptProperties->Depth*CurrentBrushRadius*Ray.Direction;
+
+		bInDrag = true;
+
+		ActiveDragPlane = FPlane(BrushStartCenterWorld, -Ray.Direction);
+		LastHitPosWorld = FMath::RayPlaneIntersection(Ray.Origin, Ray.Direction, ActiveDragPlane);
+
+		LastBrushPosWorld = LastHitPosWorld;
+		LastBrushPosNormalWorld = ActiveDragPlane;
+		FTransform Transform = ComponentTarget->GetWorldTransform();
+		LastBrushPosLocal = Transform.InverseTransformPosition(LastHitPosWorld);
+
+		BeginChange(bEnableRemeshing == false);
+
+		UpdateROI(LastBrushPosLocal);
+
+		if (SculptProperties->PrimaryBrushType == EDynamicMeshSculptBrushType::Flatten)
+		{
+			ComputeFlattenFrame();
+		}
+
+		// apply initial stamp
+		PendingStampRay = Ray;
+		bStampPending = true;
+	}
+}
+
+
+
+void UDynamicMeshSculptTool::UpdateROI(const FVector& BrushPos)
+{
+	SCOPE_CYCLE_COUNTER(SculptTool_UpdateROI);
+
+	// TODO: need dynamic vertex hash table!
+
+	float RadiusSqr = CurrentBrushRadius * CurrentBrushRadius;
+
+
+	FAxisAlignedBox3d BrushBox(
+		BrushPos - CurrentBrushRadius * FVector3d::One(),
+		BrushPos + CurrentBrushRadius * FVector3d::One());
+
+	VertexSetBuffer.Reset();
+	FDynamicMesh3* Mesh = DynamicMeshComponent->GetMesh();
+	FDynamicMeshOctree3* Octree = DynamicMeshComponent->GetOctree();
+	Octree->RangeQuery(BrushBox,
+		[&](int TriIdx) {
+		
+		FIndex3i TriV = Mesh->GetTriangle(TriIdx);
+		for (int j = 0; j < 3; ++j)
+		{
+			FVector3d Position = Mesh->GetVertex(TriV[j]);
+			if ((Position - BrushPos).SquaredLength() < RadiusSqr)
+			{
+				VertexSetBuffer.Add(TriV[j]);
+			}
+		}
+
+	});
+
+	VertexROI.SetNum(0, false);
+	BufferUtil::AppendElements(VertexROI, VertexSetBuffer);
+
+	TriangleROI.Reset();
+	MeshIndexUtil::VertexToTriangleOneRing(Mesh, VertexROI, TriangleROI);
+}
+
+
+void UDynamicMeshSculptTool::OnUpdateDrag(const FRay& WorldRay)
+{
+	if (bInDrag)
+	{
+		PendingStampRay = WorldRay;
+		bStampPending = true;
+	}
+}
+
+
+
+void UDynamicMeshSculptTool::CalculateBrushRadius()
+{
+	CurrentBrushRadius = 0.5 * BrushRelativeSizeRange.Interpolate(BrushProperties->BrushSize);
+	if (BrushProperties->bSpecifyRadius)
+	{
+		CurrentBrushRadius = BrushProperties->BrushRadius;
+	}
+	else
+	{
+		BrushProperties->BrushRadius = CurrentBrushRadius;
+	}
+}
+
+
+
+void UDynamicMeshSculptTool::ApplyStamp(const FRay& WorldRay)
+{
+	SCOPE_CYCLE_COUNTER(STAT_SculptToolApplyStamp);
+
+	CalculateBrushRadius();
+
+	SaveActiveROI();
+
+	if (bSmoothing)
+	{
+		ApplySmoothBrush(WorldRay);
+		return;
+	}
+
+	switch (SculptProperties->PrimaryBrushType)
+	{
+		case EDynamicMeshSculptBrushType::Offset:
+			ApplyOffsetBrush(WorldRay);
+			break;
+		case EDynamicMeshSculptBrushType::Pull:
+			ApplyMoveBrush(WorldRay);
+			break;
+		case EDynamicMeshSculptBrushType::Pinch:
+			ApplyPinchBrush(WorldRay);
+			break;
+		case EDynamicMeshSculptBrushType::Flatten:
+			ApplyFlattenBrush(WorldRay);
+			break;
+	}
+}
+
+
+
+
+void UDynamicMeshSculptTool::ApplySmoothBrush(const FRay& WorldRay)
+{
+	// both these are useful...
+
+	//bool bHit = UpdateBrushPositionOnActivePlane(WorldRay);
+	bool bHit = UpdateBrushPositionOnTargetMesh(WorldRay);
+	if (bHit == false)
+	{
+		return;
+	}
+
+
+	FTransform Transform = ComponentTarget->GetWorldTransform();
+	FVector NewBrushPosLocal = Transform.InverseTransformPosition(LastBrushPosWorld);
+
+	FDynamicMesh3* Mesh = DynamicMeshComponent->GetMesh();
+	for (int VertIdx : VertexROI)
+	{
+		FVector3d OrigPos = Mesh->GetVertex(VertIdx);
+
+		double UseDist = (OrigPos - NewBrushPosLocal).Length();
+		double d = UseDist / CurrentBrushRadius;
+		double w = 1;
+		if (d > 0.5) 
+		{
+			d = VectorUtil::Clamp((d - 0.5) / 0.5, 0.0, 1.0);
+			w = (1 - d * d);
+			w = w * w * w;
+		}
+
+		FVector3d SmoothedPos = (SculptProperties->SmoothingType == EMeshSculptToolSmoothType::TexturePreserving) ?
+			FMeshWeights::MeanValueCentroid(*Mesh, VertIdx) : FMeshWeights::UniformCentroid(*Mesh, VertIdx);
+
+		FVector3d NewPos = FVector3d::Lerp(OrigPos, SmoothedPos, w*SculptProperties->SmoothPower);
+
+		Mesh->SetVertex(VertIdx, NewPos);
+		UpdateSavedVertex(VertIdx, OrigPos, NewPos);
+	}
+
+	bRemeshPending = bEnableRemeshing;
+	LastBrushPosLocal = NewBrushPosLocal;
+}
+
+
+
+
+
+void UDynamicMeshSculptTool::ApplyMoveBrush(const FRay& WorldRay)
+{
+	bool bHit = UpdateBrushPositionOnActivePlane(WorldRay);
+
+	FTransform Transform = ComponentTarget->GetWorldTransform();
+	FVector NewBrushPosLocal = Transform.InverseTransformPosition(LastBrushPosWorld);
+
+	FVector3d MoveVec = NewBrushPosLocal - LastBrushPosLocal;
+
+	if (MoveVec.SquaredLength() > 0)
+	{
+		FDynamicMesh3* Mesh = DynamicMeshComponent->GetMesh();
+		for (int VertIdx : VertexROI)
+		{
+			FVector3d OrigPos = Mesh->GetVertex(VertIdx);
+
+			double PrevDist = (OrigPos - LastBrushPosLocal).Length();
+			double NewDist = (OrigPos - NewBrushPosLocal).Length();
+			double UseDist = FMath::Min(PrevDist, NewDist);
+			double d = UseDist / CurrentBrushRadius;
+			double w = 1;
+			if (d > 0.5) 
+			{
+				d = VectorUtil::Clamp((d - 0.5) / 0.5, 0.0, 1.0);
+				w = (1 - d * d);
+				w = w * w * w;
+			}
+
+			FVector3d NewPos = OrigPos + w * MoveVec;
+
+			Mesh->SetVertex(VertIdx, NewPos);
+			UpdateSavedVertex(VertIdx, OrigPos, NewPos);
+		}
+
+		bRemeshPending = bEnableRemeshing;
+	}
+
+	LastBrushPosLocal = NewBrushPosLocal;
+}
+
+
+
+
+
+
+void UDynamicMeshSculptTool::ApplyOffsetBrush(const FRay& WorldRay)
+{
+	bool bHit = UpdateBrushPositionOnTargetMesh(WorldRay);
+	if (bHit == false)
+	{
+		return;
+	}
+
+	FTransform Transform = ComponentTarget->GetWorldTransform();
+	FVector NewBrushPosLocal = Transform.InverseTransformPosition(LastBrushPosWorld);
+
+	FDynamicMesh3* Mesh = DynamicMeshComponent->GetMesh();
+	for (int VertIdx : VertexROI)
+	{
+		FVector3d OrigPos = Mesh->GetVertex(VertIdx);
+
+		FVector3d BasePos, BaseNormal;
+		GetTargetMeshNearest(OrigPos, (double)(2 * CurrentBrushRadius), BasePos, BaseNormal);
+
+		FVector3d MoveVec = SculptProperties->OffsetPower*FMathd::Sqrt(CurrentBrushRadius)*BaseNormal;
+
+		double UseDist = (OrigPos - NewBrushPosLocal).Length();
+		double d = UseDist / CurrentBrushRadius;
+		double w = 1;
+		if (d > 0.5) 
+		{
+			d = VectorUtil::Clamp((d - 0.5) / 0.5, 0.0, 1.0);
+			w = (1 - d * d);
+			w = w * w * w;
+		}
+
+		FVector3d NewPos = OrigPos + w * MoveVec;
+
+		Mesh->SetVertex(VertIdx, NewPos);
+		UpdateSavedVertex(VertIdx, OrigPos, NewPos);
+	}
+
+	bRemeshPending = bEnableRemeshing;
+
+	LastBrushPosLocal = NewBrushPosLocal;
+}
+
+
+
+
+void UDynamicMeshSculptTool::ApplyPinchBrush(const FRay& WorldRay)
+{
+	bool bHit = UpdateBrushPositionOnTargetMesh(WorldRay);
+	if (bHit == false)
+	{
+		return;
+	}
+
+	FTransform Transform = ComponentTarget->GetWorldTransform();
+	FVector NewBrushPosLocal = Transform.InverseTransformPosition(LastBrushPosWorld);
+	FVector BrushNormalLocal = Transform.InverseTransformVectorNoScale(LastBrushPosNormalWorld);
+	FVector OffsetBrushPosLocal = NewBrushPosLocal - SculptProperties->Depth * CurrentBrushRadius * BrushNormalLocal;
+
+	FVector3d MotionVec = NewBrushPosLocal - LastBrushPosLocal;
+	MotionVec.Normalize();
+
+	FDynamicMesh3* Mesh = DynamicMeshComponent->GetMesh();
+	for (int VertIdx : VertexROI)
+	{
+		FVector3d OrigPos = Mesh->GetVertex(VertIdx);
+		FVector3d Delta = OffsetBrushPosLocal - OrigPos;
+		FVector3d MoveVec = SculptProperties->OffsetPower*Delta;
+
+		double UseDist = (OrigPos - NewBrushPosLocal).Length();
+		double d = UseDist / CurrentBrushRadius;
+		double w = 1;
+		if (d > 0.5)
+		{
+			d = VectorUtil::Clamp((d - 0.5) / 0.5, 0.0, 1.0);
+			w = (1 - d * d);
+			w = w * w * w;
+		}
+		w = w * w * w;
+
+		double AnglePower = 1.0 - FMathd::Abs(MoveVec.Normalized().Dot(MotionVec));
+		w *= AnglePower;
+
+		FVector3d NewPos = OrigPos + w * MoveVec;
+
+		Mesh->SetVertex(VertIdx, NewPos);
+		UpdateSavedVertex(VertIdx, OrigPos, NewPos);
+	}
+
+	bRemeshPending = bEnableRemeshing;
+
+	LastBrushPosLocal = NewBrushPosLocal;
+}
+
+
+
+
+void UDynamicMeshSculptTool::ComputeFlattenFrame()
+{
+	FDynamicMesh3* Mesh = DynamicMeshComponent->GetMesh();
+	FVector3d AverageNormal(0, 0, 0);
+	FVector3d AveragePos(0, 0, 0);
+	for (int TriID : TriangleROI)
+	{
+		AverageNormal += Mesh->GetTriNormal(TriID);
+		AveragePos += Mesh->GetTriCentroid(TriID);
+	}
+	AverageNormal.Normalize();
+	AveragePos /= TriangleROI.Num();
+	ActiveFlattenFrame = FFrame3d(AveragePos, AverageNormal);
+
+	ActiveFlattenFrame.Origin -= SculptProperties->Depth * CurrentBrushRadius * ActiveFlattenFrame.Z();
+}
+
+void UDynamicMeshSculptTool::ApplyFlattenBrush(const FRay& WorldRay)
+{
+	bool bHit = UpdateBrushPositionOnTargetMesh(WorldRay);
+	if (bHit == false)
+	{
+		return;
+	}
+
+	FTransform Transform = ComponentTarget->GetWorldTransform();
+	FVector NewBrushPosLocal = Transform.InverseTransformPosition(LastBrushPosWorld);
+	FVector BrushNormalLocal = Transform.InverseTransformVectorNoScale(LastBrushPosNormalWorld);
+	FVector OffsetBrushPosLocal = NewBrushPosLocal - SculptProperties->Depth * CurrentBrushRadius * BrushNormalLocal;
+
+	FDynamicMesh3* Mesh = DynamicMeshComponent->GetMesh();
+
+	for (int VertIdx : VertexROI)
+	{
+		FVector3d OrigPos = Mesh->GetVertex(VertIdx);
+		FVector3d PlanePos = ActiveFlattenFrame.ToPlane(OrigPos, 2);
+		FVector3d Delta = PlanePos - OrigPos;
+		FVector3d MoveVec = FMathd::Sqrt(SculptProperties->OffsetPower) * Delta;
+
+		double UseDist = (OrigPos - NewBrushPosLocal).Length();
+		double d = UseDist / CurrentBrushRadius;
+		double w = 1;
+		if (d > 0.5)
+		{
+			d = VectorUtil::Clamp((d - 0.5) / 0.5, 0.0, 1.0);
+			w = (1 - d * d);
+			w = w * w * w;
+		}
+
+		FVector3d NewPos = OrigPos + w * MoveVec;
+
+		Mesh->SetVertex(VertIdx, NewPos);
+		UpdateSavedVertex(VertIdx, OrigPos, NewPos);
+	}
+
+	bRemeshPending = bEnableRemeshing;
+
+	LastBrushPosLocal = NewBrushPosLocal;
+}
+
+
+
+
+
+
+
+bool UDynamicMeshSculptTool::UpdateBrushPositionOnActivePlane(const FRay& WorldRay)
+{
+	FVector NewHitPosWorld = FMath::RayPlaneIntersection(WorldRay.Origin, WorldRay.Direction, ActiveDragPlane);
+	LastBrushPosWorld = NewHitPosWorld;
+	LastBrushPosNormalWorld = ActiveDragPlane;
+	return true;
+}
+
+bool UDynamicMeshSculptTool::UpdateBrushPositionOnTargetMesh(const FRay& WorldRay)
+{
+	FTransform Transform = ComponentTarget->GetWorldTransform();
+
+	FRay3d LocalRay(Transform.InverseTransformPosition(WorldRay.Origin),
+		Transform.InverseTransformVectorNoScale(WorldRay.Direction));
+	LocalRay.Direction.Normalize();
+
+	int HitTID = BrushTargetMeshSpatial.FindNearestHitTriangle(LocalRay);
+	if (HitTID != IndexConstants::InvalidID)
+	{
+		const FDynamicMesh3* TargetMesh = BrushTargetMeshSpatial.GetMesh();
+
+		FTriangle3d Triangle;
+		TargetMesh->GetTriVertices(HitTID, Triangle.V[0], Triangle.V[1], Triangle.V[2]);
+		FIntrRay3Triangle3d Query(LocalRay, Triangle);
+		Query.Find();
+
+		LastBrushPosNormalWorld = Transform.TransformVectorNoScale(TargetMesh->GetTriNormal(HitTID));
+		LastBrushPosWorld = Transform.TransformPosition(LocalRay.PointAt(Query.RayParameter));
+		return true;
+	}
+	return false;
+}
+
+
+
+
+void UDynamicMeshSculptTool::OnEndDrag(const FRay& Ray)
+{
+	bInDrag = false;
+
+	// cancel these! otherwise change record could become invalid 
+	bStampPending = false;
+	bRemeshPending = false;
+
+	// update spatial
+	bTargetDirty = true;
+
+	// close change record
+	EndChange();
+}
+
+
+void UDynamicMeshSculptTool::OnUpdateHover(const FInputDeviceRay& DevicePos)
+{
+	if (bInDrag)
+	{
+		FVector NewHitPosWorld = FMath::RayPlaneIntersection(DevicePos.WorldRay.Origin, DevicePos.WorldRay.Direction, ActiveDragPlane);
+		LastBrushPosWorld = NewHitPosWorld;
+		LastBrushPosNormalWorld = ActiveDragPlane;
+	}
+	else
+	{
+		FHitResult OutHit;
+		if (HitTest(DevicePos.WorldRay, OutHit))
+		{
+			LastBrushPosWorld = DevicePos.WorldRay.PointAt(OutHit.Distance + SculptProperties->Depth*CurrentBrushRadius);
+			LastBrushPosNormalWorld = OutHit.Normal;
+		}
+	}
+}
+
+
+
+
+void UDynamicMeshSculptTool::Render(IToolsContextRenderAPI* RenderAPI)
+{
+	UMeshSurfacePointTool::Render(RenderAPI);
+	Indicators->Render(RenderAPI);
+
+	//MeshDebugDraw::DrawNormals(DynamicMeshComponent->GetMesh()->Attributes()->PrimaryNormals(),
+	//	25.0f, FColor::Red, 5.0f, true, RenderAPI->GetPrimitiveDrawInterface(), .GetWorldTransform() );
+}
+
+
+void UDynamicMeshSculptTool::Tick(float DeltaTime)
+{
+	SCOPE_CYCLE_COUNTER(STAT_SculptToolTick);
+
+	UMeshSurfacePointTool::Tick(DeltaTime);
+	Indicators->Tick(DeltaTime);
+
+	bool bMeshModified = false;
+	bool bMeshShapeModified = false;
+
+	FDynamicMesh3* Mesh = DynamicMeshComponent->GetMesh();
+	FDynamicMeshOctree3* Octree = DynamicMeshComponent->GetOctree();
+
+	//Octree->CheckValidity(EValidityCheckFailMode::Check, false, true);
+
+	//
+	// Apply stamp
+	// 
+
+	if (bStampPending)
+	{
+		SCOPE_CYCLE_COUNTER(STAT_SculptTool_Tick_ApplyStampBlock);
+
+		// would this ever be true? does stamp require this?
+		bool bRemoveTrianglesBeforeStamp = false;
+
+		if ( bRemoveTrianglesBeforeStamp )
+		{
+			SCOPE_CYCLE_COUNTER(SculptTool_Tick_ApplyStamp_Remove);
+			Octree->RemoveTriangles(TriangleROI);
+		}
+		else
+		{
+			SCOPE_CYCLE_COUNTER(SculptTool_Tick_ApplyStamp_Remove);
+			Octree->NotifyPendingModification(TriangleROI);		// to mark initial positions
+		}
+
+		ApplyStamp(PendingStampRay);
+		bStampPending = false;
+
+		bNormalUpdatePending = true;
+		bMeshModified = true;
+		bMeshShapeModified = true;
+
+		// flow
+		if (bInDrag)
+		{
+			bStampPending = true;
+		}
+
+		if (bRemoveTrianglesBeforeStamp)
+		{
+			SCOPE_CYCLE_COUNTER(SculptTool_Tick_ApplyStamp_Insert);
+			Octree->InsertTriangles(TriangleROI);
+		}
+		else		
+		{
+			SCOPE_CYCLE_COUNTER(SculptTool_Tick_ApplyStamp_Insert);
+			Octree->ReinsertTriangles(TriangleROI);
+		}
+
+	}
+
+
+	bool bUpdatedROIInRemesh = false;
+	if (bRemeshPending)
+	{
+		SCOPE_CYCLE_COUNTER(STAT_SculptTool_Tick_RemeshBlock);
+
+		check(bInDrag == true);    // this would break undo otherwise!
+
+		RemeshROIPass();
+
+		bMeshModified = true;
+		bMeshShapeModified = true;
+		bRemeshPending = false;
+		bNormalUpdatePending = true;
+		bHaveRemeshed = true;
+		bUpdatedROIInRemesh = true;
+	}
+
+	//Octree->CheckValidity(EValidityCheckFailMode::Check, false, true);
+
+	if (bNormalUpdatePending)
+	{
+		SCOPE_CYCLE_COUNTER(STAT_SculptTool_Tick_NormalsBlock);
+
+		if (Mesh->HasAttributes() && Mesh->Attributes()->PrimaryNormals() != nullptr)
+		{
+			RecalculateNormals_Overlay();
+		}
+		else
+		{
+			RecalculateNormals_PerVertex();
+		}
+		bNormalUpdatePending = false;
+		bMeshModified = true;
+	}
+
+	if (bMeshModified)
+	{
+		SCOPE_CYCLE_COUNTER(STAT_SculptTool_Tick_UpdateMeshBlock);
+
+		if (bMeshModified)
+
+		DynamicMeshComponent->NotifyMeshUpdated();
+		GetToolManager()->PostInvalidation();
+
+		bMeshModified = false;
+
+		if (bUpdatedROIInRemesh == false)
+		{
+			UpdateROI(LastBrushPosLocal);
+		}
+	}
+
+	if (bTargetDirty)
+	{
+		UpdateTarget();
+		bTargetDirty = false;
+	}
+}
+
+
+
+void UDynamicMeshSculptTool::RemeshROIPass()
+{
+	FDynamicMesh3* Mesh = DynamicMeshComponent->GetMesh();
+	FDynamicMeshOctree3* Octree = DynamicMeshComponent->GetOctree();
+
+	FSubRegionRemesher Remesher(Mesh);
+	Remesher.SetTargetEdgeLength(RemeshProperties->RelativeSize * InitialEdgeLength);
+
+	// this is a temporary tweak for pinch brush. Remesh params should be per-brush!
+	if (SculptProperties->PrimaryBrushType == EDynamicMeshSculptBrushType::Pinch && bSmoothing == false)
+	{
+		Remesher.MinEdgeLength = Remesher.MinEdgeLength / 4.0;
+	}
+
+	Remesher.SmoothSpeedT = RemeshProperties->Smoothing;
+	Remesher.SmoothType = (SculptProperties->SmoothingType == EMeshSculptToolSmoothType::TexturePreserving) ?
+		FRemesher::ESmoothTypes::MeanValue : FRemesher::ESmoothTypes::Uniform;
+	bool bIsUniformSmooth = (Remesher.SmoothType == FRemesher::ESmoothTypes::Uniform);
+
+	Remesher.bEnableCollapses = RemeshProperties->bCollapses;
+	Remesher.bEnableFlips = RemeshProperties->bFlips;
+	Remesher.bEnableSplits = RemeshProperties->bSplits;
+	Remesher.bPreventNormalFlips = RemeshProperties->bPreventNormalFlips;
+
+	{
+		SCOPE_CYCLE_COUNTER(STAT_SculptTool_Remesh_Setup);
+		for (int VertIdx : VertexROI)
+		{
+			Remesher.VertexROI.Add(VertIdx);
+		}
+		Remesher.InitializeFromVertexROI();
+		Remesher.UpdateROI();	// required to use roi in constraints fn below
+		Octree->RemoveTriangles(Remesher.GetCurrentTriangleROI());
+	}
+
+	FMeshConstraints constraints;
+	bool bConstraintAllowSplits = true;
+	bool bConstraintAllowSmoothing = false;
+	{
+		SCOPE_CYCLE_COUNTER(STAT_SculptTool_Remesh_Constraints);
+
+		// TODO: only constrain in ROI. This is quite difficult to do externally because we need to update based on
+		// the changing triangle set in Remesher. Perhaps FSubRegionRemesher should update the constraints itself?
+
+		FMeshConstraintsUtil::ConstrainAllSeams(constraints, *Mesh, bConstraintAllowSplits, bConstraintAllowSmoothing);
+		Remesher.SetExternalConstraints(&constraints);
+	}
+
+	for (int k = 0; k < 5; ++k)
+	{
+		if (bIsUniformSmooth == false)
+		{
+			Remesher.bEnableFlips = (k < 2);
+		}
+
+		{
+			SCOPE_CYCLE_COUNTER(STAT_SculptTool_Remesh_RemeshROIUpdate);
+
+			Remesher.UpdateROI();
+
+			if (ActiveMeshChange != nullptr)
+			{
+				Remesher.SaveActiveROI(ActiveMeshChange);
+				//ActiveMeshChange->VerifySaveState();    // useful for debugging
+			}
+
+			Remesher.BeginTrackRemovedTrisInPass();
+		}
+
+		{
+			SCOPE_CYCLE_COUNTER(STAT_SculptTool_Remesh_RemeshPass);
+			Remesher.BasicRemeshPass();
+		}
+
+		{
+			SCOPE_CYCLE_COUNTER(STAT_SculptTool_Remesh_PassOctreeUpdate);
+			const TSet<int32>& TrisRemovedInPass = Remesher.EndTrackRemovedTrisInPass();
+			Octree->RemoveTriangles(TrisRemovedInPass);
+		}
+	}
+	//UE_LOG(LogTemp, Warning, TEXT("Triangle Count %d after update"), Mesh->TriangleCount());
+
+	{
+		SCOPE_CYCLE_COUNTER(STAT_SculptTool_Remesh_Finish);
+
+		// reinsert new ROI into octree
+		Octree->ReinsertTriangles(Remesher.GetCurrentTriangleROI());
+
+		//Octree->CheckValidity(EValidityCheckFailMode::Check, false, true);
+
+		UpdateROI(LastBrushPosLocal);
+	}
+}
+
+
+
+
+void UDynamicMeshSculptTool::RecalculateNormals_PerVertex()
+{
+	FDynamicMesh3* Mesh = DynamicMeshComponent->GetMesh();
+	FDynamicMeshOctree3* Octree = DynamicMeshComponent->GetOctree();
+
+	int MaxVertexID = Mesh->MaxVertexID();
+	if (NormalsVertexFlags.Num() < MaxVertexID)
+	{
+		NormalsVertexFlags.Init(false, MaxVertexID * 2);
+	}
+
+
+	{
+		SCOPE_CYCLE_COUNTER(SculptTool_Normals_Collect);
+
+		TrianglesBuffer.Reset();
+		NormalsBuffer.Reset();
+		Octree->RangeQuery(Octree->ModifiedBounds, TrianglesBuffer);
+		for (int TriangleID : TrianglesBuffer)
+		{
+			FIndex3i TriV = Mesh->GetTriangle(TriangleID);
+			for (int j = 0; j < 3; ++j)
+			{
+				int vid = TriV[j];
+				if (NormalsVertexFlags[vid] == false)
+				{
+					NormalsBuffer.Add(vid);
+					NormalsVertexFlags[vid] = true;
+				}
+			}
+		}
+	}
+
+	//UE_LOG(LogTemp, Warning, TEXT("Computing %d normals"), NormalsBuffer.Num());
+
+	{
+		SCOPE_CYCLE_COUNTER(SculptTool_Normals_Compute);
+
+		ParallelFor(NormalsBuffer.Num(), [&](int k) {
+			int vid = NormalsBuffer[k];
+			FVector3d NewNormal = FMeshNormals::ComputeVertexNormal(*Mesh, vid);
+			Mesh->SetVertexNormal(vid, (FVector3f)NewNormal);
+			NormalsVertexFlags[vid] = false;
+		});
+	}
+
+}
+
+
+
+
+
+
+
+
+
+void UDynamicMeshSculptTool::RecalculateNormals_Overlay()
+{
+	FDynamicMesh3* Mesh = DynamicMeshComponent->GetMesh();
+	FDynamicMeshNormalOverlay* Normals = Mesh->HasAttributes() ? Mesh->Attributes()->PrimaryNormals() : nullptr;
+	check(Normals != nullptr);
+
+
+	FDynamicMeshOctree3* Octree = DynamicMeshComponent->GetOctree();
+
+	int MaxElementID = Normals->MaxElementID();
+	if (NormalsVertexFlags.Num() < MaxElementID)
+	{
+		NormalsVertexFlags.Init(false, MaxElementID * 2);
+	}
+
+
+	{
+		SCOPE_CYCLE_COUNTER(SculptTool_Normals_Collect);
+
+		TrianglesBuffer.Reset();
+		NormalsBuffer.Reset();
+		Octree->RangeQuery(Octree->ModifiedBounds, TrianglesBuffer);
+		for (int TriangleID : TrianglesBuffer)
+		{
+			FIndex3i TriElems = Normals->GetTriangle(TriangleID);
+			for (int j = 0; j < 3; ++j)
+			{
+				int elemid = TriElems[j];
+				if (NormalsVertexFlags[elemid] == false)
+				{
+					NormalsBuffer.Add(elemid);
+					NormalsVertexFlags[elemid] = true;
+				}
+			}
+		}
+	}
+
+	//UE_LOG(LogTemp, Warning, TEXT("Computing %d normals"), NormalsBuffer.Num());
+
+	{
+		SCOPE_CYCLE_COUNTER(SculptTool_Normals_Compute);
+
+		ParallelFor(NormalsBuffer.Num(), [&](int k) {
+			int elemid = NormalsBuffer[k];
+			FVector3d NewNormal = FMeshNormals::ComputeOverlayNormal(*Mesh, Normals, elemid);
+			Normals->SetElement(elemid, (FVector3f)NewNormal);
+			NormalsVertexFlags[elemid] = false;
+		});
+	}
+
+}
+
+
+
+
+
+
+
+void UDynamicMeshSculptTool::UpdateTarget()
+{
+	BrushTargetMesh.Copy(*DynamicMeshComponent->GetMesh(), false, false, false, false);
+	BrushTargetMeshSpatial.SetMesh(&BrushTargetMesh, true);
+
+	BrushTargetNormals.SetMesh(&BrushTargetMesh);
+	BrushTargetNormals.ComputeVertexNormals();
+}
+
+bool UDynamicMeshSculptTool::GetTargetMeshNearest(const FVector3d& Position, double SearchRadius, FVector3d& TargetPosOut, FVector3d& TargetNormalOut)
+{
+	double fDistSqr;
+	int NearTID = BrushTargetMeshSpatial.FindNearestTriangle(Position, fDistSqr, SearchRadius);
+	if (NearTID <= 0)
+	{
+		return false;
+	}
+	FTriangle3d Triangle;
+	BrushTargetMesh.GetTriVertices(NearTID, Triangle.V[0], Triangle.V[1], Triangle.V[2]);
+	FDistPoint3Triangle3d Query(Position, Triangle);
+	Query.Get();
+	FIndex3i Tri = BrushTargetMesh.GetTriangle(NearTID);
+	TargetNormalOut =
+		Query.TriangleBaryCoords.X*BrushTargetNormals[Tri.A]
+		+ Query.TriangleBaryCoords.Y*BrushTargetNormals[Tri.B]
+		+ Query.TriangleBaryCoords.Z*BrushTargetNormals[Tri.C];
+	TargetNormalOut.Normalize();
+	TargetPosOut = Query.ClosestTrianglePoint;
+	return true;
+}
+
+
+double UDynamicMeshSculptTool::EstimateIntialSafeTargetLength(const FDynamicMesh3& Mesh, int MinTargetTriCount)
+{
+	double AreaSum = 0;
+	for (int tid : Mesh.TriangleIndicesItr())
+	{
+		AreaSum += Mesh.GetTriArea(tid);
+	}
+
+	int TriCount = Mesh.TriangleCount();
+	double TargetTriArea = 1.0;
+	if (TriCount < MinTargetTriCount)
+	{
+		TargetTriArea = AreaSum / (double)MinTargetTriCount;
+	}
+	else
+	{
+		TargetTriArea = AreaSum / (double)TriCount;
+	}
+
+	double EdgeLen = TriangleUtil::EquilateralEdgeLengthForArea(TargetTriArea);
+	return (double)FMath::RoundToInt(EdgeLen*100.0) / 100.0;
+}
+
+
+
+
+void UDynamicMeshSculptTool::IncreaseBrushRadiusAction()
+{
+	BrushProperties->BrushSize = FMath::Clamp(BrushProperties->BrushSize + 0.025f, 0.0f, 1.0f);
+	CalculateBrushRadius();
+}
+
+void UDynamicMeshSculptTool::DecreaseBrushRadiusAction()
+{
+	BrushProperties->BrushSize = FMath::Clamp(BrushProperties->BrushSize - 0.025f, 0.0f, 1.0f);
+	CalculateBrushRadius();
+}
+
+void UDynamicMeshSculptTool::RegisterActions(FInteractiveToolActionSet& ActionSet)
+{
+	ActionSet.RegisterAction(this, (int32)EStandardToolActions::IncreaseBrushSize,
+		TEXT("SculptIncreaseRadius"), 
+		LOCTEXT("SculptIncreaseRadius", "Increase Sculpt Radius"),
+		LOCTEXT("SculptIncreaseRadiusTooltip", "Increase radius of sculpting brush"),
+		EModifierKey::None, EKeys::RightBracket,
+		[this]() { IncreaseBrushRadiusAction(); } );
+
+	ActionSet.RegisterAction(this, (int32)EStandardToolActions::DecreaseBrushSize,
+		TEXT("SculptDecreaseRadius"), 
+		LOCTEXT("SculptDecreaseRadius", "Decrease Sculpt Radius"),
+		LOCTEXT("SculptDecreaseRadiusTooltip", "Decrease radius of sculpting brush"),
+		EModifierKey::None, EKeys::LeftBracket,
+		[this]() { DecreaseBrushRadiusAction(); } );
+}
+
+
+
+
+//
+// Change Tracking
+//
+
+
+void UDynamicMeshSculptTool::BeginChange(bool bIsVertexChange)
+{
+	check(ActiveVertexChange == nullptr);
+	check(ActiveMeshChange == nullptr);
+	if (bIsVertexChange) 
+	{
+		ActiveVertexChange = new FMeshVertexChangeBuilder();
+	}
+	else
+	{
+		ActiveMeshChange = new FDynamicMeshChangeTracker(DynamicMeshComponent->GetMesh());
+		ActiveMeshChange->BeginChange();
+	}
+}
+
+
+void UDynamicMeshSculptTool::EndChange()
+{
+	if (ActiveVertexChange != nullptr)
+	{
+		GetToolManager()->EmitObjectChange(DynamicMeshComponent, MoveTemp(ActiveVertexChange->Change), LOCTEXT("VertexSculptChange", "Brush Stroke"));
+
+		delete ActiveVertexChange;
+		ActiveVertexChange = nullptr;
+	}
+
+	if (ActiveMeshChange != nullptr)
+	{
+		FMeshChange* NewMeshChange = new FMeshChange();
+		NewMeshChange->DynamicMeshChange = ActiveMeshChange->EndChange();
+		//NewMeshChange->DynamicMeshChange->CheckValidity();
+		TUniquePtr<FMeshChange> NewChange(NewMeshChange);
+		GetToolManager()->EmitObjectChange(DynamicMeshComponent, MoveTemp(NewChange), LOCTEXT("MeshSculptChange", "Brush Stroke"));
+		delete ActiveMeshChange;
+		ActiveMeshChange = nullptr;
+	}
+}
+
+
+void UDynamicMeshSculptTool::SaveActiveROI()
+{
+	if (ActiveMeshChange != nullptr)
+	{
+		for (int vid : VertexROI)
+		{
+			ActiveMeshChange->SaveVertex(vid);
+		}
+	}
+}
+
+
+void UDynamicMeshSculptTool::UpdateSavedVertex(int vid, const FVector3d& OldPosition, const FVector3d& NewPosition)
+{
+	if (ActiveVertexChange != nullptr)
+	{
+		ActiveVertexChange->UpdateVertex(vid, OldPosition, NewPosition);
+	}
+}
+
+
+#undef LOCTEXT_NAMESPACE
diff --git a/Engine/Plugins/Experimental/MeshModelingToolset/Source/MeshModelingTools/Private/EditNormalsTool.cpp b/Engine/Plugins/Experimental/MeshModelingToolset/Source/MeshModelingTools/Private/EditNormalsTool.cpp
new file mode 100644
--- /dev/null
+++ b/Engine/Plugins/Experimental/MeshModelingToolset/Source/MeshModelingTools/Private/EditNormalsTool.cpp
@@ -0,0 +1,285 @@
+// Copyright 1998-2019 Epic Games, Inc. All Rights Reserved.
+
+#include "EditNormalsTool.h"
+#include "InteractiveToolManager.h"
+#include "ToolBuilderUtil.h"
+
+#include "ToolSetupUtil.h"
+
+#include "DynamicMesh3.h"
+#include "BaseBehaviors/MultiClickSequenceInputBehavior.h"
+#include "Selection/SelectClickedAction.h"
+
+#include "MeshDescriptionToDynamicMesh.h"
+#include "DynamicMeshToMeshDescription.h"
+
+#include "InteractiveGizmoManager.h"
+
+#include "AssetGenerationUtil.h"
+
+
+#define LOCTEXT_NAMESPACE "UEditNormalsTool"
+
+
+/*
+ * ToolBuilder
+ */
+
+
+bool UEditNormalsToolBuilder::CanBuildTool(const FToolBuilderState& SceneState) const
+{
+	return ToolBuilderUtil::CountComponents(SceneState, CanMakeComponentTarget) > 0;
+}
+
+UInteractiveTool* UEditNormalsToolBuilder::BuildTool(const FToolBuilderState& SceneState) const
+{
+	UEditNormalsTool* NewTool = NewObject<UEditNormalsTool>(SceneState.ToolManager);
+
+	TArray<UActorComponent*> Components = ToolBuilderUtil::FindAllComponents(SceneState, CanMakeComponentTarget);
+	check(Components.Num() > 0);
+
+	TArray<TUniquePtr<FPrimitiveComponentTarget>> ComponentTargets;
+	for (UActorComponent* ActorComponent : Components)
+	{
+		auto* MeshComponent = Cast<UPrimitiveComponent>(ActorComponent);
+		if ( MeshComponent )
+		{
+			ComponentTargets.Add(MakeComponentTarget(MeshComponent));
+		}
+	}
+
+	NewTool->SetSelection(MoveTemp(ComponentTargets));
+	NewTool->SetWorld(SceneState.World);
+	NewTool->SetAssetAPI(AssetAPI);
+
+	return NewTool;
+}
+
+
+
+
+/*
+ * Tool
+ */
+
+UEditNormalsToolProperties::UEditNormalsToolProperties()
+{
+	bInvertNormals = false;
+	bRecomputeNormals = true;
+	NormalCalculationMethod = ENormalCalculationMethod::AreaAngleWeighting;
+	bRecomputeNormalTopologyAndEdgeSharpness = false;
+	SharpEdgeAngleThreshold = 60;
+	bAllowSharpVertices = true;
+}
+
+UEditNormalsAdvancedProperties::UEditNormalsAdvancedProperties()
+{
+}
+
+
+UEditNormalsTool::UEditNormalsTool()
+{
+}
+
+void UEditNormalsTool::SetWorld(UWorld* World)
+{
+	this->TargetWorld = World;
+}
+
+void UEditNormalsTool::Setup()
+{
+	UInteractiveTool::Setup();
+
+	// hide input StaticMeshComponent
+	for (TUniquePtr<FPrimitiveComponentTarget>& ComponentTarget : ComponentTargets)
+	{
+		ComponentTarget->SetOwnerVisibility(false);
+	}
+
+	BasicProperties = NewObject<UEditNormalsToolProperties>(this, TEXT("Mesh Normals Settings"));
+	AdvancedProperties = NewObject<UEditNormalsAdvancedProperties>(this, TEXT("Advanced Settings"));
+
+	// initialize our properties
+	AddToolPropertySource(BasicProperties);
+	AddToolPropertySource(AdvancedProperties);
+
+	// initialize the PreviewMesh+BackgroundCompute object
+	UpdateNumPreviews();
+
+	
+	
+	for (UMeshOpPreviewWithBackgroundCompute* Preview : Previews)
+	{
+		Preview->InvalidateResult();
+	}
+}
+
+
+void UEditNormalsTool::UpdateNumPreviews()
+{
+	int32 CurrentNumPreview = Previews.Num();
+	int32 TargetNumPreview = ComponentTargets.Num();
+	if (TargetNumPreview < CurrentNumPreview)
+	{
+		for (int32 PreviewIdx = CurrentNumPreview - 1; PreviewIdx >= TargetNumPreview; PreviewIdx--)
+		{
+			Previews[PreviewIdx]->Cancel();
+		}
+		Previews.SetNum(TargetNumPreview);
+		OriginalDynamicMeshes.SetNum(TargetNumPreview);
+	}
+	else
+	{
+		OriginalDynamicMeshes.SetNum(TargetNumPreview);
+		for (int32 PreviewIdx = CurrentNumPreview; PreviewIdx < TargetNumPreview; PreviewIdx++)
+		{
+			UEditNormalsOperatorFactory *OpFactory = NewObject<UEditNormalsOperatorFactory>();
+			OpFactory->Tool = this;
+			OpFactory->ComponentIndex = PreviewIdx;
+			OriginalDynamicMeshes[PreviewIdx] = MakeShared<FDynamicMesh3>();
+			FMeshDescriptionToDynamicMesh Converter;
+			Converter.bPrintDebugMessages = true;
+			Converter.Convert(ComponentTargets[PreviewIdx]->GetMesh(), *OriginalDynamicMeshes[PreviewIdx]);
+
+			UMeshOpPreviewWithBackgroundCompute* Preview = Previews.Add_GetRef(NewObject<UMeshOpPreviewWithBackgroundCompute>(OpFactory, "Preview"));
+			Preview->Setup(this->TargetWorld, OpFactory);
+			Preview->ConfigureMaterials(
+				ToolSetupUtil::GetDefaultMaterial(GetToolManager(), ComponentTargets[PreviewIdx]->GetMaterial(0)),
+				ToolSetupUtil::GetDefaultWorkingMaterial(GetToolManager())
+			);
+			Preview->SetVisibility(true);
+		}
+	}
+}
+
+
+void UEditNormalsTool::Shutdown(EToolShutdownType ShutdownType)
+{
+	// Restore (unhide) the source meshes
+	for (TUniquePtr<FPrimitiveComponentTarget>& ComponentTarget : ComponentTargets)
+	{
+		ComponentTarget->SetOwnerVisibility(true);
+	}
+
+	TArray<TUniquePtr<FDynamicMeshOpResult>> Results;
+	for (UMeshOpPreviewWithBackgroundCompute* Preview : Previews)
+	{
+		Results.Emplace(Preview->Shutdown());
+	}
+	if (ShutdownType == EToolShutdownType::Accept)
+	{
+		GenerateAsset(Results);
+	}
+}
+
+void UEditNormalsTool::SetAssetAPI(IToolsContextAssetAPI* AssetAPIIn)
+{
+	this->AssetAPI = AssetAPIIn;
+}
+
+TSharedPtr<FDynamicMeshOperator> UEditNormalsOperatorFactory::MakeNewOperator()
+{
+	TSharedPtr<FEditNormalsOp> NormalsOp = MakeShared<FEditNormalsOp>();
+	NormalsOp->bInvertNormals = Tool->BasicProperties->bInvertNormals;
+	NormalsOp->bRecomputeNormals = Tool->BasicProperties->bRecomputeNormals;
+	NormalsOp->bSplitNormals = Tool->BasicProperties->bRecomputeNormalTopologyAndEdgeSharpness;
+	NormalsOp->bAllowSharpVertices = Tool->BasicProperties->bAllowSharpVertices;
+	NormalsOp->NormalCalculationMethod = Tool->BasicProperties->NormalCalculationMethod;
+	NormalsOp->NormalSplitThreshold = Tool->BasicProperties->SharpEdgeAngleThreshold;
+
+	FTransform LocalToWorld = Tool->ComponentTargets[ComponentIndex]->GetWorldTransform();
+	NormalsOp->OriginalMesh = Tool->OriginalDynamicMeshes[ComponentIndex];
+	
+	NormalsOp->SetTransform(LocalToWorld);
+
+	return NormalsOp;
+}
+
+
+
+void UEditNormalsTool::Render(IToolsContextRenderAPI* RenderAPI)
+{
+}
+
+void UEditNormalsTool::Tick(float DeltaTime)
+{
+	for (UMeshOpPreviewWithBackgroundCompute* Preview : Previews)
+	{
+		Preview->Tick(DeltaTime);
+	}
+}
+
+
+
+#if WITH_EDITOR
+void UEditNormalsTool::PostEditChangeProperty(FPropertyChangedEvent& PropertyChangedEvent)
+{
+	UpdateNumPreviews();
+	for (UMeshOpPreviewWithBackgroundCompute* Preview : Previews)
+	{
+		Preview->InvalidateResult();
+	}
+}
+#endif
+
+void UEditNormalsTool::OnPropertyModified(UObject* PropertySet, UProperty* Property)
+{
+	UpdateNumPreviews();
+	for (UMeshOpPreviewWithBackgroundCompute* Preview : Previews)
+	{
+		Preview->InvalidateResult();
+	}
+}
+
+
+
+bool UEditNormalsTool::HasAccept() const
+{
+	return true;
+}
+
+bool UEditNormalsTool::CanAccept() const
+{
+	for (UMeshOpPreviewWithBackgroundCompute* Preview : Previews)
+	{
+		if (!Preview->HaveValidResult())
+		{
+			return false;
+		}
+	}
+	return true;
+}
+
+
+void UEditNormalsTool::GenerateAsset(const TArray<TUniquePtr<FDynamicMeshOpResult>>& Results)
+{
+	GetToolManager()->BeginUndoTransaction(LOCTEXT("EditNormalsToolTransactionName", "Edit Normals Tool"));
+
+	check(Results.Num() == ComponentTargets.Num());
+	
+	for (int32 ComponentIdx = 0; ComponentIdx < ComponentTargets.Num(); ComponentIdx++)
+	{
+		check(Results[ComponentIdx]->Mesh.Get() != nullptr);
+		ComponentTargets[ComponentIdx]->CommitMesh([&Results, &ComponentIdx, this](FMeshDescription* MeshDescription)
+		{
+			FDynamicMeshToMeshDescription Converter;
+			if (BasicProperties->bRecomputeNormalTopologyAndEdgeSharpness)
+			{
+				// full conversion if normal topology changed
+				Converter.Convert(Results[ComponentIdx]->Mesh.Get(), *MeshDescription);
+			}
+			else
+			{
+				// otherwise just copy attributes
+				Converter.UpdateAttributes(Results[ComponentIdx]->Mesh.Get(), *MeshDescription, true, false);
+			}
+		});
+	}
+
+	GetToolManager()->EndUndoTransaction();
+}
+
+
+
+
+#undef LOCTEXT_NAMESPACE
