Mesh Modeling Tools moved from NotForLicensees to Experimental

diff --git a/Engine/Plugins/Experimental/MeshModelingToolset/Source/MeshModelingTools/Private/MakeMeshTestTool.cpp b/Engine/Plugins/Experimental/MeshModelingToolset/Source/MeshModelingTools/Private/MakeMeshTestTool.cpp
new file mode 100644
--- /dev/null
+++ b/Engine/Plugins/Experimental/MeshModelingToolset/Source/MeshModelingTools/Private/MakeMeshTestTool.cpp
@@ -0,0 +1,304 @@
+// Copyright 1998-2019 Epic Games, Inc. All Rights Reserved.
+
+#include "MakeMeshTestTool.h"
+#include "ToolBuilderUtil.h"
+#include "InteractiveToolManager.h"
+#include "BaseBehaviors/MouseHoverBehavior.h"
+#include "Selection/ToolSelectionUtil.h"
+#include "AssetGenerationUtil.h"
+
+#include "MeshDescriptionBuilder.h"
+#include "Generators/SweepGenerator.h"
+#include "Generators/GridBoxMeshGenerator.h"
+#include "Generators/RectangleMeshGenerator.h"
+#include "Generators/SphereGenerator.h"
+#include "DynamicMeshAttributeSet.h"
+
+#include "StaticMeshComponentBuilder.h"
+#include "Drawing/MeshDebugDrawing.h"
+
+#define LOCTEXT_NAMESPACE "UMakeMeshTestTool"
+
+
+/*
+ * ToolBuilder
+ */
+
+
+bool UMakeMeshTestToolBuilder::CanBuildTool(const FToolBuilderState& SceneState) const
+{
+	return (this->AssetAPI != nullptr);
+}
+
+UInteractiveTool* UMakeMeshTestToolBuilder::BuildTool(const FToolBuilderState& SceneState) const
+{
+	UMakeMeshTestTool* NewTool = NewObject<UMakeMeshTestTool>(SceneState.ToolManager);
+	NewTool->SetWorld(SceneState.World);
+	NewTool->SetAssetAPI(AssetAPI);
+	return NewTool;
+}
+
+
+
+/*
+ * Tool
+ */
+
+UProceduralShapeToolProperties::UProceduralShapeToolProperties()
+{
+	Shape = EMakeMeshShapeType::Box;
+	Width = 100;
+	Height = 200;
+	//StartAngle = 0;
+	//EndAngle = 360;
+	Slices = 16;
+	Subdivisions = 0;
+	bCentered = false;
+	PlaceMode = EMakeMeshPlacementType::OnScene;
+}
+
+
+void UMakeMeshTestTool::SetWorld(UWorld* World)
+{
+	this->TargetWorld = World;
+}
+
+void UMakeMeshTestTool::SetAssetAPI(IToolsContextAssetAPI* AssetAPIIn)
+{
+	this->AssetAPI = AssetAPIIn;
+}
+
+
+void UMakeMeshTestTool::Setup()
+{
+	USingleClickTool::Setup();
+
+	UMouseHoverBehavior* HoverBehavior = NewObject<UMouseHoverBehavior>(this);
+	HoverBehavior->Initialize(this);
+	AddInputBehavior(HoverBehavior);
+
+	ShapeSettings = NewObject<UProceduralShapeToolProperties>(this);
+	AddToolPropertySource(ShapeSettings);
+
+	// create preview mesh object
+	PreviewMesh = NewObject<UPreviewMesh>(this, TEXT("PreviewMesh"));
+	PreviewMesh->CreateInWorld(TargetWorld, FTransform::Identity);
+	PreviewMesh->SetVisible(false);
+	UpdatePreviewMesh();
+
+	//if (MeshSource->GetMaterial(0) != nullptr)
+	//{
+	//	PreviewMesh->SetMaterial(MeshSource->GetMaterial(0));
+	//}
+}
+
+
+void UMakeMeshTestTool::Shutdown(EToolShutdownType ShutdownType)
+{
+	PreviewMesh->SetVisible(false);
+	PreviewMesh->Disconnect();
+	PreviewMesh = nullptr;
+}
+
+
+void UMakeMeshTestTool::Render(IToolsContextRenderAPI* RenderAPI)
+{
+	//FPrimitiveDrawInterface* PDI = RenderAPI->GetPrimitiveDrawInterface();
+	//MeshDebugDraw::DrawSimpleGrid(ShapeFrame, 13, 5.0f, 1.0f, FColor::Orange, false, PDI, FTransform::Identity);
+}
+
+
+
+void UMakeMeshTestTool::OnPropertyModified(UObject* PropertySet, UProperty* Property)
+{
+	UpdatePreviewMesh();
+}
+
+
+
+
+void UMakeMeshTestTool::OnUpdateHover(const FInputDeviceRay& DevicePos)
+{
+	UpdatePreviewPosition(DevicePos);
+}
+
+
+void UMakeMeshTestTool::UpdatePreviewPosition(const FInputDeviceRay& DeviceClickPos)
+{
+	FRay ClickPosWorldRay = DeviceClickPos.WorldRay;
+
+	// hit position (temp)
+	bool bHit = false;
+
+	if (ShapeSettings->PlaceMode == EMakeMeshPlacementType::OnPlane)
+	{
+		FPlane DrawPlane(FVector::ZeroVector, FVector(0, 0, 1));
+		FVector DrawPlanePos = FMath::RayPlaneIntersection(ClickPosWorldRay.Origin, ClickPosWorldRay.Direction, DrawPlane);
+		bHit = true;
+		ShapeFrame = FFrame3f(DrawPlanePos);
+	}
+	else
+	{
+		// cast ray into scene
+		FVector RayStart = ClickPosWorldRay.Origin;
+		FVector RayEnd = ClickPosWorldRay.PointAt(999999);
+		FCollisionObjectQueryParams QueryParams(FCollisionObjectQueryParams::AllObjects);
+		FHitResult Result;
+		bHit = TargetWorld->LineTraceSingleByObjectType(Result, RayStart, RayEnd, QueryParams);
+		if (bHit)
+		{
+			ShapeFrame = FFrame3f(Result.ImpactPoint, Result.ImpactNormal);
+			ShapeFrame.ConstrainedAlignPerpAxes();
+		}
+	}
+
+	if (ShapeSettings->Rotation != 0)
+	{
+		ShapeFrame.Rotate(FQuaternionf(ShapeFrame.Z(), ShapeSettings->Rotation, true));
+	}
+
+	if (bHit)
+	{
+		PreviewMesh->SetVisible(true);
+		PreviewMesh->SetTransform(ShapeFrame.ToFTransform());
+	}
+	else
+	{
+		PreviewMesh->SetVisible(false);
+	}
+}
+
+
+void UMakeMeshTestTool::UpdatePreviewMesh()
+{
+	FDynamicMesh3 NewMesh;
+	switch (ShapeSettings->Shape)
+	{
+	case EMakeMeshShapeType::Plane:
+		GeneratePlane(&NewMesh);
+		break;
+
+	case EMakeMeshShapeType::Cylinder:
+		GenerateCylinder(&NewMesh);
+		break;
+
+	case EMakeMeshShapeType::Cone:
+		GenerateCone(&NewMesh);
+		break;
+
+	case EMakeMeshShapeType::Sphere:
+		GenerateSphere(&NewMesh);
+		break;
+
+	case EMakeMeshShapeType::Box:
+	default:
+		GenerateBox(&NewMesh);
+		break;
+	}
+
+	// set mesh position
+	FAxisAlignedBox3d Bounds = NewMesh.GetCachedBounds();
+	FVector3d TargetOrigin = Bounds.Center();
+	if (ShapeSettings->bCentered == false)
+	{
+		TargetOrigin.Z = Bounds.Min.Z;
+	}
+	for (int vid : NewMesh.VertexIndicesItr())
+	{
+		FVector3d Pos = NewMesh.GetVertex(vid);
+		Pos -= TargetOrigin;
+		NewMesh.SetVertex(vid, Pos);
+	}
+
+
+	PreviewMesh->UpdatePreview(&NewMesh);
+}
+
+
+
+
+void UMakeMeshTestTool::OnClicked(const FInputDeviceRay& DeviceClickPos)
+{
+#if WITH_EDITOR
+	const FDynamicMesh3* CurMesh = PreviewMesh->GetPreviewDynamicMesh();
+	FTransform3d CurTransform(PreviewMesh->GetTransform());
+
+	GetToolManager()->BeginUndoTransaction(LOCTEXT("MergeMeshes", "Merge Meshes"));
+
+	AActor* NewActor = AssetGenerationUtil::GenerateStaticMeshActor(
+		AssetAPI, TargetWorld,
+		CurMesh, CurTransform, TEXT("ShapeMesh"),
+		AssetGenerationUtil::GetDefaultAutoGeneratedAssetPath());
+	
+	// select newly-created object
+	ToolSelectionUtil::SetNewActorSelection(GetToolManager(), NewActor);
+
+	GetToolManager()->EndUndoTransaction();
+#else
+	checkNoEntry();
+#endif
+}
+
+
+
+void UMakeMeshTestTool::GenerateBox(FDynamicMesh3* OutMesh)
+{
+	FGridBoxMeshGenerator BoxGen;
+	BoxGen.Box = FOrientedBox3d(FVector3d::Zero(), 0.5*FVector3d(ShapeSettings->Width, ShapeSettings->Width, ShapeSettings->Height));
+	int EdgeNum = ShapeSettings->Subdivisions + 2;
+	BoxGen.EdgeVertices = FIndex3i(EdgeNum, EdgeNum, EdgeNum);
+	BoxGen.Generate();
+	OutMesh->Copy(&BoxGen);
+}
+
+
+void UMakeMeshTestTool::GeneratePlane(FDynamicMesh3* OutMesh)
+{
+	FRectangleMeshGenerator RectGen;
+	RectGen.Width = ShapeSettings->Width;
+	RectGen.Height = ShapeSettings->Height;
+	RectGen.WidthVertexCount = RectGen.HeightVertexCount = ShapeSettings->Subdivisions + 2;
+	RectGen.Generate();
+	OutMesh->Copy(&RectGen);
+}
+
+
+void UMakeMeshTestTool::GenerateCylinder(FDynamicMesh3* OutMesh)
+{
+	FCylinderGenerator CylGen;
+	CylGen.Radius[0] = ShapeSettings->Width * 0.5f;
+	CylGen.Radius[1] = CylGen.Radius[0];
+	CylGen.Height = ShapeSettings->Height;
+	CylGen.AngleSamples = ShapeSettings->Slices;
+	CylGen.bCapped = true;
+	CylGen.Generate();
+	OutMesh->Copy(&CylGen);
+}
+
+
+void UMakeMeshTestTool::GenerateCone(FDynamicMesh3* OutMesh)
+{
+	// Unreal's standard cone is just a cylinder with a very small top
+	FCylinderGenerator CylGen;
+	CylGen.Radius[0] = ShapeSettings->Width * 0.5f;
+	CylGen.Radius[1] = .01;
+	CylGen.Height = ShapeSettings->Height;
+	CylGen.AngleSamples = ShapeSettings->Slices;
+	CylGen.bCapped = true;
+	CylGen.Generate();
+	OutMesh->Copy(&CylGen);
+}
+
+
+void UMakeMeshTestTool::GenerateSphere(FDynamicMesh3* OutMesh)
+{
+	FSphereGenerator SphereGen;
+	SphereGen.Radius = ShapeSettings->Width * 0.5f;
+	SphereGen.NumTheta = ShapeSettings->Slices;
+	SphereGen.NumPhi = ShapeSettings->Slices;
+	SphereGen.Generate();
+	OutMesh->Copy(&SphereGen);
+}
+
+
+#undef LOCTEXT_NAMESPACE
diff --git a/Engine/Plugins/Experimental/MeshModelingToolset/Source/MeshModelingTools/Private/MeshInspectorTool.cpp b/Engine/Plugins/Experimental/MeshModelingToolset/Source/MeshModelingTools/Private/MeshInspectorTool.cpp
new file mode 100644
--- /dev/null
+++ b/Engine/Plugins/Experimental/MeshModelingToolset/Source/MeshModelingTools/Private/MeshInspectorTool.cpp
@@ -0,0 +1,332 @@
+// Copyright 1998-2019 Epic Games, Inc. All Rights Reserved.
+
+#include "MeshInspectorTool.h"
+#include "InteractiveToolManager.h"
+#include "ToolBuilderUtil.h"
+
+#include "DynamicMesh3.h"
+#include "DynamicMeshAttributeSet.h"
+#include "MeshNormals.h"
+
+#include "SimpleDynamicMeshComponent.h"
+#include "Materials/MaterialInstanceDynamic.h"
+
+#include "SceneManagement.h" // for FPrimitiveDrawInterface
+
+#define LOCTEXT_NAMESPACE "UMeshInspectorTool"
+
+
+/*
+ * ToolBuilder
+ */
+
+
+bool UMeshInspectorToolBuilder::CanBuildTool(const FToolBuilderState& SceneState) const
+{
+	return ToolBuilderUtil::CountComponents(SceneState, CanMakeComponentTarget) == 1;
+}
+
+UInteractiveTool* UMeshInspectorToolBuilder::BuildTool(const FToolBuilderState& SceneState) const
+{
+	UMeshInspectorTool* NewTool = NewObject<UMeshInspectorTool>(SceneState.ToolManager);
+
+	UActorComponent* ActorComponent = ToolBuilderUtil::FindFirstComponent(SceneState, CanMakeComponentTarget);
+	auto* MeshComponent = Cast<UPrimitiveComponent>(ActorComponent);
+	check(MeshComponent != nullptr);
+	NewTool->SetSelection(MakeComponentTarget(MeshComponent));
+
+	return NewTool;
+}
+
+
+
+/*
+ * Tool
+ */
+
+UMeshInspectorTool::UMeshInspectorTool()
+{
+}
+
+void UMeshInspectorTool::Setup()
+{
+	UInteractiveTool::Setup();
+
+	// create dynamic mesh component to use for live preview
+	DynamicMeshComponent = NewObject<USimpleDynamicMeshComponent>(ComponentTarget->GetOwnerActor(), "DynamicMesh");
+	DynamicMeshComponent->SetupAttachment(ComponentTarget->GetOwnerActor()->GetRootComponent());
+	DynamicMeshComponent->RegisterComponent();
+	DynamicMeshComponent->SetWorldTransform(ComponentTarget->GetWorldTransform());
+
+	// copy material if there is one
+	DefaultMaterial = ComponentTarget->GetMaterial(0);
+	if (DefaultMaterial != nullptr)
+	{
+		DynamicMeshComponent->SetMaterial(0, DefaultMaterial);
+	}
+	ActiveMaterialMode = EInspectorMaterialMode::Default;
+
+	DynamicMeshComponent->TangentsType = EDynamicMeshTangentCalcType::ExternallyCalculated;
+	DynamicMeshComponent->InitializeMesh(ComponentTarget->GetMesh());
+
+	Precompute();
+
+	// hide input StaticMeshComponent
+	ComponentTarget->SetOwnerVisibility(false);
+
+	// initialize our properties
+	Settings = NewObject<UMeshInspectorProperties>(this);
+	AddToolPropertySource(Settings);
+
+	UMaterial* CheckerMaterialBase = LoadObject<UMaterial>(nullptr, TEXT("/MeshModelingToolset/Materials/CheckerMaterial"));
+	if (CheckerMaterialBase != nullptr)
+	{
+		CheckerMaterial = UMaterialInstanceDynamic::Create(CheckerMaterialBase, NULL);
+		if (CheckerMaterial != nullptr)
+		{
+			CheckerMaterial->SetScalarParameterValue("Density", Settings->CheckerDensity);
+		}
+	}
+	ActiveCheckerDensity = Settings->CheckerDensity;
+
+	DynamicMeshComponent->bExplicitShowWireframe = Settings->bWireframe;
+
+	UMeshStatisticsProperties* Statistics = NewObject<UMeshStatisticsProperties>(this);
+	Statistics->Update(*DynamicMeshComponent->GetMesh());
+	AddToolPropertySource(Statistics);
+}
+
+
+void UMeshInspectorTool::Shutdown(EToolShutdownType ShutdownType)
+{
+	if (DynamicMeshComponent != nullptr)
+	{
+		ComponentTarget->SetOwnerVisibility(true);
+
+		DynamicMeshComponent->UnregisterComponent();
+		DynamicMeshComponent->DestroyComponent();
+		DynamicMeshComponent = nullptr;
+	}
+}
+
+
+void UMeshInspectorTool::Precompute()
+{
+	BoundaryEdges.Reset();
+	UVSeamEdges.Reset();
+	NormalSeamEdges.Reset();
+
+	FDynamicMesh3* TargetMesh = DynamicMeshComponent->GetMesh();
+	FDynamicMeshUVOverlay* UVOverlay =
+		TargetMesh->HasAttributes() ? TargetMesh->Attributes()->PrimaryUV() : nullptr;
+	FDynamicMeshNormalOverlay* NormalOverlay =
+		TargetMesh->HasAttributes() ? TargetMesh->Attributes()->PrimaryNormals() : nullptr;
+
+	for (int eid : TargetMesh->EdgeIndicesItr())
+	{
+		if (TargetMesh->IsBoundaryEdge(eid))
+		{
+			BoundaryEdges.Add(eid);
+		}
+		if (UVOverlay != nullptr && UVOverlay->IsSeamEdge(eid))
+		{
+			UVSeamEdges.Add(eid);
+		}
+		if (NormalOverlay != nullptr && NormalOverlay->IsSeamEdge(eid))
+		{
+			NormalSeamEdges.Add(eid);
+		}
+		if (TargetMesh->IsGroupBoundaryEdge(eid))
+		{
+			GroupBoundaryEdges.Add(eid);
+		}
+	}
+}
+
+
+void UMeshInspectorTool::Render(IToolsContextRenderAPI* RenderAPI)
+{
+	FPrimitiveDrawInterface* PDI = RenderAPI->GetPrimitiveDrawInterface();
+	FTransform Transform = ComponentTarget->GetWorldTransform();
+
+	FColor WireframeColor(40, 40, 40);
+	float WireframeThickness = LineWidthMultiplier * 0.5;
+	FColor BoundaryEdgeColor(240, 15, 15);
+	float BoundaryEdgeThickness = LineWidthMultiplier * 4.0;
+	FColor UVSeamColor(15, 240, 15);
+	float UVSeamThickness = LineWidthMultiplier * 2.0;
+	FColor NormalSeamColor(15, 240, 240);
+	float NormalSeamThickness = LineWidthMultiplier * 2.0;
+	FColor PolygonBorderColor(240, 15, 240);
+	float PolygonBorderThickness = LineWidthMultiplier * 2.0;
+	FColor NormalColor(15, 15, 240);
+	float NormalThickness = LineWidthMultiplier * 2.0f;
+	FColor TangentColor(240, 15, 15);
+	FColor BinormalColor(15, 240, 15);
+	float TangentThickness = LineWidthMultiplier * 2.0f;
+
+	FDynamicMesh3* TargetMesh = DynamicMeshComponent->GetMesh();
+	FVector3d A, B;
+
+	if (Settings->bBoundaryEdges)
+	{
+		for (int eid : BoundaryEdges)
+		{
+			TargetMesh->GetEdgeV(eid, A, B);
+			PDI->DrawLine(Transform.TransformPosition(A), Transform.TransformPosition(B),
+				BoundaryEdgeColor, 0, BoundaryEdgeThickness, 2.0f, true);
+		}
+	}
+
+	if (Settings->bUVSeams)
+	{
+		for (int eid : UVSeamEdges)
+		{
+			TargetMesh->GetEdgeV(eid, A, B);
+			PDI->DrawLine(Transform.TransformPosition(A), Transform.TransformPosition(B),
+				UVSeamColor, 0, UVSeamThickness, 3.0f, true);
+		}
+	}
+
+	if (Settings->bNormalSeams)
+	{
+		for (int eid : NormalSeamEdges)
+		{
+			TargetMesh->GetEdgeV(eid, A, B);
+			PDI->DrawLine(Transform.TransformPosition(A), Transform.TransformPosition(B),
+				NormalSeamColor, 0, NormalSeamThickness, 3.0f, true);
+		}
+	}
+
+	if (Settings->bPolygonBorders)
+	{
+		for (int eid : GroupBoundaryEdges)
+		{
+			TargetMesh->GetEdgeV(eid, A, B);
+			PDI->DrawLine(Transform.TransformPosition(A), Transform.TransformPosition(B),
+				PolygonBorderColor, 0, PolygonBorderThickness, 2.0f, true);
+		}
+	}
+
+	if (Settings->bNormalVectors && TargetMesh->HasAttributes() && TargetMesh->Attributes()->PrimaryNormals() != nullptr)
+	{
+		const FDynamicMeshNormalOverlay* NormalOverlay = TargetMesh->Attributes()->PrimaryNormals();
+		FVector3d TriV[3];
+		FVector3f TriN[3];
+		for (int tid : TargetMesh->TriangleIndicesItr())
+		{
+			TargetMesh->GetTriVertices(tid, TriV[0], TriV[1], TriV[2]);
+			NormalOverlay->GetTriElements(tid, TriN[0], TriN[1], TriN[2]);
+			for (int j = 0; j < 3; ++j)
+			{
+				TriV[j] = Transform.TransformPosition(TriV[j]);
+				PDI->DrawLine(TriV[j], TriV[j] + Settings->NormalLength * Transform.TransformVectorNoScale(TriN[j]),
+					NormalColor, SDPG_World, NormalThickness, 0.0f, true);
+			}
+		}
+	}
+
+	if (Settings->bTangentVectors)
+	{
+		const FMeshTangentsf* Tangents = DynamicMeshComponent->GetTangents();
+		for (int TID : TargetMesh->TriangleIndicesItr())
+		{
+			FVector3d TriV[3];
+			TargetMesh->GetTriVertices(TID, TriV[0], TriV[1], TriV[2]);
+			for (int SubIdx = 0; SubIdx < 3; SubIdx++)
+			{
+				FVector Vert = Transform.TransformPosition(TriV[SubIdx]);
+				FVector3f Tangent, Bitangent;
+				Tangents->GetPerTriangleTangent(TID, SubIdx, Tangent, Bitangent);
+				PDI->DrawLine(Vert, Vert + Settings->TangentLength * Transform.TransformVectorNoScale(Tangent),
+					TangentColor, SDPG_World, TangentThickness, 3.5f, true);
+				PDI->DrawLine(Vert, Vert + Settings->TangentLength * Transform.TransformVectorNoScale(Bitangent),
+					BinormalColor, SDPG_World, TangentThickness, 3.5f, true);
+			}
+		}
+	}
+
+
+}
+
+
+void UMeshInspectorTool::OnPropertyModified(UObject* PropertySet, UProperty* Property)
+{
+	GetToolManager()->PostInvalidation();
+	DynamicMeshComponent->bExplicitShowWireframe = Settings->bWireframe;
+
+	if (CheckerMaterial != nullptr)
+	{
+		CheckerMaterial->SetScalarParameterValue("Density", Settings->CheckerDensity);
+	}
+
+	if (Settings->MaterialMode != ActiveMaterialMode)
+	{
+		if (Settings->MaterialMode == EInspectorMaterialMode::Checkerboard && CheckerMaterial != nullptr)
+		{
+			DynamicMeshComponent->SetMaterial(0, CheckerMaterial);
+			ActiveMaterialMode = EInspectorMaterialMode::Checkerboard;
+			return;
+		}
+		if (Settings->MaterialMode == EInspectorMaterialMode::Override && Settings->OverrideMaterial != nullptr)
+		{
+			DynamicMeshComponent->SetMaterial(0, Settings->OverrideMaterial);
+			ActiveMaterialMode = EInspectorMaterialMode::Override;
+			return;
+		}
+		// default or fallback
+		if (DefaultMaterial != nullptr)
+		{
+			DynamicMeshComponent->SetMaterial(0, DefaultMaterial);
+		}
+		ActiveMaterialMode = EInspectorMaterialMode::Default;
+	}
+}
+
+
+bool UMeshInspectorTool::HasAccept() const
+{
+	return false;
+}
+
+bool UMeshInspectorTool::CanAccept() const
+{
+	return false;
+}
+
+
+
+void UMeshInspectorTool::RegisterActions(FInteractiveToolActionSet& ActionSet)
+{
+	ActionSet.RegisterAction(this, (int32)EStandardToolActions::BaseClientDefinedActionID + 1,
+		TEXT("IncreaseLineWidth"), 
+		LOCTEXT("IncreaseLineWidth", "Increase Line Width"),
+		LOCTEXT("IncreaseLineWidthTooltip", "Increase line width of rendering"),
+		EModifierKey::Shift, EKeys::Equals,
+		[this]() { IncreaseLineWidthAction(); });
+
+	ActionSet.RegisterAction(this, (int32)EStandardToolActions::BaseClientDefinedActionID + 2,
+		TEXT("DecreaseLineWidth"), 
+		LOCTEXT("DecreaseLineWidth", "Decrease Line Width"),
+		LOCTEXT("DecreaseLineWidthTooltip", "Decrease line width of rendering"),
+		EModifierKey::None, EKeys::Equals,
+		[this]() { DecreaseLineWidthAction(); });
+}
+
+
+
+void UMeshInspectorTool::IncreaseLineWidthAction()
+{
+	LineWidthMultiplier = LineWidthMultiplier * 1.25f;
+	GetToolManager()->PostInvalidation();
+}
+
+void UMeshInspectorTool::DecreaseLineWidthAction()
+{
+	LineWidthMultiplier = LineWidthMultiplier * (1.0f / 1.25f);
+	GetToolManager()->PostInvalidation();
+}
+
+
+
+#undef LOCTEXT_NAMESPACE
diff --git a/Engine/Plugins/Experimental/MeshModelingToolset/Source/MeshModelingTools/Private/MeshSelectionTool.cpp b/Engine/Plugins/Experimental/MeshModelingToolset/Source/MeshModelingTools/Private/MeshSelectionTool.cpp
new file mode 100644
--- /dev/null
+++ b/Engine/Plugins/Experimental/MeshModelingToolset/Source/MeshModelingTools/Private/MeshSelectionTool.cpp
@@ -0,0 +1,594 @@
+// Copyright 1998-2019 Epic Games, Inc. All Rights Reserved.
+
+#include "MeshSelectionTool.h"
+#include "InteractiveToolManager.h"
+#include "ToolBuilderUtil.h"
+#include "Drawing/MeshDebugDrawing.h"
+#include "DynamicMeshEditor.h"
+#include "DynamicMeshChangeTracker.h"
+#include "Changes/MeshChange.h"
+#include "AssetGenerationUtil.h"
+
+#define LOCTEXT_NAMESPACE "UMeshSelectionTool"
+
+/*
+ * ToolBuilder
+ */
+
+UMeshSurfacePointTool* UMeshSelectionToolBuilder::CreateNewTool(const FToolBuilderState& SceneState) const
+{
+	UMeshSelectionTool* SelectionTool = NewObject<UMeshSelectionTool>(SceneState.ToolManager);
+	SelectionTool->SetWorld(SceneState.World);
+	SelectionTool->SetAssetAPI(AssetAPI);
+	return SelectionTool;
+}
+
+
+/*
+ * Tool
+ */
+
+UMeshSelectionTool::UMeshSelectionTool()
+{
+}
+
+void UMeshSelectionTool::SetWorld(UWorld* World)
+{
+	this->TargetWorld = World;
+}
+
+void UMeshSelectionTool::SetAssetAPI(IToolsContextAssetAPI* AssetAPIIn)
+{
+	this->AssetAPI = AssetAPIIn;
+}
+
+
+
+void UMeshSelectionTool::Setup()
+{
+	UDynamicMeshBrushTool::Setup();
+
+	// use ourself as tool prop source for actions
+	AddToolPropertySource(this);
+
+	// enable wireframe on component
+	PreviewMesh->EnableWireframe(true);
+
+	// set vertex color material on base component so we can see selection
+	UMaterialInterface* VertexColorMat = 
+		GetToolManager()->GetContextQueriesAPI()->GetStandardMaterial(EStandardToolContextMaterials::VertexColorMaterial);
+	if (VertexColorMat != nullptr)
+	{
+		PreviewMesh->SetMaterial(VertexColorMat);
+	}
+
+	const FDynamicMesh3* Mesh = PreviewMesh->GetPreviewDynamicMesh();
+	SelectedVertices = TBitArray<>(false, Mesh->MaxVertexID());
+	SelectedTriangles = TBitArray<>(false, Mesh->MaxTriangleID());
+
+	this->Selection = NewObject<UMeshSelectionSet>(this);
+	Selection->GetOnModified().AddLambda([this](USelectionSet* SelectionObj)
+	{
+		OnExternalSelectionChange();
+	});
+
+}
+
+
+
+
+void UMeshSelectionTool::OnShutdown(EToolShutdownType ShutdownType)
+{
+	if (bHaveModifiedMesh && ShutdownType == EToolShutdownType::Accept)
+	{
+		// this block bakes the modified DynamicMeshComponent back into the StaticMeshComponent inside an undo transaction
+		GetToolManager()->BeginUndoTransaction(LOCTEXT("MeshSelectionToolTransactionName", "Edit Mesh"));
+
+		ComponentTarget->CommitMesh([=](FMeshDescription* MeshDescription)
+		{
+			PreviewMesh->Bake(MeshDescription, true);
+		});
+		GetToolManager()->EndUndoTransaction();
+	}
+}
+
+
+
+
+void UMeshSelectionTool::RegisterActions(FInteractiveToolActionSet& ActionSet)
+{
+	UDynamicMeshBrushTool::RegisterActions(ActionSet);
+
+	ActionSet.RegisterAction(this, (int32)EStandardToolActions::BaseClientDefinedActionID + 1,
+		TEXT("MeshSelectionToolDelete"),
+		LOCTEXT("MeshSelectionToolDelete", "Delete"),
+		LOCTEXT("MeshSelectionToolDeleteTooltip", "Delete Selected Elements"),
+		EModifierKey::None, EKeys::Delete,
+		[this]() { DeleteSelectedTriangles(); });
+}
+
+
+
+void UMeshSelectionTool::OnExternalSelectionChange()
+{
+	SelectedVertices.SetRange(0, SelectedVertices.Num(), false);
+	SelectedTriangles.SetRange(0, SelectedTriangles.Num(), false);
+
+	if (SelectionType == EMeshSelectionElementType::Vertex)
+	{
+		for (int VertIdx : Selection->Vertices)
+		{
+			SelectedVertices[VertIdx] = true;
+		}
+	}
+	else if (SelectionType == EMeshSelectionElementType::Face)
+	{
+		for (int FaceIdx : Selection->Faces)
+		{
+			SelectedTriangles[FaceIdx] = true;
+		}
+	}
+
+	OnSelectionUpdated();
+}
+
+
+
+
+void UMeshSelectionTool::OnBeginDrag(const FRay& WorldRay)
+{
+	UDynamicMeshBrushTool::OnBeginDrag(WorldRay);
+
+	PreviewBrushROI.Reset();
+	if (IsInBrushStroke())
+	{
+		bInRemoveStroke = GetShiftToggle();
+		BeginChange(bInRemoveStroke == false);
+		StartStamp = UBaseBrushTool::LastBrushStamp;
+		LastStamp = StartStamp;
+		bStampPending = true;
+	}
+}
+
+
+
+void UMeshSelectionTool::OnUpdateDrag(const FRay& WorldRay)
+{
+	UDynamicMeshBrushTool::OnUpdateDrag(WorldRay);
+	if (IsInBrushStroke())
+	{
+		LastStamp = UBaseBrushTool::LastBrushStamp;
+		bStampPending = true;
+	}
+}
+
+
+
+
+void UMeshSelectionTool::CalculateVertexROI(const FBrushStampData& Stamp, TArray<int>& VertexROI)
+{
+	FTransform Transform = ComponentTarget->GetWorldTransform();
+	FVector StampPosLocal = Transform.InverseTransformPosition(Stamp.WorldPosition);
+
+	// TODO: need dynamic vertex hash table!
+
+	float RadiusSqr = CurrentBrushRadius * CurrentBrushRadius;
+	const FDynamicMesh3* Mesh = PreviewMesh->GetPreviewDynamicMesh();
+	for (int VertIdx : Mesh->VertexIndicesItr())
+	{
+		FVector3d Position = Mesh->GetVertex(VertIdx);
+		if ((Position - StampPosLocal).SquaredLength() < RadiusSqr)
+		{
+			VertexROI.Add(VertIdx);
+		}
+	}
+}
+
+
+void UMeshSelectionTool::CalculateTriangleROI(const FBrushStampData& Stamp, TArray<int>& TriangleROI)
+{
+	FTransform Transform = ComponentTarget->GetWorldTransform();
+	FVector StampPosLocal = Transform.InverseTransformPosition(Stamp.WorldPosition);
+
+	float RadiusSqr = CurrentBrushRadius * CurrentBrushRadius;
+	const FDynamicMesh3* Mesh = PreviewMesh->GetPreviewDynamicMesh();
+	for (int TriIdx : Mesh->TriangleIndicesItr())
+	{
+		FVector3d Position = Mesh->GetTriCentroid(TriIdx);
+		if ((Position - StampPosLocal).SquaredLength() < RadiusSqr)
+		{
+			TriangleROI.Add(TriIdx);
+		}
+	}
+}
+
+
+
+
+static void UpdateList(TArray<int>& List, int Value, bool bAdd)
+{
+	if (bAdd)
+	{
+		List.Add(Value);
+	}
+	else
+	{
+		List.RemoveSwap(Value);
+	}
+}
+
+
+void UMeshSelectionTool::ApplyStamp(const FBrushStampData& Stamp)
+{
+	//const FDynamicMesh3* Mesh = PreviewMesh->GetPreviewDynamicMesh();
+
+	IndexBuf.Reset();
+
+	bool bDesiredValue = bInRemoveStroke ? false : true;
+
+	if (SelectionType == EMeshSelectionElementType::Face)
+	{
+		CalculateTriangleROI(Stamp, IndexBuf);
+		for (int TriIdx : IndexBuf)
+		{
+			if (SelectedTriangles[TriIdx] != bDesiredValue)
+			{
+				SelectedTriangles[TriIdx] = bDesiredValue;
+				UpdateList(Selection->Faces, TriIdx, bDesiredValue);
+				if (ActiveSelectionChange != nullptr)
+				{
+					ActiveSelectionChange->Add(TriIdx);
+				}
+			}
+		}
+	}
+	else
+	{
+		CalculateVertexROI(Stamp, IndexBuf);
+		for (int VertIdx : IndexBuf)
+		{
+			if (SelectedVertices[VertIdx] != bDesiredValue)
+			{
+				SelectedVertices[VertIdx] = bDesiredValue;
+				UpdateList(Selection->Vertices, VertIdx, bDesiredValue);
+				if (ActiveSelectionChange != nullptr)
+				{
+					ActiveSelectionChange->Add(VertIdx);
+				}
+			}
+		}
+	}
+
+	OnSelectionUpdated();
+}
+
+
+
+
+void UMeshSelectionTool::OnEndDrag(const FRay& Ray)
+{
+	UDynamicMeshBrushTool::OnEndDrag(Ray);
+
+	bInRemoveStroke = false;
+	bStampPending = false;
+
+	// close change record
+	TUniquePtr<FMeshSelectionChange> Change = EndChange();
+	GetToolManager()->EmitObjectChange(Selection, MoveTemp(Change), LOCTEXT("MeshSelectionChange", "Mesh Selection"));
+}
+
+
+void UMeshSelectionTool::OnUpdateHover(const FInputDeviceRay& DevicePos)
+{
+	UDynamicMeshBrushTool::OnUpdateHover(DevicePos);
+
+	// todo get rid of this redundant hit test!
+	FHitResult OutHit;
+	if ( UDynamicMeshBrushTool::HitTest(DevicePos.WorldRay, OutHit) )
+	{
+		PreviewBrushROI.Reset();
+		if (SelectionType == EMeshSelectionElementType::Face)
+		{
+			CalculateTriangleROI(LastBrushStamp, PreviewBrushROI);
+		}
+		else
+		{
+			CalculateVertexROI(LastBrushStamp, PreviewBrushROI);
+		}
+	}
+}
+
+
+
+
+
+
+void UMeshSelectionTool::OnSelectionUpdated()
+{
+	UpdateVisualization();
+}
+
+void UMeshSelectionTool::UpdateVisualization()
+{
+	if (SelectionType == EMeshSelectionElementType::Face)
+	{
+		PreviewMesh->SetTriangleColorFunction([this](int TriangleID)
+		{
+			return SelectedTriangles[TriangleID] ? FColor::Red : FColor::White;
+		}, UPreviewMesh::ERenderUpdateMode::FullUpdate);
+
+		//bool bIsFirst = (! BaseMeshComponent->TriangleColorFunc);
+		//BaseMeshComponent->TriangleColorFunc =
+		//if (bIsFirst)
+		//{
+		//	BaseMeshComponent->NotifyMeshUpdated();
+		//}
+		//else
+		//{
+		//	BaseMeshComponent->NotifyMeshUpdated();
+		//	//BaseMeshComponent->FastNotifyColorsUpdated();
+		//}
+	}
+	else
+	{
+		PreviewMesh->ClearTriangleColorFunction(UPreviewMesh::ERenderUpdateMode::FullUpdate);
+	}
+}
+
+
+
+
+void UMeshSelectionTool::Render(IToolsContextRenderAPI* RenderAPI)
+{
+	UDynamicMeshBrushTool::Render(RenderAPI);
+
+	FTransform WorldTransform = ComponentTarget->GetWorldTransform();
+	const FDynamicMesh3* Mesh = PreviewMesh->GetPreviewDynamicMesh();
+
+	if (SelectionType == EMeshSelectionElementType::Vertex)
+	{
+		MeshDebugDraw::DrawVertices(Mesh, Selection->Vertices,
+			12.0f, FColor::Orange, RenderAPI->GetPrimitiveDrawInterface(), WorldTransform);
+		MeshDebugDraw::DrawVertices(Mesh, PreviewBrushROI,
+			8.0f, FColor(40, 200, 40), RenderAPI->GetPrimitiveDrawInterface(), WorldTransform);
+	}
+	else
+	{
+		// drawn via material
+		MeshDebugDraw::DrawTriCentroids(Mesh, Selection->Faces,
+			12.0f, FColor::Green, RenderAPI->GetPrimitiveDrawInterface(), WorldTransform);
+		MeshDebugDraw::DrawTriCentroids(Mesh, PreviewBrushROI,
+			8.0f, FColor(40, 200, 40), RenderAPI->GetPrimitiveDrawInterface(), WorldTransform);
+	}
+}
+
+
+void UMeshSelectionTool::Tick(float DeltaTime)
+{
+	UDynamicMeshBrushTool::Tick(DeltaTime);
+	Indicators->Tick(DeltaTime);
+
+
+	if (bStampPending)
+	{
+		ApplyStamp(LastStamp);
+		bStampPending = false;
+	}
+}
+
+
+
+void UMeshSelectionTool::BeginChange(bool bAdding)
+{
+	check(ActiveSelectionChange == nullptr);
+	ActiveSelectionChange = new FMeshSelectionChangeBuilder(SelectionType, bAdding);
+}
+
+void UMeshSelectionTool::CancelChange()
+{
+	if (ActiveSelectionChange != nullptr)
+	{
+		delete ActiveSelectionChange;
+		ActiveSelectionChange = nullptr;
+	}
+}
+
+TUniquePtr<FMeshSelectionChange> UMeshSelectionTool::EndChange()
+{
+	check(ActiveSelectionChange);
+	if (ActiveSelectionChange != nullptr)
+	{
+		TUniquePtr<FMeshSelectionChange> Result = MoveTemp(ActiveSelectionChange->Change);
+		delete ActiveSelectionChange;
+		ActiveSelectionChange = nullptr;
+
+		return Result;
+	}
+	return TUniquePtr<FMeshSelectionChange>();
+}
+
+/**
+ * FCommandChangeSequence contains a list of FCommandChanges and associated target UObjects.
+ * The sequence of changes is applied atomically.
+ * @warning if the target weak UObject pointers become invalid, those changes are skipped. This likely leaves this in an undefined state.
+ */
+class FCommandChangeSequence : public FCommandChange
+{
+protected:
+	struct FChangeElem
+	{
+		TWeakObjectPtr<UObject> TargetObject;
+		TUniquePtr<FCommandChange> Change;
+	};
+
+	TArray<TSharedPtr<FChangeElem>> Sequence;
+
+public:
+	FCommandChangeSequence()
+	{
+	}
+
+	/** Add a change to the sequence */
+	void AppendChange(UObject* Target, TUniquePtr<FCommandChange> Change)
+	{
+		TSharedPtr<FChangeElem> Elem = MakeShared<FChangeElem>();
+		Elem->TargetObject = Target;
+		Elem->Change = MoveTemp(Change);
+		Sequence.Add(Elem);
+	}
+	
+	/** Apply sequence of changes in-order */
+	virtual void Apply(UObject* Object) override
+	{
+		for (int k = 0; k < Sequence.Num(); ++k)
+		{
+			TSharedPtr<FChangeElem> Elem = Sequence[k];
+			check(Elem->TargetObject.IsValid());
+			if (Elem->TargetObject.IsValid())
+			{
+				Elem->Change->Apply(Elem->TargetObject.Get());
+			}
+		}
+	}
+
+	/** Reverts sequence of changes in reverse-order */
+	virtual void Revert(UObject* Object) override
+	{
+		for (int k = Sequence.Num()-1; k >= 0; --k)
+		{
+			TSharedPtr<FChangeElem> Elem = Sequence[k];
+			check(Elem->TargetObject.IsValid());
+			if (Elem->TargetObject.IsValid())
+			{
+				Elem->Change->Revert(Elem->TargetObject.Get());
+			}
+		}
+	}
+
+	/** @return string describing this change sequenece */
+	virtual FString ToString() const override
+	{
+		FString Result = TEXT("FCommandChangeSequence: ");
+		for (int k = 0; k < Sequence.Num(); ++k)
+		{
+			Result = Result + Sequence[k]->Change->ToString() + TEXT(" ");
+		}
+		return Result;
+	}
+};
+
+
+
+void UMeshSelectionTool::DeleteSelectedTriangles()
+{
+	check(SelectionType == EMeshSelectionElementType::Face);
+	TArray<int32> SelectedFaces = Selection->GetElements(EMeshSelectionElementType::Face);
+	if (SelectedFaces.Num() == 0)
+	{
+		return;
+	}
+
+	TUniquePtr<FCommandChangeSequence> ChangeSeq = MakeUnique<FCommandChangeSequence>();
+
+	// clear current selection
+	BeginChange(false);
+	for (int tid : SelectedFaces)
+	{
+		ActiveSelectionChange->Add(tid);
+	}
+	Selection->RemoveIndices(EMeshSelectionElementType::Face, SelectedFaces);
+	TUniquePtr<FMeshSelectionChange> SelectionChange = EndChange();
+	ChangeSeq->AppendChange(Selection, MoveTemp(SelectionChange));
+
+	// delete triangles and emit delete triangles change
+	TUniquePtr<FMeshChange> MeshChange = PreviewMesh->TrackedEditMesh(
+		[&SelectedFaces](FDynamicMesh3& Mesh, FDynamicMeshChangeTracker& ChangeTracker)
+	{
+		FDynamicMeshEditor Editor(&Mesh);
+		Editor.RemoveTriangles(SelectedFaces, true, [&ChangeTracker](int TriangleID) { ChangeTracker.SaveTriangle(TriangleID, true); });
+	});
+	ChangeSeq->AppendChange(PreviewMesh, MoveTemp(MeshChange));
+
+	// emit combined change sequence
+	GetToolManager()->EmitObjectChange(this, MoveTemp(ChangeSeq), LOCTEXT("MeshSelectionToolDelete", "Delete Faces"));
+
+	OnSelectionUpdated();
+
+	bHaveModifiedMesh = true;
+}
+
+
+
+void AssignMaterial(AActor* ToActor, const TUniquePtr<FPrimitiveComponentTarget>& FromTarget)
+{
+	UMaterialInterface* Material = FromTarget->GetMaterial(0);
+	if (!Material)
+	{
+		return;
+	}
+
+	//if (Cast<AStaticMeshActor>(ToActor) != nullptr)
+	//{
+	//	UStaticMeshComponent* Component = Cast<AStaticMeshActor>(ToActor)->GetStaticMeshComponent();
+	//	if (Component)
+	//	{
+	//		Component->SetMaterial(0, Material);
+	//	}
+	//} 
+	//else
+	//{
+		USceneComponent* Component = ToActor->GetRootComponent();
+		if (Cast<UPrimitiveComponent>(Component) != nullptr)
+		{
+			Cast<UPrimitiveComponent>(Component)->SetMaterial(0, Material);
+		}
+	//}
+}
+
+
+void UMeshSelectionTool::SeparateSelectedTriangles()
+{
+#if WITH_EDITOR
+	// currently AssetGenerationUtil::GenerateStaticMeshActor only defined in editor
+
+	check(SelectionType == EMeshSelectionElementType::Face);
+	TArray<int32> SelectedFaces = Selection->GetElements(EMeshSelectionElementType::Face);
+	if (SelectedFaces.Num() == 0)
+	{
+		return;
+	}
+
+	const FDynamicMesh3* SourceMesh = PreviewMesh->GetPreviewDynamicMesh();
+	if (SelectedFaces.Num() == SourceMesh->TriangleCount())
+	{
+		return;		// don't separate entire mesh
+	}
+
+
+	// extract copy of triangles
+	FDynamicMesh3 SeparatedMesh;
+	SeparatedMesh.EnableAttributes();
+	FDynamicMeshEditor Editor(&SeparatedMesh);
+	FMeshIndexMappings Mappings; FDynamicMeshEditResult EditResult;
+	Editor.AppendTriangles(SourceMesh, SelectedFaces, Mappings, EditResult);
+
+	// emit new asset
+	FTransform3d Transform(PreviewMesh->GetTransform());
+	GetToolManager()->BeginUndoTransaction(LOCTEXT("MeshSelectionToolSeparate", "Separate"));
+	AActor* NewActor = AssetGenerationUtil::GenerateStaticMeshActor(
+		AssetAPI, TargetWorld, &SeparatedMesh, Transform, TEXT("Submesh"),
+		AssetGenerationUtil::GetDefaultAutoGeneratedAssetPath());
+	AssignMaterial(NewActor, ComponentTarget);
+	GetToolManager()->EndUndoTransaction();
+
+	// todo: undo won't remove this asset...
+
+	// delete selected triangles from this mesh
+	DeleteSelectedTriangles();
+#endif
+}
+
+
+
+
+#undef LOCTEXT_NAMESPACE
\ No newline at end of file
diff --git a/Engine/Plugins/Experimental/MeshModelingToolset/Source/MeshModelingTools/Private/MeshSpaceDeformerTool.cpp b/Engine/Plugins/Experimental/MeshModelingToolset/Source/MeshModelingTools/Private/MeshSpaceDeformerTool.cpp
new file mode 100644
--- /dev/null
+++ b/Engine/Plugins/Experimental/MeshModelingToolset/Source/MeshModelingTools/Private/MeshSpaceDeformerTool.cpp
@@ -0,0 +1,522 @@
+// Copyright 1998-2019 Epic Games, Inc. All Rights Reserved.
+
+#include "MeshSpaceDeformerTool.h"
+#include "InteractiveToolManager.h"
+#include "ToolBuilderUtil.h"
+
+#include "SegmentTypes.h"
+#include "DynamicMeshAttributeSet.h"
+#include "MeshNormals.h"
+#include "ToolSceneQueriesUtil.h"
+#include "Intersection/IntersectionUtil.h"
+
+#include "Async/ParallelFor.h"
+#include "Containers/BitArray.h"
+
+///////////////////////////////
+// EIGEN
+
+#if defined(_MSC_VER) && USING_CODE_ANALYSIS
+#pragma warning(push)
+#pragma warning(disable : 6011)
+#pragma warning(disable : 6387)
+#pragma warning(disable : 6313)
+#pragma warning(disable : 6294)
+#endif
+THIRD_PARTY_INCLUDES_START
+#include <Eigen/Core>
+#include <Eigen/Dense> 
+#include <Eigen/Eigenvalues>
+
+#ifndef EIGEN_MPL2_ONLY
+#endif
+THIRD_PARTY_INCLUDES_END
+#if defined(_MSC_VER) && USING_CODE_ANALYSIS
+#pragma warning(pop)
+#endif
+
+///////////////////////////////
+
+
+#define LOCTEXT_NAMESPACE "MeshSpaceDeformerTool"
+
+
+//////////////////////////////
+// DEBUG_SETTINGS
+
+#define DEBUG_ASYNC_TASK
+
+#define DEBUG_DRAW_AXES
+//////////////////////////////
+
+
+#ifdef DEBUG_ASYNC_TASK
+#define print(text,...) UE_LOG(LogTemp, Warning, TEXT(text), __VA_ARGS__)
+#else
+#define print(x)
+#endif
+
+
+
+double RayToSegmentSquareDist(const FRay3d& Ray, const FVector3d& V0, const FVector3d& V1, FVector3d* ClosestPtOnSegment, FVector3d* ClosestPtOnRay)
+{
+	if (V0 == V1)
+	{
+		if (ClosestPtOnRay) 
+		{
+			*ClosestPtOnRay = V0;
+		}
+
+		if (ClosestPtOnSegment)
+		{
+			*ClosestPtOnSegment = V0;
+		}
+		return 0.;
+	}
+	FVector3d SegCenter = (V0 + V1) * 0.5;
+	FVector3d SegDir = (V1 - V0).Normalized();
+	FVector3d Diff = Ray.Origin - SegCenter;
+
+	double SegExtent = V0.Distance(V1) * 0.5;
+	double A01 = -Ray.Direction.Dot(SegDir);
+	double B0 = Diff.Dot(Ray.Direction);
+	double B1 = -Diff.Dot(SegDir);
+	double C = (Diff).SquaredLength();
+	double det = FMath::Abs(1. - A01 * A01);
+	double S0, S1, SqrDist, ExtDet;
+
+	if (det > 0.)
+	{
+		// The ray and segment are not parallel.
+
+		S0 = A01 * B1 - B0;
+		S1 = A01 * B0 - B1;
+		ExtDet = SegExtent * det;
+
+		if (S0 >= 0.)
+		{
+			if (S1 >= -ExtDet)
+			{
+				if (S1 <= ExtDet)
+				{
+					// region 0
+					// Minimum at interior points of ray and segment.
+					float invDet = 1. / det;
+					S0 *= invDet;
+					S1 *= invDet;
+					SqrDist = S0 * (S0 + A01 * S1 + 2. * B0) + S1 * (A01 * S0 + S1 + 2. * B1) + C;
+				}
+				else
+				{
+					// region 1
+					S1 = SegExtent;
+					S0 = std::max(0., -(A01 * S1 + B0));
+					SqrDist = -S0 * S0 + S1 * (S1 + 2. * B1) + C;
+				}
+			}
+			else
+			{
+				// region 5
+				S1 = -SegExtent;
+				S0 = std::max(0., -(A01 * S1 + B0));
+				SqrDist = -S0 * S0 + S1 * (S1 + 2. * B1) + C;
+			}
+		}
+		else
+		{
+			if (S1 <= -ExtDet)
+			{
+				// region 4
+				S0 = std::max(0., -(-A01 * SegExtent + B0));
+				S1 = (S0 > 0.) ? -SegExtent : std::min(std::max(-SegExtent, -B1), SegExtent);
+				SqrDist = -S0 * S0 + S1 * (S1 + 2. * B1) + C;
+			}
+			else if (S1 <= ExtDet)
+			{
+				// region 3
+				S0 = 0;
+				S1 = std::min(std::max(-SegExtent, -B1), SegExtent);
+				SqrDist = S1 * (S1 + 2. * B1) + C;
+			}
+			else
+			{
+				// region 2
+				S0 = std::max(0., -(A01 * SegExtent + B0));
+				S1 = (S0 > 0.) ? SegExtent : std::min(std::max(-SegExtent, -B1), SegExtent);
+				SqrDist = -S0 * S0 + S1 * (S1 + 2. * B1) + C;
+			}
+		}
+	}
+	else
+	{
+		// Ray and segment are parallel.
+
+		S1 = (A01 > 0.) ? -SegExtent : SegExtent;
+		S0 = std::max(0., -(A01 * S1 + B0));
+		SqrDist = -S0 * S0 + S1 * (S1 + 2. * B1) + C;
+	}
+
+	if (ClosestPtOnSegment)
+	{
+		*ClosestPtOnSegment = SegDir * S1 + SegCenter;
+	}
+
+	if (ClosestPtOnRay)
+	{
+		*ClosestPtOnRay = Ray.Direction * S0 + Ray.Origin;
+	}
+	return SqrDist;
+}
+
+FMatrix3d SolveSampleCovariance(double* Array, const int NCols, const int NRows, FVector3d& ExtentsOut, FVector3d& OriginOut)
+{
+	using RowMajorVec3dMatrix = Eigen::Matrix<double, Eigen::Dynamic, 3, Eigen::RowMajor>;
+	using RowMajorVec3d = Eigen::Matrix<double, 1, 3, Eigen::RowMajor>;
+	RowMajorVec3dMatrix VertexMatrix = Eigen::Map<RowMajorVec3dMatrix>(Array, NRows, NCols);
+
+	check(VertexMatrix(0, 0) == *Array); // Ensure the mapping is in expected format
+	check(VertexMatrix(0, 1) == Array[1]);
+	check(VertexMatrix(0, 2) == Array[2]);
+
+	//Using adjoint eigen solver
+	auto Origin = VertexMatrix.colwise().mean();															 // Get the centroid of the point cloud
+	Eigen::MatrixXd CenteredVertices = VertexMatrix.rowwise() - Origin;										 // Translate each point so centroid is (0,0,0)
+	Eigen::MatrixXd CovarianceMatrix = (CenteredVertices.adjoint() * CenteredVertices) / double(NRows - 1);	 // note: adjoint ~= conjugate transpose; formula from https://en.wikipedia.org/wiki/Sample_mean_and_covariance#Sample_covariance
+	
+	// note: SelfAdjointEigenSolver automatically sorts the eigen vectors by the order of corresponding eigenvalues in increasing order
+	Eigen::SelfAdjointEigenSolver<Eigen::MatrixXd> EigenSolver(CovarianceMatrix);							 
+	Eigen::Matrix3d EigenVectors = EigenSolver.eigenvectors();
+
+	//Align object with the standard basis by rotating all the points by the inverse of EigenVectors.
+	//EigenVectors is an orthonormal basis (3x3) so transpose is the inverse; 
+	Eigen::Matrix3Xd RotatedVertices = EigenVectors.transpose() * CenteredVertices.transpose();
+
+	//Since our point cloud is now aligned with the standard axes, the dimensions of its AABB should be the dimensions of its 
+	RowMajorVec3d Extents = RotatedVertices.rowwise().maxCoeff() - RotatedVertices.rowwise().minCoeff();
+
+	std::sort(Extents.data(), Extents.data() +3);						//Sort the extents to match our sorted eigenvectors
+	ExtentsOut = FVector3d{ FMath::Abs(Extents[0]), FMath::Abs(Extents[1]), FMath::Abs(Extents[2]) } / 2;	//Divide by two because extent represents distance from center
+	OriginOut = FVector3d{ Origin(0,0), Origin(0,1), Origin(0,2) };		//Origin in world space
+
+	return FMatrix3d{ FVector3d{ EigenVectors(0,0), EigenVectors(1,0), EigenVectors(2,0) },
+					  FVector3d{ EigenVectors(0,1), EigenVectors(1,1), EigenVectors(2,1) },
+					  FVector3d{ EigenVectors(0,2), EigenVectors(1,2), EigenVectors(2,2) }, 
+					  false };
+}
+
+FMatrix3d UMeshSpaceDeformerTool::CalculateBestAxis(FDynamicMesh3* Mesh, const TArray<FVector3d>& PositionBuffer, FTransform WorldTransform, FVector3d * ExtentsOut, FVector3d * OriginOut)
+{
+	PCAVertexPositionBuffer.Reset(3 * Mesh->VertexCount());
+	double* VectorData = PCAVertexPositionBuffer.GetData();
+	FVector3d* InitializePtr = reinterpret_cast<FVector3d*>(VectorData);
+	for (const int& VertexID : Mesh->VertexIndicesItr())
+	{
+		const FVector3d& Point = PositionBuffer[VertexID];
+		*InitializePtr = WorldTransform.TransformPosition(Point);
+		++InitializePtr;
+	}
+
+	const int NCols = 3; // R^3;
+	const int NRows = Mesh->VertexCount();
+
+	return SolveSampleCovariance(VectorData, NCols, NRows, *ExtentsOut, *OriginOut);
+}
+
+/*
+ * ToolBuilder
+ */
+UMeshSurfacePointTool* UMeshSpaceDeformerToolBuilder::CreateNewTool(const FToolBuilderState& SceneState) const
+{
+	UMeshSpaceDeformerTool* NonlinearTool = NewObject<UMeshSpaceDeformerTool>(SceneState.ToolManager);
+	return NonlinearTool;
+}
+
+UMeshSpaceDeformerTool::UMeshSpaceDeformerTool()
+{
+}
+
+void UMeshSpaceDeformerTool::Setup()
+{
+	UMeshSurfacePointTool::Setup();
+
+	//// create dynamic mesh component to use for live preview
+	DynamicMeshComponent = NewObject<USimpleDynamicMeshComponent>(ComponentTarget->GetOwnerActor(), "Dynamic Mesh");
+	DynamicMeshComponent->SetupAttachment(ComponentTarget->GetOwnerActor()->GetRootComponent());
+	DynamicMeshComponent->RegisterComponent();
+	DynamicMeshComponent->SetWorldTransform(ComponentTarget->GetWorldTransform());
+
+
+	//// copy material if there is one
+	auto Material = ComponentTarget->GetMaterial(0);
+	if (Material != nullptr)
+	{
+		DynamicMeshComponent->SetMaterial(0, Material);
+	}
+	//// dynamic mesh configuration settings
+	ComponentTarget->SetOwnerVisibility(false);
+	DynamicMeshComponent->TangentsType = EDynamicMeshTangentCalcType::AutoCalculated;
+	DynamicMeshComponent->InitializeMesh(ComponentTarget->GetMesh());
+	OnDynamicMeshComponentChangedHandle = DynamicMeshComponent->OnMeshChanged.Add(
+		FSimpleMulticastDelegate::FDelegate::CreateUObject(this, &UMeshSpaceDeformerTool::OnDynamicMeshComponentChanged));
+
+	FDynamicMesh3* TargetMesh = DynamicMeshComponent->GetMesh();
+
+	OriginalPositions.SetNumUninitialized(TargetMesh->MaxVertexID());
+	for (int32 VertexID : TargetMesh->VertexIndicesItr())
+	{
+		OriginalPositions[VertexID] = TargetMesh->GetVertex(VertexID);
+	}
+
+	bInDrag = false;
+
+	//// add properties
+	ToolPropertyObjects.Add(this);
+
+	//// set up visualizers
+	AxisRenderer.LineColor = FLinearColor::Red;
+	AxisRenderer.LineThickness = 2.0;
+	AxisRenderer.bDepthTested = false;
+
+	for (int8 i = 0; i < 3; ++i)
+	{
+		FMeshSpaceDeformerOp* Op = Operators[i];
+		Op->UpdateMesh(DynamicMeshComponent->GetMesh());
+	}
+
+	AutoDetectAxes();
+
+}
+
+void UMeshSpaceDeformerTool::Shutdown(EToolShutdownType ShutdownType)
+{
+
+	if (DynamicMeshComponent != nullptr)
+	{
+		DynamicMeshComponent->OnMeshChanged.Remove(OnDynamicMeshComponentChangedHandle);
+
+		ComponentTarget->SetOwnerVisibility(true);
+
+		if (ShutdownType == EToolShutdownType::Accept)
+		{
+			// this block bakes the modified DynamicMeshComponent back into the StaticMeshComponent inside an undo transaction
+			GetToolManager()->BeginUndoTransaction(LOCTEXT("MeshSpaceDeformerToolTransactionName", "Deform Mesh"));
+			ComponentTarget->CommitMesh([&](FMeshDescription* MeshDescription)
+			{
+				DynamicMeshComponent->Bake(MeshDescription, false);
+			});
+			GetToolManager()->EndUndoTransaction();
+		}
+			DynamicMeshComponent->UnregisterComponent();
+			DynamicMeshComponent->DestroyComponent();
+			DynamicMeshComponent = nullptr;
+	}
+}
+
+
+void UMeshSpaceDeformerTool::RegisterActions(FInteractiveToolActionSet& ActionSet)
+{
+}
+
+void UMeshSpaceDeformerTool::OnDynamicMeshComponentChanged()
+{
+}
+
+bool UMeshSpaceDeformerTool::HitTest(const FRay & Ray, FHitResult & OutHit)
+{
+	return Handle.SelectedVertexID != -1;
+}
+
+void UMeshSpaceDeformerTool::OnBeginDrag(const FRay& WorldRay)
+{
+	if (Handle.SelectedVertexID != -1)
+	{
+		bInDrag = true;
+		Handle.DragUpdateBounds(WorldRay);
+		UpdateIntervalsFromDrag();
+		UpdateOp();
+	}
+}
+
+
+void UMeshSpaceDeformerTool::OnUpdateDrag(const FRay& Ray)
+{
+	if (bInDrag)
+	{
+		Handle.DragUpdateBounds(Ray);
+		UpdateIntervalsFromDrag();
+		UpdateOp();
+	}
+}
+
+void UMeshSpaceDeformerTool::OnEndDrag(const FRay& Ray)
+{
+	bInDrag = false;
+
+}
+
+inline void UMeshSpaceDeformerTool::AutoDetectAxes()
+{
+	FTransform WorldTransform = ComponentTarget->GetWorldTransform();
+	FTransform WorldToObject = WorldTransform.Inverse();
+
+	//Retrieves the best axis in world space via principal component analysis.
+	FMatrix3d Basis = CalculateBestAxis(DynamicMeshComponent->GetMesh(), OriginalPositions, WorldTransform, &PrincipalAxesHalfExtentCoeff, &AxisCentroidWorldSpace);
+
+	PrincipalAxesWorldSpace[0] = FVector3d{ Basis(0,0),Basis(1,0),Basis(2,0) };
+	PrincipalAxesWorldSpace[1] = FVector3d{ Basis(0,1),Basis(1,1),Basis(2,1) };
+	PrincipalAxesWorldSpace[2] = FVector3d{ Basis(0,2),Basis(1,2),Basis(2,2) };
+
+	PrincipalAxesObjectSpace[0] = WorldToObject.TransformVectorNoScale(PrincipalAxesWorldSpace[0]);
+	PrincipalAxesObjectSpace[1] = WorldToObject.TransformVectorNoScale(PrincipalAxesWorldSpace[1]);
+	PrincipalAxesObjectSpace[2] = WorldToObject.TransformVectorNoScale(PrincipalAxesWorldSpace[2]);
+
+	UpdateObjectSpaceAxisCentroid();
+}
+
+void UMeshSpaceDeformerTool::SwapSecondaryAxis()
+{
+	std::swap(SecondaryAxis,ThirdAxis);
+}
+
+inline void UMeshSpaceDeformerTool::UpdateObjectSpaceAxisCentroid()
+{
+	AxisCentroidObjectSpace = ComponentTarget->GetWorldTransform().InverseTransformPositionNoScale(AxisCentroidWorldSpace);
+}
+
+void UMeshSpaceDeformerTool::UpdateIntervalsFromDrag()
+{
+	FHandleWidgetVertex* SelectedVertex = Handle.GetSelectedVertex();
+	double CurrentSegmentLength = (SelectedVertex->WorldCenter - AxisCentroidWorldSpace).Length();
+	double MaxLength = PrincipalAxesHalfExtentCoeff[SelectedAxis];
+	double NewInterval = FMath::Clamp(CurrentSegmentLength / MaxLength, 0.0, 1.5);
+	if (Handle.SelectedVertexID == 0)
+	{
+		LowerBoundsInterval = -NewInterval;
+	}
+	else
+	{
+		UpperBoundsInterval = NewInterval;
+	}
+}
+
+void UMeshSpaceDeformerTool::OnUpdateHover(const FInputDeviceRay& DevicePos)
+{
+	if (!bInDrag)
+	{
+		const auto ToleranceTest = [this](const FVector3d& Position1, const FVector3d& Position2) {
+			return ToolSceneQueriesUtil::PointSnapQuery(this->CameraState, Position1, Position2, VisualAngleSnapThreshold);
+		};
+		Handle.UpdateHover(DevicePos, ComponentTarget->GetWorldTransform(), TFunction<bool(const FVector3d&, const FVector3d&)>{ToleranceTest});
+	}
+}
+
+#if WITH_EDITOR
+void UMeshSpaceDeformerTool::PostEditChangeProperty(FPropertyChangedEvent & PropertyChangedEvent)
+{
+	//UpdateOp();
+}
+#endif
+
+void UMeshSpaceDeformerTool::UpdateOp()
+{
+	FMeshSpaceDeformerOp* Op = Operators[static_cast<int8>(SelectedOperationType)];
+
+	//Provide the basis so that the selected axis is in the Y direction (this will NEED to be rotated based on the orientation the operator expects)
+	// i.e. the Bend operator expects the bend along the Y axis, the twist expects along the Z axis. Just swap the axes as needed inside the operator
+	FMatrix3d Basis
+	{
+		PrincipalAxesObjectSpace[SecondaryAxis],
+		PrincipalAxesObjectSpace[SelectedAxis],
+		PrincipalAxesObjectSpace[ThirdAxis],
+		false
+	};
+
+	//This allows the object to move, or the handle position to move and have the changes reflected based on their relative world-space positions
+	UpdateObjectSpaceAxisCentroid();
+
+	//Pass the relevant updated data to the operator then run the operation.
+	Op->UpdateAxisData(Basis, AxisCentroidObjectSpace, PrincipalAxesHalfExtentCoeff, LowerBoundsInterval, UpperBoundsInterval, ModifierPercent);
+	Op->CalculateResult(nullptr);
+
+	DynamicMeshComponent->FastNotifyPositionsUpdated();
+}
+
+void UMeshSpaceDeformerTool::Tick(float DeltaTime)
+{
+	UMeshSurfacePointTool::Tick(DeltaTime);
+}
+
+
+void UMeshSpaceDeformerTool::Render(IToolsContextRenderAPI* RenderAPI)
+{
+	GetToolManager()->GetContextQueriesAPI()->GetCurrentViewState(CameraState);
+
+	FDynamicMesh3* TargetMesh = DynamicMeshComponent->GetMesh();
+
+	AxisRenderer.BeginFrame(RenderAPI, CameraState);
+
+	//Update the operator (pass it relevant data)
+	UpdateOp();
+
+	//Update the handle visualization and render it.
+	Handle.SetBasis(PrincipalAxesObjectSpace[0], PrincipalAxesObjectSpace[1], PrincipalAxesObjectSpace[2]);
+	Handle.UpdateDisplayData(SelectedAxis, LowerBoundsInterval, UpperBoundsInterval, PrincipalAxesHalfExtentCoeff);
+	Handle.SetPosition(AxisCentroidWorldSpace);
+	Handle.Render();
+
+
+	AxisRenderer.EndFrame();
+	GetToolManager()->PostInvalidation();
+}
+
+
+
+void FHandleWidget::UpdateHover(const FInputDeviceRay & DevicePos , const FTransform& WorldTransform, TFunction<bool(const FVector3d&, const FVector3d&)> ToleranceFunction)
+{
+	FVector3d LocalPosition, LocalNormal;
+	FGeometrySet3::FNearest Nearest;
+
+	//Reset all handle axes and vertex states
+	for (unsigned k = 0; k < 3; ++k)
+	{
+		Axes[k].State = IHandleWidgetSelectable::Default;
+		Axes[k].Vertices[0].State = IHandleWidgetSelectable::Default;
+		Axes[k].Vertices[1].State = IHandleWidgetSelectable::Default;
+	}
+
+	//If we're hovering over a vertex, change its state
+	bool bFound = GeometrySet.FindNearestPointToRay(DevicePos.WorldRay, Nearest, ToleranceFunction);
+	if (bFound)
+	{ // Found a hit point
+		int j = Nearest.ID % 2;
+		int i = (Nearest.ID) / 2;
+		SelectedVertexID = j;
+		Axes[i].Vertices[j].State = IHandleWidgetSelectable::Hovering;
+		Axes[i].Vertices[(j + 1) % 2].State = IHandleWidgetSelectable::Default;
+	}else
+	{
+		SelectedVertexID = -1;
+	}
+	
+	//TODO: Do the same for the axes. The geometry set should already contain the axes, just need to decide how the UI will work.
+
+
+}
+
+void FHandleWidget::DragUpdateBounds(const FRay3d& WorldRay)
+{
+	FVector3d ClosestPtOnSegment;
+	FHandleWidgetAxis* SelectedAxis = GetSelectedAxis();
+
+	FVector3d V0 = SelectedAxis->WorldCenter + SelectedAxis->Axis * 99999.0;
+	FVector3d V1 = SelectedAxis->WorldCenter - SelectedAxis->Axis * 99999.0;
+
+	double SqDist = RayToSegmentSquareDist(WorldRay, V0, V1, &ClosestPtOnSegment, nullptr);
+
+	FHandleWidgetVertex* SelectedVertex = GetSelectedVertex();
+	SelectedVertex->State = IHandleWidgetSelectable::Selected;
+	SelectedVertex->SetPosition(ClosestPtOnSegment);
+}
+
+#undef LOCTEXT_NAMESPACE
+
diff --git a/Engine/Plugins/Experimental/MeshModelingToolset/Source/MeshModelingTools/Private/PolygonMeshDeformTool.cpp b/Engine/Plugins/Experimental/MeshModelingToolset/Source/MeshModelingTools/Private/PolygonMeshDeformTool.cpp
new file mode 100644
--- /dev/null
+++ b/Engine/Plugins/Experimental/MeshModelingToolset/Source/MeshModelingTools/Private/PolygonMeshDeformTool.cpp
@@ -0,0 +1,1411 @@
+// Copyright 1998-2019 Epic Games, Inc. All Rights Reserved.
+
+#include "PolygonMeshDeformTool.h"
+#include "InteractiveToolManager.h"
+#include "ToolBuilderUtil.h"
+
+#include "SegmentTypes.h"
+#include "DynamicMeshAttributeSet.h"
+#include "MeshNormals.h"
+#include "ToolSceneQueriesUtil.h"
+#include "Intersection/IntersectionUtil.h"
+
+#include "Async/ParallelFor.h"
+#include "Containers/BitArray.h"
+
+#define LOCTEXT_NAMESPACE "UPolygonMeshDeformTool"
+
+
+
+//////////////////////////////
+// DEBUG_SETTINGS
+
+//Draw white triangles defining the selection subset
+//#define DEBUG_ROI_TRIANGLES
+
+//Draw pink circles around the handles
+//#define DEBUG_ROI_HANDLES
+
+//Draw points on the ROI vertices, White => Weight == 0, Black => Weight == 1
+//#define DEBUG_ROI_WEIGHTS
+
+//////////////////////////////
+
+
+/*
+ * ToolBuilder
+ */
+UMeshSurfacePointTool* UPolygonMeshDeformToolBuilder::CreateNewTool(const FToolBuilderState& SceneState) const
+{
+	UPolygonMeshDeformTool* DeformTool = NewObject<UPolygonMeshDeformTool>(SceneState.ToolManager);
+	return DeformTool;
+}
+
+/*
+ * Tool
+ */
+UPolyEditTransformProperties::UPolyEditTransformProperties()
+{
+	bUseLocalAxis = true;
+	TransformMode = EQuickTransformerMode::AxisTranslation;
+	bSelectVertices = true;
+	bSelectFaces = true;
+	bSelectEdges = true;
+}
+
+/*
+ * Asynchronous Task
+ */
+
+
+void FConstrainedMeshDeformerTask::UpdateDeformer(const ELaplacianWeightScheme SelectedWeightScheme, const FDynamicMesh3& SrcMesh, const TArray<FConstraintData>& ConstraintArray, const TArray<int32>& SrcIDBufferSubset, bool bNewTransaction, const FRichCurve* Curve)
+{
+	bIsNewTransaction = bNewTransaction;
+	SrcMeshMaxVertexID = SrcMesh.MaxVertexID();
+
+	LaplacianWeightScheme = SelectedWeightScheme;
+	
+	bAttenuateWeights = (Curve != nullptr);
+	if (bAttenuateWeights)
+	{
+		WeightAttenuationCurve = *Curve;
+	}
+
+	// Set-up the subset mesh.
+	if (bIsNewTransaction)
+	{
+		//Copy the part of the mesh we want to deform into the SubsetMesh and create map from Src Mesh to the SubsetMesh.
+		InitializeSubsetMesh(SrcMesh, SrcIDBufferSubset);
+	}
+
+	// only want the subset of constraints that correspond to our subset mesh.
+	{
+		const int32 NumSubsetVerts = SubsetVertexIDToSrcVertexIDMap.Num();
+		SubsetConstraintBuffer.Empty(NumSubsetVerts);
+		SubsetConstraintBuffer.AddUninitialized(NumSubsetVerts);
+
+		for (int32 SubVertexID = 0; SubVertexID < SubsetVertexIDToSrcVertexIDMap.Num(); ++SubVertexID)
+		{
+			int32 SrcVtxID = SubsetVertexIDToSrcVertexIDMap[SubVertexID];
+			SubsetConstraintBuffer[SubVertexID] = ConstraintArray[SrcVtxID];
+		}
+	}
+
+	check(bIsNewTransaction || ConstrainedDeformer.IsValid());
+	
+}
+
+
+void FConstrainedMeshDeformerTask::DoWork()
+{
+
+	//TODO: (simple optimization) -
+	//	Instead of SrcVertexIDtoSubsetVertexIDMap, use SubsetVertexIDToSetVertexIDMap - then we can use the VertexIndicesItr()
+	//	on the SubsetMesh to minimize the quantity of vertex indices we need to iterate at every following step.
+
+	if (*bAbortSource == true)
+	{
+		return;
+	}
+
+	if (bIsNewTransaction) //Will only be true once per input transaction (click+drag)
+	{
+		
+		// Create a new deformation solver.
+		ConstrainedDeformer = MakeUnique<FConstrainedMeshDeformer>(SubsetMesh, LaplacianWeightScheme);
+
+		if (bAttenuateWeights)
+		{
+			ApplyAttenuation();
+		}
+
+		//Update our deformer's constraints before deforming using the copy of the constraint buffer
+		for (int32 SubsetVertexID = 0; SubsetVertexID < SubsetConstraintBuffer.Num(); ++SubsetVertexID)
+		{
+			FConstraintData& CData = SubsetConstraintBuffer[SubsetVertexID];
+			ConstrainedDeformer->AddConstraint(SubsetVertexID, CData.Weight, CData.Position, CData.bPostFix);
+		}
+
+		bIsNewTransaction = false;
+
+	}
+	else
+	{
+		//This else block is run every consecutive frame after the start of the input transaction because UpdateConstraintPosition() is very cheap (no factorizing or rebuilding)
+		//Update only the positions of the constraints, as the weights cannot change mid-transaction
+
+		for (int32 SubsetVertexID = 0; SubsetVertexID < SubsetConstraintBuffer.Num(); ++SubsetVertexID)
+		{
+			FConstraintData& CData = SubsetConstraintBuffer[SubsetVertexID];
+			ConstrainedDeformer->UpdateConstraintPosition(SubsetVertexID, CData.Position, CData.bPostFix);
+		}
+
+	}
+
+	if (*bAbortSource == true)
+	{
+		return;
+	}
+
+	//Run the deformation process
+
+	const bool bSuccessfulSolve = ConstrainedDeformer->Deform(SubsetPositionBuffer);
+	
+	
+	if (bSuccessfulSolve)
+	{
+		if (*bAbortSource == true)
+		{
+			return;
+		}
+	}
+	else
+	{
+		UE_LOG(LogTemp, Warning, TEXT("Laplacian deformation failed"));
+	}
+
+}
+
+inline void FConstrainedMeshDeformerTask::InitializeSubsetMesh(const FDynamicMesh3& SrcMesh, const TArray<int32>& SrcIDBufferSubset)
+{
+	//These can be re-used until the user stops dragging
+	SubsetMesh.Clear();
+	SubsetPositionBuffer.Reset();
+
+	//Initialize every element to -1, helps us keep track of vertices we've already added while iterating the triangles
+	TArray<int32> SrcVertexIDToSubsetVertexIDMap;
+	SrcVertexIDToSubsetVertexIDMap.Init(INACTIVE_SUBSET_ID, SrcMeshMaxVertexID);
+
+
+	//Iterate the triangle array to append vertices, and then triangles to the temporary subset mesh all at once
+	for (int32 i = 0; i < SrcIDBufferSubset.Num(); i += 3)
+	{
+
+		// Build the triangle
+		FIndex3i Triangle;
+		for (int32 v = 0; v < 3; ++v)
+		{
+			//It's the SrcVertexID because every element in the SrcIDBufferSubset is the Vertex ID of a vertex in the original mesh.
+			const int32 SrcVertexID = SrcIDBufferSubset[i + v];
+			int32& SubsetID = SrcVertexIDToSubsetVertexIDMap[SrcVertexID];
+			
+			if (SubsetID == INACTIVE_SUBSET_ID) // we haven't already visited this vertex
+			{
+				const FVector3d Vertex = SrcMesh.GetVertex(SrcVertexID);
+				SubsetID = SubsetMesh.AppendVertex(Vertex);
+			}
+
+			Triangle[v] = SubsetID;
+		}
+		SubsetMesh.AppendTriangle(Triangle);
+	}
+
+	// create a mapping back to the original vertex IDs from the subset mesh
+	int32 MaxSubMeshVertexID = SubsetMesh.MaxVertexID(); // Really MaxID + 1
+	SubsetVertexIDToSrcVertexIDMap.Reset(MaxSubMeshVertexID);
+	SubsetVertexIDToSrcVertexIDMap.AddUninitialized(MaxSubMeshVertexID);
+
+	for (int32 SrcID = 0; SrcID < SrcVertexIDToSubsetVertexIDMap.Num(); ++SrcID)
+	{
+		const int32 SubsetVertexID = SrcVertexIDToSubsetVertexIDMap[SrcID];
+		if (SubsetVertexID != INACTIVE_SUBSET_ID)
+		{
+			SubsetVertexIDToSrcVertexIDMap[SubsetVertexID] = SrcID;
+		}
+	}
+
+}
+
+void FConstrainedMeshDeformerTask::ExportResults(FDynamicMesh3& TargetMesh) const
+{
+
+	//Update the position buffer result
+	for (int32 SubsetVertexID = 0; SubsetVertexID < SubsetVertexIDToSrcVertexIDMap.Num(); ++SubsetVertexID)
+	{
+		const int32 SrcVertexID = SubsetVertexIDToSrcVertexIDMap[SubsetVertexID];
+		const FVector3d Position = SubsetPositionBuffer[SubsetVertexID];
+		
+		TargetMesh.SetVertex(SrcVertexID, Position);
+	}
+}
+
+void FConstrainedMeshDeformerTask::ApplyAttenuation()
+{
+	size_t Size = SrcMeshMaxVertexID;
+	TSet<int> Handles;
+
+	auto InPlaceMinMaxElements = [](FVector3d& Min, FVector3d& Max, const FVector3d Test)
+	{
+		for (uint8 i = 0; i < 3; ++i)
+		{
+			Min[i] = Test[i] < Min[i] ? Test[i] : Min[i];
+			Max[i] = Test[i] > Max[i] ? Test[i] : Max[i];
+		}
+	};
+
+	//Experimental approach: Just going to try grabbing the bounding box of the entire mesh, then the bounding box of the handles as a point cloud.
+	//						 We need a T value to pass to the Weights curve, so let's try finding the distance of each vertex V from line segment formed by the min/max handles
+	//						 Divide the distance from the handles to vertex V by the length of the mesh's bounding box extent, 
+	//						 and that will provide a **ROUGH** approximation of the time value for our curve.
+	//
+	//											  Distance( LineSegment(MaxHandle,MinHandle) , V )
+	// where T(V) is time value at V     T(V) = -----------------------------------------------------
+	//   and V is the position								Length(MeshMin - MeshMax)   
+	//      of each vertex
+
+	FVector3d Min{ std::numeric_limits<double>::max(),std::numeric_limits<double>::max(),std::numeric_limits<double>::max() };
+	FVector3d Max{ std::numeric_limits<double>::min(),std::numeric_limits<double>::min(),std::numeric_limits<double>::min() };
+	FVector3d MinHandles = Min;
+	FVector3d MaxHandles = Max;
+	double LeastWeight = std::numeric_limits<double>::max();
+
+
+
+	for (int32 SubVertexID = 0; SubVertexID < SubsetConstraintBuffer.Num(); ++SubVertexID)
+	{
+
+		FConstraintData& CData = SubsetConstraintBuffer[SubVertexID];
+		// Update bounding box
+		InPlaceMinMaxElements(Min, Max, CData.Position);
+		
+		if (CData.Weight > 0.0)
+		{
+			LeastWeight = CData.Weight < LeastWeight ? CData.Weight : LeastWeight;
+
+			// update bounding box
+			InPlaceMinMaxElements(MinHandles, MaxHandles, CData.Position);
+			Handles.Add(SubVertexID);
+		}
+	}
+
+	double ExtentLength = Min.Distance(Max);
+
+	// Is this why the system has memory?
+	for (int32 SubVertexID = 0; SubVertexID < SubsetConstraintBuffer.Num(); ++SubVertexID)
+	{
+		if (!Handles.Contains(SubVertexID))
+		{
+			FConstraintData& CData = SubsetConstraintBuffer[SubVertexID];
+			double T = CData.Position.Distance(FMath::ClosestPointOnSegment(CData.Position, MinHandles, MaxHandles)) / ExtentLength;
+			CData.Weight = WeightAttenuationCurve.Eval(T) * LeastWeight;
+		}
+
+	}
+}
+
+
+/*
+*	FGroupTopologyLaplacianDeformer methods
+*/
+
+
+void FGroupTopologyLaplacianDeformer::InitBackgroundWorker(const ELaplacianWeightScheme WeightScheme)
+{
+
+	//Initialize asynchronous deformation objects
+	if (AsyncMeshDeformTask == nullptr)
+	{
+		AsyncMeshDeformTask = new FAsyncTaskExecuterWithAbort<FConstrainedMeshDeformerTask>(WeightScheme);
+	}
+}
+
+void FGroupTopologyLaplacianDeformer::InitializeConstraintBuffer()
+{
+	//MaxVertexID is used because the array is potentially sparse.
+	int MaxVertexID = Mesh->MaxVertexID();
+
+	SrcMeshConstraintBuffer.SetNum(MaxVertexID);
+	
+	for (int32 VertexID : Mesh->VertexIndicesItr())
+	{
+		FConstraintData& CD = SrcMeshConstraintBuffer[VertexID];
+		CD.Position = Mesh->GetVertex(VertexID); 
+		CD.Weight   = 0.0;
+		CD.bPostFix = false;
+	}
+}
+
+bool FGroupTopologyLaplacianDeformer::IsTaskInFlight() const
+{
+	return ( AsyncMeshDeformTask != nullptr && !AsyncMeshDeformTask->IsDone() );
+}
+
+bool FGroupTopologyLaplacianDeformer::UpdateAndLaunchdWorker(const ELaplacianWeightScheme SelectedWeightScheme, const FRichCurve*  Curve)
+{
+	/* Deformer needs to run if we've modified the constraints since the last time it finished. */
+	if (AsyncMeshDeformTask == nullptr)
+	{
+		InitBackgroundWorker(SelectedWeightScheme);
+	}
+
+	if (bDeformerNeedsToRun && AsyncMeshDeformTask->IsDone())
+	{		
+		bool bRebuildSubsetMesh = bTaskSubmeshIsDirty;
+
+		FConstrainedMeshDeformerTask& Task = AsyncMeshDeformTask->GetTask();
+
+		// Update the deformer's buffers and weight scheme
+		// this creates the subset mesh if needed.
+		Task.UpdateDeformer(SelectedWeightScheme, *Mesh, SrcMeshConstraintBuffer, SubsetIDBuffer, bRebuildSubsetMesh, Curve);
+
+		// task now has valid submesh
+
+		bTaskSubmeshIsDirty = false;
+
+		//Launch second thread
+		AsyncMeshDeformTask->StartBackgroundTask();
+
+		bDeformerNeedsToRun = false;	 // This was set to true above in UpdateSolution()
+		bVertexPositionsNeedSync = true; // The task will generate new vertex positions. 
+
+		return true;
+
+	}
+	return false;
+}
+
+void FGroupTopologyLaplacianDeformer::SetActiveHandleFaces(const TArray<int>& FaceGroupIDs)
+{
+	Reset();
+
+	check(FaceGroupIDs.Num() == 1);   // multi-face not supported yet
+	int GroupID = FaceGroupIDs[0];
+
+	// find set of vertices in handle 
+	Topology->CollectGroupVertices(GroupID, HandleVertices);
+	Topology->CollectGroupBoundaryVertices(GroupID, HandleBoundaryVertices);
+	ModifiedVertices = HandleVertices;
+
+	// find neighbour group set
+	TArray<int> HandleGroups = FaceGroupIDs;
+	const TArray<int>& GroupNbrGroups = Topology->GetGroupNbrGroups(GroupID);
+
+
+	CalculateROI(HandleGroups, GroupNbrGroups);
+
+
+	TArray<int> NeighborGroups;
+
+	//Add to output parameter 
+	NeighborGroups.Add(GroupID);
+	NeighborGroups.Append(GroupNbrGroups);
+
+	UpdateSelection(Mesh, NeighborGroups, bLocalize);
+
+	// Save the positions of the selected region.
+	SaveInitialPositions();
+
+}
+
+
+void FGroupTopologyLaplacianDeformer::SetActiveHandleEdges(const TArray<int>& TopologyEdgeIDs)
+{
+	Reset();
+
+	for (int EdgeID : TopologyEdgeIDs)
+	{
+		const TArray<int>& EdgeVerts = Topology->GetGroupEdgeVertices(EdgeID);
+		for (int VertID : EdgeVerts)
+		{
+			HandleVertices.Add(VertID);
+		}
+	}
+	HandleBoundaryVertices = HandleVertices;
+	ModifiedVertices = HandleVertices;
+
+	TArray<int> HandleGroups;
+	TArray<int> NbrGroups;
+	Topology->FindEdgeNbrGroups(TopologyEdgeIDs, NbrGroups);
+
+	CalculateROI(HandleGroups, NbrGroups);
+
+	UpdateSelection(Mesh, NbrGroups, bLocalize);
+
+	// Save the positions of the selected region.
+	SaveInitialPositions();
+}
+
+void FGroupTopologyLaplacianDeformer::SetActiveHandleCorners(const TArray<int>& CornerIDs)
+{
+	Reset();
+
+	for (int CornerID : CornerIDs)
+	{
+		int VertID = Topology->GetCornerVertexID(CornerID);
+		if (VertID >= 0)
+		{
+			HandleVertices.Add(VertID);
+		}
+	}
+	HandleBoundaryVertices = HandleVertices;
+	ModifiedVertices = HandleVertices;
+
+	TArray<int> HandleGroups;
+	TArray<int> NbrGroups;
+
+	Topology->FindCornerNbrGroups(CornerIDs, NbrGroups);
+
+
+	CalculateROI(HandleGroups, NbrGroups);
+
+	UpdateSelection(Mesh, NbrGroups, bLocalize);
+
+	// Save the positions of the selected region.
+	SaveInitialPositions();
+}
+
+
+
+void FGroupTopologyLaplacianDeformer::UpdateSelection(const FDynamicMesh3* TargetMesh, const TArray<int>& Groups, bool bLocalizeDeformation)
+{
+
+	// Build an index buffer (SubsetIdBuffer) and a vertexId buffer (ModifidedVertices) for the region we want to change
+
+	if (bLocalizeDeformation)
+	{
+		//For each group ID, retrieve the array of all TriangleIDs associated with that GroupID and append that array to the end of the TriSet to remove duplicates
+		TSet<int> TriSet;
+		for (const int32& GroupID : Groups)
+		{
+			TriSet.Append(Topology->GetGroupFaces(GroupID));
+		}//Now we have every triangle ID involved in the transaction 
+
+		//Since we are flattening the Face to a set of 3 indices, we do 3 * number of triangles though it is too many.
+		SubsetIDBuffer.Reset(3 * TriSet.Num());
+		//Add each triangle's A,B, and C indices to the subset triangle array.
+		for (const int& Tri : TriSet)
+		{
+			FIndex3i Triple = TargetMesh->GetTriangle(Tri);
+			SubsetIDBuffer.Add(Triple.A);
+			SubsetIDBuffer.Add(Triple.B);
+			SubsetIDBuffer.Add(Triple.C);
+		
+		}
+	}
+	else
+	{
+		// the entire mesh.
+		const int32 NumTris = TargetMesh->TriangleCount();
+		SubsetIDBuffer.Reset(3 * NumTris);
+		for (int TriId : TargetMesh->TriangleIndicesItr()) 
+		{ 
+			FIndex3i Triple = TargetMesh->GetTriangle(TriId);
+			SubsetIDBuffer.Add(Triple.A);
+			SubsetIDBuffer.Add(Triple.B);
+			SubsetIDBuffer.Add(Triple.C);
+		}
+	}
+
+	// Add the vertices to the set (eliminates duplicates.)  Todo: don't use a set.
+	ResetModifiedVertices();
+	for (int32 VertexID : SubsetIDBuffer)
+	{
+		RecordModifiedVertex(VertexID);
+	}
+
+}
+
+// This actually updates constraints that correspond to the handle vertices.
+void FGroupTopologyLaplacianDeformer::UpdateSolution(FDynamicMesh3 * TargetMesh, const TFunction<FVector3d(FDynamicMesh3* Mesh, int)>& HandleVertexDeformFunc)
+{
+	// copy the current positions.
+	FVertexPositionCache CurrentPositions;
+	for (int VertexID : InitialPositions.Vertices)
+	{
+		CurrentPositions.AddVertex(TargetMesh, VertexID);
+	}
+
+	// Set the target mesh to the initial positions.
+	// Note: this only updates the vertices in the selected region.
+	InitialPositions.SetPositions(TargetMesh);
+
+	//Reset the constraints
+	for (int32 VertexID : ModifiedVertices)
+	{
+		//Get the vertex's data from the constraint buffer
+		FConstraintData& CData = SrcMeshConstraintBuffer[VertexID];
+
+		CData.Position = TargetMesh->GetVertex(VertexID);
+		CData.Weight   = 0.0; //A weight of zero is used to allow this point to move freely when moving the handles
+		CData.bPostFix = false;
+	}
+
+	//Actually deform the handles and add a constraint.
+	for (int VertexID : HandleVertices)
+	{
+		const FVector3d DeformPos = HandleVertexDeformFunc(TargetMesh, VertexID);
+
+		//Get the vertex's data from the constraint buffer
+		FConstraintData& CData = SrcMeshConstraintBuffer[VertexID];
+
+		//Set the new vertex data
+		CData.Position = DeformPos;
+		CData.Weight   = HandleWeights;
+		CData.bPostFix = bPostfixHandles;
+	}
+
+	// Restore Current Positions.  This is done because the target mesh is being used to define the highlight region.
+	// if we don't reset the positions the highlight mesh will appear to reset momentarily until the first laplacian solver result is available 
+	CurrentPositions.SetPositions(TargetMesh);
+
+	bDeformerNeedsToRun = true;
+
+}
+
+void FGroupTopologyLaplacianDeformer::ExportDeformedPositions(FDynamicMesh3* TargetMesh)
+{
+	bool bIsWorking = IsTaskInFlight();
+	if (AsyncMeshDeformTask != nullptr && !bIsWorking)
+	{
+		const FConstrainedMeshDeformerTask& Task = AsyncMeshDeformTask->GetTask();
+		Task.ExportResults(*TargetMesh);
+	}
+}
+
+inline FGroupTopologyLaplacianDeformer::~FGroupTopologyLaplacianDeformer()
+{
+	Shutdown();
+}
+
+inline void FGroupTopologyLaplacianDeformer::Shutdown()
+{	
+	if (AsyncMeshDeformTask != nullptr)
+	{
+		if (AsyncMeshDeformTask->IsDone())
+		{
+			delete AsyncMeshDeformTask;
+		}
+		else
+		{
+			AsyncMeshDeformTask->CancelAndDelete();
+		}
+
+		AsyncMeshDeformTask = nullptr;
+	}
+}
+
+/*
+* Tool methods
+*/
+
+UPolygonMeshDeformTool::UPolygonMeshDeformTool()
+{
+}
+
+void UPolygonMeshDeformTool::Setup()
+{
+	UMeshSurfacePointTool::Setup();
+
+	// create dynamic mesh component to use for live preview
+	DynamicMeshComponent = NewObject<USimpleDynamicMeshComponent>(ComponentTarget->GetOwnerActor(), "DynamicMesh");
+	DynamicMeshComponent->SetupAttachment(ComponentTarget->GetOwnerActor()->GetRootComponent());
+	DynamicMeshComponent->RegisterComponent();
+	DynamicMeshComponent->SetWorldTransform(ComponentTarget->GetWorldTransform());
+
+	// copy material if there is one
+	auto Material = ComponentTarget->GetMaterial(0);
+	if (Material != nullptr)
+	{
+		DynamicMeshComponent->SetMaterial(0, Material);
+	}
+
+	// dynamic mesh configuration settings
+	DynamicMeshComponent->TangentsType = EDynamicMeshTangentCalcType::AutoCalculated;
+	DynamicMeshComponent->InitializeMesh(ComponentTarget->GetMesh());
+	OnDynamicMeshComponentChangedHandle = DynamicMeshComponent->OnMeshChanged.Add(
+		FSimpleMulticastDelegate::FDelegate::CreateUObject(this, &UPolygonMeshDeformTool::OnDynamicMeshComponentChanged));
+
+	// initialize AABBTree
+	MeshSpatial.SetMesh(DynamicMeshComponent->GetMesh());
+	PrecomputeTopology();
+
+	//initialize topology selector
+	TopoSelector.Initialize(DynamicMeshComponent->GetMesh(), &Topology);
+	TopoSelector.SetSpatialSource([this]() {return &GetSpatial(); });
+	TopoSelector.PointsWithinToleranceTest = [this](const FVector3d& Position1, const FVector3d& Position2) {
+		return ToolSceneQueriesUtil::PointSnapQuery(this->CameraState, Position1, Position2, VisualAngleSnapThreshold);
+	};
+
+	// hide input StaticMeshComponent
+	ComponentTarget->SetOwnerVisibility(false);
+
+	// init state flags flags
+	bInDrag = false;
+
+	// initialize snap solver
+	QuickAxisTranslater.Initialize();
+	QuickAxisRotator.Initialize();
+
+	// add properties
+	TransformProps = NewObject<UPolyEditTransformProperties>(this);
+	AddToolPropertySource(TransformProps);
+
+	// set up visualizers
+	PolyEdgesRenderer.LineColor = FLinearColor::Red;
+	PolyEdgesRenderer.LineThickness = 2.0;
+	HilightRenderer.LineColor = FLinearColor::Green;
+	HilightRenderer.LineThickness = 4.0f;
+
+	// Allocates buffers, sets up the asynchronous task
+	// Copies the source mesh positions.
+	const ELaplacianWeightScheme LaplacianWeightScheme = ConvertToLaplacianWeightScheme(TransformProps->SelectedWeightScheme);
+	LaplacianDeformer.InitBackgroundWorker(LaplacianWeightScheme);
+	
+
+	/**
+	// How to add a curve for the weights.
+	//Add a default curve for falloff
+	FKeyHandle Keys[5];
+	Keys[0] = TransformProps->DefaultFalloffCurve.UpdateOrAddKey(0.f, 1.f);
+	Keys[1] = TransformProps->DefaultFalloffCurve.UpdateOrAddKey(0.25f, 0.25f);
+	Keys[2] = TransformProps->DefaultFalloffCurve.UpdateOrAddKey(0.3333333f, 0.25f);
+	Keys[3] = TransformProps->DefaultFalloffCurve.UpdateOrAddKey(0.6666667f, 1.25f);
+	Keys[4] = TransformProps->DefaultFalloffCurve.UpdateOrAddKey(1.f, 1.4f);
+	for (uint8 i = 0; i < 5; ++i)
+	{
+		TransformProps->DefaultFalloffCurve.SetKeyInterpMode(Keys[i], ERichCurveInterpMode::RCIM_Cubic);
+	}
+	TransformProps->WeightAttenuationCurve.EditorCurveData = TransformProps->DefaultFalloffCurve;
+	*/
+}
+
+void UPolygonMeshDeformTool::Shutdown(EToolShutdownType ShutdownType)
+{
+	//Tell the background thread to cancel the rest of its jobs before we close;
+	LaplacianDeformer.Shutdown();
+
+	if (DynamicMeshComponent != nullptr)
+	{
+		DynamicMeshComponent->OnMeshChanged.Remove(OnDynamicMeshComponentChangedHandle);
+
+		ComponentTarget->SetOwnerVisibility(true);
+
+		if (ShutdownType == EToolShutdownType::Accept)
+		{
+			// this block bakes the modified DynamicMeshComponent back into the StaticMeshComponent inside an undo transaction
+			GetToolManager()->BeginUndoTransaction(LOCTEXT("PolygonMeshDeformToolTransactionName", "Deform Mesh"));
+			ComponentTarget->CommitMesh([=](FMeshDescription* MeshDescription)
+			{
+				DynamicMeshComponent->Bake(MeshDescription, false);
+			});
+			GetToolManager()->EndUndoTransaction();
+		}
+
+		DynamicMeshComponent->UnregisterComponent();
+		DynamicMeshComponent->DestroyComponent();
+		DynamicMeshComponent = nullptr;
+	}
+}
+
+
+
+void UPolygonMeshDeformTool::ToggleLocalTransformFrameAction()
+{
+	TransformProps->bUseLocalAxis = !TransformProps->bUseLocalAxis;
+	if (bInDrag)
+	{
+		UpdateQuickTransformer();
+	}
+}
+
+void UPolygonMeshDeformTool::NextTransformTypeAction()
+{
+	if (bInDrag == false)
+	{
+		if (TransformProps->TransformMode == EQuickTransformerMode::AxisRotation)
+		{
+			TransformProps->TransformMode = EQuickTransformerMode::AxisTranslation;
+		}
+		else
+		{
+			TransformProps->TransformMode = EQuickTransformerMode::AxisRotation;
+		}
+		UpdateQuickTransformer();
+	}
+}
+
+
+
+void UPolygonMeshDeformTool::RegisterActions(FInteractiveToolActionSet& ActionSet)
+{
+	ActionSet.RegisterAction(this, (int32)EStandardToolActions::BaseClientDefinedActionID + 1,
+		TEXT("ToggleFrame"),
+		LOCTEXT("ToggleFrame", "Toggle World/Local Frame"),
+		LOCTEXT("ToggleFrameTooltip", "Toggle between World and Local transformation frames"),
+		EModifierKey::None, EKeys::A,
+		[this]() { ToggleLocalTransformFrameAction(); });
+
+
+	ActionSet.RegisterAction(this, (int32)EStandardToolActions::BaseClientDefinedActionID + 2,
+		TEXT("NextTransformType"),
+		LOCTEXT("NextTransformType", "Next Transform Type"),
+		LOCTEXT("NextTransformTypeTooltip", "Cycle to next transform type"),
+		EModifierKey::None, EKeys::Q,
+		[this]() { NextTransformTypeAction(); });
+}
+
+
+
+
+void UPolygonMeshDeformTool::OnDynamicMeshComponentChanged()
+{
+	bSpatialDirty = true;
+	TopoSelector.Invalidate(true, false);
+
+
+	//Makes sure the constraint buffer and position buffers reflect Undo/Redo changes
+	FDynamicMesh3* Mesh = DynamicMeshComponent->GetMesh();
+	
+	
+	//Apply Undo/redo
+	for (int VertexID : Mesh->VertexIndicesItr())
+	{
+		const FVector3d Position = Mesh->GetVertex(VertexID);
+		LaplacianDeformer.SrcMeshConstraintBuffer[VertexID].Position = Position;
+	}
+
+	// a deform task could still be in flight.
+	if (LaplacianDeformer.AsyncMeshDeformTask != nullptr)
+	{
+		LaplacianDeformer.AsyncMeshDeformTask->CancelAndDelete();
+		LaplacianDeformer.AsyncMeshDeformTask = nullptr;
+		LaplacianDeformer.bTaskSubmeshIsDirty = true;
+	}
+}
+
+FDynamicMeshAABBTree3& UPolygonMeshDeformTool::GetSpatial()
+{
+	if (bSpatialDirty)
+	{
+		MeshSpatial.Build();
+		bSpatialDirty = false;
+	}
+	return MeshSpatial;
+}
+
+
+bool UPolygonMeshDeformTool::HitTest(const FRay& WorldRay, FHitResult& OutHit)
+{
+	FTransform Transform = ComponentTarget->GetWorldTransform();
+	FRay3d LocalRay(Transform.InverseTransformPosition(WorldRay.Origin),
+		Transform.InverseTransformVector(WorldRay.Direction));
+	LocalRay.Direction.Normalize();
+
+	TopoSelector.UpdateEnableFlags(TransformProps->bSelectFaces, TransformProps->bSelectEdges, TransformProps->bSelectVertices);
+	FGroupTopologySelection Selection;
+	FVector3d LocalPosition, LocalNormal;
+	if (TopoSelector.FindSelectedElement(LocalRay, Selection, LocalPosition, LocalNormal) == false)
+	{
+		return false;
+	}
+
+	if (Selection.SelectedCornerIDs.Num() > 0)
+	{
+		OutHit.FaceIndex = Selection.SelectedCornerIDs[0];
+		OutHit.Distance = LocalRay.Project(LocalPosition);
+		OutHit.ImpactPoint = Transform.TransformPosition(LocalRay.PointAt(OutHit.Distance));
+	}
+	else if (Selection.SelectedEdgeIDs.Num() > 0)
+	{
+		OutHit.FaceIndex = Selection.SelectedEdgeIDs[0];
+		OutHit.Distance = LocalRay.Project(LocalPosition);
+		OutHit.ImpactPoint = Transform.TransformPosition(LocalRay.PointAt(OutHit.Distance));
+	}
+	else
+	{
+		int HitTID = GetSpatial().FindNearestHitTriangle(LocalRay);
+		if (HitTID != IndexConstants::InvalidID)
+		{
+			FTriangle3d Triangle;
+			GetSpatial().GetMesh()->GetTriVertices(HitTID, Triangle.V[0], Triangle.V[1], Triangle.V[2]);
+			FIntrRay3Triangle3d Query(LocalRay, Triangle);
+			Query.Find();
+			OutHit.FaceIndex = HitTID;
+			OutHit.Distance = Query.RayParameter;
+			OutHit.Normal = Transform.TransformVectorNoScale(GetSpatial().GetMesh()->GetTriNormal(HitTID));
+			OutHit.ImpactPoint = Transform.TransformPosition(LocalRay.PointAt(Query.RayParameter));
+		}
+	}
+	return true;
+}
+
+
+void UPolygonMeshDeformTool::OnBeginDrag(const FRay& WorldRay)
+{
+	FTransform Transform = ComponentTarget->GetWorldTransform();
+	FRay3d LocalRay(Transform.InverseTransformPosition(WorldRay.Origin),
+		Transform.InverseTransformVector(WorldRay.Direction));
+	LocalRay.Direction.Normalize();
+
+	HilightSelection.Clear();
+
+	TopoSelector.UpdateEnableFlags(TransformProps->bSelectFaces, TransformProps->bSelectEdges, TransformProps->bSelectVertices);
+	FGroupTopologySelection Selection;
+	FVector3d LocalPosition, LocalNormal;
+	bool bHit = TopoSelector.FindSelectedElement(LocalRay, Selection, LocalPosition, LocalNormal);
+
+	if (bHit == false)
+	{
+		bInDrag = false;
+		return;
+	}
+
+	HilightSelection = Selection;
+
+	FVector3d WorldHitPos = Transform.TransformPosition(LocalPosition);
+	FVector3d WorldHitNormal = Transform.TransformVector(LocalNormal);
+
+	bInDrag = true;
+	StartHitPosWorld = WorldHitPos;
+	LastHitPosWorld = StartHitPosWorld;
+	StartHitNormalWorld = WorldHitNormal;
+
+	QuickAxisRotator.ClearAxisLock();
+	UpdateActiveSurfaceFrame(HilightSelection);
+	UpdateQuickTransformer();
+
+	LastBrushPosLocal = Transform.InverseTransformPosition(LastHitPosWorld);
+	StartBrushPosLocal = LastBrushPosLocal;
+
+	// Capture the part of the mesh that will deform
+
+	if (TransformProps->DeformationStrategy == EGroupTopologyDeformationStrategy::Linear)
+	{
+		//Determine which of the following (corners, edges or faces) has been selected by counting the associated feature's IDs
+		if (Selection.SelectedCornerIDs.Num() > 0)
+		{
+			//Add all the the Corner's adjacent poly-groups (NbrGroups) to the ongoing array of groups.
+			LinearDeformer.SetActiveHandleCorners(Selection.SelectedCornerIDs);
+		}
+		else if (Selection.SelectedEdgeIDs.Num() > 0)
+		{
+			//Add all the the edge's adjacent poly-groups (NbrGroups) to the ongoing array of groups.
+			LinearDeformer.SetActiveHandleEdges(Selection.SelectedEdgeIDs);
+		}
+		else if (Selection.SelectedGroupIDs.Num() > 0)
+		{
+			LinearDeformer.SetActiveHandleFaces(Selection.SelectedGroupIDs);
+		}
+	}
+	else //(TransformProps->DeformationStrategy == EGroupTopologyDeformationStrategy::Laplacian)
+	{
+		LaplacianDeformer.bLocalize = true; // TransformProps->bLocalizeDeformation;
+
+		//Determine which of the following (corners, edges or faces) has been selected by counting the associated feature's IDs
+		if (Selection.SelectedCornerIDs.Num() > 0)
+		{
+			//Add all the the Corner's adjacent poly-groups (NbrGroups) to the ongoing array of groups.
+			LaplacianDeformer.SetActiveHandleCorners(Selection.SelectedCornerIDs);
+		}
+		else if (Selection.SelectedEdgeIDs.Num() > 0)
+		{
+			//Add all the the edge's adjacent poly-groups (NbrGroups) to the ongoing array of groups.
+			LaplacianDeformer.SetActiveHandleEdges(Selection.SelectedEdgeIDs);
+		}
+		else if (Selection.SelectedGroupIDs.Num() > 0)
+		{
+			LaplacianDeformer.SetActiveHandleFaces(Selection.SelectedGroupIDs);
+		}
+
+		// the task will need a new mesh that corresponds to the selected region.
+		LaplacianDeformer.bTaskSubmeshIsDirty = true;
+	}
+
+	BeginChange();
+}
+
+
+void UPolygonMeshDeformTool::UpdateActiveSurfaceFrame(FGroupTopologySelection& Selection)
+{
+	FTransform Transform = ComponentTarget->GetWorldTransform();
+
+	// update surface frame
+	ActiveSurfaceFrame.Origin = StartHitPosWorld;
+	if (HilightSelection.SelectedCornerIDs.Num() == 1)
+	{
+		// just keeping existing axes...we don't have enough info to do something smarter
+	}
+	else
+	{
+		ActiveSurfaceFrame.AlignAxis(2, StartHitNormalWorld);
+		if (HilightSelection.SelectedEdgeIDs.Num() == 1)
+		{
+			FVector3d Tangent;
+			if (Topology.GetGroupEdgeTangent(HilightSelection.SelectedEdgeIDs[0], Tangent))
+			{
+				Tangent = Transform.TransformVector(Tangent);
+				ActiveSurfaceFrame.ConstrainedAlignAxis(0, Tangent, ActiveSurfaceFrame.Z());
+			}
+		}
+	}
+}
+
+
+FQuickTransformer* UPolygonMeshDeformTool::GetActiveQuickTransformer()
+{
+	if (TransformProps->TransformMode == EQuickTransformerMode::AxisRotation)
+	{
+		return &QuickAxisRotator;
+	}
+	else
+	{
+		return &QuickAxisTranslater;
+	}
+}
+
+
+void UPolygonMeshDeformTool::UpdateQuickTransformer()
+{
+	if (TransformProps->bUseLocalAxis)
+	{
+		GetActiveQuickTransformer()->SetActiveWorldFrame(ActiveSurfaceFrame);
+	}
+	else
+	{
+		GetActiveQuickTransformer()->SetActiveFrameFromWorldAxes(StartHitPosWorld);
+	}
+}
+
+
+
+
+
+void UPolygonMeshDeformTool::UpdateChangeFromROI(bool bFinal)
+{
+	if (ActiveVertexChange == nullptr)
+	{
+		return;
+	}
+	const bool bIsLaplacian = (TransformProps->DeformationStrategy == EGroupTopologyDeformationStrategy::Laplacian);
+
+	FDynamicMesh3* Mesh = DynamicMeshComponent->GetMesh();
+	const TSet<int>& ModifiedVertices = (bIsLaplacian) ? LaplacianDeformer.GetModifiedVertices() : LinearDeformer.GetModifiedVertices();
+	
+	ActiveVertexChange->SavePositions(Mesh, ModifiedVertices, !bFinal);
+
+}
+
+
+void UPolygonMeshDeformTool::OnUpdateDrag(const FRay& Ray)
+{
+	if (bInDrag)
+	{
+		bUpdatePending = true;
+		UpdateRay = Ray;
+	}
+}
+
+void UPolygonMeshDeformTool::OnEndDrag(const FRay& Ray)
+{
+	bInDrag = false;
+	bUpdatePending = false;
+
+	// update spatial
+	bSpatialDirty = true;
+
+	HilightSelection.Clear(); 
+	TopoSelector.Invalidate(true, false);
+	QuickAxisRotator.ClearAxisLock();
+
+	//If it's linear, it's computed real time with no delay. This may need to be restructured for clarity by using the background task for this as well.
+	if (TransformProps->DeformationStrategy == EGroupTopologyDeformationStrategy::Linear)
+	{
+		// close change record
+		EndChange();
+	}
+}
+
+
+
+void UPolygonMeshDeformTool::OnUpdateHover(const FInputDeviceRay& DevicePos)
+{
+	//if (!bNeedEmitEndChange)
+	if (ActiveVertexChange == nullptr)
+	{
+		FTransform Transform = ComponentTarget->GetWorldTransform();
+		FRay3d LocalRay(Transform.InverseTransformPosition(DevicePos.WorldRay.Origin),
+		Transform.InverseTransformVector(DevicePos.WorldRay.Direction));
+		LocalRay.Direction.Normalize();
+
+		HilightSelection.Clear();
+		TopoSelector.UpdateEnableFlags(TransformProps->bSelectFaces, TransformProps->bSelectEdges, TransformProps->bSelectVertices);
+		FVector3d LocalPosition, LocalNormal;
+		bool bHit = TopoSelector.FindSelectedElement(LocalRay, HilightSelection, LocalPosition, LocalNormal);
+
+		if (bHit)
+		{
+			StartHitPosWorld = Transform.TransformPosition(LocalPosition);
+			StartHitNormalWorld = Transform.TransformVector(LocalNormal);
+
+			UpdateActiveSurfaceFrame(HilightSelection);
+			UpdateQuickTransformer();
+		}
+	}
+}
+
+
+
+
+void UPolygonMeshDeformTool::ComputeUpdate()
+{
+
+	if (bUpdatePending == true)
+	{
+		// Linear Deformer : Update the solution
+		// Laplacain Deformer : Update the constraints (positions and weights) - the region was identified in onBeginDrag
+
+
+		if (TransformProps->TransformMode == EQuickTransformerMode::AxisRotation)
+		{
+			ComputeUpdate_Rotate();
+		}
+		else
+		{
+			ComputeUpdate_Translate();
+		}
+	}
+
+	if (TransformProps->DeformationStrategy == EGroupTopologyDeformationStrategy::Laplacian)
+	{
+		bool bIsWorking = LaplacianDeformer.IsTaskInFlight();
+
+		if (!bIsWorking)
+		{
+			// Sync update if we have new results.
+			if (LaplacianDeformer.bVertexPositionsNeedSync)
+			{
+				
+				//Update the mesh with the provided solutions.
+				LaplacianDeformer.ExportDeformedPositions(DynamicMeshComponent->GetMesh());
+
+				LaplacianDeformer.bVertexPositionsNeedSync = false;
+
+				//Re-sync mesh, and flag the spatial data struct & topology for re-evaluation
+				DynamicMeshComponent->FastNotifyPositionsUpdated();
+
+				GetToolManager()->PostInvalidation();
+				bSpatialDirty = true;
+				TopoSelector.Invalidate(true, false);
+			}
+
+			// emit end change if we are done with the drag
+			if (!LaplacianDeformer.bDeformerNeedsToRun && !bInDrag)
+			{
+				EndChange();
+			}
+
+			// Not working but we have more work for it to do..
+
+			if (LaplacianDeformer.bDeformerNeedsToRun)
+			{
+			
+				FRichCurve* Curve = NULL;
+
+				/**
+				// How to add a deformation curve
+				const bool bApplyAttenuationCurve = TransformProps->bApplyAttenuationCurve
+				if (bApplyAttenuationCurve)
+				{
+					Curve = TransformProps->WeightAttenuationCurve.GetRichCurve();
+				}
+				*/
+				const ELaplacianWeightScheme LaplacianWeightScheme = ConvertToLaplacianWeightScheme(TransformProps->SelectedWeightScheme);
+				LaplacianDeformer.UpdateAndLaunchdWorker(LaplacianWeightScheme, Curve);
+
+			}
+		}
+	}
+}
+
+
+
+
+
+
+void UPolygonMeshDeformTool::ComputeUpdate_Rotate()
+{
+	const bool bIsLaplacian = (TransformProps->DeformationStrategy == EGroupTopologyDeformationStrategy::Laplacian); 
+	FGroupTopologyDeformer& SelectedDeformer = (bIsLaplacian) ? LaplacianDeformer : LinearDeformer;
+
+	FDynamicMesh3* Mesh = DynamicMeshComponent->GetMesh();
+	FTransform Transform = ComponentTarget->GetWorldTransform();
+	FVector NewHitPosWorld = LastHitPosWorld;
+
+	FVector3d SnappedPoint;
+	if (QuickAxisRotator.UpdateSnap(UpdateRay, SnappedPoint))
+	{
+		NewHitPosWorld = SnappedPoint;
+	}
+	else
+	{
+		return;
+	}
+
+	// check if we are on back-facing part of rotation in which case we ignore...
+	FVector3d SphereCenter = QuickAxisRotator.GetActiveWorldFrame().Origin;
+	if (QuickAxisRotator.HaveActiveSnapRotation() && QuickAxisRotator.GetHaveLockedToAxis() == false)
+	{
+		FVector3d ToSnapPointVec = (SnappedPoint - SphereCenter);
+		FVector3d ToEyeVec = (SnappedPoint - (FVector3d)CameraState.Position);
+		if (ToSnapPointVec.Dot(ToEyeVec) > 0)
+		{
+			return;
+		}
+	}
+
+
+	// if we haven't snapped to a rotation we can exit
+	if (QuickAxisRotator.HaveActiveSnapRotation() == false)
+	{
+		QuickAxisRotator.ClearAxisLock();
+
+		SelectedDeformer.ClearSolution(Mesh);
+
+		//TODO: This is unseemly here, need to potentially defer this so that it's handled the same way as laplacian. Placeholder for now.
+		if (TransformProps->DeformationStrategy == EGroupTopologyDeformationStrategy::Linear)
+		{
+
+			DynamicMeshComponent->FastNotifyPositionsUpdated();
+			GetToolManager()->PostInvalidation();
+		}
+		bUpdatePending = false;
+		return;
+	}
+
+	// ok we have an axis...
+	if (QuickAxisRotator.GetHaveLockedToAxis() == false)
+	{
+		QuickAxisRotator.SetAxisLock();
+		RotationStartPointWorld = SnappedPoint;
+		RotationStartFrame = QuickAxisRotator.GetActiveRotationFrame();
+	}
+
+	FVector2d RotateStartVec = RotationStartFrame.ToPlaneUV(RotationStartPointWorld, 2);
+	RotateStartVec.Normalize();
+	FVector2d RotateToVec = RotationStartFrame.ToPlaneUV(NewHitPosWorld, 2);
+	RotateToVec.Normalize();
+	double AngleRad = RotateStartVec.SignedAngleR(RotateToVec);
+	FQuaterniond Rotation(
+		Transform.InverseTransformVectorNoScale(RotationStartFrame.Z()), AngleRad, false);
+	FVector3d LocalOrigin = Transform.InverseTransformPosition(RotationStartFrame.Origin);
+
+	// Linear Deformer: Update Mesh the rotation,
+	// Laplacian Deformer:  Update handles constraints with the rotation and set bDeformerNeedsToRun = true;.
+	SelectedDeformer.UpdateSolution(Mesh, [this, LocalOrigin, Rotation](FDynamicMesh3* TargetMesh, int VertIdx)
+	{
+		FVector3d V = TargetMesh->GetVertex(VertIdx);
+		V -= LocalOrigin;
+		V = Rotation * V;
+		V += LocalOrigin;
+		return V;
+	});
+
+	//TODO: This is unseemly here, need to potentially defer this so that it's handled the same way as laplacian. Placeholder for now.
+	if (!bIsLaplacian)
+	{
+		DynamicMeshComponent->FastNotifyPositionsUpdated();
+		GetToolManager()->PostInvalidation();
+	}
+	bUpdatePending = false;
+}
+
+
+
+
+void UPolygonMeshDeformTool::ComputeUpdate_Translate()
+{
+	const bool bIsLaplacian = (TransformProps->DeformationStrategy == EGroupTopologyDeformationStrategy::Laplacian);
+	FGroupTopologyDeformer& SelectedDeformer = (bIsLaplacian) ? LaplacianDeformer : LinearDeformer;
+
+	FTransform Transform = ComponentTarget->GetWorldTransform();
+	FVector NewHitPosWorld = LastHitPosWorld;
+	FVector3d SnappedPoint;
+	if (QuickAxisTranslater.UpdateSnap(UpdateRay, SnappedPoint))
+	{
+		NewHitPosWorld = SnappedPoint;
+	}
+	else
+	{
+		return;
+	}
+
+	FVector NewBrushPosLocal = Transform.InverseTransformPosition(NewHitPosWorld);
+	FVector3d NewMoveDelta = NewBrushPosLocal - StartBrushPosLocal;
+
+	FDynamicMesh3* Mesh = DynamicMeshComponent->GetMesh();
+	if (LastMoveDelta.SquaredLength() > 0.)
+	{
+		if (NewMoveDelta.SquaredLength() > 0.)
+		{
+			// Linear Deformer: Update Mesh with the translation,
+			// Laplacian Deformer:  Update handles constraints and set bDeformerNeedsToRun = true;.
+
+			SelectedDeformer.UpdateSolution(Mesh, [this, NewMoveDelta](FDynamicMesh3* TargetMesh, int VertIdx)
+			{
+				return TargetMesh->GetVertex(VertIdx) + NewMoveDelta;
+			});
+		}
+		else
+		{
+			// Reset mesh to initial positions.
+
+			SelectedDeformer.ClearSolution(Mesh);
+
+		}
+		//TODO: This is unseemly here, need to potentially defer this so that it's handled the same way as laplacian. Placeholder for now.
+		if (!bIsLaplacian)
+		{
+			DynamicMeshComponent->FastNotifyPositionsUpdated();
+			GetToolManager()->PostInvalidation();
+		}
+	}
+
+	LastMoveDelta = NewMoveDelta;
+	LastBrushPosLocal = NewBrushPosLocal;
+
+	bUpdatePending = false;
+}
+
+
+
+void UPolygonMeshDeformTool::Tick(float DeltaTime)
+{
+	UMeshSurfacePointTool::Tick(DeltaTime);
+	
+	LaplacianDeformer.HandleWeights   = TransformProps->HandleWeight;
+	LaplacianDeformer.bPostfixHandles = TransformProps->bPostFixHandles;
+}
+
+
+
+void UPolygonMeshDeformTool::PrecomputeTopology()
+{
+	FDynamicMesh3* Mesh = DynamicMeshComponent->GetMesh();
+	Topology = FGroupTopology(Mesh, true);
+
+	LinearDeformer.Initialize(Mesh, &Topology);
+	LaplacianDeformer.Initialize(Mesh, &Topology);
+
+	// Make the Constraint Buffer, zero weights, but current pos
+	LaplacianDeformer.InitializeConstraintBuffer();
+}
+
+
+
+
+void UPolygonMeshDeformTool::Render(IToolsContextRenderAPI* RenderAPI)
+{
+
+	ComputeUpdate();
+		
+	GetToolManager()->GetContextQueriesAPI()->GetCurrentViewState(CameraState);
+	GetActiveQuickTransformer()->UpdateCameraState(CameraState);
+
+	FDynamicMesh3* TargetMesh = DynamicMeshComponent->GetMesh();
+
+	PolyEdgesRenderer.BeginFrame(RenderAPI, CameraState);
+	PolyEdgesRenderer.SetTransform(ComponentTarget->GetWorldTransform());
+
+
+	for (FGroupTopology::FGroupEdge& Edge : Topology.Edges)
+	{
+		FVector3d A, B;
+		for (int eid : Edge.Span.Edges)
+		{
+			TargetMesh->GetEdgeV(eid, A, B);
+			PolyEdgesRenderer.DrawLine(A, B);
+		}
+	}
+
+	PolyEdgesRenderer.EndFrame();
+
+
+	HilightRenderer.BeginFrame(RenderAPI, CameraState);
+	HilightRenderer.SetTransform(ComponentTarget->GetWorldTransform());
+
+#ifdef DEBUG_ROI_WEIGHTS
+	FDynamicMesh3* MeshPtr = DynamicMeshComponent->GetMesh();
+	for (int32 VertexID : DynamicMeshComponent->GetMesh()->VertexIndicesItr())
+	{
+		float Color = 1.f - SrcMeshConstraintBuffer[VertexID].Weight;
+		HilightRenderer.DrawPoint(MeshPtr->GetVertex(VertexID), FLinearColor(Color, Color, Color, 1.f), 8, true);
+	}
+
+#endif
+
+#ifdef DEBUG_ROI_HANDLES
+	const FLinearColor FOOF{ 1.f,0.f,1.f,1.f };
+	for (int VertIdx : HandleVertices)
+	{
+		HilightRenderer.DrawViewFacingCircle(TargetMesh->GetVertex(VertIdx), 0.8f, 8, FOOF, 3, false);
+	}
+#endif 
+
+#ifdef DEBUG_ROI_TRIANGLES
+	const FLinearColor Whiteish{ 0.67f,0.67f,0.67f,1.f };
+	for (int32 i = 0; i < SubsetIDBuffer.Num(); i += 3)
+	{
+		FVector3d A = TargetMesh->GetVertex(SubsetIDBuffer[i]);
+		FVector3d B = TargetMesh->GetVertex(SubsetIDBuffer[i + 1]);
+		FVector3d C = TargetMesh->GetVertex(SubsetIDBuffer[i + 2]);
+		HilightRenderer.DrawLine(A, B, Whiteish, 2.7f, true);
+		HilightRenderer.DrawLine(B, C, Whiteish, 2.7f, true);
+		HilightRenderer.DrawLine(C, A, Whiteish, 2.7f, true);
+	}
+#endif
+
+	TopoSelector.VisualAngleSnapThreshold = this->VisualAngleSnapThreshold;
+	TopoSelector.DrawSelection(HilightSelection, &HilightRenderer, &CameraState);
+	HilightRenderer.EndFrame();
+
+
+	if (bInDrag)
+	{
+
+		GetActiveQuickTransformer()->Render(RenderAPI);
+	}
+	else
+	{
+		GetActiveQuickTransformer()->PreviewRender(RenderAPI);
+	}
+}
+
+
+
+
+
+//
+// Change Tracking
+//
+
+
+void UPolygonMeshDeformTool::BeginChange()
+{
+	const bool bIsLaplacian = (TransformProps->DeformationStrategy == EGroupTopologyDeformationStrategy::Laplacian);
+	if (!bIsLaplacian || LaplacianDeformer.IsDone())
+	{
+		if (ActiveVertexChange == nullptr)
+		{
+			ActiveVertexChange = new FMeshVertexChangeBuilder();
+			UpdateChangeFromROI(false);
+		}
+	}
+}
+
+
+void UPolygonMeshDeformTool::EndChange()
+{
+
+	if (ActiveVertexChange != nullptr)
+	{
+		UpdateChangeFromROI(true);
+		GetToolManager()->EmitObjectChange(DynamicMeshComponent, MoveTemp(ActiveVertexChange->Change), LOCTEXT("PolyMeshDeformationChange", "PolyMesh Edit"));
+	}
+
+	delete ActiveVertexChange;
+	ActiveVertexChange = nullptr;
+}
+
+
+
+#undef LOCTEXT_NAMESPACE
diff --git a/Engine/Plugins/Experimental/MeshModelingToolset/Source/MeshModelingTools/Private/Properties/MeshStatisticsProperties.cpp b/Engine/Plugins/Experimental/MeshModelingToolset/Source/MeshModelingTools/Private/Properties/MeshStatisticsProperties.cpp
new file mode 100644
--- /dev/null
+++ b/Engine/Plugins/Experimental/MeshModelingToolset/Source/MeshModelingTools/Private/Properties/MeshStatisticsProperties.cpp
@@ -0,0 +1,22 @@
+// Copyright 1998-2019 Epic Games, Inc. All Rights Reserved.
+
+#include "Properties/MeshStatisticsProperties.h"
+
+#include "DynamicMesh3.h"
+#include "DynamicMeshAttributeSet.h"
+#include "MeshNormals.h"
+
+#include "SimpleDynamicMeshComponent.h"
+
+#define LOCTEXT_NAMESPACE "UMeshStatisticsProperites"
+
+
+void UMeshStatisticsProperties::Update(const FDynamicMesh3& Mesh)
+{
+	TriangleCount = Mesh.TriangleCount();
+	VertexCount = Mesh.VertexCount();
+}
+
+
+
+#undef LOCTEXT_NAMESPACE
