Mesh Modeling Tools moved from NotForLicensees to Experimental

diff --git a/Engine/Plugins/Experimental/MeshModelingToolset/Source/MeshModelingTools/Private/RemeshMeshTool.cpp b/Engine/Plugins/Experimental/MeshModelingToolset/Source/MeshModelingTools/Private/RemeshMeshTool.cpp
new file mode 100644
--- /dev/null
+++ b/Engine/Plugins/Experimental/MeshModelingToolset/Source/MeshModelingTools/Private/RemeshMeshTool.cpp
@@ -0,0 +1,295 @@
+// Copyright 1998-2019 Epic Games, Inc. All Rights Reserved.
+
+#include "RemeshMeshTool.h"
+#include "InteractiveToolManager.h"
+#include "ToolBuilderUtil.h"
+
+#include "DynamicMesh3.h"
+#include "DynamicMeshAttributeSet.h"
+#include "Remesher.h"
+#include "MeshConstraintsUtil.h"
+#include "ProjectionTargets.h"
+#include "MeshNormals.h"
+
+#include "SimpleDynamicMeshComponent.h"
+
+#include "SceneManagement.h" // for FPrimitiveDrawInterface
+#include "Async/ParallelFor.h"
+
+#define LOCTEXT_NAMESPACE "URemeshMeshTool"
+
+
+/*
+ * ToolBuilder
+ */
+
+
+bool URemeshMeshToolBuilder::CanBuildTool(const FToolBuilderState& SceneState) const
+{
+	return ToolBuilderUtil::CountComponents(SceneState, CanMakeComponentTarget) == 1;
+}
+
+UInteractiveTool* URemeshMeshToolBuilder::BuildTool(const FToolBuilderState& SceneState) const
+{
+	URemeshMeshTool* NewTool = NewObject<URemeshMeshTool>(SceneState.ToolManager, "Remesh Tool");
+
+	UActorComponent* ActorComponent = ToolBuilderUtil::FindFirstComponent(SceneState, CanMakeComponentTarget);
+	auto* MeshComponent = Cast<UPrimitiveComponent>(ActorComponent);
+	check(MeshComponent != nullptr);
+	NewTool->SetSelection(MakeComponentTarget(MeshComponent));
+
+	return NewTool;
+}
+
+
+
+/*
+ * Tool
+ */
+
+URemeshMeshToolProperties::URemeshMeshToolProperties()
+{
+	TargetTriangleCount = 5000;
+	SmoothingSpeed = 0.25;
+	RemeshIterations = 20;
+	bDiscardAttributes = false;
+	SmoothingType = ERemeshSmoothingType::MeanValue;
+	bPreserveSharpEdges = true;
+}
+
+URemeshMeshAdvancedProperties::URemeshMeshAdvancedProperties()
+{
+	TargetEdgeLength = 5.0;
+	bFlips = true;
+	bSplits = true;
+	bCollapses = true;
+	bReproject = true;
+	bPreventNormalFlips = true;
+	bUseTargetEdgeLength = false;
+}
+
+
+URemeshMeshTool::URemeshMeshTool()
+{
+}
+
+
+void URemeshMeshTool::Setup()
+{
+	UInteractiveTool::Setup();
+
+	// create dynamic mesh component to use for live preview
+	DynamicMeshComponent = NewObject<USimpleDynamicMeshComponent>(ComponentTarget->GetOwnerActor(), "DynamicMesh");
+	DynamicMeshComponent->SetupAttachment(ComponentTarget->GetOwnerActor()->GetRootComponent());
+	DynamicMeshComponent->RegisterComponent();
+	DynamicMeshComponent->SetWorldTransform(ComponentTarget->GetWorldTransform());
+
+	// copy material if there is one
+	auto Material = ComponentTarget->GetMaterial(0);
+	if (Material != nullptr)
+	{
+		DynamicMeshComponent->SetMaterial(0, Material);
+	}
+	DynamicMeshComponent->bExplicitShowWireframe = true;
+
+	DynamicMeshComponent->InitializeMesh(ComponentTarget->GetMesh());
+	OriginalMesh.Copy(*DynamicMeshComponent->GetMesh());
+	OriginalMeshSpatial.SetMesh(&OriginalMesh, true);
+
+	// hide input StaticMeshComponent
+	ComponentTarget->SetOwnerVisibility(false);
+
+	// calculate initial mesh area (no utility fn yet)
+	InitialMeshArea = 0;
+	for (int tid : OriginalMesh.TriangleIndicesItr())
+	{
+		InitialMeshArea += OriginalMesh.GetTriArea(tid);
+	}
+
+	BasicProperties = NewObject<URemeshMeshToolProperties>(this, TEXT("Remesh Settings"));
+	// arbitrary threshold of 5000 tris seems reasonable? 
+	BasicProperties->TargetTriangleCount = (OriginalMesh.TriangleCount() < 5000) ? 5000 : OriginalMesh.TriangleCount();
+
+	AdvancedProperties = NewObject<URemeshMeshAdvancedProperties>(this, TEXT("Advanced Settings"));
+	AdvancedProperties->TargetEdgeLength = CalculateTargetEdgeLength(BasicProperties->TargetTriangleCount);
+
+	// initialize our properties
+	AddToolPropertySource(BasicProperties);
+	AddToolPropertySource(AdvancedProperties);
+
+	MeshStatisticsProperties = NewObject<UMeshStatisticsProperties>(this);
+	AddToolPropertySource(MeshStatisticsProperties);
+	MeshStatisticsProperties->Update(*DynamicMeshComponent->GetMesh());
+
+	bResultValid = false;
+}
+
+
+void URemeshMeshTool::Shutdown(EToolShutdownType ShutdownType)
+{
+	if (DynamicMeshComponent != nullptr)
+	{
+		//DynamicMeshComponent->OnMeshChanged.Remove(OnDynamicMeshComponentChangedHandle);
+
+		ComponentTarget->SetOwnerVisibility(true);
+
+		if (ShutdownType == EToolShutdownType::Accept)
+		{
+			// this block bakes the modified DynamicMeshComponent back into the StaticMeshComponent inside an undo transaction
+			GetToolManager()->BeginUndoTransaction(LOCTEXT("RemeshMeshToolTransactionName", "Remesh Mesh"));
+			ComponentTarget->CommitMesh([=](FMeshDescription* MeshDescription)
+			{
+				DynamicMeshComponent->Bake(MeshDescription, true);
+			});
+			GetToolManager()->EndUndoTransaction();
+		}
+
+		DynamicMeshComponent->UnregisterComponent();
+		DynamicMeshComponent->DestroyComponent();
+		DynamicMeshComponent = nullptr;
+	}
+}
+
+
+void URemeshMeshTool::Render(IToolsContextRenderAPI* RenderAPI)
+{
+	UpdateResult();
+
+
+	FPrimitiveDrawInterface* PDI = RenderAPI->GetPrimitiveDrawInterface();
+	FTransform Transform = ComponentTarget->GetWorldTransform(); //Actor->GetTransform();
+
+	FColor LineColor(255, 0, 0);
+	FDynamicMesh3* TargetMesh = DynamicMeshComponent->GetMesh();
+	if (TargetMesh->HasAttributes())
+	{
+		FDynamicMeshUVOverlay* UVOverlay = TargetMesh->Attributes()->PrimaryUV();
+		for (int eid : TargetMesh->EdgeIndicesItr()) 
+		{
+			if (UVOverlay->IsSeamEdge(eid)) 
+			{
+				FVector3d A, B;
+				TargetMesh->GetEdgeV(eid, A, B);
+				PDI->DrawLine(Transform.TransformPosition(A), Transform.TransformPosition(B),
+					LineColor, 0, 2.0, 1.0f, true);
+			}
+		}
+	}
+}
+
+#if WITH_EDITOR
+void URemeshMeshTool::PostEditChangeProperty(FPropertyChangedEvent& PropertyChangedEvent)
+{
+	UProperty* PropertyThatChanged = PropertyChangedEvent.Property;
+	bResultValid = false;
+}
+#endif
+
+void URemeshMeshTool::OnPropertyModified(UObject* PropertySet, UProperty* Property)
+{
+	bResultValid = false;
+}
+
+
+double URemeshMeshTool::CalculateTargetEdgeLength(int TargetTriCount)
+{
+	double TargetTriArea = InitialMeshArea / (double)TargetTriCount;
+	double EdgeLen = TriangleUtil::EquilateralEdgeLengthForArea(TargetTriArea);
+	return (double)FMath::RoundToInt(EdgeLen*100.0) / 100.0;
+}
+
+
+
+void URemeshMeshTool::UpdateResult()
+{
+	if (bResultValid) 
+	{
+		return;
+	}
+
+	FDynamicMesh3* TargetMesh = DynamicMeshComponent->GetMesh();
+	TargetMesh->Copy(OriginalMesh);
+
+	if (BasicProperties->bDiscardAttributes && BasicProperties->bPreserveSharpEdges == false)
+	{
+		TargetMesh->DiscardAttributes();
+	}
+
+	FRemesher Remesher(TargetMesh);
+	Remesher.bEnableSplits = AdvancedProperties->bSplits;
+	Remesher.bEnableFlips = AdvancedProperties->bFlips;
+	Remesher.bEnableCollapses = AdvancedProperties->bCollapses;
+
+	// recalculate target edge length
+	if (AdvancedProperties->bUseTargetEdgeLength == false)
+	{
+		AdvancedProperties->TargetEdgeLength = CalculateTargetEdgeLength(BasicProperties->TargetTriangleCount);
+	}
+
+	Remesher.SetTargetEdgeLength(AdvancedProperties->TargetEdgeLength);
+
+	Remesher.ProjectionMode = (AdvancedProperties->bReproject) ? 
+		FRemesher::ETargetProjectionMode::AfterRefinement : FRemesher::ETargetProjectionMode::NoProjection;
+
+	Remesher.bEnableSmoothing = (BasicProperties->SmoothingSpeed > 0);
+	Remesher.SmoothSpeedT = BasicProperties->SmoothingSpeed;
+	//Remesher.SmoothType = FRemesher::ESmoothTypes::MeanValue;
+	Remesher.SmoothType = (BasicProperties->bDiscardAttributes) ?
+		FRemesher::ESmoothTypes::Uniform :
+		(FRemesher::ESmoothTypes)(int)BasicProperties->SmoothingType;
+	bool bIsUniformSmooth = (Remesher.SmoothType == FRemesher::ESmoothTypes::Uniform);
+
+	Remesher.bPreventNormalFlips = AdvancedProperties->bPreventNormalFlips;
+
+	Remesher.DEBUG_CHECK_LEVEL = 0;
+
+	FMeshConstraints constraints;
+	FMeshConstraintsUtil::ConstrainAllSeams(constraints, *TargetMesh, true, !BasicProperties->bPreserveSharpEdges);
+	Remesher.SetExternalConstraints(&constraints);
+
+	FMeshProjectionTarget ProjTarget(&OriginalMesh, &OriginalMeshSpatial);
+	Remesher.SetProjectionTarget(&ProjTarget);
+
+	if (BasicProperties->bDiscardAttributes && BasicProperties->bPreserveSharpEdges == true)
+	{
+		TargetMesh->DiscardAttributes();
+	}
+
+	// run the remesh iterations
+	for (int k = 0; k < BasicProperties->RemeshIterations; ++k)
+	{
+		// If we are not uniform smoothing, then flips seem to often make things worse.
+		// Possibly this is because without the tangential flow, we won't get to the nice tris.
+		// In this case we are better off basically not flipping, and just letting collapses resolve things
+		// regular-valence polygons - things stay "stuck". 
+		// @todo try implementing edge-length flip criteria instead of valence-flip
+		if (bIsUniformSmooth == false)
+		{
+			bool bUseFlipsThisPass = (k % 2 == 0 && k < BasicProperties->RemeshIterations/2);
+			Remesher.bEnableFlips = bUseFlipsThisPass && AdvancedProperties->bFlips;
+		}
+
+		Remesher.BasicRemeshPass();
+	}
+
+	DynamicMeshComponent->NotifyMeshUpdated();
+	GetToolManager()->PostInvalidation();
+	MeshStatisticsProperties->Update(*DynamicMeshComponent->GetMesh());
+
+	bResultValid = true;
+}
+
+bool URemeshMeshTool::HasAccept() const
+{
+	return true;
+}
+
+bool URemeshMeshTool::CanAccept() const
+{
+	return true;
+}
+
+
+
+
+#undef LOCTEXT_NAMESPACE
diff --git a/Engine/Plugins/Experimental/MeshModelingToolset/Source/MeshModelingTools/Private/TransformMeshesTool.cpp b/Engine/Plugins/Experimental/MeshModelingToolset/Source/MeshModelingTools/Private/TransformMeshesTool.cpp
new file mode 100644
--- /dev/null
+++ b/Engine/Plugins/Experimental/MeshModelingToolset/Source/MeshModelingTools/Private/TransformMeshesTool.cpp
@@ -0,0 +1,284 @@
+// Copyright 1998-2019 Epic Games, Inc. All Rights Reserved.
+
+#include "TransformMeshesTool.h"
+#include "InteractiveToolManager.h"
+#include "InteractiveGizmoManager.h"
+#include "ToolBuilderUtil.h"
+#include "ToolSetupUtil.h"
+#include "DynamicMesh3.h"
+#include "BaseBehaviors/ClickDragBehavior.h"
+
+#include "BaseGizmos/GizmoComponents.h"
+#include "BaseGizmos/TransformGizmo.h"
+
+#include "Components/PrimitiveComponent.h"
+#include "CollisionQueryParams.h"
+#include "Engine/World.h"
+
+#define LOCTEXT_NAMESPACE "UTransformMeshesTool"
+
+/*
+ * ToolBuilder
+ */
+
+
+bool UTransformMeshesToolBuilder::CanBuildTool(const FToolBuilderState& SceneState) const
+{
+	return ToolBuilderUtil::CountComponents(SceneState, CanMakeComponentTarget) >= 1;
+}
+
+UInteractiveTool* UTransformMeshesToolBuilder::BuildTool(const FToolBuilderState& SceneState) const
+{
+	UTransformMeshesTool* NewTool = NewObject<UTransformMeshesTool>(SceneState.ToolManager);
+
+	TArray<UActorComponent*> Components = ToolBuilderUtil::FindAllComponents(SceneState, CanMakeComponentTarget);
+	check(Components.Num() > 0);
+
+	TArray<TUniquePtr<FPrimitiveComponentTarget>> ComponentTargets;
+	for (UActorComponent* ActorComponent : Components)
+	{
+		auto* MeshComponent = Cast<UPrimitiveComponent>(ActorComponent);
+		if ( MeshComponent )
+		{
+			ComponentTargets.Add(MakeComponentTarget(MeshComponent));
+		}
+	}
+
+	NewTool->SetSelection(MoveTemp(ComponentTargets));
+	NewTool->SetWorld(SceneState.World, SceneState.GizmoManager);
+
+	return NewTool;
+}
+
+
+
+
+/*
+ * Tool
+ */
+
+UTransformMeshesTool::UTransformMeshesTool()
+{
+}
+
+void UTransformMeshesTool::SetWorld(UWorld* World, UInteractiveGizmoManager* GizmoManagerIn)
+{
+	this->TargetWorld = World;
+	this->GizmoManager = GizmoManagerIn;
+}
+
+
+void UTransformMeshesTool::Setup()
+{
+	UInteractiveTool::Setup();
+
+	//UClickDragInputBehavior* ClickDragBehavior = NewObject<UClickDragInputBehavior>(this);
+	//ClickDragBehavior->Initialize(this);
+	//AddInputBehavior(ClickDragBehavior);
+
+	TransformProps = NewObject<UTransformMeshesToolProperties>();
+	AddToolPropertySource(TransformProps);
+
+	UpdateTransformMode(TransformProps->TransformMode);
+}
+
+
+
+
+
+
+
+void UTransformMeshesTool::Shutdown(EToolShutdownType ShutdownType)
+{
+	GizmoManager->DestroyAllGizmosByOwner(this);
+}
+
+
+
+void UTransformMeshesTool::Tick(float DeltaTime)
+{
+}
+
+
+void UTransformMeshesTool::Render(IToolsContextRenderAPI* RenderAPI)
+{
+	FPrimitiveDrawInterface* PDI = RenderAPI->GetPrimitiveDrawInterface();
+}
+
+void UTransformMeshesTool::OnPropertyModified(UObject* PropertySet, UProperty* Property)
+{
+	if (CurTransformMode != TransformProps->TransformMode)
+	{
+		UpdateTransformMode(TransformProps->TransformMode);
+	}
+
+	bool bEnableSetPivot = TransformProps->bSetPivot;
+	if (TransformProps->TransformMode == ETransformMeshesTransformMode::SharedGizmoLocal ||
+		TransformProps->TransformMode == ETransformMeshesTransformMode::QuickTranslate)
+	{
+		bEnableSetPivot = false;
+	}
+	UpdateSetPivotModes(bEnableSetPivot);
+}
+
+
+void UTransformMeshesTool::UpdateSetPivotModes(bool bEnableSetPivot)
+{
+	for (FTransformMeshesTarget& Target : ActiveGizmos)
+	{
+		Target.TransformProxy->bSetPivotMode = bEnableSetPivot;
+	}
+}
+
+
+
+void UTransformMeshesTool::UpdateTransformMode(ETransformMeshesTransformMode NewMode)
+{
+	ResetActiveGizmos();
+
+	switch (NewMode)
+	{
+		default:
+		case ETransformMeshesTransformMode::SharedGizmo:
+			SetActiveGizmos_Single(false);
+			break;
+
+		case ETransformMeshesTransformMode::SharedGizmoLocal:
+			SetActiveGizmos_Single(true);
+			break;
+
+		case ETransformMeshesTransformMode::PerObjectGizmo:
+			SetActiveGizmos_PerObject();
+			break;
+	}
+
+	CurTransformMode = NewMode;
+}
+
+
+
+void UTransformMeshesTool::SetActiveGizmos_Single(bool bLocalRotations)
+{
+	check(ActiveGizmos.Num() == 0);
+
+	FTransformMeshesTarget Transformable;
+	Transformable.TransformProxy = NewObject<UTransformProxy>(this);
+	Transformable.TransformProxy->bRotatePerObject = bLocalRotations;
+
+	for (TUniquePtr<FPrimitiveComponentTarget>& Target : ComponentTargets)
+	{
+		Transformable.TransformProxy->AddComponent(Target->GetOwnerComponent());
+	}
+	Transformable.TransformGizmo = GizmoManager->Create3AxisTransformGizmo(FString(), this);
+	Transformable.TransformGizmo->SetActiveTarget(Transformable.TransformProxy);
+
+	ActiveGizmos.Add(Transformable);
+}
+
+void UTransformMeshesTool::SetActiveGizmos_PerObject()
+{
+	check(ActiveGizmos.Num() == 0);
+
+	for (TUniquePtr<FPrimitiveComponentTarget>& Target : ComponentTargets)
+	{
+		FTransformMeshesTarget Transformable;
+		Transformable.TransformProxy = NewObject<UTransformProxy>(this);
+		Transformable.TransformProxy->AddComponent(Target->GetOwnerComponent());
+		Transformable.TransformGizmo = GizmoManager->Create3AxisTransformGizmo(FString(), this);
+		Transformable.TransformGizmo->SetActiveTarget(Transformable.TransformProxy);
+		ActiveGizmos.Add(Transformable);
+	}
+}
+
+void UTransformMeshesTool::ResetActiveGizmos()
+{
+	GizmoManager->DestroyAllGizmosByOwner(this);
+	ActiveGizmos.Reset();
+}
+
+
+
+// does not make sense that CanBeginClickDragSequence() returns a RayHit? Needs to be an out-argument...
+FInputRayHit UTransformMeshesTool::CanBeginClickDragSequence(const FInputDeviceRay& PressPos)
+{
+	check(ActiveGizmos.Num() == 1);
+
+	float MinHitDistance = TNumericLimits<float>::Max();
+	for (TUniquePtr<FPrimitiveComponentTarget>& Target : ComponentTargets)
+	{
+		FHitResult WorldHit;
+		if (Target->HitTest(PressPos.WorldRay, WorldHit))
+		{
+			MinHitDistance = FMath::Min(MinHitDistance, WorldHit.Distance);
+		}
+	}
+	return (MinHitDistance < TNumericLimits<float>::Max()) ? FInputRayHit(MinHitDistance) : FInputRayHit();
+}
+
+void UTransformMeshesTool::OnClickPress(const FInputDeviceRay& PressPos)
+{
+	check(ActiveGizmos.Num() == 1);
+
+	FInputRayHit HitPos = CanBeginClickDragSequence(PressPos);
+	check(HitPos.bHit);
+
+	GetToolManager()->BeginUndoTransaction(LOCTEXT("TransformToolTransformTxnName", "Transform"));
+
+	StartDragFrameWorld = FFrame3d(PressPos.WorldRay.PointAt(HitPos.HitDepth));
+	USceneComponent* GizmoComponent = ActiveGizmos[0].TransformGizmo->GetGizmoActor()->GetRootComponent();
+	StartDragTransform = GizmoComponent->GetComponentToWorld();
+	GizmoComponent->Modify();
+
+	for (TUniquePtr<FPrimitiveComponentTarget>& Target : ComponentTargets)
+	{
+		Target->GetOwnerComponent()->Modify();
+	}
+}
+
+void UTransformMeshesTool::OnClickDrag(const FInputDeviceRay& DragPos)
+{
+	for (TUniquePtr<FPrimitiveComponentTarget>& Target : ComponentTargets)
+	{
+		Target->SetOwnerVisibility(false);
+	}
+
+	FCollisionObjectQueryParams QueryParams(FCollisionObjectQueryParams::AllObjects);
+	FHitResult Result;
+	bool bWorldHit = TargetWorld->LineTraceSingleByObjectType(Result, DragPos.WorldRay.Origin, DragPos.WorldRay.PointAt(999999), QueryParams);
+	if (bWorldHit)
+	{
+		FVector HitPos = Result.ImpactPoint;
+		FFrame3d NewFrame = StartDragFrameWorld;
+		NewFrame.Origin = (FVector3d)HitPos;
+
+		FTransform RelTransform = FTransform::Identity;
+		RelTransform.SetTranslation(
+			(FVector)(NewFrame.Origin - StartDragFrameWorld.Origin));
+
+		USceneComponent* GizmoComponent = ActiveGizmos[0].TransformGizmo->GetGizmoActor()->GetRootComponent();
+		FTransform NewTransform = StartDragTransform;
+		NewTransform.Accumulate(RelTransform);
+		GizmoComponent->SetWorldTransform(NewTransform);
+	}
+
+	for (TUniquePtr<FPrimitiveComponentTarget>& Target : ComponentTargets)
+	{
+		Target->SetOwnerVisibility(true);
+	}
+}
+
+void UTransformMeshesTool::OnClickRelease(const FInputDeviceRay& ReleasePos)
+{
+	GetToolManager()->EndUndoTransaction();
+}
+
+void UTransformMeshesTool::OnTerminateDragSequence()
+{
+	GetToolManager()->EndUndoTransaction();
+}
+
+
+
+
+
+#undef LOCTEXT_NAMESPACE
diff --git a/Engine/Plugins/Experimental/MeshModelingToolset/Source/MeshModelingTools/Private/UVProjectionTool.cpp b/Engine/Plugins/Experimental/MeshModelingToolset/Source/MeshModelingTools/Private/UVProjectionTool.cpp
new file mode 100644
--- /dev/null
+++ b/Engine/Plugins/Experimental/MeshModelingToolset/Source/MeshModelingTools/Private/UVProjectionTool.cpp
@@ -0,0 +1,367 @@
+// Copyright 1998-2019 Epic Games, Inc. All Rights Reserved.
+
+#include "UVProjectionTool.h"
+#include "InteractiveToolManager.h"
+#include "ToolBuilderUtil.h"
+
+#include "ToolSetupUtil.h"
+
+#include "DynamicMesh3.h"
+#include "BaseBehaviors/MultiClickSequenceInputBehavior.h"
+#include "Selection/SelectClickedAction.h"
+
+#include "MeshDescriptionToDynamicMesh.h"
+#include "DynamicMeshToMeshDescription.h"
+
+#include "InteractiveGizmoManager.h"
+
+#include "AssetGenerationUtil.h"
+
+#include "BaseGizmos/GizmoComponents.h"
+#include "BaseGizmos/TransformGizmo.h"
+
+
+#define LOCTEXT_NAMESPACE "UUVProjectionTool"
+
+
+/*
+ * ToolBuilder
+ */
+
+
+bool UUVProjectionToolBuilder::CanBuildTool(const FToolBuilderState& SceneState) const
+{
+	// note most of the code is written to support working on any number > 0, but that seems maybe confusing UI-wise and is not fully tested, so I've limited it to acting on one for now
+	// TODO: if enable tool working on multiple components, figure out what to do if we have multiple component targets that point to the same underlying mesh data?
+	return ToolBuilderUtil::CountComponents(SceneState, CanMakeComponentTarget) == 1;
+}
+
+UInteractiveTool* UUVProjectionToolBuilder::BuildTool(const FToolBuilderState& SceneState) const
+{
+	UUVProjectionTool* NewTool = NewObject<UUVProjectionTool>(SceneState.ToolManager);
+
+	TArray<UActorComponent*> Components = ToolBuilderUtil::FindAllComponents(SceneState, CanMakeComponentTarget);
+	check(Components.Num() > 0);
+
+	TArray<TUniquePtr<FPrimitiveComponentTarget>> ComponentTargets;
+	for (UActorComponent* ActorComponent : Components)
+	{
+		auto* MeshComponent = Cast<UPrimitiveComponent>(ActorComponent);
+		if ( MeshComponent )
+		{
+			ComponentTargets.Add(MakeComponentTarget(MeshComponent));
+		}
+	}
+
+	NewTool->SetSelection(MoveTemp(ComponentTargets));
+	NewTool->SetWorld(SceneState.World, SceneState.GizmoManager);
+	NewTool->SetAssetAPI(AssetAPI);
+
+	return NewTool;
+}
+
+
+
+
+/*
+ * Tool
+ */
+
+UUVProjectionToolProperties::UUVProjectionToolProperties()
+{
+	UVProjectionMethod = EUVProjectionMethod::Cube;
+	ProjectionPrimitiveScale = FVector::OneVector;
+	UVScale = FVector2D::UnitVector;
+	UVOffset = FVector2D::ZeroVector;
+}
+
+UUVProjectionAdvancedProperties::UUVProjectionAdvancedProperties()
+{
+}
+
+
+UUVProjectionTool::UUVProjectionTool()
+{
+}
+
+void UUVProjectionTool::SetWorld(UWorld* World, UInteractiveGizmoManager* GizmoManagerIn)
+{
+	this->TargetWorld = World;
+	this->GizmoManager = GizmoManagerIn;
+}
+
+void UUVProjectionTool::Setup()
+{
+	UInteractiveTool::Setup();
+
+	// hide input StaticMeshComponent
+	for (TUniquePtr<FPrimitiveComponentTarget>& ComponentTarget : ComponentTargets)
+	{
+		ComponentTarget->SetOwnerVisibility(false);
+	}
+
+	BasicProperties = NewObject<UUVProjectionToolProperties>(this, TEXT("UV Projection Settings"));
+	AdvancedProperties = NewObject<UUVProjectionAdvancedProperties>(this, TEXT("Advanced Settings"));
+
+	// initialize our properties
+	AddToolPropertySource(BasicProperties);
+	AddToolPropertySource(AdvancedProperties);
+
+	// Set up check material
+	UMaterial* CheckerMaterialBase = LoadObject<UMaterial>(nullptr, TEXT("/MeshModelingToolset/Materials/CheckerMaterial"));
+	CheckerMaterial = UMaterialInstanceDynamic::Create(CheckerMaterialBase, NULL);
+	CheckerMaterial->SetScalarParameterValue("Density", 10);
+
+	// initialize the PreviewMesh+BackgroundCompute object
+	UpdateNumPreviews();
+
+	// set up visualizers
+	ProjectionShapeVisualizer.LineColor = FLinearColor::Red;
+	ProjectionShapeVisualizer.LineThickness = 2.0;
+
+	
+	for (UMeshOpPreviewWithBackgroundCompute* Preview : Previews)
+	{
+		Preview->InvalidateResult();
+	}
+}
+
+
+void UUVProjectionTool::UpdateNumPreviews()
+{
+	int32 CurrentNumPreview = Previews.Num();
+	int32 TargetNumPreview = ComponentTargets.Num();
+	if (TargetNumPreview < CurrentNumPreview)
+	{
+		for (int32 PreviewIdx = CurrentNumPreview - 1; PreviewIdx >= TargetNumPreview; PreviewIdx--)
+		{
+			Previews[PreviewIdx]->Cancel();
+			GizmoManager->DestroyGizmo(TransformGizmos[PreviewIdx]);
+		}
+		Previews.SetNum(TargetNumPreview);
+		TransformGizmos.SetNum(TargetNumPreview);
+		TransformProxies.SetNum(TargetNumPreview);
+		OriginalDynamicMeshes.SetNum(TargetNumPreview);
+	}
+	else
+	{
+		OriginalDynamicMeshes.SetNum(TargetNumPreview);
+		for (int32 PreviewIdx = CurrentNumPreview; PreviewIdx < TargetNumPreview; PreviewIdx++)
+		{
+			UUVProjectionOperatorFactory *OpFactory = NewObject<UUVProjectionOperatorFactory>();
+			OpFactory->Tool = this;
+			OpFactory->ComponentIndex = PreviewIdx;
+			OriginalDynamicMeshes[PreviewIdx] = MakeShared<FDynamicMesh3>();
+			FMeshDescriptionToDynamicMesh Converter;
+			Converter.bPrintDebugMessages = true;
+			Converter.Convert(ComponentTargets[PreviewIdx]->GetMesh(), *OriginalDynamicMeshes[PreviewIdx]);
+
+			FVector Center, Extents;
+			FBoxSphereBounds Bounds = ComponentTargets[PreviewIdx]->GetOwnerComponent()->CalcLocalBounds();
+			
+			FTransform LocalXF(Bounds.Origin);
+			LocalXF.SetScale3D(Bounds.BoxExtent);
+
+			UMeshOpPreviewWithBackgroundCompute* Preview = Previews.Add_GetRef(NewObject<UMeshOpPreviewWithBackgroundCompute>(OpFactory, "Preview"));
+			Preview->Setup(this->TargetWorld, OpFactory);
+			Preview->ConfigureMaterials(
+				ToolSetupUtil::GetDefaultMaterial(GetToolManager(), CheckerMaterial),//ComponentTargets[PreviewIdx]->GetMaterial(0)),
+				ToolSetupUtil::GetDefaultWorkingMaterial(GetToolManager())
+			);
+			Preview->SetVisibility(true);
+
+
+			UTransformProxy* TransformProxy = TransformProxies.Add_GetRef(NewObject<UTransformProxy>(this));
+			TransformProxy->SetTransform(LocalXF * ComponentTargets[PreviewIdx]->GetWorldTransform());
+			TransformProxy->OnTransformChanged.AddUObject(this, &UUVProjectionTool::TransformChanged);
+
+			UTransformGizmo* TransformGizmo = TransformGizmos.Add_GetRef(GizmoManager->Create3AxisTransformGizmo(TEXT("UVProjectionToolTransformGizmo"), this));
+			TransformGizmo->SetActiveTarget(TransformProxy);
+		}
+		check(TransformProxies.Num() == TargetNumPreview);
+		check(TransformGizmos.Num() == TargetNumPreview);
+	}
+}
+
+
+void UUVProjectionTool::Shutdown(EToolShutdownType ShutdownType)
+{
+	// Restore (unhide) the source meshes
+	for (TUniquePtr<FPrimitiveComponentTarget>& ComponentTarget : ComponentTargets)
+	{
+		ComponentTarget->SetOwnerVisibility(true);
+	}
+
+	TArray<TUniquePtr<FDynamicMeshOpResult>> Results;
+	for (UMeshOpPreviewWithBackgroundCompute* Preview : Previews)
+	{
+		Results.Emplace(Preview->Shutdown());
+	}
+	if (ShutdownType == EToolShutdownType::Accept)
+	{
+		GenerateAsset(Results);
+	}
+
+	// TODO MAKE THIS UNNECESSARY GIZMOS SHOULD BE AUTO DESTROYED!!
+	for (UTransformGizmo* TransformGizmo : TransformGizmos)
+	{
+		if (TransformGizmo != nullptr)
+		{
+			GizmoManager->DestroyGizmo(TransformGizmo);
+		}
+	}
+	TransformGizmos.Empty();
+
+	GizmoManager->DeregisterGizmoType(TEXT("TransformGizmo"));
+}
+
+void UUVProjectionTool::SetAssetAPI(IToolsContextAssetAPI* AssetAPIIn)
+{
+	this->AssetAPI = AssetAPIIn;
+}
+
+TSharedPtr<FDynamicMeshOperator> UUVProjectionOperatorFactory::MakeNewOperator()
+{
+	TSharedPtr<FUVProjectionOp> Op = MakeShared<FUVProjectionOp>();
+	Op->ProjectionMethod = Tool->BasicProperties->UVProjectionMethod;
+
+	// TODO: de-dupe this logic (it's also in Render, below)
+	FTransform LocalScale = FTransform::Identity;
+	LocalScale.SetScale3D(Tool->BasicProperties->ProjectionPrimitiveScale);
+	Op->ProjectionTransform = LocalScale * Tool->TransformProxies[ComponentIndex]->GetTransform();
+	Op->CylinderProjectToTopOrBottomAngleThreshold = Tool->BasicProperties->CylinderProjectToTopOrBottomAngleThreshold;
+	Op->UVScale = Tool->BasicProperties->UVScale;
+	Op->UVOffset = Tool->BasicProperties->UVOffset;
+
+	FTransform LocalToWorld = Tool->ComponentTargets[ComponentIndex]->GetWorldTransform();
+	Op->OriginalMesh = Tool->OriginalDynamicMeshes[ComponentIndex];
+	
+	Op->SetTransform(LocalToWorld);
+
+	return Op;
+}
+
+
+
+void UUVProjectionTool::Render(IToolsContextRenderAPI* RenderAPI)
+{
+	ProjectionShapeVisualizer.bDepthTested = false;
+	ProjectionShapeVisualizer.BeginFrame(RenderAPI, CameraState);
+	FTransform LocalScale = FTransform::Identity;
+	LocalScale.SetScale3D(BasicProperties->ProjectionPrimitiveScale);
+	for (UTransformProxy* TransformProxy : TransformProxies)
+	{
+		ProjectionShapeVisualizer.SetTransform(LocalScale * TransformProxy->GetTransform());
+
+		switch (BasicProperties->UVProjectionMethod)
+		{
+		case EUVProjectionMethod::Cube:
+			ProjectionShapeVisualizer.DrawWireBox(FBox(FVector(-1, -1, -1), FVector(1, 1, 1)));
+			break;
+		case EUVProjectionMethod::Cylinder:
+			ProjectionShapeVisualizer.DrawWireCylinder(FVector(0, 0, -1), FVector(0, 0, 1), 1, 2, 20);
+			break;
+		case EUVProjectionMethod::Plane:
+			ProjectionShapeVisualizer.DrawSquare(FVector(0, 0, 0), FVector(2, 0, 0), FVector(0, 2, 0));
+			break;
+		}
+	}
+	
+	ProjectionShapeVisualizer.EndFrame();
+}
+
+void UUVProjectionTool::Tick(float DeltaTime)
+{
+	for (UMeshOpPreviewWithBackgroundCompute* Preview : Previews)
+	{
+		Preview->Tick(DeltaTime);
+	}
+}
+
+
+
+#if WITH_EDITOR
+void UUVProjectionTool::PostEditChangeProperty(FPropertyChangedEvent& PropertyChangedEvent)
+{
+	UpdateNumPreviews();
+	for (UMeshOpPreviewWithBackgroundCompute* Preview : Previews)
+	{
+		Preview->InvalidateResult();
+	}
+}
+#endif
+
+void UUVProjectionTool::OnPropertyModified(UObject* PropertySet, UProperty* Property)
+{
+	if (Property->GetFName() == GET_MEMBER_NAME_CHECKED(UUVProjectionToolProperties, CheckerDensity))
+	{
+		CheckerMaterial->SetScalarParameterValue("Density", BasicProperties->CheckerDensity);
+	}
+	UpdateNumPreviews();
+	for (UMeshOpPreviewWithBackgroundCompute* Preview : Previews)
+	{
+		Preview->InvalidateResult();
+	}
+}
+
+
+void UUVProjectionTool::TransformChanged(UTransformProxy* Proxy, FTransform Transform)
+{
+	// TODO: if multi-select is re-enabled, only invalidate the preview that actually needs it?
+	for (UMeshOpPreviewWithBackgroundCompute* Preview : Previews)
+	{
+		Preview->InvalidateResult();
+	}
+}
+
+
+bool UUVProjectionTool::HasAccept() const
+{
+	return true;
+}
+
+bool UUVProjectionTool::CanAccept() const
+{
+	for (UMeshOpPreviewWithBackgroundCompute* Preview : Previews)
+	{
+		if (!Preview->HaveValidResult())
+		{
+			return false;
+		}
+	}
+	return true;
+}
+
+
+void UUVProjectionTool::GenerateAsset(const TArray<TUniquePtr<FDynamicMeshOpResult>>& Results)
+{
+	GetToolManager()->BeginUndoTransaction(LOCTEXT("UVProjectionToolTransactionName", "UV Projection Tool"));
+
+	check(Results.Num() == ComponentTargets.Num());
+	
+	for (int32 ComponentIdx = 0; ComponentIdx < ComponentTargets.Num(); ComponentIdx++)
+	{
+		check(Results[ComponentIdx]->Mesh.Get() != nullptr);
+		ComponentTargets[ComponentIdx]->CommitMesh([&Results, &ComponentIdx, this](FMeshDescription* MeshDescription)
+		{
+			FDynamicMeshToMeshDescription Converter;
+			//if (??) // TODO check if UV topology changed?  or remove all traces of this if statement (may be safe to assume it almost always changes w/ a uv projection op)
+			{
+				// full conversion if topology changed
+				Converter.Convert(Results[ComponentIdx]->Mesh.Get(), *MeshDescription);
+			}
+			//else
+			//{
+			//	// otherwise just copy attributes
+			//	Converter.UpdateAttributes(Results[ComponentIdx]->Mesh.Get(), *MeshDescription, false, true);
+			//}
+		});
+	}
+
+	GetToolManager()->EndUndoTransaction();
+}
+
+
+
+
+#undef LOCTEXT_NAMESPACE
diff --git a/Engine/Plugins/Experimental/MeshModelingToolset/Source/MeshModelingTools/Public/BaseBrushTool.h b/Engine/Plugins/Experimental/MeshModelingToolset/Source/MeshModelingTools/Public/BaseBrushTool.h
new file mode 100644
--- /dev/null
+++ b/Engine/Plugins/Experimental/MeshModelingToolset/Source/MeshModelingTools/Public/BaseBrushTool.h
@@ -0,0 +1,135 @@
+// Copyright 1998-2019 Epic Games, Inc. All Rights Reserved.
+
+#pragma once
+
+#include "CoreMinimal.h"
+#include "BaseTools/MeshSurfacePointTool.h"
+#include "BoxTypes.h"
+#include "ToolIndicatorSet.h"
+#include "BaseBrushTool.generated.h"
+
+
+
+/**
+ * Standard properties for a Brush-type Tool
+ */
+UCLASS()
+class MESHMODELINGTOOLS_API UBrushBaseProperties : public UInteractiveToolPropertySet
+{
+	GENERATED_BODY()
+
+public:
+	UBrushBaseProperties();
+
+	/** Relative size of brush */
+	UPROPERTY(EditAnywhere, Category = Options, meta = (DisplayName = "Size", UIMin = "0.0", UIMax = "1.0"))
+	float BrushSize;
+
+	/** If true, ignore relative Brush Size and use explicit Radius*/
+	UPROPERTY(EditAnywhere, Category = Options, AdvancedDisplay)
+	bool bSpecifyRadius;
+
+	/** Radius of brush */
+	UPROPERTY(EditAnywhere, Category = Options, AdvancedDisplay, meta = (DisplayName = "Radius", UIMin = "1.0", UIMax = "1000.0", ClampMin = "0.1", ClampMax = "50000.0"))
+	float BrushRadius;
+};
+
+
+/**
+ * Generic Brush Stamp data
+ */
+USTRUCT()
+struct MESHMODELINGTOOLS_API FBrushStampData
+{
+	GENERATED_BODY();
+	/** Radius of brush stamp */
+	float Radius;
+	/** World Position of brush stamp */
+	FVector WorldPosition;
+	/** World Normal of brush stamp */
+	FVector WorldNormal;
+
+	/** Hit Result provided by implementations - may not be fully populated */
+	FHitResult HitResult;
+};
+
+
+
+/**
+ * UBaseBrushTool implements standard brush-style functionality for an InteractiveTool.
+ * This includes:
+ *   1) brush radius property set with dimension-relative brush sizing and default brush radius hotkeys
+ *   2) brush indicator visualization
+ *   3) tracking of last brush stamp location via .LastBrushStamp UProperty
+ *   4) status of brush stroke via .bInBrushStroke UProperty
+ */
+UCLASS()
+class MESHMODELINGTOOLS_API UBaseBrushTool : public UMeshSurfacePointTool
+{
+	GENERATED_BODY()
+
+public:
+	UBaseBrushTool();
+
+	virtual void RegisterActions(FInteractiveToolActionSet& ActionSet) override;
+
+	virtual void Setup() override;
+	virtual void Shutdown(EToolShutdownType ShutdownType) override;
+
+	virtual void Tick(float DeltaTime) override;
+	virtual void Render(IToolsContextRenderAPI* RenderAPI) override;
+
+	virtual void OnPropertyModified(UObject* PropertySet, UProperty* Property) override;
+
+
+	//
+	// UMeshSurfacePointTool implementation
+	//
+	virtual void OnBeginDrag(const FRay& Ray) override;
+	virtual void OnUpdateDrag(const FRay& Ray) override;
+	virtual void OnEndDrag(const FRay& Ray) override;
+	virtual void OnUpdateHover(const FInputDeviceRay& DevicePos) override;
+
+public:
+
+	/** Properties that control brush size/etc*/
+	UPROPERTY()
+	UBrushBaseProperties* BrushProperties;
+
+	/** Brush indicator */
+	UPROPERTY()
+	UToolIndicatorSet* Indicators;
+
+	/** Set to true by Tool if user is currently in an active brush stroke*/
+	UPROPERTY()
+	bool bInBrushStroke = false;
+
+	/** Position of brush at last update (both during stroke and during Hover) */
+	UPROPERTY()
+	FBrushStampData LastBrushStamp;
+
+public:
+
+	virtual void IncreaseBrushSizeAction();
+	virtual void DecreaseBrushSizeAction();
+
+
+	virtual bool IsInBrushStroke() const { return bInBrushStroke; }
+
+protected:
+
+	/**
+	 * Subclasses should implement this to give an estimate of target dimension for brush size scaling
+	 */
+	virtual double EstimateMaximumTargetDimension() { return 100.0; }
+
+	// creates brush indicator widget
+	virtual IToolIndicator* MakeBrushIndicator();
+
+
+protected:
+	FInterval1d BrushRelativeSizeRange;
+	double CurrentBrushRadius;
+	void RecalculateBrushRadius();
+
+};
diff --git a/Engine/Plugins/Experimental/MeshModelingToolset/Source/MeshModelingTools/Public/BrushToolIndicator.h b/Engine/Plugins/Experimental/MeshModelingToolset/Source/MeshModelingTools/Public/BrushToolIndicator.h
new file mode 100644
--- /dev/null
+++ b/Engine/Plugins/Experimental/MeshModelingToolset/Source/MeshModelingTools/Public/BrushToolIndicator.h
@@ -0,0 +1,51 @@
+// Copyright 1998-2019 Epic Games, Inc. All Rights Reserved.
+
+#pragma once
+
+#include "CoreMinimal.h"
+#include "Misc/Attribute.h"
+#include "VectorTypes.h"
+#include "ToolIndicatorSet.h"
+#include "BrushToolIndicator.generated.h"
+
+/*
+ * Circle indicator
+ */
+UCLASS()
+class UBrushStampSizeIndicator : public UObject, public IToolIndicator
+{
+	GENERATED_BODY()
+
+public:
+	UBrushStampSizeIndicator();
+	virtual ~UBrushStampSizeIndicator();
+
+	virtual void Connect(UInteractiveTool* Tool) override;
+	virtual void Disconnect() override;
+
+	virtual void Render(IToolsContextRenderAPI* RenderAPI) override;
+	virtual void Tick(float DeltaTime) override;
+
+	UInteractiveTool* ParentTool;
+	float Radius;
+	FVector Center;
+	FVector Normal;
+	int SampleStepCount;
+	FColor LineColor;
+	float LineThickness;
+	bool IsPixelThickness;
+	int DepthLayer;
+
+	bool bDrawSecondaryLines;
+	FColor SecondaryLineColor;
+	float SecondaryLineThickness;
+
+	TAttribute<float> BrushRadius;
+	TAttribute<FVector> BrushPosition;
+	TAttribute<FVector> BrushNormal;
+
+
+protected:
+	void DrawCircle(FPrimitiveDrawInterface* PDI, const FVector3f& AxisX, const FVector3f& AxisY, float UseRadius, bool bIsSecondary);
+	void DrawLine(FPrimitiveDrawInterface* PDI, const FVector3f& Start, const FVector3f& End, bool bIsSecondary);
+};
\ No newline at end of file
diff --git a/Engine/Plugins/Experimental/MeshModelingToolset/Source/MeshModelingTools/Public/DisplaceMeshTool.h b/Engine/Plugins/Experimental/MeshModelingToolset/Source/MeshModelingTools/Public/DisplaceMeshTool.h
new file mode 100644
--- /dev/null
+++ b/Engine/Plugins/Experimental/MeshModelingToolset/Source/MeshModelingTools/Public/DisplaceMeshTool.h
@@ -0,0 +1,129 @@
+// Copyright 1998-2019 Epic Games, Inc. All Rights Reserved.
+
+#pragma once
+
+#include "CoreMinimal.h"
+#include "UObject/NoExportTypes.h"
+#include "Engine/Classes/Engine/Texture2D.h"
+#include "SingleSelectionTool.h"
+#include "InteractiveToolBuilder.h"
+#include "DynamicMesh3.h"
+#include "MeshNormals.h"
+#include "Spatial/SampledScalarField2.h"
+#include "DisplaceMeshTool.generated.h"
+
+
+// predeclarations
+struct FMeshDescription;
+class USimpleDynamicMeshComponent;
+
+
+UENUM()
+enum class EDisplaceMeshToolDisplaceType : uint8
+{
+	/** Displace with N iterations */
+	Constant UMETA(DisplayName = "Constant"),
+
+	/** Displace with N iterations */
+	RandomNoise UMETA(DisplayName = "Random Noise"),
+
+	/** Displace with N iterations */
+	DisplacementMap UMETA(DisplayName = "Displacement Map"),
+
+};
+
+
+
+/**
+ *
+ */
+UCLASS()
+class MESHMODELINGTOOLS_API UDisplaceMeshToolBuilder : public UInteractiveToolBuilder
+{
+	GENERATED_BODY()
+
+public:
+	virtual bool CanBuildTool(const FToolBuilderState& SceneState) const override;
+	virtual UInteractiveTool* BuildTool(const FToolBuilderState& SceneState) const override;
+};
+
+
+
+
+
+/**
+ * Simple Mesh Displacement Tool
+ */
+UCLASS()
+class MESHMODELINGTOOLS_API UDisplaceMeshTool : public USingleSelectionTool
+{
+	GENERATED_BODY()
+
+public:
+	UDisplaceMeshTool();
+
+	virtual void Setup() override;
+	virtual void Shutdown(EToolShutdownType ShutdownType) override;
+
+	virtual void Render(IToolsContextRenderAPI* RenderAPI) override;
+
+	virtual bool HasCancel() const override { return true; }
+	virtual bool HasAccept() const override;
+	virtual bool CanAccept() const override;
+#if WITH_EDITOR
+	virtual void PostEditChangeProperty(FPropertyChangedEvent& PropertyChangedEvent) override;
+#endif
+protected:
+	// need to update bResultValid if these are modified, so we don't publicly expose them. 
+	// @todo setters/getters for these
+
+	/** primary brush mode */
+	UPROPERTY(EditAnywhere, Category = Options)
+	EDisplaceMeshToolDisplaceType DisplacementType;
+
+	/** Displacement intensity */
+	UPROPERTY(EditAnywhere, Category = Options, meta = (UIMin = "-100.0", UIMax = "100.0", ClampMin = "-10000.0", ClampMax = "100000.0"))
+	float DisplaceIntensity;
+
+	/** Seed for randomization */
+	UPROPERTY(EditAnywhere, Category = Options)
+	int RandomSeed;
+
+	/** Subdivision iterations for mesh */
+	UPROPERTY(EditAnywhere, Category = Options, meta = (UIMin = "0", UIMax = "10", ClampMin = "0", ClampMax = "100"))
+	int Subdivisions;
+
+	/** Displacement map */
+	UPROPERTY(EditAnywhere, Category = Options)
+	UTexture2D* DisplacementMap;
+
+
+protected:
+	USimpleDynamicMeshComponent* DynamicMeshComponent;
+	
+	FDynamicMesh3 OriginalMesh;		// don't actually need whole mesh, just positions...
+
+	FDynamicMesh3 SubdividedMesh;
+	int CachedSubdivisionsCount = -1;
+	void UpdateSubdividedMesh();
+
+	FSampledScalarField2f DisplaceField;
+	UTexture2D* CachedMapSource = nullptr;
+	void UpdateMap(bool bForceUpdate = false);
+
+	TArray<FVector3d> PositionBuffer;
+	FMeshNormals NormalsBuffer;
+	TArray<FVector3d> DisplacedBuffer;
+	bool bResultValid;
+	void UpdateResult();
+
+	void ComputeDisplacement_Constant();
+	void ComputeDisplacement_RandomNoise();
+
+	void ComputeDisplacement_Map();
+
+
+private:
+	//TUniquePtr<
+
+};
diff --git a/Engine/Plugins/Experimental/MeshModelingToolset/Source/MeshModelingTools/Public/DrawPolygonTool.h b/Engine/Plugins/Experimental/MeshModelingToolset/Source/MeshModelingTools/Public/DrawPolygonTool.h
new file mode 100644
--- /dev/null
+++ b/Engine/Plugins/Experimental/MeshModelingToolset/Source/MeshModelingTools/Public/DrawPolygonTool.h
@@ -0,0 +1,288 @@
+// Copyright 1998-2019 Epic Games, Inc. All Rights Reserved.
+
+#pragma once
+
+#include "CoreMinimal.h"
+#include "UObject/NoExportTypes.h"
+#include "InteractiveGizmo.h"
+#include "InteractiveToolBuilder.h"
+#include "BaseTools/SingleClickTool.h"
+#include "PreviewMesh.h"
+#include "DynamicMeshAABBTree3.h"
+#include "Snapping/PointPlanarSnapSolver.h"
+#include "ToolSceneQueriesUtil.h"
+#include "DrawPolygonTool.generated.h"
+
+
+/**
+ *
+ */
+UCLASS()
+class MESHMODELINGTOOLS_API UDrawPolygonToolBuilder : public UInteractiveToolBuilder
+{
+	GENERATED_BODY()
+
+public:
+	IToolsContextAssetAPI* AssetAPI;
+
+	UDrawPolygonToolBuilder() 
+	{
+		AssetAPI = nullptr;
+	}
+
+	virtual bool CanBuildTool(const FToolBuilderState& SceneState) const override;
+	virtual UInteractiveTool* BuildTool(const FToolBuilderState& SceneState) const override;
+};
+
+
+
+/** Polygon Tool Draw Type */
+UENUM()
+enum class EDrawPolygonDrawMode : uint8
+{
+	/** Freehand Polygon Drawing */
+	Freehand UMETA(DisplayName = "Freehand"),
+
+	/** Circle */
+	Circle UMETA(DisplayName = "Circle"),
+
+	/** Square */
+	Square UMETA(DisplayName = "Square"),
+
+	/** Rectangle */
+	Rectangle UMETA(DisplayName = "Rectangle")
+
+};
+
+
+/** Output of Draw Polygon Tool */
+UENUM()
+enum class EDrawPolygonOutputMode : uint8
+{
+	/** Planar polygon mesh */
+	MeshedPolygon UMETA(DisplayName = "Flat Mesh"),
+
+	/** Fixed-distance Extrusion */
+	ExtrudedConstant UMETA(DisplayName = "Extruded To Distance"),
+
+	/** extrusion height is set via additional mouse input after closing polygon */
+	ExtrudedInteractive UMETA(DisplayName = "Interactive Extrude"),
+};
+
+
+
+
+
+UCLASS()
+class MESHMODELINGTOOLS_API UDrawPolygonToolStandardProperties : public UObject
+{
+	GENERATED_BODY()
+
+public:
+	UPROPERTY(EditAnywhere, NonTransactional, Category = Polygon)
+	EDrawPolygonDrawMode PolygonType = EDrawPolygonDrawMode::Freehand;
+
+	UPROPERTY(EditAnywhere, NonTransactional, Category = Polygon)
+	EDrawPolygonOutputMode OutputMode = EDrawPolygonOutputMode::ExtrudedConstant;
+
+	/** Extrusion Distance in non-interactive mode */
+	UPROPERTY(EditAnywhere, NonTransactional, Category = Polygon, meta = (UIMin = "-1000", UIMax = "1000", ClampMin = "-10000", ClampMax = "10000"))
+	float ExtrudeHeight = 100.0f;
+
+	/** Extrusion Distance in non-interactive mode */
+	UPROPERTY(EditAnywhere, NonTransactional, Category = Polygon, meta = (UIMin = "3", UIMax = "100", ClampMin = "3", ClampMax = "10000"))
+	int Steps = 16;
+
+	UPROPERTY(EditAnywhere, NonTransactional, Category = Polygon)
+	bool bAllowSelfIntersections = false;
+};
+
+
+
+
+UCLASS()
+class MESHMODELINGTOOLS_API UDrawPolygonToolSnapProperties : public UObject
+{
+	GENERATED_BODY()
+
+public:
+	UPROPERTY(EditAnywhere, NonTransactional, Category = Snapping)
+	bool bEnableSnapping = true;
+
+	UPROPERTY(EditAnywhere, NonTransactional, Category = Snapping, Meta = (EditCondition = "bEnableSnapping"))
+	bool bSnapToVertices = true;
+
+	UPROPERTY(EditAnywhere, NonTransactional, Category = Snapping, Meta = (EditCondition = "bEnableSnapping"))
+	bool bSnapToEdges = false;
+
+	UPROPERTY(EditAnywhere, NonTransactional, Category = Snapping, Meta = (EditCondition = "bEnableSnapping"))
+	bool bSnapToAngles = true;
+
+	UPROPERTY(EditAnywhere, NonTransactional, Category = Snapping, Meta = (EditCondition = "bEnableSnapping"))
+	bool bSnapToLengths = true;
+
+	UPROPERTY(EditAnywhere, NonTransactional, Category = Snapping)
+	float Length = 0.0f;
+
+	UPROPERTY(EditAnywhere, NonTransactional, Category = Snapping)
+	bool bHitSceneObjects = false;
+
+	UPROPERTY(EditAnywhere, NonTransactional, Category = Snapping, Meta = (EditCondition = "bHitSceneObjects"))
+	float HitNormalOffset = 0.0f;
+};
+
+
+
+
+/**
+ * This tool allows the user to draw and extrude 2D polygons
+ */
+UCLASS()
+class MESHMODELINGTOOLS_API UDrawPolygonTool : public UInteractiveTool, public IClickSequenceBehaviorTarget
+{
+	GENERATED_BODY()
+
+public:
+	UDrawPolygonTool();
+
+	virtual void RegisterActions(FInteractiveToolActionSet& ActionSet) override;
+
+	virtual void SetWorld(UWorld* World);
+	virtual void SetAssetAPI(IToolsContextAssetAPI* AssetAPI);
+
+	virtual void Setup() override;
+	virtual void Shutdown(EToolShutdownType ShutdownType) override;
+
+	virtual void Render(IToolsContextRenderAPI* RenderAPI) override;
+
+	virtual bool HasCancel() const override { return false; }
+	virtual bool HasAccept() const override { return false; }
+	virtual bool CanAccept() const override { return false; }
+
+	// IClickSequenceBehaviorTarget implementation
+
+	virtual void OnBeginSequencePreview(const FInputDeviceRay& ClickPos) override;
+	virtual bool CanBeginClickSequence(const FInputDeviceRay& ClickPos) override;
+	virtual void OnBeginClickSequence(const FInputDeviceRay& ClickPos) override;
+	virtual void OnNextSequencePreview(const FInputDeviceRay& ClickPos) override;
+	virtual bool OnNextSequenceClick(const FInputDeviceRay& ClickPos) override;
+	virtual void OnTerminateClickSequence() override;
+	virtual bool RequestAbortClickSequence() override;
+	virtual void OnUpdateModifierState(int ModifierID, bool bIsOn) override;
+
+	// polygon drawing functions
+	virtual void ResetPolygon();
+	virtual void UpdatePreviewVertex(const FVector& PreviewVertex);
+	virtual void AppendVertex(const FVector& Vertex);
+	virtual bool FindDrawPlaneHitPoint(const FInputDeviceRay& ClickPos, FVector& HitPosOut);
+	virtual void EmitCurrentPolygon();
+
+	virtual void BeginInteractiveExtrude();
+	virtual void EndInteractiveExtrude();
+	virtual float FindInteractiveHeightDistance(const FInputDeviceRay& ClickPos);
+
+	virtual void SetDrawPlaneFromWorldPos(const FVector& Position, const FVector& Normal);
+
+
+public:
+	virtual void PopLastVertexAction();
+
+
+protected:
+	// flags used to identify modifier keys/buttons
+	static const int IgnoreSnappingModifier = 1;
+	static const int AngleSnapModifier = 2;
+
+protected:
+
+	/** Properties that control polygon generation exposed to user via detailsview */
+	UPROPERTY()
+	UDrawPolygonToolStandardProperties* PolygonProperties;
+
+	UPROPERTY()
+	UDrawPolygonToolSnapProperties* SnapProperties;
+	
+
+	/** Origin of plane we will draw polygon on */
+	UPROPERTY()
+	FVector DrawPlaneOrigin;
+
+	/** Orientation of plane we will draw polygon on */
+	UPROPERTY()
+	FQuat DrawPlaneOrientation;
+	
+	/** Vertices of current preview polygon */
+	UPROPERTY()
+	TArray<FVector> PolygonVertices;
+
+	/** last vertex of polygon that is actively being updated as input device is moved */
+	UPROPERTY()
+	FVector PreviewVertex;
+
+protected:
+	UWorld* TargetWorld;
+	IToolsContextAssetAPI* AssetAPI;
+
+	FViewCameraState CameraState;
+
+	UPROPERTY()
+	UPreviewMesh* PreviewMesh;
+
+	// drawing plane control
+
+	UPROPERTY()
+	UInteractiveGizmo* PositionPlaneGizmo;
+
+	void UpdateDrawPlaneFromGizmo(const FFrame3d& WorldPosition);
+
+	IClickBehaviorTarget* SetPointInWorldConnector = nullptr;
+
+
+	// polygon drawing
+
+	bool bAbortActivePolygonDraw;
+
+	bool bInFixedPolygonMode = false;
+	TArray<FVector> FixedPolygonClickPoints;
+
+	// can close poly if current segment intersects existing segment
+	bool UpdateSelfIntersection();
+	bool bHaveSelfIntersection;
+	int SelfIntersectSegmentIdx;
+	FVector3f SelfIntersectionPoint;
+
+	// only used when SnapSettings.bHitSceneObjects = true
+	bool bHaveSurfaceHit;
+	FVector3f SurfaceHitPoint;
+	FVector3f SurfaceOffsetPoint;
+
+	bool bIgnoreSnappingToggle = false;		// toggled by hotkey (shift)
+	FPointPlanarSnapSolver SnapEngine;
+	ToolSceneQueriesUtil::FSnapGeometry LastSnapGeometry;
+
+	void GenerateFixedPolygon(TArray<FVector>& FixedPoints, TArray<FVector>& VerticesOut);
+
+
+	// extrusion control
+
+	bool bInInteractiveExtrude = false;
+
+	void UpdateLivePreview();
+	bool bPreviewUpdatePending;
+
+	FDynamicMesh3 PreviewHeightTarget;
+	FDynamicMeshAABBTree3 PreviewHeightTargetAABB;
+	FFrame3d PreviewHeightFrame;
+	void GeneratePreviewHeightTarget();
+
+
+	FFrame3f HitPosFrameWorld;
+
+	// we use this to generate extruded meshes
+	void GeneratePolygonMesh(const TArray<FVector>& Polygon, FDynamicMesh3* ResultMeshOut, FFrame3d& WorldFrameOut, bool bIncludePreviewVtx, double ExtrudeDistance, bool bExtrudeSymmetric);
+};
+
+
+
+
+
diff --git a/Engine/Plugins/Experimental/MeshModelingToolset/Source/MeshModelingTools/Public/DynamicMeshSculptTool.h b/Engine/Plugins/Experimental/MeshModelingToolset/Source/MeshModelingTools/Public/DynamicMeshSculptTool.h
new file mode 100644
--- /dev/null
+++ b/Engine/Plugins/Experimental/MeshModelingToolset/Source/MeshModelingTools/Public/DynamicMeshSculptTool.h
@@ -0,0 +1,314 @@
+// Copyright 1998-2019 Epic Games, Inc. All Rights Reserved.
+
+#pragma once
+
+#include "CoreMinimal.h"
+#include "UObject/NoExportTypes.h"
+#include "BaseTools/MeshSurfacePointTool.h"
+#include "OctreeDynamicMeshComponent.h"
+#include "DynamicMeshAABBTree3.h"
+#include "DynamicMeshOctree3.h"
+#include "BrushToolIndicator.h"
+#include "MeshNormals.h"
+#include "BaseBrushTool.h"
+#include "Drawing/ToolDataVisualizer.h"
+#include "DynamicMeshSculptTool.generated.h"
+
+
+DECLARE_STATS_GROUP(TEXT("SculptTool"), STATGROUP_SculptTool, STATCAT_Advanced);
+DECLARE_CYCLE_STAT(TEXT("SculptTool_UpdateROI"), SculptTool_UpdateROI, STATGROUP_SculptTool);
+DECLARE_CYCLE_STAT(TEXT("SculptTool_ApplyStamp"), STAT_SculptToolApplyStamp, STATGROUP_SculptTool );
+DECLARE_CYCLE_STAT(TEXT("SculptTool_Tick"), STAT_SculptToolTick, STATGROUP_SculptTool);
+DECLARE_CYCLE_STAT(TEXT("SculptTool_Tick_ApplyStampBlock"), STAT_SculptTool_Tick_ApplyStampBlock, STATGROUP_SculptTool);
+DECLARE_CYCLE_STAT(TEXT("SculptTool_Tick_ApplyStamp_Remove"), SculptTool_Tick_ApplyStamp_Remove, STATGROUP_SculptTool);
+DECLARE_CYCLE_STAT(TEXT("SculptTool_Tick_ApplyStamp_Insert"), SculptTool_Tick_ApplyStamp_Insert, STATGROUP_SculptTool);
+DECLARE_CYCLE_STAT(TEXT("SculptTool_Tick_RemeshBlock"), STAT_SculptTool_Tick_RemeshBlock, STATGROUP_SculptTool);
+DECLARE_CYCLE_STAT(TEXT("SculptTool_Tick_NormalsBlock"), STAT_SculptTool_Tick_NormalsBlock, STATGROUP_SculptTool);
+DECLARE_CYCLE_STAT(TEXT("SculptTool_Tick_UpdateMeshBlock"), STAT_SculptTool_Tick_UpdateMeshBlock, STATGROUP_SculptTool);
+DECLARE_CYCLE_STAT(TEXT("SculptTool_Normals_Collect"), SculptTool_Normals_Collect, STATGROUP_SculptTool);
+DECLARE_CYCLE_STAT(TEXT("SculptTool_Normals_Compute"), SculptTool_Normals_Compute, STATGROUP_SculptTool);
+DECLARE_CYCLE_STAT(TEXT("SculptTool_Remesh_1Setup"), STAT_SculptTool_Remesh_Setup, STATGROUP_SculptTool);
+DECLARE_CYCLE_STAT(TEXT("SculptTool_Remesh_2Constraints"), STAT_SculptTool_Remesh_Constraints, STATGROUP_SculptTool);
+DECLARE_CYCLE_STAT(TEXT("SculptTool_Remesh_3RemeshROIUpdate"), STAT_SculptTool_Remesh_RemeshROIUpdate, STATGROUP_SculptTool);
+DECLARE_CYCLE_STAT(TEXT("SculptTool_Remesh_4RemeshPass"), STAT_SculptTool_Remesh_RemeshPass, STATGROUP_SculptTool);
+DECLARE_CYCLE_STAT(TEXT("SculptTool_Remesh_5PassOctreeUpdate"), STAT_SculptTool_Remesh_PassOctreeUpdate, STATGROUP_SculptTool);
+DECLARE_CYCLE_STAT(TEXT("SculptTool_Remesh_6Finish"), STAT_SculptTool_Remesh_Finish, STATGROUP_SculptTool);
+
+
+class FMeshVertexChangeBuilder;
+class FDynamicMeshChangeTracker;
+
+
+/** Mesh sculpting brush types */
+UENUM()
+enum class EDynamicMeshSculptBrushType : uint8
+{
+	/** Move brush */
+	Pull UMETA(DisplayName = "Pull"),
+
+	/** Offset brush */
+	Offset UMETA(DisplayName = "Offset"),
+
+	/** Pinch brush */
+	Pinch UMETA(DisplayName = "Pinch"),
+
+	/** Flatten brush */
+	Flatten UMETA(DisplayName = "Flatten"),
+};
+
+
+
+/** Smooth type for sculpting brush */
+UENUM()
+enum class EMeshSculptToolSmoothType : uint8
+{
+	/** Uniform smoothing */
+	Uniform UMETA(DisplayName = "Uniform"),
+
+	/** Texture-preserving smoothing */
+	TexturePreserving UMETA(DisplayName = "Texture Preserving"),
+};
+
+
+
+
+
+/**
+ *
+ */
+UCLASS()
+class MESHMODELINGTOOLS_API UDynamicMeshSculptToolBuilder : public UMeshSurfacePointToolBuilder
+{
+	GENERATED_BODY()
+
+public:
+	bool bEnableRemeshing;
+
+	UDynamicMeshSculptToolBuilder()
+	{
+		bEnableRemeshing = false;
+	}
+
+	virtual UMeshSurfacePointTool* CreateNewTool(const FToolBuilderState& SceneState) const override;
+};
+
+
+
+
+
+
+
+
+
+UCLASS()
+class MESHMODELINGTOOLS_API UBrushSculptProperties : public UInteractiveToolPropertySet
+{
+	GENERATED_BODY()
+
+public:
+	UBrushSculptProperties();
+
+	/** primary brush mode */
+	UPROPERTY(EditAnywhere, Category = Options)
+	EDynamicMeshSculptBrushType PrimaryBrushType;
+
+	/** smoothing type */
+	UPROPERTY(EditAnywhere, Category = Options)
+	EMeshSculptToolSmoothType SmoothingType;
+
+	/** Smoothing speed of brush */
+	UPROPERTY(EditAnywhere, Category = Options, meta = (UIMin = "0.0", UIMax = "1.0", ClampMin = "0.0", ClampMax = "1.0"))
+	float SmoothPower;
+
+	/** Offset speed of brush */
+	UPROPERTY(EditAnywhere, Category = Options, meta = (UIMin = "0.0", UIMax = "1.0", ClampMin = "0.0", ClampMax = "1.0", EditCondition = "PrimaryBrushType != EDynamicMeshSculptBrushType::Pull"))
+	float OffsetPower;
+
+	/** Depth of Brush into surface along ray */
+	UPROPERTY(EditAnywhere, Category = Options, meta = (UIMin = "-0.5", UIMax = "0.5", ClampMin = "-1.0", ClampMax = "1.0"))
+	float Depth;
+};
+
+
+
+
+
+UCLASS()
+class MESHMODELINGTOOLS_API UBrushRemeshProperties : public UInteractiveToolPropertySet
+{
+	GENERATED_BODY()
+
+public:
+	UBrushRemeshProperties();
+
+	/** Target Relative Triangle Sizefor Dynamic Meshing */
+	UPROPERTY(EditAnywhere, Category = Options, meta = (UIMin = "0.5", UIMax = "2.0", ClampMin = "0.1", ClampMax = "100.0"))
+	float RelativeSize;
+
+	/** Smoothing speed for dynamic meshing */
+	UPROPERTY(EditAnywhere, Category = Options, meta = (UIMin = "0.0", UIMax = "1.0", ClampMin = "0.0", ClampMax = "1.0"))
+	float Smoothing;
+
+
+	/** Enable edge flips */
+	UPROPERTY(EditAnywhere, Category = Options, AdvancedDisplay)
+	bool bFlips = true;
+
+	/** Enable edge splits */
+	UPROPERTY(EditAnywhere, Category = Options, AdvancedDisplay)
+	bool bSplits = true;
+
+	/** Enable edge collapses */
+	UPROPERTY(EditAnywhere, Category = Options, AdvancedDisplay)
+	bool bCollapses = true;
+
+	/** Prevent normal flips */
+	UPROPERTY(EditAnywhere, Category = Options, AdvancedDisplay)
+	bool bPreventNormalFlips = false;
+
+};
+
+
+
+/**
+ *
+ */
+UCLASS()
+class MESHMODELINGTOOLS_API UDynamicMeshSculptTool : public UMeshSurfacePointTool
+{
+	GENERATED_BODY()
+
+public:
+	UDynamicMeshSculptTool();
+
+	virtual void RegisterActions(FInteractiveToolActionSet& ActionSet) override;
+
+	virtual void Setup() override;
+	virtual void Shutdown(EToolShutdownType ShutdownType) override;
+
+	virtual void Tick(float DeltaTime) override;
+	virtual void Render(IToolsContextRenderAPI* RenderAPI) override;
+
+	virtual bool HasCancel() const override { return true; }
+	virtual bool HasAccept() const override { return true; }
+	virtual bool CanAccept() const override { return true; }
+
+	virtual bool HitTest(const FRay& Ray, FHitResult& OutHit) override;
+
+	virtual void OnBeginDrag(const FRay& Ray) override;
+	virtual void OnUpdateDrag(const FRay& Ray) override;
+	virtual void OnEndDrag(const FRay& Ray) override;
+	virtual void OnUpdateHover(const FInputDeviceRay& DevicePos) override;
+
+	virtual void SetEnableRemeshing(bool bEnable) { bEnableRemeshing = bEnable; }
+	virtual bool GetEnableRemeshing() const { return bEnableRemeshing; }
+
+
+	virtual void OnPropertyModified(UObject* PropertySet, UProperty* Property) override;
+
+public:
+
+	/** Properties that control sculpting*/
+	UPROPERTY()
+	UBrushSculptProperties* SculptProperties;
+
+	/** Properties that control brush size/etc*/
+	UPROPERTY()
+	UBrushBaseProperties* BrushProperties;
+
+	/** Properties that control dynamic remeshing */
+	UPROPERTY()
+	UBrushRemeshProperties* RemeshProperties;
+
+
+
+public:
+
+	virtual void IncreaseBrushRadiusAction();
+	virtual void DecreaseBrushRadiusAction();
+
+
+protected:
+
+	UPROPERTY()
+	UToolIndicatorSet* Indicators;
+
+	UPROPERTY()
+	UOctreeDynamicMeshComponent* DynamicMeshComponent;
+
+	// realtime visualization
+	void OnDynamicMeshComponentChanged();
+	FDelegateHandle OnDynamicMeshComponentChangedHandle;
+
+	FInterval1d BrushRelativeSizeRange;
+	double CurrentBrushRadius;
+	void CalculateBrushRadius();
+
+	bool bEnableRemeshing;
+	double InitialEdgeLength;
+
+	bool bInDrag;
+
+	FPlane ActiveDragPlane;
+	FVector LastHitPosWorld;
+	FVector BrushStartCenterWorld;
+	FVector BrushStartNormalWorld;
+	FVector LastBrushPosLocal;
+	FVector LastBrushPosWorld;
+	FVector LastBrushPosNormalWorld;
+	
+
+	TArray<int> VertexROI;
+	TSet<int> VertexSetBuffer;
+	TSet<int> TriangleROI;
+	void UpdateROI(const FVector& BrushPos);
+
+	bool bRemeshPending;
+	bool bNormalUpdatePending;
+	bool bTargetDirty;
+
+	bool bSmoothing;
+
+	bool bHaveRemeshed;
+
+	bool bStampPending;
+	FRay PendingStampRay;
+	void ApplyStamp(const FRay& WorldRay);
+
+	FDynamicMesh3 BrushTargetMesh;
+	FDynamicMeshAABBTree3 BrushTargetMeshSpatial;
+	FMeshNormals BrushTargetNormals;
+	void UpdateTarget();
+	bool GetTargetMeshNearest(const FVector3d& Position, double SearchRadius, FVector3d& TargetPosOut, FVector3d& TargetNormalOut);
+
+	bool UpdateBrushPositionOnActivePlane(const FRay& WorldRay);
+	bool UpdateBrushPositionOnTargetMesh(const FRay& WorldRay);
+	void ApplySmoothBrush(const FRay& WorldRay);
+	void ApplyMoveBrush(const FRay& WorldRay);
+	void ApplyOffsetBrush(const FRay& WorldRay);
+	void ApplyPinchBrush(const FRay& WorldRay);
+	void ApplyFlattenBrush(const FRay& WorldRay);
+
+	FFrame3d ActiveFlattenFrame;
+	void ComputeFlattenFrame();
+
+	TArray<int> TrianglesBuffer;
+	TArray<int> NormalsBuffer;
+	TArray<bool> NormalsVertexFlags;
+	void RecalculateNormals_PerVertex();
+	void RecalculateNormals_Overlay();
+
+	void RemeshROIPass();
+
+	FMeshVertexChangeBuilder* ActiveVertexChange = nullptr;
+	FDynamicMeshChangeTracker* ActiveMeshChange = nullptr;
+	void BeginChange(bool bIsVertexChange);
+	void EndChange();
+	void SaveActiveROI();
+	void UpdateSavedVertex(int vid, const FVector3d& OldPosition, const FVector3d& NewPosition);
+
+	double EstimateIntialSafeTargetLength(const FDynamicMesh3& Mesh, int MinTargetTriCount);
+};
+
+
+
diff --git a/Engine/Plugins/Experimental/MeshModelingToolset/Source/MeshModelingTools/Public/EditNormalsTool.h b/Engine/Plugins/Experimental/MeshModelingToolset/Source/MeshModelingTools/Public/EditNormalsTool.h
new file mode 100644
--- /dev/null
+++ b/Engine/Plugins/Experimental/MeshModelingToolset/Source/MeshModelingTools/Public/EditNormalsTool.h
@@ -0,0 +1,173 @@
+// Copyright 1998-2019 Epic Games, Inc. All Rights Reserved.
+
+#pragma once
+
+#include "CoreMinimal.h"
+#include "UObject/NoExportTypes.h"
+#include "MultiSelectionTool.h"
+#include "InteractiveToolBuilder.h"
+#include "MeshOpPreviewHelpers.h"
+#include "CleaningOps/EditNormalsOp.h"
+#include "DynamicMesh3.h"
+#include "BaseTools/SingleClickTool.h"
+#include "EditNormalsTool.generated.h"
+
+
+// predeclarations
+struct FMeshDescription;
+class USimpleDynamicMeshComponent;
+
+
+
+
+
+/**
+ *
+ */
+UCLASS()
+class MESHMODELINGTOOLS_API UEditNormalsToolBuilder : public UInteractiveToolBuilder
+{
+	GENERATED_BODY()
+
+public:
+	IToolsContextAssetAPI* AssetAPI = nullptr;
+
+	virtual bool CanBuildTool(const FToolBuilderState& SceneState) const override;
+	virtual UInteractiveTool* BuildTool(const FToolBuilderState& SceneState) const override;
+};
+
+
+
+
+
+/**
+ * Standard properties
+ */
+UCLASS()
+class MESHMODELINGTOOLS_API UEditNormalsToolProperties : public UInteractiveToolPropertySet
+{
+	GENERATED_BODY()
+
+public:
+	UEditNormalsToolProperties();
+
+
+	/** Recompute all mesh normals */
+	UPROPERTY(EditAnywhere, Category = Options)
+	bool bRecomputeNormals;
+
+	/** Choose the method for computing vertex normals */
+	UPROPERTY(EditAnywhere, Category = Options)
+	ENormalCalculationMethod NormalCalculationMethod;
+
+	/** Invert (flip) all mesh normals */
+	UPROPERTY(EditAnywhere, Category = Options)
+	bool bInvertNormals;
+
+	/** Choose where to create sharp edges by having multiple normals on the same vertices, using the change in normals of adjacent faces */
+	UPROPERTY(EditAnywhere, Category = NormalsTopology)
+	bool bRecomputeNormalTopologyAndEdgeSharpness;
+
+	/** Threshold on angle of change in face normals across an edge, above which we create a sharp edge if bSplitNormals is true */
+	UPROPERTY(EditAnywhere, Category = NormalsTopology, meta = (UIMin = "0.0", UIMax = "180.0", ClampMin = "0.0", ClampMax = "180.0"))
+	float SharpEdgeAngleThreshold;
+
+	/** If split normals is enabl*/
+	UPROPERTY(EditAnywhere, Category = NormalsTopology)
+	bool bAllowSharpVertices;
+};
+
+
+
+
+/**
+ * Advanced properties
+ */
+UCLASS()
+class MESHMODELINGTOOLS_API UEditNormalsAdvancedProperties : public UInteractiveToolPropertySet
+{
+	GENERATED_BODY()
+
+public:
+	UEditNormalsAdvancedProperties();
+};
+
+
+/**
+ * Factory with enough info to spawn the background-thread Operator to do a chunk of work for the tool
+ *  stores a pointer to the tool and enough info to know which specific operator it should spawn
+ */
+UCLASS()
+class MESHMODELINGTOOLS_API UEditNormalsOperatorFactory : public UObject, public IDynamicMeshOperatorFactory
+{
+	GENERATED_BODY()
+
+public:
+	// IDynamicMeshOperatorFactory API
+	virtual TSharedPtr<FDynamicMeshOperator> MakeNewOperator() override;
+
+	UPROPERTY()
+	UEditNormalsTool *Tool;
+
+	int ComponentIndex;
+
+};
+
+/**
+ * Simple Mesh Normal Updating Tool
+ */
+UCLASS()
+class MESHMODELINGTOOLS_API UEditNormalsTool : public UMultiSelectionTool
+{
+	GENERATED_BODY()
+
+public:
+
+	friend UEditNormalsOperatorFactory;
+
+	UEditNormalsTool();
+
+	virtual void Setup() override;
+	virtual void Shutdown(EToolShutdownType ShutdownType) override;
+
+	virtual void SetWorld(UWorld* World);
+	virtual void SetAssetAPI(IToolsContextAssetAPI* AssetAPI);
+
+	virtual void Tick(float DeltaTime) override;
+	virtual void Render(IToolsContextRenderAPI* RenderAPI) override;
+
+	virtual bool HasCancel() const override { return true; }
+	virtual bool HasAccept() const override;
+	virtual bool CanAccept() const override;
+
+#if WITH_EDITOR
+	virtual void PostEditChangeProperty(FPropertyChangedEvent &PropertyChangedEvent) override;
+#endif
+
+	virtual void OnPropertyModified(UObject* PropertySet, UProperty* Property) override;
+
+protected:
+
+	UPROPERTY()
+	UEditNormalsToolProperties* BasicProperties;
+
+	UPROPERTY()
+	UEditNormalsAdvancedProperties* AdvancedProperties;
+
+
+	UPROPERTY()
+	TArray<UMeshOpPreviewWithBackgroundCompute*> Previews;
+
+
+protected:
+	TArray<TSharedPtr<FDynamicMesh3>> OriginalDynamicMeshes;
+
+	UWorld* TargetWorld;
+	IToolsContextAssetAPI* AssetAPI;
+
+	FViewCameraState CameraState;
+
+	void UpdateNumPreviews();
+
+	void GenerateAsset(const TArray<TUniquePtr<FDynamicMeshOpResult>>& Results);
+};
diff --git a/Engine/Plugins/Experimental/MeshModelingToolset/Source/MeshModelingTools/Public/MakeMeshTestTool.h b/Engine/Plugins/Experimental/MeshModelingToolset/Source/MeshModelingTools/Public/MakeMeshTestTool.h
new file mode 100644
--- /dev/null
+++ b/Engine/Plugins/Experimental/MeshModelingToolset/Source/MeshModelingTools/Public/MakeMeshTestTool.h
@@ -0,0 +1,167 @@
+// Copyright 1998-2019 Epic Games, Inc. All Rights Reserved.
+
+#pragma once
+
+#include "CoreMinimal.h"
+#include "UObject/NoExportTypes.h"
+#include "InteractiveToolBuilder.h"
+#include "BaseTools/SingleClickTool.h"
+#include "PreviewMesh.h"
+#include "MakeMeshTestTool.generated.h"
+
+class FDynamicMesh3;
+
+/**
+ *
+ */
+UCLASS()
+class MESHMODELINGTOOLS_API UMakeMeshTestToolBuilder : public UInteractiveToolBuilder
+{
+	GENERATED_BODY()
+
+public:
+	IToolsContextAssetAPI* AssetAPI;
+
+	UMakeMeshTestToolBuilder() 
+	{
+		AssetAPI = nullptr;
+	}
+
+	virtual bool CanBuildTool(const FToolBuilderState& SceneState) const override;
+	virtual UInteractiveTool* BuildTool(const FToolBuilderState& SceneState) const override;
+};
+
+
+
+/** Shape Types */
+UENUM()
+enum class EMakeMeshShapeType : uint8
+{
+	Box = 0 UMETA(DisplayName = "Box"),
+	Cylinder = 1 UMETA(DisplayName = "Cylinder"),
+	Cone = 2 UMETA(DisplayName = "Cone"),
+	Plane = 3 UMETA(DisplayName = "Plane"),
+	Sphere = 4 UMETA(DisplayName = "Sphere")
+};
+
+
+
+/** Placement Target Types */
+UENUM()
+enum class EMakeMeshPlacementType : uint8
+{
+	OnPlane = 0,
+	OnScene = 1
+};
+
+
+
+
+
+UCLASS()
+class MESHMODELINGTOOLS_API UProceduralShapeToolProperties : public UInteractiveToolPropertySet
+{
+	GENERATED_BODY()
+
+public:
+	UProceduralShapeToolProperties();
+
+	/** Type of shape to generate */
+	UPROPERTY(EditAnywhere, Category = ShapeSettings)
+	EMakeMeshShapeType Shape;
+
+	/** Width of Shape */
+	UPROPERTY(EditAnywhere, Category = ShapeSettings, meta = (DisplayName = "Width", UIMin = "1.0", UIMax = "1000.0", ClampMin = "0.0001", ClampMax = "1000000.0"))
+	float Width;
+
+	/** Height of Shape */
+	UPROPERTY(EditAnywhere, Category = ShapeSettings, meta = (DisplayName = "Height", UIMin = "1.0", UIMax = "1000.0", ClampMin = "0.0001", ClampMax = "1000000.0"))
+	float Height;
+
+	/** Rotation around up axis */
+	UPROPERTY(EditAnywhere, Category = ShapeSettings, meta = (DisplayName = "Rotation", UIMin = "0.0", UIMax = "360.0"))
+	float Rotation;
+
+	/** Type of shape to generate */
+	UPROPERTY(EditAnywhere, Category = ShapeSettings)
+	EMakeMeshPlacementType PlaceMode;
+
+	/** Center shape at click point */
+	UPROPERTY(EditAnywhere, Category = ShapeSettings)
+	bool bCentered;
+
+	///** Start Angle of Shape */
+	//UPROPERTY(EditAnywhere, Category = ShapeSettings, meta = (DisplayName = "Start Angle", UIMin = "0.0", UIMax = "360.0", ClampMin = "-10000", ClampMax = "10000.0"))
+	//float StartAngle;
+
+	///** End Angle of Shape */
+	//UPROPERTY(EditAnywhere, Category = ShapeSettings, meta = (DisplayName = "End Angle", UIMin = "0.0", UIMax = "360.0", ClampMin = "-10000", ClampMax = "10000.0"))
+	//float EndAngle;
+
+
+	/** Number of Slices */
+	UPROPERTY(EditAnywhere, Category = ShapeSettings, meta = (DisplayName = "Slices", UIMin = "3", UIMax = "128", ClampMin = "3", ClampMax = "999"))
+	int Slices;
+
+	/** Subdivisions */
+	UPROPERTY(EditAnywhere, Category = ShapeSettings, meta = (DisplayName = "Subdivisions", UIMin = "0", UIMax = "100", ClampMin = "0", ClampMax = "4000"))
+	int Subdivisions;
+
+};
+
+
+
+
+
+
+
+/**
+ *
+ */
+UCLASS()
+class MESHMODELINGTOOLS_API UMakeMeshTestTool : public USingleClickTool, public IHoverBehaviorTarget
+{
+	GENERATED_BODY()
+
+public:
+	virtual void SetWorld(UWorld* World);
+	virtual void SetAssetAPI(IToolsContextAssetAPI* AssetAPI);
+
+	virtual void Setup() override;
+	virtual void Shutdown(EToolShutdownType ShutdownType) override;
+
+	virtual void Render(IToolsContextRenderAPI* RenderAPI) override;
+
+	virtual bool HasCancel() const override { return false; }
+	virtual bool HasAccept() const override { return false; }
+	virtual bool CanAccept() const override { return false; }
+
+	virtual void OnPropertyModified(UObject* PropertySet, UProperty* Property) override;
+
+	virtual void OnUpdateHover(const FInputDeviceRay& DevicePos) override;
+	virtual void OnClicked(const FInputDeviceRay& ClickPos) override;
+
+
+
+protected:
+	UPROPERTY()
+	UProceduralShapeToolProperties* ShapeSettings;
+
+	UPROPERTY()
+	UPreviewMesh* PreviewMesh;
+
+protected:
+	UWorld* TargetWorld;
+	IToolsContextAssetAPI* AssetAPI;
+
+	void UpdatePreviewPosition(const FInputDeviceRay& ClickPos);
+	FFrame3f ShapeFrame;
+
+	void UpdatePreviewMesh();
+
+	void GenerateCylinder(FDynamicMesh3* OutMesh);
+	void GenerateCone(FDynamicMesh3* OutMesh);
+	void GenerateBox(FDynamicMesh3* OutMesh);
+	void GeneratePlane(FDynamicMesh3* OutMesh);
+	void GenerateSphere(FDynamicMesh3* OutMesh);
+};
\ No newline at end of file
diff --git a/Engine/Plugins/Experimental/MeshModelingToolset/Source/MeshModelingTools/Public/MeshInspectorTool.h b/Engine/Plugins/Experimental/MeshModelingToolset/Source/MeshModelingTools/Public/MeshInspectorTool.h
new file mode 100644
--- /dev/null
+++ b/Engine/Plugins/Experimental/MeshModelingToolset/Source/MeshModelingTools/Public/MeshInspectorTool.h
@@ -0,0 +1,164 @@
+// Copyright 1998-2019 Epic Games, Inc. All Rights Reserved.
+
+#pragma once
+
+#include "CoreMinimal.h"
+#include "UObject/NoExportTypes.h"
+#include "SingleSelectionTool.h"
+#include "InteractiveToolBuilder.h"
+#include "DynamicMesh3.h"
+#include "DynamicMeshAABBTree3.h"
+#include "Properties/MeshStatisticsProperties.h"
+#include "MeshInspectorTool.generated.h"
+
+
+// predeclarations
+struct FMeshDescription;
+class USimpleDynamicMeshComponent;
+class UMaterialInstanceDynamic;
+
+/**
+ *
+ */
+UCLASS()
+class MESHMODELINGTOOLS_API UMeshInspectorToolBuilder : public UInteractiveToolBuilder
+{
+	GENERATED_BODY()
+
+public:
+	virtual bool CanBuildTool(const FToolBuilderState& SceneState) const override;
+	virtual UInteractiveTool* BuildTool(const FToolBuilderState& SceneState) const override;
+};
+
+
+
+/** Material modes for MeshInspectorTool */
+UENUM()
+enum class EInspectorMaterialMode : uint8
+{
+	/** Input material */
+	Default,
+
+	/** Checkerboard material */
+	Checkerboard,
+
+	/** Override material */
+	Override
+};
+
+
+
+
+
+UCLASS()
+class MESHMODELINGTOOLS_API UMeshInspectorProperties : public UInteractiveToolPropertySet
+{
+	GENERATED_BODY()
+
+public:
+	/** Toggle visibility of all mesh edges */
+	UPROPERTY(EditAnywhere, Category = Options)
+	bool bWireframe = true;
+
+	/** Toggle visibility of open boundary edges */
+	UPROPERTY(EditAnywhere, Category = Options)
+	bool bBoundaryEdges = true;
+
+	/** Toggle visibility of polygon borders */
+	UPROPERTY(EditAnywhere, Category = Options)
+	bool bPolygonBorders = false;
+
+	/** Toggle visibility of UV seam edges */
+	UPROPERTY(EditAnywhere, Category = Options)
+	bool bUVSeams = false;
+
+	/** Toggle visibility of Normal seam edges */
+	UPROPERTY(EditAnywhere, Category = Options)
+	bool bNormalSeams = false;
+
+	/** Toggle visibility of normal vectors */
+	UPROPERTY(EditAnywhere, Category = Options)
+	bool bNormalVectors = false;
+
+	/** Toggle visibility of tangent vectors */
+	UPROPERTY(EditAnywhere, Category = Options)
+	bool bTangentVectors = false;
+
+	/** Length of line segments representing normal vectors */
+	UPROPERTY(EditAnywhere, Category = Options, meta = (EditCondition = "bNormalVectors"))
+	float NormalLength = 5.0f;
+
+	/** Length of line segments representing tangent vectors */
+	UPROPERTY(EditAnywhere, Category = Options, meta = (EditCondition = "bTangentVectors"))
+	float TangentLength = 5.0f;
+
+
+	/** Material that will be used on the mesh */
+	UPROPERTY(EditAnywhere, Category = Options)
+	EInspectorMaterialMode MaterialMode = EInspectorMaterialMode::Default;
+
+	UPROPERTY(EditAnywhere, Category = Options, meta = (UIMin = "1.0", UIMax = "40.0", ClampMin = "0.01", ClampMax = "1000.0", EditCondition = "MaterialMode == EInspectorMaterialMode::Checkerboard"))
+	float CheckerDensity = 20.0f;
+
+	UPROPERTY(EditAnywhere, Category = Options, meta = (EditCondition = "MaterialMode == EInspectorMaterialMode::Override"))
+	UMaterialInterface* OverrideMaterial = nullptr;
+};
+
+
+
+
+/**
+ * Mesh Inspector Tool for visualizing mesh information
+ */
+UCLASS()
+class MESHMODELINGTOOLS_API UMeshInspectorTool : public USingleSelectionTool
+{
+	GENERATED_BODY()
+
+public:
+	UMeshInspectorTool();
+
+	virtual void RegisterActions(FInteractiveToolActionSet& ActionSet) override;
+
+	virtual void Setup() override;
+	virtual void Shutdown(EToolShutdownType ShutdownType) override;
+
+	virtual void Render(IToolsContextRenderAPI* RenderAPI) override;
+
+	virtual bool HasCancel() const override { return false; }
+	virtual bool HasAccept() const override;
+	virtual bool CanAccept() const override;
+
+	virtual void OnPropertyModified(UObject* PropertySet, UProperty* Property) override;
+
+public:
+
+	virtual void IncreaseLineWidthAction();
+	virtual void DecreaseLineWidthAction();
+
+protected:
+
+	UPROPERTY()
+	UMeshInspectorProperties* Settings = nullptr;
+
+	UPROPERTY()
+	UMaterialInterface* DefaultMaterial = nullptr;
+
+	UPROPERTY()
+	UMaterialInstanceDynamic* CheckerMaterial = nullptr;
+
+
+	float LineWidthMultiplier = 1.0f;
+	EInspectorMaterialMode ActiveMaterialMode;
+	float ActiveCheckerDensity = 0.0f;
+
+protected:
+	UPROPERTY()
+	USimpleDynamicMeshComponent* DynamicMeshComponent;
+
+	TArray<int> BoundaryEdges;
+	TArray<int> UVSeamEdges;
+	TArray<int> NormalSeamEdges;
+	TArray<int> GroupBoundaryEdges;
+	void Precompute();
+};
diff --git a/Engine/Plugins/Experimental/MeshModelingToolset/Source/MeshModelingTools/Public/MeshSelectionTool.h b/Engine/Plugins/Experimental/MeshModelingToolset/Source/MeshModelingTools/Public/MeshSelectionTool.h
new file mode 100644
--- /dev/null
+++ b/Engine/Plugins/Experimental/MeshModelingToolset/Source/MeshModelingTools/Public/MeshSelectionTool.h
@@ -0,0 +1,117 @@
+// Copyright 1998-2019 Epic Games, Inc. All Rights Reserved.
+
+#pragma once
+
+#include "CoreMinimal.h"
+#include "DynamicMeshBrushTool.h"
+#include "SelectionSet.h"
+#include "MeshSelectionChange.h"
+#include "MeshSelectionTool.generated.h"
+
+
+
+/**
+ *
+ */
+UCLASS()
+class MESHMODELINGTOOLS_API UMeshSelectionToolBuilder : public UMeshSurfacePointToolBuilder
+{
+	GENERATED_BODY()
+
+public:
+	IToolsContextAssetAPI* AssetAPI = nullptr;
+
+	virtual UMeshSurfacePointTool* CreateNewTool(const FToolBuilderState& SceneState) const override;
+};
+
+
+
+/**
+ *
+ */
+UCLASS()
+class MESHMODELINGTOOLS_API UMeshSelectionTool : public UDynamicMeshBrushTool
+{
+	GENERATED_BODY()
+
+public:
+	UMeshSelectionTool();
+
+	virtual void SetWorld(UWorld* World);
+	virtual void SetAssetAPI(IToolsContextAssetAPI* AssetAPI);
+
+	virtual void RegisterActions(FInteractiveToolActionSet& ActionSet) override;
+
+	virtual void Setup() override;
+
+	virtual void Tick(float DeltaTime) override;
+	virtual void Render(IToolsContextRenderAPI* RenderAPI) override;
+
+	virtual bool HasCancel() const override { return true; }
+	virtual bool HasAccept() const override { return bHaveModifiedMesh; }
+	virtual bool CanAccept() const override { return bHaveModifiedMesh; }
+
+	// UBaseBrushTool overrides
+	virtual void OnBeginDrag(const FRay& Ray) override;
+	virtual void OnUpdateDrag(const FRay& Ray) override;
+	virtual void OnEndDrag(const FRay& Ray) override;
+	virtual void OnUpdateHover(const FInputDeviceRay& DevicePos) override;
+
+public:
+
+
+	UFUNCTION(CallInEditor, Category = Options, meta = (DisplayName = "Delete"))
+	void DeleteSelectedTriangles();
+
+	UFUNCTION(CallInEditor, Category = Options, meta = (DisplayName = "Separate"))
+	void SeparateSelectedTriangles();
+
+
+protected:
+
+	virtual void ApplyStamp(const FBrushStampData& Stamp);
+
+
+	virtual void OnShutdown(EToolShutdownType ShutdownType) override;
+
+protected:
+	UPROPERTY()
+	UMeshSelectionSet* Selection;
+
+	UWorld* TargetWorld;
+	IToolsContextAssetAPI* AssetAPI;
+
+	EMeshSelectionElementType SelectionType = EMeshSelectionElementType::Face;
+
+	bool bInRemoveStroke = false;
+
+	FBrushStampData StartStamp;
+	FBrushStampData LastStamp;
+
+	bool bStampPending;
+
+
+	// temp
+	TArray<int> IndexBuf;
+	void CalculateVertexROI(const FBrushStampData& Stamp, TArray<int>& VertexROI);
+	void CalculateTriangleROI(const FBrushStampData& Stamp, TArray<int>& TriangleROI);
+	TArray<int> PreviewBrushROI;
+	TBitArray<> SelectedVertices;
+	TBitArray<> SelectedTriangles;
+	void OnExternalSelectionChange();
+
+	void OnSelectionUpdated();
+	void UpdateVisualization();
+
+	// selection change
+	FMeshSelectionChangeBuilder* ActiveSelectionChange = nullptr;
+	void BeginChange(bool bAdding);
+	TUniquePtr<FMeshSelectionChange> EndChange();
+	void CancelChange();
+
+
+	bool bHaveModifiedMesh = false;
+};
+
+
+
diff --git a/Engine/Plugins/Experimental/MeshModelingToolset/Source/MeshModelingTools/Public/MeshSpaceDeformerTool.h b/Engine/Plugins/Experimental/MeshModelingToolset/Source/MeshModelingTools/Public/MeshSpaceDeformerTool.h
new file mode 100644
--- /dev/null
+++ b/Engine/Plugins/Experimental/MeshModelingToolset/Source/MeshModelingTools/Public/MeshSpaceDeformerTool.h
@@ -0,0 +1,464 @@
+// Copyright 1998-2019 Epic Games, Inc. All Rights Reserved.
+
+#pragma once
+
+#include "CoreMinimal.h"
+#include "UObject/NoExportTypes.h"
+#include "BaseTools/MeshSurfacePointTool.h"
+#include "DynamicMeshAABBTree3.h"
+#include "Drawing/ToolDataVisualizer.h"
+#include "Transforms/QuickAxisTranslater.h"
+#include "Transforms/QuickAxisRotator.h"
+#include "MeshVertexChange.h"
+#include "GroupTopology.h"
+#include "Spatial/GeometrySet3.h"
+#include "Selection/GroupTopologySelector.h"
+#include "Operations/GroupTopologyDeformer.h"
+#include "Curves/CurveFloat.h"
+#include "SpaceDeformerOps/MeshSpaceDeformerOp.h"
+#include "SpaceDeformerOps/BendMeshOp.h"
+#include "SpaceDeformerOps/TwistMeshOp.h"
+#include "SpaceDeformerOps/FlareMeshOp.h"
+#include "SimpleDynamicMeshComponent.h"
+
+#include "ModelingOperators/Public/ModelingTaskTypes.h"
+
+
+#include "MeshSpaceDeformerTool.generated.h"
+
+class FMeshVertexChangeBuilder;
+
+/**
+ * ToolBuilder
+ */
+UCLASS()
+class MESHMODELINGTOOLS_API UMeshSpaceDeformerToolBuilder : public UMeshSurfacePointToolBuilder
+{
+	GENERATED_BODY()
+
+public:
+	UMeshSpaceDeformerToolBuilder()
+	{
+	}
+
+	virtual UMeshSurfacePointTool* CreateNewTool(const FToolBuilderState& SceneState) const override;
+};
+
+/** ENonlinearOperation determines which type of nonlinear deformation will be applied*/
+UENUM()
+enum class  ENonlinearOperationType : int8
+{
+	Bend		UMETA(DisplayName = "Bend"),
+	Flare		UMETA(DisplayName = "Flare"),
+	Twist		UMETA(DisplayName = "Twist"),
+	//Sinusoid	UMETA(DisplayName = "Sinusoid"),
+	//Wave		UMETA(DisplayName = "Wave"),
+	//Squish	UMETA(DisplayName = "Squish")
+};
+
+/**
+* \brief IHandleWidgetSelectable is a pure virtual base class for elements of the HandleWidget which may be selectable (i.e. the end points of line segments or the axes themselves)
+*/
+class IHandleWidgetSelectable
+{
+public:
+	IHandleWidgetSelectable() = delete;
+	IHandleWidgetSelectable(const FLinearColor& DefaultColor, const FLinearColor& HoveringColor, const FLinearColor& SelectedColor) : StateColors{ DefaultColor,HoveringColor,SelectedColor } 
+	{
+	};
+	virtual ~IHandleWidgetSelectable() 
+	{
+	};
+
+	/**
+	* \brief Sets the renderer to be used by the tool
+	*/
+	virtual inline void SetRenderer(FToolDataVisualizer* Target) 
+	{ 
+		Renderer = Target;
+	};
+
+	/**
+	* \brief Sets the world position of the tool
+	*/ 
+	virtual inline void SetPosition(const FVector3d& Pos) 
+	{
+		WorldCenter = Pos;
+	};
+
+	/**
+	* \brief Gets the color depnding on the state of the selectable element
+	*/
+	virtual inline const FLinearColor& GetColor() const 
+	{
+		return StateColors[static_cast<int8>(State)];
+	}
+
+	/**
+	* \brief Sets the visibility of the specified widget element
+	*/
+	virtual inline void SetVisible(bool Enable) 
+	{
+		Visible = Enable;
+	};
+
+	FToolDataVisualizer* Renderer{nullptr};
+
+	enum EDisplayState : int8
+	{
+		Default,  // Unselected, the default display color
+		Hovering, // The element is being hovered over
+		Selected  // The element is selected
+	};
+
+	EDisplayState State = EDisplayState::Default;
+	FLinearColor StateColors[3];
+
+	FVector3d WorldCenter;
+	bool Visible{false};
+	virtual void Render() = 0;
+};
+
+/**
+* \brief This object is located on the ends of the line segments which extend from the center of the widget. 
+*/
+class FHandleWidgetVertex : public IHandleWidgetSelectable
+{
+public:
+	FHandleWidgetVertex() : IHandleWidgetSelectable
+	{
+		FLinearColor{1.f,1.f,1.f,1.f},	// White
+		FLinearColor{0.f,1.f,1.f,1.f},	// Yellow
+		FLinearColor{1.f,0.0f,1.f,1.f}	// Pink
+	} {};
+	virtual void Render() override
+	{
+		if (Visible)
+		{
+			float Width = State != Default ? 4.0 : 2.0;
+			Renderer->DrawViewFacingCircle(WorldCenter,4.0, 16, GetColor(), Width,false);
+		}
+	};
+};
+
+
+/**
+* \brief Represents an axis of the handle widget, each axis is selectable and has two vertices located at the end points
+*/
+class FHandleWidgetAxis : public IHandleWidgetSelectable
+{
+public:
+	FHandleWidgetAxis() : IHandleWidgetSelectable
+	{
+		FLinearColor{0.2f,0.2f,1.f,1.f}, // Light-blue
+		FLinearColor{0.f,1.f,1.f,1.f},	 // Yellow
+		FLinearColor{1.f,0.f,1.f,1.f}	 // Pink
+	} {};
+
+	double LowerExtentLength;
+	double UpperExtentLength;
+	
+	//Identity
+	FMatrix3d Basis{ 
+		FVector3d{1.0,0.0,0.0}, 
+		FVector3d{0.0,1.0,0.0}, 
+		FVector3d{0.0,0.0,1.0}, 
+		false };
+
+	//The direction vector representing the local space direction of the line segment this handle axis represents
+	FVector3d Axis;
+
+	void SetAxis(const FVector3d& TargetAxis) {Axis = TargetAxis;};
+
+	void SetBasis(const FVector3d& V0, const FVector3d& V1, const FVector3d& V2 ) { Basis = FMatrix3d{ V0, V1, V2, false }; };
+
+	FHandleWidgetVertex Vertices[2];
+	double SelectedWidth = 2.5;
+	double DefaultWidth = 1.5;
+
+	/**
+	* \brief Sets the center position from the center of the handle. This, in turn, sets the position of its associated vertices
+	*/
+	virtual inline void SetPosition(const FVector3d& Center) override
+	{
+		WorldCenter = Center;
+
+		const FVector3d Positions[2]
+		{
+			WorldCenter + Axis * LowerExtentLength,
+			WorldCenter + Axis * UpperExtentLength
+		};
+
+		for (uint8 i = 0; i < 2; ++i)
+		{
+			Vertices[i].SetPosition(Positions[i]);
+		}
+
+	};
+
+	virtual void Render() override
+	{
+		double Width = State == EDisplayState::Default ? DefaultWidth : SelectedWidth;
+
+		for (uint8 i = 0; i < 2; ++i)
+		{
+			Vertices[i].Render();
+		}
+
+		if (Visible)
+		{
+			Renderer->DrawLine(Vertices[0].WorldCenter, Vertices[1].WorldCenter, GetColor(), Width, false);
+		}
+	};
+
+};
+
+/**
+* \brief This purpose of this "widget" is to visualize the world space center of the transformation being applied by a Mesh Space Deformer operator. 
+  \par This is meant to be an analogous visualizer to the "handle" widgets used by popular modeling tools i.e. maya
+*/
+class FHandleWidget
+{
+public:
+	FHandleWidget(FToolDataVisualizer* Renderer) 
+	{
+		for (unsigned i=0;i<3;++i)
+		{
+			Axes[i].SetRenderer(Renderer);
+
+
+			for (unsigned j=0;j<2;++j)
+			{
+				Axes[i].Vertices[j].SetRenderer(Renderer);
+			}
+			const FVector3d& A = Axes[i].Vertices[0].WorldCenter;
+			const FVector3d& B = Axes[i].Vertices[1].WorldCenter;
+			GeometrySet.AddPoint(i * 2 , A);
+			GeometrySet.AddPoint(i * 2 + 1, B);
+			GeometrySet.AddCurve(i, FPolyline3d{ TArray<FVector3d>{A,B} });
+
+		}
+
+
+	};
+	virtual ~FHandleWidget() {};
+
+	/**
+	* \brief Returns a pointer to the Handle Widget Axis related object as specified by the Selection
+	*/
+	FHandleWidgetAxis* GetSelectedAxis() { return &Axes[SelectedAxisID];};
+
+	/**
+	* \brief Returns a pointer to the Handle Widget Axis related object as specified by the Selection
+	*/
+	FHandleWidgetVertex* GetSelectedVertex() { return SelectedVertexID == -1 ? nullptr : &Axes[SelectedAxisID].Vertices[SelectedVertexID]; };
+	int SelectedAxisID;
+	int SelectedVertexID;
+
+	FHandleWidgetAxis Axes[3];
+	FVector3d WorldCenter;
+
+	/** Given a ray, tests the geometry collection and updates the corresponding selectable handle elements */
+	void UpdateHover(const FInputDeviceRay& DevicePos, const FTransform& WorldTransform, TFunction<bool(const FVector3d&, const FVector3d&)> ToleranceFunction);
+	void SetPosition(const FVector3d& Center)
+	{
+		WorldCenter = Center;
+
+
+		for (unsigned i = 0; i < 3; ++i)
+		{
+			Axes[i].SetPosition(Center);
+			const FVector3d& A = Axes[i].Vertices[0].WorldCenter;
+			const FVector3d& B = Axes[i].Vertices[1].WorldCenter;
+			GeometrySet.UpdatePoint(i * 2 , A);
+			GeometrySet.UpdatePoint(i * 2 + 1, B);
+			GeometrySet.UpdateCurve(i, FPolyline3d{ TArray<FVector3d>{A,B} });
+		}
+	};
+
+	void SetBasis(const FVector3d& Axis0, const FVector3d& Axis1, const FVector3d& Axis2 )
+	{
+		Axes[0].SetAxis(Axis0);
+		Axes[1].SetAxis(Axis1);
+		Axes[2].SetAxis(Axis2);
+	};
+	void UpdateDisplayData(int SelectedAxis, double LowerBoundsInterval, double UpperBoundsInterval, const FVector3d& HalfExtents)
+	{
+
+		SelectedAxisID = SelectedAxis;
+		for (unsigned i = 0; i < 3; ++i)
+		{
+			if ( i == SelectedAxis )
+			{
+				Axes[i].State = IHandleWidgetSelectable::EDisplayState::Selected;
+				Axes[i].SetVisible(true);
+				Axes[i].Vertices[0].SetVisible(true);
+				Axes[i].Vertices[1].SetVisible(true);
+				Axes[i].LowerExtentLength = HalfExtents[i] * LowerBoundsInterval;
+				Axes[i].UpperExtentLength = HalfExtents[i] * UpperBoundsInterval;
+			}else
+			{
+				Axes[i].SetVisible(true);
+				Axes[i].Vertices[0].SetVisible(false);
+				Axes[i].Vertices[1].SetVisible(false);
+				Axes[i].LowerExtentLength = -HalfExtents[i] ;
+				Axes[i].UpperExtentLength = HalfExtents[i] ;
+			}
+
+		}
+	}
+
+	/** Called when updating a click-and-drag transaction, updates the positions of the selected vertices based on the new mouse input position*/
+	void DragUpdateBounds(const FRay3d& WorldRay);
+
+	void Render() 
+	{
+		for (unsigned i=0;i<3;++i)
+		{
+			Axes[i].Render();
+		}
+	};
+
+	/** The set of geometry containing the vertices which are draggable and the axes objects which are selectable */
+	FGeometrySet3 GeometrySet;
+
+};
+
+
+/**
+ * Applies non-linear deformations to a mesh 
+ */
+UCLASS()
+class MESHMODELINGTOOLS_API UMeshSpaceDeformerTool : public UMeshSurfacePointTool
+{
+	GENERATED_BODY()
+
+public:
+	UMeshSpaceDeformerTool();
+
+	virtual void RegisterActions(FInteractiveToolActionSet& ActionSet) override;
+
+	FMatrix3d CalculateBestAxis(FDynamicMesh3 * Mesh, const TArray<FVector3d>& PositionBuffer, FTransform WorldTransform, FVector3d * ExtentsOut, FVector3d * OriginOut );
+
+	virtual void Setup() override;
+	virtual void Shutdown(EToolShutdownType ShutdownType) override;
+
+	virtual void Tick(float DeltaTime) override;
+	virtual void Render(IToolsContextRenderAPI* RenderAPI) override;
+
+	virtual bool HasCancel() const override { return true; }
+	virtual bool HasAccept() const override { return true; }
+	virtual bool CanAccept() const override { return true; }
+
+	virtual bool HitTest(const FRay& Ray, FHitResult& OutHit) override;
+	
+	//
+	// Input behavior
+	//
+	virtual void OnBeginDrag(const FRay& Ray) override;
+	virtual void OnUpdateDrag(const FRay& Ray) override;
+	virtual void OnEndDrag(const FRay& Ray) override;
+
+	virtual void OnUpdateHover(const FInputDeviceRay& DevicePos) override;
+
+#if WITH_EDITOR
+	virtual void PostEditChangeProperty(FPropertyChangedEvent& PropertyChangedEvent) override;
+#endif
+
+	void UpdateOp();
+public:
+
+	float VisualAngleSnapThreshold = 1.5f;
+
+protected:
+
+	//Options
+	UPROPERTY(EditAnywhere, Category = Options, meta = (DisplayName = "Operation Type"))
+	ENonlinearOperationType SelectedOperationType;
+
+	/** The upper bounds interval corresponds to the region of space which the selected operator will affect. A setting of -1.0 should envelope all points in the "lower" half of the mesh given the axis has been auto-detected. The corresponding upper value of 1 will cover the entire mesh. */
+	UPROPERTY(EditAnywhere, Category = Options, meta = (DisplayName = "Lower Bound", UIMin = "-1.5", UIMax = "0.0", ClampMin = "-1.5", ClampMax = "0.0"))
+	double LowerBoundsInterval{ -1.0 };
+
+	/** The upper bounds interval corresponds to the region of space which the selected operator will affect. A setting of 1.0 should envelope all points in the "upper" half of the mesh given the axis has been auto-detected. The corresponding lower value of -1 will cover the entire mesh. */
+	UPROPERTY(EditAnywhere, Category = Options, meta = (DisplayName = "Upper Bound", UIMin = "0.0", UIMax = "1.5", ClampMin = "0.0", ClampMax = "1.5"))
+	double UpperBoundsInterval{ 1.0 };
+
+	/** As each operator has a range of values (i.e. curvature, angle of twist, scale), this represents the percentage passed to the operator as a parameter. In the future, for more control, this should be separated into individual settings for each operator for more precise control */
+	UPROPERTY(EditAnywhere, Category = Options, meta = (DisplayName = "Modifier Percent", UIMin = "-100.0", UIMax = "100.0", ClampMin = "-100.0", ClampMax = "100.0"))
+	double ModifierPercent{ 0.0 };
+
+	/** Automatically detects the principal axes of the point cloud of the mesh, as well as determining the lengths of the maximal elements in those directions, setting the intervals accordingly. */
+	UFUNCTION(CallInEditor, Category = Options, meta = (DisplayName = "Auto-detect Axis", ToolTip = "Attempts to automatically align the handle with principal axis."))
+	void AutoDetectAxes();
+
+	/** Swaps the secondary axis upon which the operator will work. I.e. bending */
+	UFUNCTION(CallInEditor, Category = Options, meta = (DisplayName = "Swap secondary axis"))
+	void SwapSecondaryAxis();
+
+	//Calculate the updated object-space position which gets passed to the handle widget.
+	void UpdateObjectSpaceAxisCentroid();
+
+	//Updates the upper and lower intervals from the drag update
+	void UpdateIntervalsFromDrag();
+
+
+	UPROPERTY()
+	USimpleDynamicMeshComponent* DynamicMeshComponent;
+
+	FViewCameraState CameraState;
+
+	//For buffer re-use and avoiding recurrent dynamic allocation
+	TArray<double> PCAVertexPositionBuffer;
+
+	//Array of the original positions of each vertex
+	TArray<FVector3d> OriginalPositions;
+
+	//Index of the axes being used. 
+	uint8 SelectedAxis{ 2 };
+	uint8 SecondaryAxis{ 1 };
+	uint8 ThirdAxis{0};
+
+	//Provided PCA, these are the direction vectors representing the principal axes to best align the mesh
+	FVector3d PrincipalAxesWorldSpace[3];
+	FVector3d PrincipalAxesObjectSpace[3];
+
+	// These represent the distance from the centroid to the farthest elements in the direction (respectively) of the PCA bases. 
+	// Specifically, the X element is the shortest of the PCA axes, the Y element is the middle axis, and the Z element represents the axis with the largest degree of variance, meaning Z is the greatest half extent.
+	FVector3d PrincipalAxesHalfExtentCoeff;
+
+	//Where the handle widget will be rendered
+	FVector3d AxisCentroidWorldSpace;
+	FVector3d AxisCentroidObjectSpace;
+
+	FToolDataVisualizer AxisRenderer;
+	FHandleWidget Handle{&AxisRenderer};
+
+	FBendMeshOp BendOp;
+	FTwistMeshOp TwistOp;
+	FFlareMeshOp FlareOp;
+
+	TArray<FMeshSpaceDeformerOp*> Operators
+	{
+		&BendOp,
+		&FlareOp,
+		&TwistOp
+		//Sinusoid
+		//Wave
+		//Squish
+	};
+
+
+	// realtime visualization
+	void OnDynamicMeshComponentChanged();
+	FDelegateHandle OnDynamicMeshComponentChangedHandle;
+	
+	//
+	// Update	
+	//
+	/** True for the duration of UI click+drag */
+	bool bInDrag;
+	
+private:
+};
+
