Mesh Modeling Tools moved from NotForLicensees to Experimental

diff --git a/Engine/Plugins/Experimental/MeshModelingToolset/Source/MeshModelingTools/Public/PolygonMeshDeformTool.h b/Engine/Plugins/Experimental/MeshModelingToolset/Source/MeshModelingTools/Public/PolygonMeshDeformTool.h
new file mode 100644
--- /dev/null
+++ b/Engine/Plugins/Experimental/MeshModelingToolset/Source/MeshModelingTools/Public/PolygonMeshDeformTool.h
@@ -0,0 +1,485 @@
+// Copyright 1998-2019 Epic Games, Inc. All Rights Reserved.
+
+#pragma once
+
+#include "CoreMinimal.h"
+#include "UObject/NoExportTypes.h"
+#include "BaseTools/MeshSurfacePointTool.h"
+#include "SimpleDynamicMeshComponent.h"
+#include "DynamicMeshAABBTree3.h"
+#include "Drawing/ToolDataVisualizer.h"
+#include "Transforms/QuickAxisTranslater.h"
+#include "Transforms/QuickAxisRotator.h"
+#include "MeshVertexChange.h"
+#include "GroupTopology.h"
+#include "Spatial/GeometrySet3.h"
+#include "Selection/GroupTopologySelector.h"
+#include "Operations/GroupTopologyDeformer.h"
+#include "MeshSolverUtilities/Private/LaplacianMeshSmoother.h"
+#include "Curves/CurveFloat.h"
+#include "ModelingOperators/Public/ModelingTaskTypes.h"
+#include "PolygonMeshDeformTool.generated.h"
+
+class FMeshVertexChangeBuilder;
+class FGroupTopologyLaplacianDeformer;
+class FDeformTask;
+struct FConstraintData;
+
+/**
+ * ToolBuilder
+ */
+UCLASS()
+class MESHMODELINGTOOLS_API UPolygonMeshDeformToolBuilder : public UMeshSurfacePointToolBuilder
+{
+	GENERATED_BODY()
+
+public:
+	UPolygonMeshDeformToolBuilder()
+	{
+	}
+
+	virtual UMeshSurfacePointTool* CreateNewTool(const FToolBuilderState& SceneState) const override;
+};
+
+/** Deformation strategies */
+UENUM()
+enum class EGroupTopologyDeformationStrategy : uint8
+{
+	/** Deforms the mesh using linear translations*/
+	Linear UMETA(DisplayName = "Linear"),
+
+	/** Deforms the mesh using laplacian deformation*/
+	Laplacian UMETA(DisplayName = "Laplacian")
+};
+
+/** Laplacian weight schemes determine how we will look at the curvature at a given vertex in relation to its neighborhood*/
+UENUM()
+enum class EWeightScheme 
+{
+	Uniform				UMETA(DisplayName = "Uniform"),
+	Umbrella			UMETA(DisplayName = "Umbrella"),
+	Valence				UMETA(DisplayName = "Valence"),
+	MeanValue			UMETA(DisplayName = "MeanValue"),
+	Cotangent			UMETA(DisplayName = "Cotangent"),
+	ClampedCotangent	UMETA(DisplayName = "ClampedCotangent")
+};
+
+/** The ELaplacianWeightScheme enum is the same..*/
+static ELaplacianWeightScheme ConvertToLaplacianWeightScheme(const EWeightScheme WeightScheme)
+{
+	return static_cast<ELaplacianWeightScheme>(WeightScheme);
+}
+
+/** Modes for quick transformer */
+UENUM()
+enum class EQuickTransformerMode : uint8
+{
+	/** Translation along frame axes */
+	AxisTranslation = 0 UMETA(DisplayName = "Translate"),
+
+	/** Rotation around frame axes*/
+	AxisRotation = 1 UMETA(DisplayName = "Rotate"),
+};
+
+
+
+UCLASS()
+class MESHMODELINGTOOLS_API UPolyEditTransformProperties : public UInteractiveToolPropertySet
+{
+	GENERATED_BODY()
+
+public:
+	UPolyEditTransformProperties();
+
+	//Options
+
+	UPROPERTY(EditAnywhere, Category = Options, meta = (DisplayName = "Deformation Type", ToolTip = "Select the type of deformation you wish to employ on a polygroup."))
+	EGroupTopologyDeformationStrategy DeformationStrategy;
+
+	UPROPERTY(EditAnywhere, Category = Options)
+	bool bUseLocalAxis;
+
+	UPROPERTY(EditAnywhere, Category = Options)
+	EQuickTransformerMode TransformMode;
+
+	UPROPERTY(EditAnywhere, Category = Options)
+	bool bSelectFaces;
+
+	UPROPERTY(EditAnywhere, Category = Options)
+	bool bSelectEdges;
+
+	UPROPERTY(EditAnywhere, Category = Options)
+	bool bSelectVertices;
+
+	//Laplacian Deformation Options
+
+	UPROPERTY(EditAnywhere, Category = LaplacianOptions, meta = (EditCondition="DeformationStrategy == EGroupTopologyDeformationStrategy::Laplacian", DisplayName = "Laplacian Weight Scheme", ToolTip = "Laplacian weight schemes determine how the deformer will interpret the curvature at a given vertex in relation to nearby vertices in order to generate a solution."))
+	EWeightScheme SelectedWeightScheme = EWeightScheme::ClampedCotangent;
+
+	UPROPERTY(EditAnywhere, Category = LaplacianOptions, meta = (EditCondition="DeformationStrategy == EGroupTopologyDeformationStrategy::Laplacian", DisplayName = "Handle Weight", UIMin = "0.0001", UIMax = "1.0", ClampMin = "0.0001", ClampMax = "100.0", ToolTip = "The handle weight describes how effective dragging the selected PolyGroup will be."))
+	double HandleWeight = 1.0;
+
+	UPROPERTY(EditAnywhere, Category = LaplacianOptions, meta = (EditCondition="DeformationStrategy == EGroupTopologyDeformationStrategy::Laplacian", DisplayName = "Post-fix Handles", ToolTip = "Post-fixing the handles will attempt to force their vertices to arrive exactly where you wanted them. This may result\nin undesirable or choppy curvature in some scenarios."))
+	bool bPostFixHandles = false;
+/**
+	// How to add a weight curve
+	UPROPERTY(EditAnywhere, Category = LaplacianOptions, meta = (EditCondition="DeformationStrategy == EGroupTopologyDeformationStrategy::Laplacian", DisplayName = "Localize Deformation", ToolTip = "When enabled, only the vertices in the polygroups immediately adjacent to the selected group will be affected by the deformation.\nWhen disabled, the deformer will solve for the curvature of the entire mesh (slower)"))
+	bool bLocalizeDeformation{ true};
+
+	UPROPERTY(EditAnywhere, Category = LaplacianOptions, meta = (EditCondition="DeformationStrategy == EGroupTopologyDeformationStrategy::Laplacian", DisplayName = "Apply Weight Attenuation Curve", ToolTip = "When enabled, the curve below will be used to calculate the weight at a given vertex based on distance from the handles"))
+	bool bApplyAttenuationCurve{ false };
+	
+	FRichCurve DefaultFalloffCurve;
+
+	UPROPERTY(EditAnywhere,  Category = LaplacianOptions, meta = ( EditCondition="DeformationStrategy == EGroupTopologyDeformationStrategy::Laplacian && bApplyAttenuationCurve", DisplayName = "Distance-Weight Attenuation Curve",UIMin = "1.0", UIMax = "1.0", ClampMin = "1.0", ClampMax = "1.0", ToolTip = "This curve determines the weight attenuation over the distance of the mesh.\nThe selected polygroup handle is t=0.0, and t=1.0 is roughly the farthest vertices from the handles.\nThe value of the curve at each time interval represents the weight of the vertices at that distance from the selection."))
+	FRuntimeFloatCurve WeightAttenuationCurve;
+*/ 
+};
+
+//Stores per-vertex data needed by the laplacian deformer object
+//TODO: May be a candidate for a subclass of the FGroupTopologyLaplacianDeformer
+struct FConstraintData
+{
+	FConstraintData& operator=(const FConstraintData& other)
+	{
+		Position = other.Position;
+		Weight   = other.Weight;
+		bPostFix = other.bPostFix;
+		return *this;
+	}
+
+	FVector3d Position;
+	double Weight{ 0.0 };
+	bool bPostFix{ false };
+};
+
+/**
+*	FDeformTask is an object which wraps an asynchronous task to be run multiple times on a separate thread.
+*	The Laplacian deformation process requires the use of potentially large sparse matrices and sparse multiplication.
+*	
+*   Expected usage:
+*
+*   
+*   // define constraints.  Need Constraints[VertID] to hold the constraints for the corresponding vertex.
+*   TArray<FConstraintData> Constraints;
+*   ....
+*
+*   // populate with the VertexIDs of the vertices that are in the region you wish to deform.  
+*   TArray<int32> SrcVertIDs;  //Basically a mini-index buffer.
+*   ...
+* 
+*   // Create or reuse a laplacian deformation task.
+*   FDeformTask*   DeformTask = New FDeformTask(WeightScheme);
+*
+*   // the deformer will have to build a new mesh that represents the regions in SrcVertIDs;
+*   // but set this to false on subsequent calls to UpdateDeformer if the SrcVertIDs array hasn't changed.
+*   bool bRequiresRegion = true;
+
+*   DefTask->UpdateDeformer(WeightScheme, Mesh, Constraints, SrcVertIDs, bRequiresRegion);
+*
+*   DeformTask->DoWork();  or DeformTask->StartBackgroundTask(); //which calls DoWork on background thread.
+*
+*  // wheh DeformTask->IsDone == true; you can copy the results back to the mesh
+*  DeformTask->ExportResults(Mesh);
+*
+* Note: if only the positions in the Constraints change (e.g. handle positions) then subsequent calls 
+* to UpdateDeformer() and DoWork() will be much faster as the matrix system will not be rebuilt or re-factored
+*/
+class FConstrainedMeshDeformerTask : public FNonAbandonableTask
+{
+	friend class FAsyncTask<FDeformTask>;
+public:
+
+	enum
+	{
+		INACTIVE_SUBSET_ID = -1
+	};
+
+	FConstrainedMeshDeformerTask(const ELaplacianWeightScheme SelectedWeightScheme)
+	{
+		LaplacianWeightScheme = SelectedWeightScheme;
+	}
+
+	virtual ~FConstrainedMeshDeformerTask() {};
+
+	//NO idea what this is meant to do. Performance analysis maybe? Scheduling data?
+	FORCEINLINE TStatId GetStatId() const
+	{
+		RETURN_QUICK_DECLARE_CYCLE_STAT(FConstrainedMeshDeformerTask, STATGROUP_ThreadPoolAsyncTasks);
+	}
+
+	
+	/** Called by the main thread in the tool, this copies the Constraint buffer right before the task begins on another thread.
+	  * Ensures the FConstrainedMeshDeformer is using correct mesh subset and the selected settings, then updates on change in properties, i.e. weight scheme */
+	void UpdateDeformer(const ELaplacianWeightScheme SelectedWeightScheme, const FDynamicMesh3& Mesh, const TArray<FConstraintData>& ConstraintArray, const TArray<int32>& SrcIDBufferSubset, bool bNewTransaction, const FRichCurve* Curve);
+
+	/** Required by the FAsyncTaskExecutor */
+	void SetAbortSource(bool* bAbort) { bAbortSource = bAbort; };
+
+	/** Called by the FAsyncTask<FDeformTask> object for background computation. */
+	void DoWork();
+
+
+	/** Updates the positions in the target mesh for regions that correspond to the subset mesh */
+	void ExportResults(FDynamicMesh3& TargetMesh) const;
+
+private:
+
+	/** Creates the mesh (i.e. SubsetMesh) that corresponds to the region of the SrcMesh defined by the partial index buffer SrcIDBufferSubset */
+	void InitializeSubsetMesh(const FDynamicMesh3& SrcMesh, const TArray<int32>& SrcIDBufferSubset);
+
+	/** Attenuates the weights of the constraints using the selected curve */
+	void ApplyAttenuation();
+
+	/** Denotes the weight scheme being used by the running background task. Changes when selected property changes in editor. */
+	ELaplacianWeightScheme LaplacianWeightScheme;
+
+	/** positions for each vertex in the subset mesh - for use in the deformer */
+	TArray<FVector3d> SubsetPositionBuffer;
+
+	/** constraint data for each vertex in subset mesh - for use by the deformer*/
+	TArray<FConstraintData> SubsetConstraintBuffer;
+
+	FRichCurve WeightAttenuationCurve;
+
+	/** True only for the first update, and then false for the duration of the Input transaction
+	  * It's passed in and copied in UpdateDeformer() */
+	bool bIsNewTransaction = true;
+
+	/** When true, the constraint weights will be attenuated based on distance using the provided curve object*/
+	bool bAttenuateWeights = false;
+
+	/** The abort bool used by the Task Deleter */
+	bool* bAbortSource = nullptr;
+
+	/** Used to initialize the array mapping, updated during the UpdateDeformer() function */
+	int SrcMeshMaxVertexID;
+
+	/** A subset of the original mesh */
+	FDynamicMesh3 SubsetMesh;
+
+	/** Maps Subset Mesh VertexID to Src Mesh VertexID */
+	TArray<int32> SubsetVertexIDToSrcVertexIDMap;
+
+	/** Laplacian deformer object gets rebuilt each new transaction */
+	TUniquePtr<FConstrainedMeshDeformer> ConstrainedDeformer;
+
+private :
+	FConstrainedMeshDeformerTask();
+
+};
+
+class FGroupTopologyLaplacianDeformer : public FGroupTopologyDeformer
+{
+
+public:
+
+	FGroupTopologyLaplacianDeformer() = default;
+
+	virtual ~FGroupTopologyLaplacianDeformer();
+
+	/** Used to begin a procedural addition of modified vertices */
+	inline void ResetModifiedVertices()
+	{
+		ModifiedVertices.Empty();
+	};
+
+	/** Change tracking */
+	template <typename ValidSetAppendContainerType>
+	void RecordModifiedVertices(const ValidSetAppendContainerType& Container)
+	{
+		ModifiedVertices.Empty();
+		ModifiedVertices.Append(Container);
+	}
+
+	/** Used to iteratively add to the active change set (TSet<>)*/
+	inline void RecordModifiedVertex(int32 VertexID)
+	{
+		ModifiedVertices.Add(VertexID);
+	};
+
+
+	void SetActiveHandleFaces(const TArray<int>& FaceGroupIDs) override;
+	void SetActiveHandleEdges(const TArray<int>& TopologyEdgeIDs) override;
+	void SetActiveHandleCorners(const TArray<int>& TopologyCornerIDs) override;
+
+
+
+	/** Allocates shared storage for use in task synchronization */
+	void InitBackgroundWorker(const ELaplacianWeightScheme WeightScheme);
+
+
+	/** Coordinates the background tasks. Returns false if the worker was already running */
+	bool UpdateAndLaunchdWorker(const ELaplacianWeightScheme WeightScheme, const FRichCurve* Curve = nullptr);
+
+	/** Capture data about background task state.*/
+	bool IsTaskInFlight() const;
+
+
+
+	/** Sets the SrcMeshConstraintBuffer to have a size of MaxVertexID, and initializes with the current mesh positions, but weight zero*/
+	void InitializeConstraintBuffer();
+
+	/** Given an array of Group IDs, update the selection and record vertices */
+	void UpdateSelection(const FDynamicMesh3* TargetMesh, const TArray<int>& Groups, bool bLocalizeDeformation);
+
+	/** Updates the mesh preview and/or solvers upon user input, provided a deformation strategy */
+	void UpdateSolution(FDynamicMesh3* TargetMesh, const TFunction<FVector3d(FDynamicMesh3* Mesh, int)>& HandleVertexDeformFunc) override;
+
+	/** Updates the vertex positions of the mesh with the result from the last deformation solve. */
+	void ExportDeformedPositions(FDynamicMesh3* TargetMesh);
+
+	/** Returns true if the asynchronous task has finished. */
+	inline bool IsDone() { return AsyncMeshDeformTask == nullptr || AsyncMeshDeformTask->IsDone(); };
+
+	/** Triggers abort on task and passes off ownership to deleter object */
+	inline void Shutdown();
+
+	/** Stores the position of the vertex constraints and corresponding weights for the entire mesh.  This is used as a form of scratch space.*/
+	TArray<FConstraintData> SrcMeshConstraintBuffer;
+
+	/** Array of vertex indices organized in groups of three - basically an index buffer - that defines the subset of the mesh that the deformation task will work on.*/
+	TArray<int32> SubsetIDBuffer;
+
+	/** Need to update the task with the current submesh */
+	bool bTaskSubmeshIsDirty = true;
+
+	/** Asynchronous task object. This object deals with expensive matrix functionality that computes the deformation of a local mesh. */
+	FAsyncTaskExecuterWithAbort<FConstrainedMeshDeformerTask>* AsyncMeshDeformTask = nullptr;
+
+
+	/** The weight which will be applied to the constraints corresponding to the handle vertices. */
+	double HandleWeights = 1.0;
+
+	/** This is set to true whenever the user interacts with the tool under laplacian deformation mode.
+	  * It is set to false immediately before beginning a background task and cannot be set to false again until the work is done. */
+	bool bDeformerNeedsToRun = false;
+
+
+	/** When true, tells the solver to attempt to postfix the actual position of the handles to the constrained position */
+	bool bPostfixHandles = false;
+
+	//This is set to false only after 
+	//	1) the asynchronous deformation task is complete
+	//	2) the main thread has seen it complete, and
+	//	3) the main thread updates the vertex positions of the mesh one last time
+	bool bVertexPositionsNeedSync = false;
+
+	bool bLocalize =  true;
+
+};
+
+
+/**
+ *
+ */
+UCLASS()
+class MESHMODELINGTOOLS_API UPolygonMeshDeformTool : public UMeshSurfacePointTool
+{
+	GENERATED_BODY()
+
+public:
+	UPolygonMeshDeformTool();
+
+	virtual void RegisterActions(FInteractiveToolActionSet& ActionSet) override;
+
+	virtual void Setup() override;
+	virtual void Shutdown(EToolShutdownType ShutdownType) override;
+
+	virtual void Tick(float DeltaTime) override;
+	virtual void Render(IToolsContextRenderAPI* RenderAPI) override;
+
+	virtual bool HasCancel() const override { return true; }
+	virtual bool HasAccept() const override { return true; }
+	virtual bool CanAccept() const override { return true; }
+
+
+	virtual bool HitTest(const FRay& Ray, FHitResult& OutHit) override;
+	
+	virtual void OnBeginDrag(const FRay& Ray) override;
+	virtual void OnUpdateDrag(const FRay& Ray) override;
+	virtual void OnEndDrag(const FRay& Ray) override;
+	virtual void OnUpdateHover(const FInputDeviceRay& DevicePos) override;
+
+public:
+	virtual void ToggleLocalTransformFrameAction();
+	virtual void NextTransformTypeAction();
+
+	//
+	float VisualAngleSnapThreshold = 0.5;
+
+protected:
+	UPROPERTY()
+	USimpleDynamicMeshComponent* DynamicMeshComponent;
+
+	UPROPERTY()
+	UPolyEditTransformProperties* TransformProps;
+
+	// realtime visualization
+	void OnDynamicMeshComponentChanged();
+	FDelegateHandle OnDynamicMeshComponentChangedHandle;
+	
+	// camera state at last render
+	FViewCameraState CameraState;
+
+	FToolDataVisualizer PolyEdgesRenderer;
+
+	// True for the duration of UI click+drag
+	bool bInDrag;
+
+	FPlane ActiveDragPlane;
+	FVector StartHitPosWorld;
+	FVector StartHitNormalWorld;
+	FVector LastHitPosWorld;
+	FVector LastBrushPosLocal;
+	FVector StartBrushPosLocal;
+
+	FFrame3d ActiveSurfaceFrame;
+	FQuickTransformer* GetActiveQuickTransformer();
+	void UpdateActiveSurfaceFrame(FGroupTopologySelection& Selection);
+	void UpdateQuickTransformer();
+
+	FRay UpdateRay;
+	bool bUpdatePending = false;
+	void ComputeUpdate();
+
+	FVector3d LastMoveDelta;
+	FQuickAxisTranslater QuickAxisTranslater;
+	void ComputeUpdate_Translate();
+
+	FQuickAxisRotator QuickAxisRotator;
+	FVector3d RotationStartPointWorld;
+	FFrame3d RotationStartFrame;
+	void ComputeUpdate_Rotate();
+
+
+	FGroupTopology Topology;
+	void PrecomputeTopology();
+
+	FGroupTopologySelector TopoSelector;
+	
+	//
+	// data for current drag
+	//
+
+	FGroupTopologySelection HilightSelection;
+	FToolDataVisualizer HilightRenderer;
+
+	FDynamicMeshAABBTree3 MeshSpatial;
+	FDynamicMeshAABBTree3& GetSpatial();
+
+	FMeshVertexChangeBuilder* ActiveVertexChange;
+
+	// The two deformer type options.
+	FGroupTopologyDeformer LinearDeformer;
+	FGroupTopologyLaplacianDeformer LaplacianDeformer;
+	
+	// This is true when the spatial index needs to reflect a modification
+	bool bSpatialDirty; 
+
+	void BeginChange();
+	void EndChange();
+	void UpdateChangeFromROI(bool bFinal);
+
+};
+
diff --git a/Engine/Plugins/Experimental/MeshModelingToolset/Source/MeshModelingTools/Public/Properties/MeshStatisticsProperties.h b/Engine/Plugins/Experimental/MeshModelingToolset/Source/MeshModelingTools/Public/Properties/MeshStatisticsProperties.h
new file mode 100644
--- /dev/null
+++ b/Engine/Plugins/Experimental/MeshModelingToolset/Source/MeshModelingTools/Public/Properties/MeshStatisticsProperties.h
@@ -0,0 +1,31 @@
+// Copyright 1998-2019 Epic Games, Inc. All Rights Reserved.
+
+#pragma once
+
+#include "CoreMinimal.h"
+
+#include "InteractiveToolBuilder.h"
+
+#include "MeshStatisticsProperties.generated.h"
+
+
+// predeclarations
+class FDynamicMesh3;
+
+
+
+UCLASS()
+class MESHMODELINGTOOLS_API UMeshStatisticsProperties : public UInteractiveToolPropertySet
+{
+	GENERATED_BODY()
+public:
+
+	UPROPERTY(VisibleAnywhere, Category = MeshStatistics)
+	int TriangleCount = 0;
+
+	UPROPERTY(VisibleAnywhere, Category = MeshStatistics)
+	int VertexCount = 0;
+
+	void Update(const FDynamicMesh3& Mesh);
+};
+
diff --git a/Engine/Plugins/Experimental/MeshModelingToolset/Source/MeshModelingTools/Public/RemeshMeshTool.h b/Engine/Plugins/Experimental/MeshModelingToolset/Source/MeshModelingTools/Public/RemeshMeshTool.h
new file mode 100644
--- /dev/null
+++ b/Engine/Plugins/Experimental/MeshModelingToolset/Source/MeshModelingTools/Public/RemeshMeshTool.h
@@ -0,0 +1,176 @@
+// Copyright 1998-2019 Epic Games, Inc. All Rights Reserved.
+
+#pragma once
+
+#include "CoreMinimal.h"
+#include "UObject/NoExportTypes.h"
+#include "SingleSelectionTool.h"
+#include "InteractiveToolBuilder.h"
+#include "DynamicMesh3.h"
+#include "DynamicMeshAABBTree3.h"
+#include "Properties/MeshStatisticsProperties.h"
+#include "RemeshMeshTool.generated.h"
+
+// predeclarations
+struct FMeshDescription;
+class USimpleDynamicMeshComponent;
+
+/** Mesh paint color view modes (somewhat maps to EVertexColorViewMode engine enum.) */
+UENUM()
+enum class ERemeshSmoothingType : uint8
+{
+	/** Uniform Smoothing */
+	Uniform = 0 UMETA(DisplayName = "Uniform"),
+
+	/** Cotangent Smoothing */
+	Cotangent = 1 UMETA(DisplayName = "Shape Preserving"),
+
+	/** Mean Value Smoothing */
+	MeanValue = 2 UMETA(DisplayName = "Mixed")
+};
+
+/**
+ *
+ */
+UCLASS()
+class MESHMODELINGTOOLS_API URemeshMeshToolBuilder : public UInteractiveToolBuilder
+{
+	GENERATED_BODY()
+
+public:
+	virtual bool CanBuildTool(const FToolBuilderState& SceneState) const override;
+	virtual UInteractiveTool* BuildTool(const FToolBuilderState& SceneState) const override;
+};
+
+/**
+ * Standard properties of the Remesh operation
+ */
+UCLASS()
+class MESHMODELINGTOOLS_API URemeshMeshToolProperties : public UInteractiveToolPropertySet
+{
+	GENERATED_BODY()
+
+public:
+	URemeshMeshToolProperties();
+
+	/** Target triangle count */
+	UPROPERTY(EditAnywhere, Category = Options)
+	int TargetTriangleCount;
+
+
+	/** Smoothing speed */
+	UPROPERTY(EditAnywhere, Category = Options, meta = (UIMin = "0.0", UIMax = "1.0", ClampMin = "0.0", ClampMax = "1.0"))
+	float SmoothingSpeed;
+
+	/** Smoothing type */
+	UPROPERTY(EditAnywhere, Category = Options)
+	ERemeshSmoothingType SmoothingType;
+
+	/** Number of Remeshing passes */
+	UPROPERTY(EditAnywhere, Category = Options, meta = (UIMin = "0", UIMax = "50", ClampMin = "0", ClampMax = "1000"))
+	int RemeshIterations;
+
+
+	/** If true, UVs and Normals are discarded  */
+	UPROPERTY(EditAnywhere, Category = Options)
+	bool bDiscardAttributes;
+
+	/** If true, sharp edges are preserved  */
+	UPROPERTY(EditAnywhere, Category = Options)
+	bool bPreserveSharpEdges;
+
+};
+
+/**
+ * Advanced properties of the remesh operation
+ */
+UCLASS()
+class MESHMODELINGTOOLS_API URemeshMeshAdvancedProperties : public UInteractiveToolPropertySet
+{
+	GENERATED_BODY()
+
+public:
+	URemeshMeshAdvancedProperties();
+
+	/** If true, the target count is ignored and the target edge length is used directly */
+	UPROPERTY(EditAnywhere, Category = Advanced)
+	bool bUseTargetEdgeLength;
+
+	/** Target edge length */
+	UPROPERTY(EditAnywhere, Category = Advanced, meta = (NoSpinbox = "true"))
+	float TargetEdgeLength;
+
+	/** Enable edge flips */
+	UPROPERTY(EditAnywhere, Category = Advanced)
+	bool bFlips;
+
+	/** Enable edge splits */
+	UPROPERTY(EditAnywhere, Category = Advanced)
+	bool bSplits;
+
+	/** Enable edge collapses */
+	UPROPERTY(EditAnywhere, Category = Advanced)
+	bool bCollapses;
+
+	/** Enable projection back to input mesh */
+	UPROPERTY(EditAnywhere, Category = Advanced)
+	bool bReproject;
+
+	/** Prevent normal flips */
+	UPROPERTY(EditAnywhere, Category = Advanced)
+	bool bPreventNormalFlips;
+
+};
+
+/**
+ * Simple Mesh Remeshing Tool
+ */
+UCLASS()
+class MESHMODELINGTOOLS_API URemeshMeshTool : public USingleSelectionTool
+{
+	GENERATED_BODY()
+
+public:
+	URemeshMeshTool();
+
+	virtual void Setup() override;
+	virtual void Shutdown(EToolShutdownType ShutdownType) override;
+
+	virtual void Render(IToolsContextRenderAPI* RenderAPI) override;
+
+	virtual bool HasCancel() const override { return true; }
+	virtual bool HasAccept() const override;
+	virtual bool CanAccept() const override;
+
+
+#if WITH_EDITOR
+	virtual void PostEditChangeProperty(FPropertyChangedEvent &PropertyChangedEvent) override;
+#endif
+
+	virtual void OnPropertyModified(UObject* PropertySet, UProperty* Property) override;
+
+public:
+
+	UPROPERTY()
+	URemeshMeshToolProperties* BasicProperties;
+
+	UPROPERTY()
+	URemeshMeshAdvancedProperties* AdvancedProperties;
+
+	UPROPERTY()
+	UMeshStatisticsProperties* MeshStatisticsProperties;
+
+
+protected:
+	USimpleDynamicMeshComponent* DynamicMeshComponent;
+	FDynamicMesh3 OriginalMesh;
+	FDynamicMeshAABBTree3 OriginalMeshSpatial;
+	bool bResultValid;
+	void UpdateResult();
+
+	double InitialMeshArea;
+
+	double CalculateTargetEdgeLength(int TargetTriCount);
+
+	void ComputeRemeshing();
+};
diff --git a/Engine/Plugins/Experimental/MeshModelingToolset/Source/MeshModelingTools/Public/TransformMeshesTool.h b/Engine/Plugins/Experimental/MeshModelingToolset/Source/MeshModelingTools/Public/TransformMeshesTool.h
new file mode 100644
--- /dev/null
+++ b/Engine/Plugins/Experimental/MeshModelingToolset/Source/MeshModelingTools/Public/TransformMeshesTool.h
@@ -0,0 +1,141 @@
+// Copyright 1998-2019 Epic Games, Inc. All Rights Reserved.
+
+#pragma once
+
+#include "CoreMinimal.h"
+#include "MultiSelectionTool.h"
+#include "InteractiveToolBuilder.h"
+#include "BaseBehaviors//BehaviorTargetInterfaces.h"
+#include "FrameTypes.h"
+#include "TransformMeshesTool.generated.h"
+
+class UBaseAxisTranslationGizmo;
+class UAxisAngleGizmo;
+class UTransformGizmo;
+class UTransformProxy;
+
+/**
+ *
+ */
+UCLASS()
+class MESHMODELINGTOOLS_API UTransformMeshesToolBuilder : public UInteractiveToolBuilder
+{
+	GENERATED_BODY()
+
+public:
+	virtual bool CanBuildTool(const FToolBuilderState& SceneState) const override;
+	virtual UInteractiveTool* BuildTool(const FToolBuilderState& SceneState) const override;
+};
+
+
+
+/** Mesh Transform modes */
+UENUM()
+enum class ETransformMeshesTransformMode : uint8
+{
+	/** Shared Gizmo */
+	SharedGizmo = 0 UMETA(DisplayName = "Shared Gizmo"),
+
+	/** Shared Gizmo, Local Transformations*/
+	SharedGizmoLocal = 1 UMETA(DisplayName = "Shared Gizmo (Local)"),
+
+	/** Per Object Gizmo */
+	PerObjectGizmo = 2 UMETA(DisplayName = "Multi-Gizmo"),
+
+	/** Quick Translate Transformer */
+	QuickTranslate = 3 UMETA(DisplayName = "QuickTranslate")
+};
+
+
+
+/**
+ * Standard properties of the Transform Meshes operation
+ */
+UCLASS()
+class MESHMODELINGTOOLS_API UTransformMeshesToolProperties : public UInteractiveToolPropertySet
+{
+	GENERATED_BODY()
+public:
+
+	UPROPERTY(EditAnywhere, Category = Options)
+	ETransformMeshesTransformMode TransformMode = ETransformMeshesTransformMode::SharedGizmo;
+
+
+	UPROPERTY(EditAnywhere, Category = Options, meta = (EditCondition = "TransformMode == ETransformMeshesTransformMode::SharedGizmo || TransformMode == ETransformMeshesTransformMode::PerObjectGizmo"))
+	bool bSetPivot = false;
+};
+
+
+USTRUCT()
+struct FTransformMeshesTarget
+{
+	GENERATED_BODY()
+
+	UPROPERTY()
+	UTransformProxy* TransformProxy;
+
+	UPROPERTY()
+	UTransformGizmo* TransformGizmo;
+};
+
+
+/**
+ *
+ */
+UCLASS()
+class MESHMODELINGTOOLS_API UTransformMeshesTool : public UMultiSelectionTool, public IClickDragBehaviorTarget
+{
+	GENERATED_BODY()
+
+public:
+	UTransformMeshesTool();
+
+	virtual void SetWorld(UWorld* World, UInteractiveGizmoManager* GizmoManager);
+
+	virtual void Setup() override;
+	virtual void Shutdown(EToolShutdownType ShutdownType) override;
+
+	virtual void Tick(float DeltaTime) override;
+	virtual void Render(IToolsContextRenderAPI* RenderAPI) override;
+
+	virtual bool HasCancel() const override { return false; }
+	virtual bool HasAccept() const override { return false; }
+	virtual bool CanAccept() const override { return false; }
+
+	virtual void OnPropertyModified(UObject* PropertySet, UProperty* Property) override;
+
+
+	// ICLickDragBehaviorTarget interface
+	virtual FInputRayHit CanBeginClickDragSequence(const FInputDeviceRay& PressPos) override;
+	virtual void OnClickPress(const FInputDeviceRay& PressPos) override;
+	virtual void OnClickDrag(const FInputDeviceRay& DragPos) override;
+	virtual void OnClickRelease(const FInputDeviceRay& ReleasePos) override;
+	virtual void OnTerminateDragSequence() override;
+
+protected:
+	UPROPERTY()
+	UTransformMeshesToolProperties* TransformProps;
+
+protected:
+	UWorld* TargetWorld;
+	UInteractiveGizmoManager* GizmoManager;
+
+	UPROPERTY()
+	TArray<FTransformMeshesTarget> ActiveGizmos;
+
+
+	ETransformMeshesTransformMode CurTransformMode;
+	void UpdateTransformMode(ETransformMeshesTransformMode NewMode);
+	void UpdateSetPivotModes(bool bEnableSetPivot);
+
+	void SetActiveGizmos_Single(bool bLocalRotations);
+	void SetActiveGizmos_PerObject();
+	void ResetActiveGizmos();
+
+
+	FFrame3d StartDragFrameWorld;
+	FTransform StartDragTransform;
+
+
+	void OnParametersUpdated();
+};
diff --git a/Engine/Plugins/Experimental/MeshModelingToolset/Source/MeshModelingTools/Public/UVProjectionTool.h b/Engine/Plugins/Experimental/MeshModelingToolset/Source/MeshModelingTools/Public/UVProjectionTool.h
new file mode 100644
--- /dev/null
+++ b/Engine/Plugins/Experimental/MeshModelingToolset/Source/MeshModelingTools/Public/UVProjectionTool.h
@@ -0,0 +1,191 @@
+// Copyright 1998-2019 Epic Games, Inc. All Rights Reserved.
+
+#pragma once
+
+#include "CoreMinimal.h"
+#include "UObject/NoExportTypes.h"
+#include "MultiSelectionTool.h"
+#include "InteractiveToolBuilder.h"
+#include "MeshOpPreviewHelpers.h"
+#include "Drawing/ToolDataVisualizer.h"
+#include "ParameterizationOps/UVProjectionOp.h"
+#include "DynamicMesh3.h"
+#include "BaseTools/SingleClickTool.h"
+
+#include "UVProjectionTool.generated.h"
+
+
+// predeclarations
+struct FMeshDescription;
+class USimpleDynamicMeshComponent;
+class UTransformGizmo;
+class UTransformProxy;
+
+
+/**
+ *
+ */
+UCLASS()
+class MESHMODELINGTOOLS_API UUVProjectionToolBuilder : public UInteractiveToolBuilder
+{
+	GENERATED_BODY()
+
+public:
+	IToolsContextAssetAPI* AssetAPI = nullptr;
+
+	virtual bool CanBuildTool(const FToolBuilderState& SceneState) const override;
+	virtual UInteractiveTool* BuildTool(const FToolBuilderState& SceneState) const override;
+};
+
+
+
+
+
+/**
+ * Standard properties
+ */
+UCLASS()
+class MESHMODELINGTOOLS_API UUVProjectionToolProperties : public UInteractiveToolPropertySet
+{
+	GENERATED_BODY()
+
+public:
+	UUVProjectionToolProperties();
+
+
+
+	/** Choose the UV projection method (cube, cylinder, plane) */
+	UPROPERTY(EditAnywhere, Category = Options)
+	EUVProjectionMethod UVProjectionMethod;
+
+	/** Per-axis scaling of projection primitive */
+	UPROPERTY(EditAnywhere, Category = Options)
+	FVector ProjectionPrimitiveScale;
+
+	/** If triangle normal direction is within this threshold degrees of the cylinder top/bottom plane direction, project UVs to the top/bottom plane instead of the sides */
+	UPROPERTY(EditAnywhere, Category = Options, meta = (Default = "1", UIMin = "0", UIMax = "20", ClampMin = "0.0", ClampMax = "90.0", EditCondition = "UVProjectionMethod == EUVProjectionMethod::Cylinder"))
+	float CylinderProjectToTopOrBottomAngleThreshold;
+
+	/** Choose the UV scale factors */
+	UPROPERTY(EditAnywhere, Category = Options)
+	FVector2D UVScale;
+
+	/** Choose the UV offsets */
+	UPROPERTY(EditAnywhere, Category = Options)
+	FVector2D UVOffset;
+
+	/** For display only - the density of the checkerboard pattern in the viewport.*/
+	UPROPERTY(EditAnywhere, Category = Options, meta = (Default = "20", UIMin = "1.0", UIMax = "40.0", ClampMin = "0.01", ClampMax = "1000.0"))
+	float CheckerDensity = 10.f;
+
+};
+
+
+
+
+/**
+ * Advanced properties
+ */
+UCLASS()
+class MESHMODELINGTOOLS_API UUVProjectionAdvancedProperties : public UInteractiveToolPropertySet
+{
+	GENERATED_BODY()
+
+public:
+	UUVProjectionAdvancedProperties();
+};
+
+
+/**
+ * Factory with enough info to spawn the background-thread Operator to do a chunk of work for the tool
+ *  stores a pointer to the tool and enough info to know which specific operator it should spawn
+ */
+UCLASS()
+class MESHMODELINGTOOLS_API UUVProjectionOperatorFactory : public UObject, public IDynamicMeshOperatorFactory
+{
+	GENERATED_BODY()
+
+public:
+	// IDynamicMeshOperatorFactory API
+	virtual TSharedPtr<FDynamicMeshOperator> MakeNewOperator() override;
+
+	UPROPERTY()
+	UUVProjectionTool *Tool;
+
+	int ComponentIndex;
+
+};
+
+/**
+ * Simple Mesh Normal Updating Tool
+ */
+UCLASS()
+class MESHMODELINGTOOLS_API UUVProjectionTool : public UMultiSelectionTool
+{
+	GENERATED_BODY()
+
+public:
+
+	friend UUVProjectionOperatorFactory;
+
+	UUVProjectionTool();
+
+	virtual void Setup() override;
+	virtual void Shutdown(EToolShutdownType ShutdownType) override;
+
+	virtual void SetWorld(UWorld* World, UInteractiveGizmoManager* GizmoManagerIn);
+	virtual void SetAssetAPI(IToolsContextAssetAPI* AssetAPI);
+
+	virtual void Tick(float DeltaTime) override;
+	virtual void Render(IToolsContextRenderAPI* RenderAPI) override;
+
+	virtual bool HasCancel() const override { return true; }
+	virtual bool HasAccept() const override;
+	virtual bool CanAccept() const override;
+
+#if WITH_EDITOR
+	virtual void PostEditChangeProperty(FPropertyChangedEvent &PropertyChangedEvent) override;
+#endif
+
+	virtual void OnPropertyModified(UObject* PropertySet, UProperty* Property) override;
+
+protected:
+
+	UPROPERTY()
+	UUVProjectionToolProperties* BasicProperties;
+
+	UPROPERTY()
+	UUVProjectionAdvancedProperties* AdvancedProperties;
+
+
+	UPROPERTY()
+	TArray<UMeshOpPreviewWithBackgroundCompute*> Previews;
+
+
+	UPROPERTY()
+	UMaterialInstanceDynamic* CheckerMaterial = nullptr;
+
+	UPROPERTY()
+	TArray<UTransformGizmo*> TransformGizmos;
+	
+	UPROPERTY()
+	TArray<UTransformProxy*> TransformProxies;
+
+protected:
+	TArray<TSharedPtr<FDynamicMesh3>> OriginalDynamicMeshes;
+	TArray<FDynamicMesh3> ReferencePrimitives;
+
+	UWorld* TargetWorld;
+	IToolsContextAssetAPI* AssetAPI;
+	UInteractiveGizmoManager* GizmoManager;
+
+	FViewCameraState CameraState;
+
+	FToolDataVisualizer ProjectionShapeVisualizer;
+
+	void UpdateNumPreviews();
+
+	void GenerateAsset(const TArray<TUniquePtr<FDynamicMeshOpResult>>& Results);
+
+	void TransformChanged(UTransformProxy* Proxy, FTransform Transform);
+};
diff --git a/Engine/Plugins/Experimental/MeshModelingToolset/Source/MeshModelingTools/Public/WeldMeshEdgesTool.h b/Engine/Plugins/Experimental/MeshModelingToolset/Source/MeshModelingTools/Public/WeldMeshEdgesTool.h
new file mode 100644
--- /dev/null
+++ b/Engine/Plugins/Experimental/MeshModelingToolset/Source/MeshModelingTools/Public/WeldMeshEdgesTool.h
@@ -0,0 +1,71 @@
+// Copyright 1998-2019 Epic Games, Inc. All Rights Reserved.
+
+#pragma once
+
+#include "CoreMinimal.h"
+#include "UObject/NoExportTypes.h"
+#include "SingleSelectionTool.h"
+#include "InteractiveToolBuilder.h"
+#include "DynamicMesh3.h"
+#include "DynamicMeshAABBTree3.h"
+#include "WeldMeshEdgesTool.generated.h"
+
+
+// predeclarations
+struct FMeshDescription;
+class USimpleDynamicMeshComponent;
+
+/**
+ *
+ */
+UCLASS()
+class MESHMODELINGTOOLS_API UWeldMeshEdgesToolBuilder : public UInteractiveToolBuilder
+{
+	GENERATED_BODY()
+
+public:
+	virtual bool CanBuildTool(const FToolBuilderState& SceneState) const override;
+	virtual UInteractiveTool* BuildTool(const FToolBuilderState& SceneState) const override;
+};
+
+/**
+ * Mesh Weld Edges Tool
+ */
+UCLASS()
+class MESHMODELINGTOOLS_API UWeldMeshEdgesTool : public USingleSelectionTool
+{
+	GENERATED_BODY()
+
+public:
+	UWeldMeshEdgesTool();
+
+	virtual void Setup() override;
+	virtual void Shutdown(EToolShutdownType ShutdownType) override;
+
+	virtual void Render(IToolsContextRenderAPI* RenderAPI) override;
+
+	virtual bool HasCancel() const override { return true; }
+	virtual bool HasAccept() const override;
+	virtual bool CanAccept() const override;
+
+#if WITH_EDITOR
+	virtual void PostEditChangeProperty(FPropertyChangedEvent& PropertyChangedEvent) override;
+#endif
+protected:
+	// need to update bResultValid if these are modified, so we don't publicly expose them. 
+	// @todo setters/getters for these
+
+	/** Target edge length */
+	UPROPERTY(EditAnywhere, Category = Options, meta = (UIMin = "0.000001", UIMax = "0.01", ClampMin = "0.00000001", ClampMax = "1000.0"))
+	float Tolerance;
+
+	/** Enable edge flips */
+	UPROPERTY(EditAnywhere, Category = Advanced)
+	bool bOnlyUnique;
+
+protected:
+	USimpleDynamicMeshComponent* DynamicMeshComponent;
+	FDynamicMesh3 OriginalMesh;
+	bool bResultValid;
+	void UpdateResult();
+};
diff --git a/Engine/Plugins/Experimental/MeshModelingToolset/Source/MeshModelingToolsEditorOnly/MeshModelingToolsEditorOnly.Build.cs b/Engine/Plugins/Experimental/MeshModelingToolset/Source/MeshModelingToolsEditorOnly/MeshModelingToolsEditorOnly.Build.cs
new file mode 100644
--- /dev/null
+++ b/Engine/Plugins/Experimental/MeshModelingToolset/Source/MeshModelingToolsEditorOnly/MeshModelingToolsEditorOnly.Build.cs
@@ -0,0 +1,75 @@
+// Copyright 1998-2019 Epic Games, Inc. All Rights Reserved.
+
+using UnrealBuildTool;
+
+public class MeshModelingToolsEditorOnly : ModuleRules
+{
+	public MeshModelingToolsEditorOnly(ReadOnlyTargetRules Target) : base(Target)
+	{
+		PCHUsage = ModuleRules.PCHUsageMode.UseExplicitOrSharedPCHs;
+		
+		PublicIncludePaths.AddRange(
+			new string[] {
+				// ... add public include paths required here ...
+			}
+			);
+				
+		
+		PrivateIncludePaths.AddRange(
+			new string[] {
+				// ... add other private include paths required here ...
+			}
+			);
+			
+		
+		PublicDependencyModuleNames.AddRange(
+			new string[]
+			{
+				"Core",
+                "InteractiveToolsFramework",
+                "MeshDescription",
+				"GeometricObjects",
+				"DynamicMesh",
+				"MeshConversion",
+				"MeshModelingTools",
+                "MeshSolverUtilities",
+				"ModelingComponents",
+				"ModelingOperators",
+				"ModelingOperatorsEditorOnly",
+                "ProxyLODMeshReduction", // currently required to be public due to IVoxelBasedCSG API
+				// ... add other public dependencies that you statically link with here ...
+			}
+            );
+			
+		
+		PrivateDependencyModuleNames.AddRange(
+			new string[]
+			{
+				"CoreUObject",
+				"Engine",
+                "RenderCore",
+                "RHI",
+				"InputCore",
+
+				"MeshUtilities",    // temp for saving mesh asset
+				"UnrealEd",
+				"MeshBuilder",
+                "MeshDescriptionOperations",
+				"MeshUtilitiesCommon",  
+				"MeshReductionInterface", // for UE4 standard simplification 
+                "ProxyLODMeshReduction", // for mesh merging voxel-based csg
+				//"Slate",
+				//"SlateCore",
+				// ... add private dependencies that you statically link with here ...	
+			}
+			);
+		
+		
+		DynamicallyLoadedModuleNames.AddRange(
+			new string[]
+			{
+				// ... add any modules that your module loads dynamically here ...
+			}
+			);
+	}
+}
diff --git a/Engine/Plugins/Experimental/MeshModelingToolset/Source/MeshModelingToolsEditorOnly/Private/MergeMeshesTool.cpp b/Engine/Plugins/Experimental/MeshModelingToolset/Source/MeshModelingToolsEditorOnly/Private/MergeMeshesTool.cpp
new file mode 100644
--- /dev/null
+++ b/Engine/Plugins/Experimental/MeshModelingToolset/Source/MeshModelingToolsEditorOnly/Private/MergeMeshesTool.cpp
@@ -0,0 +1,198 @@
+// Copyright 1998-2019 Epic Games, Inc. All Rights Reserved.
+
+#include "MergeMeshesTool.h"
+#include "InteractiveToolManager.h"
+#include "ToolBuilderUtil.h"
+#include "ToolSetupUtil.h"
+#include "AssetGenerationUtil.h"
+#include "Selection/ToolSelectionUtil.h"
+
+#include "DynamicMesh3.h"
+#include "MeshDescriptionToDynamicMesh.h"
+
+#include "CompositionOps/VoxelMergeMeshesOp.h"
+
+#define LOCTEXT_NAMESPACE "UMergeMeshesTool"
+
+/*
+ * ToolBuilder
+ */
+
+
+bool UMergeMeshesToolBuilder::CanBuildTool(const FToolBuilderState& SceneState) const
+{
+	bool bHasBuildAPI = (this->AssetAPI != nullptr);
+	bool bHasComponents = ToolBuilderUtil::CountComponents(SceneState, CanMakeComponentTarget) > 1;
+	return ( bHasBuildAPI && bHasComponents );
+}
+
+UInteractiveTool* UMergeMeshesToolBuilder::BuildTool(const FToolBuilderState& SceneState) const
+{
+	UMergeMeshesTool* NewTool = NewObject<UMergeMeshesTool>(SceneState.ToolManager);
+
+	TArray<UActorComponent*> Components = ToolBuilderUtil::FindAllComponents(SceneState, CanMakeComponentTarget);
+	check(Components.Num() > 0);
+
+	TArray<TUniquePtr<FPrimitiveComponentTarget>> ComponentTargets;
+	for (UActorComponent* ActorComponent : Components)
+	{
+		auto* MeshComponent = Cast<UPrimitiveComponent>(ActorComponent);
+		if ( MeshComponent )
+		{
+			ComponentTargets.Add(MakeComponentTarget(MeshComponent));
+		}
+	}
+
+	NewTool->SetSelection(MoveTemp(ComponentTargets));
+	NewTool->SetWorld(SceneState.World);
+	NewTool->SetAssetAPI(AssetAPI);
+
+	return NewTool;
+}
+
+
+
+
+/*
+ * Tool
+ */
+
+UMergeMeshesTool::UMergeMeshesTool()
+{
+}
+
+void UMergeMeshesTool::SetWorld(UWorld* World)
+{
+	this->TargetWorld = World;
+}
+
+void UMergeMeshesTool::SetAssetAPI(IToolsContextAssetAPI* AssetAPIIn)
+{
+	this->AssetAPI = AssetAPIIn;
+}
+
+
+
+void UMergeMeshesTool::Setup()
+{
+	UInteractiveTool::Setup();
+
+	MergeProps = NewObject<UMergeMeshesToolProperties>();
+	AddToolPropertySource(MergeProps);
+
+	// Hide the source meshes
+	for (auto& ComponentTarget : ComponentTargets)
+	{
+		ComponentTarget->SetOwnerVisibility(false);
+	}
+
+	// save transformed version of input meshes (maybe this could happen in the Operator?)
+	CacheInputMeshes();
+
+	// initialize the PreviewMesh+BackgroundCompute object
+	Preview = NewObject<UMeshOpPreviewWithBackgroundCompute>(this, "Preview");
+	Preview->Setup(this->TargetWorld, this);
+	Preview->ConfigureMaterials(
+		// TODO: This is is very likely a bug waiting to happen
+		ToolSetupUtil::GetDefaultMaterial(GetToolManager(), ComponentTargets[0]->GetMaterial(0)),
+		ToolSetupUtil::GetDefaultWorkingMaterial(GetToolManager())
+	);
+	Preview->SetVisibility(false);
+	Preview->InvalidateResult();    // start compute
+}
+
+
+void UMergeMeshesTool::Shutdown(EToolShutdownType ShutdownType)
+{
+	TUniquePtr<FDynamicMeshOpResult> Result = Preview->Shutdown();
+	if (ShutdownType == EToolShutdownType::Accept)
+	{
+		GenerateAsset(Result);
+	}
+
+	// Restore (unhide) the source meshes
+	for (auto& ComponentTarget : ComponentTargets)
+	{
+		ComponentTarget->SetOwnerVisibility(true);
+	}
+}
+
+
+
+void UMergeMeshesTool::Tick(float DeltaTime)
+{
+	Preview->Tick(DeltaTime);
+}
+
+
+void UMergeMeshesTool::Render(IToolsContextRenderAPI* RenderAPI)
+{
+}
+
+
+bool UMergeMeshesTool::HasAccept() const
+{
+	return true;
+}
+
+bool UMergeMeshesTool::CanAccept() const
+{
+	return Preview->HaveValidResult();
+}
+
+void UMergeMeshesTool::OnPropertyModified(UObject* PropertySet, UProperty* Property)
+{
+	Preview->InvalidateResult();
+}
+
+
+
+
+TSharedPtr<FDynamicMeshOperator> UMergeMeshesTool::MakeNewOperator()
+{
+	TSharedPtr<FVoxelMergeMeshesOp> MergeOp = MakeShared<FVoxelMergeMeshesOp>();
+	MergeOp->VoxelCount = MergeProps->VoxelCount;
+	MergeOp->AdaptivityD = MergeProps->Adaptivity;
+	MergeOp->IsoSurfaceD = MergeProps->IsoSurface;
+	MergeOp->InputMeshArray = InputMeshes;
+	return MergeOp;
+}
+
+
+
+void UMergeMeshesTool::CacheInputMeshes()
+{
+	InputMeshes = MakeShared<TArray<IVoxelBasedCSG::FPlacedMesh>>();
+
+	// Package the selected meshes and transforms for consumption by the CSGTool
+	for (auto& ComponentTarget : ComponentTargets)
+	{
+		IVoxelBasedCSG::FPlacedMesh PlacedMesh;
+		PlacedMesh.Mesh = ComponentTarget->GetMesh();
+		PlacedMesh.Transform = ComponentTarget->GetWorldTransform();
+		InputMeshes->Add(PlacedMesh);
+	}
+}
+
+
+
+void UMergeMeshesTool::GenerateAsset(const TUniquePtr<FDynamicMeshOpResult>& Result)
+{
+	check(Result->Mesh.Get() != nullptr);
+
+	GetToolManager()->BeginUndoTransaction(LOCTEXT("MergeMeshes", "Merge Meshes"));
+
+	AActor* NewActor = AssetGenerationUtil::GenerateStaticMeshActor(
+		AssetAPI, TargetWorld,
+		Result->Mesh.Get(), Result->Transform, TEXT("MergedMesh"),
+		AssetGenerationUtil::GetDefaultAutoGeneratedAssetPath());
+
+	// select newly-created object
+	ToolSelectionUtil::SetNewActorSelection(GetToolManager(), NewActor);
+
+	GetToolManager()->EndUndoTransaction();
+}
+
+
+
+#undef LOCTEXT_NAMESPACE
diff --git a/Engine/Plugins/Experimental/MeshModelingToolset/Source/MeshModelingToolsEditorOnly/Private/ParameterizeMeshTool.cpp b/Engine/Plugins/Experimental/MeshModelingToolset/Source/MeshModelingToolsEditorOnly/Private/ParameterizeMeshTool.cpp
new file mode 100644
--- /dev/null
+++ b/Engine/Plugins/Experimental/MeshModelingToolset/Source/MeshModelingToolsEditorOnly/Private/ParameterizeMeshTool.cpp
@@ -0,0 +1,245 @@
+// Copyright 1998-2019 Epic Games, Inc. All Rights Reserved.
+
+#include "ParameterizeMeshTool.h"
+#include "InteractiveToolManager.h"
+#include "ToolBuilderUtil.h"
+
+#include "DynamicMesh3.h"
+#include "DynamicMeshToMeshDescription.h"
+
+
+
+#include "SimpleDynamicMeshComponent.h"
+#include "Materials/MaterialInstanceDynamic.h"
+#include "MeshDescription.h"
+#include "MeshDescriptionToDynamicMesh.h"
+
+#include "ParameterizationOps/ParameterizeMeshOp.h"
+#include "ToolSetupUtil.h"
+
+#define LOCTEXT_NAMESPACE "UParameterizeMeshTool"
+
+
+DEFINE_LOG_CATEGORY_STATIC(LogParameterizeMeshTool, Log, All);
+
+/*
+ * ToolBuilder
+ */
+
+bool UParameterizeMeshToolBuilder::CanBuildTool(const FToolBuilderState& SceneState) const
+{
+	return ToolBuilderUtil::CountComponents(SceneState, CanMakeComponentTarget) == 1;
+}
+
+UInteractiveTool* UParameterizeMeshToolBuilder::BuildTool(const FToolBuilderState& SceneState) const
+{
+	UParameterizeMeshTool* NewTool = NewObject<UParameterizeMeshTool>(SceneState.ToolManager);
+
+	UActorComponent* ActorComponent = ToolBuilderUtil::FindFirstComponent(SceneState, CanMakeComponentTarget);
+	auto* MeshComponent = Cast<UPrimitiveComponent>(ActorComponent);
+	check(MeshComponent != nullptr);
+	NewTool->SetSelection(MakeComponentTarget(MeshComponent));
+	NewTool->SetWorld(SceneState.World);
+	NewTool->SetAssetAPI(AssetAPI);
+
+	return NewTool;
+}
+
+
+/*
+ * Tool
+ */
+
+UParameterizeMeshTool::UParameterizeMeshTool()
+{
+	ChartStretch = 0.11f;
+}
+
+void UParameterizeMeshTool::SetWorld(UWorld* World)
+{
+	this->TargetWorld = World;
+}
+
+void UParameterizeMeshTool::SetAssetAPI(IToolsContextAssetAPI* AssetAPIIn)
+{
+	this->AssetAPI = AssetAPIIn;
+}
+
+void UParameterizeMeshTool::Setup()
+{
+	UInteractiveTool::Setup();
+
+	// Deep copy of input mesh to be shared with the UV generation tool.
+	InputMesh = MakeShared<FMeshDescription>(*ComponentTarget->GetMesh());
+
+
+	// Set up check material
+
+	UMaterial* CheckerMaterialBase = LoadObject<UMaterial>(nullptr, TEXT("/MeshModelingToolset/Materials/CheckerMaterial"));
+	CheckerMaterial = UMaterialInstanceDynamic::Create(CheckerMaterialBase, NULL);
+	CheckerMaterial->SetScalarParameterValue("Density", CheckerDensity);
+
+	// Copy existing material if there is one
+	
+	DefaultMaterial = ComponentTarget->GetMaterial(0);
+	if (DefaultMaterial == nullptr)
+	{
+		DefaultMaterial = LoadObject<UMaterial>(nullptr, TEXT("/Engine/EngineMaterials/DefaultMaterial"));  
+	}
+
+	
+
+	// hide input StaticMeshComponent
+	ComponentTarget->SetOwnerVisibility(false);
+
+	// initialize our properties
+	ToolPropertyObjects.Add(this);
+
+	
+	// Construct the preview object and set the material on it
+	Preview = NewObject<UMeshOpPreviewWithBackgroundCompute>(this, "Preview");
+	Preview->Setup(this->TargetWorld, this);
+
+	// Initialize the preview mesh with a copy of the source mesh.
+	{
+		FDynamicMesh3 Mesh;
+		FMeshDescriptionToDynamicMesh Converter;
+		Converter.bPrintDebugMessages = false;
+		Converter.Convert(InputMesh.Get(), Mesh);
+
+		Preview->PreviewMesh->UpdatePreview(&Mesh);
+		Preview->PreviewMesh->SetTransform(ComponentTarget->GetWorldTransform());
+	}
+	Preview->ConfigureMaterials(
+		ToolSetupUtil::GetDefaultMaterial(GetToolManager(), DefaultMaterial),
+		ToolSetupUtil::GetDefaultWorkingMaterial(GetToolManager())
+	);
+
+	// Choose the material for display
+	if (MaterialMode == EParameterizeMeshMaterialMode::Checkerboard)
+	{
+		Preview->ConfigureMaterials(
+			ToolSetupUtil::GetDefaultMaterial(GetToolManager(), CheckerMaterial),
+			ToolSetupUtil::GetDefaultWorkingMaterial(GetToolManager())
+		);
+	}
+
+
+	Preview->SetVisibility(true);
+	Preview->InvalidateResult();    // start compute
+}
+
+
+void UParameterizeMeshTool::Shutdown(EToolShutdownType ShutdownType)
+{
+
+	TUniquePtr<FDynamicMeshOpResult> Result = Preview->Shutdown();
+	if (ShutdownType == EToolShutdownType::Accept)
+	{
+		FDynamicMesh3* DynamicMeshResult = Result->Mesh.Get();
+		check(DynamicMeshResult != nullptr);
+		GetToolManager()->BeginUndoTransaction(LOCTEXT("ParameterizeMesh", "Parameterize Mesh"));
+
+		ComponentTarget->CommitMesh([DynamicMeshResult](FMeshDescription* MeshDescription)
+		{
+			FDynamicMeshToMeshDescription Converter;
+			Converter.Convert(DynamicMeshResult, *MeshDescription);
+		});
+		GetToolManager()->EndUndoTransaction();
+	}
+
+	// Restore (unhide) the source meshes
+	ComponentTarget->SetOwnerVisibility(true);
+
+}
+
+void UParameterizeMeshTool::Tick(float DeltaTime)
+{
+	Preview->Tick(DeltaTime);
+}
+
+void UParameterizeMeshTool::Render(IToolsContextRenderAPI* RenderAPI)
+{
+}
+
+bool UParameterizeMeshTool::HasAccept() const
+{
+	return true;
+}
+
+bool UParameterizeMeshTool::CanAccept() const
+{
+	return Preview->HaveValidResult();
+}
+
+void UParameterizeMeshTool::PostEditChangeProperty(FPropertyChangedEvent& PropertyChangedEvent)
+{
+	UProperty* PropertyThatChanged = PropertyChangedEvent.Property;
+
+	if (PropertyThatChanged->GetFName() == GET_MEMBER_NAME_CHECKED(UParameterizeMeshTool, CheckerDensity))
+	{
+		CheckerMaterial->SetScalarParameterValue("Density", CheckerDensity);
+	}
+	else if (PropertyThatChanged->GetFName() == GET_MEMBER_NAME_CHECKED(UParameterizeMeshTool, MaterialMode))
+	{
+		if (MaterialMode == EParameterizeMeshMaterialMode::Checkerboard && CheckerMaterial != nullptr)
+		{
+			Preview->ConfigureMaterials(
+				ToolSetupUtil::GetDefaultMaterial(GetToolManager(), CheckerMaterial),
+				ToolSetupUtil::GetDefaultWorkingMaterial(GetToolManager())
+			);
+		}
+		if (MaterialMode == EParameterizeMeshMaterialMode::Override && OverrideMaterial != nullptr)
+		{
+			Preview->ConfigureMaterials(
+				ToolSetupUtil::GetDefaultMaterial(GetToolManager(), OverrideMaterial),
+				ToolSetupUtil::GetDefaultWorkingMaterial(GetToolManager())
+			);
+		}
+		if (MaterialMode == EParameterizeMeshMaterialMode::Default && DefaultMaterial != nullptr)
+		{
+		
+			Preview->ConfigureMaterials(
+				ToolSetupUtil::GetDefaultMaterial(GetToolManager(), DefaultMaterial),
+				ToolSetupUtil::GetDefaultWorkingMaterial(GetToolManager())
+			);
+		}
+
+		
+
+	}
+	else if (PropertyThatChanged->GetFName() == GET_MEMBER_NAME_CHECKED(UParameterizeMeshTool, OverrideMaterial))
+	{
+		if (MaterialMode == EParameterizeMeshMaterialMode::Override && OverrideMaterial != nullptr)
+		{
+			
+			Preview->ConfigureMaterials(
+				ToolSetupUtil::GetDefaultMaterial(GetToolManager(), OverrideMaterial),
+				ToolSetupUtil::GetDefaultWorkingMaterial(GetToolManager())
+			);
+		}
+	}
+	else
+	{
+		// One of the UV generation properties must have changed.  Dirty the result to force a recompute
+		Preview->InvalidateResult();
+	}
+}
+
+TSharedPtr<FDynamicMeshOperator> UParameterizeMeshTool::MakeNewOperator()
+{
+	TSharedPtr<FParameterizeMeshOp> ParamertizeMeshOp = MakeShared<FParameterizeMeshOp>();
+	ParamertizeMeshOp->Stretch   = ChartStretch;
+	ParamertizeMeshOp->NumCharts = 0;
+	ParamertizeMeshOp->InputMesh = InputMesh;
+
+	const FTransform XForm = ComponentTarget->GetWorldTransform();
+	FTransform3d XForm3d(XForm);
+	ParamertizeMeshOp->SetTransform(XForm3d);
+
+	return ParamertizeMeshOp;
+}
+
+
+
+#undef LOCTEXT_NAMESPACE
diff --git a/Engine/Plugins/Experimental/MeshModelingToolset/Source/MeshModelingToolsEditorOnly/Private/PolygonOnMeshTool.cpp b/Engine/Plugins/Experimental/MeshModelingToolset/Source/MeshModelingToolsEditorOnly/Private/PolygonOnMeshTool.cpp
new file mode 100644
--- /dev/null
+++ b/Engine/Plugins/Experimental/MeshModelingToolset/Source/MeshModelingToolsEditorOnly/Private/PolygonOnMeshTool.cpp
@@ -0,0 +1,347 @@
+// Copyright 1998-2019 Epic Games, Inc. All Rights Reserved.
+
+#include "PolygonOnMeshTool.h"
+#include "InteractiveToolManager.h"
+#include "ToolBuilderUtil.h"
+
+#include "ToolSetupUtil.h"
+
+#include "DynamicMesh3.h"
+#include "BaseBehaviors/MultiClickSequenceInputBehavior.h"
+#include "Selection/SelectClickedAction.h"
+
+#include "MeshDescriptionToDynamicMesh.h"
+#include "DynamicMeshToMeshDescription.h"
+
+#include "InteractiveGizmoManager.h"
+
+#include "PositionPlaneGizmo.h"
+#include "AssetGenerationUtil.h"
+
+
+#define LOCTEXT_NAMESPACE "UPolygonOnMeshTool"
+
+
+/*
+ * ToolBuilder
+ */
+
+
+bool UPolygonOnMeshToolBuilder::CanBuildTool(const FToolBuilderState& SceneState) const
+{
+	return ToolBuilderUtil::CountComponents(SceneState, CanMakeComponentTarget) == 1;
+}
+
+UInteractiveTool* UPolygonOnMeshToolBuilder::BuildTool(const FToolBuilderState& SceneState) const
+{
+	UPolygonOnMeshTool* NewTool = NewObject<UPolygonOnMeshTool>(SceneState.ToolManager);
+
+	UActorComponent* ActorComponent = ToolBuilderUtil::FindFirstComponent(SceneState, CanMakeComponentTarget);
+	auto* MeshComponent = Cast<UPrimitiveComponent>(ActorComponent);
+	check(MeshComponent != nullptr);
+	NewTool->SetSelection(MakeComponentTarget(MeshComponent));
+
+	NewTool->SetWorld(SceneState.World);
+	NewTool->SetAssetAPI(AssetAPI);
+	return NewTool;
+}
+
+
+
+
+/*
+ * Tool
+ */
+
+UPolygonOnMeshToolProperties::UPolygonOnMeshToolProperties()
+{
+	PolygonOperation = EEmbeddedPolygonOpMethod::CutAndFill;
+	PolygonScale = 10;
+	// ExtrudeDistance = 10;
+	bDiscardAttributes = false;
+}
+
+UPolygonOnMeshAdvancedProperties::UPolygonOnMeshAdvancedProperties()
+{
+}
+
+
+UPolygonOnMeshTool::UPolygonOnMeshTool()
+{
+	EmbedPolygonOrigin = FVector::ZeroVector;
+	EmbedPolygonOrientation = FQuat::Identity;
+}
+
+void UPolygonOnMeshTool::SetWorld(UWorld* World)
+{
+	this->TargetWorld = World;
+}
+
+void UPolygonOnMeshTool::Setup()
+{
+	UInteractiveTool::Setup();
+
+	// hide input StaticMeshComponent
+	ComponentTarget->SetOwnerVisibility(false);
+
+	// click to set plane behavior
+	FSelectClickedAction* SetPlaneAction = new FSelectClickedAction();
+	SetPlaneAction->World = this->TargetWorld;
+	SetPlaneAction->OnClickedPositionFunc = [this](const FHitResult& Hit)
+	{
+		SetPlaneFromWorldPos(Hit.ImpactPoint, Hit.ImpactNormal);
+		for (UMeshOpPreviewWithBackgroundCompute* Preview : Previews)
+		{
+			Preview->InvalidateResult();
+		}
+	};
+	SetPointInWorldConnector = SetPlaneAction;
+
+	USingleClickInputBehavior* ClickToSetPlaneBehavior = NewObject<USingleClickInputBehavior>();
+	ClickToSetPlaneBehavior->ModifierCheckFunc = FInputDeviceState::IsCtrlKeyDown;
+	ClickToSetPlaneBehavior->Initialize(SetPointInWorldConnector);
+	AddInputBehavior(ClickToSetPlaneBehavior);
+
+
+	UPositionPlaneGizmoBuilder* PositionPlaneGizmoBuilder = NewObject<UPositionPlaneGizmoBuilder>();
+	UInteractiveGizmoManager* GizmoManager = GetToolManager()->GetPairedGizmoManager();
+	GizmoManager->RegisterGizmoType(TEXT("PolygonPlaneGizmo"), PositionPlaneGizmoBuilder);
+	PositionPlaneGizmo = GizmoManager->CreateGizmo(TEXT("PolygonPlaneGizmo"), TEXT("PolygonPlaneGizmo"));
+	Cast<UPositionPlaneGizmo>(PositionPlaneGizmo)->OnPositionUpdatedFunc = [this](const FFrame3d& WorldFrame)
+	{
+		UpdatePlaneFromGizmo(WorldFrame);
+	};
+
+	BasicProperties = NewObject<UPolygonOnMeshToolProperties>(this, TEXT("Polygon On Mesh Settings"));
+	AdvancedProperties = NewObject<UPolygonOnMeshAdvancedProperties>(this, TEXT("Advanced Settings"));
+
+	// initialize our properties
+	AddToolPropertySource(BasicProperties);
+	AddToolPropertySource(AdvancedProperties);
+
+
+	// initialize the PreviewMesh+BackgroundCompute object
+	UpdateNumPreviews();
+
+	FVector DefaultOrigin, Extents;
+	ComponentTarget->GetOwnerActor()->GetActorBounds(false, DefaultOrigin, Extents);
+	SetPlaneFromWorldPos(DefaultOrigin, FVector::UpVector);
+
+
+
+	// Convert input mesh description to dynamic mesh
+	OriginalDynamicMesh = MakeShared<FDynamicMesh3>();
+	FMeshDescriptionToDynamicMesh Converter;
+	Converter.bPrintDebugMessages = true;
+	Converter.Convert(ComponentTarget->GetMesh(), *OriginalDynamicMesh);
+	// TODO: consider adding an AABB tree construction here?  tradeoff vs doing a raycast against full every time a param change happens ...
+
+	for (UMeshOpPreviewWithBackgroundCompute* Preview : Previews)
+	{
+		Preview->InvalidateResult();
+	}
+}
+
+
+void UPolygonOnMeshTool::UpdateNumPreviews()
+{
+	int32 CurrentNumPreview = Previews.Num();
+	int32 TargetNumPreview = 1; // currently only have single preview use case; TODO consider cases where multiple meshes are generated with separate preview objects
+	if (TargetNumPreview < CurrentNumPreview)
+	{
+		for (int32 PreviewIdx = CurrentNumPreview - 1; PreviewIdx >= TargetNumPreview; PreviewIdx--)
+		{
+			Previews[PreviewIdx]->Cancel();
+		}
+		Previews.SetNum(TargetNumPreview);
+	}
+	else
+	{
+		for (int32 PreviewIdx = CurrentNumPreview; PreviewIdx < TargetNumPreview; PreviewIdx++)
+		{
+			UPolygonOnMeshOperatorFactory *OpFactory = NewObject<UPolygonOnMeshOperatorFactory>();
+			OpFactory->Tool = this;
+			// TODO: give the OpFactory more info about what specifically it is doing differently vs the other previews
+			UMeshOpPreviewWithBackgroundCompute* Preview = Previews.Add_GetRef(NewObject<UMeshOpPreviewWithBackgroundCompute>(OpFactory, "Preview"));
+			Preview->Setup(this->TargetWorld, OpFactory);
+			Preview->ConfigureMaterials(
+				ToolSetupUtil::GetDefaultMaterial(GetToolManager(), ComponentTarget->GetMaterial(0)),
+				ToolSetupUtil::GetDefaultWorkingMaterial(GetToolManager())
+			);
+			Preview->SetVisibility(true);
+		}
+	}
+}
+
+
+void UPolygonOnMeshTool::Shutdown(EToolShutdownType ShutdownType)
+{
+	// Restore (unhide) the source meshes
+	ComponentTarget->SetOwnerVisibility(true);
+
+	TArray<TUniquePtr<FDynamicMeshOpResult>> Results;
+	for (UMeshOpPreviewWithBackgroundCompute* Preview : Previews)
+	{
+		Results.Emplace(Preview->Shutdown());
+	}
+	if (ShutdownType == EToolShutdownType::Accept)
+	{
+		GenerateAsset(Results);
+	}
+
+	if (SetPointInWorldConnector != nullptr)
+	{
+		delete SetPointInWorldConnector;
+	}
+	UInteractiveGizmoManager* GizmoManager = GetToolManager()->GetPairedGizmoManager();
+	GizmoManager->DestroyGizmo(PositionPlaneGizmo);
+	PositionPlaneGizmo = nullptr;
+	GizmoManager->DeregisterGizmoType(TEXT("PolygonPlaneGizmo"));
+}
+
+void UPolygonOnMeshTool::SetAssetAPI(IToolsContextAssetAPI* AssetAPIIn)
+{
+	this->AssetAPI = AssetAPIIn;
+}
+
+TSharedPtr<FDynamicMeshOperator> UPolygonOnMeshOperatorFactory::MakeNewOperator()
+{
+	TSharedPtr<FEmbedPolygonsOp> EmbedOp = MakeShared<FEmbedPolygonsOp>();
+	EmbedOp->bDiscardAttributes = Tool->BasicProperties->bDiscardAttributes;
+	EmbedOp->Operation = Tool->BasicProperties->PolygonOperation;
+	EmbedOp->PolygonScale = Tool->BasicProperties->PolygonScale;
+	// EmbedOp->ExtrudeDistance = Tool->BasicProperties->ExtrudeDistance;
+	// TODO: also put the polygon info, etc, into the EmbedOp
+
+	FTransform LocalToWorld = Tool->ComponentTarget->GetWorldTransform();
+	FTransform WorldToLocal = LocalToWorld.Inverse();
+	FVector LocalOrigin = WorldToLocal.TransformPosition(Tool->EmbedPolygonOrigin);
+	FVector WorldNormal = Tool->EmbedPolygonOrientation.GetAxisZ();
+	FVector LocalNormal = WorldToLocal.TransformVectorNoScale(WorldNormal);
+	EmbedOp->LocalPlaneOrigin = LocalOrigin;
+	EmbedOp->LocalPlaneNormal = LocalNormal;
+	EmbedOp->OriginalMesh = Tool->OriginalDynamicMesh;
+	
+	EmbedOp->SetTransform(LocalToWorld);
+
+	return EmbedOp;
+}
+
+
+
+void UPolygonOnMeshTool::Render(IToolsContextRenderAPI* RenderAPI)
+{
+}
+
+void UPolygonOnMeshTool::Tick(float DeltaTime)
+{
+	for (UMeshOpPreviewWithBackgroundCompute* Preview : Previews)
+	{
+		Preview->Tick(DeltaTime);
+	}
+}
+
+
+
+void UPolygonOnMeshTool::PostEditChangeProperty(FPropertyChangedEvent& PropertyChangedEvent)
+{
+	UpdateNumPreviews();
+	for (UMeshOpPreviewWithBackgroundCompute* Preview : Previews)
+	{
+		Preview->InvalidateResult();
+	}
+}
+
+void UPolygonOnMeshTool::OnPropertyModified(UObject* PropertySet, UProperty* Property)
+{
+	UpdateNumPreviews();
+	for (UMeshOpPreviewWithBackgroundCompute* Preview : Previews)
+	{
+		Preview->InvalidateResult();
+	}
+}
+
+
+
+
+void UPolygonOnMeshTool::SetPlaneFromWorldPos(const FVector& Position, const FVector& Normal)
+{
+	EmbedPolygonOrigin = Position;
+
+	FFrame3f PlaneFrame(Position, Normal);
+	EmbedPolygonOrientation = PlaneFrame.Rotation;
+
+	UPositionPlaneGizmo* Gizmo = Cast<UPositionPlaneGizmo>(PositionPlaneGizmo);
+	Gizmo->ExternalUpdatePosition(EmbedPolygonOrigin, EmbedPolygonOrientation, false);
+}
+
+
+void UPolygonOnMeshTool::UpdatePlaneFromGizmo(const FFrame3d& WorldPosition)
+{
+	EmbedPolygonOrientation = WorldPosition.Rotation;
+	EmbedPolygonOrigin = (FVector)WorldPosition.Origin;
+	for (UMeshOpPreviewWithBackgroundCompute* Preview : Previews)
+	{
+		Preview->InvalidateResult();
+	}
+}
+
+
+bool UPolygonOnMeshTool::HasAccept() const
+{
+	return true;
+}
+
+bool UPolygonOnMeshTool::CanAccept() const
+{
+	for (UMeshOpPreviewWithBackgroundCompute* Preview : Previews)
+	{
+		if (!Preview->HaveValidResult())
+		{
+			return false;
+		}
+	}
+	return true;
+}
+
+
+void UPolygonOnMeshTool::GenerateAsset(const TArray<TUniquePtr<FDynamicMeshOpResult>>& Results)
+{
+	GetToolManager()->BeginUndoTransaction(LOCTEXT("PolygonOnMeshToolTransactionName", "Polygon On Mesh Tool"));
+	
+
+	// first preview always replaces, and any subsequent previews can add new actors
+	// TODO: options to support other choices re what should be a new actor
+	check(Results.Num() > 0);
+	check(Results[0]->Mesh.Get() != nullptr);
+	ComponentTarget->CommitMesh([&Results](FMeshDescription* MeshDescription)
+	{
+		FDynamicMeshToMeshDescription Converter;
+		Converter.Convert(Results[0]->Mesh.Get(), *MeshDescription);
+	});
+
+	// currently nothing generates more than the single result for this tool
+	if (Results.Num() > 1)
+	{
+		unimplemented();
+		//FSelectedOjectsChangeList NewSelection;
+		//NewSelection.ModificationType = ESelectedObjectsModificationType::Replace;
+		//NewSelection.Actors.Add(ComponentTarget->GetOwnerActor());
+
+		//check(Results[1]->Mesh.Get() != nullptr);
+		//// TODO: copy over material?
+		//AActor* NewActor = AssetGenerationUtil::GenerateStaticMeshActor(
+		//	AssetAPI, TargetWorld,
+		//	Results[1]->Mesh.Get(), Results[1]->Transform, TEXT("Additional Polygon On Mesh Result"),
+		//	AssetGenerationUtil::GetDefaultAutoGeneratedAssetPath());
+		//NewSelection.Actors.Add(NewActor);
+		//GetToolManager()->RequestSelectionChange(NewSelection);
+	}
+
+	GetToolManager()->EndUndoTransaction();
+}
+
+
+
+
+#undef LOCTEXT_NAMESPACE
diff --git a/Engine/Plugins/Experimental/MeshModelingToolset/Source/MeshModelingToolsEditorOnly/Private/SimplifyMeshTool.cpp b/Engine/Plugins/Experimental/MeshModelingToolset/Source/MeshModelingToolsEditorOnly/Private/SimplifyMeshTool.cpp
new file mode 100644
--- /dev/null
+++ b/Engine/Plugins/Experimental/MeshModelingToolset/Source/MeshModelingToolsEditorOnly/Private/SimplifyMeshTool.cpp
@@ -0,0 +1,293 @@
+// Copyright 1998-2019 Epic Games, Inc. All Rights Reserved.
+
+#include "SimplifyMeshTool.h"
+#include "InteractiveToolManager.h"
+#include "ToolBuilderUtil.h"
+
+#include "DynamicMesh3.h"
+#include "DynamicMeshAttributeSet.h"
+#include "MeshSimplification.h"
+#include "MeshConstraintsUtil.h"
+#include "ProjectionTargets.h"
+
+#include "SimpleDynamicMeshComponent.h"
+
+#include "SceneManagement.h" // for FPrimitiveDrawInterface
+
+//#include "ProfilingDebugging/ScopedTimers.h" // enable this to use the timer.
+#include "MeshDescription.h"
+#include "MeshDescriptionOperations.h"
+#include "OverlappingCorners.h"
+#include "IMeshReductionManagerModule.h"
+#include "IMeshReductionInterfaces.h"
+#include "Modules/ModuleManager.h"
+
+#define LOCTEXT_NAMESPACE "USimplifyMeshTool"
+
+
+DEFINE_LOG_CATEGORY_STATIC(LogMeshSimplification, Log, All);
+
+/*
+ * ToolBuilder
+ */
+
+
+bool USimplifyMeshToolBuilder::CanBuildTool(const FToolBuilderState& SceneState) const
+{
+	return ToolBuilderUtil::CountComponents(SceneState, CanMakeComponentTarget) == 1;
+}
+
+UInteractiveTool* USimplifyMeshToolBuilder::BuildTool(const FToolBuilderState& SceneState) const
+{
+	USimplifyMeshTool* NewTool = NewObject<USimplifyMeshTool>(SceneState.ToolManager);
+
+	UActorComponent* ActorComponent = ToolBuilderUtil::FindFirstComponent(SceneState, CanMakeComponentTarget);
+	auto* MeshComponent = Cast<UPrimitiveComponent>(ActorComponent);
+	check(MeshComponent != nullptr);
+	NewTool->SetSelection(MakeComponentTarget(MeshComponent));
+
+	return NewTool;
+}
+
+
+
+/*
+ * Tool
+ */
+
+USimplifyMeshTool::USimplifyMeshTool()
+{
+	SimplifierType = ESimplifyType::QEM;
+	TargetMode = ESimplifyTargetType::Percentage;
+	TargetPercentage = 50;
+	TargetCount = 1000;
+	TargetEdgeLength = 5.0;
+	bReproject = false;
+	bPreventNormalFlips = true;
+	bDiscardAttributes = false;
+}
+
+void USimplifyMeshTool::Setup()
+{
+	UInteractiveTool::Setup();
+
+	// create dynamic mesh component to use for live preview
+	DynamicMeshComponent = NewObject<USimpleDynamicMeshComponent>(ComponentTarget->GetOwnerActor(), "DynamicMesh");
+	DynamicMeshComponent->SetupAttachment(ComponentTarget->GetOwnerActor()->GetRootComponent());
+	DynamicMeshComponent->RegisterComponent();
+	DynamicMeshComponent->SetWorldTransform(ComponentTarget->GetWorldTransform());
+
+	// copy material if there is one
+	auto Material = ComponentTarget->GetMaterial(0);
+	if (Material != nullptr)
+	{
+		DynamicMeshComponent->SetMaterial(0, Material);
+	}
+	DynamicMeshComponent->bExplicitShowWireframe = true;
+
+	DynamicMeshComponent->InitializeMesh(ComponentTarget->GetMesh());
+	OriginalMesh.Copy(*DynamicMeshComponent->GetMesh());
+	OriginalMeshSpatial.SetMesh(&OriginalMesh, true);
+
+	// hide input StaticMeshComponent
+	ComponentTarget->SetOwnerVisibility(false);
+
+	// initialize our properties
+	ToolPropertyObjects.Add(this);
+
+	MeshStatisticsProperties = NewObject<UMeshStatisticsProperties>(this);
+	AddToolPropertySource(MeshStatisticsProperties);
+	MeshStatisticsProperties->Update(*DynamicMeshComponent->GetMesh());
+
+	bResultValid = false;
+}
+
+
+void USimplifyMeshTool::Shutdown(EToolShutdownType ShutdownType)
+{
+	if (DynamicMeshComponent != nullptr)
+	{
+		//DynamicMeshComponent->OnMeshChanged.Remove(OnDynamicMeshComponentChangedHandle);
+
+		ComponentTarget->SetOwnerVisibility(true);
+
+		if (ShutdownType == EToolShutdownType::Accept)
+		{
+			// this block bakes the modified DynamicMeshComponent back into the StaticMeshComponent inside an undo transaction
+			GetToolManager()->BeginUndoTransaction(LOCTEXT("SimplifyMeshToolTransactionName", "Simplify Mesh"));
+			ComponentTarget->CommitMesh([=](FMeshDescription* MeshDescription)
+			{
+				DynamicMeshComponent->Bake(MeshDescription, true);
+			});
+			GetToolManager()->EndUndoTransaction();
+		}
+
+		DynamicMeshComponent->UnregisterComponent();
+		DynamicMeshComponent->DestroyComponent();
+		DynamicMeshComponent = nullptr;
+	}
+}
+
+
+void USimplifyMeshTool::Render(IToolsContextRenderAPI* RenderAPI)
+{
+	UpdateResult();
+
+
+	FPrimitiveDrawInterface* PDI = RenderAPI->GetPrimitiveDrawInterface();
+	FTransform Transform = ComponentTarget->GetWorldTransform(); //Actor->GetTransform();
+
+	FColor LineColor(255, 0, 0);
+	FDynamicMesh3* TargetMesh = DynamicMeshComponent->GetMesh();
+	if (TargetMesh->HasAttributes())
+	{
+		FDynamicMeshUVOverlay* UVOverlay = TargetMesh->Attributes()->PrimaryUV();
+		for (int eid : TargetMesh->EdgeIndicesItr()) 
+		{
+			if (UVOverlay->IsSeamEdge(eid)) 
+			{
+				FVector3d A, B;
+				TargetMesh->GetEdgeV(eid, A, B);
+				PDI->DrawLine(Transform.TransformPosition(A), Transform.TransformPosition(B),
+					LineColor, 0, 2.0, 1.0f, true);
+			}
+		}
+	}
+}
+
+
+void USimplifyMeshTool::PostEditChangeProperty(FPropertyChangedEvent& PropertyChangedEvent)
+{
+	UProperty* PropertyThatChanged = PropertyChangedEvent.Property;
+	bResultValid = false;
+}
+
+template <typename SimplificationType>
+void ComputeResult(FDynamicMesh3* TargetMesh, const bool bReproject, int OriginalTriCount, FDynamicMesh3& OriginalMesh, FDynamicMeshAABBTree3& OriginalMeshSpatial, const ESimplifyTargetType TargetMode, const float TargetPercentage, const int TargetCount, const float TargetEdgeLength)
+{
+	SimplificationType Reducer(TargetMesh);
+
+	Reducer.ProjectionMode = (bReproject) ? 
+		SimplificationType::ETargetProjectionMode::AfterRefinement : SimplificationType::ETargetProjectionMode::NoProjection;
+
+	Reducer.DEBUG_CHECK_LEVEL = 0;
+
+	FMeshConstraints constraints;
+	FMeshConstraintsUtil::ConstrainAllSeams(constraints, *TargetMesh, true, false);
+	//FMeshConstraintsUtil::ConstrainAllSeamJunctions(constraints, *TargetMesh, true, false);
+	Reducer.SetExternalConstraints(&constraints);
+
+	FMeshProjectionTarget ProjTarget(&OriginalMesh, &OriginalMeshSpatial);
+	Reducer.SetProjectionTarget(&ProjTarget);
+
+	if (TargetMode == ESimplifyTargetType::Percentage)
+	{
+		double Ratio = (double)TargetPercentage / 100.0;
+		int UseTarget = FMath::Max(4, (int)(Ratio * (double)OriginalTriCount));
+		Reducer.SimplifyToTriangleCount(UseTarget);
+	} 
+	else if (TargetMode == ESimplifyTargetType::TriangleCount)
+	{
+		Reducer.SimplifyToTriangleCount(TargetCount);
+	}
+	else if (TargetMode == ESimplifyTargetType::EdgeLength)
+	{
+		Reducer.SimplifyToEdgeLength(TargetEdgeLength);
+	}
+}
+
+
+void USimplifyMeshTool::UpdateResult()
+{
+	
+	if (bResultValid) 
+	{
+		return;
+	}
+
+	//FScopedDurationTimeLogger Timer(TEXT("Simplification Time"));
+
+	FDynamicMesh3* TargetMesh = DynamicMeshComponent->GetMesh();
+	TargetMesh->Copy(OriginalMesh);
+	int OriginalTriCount = OriginalMesh.TriangleCount();
+
+	if (bDiscardAttributes)
+	{
+		TargetMesh->DiscardAttributes();
+	}
+	if (SimplifierType == ESimplifyType::QEM)
+	{
+		ComputeResult<FQEMSimplification>(TargetMesh, bReproject, OriginalTriCount, OriginalMesh, OriginalMeshSpatial, TargetMode, TargetPercentage, TargetCount, TargetEdgeLength);
+	}
+	else if (SimplifierType == ESimplifyType::Attribute)
+	{
+		ComputeResult<FAttrMeshSimplification>(TargetMesh, bReproject, OriginalTriCount, OriginalMesh, OriginalMeshSpatial, TargetMode, TargetPercentage, TargetCount, TargetEdgeLength);
+	}
+	else //if (SimplifierType == ESimplifyType::UE4Standard)
+	{
+
+
+		const FMeshDescription* SrcMeshDescription = ComponentTarget->GetMesh();
+		FMeshDescription DstMeshDescription(*SrcMeshDescription);
+
+
+		FOverlappingCorners OverlappingCorners;
+		FMeshDescriptionOperations::FindOverlappingCorners(OverlappingCorners, *SrcMeshDescription, 1.e-5);
+
+		FMeshReductionSettings ReductionSettings;
+		if (TargetMode == ESimplifyTargetType::Percentage)
+		{
+			ReductionSettings.PercentTriangles = FMath::Max(TargetPercentage / 100., .001);  // Only support triangle percentage and count, but not edge length
+		}
+		else if (TargetMode == ESimplifyTargetType::TriangleCount)
+		{
+			int32 NumTris = SrcMeshDescription->Polygons().Num();
+			ReductionSettings.PercentTriangles = (float)TargetCount / (float)NumTris;
+		}
+
+		float Error;
+		{
+			IMeshReductionManagerModule& MeshReductionModule = FModuleManager::Get().LoadModuleChecked<IMeshReductionManagerModule>("MeshReductionInterface");
+			IMeshReduction* MeshReduction = MeshReductionModule.GetStaticMeshReductionInterface();
+
+
+			if (!MeshReduction)
+			{
+				// no reduction possible
+				Error = 0.f;
+				return;
+			}
+
+			Error = ReductionSettings.MaxDeviation;
+			MeshReduction->ReduceMeshDescription(DstMeshDescription, Error, *SrcMeshDescription, OverlappingCorners, ReductionSettings);
+		}
+
+		UE_LOG(LogMeshSimplification, Log, TEXT("Mesh simplified to %d triangles"), DstMeshDescription.Polygons().Num());
+
+		// Put the reduced mesh into the target...
+		DynamicMeshComponent->InitializeMesh(&DstMeshDescription);
+	}
+
+	UE_LOG(LogMeshSimplification, Log, TEXT("Mesh simplified to %d triangles"), TargetMesh->TriangleCount());
+
+	DynamicMeshComponent->NotifyMeshUpdated();
+	GetToolManager()->PostInvalidation();
+	MeshStatisticsProperties->Update(*DynamicMeshComponent->GetMesh());
+
+	bResultValid = true;
+}
+
+bool USimplifyMeshTool::HasAccept() const
+{
+	return true;
+}
+
+bool USimplifyMeshTool::CanAccept() const
+{
+	return true;
+}
+
+
+
+
+#undef LOCTEXT_NAMESPACE
diff --git a/Engine/Plugins/Experimental/MeshModelingToolset/Source/MeshModelingToolsEditorOnly/Private/VoxelCSGMeshesTool.cpp b/Engine/Plugins/Experimental/MeshModelingToolset/Source/MeshModelingToolsEditorOnly/Private/VoxelCSGMeshesTool.cpp
new file mode 100644
--- /dev/null
+++ b/Engine/Plugins/Experimental/MeshModelingToolset/Source/MeshModelingToolsEditorOnly/Private/VoxelCSGMeshesTool.cpp
@@ -0,0 +1,187 @@
+// Copyright 1998-2019 Epic Games, Inc. All Rights Reserved.
+
+#include "VoxelCSGMeshesTool.h"
+#include "InteractiveToolManager.h"
+#include "ToolBuilderUtil.h"
+#include "ToolSetupUtil.h"
+#include "AssetGenerationUtil.h"
+#include "Selection/ToolSelectionUtil.h"
+
+#include "DynamicMesh3.h"
+#include "MeshDescriptionToDynamicMesh.h"
+
+#include "CompositionOps/VoxelBooleanMeshesOp.h"
+
+#define LOCTEXT_NAMESPACE "UVoxelCSGMeshesTool"
+
+
+/*
+ * ToolBuilder
+ */
+bool UVoxelCSGMeshesToolBuilder::CanBuildTool(const FToolBuilderState& SceneState) const
+{
+	bool bHasBuildAPI = (this->AssetAPI != nullptr);
+	bool bHasComponents = ToolBuilderUtil::CountComponents(SceneState, CanMakeComponentTarget) == 2;
+	return (bHasBuildAPI && bHasComponents);
+}
+
+UInteractiveTool* UVoxelCSGMeshesToolBuilder::BuildTool(const FToolBuilderState& SceneState) const
+{
+	UVoxelCSGMeshesTool* NewTool = NewObject<UVoxelCSGMeshesTool>(SceneState.ToolManager);
+
+	TArray<UActorComponent*> Components = ToolBuilderUtil::FindAllComponents(SceneState, CanMakeComponentTarget);
+	check(Components.Num() > 0);
+
+	TArray<TUniquePtr<FPrimitiveComponentTarget>> ComponentTargets;
+	for (UActorComponent* ActorComponent : Components)
+	{
+		auto* MeshComponent = Cast<UPrimitiveComponent>(ActorComponent);
+		if ( MeshComponent )
+		{
+			ComponentTargets.Add(MakeComponentTarget(MeshComponent));
+		}
+	}
+
+	NewTool->SetSelection(MoveTemp(ComponentTargets));
+	NewTool->SetWorld(SceneState.World);
+	NewTool->SetAssetAPI(AssetAPI);
+
+	return NewTool;
+}
+
+/*
+ * Tool
+ */
+
+UVoxelCSGMeshesTool::UVoxelCSGMeshesTool()
+{
+}
+
+void UVoxelCSGMeshesTool::SetWorld(UWorld* World)
+{
+	this->TargetWorld = World;
+}
+
+void UVoxelCSGMeshesTool::SetAssetAPI(IToolsContextAssetAPI* AssetAPIIn)
+{
+	this->AssetAPI = AssetAPIIn;
+}
+
+
+
+void UVoxelCSGMeshesTool::Setup()
+{
+	UInteractiveTool::Setup();
+
+	CSGProps = NewObject<UVoxelCSGMeshesToolProperties>();
+	CSGProps->VoxelCount = 128;
+	CSGProps->Adaptivity = 0.01f;
+	CSGProps->IsoSurface = 0.0f;
+	AddToolPropertySource(CSGProps);
+
+	// Hide the source meshes
+	for (auto& ComponentTarget : ComponentTargets)
+	{
+		ComponentTarget->SetOwnerVisibility(false);
+	}
+
+	// save transformed version of input meshes (maybe this could happen in the Operator?)
+	CacheInputMeshes();
+
+	// initialize the PreviewMesh+BackgroundCompute object
+	Preview = NewObject<UMeshOpPreviewWithBackgroundCompute>(this, "Preview");
+	Preview->Setup(this->TargetWorld, this);
+	Preview->ConfigureMaterials(
+		// TODO : Probably multi-select bug - MCD
+		ToolSetupUtil::GetDefaultMaterial(GetToolManager(), ComponentTargets[0]->GetMaterial(0)),
+		ToolSetupUtil::GetDefaultWorkingMaterial(GetToolManager())
+	);
+	Preview->SetVisibility(false);
+	Preview->InvalidateResult();    // start compute
+}
+
+
+void UVoxelCSGMeshesTool::Shutdown(EToolShutdownType ShutdownType)
+{
+	TUniquePtr<FDynamicMeshOpResult> Result = Preview->Shutdown();
+	if (ShutdownType == EToolShutdownType::Accept)
+	{
+		GenerateAsset(Result);
+	}
+
+	// Restore (unhide) the source meshes
+	for (auto& ComponentTarget : ComponentTargets)
+	{
+		ComponentTarget->SetOwnerVisibility(true);
+	}
+}
+
+
+
+void UVoxelCSGMeshesTool::Tick(float DeltaTime)
+{
+	Preview->Tick(DeltaTime);
+}
+
+void UVoxelCSGMeshesTool::Render(IToolsContextRenderAPI* RenderAPI)
+{
+}
+
+bool UVoxelCSGMeshesTool::HasAccept() const
+{
+	return true;
+}
+
+bool UVoxelCSGMeshesTool::CanAccept() const
+{
+	return Preview->HaveValidResult();
+}
+
+void UVoxelCSGMeshesTool::OnPropertyModified(UObject* PropertySet, UProperty* Property)
+{
+	Preview->InvalidateResult();
+}
+
+TSharedPtr<FDynamicMeshOperator> UVoxelCSGMeshesTool::MakeNewOperator()
+{
+	TSharedPtr<FVoxelBooleanMeshesOp> CSGOp = MakeShared<FVoxelBooleanMeshesOp>();
+	CSGOp->Operation = (FVoxelBooleanMeshesOp::EBooleanOperation)(int)CSGProps->Operation;
+	CSGOp->VoxelCount = CSGProps->VoxelCount;
+	CSGOp->AdaptivityD = CSGProps->Adaptivity;
+	CSGOp->IsoSurfaceD = CSGProps->IsoSurface;
+	CSGOp->InputMeshArray = InputMeshes;
+	return CSGOp;
+}
+
+void UVoxelCSGMeshesTool::CacheInputMeshes()
+{
+	InputMeshes = MakeShared<TArray<IVoxelBasedCSG::FPlacedMesh>>();
+
+	// Package the selected meshes and transforms for consumption by the CSGTool
+	for (auto& ComponentTarget : ComponentTargets)
+	{
+		IVoxelBasedCSG::FPlacedMesh PlacedMesh;
+		PlacedMesh.Mesh = ComponentTarget->GetMesh();
+		PlacedMesh.Transform = ComponentTarget->GetWorldTransform();
+		InputMeshes->Add(PlacedMesh);
+	}
+}
+
+void UVoxelCSGMeshesTool::GenerateAsset(const TUniquePtr<FDynamicMeshOpResult>& Result)
+{
+	check(Result->Mesh.Get() != nullptr);
+
+	GetToolManager()->BeginUndoTransaction(LOCTEXT("VoxelCSGMeshes", "Boolean Meshes"));
+
+	AActor* NewActor = AssetGenerationUtil::GenerateStaticMeshActor(
+		AssetAPI, TargetWorld,
+		Result->Mesh.Get(), Result->Transform, TEXT("CSGMesh"),
+		AssetGenerationUtil::GetDefaultAutoGeneratedAssetPath());
+
+	// select newly-created object
+	ToolSelectionUtil::SetNewActorSelection(GetToolManager(), NewActor);
+
+	GetToolManager()->EndUndoTransaction();
+}
+
+#undef LOCTEXT_NAMESPACE
diff --git a/Engine/Plugins/Experimental/MeshModelingToolset/Source/MeshModelingToolsEditorOnly/Public/MergeMeshesTool.h b/Engine/Plugins/Experimental/MeshModelingToolset/Source/MeshModelingToolsEditorOnly/Public/MergeMeshesTool.h
new file mode 100644
--- /dev/null
+++ b/Engine/Plugins/Experimental/MeshModelingToolset/Source/MeshModelingToolsEditorOnly/Public/MergeMeshesTool.h
@@ -0,0 +1,101 @@
+// Copyright 1998-2019 Epic Games, Inc. All Rights Reserved.
+
+#pragma once
+
+#include "CoreMinimal.h"
+#include "MultiSelectionTool.h"
+#include "InteractiveToolBuilder.h"
+#include "MeshOpPreviewHelpers.h"
+#include "ProxyLODVolume.h"
+#include "MergeMeshesTool.generated.h"
+
+
+/**
+ *
+ */
+UCLASS()
+class MESHMODELINGTOOLSEDITORONLY_API UMergeMeshesToolBuilder : public UInteractiveToolBuilder
+{
+	GENERATED_BODY()
+
+public:
+	IToolsContextAssetAPI* AssetAPI;
+
+	UMergeMeshesToolBuilder()
+	{
+		AssetAPI = nullptr;
+	}
+
+	virtual bool CanBuildTool(const FToolBuilderState& SceneState) const override;
+	virtual UInteractiveTool* BuildTool(const FToolBuilderState& SceneState) const override;
+};
+
+
+
+/**
+ * Standard properties of the Merge Meshes operation
+ */
+UCLASS()
+class MESHMODELINGTOOLSEDITORONLY_API UMergeMeshesToolProperties : public UInteractiveToolPropertySet
+{
+	GENERATED_BODY()
+public:
+
+	/** The size of the geometry bounding box major axis measured in voxels.*/
+	UPROPERTY(EditAnywhere, Category = Options, meta = (UIMin = "16", UIMax = "2048", ClampMin = "16", ClampMax = "2048"))
+	int32 VoxelCount = 128;
+
+	/** Remeshing adaptivity */
+	UPROPERTY(EditAnywhere, Category = Options, meta = (UIMin = "0", UIMax = "1", ClampMin = "0", ClampMax = "1"))
+	float Adaptivity = 0.01f;
+
+	/** Offset when remeshing, measured in voxels units */
+	UPROPERTY(EditAnywhere, Category = Options, meta = (UIMin = "-2", UIMax = "2", ClampMin = "-2", ClampMax = "2"))
+	float IsoSurface = 0;
+};
+
+/**
+ *
+ */
+UCLASS()
+class MESHMODELINGTOOLSEDITORONLY_API UMergeMeshesTool : public UMultiSelectionTool, public IDynamicMeshOperatorFactory
+{
+	GENERATED_BODY()
+
+public:
+	UMergeMeshesTool();
+
+	virtual void SetWorld(UWorld* World);
+	virtual void SetAssetAPI(IToolsContextAssetAPI* AssetAPI);
+
+	virtual void Setup() override;
+	virtual void Shutdown(EToolShutdownType ShutdownType) override;
+
+	virtual void Tick(float DeltaTime) override;
+	virtual void Render(IToolsContextRenderAPI* RenderAPI) override;
+
+	virtual bool HasCancel() const override { return true; }
+	virtual bool HasAccept() const override;
+	virtual bool CanAccept() const override;
+
+	virtual void OnPropertyModified(UObject* PropertySet, UProperty* Property) override;
+
+	// IDynamicMeshOperatorFactory API
+	virtual TSharedPtr<FDynamicMeshOperator> MakeNewOperator() override;
+
+protected:
+	UPROPERTY()
+	UMergeMeshesToolProperties* MergeProps;
+
+	UPROPERTY()
+	UMeshOpPreviewWithBackgroundCompute* Preview;
+
+protected:
+	UWorld* TargetWorld;
+	IToolsContextAssetAPI* AssetAPI;
+
+	TSharedPtr<TArray<IVoxelBasedCSG::FPlacedMesh>> InputMeshes;
+	void CacheInputMeshes();
+
+	void GenerateAsset(const TUniquePtr<FDynamicMeshOpResult>& Result);
+};
diff --git a/Engine/Plugins/Experimental/MeshModelingToolset/Source/MeshModelingToolsEditorOnly/Public/ParameterizeMeshTool.h b/Engine/Plugins/Experimental/MeshModelingToolset/Source/MeshModelingToolsEditorOnly/Public/ParameterizeMeshTool.h
new file mode 100644
--- /dev/null
+++ b/Engine/Plugins/Experimental/MeshModelingToolset/Source/MeshModelingToolsEditorOnly/Public/ParameterizeMeshTool.h
@@ -0,0 +1,132 @@
+// Copyright 1998-2019 Epic Games, Inc. All Rights Reserved.
+
+#pragma once
+
+#include "CoreMinimal.h"
+#include "UObject/NoExportTypes.h"
+#include "SingleSelectionTool.h"
+#include "InteractiveToolBuilder.h"
+#include "DynamicMesh3.h"
+#include "MeshOpPreviewHelpers.h"
+#include "ParameterizeMeshTool.generated.h"
+
+
+// predeclarations
+struct FMeshDescription;
+class USimpleDynamicMeshComponent;
+class UMaterialInterface;
+class UMaterialInstanceDynamic;
+
+
+
+/** Material modes for MeshInspectorTool */
+UENUM()
+enum class EParameterizeMeshMaterialMode : uint8
+{
+	/** Input material */
+	Default,
+
+	/** Checkerboard material */
+	Checkerboard,
+
+	/** Override material */
+	Override
+};
+
+/**
+ *
+ */
+UCLASS()
+class MESHMODELINGTOOLSEDITORONLY_API UParameterizeMeshToolBuilder : public UInteractiveToolBuilder
+{
+	GENERATED_BODY()
+
+public:
+
+	IToolsContextAssetAPI* AssetAPI;
+
+	UParameterizeMeshToolBuilder()
+	{
+		AssetAPI = nullptr;
+	}
+
+	virtual bool CanBuildTool(const FToolBuilderState& SceneState) const override;
+	virtual UInteractiveTool* BuildTool(const FToolBuilderState& SceneState) const override;
+};
+
+
+
+
+
+/**
+ * Simple Mesh Simplifying Tool
+ */
+UCLASS()
+class MESHMODELINGTOOLSEDITORONLY_API UParameterizeMeshTool : public USingleSelectionTool, public IDynamicMeshOperatorFactory
+{
+	GENERATED_BODY()
+
+public:
+	UParameterizeMeshTool();
+
+	virtual void SetWorld(UWorld* World);
+	virtual void SetAssetAPI(IToolsContextAssetAPI* AssetAPI);
+
+	virtual void Setup() override;
+	virtual void Shutdown(EToolShutdownType ShutdownType) override;
+
+	virtual void Tick(float DeltaTime) override;
+	virtual void Render(IToolsContextRenderAPI* RenderAPI) override;
+
+	virtual bool HasCancel() const override { return true; }
+	virtual bool HasAccept() const override;
+	virtual bool CanAccept() const override;
+
+#if WITH_EDITOR
+	virtual void PostEditChangeProperty(FPropertyChangedEvent& PropertyChangedEvent) override;
+#endif
+
+	// IDynamicMeshOperatorFactory API
+	virtual TSharedPtr<FDynamicMeshOperator> MakeNewOperator() override;
+
+protected:
+	// need to update bResultValid if these are modified, so we don't publicly expose them. 
+	// @todo setters/getters for these
+
+	/** Maximum amount of stretch, from none to any.  If zero stretch is specified each triangle will likey be its own chart */
+	UPROPERTY(EditAnywhere, Category = Options, meta = (Default = "0.166", UIMin = "0", UIMax = "1", ClampMin = "0", ClampMax= "1"))
+	float ChartStretch = 0.11f;
+	
+	/** Material that will be displayed on the mesh */
+	UPROPERTY(EditAnywhere, Category = Options)
+	EParameterizeMeshMaterialMode MaterialMode = EParameterizeMeshMaterialMode::Default;
+
+	/** The material that will be displayed on the mesh when override material mode has been selected */
+	UPROPERTY(EditAnywhere, Category = Options)
+	UMaterialInterface* OverrideMaterial = nullptr;
+	
+	/** For display only - the density of the checkerboard pattern in the viewport.*/
+	UPROPERTY(EditAnywhere, Category = Options, meta = (Default = "20", UIMin = "1.0", UIMax = "40.0", ClampMin = "0.01", ClampMax = "1000.0"))
+	float CheckerDensity = 20.f;
+
+protected:
+	UPROPERTY()
+	UMaterialInterface* DefaultMaterial = nullptr;
+
+	UPROPERTY()
+	UMaterialInstanceDynamic* DisplayMaterial = nullptr;
+
+	UPROPERTY()
+	UMaterialInstanceDynamic* CheckerMaterial = nullptr;
+	
+	UPROPERTY()
+	UMeshOpPreviewWithBackgroundCompute* Preview = nullptr;
+
+
+protected:
+	UWorld* TargetWorld;
+	IToolsContextAssetAPI* AssetAPI;
+
+	TSharedPtr<FMeshDescription> InputMesh;
+
+};
diff --git a/Engine/Plugins/Experimental/MeshModelingToolset/Source/MeshModelingToolsEditorOnly/Public/PolygonOnMeshTool.h b/Engine/Plugins/Experimental/MeshModelingToolset/Source/MeshModelingToolsEditorOnly/Public/PolygonOnMeshTool.h
new file mode 100644
--- /dev/null
+++ b/Engine/Plugins/Experimental/MeshModelingToolset/Source/MeshModelingToolsEditorOnly/Public/PolygonOnMeshTool.h
@@ -0,0 +1,180 @@
+// Copyright 1998-2019 Epic Games, Inc. All Rights Reserved.
+
+#pragma once
+
+#include "CoreMinimal.h"
+#include "UObject/NoExportTypes.h"
+#include "InteractiveGizmo.h"
+#include "SingleSelectionTool.h"
+#include "InteractiveToolBuilder.h"
+#include "MeshOpPreviewHelpers.h"
+#include "CuttingOps/EmbedPolygonsOp.h"
+#include "DynamicMesh3.h"
+#include "BaseTools/SingleClickTool.h"
+#include "PolygonOnMeshTool.generated.h"
+
+
+// predeclarations
+struct FMeshDescription;
+class USimpleDynamicMeshComponent;
+
+
+
+
+
+/**
+ *
+ */
+UCLASS()
+class MESHMODELINGTOOLSEDITORONLY_API UPolygonOnMeshToolBuilder : public UInteractiveToolBuilder
+{
+	GENERATED_BODY()
+
+public:
+	IToolsContextAssetAPI* AssetAPI = nullptr;
+
+	virtual bool CanBuildTool(const FToolBuilderState& SceneState) const override;
+	virtual UInteractiveTool* BuildTool(const FToolBuilderState& SceneState) const override;
+};
+
+
+
+
+
+/**
+ * Standard properties of the polygon-on-mesh operations
+ */
+UCLASS()
+class MESHMODELINGTOOLSEDITORONLY_API UPolygonOnMeshToolProperties : public UInteractiveToolPropertySet
+{
+	GENERATED_BODY()
+
+public:
+	UPolygonOnMeshToolProperties();
+
+	/** If true, UVs and Normals are discarded  */
+	UPROPERTY(EditAnywhere, Category = Options)
+	bool bDiscardAttributes;
+
+	/** What operation to do with the polygon -- e.g. extrude it, cut through with it, or delete inside it */
+	UPROPERTY(EditAnywhere, Category = Options)
+	EEmbeddedPolygonOpMethod PolygonOperation;
+
+	// TODO: re-add if/when extrude is added as a supported operation
+	///** Amount to extrude, if extrude is enabled */
+	//UPROPERTY(EditAnywhere, Category = Options)
+	//float ExtrudeDistance;
+
+	/** Scale of polygon to embed */
+	UPROPERTY(EditAnywhere, Category = Options)
+	float PolygonScale;
+};
+
+
+
+
+/**
+ * Advanced properties of the plane cut operation
+ */
+UCLASS()
+class MESHMODELINGTOOLSEDITORONLY_API UPolygonOnMeshAdvancedProperties : public UInteractiveToolPropertySet
+{
+	GENERATED_BODY()
+
+public:
+	UPolygonOnMeshAdvancedProperties();
+
+	// TODO: add additional properties, e.g. perhaps more advanced triangulation options, perhaps post-process remeshing options (because the polygon-in-mesh embedding can create slivers)
+
+};
+
+
+UCLASS()
+class MESHMODELINGTOOLSEDITORONLY_API UPolygonOnMeshOperatorFactory : public UObject, public IDynamicMeshOperatorFactory
+{
+	GENERATED_BODY()
+
+public:
+	// IDynamicMeshOperatorFactory API
+	virtual TSharedPtr<FDynamicMeshOperator> MakeNewOperator() override;
+
+	UPROPERTY()
+	UPolygonOnMeshTool *Tool;
+
+	// TODO: add info about what this factory is specifically meant to do with the tool (e.g. cut inside vs outside)
+	// currently not needed as we only do one thing at a time, as configured by the tool properties
+};
+
+/**
+ * Simple Mesh Plane Cutting Tool
+ */
+UCLASS()
+class MESHMODELINGTOOLSEDITORONLY_API UPolygonOnMeshTool : public USingleSelectionTool
+{
+	GENERATED_BODY()
+
+public:
+
+	friend UPolygonOnMeshOperatorFactory;
+
+	UPolygonOnMeshTool();
+
+	virtual void Setup() override;
+	virtual void Shutdown(EToolShutdownType ShutdownType) override;
+
+	virtual void SetWorld(UWorld* World);
+	virtual void SetAssetAPI(IToolsContextAssetAPI* AssetAPI);
+
+	virtual void Tick(float DeltaTime) override;
+	virtual void Render(IToolsContextRenderAPI* RenderAPI) override;
+
+	virtual bool HasCancel() const override { return true; }
+	virtual bool HasAccept() const override;
+	virtual bool CanAccept() const override;
+
+	virtual void PostEditChangeProperty(FPropertyChangedEvent& PropertyChangedEvent) override;
+
+	virtual void OnPropertyModified(UObject* PropertySet, UProperty* Property) override;
+
+protected:
+
+	UPROPERTY()
+	UPolygonOnMeshToolProperties* BasicProperties;
+
+	UPROPERTY()
+	UPolygonOnMeshAdvancedProperties* AdvancedProperties;
+
+	/** Origin of polygon projection plane */
+	UPROPERTY()
+	FVector EmbedPolygonOrigin;
+
+	/** Orientation of polygon projection plane */
+	UPROPERTY()
+	FQuat EmbedPolygonOrientation;
+
+	UPROPERTY()
+	TArray<UMeshOpPreviewWithBackgroundCompute*> Previews;
+
+
+protected:
+	TSharedPtr<FDynamicMesh3> OriginalDynamicMesh;
+
+	UWorld* TargetWorld;
+	IToolsContextAssetAPI* AssetAPI;
+
+	FViewCameraState CameraState;
+
+	// plane control
+	UPROPERTY()
+	UInteractiveGizmo* PositionPlaneGizmo;
+
+	void UpdatePlaneFromGizmo(const FFrame3d& WorldPosition);
+
+	void UpdateNumPreviews();
+
+	IClickBehaviorTarget* SetPointInWorldConnector = nullptr;
+
+	virtual void SetPlaneFromWorldPos(const FVector& Position, const FVector& Normal);
+
+	void GenerateAsset(const TArray<TUniquePtr<FDynamicMeshOpResult>>& Results);
+};
diff --git a/Engine/Plugins/Experimental/MeshModelingToolset/Source/MeshModelingToolsEditorOnly/Public/SimplifyMeshTool.h b/Engine/Plugins/Experimental/MeshModelingToolset/Source/MeshModelingToolsEditorOnly/Public/SimplifyMeshTool.h
new file mode 100644
--- /dev/null
+++ b/Engine/Plugins/Experimental/MeshModelingToolset/Source/MeshModelingToolsEditorOnly/Public/SimplifyMeshTool.h
@@ -0,0 +1,141 @@
+// Copyright 1998-2019 Epic Games, Inc. All Rights Reserved.
+
+#pragma once
+
+#include "CoreMinimal.h"
+#include "UObject/NoExportTypes.h"
+#include "SingleSelectionTool.h"
+#include "InteractiveToolBuilder.h"
+#include "DynamicMesh3.h"
+#include "DynamicMeshAABBTree3.h"
+#include "Properties/MeshStatisticsProperties.h"
+#include "SimplifyMeshTool.generated.h"
+
+
+// predeclarations
+struct FMeshDescription;
+class USimpleDynamicMeshComponent;
+
+
+
+
+/**  */
+UENUM()
+enum class ESimplifyTargetType : uint8
+{
+	/** Percentage of input triangles */
+	Percentage = 0 UMETA(DisplayName = "Percentage"),
+
+	/** Target triangle count */
+	TriangleCount = 1 UMETA(DisplayName = "Triangle Count"),
+
+	/** Target edge length */
+	EdgeLength = 2 UMETA(DisplayName = "Edge Length")
+};
+
+/**  */
+UENUM()
+enum class ESimplifyType : uint8
+{
+	/** Fastest. Standard quadric error metric. Will not simplify UV boundaries.*/
+	QEM = 0 UMETA(DisplayName = "QEM"),
+
+	/** Potentially higher quality. Takes the normal into account. Will not simplify UV bounaries. */
+	Attribute = 1 UMETA(DisplayName = "Normal Aware"),
+
+	/** Highest quality reduction.  Will simplify UV boundaries. */
+	UE4Standard = 2 UMETA(DisplayName = "UE4 Standard"),
+};
+
+
+
+
+/**
+ *
+ */
+UCLASS()
+class MESHMODELINGTOOLSEDITORONLY_API USimplifyMeshToolBuilder : public UInteractiveToolBuilder
+{
+	GENERATED_BODY()
+
+public:
+	virtual bool CanBuildTool(const FToolBuilderState& SceneState) const override;
+	virtual UInteractiveTool* BuildTool(const FToolBuilderState& SceneState) const override;
+};
+
+
+
+
+
+/**
+ * Simple Mesh Simplifying Tool
+ */
+UCLASS()
+class MESHMODELINGTOOLSEDITORONLY_API USimplifyMeshTool : public USingleSelectionTool
+{
+	GENERATED_BODY()
+
+public:
+	USimplifyMeshTool();
+
+	virtual void Setup() override;
+	virtual void Shutdown(EToolShutdownType ShutdownType) override;
+
+	virtual void Render(IToolsContextRenderAPI* RenderAPI) override;
+
+	virtual bool HasCancel() const override { return true; }
+	virtual bool HasAccept() const override;
+	virtual bool CanAccept() const override;
+
+#if WITH_EDITOR
+	virtual void PostEditChangeProperty(FPropertyChangedEvent& PropertyChangedEvent) override;
+#endif
+
+protected:
+	// need to update bResultValid if these are modified, so we don't publicly expose them. 
+	// @todo setters/getters for these
+
+	/** Simplification Target Type  */
+	UPROPERTY(EditAnywhere, Category = Options)
+	ESimplifyTargetType TargetMode;
+
+	/** Simplification Scheme  */
+	UPROPERTY(EditAnywhere, Category = Options)
+	ESimplifyType SimplifierType;
+
+	/** Target percentage of original triangle count */
+	UPROPERTY(EditAnywhere, Category = Options, meta = (UIMin = "0", UIMax = "100"))
+	int TargetPercentage;
+
+	/** Target edge length */
+	UPROPERTY(EditAnywhere, Category = Options, meta = (UIMin = "3.0", UIMax = "10.0", ClampMin = "0.001", ClampMax = "1000.0"))
+	float TargetEdgeLength;
+
+	/** Target triangle count */
+	UPROPERTY(EditAnywhere, Category = Options, meta = (UIMin = "4", UIMax = "10000", ClampMin = "1", ClampMax = "9999999999"))
+	int TargetCount;
+
+	/** If true, UVs and Normals are discarded  */
+	UPROPERTY(EditAnywhere, Category = Options)
+	bool bDiscardAttributes;
+
+	/** Enable projection back to input mesh */
+	UPROPERTY(EditAnywhere, Category = Advanced)
+	bool bReproject;
+
+	/** Prevent normal flips */
+	UPROPERTY(EditAnywhere, Category = Advanced)
+	bool bPreventNormalFlips;
+
+	UPROPERTY()
+	UMeshStatisticsProperties* MeshStatisticsProperties;
+
+protected:
+	USimpleDynamicMeshComponent* DynamicMeshComponent;
+	FDynamicMesh3 OriginalMesh;
+	FDynamicMeshAABBTree3 OriginalMeshSpatial;
+	bool bResultValid;
+	void UpdateResult();
+
+	void ComputeSimplifying();
+};
diff --git a/Engine/Plugins/Experimental/MeshModelingToolset/Source/MeshModelingToolsEditorOnly/Public/VoxelCSGMeshesTool.h b/Engine/Plugins/Experimental/MeshModelingToolset/Source/MeshModelingToolsEditorOnly/Public/VoxelCSGMeshesTool.h
new file mode 100644
--- /dev/null
+++ b/Engine/Plugins/Experimental/MeshModelingToolset/Source/MeshModelingToolsEditorOnly/Public/VoxelCSGMeshesTool.h
@@ -0,0 +1,127 @@
+// Copyright 1998-2019 Epic Games, Inc. All Rights Reserved.
+
+#pragma once
+
+#include "CoreMinimal.h"
+#include "UObject/NoExportTypes.h"
+#include "InteractiveToolBuilder.h"
+#include "MeshOpPreviewHelpers.h"
+#include "MultiSelectionTool.h"
+#include "ProxyLODVolume.h"
+#include "VoxelCSGMeshesTool.generated.h"
+
+
+/**
+ *
+ */
+UCLASS()
+class MESHMODELINGTOOLSEDITORONLY_API UVoxelCSGMeshesToolBuilder : public UInteractiveToolBuilder
+{
+	GENERATED_BODY()
+
+public:
+	IToolsContextAssetAPI* AssetAPI;
+
+	UVoxelCSGMeshesToolBuilder()
+	{
+		AssetAPI = nullptr;
+	}
+
+	virtual bool CanBuildTool(const FToolBuilderState& SceneState) const override;
+	virtual UInteractiveTool* BuildTool(const FToolBuilderState& SceneState) const override;
+};
+
+
+/**  */
+UENUM()
+enum class EVoxelCSGOperation : uint8
+{
+	/** Subtracts the first object from the second */
+	DifferenceAB = 0 UMETA(DisplayName = "A - B"),
+
+	/** Subtracts the second object from the first */
+	DifferenceBA = 1 UMETA(DisplayName = "B - A"),
+
+	/** intersection of two objects */
+	Intersect = 2 UMETA(DisplayName = "Intersect"),
+
+	/** union of two objects */
+	Union = 3 UMETA(DisplayName = "Union"),
+
+};
+
+
+/**
+ * Standard properties of the Voxel CSG operation
+ */
+UCLASS()
+class MESHMODELINGTOOLSEDITORONLY_API UVoxelCSGMeshesToolProperties : public UInteractiveToolPropertySet
+{
+	GENERATED_BODY()
+public:
+	/** The type of operation  */
+	UPROPERTY(EditAnywhere, Category = Options)
+	EVoxelCSGOperation Operation;
+
+	/** The size of the geometry bounding box major axis measured in voxels.*/
+	UPROPERTY(EditAnywhere, Category = Options, meta = (UIMin = "16", UIMax = "2048", ClampMin = "16", ClampMax = "2048"))
+	int32 VoxelCount = 128;
+
+	/** Remeshing adaptivity */
+	UPROPERTY(EditAnywhere, Category = Options, meta = (UIMin = "0", UIMax = "1", ClampMin = "0", ClampMax = "1"))
+	float Adaptivity = 0.01f;
+
+	/** Offset when remeshing, measured in voxels units */
+	UPROPERTY(EditAnywhere, Category = Options, meta = (UIMin = "-2", UIMax = "2", ClampMin = "-2", ClampMax = "2"))
+	float IsoSurface = 0;
+};
+
+
+
+
+/**
+ *
+ */
+UCLASS()
+class MESHMODELINGTOOLSEDITORONLY_API UVoxelCSGMeshesTool : public UMultiSelectionTool, public IDynamicMeshOperatorFactory
+{
+	GENERATED_BODY()
+
+public:
+	UVoxelCSGMeshesTool();
+
+	virtual void SetWorld(UWorld* World);
+	virtual void SetAssetAPI(IToolsContextAssetAPI* AssetAPI);
+
+	virtual void Setup() override;
+	virtual void Shutdown(EToolShutdownType ShutdownType) override;
+
+	virtual void Tick(float DeltaTime) override;
+	virtual void Render(IToolsContextRenderAPI* RenderAPI) override;
+
+	virtual bool HasCancel() const override { return true; }
+	virtual bool HasAccept() const override;
+	virtual bool CanAccept() const override;
+
+	virtual void OnPropertyModified(UObject* PropertySet, UProperty* Property) override;
+
+	// IDynamicMeshOperatorFactory API
+	virtual TSharedPtr<FDynamicMeshOperator> MakeNewOperator() override;
+
+
+protected:
+	UPROPERTY()
+	UVoxelCSGMeshesToolProperties* CSGProps;
+
+	UPROPERTY()
+	UMeshOpPreviewWithBackgroundCompute* Preview;
+
+protected:
+	UWorld* TargetWorld;
+	IToolsContextAssetAPI* AssetAPI;
+
+	TSharedPtr<TArray<IVoxelBasedCSG::FPlacedMesh>> InputMeshes;
+	void CacheInputMeshes();
+
+	void GenerateAsset(const TUniquePtr<FDynamicMeshOpResult>& Result);
+};
diff --git a/Engine/Plugins/Experimental/MeshModelingToolset/Source/ModelingComponents/ModelingComponents.Build.cs b/Engine/Plugins/Experimental/MeshModelingToolset/Source/ModelingComponents/ModelingComponents.Build.cs
new file mode 100644
--- /dev/null
+++ b/Engine/Plugins/Experimental/MeshModelingToolset/Source/ModelingComponents/ModelingComponents.Build.cs
@@ -0,0 +1,65 @@
+// Copyright 1998-2019 Epic Games, Inc. All Rights Reserved.
+
+using UnrealBuildTool;
+
+public class ModelingComponents : ModuleRules
+{
+	public ModelingComponents(ReadOnlyTargetRules Target) : base(Target)
+	{
+		PCHUsage = ModuleRules.PCHUsageMode.UseExplicitOrSharedPCHs;
+		
+		PublicIncludePaths.AddRange(
+			new string[] {
+				// ... add public include paths required here ...
+			}
+			);
+				
+		
+		PrivateIncludePaths.AddRange(
+			new string[] {
+				// ... add other private include paths required here ...
+			}
+			);
+			
+		
+		PublicDependencyModuleNames.AddRange(
+			new string[]
+			{
+				"Core",
+                "InteractiveToolsFramework",
+                "MeshDescription",
+				"GeometricObjects",
+				"GeometryAlgorithms",
+				"DynamicMesh",
+				"MeshConversion",
+				"ModelingOperators"
+				// ... add other public dependencies that you statically link with here ...
+			}
+            );
+			
+		
+		PrivateDependencyModuleNames.AddRange(
+			new string[]
+			{
+				"CoreUObject",
+				"Engine",
+                "RenderCore",
+                "RHI",
+
+				//"MeshUtilities",    // temp for saving mesh asset
+				//"UnrealEd"
+				//"Slate",
+				//"SlateCore",
+				// ... add private dependencies that you statically link with here ...	
+			}
+			);
+		
+		
+		DynamicallyLoadedModuleNames.AddRange(
+			new string[]
+			{
+				// ... add any modules that your module loads dynamically here ...
+			}
+			);
+	}
+}
