Mesh Modeling Tools moved from NotForLicensees to Experimental

diff --git a/Engine/Plugins/Experimental/MeshModelingToolset/Source/ModelingComponents/Private/AssetGenerationUtil.cpp b/Engine/Plugins/Experimental/MeshModelingToolset/Source/ModelingComponents/Private/AssetGenerationUtil.cpp
new file mode 100644
--- /dev/null
+++ b/Engine/Plugins/Experimental/MeshModelingToolset/Source/ModelingComponents/Private/AssetGenerationUtil.cpp
@@ -0,0 +1,84 @@
+// Copyright 1998-2019 Epic Games, Inc. All Rights Reserved.
+
+#include "AssetGenerationUtil.h"
+#include "InteractiveTool.h"
+#include "InteractiveToolManager.h"
+#include "Materials/Material.h"
+
+#include "Components/StaticMeshComponent.h"
+#include "Engine/StaticMesh.h"
+#include "Engine/StaticMeshActor.h"
+
+#include "StaticMeshComponentBuilder.h"
+#include "MeshDescriptionToDynamicMesh.h"
+#include "MeshDescriptionBuilder.h"
+
+
+FString AssetGenerationUtil::GetDefaultAutoGeneratedAssetPath()
+{
+	return FString("/Game/_GENERATED");
+}
+
+
+#if WITH_EDITOR
+AActor* AssetGenerationUtil::GenerateStaticMeshActor(
+	IToolsContextAssetAPI* AssetAPI,
+	UWorld* TargetWorld,
+	const FDynamicMesh3* Mesh,
+	const FTransform3d& Transform,
+	FString ObjectName,
+	FString PackagePath
+)
+{
+	check(AssetAPI);
+	check(TargetWorld);
+	check(Mesh);
+
+	// asset name and location (temp)
+	FString PackageFolderPath = (PackagePath.IsEmpty()) ? GetDefaultAutoGeneratedAssetPath() : PackagePath;
+	//FString PackageFolderPath = AssetAPI->GetActiveAssetFolderPath();
+	FString MeshName = TEXT("MergedMesh");
+
+	// create new package
+	FString UniqueMeshName = AssetAPI->MakeUniqueAssetName(MeshName, PackageFolderPath);
+	UPackage* AssetPackage = AssetAPI->CreateNewPackage(UniqueMeshName, PackageFolderPath);
+
+
+	// create new actor
+	FRotator Rotation(0.0f, 0.0f, 0.0f);
+	FActorSpawnParameters SpawnInfo;
+	// @todo nothing here is specific to AStaticMeshActor...could we pass in a CDO and clone it instead of using SpawnActor?
+	AStaticMeshActor* NewActor = TargetWorld->SpawnActor<AStaticMeshActor>(FVector::ZeroVector, Rotation, SpawnInfo);
+	NewActor->Rename(*UniqueMeshName);
+
+	// construct new static mesh
+	FStaticMeshComponentBuilder Builder;
+	Builder.Initialize(AssetPackage, FName(*UniqueMeshName));
+
+	// Populate the MeshDescription associated with the asset.
+	FMeshDescriptionBuilder MeshDescripBuilder;
+	MeshDescripBuilder.SetMeshDescription(Builder.MeshDescription);
+	if (Mesh != nullptr)
+	{
+		MeshDescripBuilder.AppendMesh(Mesh, true);
+	}
+	else
+	{
+		// should generate default sphere here or something...
+	}
+
+	// create new mesh component and set as root of NewActor.
+	// NB: This should to be called only after the mesh has been created / populated,
+	// internally it calls  NewStaticMesh->CommitMeshDescription(0);
+	Builder.CreateAndSetAsRootComponent(NewActor);
+
+	// transform new component to origin
+	Builder.NewMeshComponent->SetWorldTransform((FTransform)Transform);
+
+	// force save of asset
+	//AssetAPI->InteractiveSaveGeneratedAsset(Builder.NewStaticMesh, AssetPackage);
+	AssetAPI->AutoSaveGeneratedAsset(Builder.NewStaticMesh, AssetPackage);
+
+	return NewActor;
+}
+#endif
diff --git a/Engine/Plugins/Experimental/MeshModelingToolset/Source/ModelingComponents/Private/MeshOpPreviewHelpers.cpp b/Engine/Plugins/Experimental/MeshModelingToolset/Source/ModelingComponents/Private/MeshOpPreviewHelpers.cpp
new file mode 100644
--- /dev/null
+++ b/Engine/Plugins/Experimental/MeshModelingToolset/Source/ModelingComponents/Private/MeshOpPreviewHelpers.cpp
@@ -0,0 +1,103 @@
+// Copyright 1998-2019 Epic Games, Inc. All Rights Reserved.
+
+#include "MeshOpPreviewHelpers.h"
+
+
+void UMeshOpPreviewWithBackgroundCompute::Setup(UWorld* InWorld, IDynamicMeshOperatorFactory* OpGenerator)
+{
+	PreviewMesh = NewObject<UPreviewMesh>(this, TEXT("PreviewMesh"));
+	PreviewMesh->CreateInWorld(InWorld, FTransform::Identity);
+
+	BackgroundCompute = MakeUnique<FBackgroundDynamicMeshComputeSource>(OpGenerator);
+	bResultValid = false;
+}
+
+TUniquePtr<FDynamicMeshOpResult> UMeshOpPreviewWithBackgroundCompute::Shutdown()
+{
+	BackgroundCompute->CancelActiveCompute();
+
+
+	TUniquePtr<FDynamicMeshOpResult> Result = MakeUnique<FDynamicMeshOpResult>();
+	Result->Mesh = PreviewMesh->ExtractPreviewMesh();
+	Result->Transform = FTransform3d(PreviewMesh->GetTransform());
+
+	PreviewMesh->SetVisible(false);
+	PreviewMesh->Disconnect();
+	PreviewMesh = nullptr;
+
+	return MoveTemp(Result);
+}
+
+
+void UMeshOpPreviewWithBackgroundCompute::Cancel()
+{
+	BackgroundCompute->CancelActiveCompute();
+
+	PreviewMesh->SetVisible(false);
+	PreviewMesh->Disconnect();
+	PreviewMesh = nullptr;
+}
+
+
+void UMeshOpPreviewWithBackgroundCompute::Tick(float DeltaTime)
+{
+	bool bIsLongDelay = false;
+	if (BackgroundCompute)
+	{
+		BackgroundCompute->Tick(DeltaTime);
+		bIsLongDelay = BackgroundCompute->GetElapsedComputeTime() > 2.0;
+	}
+
+	UpdateResults();
+
+	if (bResultValid || WorkingMaterial == nullptr || bIsLongDelay == false)
+	{
+		PreviewMesh->SetMaterial(StandardMaterial);
+	}
+	else
+	{
+		PreviewMesh->SetMaterial(WorkingMaterial);
+	}
+}
+
+
+void UMeshOpPreviewWithBackgroundCompute::UpdateResults()
+{
+	EBackgroundComputeTaskStatus Status = BackgroundCompute->CheckStatus();
+	if (Status == EBackgroundComputeTaskStatus::NewResultAvailable)
+	{
+		TSharedPtr<FDynamicMeshOperator> MeshOp =
+			StaticCastSharedPtr<FDynamicMeshOperator>(BackgroundCompute->ExtractResult());
+
+		TUniquePtr<FDynamicMesh3> ResultMesh = MeshOp->ExtractResult();
+		PreviewMesh->SetTransform((FTransform)MeshOp->GetResultTransform());
+		PreviewMesh->UpdatePreview(ResultMesh.Get());  // copies the mesh @todo we could just give ownership to the Preview!
+		PreviewMesh->SetVisible(true);
+		bResultValid = true;
+	}
+}
+
+
+void UMeshOpPreviewWithBackgroundCompute::InvalidateResult()
+{
+	BackgroundCompute->NotifyActiveComputeInvalidated();
+	bResultValid = false;
+}
+
+
+void UMeshOpPreviewWithBackgroundCompute::ConfigureMaterials(UMaterialInterface* StandardMaterialIn, UMaterialInterface* WorkingMaterialIn)
+{
+	this->StandardMaterial = StandardMaterialIn;
+	this->WorkingMaterial = WorkingMaterialIn;
+
+	if (PreviewMesh != nullptr)
+	{
+		PreviewMesh->SetMaterial(this->StandardMaterial);
+	}
+}
+
+
+void UMeshOpPreviewWithBackgroundCompute::SetVisibility(bool bVisible)
+{
+	PreviewMesh->SetVisible(bVisible);
+}
diff --git a/Engine/Plugins/Experimental/MeshModelingToolset/Source/ModelingComponents/Private/OctreeDynamicMeshComponent.cpp b/Engine/Plugins/Experimental/MeshModelingToolset/Source/ModelingComponents/Private/OctreeDynamicMeshComponent.cpp
new file mode 100644
--- /dev/null
+++ b/Engine/Plugins/Experimental/MeshModelingToolset/Source/ModelingComponents/Private/OctreeDynamicMeshComponent.cpp
@@ -0,0 +1,380 @@
+// Copyright 1998-2019 Epic Games, Inc. All Rights Reserved. 
+
+#include "OctreeDynamicMeshComponent.h"
+#include "RenderingThread.h"
+#include "RenderResource.h"
+#include "PrimitiveViewRelevance.h"
+#include "PrimitiveSceneProxy.h"
+#include "VertexFactory.h"
+#include "MaterialShared.h"
+#include "Engine/CollisionProfile.h"
+#include "Materials/Material.h"
+#include "LocalVertexFactory.h"
+#include "SceneManagement.h"
+#include "DynamicMeshBuilder.h"
+#include "EngineGlobals.h"
+#include "Engine/Engine.h"
+#include "StaticMeshResources.h"
+
+#include "MeshDescription.h"
+#include "MeshAttributes.h"
+
+#include "DynamicMeshAttributeSet.h"
+#include "MeshNormals.h"
+#include "MeshDescriptionToDynamicMesh.h"
+#include "DynamicMeshToMeshDescription.h"
+
+#include "MeshVertexChange.h"
+#include "Changes/MeshChange.h"
+#include "DynamicMeshChangeTracker.h"
+
+
+
+// default proxy for this component
+#include "OctreeDynamicMeshSceneProxy.h"
+
+
+
+
+UOctreeDynamicMeshComponent::UOctreeDynamicMeshComponent(const FObjectInitializer& ObjectInitializer)
+	: Super(ObjectInitializer)
+{
+	PrimaryComponentTick.bCanEverTick = false;
+
+	SetCollisionProfileName(UCollisionProfile::NoCollision_ProfileName);
+
+	InitializeNewMesh();
+}
+
+
+
+void UOctreeDynamicMeshComponent::InitializeMesh(FMeshDescription* MeshDescription)
+{
+	FMeshDescriptionToDynamicMesh Converter;
+	Converter.bPrintDebugMessages = true;
+	Mesh->Clear();
+	Converter.Convert(MeshDescription, *Mesh);
+
+	FAxisAlignedBox3d MeshBounds = Mesh->GetCachedBounds();
+	Octree = MakeUnique<FDynamicMeshOctree3>();
+	Octree->RootDimension = MeshBounds.MaxDim() * 0.5;
+
+	Octree->Initialize(GetMesh());
+	
+	FDynamicMeshOctree3::FStatistics Stats;
+	Octree->ComputeStatistics(Stats);
+	UE_LOG(LogTemp, Warning, TEXT("OctreeStats %s"), *Stats.ToString());
+
+	OctreeCut = MakeUnique<FDynamicMeshOctree3::FTreeCutSet>();
+
+	NotifyMeshUpdated();
+}
+
+
+TUniquePtr<FDynamicMesh3> UOctreeDynamicMeshComponent::ExtractMesh(bool bNotifyUpdate)
+{
+	TUniquePtr<FDynamicMesh3> CurMesh = MoveTemp(Mesh);
+	InitializeNewMesh();
+	if (bNotifyUpdate)
+	{
+		NotifyMeshUpdated();
+	}
+	return CurMesh;
+}
+
+
+void UOctreeDynamicMeshComponent::InitializeNewMesh()
+{
+	Mesh = MakeUnique<FDynamicMesh3>();
+	// discard any attributes/etc initialized by default
+	Mesh->Clear();
+	Octree = MakeUnique<FDynamicMeshOctree3>();
+	Octree->Initialize(GetMesh());
+	OctreeCut = MakeUnique<FDynamicMeshOctree3::FTreeCutSet>();
+}
+
+
+void UOctreeDynamicMeshComponent::Bake(FMeshDescription* MeshDescription, bool bHaveModifiedToplogy)
+{
+	if (bHaveModifiedToplogy == false)
+	{
+		FDynamicMeshToMeshDescription Converter;
+		Converter.Update(Mesh.Get(), *MeshDescription);
+	}
+	else
+	{
+		FDynamicMeshToMeshDescription Converter;
+		Converter.Convert(Mesh.Get(), *MeshDescription);
+
+		//UE_LOG(LogTemp, Warning, TEXT("MeshDescription has %d instances"), MeshDescription->VertexInstances().Num());
+	}
+}
+
+
+
+void UOctreeDynamicMeshComponent::NotifyMeshUpdated()
+{
+	if (CurrentProxy != nullptr)
+	{
+		FAxisAlignedBox3d DirtyBox = Octree->ModifiedBounds;
+		Octree->ResetModifiedBounds();
+
+		// update existing cells
+		int TotalTriangleCount = 0;
+		int SpillTriangleCount = 0;
+		TArray<int32> SetsToUpdate;
+
+		{
+			SCOPE_CYCLE_COUNTER(STAT_SculptToolOctree_UpdateExisting);
+			FCriticalSection SetsLock;
+			int NumCutCells = CutCellSetMap.Num();
+			ParallelFor(NumCutCells, [&](int i) 
+			{
+				const FCutCellIndexSet& CutCellSet = CutCellSetMap[i];
+
+				if (Octree->TestCellIntersection(CutCellSet.CellRef, DirtyBox) == false)
+				{
+					return;
+				}
+
+				TArray<int32>& TriangleSet = TriangleDecomposition.GetIndexSetArray(CutCellSet.DecompSetID);
+				TriangleSet.Reset();
+				Octree->CollectTriangles(CutCellSet.CellRef, [&TriangleSet](int TriangleID) {
+					TriangleSet.Add(TriangleID);
+				});
+
+				SetsLock.Lock();
+				TotalTriangleCount += TriangleSet.Num();
+				SetsToUpdate.Add(CutCellSet.DecompSetID);
+				SetsLock.Unlock();
+			}, false);
+		}
+
+		// update cut set to find new cells
+		TArray<FDynamicMeshOctree3::FCellReference> NewCutCells;
+		{
+			SCOPE_CYCLE_COUNTER(STAT_SculptToolOctree_UpdateCutSet);
+			Octree->UpdateLevelCutSet(*OctreeCut, NewCutCells);
+		}
+
+		// add new ones
+		{
+			SCOPE_CYCLE_COUNTER(STAT_SculptToolOctree_CreateNew);
+			for (const FDynamicMeshOctree3::FCellReference& CellRef : NewCutCells)
+			{
+
+				int32 IndexSetID = TriangleDecomposition.CreateNewIndexSet();
+				TArray<int32>& TriangleSet = TriangleDecomposition.GetIndexSetArray(IndexSetID);
+				Octree->CollectTriangles(CellRef, [&TriangleSet](int TriangleID) {
+					TriangleSet.Add(TriangleID);
+				});
+				TotalTriangleCount += TriangleSet.Num();
+				FCutCellIndexSet SetMapEntry = { CellRef, IndexSetID };
+				CutCellSetMap.Add(SetMapEntry);
+				SetsToUpdate.Add(IndexSetID);
+			}
+		}
+
+		// rebuild spill set (always for now...should track bounds? and do separately for each root cell?)
+		{
+			SCOPE_CYCLE_COUNTER(STAT_SculptToolOctree_UpdateSpill);
+			TArray<int32>& SpillTriangleSet = TriangleDecomposition.GetIndexSetArray(SpillDecompSetID);
+			SpillTriangleSet.Reset();
+			Octree->CollectRootTriangles(*OctreeCut,
+				[&SpillTriangleSet](int TriangleID) {
+				SpillTriangleSet.Add(TriangleID);
+			});
+			Octree->CollectSpillTriangles(
+				[&SpillTriangleSet](int TriangleID) {
+				SpillTriangleSet.Add(TriangleID);
+			});
+			TotalTriangleCount += SpillTriangleSet.Num();
+			SpillTriangleCount += SpillTriangleSet.Num();
+			SetsToUpdate.Add(SpillDecompSetID);
+		}
+
+
+		UE_LOG(LogTemp, Warning, TEXT("Updating %d of %d decomposition sets, %d tris total, spillcount %d"), SetsToUpdate.Num(), CutCellSetMap.Num(), TotalTriangleCount, SpillTriangleCount);
+		{
+			SCOPE_CYCLE_COUNTER(STAT_SculptToolOctree_UpdateFromDecomp);
+			CurrentProxy->UpdateFromDecomposition(TriangleDecomposition, SetsToUpdate);
+		}
+
+	}
+
+	//MarkRenderStateDirty();
+	//UpdateBounds();
+	//CurrentProxy = nullptr;
+}
+
+
+
+
+static void InitializeOctreeCutSet(const FDynamicMesh3& Mesh, const FDynamicMeshOctree3& Octree, TUniquePtr<FDynamicMeshOctree3::FTreeCutSet>& CutSet)
+{
+	int TriangleCount = Mesh.TriangleCount();
+	if (TriangleCount < 50000)
+	{
+		*CutSet = Octree.BuildLevelCutSet(1);
+		return;
+	}
+
+	FDynamicMeshOctree3::FStatistics Stats;
+	Octree.ComputeStatistics(Stats);
+	int MaxLevel = Stats.Levels;
+
+	int CutLevel = 0;
+	while ( CutLevel < MaxLevel-1 && Stats.LevelBoxCounts[CutLevel] < 200 && Stats.LevelBoxCounts[CutLevel+1] < 300)
+	{
+		CutLevel++;
+	}
+	*CutSet = Octree.BuildLevelCutSet(CutLevel);
+	//UE_LOG(LogTemp, Warning, TEXT("InitializeOctreeCutSet - level %d cells %d"), CutLevel, CutSet->CutCells.Num());
+}
+
+
+
+
+FPrimitiveSceneProxy* UOctreeDynamicMeshComponent::CreateSceneProxy()
+{
+	CurrentProxy = nullptr;
+	if (Mesh->TriangleCount() > 0)
+	{
+		CurrentProxy = new FOctreeDynamicMeshSceneProxy(this);
+
+		if (TriangleColorFunc != nullptr)
+		{
+			CurrentProxy->bUsePerTriangleColor = true;
+			CurrentProxy->PerTriangleColorFunc = [this](int TriangleID) { return GetTriangleColor(TriangleID); };
+		}
+
+
+		OctreeCut = MakeUnique<FDynamicMeshOctree3::FTreeCutSet>();
+		InitializeOctreeCutSet(*Mesh, *Octree, OctreeCut);
+
+		TriangleDecomposition = FArrayIndexSetsDecomposition();
+
+		SpillDecompSetID = TriangleDecomposition.CreateNewIndexSet();
+
+		CutCellSetMap.Reset();
+		for (auto CellRef : OctreeCut->CutCells)
+		{
+			int32 IndexSetID = TriangleDecomposition.CreateNewIndexSet();
+			TArray<int32>& TriangleSet = TriangleDecomposition.GetIndexSetArray(IndexSetID);
+			Octree->CollectTriangles(CellRef, [&TriangleSet](int TriangleID) {
+				TriangleSet.Add(TriangleID);
+			});
+
+			FCutCellIndexSet SetMapEntry = { CellRef, IndexSetID };
+			CutCellSetMap.Add(SetMapEntry);
+		}
+
+		// collect spill triangles
+		{
+			TArray<int32>& SpillTriangleSet = TriangleDecomposition.GetIndexSetArray(SpillDecompSetID);
+			Octree->CollectRootTriangles(*OctreeCut, 
+				[&SpillTriangleSet](int TriangleID) {
+					SpillTriangleSet.Add(TriangleID);
+			});
+			Octree->CollectSpillTriangles(
+				[&SpillTriangleSet](int TriangleID) {
+				SpillTriangleSet.Add(TriangleID);
+			});
+		}
+
+		CurrentProxy->InitializeFromDecomposition(TriangleDecomposition);
+	}
+	return CurrentProxy;
+}
+
+int32 UOctreeDynamicMeshComponent::GetNumMaterials() const
+{
+	return 1;
+}
+
+
+FColor UOctreeDynamicMeshComponent::GetTriangleColor(int TriangleID)
+{
+	if (TriangleColorFunc != nullptr)
+	{
+		return TriangleColorFunc(TriangleID);
+	}
+	else
+	{
+		return (TriangleID % 2 == 0) ? FColor::Red : FColor::White;
+	}
+}
+
+
+
+FBoxSphereBounds UOctreeDynamicMeshComponent::CalcBounds(const FTransform& LocalToWorld) const
+{
+	// Bounds are tighter if the box is generated from pre-transformed vertices.
+	FBox BoundingBox(ForceInit);
+	for ( FVector3d Vertex : Mesh->VerticesItr() ) 
+	{
+		BoundingBox += LocalToWorld.TransformPosition(Vertex);
+	}
+
+	FBoxSphereBounds NewBounds;
+	NewBounds.BoxExtent = BoundingBox.GetExtent();
+	NewBounds.Origin = BoundingBox.GetCenter();
+	NewBounds.SphereRadius = NewBounds.BoxExtent.Size();
+
+	return NewBounds;
+}
+
+
+
+void UOctreeDynamicMeshComponent::ApplyChange(const FMeshVertexChange* Change, bool bRevert)
+{
+	int NV = Change->Vertices.Num();
+	const TArray<FVector3d>& Positions = (bRevert) ? Change->OldPositions : Change->NewPositions;
+	
+	Octree->ResetModifiedBounds();
+	TSet<int> TrianglesToUpdate;
+
+	for (int k = 0; k < NV; ++k)
+	{
+		int vid = Change->Vertices[k];
+
+		for (int tid : Mesh->VtxTrianglesItr(vid))
+		{
+			if (TrianglesToUpdate.Contains(tid) == false)
+			{
+				Octree->NotifyPendingModification(tid);
+				TrianglesToUpdate.Add(tid);
+			}
+		}
+
+		Mesh->SetVertex(vid, Positions[k]);
+	}
+
+	Octree->ReinsertTriangles(TrianglesToUpdate);
+
+	//NotifyMeshUpdated();
+
+	OnMeshChanged.Broadcast();
+}
+
+
+
+
+void UOctreeDynamicMeshComponent::ApplyChange(const FMeshChange* Change, bool bRevert)
+{
+	TArray<int> RemoveTriangles, AddTriangles;
+	bool bRemoveOld = ! bRevert;
+	Change->DynamicMeshChange->GetSavedTriangleList(RemoveTriangles, bRemoveOld);
+	Change->DynamicMeshChange->GetSavedTriangleList(AddTriangles, !bRemoveOld);
+
+	Octree->ResetModifiedBounds();
+	Octree->RemoveTriangles(RemoveTriangles);
+
+	Change->DynamicMeshChange->Apply(Mesh.Get(), bRevert);
+
+	Octree->InsertTriangles(AddTriangles);
+
+	//NotifyMeshUpdated();
+
+	OnMeshChanged.Broadcast();
+}
\ No newline at end of file
diff --git a/Engine/Plugins/Experimental/MeshModelingToolset/Source/ModelingComponents/Private/OctreeDynamicMeshSceneProxy.h b/Engine/Plugins/Experimental/MeshModelingToolset/Source/ModelingComponents/Private/OctreeDynamicMeshSceneProxy.h
new file mode 100644
--- /dev/null
+++ b/Engine/Plugins/Experimental/MeshModelingToolset/Source/ModelingComponents/Private/OctreeDynamicMeshSceneProxy.h
@@ -0,0 +1,540 @@
+// Copyright 1998-2019 Epic Games, Inc. All Rights Reserved. 
+
+#pragma once
+
+#include "OctreeDynamicMeshComponent.h"
+#include "Util/IndexSetDecompositions.h"
+
+
+DECLARE_STATS_GROUP(TEXT("SculptToolOctree"), STATGROUP_SculptToolOctree, STATCAT_Advanced);
+DECLARE_CYCLE_STAT(TEXT("SculptToolOctree_UpdateExisting"), STAT_SculptToolOctree_UpdateExisting, STATGROUP_SculptToolOctree);
+DECLARE_CYCLE_STAT(TEXT("SculptToolOctree_UpdateCutSet"), STAT_SculptToolOctree_UpdateCutSet, STATGROUP_SculptToolOctree);
+DECLARE_CYCLE_STAT(TEXT("SculptToolOctree_CreateNew"), STAT_SculptToolOctree_CreateNew, STATGROUP_SculptToolOctree);
+DECLARE_CYCLE_STAT(TEXT("SculptToolOctree_UpdateSpill"), STAT_SculptToolOctree_UpdateSpill, STATGROUP_SculptToolOctree);
+DECLARE_CYCLE_STAT(TEXT("SculptToolOctree_UpdateFromDecomp"), STAT_SculptToolOctree_UpdateFromDecomp, STATGROUP_SculptToolOctree);
+DECLARE_CYCLE_STAT(TEXT("SculptToolOctree_UpdateDecompDestroy"), STAT_SculptToolOctree_UpdateDecompDestroy, STATGROUP_SculptToolOctree);
+DECLARE_CYCLE_STAT(TEXT("SculptToolOctree_UpdateDecompCreate"), STAT_SculptToolOctree_UpdateDecompCreate, STATGROUP_SculptToolOctree);
+DECLARE_CYCLE_STAT(TEXT("SculptToolOctree_InitializeBufferFromOverlay"), STAT_SculptToolOctree_InitializeBufferFromOverlay, STATGROUP_SculptToolOctree);
+DECLARE_CYCLE_STAT(TEXT("SculptToolOctree_BufferUpload"), STAT_SculptToolOctree_BufferUpload, STATGROUP_SculptToolOctree);
+
+class FMeshRenderBufferSet
+{
+public:
+	int TriangleCount = 0;
+
+	/** The buffer containing vertex data. */
+	FStaticMeshVertexBuffer StaticMeshVertexBuffer;
+	/** The buffer containing the position vertex data. */
+	FPositionVertexBuffer PositionVertexBuffer;
+	/** The buffer containing the vertex color data. */
+	FColorVertexBuffer ColorVertexBuffer;
+
+	FDynamicMeshIndexBuffer32 IndexBuffer;
+	FLocalVertexFactory VertexFactory;
+
+	FMeshRenderBufferSet(ERHIFeatureLevel::Type FeatureLevelType)
+		: VertexFactory(FeatureLevelType, "FMeshRenderBufferSet")
+	{
+
+	}
+
+
+	virtual ~FMeshRenderBufferSet()
+	{
+		check(IsInRenderingThread());
+
+		if (TriangleCount > 0)
+		{
+			PositionVertexBuffer.ReleaseResource();
+			StaticMeshVertexBuffer.ReleaseResource();
+			ColorVertexBuffer.ReleaseResource();
+			IndexBuffer.ReleaseResource();
+			VertexFactory.ReleaseResource();
+		}
+	}
+
+
+	static void DestroyRenderBufferSet(FMeshRenderBufferSet* BufferSet)
+	{
+		if (BufferSet->TriangleCount == 0)
+		{
+			return;
+		}
+
+		ENQUEUE_RENDER_COMMAND(FMeshRenderBufferSetDestroy)(
+			[BufferSet](FRHICommandListImmediate& RHICmdList)
+		{
+			delete BufferSet;
+		});
+	}
+
+	void Upload()
+	{
+		check(IsInRenderingThread());
+
+		if (TriangleCount == 0)
+		{
+			return;
+		}
+
+		InitOrUpdateResource(&this->PositionVertexBuffer);
+		InitOrUpdateResource(&this->StaticMeshVertexBuffer);
+		InitOrUpdateResource(&this->ColorVertexBuffer);
+
+		FLocalVertexFactory::FDataType Data;
+		this->PositionVertexBuffer.BindPositionVertexBuffer(&this->VertexFactory, Data);
+		this->StaticMeshVertexBuffer.BindTangentVertexBuffer(&this->VertexFactory, Data);
+		this->StaticMeshVertexBuffer.BindPackedTexCoordVertexBuffer(&this->VertexFactory, Data);
+		//this->StaticMeshVertexBuffer.BindLightMapVertexBuffer(&this->VertexFactory, Data, LightMapIndex);
+		this->ColorVertexBuffer.BindColorVertexBuffer(&this->VertexFactory, Data);
+		this->VertexFactory.SetData(Data);
+
+		InitOrUpdateResource(&this->VertexFactory);
+		PositionVertexBuffer.InitResource();
+		StaticMeshVertexBuffer.InitResource();
+		ColorVertexBuffer.InitResource();
+		IndexBuffer.InitResource();
+		VertexFactory.InitResource();
+	}
+
+	// copied from StaticMesh.cpp
+	void InitOrUpdateResource(FRenderResource* Resource)
+	{
+		check(IsInRenderingThread());
+
+		if (!Resource->IsInitialized())
+		{
+			Resource->InitResource();
+		}
+		else
+		{
+			Resource->UpdateRHI();
+		}
+	}
+
+};
+
+/**
+ * Scene Proxy for a mesh buffer.
+ * 
+ * Based on FProceduralMeshSceneProxy but simplified in various ways.
+ * 
+ * Supports wireframe-on-shaded rendering.
+ * 
+ */
+class FOctreeDynamicMeshSceneProxy final : public FPrimitiveSceneProxy
+{
+private:
+	UMaterialInterface* Material;
+
+	FMaterialRelevance MaterialRelevance;
+
+	TMap<int32, FMeshRenderBufferSet*> RenderBufferSets;
+
+public:
+	/** Component that created this proxy (is there a way to look this up?) */
+	UOctreeDynamicMeshComponent* ParentComponent;
+
+
+	FColor ConstantVertexColor = FColor::White;
+	bool bIgnoreVertexColors = false;
+	bool bIgnoreVertexNormals = false;
+
+
+	bool bUsePerTriangleColor = false;
+	TFunction<FColor(int)> PerTriangleColorFunc = nullptr;
+
+
+	FOctreeDynamicMeshSceneProxy(UOctreeDynamicMeshComponent* Component)
+		: FPrimitiveSceneProxy(Component)
+		, MaterialRelevance(Component->GetMaterialRelevance(GetScene().GetFeatureLevel()))
+	{
+		check(!IsInRenderingThread());
+
+		ParentComponent = Component;
+
+		// Grab material
+		Material = Component->GetMaterial(0);
+		if (Material == NULL)
+		{
+			Material = UMaterial::GetDefaultMaterial(MD_Surface);
+		}
+	}
+
+
+	virtual ~FOctreeDynamicMeshSceneProxy()
+	{
+		check(IsInRenderingThread());
+
+		for (auto MapPair : RenderBufferSets)
+		{
+			FMeshRenderBufferSet* BufferSet = MapPair.Value;
+			FMeshRenderBufferSet::DestroyRenderBufferSet(BufferSet);
+		}
+		RenderBufferSets.Reset();
+	}
+
+
+	FMeshRenderBufferSet* AllocateNewRenderBufferSet()
+	{
+		FMeshRenderBufferSet* RenderBufferSet = new FMeshRenderBufferSet(GetScene().GetFeatureLevel());
+		return RenderBufferSet;
+	}
+
+
+
+	virtual void InitializeSingleBuffer()
+	{
+			check(RenderBufferSets.Num() == 0);
+
+			FDynamicMesh3* Mesh = ParentComponent->GetMesh();
+
+			FMeshRenderBufferSet* RenderBuffers = AllocateNewRenderBufferSet();
+
+			// find suitable overlays
+			FDynamicMeshUVOverlay* UVOverlay = nullptr;
+			FDynamicMeshNormalOverlay* NormalOverlay = nullptr;
+			if (Mesh->HasAttributes())
+			{
+				UVOverlay = Mesh->Attributes()->PrimaryUV();
+				NormalOverlay = Mesh->Attributes()->PrimaryNormals();
+			}
+
+			InitializeBuffersFromOverlays(Mesh,
+				Mesh->TriangleCount(), Mesh->TriangleIndicesItr(),
+				UVOverlay, NormalOverlay, RenderBuffers);
+
+			ENQUEUE_RENDER_COMMAND(FOctreeDynamicMeshSceneProxyInitializeSingle)(
+				[this, RenderBuffers](FRHICommandListImmediate& RHICmdList)
+			{
+				RenderBuffers->Upload();
+				RenderBufferSets.Add(0, RenderBuffers);
+			});
+	}
+
+
+
+	virtual void InitializeFromDecomposition(const FArrayIndexSetsDecomposition& Decomposition)
+	{
+		check(RenderBufferSets.Num() == 0);
+
+		FDynamicMesh3* Mesh = ParentComponent->GetMesh();
+
+		// find suitable overlays
+		FDynamicMeshUVOverlay* UVOverlay = nullptr;
+		FDynamicMeshNormalOverlay* NormalOverlay = nullptr;
+		if (Mesh->HasAttributes())
+		{
+			UVOverlay = Mesh->Attributes()->PrimaryUV();
+			NormalOverlay = Mesh->Attributes()->PrimaryNormals();
+		}
+
+		const TArray<int32>& SetIDs = Decomposition.GetIndexSetIDs();
+		for (int32 SetID : SetIDs)
+		{
+			const TArray<int32>& Tris = Decomposition.GetIndexSetArray(SetID);
+			
+			FMeshRenderBufferSet* RenderBuffers = AllocateNewRenderBufferSet();
+
+			InitializeBuffersFromOverlays(Mesh,
+				Tris.Num(), Tris,
+				UVOverlay, NormalOverlay, RenderBuffers);
+
+			ENQUEUE_RENDER_COMMAND(FOctreeDynamicMeshSceneProxyInitializeFromDecomposition)(
+				[this, SetID, RenderBuffers](FRHICommandListImmediate& RHICmdList)
+			{
+				RenderBuffers->Upload();
+				RenderBufferSets.Add(SetID, RenderBuffers);
+			});
+		}
+	}
+
+
+
+	virtual void UpdateFromDecomposition(const FArrayIndexSetsDecomposition& Decomposition, const TArray<int32>& SetsToUpdate )
+	{
+		// CAN WE REUSE EXISTING BUFFER SETS??
+		//   - could have timestamp for each decomposition set array...if tris don't change we only have to update vertices
+		//   - can re-use allocated memory if new data is smaller
+
+		SCOPE_CYCLE_COUNTER(STAT_SculptToolOctree_UpdateFromDecomp);
+
+		// remove sets to update
+		ENQUEUE_RENDER_COMMAND(FOctreeDynamicMeshSceneProxyUpdatePreClean)(
+			[this, SetsToUpdate](FRHICommandListImmediate& RHICmdList)
+		{
+			SCOPE_CYCLE_COUNTER(STAT_SculptToolOctree_UpdateDecompDestroy);
+			for (int32 SetID : SetsToUpdate)
+			{
+				if (RenderBufferSets.Contains(SetID))
+				{
+					FMeshRenderBufferSet* BufferSet = RenderBufferSets.FindAndRemoveChecked(SetID);
+					FMeshRenderBufferSet::DestroyRenderBufferSet(BufferSet);
+				}
+			}
+		});
+
+		FDynamicMesh3* Mesh = ParentComponent->GetMesh();
+
+		// find suitable overlays
+		FDynamicMeshUVOverlay* UVOverlay = nullptr;
+		FDynamicMeshNormalOverlay* NormalOverlay = nullptr;
+		if (Mesh->HasAttributes())
+		{
+			UVOverlay = Mesh->Attributes()->PrimaryUV();
+			NormalOverlay = Mesh->Attributes()->PrimaryNormals();
+		}
+
+		{
+			FCriticalSection SectionLock;
+
+			SCOPE_CYCLE_COUNTER(STAT_SculptToolOctree_UpdateDecompCreate);
+			int NumSets = SetsToUpdate.Num();
+			ParallelFor(NumSets, [&](int k)
+			{
+				int32 SetID = SetsToUpdate[k];
+				const TArray<int32>& Tris = Decomposition.GetIndexSetArray(SetID);
+
+				FMeshRenderBufferSet* RenderBuffers = AllocateNewRenderBufferSet();
+
+				InitializeBuffersFromOverlays(Mesh,
+					Tris.Num(), Tris,
+					UVOverlay, NormalOverlay, RenderBuffers);
+
+				ENQUEUE_RENDER_COMMAND(FOctreeDynamicMeshSceneProxyUpdateAddOne)(
+					[this, SetID, RenderBuffers](FRHICommandListImmediate& RHICmdList)
+				{
+					SCOPE_CYCLE_COUNTER(STAT_SculptToolOctree_BufferUpload);
+					RenderBuffers->Upload();
+					RenderBufferSets.Add(SetID, RenderBuffers);
+				});
+			}, false);
+
+
+		}
+
+		UE_LOG(LogTemp, Warning, TEXT("Have %d renderbuffers"), RenderBufferSets.Num());
+	}
+
+
+
+protected:
+
+
+	/**
+	 * Initialize rendering buffers from given attribute overlays.
+	 * Creates three vertices per triangle, IE no shared vertices in buffers.
+	 */
+	template<typename TriangleEnumerable>
+	void InitializeBuffersFromOverlays(const FDynamicMesh3* Mesh, 
+		int NumTriangles, TriangleEnumerable Enumerable,
+		FDynamicMeshUVOverlay* UVOverlay, 
+		FDynamicMeshNormalOverlay* NormalOverlay,
+		FMeshRenderBufferSet* RenderBuffers)
+	{
+		SCOPE_CYCLE_COUNTER(STAT_SculptToolOctree_InitializeBufferFromOverlay);
+
+		RenderBuffers->TriangleCount = NumTriangles;
+		if (NumTriangles == 0)
+		{
+			return;
+		}
+
+		bool bHaveColors = Mesh->HasVertexColors() && (bIgnoreVertexColors == false);
+
+		int NumVertices = NumTriangles * 3;
+		int NumTexCoords = 1;		// no! zero!
+
+		{
+			RenderBuffers->PositionVertexBuffer.Init(NumVertices);
+			RenderBuffers->StaticMeshVertexBuffer.Init(NumVertices, NumTexCoords);
+			RenderBuffers->ColorVertexBuffer.Init(NumVertices);
+			RenderBuffers->IndexBuffer.Indices.AddUninitialized(NumTriangles * 3);
+		}
+
+		int TriIdx = 0, VertIdx = 0;
+		FVector3f TangentX, TangentY;
+		for (int TriangleID : Enumerable)
+		{
+			FIndex3i Tri = Mesh->GetTriangle(TriangleID);
+			FIndex3i TriUV = (UVOverlay != nullptr) ? UVOverlay->GetTriangle(TriangleID) : FIndex3i::Zero();
+			FIndex3i TriNormal = (NormalOverlay != nullptr) ? NormalOverlay->GetTriangle(TriangleID) : FIndex3i::Zero();
+
+			FColor TriColor = ConstantVertexColor;
+			if (bUsePerTriangleColor && PerTriangleColorFunc != nullptr)
+			{
+				TriColor = PerTriangleColorFunc(TriangleID);
+				bHaveColors = false;
+			}
+
+			for (int j = 0; j < 3; ++j)
+			{
+				RenderBuffers->PositionVertexBuffer.VertexPosition(VertIdx) = (FVector)Mesh->GetVertex(Tri[j]);
+
+				FVector3f Normal = (NormalOverlay != nullptr && TriNormal[j] != FDynamicMesh3::InvalidID) ? 
+					NormalOverlay->GetElement(TriNormal[j]) : Mesh->GetVertexNormal(Tri[j]);
+
+				// calculate a nonsense tangent
+				VectorUtil::MakePerpVectors(Normal, TangentX, TangentY);
+				RenderBuffers->StaticMeshVertexBuffer.SetVertexTangents(VertIdx, TangentX, TangentY, Normal);
+
+				FVector2f UV = (UVOverlay != nullptr && TriUV[j] != FDynamicMesh3::InvalidID) ? 
+					UVOverlay->GetElement(TriUV[j]) : FVector2f::Zero();
+				RenderBuffers->StaticMeshVertexBuffer.SetVertexUV(VertIdx, 0, UV);
+
+				RenderBuffers->ColorVertexBuffer.VertexColor(VertIdx) = (bHaveColors) ?
+					(FColor)Mesh->GetVertexColor(Tri[j]) : TriColor;
+
+				RenderBuffers->IndexBuffer.Indices[TriIdx++] = VertIdx;
+				VertIdx++;
+			}
+		}
+	}
+
+
+
+
+
+public:
+
+	virtual void GetDynamicMeshElements(const TArray<const FSceneView*>& Views, const FSceneViewFamily& ViewFamily, uint32 VisibilityMap, FMeshElementCollector& Collector) const override
+	{
+		QUICK_SCOPE_CYCLE_COUNTER(STAT_OctreeDynamicMeshSceneProxy_GetDynamicMeshElements);
+
+		const bool bWireframe = (AllowDebugViewmodes() && ViewFamily.EngineShowFlags.Wireframe)
+			|| ParentComponent->bExplicitShowWireframe;
+
+		FColoredMaterialRenderProxy* WireframeMaterialInstance = new FColoredMaterialRenderProxy(
+			GEngine->WireframeMaterial ? GEngine->WireframeMaterial->GetRenderProxy() : nullptr,
+			FLinearColor(0, 0.5f, 1.f)
+		);
+
+		Collector.RegisterOneFrameMaterialProxy(WireframeMaterialInstance);
+
+		FMaterialRenderProxy* MaterialProxy = Material->GetRenderProxy();
+		FMaterialRenderProxy* WireframeMaterialProxy = WireframeMaterialInstance;
+
+		//ESceneDepthPriorityGroup DepthPriority = (ParentComponent->bDrawOnTop) ? SDPG_Foreground : SDPG_World;
+		ESceneDepthPriorityGroup DepthPriority = SDPG_World;
+
+
+		for (int32 ViewIndex = 0; ViewIndex < Views.Num(); ViewIndex++)
+		{
+			if (VisibilityMap & (1 << ViewIndex))
+			{
+				const FSceneView* View = Views[ViewIndex];
+
+				bool bHasPrecomputedVolumetricLightmap;
+				FMatrix PreviousLocalToWorld;
+				int32 SingleCaptureIndex;
+				bool bOutputVelocity;
+				GetScene().GetPrimitiveUniformShaderParameters_RenderThread(GetPrimitiveSceneInfo(), bHasPrecomputedVolumetricLightmap, PreviousLocalToWorld, SingleCaptureIndex, bOutputVelocity);
+
+				// Draw the mesh.
+				for (auto MapPair : RenderBufferSets)
+				{
+					const FMeshRenderBufferSet& RenderBuffers = *MapPair.Value;
+					if (RenderBuffers.TriangleCount == 0)
+					{
+						continue;
+					}
+
+					// do we need separate one of these for each MeshRenderBufferSet?
+					FDynamicPrimitiveUniformBuffer& DynamicPrimitiveUniformBuffer = Collector.AllocateOneFrameResource<FDynamicPrimitiveUniformBuffer>();
+					DynamicPrimitiveUniformBuffer.Set(GetLocalToWorld(), PreviousLocalToWorld, GetBounds(), GetLocalBounds(), true, bHasPrecomputedVolumetricLightmap, DrawsVelocity(), bOutputVelocity);
+
+					DrawBatch(Collector, RenderBuffers, MaterialProxy, false, DepthPriority, ViewIndex, DynamicPrimitiveUniformBuffer);
+					if (bWireframe)
+					{
+						DrawBatch(Collector, RenderBuffers, WireframeMaterialProxy, true, DepthPriority, ViewIndex, DynamicPrimitiveUniformBuffer);
+					}
+				}
+			}
+		}
+	}
+
+
+
+	virtual void DrawBatch(FMeshElementCollector& Collector, 
+		const FMeshRenderBufferSet& RenderBuffers,
+		FMaterialRenderProxy* UseMaterial, 
+		bool bWireframe,
+		ESceneDepthPriorityGroup DepthPriority,
+		int ViewIndex,
+		FDynamicPrimitiveUniformBuffer& DynamicPrimitiveUniformBuffer) const
+	{
+		FMeshBatch& Mesh = Collector.AllocateMesh();
+		FMeshBatchElement& BatchElement = Mesh.Elements[0];
+		BatchElement.IndexBuffer = &RenderBuffers.IndexBuffer;
+		Mesh.bWireframe = bWireframe;
+		Mesh.VertexFactory = &RenderBuffers.VertexFactory;
+		Mesh.MaterialRenderProxy = UseMaterial;
+
+		BatchElement.PrimitiveUniformBufferResource = &DynamicPrimitiveUniformBuffer.UniformBuffer;
+
+		BatchElement.FirstIndex = 0;
+		BatchElement.NumPrimitives = RenderBuffers.IndexBuffer.Indices.Num() / 3;
+		BatchElement.MinVertexIndex = 0;
+		BatchElement.MaxVertexIndex = RenderBuffers.PositionVertexBuffer.GetNumVertices() - 1;
+		Mesh.ReverseCulling = IsLocalToWorldDeterminantNegative();
+		Mesh.Type = PT_TriangleList;
+		Mesh.DepthPriorityGroup = DepthPriority;
+		Mesh.bCanApplyViewModeOverrides = false;
+		Collector.AddMesh(ViewIndex, Mesh);
+	}
+
+
+
+	virtual FPrimitiveViewRelevance GetViewRelevance(const FSceneView* View) const override
+	{
+		FPrimitiveViewRelevance Result;
+
+		Result.bDrawRelevance = IsShown(View);
+		Result.bShadowRelevance = IsShadowCast(View);
+		Result.bDynamicRelevance = true;
+		Result.bRenderInMainPass = ShouldRenderInMainPass();
+		Result.bUsesLightingChannels = GetLightingChannelMask() != GetDefaultLightingChannelMask();
+		Result.bTranslucentSelfShadow = bCastVolumetricTranslucentShadow;
+		Result.bRenderCustomDepth = ShouldRenderCustomDepth();
+		MaterialRelevance.SetPrimitiveViewRelevance(Result);
+		Result.bVelocityRelevance = IsMovable() && Result.bOpaqueRelevance && Result.bRenderInMainPass;
+
+		return Result;
+	}
+
+
+	virtual void GetLightRelevance(const FLightSceneProxy* LightSceneProxy, bool& bDynamic, bool& bRelevant, bool& bLightMapped, bool& bShadowMapped) const override
+	{
+		FPrimitiveSceneProxy::GetLightRelevance(LightSceneProxy, bDynamic, bRelevant, bLightMapped, bShadowMapped);
+	}
+
+
+	virtual bool CanBeOccluded() const override
+	{
+		return !MaterialRelevance.bDisableDepthTest;
+	}
+
+	virtual uint32 GetMemoryFootprint(void) const override { return(sizeof(*this) + GetAllocatedSize()); }
+
+	uint32 GetAllocatedSize(void) const { return(FPrimitiveSceneProxy::GetAllocatedSize()); }
+
+
+	virtual void SetMaterial(UMaterialInterface* MaterialIn)
+	{
+		this->Material = MaterialIn;
+	}
+
+
+	SIZE_T GetTypeHash() const override
+	{
+		static size_t UniquePointer;
+		return reinterpret_cast<size_t>(&UniquePointer);
+	}
+};
+
+
+
+
+
+
+
diff --git a/Engine/Plugins/Experimental/MeshModelingToolset/Source/ModelingComponents/Private/SimpleDynamicMeshComponent.cpp b/Engine/Plugins/Experimental/MeshModelingToolset/Source/ModelingComponents/Private/SimpleDynamicMeshComponent.cpp
new file mode 100644
--- /dev/null
+++ b/Engine/Plugins/Experimental/MeshModelingToolset/Source/ModelingComponents/Private/SimpleDynamicMeshComponent.cpp
@@ -0,0 +1,270 @@
+// Copyright 1998-2019 Epic Games, Inc. All Rights Reserved. 
+
+#include "SimpleDynamicMeshComponent.h"
+#include "RenderingThread.h"
+#include "RenderResource.h"
+#include "PrimitiveViewRelevance.h"
+#include "PrimitiveSceneProxy.h"
+#include "VertexFactory.h"
+#include "MaterialShared.h"
+#include "Engine/CollisionProfile.h"
+#include "Materials/Material.h"
+#include "LocalVertexFactory.h"
+#include "SceneManagement.h"
+#include "DynamicMeshBuilder.h"
+#include "EngineGlobals.h"
+#include "Engine/Engine.h"
+#include "StaticMeshResources.h"
+
+#include "MeshDescription.h"
+#include "MeshAttributes.h"
+
+#include "DynamicMeshAttributeSet.h"
+#include "MeshNormals.h"
+#include "MeshDescriptionToDynamicMesh.h"
+#include "DynamicMeshToMeshDescription.h"
+
+#include "MeshVertexChange.h"
+#include "Changes/MeshChange.h"
+#include "DynamicMeshChangeTracker.h"
+
+
+
+// default proxy for this component
+#include "SimpleDynamicMeshSceneProxy.h"
+
+
+
+
+USimpleDynamicMeshComponent::USimpleDynamicMeshComponent(const FObjectInitializer& ObjectInitializer)
+	: Super(ObjectInitializer)
+{
+	PrimaryComponentTick.bCanEverTick = false;
+
+	SetCollisionProfileName(UCollisionProfile::NoCollision_ProfileName);
+
+	InitializeNewMesh();
+}
+
+
+
+void USimpleDynamicMeshComponent::InitializeMesh(FMeshDescription* MeshDescription)
+{
+	FMeshDescriptionToDynamicMesh Converter;
+	Converter.bPrintDebugMessages = true;
+	Mesh->Clear();
+	Converter.Convert(MeshDescription, *Mesh);
+	if (TangentsType == EDynamicMeshTangentCalcType::ExternallyCalculated)
+	{
+		Converter.CopyTangents(MeshDescription, Mesh.Get(), Tangents);
+	}
+
+	NotifyMeshUpdated();
+}
+
+
+TUniquePtr<FDynamicMesh3> USimpleDynamicMeshComponent::ExtractMesh(bool bNotifyUpdate)
+{
+	TUniquePtr<FDynamicMesh3> CurMesh = MoveTemp(Mesh);
+	InitializeNewMesh();
+	if (bNotifyUpdate)
+	{
+		NotifyMeshUpdated();
+	}
+	return CurMesh;
+}
+
+
+void USimpleDynamicMeshComponent::InitializeNewMesh()
+{
+	Mesh = MakeUnique<FDynamicMesh3>();
+	// discard any attributes/etc initialized by default
+	Mesh->Clear();
+
+	Tangents.SetMesh(Mesh.Get());
+}
+
+
+void USimpleDynamicMeshComponent::Bake(FMeshDescription* MeshDescription, bool bHaveModifiedToplogy)
+{
+	if (bHaveModifiedToplogy == false)
+	{
+		FDynamicMeshToMeshDescription Converter;
+		Converter.Update(Mesh.Get(), *MeshDescription);
+	}
+	else
+	{
+		FDynamicMeshToMeshDescription Converter;
+		Converter.Convert(Mesh.Get(), *MeshDescription);
+
+		//UE_LOG(LogTemp, Warning, TEXT("MeshDescription has %d instances"), MeshDescription->VertexInstances().Num());
+	}
+}
+
+
+
+
+
+
+FMeshTangentsf* USimpleDynamicMeshComponent::GetTangents()
+{
+	if (TangentsType == EDynamicMeshTangentCalcType::NoTangents)
+	{
+		return nullptr;
+	}
+	
+	if (TangentsType == EDynamicMeshTangentCalcType::AutoCalculated)
+	{
+		if (bTangentsValid == false && Mesh->HasAttributes())
+		{
+			FDynamicMeshUVOverlay* UVOverlay = Mesh->Attributes()->PrimaryUV();
+			FDynamicMeshNormalOverlay* NormalOverlay = Mesh->Attributes()->PrimaryNormals();
+			if (UVOverlay != nullptr && NormalOverlay != nullptr)
+			{
+				Tangents.ComputePerTriangleTangents(NormalOverlay, UVOverlay);
+				bTangentsValid = true;
+			}
+		}
+		return (bTangentsValid) ? &Tangents : nullptr;
+	}
+
+	// in this mode we assume the tangents are valid
+	check(TangentsType == EDynamicMeshTangentCalcType::ExternallyCalculated);
+	return &Tangents;
+}
+
+
+
+void USimpleDynamicMeshComponent::SetDrawOnTop(bool bSet)
+{
+	bDrawOnTop = bSet;
+	bUseEditorCompositing = bSet;
+}
+
+
+void USimpleDynamicMeshComponent::NotifyMeshUpdated()
+{
+	// Need to recreate scene proxy to send it over
+	MarkRenderStateDirty();
+	UpdateBounds();
+	CurrentProxy = nullptr;
+
+	if (TangentsType != EDynamicMeshTangentCalcType::ExternallyCalculated)
+	{
+		bTangentsValid = false;
+	}
+}
+
+
+void USimpleDynamicMeshComponent::FastNotifyColorsUpdated()
+{
+	if (CurrentProxy != nullptr)
+	{
+		CurrentProxy->FastUpdateColors();
+	}
+	else
+	{
+		NotifyMeshUpdated();
+	}
+}
+
+
+
+void USimpleDynamicMeshComponent::FastNotifyPositionsUpdated()
+{
+	if (CurrentProxy != nullptr)
+	{
+		bTangentsValid = false;
+		CurrentProxy->FastUpdatePositions(false);
+	}
+	else
+	{
+		NotifyMeshUpdated();
+	}
+}
+
+
+
+
+FPrimitiveSceneProxy* USimpleDynamicMeshComponent::CreateSceneProxy()
+{
+	CurrentProxy = nullptr;
+	if (Mesh->TriangleCount() > 0)
+	{
+		CurrentProxy = new FSimpleDynamicMeshSceneProxy(this);
+
+		if (TriangleColorFunc)
+		{
+			CurrentProxy->bUsePerTriangleColor = true;
+			CurrentProxy->PerTriangleColorFunc = [this](int TriangleID) { return GetTriangleColor(TriangleID); };
+		}
+
+		CurrentProxy->Initialize();
+	}
+	return CurrentProxy;
+}
+
+int32 USimpleDynamicMeshComponent::GetNumMaterials() const
+{
+	return 1;
+}
+
+
+FColor USimpleDynamicMeshComponent::GetTriangleColor(int TriangleID)
+{
+	if (TriangleColorFunc)
+	{
+		return TriangleColorFunc(TriangleID);
+	}
+	else
+	{
+		return (TriangleID % 2 == 0) ? FColor::Red : FColor::White;
+	}
+}
+
+
+
+FBoxSphereBounds USimpleDynamicMeshComponent::CalcBounds(const FTransform& LocalToWorld) const
+{
+	// Bounds are tighter if the box is generated from pre-transformed vertices.
+	FBox BoundingBox(ForceInit);
+	for ( FVector3d Vertex : Mesh->VerticesItr() ) 
+	{
+		BoundingBox += LocalToWorld.TransformPosition(Vertex);
+	}
+
+	FBoxSphereBounds NewBounds;
+	NewBounds.BoxExtent = BoundingBox.GetExtent();
+	NewBounds.Origin = BoundingBox.GetCenter();
+	NewBounds.SphereRadius = NewBounds.BoxExtent.Size();
+
+	return NewBounds;
+}
+
+
+
+void USimpleDynamicMeshComponent::ApplyChange(const FMeshVertexChange* Change, bool bRevert)
+{
+	int NV = Change->Vertices.Num();
+	const TArray<FVector3d>& Positions = (bRevert) ? Change->OldPositions : Change->NewPositions;
+	
+	for (int k = 0; k < NV; ++k)
+	{
+		int vid = Change->Vertices[k];
+		Mesh->SetVertex(vid, Positions[k]);
+	}
+
+	NotifyMeshUpdated();
+	OnMeshChanged.Broadcast();
+}
+
+
+
+
+void USimpleDynamicMeshComponent::ApplyChange(const FMeshChange* Change, bool bRevert)
+{
+	Change->DynamicMeshChange->Apply(Mesh.Get(), bRevert);
+
+	NotifyMeshUpdated();
+	OnMeshChanged.Broadcast();
+}
\ No newline at end of file
diff --git a/Engine/Plugins/Experimental/MeshModelingToolset/Source/ModelingComponents/Private/ToolIndicatorSet.cpp b/Engine/Plugins/Experimental/MeshModelingToolset/Source/ModelingComponents/Private/ToolIndicatorSet.cpp
new file mode 100644
--- /dev/null
+++ b/Engine/Plugins/Experimental/MeshModelingToolset/Source/ModelingComponents/Private/ToolIndicatorSet.cpp
@@ -0,0 +1,78 @@
+// Copyright 1998-2019 Epic Games, Inc. All Rights Reserved.
+
+
+#include "ToolIndicatorSet.h"
+
+UToolIndicatorSet::UToolIndicatorSet()
+{
+	Owner = nullptr;
+}
+
+
+UToolIndicatorSet::~UToolIndicatorSet()
+{
+	// if either of these is false, we did not disconnect properly
+	check(Indicators.Num() == 0);
+	check(Owner == nullptr);
+}
+
+
+void UToolIndicatorSet::Connect(UInteractiveTool* Tool)
+{
+	check(Tool != nullptr);
+	this->Owner = Tool;
+}
+
+
+void UToolIndicatorSet::Disconnect()
+{
+	for (UObject* object : Indicators)
+	{
+		IToolIndicator* Indicator = Cast<IToolIndicator>(object);
+		if (Indicator != nullptr)
+		{
+			Indicator->Disconnect();
+		}
+	}
+	Indicators.Reset();
+
+	this->Owner = nullptr;
+}
+
+
+
+void UToolIndicatorSet::AddIndicator(IToolIndicator* Indicator)
+{
+	check(Owner != nullptr);
+	check(Indicator != nullptr);
+
+	Indicator->Connect(Owner);
+
+	Indicators.Add( Cast<UObject>(Indicator) );
+}
+
+
+
+void UToolIndicatorSet::Render(IToolsContextRenderAPI* RenderAPI)
+{
+	for (UObject* object : Indicators)
+	{
+		IToolIndicator* Indicator = Cast<IToolIndicator>(object);
+		if (Indicator != nullptr)
+		{
+			Indicator->Render(RenderAPI);
+		}
+	}
+}
+
+void UToolIndicatorSet::Tick(float DeltaTime)
+{
+	for (UObject* object : Indicators)
+	{
+		IToolIndicator* Indicator = Cast<IToolIndicator>(object);
+		if (Indicator != nullptr)
+		{
+			Indicator->Tick(DeltaTime);
+		}
+	}
+}
diff --git a/Engine/Plugins/Experimental/MeshModelingToolset/Source/ModelingComponents/Public/AssetGenerationUtil.h b/Engine/Plugins/Experimental/MeshModelingToolset/Source/ModelingComponents/Public/AssetGenerationUtil.h
new file mode 100644
--- /dev/null
+++ b/Engine/Plugins/Experimental/MeshModelingToolset/Source/ModelingComponents/Public/AssetGenerationUtil.h
@@ -0,0 +1,47 @@
+// Copyright 1998-2019 Epic Games, Inc. All Rights Reserved.
+
+#pragma once
+
+#include "CoreMinimal.h"
+#include "TransformTypes.h"
+
+class UMaterialInterface;
+class UInteractiveToolManager;
+class IToolsContextAssetAPI;
+struct FMeshDescription;
+class FDynamicMesh3;
+class AActor;
+
+/**
+ * Utility functions for generating Assets in InteractiveTool implementations
+ */
+namespace AssetGenerationUtil
+{
+	/**
+	 * @return default path to use for auto-generated assets
+	 * @warning currently hardcoded to "/Game/_GENERATED"
+	 */
+	MODELINGCOMPONENTS_API FString GetDefaultAutoGeneratedAssetPath();
+
+	/**
+	 * Creates a new StaticMesh actor/component, with a new mesh asset stored at the given PackagePath created via the AssetAPI
+	 * @param AssetAPI pointer to context asset API that will be used to create new asset
+	 * @param TargetWorld world that Actor will be created in
+	 * @param Mesh geometry for the mesh
+	 * @param Transform transformation for the new actor
+	 * @param ObjectName name of the new asset
+	 * @param PackagePath path of the new asset
+	 * @return new mesh actor
+	 */
+#if WITH_EDITOR
+	MODELINGCOMPONENTS_API AActor* GenerateStaticMeshActor(
+		IToolsContextAssetAPI* AssetAPI,
+		UWorld* TargetWorld,
+		const FDynamicMesh3* Mesh,
+		const FTransform3d& Transform,
+		FString ObjectName,
+		FString PackagePath
+	);
+#endif
+
+}
diff --git a/Engine/Plugins/Experimental/MeshModelingToolset/Source/ModelingComponents/Public/BaseDynamicMeshComponent.h b/Engine/Plugins/Experimental/MeshModelingToolset/Source/ModelingComponents/Public/BaseDynamicMeshComponent.h
new file mode 100644
--- /dev/null
+++ b/Engine/Plugins/Experimental/MeshModelingToolset/Source/ModelingComponents/Public/BaseDynamicMeshComponent.h
@@ -0,0 +1,74 @@
+// Copyright 1998-2019 Epic Games, Inc. All Rights Reserved.
+
+#pragma once
+
+#include "CoreMinimal.h"
+#include "UObject/ObjectMacros.h"
+#include "Components/MeshComponent.h"
+#include "MeshVertexChange.h"
+#include "Changes/MeshChange.h"
+
+#include "DynamicMesh3.h"
+
+#include "BaseDynamicMeshComponent.generated.h"
+
+
+// predecl
+class FMeshVertexChange;
+class FMeshChange;
+
+
+
+
+/**
+ * Tangent calculation modes
+ */
+UENUM()
+enum class EDynamicMeshTangentCalcType : uint8
+{
+	/** Tangents are not used/available, proceed accordingly (eg generate arbitrary orthogonal basis) */
+	NoTangents,
+	/** Tangents should be automatically calculated on demand */
+	AutoCalculated,
+	/** Tangents are externally calculated (behavior undefined if they are not actually externally calculated!) */
+	ExternallyCalculated
+};
+
+
+/**
+ * UBaseDynamicMeshComponent is a base interface for a UMeshComponent based on a FDynamicMesh.
+ * Currently no functionality lives here, only some interface functions are defined that various subclasses implement.
+ */
+UCLASS(hidecategories = (LOD, Physics, Collision), editinlinenew, ClassGroup = Rendering)
+class MODELINGCOMPONENTS_API UBaseDynamicMeshComponent : public UMeshComponent, public IMeshVertexCommandChangeTarget, public IMeshCommandChangeTarget
+{
+	GENERATED_UCLASS_BODY()
+
+public:
+
+	/**
+	 * Call this if you update the mesh via GetMesh()
+	 * @todo should provide a function that calls a lambda to modify the mesh, and only return const mesh pointer
+	 */
+	virtual void NotifyMeshUpdated()
+	{
+		unimplemented();
+	}
+
+	/**
+	 * Apply a vertex deformation change to the internal mesh  (implements IMeshVertexCommandChangeTarget)
+	 */
+	virtual void ApplyChange(const FMeshVertexChange* Change, bool bRevert) override
+	{
+		unimplemented();
+	}
+
+	/**
+	 * Apply a general mesh change to the internal mesh  (implements IMeshCommandChangeTarget)
+	 */
+	virtual void ApplyChange(const FMeshChange* Change, bool bRevert) override
+	{
+		unimplemented();
+	}
+
+};
diff --git a/Engine/Plugins/Experimental/MeshModelingToolset/Source/ModelingComponents/Public/Changes/MeshChange.h b/Engine/Plugins/Experimental/MeshModelingToolset/Source/ModelingComponents/Public/Changes/MeshChange.h
new file mode 100644
--- /dev/null
+++ b/Engine/Plugins/Experimental/MeshModelingToolset/Source/ModelingComponents/Public/Changes/MeshChange.h
@@ -0,0 +1,61 @@
+// Copyright 1998-2019 Epic Games, Inc. All Rights Reserved.
+
+#pragma once
+
+#include "CoreMinimal.h"
+#include "UObject/Interface.h"
+#include "Misc/Change.h"
+#include "VectorTypes.h"
+#include "DynamicMeshChangeTracker.h"
+#include "MeshChange.generated.h"
+
+//class FDynamicMeshChange;		// need to refactor this out of DynamicMeshChangeTracker
+
+
+
+
+/**
+ * FMeshChange represents an undoable change to a FDynamicMesh3.
+ * Currently only valid to call Apply/Revert when the Object is a USimpleDynamicMeshComponent
+ */
+class MODELINGCOMPONENTS_API FMeshChange : public FCommandChange
+{
+public:
+	FMeshChange();
+	FMeshChange(TUniquePtr<FDynamicMeshChange> DynamicMeshChangeIn);
+
+	TUniquePtr<FDynamicMeshChange> DynamicMeshChange;
+
+	/** This function is called on Apply and Revert (last argument is true on Apply)*/
+	TFunction<void(FMeshChange*, UObject*, bool)> OnChangeAppliedFunc;
+
+	/** Makes the change to the object */
+	virtual void Apply(UObject* Object) override;
+
+	/** Reverts change to the object */
+	virtual void Revert(UObject* Object) override;
+
+	/** Describes this change (for debugging) */
+	virtual FString ToString() const override;
+};
+
+
+
+
+
+UINTERFACE()
+class MODELINGCOMPONENTS_API UMeshCommandChangeTarget : public UInterface
+{
+	GENERATED_BODY()
+};
+/**
+ * IMeshCommandChangeTarget is an interface which is used to apply a mesh change
+ */
+class MODELINGCOMPONENTS_API IMeshCommandChangeTarget
+{
+	GENERATED_BODY()
+public:
+	virtual void ApplyChange(const FMeshChange* Change, bool bRevert) = 0;
+};
+
+
diff --git a/Engine/Plugins/Experimental/MeshModelingToolset/Source/ModelingComponents/Public/OctreeDynamicMeshComponent.h b/Engine/Plugins/Experimental/MeshModelingToolset/Source/ModelingComponents/Public/OctreeDynamicMeshComponent.h
new file mode 100644
--- /dev/null
+++ b/Engine/Plugins/Experimental/MeshModelingToolset/Source/ModelingComponents/Public/OctreeDynamicMeshComponent.h
@@ -0,0 +1,134 @@
+// Copyright 1998-2019 Epic Games, Inc. All Rights Reserved.
+
+#pragma once
+
+#include "CoreMinimal.h"
+#include "BaseDynamicMeshComponent.h"
+
+#include "DynamicMeshOctree3.h"
+#include "Util/IndexSetDecompositions.h"
+
+#include "OctreeDynamicMeshComponent.generated.h"
+
+// predecl
+struct FMeshDescription;
+
+/** internal FPrimitiveSceneProxy defined in OctreeDynamicMeshSceneProxy.h */
+class FOctreeDynamicMeshSceneProxy;
+
+
+/** 
+ * UOctreeDynamicMeshComponent is a mesh component similar to UProceduralMeshComponent,
+ * except it bases the renderable geometry off an internal FDynamicMesh3 instance.
+ * The class generally has the same capabilities as USimpleDynamicMeshComponent.
+ * 
+ * A FDynamicMeshOctree3 is available to dynamically track the triangles of the mesh
+ * (however the client is responsible for updating this octree).
+ * Based on the Octree, the mesh is partitioned into chunks that are stored in separate
+ * RenderBuffers in the FOctreeDynamicMeshSceneProxy.
+ * Calling NotifyMeshUpdated() will result in only the "dirty" chunks being updated,
+ * rather than the entire mesh.
+ * 
+ * (So, if you don't need this capability, and don't want to update an Octree, use USimpleDynamicMeshComponent!)
+ */
+UCLASS(hidecategories = (LOD, Physics, Collision), editinlinenew, ClassGroup = Rendering)
+class MODELINGCOMPONENTS_API UOctreeDynamicMeshComponent : public UBaseDynamicMeshComponent
+{
+	GENERATED_UCLASS_BODY()
+
+public:
+	/**
+	 * initialize the internal mesh from a MeshDescription
+	 */
+	void InitializeMesh(FMeshDescription* MeshDescription);
+
+	/**
+	 * @return pointer to internal mesh
+	 */
+	FDynamicMesh3* GetMesh() { return Mesh.Get(); }
+
+	FDynamicMeshOctree3* GetOctree() { return Octree.Get(); }
+
+	/**
+	 * @return the current internal mesh, which is replaced with an empty mesh
+	 */
+	TUniquePtr<FDynamicMesh3> ExtractMesh(bool bNotifyUpdate);
+
+	/**
+	 * Write the internal mesh to a MeshDescription
+	 * @param bHaveModifiedToplogy if false, we only update the vertex positions in the MeshDescription, otherwise it is Empty()'d and regenerated entirely
+	 */
+	void Bake(FMeshDescription* MeshDescription, bool bHaveModifiedToplogy);
+
+
+	//
+	// change tracking/etc
+	//
+
+	/**
+	 * Call this if you update the mesh via GetMesh()
+	 * @todo should provide a function that calls a lambda to modify the mesh, and only return const mesh pointer
+	 */
+	virtual void NotifyMeshUpdated() override;
+
+	/**
+	 * Apply a vertex deformation change to the internal mesh
+	 */
+	virtual void ApplyChange(const FMeshVertexChange* Change, bool bRevert) override;
+
+	/**
+	 * Apply a general mesh change to the internal mesh
+	 */
+	virtual void ApplyChange(const FMeshChange* Change, bool bRevert) override;
+
+
+	/**
+	 * This delegate fires when a FCommandChange is applied to this component, so that
+	 * parent objects know the mesh has changed.
+	 */
+	FSimpleMulticastDelegate OnMeshChanged;
+
+	/**
+	 * if true, we always show the wireframe on top of the shaded mesh, even when not in wireframe mode
+	 */
+	UPROPERTY()
+	bool bExplicitShowWireframe = false;
+
+	TFunction<FColor(int)> TriangleColorFunc = nullptr;
+
+private:
+
+	FOctreeDynamicMeshSceneProxy* CurrentProxy = nullptr;
+
+	//~ Begin UPrimitiveComponent Interface.
+	virtual FPrimitiveSceneProxy* CreateSceneProxy() override;
+	//~ End UPrimitiveComponent Interface.
+
+	//~ Begin UMeshComponent Interface.
+	virtual int32 GetNumMaterials() const override;
+	//~ End UMeshComponent Interface.
+
+	//~ Begin USceneComponent Interface.
+	virtual FBoxSphereBounds CalcBounds(const FTransform& LocalToWorld) const override;
+	//~ Begin USceneComponent Interface.
+
+	TUniquePtr<FDynamicMesh3> Mesh;
+	TUniquePtr<FDynamicMeshOctree3> Octree;
+	TUniquePtr<FDynamicMeshOctree3::FTreeCutSet> OctreeCut;
+	FArrayIndexSetsDecomposition TriangleDecomposition;
+	struct FCutCellIndexSet
+	{
+		FDynamicMeshOctree3::FCellReference CellRef;
+		int32 DecompSetID;
+	};
+	TArray<FCutCellIndexSet> CutCellSetMap;
+	int32 SpillDecompSetID;
+
+	void InitializeNewMesh();
+
+
+
+	FColor GetTriangleColor(int TriangleID);
+
+	//friend class FCustomMeshSceneProxy;
+};
diff --git a/Engine/Plugins/Experimental/MeshModelingToolset/Source/ModelingComponents/Public/PreviewMesh.h b/Engine/Plugins/Experimental/MeshModelingToolset/Source/ModelingComponents/Public/PreviewMesh.h
new file mode 100644
--- /dev/null
+++ b/Engine/Plugins/Experimental/MeshModelingToolset/Source/ModelingComponents/Public/PreviewMesh.h
@@ -0,0 +1,255 @@
+// Copyright 1998-2019 Epic Games, Inc. All Rights Reserved.
+
+#pragma once
+
+#include "CoreMinimal.h"
+#include "SimpleDynamicMeshComponent.h"
+#include "DynamicMeshAABBTree3.h"
+#include "InteractiveToolObjects.h"
+#include "MeshVertexChange.h"
+#include "Changes/MeshChange.h"
+#include "PreviewMesh.generated.h"
+
+class FDynamicMesh3;
+struct FMeshDescription;
+
+
+
+/**
+ * UPreviewMesh internally spawns a APreviewMeshActor to hold the preview mesh object.
+ * We use this AInternalToolFrameworkActor subclass so that we can identify such objects
+ * at higher levels (for example to prevent them from being deleted in the Editor)
+ */
+UCLASS()
+class MODELINGCOMPONENTS_API APreviewMeshActor : public AInternalToolFrameworkActor
+{
+	GENERATED_BODY()
+private:
+
+public:
+
+};
+
+
+
+/** 
+ * UPreviewMesh is a utility object that spawns and owns a transient mesh object in the World.
+ * This can be used to show live preview geometry during modeling operations.
+ * Call CreateInWorld() to set it up, and Disconnect() to shut it down.
+ * 
+ * Currently implemented via an internal Actor that has a USimpleDynamicMeshComponent root component,
+ * with an AABBTree created/updated if UProperty bBuildSpatialDataStructure=true.
+ * The Actor is destroyed on Disconnect().
+ * 
+ * The intention with UPreviewMesh is to provide a higher-level interface than the Component.
+ * In future the internal Component may be replaced with another class (eg OctreeDynamicMeshComponent),
+ * or automatically swap between the two, etc.
+ * 
+ * As a result direct access to the Actor/Component, or a non-const FDynamicMesh3, is intentionally not provided.
+ * Wrapper functions are provided (or should be added) for necessary Actor/Component parameters.
+ * To edit the mesh either a copy is done, or EditMesh()/ApplyChange() must be used.
+ * These functions automatically update necessary internal data structures.
+ * 
+ */
+UCLASS(Transient)
+class MODELINGCOMPONENTS_API UPreviewMesh : public UObject, public IMeshVertexCommandChangeTarget, public IMeshCommandChangeTarget
+{
+	GENERATED_BODY()
+
+public:
+	UPreviewMesh();
+	virtual ~UPreviewMesh();
+
+	//
+	// construction / destruction
+	// 
+
+	/**
+	 * Create preview mesh in the World with the given transform
+	 */
+	void CreateInWorld(UWorld* World, const FTransform& WithTransform);
+
+	/**
+	 * Remove and destroy preview mesh
+	 */
+	void Disconnect();
+
+
+	//
+	// visualization parameters
+	// 
+
+	/**
+	 * Enable/disable wireframe overlay rendering
+	 */
+	void EnableWireframe(bool bEnable);
+
+	/**
+	 * Set material on the preview mesh
+	 */
+	void SetMaterial(UMaterialInterface* Material);
+
+	/**
+	 * Get the current transform on the preview mesh
+	 */
+	FTransform GetTransform() const;
+
+	/**
+	 * Set the transform on the preview mesh
+	 */
+	void SetTransform(const FTransform& UseTransform);
+
+	/**
+	 * @return true if the preview mesh is visible
+	 */
+	bool IsVisible() const;
+
+	/**
+	 * Set visibility state of the preview mesh
+	 */
+	void SetVisible(bool bVisible);
+
+
+	/** Render data update hint */
+	enum class ERenderUpdateMode
+	{
+		/** Do not update render data */
+		NoUpdate,	
+		/** Invalidate overlay of internal component, rebuilding all render data */
+		FullUpdate,
+		/** Attempt to do partial update of render data if possible */
+		FastUpdate
+	};
+
+	/**
+	 * Set the triangle color function for rendering / render data construction
+	 */
+	void SetTriangleColorFunction(TFunction<FColor(int)> TriangleColorFunc, ERenderUpdateMode UpdateMode = ERenderUpdateMode::FullUpdate);
+
+	/**
+	 * Clear the triangle color function for rendering / render data construction
+	 */
+	void ClearTriangleColorFunction(ERenderUpdateMode UpdateMode = ERenderUpdateMode::FullUpdate);
+
+
+
+
+
+	//
+	// Queries
+	// 
+
+	/**
+	 * Test for ray intersection with the preview mesh.
+	 * Requires that bBuildSpatialDataStructure = true
+	 * @return true if ray intersections mesh
+	 */
+	bool TestRayIntersection(const FRay3d& WorldRay);
+
+	/**
+	 * Find ray intersection with the preview mesh.
+	 * Requires that bBuildSpatialDataStructure = true
+	 * @param WorldRay ray in world space
+	 * @param HitOut output hit data (only certain members are initialized, see implementation)
+	 * @return true if ray intersections mesh
+	 */
+	bool FindRayIntersection(const FRay3d& WorldRay, FHitResult& HitOut);
+
+
+
+
+
+
+	//
+	// Read access to internal mesh
+	// 
+
+	/**
+	 * Clear the preview mesh
+	 */
+	void ClearPreview();
+
+	/**
+	 * Update the internal mesh by copying the given Mesh
+	 */
+	void UpdatePreview(const FDynamicMesh3* Mesh);
+
+	/**
+	 * Initialize the internal mesh based on the given MeshDescription
+	 */
+	void InitializeMesh(FMeshDescription* MeshDescription);
+
+	/**
+	* @return pointer to the current FDynamicMesh used for preview
+	*/
+	const FDynamicMesh3* GetPreviewDynamicMesh() const;
+
+	/**
+	 * @return the current preview FDynamicMesh, and replace with a new empty mesh
+	 */
+	TUniquePtr<FDynamicMesh3> ExtractPreviewMesh() const;
+
+	/**
+	 * Write the internal mesh to a MeshDescription
+	 * @param bHaveModifiedToplogy if false, we only update the vertex positions in the MeshDescription, otherwise it is Empty()'d and regenerated entirely
+	 */
+	void Bake(FMeshDescription* MeshDescription, bool bHaveModifiedToplogy);
+
+
+
+
+	//
+	// Edit access to internal mesh, and change-tracking/notification
+	// 
+
+	/**
+	 * Apply EditFunc to the internal mesh and update internal data structures as necessary
+	 */
+	void EditMesh(TFunctionRef<void(FDynamicMesh3&)> EditFunc);
+
+	/**
+	 * Apply EditFunc to the internal mesh and update internal data structures as necessary.
+	 * EditFunc is required to notify the given FDynamicMeshChangeTracker about all mesh changes
+	 * @return FMeshChange extracted from FDynamicMeshChangeTracker that represents mesh edit
+	 */
+	TUniquePtr<FMeshChange> TrackedEditMesh(TFunctionRef<void(FDynamicMesh3&, FDynamicMeshChangeTracker&)> EditFunc);
+
+	/**
+	 * Apply/Revert a vertex deformation change to the internal mesh (implements IMeshVertexCommandChangeTarget)
+	 */
+	virtual void ApplyChange(const FMeshVertexChange* Change, bool bRevert) override;
+
+	/**
+	 * Apply/Revert a general mesh change to the internal mesh   (implements IMeshCommandChangeTarget)
+	 */
+	virtual void ApplyChange(const FMeshChange* Change, bool bRevert) override;
+
+	/** @return delegate that is broadcast whenever the internal mesh component is changed */
+	FSimpleMulticastDelegate& GetOnMeshChanged();
+
+
+
+
+
+
+
+public:
+	/** If true, we build a spatial data structure internally for the preview mesh, which allows for hit-testing */
+	UPROPERTY()
+	bool bBuildSpatialDataStructure;
+
+	// results in component drawing w/o z-testing and with editor compositing. Do not recommend using this flag, will be deprecated/removed.
+	UPROPERTY()
+	bool bDrawOnTop;
+
+protected:
+	/** The temporary actor we create internally to own the preview mesh component */
+	APreviewMeshActor* TemporaryParentActor = nullptr;
+
+	/** This component is set as the root component of TemporaryParentActor */
+	UPROPERTY()
+	USimpleDynamicMeshComponent* DynamicMeshComponent = nullptr;
+
+	/** Spatial data structure that is initialized if bBuildSpatialDataStructure = true when UpdatePreview() is called */
+	FDynamicMeshAABBTree3 MeshAABBTree;
+};
diff --git a/Engine/Plugins/Experimental/MeshModelingToolset/Source/ModelingComponents/Public/SimpleDynamicMeshComponent.h b/Engine/Plugins/Experimental/MeshModelingToolset/Source/ModelingComponents/Public/SimpleDynamicMeshComponent.h
new file mode 100644
--- /dev/null
+++ b/Engine/Plugins/Experimental/MeshModelingToolset/Source/ModelingComponents/Public/SimpleDynamicMeshComponent.h
@@ -0,0 +1,142 @@
+// Copyright 1998-2019 Epic Games, Inc. All Rights Reserved.
+
+#pragma once
+
+#include "CoreMinimal.h"
+#include "BaseDynamicMeshComponent.h"
+
+#include "MeshTangents.h"
+
+#include "SimpleDynamicMeshComponent.generated.h"
+
+// predecl
+struct FMeshDescription;
+
+/** internal FPrimitiveSceneProxy defined in SimpleDynamicMeshSceneProxy.h */
+class FSimpleDynamicMeshSceneProxy;
+
+
+
+/** 
+ * USimpleDynamicMeshComponent is a mesh component similar to UProceduralMeshComponent,
+ * except it bases the renderable geometry off an internal FDynamicMesh3 instance.
+ * 
+ * There is some support for undo/redo on the component (@todo is this the right place?)
+ * 
+ * This component draws wireframe-on-shaded when Wireframe is enabled, or when bExplicitShowWireframe = true
+ *
+ */
+UCLASS(hidecategories = (LOD, Physics, Collision), editinlinenew, ClassGroup = Rendering)
+class MODELINGCOMPONENTS_API USimpleDynamicMeshComponent : public UBaseDynamicMeshComponent
+{
+	GENERATED_UCLASS_BODY()
+
+
+public:
+	/** How should Tangents be calculated/handled */
+	UPROPERTY()
+	EDynamicMeshTangentCalcType TangentsType = EDynamicMeshTangentCalcType::NoTangents;
+
+public:
+	/**
+	 * initialize the internal mesh from a MeshDescription
+	 */
+	void InitializeMesh(FMeshDescription* MeshDescription);
+
+	/**
+	 * @return pointer to internal mesh
+	 */
+	FDynamicMesh3* GetMesh() { return Mesh.Get(); }
+
+	/**
+	 * @return the current internal mesh, which is replaced with an empty mesh
+	 */
+	TUniquePtr<FDynamicMesh3> ExtractMesh(bool bNotifyUpdate);
+
+	/**
+	 * @return pointer to internal tangents object. 
+	 * @warning calling this with TangentsType = AutoCalculated will result in possibly-expensive Tangents calculation
+	 */
+	FMeshTangentsf* GetTangents();
+
+
+	/**
+	 * Write the internal mesh to a MeshDescription
+	 * @param bHaveModifiedToplogy if false, we only update the vertex positions in the MeshDescription, otherwise it is Empty()'d and regenerated entirely
+	 */
+	void Bake(FMeshDescription* MeshDescription, bool bHaveModifiedToplogy);
+
+
+	//
+	// change tracking/etc
+	//
+
+	/**
+	 * Call this if you update the mesh via GetMesh()
+	 * @todo should provide a function that calls a lambda to modify the mesh, and only return const mesh pointer
+	 */
+	virtual void NotifyMeshUpdated() override;
+
+	/**
+	 * Apply a vertex deformation change to the internal mesh
+	 */
+	virtual void ApplyChange(const FMeshVertexChange* Change, bool bRevert) override;
+
+	/**
+	 * Apply a general mesh change to the internal mesh
+	 */
+	virtual void ApplyChange(const FMeshChange* Change, bool bRevert) override;
+
+
+	/**
+	 * This delegate fires when a FCommandChange is applied to this component, so that
+	 * parent objects know the mesh has changed.
+	 */
+	FSimpleMulticastDelegate OnMeshChanged;
+
+
+	void SetDrawOnTop(bool bSet);
+
+	/**
+	 * if true, we always show the wireframe on top of the shaded mesh, even when not in wireframe mode
+	 */
+	UPROPERTY()
+	bool bExplicitShowWireframe = false;
+
+
+	UPROPERTY()
+	bool bDrawOnTop = false;
+
+
+
+	TFunction<FColor(int)> TriangleColorFunc = nullptr;
+	void FastNotifyColorsUpdated();
+
+	void FastNotifyPositionsUpdated();
+
+private:
+
+	FSimpleDynamicMeshSceneProxy* CurrentProxy = nullptr;
+
+	//~ Begin UPrimitiveComponent Interface.
+	virtual FPrimitiveSceneProxy* CreateSceneProxy() override;
+	//~ End UPrimitiveComponent Interface.
+
+	//~ Begin UMeshComponent Interface.
+	virtual int32 GetNumMaterials() const override;
+	//~ End UMeshComponent Interface.
+
+	//~ Begin USceneComponent Interface.
+	virtual FBoxSphereBounds CalcBounds(const FTransform& LocalToWorld) const override;
+	//~ Begin USceneComponent Interface.
+
+	TUniquePtr<FDynamicMesh3> Mesh;
+	void InitializeNewMesh();
+
+	bool bTangentsValid = false;
+	FMeshTangentsf Tangents;
+	
+	FColor GetTriangleColor(int TriangleID);
+
+	//friend class FCustomMeshSceneProxy;
+};
diff --git a/Engine/Plugins/Experimental/MeshModelingToolset/Source/ModelingComponents/Public/ToolIndicatorSet.h b/Engine/Plugins/Experimental/MeshModelingToolset/Source/ModelingComponents/Public/ToolIndicatorSet.h
new file mode 100644
--- /dev/null
+++ b/Engine/Plugins/Experimental/MeshModelingToolset/Source/ModelingComponents/Public/ToolIndicatorSet.h
@@ -0,0 +1,61 @@
+// Copyright 1998-2019 Epic Games, Inc. All Rights Reserved.
+
+#pragma once
+
+#include "CoreMinimal.h"
+#include "UObject/Interface.h"
+#include "InteractiveTool.h"
+#include "ToolIndicatorSet.generated.h"
+
+
+// UInterface for IInputBehavior
+UINTERFACE(MinimalAPI)
+class UToolIndicator : public UInterface
+{
+	GENERATED_BODY()
+};
+
+
+/**
+ *
+ */
+class MODELINGCOMPONENTS_API IToolIndicator
+{
+	GENERATED_BODY()
+
+public:
+
+	virtual void Connect(UInteractiveTool* Tool) = 0;
+	virtual void Disconnect() = 0;
+
+	virtual void Render(IToolsContextRenderAPI* RenderAPI) = 0;
+	virtual void Tick(float DeltaTime) = 0;
+};
+
+
+/**
+ *
+ */
+UCLASS(Transient)
+class MODELINGCOMPONENTS_API UToolIndicatorSet : public UObject
+{
+	GENERATED_BODY()
+
+public:
+	UToolIndicatorSet();
+	virtual ~UToolIndicatorSet();
+
+	virtual void Connect(UInteractiveTool* Tool);
+	virtual void Disconnect();
+
+	virtual void AddIndicator(IToolIndicator* Indicator);
+
+	virtual void Render(IToolsContextRenderAPI* RenderAPI);
+	virtual void Tick(float DeltaTime);
+
+protected:
+	UInteractiveTool* Owner;
+
+	UPROPERTY()
+	TSet<UObject*> Indicators;		// these are IToolIndicator* but UPROPERTY requires storing as UObject
+};
\ No newline at end of file
diff --git a/Engine/Plugins/Experimental/MeshModelingToolset/Source/ModelingOperators/ModelingOperators.Build.cs b/Engine/Plugins/Experimental/MeshModelingToolset/Source/ModelingOperators/ModelingOperators.Build.cs
new file mode 100644
--- /dev/null
+++ b/Engine/Plugins/Experimental/MeshModelingToolset/Source/ModelingOperators/ModelingOperators.Build.cs
@@ -0,0 +1,76 @@
+// Copyright 1998-2019 Epic Games, Inc. All Rights Reserved.
+
+using UnrealBuildTool;
+
+public class ModelingOperators : ModuleRules
+{
+	public ModelingOperators(ReadOnlyTargetRules Target) : base(Target)
+	{
+		PCHUsage = ModuleRules.PCHUsageMode.UseExplicitOrSharedPCHs;
+		
+		PublicIncludePaths.AddRange(
+			new string[] {
+				// ... add public include paths required here ...
+			}
+			);
+				
+		
+		PrivateIncludePaths.AddRange(
+			new string[] {
+				// ... add other private include paths required here ...
+			}
+			);
+			
+		
+		PublicDependencyModuleNames.AddRange(
+			new string[]
+			{
+				// These dependencies were commented out when we split out ModelingOperatorsEditorOnly
+				// Many of them are editor only and should not be depended on here.  If you re-add any
+				// dependent modules, please confirm by "launching" or packaging that you are not introducing
+				// editor dependencies
+				//
+				//"MeshSolverUtilities",
+				//"GeometryAlgorithms", // required for constrained Delaunay triangulation
+                //"ProxyLODMeshReduction", // currently required to be public due to IVoxelBasedCSG API
+                //"MeshDescription",
+
+				"Core",
+                "InteractiveToolsFramework",
+				"GeometricObjects",
+				"DynamicMesh",
+				"MeshConversion",
+                "MeshSolverUtilities", // required by the smoothing operators
+				// ... add other public dependencies that you statically link with here ...
+			}
+            );
+			
+		
+		PrivateDependencyModuleNames.AddRange(
+			new string[]
+			{
+				// These dependencies were commented out when we split out ModelingOperatorsEditorOnly
+				// Many of them are editor only and should not be depended on here.  If you re-add any
+				// dependent modules, please confirm by "launching" or packaging that you are not introducing
+				// editor dependencies
+				//
+				//"MeshBuilder",
+                //"MeshDescriptionOperations",
+				//"MeshUtilitiesCommon",
+				//"MeshReductionInterface", // for UE4 standard simplification
+
+				"CoreUObject",
+				"Engine",
+
+			}
+			);
+		
+		
+		DynamicallyLoadedModuleNames.AddRange(
+			new string[]
+			{
+				// ... add any modules that your module loads dynamically here ...
+			}
+			);
+	}
+}
diff --git a/Engine/Plugins/Experimental/MeshModelingToolset/Source/ModelingOperators/Private/CleaningOps/EditNormalsOp.cpp b/Engine/Plugins/Experimental/MeshModelingToolset/Source/ModelingOperators/Private/CleaningOps/EditNormalsOp.cpp
new file mode 100644
--- /dev/null
+++ b/Engine/Plugins/Experimental/MeshModelingToolset/Source/ModelingOperators/Private/CleaningOps/EditNormalsOp.cpp
@@ -0,0 +1,67 @@
+// Copyright 1998-2019 Epic Games, Inc. All Rights Reserved.
+
+#pragma once
+
+#include "CleaningOps/EditNormalsOp.h"
+
+#include "Engine/StaticMesh.h"
+
+#include "MeshDescriptionToDynamicMesh.h"
+#include "MeshSimplification.h"
+#include "MeshConstraintsUtil.h"
+
+
+#include "MeshNormals.h"
+
+
+void FEditNormalsOp::SetTransform(const FTransform& Transform) {
+	ResultTransform = (FTransform3d)Transform;
+}
+
+void FEditNormalsOp::CalculateResult(FProgressCancel* Progress)
+{
+	if (Progress->Cancelled())
+	{
+		return;
+	}
+	bool bDiscardAttributes = false;
+	ResultMesh->Copy(*OriginalMesh, true, true, true, !bDiscardAttributes);
+
+	if (!ensureMsgf(ResultMesh->HasAttributes(), TEXT("Attributes not found on mesh? Conversion should always create them, so this operator should not need to do so.")))
+	{
+		ResultMesh->EnableAttributes();
+	}
+
+	if (bSplitNormals)
+	{
+		FMeshNormals FaceNormals(ResultMesh.Get());
+		FaceNormals.ComputeTriangleNormals();
+		const TArray<FVector3d>& Normals = FaceNormals.GetNormals();
+		float Threshold = FMathf::Cos(NormalSplitThreshold * FMathf::DegToRad);
+		ResultMesh->Attributes()->PrimaryNormals()->CreateFromPredicate([&Normals, &Threshold](int VID, int TA, int TB)
+		{
+			return Normals[TA].Dot(Normals[TB]) > Threshold;
+		}, 0, bAllowSharpVertices);
+	}
+
+	bool bAreaWeight = (NormalCalculationMethod == ENormalCalculationMethod::AreaWeighted || NormalCalculationMethod == ENormalCalculationMethod::AreaAngleWeighting);
+	bool bAngleWeight = (NormalCalculationMethod == ENormalCalculationMethod::AngleWeighted || NormalCalculationMethod == ENormalCalculationMethod::AreaAngleWeighting);
+
+	if (bSplitNormals || bRecomputeNormals) // split normals requires recompute, as it will nuke whatever normals were in the original overlay
+	{
+		FMeshNormals MeshNormals(ResultMesh.Get());
+		MeshNormals.RecomputeOverlayNormals(ResultMesh->Attributes()->PrimaryNormals(), bAreaWeight, bAngleWeight);
+		MeshNormals.CopyToOverlay(ResultMesh->Attributes()->PrimaryNormals(), bInvertNormals);
+	}
+	else if (bInvertNormals)
+	{
+		// just invert with no recompute
+		FDynamicMeshNormalOverlay* Normals = ResultMesh->Attributes()->PrimaryNormals();
+		for (int ElID : Normals->ElementIndicesItr())
+		{
+			auto El = Normals->GetElement(ElID);
+			Normals->SetElement(ElID, -El);
+		}
+	}
+
+}
\ No newline at end of file
diff --git a/Engine/Plugins/Experimental/MeshModelingToolset/Source/ModelingOperators/Private/ParameterizationOps/UVProjectionOp.cpp b/Engine/Plugins/Experimental/MeshModelingToolset/Source/ModelingOperators/Private/ParameterizationOps/UVProjectionOp.cpp
new file mode 100644
--- /dev/null
+++ b/Engine/Plugins/Experimental/MeshModelingToolset/Source/ModelingOperators/Private/ParameterizationOps/UVProjectionOp.cpp
@@ -0,0 +1,226 @@
+// Copyright 1998-2019 Epic Games, Inc. All Rights Reserved.
+
+#pragma once
+
+#include "ParameterizationOps/UVProjectionOp.h"
+
+#include "Engine/StaticMesh.h"
+
+#include "MeshDescriptionToDynamicMesh.h"
+#include "MeshSimplification.h"
+#include "MeshConstraintsUtil.h"
+#include "MatrixTypes.h"
+
+#include "Async/ParallelFor.h"
+
+#include "MeshNormals.h"
+
+
+void FUVProjectionOp::SetTransform(const FTransform& Transform) {
+	ResultTransform = (FTransform3d)Transform;
+}
+
+void FUVProjectionOp::CalculateResult(FProgressCancel* Progress)
+{
+	if (Progress->Cancelled())
+	{
+		return;
+	}
+	bool bDiscardAttributes = false;
+	ResultMesh->Copy(*OriginalMesh, true, true, true, !bDiscardAttributes);
+
+	if (!ensureMsgf(ResultMesh->HasAttributes(), TEXT("Attributes not found on mesh? Conversion should always create them, so this operator should not need to do so.")))
+	{
+		ResultMesh->EnableAttributes();
+	}
+
+	// TODO: get transform from mesh vertex space to projection primitive space
+	FTransform ToProjPrimXF = FTransform(ResultTransform) * ProjectionTransform.Inverse();
+	FTransform3d ToProjPrim(ToProjPrimXF);
+
+	TArray<FVector3d> TransformedVertices; TransformedVertices.SetNumUninitialized(ResultMesh->VertexCount());
+	ParallelFor(ResultMesh->VertexCount(), [this, &ToProjPrim, &TransformedVertices](int32 VertexID)
+	{
+		if (ResultMesh->IsVertex(VertexID))
+		{
+			TransformedVertices[VertexID] = ToProjPrim.TransformPosition(ResultMesh->GetVertex(VertexID));
+		}
+	});
+	
+	
+	int NumUVLayers = ResultMesh->Attributes()->NumUVLayers();
+	int LayerIndex = 0; // TODO: support non-zero UV layer, make this a member of FUVProjectionOp and a setting in the UI (will require updating dmesh attribute code)
+	FDynamicMeshUVOverlay* UVLayer = ResultMesh->Attributes()->GetUVLayer(LayerIndex);
+	UVLayer->ClearElements();
+
+	FVector2D Scale = FVector2D(.5,.5)*UVScale, Offset = FVector2D(.5,.5)+UVOffset;
+	auto ProjAxis = [&Scale, &Offset](const FVector3d& P, int Ax1, int Ax2, float Ax1Scale, float Ax2Scale)
+	{
+		return FVector2f(float(P[Ax1]) * Ax1Scale * Scale.X + Offset.X, float(P[Ax2]) * Ax2Scale * Scale.Y + Offset.Y);
+	};
+
+	int MaxVertexID = ResultMesh->MaxVertexID();
+	int MaxTriangleID = ResultMesh->MaxTriangleID();
+	if (ProjectionMethod == EUVProjectionMethod::Plane)
+	{
+		UVLayer->BeginUnsafeElementsInsert();
+		// 1:1 with vertices; note we cannot just delete the UV layer and switch to vertex UVs because that would delete the normals layer too
+		for (int VID = MaxVertexID-1; VID >= 0; VID--)
+		{
+			if (ResultMesh->IsVertex(VID))
+			{
+				FVector2f UV = ProjAxis(TransformedVertices[VID], 0, 1, 1, 1);
+				UVLayer->InsertElement(VID, UV, VID, true);
+			}
+		}
+		UVLayer->EndUnsafeElementsInsert();
+		for (int TID = 0; TID < MaxTriangleID; TID++)
+		{
+			if (ResultMesh->IsTriangle(TID))
+			{
+				UVLayer->SetTriangle(TID, ResultMesh->GetTriangle(TID));
+			}
+		}
+		return;
+	}
+	
+	// compute normals on the transformed vertices
+	TArray<FVector3d> TransformedNormals; TransformedNormals.SetNumUninitialized(ResultMesh->TriangleCount());
+	ParallelFor(ResultMesh->TriangleCount(), [this, &TransformedVertices, &TransformedNormals](int32 TriangleID)
+	{
+		if (ResultMesh->IsTriangle(TriangleID))
+		{
+			FIndex3i Triangle = ResultMesh->GetTriangle(TriangleID);
+			TransformedNormals[TriangleID] = VectorUtil::Normal(TransformedVertices[Triangle.A], TransformedVertices[Triangle.B], TransformedVertices[Triangle.C]);
+		}
+	});
+	TMap<FIndex2i, int> VertIdxAndBucketIDToElementID;
+
+	if (ProjectionMethod == EUVProjectionMethod::Cube)
+	{
+		for (int TID = 0; TID < MaxTriangleID; TID++)
+		{
+			if (!ResultMesh->IsTriangle(TID))
+			{
+				continue;
+			}
+			const FVector3d& N = TransformedNormals[TID];
+			FVector3d NAbs(FMathd::Abs(N.X), FMathd::Abs(N.Y), FMathd::Abs(N.Z));
+			int MajorAxis = NAbs[0] > NAbs[1] ? (NAbs[0] > NAbs[2] ? 0 : 2) : (NAbs[1] > NAbs[2] ? 1 : 2);
+			double MajorAxisSign = FMathd::Sign(N[MajorAxis]);
+			int Minor1 = (MajorAxis + 1) % 3;
+			int Minor2 = (MajorAxis + 2) % 3;
+			int Bucket = MajorAxis;
+			if (MajorAxisSign > 0)
+			{
+				Bucket += 3;
+			}
+			FIndex3i OverlayTri;
+			FIndex3i SourceTri = ResultMesh->GetTriangle(TID);
+			for (int SubIdx = 0; SubIdx < 3; SubIdx++)
+			{
+				FIndex2i ElementKey(SourceTri[SubIdx], Bucket);
+				const int* ElementIdx = VertIdxAndBucketIDToElementID.Find(ElementKey);
+				if (!ElementIdx)
+				{
+					FVector3d LocalV = TransformedVertices[SourceTri[SubIdx]];
+					int NewElementIdx = UVLayer->AppendElement(ProjAxis(LocalV, Minor1, Minor2, MajorAxisSign, 1), ElementKey.A);
+					VertIdxAndBucketIDToElementID.Add(ElementKey, NewElementIdx);
+					OverlayTri[SubIdx] = NewElementIdx;
+				}
+				else
+				{
+					OverlayTri[SubIdx] = *ElementIdx;
+				}
+			}
+			UVLayer->SetTriangle(TID, OverlayTri);
+		}
+
+		return;
+	}
+
+	if (ProjectionMethod == EUVProjectionMethod::Cylinder)
+	{
+		double DotThresholdRejectFromPlane = FMathd::Cos(FMathd::DegToRad * CylinderProjectToTopOrBottomAngleThreshold);
+		for (int TID = 0; TID < MaxTriangleID; TID++)
+		{
+			if (!ResultMesh->IsTriangle(TID))
+			{
+				continue;
+			}
+			const FVector3d& N = TransformedNormals[TID];
+			FIndex3i OverlayTri;
+			FIndex3i SourceTri = ResultMesh->GetTriangle(TID);
+			if (FMathd::Abs(N.Z) > DotThresholdRejectFromPlane)
+			{
+				int MajorAxis = 2;
+				double MajorAxisSign = FMathd::Sign(N[MajorAxis]);
+				// project to +/- Z
+				int Bucket = MajorAxisSign > 0 ? 1 : 0;
+				int Minor1 = 0;
+				int Minor2 = 1;
+
+				// TODO: extract common major axis plane projection sub-function from cube projection & this?
+				for (int SubIdx = 0; SubIdx < 3; SubIdx++)
+				{
+					FIndex2i ElementKey(SourceTri[SubIdx], Bucket);
+					const int* ElementIdx = VertIdxAndBucketIDToElementID.Find(ElementKey);
+					if (!ElementIdx)
+					{
+						FVector3d LocalV = TransformedVertices[SourceTri[SubIdx]];
+						int NewElementIdx = UVLayer->AppendElement(ProjAxis(LocalV, Minor1, Minor2, MajorAxisSign, 1), ElementKey.A);
+						VertIdxAndBucketIDToElementID.Add(ElementKey, NewElementIdx);
+						OverlayTri[SubIdx] = NewElementIdx;
+					}
+					else
+					{
+						OverlayTri[SubIdx] = *ElementIdx;
+					}
+				}
+			}
+			else
+			{
+				// project to cylinder
+				int Bucket = 3;
+				FVector3d Centroid = (TransformedVertices[SourceTri.A] + TransformedVertices[SourceTri.B] + TransformedVertices[SourceTri.C]) * (1.0 / 3.0);
+				double CentroidAngle = FMathd::Atan2Positive(Centroid.Y, Centroid.X);
+				for (int SubIdx = 0; SubIdx < 3; SubIdx++)
+				{
+					FVector3d LocalV = TransformedVertices[SourceTri[SubIdx]];
+					double VAngle = FMathd::Atan2Positive(LocalV.Y, LocalV.X);
+					double AngleDiff = VAngle - CentroidAngle;
+					if (FMathd::Abs(AngleDiff) > FMathd::Pi) {
+						if (AngleDiff < 0)
+						{
+							VAngle += FMathd::TwoPi;
+							Bucket--;
+						}
+						else
+						{
+							VAngle -= FMathd::TwoPi;
+							Bucket++;
+						}
+					}
+					FIndex2i ElementKey(SourceTri[SubIdx], Bucket);
+					const int* ElementIdx = VertIdxAndBucketIDToElementID.Find(ElementKey);
+					if (!ElementIdx)
+					{
+						FVector2f UV((float(VAngle)*FMathf::InvPi-1.0f)*Scale.X+Offset.X, float(LocalV.Z)*Scale.Y + Offset.Y);
+						int NewElementIdx = UVLayer->AppendElement(UV, ElementKey.A);
+						VertIdxAndBucketIDToElementID.Add(ElementKey, NewElementIdx);
+						OverlayTri[SubIdx] = NewElementIdx;
+					}
+					else
+					{
+						OverlayTri[SubIdx] = *ElementIdx;
+					}
+				}
+			}
+
+			UVLayer->SetTriangle(TID, OverlayTri);
+		}
+
+		return;
+	}
+
+}
diff --git a/Engine/Plugins/Experimental/MeshModelingToolset/Source/ModelingOperators/Private/SmoothingOps/SmoothingOpBase.cpp b/Engine/Plugins/Experimental/MeshModelingToolset/Source/ModelingOperators/Private/SmoothingOps/SmoothingOpBase.cpp
new file mode 100644
--- /dev/null
+++ b/Engine/Plugins/Experimental/MeshModelingToolset/Source/ModelingOperators/Private/SmoothingOps/SmoothingOpBase.cpp
@@ -0,0 +1,52 @@
+
+// Copyright 1998-2019 Epic Games, Inc. All Rights Reserved.
+#include "SmoothingOps/SmoothingOpBase.h"
+
+#include "DynamicMesh3.h"
+#include "MeshNormals.h"
+
+
+FSmoothingOpBase::FSmoothingOpBase(const FDynamicMesh3* Mesh, float Speed, int32 Iterations) :
+		FDynamicMeshOperator(),
+		SmoothSpeed(Speed),
+		SmoothIterations(Iterations)
+{
+	// deep copy the src mesh into the result mesh.  This ResultMesh will be directly updated by the smoothing.
+	ResultMesh->Copy(*Mesh);
+	
+	PositionBuffer.SetNum(ResultMesh->MaxVertexID());
+
+	for (int vid : ResultMesh->VertexIndicesItr())
+	{
+		PositionBuffer[vid] = ResultMesh->GetVertex(vid);
+	}
+}
+
+void FSmoothingOpBase::SetTransform(FTransform3d& XForm)
+{
+	ResultTransform = XForm;
+}
+
+
+void FSmoothingOpBase::UpdateResultMesh()
+{
+	// move the verts to the new location
+	for (int32 vid : ResultMesh->VertexIndicesItr())
+	{
+		const FVector3d Pos = PositionBuffer[vid];
+		ResultMesh->SetVertex(vid, Pos);
+	}
+
+	// recalculate normals
+	if (ResultMesh->HasAttributes())
+	{
+		FMeshNormals Normals(ResultMesh.Get());
+		FDynamicMeshNormalOverlay* NormalOverlay = ResultMesh->Attributes()->PrimaryNormals();
+		Normals.RecomputeOverlayNormals(NormalOverlay);
+		Normals.CopyToOverlay(NormalOverlay);
+	}
+	else
+	{
+		FMeshNormals::QuickComputeVertexNormals(*ResultMesh);
+	}
+}
diff --git a/Engine/Plugins/Experimental/MeshModelingToolset/Source/ModelingOperators/Public/CleaningOps/EditNormalsOp.h b/Engine/Plugins/Experimental/MeshModelingToolset/Source/ModelingOperators/Public/CleaningOps/EditNormalsOp.h
new file mode 100644
--- /dev/null
+++ b/Engine/Plugins/Experimental/MeshModelingToolset/Source/ModelingOperators/Public/CleaningOps/EditNormalsOp.h
@@ -0,0 +1,46 @@
+// Copyright 1998-2019 Epic Games, Inc. All Rights Reserved.
+
+#pragma once
+
+#include "CoreMinimal.h"
+#include "Util/ProgressCancel.h"
+#include "ModelingOperators.h"
+
+#include "EditNormalsOp.generated.h"
+
+
+UENUM()
+enum class ENormalCalculationMethod : uint8
+{
+	AreaWeighted,
+	AngleWeighted,
+	AreaAngleWeighting
+};
+
+
+
+class MODELINGOPERATORS_API FEditNormalsOp : public FDynamicMeshOperator
+{
+public:
+	virtual ~FEditNormalsOp() {}
+
+	// inputs
+	TSharedPtr<FDynamicMesh3> OriginalMesh;
+
+	bool bInvertNormals;
+	bool bRecomputeNormals;
+	ENormalCalculationMethod NormalCalculationMethod;
+	bool bSplitNormals;
+	bool bAllowSharpVertices;
+	float NormalSplitThreshold;
+
+	void SetTransform(const FTransform& Transform);
+
+	//
+	// FDynamicMeshOperator implementation
+	// 
+
+	virtual void CalculateResult(FProgressCancel* Progress) override;
+};
+
+
diff --git a/Engine/Plugins/Experimental/MeshModelingToolset/Source/ModelingOperators/Public/ParameterizationOps/UVProjectionOp.h b/Engine/Plugins/Experimental/MeshModelingToolset/Source/ModelingOperators/Public/ParameterizationOps/UVProjectionOp.h
new file mode 100644
--- /dev/null
+++ b/Engine/Plugins/Experimental/MeshModelingToolset/Source/ModelingOperators/Public/ParameterizationOps/UVProjectionOp.h
@@ -0,0 +1,44 @@
+// Copyright 1998-2019 Epic Games, Inc. All Rights Reserved.
+
+#pragma once
+
+#include "CoreMinimal.h"
+#include "Util/ProgressCancel.h"
+#include "ModelingOperators.h"
+
+#include "UVProjectionOp.generated.h"
+
+
+UENUM()
+enum class EUVProjectionMethod : uint8
+{
+	Cube,
+	Cylinder,
+	Plane
+};
+
+
+
+class MODELINGOPERATORS_API FUVProjectionOp : public FDynamicMeshOperator
+{
+public:
+	virtual ~FUVProjectionOp() {}
+
+	// inputs
+	TSharedPtr<FDynamicMesh3> OriginalMesh;
+
+	EUVProjectionMethod ProjectionMethod;
+	float CylinderProjectToTopOrBottomAngleThreshold;
+	FTransform ProjectionTransform;
+	FVector2D UVScale, UVOffset;
+
+	void SetTransform(const FTransform& Transform);
+
+	//
+	// FDynamicMeshOperator implementation
+	// 
+
+	virtual void CalculateResult(FProgressCancel* Progress) override;
+};
+
+
diff --git a/Engine/Plugins/Experimental/MeshModelingToolset/Source/ModelingOperatorsEditorOnly/ModelingOperatorsEditorOnly.Build.cs b/Engine/Plugins/Experimental/MeshModelingToolset/Source/ModelingOperatorsEditorOnly/ModelingOperatorsEditorOnly.Build.cs
new file mode 100644
--- /dev/null
+++ b/Engine/Plugins/Experimental/MeshModelingToolset/Source/ModelingOperatorsEditorOnly/ModelingOperatorsEditorOnly.Build.cs
@@ -0,0 +1,64 @@
+// Copyright 1998-2019 Epic Games, Inc. All Rights Reserved.
+
+using UnrealBuildTool;
+
+public class ModelingOperatorsEditorOnly : ModuleRules
+{
+	public ModelingOperatorsEditorOnly(ReadOnlyTargetRules Target) : base(Target)
+	{
+		PCHUsage = ModuleRules.PCHUsageMode.UseExplicitOrSharedPCHs;
+		
+		PublicIncludePaths.AddRange(
+			new string[] {
+				// ... add public include paths required here ...
+			}
+			);
+				
+		
+		PrivateIncludePaths.AddRange(
+			new string[] {
+				// ... add other private include paths required here ...
+			}
+			);
+			
+		
+		PublicDependencyModuleNames.AddRange(
+			new string[]
+			{
+				"Core",
+                "InteractiveToolsFramework",
+                "MeshDescription",
+				"ModelingOperators",
+				"GeometricObjects",
+				"DynamicMesh",
+				"MeshConversion",
+				"MeshSolverUtilities",
+				"GeometryAlgorithms", // required for constrained Delaunay triangulation
+                "ProxyLODMeshReduction", // currently required to be public due to IVoxelBasedCSG API
+				// ... add other public dependencies that you statically link with here ...
+			}
+            );
+			
+		
+		PrivateDependencyModuleNames.AddRange(
+			new string[]
+			{
+				"CoreUObject",
+				"Engine",
+
+				"MeshBuilder",
+                "MeshDescriptionOperations",
+				"MeshUtilitiesCommon",  
+				"MeshReductionInterface", // for UE4 standard simplification 
+			}
+			);
+		
+		
+		DynamicallyLoadedModuleNames.AddRange(
+			new string[]
+			{
+				// ... add any modules that your module loads dynamically here ...
+			}
+			);
+	}
+}
diff --git a/Engine/Plugins/Experimental/MeshModelingToolset/Source/ModelingOperatorsEditorOnly/Private/CompositionOps/VoxelBooleanMeshesOp.cpp b/Engine/Plugins/Experimental/MeshModelingToolset/Source/ModelingOperatorsEditorOnly/Private/CompositionOps/VoxelBooleanMeshesOp.cpp
new file mode 100644
--- /dev/null
+++ b/Engine/Plugins/Experimental/MeshModelingToolset/Source/ModelingOperatorsEditorOnly/Private/CompositionOps/VoxelBooleanMeshesOp.cpp
@@ -0,0 +1,137 @@
+// Copyright 1998-2019 Epic Games, Inc. All Rights Reserved.
+
+#pragma once
+
+#include "CompositionOps/VoxelBooleanMeshesOp.h"
+
+#include "Engine/StaticMesh.h"
+
+#include "MeshDescription.h"
+#include "MeshDescriptionToDynamicMesh.h"
+#include "MeshSimplification.h"
+#include "MeshConstraintsUtil.h"
+
+void FVoxelBooleanMeshesOp::CalculateResult(FProgressCancel* Progress)
+{
+	FMeshDescription ResultMeshDescription;
+	UStaticMesh::RegisterMeshAttributes(ResultMeshDescription);
+
+	if (VoxelCount < 16)
+	{
+		VoxelCount = 16;
+	}
+
+	// Use the world space bounding box of each mesh to compute the voxel size
+	VoxelSizeD = ComputeVoxelSize();
+
+	// give this an absolute min since the user might scale both objects to zero..
+	VoxelSizeD = FMath::Max(VoxelSizeD, 0.001);
+
+	// Create BooleanTool and Boolean the meshes.
+	TUniquePtr<IVoxelBasedCSG> VoxelCSGTool = IVoxelBasedCSG::CreateCSGTool(VoxelSizeD);
+	FVector MergedOrigin;
+
+	IVoxelBasedCSG::FPlacedMesh& A = (*InputMeshArray)[0];
+	IVoxelBasedCSG::FPlacedMesh& B = (*InputMeshArray)[1];
+
+	switch (Operation)
+	{
+	case EBooleanOperation::DifferenceAB:
+		MergedOrigin = VoxelCSGTool->ComputeDifference(A, B, ResultMeshDescription, AdaptivityD, IsoSurfaceD);
+		break;
+	case EBooleanOperation::DifferenceBA:
+		MergedOrigin = VoxelCSGTool->ComputeDifference(B, A, ResultMeshDescription, AdaptivityD, IsoSurfaceD);
+		break;
+	case EBooleanOperation::Intersect:
+		MergedOrigin = VoxelCSGTool->ComputeIntersection(A, B, ResultMeshDescription, AdaptivityD, IsoSurfaceD);
+		break;
+	case EBooleanOperation::Union:
+		MergedOrigin = VoxelCSGTool->ComputeUnion(A, B, ResultMeshDescription, AdaptivityD, IsoSurfaceD);
+		break;
+	default:
+		check(0);
+	}
+
+	ResultTransform = FTransform3d(MergedOrigin);
+
+	if (Progress->Cancelled())
+	{
+		return;
+	}
+
+	// Convert to dynamic mesh
+	FMeshDescriptionToDynamicMesh Converter;
+	Converter.bPrintDebugMessages = true;
+	Converter.Convert(&ResultMeshDescription, *ResultMesh);
+
+	if (Progress->Cancelled())
+	{
+		return;
+	}
+
+	if (AdaptivityD > 0)
+	{
+		FQEMSimplification Reducer(ResultMesh.Get());
+		FMeshConstraints constraints;
+		FMeshConstraintsUtil::ConstrainAllSeams(constraints, *ResultMesh, true, false);
+		Reducer.SetExternalConstraints(&constraints);
+		Reducer.Progress = Progress;
+		for (int k = 0; k < NumFastCollapsePasses; ++k)
+		{
+			Reducer.FastCollapsePass(FastCollapseVoxelSizeMultipler * VoxelSizeD);
+			if (Progress->Cancelled())
+			{
+				return;
+			}
+		}
+	}
+}
+
+float FVoxelBooleanMeshesOp::ComputeVoxelSize() const 
+{ 
+	const TArray<IVoxelBasedCSG::FPlacedMesh>& PlacedMeshes = *(InputMeshArray);
+
+	float Size = 0.f;
+	for (int32 i = 0; i < PlacedMeshes.Num(); ++i)
+	{
+		//Bounding box
+		FVector BBoxMin(FLT_MAX, FLT_MAX, FLT_MAX);
+		FVector BBoxMax(-FLT_MAX, -FLT_MAX, -FLT_MAX);
+
+		// Use the scale define by the transform.  We don't care about actual placement
+		// in the world for this.
+		FVector Scale = PlacedMeshes[i].Transform.GetScale3D();
+		const FMeshDescription&  MeshDescription = *PlacedMeshes[i].Mesh;
+
+		TVertexAttributesConstRef<FVector> VertexPositions = MeshDescription.VertexAttributes().GetAttributesRef<FVector>(MeshAttribute::Vertex::Position);
+		for (const FVertexID VertexID : MeshDescription.Vertices().GetElementIDs())
+		{
+			const FVector Pos = VertexPositions[VertexID];
+            
+			BBoxMin.X = FMath::Min(BBoxMin.X, Pos.X);
+			BBoxMin.Y = FMath::Min(BBoxMin.Y, Pos.Y);
+			BBoxMin.Z = FMath::Min(BBoxMin.Z, Pos.Z);
+
+			BBoxMax.X = FMath::Max(BBoxMax.X, Pos.X);
+			BBoxMax.Y = FMath::Max(BBoxMax.Y, Pos.Y);
+			BBoxMax.Z = FMath::Max(BBoxMax.Z, Pos.Z);
+		}
+
+		// The size of the BBox in each direction
+		FVector Extents(BBoxMax.X - BBoxMin.X, BBoxMax.Y - BBoxMin.Y, BBoxMax.Z - BBoxMin.Z);
+		
+		// Scale with the local space scale.
+		Extents.X = Extents.X * Scale.X;
+		Extents.Y = Extents.Y * Scale.Y;
+		Extents.Z = Extents.Z * Scale.Z;
+
+		float MajorAxisSize = FMath::Max3(Extents.X, Extents.Y, Extents.Z);
+
+
+		Size = FMath::Max(MajorAxisSize / VoxelCount, Size);
+
+	}
+	
+	return Size;
+
+}
diff --git a/Engine/Plugins/Experimental/MeshModelingToolset/Source/ModelingOperatorsEditorOnly/Private/CompositionOps/VoxelMergeMeshesOp.cpp b/Engine/Plugins/Experimental/MeshModelingToolset/Source/ModelingOperatorsEditorOnly/Private/CompositionOps/VoxelMergeMeshesOp.cpp
new file mode 100644
--- /dev/null
+++ b/Engine/Plugins/Experimental/MeshModelingToolset/Source/ModelingOperatorsEditorOnly/Private/CompositionOps/VoxelMergeMeshesOp.cpp
@@ -0,0 +1,112 @@
+// Copyright 1998-2019 Epic Games, Inc. All Rights Reserved.
+
+#pragma once
+
+#include "CompositionOps/VoxelMergeMeshesOp.h"
+
+#include "Engine/StaticMesh.h"
+
+#include "MeshDescription.h"
+#include "MeshDescriptionToDynamicMesh.h"
+#include "MeshSimplification.h"
+#include "MeshConstraintsUtil.h"
+
+
+void FVoxelMergeMeshesOp::CalculateResult(FProgressCancel* Progress)
+{
+	FMeshDescription MergedMeshesDescription;
+	UStaticMesh::RegisterMeshAttributes(MergedMeshesDescription);
+
+	// Use the world space bounding box of each mesh to compute the voxel size
+	VoxelSizeD = ComputeVoxelSize();
+
+	// give this an absolute min since the user might scale both objects to zero..
+	VoxelSizeD = FMath::Max(VoxelSizeD, 0.001);
+
+	// Create CSGTool and merge the meshes.
+	TUniquePtr<IVoxelBasedCSG> VoxelCSGTool = IVoxelBasedCSG::CreateCSGTool(VoxelSizeD);
+	FVector MergedOrigin = VoxelCSGTool->ComputeUnion(*InputMeshArray, MergedMeshesDescription, AdaptivityD, IsoSurfaceD);
+	ResultTransform = FTransform3d(MergedOrigin);
+
+	if (Progress->Cancelled())
+	{
+		return;
+	}
+
+	// Convert to dynamic mesh
+	FMeshDescriptionToDynamicMesh Converter;
+	Converter.bPrintDebugMessages = true;
+	Converter.Convert(&MergedMeshesDescription, *ResultMesh);
+
+	if (Progress->Cancelled())
+	{
+		return;
+	}
+
+	if (AdaptivityD > 0)
+	{
+		FQEMSimplification Reducer(ResultMesh.Get());
+		FMeshConstraints constraints;
+		FMeshConstraintsUtil::ConstrainAllSeams(constraints, *ResultMesh, true, false);
+		Reducer.SetExternalConstraints(&constraints);
+		Reducer.Progress = Progress;
+		for (int k = 0; k < NumFastCollapsePasses; ++k)
+		{
+			Reducer.FastCollapsePass(FastCollapseVoxelSizeMultipler * VoxelSizeD);
+			if (Progress->Cancelled())
+			{
+				return;
+			}
+		}
+	}
+}
+
+
+float FVoxelMergeMeshesOp::ComputeVoxelSize() const
+{
+	const TArray<IVoxelBasedCSG::FPlacedMesh>& PlacedMeshes = *(InputMeshArray);
+
+	float Size = 0.f;
+	for (int32 i = 0; i < PlacedMeshes.Num(); ++i)
+	{
+		//Bounding box
+		FVector BBoxMin(FLT_MAX, FLT_MAX, FLT_MAX);
+		FVector BBoxMax(-FLT_MAX, -FLT_MAX, -FLT_MAX);
+
+		// Use the scale define by the transform.  We don't care about actual placement
+		// in the world for this.
+		FVector Scale = PlacedMeshes[i].Transform.GetScale3D();
+		const FMeshDescription&  MeshDescription = *PlacedMeshes[i].Mesh;
+
+		TVertexAttributesConstRef<FVector> VertexPositions = MeshDescription.VertexAttributes().GetAttributesRef<FVector>(MeshAttribute::Vertex::Position);
+		for (const FVertexID VertexID : MeshDescription.Vertices().GetElementIDs())
+		{
+			const FVector Pos = VertexPositions[VertexID];
+
+			BBoxMin.X = FMath::Min(BBoxMin.X, Pos.X);
+			BBoxMin.Y = FMath::Min(BBoxMin.Y, Pos.Y);
+			BBoxMin.Z = FMath::Min(BBoxMin.Z, Pos.Z);
+
+			BBoxMax.X = FMath::Max(BBoxMax.X, Pos.X);
+			BBoxMax.Y = FMath::Max(BBoxMax.Y, Pos.Y);
+			BBoxMax.Z = FMath::Max(BBoxMax.Z, Pos.Z);
+		}
+
+		// The size of the BBox in each direction
+		FVector Extents(BBoxMax.X - BBoxMin.X, BBoxMax.Y - BBoxMin.Y, BBoxMax.Z - BBoxMin.Z);
+
+		// Scale with the local space scale.
+		Extents.X = Extents.X * Scale.X;
+		Extents.Y = Extents.Y * Scale.Y;
+		Extents.Z = Extents.Z * Scale.Z;
+
+		float MajorAxisSize = FMath::Max3(Extents.X, Extents.Y, Extents.Z);
+
+
+		Size = FMath::Max(MajorAxisSize / VoxelCount, Size);
+
+	}
+
+	return Size;
+
+}
diff --git a/Engine/Plugins/Experimental/MeshModelingToolset/Source/ModelingOperatorsEditorOnly/Public/CompositionOps/VoxelBooleanMeshesOp.h b/Engine/Plugins/Experimental/MeshModelingToolset/Source/ModelingOperatorsEditorOnly/Public/CompositionOps/VoxelBooleanMeshesOp.h
new file mode 100644
--- /dev/null
+++ b/Engine/Plugins/Experimental/MeshModelingToolset/Source/ModelingOperatorsEditorOnly/Public/CompositionOps/VoxelBooleanMeshesOp.h
@@ -0,0 +1,50 @@
+// Copyright 1998-2019 Epic Games, Inc. All Rights Reserved.
+
+#pragma once
+
+#include "CoreMinimal.h"
+#include "Util/ProgressCancel.h"
+#include "ModelingOperators.h"
+#include "ProxyLODVolume.h"
+
+
+
+
+class MODELINGOPERATORSEDITORONLY_API FVoxelBooleanMeshesOp : public FDynamicMeshOperator
+{
+public:
+	virtual ~FVoxelBooleanMeshesOp() {}
+
+	enum class EBooleanOperation
+	{
+		DifferenceAB = 0,
+		DifferenceBA = 1,
+		Intersect = 2,
+		Union = 3
+	};
+
+	// inputs
+	TSharedPtr<TArray<IVoxelBasedCSG::FPlacedMesh>> InputMeshArray;
+	int32 VoxelCount = 128;
+	double VoxelSizeD = 1.0;
+	double AdaptivityD = 0;
+	double IsoSurfaceD = 0;
+	EBooleanOperation Operation = EBooleanOperation::DifferenceBA;
+
+	double FastCollapseVoxelSizeMultipler = 1.2;
+	int NumFastCollapsePasses = 10;
+
+	//
+	// FDynamicMeshOperator implementation
+	// 
+
+	virtual void CalculateResult(FProgressCancel* Progress) override;
+
+private:
+
+	// compute the voxel size based on the voxel count and the input geometry
+	float ComputeVoxelSize() const;
+
+};
+
+
diff --git a/Engine/Plugins/Experimental/MeshModelingToolset/Source/ModelingOperatorsEditorOnly/Public/CompositionOps/VoxelMergeMeshesOp.h b/Engine/Plugins/Experimental/MeshModelingToolset/Source/ModelingOperatorsEditorOnly/Public/CompositionOps/VoxelMergeMeshesOp.h
new file mode 100644
--- /dev/null
+++ b/Engine/Plugins/Experimental/MeshModelingToolset/Source/ModelingOperatorsEditorOnly/Public/CompositionOps/VoxelMergeMeshesOp.h
@@ -0,0 +1,41 @@
+// Copyright 1998-2019 Epic Games, Inc. All Rights Reserved.
+
+#pragma once
+
+#include "CoreMinimal.h"
+#include "Util/ProgressCancel.h"
+#include "ModelingOperators.h"
+#include "ProxyLODVolume.h"
+
+
+
+
+
+class MODELINGOPERATORSEDITORONLY_API FVoxelMergeMeshesOp : public FDynamicMeshOperator
+{
+public:
+	virtual ~FVoxelMergeMeshesOp() {}
+
+	// inputs
+	TSharedPtr<TArray<IVoxelBasedCSG::FPlacedMesh>> InputMeshArray;
+	int32 VoxelCount = 128;
+	double VoxelSizeD = 1.0;
+	double AdaptivityD = 0;
+	double IsoSurfaceD = 0;
+
+	double FastCollapseVoxelSizeMultipler = 1.2;
+	int NumFastCollapsePasses = 10;
+
+	//
+	// FDynamicMeshOperator implementation
+	// 
+
+	virtual void CalculateResult(FProgressCancel* Progress) override;
+
+private:
+
+	// compute the voxel size based on the voxel count and the input geometry
+	float ComputeVoxelSize() const;
+};
+
+
