make plane cut not be editor only + change plane manip gizmo, + minor fixes

#rnx

diff --git a/Engine/Plugins/Experimental/GeometryProcessing/Source/DynamicMesh/Private/Operations/MeshPlaneCut.cpp b/Engine/Plugins/Experimental/GeometryProcessing/Source/DynamicMesh/Private/Operations/MeshPlaneCut.cpp
--- a/Engine/Plugins/Experimental/GeometryProcessing/Source/DynamicMesh/Private/Operations/MeshPlaneCut.cpp
+++ b/Engine/Plugins/Experimental/GeometryProcessing/Source/DynamicMesh/Private/Operations/MeshPlaneCut.cpp
@@ -93,13 +93,12 @@ bool FMeshPlaneCut::Cut()
 		}
 
 		FDynamicMesh3::FEdgeSplitInfo splitInfo;
-		EMeshResult result = Mesh->SplitEdge(EID, splitInfo);
-		ensureMsgf(result == EMeshResult::Ok, TEXT("FMeshPlaneCut::Cut: failed to SplitEdge")); // edge split really shouldn't fail
-
-		// SplitEdge just bisects edge; now move created vertex to plane intersection pt
 		double t = f0 / (f0 - f1);
-		FVector3d newPos = (1 - t) * Mesh->GetVertex(ev.A) + (t) * Mesh->GetVertex(ev.B);
-		Mesh->SetVertex(splitInfo.NewVertex, newPos);
+		EMeshResult result = Mesh->SplitEdge(EID, splitInfo, t);
+		if (!ensureMsgf(result == EMeshResult::Ok, TEXT("FMeshPlaneCut::Cut: failed to SplitEdge")))
+		{
+			continue; // edge split really shouldn't fail; skip the edge if it somehow does
+		}
 
 		NewEdges.Add(splitInfo.NewEdges.A); // TODO: why not OnCutEdges.Add(splitInfo.NewEdges.A);  ???
 		NewEdges.Add(splitInfo.NewEdges.B); OnCutEdges.Add(splitInfo.NewEdges.B);
diff --git a/Engine/Plugins/Experimental/GeometryProcessing/Source/GeometryAlgorithms/Private/ConstrainedDelaunay2.cpp b/Engine/Plugins/Experimental/GeometryProcessing/Source/GeometryAlgorithms/Private/ConstrainedDelaunay2.cpp
--- a/Engine/Plugins/Experimental/GeometryProcessing/Source/GeometryAlgorithms/Private/ConstrainedDelaunay2.cpp
+++ b/Engine/Plugins/Experimental/GeometryProcessing/Source/GeometryAlgorithms/Private/ConstrainedDelaunay2.cpp
@@ -224,7 +224,7 @@ bool TConstrainedDelaunay2<RealType>::Triangulate()
 		int BaseIdx = TriIdx * 3;
 		for (int SubIdx = 0, NextIdx = 2; SubIdx < 3; NextIdx = SubIdx++)
 		{
-			if (Adj[BaseIdx + NextIdx] == -1) // on hull
+			if (Adj[BaseIdx + NextIdx] < 0) // on hull
 			{
 				int VertA = Indices[BaseIdx + SubIdx], VertB = Indices[BaseIdx + NextIdx];
 				if (HasUnorderedEdge(HoleMap, VertA, VertB))
@@ -256,7 +256,7 @@ bool TConstrainedDelaunay2<RealType>::Triangulate()
 				continue; // cannot cross hole edges
 			}
 			int AdjTri = Adj[BaseIdx + NextIdx];
-			if (AdjTri != -1 && Keep[AdjTri] == 0)
+			if (AdjTri >= 0 && Keep[AdjTri] == 0)
 			{
 				int WindingChange = WindingAcross(BoundaryMap, VertA, VertB);
 				int Winding = LastWinding + WindingChange;
diff --git a/Engine/Plugins/Experimental/MeshModelingToolset/Source/MeshModelingTools/Private/PlaneCutTool.cpp b/Engine/Plugins/Experimental/MeshModelingToolset/Source/MeshModelingTools/Private/PlaneCutTool.cpp
new file mode 100644
--- /dev/null
+++ b/Engine/Plugins/Experimental/MeshModelingToolset/Source/MeshModelingTools/Private/PlaneCutTool.cpp
@@ -0,0 +1,362 @@
+// Copyright 1998-2019 Epic Games, Inc. All Rights Reserved.
+
+#include "PlaneCutTool.h"
+#include "InteractiveToolManager.h"
+#include "ToolBuilderUtil.h"
+
+#include "ToolSetupUtil.h"
+
+#include "DynamicMesh3.h"
+#include "BaseBehaviors/MultiClickSequenceInputBehavior.h"
+#include "Selection/SelectClickedAction.h"
+
+#include "MeshDescriptionToDynamicMesh.h"
+#include "DynamicMeshToMeshDescription.h"
+
+#include "InteractiveGizmoManager.h"
+
+#include "BaseGizmos/GizmoComponents.h"
+#include "BaseGizmos/TransformGizmo.h"
+
+#include "Drawing/MeshDebugDrawing.h"
+#include "AssetGenerationUtil.h"
+
+#include "CuttingOps/PlaneCutOp.h"
+
+#define LOCTEXT_NAMESPACE "UPlaneCutTool"
+
+
+/*
+ * ToolBuilder
+ */
+
+
+bool UPlaneCutToolBuilder::CanBuildTool(const FToolBuilderState& SceneState) const
+{
+	return ToolBuilderUtil::CountComponents(SceneState, CanMakeComponentTarget) == 1;
+}
+
+UInteractiveTool* UPlaneCutToolBuilder::BuildTool(const FToolBuilderState& SceneState) const
+{
+	UPlaneCutTool* NewTool = NewObject<UPlaneCutTool>(SceneState.ToolManager);
+
+	UActorComponent* ActorComponent = ToolBuilderUtil::FindFirstComponent(SceneState, CanMakeComponentTarget);
+	auto* MeshComponent = Cast<UPrimitiveComponent>(ActorComponent);
+	check(MeshComponent != nullptr);
+	NewTool->SetSelection(MakeComponentTarget(MeshComponent));
+
+	NewTool->SetWorld(SceneState.World);
+	NewTool->SetAssetAPI(AssetAPI);
+	return NewTool;
+}
+
+
+
+/*
+ * Tool
+ */
+
+UPlaneCutToolProperties::UPlaneCutToolProperties()
+{
+	bDiscardAttributes = false;
+	bKeepBothHalves = false;
+	bFillCutHole = true;
+	SpacingBetweenHalves = 1;
+}
+
+UPlaneCutAdvancedProperties::UPlaneCutAdvancedProperties()
+{
+	bFillSpans = false;
+}
+
+
+UPlaneCutTool::UPlaneCutTool()
+{
+	CutPlaneOrigin = FVector::ZeroVector;
+	CutPlaneOrientation = FQuat::Identity;
+}
+
+void UPlaneCutTool::SetWorld(UWorld* World)
+{
+	this->TargetWorld = World;
+}
+
+void UPlaneCutTool::Setup()
+{
+	UInteractiveTool::Setup();
+
+	// hide input StaticMeshComponent
+	ComponentTarget->SetOwnerVisibility(false);
+
+	// click to set plane behavior
+	FSelectClickedAction* SetPlaneAction = new FSelectClickedAction();
+	SetPlaneAction->World = this->TargetWorld;
+	SetPlaneAction->OnClickedPositionFunc = [this](const FHitResult& Hit)
+	{
+		SetCutPlaneFromWorldPos(Hit.ImpactPoint, Hit.ImpactNormal);
+		for (UMeshOpPreviewWithBackgroundCompute* Preview : Previews)
+		{
+			Preview->InvalidateResult();
+		}
+	};
+	SetPointInWorldConnector = SetPlaneAction;
+
+	USingleClickInputBehavior* ClickToSetPlaneBehavior = NewObject<USingleClickInputBehavior>();
+	ClickToSetPlaneBehavior->ModifierCheckFunc = FInputDeviceState::IsCtrlKeyDown;
+	ClickToSetPlaneBehavior->Initialize(SetPointInWorldConnector);
+	AddInputBehavior(ClickToSetPlaneBehavior);
+
+	// create proxy and gizmo (but don't attach yet)
+	UInteractiveGizmoManager* GizmoManager = GetToolManager()->GetPairedGizmoManager();
+	PlaneTransformProxy = NewObject<UTransformProxy>(this);
+	PlaneTransformGizmo = GizmoManager->Create3AxisTransformGizmo(TEXT("PlaneCutToolTransformGizmo"), this);
+
+	BasicProperties = NewObject<UPlaneCutToolProperties>(this, TEXT("Plane Cut Settings"));
+	AdvancedProperties = NewObject<UPlaneCutAdvancedProperties>(this, TEXT("Advanced Settings"));
+
+	// initialize our properties
+	AddToolPropertySource(BasicProperties);
+	AddToolPropertySource(AdvancedProperties);
+
+
+	// initialize the PreviewMesh+BackgroundCompute object
+	UpdateNumPreviews();
+
+	// set initial cut plane (also attaches gizmo/proxy)
+	FVector DefaultOrigin, Extents;
+	ComponentTarget->GetOwnerActor()->GetActorBounds(false, DefaultOrigin, Extents);
+	SetCutPlaneFromWorldPos(DefaultOrigin, FVector::UpVector);
+	// hook up callback so further changes trigger recut
+	PlaneTransformProxy->OnTransformChanged.AddUObject(this, &UPlaneCutTool::TransformChanged);
+
+
+
+	// Convert input mesh description to dynamic mesh
+	OriginalDynamicMesh = MakeShared<FDynamicMesh3>();
+	FMeshDescriptionToDynamicMesh Converter;
+	Converter.bPrintDebugMessages = true;
+	Converter.Convert(ComponentTarget->GetMesh(), *OriginalDynamicMesh);
+
+	for (UMeshOpPreviewWithBackgroundCompute* Preview : Previews)
+	{
+		Preview->InvalidateResult();
+	}
+}
+
+
+void UPlaneCutTool::UpdateNumPreviews()
+{
+	int32 CurrentNumPreview = Previews.Num();
+	int32 TargetNumPreview = BasicProperties->bKeepBothHalves ? 2 : 1;
+	if (TargetNumPreview < CurrentNumPreview)
+	{
+		for (int32 PreviewIdx = CurrentNumPreview - 1; PreviewIdx >= TargetNumPreview; PreviewIdx--)
+		{
+			Previews[PreviewIdx]->Cancel();
+		}
+		Previews.SetNum(TargetNumPreview);
+	}
+	else
+	{
+		for (int32 PreviewIdx = CurrentNumPreview; PreviewIdx < TargetNumPreview; PreviewIdx++)
+		{
+			UPlaneCutOperatorFactory *CutSide = NewObject<UPlaneCutOperatorFactory>();
+			CutSide->CutTool = this;
+			CutSide->bCutBackSide = PreviewIdx == 1;
+			UMeshOpPreviewWithBackgroundCompute* Preview = Previews.Add_GetRef(NewObject<UMeshOpPreviewWithBackgroundCompute>(CutSide, "Preview"));
+			Preview->Setup(this->TargetWorld, CutSide);
+			Preview->ConfigureMaterials(
+				ToolSetupUtil::GetDefaultMaterial(GetToolManager(), ComponentTarget->GetMaterial(0)),
+				ToolSetupUtil::GetDefaultWorkingMaterial(GetToolManager())
+			);
+			Preview->SetVisibility(true);
+		}
+	}
+}
+
+
+void UPlaneCutTool::Shutdown(EToolShutdownType ShutdownType)
+{
+	// Restore (unhide) the source meshes
+	ComponentTarget->SetOwnerVisibility(true);
+
+	TArray<TUniquePtr<FDynamicMeshOpResult>> Results;
+	for (UMeshOpPreviewWithBackgroundCompute* Preview : Previews)
+	{
+		Results.Emplace(Preview->Shutdown());
+	}
+	if (ShutdownType == EToolShutdownType::Accept)
+	{
+		GenerateAsset(Results);
+	}
+
+	if (SetPointInWorldConnector != nullptr)
+	{
+		delete SetPointInWorldConnector;
+	}
+	UInteractiveGizmoManager* GizmoManager = GetToolManager()->GetPairedGizmoManager();
+	GizmoManager->DestroyAllGizmosByOwner(this);
+}
+
+void UPlaneCutTool::SetAssetAPI(IToolsContextAssetAPI* AssetAPIIn)
+{
+	this->AssetAPI = AssetAPIIn;
+}
+
+TSharedPtr<FDynamicMeshOperator> UPlaneCutOperatorFactory::MakeNewOperator()
+{
+	TSharedPtr<FPlaneCutOp> CutOp = MakeShared<FPlaneCutOp>();
+	CutOp->bDiscardAttributes = CutTool->BasicProperties->bDiscardAttributes;
+	CutOp->bFillCutHole = CutTool->BasicProperties->bFillCutHole;
+	CutOp->bFillSpans = CutTool->AdvancedProperties->bFillSpans;
+
+	FTransform LocalToWorld = CutTool->ComponentTarget->GetWorldTransform();
+	FTransform WorldToLocal = LocalToWorld.Inverse();
+	FVector LocalOrigin = WorldToLocal.TransformPosition(CutTool->CutPlaneOrigin);
+	FVector WorldNormal = CutTool->CutPlaneOrientation.GetAxisZ();
+	FVector LocalNormal = WorldToLocal.TransformVectorNoScale(WorldNormal);
+	if (bCutBackSide)
+	{
+		LocalNormal = -LocalNormal;
+	}
+	CutOp->LocalPlaneOrigin = LocalOrigin;
+	CutOp->LocalPlaneNormal = LocalNormal;
+	CutOp->OriginalMesh = CutTool->OriginalDynamicMesh;
+	
+	if (bCutBackSide)
+	{
+		LocalToWorld *= FTransform(CutTool->BasicProperties->SpacingBetweenHalves * WorldNormal);
+	}
+	CutOp->SetTransform(LocalToWorld);
+
+	return CutOp;
+}
+
+
+
+void UPlaneCutTool::Render(IToolsContextRenderAPI* RenderAPI)
+{
+	FPrimitiveDrawInterface* PDI = RenderAPI->GetPrimitiveDrawInterface();
+	FColor GridColor(128, 128, 128, 32);
+	float GridThickness = 0.5f;
+	float GridLineSpacing = 25.0f;   // @todo should be relative to view
+	int NumGridLines = 10;
+	
+	FFrame3f DrawFrame(CutPlaneOrigin, CutPlaneOrientation);
+	MeshDebugDraw::DrawSimpleGrid(DrawFrame, NumGridLines, GridLineSpacing, GridThickness, GridColor, false, PDI, FTransform::Identity);
+}
+
+void UPlaneCutTool::Tick(float DeltaTime)
+{
+	for (UMeshOpPreviewWithBackgroundCompute* Preview : Previews)
+	{
+		Preview->Tick(DeltaTime);
+	}
+}
+
+
+#if WITH_EDITOR
+void UPlaneCutTool::PostEditChangeProperty(FPropertyChangedEvent& PropertyChangedEvent)
+{
+	UpdateNumPreviews();
+	for (UMeshOpPreviewWithBackgroundCompute* Preview : Previews)
+	{
+		Preview->InvalidateResult();
+	}
+}
+#endif
+
+void UPlaneCutTool::OnPropertyModified(UObject* PropertySet, UProperty* Property)
+{
+	UpdateNumPreviews();
+	for (UMeshOpPreviewWithBackgroundCompute* Preview : Previews)
+	{
+		Preview->InvalidateResult();
+	}
+}
+
+
+void UPlaneCutTool::TransformChanged(UTransformProxy* Proxy, FTransform Transform)
+{
+	// TODO: if multi-select is re-enabled, only invalidate the preview that actually needs it?
+	CutPlaneOrientation = Transform.GetRotation();
+	CutPlaneOrigin = (FVector)Transform.GetTranslation();
+	for (UMeshOpPreviewWithBackgroundCompute* Preview : Previews)
+	{
+		Preview->InvalidateResult();
+	}
+}
+
+
+void UPlaneCutTool::SetCutPlaneFromWorldPos(const FVector& Position, const FVector& Normal)
+{
+	CutPlaneOrigin = Position;
+
+	FFrame3f CutPlane(Position, Normal);
+	CutPlaneOrientation = CutPlane.Rotation;
+
+	PlaneTransformProxy->SetTransform(CutPlane.ToFTransform());
+	PlaneTransformGizmo->SetActiveTarget(PlaneTransformProxy);
+}
+
+
+bool UPlaneCutTool::HasAccept() const
+{
+	return true;
+}
+
+bool UPlaneCutTool::CanAccept() const
+{
+	for (UMeshOpPreviewWithBackgroundCompute* Preview : Previews)
+	{
+		if (!Preview->HaveValidResult())
+		{
+			return false;
+		}
+	}
+	return true;
+}
+
+
+void UPlaneCutTool::GenerateAsset(const TArray<TUniquePtr<FDynamicMeshOpResult>>& Results)
+{
+	GetToolManager()->BeginUndoTransaction(LOCTEXT("PlaneCutToolTransactionName", "Plane Cut Tool"));
+	
+
+	// currently in-place replaces the first half, and adds a new actor for the second half (if it was generated)
+	// TODO: options to support other choices re what should be a new actor
+	check(Results.Num() > 0);
+	check(Results[0]->Mesh.Get() != nullptr);
+	ComponentTarget->CommitMesh([&Results](FMeshDescription* MeshDescription)
+	{
+		FDynamicMeshToMeshDescription Converter;
+		Converter.Convert(Results[0]->Mesh.Get(), *MeshDescription);
+	});
+
+
+	// The method for creating a new mesh (AssetGenerationUtil::GenerateStaticMeshActor) is editor-only; just creating the other half if not in editor
+#if WITH_EDITOR
+	if (Results.Num() == 2)
+	{
+		FSelectedOjectsChangeList NewSelection;
+		NewSelection.ModificationType = ESelectedObjectsModificationType::Replace;
+		NewSelection.Actors.Add(ComponentTarget->GetOwnerActor());
+
+		check(Results[1]->Mesh.Get() != nullptr);
+		// TODO: copy over material?
+		AActor* NewActor = AssetGenerationUtil::GenerateStaticMeshActor(
+			AssetAPI, TargetWorld,
+			Results[1]->Mesh.Get(), Results[1]->Transform, TEXT("Plane Cut Other Half"),
+			AssetGenerationUtil::GetDefaultAutoGeneratedAssetPath());
+		NewSelection.Actors.Add(NewActor);
+		GetToolManager()->RequestSelectionChange(NewSelection);
+	}
+#endif
+
+	GetToolManager()->EndUndoTransaction();
+}
+
+
+
+
+#undef LOCTEXT_NAMESPACE
diff --git a/Engine/Plugins/Experimental/MeshModelingToolset/Source/MeshModelingTools/Private/UVProjectionTool.cpp b/Engine/Plugins/Experimental/MeshModelingToolset/Source/MeshModelingTools/Private/UVProjectionTool.cpp
--- a/Engine/Plugins/Experimental/MeshModelingToolset/Source/MeshModelingTools/Private/UVProjectionTool.cpp
+++ b/Engine/Plugins/Experimental/MeshModelingToolset/Source/MeshModelingTools/Private/UVProjectionTool.cpp
@@ -202,17 +202,8 @@ void UUVProjectionTool::Shutdown(EToolShutdownType ShutdownType)
 		GenerateAsset(Results);
 	}
 
-	// TODO MAKE THIS UNNECESSARY GIZMOS SHOULD BE AUTO DESTROYED!!
-	for (UTransformGizmo* TransformGizmo : TransformGizmos)
-	{
-		if (TransformGizmo != nullptr)
-		{
-			GizmoManager->DestroyGizmo(TransformGizmo);
-		}
-	}
+	GizmoManager->DestroyAllGizmosByOwner(this);
 	TransformGizmos.Empty();
-
-	GizmoManager->DeregisterGizmoType(TEXT("TransformGizmo"));
 }
 
 void UUVProjectionTool::SetAssetAPI(IToolsContextAssetAPI* AssetAPIIn)
diff --git a/Engine/Plugins/Experimental/MeshModelingToolset/Source/MeshModelingTools/Public/PlaneCutTool.h b/Engine/Plugins/Experimental/MeshModelingToolset/Source/MeshModelingTools/Public/PlaneCutTool.h
new file mode 100644
--- /dev/null
+++ b/Engine/Plugins/Experimental/MeshModelingToolset/Source/MeshModelingTools/Public/PlaneCutTool.h
@@ -0,0 +1,190 @@
+// Copyright 1998-2019 Epic Games, Inc. All Rights Reserved.
+
+#pragma once
+
+#include "CoreMinimal.h"
+#include "UObject/NoExportTypes.h"
+#include "InteractiveGizmo.h"
+#include "SingleSelectionTool.h"
+#include "InteractiveToolBuilder.h"
+#include "MeshOpPreviewHelpers.h"
+#include "DynamicMesh3.h"
+#include "BaseTools/SingleClickTool.h"
+#include "PlaneCutTool.generated.h"
+
+
+// predeclarations
+struct FMeshDescription;
+class USimpleDynamicMeshComponent;
+class UTransformGizmo;
+class UTransformProxy;
+
+
+
+
+
+/**
+ *
+ */
+UCLASS()
+class MESHMODELINGTOOLS_API UPlaneCutToolBuilder : public UInteractiveToolBuilder
+{
+	GENERATED_BODY()
+
+public:
+	IToolsContextAssetAPI* AssetAPI = nullptr;
+
+	virtual bool CanBuildTool(const FToolBuilderState& SceneState) const override;
+	virtual UInteractiveTool* BuildTool(const FToolBuilderState& SceneState) const override;
+};
+
+
+
+
+
+/**
+ * Standard properties of the plane cut operation
+ */
+UCLASS()
+class MESHMODELINGTOOLS_API UPlaneCutToolProperties : public UInteractiveToolPropertySet
+{
+	GENERATED_BODY()
+
+public:
+	UPlaneCutToolProperties();
+
+	/** If true, UVs and Normals are discarded  */
+	UPROPERTY(EditAnywhere, Category = Options)
+	bool bDiscardAttributes;
+
+	/** If true, both halves of the cut are computed */
+	UPROPERTY(EditAnywhere, Category = Options)
+	bool bKeepBothHalves;
+
+	/** If keeping both halves, separate the two pieces by this amount */
+	UPROPERTY(EditAnywhere, Category = Options)
+	float SpacingBetweenHalves;
+
+
+	/** If true, the cut surface is filled with simple planar hole fill surface(s) */
+	UPROPERTY(EditAnywhere, Category = Advanced)
+	bool bFillCutHole;
+};
+
+
+
+
+/**
+ * Advanced properties of the plane cut operation
+ */
+UCLASS()
+class MESHMODELINGTOOLS_API UPlaneCutAdvancedProperties : public UInteractiveToolPropertySet
+{
+	GENERATED_BODY()
+
+public:
+	UPlaneCutAdvancedProperties();
+
+	/** 
+	 *  If true, try to treat open edge spans as fillable loops
+	 *  TODO: add some tolerance-based span joining code to try to turn open spans into loops instead?
+	 */
+	UPROPERTY(EditAnywhere, Category = Advanced)
+	bool bFillSpans;
+
+};
+
+
+UCLASS()
+class MESHMODELINGTOOLS_API UPlaneCutOperatorFactory : public UObject, public IDynamicMeshOperatorFactory
+{
+	GENERATED_BODY()
+
+public:
+	// IDynamicMeshOperatorFactory API
+	virtual TSharedPtr<FDynamicMeshOperator> MakeNewOperator() override;
+
+	UPROPERTY()
+	UPlaneCutTool *CutTool;
+
+	UPROPERTY()
+	bool bCutBackSide = false;
+};
+
+/**
+ * Simple Mesh Plane Cutting Tool
+ */
+UCLASS()
+class MESHMODELINGTOOLS_API UPlaneCutTool : public USingleSelectionTool
+{
+	GENERATED_BODY()
+
+public:
+
+	friend UPlaneCutOperatorFactory;
+
+	UPlaneCutTool();
+
+	virtual void Setup() override;
+	virtual void Shutdown(EToolShutdownType ShutdownType) override;
+
+	virtual void SetWorld(UWorld* World);
+	virtual void SetAssetAPI(IToolsContextAssetAPI* AssetAPI);
+
+	virtual void Tick(float DeltaTime) override;
+	virtual void Render(IToolsContextRenderAPI* RenderAPI) override;
+
+	virtual bool HasCancel() const override { return true; }
+	virtual bool HasAccept() const override;
+	virtual bool CanAccept() const override;
+
+#if WITH_EDITOR
+	virtual void PostEditChangeProperty(FPropertyChangedEvent &PropertyChangedEvent) override;
+#endif
+
+	virtual void OnPropertyModified(UObject* PropertySet, UProperty* Property) override;
+
+protected:
+
+	UPROPERTY()
+	UPlaneCutToolProperties* BasicProperties;
+
+	UPROPERTY()
+	UPlaneCutAdvancedProperties* AdvancedProperties;
+
+	/** Origin of cutting plane */
+	UPROPERTY()
+	FVector CutPlaneOrigin;
+
+	/** Orientation of cutting plane */
+	UPROPERTY()
+	FQuat CutPlaneOrientation;
+
+	UPROPERTY()
+	TArray<UMeshOpPreviewWithBackgroundCompute*> Previews;
+
+
+protected:
+	TSharedPtr<FDynamicMesh3> OriginalDynamicMesh;
+
+	UWorld* TargetWorld;
+	IToolsContextAssetAPI* AssetAPI;
+
+	FViewCameraState CameraState;
+
+	UPROPERTY()
+	UTransformGizmo* PlaneTransformGizmo;
+
+	UPROPERTY()
+	UTransformProxy* PlaneTransformProxy;
+
+	void TransformChanged(UTransformProxy* Proxy, FTransform Transform);
+
+	void UpdateNumPreviews();
+
+	IClickBehaviorTarget* SetPointInWorldConnector = nullptr;
+
+	virtual void SetCutPlaneFromWorldPos(const FVector& Position, const FVector& Normal);
+
+	void GenerateAsset(const TArray<TUniquePtr<FDynamicMeshOpResult>>& Results);
+};
diff --git a/Engine/Plugins/Experimental/MeshModelingToolset/Source/MeshModelingToolsEditorOnly/Private/PlaneCutTool.cpp b/Engine/Plugins/Experimental/MeshModelingToolset/Source/MeshModelingToolsEditorOnly/Private/PlaneCutTool.cpp
deleted file mode 100644
diff --git a/Engine/Plugins/Experimental/MeshModelingToolset/Source/MeshModelingToolsEditorOnly/Public/PlaneCutTool.h b/Engine/Plugins/Experimental/MeshModelingToolset/Source/MeshModelingToolsEditorOnly/Public/PlaneCutTool.h
deleted file mode 100644
diff --git a/Engine/Plugins/Experimental/MeshModelingToolset/Source/ModelingOperators/ModelingOperators.Build.cs b/Engine/Plugins/Experimental/MeshModelingToolset/Source/ModelingOperators/ModelingOperators.Build.cs
--- a/Engine/Plugins/Experimental/MeshModelingToolset/Source/ModelingOperators/ModelingOperators.Build.cs
+++ b/Engine/Plugins/Experimental/MeshModelingToolset/Source/ModelingOperators/ModelingOperators.Build.cs
@@ -31,7 +31,6 @@ public class ModelingOperators : ModuleRules
 				// editor dependencies
 				//
 				//"MeshSolverUtilities",
-				//"GeometryAlgorithms", // required for constrained Delaunay triangulation
                 //"ProxyLODMeshReduction", // currently required to be public due to IVoxelBasedCSG API
                 //"MeshDescription",
 
@@ -40,6 +39,7 @@ public class ModelingOperators : ModuleRules
 				"GeometricObjects",
 				"DynamicMesh",
 				"MeshConversion",
+                "GeometryAlgorithms", // required for constrained Delaunay triangulation
                 "MeshSolverUtilities", // required by the smoothing operators
 				// ... add other public dependencies that you statically link with here ...
 			}
diff --git a/Engine/Plugins/Experimental/MeshModelingToolset/Source/ModelingOperators/Private/CuttingOps/PlaneCutOp.cpp b/Engine/Plugins/Experimental/MeshModelingToolset/Source/ModelingOperators/Private/CuttingOps/PlaneCutOp.cpp
new file mode 100644
--- /dev/null
+++ b/Engine/Plugins/Experimental/MeshModelingToolset/Source/ModelingOperators/Private/CuttingOps/PlaneCutOp.cpp
@@ -0,0 +1,52 @@
+// Copyright 1998-2019 Epic Games, Inc. All Rights Reserved.
+
+#pragma once
+
+#include "CuttingOps/PlaneCutOp.h"
+
+#include "Engine/StaticMesh.h"
+
+#include "MeshDescriptionToDynamicMesh.h"
+#include "MeshSimplification.h"
+#include "MeshConstraintsUtil.h"
+#include "Operations/MeshPlaneCut.h"
+#include "ConstrainedDelaunay2.h"
+
+
+void FPlaneCutOp::SetTransform(const FTransform& Transform) {
+	ResultTransform = (FTransform3d)Transform;
+}
+
+void FPlaneCutOp::CalculateResult(FProgressCancel* Progress)
+{
+	if (Progress->Cancelled())
+	{
+		return;
+	}
+	ResultMesh->Copy(*OriginalMesh, true, true, true, !bDiscardAttributes);
+
+	FMeshPlaneCut Cut(ResultMesh.Get(), LocalPlaneOrigin, LocalPlaneNormal);
+
+	if (Progress->Cancelled())
+	{
+		return;
+	}
+	Cut.UVScaleFactor = 1.0 / ResultMesh->GetBounds().MaxDim();
+
+
+	if (Progress->Cancelled())
+	{
+		return;
+	}
+	Cut.Cut();
+
+
+	if (Progress->Cancelled())
+	{
+		return;
+	}
+	if (bFillCutHole)
+	{
+		Cut.HoleFill(ConstrainedDelaunayTriangulate<double>, bFillSpans);
+	}
+}
\ No newline at end of file
diff --git a/Engine/Plugins/Experimental/MeshModelingToolset/Source/ModelingOperators/Public/CuttingOps/PlaneCutOp.h b/Engine/Plugins/Experimental/MeshModelingToolset/Source/ModelingOperators/Public/CuttingOps/PlaneCutOp.h
new file mode 100644
--- /dev/null
+++ b/Engine/Plugins/Experimental/MeshModelingToolset/Source/ModelingOperators/Public/CuttingOps/PlaneCutOp.h
@@ -0,0 +1,32 @@
+// Copyright 1998-2019 Epic Games, Inc. All Rights Reserved.
+
+#pragma once
+
+#include "CoreMinimal.h"
+#include "Util/ProgressCancel.h"
+#include "ModelingOperators.h"
+
+
+
+
+
+class MODELINGOPERATORS_API FPlaneCutOp : public FDynamicMeshOperator
+{
+public:
+	virtual ~FPlaneCutOp() {}
+
+	// inputs
+	FVector LocalPlaneOrigin, LocalPlaneNormal;
+	bool bFillCutHole, bDiscardAttributes, bFillSpans;
+	TSharedPtr<FDynamicMesh3> OriginalMesh;
+
+	void SetTransform(const FTransform& Transform);
+
+	//
+	// FDynamicMeshOperator implementation
+	// 
+
+	virtual void CalculateResult(FProgressCancel* Progress) override;
+};
+
+
diff --git a/Engine/Plugins/Experimental/MeshModelingToolset/Source/ModelingOperatorsEditorOnly/Private/CuttingOps/PlaneCutOp.cpp b/Engine/Plugins/Experimental/MeshModelingToolset/Source/ModelingOperatorsEditorOnly/Private/CuttingOps/PlaneCutOp.cpp
deleted file mode 100644
diff --git a/Engine/Plugins/Experimental/MeshModelingToolset/Source/ModelingOperatorsEditorOnly/Public/CuttingOps/PlaneCutOp.h b/Engine/Plugins/Experimental/MeshModelingToolset/Source/ModelingOperatorsEditorOnly/Public/CuttingOps/PlaneCutOp.h
deleted file mode 100644
