Rename "make mesh test tool" -> "add primitive tool"

#rnx

diff --git a/Engine/Plugins/Experimental/MeshModelingToolset/Source/MeshModelingTools/Private/AddPrimitiveTool.cpp b/Engine/Plugins/Experimental/MeshModelingToolset/Source/MeshModelingTools/Private/AddPrimitiveTool.cpp
new file mode 100644
--- /dev/null
+++ b/Engine/Plugins/Experimental/MeshModelingToolset/Source/MeshModelingTools/Private/AddPrimitiveTool.cpp
@@ -0,0 +1,304 @@
+// Copyright 1998-2019 Epic Games, Inc. All Rights Reserved.
+
+#include "AddPrimitiveTool.h"
+#include "ToolBuilderUtil.h"
+#include "InteractiveToolManager.h"
+#include "BaseBehaviors/MouseHoverBehavior.h"
+#include "Selection/ToolSelectionUtil.h"
+#include "AssetGenerationUtil.h"
+
+#include "MeshDescriptionBuilder.h"
+#include "Generators/SweepGenerator.h"
+#include "Generators/GridBoxMeshGenerator.h"
+#include "Generators/RectangleMeshGenerator.h"
+#include "Generators/SphereGenerator.h"
+#include "DynamicMeshAttributeSet.h"
+
+#include "StaticMeshComponentBuilder.h"
+#include "Drawing/MeshDebugDrawing.h"
+
+#define LOCTEXT_NAMESPACE "UAddPrimitiveTool"
+
+
+/*
+ * ToolBuilder
+ */
+
+
+bool UAddPrimitiveToolBuilder::CanBuildTool(const FToolBuilderState& SceneState) const
+{
+	return (this->AssetAPI != nullptr);
+}
+
+UInteractiveTool* UAddPrimitiveToolBuilder::BuildTool(const FToolBuilderState& SceneState) const
+{
+	UAddPrimitiveTool* NewTool = NewObject<UAddPrimitiveTool>(SceneState.ToolManager);
+	NewTool->SetWorld(SceneState.World);
+	NewTool->SetAssetAPI(AssetAPI);
+	return NewTool;
+}
+
+
+
+/*
+ * Tool
+ */
+
+UProceduralShapeToolProperties::UProceduralShapeToolProperties()
+{
+	Shape = EMakeMeshShapeType::Box;
+	Width = 100;
+	Height = 200;
+	//StartAngle = 0;
+	//EndAngle = 360;
+	Slices = 16;
+	Subdivisions = 0;
+	bCentered = false;
+	PlaceMode = EMakeMeshPlacementType::OnScene;
+}
+
+
+void UAddPrimitiveTool::SetWorld(UWorld* World)
+{
+	this->TargetWorld = World;
+}
+
+void UAddPrimitiveTool::SetAssetAPI(IToolsContextAssetAPI* AssetAPIIn)
+{
+	this->AssetAPI = AssetAPIIn;
+}
+
+
+void UAddPrimitiveTool::Setup()
+{
+	USingleClickTool::Setup();
+
+	UMouseHoverBehavior* HoverBehavior = NewObject<UMouseHoverBehavior>(this);
+	HoverBehavior->Initialize(this);
+	AddInputBehavior(HoverBehavior);
+
+	ShapeSettings = NewObject<UProceduralShapeToolProperties>(this);
+	AddToolPropertySource(ShapeSettings);
+
+	// create preview mesh object
+	PreviewMesh = NewObject<UPreviewMesh>(this, TEXT("PreviewMesh"));
+	PreviewMesh->CreateInWorld(TargetWorld, FTransform::Identity);
+	PreviewMesh->SetVisible(false);
+	UpdatePreviewMesh();
+
+	//if (MeshSource->GetMaterial(0) != nullptr)
+	//{
+	//	PreviewMesh->SetMaterial(MeshSource->GetMaterial(0));
+	//}
+}
+
+
+void UAddPrimitiveTool::Shutdown(EToolShutdownType ShutdownType)
+{
+	PreviewMesh->SetVisible(false);
+	PreviewMesh->Disconnect();
+	PreviewMesh = nullptr;
+}
+
+
+void UAddPrimitiveTool::Render(IToolsContextRenderAPI* RenderAPI)
+{
+	//FPrimitiveDrawInterface* PDI = RenderAPI->GetPrimitiveDrawInterface();
+	//MeshDebugDraw::DrawSimpleGrid(ShapeFrame, 13, 5.0f, 1.0f, FColor::Orange, false, PDI, FTransform::Identity);
+}
+
+
+
+void UAddPrimitiveTool::OnPropertyModified(UObject* PropertySet, UProperty* Property)
+{
+	UpdatePreviewMesh();
+}
+
+
+
+
+void UAddPrimitiveTool::OnUpdateHover(const FInputDeviceRay& DevicePos)
+{
+	UpdatePreviewPosition(DevicePos);
+}
+
+
+void UAddPrimitiveTool::UpdatePreviewPosition(const FInputDeviceRay& DeviceClickPos)
+{
+	FRay ClickPosWorldRay = DeviceClickPos.WorldRay;
+
+	// hit position (temp)
+	bool bHit = false;
+
+	if (ShapeSettings->PlaceMode == EMakeMeshPlacementType::OnPlane)
+	{
+		FPlane DrawPlane(FVector::ZeroVector, FVector(0, 0, 1));
+		FVector DrawPlanePos = FMath::RayPlaneIntersection(ClickPosWorldRay.Origin, ClickPosWorldRay.Direction, DrawPlane);
+		bHit = true;
+		ShapeFrame = FFrame3f(DrawPlanePos);
+	}
+	else
+	{
+		// cast ray into scene
+		FVector RayStart = ClickPosWorldRay.Origin;
+		FVector RayEnd = ClickPosWorldRay.PointAt(999999);
+		FCollisionObjectQueryParams QueryParams(FCollisionObjectQueryParams::AllObjects);
+		FHitResult Result;
+		bHit = TargetWorld->LineTraceSingleByObjectType(Result, RayStart, RayEnd, QueryParams);
+		if (bHit)
+		{
+			ShapeFrame = FFrame3f(Result.ImpactPoint, Result.ImpactNormal);
+			ShapeFrame.ConstrainedAlignPerpAxes();
+		}
+	}
+
+	if (ShapeSettings->Rotation != 0)
+	{
+		ShapeFrame.Rotate(FQuaternionf(ShapeFrame.Z(), ShapeSettings->Rotation, true));
+	}
+
+	if (bHit)
+	{
+		PreviewMesh->SetVisible(true);
+		PreviewMesh->SetTransform(ShapeFrame.ToFTransform());
+	}
+	else
+	{
+		PreviewMesh->SetVisible(false);
+	}
+}
+
+
+void UAddPrimitiveTool::UpdatePreviewMesh()
+{
+	FDynamicMesh3 NewMesh;
+	switch (ShapeSettings->Shape)
+	{
+	case EMakeMeshShapeType::Plane:
+		GeneratePlane(&NewMesh);
+		break;
+
+	case EMakeMeshShapeType::Cylinder:
+		GenerateCylinder(&NewMesh);
+		break;
+
+	case EMakeMeshShapeType::Cone:
+		GenerateCone(&NewMesh);
+		break;
+
+	case EMakeMeshShapeType::Sphere:
+		GenerateSphere(&NewMesh);
+		break;
+
+	case EMakeMeshShapeType::Box:
+	default:
+		GenerateBox(&NewMesh);
+		break;
+	}
+
+	// set mesh position
+	FAxisAlignedBox3d Bounds = NewMesh.GetCachedBounds();
+	FVector3d TargetOrigin = Bounds.Center();
+	if (ShapeSettings->bCentered == false)
+	{
+		TargetOrigin.Z = Bounds.Min.Z;
+	}
+	for (int vid : NewMesh.VertexIndicesItr())
+	{
+		FVector3d Pos = NewMesh.GetVertex(vid);
+		Pos -= TargetOrigin;
+		NewMesh.SetVertex(vid, Pos);
+	}
+
+
+	PreviewMesh->UpdatePreview(&NewMesh);
+}
+
+
+
+
+void UAddPrimitiveTool::OnClicked(const FInputDeviceRay& DeviceClickPos)
+{
+#if WITH_EDITOR
+	const FDynamicMesh3* CurMesh = PreviewMesh->GetPreviewDynamicMesh();
+	FTransform3d CurTransform(PreviewMesh->GetTransform());
+
+	GetToolManager()->BeginUndoTransaction(LOCTEXT("MergeMeshes", "Merge Meshes"));
+
+	AActor* NewActor = AssetGenerationUtil::GenerateStaticMeshActor(
+		AssetAPI, TargetWorld,
+		CurMesh, CurTransform, TEXT("ShapeMesh"),
+		AssetGenerationUtil::GetDefaultAutoGeneratedAssetPath());
+	
+	// select newly-created object
+	ToolSelectionUtil::SetNewActorSelection(GetToolManager(), NewActor);
+
+	GetToolManager()->EndUndoTransaction();
+#else
+	checkNoEntry();
+#endif
+}
+
+
+
+void UAddPrimitiveTool::GenerateBox(FDynamicMesh3* OutMesh)
+{
+	FGridBoxMeshGenerator BoxGen;
+	BoxGen.Box = FOrientedBox3d(FVector3d::Zero(), 0.5*FVector3d(ShapeSettings->Width, ShapeSettings->Width, ShapeSettings->Height));
+	int EdgeNum = ShapeSettings->Subdivisions + 2;
+	BoxGen.EdgeVertices = FIndex3i(EdgeNum, EdgeNum, EdgeNum);
+	BoxGen.Generate();
+	OutMesh->Copy(&BoxGen);
+}
+
+
+void UAddPrimitiveTool::GeneratePlane(FDynamicMesh3* OutMesh)
+{
+	FRectangleMeshGenerator RectGen;
+	RectGen.Width = ShapeSettings->Width;
+	RectGen.Height = ShapeSettings->Height;
+	RectGen.WidthVertexCount = RectGen.HeightVertexCount = ShapeSettings->Subdivisions + 2;
+	RectGen.Generate();
+	OutMesh->Copy(&RectGen);
+}
+
+
+void UAddPrimitiveTool::GenerateCylinder(FDynamicMesh3* OutMesh)
+{
+	FCylinderGenerator CylGen;
+	CylGen.Radius[0] = ShapeSettings->Width * 0.5f;
+	CylGen.Radius[1] = CylGen.Radius[0];
+	CylGen.Height = ShapeSettings->Height;
+	CylGen.AngleSamples = ShapeSettings->Slices;
+	CylGen.bCapped = true;
+	CylGen.Generate();
+	OutMesh->Copy(&CylGen);
+}
+
+
+void UAddPrimitiveTool::GenerateCone(FDynamicMesh3* OutMesh)
+{
+	// Unreal's standard cone is just a cylinder with a very small top
+	FCylinderGenerator CylGen;
+	CylGen.Radius[0] = ShapeSettings->Width * 0.5f;
+	CylGen.Radius[1] = .01;
+	CylGen.Height = ShapeSettings->Height;
+	CylGen.AngleSamples = ShapeSettings->Slices;
+	CylGen.bCapped = true;
+	CylGen.Generate();
+	OutMesh->Copy(&CylGen);
+}
+
+
+void UAddPrimitiveTool::GenerateSphere(FDynamicMesh3* OutMesh)
+{
+	FSphereGenerator SphereGen;
+	SphereGen.Radius = ShapeSettings->Width * 0.5f;
+	SphereGen.NumTheta = ShapeSettings->Slices;
+	SphereGen.NumPhi = ShapeSettings->Slices;
+	SphereGen.Generate();
+	OutMesh->Copy(&SphereGen);
+}
+
+
+#undef LOCTEXT_NAMESPACE
diff --git a/Engine/Plugins/Experimental/MeshModelingToolset/Source/MeshModelingTools/Private/MakeMeshTestTool.cpp b/Engine/Plugins/Experimental/MeshModelingToolset/Source/MeshModelingTools/Private/MakeMeshTestTool.cpp
deleted file mode 100644
diff --git a/Engine/Plugins/Experimental/MeshModelingToolset/Source/MeshModelingTools/Public/AddPrimitiveTool.h b/Engine/Plugins/Experimental/MeshModelingToolset/Source/MeshModelingTools/Public/AddPrimitiveTool.h
new file mode 100644
--- /dev/null
+++ b/Engine/Plugins/Experimental/MeshModelingToolset/Source/MeshModelingTools/Public/AddPrimitiveTool.h
@@ -0,0 +1,167 @@
+// Copyright 1998-2019 Epic Games, Inc. All Rights Reserved.
+
+#pragma once
+
+#include "CoreMinimal.h"
+#include "UObject/NoExportTypes.h"
+#include "InteractiveToolBuilder.h"
+#include "BaseTools/SingleClickTool.h"
+#include "PreviewMesh.h"
+#include "AddPrimitiveTool.generated.h"
+
+class FDynamicMesh3;
+
+/**
+ *
+ */
+UCLASS()
+class MESHMODELINGTOOLS_API UAddPrimitiveToolBuilder : public UInteractiveToolBuilder
+{
+	GENERATED_BODY()
+
+public:
+	IToolsContextAssetAPI* AssetAPI;
+
+	UAddPrimitiveToolBuilder() 
+	{
+		AssetAPI = nullptr;
+	}
+
+	virtual bool CanBuildTool(const FToolBuilderState& SceneState) const override;
+	virtual UInteractiveTool* BuildTool(const FToolBuilderState& SceneState) const override;
+};
+
+
+
+/** Shape Types */
+UENUM()
+enum class EMakeMeshShapeType : uint8
+{
+	Box = 0 UMETA(DisplayName = "Box"),
+	Cylinder = 1 UMETA(DisplayName = "Cylinder"),
+	Cone = 2 UMETA(DisplayName = "Cone"),
+	Plane = 3 UMETA(DisplayName = "Plane"),
+	Sphere = 4 UMETA(DisplayName = "Sphere")
+};
+
+
+
+/** Placement Target Types */
+UENUM()
+enum class EMakeMeshPlacementType : uint8
+{
+	OnPlane = 0,
+	OnScene = 1
+};
+
+
+
+
+
+UCLASS()
+class MESHMODELINGTOOLS_API UProceduralShapeToolProperties : public UInteractiveToolPropertySet
+{
+	GENERATED_BODY()
+
+public:
+	UProceduralShapeToolProperties();
+
+	/** Type of shape to generate */
+	UPROPERTY(EditAnywhere, Category = ShapeSettings)
+	EMakeMeshShapeType Shape;
+
+	/** Width of Shape */
+	UPROPERTY(EditAnywhere, Category = ShapeSettings, meta = (DisplayName = "Width", UIMin = "1.0", UIMax = "1000.0", ClampMin = "0.0001", ClampMax = "1000000.0"))
+	float Width;
+
+	/** Height of Shape */
+	UPROPERTY(EditAnywhere, Category = ShapeSettings, meta = (DisplayName = "Height", UIMin = "1.0", UIMax = "1000.0", ClampMin = "0.0001", ClampMax = "1000000.0"))
+	float Height;
+
+	/** Rotation around up axis */
+	UPROPERTY(EditAnywhere, Category = ShapeSettings, meta = (DisplayName = "Rotation", UIMin = "0.0", UIMax = "360.0"))
+	float Rotation;
+
+	/** Type of shape to generate */
+	UPROPERTY(EditAnywhere, Category = ShapeSettings)
+	EMakeMeshPlacementType PlaceMode;
+
+	/** Center shape at click point */
+	UPROPERTY(EditAnywhere, Category = ShapeSettings)
+	bool bCentered;
+
+	///** Start Angle of Shape */
+	//UPROPERTY(EditAnywhere, Category = ShapeSettings, meta = (DisplayName = "Start Angle", UIMin = "0.0", UIMax = "360.0", ClampMin = "-10000", ClampMax = "10000.0"))
+	//float StartAngle;
+
+	///** End Angle of Shape */
+	//UPROPERTY(EditAnywhere, Category = ShapeSettings, meta = (DisplayName = "End Angle", UIMin = "0.0", UIMax = "360.0", ClampMin = "-10000", ClampMax = "10000.0"))
+	//float EndAngle;
+
+
+	/** Number of Slices */
+	UPROPERTY(EditAnywhere, Category = ShapeSettings, meta = (DisplayName = "Slices", UIMin = "3", UIMax = "128", ClampMin = "3", ClampMax = "999"))
+	int Slices;
+
+	/** Subdivisions */
+	UPROPERTY(EditAnywhere, Category = ShapeSettings, meta = (DisplayName = "Subdivisions", UIMin = "0", UIMax = "100", ClampMin = "0", ClampMax = "4000"))
+	int Subdivisions;
+
+};
+
+
+
+
+
+
+
+/**
+ *
+ */
+UCLASS()
+class MESHMODELINGTOOLS_API UAddPrimitiveTool : public USingleClickTool, public IHoverBehaviorTarget
+{
+	GENERATED_BODY()
+
+public:
+	virtual void SetWorld(UWorld* World);
+	virtual void SetAssetAPI(IToolsContextAssetAPI* AssetAPI);
+
+	virtual void Setup() override;
+	virtual void Shutdown(EToolShutdownType ShutdownType) override;
+
+	virtual void Render(IToolsContextRenderAPI* RenderAPI) override;
+
+	virtual bool HasCancel() const override { return false; }
+	virtual bool HasAccept() const override { return false; }
+	virtual bool CanAccept() const override { return false; }
+
+	virtual void OnPropertyModified(UObject* PropertySet, UProperty* Property) override;
+
+	virtual void OnUpdateHover(const FInputDeviceRay& DevicePos) override;
+	virtual void OnClicked(const FInputDeviceRay& ClickPos) override;
+
+
+
+protected:
+	UPROPERTY()
+	UProceduralShapeToolProperties* ShapeSettings;
+
+	UPROPERTY()
+	UPreviewMesh* PreviewMesh;
+
+protected:
+	UWorld* TargetWorld;
+	IToolsContextAssetAPI* AssetAPI;
+
+	void UpdatePreviewPosition(const FInputDeviceRay& ClickPos);
+	FFrame3f ShapeFrame;
+
+	void UpdatePreviewMesh();
+
+	void GenerateCylinder(FDynamicMesh3* OutMesh);
+	void GenerateCone(FDynamicMesh3* OutMesh);
+	void GenerateBox(FDynamicMesh3* OutMesh);
+	void GeneratePlane(FDynamicMesh3* OutMesh);
+	void GenerateSphere(FDynamicMesh3* OutMesh);
+};
\ No newline at end of file
diff --git a/Engine/Plugins/Experimental/MeshModelingToolset/Source/MeshModelingTools/Public/MakeMeshTestTool.h b/Engine/Plugins/Experimental/MeshModelingToolset/Source/MeshModelingTools/Public/MakeMeshTestTool.h
deleted file mode 100644
diff --git a/Engine/Plugins/Experimental/ModelingToolsEditorMode/Source/ModelingToolsEditorMode/Private/ModelingToolsEditorMode.cpp b/Engine/Plugins/Experimental/ModelingToolsEditorMode/Source/ModelingToolsEditorMode/Private/ModelingToolsEditorMode.cpp
--- a/Engine/Plugins/Experimental/ModelingToolsEditorMode/Source/ModelingToolsEditorMode/Private/ModelingToolsEditorMode.cpp
+++ b/Engine/Plugins/Experimental/ModelingToolsEditorMode/Source/ModelingToolsEditorMode/Private/ModelingToolsEditorMode.cpp
@@ -14,7 +14,7 @@
 #include "DynamicMeshSculptTool.h"
 #include "PolygonMeshDeformTool.h"
 #include "ConvertToPolygonsTool.h"
-#include "MakeMeshTestTool.h"
+#include "AddPrimitiveTool.h"
 #include "SmoothMeshTool.h"
 #include "RemeshMeshTool.h"
 #include "SimplifyMeshTool.h"
@@ -93,7 +93,7 @@ void FModelingToolsEditorMode::BuildModeToolbar(FToolBarBuilder& ToolbarBuilder)
 
 	ToolbarBuilder.AddSeparator();
 
-	ToolbarBuilder.AddToolBarButton(FModelingToolsManagerCommands::Get().BeginMakePrimitiveTool);
+	ToolbarBuilder.AddToolBarButton(FModelingToolsManagerCommands::Get().BeginAddPrimitiveTool);
 	ToolbarBuilder.AddToolBarButton(FModelingToolsManagerCommands::Get().BeginDrawPolygonTool);
 	if (CVarEnablePrototypeModelingTools.GetValueOnGameThread() > 0)
 	{
@@ -379,9 +379,9 @@ void FModelingToolsEditorMode::Enter()
 	//
 	// make shape tools
 	//
-	auto MakeMeshTestToolBuilder = NewObject<UMakeMeshTestToolBuilder>();
-	MakeMeshTestToolBuilder->AssetAPI = ToolsContext->GetAssetAPI();
-	RegisterToolFunc(ToolManagerCommands.BeginMakePrimitiveTool, TEXT("MakeMeshTestTool"), MakeMeshTestToolBuilder);
+	auto AddPrimitiveToolBuilder = NewObject<UAddPrimitiveToolBuilder>();
+	AddPrimitiveToolBuilder->AssetAPI = ToolsContext->GetAssetAPI();
+	RegisterToolFunc(ToolManagerCommands.BeginAddPrimitiveTool, TEXT("AddPrimitiveTool"), AddPrimitiveToolBuilder);
 
 	auto DrawPolygonToolBuilder = NewObject<UDrawPolygonToolBuilder>();
 	DrawPolygonToolBuilder->AssetAPI = ToolsContext->GetAssetAPI();
diff --git a/Engine/Plugins/Experimental/ModelingToolsEditorMode/Source/ModelingToolsEditorMode/Private/ModelingToolsEditorModeStyle.cpp b/Engine/Plugins/Experimental/ModelingToolsEditorMode/Source/ModelingToolsEditorMode/Private/ModelingToolsEditorModeStyle.cpp
--- a/Engine/Plugins/Experimental/ModelingToolsEditorMode/Source/ModelingToolsEditorMode/Private/ModelingToolsEditorModeStyle.cpp
+++ b/Engine/Plugins/Experimental/ModelingToolsEditorMode/Source/ModelingToolsEditorMode/Private/ModelingToolsEditorModeStyle.cpp
@@ -77,8 +77,8 @@ void FModelingToolsEditorModeStyle::Initialize()
 		StyleSet->Set("ModelingToolsManagerCommands.CompleteActiveTool.Small", new IMAGE_PLUGIN_BRUSH("Icons/icon_ActiveTool_Accept_40x", Icon28x28));
 
 
-		StyleSet->Set("ModelingToolsManagerCommands.BeginMakePrimitiveTool", new IMAGE_PLUGIN_BRUSH("Icons/icon_Tool_MakePrimitive_40x", Icon28x28));
-		StyleSet->Set("ModelingToolsManagerCommands.BeginMakePrimitiveTool.Small", new IMAGE_PLUGIN_BRUSH("Icons/icon_Tool_MakePrimitive_40x", Icon28x28));
+		StyleSet->Set("ModelingToolsManagerCommands.BeginAddPrimitiveTool", new IMAGE_PLUGIN_BRUSH("Icons/icon_Tool_MakePrimitive_40x", Icon28x28));
+		StyleSet->Set("ModelingToolsManagerCommands.BeginAddPrimitiveTool.Small", new IMAGE_PLUGIN_BRUSH("Icons/icon_Tool_MakePrimitive_40x", Icon28x28));
 		StyleSet->Set("ModelingToolsManagerCommands.BeginDrawPolygonTool", new IMAGE_PLUGIN_BRUSH("Icons/icon_Tool_DrawPolygon_40x", Icon28x28));
 		StyleSet->Set("ModelingToolsManagerCommands.BeginDrawPolygonTool.Small", new IMAGE_PLUGIN_BRUSH("Icons/icon_Tool_DrawPolygon_40x", Icon28x28));
 		StyleSet->Set("ModelingToolsManagerCommands.BeginShapeSprayTool", new IMAGE_PLUGIN_BRUSH("Icons/icon_Tool_ShapeSpray_40x", Icon28x28));
diff --git a/Engine/Plugins/Experimental/ModelingToolsEditorMode/Source/ModelingToolsEditorMode/Private/ModelingToolsManagerActions.cpp b/Engine/Plugins/Experimental/ModelingToolsEditorMode/Source/ModelingToolsEditorMode/Private/ModelingToolsManagerActions.cpp
--- a/Engine/Plugins/Experimental/ModelingToolsEditorMode/Source/ModelingToolsEditorMode/Private/ModelingToolsManagerActions.cpp
+++ b/Engine/Plugins/Experimental/ModelingToolsEditorMode/Source/ModelingToolsEditorMode/Private/ModelingToolsManagerActions.cpp
@@ -21,7 +21,7 @@ FModelingToolsManagerCommands::FModelingToolsManagerCommands() :
 
 void FModelingToolsManagerCommands::RegisterCommands()
 {
-	UI_COMMAND(BeginMakePrimitiveTool, "Primitive", "Start the Make Primitive Tool", EUserInterfaceActionType::Button, FInputChord());
+	UI_COMMAND(BeginAddPrimitiveTool, "Primitive", "Start the Add Primitive Tool", EUserInterfaceActionType::Button, FInputChord());
 	UI_COMMAND(BeginDrawPolygonTool, "Polygon", "Start the Draw Polygon Tool", EUserInterfaceActionType::Button, FInputChord());
 	UI_COMMAND(BeginShapeSprayTool, "ShapeSpray", "Start the Shape Spray Tool", EUserInterfaceActionType::Button, FInputChord());
 
diff --git a/Engine/Plugins/Experimental/ModelingToolsEditorMode/Source/ModelingToolsEditorMode/Public/ModelingToolsManagerActions.h b/Engine/Plugins/Experimental/ModelingToolsEditorMode/Source/ModelingToolsEditorMode/Public/ModelingToolsManagerActions.h
--- a/Engine/Plugins/Experimental/ModelingToolsEditorMode/Source/ModelingToolsEditorMode/Public/ModelingToolsManagerActions.h
+++ b/Engine/Plugins/Experimental/ModelingToolsEditorMode/Source/ModelingToolsEditorMode/Public/ModelingToolsManagerActions.h
@@ -15,7 +15,7 @@ public:
 	FModelingToolsManagerCommands();
 
 
-	TSharedPtr<FUICommandInfo> BeginMakePrimitiveTool;
+	TSharedPtr<FUICommandInfo> BeginAddPrimitiveTool;
 	TSharedPtr<FUICommandInfo> BeginDrawPolygonTool;
 	TSharedPtr<FUICommandInfo> BeginShapeSprayTool;
 
