MeshModelingToolset: Add Primitive and Draw Polygon now have configurable materials

diff --git a/Engine/Plugins/Experimental/MeshModelingToolset/Source/MeshModelingTools/Private/AddPrimitiveTool.cpp b/Engine/Plugins/Experimental/MeshModelingToolset/Source/MeshModelingTools/Private/AddPrimitiveTool.cpp
--- a/Engine/Plugins/Experimental/MeshModelingToolset/Source/MeshModelingTools/Private/AddPrimitiveTool.cpp
+++ b/Engine/Plugins/Experimental/MeshModelingToolset/Source/MeshModelingTools/Private/AddPrimitiveTool.cpp
@@ -55,6 +55,7 @@ UProceduralShapeToolProperties::UProceduralShapeToolProperties()
 	Subdivisions = 0;
 	bCentered = false;
 	PlaceMode = EMakeMeshPlacementType::OnScene;
+    Material = CreateDefaultSubobject<UMaterialInterface>(TEXT("MATERIAL"));
 }
 
 
@@ -84,12 +85,9 @@ void UAddPrimitiveTool::Setup()
 	PreviewMesh = NewObject<UPreviewMesh>(this, TEXT("PreviewMesh"));
 	PreviewMesh->CreateInWorld(TargetWorld, FTransform::Identity);
 	PreviewMesh->SetVisible(false);
-	UpdatePreviewMesh();
+	PreviewMesh->SetMaterial(ShapeSettings->Material);
 
-	//if (MeshSource->GetMaterial(0) != nullptr)
-	//{
-	//	PreviewMesh->SetMaterial(MeshSource->GetMaterial(0));
-	//}
+	UpdatePreviewMesh();
 }
 
 
@@ -111,6 +109,7 @@ void UAddPrimitiveTool::Render(IToolsContextRenderAPI* RenderAPI)
 
 void UAddPrimitiveTool::OnPropertyModified(UObject* PropertySet, UProperty* Property)
 {
+	PreviewMesh->SetMaterial(ShapeSettings->Material);
 	UpdatePreviewMesh();
 }
 
@@ -222,13 +221,15 @@ void UAddPrimitiveTool::OnClicked(const FInputDeviceRay& DeviceClickPos)
 #if WITH_EDITOR
 	const FDynamicMesh3* CurMesh = PreviewMesh->GetPreviewDynamicMesh();
 	FTransform3d CurTransform(PreviewMesh->GetTransform());
-
+	UMaterialInterface* Material = PreviewMesh->GetMaterial();
 	GetToolManager()->BeginUndoTransaction(LOCTEXT("MergeMeshes", "Merge Meshes"));
 
 	AActor* NewActor = AssetGenerationUtil::GenerateStaticMeshActor(
 		AssetAPI, TargetWorld,
 		CurMesh, CurTransform, TEXT("ShapeMesh"),
-		AssetGenerationUtil::GetDefaultAutoGeneratedAssetPath());
+		AssetGenerationUtil::GetDefaultAutoGeneratedAssetPath(),
+		Material
+		);
 	
 	// select newly-created object
 	ToolSelectionUtil::SetNewActorSelection(GetToolManager(), NewActor);
diff --git a/Engine/Plugins/Experimental/MeshModelingToolset/Source/MeshModelingTools/Private/DrawPolygonTool.cpp b/Engine/Plugins/Experimental/MeshModelingToolset/Source/MeshModelingTools/Private/DrawPolygonTool.cpp
--- a/Engine/Plugins/Experimental/MeshModelingToolset/Source/MeshModelingTools/Private/DrawPolygonTool.cpp
+++ b/Engine/Plugins/Experimental/MeshModelingToolset/Source/MeshModelingTools/Private/DrawPolygonTool.cpp
@@ -41,12 +41,9 @@
 /*
  * ToolBuilder
  */
-
-
 constexpr int StartPointSnapID = FPointPlanarSnapSolver::BaseExternalPointID + 1;
 constexpr int CurrentSceneSnapID = FPointPlanarSnapSolver::BaseExternalPointID + 2;
 
-
 bool UDrawPolygonToolBuilder::CanBuildTool(const FToolBuilderState& SceneState) const
 {
 	return (this->AssetAPI != nullptr);
@@ -60,12 +57,17 @@ UInteractiveTool* UDrawPolygonToolBuilder::BuildTool(const FToolBuilderState& Sc
 	return NewTool;
 }
 
-
+/*
+ * Properties
+ */
+UDrawPolygonToolStandardProperties::UDrawPolygonToolStandardProperties()
+{
+	Material = CreateDefaultSubobject<UMaterialInterface>(TEXT("MATERIAL"));
+}
 
 /*
  * Tool
  */
-
 UDrawPolygonTool::UDrawPolygonTool()
 {
 	DrawPlaneOrigin = FVector::ZeroVector;
@@ -73,7 +75,6 @@ UDrawPolygonTool::UDrawPolygonTool()
 	bInInteractiveExtrude = false;
 }
 
-
 void UDrawPolygonTool::SetWorld(UWorld* World)
 {
 	this->TargetWorld = World;
@@ -84,9 +85,6 @@ void UDrawPolygonTool::SetAssetAPI(IToolsContextAssetAPI* AssetAPIIn)
 	this->AssetAPI = AssetAPIIn;
 }
 
-
-
-
 void UDrawPolygonTool::Setup()
 {
 	UInteractiveTool::Setup();
@@ -133,6 +131,7 @@ void UDrawPolygonTool::Setup()
 	PreviewMesh = NewObject<UPreviewMesh>(this, TEXT("PreviewMesh"));
 	PreviewMesh->CreateInWorld(this->TargetWorld, FTransform::Identity);
 	PreviewMesh->SetVisible(false);
+	PreviewMesh->SetMaterial(PolygonProperties->Material);
 	bPreviewUpdatePending = false;
 
 	SnapEngine.SnapMetricTolerance = ToolSceneQueriesUtil::GetDefaultVisualAngleSnapThreshD();
@@ -166,9 +165,6 @@ void UDrawPolygonTool::Shutdown(EToolShutdownType ShutdownType)
 	GizmoManager->DeregisterGizmoType(TEXT("DrawPolygonPlaneGizmo"));
 }
 
-
-
-
 void UDrawPolygonTool::RegisterActions(FInteractiveToolActionSet& ActionSet)
 {
 	ActionSet.RegisterAction(this, (int32)EStandardToolActions::BaseClientDefinedActionID + 1,
@@ -179,7 +175,6 @@ void UDrawPolygonTool::RegisterActions(FInteractiveToolActionSet& ActionSet)
 		[this]() { PopLastVertexAction(); });
 }
 
-
 void UDrawPolygonTool::PopLastVertexAction()
 {
 	if (bInInteractiveExtrude || PolygonVertices.Num() == 0)
@@ -215,9 +210,6 @@ void UDrawPolygonTool::PopLastVertexAction()
 	}
 }
 
-
-
-
 void DrawEdgeTicks(FPrimitiveDrawInterface* PDI, 
 	const FSegment3f& Segment, float Height,
 	const FVector3f& PlaneNormal, 
@@ -235,8 +227,6 @@ void DrawEdgeTicks(FPrimitiveDrawInterface* PDI,
 	PDI->DrawLine(A, B, Color, DepthPriorityGroup, LineThickness, 0.0f, bIsScreenSpace);
 }
 
-
-
 void UDrawPolygonTool::Render(IToolsContextRenderAPI* RenderAPI)
 {
 	FPrimitiveDrawInterface* PDI = RenderAPI->GetPrimitiveDrawInterface();
@@ -408,8 +398,6 @@ void UDrawPolygonTool::Render(IToolsContextRenderAPI* RenderAPI)
 	}
 }
 
-
-
 void UDrawPolygonTool::ResetPolygon()
 {
 	PolygonVertices.Reset();
@@ -430,13 +418,11 @@ void UDrawPolygonTool::UpdatePreviewVertex(const FVector& PreviewVertexIn)
 	}
 }
 
-
 void UDrawPolygonTool::AppendVertex(const FVector& Vertex)
 {
 	PolygonVertices.Add(Vertex);
 }
 
-
 bool UDrawPolygonTool::FindDrawPlaneHitPoint(const FInputDeviceRay& ClickPos, FVector& HitPosOut)
 {
 	bHaveSurfaceHit = false;
@@ -504,7 +490,6 @@ bool UDrawPolygonTool::FindDrawPlaneHitPoint(const FInputDeviceRay& ClickPos, FV
 	return true;
 }
 
-
 void UDrawPolygonTool::OnBeginSequencePreview(const FInputDeviceRay& DevicePos)
 {
 	// just update snapped point preview
@@ -516,7 +501,6 @@ void UDrawPolygonTool::OnBeginSequencePreview(const FInputDeviceRay& DevicePos)
 	
 }
 
-
 bool UDrawPolygonTool::CanBeginClickSequence(const FInputDeviceRay& ClickPos)
 {
 	return true;
@@ -692,8 +676,6 @@ bool UDrawPolygonTool::RequestAbortClickSequence()
 	return false;
 }
 
-
-
 void UDrawPolygonTool::OnUpdateModifierState(int ModifierID, bool bIsOn)
 {
 	if (ModifierID == IgnoreSnappingModifier)
@@ -706,9 +688,6 @@ void UDrawPolygonTool::OnUpdateModifierState(int ModifierID, bool bIsOn)
 	}
 }
 
-
-
-
 bool UDrawPolygonTool::UpdateSelfIntersection()
 {
 	bHaveSelfIntersection = false;
@@ -737,8 +716,6 @@ bool UDrawPolygonTool::UpdateSelfIntersection()
 	return false;
 }
 
-
-
 void UDrawPolygonTool::GenerateFixedPolygon(TArray<FVector>& FixedPoints, TArray<FVector>& VerticesOut)
 {
 	FFrame3f DrawFrame(DrawPlaneOrigin, DrawPlaneOrientation);
@@ -779,8 +756,6 @@ void UDrawPolygonTool::GenerateFixedPolygon(TArray<FVector>& FixedPoints, TArray
 	}
 }
 
-
-
 void UDrawPolygonTool::BeginInteractiveExtrude()
 {
 	bInInteractiveExtrude = true;
@@ -788,7 +763,6 @@ void UDrawPolygonTool::BeginInteractiveExtrude()
 	GeneratePreviewHeightTarget();
 }
 
-
 void UDrawPolygonTool::EndInteractiveExtrude()
 {
 	EmitCurrentPolygon();
@@ -799,7 +773,6 @@ void UDrawPolygonTool::EndInteractiveExtrude()
 	bInInteractiveExtrude = false;
 }
 
-
 float UDrawPolygonTool::FindInteractiveHeightDistance(const FInputDeviceRay& ClickPos)
 {
 	float NearestHitDist = TNumericLimits<float>::Max();
@@ -855,10 +828,6 @@ float UDrawPolygonTool::FindInteractiveHeightDistance(const FInputDeviceRay& Cli
 	
 }
 
-
-
-
-
 void UDrawPolygonTool::SetDrawPlaneFromWorldPos(const FVector& Position, const FVector& Normal)
 {
 	DrawPlaneOrigin = Position;
@@ -877,7 +846,6 @@ void UDrawPolygonTool::SetDrawPlaneFromWorldPos(const FVector& Position, const F
 	Gizmo->ExternalUpdatePosition(DrawPlaneOrigin, DrawPlaneOrientation, false);
 }
 
-
 void UDrawPolygonTool::UpdateDrawPlaneFromGizmo(const FFrame3d& WorldPosition)
 {
 	if (bIgnoreSnappingToggle == false)
@@ -895,10 +863,6 @@ void UDrawPolygonTool::UpdateDrawPlaneFromGizmo(const FFrame3d& WorldPosition)
 	}
 }
 
-
-
-
-
 void UDrawPolygonTool::EmitCurrentPolygon()
 {
 #if WITH_EDITOR
@@ -915,7 +879,8 @@ void UDrawPolygonTool::EmitCurrentPolygon()
 	AActor* NewActor = AssetGenerationUtil::GenerateStaticMeshActor(
 		AssetAPI, TargetWorld,
 		&Mesh, PlaneFrameOut.ToTransform(), TEXT("Polygon"),
-		AssetGenerationUtil::GetDefaultAutoGeneratedAssetPath());
+		AssetGenerationUtil::GetDefaultAutoGeneratedAssetPath(),
+		PolygonProperties->Material);
 
 	// select newly-created object
 	ToolSelectionUtil::SetNewActorSelection(GetToolManager(), NewActor);
@@ -927,11 +892,6 @@ void UDrawPolygonTool::EmitCurrentPolygon()
 	ResetPolygon();
 }
 
-
-
-
-
-
 void UDrawPolygonTool::UpdateLivePreview()
 {
 	int NumVerts = PolygonVertices.Num();
@@ -947,12 +907,10 @@ void UDrawPolygonTool::UpdateLivePreview()
 	GeneratePolygonMesh(PolygonVertices, &Mesh, PlaneFrame, false, ExtrudeDist, false);
 
 	PreviewMesh->SetTransform(PlaneFrame.ToFTransform());
+	PreviewMesh->SetMaterial(PolygonProperties->Material);
 	PreviewMesh->UpdatePreview(&Mesh);
 }
 
-
-
-
 void UDrawPolygonTool::GeneratePolygonMesh(const TArray<FVector>& Polygon, FDynamicMesh3* ResultMeshOut, FFrame3d& WorldFrameOut, bool bIncludePreviewVtx, double ExtrudeDistance, bool bExtrudeSymmetric)
 {
 	// construct centered frame for polygon
@@ -1025,17 +983,10 @@ void UDrawPolygonTool::GeneratePolygonMesh(const TArray<FVector>& Polygon, FDyna
 	}
 }
 
-
-
-
 void UDrawPolygonTool::GeneratePreviewHeightTarget()
 {
 	GeneratePolygonMesh(PolygonVertices, &PreviewHeightTarget, PreviewHeightFrame, false, 99999, true);
 	PreviewHeightTargetAABB.SetMesh(&PreviewHeightTarget);
 }
 
-
-
-
-
 #undef LOCTEXT_NAMESPACE
diff --git a/Engine/Plugins/Experimental/MeshModelingToolset/Source/MeshModelingTools/Public/AddPrimitiveTool.h b/Engine/Plugins/Experimental/MeshModelingToolset/Source/MeshModelingTools/Public/AddPrimitiveTool.h
--- a/Engine/Plugins/Experimental/MeshModelingToolset/Source/MeshModelingTools/Public/AddPrimitiveTool.h
+++ b/Engine/Plugins/Experimental/MeshModelingToolset/Source/MeshModelingTools/Public/AddPrimitiveTool.h
@@ -3,6 +3,7 @@
 #pragma once
 
 #include "CoreMinimal.h"
+#include "MaterialShared.h"
 #include "UObject/NoExportTypes.h"
 #include "InteractiveToolBuilder.h"
 #include "BaseTools/SingleClickTool.h"
@@ -107,6 +108,9 @@ public:
 	UPROPERTY(EditAnywhere, Category = ShapeSettings, meta = (DisplayName = "Subdivisions", UIMin = "0", UIMax = "100", ClampMin = "0", ClampMax = "4000"))
 	int Subdivisions;
 
+	/** Material */
+	UPROPERTY(EditAnywhere, Category = ShapeSettings, meta = (DisplayName = "Material"))
+	UMaterialInterface* Material;
 };
 
 
@@ -164,4 +168,4 @@ protected:
 	void GenerateBox(FDynamicMesh3* OutMesh);
 	void GeneratePlane(FDynamicMesh3* OutMesh);
 	void GenerateSphere(FDynamicMesh3* OutMesh);
-};
\ No newline at end of file
+};
diff --git a/Engine/Plugins/Experimental/MeshModelingToolset/Source/MeshModelingTools/Public/DrawPolygonTool.h b/Engine/Plugins/Experimental/MeshModelingToolset/Source/MeshModelingTools/Public/DrawPolygonTool.h
--- a/Engine/Plugins/Experimental/MeshModelingToolset/Source/MeshModelingTools/Public/DrawPolygonTool.h
+++ b/Engine/Plugins/Experimental/MeshModelingToolset/Source/MeshModelingTools/Public/DrawPolygonTool.h
@@ -78,7 +78,9 @@ class MESHMODELINGTOOLS_API UDrawPolygonToolStandardProperties : public UObject
 {
 	GENERATED_BODY()
 
-public:
+	public:
+	UDrawPolygonToolStandardProperties();
+
 	UPROPERTY(EditAnywhere, NonTransactional, Category = Polygon)
 	EDrawPolygonDrawMode PolygonType = EDrawPolygonDrawMode::Freehand;
 
@@ -95,6 +97,10 @@ public:
 
 	UPROPERTY(EditAnywhere, NonTransactional, Category = Polygon)
 	bool bAllowSelfIntersections = false;
+
+	/** Material for new mesh*/
+	UPROPERTY(EditAnywhere, NonTransactional, Category = Polygon, meta = (DisplayName = "Material"))
+	UMaterialInterface* Material;
 };
 
 
diff --git a/Engine/Plugins/Experimental/MeshModelingToolset/Source/ModelingComponents/Private/AssetGenerationUtil.cpp b/Engine/Plugins/Experimental/MeshModelingToolset/Source/ModelingComponents/Private/AssetGenerationUtil.cpp
--- a/Engine/Plugins/Experimental/MeshModelingToolset/Source/ModelingComponents/Private/AssetGenerationUtil.cpp
+++ b/Engine/Plugins/Experimental/MeshModelingToolset/Source/ModelingComponents/Private/AssetGenerationUtil.cpp
@@ -27,7 +27,8 @@ AActor* AssetGenerationUtil::GenerateStaticMeshActor(
 	const FDynamicMesh3* Mesh,
 	const FTransform3d& Transform,
 	FString ObjectName,
-	FString PackagePath
+	FString PackagePath,
+	UMaterialInterface* Material
 )
 {
 	check(AssetAPI);
@@ -74,7 +75,12 @@ AActor* AssetGenerationUtil::GenerateStaticMeshActor(
 
 	// transform new component to origin
 	Builder.NewMeshComponent->SetWorldTransform((FTransform)Transform);
-
+	UMaterialInstance;
+	UMaterialInterface;
+	if (Material)
+	{
+		Builder.NewMeshComponent->SetMaterial(0, Material);
+	}
 	// force save of asset
 	//AssetAPI->InteractiveSaveGeneratedAsset(Builder.NewStaticMesh, AssetPackage);
 	AssetAPI->AutoSaveGeneratedAsset(Builder.NewStaticMesh, AssetPackage);
diff --git a/Engine/Plugins/Experimental/MeshModelingToolset/Source/ModelingComponents/Private/PreviewMesh.cpp b/Engine/Plugins/Experimental/MeshModelingToolset/Source/ModelingComponents/Private/PreviewMesh.cpp
--- a/Engine/Plugins/Experimental/MeshModelingToolset/Source/ModelingComponents/Private/PreviewMesh.cpp
+++ b/Engine/Plugins/Experimental/MeshModelingToolset/Source/ModelingComponents/Private/PreviewMesh.cpp
@@ -56,6 +56,13 @@ void UPreviewMesh::SetMaterial(UMaterialInterface* Material)
 	DynamicMeshComponent->SetMaterial(0, Material);
 }
 
+UMaterialInterface*
+UPreviewMesh::GetMaterial() const
+{
+	check(DynamicMeshComponent);
+	return DynamicMeshComponent->GetMaterial(0);
+}
+
 
 void UPreviewMesh::EnableWireframe(bool bEnable)
 {
diff --git a/Engine/Plugins/Experimental/MeshModelingToolset/Source/ModelingComponents/Public/AssetGenerationUtil.h b/Engine/Plugins/Experimental/MeshModelingToolset/Source/ModelingComponents/Public/AssetGenerationUtil.h
--- a/Engine/Plugins/Experimental/MeshModelingToolset/Source/ModelingComponents/Public/AssetGenerationUtil.h
+++ b/Engine/Plugins/Experimental/MeshModelingToolset/Source/ModelingComponents/Public/AssetGenerationUtil.h
@@ -40,7 +40,8 @@ namespace AssetGenerationUtil
 		const FDynamicMesh3* Mesh,
 		const FTransform3d& Transform,
 		FString ObjectName,
-		FString PackagePath
+		FString PackagePath,
+		UMaterialInterface* Material = nullptr
 	);
 #endif
 
diff --git a/Engine/Plugins/Experimental/MeshModelingToolset/Source/ModelingComponents/Public/PreviewMesh.h b/Engine/Plugins/Experimental/MeshModelingToolset/Source/ModelingComponents/Public/PreviewMesh.h
--- a/Engine/Plugins/Experimental/MeshModelingToolset/Source/ModelingComponents/Public/PreviewMesh.h
+++ b/Engine/Plugins/Experimental/MeshModelingToolset/Source/ModelingComponents/Public/PreviewMesh.h
@@ -89,6 +89,10 @@ public:
 	 */
 	void SetMaterial(UMaterialInterface* Material);
 
+	/**
+	* Get material from the preview mesh
+	*/
+	UMaterialInterface* GetMaterial() const;
 	/**
 	 * Get the current transform on the preview mesh
 	 */
