ModelingTools: Selection tool now always at least selects triangle hit by brush ray, and shows accept/cancel buttons
#rnx

diff --git a/Engine/Plugins/Experimental/MeshModelingToolset/Source/MeshModelingTools/Private/MeshSelectionTool.cpp b/Engine/Plugins/Experimental/MeshModelingToolset/Source/MeshModelingTools/Private/MeshSelectionTool.cpp
--- a/Engine/Plugins/Experimental/MeshModelingToolset/Source/MeshModelingTools/Private/MeshSelectionTool.cpp
+++ b/Engine/Plugins/Experimental/MeshModelingToolset/Source/MeshModelingTools/Private/MeshSelectionTool.cpp
@@ -190,8 +190,14 @@ void UMeshSelectionTool::CalculateTriangleROI(const FBrushStampData& Stamp, TArr
 	FTransform Transform = ComponentTarget->GetWorldTransform();
 	FVector StampPosLocal = Transform.InverseTransformPosition(Stamp.WorldPosition);
 
-	float RadiusSqr = CurrentBrushRadius * CurrentBrushRadius;
+	// always select first triangle
 	const FDynamicMesh3* Mesh = PreviewMesh->GetPreviewDynamicMesh();
+	if (Mesh->IsTriangle(Stamp.HitResult.FaceIndex))
+	{
+		TriangleROI.Add(Stamp.HitResult.FaceIndex);
+	}
+
+	float RadiusSqr = CurrentBrushRadius * CurrentBrushRadius;
 	for (int TriIdx : Mesh->TriangleIndicesItr())
 	{
 		FVector3d Position = Mesh->GetTriCentroid(TriIdx);
diff --git a/Engine/Plugins/Experimental/MeshModelingToolset/Source/MeshModelingTools/Public/MeshSelectionTool.h b/Engine/Plugins/Experimental/MeshModelingToolset/Source/MeshModelingTools/Public/MeshSelectionTool.h
--- a/Engine/Plugins/Experimental/MeshModelingToolset/Source/MeshModelingTools/Public/MeshSelectionTool.h
+++ b/Engine/Plugins/Experimental/MeshModelingToolset/Source/MeshModelingTools/Public/MeshSelectionTool.h
@@ -48,7 +48,7 @@ public:
 	virtual void Render(IToolsContextRenderAPI* RenderAPI) override;
 
 	virtual bool HasCancel() const override { return true; }
-	virtual bool HasAccept() const override { return bHaveModifiedMesh; }
+	virtual bool HasAccept() const override { return true; }
 	virtual bool CanAccept() const override { return bHaveModifiedMesh; }
 
 	// UBaseBrushTool overrides
