Rename PolygonMeshDeformTool -> EditMeshPolygonsTool

#rnx

diff --git a/Engine/Plugins/Experimental/MeshModelingToolset/Source/MeshModelingTools/Private/EditMeshPolygonsTool.cpp b/Engine/Plugins/Experimental/MeshModelingToolset/Source/MeshModelingTools/Private/EditMeshPolygonsTool.cpp
new file mode 100644
--- /dev/null
+++ b/Engine/Plugins/Experimental/MeshModelingToolset/Source/MeshModelingTools/Private/EditMeshPolygonsTool.cpp
@@ -0,0 +1,1411 @@
+// Copyright 1998-2019 Epic Games, Inc. All Rights Reserved.
+
+#include "EditMeshPolygonsTool.h"
+#include "InteractiveToolManager.h"
+#include "ToolBuilderUtil.h"
+
+#include "SegmentTypes.h"
+#include "DynamicMeshAttributeSet.h"
+#include "MeshNormals.h"
+#include "ToolSceneQueriesUtil.h"
+#include "Intersection/IntersectionUtil.h"
+
+#include "Async/ParallelFor.h"
+#include "Containers/BitArray.h"
+
+#define LOCTEXT_NAMESPACE "UEditMeshPolygonsTool"
+
+
+
+//////////////////////////////
+// DEBUG_SETTINGS
+
+//Draw white triangles defining the selection subset
+//#define DEBUG_ROI_TRIANGLES
+
+//Draw pink circles around the handles
+//#define DEBUG_ROI_HANDLES
+
+//Draw points on the ROI vertices, White => Weight == 0, Black => Weight == 1
+//#define DEBUG_ROI_WEIGHTS
+
+//////////////////////////////
+
+
+/*
+ * ToolBuilder
+ */
+UMeshSurfacePointTool* UEditMeshPolygonsToolBuilder::CreateNewTool(const FToolBuilderState& SceneState) const
+{
+	UEditMeshPolygonsTool* DeformTool = NewObject<UEditMeshPolygonsTool>(SceneState.ToolManager);
+	return DeformTool;
+}
+
+/*
+ * Tool
+ */
+UPolyEditTransformProperties::UPolyEditTransformProperties()
+{
+	bUseLocalAxis = true;
+	TransformMode = EQuickTransformerMode::AxisTranslation;
+	bSelectVertices = true;
+	bSelectFaces = true;
+	bSelectEdges = true;
+}
+
+/*
+ * Asynchronous Task
+ */
+
+
+void FConstrainedMeshDeformerTask::UpdateDeformer(const ELaplacianWeightScheme SelectedWeightScheme, const FDynamicMesh3& SrcMesh, const TArray<FConstraintData>& ConstraintArray, const TArray<int32>& SrcIDBufferSubset, bool bNewTransaction, const FRichCurve* Curve)
+{
+	bIsNewTransaction = bNewTransaction;
+	SrcMeshMaxVertexID = SrcMesh.MaxVertexID();
+
+	LaplacianWeightScheme = SelectedWeightScheme;
+	
+	bAttenuateWeights = (Curve != nullptr);
+	if (bAttenuateWeights)
+	{
+		WeightAttenuationCurve = *Curve;
+	}
+
+	// Set-up the subset mesh.
+	if (bIsNewTransaction)
+	{
+		//Copy the part of the mesh we want to deform into the SubsetMesh and create map from Src Mesh to the SubsetMesh.
+		InitializeSubsetMesh(SrcMesh, SrcIDBufferSubset);
+	}
+
+	// only want the subset of constraints that correspond to our subset mesh.
+	{
+		const int32 NumSubsetVerts = SubsetVertexIDToSrcVertexIDMap.Num();
+		SubsetConstraintBuffer.Empty(NumSubsetVerts);
+		SubsetConstraintBuffer.AddUninitialized(NumSubsetVerts);
+
+		for (int32 SubVertexID = 0; SubVertexID < SubsetVertexIDToSrcVertexIDMap.Num(); ++SubVertexID)
+		{
+			int32 SrcVtxID = SubsetVertexIDToSrcVertexIDMap[SubVertexID];
+			SubsetConstraintBuffer[SubVertexID] = ConstraintArray[SrcVtxID];
+		}
+	}
+
+	check(bIsNewTransaction || ConstrainedDeformer.IsValid());
+	
+}
+
+
+void FConstrainedMeshDeformerTask::DoWork()
+{
+
+	//TODO: (simple optimization) -
+	//	Instead of SrcVertexIDtoSubsetVertexIDMap, use SubsetVertexIDToSetVertexIDMap - then we can use the VertexIndicesItr()
+	//	on the SubsetMesh to minimize the quantity of vertex indices we need to iterate at every following step.
+
+	if (*bAbortSource == true)
+	{
+		return;
+	}
+
+	if (bIsNewTransaction) //Will only be true once per input transaction (click+drag)
+	{
+		
+		// Create a new deformation solver.
+		ConstrainedDeformer = MakeUnique<FConstrainedMeshDeformer>(SubsetMesh, LaplacianWeightScheme);
+
+		if (bAttenuateWeights)
+		{
+			ApplyAttenuation();
+		}
+
+		//Update our deformer's constraints before deforming using the copy of the constraint buffer
+		for (int32 SubsetVertexID = 0; SubsetVertexID < SubsetConstraintBuffer.Num(); ++SubsetVertexID)
+		{
+			FConstraintData& CData = SubsetConstraintBuffer[SubsetVertexID];
+			ConstrainedDeformer->AddConstraint(SubsetVertexID, CData.Weight, CData.Position, CData.bPostFix);
+		}
+
+		bIsNewTransaction = false;
+
+	}
+	else
+	{
+		//This else block is run every consecutive frame after the start of the input transaction because UpdateConstraintPosition() is very cheap (no factorizing or rebuilding)
+		//Update only the positions of the constraints, as the weights cannot change mid-transaction
+
+		for (int32 SubsetVertexID = 0; SubsetVertexID < SubsetConstraintBuffer.Num(); ++SubsetVertexID)
+		{
+			FConstraintData& CData = SubsetConstraintBuffer[SubsetVertexID];
+			ConstrainedDeformer->UpdateConstraintPosition(SubsetVertexID, CData.Position, CData.bPostFix);
+		}
+
+	}
+
+	if (*bAbortSource == true)
+	{
+		return;
+	}
+
+	//Run the deformation process
+
+	const bool bSuccessfulSolve = ConstrainedDeformer->Deform(SubsetPositionBuffer);
+	
+	
+	if (bSuccessfulSolve)
+	{
+		if (*bAbortSource == true)
+		{
+			return;
+		}
+	}
+	else
+	{
+		UE_LOG(LogTemp, Warning, TEXT("Laplacian deformation failed"));
+	}
+
+}
+
+inline void FConstrainedMeshDeformerTask::InitializeSubsetMesh(const FDynamicMesh3& SrcMesh, const TArray<int32>& SrcIDBufferSubset)
+{
+	//These can be re-used until the user stops dragging
+	SubsetMesh.Clear();
+	SubsetPositionBuffer.Reset();
+
+	//Initialize every element to -1, helps us keep track of vertices we've already added while iterating the triangles
+	TArray<int32> SrcVertexIDToSubsetVertexIDMap;
+	SrcVertexIDToSubsetVertexIDMap.Init(INACTIVE_SUBSET_ID, SrcMeshMaxVertexID);
+
+
+	//Iterate the triangle array to append vertices, and then triangles to the temporary subset mesh all at once
+	for (int32 i = 0; i < SrcIDBufferSubset.Num(); i += 3)
+	{
+
+		// Build the triangle
+		FIndex3i Triangle;
+		for (int32 v = 0; v < 3; ++v)
+		{
+			//It's the SrcVertexID because every element in the SrcIDBufferSubset is the Vertex ID of a vertex in the original mesh.
+			const int32 SrcVertexID = SrcIDBufferSubset[i + v];
+			int32& SubsetID = SrcVertexIDToSubsetVertexIDMap[SrcVertexID];
+			
+			if (SubsetID == INACTIVE_SUBSET_ID) // we haven't already visited this vertex
+			{
+				const FVector3d Vertex = SrcMesh.GetVertex(SrcVertexID);
+				SubsetID = SubsetMesh.AppendVertex(Vertex);
+			}
+
+			Triangle[v] = SubsetID;
+		}
+		SubsetMesh.AppendTriangle(Triangle);
+	}
+
+	// create a mapping back to the original vertex IDs from the subset mesh
+	int32 MaxSubMeshVertexID = SubsetMesh.MaxVertexID(); // Really MaxID + 1
+	SubsetVertexIDToSrcVertexIDMap.Reset(MaxSubMeshVertexID);
+	SubsetVertexIDToSrcVertexIDMap.AddUninitialized(MaxSubMeshVertexID);
+
+	for (int32 SrcID = 0; SrcID < SrcVertexIDToSubsetVertexIDMap.Num(); ++SrcID)
+	{
+		const int32 SubsetVertexID = SrcVertexIDToSubsetVertexIDMap[SrcID];
+		if (SubsetVertexID != INACTIVE_SUBSET_ID)
+		{
+			SubsetVertexIDToSrcVertexIDMap[SubsetVertexID] = SrcID;
+		}
+	}
+
+}
+
+void FConstrainedMeshDeformerTask::ExportResults(FDynamicMesh3& TargetMesh) const
+{
+
+	//Update the position buffer result
+	for (int32 SubsetVertexID = 0; SubsetVertexID < SubsetVertexIDToSrcVertexIDMap.Num(); ++SubsetVertexID)
+	{
+		const int32 SrcVertexID = SubsetVertexIDToSrcVertexIDMap[SubsetVertexID];
+		const FVector3d Position = SubsetPositionBuffer[SubsetVertexID];
+		
+		TargetMesh.SetVertex(SrcVertexID, Position);
+	}
+}
+
+void FConstrainedMeshDeformerTask::ApplyAttenuation()
+{
+	size_t Size = SrcMeshMaxVertexID;
+	TSet<int> Handles;
+
+	auto InPlaceMinMaxElements = [](FVector3d& Min, FVector3d& Max, const FVector3d Test)
+	{
+		for (uint8 i = 0; i < 3; ++i)
+		{
+			Min[i] = Test[i] < Min[i] ? Test[i] : Min[i];
+			Max[i] = Test[i] > Max[i] ? Test[i] : Max[i];
+		}
+	};
+
+	//Experimental approach: Just going to try grabbing the bounding box of the entire mesh, then the bounding box of the handles as a point cloud.
+	//						 We need a T value to pass to the Weights curve, so let's try finding the distance of each vertex V from line segment formed by the min/max handles
+	//						 Divide the distance from the handles to vertex V by the length of the mesh's bounding box extent, 
+	//						 and that will provide a **ROUGH** approximation of the time value for our curve.
+	//
+	//											  Distance( LineSegment(MaxHandle,MinHandle) , V )
+	// where T(V) is time value at V     T(V) = -----------------------------------------------------
+	//   and V is the position								Length(MeshMin - MeshMax)   
+	//      of each vertex
+
+	FVector3d Min{ std::numeric_limits<double>::max(),std::numeric_limits<double>::max(),std::numeric_limits<double>::max() };
+	FVector3d Max{ std::numeric_limits<double>::min(),std::numeric_limits<double>::min(),std::numeric_limits<double>::min() };
+	FVector3d MinHandles = Min;
+	FVector3d MaxHandles = Max;
+	double LeastWeight = std::numeric_limits<double>::max();
+
+
+
+	for (int32 SubVertexID = 0; SubVertexID < SubsetConstraintBuffer.Num(); ++SubVertexID)
+	{
+
+		FConstraintData& CData = SubsetConstraintBuffer[SubVertexID];
+		// Update bounding box
+		InPlaceMinMaxElements(Min, Max, CData.Position);
+		
+		if (CData.Weight > 0.0)
+		{
+			LeastWeight = CData.Weight < LeastWeight ? CData.Weight : LeastWeight;
+
+			// update bounding box
+			InPlaceMinMaxElements(MinHandles, MaxHandles, CData.Position);
+			Handles.Add(SubVertexID);
+		}
+	}
+
+	double ExtentLength = Min.Distance(Max);
+
+	// Is this why the system has memory?
+	for (int32 SubVertexID = 0; SubVertexID < SubsetConstraintBuffer.Num(); ++SubVertexID)
+	{
+		if (!Handles.Contains(SubVertexID))
+		{
+			FConstraintData& CData = SubsetConstraintBuffer[SubVertexID];
+			double T = CData.Position.Distance(FMath::ClosestPointOnSegment(CData.Position, MinHandles, MaxHandles)) / ExtentLength;
+			CData.Weight = WeightAttenuationCurve.Eval(T) * LeastWeight;
+		}
+
+	}
+}
+
+
+/*
+*	FGroupTopologyLaplacianDeformer methods
+*/
+
+
+void FGroupTopologyLaplacianDeformer::InitBackgroundWorker(const ELaplacianWeightScheme WeightScheme)
+{
+
+	//Initialize asynchronous deformation objects
+	if (AsyncMeshDeformTask == nullptr)
+	{
+		AsyncMeshDeformTask = new FAsyncTaskExecuterWithAbort<FConstrainedMeshDeformerTask>(WeightScheme);
+	}
+}
+
+void FGroupTopologyLaplacianDeformer::InitializeConstraintBuffer()
+{
+	//MaxVertexID is used because the array is potentially sparse.
+	int MaxVertexID = Mesh->MaxVertexID();
+
+	SrcMeshConstraintBuffer.SetNum(MaxVertexID);
+	
+	for (int32 VertexID : Mesh->VertexIndicesItr())
+	{
+		FConstraintData& CD = SrcMeshConstraintBuffer[VertexID];
+		CD.Position = Mesh->GetVertex(VertexID); 
+		CD.Weight   = 0.0;
+		CD.bPostFix = false;
+	}
+}
+
+bool FGroupTopologyLaplacianDeformer::IsTaskInFlight() const
+{
+	return ( AsyncMeshDeformTask != nullptr && !AsyncMeshDeformTask->IsDone() );
+}
+
+bool FGroupTopologyLaplacianDeformer::UpdateAndLaunchdWorker(const ELaplacianWeightScheme SelectedWeightScheme, const FRichCurve*  Curve)
+{
+	/* Deformer needs to run if we've modified the constraints since the last time it finished. */
+	if (AsyncMeshDeformTask == nullptr)
+	{
+		InitBackgroundWorker(SelectedWeightScheme);
+	}
+
+	if (bDeformerNeedsToRun && AsyncMeshDeformTask->IsDone())
+	{		
+		bool bRebuildSubsetMesh = bTaskSubmeshIsDirty;
+
+		FConstrainedMeshDeformerTask& Task = AsyncMeshDeformTask->GetTask();
+
+		// Update the deformer's buffers and weight scheme
+		// this creates the subset mesh if needed.
+		Task.UpdateDeformer(SelectedWeightScheme, *Mesh, SrcMeshConstraintBuffer, SubsetIDBuffer, bRebuildSubsetMesh, Curve);
+
+		// task now has valid submesh
+
+		bTaskSubmeshIsDirty = false;
+
+		//Launch second thread
+		AsyncMeshDeformTask->StartBackgroundTask();
+
+		bDeformerNeedsToRun = false;	 // This was set to true above in UpdateSolution()
+		bVertexPositionsNeedSync = true; // The task will generate new vertex positions. 
+
+		return true;
+
+	}
+	return false;
+}
+
+void FGroupTopologyLaplacianDeformer::SetActiveHandleFaces(const TArray<int>& FaceGroupIDs)
+{
+	Reset();
+
+	check(FaceGroupIDs.Num() == 1);   // multi-face not supported yet
+	int GroupID = FaceGroupIDs[0];
+
+	// find set of vertices in handle 
+	Topology->CollectGroupVertices(GroupID, HandleVertices);
+	Topology->CollectGroupBoundaryVertices(GroupID, HandleBoundaryVertices);
+	ModifiedVertices = HandleVertices;
+
+	// find neighbour group set
+	TArray<int> HandleGroups = FaceGroupIDs;
+	const TArray<int>& GroupNbrGroups = Topology->GetGroupNbrGroups(GroupID);
+
+
+	CalculateROI(HandleGroups, GroupNbrGroups);
+
+
+	TArray<int> NeighborGroups;
+
+	//Add to output parameter 
+	NeighborGroups.Add(GroupID);
+	NeighborGroups.Append(GroupNbrGroups);
+
+	UpdateSelection(Mesh, NeighborGroups, bLocalize);
+
+	// Save the positions of the selected region.
+	SaveInitialPositions();
+
+}
+
+
+void FGroupTopologyLaplacianDeformer::SetActiveHandleEdges(const TArray<int>& TopologyEdgeIDs)
+{
+	Reset();
+
+	for (int EdgeID : TopologyEdgeIDs)
+	{
+		const TArray<int>& EdgeVerts = Topology->GetGroupEdgeVertices(EdgeID);
+		for (int VertID : EdgeVerts)
+		{
+			HandleVertices.Add(VertID);
+		}
+	}
+	HandleBoundaryVertices = HandleVertices;
+	ModifiedVertices = HandleVertices;
+
+	TArray<int> HandleGroups;
+	TArray<int> NbrGroups;
+	Topology->FindEdgeNbrGroups(TopologyEdgeIDs, NbrGroups);
+
+	CalculateROI(HandleGroups, NbrGroups);
+
+	UpdateSelection(Mesh, NbrGroups, bLocalize);
+
+	// Save the positions of the selected region.
+	SaveInitialPositions();
+}
+
+void FGroupTopologyLaplacianDeformer::SetActiveHandleCorners(const TArray<int>& CornerIDs)
+{
+	Reset();
+
+	for (int CornerID : CornerIDs)
+	{
+		int VertID = Topology->GetCornerVertexID(CornerID);
+		if (VertID >= 0)
+		{
+			HandleVertices.Add(VertID);
+		}
+	}
+	HandleBoundaryVertices = HandleVertices;
+	ModifiedVertices = HandleVertices;
+
+	TArray<int> HandleGroups;
+	TArray<int> NbrGroups;
+
+	Topology->FindCornerNbrGroups(CornerIDs, NbrGroups);
+
+
+	CalculateROI(HandleGroups, NbrGroups);
+
+	UpdateSelection(Mesh, NbrGroups, bLocalize);
+
+	// Save the positions of the selected region.
+	SaveInitialPositions();
+}
+
+
+
+void FGroupTopologyLaplacianDeformer::UpdateSelection(const FDynamicMesh3* TargetMesh, const TArray<int>& Groups, bool bLocalizeDeformation)
+{
+
+	// Build an index buffer (SubsetIdBuffer) and a vertexId buffer (ModifidedVertices) for the region we want to change
+
+	if (bLocalizeDeformation)
+	{
+		//For each group ID, retrieve the array of all TriangleIDs associated with that GroupID and append that array to the end of the TriSet to remove duplicates
+		TSet<int> TriSet;
+		for (const int32& GroupID : Groups)
+		{
+			TriSet.Append(Topology->GetGroupFaces(GroupID));
+		}//Now we have every triangle ID involved in the transaction 
+
+		//Since we are flattening the Face to a set of 3 indices, we do 3 * number of triangles though it is too many.
+		SubsetIDBuffer.Reset(3 * TriSet.Num());
+		//Add each triangle's A,B, and C indices to the subset triangle array.
+		for (const int& Tri : TriSet)
+		{
+			FIndex3i Triple = TargetMesh->GetTriangle(Tri);
+			SubsetIDBuffer.Add(Triple.A);
+			SubsetIDBuffer.Add(Triple.B);
+			SubsetIDBuffer.Add(Triple.C);
+		
+		}
+	}
+	else
+	{
+		// the entire mesh.
+		const int32 NumTris = TargetMesh->TriangleCount();
+		SubsetIDBuffer.Reset(3 * NumTris);
+		for (int TriId : TargetMesh->TriangleIndicesItr()) 
+		{ 
+			FIndex3i Triple = TargetMesh->GetTriangle(TriId);
+			SubsetIDBuffer.Add(Triple.A);
+			SubsetIDBuffer.Add(Triple.B);
+			SubsetIDBuffer.Add(Triple.C);
+		}
+	}
+
+	// Add the vertices to the set (eliminates duplicates.)  Todo: don't use a set.
+	ResetModifiedVertices();
+	for (int32 VertexID : SubsetIDBuffer)
+	{
+		RecordModifiedVertex(VertexID);
+	}
+
+}
+
+// This actually updates constraints that correspond to the handle vertices.
+void FGroupTopologyLaplacianDeformer::UpdateSolution(FDynamicMesh3 * TargetMesh, const TFunction<FVector3d(FDynamicMesh3* Mesh, int)>& HandleVertexDeformFunc)
+{
+	// copy the current positions.
+	FVertexPositionCache CurrentPositions;
+	for (int VertexID : InitialPositions.Vertices)
+	{
+		CurrentPositions.AddVertex(TargetMesh, VertexID);
+	}
+
+	// Set the target mesh to the initial positions.
+	// Note: this only updates the vertices in the selected region.
+	InitialPositions.SetPositions(TargetMesh);
+
+	//Reset the constraints
+	for (int32 VertexID : ModifiedVertices)
+	{
+		//Get the vertex's data from the constraint buffer
+		FConstraintData& CData = SrcMeshConstraintBuffer[VertexID];
+
+		CData.Position = TargetMesh->GetVertex(VertexID);
+		CData.Weight   = 0.0; //A weight of zero is used to allow this point to move freely when moving the handles
+		CData.bPostFix = false;
+	}
+
+	//Actually deform the handles and add a constraint.
+	for (int VertexID : HandleVertices)
+	{
+		const FVector3d DeformPos = HandleVertexDeformFunc(TargetMesh, VertexID);
+
+		//Get the vertex's data from the constraint buffer
+		FConstraintData& CData = SrcMeshConstraintBuffer[VertexID];
+
+		//Set the new vertex data
+		CData.Position = DeformPos;
+		CData.Weight   = HandleWeights;
+		CData.bPostFix = bPostfixHandles;
+	}
+
+	// Restore Current Positions.  This is done because the target mesh is being used to define the highlight region.
+	// if we don't reset the positions the highlight mesh will appear to reset momentarily until the first laplacian solver result is available 
+	CurrentPositions.SetPositions(TargetMesh);
+
+	bDeformerNeedsToRun = true;
+
+}
+
+void FGroupTopologyLaplacianDeformer::ExportDeformedPositions(FDynamicMesh3* TargetMesh)
+{
+	bool bIsWorking = IsTaskInFlight();
+	if (AsyncMeshDeformTask != nullptr && !bIsWorking)
+	{
+		const FConstrainedMeshDeformerTask& Task = AsyncMeshDeformTask->GetTask();
+		Task.ExportResults(*TargetMesh);
+	}
+}
+
+inline FGroupTopologyLaplacianDeformer::~FGroupTopologyLaplacianDeformer()
+{
+	Shutdown();
+}
+
+inline void FGroupTopologyLaplacianDeformer::Shutdown()
+{	
+	if (AsyncMeshDeformTask != nullptr)
+	{
+		if (AsyncMeshDeformTask->IsDone())
+		{
+			delete AsyncMeshDeformTask;
+		}
+		else
+		{
+			AsyncMeshDeformTask->CancelAndDelete();
+		}
+
+		AsyncMeshDeformTask = nullptr;
+	}
+}
+
+/*
+* Tool methods
+*/
+
+UEditMeshPolygonsTool::UEditMeshPolygonsTool()
+{
+}
+
+void UEditMeshPolygonsTool::Setup()
+{
+	UMeshSurfacePointTool::Setup();
+
+	// create dynamic mesh component to use for live preview
+	DynamicMeshComponent = NewObject<USimpleDynamicMeshComponent>(ComponentTarget->GetOwnerActor(), "DynamicMesh");
+	DynamicMeshComponent->SetupAttachment(ComponentTarget->GetOwnerActor()->GetRootComponent());
+	DynamicMeshComponent->RegisterComponent();
+	DynamicMeshComponent->SetWorldTransform(ComponentTarget->GetWorldTransform());
+
+	// copy material if there is one
+	auto Material = ComponentTarget->GetMaterial(0);
+	if (Material != nullptr)
+	{
+		DynamicMeshComponent->SetMaterial(0, Material);
+	}
+
+	// dynamic mesh configuration settings
+	DynamicMeshComponent->TangentsType = EDynamicMeshTangentCalcType::AutoCalculated;
+	DynamicMeshComponent->InitializeMesh(ComponentTarget->GetMesh());
+	OnDynamicMeshComponentChangedHandle = DynamicMeshComponent->OnMeshChanged.Add(
+		FSimpleMulticastDelegate::FDelegate::CreateUObject(this, &UEditMeshPolygonsTool::OnDynamicMeshComponentChanged));
+
+	// initialize AABBTree
+	MeshSpatial.SetMesh(DynamicMeshComponent->GetMesh());
+	PrecomputeTopology();
+
+	//initialize topology selector
+	TopoSelector.Initialize(DynamicMeshComponent->GetMesh(), &Topology);
+	TopoSelector.SetSpatialSource([this]() {return &GetSpatial(); });
+	TopoSelector.PointsWithinToleranceTest = [this](const FVector3d& Position1, const FVector3d& Position2) {
+		return ToolSceneQueriesUtil::PointSnapQuery(this->CameraState, Position1, Position2, VisualAngleSnapThreshold);
+	};
+
+	// hide input StaticMeshComponent
+	ComponentTarget->SetOwnerVisibility(false);
+
+	// init state flags flags
+	bInDrag = false;
+
+	// initialize snap solver
+	QuickAxisTranslater.Initialize();
+	QuickAxisRotator.Initialize();
+
+	// add properties
+	TransformProps = NewObject<UPolyEditTransformProperties>(this);
+	AddToolPropertySource(TransformProps);
+
+	// set up visualizers
+	PolyEdgesRenderer.LineColor = FLinearColor::Red;
+	PolyEdgesRenderer.LineThickness = 2.0;
+	HilightRenderer.LineColor = FLinearColor::Green;
+	HilightRenderer.LineThickness = 4.0f;
+
+	// Allocates buffers, sets up the asynchronous task
+	// Copies the source mesh positions.
+	const ELaplacianWeightScheme LaplacianWeightScheme = ConvertToLaplacianWeightScheme(TransformProps->SelectedWeightScheme);
+	LaplacianDeformer.InitBackgroundWorker(LaplacianWeightScheme);
+	
+
+	/**
+	// How to add a curve for the weights.
+	//Add a default curve for falloff
+	FKeyHandle Keys[5];
+	Keys[0] = TransformProps->DefaultFalloffCurve.UpdateOrAddKey(0.f, 1.f);
+	Keys[1] = TransformProps->DefaultFalloffCurve.UpdateOrAddKey(0.25f, 0.25f);
+	Keys[2] = TransformProps->DefaultFalloffCurve.UpdateOrAddKey(0.3333333f, 0.25f);
+	Keys[3] = TransformProps->DefaultFalloffCurve.UpdateOrAddKey(0.6666667f, 1.25f);
+	Keys[4] = TransformProps->DefaultFalloffCurve.UpdateOrAddKey(1.f, 1.4f);
+	for (uint8 i = 0; i < 5; ++i)
+	{
+		TransformProps->DefaultFalloffCurve.SetKeyInterpMode(Keys[i], ERichCurveInterpMode::RCIM_Cubic);
+	}
+	TransformProps->WeightAttenuationCurve.EditorCurveData = TransformProps->DefaultFalloffCurve;
+	*/
+}
+
+void UEditMeshPolygonsTool::Shutdown(EToolShutdownType ShutdownType)
+{
+	//Tell the background thread to cancel the rest of its jobs before we close;
+	LaplacianDeformer.Shutdown();
+
+	if (DynamicMeshComponent != nullptr)
+	{
+		DynamicMeshComponent->OnMeshChanged.Remove(OnDynamicMeshComponentChangedHandle);
+
+		ComponentTarget->SetOwnerVisibility(true);
+
+		if (ShutdownType == EToolShutdownType::Accept)
+		{
+			// this block bakes the modified DynamicMeshComponent back into the StaticMeshComponent inside an undo transaction
+			GetToolManager()->BeginUndoTransaction(LOCTEXT("EditMeshPolygonsToolTransactionName", "Deform Mesh"));
+			ComponentTarget->CommitMesh([=](FMeshDescription* MeshDescription)
+			{
+				DynamicMeshComponent->Bake(MeshDescription, false);
+			});
+			GetToolManager()->EndUndoTransaction();
+		}
+
+		DynamicMeshComponent->UnregisterComponent();
+		DynamicMeshComponent->DestroyComponent();
+		DynamicMeshComponent = nullptr;
+	}
+}
+
+
+
+void UEditMeshPolygonsTool::ToggleLocalTransformFrameAction()
+{
+	TransformProps->bUseLocalAxis = !TransformProps->bUseLocalAxis;
+	if (bInDrag)
+	{
+		UpdateQuickTransformer();
+	}
+}
+
+void UEditMeshPolygonsTool::NextTransformTypeAction()
+{
+	if (bInDrag == false)
+	{
+		if (TransformProps->TransformMode == EQuickTransformerMode::AxisRotation)
+		{
+			TransformProps->TransformMode = EQuickTransformerMode::AxisTranslation;
+		}
+		else
+		{
+			TransformProps->TransformMode = EQuickTransformerMode::AxisRotation;
+		}
+		UpdateQuickTransformer();
+	}
+}
+
+
+
+void UEditMeshPolygonsTool::RegisterActions(FInteractiveToolActionSet& ActionSet)
+{
+	ActionSet.RegisterAction(this, (int32)EStandardToolActions::BaseClientDefinedActionID + 1,
+		TEXT("ToggleFrame"),
+		LOCTEXT("ToggleFrame", "Toggle World/Local Frame"),
+		LOCTEXT("ToggleFrameTooltip", "Toggle between World and Local transformation frames"),
+		EModifierKey::None, EKeys::A,
+		[this]() { ToggleLocalTransformFrameAction(); });
+
+
+	ActionSet.RegisterAction(this, (int32)EStandardToolActions::BaseClientDefinedActionID + 2,
+		TEXT("NextTransformType"),
+		LOCTEXT("NextTransformType", "Next Transform Type"),
+		LOCTEXT("NextTransformTypeTooltip", "Cycle to next transform type"),
+		EModifierKey::None, EKeys::Q,
+		[this]() { NextTransformTypeAction(); });
+}
+
+
+
+
+void UEditMeshPolygonsTool::OnDynamicMeshComponentChanged()
+{
+	bSpatialDirty = true;
+	TopoSelector.Invalidate(true, false);
+
+
+	//Makes sure the constraint buffer and position buffers reflect Undo/Redo changes
+	FDynamicMesh3* Mesh = DynamicMeshComponent->GetMesh();
+	
+	
+	//Apply Undo/redo
+	for (int VertexID : Mesh->VertexIndicesItr())
+	{
+		const FVector3d Position = Mesh->GetVertex(VertexID);
+		LaplacianDeformer.SrcMeshConstraintBuffer[VertexID].Position = Position;
+	}
+
+	// a deform task could still be in flight.
+	if (LaplacianDeformer.AsyncMeshDeformTask != nullptr)
+	{
+		LaplacianDeformer.AsyncMeshDeformTask->CancelAndDelete();
+		LaplacianDeformer.AsyncMeshDeformTask = nullptr;
+		LaplacianDeformer.bTaskSubmeshIsDirty = true;
+	}
+}
+
+FDynamicMeshAABBTree3& UEditMeshPolygonsTool::GetSpatial()
+{
+	if (bSpatialDirty)
+	{
+		MeshSpatial.Build();
+		bSpatialDirty = false;
+	}
+	return MeshSpatial;
+}
+
+
+bool UEditMeshPolygonsTool::HitTest(const FRay& WorldRay, FHitResult& OutHit)
+{
+	FTransform Transform = ComponentTarget->GetWorldTransform();
+	FRay3d LocalRay(Transform.InverseTransformPosition(WorldRay.Origin),
+		Transform.InverseTransformVector(WorldRay.Direction));
+	LocalRay.Direction.Normalize();
+
+	TopoSelector.UpdateEnableFlags(TransformProps->bSelectFaces, TransformProps->bSelectEdges, TransformProps->bSelectVertices);
+	FGroupTopologySelection Selection;
+	FVector3d LocalPosition, LocalNormal;
+	if (TopoSelector.FindSelectedElement(LocalRay, Selection, LocalPosition, LocalNormal) == false)
+	{
+		return false;
+	}
+
+	if (Selection.SelectedCornerIDs.Num() > 0)
+	{
+		OutHit.FaceIndex = Selection.SelectedCornerIDs[0];
+		OutHit.Distance = LocalRay.Project(LocalPosition);
+		OutHit.ImpactPoint = Transform.TransformPosition(LocalRay.PointAt(OutHit.Distance));
+	}
+	else if (Selection.SelectedEdgeIDs.Num() > 0)
+	{
+		OutHit.FaceIndex = Selection.SelectedEdgeIDs[0];
+		OutHit.Distance = LocalRay.Project(LocalPosition);
+		OutHit.ImpactPoint = Transform.TransformPosition(LocalRay.PointAt(OutHit.Distance));
+	}
+	else
+	{
+		int HitTID = GetSpatial().FindNearestHitTriangle(LocalRay);
+		if (HitTID != IndexConstants::InvalidID)
+		{
+			FTriangle3d Triangle;
+			GetSpatial().GetMesh()->GetTriVertices(HitTID, Triangle.V[0], Triangle.V[1], Triangle.V[2]);
+			FIntrRay3Triangle3d Query(LocalRay, Triangle);
+			Query.Find();
+			OutHit.FaceIndex = HitTID;
+			OutHit.Distance = Query.RayParameter;
+			OutHit.Normal = Transform.TransformVectorNoScale(GetSpatial().GetMesh()->GetTriNormal(HitTID));
+			OutHit.ImpactPoint = Transform.TransformPosition(LocalRay.PointAt(Query.RayParameter));
+		}
+	}
+	return true;
+}
+
+
+void UEditMeshPolygonsTool::OnBeginDrag(const FRay& WorldRay)
+{
+	FTransform Transform = ComponentTarget->GetWorldTransform();
+	FRay3d LocalRay(Transform.InverseTransformPosition(WorldRay.Origin),
+		Transform.InverseTransformVector(WorldRay.Direction));
+	LocalRay.Direction.Normalize();
+
+	HilightSelection.Clear();
+
+	TopoSelector.UpdateEnableFlags(TransformProps->bSelectFaces, TransformProps->bSelectEdges, TransformProps->bSelectVertices);
+	FGroupTopologySelection Selection;
+	FVector3d LocalPosition, LocalNormal;
+	bool bHit = TopoSelector.FindSelectedElement(LocalRay, Selection, LocalPosition, LocalNormal);
+
+	if (bHit == false)
+	{
+		bInDrag = false;
+		return;
+	}
+
+	HilightSelection = Selection;
+
+	FVector3d WorldHitPos = Transform.TransformPosition(LocalPosition);
+	FVector3d WorldHitNormal = Transform.TransformVector(LocalNormal);
+
+	bInDrag = true;
+	StartHitPosWorld = WorldHitPos;
+	LastHitPosWorld = StartHitPosWorld;
+	StartHitNormalWorld = WorldHitNormal;
+
+	QuickAxisRotator.ClearAxisLock();
+	UpdateActiveSurfaceFrame(HilightSelection);
+	UpdateQuickTransformer();
+
+	LastBrushPosLocal = Transform.InverseTransformPosition(LastHitPosWorld);
+	StartBrushPosLocal = LastBrushPosLocal;
+
+	// Capture the part of the mesh that will deform
+
+	if (TransformProps->DeformationStrategy == EGroupTopologyDeformationStrategy::Linear)
+	{
+		//Determine which of the following (corners, edges or faces) has been selected by counting the associated feature's IDs
+		if (Selection.SelectedCornerIDs.Num() > 0)
+		{
+			//Add all the the Corner's adjacent poly-groups (NbrGroups) to the ongoing array of groups.
+			LinearDeformer.SetActiveHandleCorners(Selection.SelectedCornerIDs);
+		}
+		else if (Selection.SelectedEdgeIDs.Num() > 0)
+		{
+			//Add all the the edge's adjacent poly-groups (NbrGroups) to the ongoing array of groups.
+			LinearDeformer.SetActiveHandleEdges(Selection.SelectedEdgeIDs);
+		}
+		else if (Selection.SelectedGroupIDs.Num() > 0)
+		{
+			LinearDeformer.SetActiveHandleFaces(Selection.SelectedGroupIDs);
+		}
+	}
+	else //(TransformProps->DeformationStrategy == EGroupTopologyDeformationStrategy::Laplacian)
+	{
+		LaplacianDeformer.bLocalize = true; // TransformProps->bLocalizeDeformation;
+
+		//Determine which of the following (corners, edges or faces) has been selected by counting the associated feature's IDs
+		if (Selection.SelectedCornerIDs.Num() > 0)
+		{
+			//Add all the the Corner's adjacent poly-groups (NbrGroups) to the ongoing array of groups.
+			LaplacianDeformer.SetActiveHandleCorners(Selection.SelectedCornerIDs);
+		}
+		else if (Selection.SelectedEdgeIDs.Num() > 0)
+		{
+			//Add all the the edge's adjacent poly-groups (NbrGroups) to the ongoing array of groups.
+			LaplacianDeformer.SetActiveHandleEdges(Selection.SelectedEdgeIDs);
+		}
+		else if (Selection.SelectedGroupIDs.Num() > 0)
+		{
+			LaplacianDeformer.SetActiveHandleFaces(Selection.SelectedGroupIDs);
+		}
+
+		// the task will need a new mesh that corresponds to the selected region.
+		LaplacianDeformer.bTaskSubmeshIsDirty = true;
+	}
+
+	BeginChange();
+}
+
+
+void UEditMeshPolygonsTool::UpdateActiveSurfaceFrame(FGroupTopologySelection& Selection)
+{
+	FTransform Transform = ComponentTarget->GetWorldTransform();
+
+	// update surface frame
+	ActiveSurfaceFrame.Origin = StartHitPosWorld;
+	if (HilightSelection.SelectedCornerIDs.Num() == 1)
+	{
+		// just keeping existing axes...we don't have enough info to do something smarter
+	}
+	else
+	{
+		ActiveSurfaceFrame.AlignAxis(2, StartHitNormalWorld);
+		if (HilightSelection.SelectedEdgeIDs.Num() == 1)
+		{
+			FVector3d Tangent;
+			if (Topology.GetGroupEdgeTangent(HilightSelection.SelectedEdgeIDs[0], Tangent))
+			{
+				Tangent = Transform.TransformVector(Tangent);
+				ActiveSurfaceFrame.ConstrainedAlignAxis(0, Tangent, ActiveSurfaceFrame.Z());
+			}
+		}
+	}
+}
+
+
+FQuickTransformer* UEditMeshPolygonsTool::GetActiveQuickTransformer()
+{
+	if (TransformProps->TransformMode == EQuickTransformerMode::AxisRotation)
+	{
+		return &QuickAxisRotator;
+	}
+	else
+	{
+		return &QuickAxisTranslater;
+	}
+}
+
+
+void UEditMeshPolygonsTool::UpdateQuickTransformer()
+{
+	if (TransformProps->bUseLocalAxis)
+	{
+		GetActiveQuickTransformer()->SetActiveWorldFrame(ActiveSurfaceFrame);
+	}
+	else
+	{
+		GetActiveQuickTransformer()->SetActiveFrameFromWorldAxes(StartHitPosWorld);
+	}
+}
+
+
+
+
+
+void UEditMeshPolygonsTool::UpdateChangeFromROI(bool bFinal)
+{
+	if (ActiveVertexChange == nullptr)
+	{
+		return;
+	}
+	const bool bIsLaplacian = (TransformProps->DeformationStrategy == EGroupTopologyDeformationStrategy::Laplacian);
+
+	FDynamicMesh3* Mesh = DynamicMeshComponent->GetMesh();
+	const TSet<int>& ModifiedVertices = (bIsLaplacian) ? LaplacianDeformer.GetModifiedVertices() : LinearDeformer.GetModifiedVertices();
+	
+	ActiveVertexChange->SavePositions(Mesh, ModifiedVertices, !bFinal);
+
+}
+
+
+void UEditMeshPolygonsTool::OnUpdateDrag(const FRay& Ray)
+{
+	if (bInDrag)
+	{
+		bUpdatePending = true;
+		UpdateRay = Ray;
+	}
+}
+
+void UEditMeshPolygonsTool::OnEndDrag(const FRay& Ray)
+{
+	bInDrag = false;
+	bUpdatePending = false;
+
+	// update spatial
+	bSpatialDirty = true;
+
+	HilightSelection.Clear(); 
+	TopoSelector.Invalidate(true, false);
+	QuickAxisRotator.ClearAxisLock();
+
+	//If it's linear, it's computed real time with no delay. This may need to be restructured for clarity by using the background task for this as well.
+	if (TransformProps->DeformationStrategy == EGroupTopologyDeformationStrategy::Linear)
+	{
+		// close change record
+		EndChange();
+	}
+}
+
+
+
+void UEditMeshPolygonsTool::OnUpdateHover(const FInputDeviceRay& DevicePos)
+{
+	//if (!bNeedEmitEndChange)
+	if (ActiveVertexChange == nullptr)
+	{
+		FTransform Transform = ComponentTarget->GetWorldTransform();
+		FRay3d LocalRay(Transform.InverseTransformPosition(DevicePos.WorldRay.Origin),
+		Transform.InverseTransformVector(DevicePos.WorldRay.Direction));
+		LocalRay.Direction.Normalize();
+
+		HilightSelection.Clear();
+		TopoSelector.UpdateEnableFlags(TransformProps->bSelectFaces, TransformProps->bSelectEdges, TransformProps->bSelectVertices);
+		FVector3d LocalPosition, LocalNormal;
+		bool bHit = TopoSelector.FindSelectedElement(LocalRay, HilightSelection, LocalPosition, LocalNormal);
+
+		if (bHit)
+		{
+			StartHitPosWorld = Transform.TransformPosition(LocalPosition);
+			StartHitNormalWorld = Transform.TransformVector(LocalNormal);
+
+			UpdateActiveSurfaceFrame(HilightSelection);
+			UpdateQuickTransformer();
+		}
+	}
+}
+
+
+
+
+void UEditMeshPolygonsTool::ComputeUpdate()
+{
+
+	if (bUpdatePending == true)
+	{
+		// Linear Deformer : Update the solution
+		// Laplacain Deformer : Update the constraints (positions and weights) - the region was identified in onBeginDrag
+
+
+		if (TransformProps->TransformMode == EQuickTransformerMode::AxisRotation)
+		{
+			ComputeUpdate_Rotate();
+		}
+		else
+		{
+			ComputeUpdate_Translate();
+		}
+	}
+
+	if (TransformProps->DeformationStrategy == EGroupTopologyDeformationStrategy::Laplacian)
+	{
+		bool bIsWorking = LaplacianDeformer.IsTaskInFlight();
+
+		if (!bIsWorking)
+		{
+			// Sync update if we have new results.
+			if (LaplacianDeformer.bVertexPositionsNeedSync)
+			{
+				
+				//Update the mesh with the provided solutions.
+				LaplacianDeformer.ExportDeformedPositions(DynamicMeshComponent->GetMesh());
+
+				LaplacianDeformer.bVertexPositionsNeedSync = false;
+
+				//Re-sync mesh, and flag the spatial data struct & topology for re-evaluation
+				DynamicMeshComponent->FastNotifyPositionsUpdated();
+
+				GetToolManager()->PostInvalidation();
+				bSpatialDirty = true;
+				TopoSelector.Invalidate(true, false);
+			}
+
+			// emit end change if we are done with the drag
+			if (!LaplacianDeformer.bDeformerNeedsToRun && !bInDrag)
+			{
+				EndChange();
+			}
+
+			// Not working but we have more work for it to do..
+
+			if (LaplacianDeformer.bDeformerNeedsToRun)
+			{
+			
+				FRichCurve* Curve = NULL;
+
+				/**
+				// How to add a deformation curve
+				const bool bApplyAttenuationCurve = TransformProps->bApplyAttenuationCurve
+				if (bApplyAttenuationCurve)
+				{
+					Curve = TransformProps->WeightAttenuationCurve.GetRichCurve();
+				}
+				*/
+				const ELaplacianWeightScheme LaplacianWeightScheme = ConvertToLaplacianWeightScheme(TransformProps->SelectedWeightScheme);
+				LaplacianDeformer.UpdateAndLaunchdWorker(LaplacianWeightScheme, Curve);
+
+			}
+		}
+	}
+}
+
+
+
+
+
+
+void UEditMeshPolygonsTool::ComputeUpdate_Rotate()
+{
+	const bool bIsLaplacian = (TransformProps->DeformationStrategy == EGroupTopologyDeformationStrategy::Laplacian); 
+	FGroupTopologyDeformer& SelectedDeformer = (bIsLaplacian) ? LaplacianDeformer : LinearDeformer;
+
+	FDynamicMesh3* Mesh = DynamicMeshComponent->GetMesh();
+	FTransform Transform = ComponentTarget->GetWorldTransform();
+	FVector NewHitPosWorld = LastHitPosWorld;
+
+	FVector3d SnappedPoint;
+	if (QuickAxisRotator.UpdateSnap(UpdateRay, SnappedPoint))
+	{
+		NewHitPosWorld = SnappedPoint;
+	}
+	else
+	{
+		return;
+	}
+
+	// check if we are on back-facing part of rotation in which case we ignore...
+	FVector3d SphereCenter = QuickAxisRotator.GetActiveWorldFrame().Origin;
+	if (QuickAxisRotator.HaveActiveSnapRotation() && QuickAxisRotator.GetHaveLockedToAxis() == false)
+	{
+		FVector3d ToSnapPointVec = (SnappedPoint - SphereCenter);
+		FVector3d ToEyeVec = (SnappedPoint - (FVector3d)CameraState.Position);
+		if (ToSnapPointVec.Dot(ToEyeVec) > 0)
+		{
+			return;
+		}
+	}
+
+
+	// if we haven't snapped to a rotation we can exit
+	if (QuickAxisRotator.HaveActiveSnapRotation() == false)
+	{
+		QuickAxisRotator.ClearAxisLock();
+
+		SelectedDeformer.ClearSolution(Mesh);
+
+		//TODO: This is unseemly here, need to potentially defer this so that it's handled the same way as laplacian. Placeholder for now.
+		if (TransformProps->DeformationStrategy == EGroupTopologyDeformationStrategy::Linear)
+		{
+
+			DynamicMeshComponent->FastNotifyPositionsUpdated();
+			GetToolManager()->PostInvalidation();
+		}
+		bUpdatePending = false;
+		return;
+	}
+
+	// ok we have an axis...
+	if (QuickAxisRotator.GetHaveLockedToAxis() == false)
+	{
+		QuickAxisRotator.SetAxisLock();
+		RotationStartPointWorld = SnappedPoint;
+		RotationStartFrame = QuickAxisRotator.GetActiveRotationFrame();
+	}
+
+	FVector2d RotateStartVec = RotationStartFrame.ToPlaneUV(RotationStartPointWorld, 2);
+	RotateStartVec.Normalize();
+	FVector2d RotateToVec = RotationStartFrame.ToPlaneUV(NewHitPosWorld, 2);
+	RotateToVec.Normalize();
+	double AngleRad = RotateStartVec.SignedAngleR(RotateToVec);
+	FQuaterniond Rotation(
+		Transform.InverseTransformVectorNoScale(RotationStartFrame.Z()), AngleRad, false);
+	FVector3d LocalOrigin = Transform.InverseTransformPosition(RotationStartFrame.Origin);
+
+	// Linear Deformer: Update Mesh the rotation,
+	// Laplacian Deformer:  Update handles constraints with the rotation and set bDeformerNeedsToRun = true;.
+	SelectedDeformer.UpdateSolution(Mesh, [this, LocalOrigin, Rotation](FDynamicMesh3* TargetMesh, int VertIdx)
+	{
+		FVector3d V = TargetMesh->GetVertex(VertIdx);
+		V -= LocalOrigin;
+		V = Rotation * V;
+		V += LocalOrigin;
+		return V;
+	});
+
+	//TODO: This is unseemly here, need to potentially defer this so that it's handled the same way as laplacian. Placeholder for now.
+	if (!bIsLaplacian)
+	{
+		DynamicMeshComponent->FastNotifyPositionsUpdated();
+		GetToolManager()->PostInvalidation();
+	}
+	bUpdatePending = false;
+}
+
+
+
+
+void UEditMeshPolygonsTool::ComputeUpdate_Translate()
+{
+	const bool bIsLaplacian = (TransformProps->DeformationStrategy == EGroupTopologyDeformationStrategy::Laplacian);
+	FGroupTopologyDeformer& SelectedDeformer = (bIsLaplacian) ? LaplacianDeformer : LinearDeformer;
+
+	FTransform Transform = ComponentTarget->GetWorldTransform();
+	FVector NewHitPosWorld = LastHitPosWorld;
+	FVector3d SnappedPoint;
+	if (QuickAxisTranslater.UpdateSnap(UpdateRay, SnappedPoint))
+	{
+		NewHitPosWorld = SnappedPoint;
+	}
+	else
+	{
+		return;
+	}
+
+	FVector NewBrushPosLocal = Transform.InverseTransformPosition(NewHitPosWorld);
+	FVector3d NewMoveDelta = NewBrushPosLocal - StartBrushPosLocal;
+
+	FDynamicMesh3* Mesh = DynamicMeshComponent->GetMesh();
+	if (LastMoveDelta.SquaredLength() > 0.)
+	{
+		if (NewMoveDelta.SquaredLength() > 0.)
+		{
+			// Linear Deformer: Update Mesh with the translation,
+			// Laplacian Deformer:  Update handles constraints and set bDeformerNeedsToRun = true;.
+
+			SelectedDeformer.UpdateSolution(Mesh, [this, NewMoveDelta](FDynamicMesh3* TargetMesh, int VertIdx)
+			{
+				return TargetMesh->GetVertex(VertIdx) + NewMoveDelta;
+			});
+		}
+		else
+		{
+			// Reset mesh to initial positions.
+
+			SelectedDeformer.ClearSolution(Mesh);
+
+		}
+		//TODO: This is unseemly here, need to potentially defer this so that it's handled the same way as laplacian. Placeholder for now.
+		if (!bIsLaplacian)
+		{
+			DynamicMeshComponent->FastNotifyPositionsUpdated();
+			GetToolManager()->PostInvalidation();
+		}
+	}
+
+	LastMoveDelta = NewMoveDelta;
+	LastBrushPosLocal = NewBrushPosLocal;
+
+	bUpdatePending = false;
+}
+
+
+
+void UEditMeshPolygonsTool::Tick(float DeltaTime)
+{
+	UMeshSurfacePointTool::Tick(DeltaTime);
+	
+	LaplacianDeformer.HandleWeights   = TransformProps->HandleWeight;
+	LaplacianDeformer.bPostfixHandles = TransformProps->bPostFixHandles;
+}
+
+
+
+void UEditMeshPolygonsTool::PrecomputeTopology()
+{
+	FDynamicMesh3* Mesh = DynamicMeshComponent->GetMesh();
+	Topology = FGroupTopology(Mesh, true);
+
+	LinearDeformer.Initialize(Mesh, &Topology);
+	LaplacianDeformer.Initialize(Mesh, &Topology);
+
+	// Make the Constraint Buffer, zero weights, but current pos
+	LaplacianDeformer.InitializeConstraintBuffer();
+}
+
+
+
+
+void UEditMeshPolygonsTool::Render(IToolsContextRenderAPI* RenderAPI)
+{
+
+	ComputeUpdate();
+		
+	GetToolManager()->GetContextQueriesAPI()->GetCurrentViewState(CameraState);
+	GetActiveQuickTransformer()->UpdateCameraState(CameraState);
+
+	FDynamicMesh3* TargetMesh = DynamicMeshComponent->GetMesh();
+
+	PolyEdgesRenderer.BeginFrame(RenderAPI, CameraState);
+	PolyEdgesRenderer.SetTransform(ComponentTarget->GetWorldTransform());
+
+
+	for (FGroupTopology::FGroupEdge& Edge : Topology.Edges)
+	{
+		FVector3d A, B;
+		for (int eid : Edge.Span.Edges)
+		{
+			TargetMesh->GetEdgeV(eid, A, B);
+			PolyEdgesRenderer.DrawLine(A, B);
+		}
+	}
+
+	PolyEdgesRenderer.EndFrame();
+
+
+	HilightRenderer.BeginFrame(RenderAPI, CameraState);
+	HilightRenderer.SetTransform(ComponentTarget->GetWorldTransform());
+
+#ifdef DEBUG_ROI_WEIGHTS
+	FDynamicMesh3* MeshPtr = DynamicMeshComponent->GetMesh();
+	for (int32 VertexID : DynamicMeshComponent->GetMesh()->VertexIndicesItr())
+	{
+		float Color = 1.f - SrcMeshConstraintBuffer[VertexID].Weight;
+		HilightRenderer.DrawPoint(MeshPtr->GetVertex(VertexID), FLinearColor(Color, Color, Color, 1.f), 8, true);
+	}
+
+#endif
+
+#ifdef DEBUG_ROI_HANDLES
+	const FLinearColor FOOF{ 1.f,0.f,1.f,1.f };
+	for (int VertIdx : HandleVertices)
+	{
+		HilightRenderer.DrawViewFacingCircle(TargetMesh->GetVertex(VertIdx), 0.8f, 8, FOOF, 3, false);
+	}
+#endif 
+
+#ifdef DEBUG_ROI_TRIANGLES
+	const FLinearColor Whiteish{ 0.67f,0.67f,0.67f,1.f };
+	for (int32 i = 0; i < SubsetIDBuffer.Num(); i += 3)
+	{
+		FVector3d A = TargetMesh->GetVertex(SubsetIDBuffer[i]);
+		FVector3d B = TargetMesh->GetVertex(SubsetIDBuffer[i + 1]);
+		FVector3d C = TargetMesh->GetVertex(SubsetIDBuffer[i + 2]);
+		HilightRenderer.DrawLine(A, B, Whiteish, 2.7f, true);
+		HilightRenderer.DrawLine(B, C, Whiteish, 2.7f, true);
+		HilightRenderer.DrawLine(C, A, Whiteish, 2.7f, true);
+	}
+#endif
+
+	TopoSelector.VisualAngleSnapThreshold = this->VisualAngleSnapThreshold;
+	TopoSelector.DrawSelection(HilightSelection, &HilightRenderer, &CameraState);
+	HilightRenderer.EndFrame();
+
+
+	if (bInDrag)
+	{
+
+		GetActiveQuickTransformer()->Render(RenderAPI);
+	}
+	else
+	{
+		GetActiveQuickTransformer()->PreviewRender(RenderAPI);
+	}
+}
+
+
+
+
+
+//
+// Change Tracking
+//
+
+
+void UEditMeshPolygonsTool::BeginChange()
+{
+	const bool bIsLaplacian = (TransformProps->DeformationStrategy == EGroupTopologyDeformationStrategy::Laplacian);
+	if (!bIsLaplacian || LaplacianDeformer.IsDone())
+	{
+		if (ActiveVertexChange == nullptr)
+		{
+			ActiveVertexChange = new FMeshVertexChangeBuilder();
+			UpdateChangeFromROI(false);
+		}
+	}
+}
+
+
+void UEditMeshPolygonsTool::EndChange()
+{
+
+	if (ActiveVertexChange != nullptr)
+	{
+		UpdateChangeFromROI(true);
+		GetToolManager()->EmitObjectChange(DynamicMeshComponent, MoveTemp(ActiveVertexChange->Change), LOCTEXT("PolyMeshDeformationChange", "PolyMesh Edit"));
+	}
+
+	delete ActiveVertexChange;
+	ActiveVertexChange = nullptr;
+}
+
+
+
+#undef LOCTEXT_NAMESPACE
diff --git a/Engine/Plugins/Experimental/MeshModelingToolset/Source/MeshModelingTools/Private/PolygonMeshDeformTool.cpp b/Engine/Plugins/Experimental/MeshModelingToolset/Source/MeshModelingTools/Private/PolygonMeshDeformTool.cpp
deleted file mode 100644
diff --git a/Engine/Plugins/Experimental/MeshModelingToolset/Source/MeshModelingTools/Public/EditMeshPolygonsTool.h b/Engine/Plugins/Experimental/MeshModelingToolset/Source/MeshModelingTools/Public/EditMeshPolygonsTool.h
new file mode 100644
--- /dev/null
+++ b/Engine/Plugins/Experimental/MeshModelingToolset/Source/MeshModelingTools/Public/EditMeshPolygonsTool.h
@@ -0,0 +1,485 @@
+// Copyright 1998-2019 Epic Games, Inc. All Rights Reserved.
+
+#pragma once
+
+#include "CoreMinimal.h"
+#include "UObject/NoExportTypes.h"
+#include "BaseTools/MeshSurfacePointTool.h"
+#include "SimpleDynamicMeshComponent.h"
+#include "DynamicMeshAABBTree3.h"
+#include "Drawing/ToolDataVisualizer.h"
+#include "Transforms/QuickAxisTranslater.h"
+#include "Transforms/QuickAxisRotator.h"
+#include "MeshVertexChange.h"
+#include "GroupTopology.h"
+#include "Spatial/GeometrySet3.h"
+#include "Selection/GroupTopologySelector.h"
+#include "Operations/GroupTopologyDeformer.h"
+#include "MeshSolverUtilities/Private/LaplacianMeshSmoother.h"
+#include "Curves/CurveFloat.h"
+#include "ModelingOperators/Public/ModelingTaskTypes.h"
+#include "EditMeshPolygonsTool.generated.h"
+
+class FMeshVertexChangeBuilder;
+class FGroupTopologyLaplacianDeformer;
+class FDeformTask;
+struct FConstraintData;
+
+/**
+ * ToolBuilder
+ */
+UCLASS()
+class MESHMODELINGTOOLS_API UEditMeshPolygonsToolBuilder : public UMeshSurfacePointToolBuilder
+{
+	GENERATED_BODY()
+
+public:
+	UEditMeshPolygonsToolBuilder()
+	{
+	}
+
+	virtual UMeshSurfacePointTool* CreateNewTool(const FToolBuilderState& SceneState) const override;
+};
+
+/** Deformation strategies */
+UENUM()
+enum class EGroupTopologyDeformationStrategy : uint8
+{
+	/** Deforms the mesh using linear translations*/
+	Linear UMETA(DisplayName = "Linear"),
+
+	/** Deforms the mesh using laplacian deformation*/
+	Laplacian UMETA(DisplayName = "Laplacian")
+};
+
+/** Laplacian weight schemes determine how we will look at the curvature at a given vertex in relation to its neighborhood*/
+UENUM()
+enum class EWeightScheme 
+{
+	Uniform				UMETA(DisplayName = "Uniform"),
+	Umbrella			UMETA(DisplayName = "Umbrella"),
+	Valence				UMETA(DisplayName = "Valence"),
+	MeanValue			UMETA(DisplayName = "MeanValue"),
+	Cotangent			UMETA(DisplayName = "Cotangent"),
+	ClampedCotangent	UMETA(DisplayName = "ClampedCotangent")
+};
+
+/** The ELaplacianWeightScheme enum is the same..*/
+static ELaplacianWeightScheme ConvertToLaplacianWeightScheme(const EWeightScheme WeightScheme)
+{
+	return static_cast<ELaplacianWeightScheme>(WeightScheme);
+}
+
+/** Modes for quick transformer */
+UENUM()
+enum class EQuickTransformerMode : uint8
+{
+	/** Translation along frame axes */
+	AxisTranslation = 0 UMETA(DisplayName = "Translate"),
+
+	/** Rotation around frame axes*/
+	AxisRotation = 1 UMETA(DisplayName = "Rotate"),
+};
+
+
+
+UCLASS()
+class MESHMODELINGTOOLS_API UPolyEditTransformProperties : public UInteractiveToolPropertySet
+{
+	GENERATED_BODY()
+
+public:
+	UPolyEditTransformProperties();
+
+	//Options
+
+	UPROPERTY(EditAnywhere, Category = Options, meta = (DisplayName = "Deformation Type", ToolTip = "Select the type of deformation you wish to employ on a polygroup."))
+	EGroupTopologyDeformationStrategy DeformationStrategy;
+
+	UPROPERTY(EditAnywhere, Category = Options)
+	bool bUseLocalAxis;
+
+	UPROPERTY(EditAnywhere, Category = Options)
+	EQuickTransformerMode TransformMode;
+
+	UPROPERTY(EditAnywhere, Category = Options)
+	bool bSelectFaces;
+
+	UPROPERTY(EditAnywhere, Category = Options)
+	bool bSelectEdges;
+
+	UPROPERTY(EditAnywhere, Category = Options)
+	bool bSelectVertices;
+
+	//Laplacian Deformation Options
+
+	UPROPERTY(EditAnywhere, Category = LaplacianOptions, meta = (EditCondition="DeformationStrategy == EGroupTopologyDeformationStrategy::Laplacian", DisplayName = "Laplacian Weight Scheme", ToolTip = "Laplacian weight schemes determine how the deformer will interpret the curvature at a given vertex in relation to nearby vertices in order to generate a solution."))
+	EWeightScheme SelectedWeightScheme = EWeightScheme::ClampedCotangent;
+
+	UPROPERTY(EditAnywhere, Category = LaplacianOptions, meta = (EditCondition="DeformationStrategy == EGroupTopologyDeformationStrategy::Laplacian", DisplayName = "Handle Weight", UIMin = "0.0001", UIMax = "1.0", ClampMin = "0.0001", ClampMax = "100.0", ToolTip = "The handle weight describes how effective dragging the selected PolyGroup will be."))
+	double HandleWeight = 1.0;
+
+	UPROPERTY(EditAnywhere, Category = LaplacianOptions, meta = (EditCondition="DeformationStrategy == EGroupTopologyDeformationStrategy::Laplacian", DisplayName = "Post-fix Handles", ToolTip = "Post-fixing the handles will attempt to force their vertices to arrive exactly where you wanted them. This may result\nin undesirable or choppy curvature in some scenarios."))
+	bool bPostFixHandles = false;
+/**
+	// How to add a weight curve
+	UPROPERTY(EditAnywhere, Category = LaplacianOptions, meta = (EditCondition="DeformationStrategy == EGroupTopologyDeformationStrategy::Laplacian", DisplayName = "Localize Deformation", ToolTip = "When enabled, only the vertices in the polygroups immediately adjacent to the selected group will be affected by the deformation.\nWhen disabled, the deformer will solve for the curvature of the entire mesh (slower)"))
+	bool bLocalizeDeformation{ true};
+
+	UPROPERTY(EditAnywhere, Category = LaplacianOptions, meta = (EditCondition="DeformationStrategy == EGroupTopologyDeformationStrategy::Laplacian", DisplayName = "Apply Weight Attenuation Curve", ToolTip = "When enabled, the curve below will be used to calculate the weight at a given vertex based on distance from the handles"))
+	bool bApplyAttenuationCurve{ false };
+	
+	FRichCurve DefaultFalloffCurve;
+
+	UPROPERTY(EditAnywhere,  Category = LaplacianOptions, meta = ( EditCondition="DeformationStrategy == EGroupTopologyDeformationStrategy::Laplacian && bApplyAttenuationCurve", DisplayName = "Distance-Weight Attenuation Curve",UIMin = "1.0", UIMax = "1.0", ClampMin = "1.0", ClampMax = "1.0", ToolTip = "This curve determines the weight attenuation over the distance of the mesh.\nThe selected polygroup handle is t=0.0, and t=1.0 is roughly the farthest vertices from the handles.\nThe value of the curve at each time interval represents the weight of the vertices at that distance from the selection."))
+	FRuntimeFloatCurve WeightAttenuationCurve;
+*/ 
+};
+
+//Stores per-vertex data needed by the laplacian deformer object
+//TODO: May be a candidate for a subclass of the FGroupTopologyLaplacianDeformer
+struct FConstraintData
+{
+	FConstraintData& operator=(const FConstraintData& other)
+	{
+		Position = other.Position;
+		Weight   = other.Weight;
+		bPostFix = other.bPostFix;
+		return *this;
+	}
+
+	FVector3d Position;
+	double Weight{ 0.0 };
+	bool bPostFix{ false };
+};
+
+/**
+*	FDeformTask is an object which wraps an asynchronous task to be run multiple times on a separate thread.
+*	The Laplacian deformation process requires the use of potentially large sparse matrices and sparse multiplication.
+*	
+*   Expected usage:
+*
+*   
+*   // define constraints.  Need Constraints[VertID] to hold the constraints for the corresponding vertex.
+*   TArray<FConstraintData> Constraints;
+*   ....
+*
+*   // populate with the VertexIDs of the vertices that are in the region you wish to deform.  
+*   TArray<int32> SrcVertIDs;  //Basically a mini-index buffer.
+*   ...
+* 
+*   // Create or reuse a laplacian deformation task.
+*   FDeformTask*   DeformTask = New FDeformTask(WeightScheme);
+*
+*   // the deformer will have to build a new mesh that represents the regions in SrcVertIDs;
+*   // but set this to false on subsequent calls to UpdateDeformer if the SrcVertIDs array hasn't changed.
+*   bool bRequiresRegion = true;
+
+*   DefTask->UpdateDeformer(WeightScheme, Mesh, Constraints, SrcVertIDs, bRequiresRegion);
+*
+*   DeformTask->DoWork();  or DeformTask->StartBackgroundTask(); //which calls DoWork on background thread.
+*
+*  // wheh DeformTask->IsDone == true; you can copy the results back to the mesh
+*  DeformTask->ExportResults(Mesh);
+*
+* Note: if only the positions in the Constraints change (e.g. handle positions) then subsequent calls 
+* to UpdateDeformer() and DoWork() will be much faster as the matrix system will not be rebuilt or re-factored
+*/
+class FConstrainedMeshDeformerTask : public FNonAbandonableTask
+{
+	friend class FAsyncTask<FDeformTask>;
+public:
+
+	enum
+	{
+		INACTIVE_SUBSET_ID = -1
+	};
+
+	FConstrainedMeshDeformerTask(const ELaplacianWeightScheme SelectedWeightScheme)
+	{
+		LaplacianWeightScheme = SelectedWeightScheme;
+	}
+
+	virtual ~FConstrainedMeshDeformerTask() {};
+
+	//NO idea what this is meant to do. Performance analysis maybe? Scheduling data?
+	FORCEINLINE TStatId GetStatId() const
+	{
+		RETURN_QUICK_DECLARE_CYCLE_STAT(FConstrainedMeshDeformerTask, STATGROUP_ThreadPoolAsyncTasks);
+	}
+
+	
+	/** Called by the main thread in the tool, this copies the Constraint buffer right before the task begins on another thread.
+	  * Ensures the FConstrainedMeshDeformer is using correct mesh subset and the selected settings, then updates on change in properties, i.e. weight scheme */
+	void UpdateDeformer(const ELaplacianWeightScheme SelectedWeightScheme, const FDynamicMesh3& Mesh, const TArray<FConstraintData>& ConstraintArray, const TArray<int32>& SrcIDBufferSubset, bool bNewTransaction, const FRichCurve* Curve);
+
+	/** Required by the FAsyncTaskExecutor */
+	void SetAbortSource(bool* bAbort) { bAbortSource = bAbort; };
+
+	/** Called by the FAsyncTask<FDeformTask> object for background computation. */
+	void DoWork();
+
+
+	/** Updates the positions in the target mesh for regions that correspond to the subset mesh */
+	void ExportResults(FDynamicMesh3& TargetMesh) const;
+
+private:
+
+	/** Creates the mesh (i.e. SubsetMesh) that corresponds to the region of the SrcMesh defined by the partial index buffer SrcIDBufferSubset */
+	void InitializeSubsetMesh(const FDynamicMesh3& SrcMesh, const TArray<int32>& SrcIDBufferSubset);
+
+	/** Attenuates the weights of the constraints using the selected curve */
+	void ApplyAttenuation();
+
+	/** Denotes the weight scheme being used by the running background task. Changes when selected property changes in editor. */
+	ELaplacianWeightScheme LaplacianWeightScheme;
+
+	/** positions for each vertex in the subset mesh - for use in the deformer */
+	TArray<FVector3d> SubsetPositionBuffer;
+
+	/** constraint data for each vertex in subset mesh - for use by the deformer*/
+	TArray<FConstraintData> SubsetConstraintBuffer;
+
+	FRichCurve WeightAttenuationCurve;
+
+	/** True only for the first update, and then false for the duration of the Input transaction
+	  * It's passed in and copied in UpdateDeformer() */
+	bool bIsNewTransaction = true;
+
+	/** When true, the constraint weights will be attenuated based on distance using the provided curve object*/
+	bool bAttenuateWeights = false;
+
+	/** The abort bool used by the Task Deleter */
+	bool* bAbortSource = nullptr;
+
+	/** Used to initialize the array mapping, updated during the UpdateDeformer() function */
+	int SrcMeshMaxVertexID;
+
+	/** A subset of the original mesh */
+	FDynamicMesh3 SubsetMesh;
+
+	/** Maps Subset Mesh VertexID to Src Mesh VertexID */
+	TArray<int32> SubsetVertexIDToSrcVertexIDMap;
+
+	/** Laplacian deformer object gets rebuilt each new transaction */
+	TUniquePtr<FConstrainedMeshDeformer> ConstrainedDeformer;
+
+private :
+	FConstrainedMeshDeformerTask();
+
+};
+
+class FGroupTopologyLaplacianDeformer : public FGroupTopologyDeformer
+{
+
+public:
+
+	FGroupTopologyLaplacianDeformer() = default;
+
+	virtual ~FGroupTopologyLaplacianDeformer();
+
+	/** Used to begin a procedural addition of modified vertices */
+	inline void ResetModifiedVertices()
+	{
+		ModifiedVertices.Empty();
+	};
+
+	/** Change tracking */
+	template <typename ValidSetAppendContainerType>
+	void RecordModifiedVertices(const ValidSetAppendContainerType& Container)
+	{
+		ModifiedVertices.Empty();
+		ModifiedVertices.Append(Container);
+	}
+
+	/** Used to iteratively add to the active change set (TSet<>)*/
+	inline void RecordModifiedVertex(int32 VertexID)
+	{
+		ModifiedVertices.Add(VertexID);
+	};
+
+
+	void SetActiveHandleFaces(const TArray<int>& FaceGroupIDs) override;
+	void SetActiveHandleEdges(const TArray<int>& TopologyEdgeIDs) override;
+	void SetActiveHandleCorners(const TArray<int>& TopologyCornerIDs) override;
+
+
+
+	/** Allocates shared storage for use in task synchronization */
+	void InitBackgroundWorker(const ELaplacianWeightScheme WeightScheme);
+
+
+	/** Coordinates the background tasks. Returns false if the worker was already running */
+	bool UpdateAndLaunchdWorker(const ELaplacianWeightScheme WeightScheme, const FRichCurve* Curve = nullptr);
+
+	/** Capture data about background task state.*/
+	bool IsTaskInFlight() const;
+
+
+
+	/** Sets the SrcMeshConstraintBuffer to have a size of MaxVertexID, and initializes with the current mesh positions, but weight zero*/
+	void InitializeConstraintBuffer();
+
+	/** Given an array of Group IDs, update the selection and record vertices */
+	void UpdateSelection(const FDynamicMesh3* TargetMesh, const TArray<int>& Groups, bool bLocalizeDeformation);
+
+	/** Updates the mesh preview and/or solvers upon user input, provided a deformation strategy */
+	void UpdateSolution(FDynamicMesh3* TargetMesh, const TFunction<FVector3d(FDynamicMesh3* Mesh, int)>& HandleVertexDeformFunc) override;
+
+	/** Updates the vertex positions of the mesh with the result from the last deformation solve. */
+	void ExportDeformedPositions(FDynamicMesh3* TargetMesh);
+
+	/** Returns true if the asynchronous task has finished. */
+	inline bool IsDone() { return AsyncMeshDeformTask == nullptr || AsyncMeshDeformTask->IsDone(); };
+
+	/** Triggers abort on task and passes off ownership to deleter object */
+	inline void Shutdown();
+
+	/** Stores the position of the vertex constraints and corresponding weights for the entire mesh.  This is used as a form of scratch space.*/
+	TArray<FConstraintData> SrcMeshConstraintBuffer;
+
+	/** Array of vertex indices organized in groups of three - basically an index buffer - that defines the subset of the mesh that the deformation task will work on.*/
+	TArray<int32> SubsetIDBuffer;
+
+	/** Need to update the task with the current submesh */
+	bool bTaskSubmeshIsDirty = true;
+
+	/** Asynchronous task object. This object deals with expensive matrix functionality that computes the deformation of a local mesh. */
+	FAsyncTaskExecuterWithAbort<FConstrainedMeshDeformerTask>* AsyncMeshDeformTask = nullptr;
+
+
+	/** The weight which will be applied to the constraints corresponding to the handle vertices. */
+	double HandleWeights = 1.0;
+
+	/** This is set to true whenever the user interacts with the tool under laplacian deformation mode.
+	  * It is set to false immediately before beginning a background task and cannot be set to false again until the work is done. */
+	bool bDeformerNeedsToRun = false;
+
+
+	/** When true, tells the solver to attempt to postfix the actual position of the handles to the constrained position */
+	bool bPostfixHandles = false;
+
+	//This is set to false only after 
+	//	1) the asynchronous deformation task is complete
+	//	2) the main thread has seen it complete, and
+	//	3) the main thread updates the vertex positions of the mesh one last time
+	bool bVertexPositionsNeedSync = false;
+
+	bool bLocalize =  true;
+
+};
+
+
+/**
+ *
+ */
+UCLASS()
+class MESHMODELINGTOOLS_API UEditMeshPolygonsTool : public UMeshSurfacePointTool
+{
+	GENERATED_BODY()
+
+public:
+	UEditMeshPolygonsTool();
+
+	virtual void RegisterActions(FInteractiveToolActionSet& ActionSet) override;
+
+	virtual void Setup() override;
+	virtual void Shutdown(EToolShutdownType ShutdownType) override;
+
+	virtual void Tick(float DeltaTime) override;
+	virtual void Render(IToolsContextRenderAPI* RenderAPI) override;
+
+	virtual bool HasCancel() const override { return true; }
+	virtual bool HasAccept() const override { return true; }
+	virtual bool CanAccept() const override { return true; }
+
+
+	virtual bool HitTest(const FRay& Ray, FHitResult& OutHit) override;
+	
+	virtual void OnBeginDrag(const FRay& Ray) override;
+	virtual void OnUpdateDrag(const FRay& Ray) override;
+	virtual void OnEndDrag(const FRay& Ray) override;
+	virtual void OnUpdateHover(const FInputDeviceRay& DevicePos) override;
+
+public:
+	virtual void ToggleLocalTransformFrameAction();
+	virtual void NextTransformTypeAction();
+
+	//
+	float VisualAngleSnapThreshold = 0.5;
+
+protected:
+	UPROPERTY()
+	USimpleDynamicMeshComponent* DynamicMeshComponent;
+
+	UPROPERTY()
+	UPolyEditTransformProperties* TransformProps;
+
+	// realtime visualization
+	void OnDynamicMeshComponentChanged();
+	FDelegateHandle OnDynamicMeshComponentChangedHandle;
+	
+	// camera state at last render
+	FViewCameraState CameraState;
+
+	FToolDataVisualizer PolyEdgesRenderer;
+
+	// True for the duration of UI click+drag
+	bool bInDrag;
+
+	FPlane ActiveDragPlane;
+	FVector StartHitPosWorld;
+	FVector StartHitNormalWorld;
+	FVector LastHitPosWorld;
+	FVector LastBrushPosLocal;
+	FVector StartBrushPosLocal;
+
+	FFrame3d ActiveSurfaceFrame;
+	FQuickTransformer* GetActiveQuickTransformer();
+	void UpdateActiveSurfaceFrame(FGroupTopologySelection& Selection);
+	void UpdateQuickTransformer();
+
+	FRay UpdateRay;
+	bool bUpdatePending = false;
+	void ComputeUpdate();
+
+	FVector3d LastMoveDelta;
+	FQuickAxisTranslater QuickAxisTranslater;
+	void ComputeUpdate_Translate();
+
+	FQuickAxisRotator QuickAxisRotator;
+	FVector3d RotationStartPointWorld;
+	FFrame3d RotationStartFrame;
+	void ComputeUpdate_Rotate();
+
+
+	FGroupTopology Topology;
+	void PrecomputeTopology();
+
+	FGroupTopologySelector TopoSelector;
+	
+	//
+	// data for current drag
+	//
+
+	FGroupTopologySelection HilightSelection;
+	FToolDataVisualizer HilightRenderer;
+
+	FDynamicMeshAABBTree3 MeshSpatial;
+	FDynamicMeshAABBTree3& GetSpatial();
+
+	FMeshVertexChangeBuilder* ActiveVertexChange;
+
+	// The two deformer type options.
+	FGroupTopologyDeformer LinearDeformer;
+	FGroupTopologyLaplacianDeformer LaplacianDeformer;
+	
+	// This is true when the spatial index needs to reflect a modification
+	bool bSpatialDirty; 
+
+	void BeginChange();
+	void EndChange();
+	void UpdateChangeFromROI(bool bFinal);
+
+};
+
diff --git a/Engine/Plugins/Experimental/MeshModelingToolset/Source/MeshModelingTools/Public/PolygonMeshDeformTool.h b/Engine/Plugins/Experimental/MeshModelingToolset/Source/MeshModelingTools/Public/PolygonMeshDeformTool.h
deleted file mode 100644
diff --git a/Engine/Plugins/Experimental/ModelingToolsEditorMode/Source/ModelingToolsEditorMode/Private/ModelingToolsActions.cpp b/Engine/Plugins/Experimental/ModelingToolsEditorMode/Source/ModelingToolsEditorMode/Private/ModelingToolsActions.cpp
--- a/Engine/Plugins/Experimental/ModelingToolsEditorMode/Source/ModelingToolsEditorMode/Private/ModelingToolsActions.cpp
+++ b/Engine/Plugins/Experimental/ModelingToolsEditorMode/Source/ModelingToolsEditorMode/Private/ModelingToolsActions.cpp
@@ -5,7 +5,7 @@
 #include "DynamicMeshSculptTool.h"
 #include "MeshInspectorTool.h"
 #include "DrawPolygonTool.h"
-#include "PolygonMeshDeformTool.h"
+#include "EditMeshPolygonsTool.h"
 #include "ShapeSprayTool.h"
 #include "MeshSpaceDeformerTool.h"
 #include "MeshSelectionTool.h"
@@ -32,7 +32,7 @@ void FModelingToolActionCommands::GetToolDefaultObjectList(TArray<UInteractiveTo
 	ToolCDOs.Add(GetMutableDefault<UDynamicMeshSculptTool>());
 	ToolCDOs.Add(GetMutableDefault<UMeshInspectorTool>());
 	ToolCDOs.Add(GetMutableDefault<UDrawPolygonTool>());
-	ToolCDOs.Add(GetMutableDefault<UPolygonMeshDeformTool>());
+	ToolCDOs.Add(GetMutableDefault<UEditMeshPolygonsTool>());
 	ToolCDOs.Add(GetMutableDefault<UMeshSpaceDeformerTool>());
 	ToolCDOs.Add(GetMutableDefault<UShapeSprayTool>());
 	ToolCDOs.Add(GetMutableDefault<UMeshSelectionTool>());
diff --git a/Engine/Plugins/Experimental/ModelingToolsEditorMode/Source/ModelingToolsEditorMode/Private/ModelingToolsEditorMode.cpp b/Engine/Plugins/Experimental/ModelingToolsEditorMode/Source/ModelingToolsEditorMode/Private/ModelingToolsEditorMode.cpp
--- a/Engine/Plugins/Experimental/ModelingToolsEditorMode/Source/ModelingToolsEditorMode/Private/ModelingToolsEditorMode.cpp
+++ b/Engine/Plugins/Experimental/ModelingToolsEditorMode/Source/ModelingToolsEditorMode/Private/ModelingToolsEditorMode.cpp
@@ -12,7 +12,7 @@
 //#include "MeshSurfacePointTool.h"
 //#include "MeshVertexDragTool.h"
 #include "DynamicMeshSculptTool.h"
-#include "PolygonMeshDeformTool.h"
+#include "EditMeshPolygonsTool.h"
 #include "ConvertToPolygonsTool.h"
 #include "AddPrimitiveTool.h"
 #include "SmoothMeshTool.h"
@@ -400,7 +400,7 @@ void FModelingToolsEditorMode::Enter()
 	MoveVerticesToolBuilder->bEnableRemeshing = false;
 	RegisterToolFunc(ToolManagerCommands.BeginSculptMeshTool, TEXT("MoveVerticesTool"), MoveVerticesToolBuilder);
 
-	RegisterToolFunc(ToolManagerCommands.BeginPolyEditTool, TEXT("PolygonMeshDeformTool"), NewObject<UPolygonMeshDeformToolBuilder>());
+	RegisterToolFunc(ToolManagerCommands.BeginPolyEditTool, TEXT("EditMeshPolygonsTool"), NewObject<UEditMeshPolygonsToolBuilder>());
 	RegisterToolFunc(ToolManagerCommands.BeginSmoothMeshTool, TEXT("SmoothMeshTool"), NewObject<USmoothMeshToolBuilder>());
 	RegisterToolFunc(ToolManagerCommands.BeginDisplaceMeshTool, TEXT("DisplaceMeshTool"), NewObject<UDisplaceMeshToolBuilder>());
 	RegisterToolFunc(ToolManagerCommands.BeginMeshSpaceDeformerTool, TEXT("MeshSpaceDeformerTool"), NewObject<UMeshSpaceDeformerToolBuilder>());
