MeshModelingTools : Disable properties of Add Primitive Tool which are not relevant to the current Primitive

diff --git a/Engine/Plugins/Experimental/MeshModelingToolset/Source/MeshModelingTools/Private/AddPrimitiveTool.cpp b/Engine/Plugins/Experimental/MeshModelingToolset/Source/MeshModelingTools/Private/AddPrimitiveTool.cpp
--- a/Engine/Plugins/Experimental/MeshModelingToolset/Source/MeshModelingTools/Private/AddPrimitiveTool.cpp
+++ b/Engine/Plugins/Experimental/MeshModelingToolset/Source/MeshModelingTools/Private/AddPrimitiveTool.cpp
@@ -13,7 +13,6 @@
 #include "Generators/RectangleMeshGenerator.h"
 #include "Generators/SphereGenerator.h"
 #include "DynamicMeshAttributeSet.h"
-
 #include "StaticMeshComponentBuilder.h"
 #include "Drawing/MeshDebugDrawing.h"
 
@@ -23,8 +22,6 @@
 /*
  * ToolBuilder
  */
-
-
 bool UAddPrimitiveToolBuilder::CanBuildTool(const FToolBuilderState& SceneState) const
 {
 	return (this->AssetAPI != nullptr);
@@ -38,12 +35,9 @@ UInteractiveTool* UAddPrimitiveToolBuilder::BuildTool(const FToolBuilderState& S
 	return NewTool;
 }
 
-
-
 /*
  * Tool
  */
-
 UProceduralShapeToolProperties::UProceduralShapeToolProperties()
 {
 	Shape = EMakeMeshShapeType::Box;
@@ -58,6 +52,36 @@ UProceduralShapeToolProperties::UProceduralShapeToolProperties()
     Material = CreateDefaultSubobject<UMaterialInterface>(TEXT("MATERIAL"));
 }
 
+namespace
+{
+	struct {
+		FName Name;
+		EMakeMeshShapeType EnabledShapes;
+	} EnabledShapesMap[] =
+	{ 
+	  { TEXT("Shape"),        EMakeMeshShapeType::All },
+	  { TEXT("Width"),        EMakeMeshShapeType::All },
+	  { TEXT("Height"),       EMakeMeshShapeType::Box | EMakeMeshShapeType::Cylinder | EMakeMeshShapeType::Cone | EMakeMeshShapeType::Plane },
+	  { TEXT("Rotation"),     EMakeMeshShapeType::All },
+	  { TEXT("PlaceMode"),    EMakeMeshShapeType::All },
+	  { TEXT("bCentered"),    EMakeMeshShapeType::All },
+	  { TEXT("Slices"),       EMakeMeshShapeType::Cylinder | EMakeMeshShapeType::Cone | EMakeMeshShapeType::Sphere },
+	  { TEXT("Subdivisions"), EMakeMeshShapeType::Box | EMakeMeshShapeType::Plane },
+	  { TEXT("Material"),     EMakeMeshShapeType::All },
+	};
+};
+
+bool
+UProceduralShapeToolProperties::CanEditChange( const UProperty* InProperty) const
+{
+	auto* Elem = Algo::FindByPredicate(EnabledShapesMap, [InProperty](const auto& Elem) { return Elem.Name == InProperty->GetFName(); });
+	if (Elem != nullptr)
+	{
+		return (Shape & Elem->EnabledShapes) != EMakeMeshShapeType::None;
+	}
+	checkNoEntry();
+	return false;
+}
 
 void UAddPrimitiveTool::SetWorld(UWorld* World)
 {
@@ -213,9 +237,6 @@ void UAddPrimitiveTool::UpdatePreviewMesh()
 	PreviewMesh->UpdatePreview(&NewMesh);
 }
 
-
-
-
 void UAddPrimitiveTool::OnClicked(const FInputDeviceRay& DeviceClickPos)
 {
 #if WITH_EDITOR
@@ -229,8 +250,8 @@ void UAddPrimitiveTool::OnClicked(const FInputDeviceRay& DeviceClickPos)
 		CurMesh, CurTransform, TEXT("ShapeMesh"),
 		AssetGenerationUtil::GetDefaultAutoGeneratedAssetPath(),
 		Material
-		);
-	
+	);
+
 	// select newly-created object
 	ToolSelectionUtil::SetNewActorSelection(GetToolManager(), NewActor);
 
diff --git a/Engine/Plugins/Experimental/MeshModelingToolset/Source/MeshModelingTools/Public/AddPrimitiveTool.h b/Engine/Plugins/Experimental/MeshModelingToolset/Source/MeshModelingTools/Public/AddPrimitiveTool.h
--- a/Engine/Plugins/Experimental/MeshModelingToolset/Source/MeshModelingTools/Public/AddPrimitiveTool.h
+++ b/Engine/Plugins/Experimental/MeshModelingToolset/Source/MeshModelingTools/Public/AddPrimitiveTool.h
@@ -32,20 +32,19 @@ public:
 	virtual UInteractiveTool* BuildTool(const FToolBuilderState& SceneState) const override;
 };
 
-
-
 /** Shape Types */
 UENUM()
 enum class EMakeMeshShapeType : uint8
 {
-	Box = 0 UMETA(DisplayName = "Box"),
-	Cylinder = 1 UMETA(DisplayName = "Cylinder"),
-	Cone = 2 UMETA(DisplayName = "Cone"),
-	Plane = 3 UMETA(DisplayName = "Plane"),
-	Sphere = 4 UMETA(DisplayName = "Sphere")
+	None     = 0x00 UMETA(DisplayName = "None", Hidden),
+	All      = 0xff UMETA(DisplayName = "All", Hidden), 
+	Box      = 0x01 UMETA(DisplayName = "Box"),
+	Cylinder = 0x02 UMETA(DisplayName = "Cylinder"),
+	Cone     = 0x04 UMETA(DisplayName = "Cone"),
+	Plane    = 0x08 UMETA(DisplayName = "Plane"),
+	Sphere   = 0x10 UMETA(DisplayName = "Sphere")
 };
-
-
+ENUM_CLASS_FLAGS(EMakeMeshShapeType);
 
 /** Placement Target Types */
 UENUM()
@@ -55,10 +54,6 @@ enum class EMakeMeshPlacementType : uint8
 	OnScene = 1
 };
 
-
-
-
-
 UCLASS()
 class MESHMODELINGTOOLS_API UProceduralShapeToolProperties : public UInteractiveToolPropertySet
 {
@@ -66,6 +61,7 @@ class MESHMODELINGTOOLS_API UProceduralShapeToolProperties : public UInteractive
 
 public:
 	UProceduralShapeToolProperties();
+	bool CanEditChange( const UProperty* InProperty ) const override;
 
 	/** Type of shape to generate */
 	UPROPERTY(EditAnywhere, Category = ShapeSettings)
