MeshModelingToolset : Add Combo Box to set pivot location of new primitives

diff --git a/Engine/Plugins/Experimental/MeshModelingToolset/Source/MeshModelingTools/Private/AddPrimitiveTool.cpp b/Engine/Plugins/Experimental/MeshModelingToolset/Source/MeshModelingTools/Private/AddPrimitiveTool.cpp
--- a/Engine/Plugins/Experimental/MeshModelingToolset/Source/MeshModelingTools/Private/AddPrimitiveTool.cpp
+++ b/Engine/Plugins/Experimental/MeshModelingToolset/Source/MeshModelingTools/Private/AddPrimitiveTool.cpp
@@ -47,7 +47,7 @@ UProceduralShapeToolProperties::UProceduralShapeToolProperties()
 	//EndAngle = 360;
 	Slices = 16;
 	Subdivisions = 0;
-	bCentered = false;
+	PivotLocation = EMakeMeshPivotLocation::Base;
 	PlaceMode = EMakeMeshPlacementType::OnScene;
     Material = CreateDefaultSubobject<UMaterialInterface>(TEXT("MATERIAL"));
 }
@@ -59,15 +59,15 @@ namespace
 		EMakeMeshShapeType EnabledShapes;
 	} EnabledShapesMap[] =
 	{ 
-	  { TEXT("Shape"),        EMakeMeshShapeType::All },
-	  { TEXT("Width"),        EMakeMeshShapeType::All },
-	  { TEXT("Height"),       EMakeMeshShapeType::Box | EMakeMeshShapeType::Cylinder | EMakeMeshShapeType::Cone | EMakeMeshShapeType::Plane },
-	  { TEXT("Rotation"),     EMakeMeshShapeType::All },
-	  { TEXT("PlaceMode"),    EMakeMeshShapeType::All },
-	  { TEXT("bCentered"),    EMakeMeshShapeType::All },
-	  { TEXT("Slices"),       EMakeMeshShapeType::Cylinder | EMakeMeshShapeType::Cone | EMakeMeshShapeType::Sphere },
-	  { TEXT("Subdivisions"), EMakeMeshShapeType::Box | EMakeMeshShapeType::Plane },
-	  { TEXT("Material"),     EMakeMeshShapeType::All },
+	  { TEXT("Shape"),         EMakeMeshShapeType::All },
+	  { TEXT("Width"),         EMakeMeshShapeType::All },
+	  { TEXT("Height"),        EMakeMeshShapeType::Box | EMakeMeshShapeType::Cylinder | EMakeMeshShapeType::Cone | EMakeMeshShapeType::Plane },
+	  { TEXT("Rotation"),      EMakeMeshShapeType::All },
+	  { TEXT("PlaceMode"),     EMakeMeshShapeType::All },
+	  { TEXT("PivotLocation"), EMakeMeshShapeType::All },
+	  { TEXT("Slices"),        EMakeMeshShapeType::Cylinder | EMakeMeshShapeType::Cone | EMakeMeshShapeType::Sphere },
+	  { TEXT("Subdivisions"),  EMakeMeshShapeType::Box | EMakeMeshShapeType::Plane },
+	  { TEXT("Material"),      EMakeMeshShapeType::All },
 	};
 };
 
@@ -222,7 +222,7 @@ void UAddPrimitiveTool::UpdatePreviewMesh()
 	// set mesh position
 	FAxisAlignedBox3d Bounds = NewMesh.GetCachedBounds();
 	FVector3d TargetOrigin = Bounds.Center();
-	if (ShapeSettings->bCentered == false)
+	if (ShapeSettings->PivotLocation == EMakeMeshPivotLocation::Base)
 	{
 		TargetOrigin.Z = Bounds.Min.Z;
 	}
diff --git a/Engine/Plugins/Experimental/MeshModelingToolset/Source/MeshModelingTools/Public/AddPrimitiveTool.h b/Engine/Plugins/Experimental/MeshModelingToolset/Source/MeshModelingTools/Public/AddPrimitiveTool.h
--- a/Engine/Plugins/Experimental/MeshModelingToolset/Source/MeshModelingTools/Public/AddPrimitiveTool.h
+++ b/Engine/Plugins/Experimental/MeshModelingToolset/Source/MeshModelingTools/Public/AddPrimitiveTool.h
@@ -54,6 +54,14 @@ enum class EMakeMeshPlacementType : uint8
 	OnScene = 1
 };
 
+/** Placement Pivot Location */
+UENUM()
+enum class EMakeMeshPivotLocation : uint8
+{
+	Base,
+	Centered
+};
+
 UCLASS()
 class MESHMODELINGTOOLS_API UProceduralShapeToolProperties : public UInteractiveToolPropertySet
 {
@@ -85,7 +93,7 @@ public:
 
 	/** Center shape at click point */
 	UPROPERTY(EditAnywhere, Category = ShapeSettings)
-	bool bCentered;
+    EMakeMeshPivotLocation PivotLocation;
 
 	///** Start Angle of Shape */
 	//UPROPERTY(EditAnywhere, Category = ShapeSettings, meta = (DisplayName = "Start Angle", UIMin = "0.0", UIMax = "360.0", ClampMin = "-10000", ClampMax = "10000.0"))
