polyedit settings to change polygon boundaries; initial version, needs threading fix + revisit of undo/redo behavior

#rnx

diff --git a/Engine/Plugins/Experimental/GeometryProcessing/Source/GeometricObjects/Private/Spatial/GeometrySet3.cpp b/Engine/Plugins/Experimental/GeometryProcessing/Source/GeometricObjects/Private/Spatial/GeometrySet3.cpp
--- a/Engine/Plugins/Experimental/GeometryProcessing/Source/GeometricObjects/Private/Spatial/GeometrySet3.cpp
+++ b/Engine/Plugins/Experimental/GeometryProcessing/Source/GeometricObjects/Private/Spatial/GeometrySet3.cpp
@@ -10,10 +10,12 @@ void FGeometrySet3::Reset(bool bPoints, bool bCurves)
 	if (bPoints)
 	{
 		Points.Reset();
+		PointIDToIndex.Reset();
 	}
 	if (bCurves)
 	{
 		Curves.Reset();
+		CurveIDToIndex.Reset();
 	}
 }
 
diff --git a/Engine/Plugins/Experimental/GeometryProcessing/Source/MeshConversion/Private/DynamicMeshToMeshDescription.cpp b/Engine/Plugins/Experimental/GeometryProcessing/Source/MeshConversion/Private/DynamicMeshToMeshDescription.cpp
--- a/Engine/Plugins/Experimental/GeometryProcessing/Source/MeshConversion/Private/DynamicMeshToMeshDescription.cpp
+++ b/Engine/Plugins/Experimental/GeometryProcessing/Source/MeshConversion/Private/DynamicMeshToMeshDescription.cpp
@@ -150,7 +150,7 @@ void FDynamicMeshToMeshDescription::Convert_NoAttributes(const FDynamicMesh3* Me
 	Builder.SetMeshDescription(&MeshOut);
 
 	bool bCopyGroupToPolyGroup = false;
-	if (bSetPolyGroups && MeshIn->HasTriangleGroups())
+	if (ConversionOptions.bSetPolyGroups && MeshIn->HasTriangleGroups())
 	{
 		Builder.EnablePolyGroups();
 		bCopyGroupToPolyGroup = true;
@@ -216,7 +216,7 @@ void FDynamicMeshToMeshDescription::Convert_SharedInstances(const FDynamicMesh3*
 	Builder.SetMeshDescription(&MeshOut);
 
 	bool bCopyGroupToPolyGroup = false;
-	if (bSetPolyGroups && MeshIn->HasTriangleGroups())
+	if (ConversionOptions.bSetPolyGroups && MeshIn->HasTriangleGroups())
 	{
 		Builder.EnablePolyGroups();
 		bCopyGroupToPolyGroup = true;
@@ -280,7 +280,7 @@ void FDynamicMeshToMeshDescription::Convert_NoSharedInstances(const FDynamicMesh
 	Builder.SetMeshDescription(&MeshOut);
 
 	bool bCopyGroupToPolyGroup = false;
-	if (bSetPolyGroups && MeshIn->HasTriangleGroups())
+	if (ConversionOptions.bSetPolyGroups && MeshIn->HasTriangleGroups())
 	{
 		Builder.EnablePolyGroups();
 		bCopyGroupToPolyGroup = true;
diff --git a/Engine/Plugins/Experimental/GeometryProcessing/Source/MeshConversion/Public/DynamicMeshToMeshDescription.h b/Engine/Plugins/Experimental/GeometryProcessing/Source/MeshConversion/Public/DynamicMeshToMeshDescription.h
--- a/Engine/Plugins/Experimental/GeometryProcessing/Source/MeshConversion/Public/DynamicMeshToMeshDescription.h
+++ b/Engine/Plugins/Experimental/GeometryProcessing/Source/MeshConversion/Public/DynamicMeshToMeshDescription.h
@@ -5,6 +5,7 @@
 #include "CoreMinimal.h"
 #include "DynamicMesh3.h"
 #include "MeshDescription.h"
+#include "MeshConversionOptions.h"
 
 /**
  * Convert FDynamicMesh3 to FMeshDescription
@@ -16,8 +17,16 @@ public:
 	/** If true, will print some possibly-helpful debugging spew to output log */
 	bool bPrintDebugMessages = false;
 
-	/** Should DynamicMesh triangle groups be transfered to MeshDescription via custom PolyTriGroups attribute */
-	bool bSetPolyGroups = true;
+	/** General settings for conversions to mesh description */
+	FConversionToMeshDescriptionOptions ConversionOptions;
+
+	FDynamicMeshToMeshDescription()
+	{
+	}
+
+	FDynamicMeshToMeshDescription(FConversionToMeshDescriptionOptions ConversionOptions) : ConversionOptions(ConversionOptions)
+	{
+	}
 
 	/**
 	 * Default conversion of DynamicMesh to MeshDescription. Calls functions below depending on mesh state
diff --git a/Engine/Plugins/Experimental/GeometryProcessing/Source/MeshConversion/Public/MeshConversionOptions.h b/Engine/Plugins/Experimental/GeometryProcessing/Source/MeshConversion/Public/MeshConversionOptions.h
new file mode 100644
--- /dev/null
+++ b/Engine/Plugins/Experimental/GeometryProcessing/Source/MeshConversion/Public/MeshConversionOptions.h
@@ -0,0 +1,15 @@
+// Copyright 1998-2019 Epic Games, Inc. All Rights Reserved.
+
+#pragma once
+
+#include "CoreMinimal.h"
+
+
+struct MESHCONVERSION_API FConversionToMeshDescriptionOptions
+{
+public:
+	/** Should triangle groups be transfered to MeshDescription via custom PolyTriGroups attribute */
+	bool bSetPolyGroups = true;
+
+	// TODO: add additional conversion options (e.g. normal / tangent recompute, etc)
+};
\ No newline at end of file
diff --git a/Engine/Plugins/Experimental/MeshModelingToolset/Source/MeshModelingTools/Private/ConvertToPolygonsTool.cpp b/Engine/Plugins/Experimental/MeshModelingToolset/Source/MeshModelingTools/Private/ConvertToPolygonsTool.cpp
--- a/Engine/Plugins/Experimental/MeshModelingToolset/Source/MeshModelingTools/Private/ConvertToPolygonsTool.cpp
+++ b/Engine/Plugins/Experimental/MeshModelingToolset/Source/MeshModelingTools/Private/ConvertToPolygonsTool.cpp
@@ -199,7 +199,7 @@ void UConvertToPolygonsTool::ConvertToPolygons(FMeshDescription* MeshIn)
 	}
 
 	FDynamicMeshToMeshDescription Converter;
-	Converter.bSetPolyGroups = true;
+	Converter.ConversionOptions.bSetPolyGroups = true;
 	Converter.Convert(&SearchMesh, *MeshIn);
 }
 
diff --git a/Engine/Plugins/Experimental/MeshModelingToolset/Source/MeshModelingTools/Private/EditMeshPolygonsTool.cpp b/Engine/Plugins/Experimental/MeshModelingToolset/Source/MeshModelingTools/Private/EditMeshPolygonsTool.cpp
--- a/Engine/Plugins/Experimental/MeshModelingToolset/Source/MeshModelingTools/Private/EditMeshPolygonsTool.cpp
+++ b/Engine/Plugins/Experimental/MeshModelingToolset/Source/MeshModelingTools/Private/EditMeshPolygonsTool.cpp
@@ -9,6 +9,7 @@
 #include "MeshNormals.h"
 #include "ToolSceneQueriesUtil.h"
 #include "Intersection/IntersectionUtil.h"
+#include "FindPolygonsAlgorithm.h"
 
 #include "Async/ParallelFor.h"
 #include "Containers/BitArray.h"
@@ -51,6 +52,8 @@ UPolyEditTransformProperties::UPolyEditTransformProperties()
 	bSelectVertices = true;
 	bSelectFaces = true;
 	bSelectEdges = true;
+	bRecomputePolygonGroups = false;
+	PolygonGroupingAngleThreshold = .5;
 }
 
 /*
@@ -638,6 +641,8 @@ void UEditMeshPolygonsTool::Setup()
 
 	// add properties
 	TransformProps = NewObject<UPolyEditTransformProperties>(this);
+	TransformProps->bRecomputePolygonGroups = !DynamicMeshComponent->GetMesh()->HasTriangleGroups();
+	BackupTriangleGroups();
 	AddToolPropertySource(TransformProps);
 
 	// set up visualizers
@@ -686,7 +691,9 @@ void UEditMeshPolygonsTool::Shutdown(EToolShutdownType ShutdownType)
 			GetToolManager()->BeginUndoTransaction(LOCTEXT("EditMeshPolygonsToolTransactionName", "Deform Mesh"));
 			ComponentTarget->CommitMesh([=](FMeshDescription* MeshDescription)
 			{
-				DynamicMeshComponent->Bake(MeshDescription, false);
+				FConversionToMeshDescriptionOptions ConversionOptions;
+				ConversionOptions.bSetPolyGroups = false; // don't save polygroups, as we may change these temporarily in this tool just to get a different edit effect
+				DynamicMeshComponent->Bake(MeshDescription, false, ConversionOptions);
 			});
 			GetToolManager()->EndUndoTransaction();
 		}
@@ -1371,7 +1378,56 @@ void UEditMeshPolygonsTool::Render(IToolsContextRenderAPI* RenderAPI)
 }
 
 
+void UEditMeshPolygonsTool::OnPropertyModified(UObject* PropertySet, UProperty* Property)
+{
+	if (!Property) // happens on undo/redo of property modifications if object is set to be transactional
+	{
+		return;
+	}
+	if (Property->GetFName() == GET_MEMBER_NAME_CHECKED(UPolyEditTransformProperties, bRecomputePolygonGroups)
+	 || (Property->GetFName() == GET_MEMBER_NAME_CHECKED(UPolyEditTransformProperties, PolygonGroupingAngleThreshold)))
+	{
+		if (TransformProps->bRecomputePolygonGroups)
+		{
+			FDynamicMesh3* SearchMesh = DynamicMeshComponent->GetMesh();
+			FFindPolygonsAlgorithm Polygons = FFindPolygonsAlgorithm(SearchMesh);
+			double DotTolerance = 1.0 - FMathd::Cos(TransformProps->PolygonGroupingAngleThreshold * FMathd::DegToRad);
+			Polygons.FindPolygons(DotTolerance);
+			Polygons.FindPolygonEdges();
+		}
+		else
+		{
+			RestoreTriangleGroups();
+		}
+		PrecomputeTopology();
+		TopoSelector.Invalidate(false, true);
+	}
+}
+
+void UEditMeshPolygonsTool::BackupTriangleGroups()
+{
+	FDynamicMesh3* Mesh = DynamicMeshComponent->GetMesh();
+	if (Mesh->HasTriangleGroups())
+	{
+		InitialTriangleGroups = *Mesh->GetTriangleGroupsBuffer();
+	}
+	else
+	{
+		InitialTriangleGroups.SetNum(0);
+	}
+}
 
+void UEditMeshPolygonsTool::RestoreTriangleGroups()
+{
+	FDynamicMesh3* Mesh = DynamicMeshComponent->GetMesh();
+	for (int TID = 0, MaxID = InitialTriangleGroups.Num(); TID < MaxID; TID++)
+	{
+		if (Mesh->IsTriangle(TID))
+		{
+			Mesh->SetTriangleGroup(TID, InitialTriangleGroups[TID]);
+		}
+	}
+}
 
 
 //
diff --git a/Engine/Plugins/Experimental/MeshModelingToolset/Source/MeshModelingTools/Public/EditMeshPolygonsTool.h b/Engine/Plugins/Experimental/MeshModelingToolset/Source/MeshModelingTools/Public/EditMeshPolygonsTool.h
--- a/Engine/Plugins/Experimental/MeshModelingToolset/Source/MeshModelingTools/Public/EditMeshPolygonsTool.h
+++ b/Engine/Plugins/Experimental/MeshModelingToolset/Source/MeshModelingTools/Public/EditMeshPolygonsTool.h
@@ -111,6 +111,12 @@ public:
 	UPROPERTY(EditAnywhere, Category = Options)
 	bool bSelectVertices;
 
+	UPROPERTY(EditAnywhere, Category = Options)
+	bool bRecomputePolygonGroups;
+
+	UPROPERTY(EditAnywhere, Category = Options, meta = (EditCondition="bRecomputePolygonGroups"))
+	float PolygonGroupingAngleThreshold;
+
 	//Laplacian Deformation Options
 
 	UPROPERTY(EditAnywhere, Category = LaplacianOptions, meta = (EditCondition="DeformationStrategy == EGroupTopologyDeformationStrategy::Laplacian", DisplayName = "Laplacian Weight Scheme", ToolTip = "Laplacian weight schemes determine how the deformer will interpret the curvature at a given vertex in relation to nearby vertices in order to generate a solution."))
@@ -418,6 +424,8 @@ protected:
 	// realtime visualization
 	void OnDynamicMeshComponentChanged();
 	FDelegateHandle OnDynamicMeshComponentChangedHandle;
+
+	virtual void OnPropertyModified(UObject* PropertySet, UProperty* Property);
 	
 	// camera state at last render
 	FViewCameraState CameraState;
@@ -473,6 +481,11 @@ protected:
 	// The two deformer type options.
 	FGroupTopologyDeformer LinearDeformer;
 	FGroupTopologyLaplacianDeformer LaplacianDeformer;
+
+	// Initial polygon group info
+	TDynamicVector<int> InitialTriangleGroups;
+	void BackupTriangleGroups();
+	void RestoreTriangleGroups();
 	
 	// This is true when the spatial index needs to reflect a modification
 	bool bSpatialDirty; 
diff --git a/Engine/Plugins/Experimental/MeshModelingToolset/Source/ModelingComponents/Private/OctreeDynamicMeshComponent.cpp b/Engine/Plugins/Experimental/MeshModelingToolset/Source/ModelingComponents/Private/OctreeDynamicMeshComponent.cpp
--- a/Engine/Plugins/Experimental/MeshModelingToolset/Source/ModelingComponents/Private/OctreeDynamicMeshComponent.cpp
+++ b/Engine/Plugins/Experimental/MeshModelingToolset/Source/ModelingComponents/Private/OctreeDynamicMeshComponent.cpp
@@ -94,16 +94,16 @@ void UOctreeDynamicMeshComponent::InitializeNewMesh()
 }
 
 
-void UOctreeDynamicMeshComponent::Bake(FMeshDescription* MeshDescription, bool bHaveModifiedToplogy)
+void UOctreeDynamicMeshComponent::Bake(FMeshDescription* MeshDescription, bool bHaveModifiedTopology, const FConversionToMeshDescriptionOptions& ConversionOptions)
 {
-	if (bHaveModifiedToplogy == false)
+	if (bHaveModifiedTopology == false)
 	{
-		FDynamicMeshToMeshDescription Converter;
+		FDynamicMeshToMeshDescription Converter(ConversionOptions);
 		Converter.Update(Mesh.Get(), *MeshDescription);
 	}
 	else
 	{
-		FDynamicMeshToMeshDescription Converter;
+		FDynamicMeshToMeshDescription Converter(ConversionOptions);
 		Converter.Convert(Mesh.Get(), *MeshDescription);
 
 		//UE_LOG(LogTemp, Warning, TEXT("MeshDescription has %d instances"), MeshDescription->VertexInstances().Num());
diff --git a/Engine/Plugins/Experimental/MeshModelingToolset/Source/ModelingComponents/Private/SimpleDynamicMeshComponent.cpp b/Engine/Plugins/Experimental/MeshModelingToolset/Source/ModelingComponents/Private/SimpleDynamicMeshComponent.cpp
--- a/Engine/Plugins/Experimental/MeshModelingToolset/Source/ModelingComponents/Private/SimpleDynamicMeshComponent.cpp
+++ b/Engine/Plugins/Experimental/MeshModelingToolset/Source/ModelingComponents/Private/SimpleDynamicMeshComponent.cpp
@@ -85,16 +85,16 @@ void USimpleDynamicMeshComponent::InitializeNewMesh()
 }
 
 
-void USimpleDynamicMeshComponent::Bake(FMeshDescription* MeshDescription, bool bHaveModifiedToplogy)
+void USimpleDynamicMeshComponent::Bake(FMeshDescription* MeshDescription, bool bHaveModifiedTopology, const FConversionToMeshDescriptionOptions& ConversionOptions)
 {
-	if (bHaveModifiedToplogy == false)
+	if (bHaveModifiedTopology == false)
 	{
-		FDynamicMeshToMeshDescription Converter;
+		FDynamicMeshToMeshDescription Converter(ConversionOptions);
 		Converter.Update(Mesh.Get(), *MeshDescription);
 	}
 	else
 	{
-		FDynamicMeshToMeshDescription Converter;
+		FDynamicMeshToMeshDescription Converter(ConversionOptions);
 		Converter.Convert(Mesh.Get(), *MeshDescription);
 
 		//UE_LOG(LogTemp, Warning, TEXT("MeshDescription has %d instances"), MeshDescription->VertexInstances().Num());
diff --git a/Engine/Plugins/Experimental/MeshModelingToolset/Source/ModelingComponents/Public/OctreeDynamicMeshComponent.h b/Engine/Plugins/Experimental/MeshModelingToolset/Source/ModelingComponents/Public/OctreeDynamicMeshComponent.h
--- a/Engine/Plugins/Experimental/MeshModelingToolset/Source/ModelingComponents/Public/OctreeDynamicMeshComponent.h
+++ b/Engine/Plugins/Experimental/MeshModelingToolset/Source/ModelingComponents/Public/OctreeDynamicMeshComponent.h
@@ -4,6 +4,7 @@
 
 #include "CoreMinimal.h"
 #include "BaseDynamicMeshComponent.h"
+#include "MeshConversionOptions.h"
 
 #include "DynamicMeshOctree3.h"
 #include "Util/IndexSetDecompositions.h"
@@ -56,9 +57,20 @@ public:
 
 	/**
 	 * Write the internal mesh to a MeshDescription
-	 * @param bHaveModifiedToplogy if false, we only update the vertex positions in the MeshDescription, otherwise it is Empty()'d and regenerated entirely
+	 * @param bHaveModifiedTopology if false, we only update the vertex positions in the MeshDescription, otherwise it is Empty()'d and regenerated entirely
+	 * @param ConversionOptions struct of additional options for the conversion
 	 */
-	void Bake(FMeshDescription* MeshDescription, bool bHaveModifiedToplogy);
+	void Bake(FMeshDescription* MeshDescription, bool bHaveModifiedTopology, const FConversionToMeshDescriptionOptions& ConversionOptions);
+
+	/**
+	* Write the internal mesh to a MeshDescription with default conversion options
+	* @param bHaveModifiedTopology if false, we only update the vertex positions in the MeshDescription, otherwise it is Empty()'d and regenerated entirely
+	*/
+	void Bake(FMeshDescription* MeshDescription, bool bHaveModifiedTopology)
+	{
+		FConversionToMeshDescriptionOptions ConversionOptions;
+		Bake(MeshDescription, bHaveModifiedTopology, ConversionOptions);
+	}
 
 
 	//
diff --git a/Engine/Plugins/Experimental/MeshModelingToolset/Source/ModelingComponents/Public/SimpleDynamicMeshComponent.h b/Engine/Plugins/Experimental/MeshModelingToolset/Source/ModelingComponents/Public/SimpleDynamicMeshComponent.h
--- a/Engine/Plugins/Experimental/MeshModelingToolset/Source/ModelingComponents/Public/SimpleDynamicMeshComponent.h
+++ b/Engine/Plugins/Experimental/MeshModelingToolset/Source/ModelingComponents/Public/SimpleDynamicMeshComponent.h
@@ -4,6 +4,7 @@
 
 #include "CoreMinimal.h"
 #include "BaseDynamicMeshComponent.h"
+#include "MeshConversionOptions.h"
 
 #include "MeshTangents.h"
 
@@ -62,9 +63,20 @@ public:
 
 	/**
 	 * Write the internal mesh to a MeshDescription
-	 * @param bHaveModifiedToplogy if false, we only update the vertex positions in the MeshDescription, otherwise it is Empty()'d and regenerated entirely
+	 * @param bHaveModifiedTopology if false, we only update the vertex positions in the MeshDescription, otherwise it is Empty()'d and regenerated entirely
+	 * @param ConversionOptions struct of additional options for the conversion
 	 */
-	void Bake(FMeshDescription* MeshDescription, bool bHaveModifiedToplogy);
+	void Bake(FMeshDescription* MeshDescription, bool bHaveModifiedTopology, const FConversionToMeshDescriptionOptions& ConversionOptions);
+
+	/**
+	* Write the internal mesh to a MeshDescription with default conversion options
+	* @param bHaveModifiedTopology if false, we only update the vertex positions in the MeshDescription, otherwise it is Empty()'d and regenerated entirely
+	*/
+	void Bake(FMeshDescription* MeshDescription, bool bHaveModifiedTopology)
+	{
+		FConversionToMeshDescriptionOptions ConversionOptions;
+		Bake(MeshDescription, bHaveModifiedTopology, ConversionOptions);
+	}
 
 
 	//
