Fix clang compile error in AssetGenerationUtil.cpp
#rnx

diff --git a/Engine/Plugins/Experimental/MeshModelingToolset/Source/ModelingComponents/Private/AssetGenerationUtil.cpp b/Engine/Plugins/Experimental/MeshModelingToolset/Source/ModelingComponents/Private/AssetGenerationUtil.cpp
--- a/Engine/Plugins/Experimental/MeshModelingToolset/Source/ModelingComponents/Private/AssetGenerationUtil.cpp
+++ b/Engine/Plugins/Experimental/MeshModelingToolset/Source/ModelingComponents/Private/AssetGenerationUtil.cpp
@@ -75,8 +75,7 @@ AActor* AssetGenerationUtil::GenerateStaticMeshActor(
 
 	// transform new component to origin
 	Builder.NewMeshComponent->SetWorldTransform((FTransform)Transform);
-	UMaterialInstance;
-	UMaterialInterface;
+
 	if (Material)
 	{
 		Builder.NewMeshComponent->SetMaterial(0, Material);
