Attempted fix for compile error with AddPrimitiveTool.h
#rnx

diff --git a/Engine/Plugins/Experimental/MeshModelingToolset/Source/MeshModelingTools/Public/AddPrimitiveTool.h b/Engine/Plugins/Experimental/MeshModelingToolset/Source/MeshModelingTools/Public/AddPrimitiveTool.h
--- a/Engine/Plugins/Experimental/MeshModelingToolset/Source/MeshModelingTools/Public/AddPrimitiveTool.h
+++ b/Engine/Plugins/Experimental/MeshModelingToolset/Source/MeshModelingTools/Public/AddPrimitiveTool.h
@@ -69,7 +69,7 @@ class MESHMODELINGTOOLS_API UProceduralShapeToolProperties : public UInteractive
 
 public:
 	UProceduralShapeToolProperties();
-	bool CanEditChange( const UProperty* InProperty ) const override;
+	virtual bool CanEditChange( const UProperty* InProperty ) const override;
 
 	/** Type of shape to generate */
 	UPROPERTY(EditAnywhere, Category = ShapeSettings)
