ModelingTools: clean up Remesh and Simplify properties (move to ToolPropertySet, edit conditions, etc)
#rnx

diff --git a/Engine/Plugins/Experimental/MeshModelingToolset/Source/MeshModelingTools/Private/RemeshMeshTool.cpp b/Engine/Plugins/Experimental/MeshModelingToolset/Source/MeshModelingTools/Private/RemeshMeshTool.cpp
--- a/Engine/Plugins/Experimental/MeshModelingToolset/Source/MeshModelingTools/Private/RemeshMeshTool.cpp
+++ b/Engine/Plugins/Experimental/MeshModelingToolset/Source/MeshModelingTools/Private/RemeshMeshTool.cpp
@@ -55,10 +55,7 @@ URemeshMeshToolProperties::URemeshMeshToolProperties()
 	bDiscardAttributes = false;
 	SmoothingType = ERemeshSmoothingType::MeanValue;
 	bPreserveSharpEdges = true;
-}
 
-URemeshMeshAdvancedProperties::URemeshMeshAdvancedProperties()
-{
 	TargetEdgeLength = 5.0;
 	bFlips = true;
 	bSplits = true;
@@ -69,6 +66,7 @@ URemeshMeshAdvancedProperties::URemeshMeshAdvancedProperties()
 }
 
 
+
 URemeshMeshTool::URemeshMeshTool()
 {
 }
@@ -106,16 +104,13 @@ void URemeshMeshTool::Setup()
 		InitialMeshArea += OriginalMesh.GetTriArea(tid);
 	}
 
-	BasicProperties = NewObject<URemeshMeshToolProperties>(this, TEXT("Remesh Settings"));
+	BasicProperties = NewObject<URemeshMeshToolProperties>(this);
 	// arbitrary threshold of 5000 tris seems reasonable? 
 	BasicProperties->TargetTriangleCount = (OriginalMesh.TriangleCount() < 5000) ? 5000 : OriginalMesh.TriangleCount();
-
-	AdvancedProperties = NewObject<URemeshMeshAdvancedProperties>(this, TEXT("Advanced Settings"));
-	AdvancedProperties->TargetEdgeLength = CalculateTargetEdgeLength(BasicProperties->TargetTriangleCount);
+	BasicProperties->TargetEdgeLength = CalculateTargetEdgeLength(BasicProperties->TargetTriangleCount);
 
 	// initialize our properties
 	AddToolPropertySource(BasicProperties);
-	AddToolPropertySource(AdvancedProperties);
 
 	MeshStatisticsProperties = NewObject<UMeshStatisticsProperties>(this);
 	AddToolPropertySource(MeshStatisticsProperties);
@@ -177,13 +172,6 @@ void URemeshMeshTool::Render(IToolsContextRenderAPI* RenderAPI)
 	}
 }
 
-#if WITH_EDITOR
-void URemeshMeshTool::PostEditChangeProperty(FPropertyChangedEvent& PropertyChangedEvent)
-{
-	UProperty* PropertyThatChanged = PropertyChangedEvent.Property;
-	bResultValid = false;
-}
-#endif
 
 void URemeshMeshTool::OnPropertyModified(UObject* PropertySet, UProperty* Property)
 {
@@ -216,19 +204,19 @@ void URemeshMeshTool::UpdateResult()
 	}
 
 	FRemesher Remesher(TargetMesh);
-	Remesher.bEnableSplits = AdvancedProperties->bSplits;
-	Remesher.bEnableFlips = AdvancedProperties->bFlips;
-	Remesher.bEnableCollapses = AdvancedProperties->bCollapses;
+	Remesher.bEnableSplits = BasicProperties->bSplits;
+	Remesher.bEnableFlips = BasicProperties->bFlips;
+	Remesher.bEnableCollapses = BasicProperties->bCollapses;
 
 	// recalculate target edge length
-	if (AdvancedProperties->bUseTargetEdgeLength == false)
+	if (BasicProperties->bUseTargetEdgeLength == false)
 	{
-		AdvancedProperties->TargetEdgeLength = CalculateTargetEdgeLength(BasicProperties->TargetTriangleCount);
+		BasicProperties->TargetEdgeLength = CalculateTargetEdgeLength(BasicProperties->TargetTriangleCount);
 	}
 
-	Remesher.SetTargetEdgeLength(AdvancedProperties->TargetEdgeLength);
+	Remesher.SetTargetEdgeLength(BasicProperties->TargetEdgeLength);
 
-	Remesher.ProjectionMode = (AdvancedProperties->bReproject) ? 
+	Remesher.ProjectionMode = (BasicProperties->bReproject) ? 
 		FRemesher::ETargetProjectionMode::AfterRefinement : FRemesher::ETargetProjectionMode::NoProjection;
 
 	Remesher.bEnableSmoothing = (BasicProperties->SmoothingSpeed > 0);
@@ -239,7 +227,7 @@ void URemeshMeshTool::UpdateResult()
 		(FRemesher::ESmoothTypes)(int)BasicProperties->SmoothingType;
 	bool bIsUniformSmooth = (Remesher.SmoothType == FRemesher::ESmoothTypes::Uniform);
 
-	Remesher.bPreventNormalFlips = AdvancedProperties->bPreventNormalFlips;
+	Remesher.bPreventNormalFlips = BasicProperties->bPreventNormalFlips;
 
 	Remesher.DEBUG_CHECK_LEVEL = 0;
 
@@ -266,7 +254,7 @@ void URemeshMeshTool::UpdateResult()
 		if (bIsUniformSmooth == false)
 		{
 			bool bUseFlipsThisPass = (k % 2 == 0 && k < BasicProperties->RemeshIterations/2);
-			Remesher.bEnableFlips = bUseFlipsThisPass && AdvancedProperties->bFlips;
+			Remesher.bEnableFlips = bUseFlipsThisPass && BasicProperties->bFlips;
 		}
 
 		Remesher.BasicRemeshPass();
diff --git a/Engine/Plugins/Experimental/MeshModelingToolset/Source/MeshModelingTools/Public/RemeshMeshTool.h b/Engine/Plugins/Experimental/MeshModelingToolset/Source/MeshModelingTools/Public/RemeshMeshTool.h
--- a/Engine/Plugins/Experimental/MeshModelingToolset/Source/MeshModelingTools/Public/RemeshMeshTool.h
+++ b/Engine/Plugins/Experimental/MeshModelingToolset/Source/MeshModelingTools/Public/RemeshMeshTool.h
@@ -54,7 +54,7 @@ public:
 	URemeshMeshToolProperties();
 
 	/** Target triangle count */
-	UPROPERTY(EditAnywhere, Category = Options)
+	UPROPERTY(EditAnywhere, Category = Options, meta = (EditCondition = "bUseTargetEdgeLength == false"))
 	int TargetTriangleCount;
 
 
@@ -79,49 +79,39 @@ public:
 	UPROPERTY(EditAnywhere, Category = Options)
 	bool bPreserveSharpEdges;
 
-};
-
-/**
- * Advanced properties of the remesh operation
- */
-UCLASS()
-class MESHMODELINGTOOLS_API URemeshMeshAdvancedProperties : public UInteractiveToolPropertySet
-{
-	GENERATED_BODY()
-
-public:
-	URemeshMeshAdvancedProperties();
 
 	/** If true, the target count is ignored and the target edge length is used directly */
-	UPROPERTY(EditAnywhere, Category = Advanced)
+	UPROPERTY(EditAnywhere, Category = Options, AdvancedDisplay)
 	bool bUseTargetEdgeLength;
 
 	/** Target edge length */
-	UPROPERTY(EditAnywhere, Category = Advanced, meta = (NoSpinbox = "true"))
+	UPROPERTY(EditAnywhere, Category = Options, AdvancedDisplay, meta = (NoSpinbox = "true", EditCondition = "bUseTargetEdgeLength == true"))
 	float TargetEdgeLength;
 
+
 	/** Enable edge flips */
-	UPROPERTY(EditAnywhere, Category = Advanced)
+	UPROPERTY(EditAnywhere, Category = Options, AdvancedDisplay)
 	bool bFlips;
 
 	/** Enable edge splits */
-	UPROPERTY(EditAnywhere, Category = Advanced)
+	UPROPERTY(EditAnywhere, Category = Options, AdvancedDisplay)
 	bool bSplits;
 
 	/** Enable edge collapses */
-	UPROPERTY(EditAnywhere, Category = Advanced)
+	UPROPERTY(EditAnywhere, Category = Options, AdvancedDisplay)
 	bool bCollapses;
 
 	/** Enable projection back to input mesh */
-	UPROPERTY(EditAnywhere, Category = Advanced)
+	UPROPERTY(EditAnywhere, Category = Options, AdvancedDisplay)
 	bool bReproject;
 
 	/** Prevent normal flips */
-	UPROPERTY(EditAnywhere, Category = Advanced)
+	UPROPERTY(EditAnywhere, Category = Options, AdvancedDisplay)
 	bool bPreventNormalFlips;
 
 };
 
+
 /**
  * Simple Mesh Remeshing Tool
  */
@@ -142,11 +132,6 @@ public:
 	virtual bool HasAccept() const override;
 	virtual bool CanAccept() const override;
 
-
-#if WITH_EDITOR
-	virtual void PostEditChangeProperty(FPropertyChangedEvent &PropertyChangedEvent) override;
-#endif
-
 	virtual void OnPropertyModified(UObject* PropertySet, UProperty* Property) override;
 
 public:
@@ -154,9 +139,6 @@ public:
 	UPROPERTY()
 	URemeshMeshToolProperties* BasicProperties;
 
-	UPROPERTY()
-	URemeshMeshAdvancedProperties* AdvancedProperties;
-
 	UPROPERTY()
 	UMeshStatisticsProperties* MeshStatisticsProperties;
 
diff --git a/Engine/Plugins/Experimental/MeshModelingToolset/Source/MeshModelingToolsEditorOnly/Private/SimplifyMeshTool.cpp b/Engine/Plugins/Experimental/MeshModelingToolset/Source/MeshModelingToolsEditorOnly/Private/SimplifyMeshTool.cpp
--- a/Engine/Plugins/Experimental/MeshModelingToolset/Source/MeshModelingToolsEditorOnly/Private/SimplifyMeshTool.cpp
+++ b/Engine/Plugins/Experimental/MeshModelingToolset/Source/MeshModelingToolsEditorOnly/Private/SimplifyMeshTool.cpp
@@ -55,7 +55,7 @@ UInteractiveTool* USimplifyMeshToolBuilder::BuildTool(const FToolBuilderState& S
  * Tool
  */
 
-USimplifyMeshTool::USimplifyMeshTool()
+USimplifyMeshToolProperties::USimplifyMeshToolProperties()
 {
 	SimplifierType = ESimplifyType::QEM;
 	TargetMode = ESimplifyTargetType::Percentage;
@@ -67,6 +67,8 @@ USimplifyMeshTool::USimplifyMeshTool()
 	bDiscardAttributes = false;
 }
 
+
+
 void USimplifyMeshTool::Setup()
 {
 	UInteractiveTool::Setup();
@@ -93,7 +95,8 @@ void USimplifyMeshTool::Setup()
 	ComponentTarget->SetOwnerVisibility(false);
 
 	// initialize our properties
-	ToolPropertyObjects.Add(this);
+	SimplifyProperties = NewObject<USimplifyMeshToolProperties>(this);
+	AddToolPropertySource(SimplifyProperties);
 
 	MeshStatisticsProperties = NewObject<UMeshStatisticsProperties>(this);
 	AddToolPropertySource(MeshStatisticsProperties);
@@ -156,12 +159,12 @@ void USimplifyMeshTool::Render(IToolsContextRenderAPI* RenderAPI)
 }
 
 
-void USimplifyMeshTool::PostEditChangeProperty(FPropertyChangedEvent& PropertyChangedEvent)
+void USimplifyMeshTool::OnPropertyModified(UObject* PropertySet, UProperty* Property)
 {
-	UProperty* PropertyThatChanged = PropertyChangedEvent.Property;
 	bResultValid = false;
 }
 
+
 template <typename SimplificationType>
 void ComputeResult(FDynamicMesh3* TargetMesh, const bool bReproject, int OriginalTriCount, FDynamicMesh3& OriginalMesh, FDynamicMeshAABBTree3& OriginalMeshSpatial, const ESimplifyTargetType TargetMode, const float TargetPercentage, const int TargetCount, const float TargetEdgeLength)
 {
@@ -211,22 +214,23 @@ void USimplifyMeshTool::UpdateResult()
 	TargetMesh->Copy(OriginalMesh);
 	int OriginalTriCount = OriginalMesh.TriangleCount();
 
-	if (bDiscardAttributes)
+	if (SimplifyProperties->bDiscardAttributes)
 	{
 		TargetMesh->DiscardAttributes();
 	}
-	if (SimplifierType == ESimplifyType::QEM)
+	if (SimplifyProperties->SimplifierType == ESimplifyType::QEM)
 	{
-		ComputeResult<FQEMSimplification>(TargetMesh, bReproject, OriginalTriCount, OriginalMesh, OriginalMeshSpatial, TargetMode, TargetPercentage, TargetCount, TargetEdgeLength);
+		ComputeResult<FQEMSimplification>(TargetMesh, SimplifyProperties->bReproject, OriginalTriCount, OriginalMesh, OriginalMeshSpatial, 
+			SimplifyProperties->TargetMode, SimplifyProperties->TargetPercentage, SimplifyProperties->TargetCount, SimplifyProperties->TargetEdgeLength);
 	}
-	else if (SimplifierType == ESimplifyType::Attribute)
+	else if (SimplifyProperties->SimplifierType == ESimplifyType::Attribute)
 	{
-		ComputeResult<FAttrMeshSimplification>(TargetMesh, bReproject, OriginalTriCount, OriginalMesh, OriginalMeshSpatial, TargetMode, TargetPercentage, TargetCount, TargetEdgeLength);
+		ComputeResult<FAttrMeshSimplification>(TargetMesh, SimplifyProperties->bReproject, OriginalTriCount, OriginalMesh, OriginalMeshSpatial,
+			SimplifyProperties->TargetMode, SimplifyProperties->TargetPercentage, SimplifyProperties->TargetCount, SimplifyProperties->TargetEdgeLength);
 	}
 	else //if (SimplifierType == ESimplifyType::UE4Standard)
 	{
 
-
 		const FMeshDescription* SrcMeshDescription = ComponentTarget->GetMesh();
 		FMeshDescription DstMeshDescription(*SrcMeshDescription);
 
@@ -235,14 +239,14 @@ void USimplifyMeshTool::UpdateResult()
 		FMeshDescriptionOperations::FindOverlappingCorners(OverlappingCorners, *SrcMeshDescription, 1.e-5);
 
 		FMeshReductionSettings ReductionSettings;
-		if (TargetMode == ESimplifyTargetType::Percentage)
+		if (SimplifyProperties->TargetMode == ESimplifyTargetType::Percentage)
 		{
-			ReductionSettings.PercentTriangles = FMath::Max(TargetPercentage / 100., .001);  // Only support triangle percentage and count, but not edge length
+			ReductionSettings.PercentTriangles = FMath::Max(SimplifyProperties->TargetPercentage / 100., .001);  // Only support triangle percentage and count, but not edge length
 		}
-		else if (TargetMode == ESimplifyTargetType::TriangleCount)
+		else if (SimplifyProperties->TargetMode == ESimplifyTargetType::TriangleCount)
 		{
 			int32 NumTris = SrcMeshDescription->Polygons().Num();
-			ReductionSettings.PercentTriangles = (float)TargetCount / (float)NumTris;
+			ReductionSettings.PercentTriangles = (float)SimplifyProperties->TargetCount / (float)NumTris;
 		}
 
 		float Error;
diff --git a/Engine/Plugins/Experimental/MeshModelingToolset/Source/MeshModelingToolsEditorOnly/Public/SimplifyMeshTool.h b/Engine/Plugins/Experimental/MeshModelingToolset/Source/MeshModelingToolsEditorOnly/Public/SimplifyMeshTool.h
--- a/Engine/Plugins/Experimental/MeshModelingToolset/Source/MeshModelingToolsEditorOnly/Public/SimplifyMeshTool.h
+++ b/Engine/Plugins/Experimental/MeshModelingToolset/Source/MeshModelingToolsEditorOnly/Public/SimplifyMeshTool.h
@@ -65,35 +65,15 @@ public:
 
 
 
-
-
 /**
- * Simple Mesh Simplifying Tool
+ * Standard properties of the Simplify operation
  */
 UCLASS()
-class MESHMODELINGTOOLSEDITORONLY_API USimplifyMeshTool : public USingleSelectionTool
+class MESHMODELINGTOOLSEDITORONLY_API USimplifyMeshToolProperties : public UInteractiveToolPropertySet
 {
 	GENERATED_BODY()
-
 public:
-	USimplifyMeshTool();
-
-	virtual void Setup() override;
-	virtual void Shutdown(EToolShutdownType ShutdownType) override;
-
-	virtual void Render(IToolsContextRenderAPI* RenderAPI) override;
-
-	virtual bool HasCancel() const override { return true; }
-	virtual bool HasAccept() const override;
-	virtual bool CanAccept() const override;
-
-#if WITH_EDITOR
-	virtual void PostEditChangeProperty(FPropertyChangedEvent& PropertyChangedEvent) override;
-#endif
-
-protected:
-	// need to update bResultValid if these are modified, so we don't publicly expose them. 
-	// @todo setters/getters for these
+	USimplifyMeshToolProperties();
 
 	/** Simplification Target Type  */
 	UPROPERTY(EditAnywhere, Category = Options)
@@ -104,15 +84,15 @@ protected:
 	ESimplifyType SimplifierType;
 
 	/** Target percentage of original triangle count */
-	UPROPERTY(EditAnywhere, Category = Options, meta = (UIMin = "0", UIMax = "100"))
+	UPROPERTY(EditAnywhere, Category = Options, meta = (UIMin = "0", UIMax = "100", EditCondition = "TargetMode == ESimplifyTargetType::Percentage"))
 	int TargetPercentage;
 
 	/** Target edge length */
-	UPROPERTY(EditAnywhere, Category = Options, meta = (UIMin = "3.0", UIMax = "10.0", ClampMin = "0.001", ClampMax = "1000.0"))
+	UPROPERTY(EditAnywhere, Category = Options, meta = (UIMin = "3.0", UIMax = "10.0", ClampMin = "0.001", ClampMax = "1000.0", EditCondition = "TargetMode == ESimplifyTargetType::EdgeLength"))
 	float TargetEdgeLength;
 
 	/** Target triangle count */
-	UPROPERTY(EditAnywhere, Category = Options, meta = (UIMin = "4", UIMax = "10000", ClampMin = "1", ClampMax = "9999999999"))
+	UPROPERTY(EditAnywhere, Category = Options, meta = (UIMin = "4", UIMax = "10000", ClampMin = "1", ClampMax = "9999999999", EditCondition = "TargetMode == ESimplifyTargetType::TriangleCount"))
 	int TargetCount;
 
 	/** If true, UVs and Normals are discarded  */
@@ -120,13 +100,45 @@ protected:
 	bool bDiscardAttributes;
 
 	/** Enable projection back to input mesh */
-	UPROPERTY(EditAnywhere, Category = Advanced)
+	UPROPERTY(EditAnywhere, Category = Options, AdvancedDisplay)
 	bool bReproject;
 
 	/** Prevent normal flips */
-	UPROPERTY(EditAnywhere, Category = Advanced)
+	UPROPERTY(EditAnywhere, Category = Options, AdvancedDisplay)
 	bool bPreventNormalFlips;
 
+};
+
+
+
+
+/**
+ * Simple Mesh Simplifying Tool
+ */
+UCLASS()
+class MESHMODELINGTOOLSEDITORONLY_API USimplifyMeshTool : public USingleSelectionTool
+{
+	GENERATED_BODY()
+
+public:
+	virtual void Setup() override;
+	virtual void Shutdown(EToolShutdownType ShutdownType) override;
+
+	virtual void Render(IToolsContextRenderAPI* RenderAPI) override;
+
+	virtual bool HasCancel() const override { return true; }
+	virtual bool HasAccept() const override;
+	virtual bool CanAccept() const override;
+
+	virtual void OnPropertyModified(UObject* PropertySet, UProperty* Property) override;
+
+protected:
+	// need to update bResultValid if these are modified, so we don't publicly expose them. 
+	// @todo setters/getters for these
+
+	UPROPERTY()
+	USimplifyMeshToolProperties* SimplifyProperties;
+
 	UPROPERTY()
 	UMeshStatisticsProperties* MeshStatisticsProperties;
 
