Edit Mesh Polygons improvements + add FIndexFlagSet to make performance more tolerable for large meshes

#rnx

diff --git a/Engine/Plugins/Experimental/GeometryProcessing/Source/DynamicMesh/Private/MeshRegionBoundaryLoops.cpp b/Engine/Plugins/Experimental/GeometryProcessing/Source/DynamicMesh/Private/MeshRegionBoundaryLoops.cpp
--- a/Engine/Plugins/Experimental/GeometryProcessing/Source/DynamicMesh/Private/MeshRegionBoundaryLoops.cpp
+++ b/Engine/Plugins/Experimental/GeometryProcessing/Source/DynamicMesh/Private/MeshRegionBoundaryLoops.cpp
@@ -4,7 +4,7 @@
 #include "MeshRegionBoundaryLoops.h"
 #include "MeshBoundaryLoops.h"   // has a set of internal static functions we re-use
 #include "VectorUtil.h"
-
+#include "Util/SparseIndexCollectionTypes.h"
 
 
 FMeshRegionBoundaryLoops::FMeshRegionBoundaryLoops(const FDynamicMesh3* MeshIn, const TArray<int>& RegionTris, bool bAutoCompute)
@@ -12,16 +12,16 @@ FMeshRegionBoundaryLoops::FMeshRegionBoundaryLoops(const FDynamicMesh3* MeshIn,
 	this->Mesh = MeshIn;
 
 	// make flag set for included triangles
-	triangles.Init(false, Mesh->MaxTriangleID());
+	Triangles.InitAuto(Mesh->MaxTriangleID(), RegionTris.Num());
 	for (int i = 0; i < RegionTris.Num(); ++i)
 	{
-		triangles[RegionTris[i]] = true;
+		Triangles.Add(RegionTris[i]);
 	}
 
 	// make flag set for included edges
 	// NOTE: this currently processes non-boundary-edges twice. Could
 	// avoid w/ another IndexFlagSet, but the check is inexpensive...
-	edges.Init(false, Mesh->MaxEdgeID());
+	Edges.InitAuto(Mesh->MaxEdgeID(), RegionTris.Num());
 	for (int i = 0; i < RegionTris.Num(); ++i)
 	{
 		int tid = RegionTris[i];
@@ -29,13 +29,13 @@ FMeshRegionBoundaryLoops::FMeshRegionBoundaryLoops(const FDynamicMesh3* MeshIn,
 		for (int j = 0; j < 3; ++j)
 		{
 			int eid = te[j];
-			if (!ContainsElement(edges, eid))
+			if (!Edges.Contains(eid))
 			{
 				FIndex2i et = Mesh->GetEdgeT(eid);
-				if (et.B == IndexConstants::InvalidID || triangles[et.A] != triangles[et.B])
+				if (et.B == IndexConstants::InvalidID || Triangles[et.A] != Triangles[et.B])
 				{
 					edges_roi.Add(eid);
-					edges[eid] = true;
+					Edges.Add(eid);
 				}
 			}
 		}
@@ -71,8 +71,8 @@ bool FMeshRegionBoundaryLoops::Compute()
 	Loops.SetNum(0);
 
 	// Temporary memory used to indicate when we have "used" an edge.
-	TArray<bool> used_edge;
-	used_edge.Init(false, Mesh->MaxEdgeID());
+	FIndexFlagSet used_edge;
+	used_edge.InitAuto(Mesh->MaxEdgeID(), edges_roi.Num());
 
 	// current loop is stored here, cleared after each loop extracted
 	TArray<int> loop_edges;
@@ -98,7 +98,7 @@ bool FMeshRegionBoundaryLoops::Compute()
 
 		// ok this is start of a boundary chain
 		int eStart = eid;
-		used_edge[eStart] = true;
+		used_edge.Add(eStart);
 		loop_edges.Add(eStart);
 
 		int eCur = eid;
@@ -178,7 +178,7 @@ bool FMeshRegionBoundaryLoops::Compute()
 				check(used_edge[eNext] == false);
 				loop_edges.Add(eNext);
 				eCur = eNext;
-				used_edge[eCur] = true;
+				used_edge.Add(eCur);
 			}
 		}
 
@@ -219,7 +219,7 @@ bool FMeshRegionBoundaryLoops::Compute()
 // tid_in and tid_out are triangles 'in' and 'out' of set, respectively
 bool FMeshRegionBoundaryLoops::IsEdgeOnBoundary(int eid, int& tid_in, int& tid_out) const
 {
-	if (ContainsElement(edges, eid) == false)
+	if (Edges.Contains(eid) == false)
 	{
 		return false;
 	}
@@ -233,8 +233,8 @@ bool FMeshRegionBoundaryLoops::IsEdgeOnBoundary(int eid, int& tid_in, int& tid_o
 		return true;
 	}
 
-	bool in0 = triangles[et.A];
-	bool in1 = triangles[et.B];
+	bool in0 = Triangles[et.A];
+	bool in1 = Triangles[et.B];
 	if (in0 != in1)
 	{
 		tid_in = (in0) ? et.A : et.B;
@@ -315,7 +315,7 @@ FVector3d FMeshRegionBoundaryLoops::GetVertexNormal(int vid)
 // If the loops are all sane, then we will get the smallest loops by "turning left" at bowtie_v.
 // So, we compute the tangent plane at bowtie_v, and then the signed angle for each
 // viable edge : this plane. 
-int FMeshRegionBoundaryLoops::FindLeftTurnEdge(int incoming_e, int bowtie_v, TArray<int>& bdry_edges, int bdry_edges_count, TArray<bool>& used_edges)
+int FMeshRegionBoundaryLoops::FindLeftTurnEdge(int incoming_e, int bowtie_v, TArray<int>& bdry_edges, int bdry_edges_count, const FIndexFlagSet& used_edges)
 {
 	// compute normal and edge [a,bowtie]
 	FVector3d n = GetVertexNormal(bowtie_v);
diff --git a/Engine/Plugins/Experimental/GeometryProcessing/Source/DynamicMesh/Public/MeshRegionBoundaryLoops.h b/Engine/Plugins/Experimental/GeometryProcessing/Source/DynamicMesh/Public/MeshRegionBoundaryLoops.h
--- a/Engine/Plugins/Experimental/GeometryProcessing/Source/DynamicMesh/Public/MeshRegionBoundaryLoops.h
+++ b/Engine/Plugins/Experimental/GeometryProcessing/Source/DynamicMesh/Public/MeshRegionBoundaryLoops.h
@@ -6,7 +6,7 @@
 
 #include "DynamicMesh3.h"
 #include "EdgeLoop.h"
-
+#include "Util/SparseIndexCollectionTypes.h"
 
 /**
  * Extract FEdgeLoops on the boundary of a set of triangles of a mesh.
@@ -54,14 +54,11 @@ protected:
 
 	// TODO: C# code used IndexFlagSet here, which is sparse if region size << mesh size
 	// list of included triangles and edges
-	TArray<bool> triangles;
-	TArray<bool> edges;
+	FIndexFlagSet Triangles;
+	FIndexFlagSet Edges;
 	TArray<int> edges_roi;
 
-
-	static bool ContainsElement(const TArray<bool>& flagset, int index) { return flagset[index] == true; }
-
-	bool IsEdgeOnBoundary(int eid) const { return ContainsElement(edges, eid); }
+	bool IsEdgeOnBoundary(int eid) const { return Edges.Contains(eid); }
 
 	// returns true for both internal and mesh boundary edges
 	// tid_in and tid_out are triangles 'in' and 'out' of set, respectively
@@ -85,7 +82,7 @@ protected:
 	// ok, bdry_edges[0...bdry_edges_count] contains the boundary edges coming out of bowtie_v.
 	// We want to pick the best one to continue the loop that came : to bowtie_v on incoming_e.
 	// If the loops are all sane, then we will get the smallest loops by "turning left" at bowtie_v.
-	int FindLeftTurnEdge(int incoming_e, int bowtie_v, TArray<int>& bdry_edges, int bdry_edges_count, TArray<bool>& used_edges);
+	int FindLeftTurnEdge(int incoming_e, int bowtie_v, TArray<int>& bdry_edges, int bdry_edges_count, const FIndexFlagSet& used_edges);
 
 
 	// This is called when loopV contains one or more "bowtie" vertices.
diff --git a/Engine/Plugins/Experimental/GeometryProcessing/Source/GeometricObjects/Public/Util/SparseIndexCollectionTypes.h b/Engine/Plugins/Experimental/GeometryProcessing/Source/GeometricObjects/Public/Util/SparseIndexCollectionTypes.h
new file mode 100644
--- /dev/null
+++ b/Engine/Plugins/Experimental/GeometryProcessing/Source/GeometricObjects/Public/Util/SparseIndexCollectionTypes.h
@@ -0,0 +1,144 @@
+// Copyright 1998-2019 Epic Games, Inc. All Rights Reserved.
+
+// partial port of geometry3sharp Indexing.cs
+
+#pragma once
+
+#include <CoreMinimal.h>
+
+
+/**
+ *  This class provides optionally sparse or dense boolean flags for a set of integer indices
+ */
+class FIndexFlagSet
+{
+private:
+	TOptional<TArray<bool>> Dense; // TODO: consider switching to bit array; but test first as performance of bit array may by significantly worse
+	int DenseCount;      // only tracked for Dense
+	TOptional<TSet<int>> Sparse;
+
+public:
+	FIndexFlagSet(bool bSetSparse = true, int MaxIndex = -1)
+	{
+		InitManual(bSetSparse, MaxIndex);
+	}
+
+	FIndexFlagSet(int MaxIndex, int SubsetCountEst)
+	{
+		InitAuto(MaxIndex, SubsetCountEst);
+	}
+
+	void InitAuto(int MaxIndex, int SubsetCountEst)
+	{
+		bool bSmall = MaxIndex < 32000;
+		constexpr float PercentThresh = 0.05f;           
+
+		InitManual(!bSmall && ((float)SubsetCountEst / (float)MaxIndex < PercentThresh), MaxIndex);
+	}
+
+	void InitManual(bool bSetSparse, int MaxIndex = -1)
+	{
+		if (bSetSparse)
+		{
+			Sparse = TSet<int>();
+		}
+		else
+		{
+			Dense = TArray<bool>();
+			check(MaxIndex >= 0);
+			Dense->SetNumZeroed(FMath::Max(0, MaxIndex));
+		}
+		DenseCount = 0;
+	}
+
+	/**
+	 *  checks if value i is true
+	 */
+	FORCEINLINE bool Contains(int i) const
+	{
+		if (Dense.IsSet())
+		{
+			return Dense.GetValue()[i];
+		}
+		else
+		{
+			check(Sparse.IsSet());
+			return Sparse->Contains(i);
+		}
+	}
+
+	/**
+	 *  sets value i to true
+	 */
+	FORCEINLINE void Add(int i)
+	{
+		if (Dense.IsSet())
+		{
+			bool &Value = Dense.GetValue()[i];
+			if (!Value)
+			{
+				Value = true;
+				DenseCount++;
+			}
+		}
+		else
+		{
+			check(Sparse.IsSet());
+			Sparse->Add(i);
+		}
+	}
+
+	FORCEINLINE void Remove(int Index)
+	{
+		if (Dense.IsSet())
+		{
+			bool &Value = Dense.GetValue()[Index];
+			if (Value)
+			{
+				Value = false;
+				DenseCount--;
+			}
+		}
+		else
+		{
+			check(Sparse.IsSet());
+			Sparse->Remove(Index);
+		}
+	}
+
+	/**
+	 *  Returns number of true values in set
+	 */
+	FORCEINLINE int Count() const
+	{
+		if (Dense.IsSet())
+		{
+			return DenseCount;
+		}
+		else
+		{
+			return Sparse->Num();
+		}
+	}
+
+
+	inline const bool operator[](unsigned int Index) const
+	{
+		return Contains(Index);
+	}
+	
+
+	// TODO: iterator support
+
+};
+
+
+
+
+// TODO implement a similar class but as a template map
+//template<typename T>
+//class TOptionallySparseMap
+//{
+// etc
+//
+//};
diff --git a/Engine/Plugins/Experimental/MeshModelingToolset/Source/MeshModelingTools/Private/EditMeshPolygonsTool.cpp b/Engine/Plugins/Experimental/MeshModelingToolset/Source/MeshModelingTools/Private/EditMeshPolygonsTool.cpp
--- a/Engine/Plugins/Experimental/MeshModelingToolset/Source/MeshModelingTools/Private/EditMeshPolygonsTool.cpp
+++ b/Engine/Plugins/Experimental/MeshModelingToolset/Source/MeshModelingTools/Private/EditMeshPolygonsTool.cpp
@@ -52,7 +52,7 @@ UPolyEditTransformProperties::UPolyEditTransformProperties()
 	bSelectVertices = true;
 	bSelectFaces = true;
 	bSelectEdges = true;
-	bRecomputePolygonGroups = false;
+	PolygonMode = EPolygonGroupMode::KeepInputPolygons;
 	PolygonGroupingAngleThreshold = .5;
 }
 
@@ -618,6 +618,17 @@ void UEditMeshPolygonsTool::Setup()
 	OnDynamicMeshComponentChangedHandle = DynamicMeshComponent->OnMeshChanged.Add(
 		FSimpleMulticastDelegate::FDelegate::CreateUObject(this, &UEditMeshPolygonsTool::OnDynamicMeshComponentChanged));
 
+
+	// add properties
+	TransformProps = NewObject<UPolyEditTransformProperties>(this);
+	TransformProps->PolygonMode = DynamicMeshComponent->GetMesh()->HasTriangleGroups() && DynamicMeshComponent->GetMesh()->MaxGroupID() > 1 ? EPolygonGroupMode::KeepInputPolygons : EPolygonGroupMode::RecomputePolygonsByAngleThreshold;
+	if (TransformProps->PolygonMode == EPolygonGroupMode::RecomputePolygonsByAngleThreshold)
+	{
+		ComputePolygons(false);
+	}
+	BackupTriangleGroups();
+	AddToolPropertySource(TransformProps);
+
 	// initialize AABBTree
 	MeshSpatial.SetMesh(DynamicMeshComponent->GetMesh());
 	PrecomputeTopology();
@@ -639,12 +650,6 @@ void UEditMeshPolygonsTool::Setup()
 	QuickAxisTranslater.Initialize();
 	QuickAxisRotator.Initialize();
 
-	// add properties
-	TransformProps = NewObject<UPolyEditTransformProperties>(this);
-	TransformProps->bRecomputePolygonGroups = !DynamicMeshComponent->GetMesh()->HasTriangleGroups();
-	BackupTriangleGroups();
-	AddToolPropertySource(TransformProps);
-
 	// set up visualizers
 	PolyEdgesRenderer.LineColor = FLinearColor::Red;
 	PolyEdgesRenderer.LineThickness = 2.0;
@@ -1284,6 +1289,43 @@ void UEditMeshPolygonsTool::Tick(float DeltaTime)
 }
 
 
+void UEditMeshPolygonsTool::ComputePolygons(bool RecomputeTopology)
+{
+	switch (TransformProps->PolygonMode)
+	{
+	case EPolygonGroupMode::KeepInputPolygons:
+		RestoreTriangleGroups();
+		break;
+	case EPolygonGroupMode::RecomputePolygonsByAngleThreshold:
+	{
+		FDynamicMesh3* SearchMesh = DynamicMeshComponent->GetMesh();
+		FFindPolygonsAlgorithm Polygons = FFindPolygonsAlgorithm(SearchMesh);
+		double DotTolerance = 1.0 - FMathd::Cos(TransformProps->PolygonGroupingAngleThreshold * FMathd::DegToRad);
+		Polygons.FindPolygons(DotTolerance);
+		Polygons.FindPolygonEdges();
+	}
+	break;
+	case EPolygonGroupMode::PolygonsAreTriangles:
+	{
+		FDynamicMesh3* Mesh = DynamicMeshComponent->GetMesh();
+		int GID = 0;
+		for (int TID : Mesh->TriangleIndicesItr())
+		{
+			Mesh->SetTriangleGroup(TID, GID++);
+		}
+	}
+	break;
+	}
+
+	if (RecomputeTopology)
+	{
+		PrecomputeTopology();
+		TopoSelector.Invalidate(false, true);
+		HilightSelection.Clear();
+	}
+}
+
+
 
 void UEditMeshPolygonsTool::PrecomputeTopology()
 {
@@ -1380,27 +1422,13 @@ void UEditMeshPolygonsTool::Render(IToolsContextRenderAPI* RenderAPI)
 
 void UEditMeshPolygonsTool::OnPropertyModified(UObject* PropertySet, UProperty* Property)
 {
-	if (!Property) // happens on undo/redo of property modifications if object is set to be transactional
+	// if anything has changed the polygon settings, recompute polygons
+	if (Property && 
+			(Property->GetFName() == GET_MEMBER_NAME_CHECKED(UPolyEditTransformProperties, PolygonMode)
+		  || Property->GetFName() == GET_MEMBER_NAME_CHECKED(UPolyEditTransformProperties, PolygonGroupingAngleThreshold))
+		)
 	{
-		return;
-	}
-	if (Property->GetFName() == GET_MEMBER_NAME_CHECKED(UPolyEditTransformProperties, bRecomputePolygonGroups)
-	 || (Property->GetFName() == GET_MEMBER_NAME_CHECKED(UPolyEditTransformProperties, PolygonGroupingAngleThreshold)))
-	{
-		if (TransformProps->bRecomputePolygonGroups)
-		{
-			FDynamicMesh3* SearchMesh = DynamicMeshComponent->GetMesh();
-			FFindPolygonsAlgorithm Polygons = FFindPolygonsAlgorithm(SearchMesh);
-			double DotTolerance = 1.0 - FMathd::Cos(TransformProps->PolygonGroupingAngleThreshold * FMathd::DegToRad);
-			Polygons.FindPolygons(DotTolerance);
-			Polygons.FindPolygonEdges();
-		}
-		else
-		{
-			RestoreTriangleGroups();
-		}
-		PrecomputeTopology();
-		TopoSelector.Invalidate(false, true);
+		ComputePolygons(true);
 	}
 }
 
diff --git a/Engine/Plugins/Experimental/MeshModelingToolset/Source/MeshModelingTools/Private/UVProjectionTool.cpp b/Engine/Plugins/Experimental/MeshModelingToolset/Source/MeshModelingTools/Private/UVProjectionTool.cpp
--- a/Engine/Plugins/Experimental/MeshModelingToolset/Source/MeshModelingTools/Private/UVProjectionTool.cpp
+++ b/Engine/Plugins/Experimental/MeshModelingToolset/Source/MeshModelingTools/Private/UVProjectionTool.cpp
@@ -285,7 +285,8 @@ void UUVProjectionTool::PostEditChangeProperty(FPropertyChangedEvent& PropertyCh
 
 void UUVProjectionTool::OnPropertyModified(UObject* PropertySet, UProperty* Property)
 {
-	if (Property->GetFName() == GET_MEMBER_NAME_CHECKED(UUVProjectionToolProperties, CheckerDensity))
+	// if we don't know what changed, or we know checker density changed, update checker material
+	if (!Property || Property->GetFName() == GET_MEMBER_NAME_CHECKED(UUVProjectionToolProperties, CheckerDensity))
 	{
 		CheckerMaterial->SetScalarParameterValue("Density", BasicProperties->CheckerDensity);
 	}
diff --git a/Engine/Plugins/Experimental/MeshModelingToolset/Source/MeshModelingTools/Public/EditMeshPolygonsTool.h b/Engine/Plugins/Experimental/MeshModelingToolset/Source/MeshModelingTools/Public/EditMeshPolygonsTool.h
--- a/Engine/Plugins/Experimental/MeshModelingToolset/Source/MeshModelingTools/Public/EditMeshPolygonsTool.h
+++ b/Engine/Plugins/Experimental/MeshModelingToolset/Source/MeshModelingTools/Public/EditMeshPolygonsTool.h
@@ -82,6 +82,15 @@ enum class EQuickTransformerMode : uint8
 };
 
 
+UENUM()
+enum class EPolygonGroupMode : uint8
+{
+	KeepInputPolygons						UMETA(DisplayName = "Edit Input Polygons"),
+	RecomputePolygonsByAngleThreshold		UMETA(DisplayName = "Recompute Polygons Based on Angle Threshold"),
+	PolygonsAreTriangles					UMETA(DisplayName = "Edit Triangles Directly")
+};
+
+
 
 UCLASS()
 class MESHMODELINGTOOLS_API UPolyEditTransformProperties : public UInteractiveToolPropertySet
@@ -112,9 +121,9 @@ public:
 	bool bSelectVertices;
 
 	UPROPERTY(EditAnywhere, Category = Options)
-	bool bRecomputePolygonGroups;
+	EPolygonGroupMode PolygonMode;
 
-	UPROPERTY(EditAnywhere, Category = Options, meta = (EditCondition="bRecomputePolygonGroups"))
+	UPROPERTY(EditAnywhere, Category = Options, meta = (EditCondition="PolygonMode == EPolygonGroupMode::RecomputePolygonsByAngleThreshold"))
 	float PolygonGroupingAngleThreshold;
 
 	//Laplacian Deformation Options
@@ -463,6 +472,7 @@ protected:
 
 	FGroupTopology Topology;
 	void PrecomputeTopology();
+	void ComputePolygons(bool RecomputeTopology = true);
 
 	FGroupTopologySelector TopoSelector;
 	
