Fix compile error of AddPrimitiveTool.h
#rnx

diff --git a/Engine/Plugins/Experimental/MeshModelingToolset/Source/MeshModelingTools/Public/AddPrimitiveTool.h b/Engine/Plugins/Experimental/MeshModelingToolset/Source/MeshModelingTools/Public/AddPrimitiveTool.h
--- a/Engine/Plugins/Experimental/MeshModelingToolset/Source/MeshModelingTools/Public/AddPrimitiveTool.h
+++ b/Engine/Plugins/Experimental/MeshModelingToolset/Source/MeshModelingTools/Public/AddPrimitiveTool.h
@@ -69,7 +69,12 @@ class MESHMODELINGTOOLS_API UProceduralShapeToolProperties : public UInteractive
 
 public:
 	UProceduralShapeToolProperties();
+
+	// UObject interface
+#if WITH_EDITOR
 	virtual bool CanEditChange( const UProperty* InProperty ) const override;
+#endif // WITH_EDITOR	
+	// End of UObject interface
 
 	/** Type of shape to generate */
 	UPROPERTY(EditAnywhere, Category = ShapeSettings)
