add wireframe toggles, minor associated cleanup

#rnx

diff --git a/Engine/Plugins/Experimental/MeshModelingToolset/Source/MeshModelingTools/Private/AddPrimitiveTool.cpp b/Engine/Plugins/Experimental/MeshModelingToolset/Source/MeshModelingTools/Private/AddPrimitiveTool.cpp
--- a/Engine/Plugins/Experimental/MeshModelingToolset/Source/MeshModelingTools/Private/AddPrimitiveTool.cpp
+++ b/Engine/Plugins/Experimental/MeshModelingToolset/Source/MeshModelingTools/Private/AddPrimitiveTool.cpp
@@ -68,9 +68,12 @@ namespace
 	  { TEXT("Slices"),        EMakeMeshShapeType::Cylinder | EMakeMeshShapeType::Cone | EMakeMeshShapeType::Sphere },
 	  { TEXT("Subdivisions"),  EMakeMeshShapeType::Box | EMakeMeshShapeType::Plane },
 	  { TEXT("Material"),      EMakeMeshShapeType::All },
+	  { TEXT("bWireframe"),	   EMakeMeshShapeType::All }
 	};
 };
 
+// UObject interface
+#if WITH_EDITOR
 bool
 UProceduralShapeToolProperties::CanEditChange( const UProperty* InProperty) const
 {
@@ -82,6 +85,9 @@ UProceduralShapeToolProperties::CanEditChange( const UProperty* InProperty) cons
 	checkNoEntry();
 	return false;
 }
+#endif // WITH_EDITOR	
+// End of UObject interface
+
 
 void UAddPrimitiveTool::SetWorld(UWorld* World)
 {
@@ -137,6 +143,7 @@ void UAddPrimitiveTool::Render(IToolsContextRenderAPI* RenderAPI)
 
 void UAddPrimitiveTool::OnPropertyModified(UObject* PropertySet, UProperty* Property)
 {
+	PreviewMesh->EnableWireframe(ShapeSettings->bWireframe);
 	PreviewMesh->SetMaterial(ShapeSettings->Material);
 	UpdatePreviewMesh();
 }
diff --git a/Engine/Plugins/Experimental/MeshModelingToolset/Source/MeshModelingTools/Private/EditMeshPolygonsTool.cpp b/Engine/Plugins/Experimental/MeshModelingToolset/Source/MeshModelingTools/Private/EditMeshPolygonsTool.cpp
--- a/Engine/Plugins/Experimental/MeshModelingToolset/Source/MeshModelingTools/Private/EditMeshPolygonsTool.cpp
+++ b/Engine/Plugins/Experimental/MeshModelingToolset/Source/MeshModelingTools/Private/EditMeshPolygonsTool.cpp
@@ -1350,6 +1350,7 @@ void UEditMeshPolygonsTool::Render(IToolsContextRenderAPI* RenderAPI)
 	GetToolManager()->GetContextQueriesAPI()->GetCurrentViewState(CameraState);
 	GetActiveQuickTransformer()->UpdateCameraState(CameraState);
 
+	DynamicMeshComponent->bExplicitShowWireframe = TransformProps->bWireframe;
 	FDynamicMesh3* TargetMesh = DynamicMeshComponent->GetMesh();
 
 	PolyEdgesRenderer.BeginFrame(RenderAPI, CameraState);
diff --git a/Engine/Plugins/Experimental/MeshModelingToolset/Source/MeshModelingTools/Private/MeshInspectorTool.cpp b/Engine/Plugins/Experimental/MeshModelingToolset/Source/MeshModelingTools/Private/MeshInspectorTool.cpp
--- a/Engine/Plugins/Experimental/MeshModelingToolset/Source/MeshModelingTools/Private/MeshInspectorTool.cpp
+++ b/Engine/Plugins/Experimental/MeshModelingToolset/Source/MeshModelingTools/Private/MeshInspectorTool.cpp
@@ -149,8 +149,6 @@ void UMeshInspectorTool::Render(IToolsContextRenderAPI* RenderAPI)
 	FPrimitiveDrawInterface* PDI = RenderAPI->GetPrimitiveDrawInterface();
 	FTransform Transform = ComponentTarget->GetWorldTransform();
 
-	FColor WireframeColor(40, 40, 40);
-	float WireframeThickness = LineWidthMultiplier * 0.5;
 	FColor BoundaryEdgeColor(240, 15, 15);
 	float BoundaryEdgeThickness = LineWidthMultiplier * 4.0;
 	FColor UVSeamColor(15, 240, 15);
diff --git a/Engine/Plugins/Experimental/MeshModelingToolset/Source/MeshModelingTools/Public/AddPrimitiveTool.h b/Engine/Plugins/Experimental/MeshModelingToolset/Source/MeshModelingTools/Public/AddPrimitiveTool.h
--- a/Engine/Plugins/Experimental/MeshModelingToolset/Source/MeshModelingTools/Public/AddPrimitiveTool.h
+++ b/Engine/Plugins/Experimental/MeshModelingToolset/Source/MeshModelingTools/Public/AddPrimitiveTool.h
@@ -120,6 +120,10 @@ public:
 	/** Material */
 	UPROPERTY(EditAnywhere, Category = ShapeSettings, meta = (DisplayName = "Material"))
 	UMaterialInterface* Material;
+
+	/** Wireframe toggle */
+	UPROPERTY(EditAnywhere, Category = ShapeSettings, meta = (DisplayName = "Show Wireframe"))
+	bool bWireframe;
 };
 
 
diff --git a/Engine/Plugins/Experimental/MeshModelingToolset/Source/MeshModelingTools/Public/EditMeshPolygonsTool.h b/Engine/Plugins/Experimental/MeshModelingToolset/Source/MeshModelingTools/Public/EditMeshPolygonsTool.h
--- a/Engine/Plugins/Experimental/MeshModelingToolset/Source/MeshModelingTools/Public/EditMeshPolygonsTool.h
+++ b/Engine/Plugins/Experimental/MeshModelingToolset/Source/MeshModelingTools/Public/EditMeshPolygonsTool.h
@@ -120,6 +120,9 @@ public:
 	UPROPERTY(EditAnywhere, Category = Options)
 	bool bSelectVertices;
 
+	UPROPERTY(EditAnywhere, Category = Options, meta = (DisplayName = "Show Triangle Mesh Wireframe"))
+	bool bWireframe;
+
 	UPROPERTY(EditAnywhere, Category = Options)
 	EPolygonGroupMode PolygonMode;
 
