tweaks to edit mesh polygon behavior -- increase default weight to a more useful value (100), suppress interactive updates on polygon angle threshold, and always compute polygons wrt the original mesh

#rnx

diff --git a/Engine/Plugins/Experimental/MeshModelingToolset/Source/MeshModelingTools/Private/EditMeshPolygonsTool.cpp b/Engine/Plugins/Experimental/MeshModelingToolset/Source/MeshModelingTools/Private/EditMeshPolygonsTool.cpp
--- a/Engine/Plugins/Experimental/MeshModelingToolset/Source/MeshModelingTools/Private/EditMeshPolygonsTool.cpp
+++ b/Engine/Plugins/Experimental/MeshModelingToolset/Source/MeshModelingTools/Private/EditMeshPolygonsTool.cpp
@@ -56,6 +56,19 @@ UPolyEditTransformProperties::UPolyEditTransformProperties()
 	PolygonGroupingAngleThreshold = .5;
 }
 
+#ifdef WITH_EDITOR
+void UPolyEditTransformProperties::PostEditChangeProperty(FPropertyChangedEvent& PropertyChangedEvent)
+{
+	// skip interactive updates for PolygonGroupingAngleThreshold
+	// TODO: thread the polygon group compute and remove this update skip
+	if (PropertyChangedEvent.GetPropertyName() == GET_MEMBER_NAME_CHECKED(UPolyEditTransformProperties, PolygonGroupingAngleThreshold) && PropertyChangedEvent.ChangeType == EPropertyChangeType::Interactive)
+	{
+		return;
+	}
+	Super::PostEditChangeProperty(PropertyChangedEvent);
+}
+#endif
+
 /*
  * Asynchronous Task
  */
@@ -615,6 +628,7 @@ void UEditMeshPolygonsTool::Setup()
 	// dynamic mesh configuration settings
 	DynamicMeshComponent->TangentsType = EDynamicMeshTangentCalcType::AutoCalculated;
 	DynamicMeshComponent->InitializeMesh(ComponentTarget->GetMesh());
+	InitialMesh = MakeUnique<FDynamicMesh3>(*DynamicMeshComponent->GetMesh());
 	OnDynamicMeshComponentChangedHandle = DynamicMeshComponent->OnMeshChanged.Add(
 		FSimpleMulticastDelegate::FDelegate::CreateUObject(this, &UEditMeshPolygonsTool::OnDynamicMeshComponentChanged));
 
@@ -1294,15 +1308,15 @@ void UEditMeshPolygonsTool::ComputePolygons(bool RecomputeTopology)
 	switch (TransformProps->PolygonMode)
 	{
 	case EPolygonGroupMode::KeepInputPolygons:
-		RestoreTriangleGroups();
+		SetTriangleGroups(InitialTriangleGroups);
 		break;
 	case EPolygonGroupMode::RecomputePolygonsByAngleThreshold:
 	{
-		FDynamicMesh3* SearchMesh = DynamicMeshComponent->GetMesh();
-		FFindPolygonsAlgorithm Polygons = FFindPolygonsAlgorithm(SearchMesh);
+		FFindPolygonsAlgorithm Polygons = FFindPolygonsAlgorithm(InitialMesh.Get());
 		double DotTolerance = 1.0 - FMathd::Cos(TransformProps->PolygonGroupingAngleThreshold * FMathd::DegToRad);
 		Polygons.FindPolygons(DotTolerance);
 		Polygons.FindPolygonEdges();
+		SetTriangleGroups(*InitialMesh->GetTriangleGroupsBuffer());
 	}
 	break;
 	case EPolygonGroupMode::PolygonsAreTriangles:
@@ -1421,6 +1435,7 @@ void UEditMeshPolygonsTool::Render(IToolsContextRenderAPI* RenderAPI)
 }
 
 
+
 void UEditMeshPolygonsTool::OnPropertyModified(UObject* PropertySet, UProperty* Property)
 {
 	// if anything has changed the polygon settings, recompute polygons
@@ -1446,14 +1461,14 @@ void UEditMeshPolygonsTool::BackupTriangleGroups()
 	}
 }
 
-void UEditMeshPolygonsTool::RestoreTriangleGroups()
+void UEditMeshPolygonsTool::SetTriangleGroups(const TDynamicVector<int>& Groups)
 {
 	FDynamicMesh3* Mesh = DynamicMeshComponent->GetMesh();
-	for (int TID = 0, MaxID = InitialTriangleGroups.Num(); TID < MaxID; TID++)
+	for (int TID = 0, MaxID = Groups.Num(); TID < MaxID; TID++)
 	{
 		if (Mesh->IsTriangle(TID))
 		{
-			Mesh->SetTriangleGroup(TID, InitialTriangleGroups[TID]);
+			Mesh->SetTriangleGroup(TID, Groups[TID]);
 		}
 	}
 }
diff --git a/Engine/Plugins/Experimental/MeshModelingToolset/Source/MeshModelingTools/Public/EditMeshPolygonsTool.h b/Engine/Plugins/Experimental/MeshModelingToolset/Source/MeshModelingTools/Public/EditMeshPolygonsTool.h
--- a/Engine/Plugins/Experimental/MeshModelingToolset/Source/MeshModelingTools/Public/EditMeshPolygonsTool.h
+++ b/Engine/Plugins/Experimental/MeshModelingToolset/Source/MeshModelingTools/Public/EditMeshPolygonsTool.h
@@ -100,6 +100,10 @@ class MESHMODELINGTOOLS_API UPolyEditTransformProperties : public UInteractiveTo
 public:
 	UPolyEditTransformProperties();
 
+#ifdef WITH_EDITOR
+	virtual void PostEditChangeProperty(FPropertyChangedEvent& PropertyChangedEvent) override;
+#endif
+
 	//Options
 
 	UPROPERTY(EditAnywhere, Category = Options, meta = (DisplayName = "Deformation Type", ToolTip = "Select the type of deformation you wish to employ on a polygroup."))
@@ -134,8 +138,8 @@ public:
 	UPROPERTY(EditAnywhere, Category = LaplacianOptions, meta = (EditCondition="DeformationStrategy == EGroupTopologyDeformationStrategy::Laplacian", DisplayName = "Laplacian Weight Scheme", ToolTip = "Laplacian weight schemes determine how the deformer will interpret the curvature at a given vertex in relation to nearby vertices in order to generate a solution."))
 	EWeightScheme SelectedWeightScheme = EWeightScheme::ClampedCotangent;
 
-	UPROPERTY(EditAnywhere, Category = LaplacianOptions, meta = (EditCondition="DeformationStrategy == EGroupTopologyDeformationStrategy::Laplacian", DisplayName = "Handle Weight", UIMin = "0.0001", UIMax = "1.0", ClampMin = "0.0001", ClampMax = "100.0", ToolTip = "The handle weight describes how effective dragging the selected PolyGroup will be."))
-	double HandleWeight = 1.0;
+	UPROPERTY(EditAnywhere, Category = LaplacianOptions, meta = (EditCondition="DeformationStrategy == EGroupTopologyDeformationStrategy::Laplacian", DisplayName = "Handle Weight", UIMin = "0.0001", UIMax = "100.0", ClampMin = "0.0001", ClampMax = "1000.0", ToolTip = "The handle weight describes how effective dragging the selected PolyGroup will be."))
+	double HandleWeight = 100.0;
 
 	UPROPERTY(EditAnywhere, Category = LaplacianOptions, meta = (EditCondition="DeformationStrategy == EGroupTopologyDeformationStrategy::Laplacian", DisplayName = "Post-fix Handles", ToolTip = "Post-fixing the handles will attempt to force their vertices to arrive exactly where you wanted them. This may result\nin undesirable or choppy curvature in some scenarios."))
 	bool bPostFixHandles = false;
@@ -495,10 +499,11 @@ protected:
 	FGroupTopologyDeformer LinearDeformer;
 	FGroupTopologyLaplacianDeformer LaplacianDeformer;
 
-	// Initial polygon group info
+	// Initial polygon group and mesh info
 	TDynamicVector<int> InitialTriangleGroups;
+	TUniquePtr<FDynamicMesh3> InitialMesh;
 	void BackupTriangleGroups();
-	void RestoreTriangleGroups();
+	void SetTriangleGroups(const TDynamicVector<int>& Groups);
 	
 	// This is true when the spatial index needs to reflect a modification
 	bool bSpatialDirty; 
