fix asset names on mesh generation

#rnx

diff --git a/Engine/Plugins/Experimental/MeshModelingToolset/Source/MeshModelingTools/Private/AddPrimitiveTool.cpp b/Engine/Plugins/Experimental/MeshModelingToolset/Source/MeshModelingTools/Private/AddPrimitiveTool.cpp
--- a/Engine/Plugins/Experimental/MeshModelingToolset/Source/MeshModelingTools/Private/AddPrimitiveTool.cpp
+++ b/Engine/Plugins/Experimental/MeshModelingToolset/Source/MeshModelingTools/Private/AddPrimitiveTool.cpp
@@ -254,11 +254,14 @@ void UAddPrimitiveTool::OnClicked(const FInputDeviceRay& DeviceClickPos)
 	const FDynamicMesh3* CurMesh = PreviewMesh->GetPreviewDynamicMesh();
 	FTransform3d CurTransform(PreviewMesh->GetTransform());
 	UMaterialInterface* Material = PreviewMesh->GetMaterial();
-	GetToolManager()->BeginUndoTransaction(LOCTEXT("MergeMeshes", "Merge Meshes"));
+	GetToolManager()->BeginUndoTransaction(LOCTEXT("UAddPrimitiveTool", "Add Primitive Mesh"));
+
+	const UEnum* const MakeMeshShapeTypeEnum = StaticEnum<EMakeMeshShapeType>();		
+	FString ShapeTypeName = MakeMeshShapeTypeEnum->GetNameStringByValue((uint8)ShapeSettings->Shape);
 
 	AActor* NewActor = AssetGenerationUtil::GenerateStaticMeshActor(
 		AssetAPI, TargetWorld,
-		CurMesh, CurTransform, TEXT("ShapeMesh"),
+		CurMesh, CurTransform, ShapeTypeName,
 		AssetGenerationUtil::GetDefaultAutoGeneratedAssetPath(),
 		Material
 	);
diff --git a/Engine/Plugins/Experimental/MeshModelingToolset/Source/ModelingComponents/Private/AssetGenerationUtil.cpp b/Engine/Plugins/Experimental/MeshModelingToolset/Source/ModelingComponents/Private/AssetGenerationUtil.cpp
--- a/Engine/Plugins/Experimental/MeshModelingToolset/Source/ModelingComponents/Private/AssetGenerationUtil.cpp
+++ b/Engine/Plugins/Experimental/MeshModelingToolset/Source/ModelingComponents/Private/AssetGenerationUtil.cpp
@@ -38,7 +38,7 @@ AActor* AssetGenerationUtil::GenerateStaticMeshActor(
 	// asset name and location (temp)
 	FString PackageFolderPath = (PackagePath.IsEmpty()) ? GetDefaultAutoGeneratedAssetPath() : PackagePath;
 	//FString PackageFolderPath = AssetAPI->GetActiveAssetFolderPath();
-	FString MeshName = TEXT("MergedMesh");
+	FString MeshName = ObjectName;
 
 	// create new package
 	FString UniqueMeshName = AssetAPI->MakeUniqueAssetName(MeshName, PackageFolderPath);
