consistent UV scaling settings for Add Primitive and Draw Polygon tools

#rnx

diff --git a/Engine/Plugins/Experimental/GeometryProcessing/Source/DynamicMesh/Private/DynamicMeshEditor.cpp b/Engine/Plugins/Experimental/GeometryProcessing/Source/DynamicMesh/Private/DynamicMeshEditor.cpp
--- a/Engine/Plugins/Experimental/GeometryProcessing/Source/DynamicMesh/Private/DynamicMeshEditor.cpp
+++ b/Engine/Plugins/Experimental/GeometryProcessing/Source/DynamicMesh/Private/DynamicMeshEditor.cpp
@@ -574,6 +574,43 @@ void FDynamicMeshEditor::SetQuadUVsFromProjection(const FIndex2i& QuadTris, cons
 }
 
 
+void FDynamicMeshEditor::RescaleAttributeUVs(float UVScale, bool bWorldSpace, int UVLayerIndex)
+{
+	check(Mesh->HasAttributes() && Mesh->Attributes()->NumUVLayers() > UVLayerIndex );
+	FDynamicMeshUVOverlay* UVs = Mesh->Attributes()->GetUVLayer(UVLayerIndex);
+
+	if (bWorldSpace)
+	{
+		FVector2f TriUVs[3];
+		FVector3d TriVs[3];
+		float TotalEdgeUVLen = 0;
+		double TotalEdgeLen = 0;
+		for (int TID : Mesh->TriangleIndicesItr())
+		{
+			UVs->GetTriElements(TID, TriUVs[0], TriUVs[1], TriUVs[2]);
+			Mesh->GetTriVertices(TID, TriVs[0], TriVs[1], TriVs[2]);
+			for (int j = 2, i = 0; i < 3; j = i++)
+			{
+				TotalEdgeUVLen += TriUVs[j].Distance(TriUVs[i]);
+				TotalEdgeLen += TriVs[j].Distance(TriVs[i]);
+			}
+		}
+		if (TotalEdgeUVLen > KINDA_SMALL_NUMBER)
+		{
+			float AvgUVScale = TotalEdgeLen / TotalEdgeUVLen;
+			UVScale *= AvgUVScale;
+		}
+	}
+
+	for (int UVID : UVs->ElementIndicesItr())
+	{
+		FVector2f UV;
+		UVs->GetElement(UVID, UV);
+		UVs->SetElement(UVID, UV*UVScale);
+	}
+}
+
+
 
 
 
diff --git a/Engine/Plugins/Experimental/GeometryProcessing/Source/DynamicMesh/Public/DynamicMeshEditor.h b/Engine/Plugins/Experimental/GeometryProcessing/Source/DynamicMesh/Public/DynamicMeshEditor.h
--- a/Engine/Plugins/Experimental/GeometryProcessing/Source/DynamicMesh/Public/DynamicMeshEditor.h
+++ b/Engine/Plugins/Experimental/GeometryProcessing/Source/DynamicMesh/Public/DynamicMeshEditor.h
@@ -281,6 +281,14 @@ public:
 	void SetTriangleUVsFromProjection(const TArray<int>& Triangles, const FFrame3d& ProjectionFrame, float UVScaleFactor, int UVLayerIndex = 0);
 
 
+	/**
+	 * Rescale UVs for the whole mesh, for the given UV attribute layer
+	 * @param UVScale Scale factor to multiply into UVs.  If in world space, this is in centimeters relative to the average UV scale
+	 * @param bWorldSpace If true, UVs are rescaled relative to an absolute world scale.
+	 * @param UVLayerIndex which UV layer to operate on (must exist)
+	 */
+	void RescaleAttributeUVs(float UVScale = 1.0f, bool bWorldSpace = false, int UVLayerIndex = 0);
+
 
 	//////////////////////////////////////////////////////////////////////////
 	// mesh element copying / duplication
diff --git a/Engine/Plugins/Experimental/MeshModelingToolset/Source/MeshModelingTools/Private/AddPrimitiveTool.cpp b/Engine/Plugins/Experimental/MeshModelingToolset/Source/MeshModelingTools/Private/AddPrimitiveTool.cpp
--- a/Engine/Plugins/Experimental/MeshModelingToolset/Source/MeshModelingTools/Private/AddPrimitiveTool.cpp
+++ b/Engine/Plugins/Experimental/MeshModelingToolset/Source/MeshModelingTools/Private/AddPrimitiveTool.cpp
@@ -16,6 +16,8 @@
 #include "StaticMeshComponentBuilder.h"
 #include "Drawing/MeshDebugDrawing.h"
 
+#include "DynamicMeshEditor.h"
+
 #define LOCTEXT_NAMESPACE "UAddPrimitiveTool"
 
 
@@ -68,6 +70,8 @@ namespace
 	  { TEXT("Slices"),        EMakeMeshShapeType::Cylinder | EMakeMeshShapeType::Cone | EMakeMeshShapeType::Sphere },
 	  { TEXT("Subdivisions"),  EMakeMeshShapeType::Box | EMakeMeshShapeType::Plane },
 	  { TEXT("Material"),      EMakeMeshShapeType::All },
+	  { TEXT("bWorldSpaceUVScale"), EMakeMeshShapeType::All },
+	  { TEXT("UVScale"),	   EMakeMeshShapeType::All },
 	  { TEXT("bWireframe"),	   EMakeMeshShapeType::All }
 	};
 };
@@ -230,6 +234,13 @@ void UAddPrimitiveTool::UpdatePreviewMesh()
 		break;
 	}
 
+	if (ShapeSettings->UVScale != 1.0 || ShapeSettings->bWorldSpaceUVScale)
+	{
+		FDynamicMeshEditor Editor(&NewMesh);
+		float WorldUnitsInMetersFactor = ShapeSettings->bWorldSpaceUVScale ? .01f : 1.0f;
+		Editor.RescaleAttributeUVs(ShapeSettings->UVScale * WorldUnitsInMetersFactor, ShapeSettings->bWorldSpaceUVScale);
+	}
+
 	// set mesh position
 	FAxisAlignedBox3d Bounds = NewMesh.GetCachedBounds();
 	FVector3d TargetOrigin = Bounds.Center();
diff --git a/Engine/Plugins/Experimental/MeshModelingToolset/Source/MeshModelingTools/Private/DrawPolygonTool.cpp b/Engine/Plugins/Experimental/MeshModelingToolset/Source/MeshModelingTools/Private/DrawPolygonTool.cpp
--- a/Engine/Plugins/Experimental/MeshModelingToolset/Source/MeshModelingTools/Private/DrawPolygonTool.cpp
+++ b/Engine/Plugins/Experimental/MeshModelingToolset/Source/MeshModelingTools/Private/DrawPolygonTool.cpp
@@ -20,6 +20,8 @@
 #include "MeshQueries.h"
 #include "ToolSceneQueriesUtil.h"
 
+#include "DynamicMeshEditor.h"
+
 #include "PositionPlaneGizmo.h"
 #include "Drawing/MeshDebugDrawing.h"
 
@@ -985,6 +987,17 @@ void UDrawPolygonTool::GeneratePolygonMesh(const TArray<FVector>& Polygon, FDyna
 
 		Extruder.Apply();
 	}
+
+	FDynamicMeshEditor Editor(ResultMeshOut);
+	float InitialUVScale = 1.0 / PolyMeshGen.Polygon.Bounds().MaxDim(); // this is the UV scale used by both the polymeshgen and the extruder above
+	// default global rescale -- initial scale doesn't factor in extrude distance; rescale so UVScale of 1.0 fits in the unit square texture
+	float GlobalUVRescale = PolygonProperties->UVScale / FMathf::Max(1.0f, ExtrudeDistance * InitialUVScale);
+	if (PolygonProperties->bWorldSpaceUVScale)
+	{
+		// since we know the initial uv scale, directly compute the global scale (relative to 1 meter as a standard scale)
+		GlobalUVRescale = PolygonProperties->UVScale * .01 / InitialUVScale;
+	}
+	Editor.RescaleAttributeUVs(GlobalUVRescale, false);
 }
 
 void UDrawPolygonTool::GeneratePreviewHeightTarget()
diff --git a/Engine/Plugins/Experimental/MeshModelingToolset/Source/MeshModelingTools/Public/AddPrimitiveTool.h b/Engine/Plugins/Experimental/MeshModelingToolset/Source/MeshModelingTools/Public/AddPrimitiveTool.h
--- a/Engine/Plugins/Experimental/MeshModelingToolset/Source/MeshModelingTools/Public/AddPrimitiveTool.h
+++ b/Engine/Plugins/Experimental/MeshModelingToolset/Source/MeshModelingTools/Public/AddPrimitiveTool.h
@@ -121,6 +121,12 @@ public:
 	UPROPERTY(EditAnywhere, Category = ShapeSettings, meta = (DisplayName = "Material"))
 	UMaterialInterface* Material;
 
+	UPROPERTY(EditAnywhere, Category = ShapeSettings, meta = (DisplayName = "UV Scale"))
+	float UVScale = 1.0;
+	
+	UPROPERTY(EditAnywhere, Category = ShapeSettings, meta = (DisplayName = "UV Scale Relative to World Space"))
+	bool bWorldSpaceUVScale = false;
+
 	/** Wireframe toggle */
 	UPROPERTY(EditAnywhere, Category = ShapeSettings, meta = (DisplayName = "Show Wireframe"))
 	bool bWireframe;
diff --git a/Engine/Plugins/Experimental/MeshModelingToolset/Source/MeshModelingTools/Public/DrawPolygonTool.h b/Engine/Plugins/Experimental/MeshModelingToolset/Source/MeshModelingTools/Public/DrawPolygonTool.h
--- a/Engine/Plugins/Experimental/MeshModelingToolset/Source/MeshModelingTools/Public/DrawPolygonTool.h
+++ b/Engine/Plugins/Experimental/MeshModelingToolset/Source/MeshModelingTools/Public/DrawPolygonTool.h
@@ -101,6 +101,12 @@ class MESHMODELINGTOOLS_API UDrawPolygonToolStandardProperties : public UObject
 	/** Material for new mesh*/
 	UPROPERTY(EditAnywhere, NonTransactional, Category = Polygon, meta = (DisplayName = "Material"))
 	UMaterialInterface* Material;
+
+	UPROPERTY(EditAnywhere, Category = ShapeSettings, meta = (DisplayName = "UV Scale"))
+	float UVScale = 1.0;
+
+	UPROPERTY(EditAnywhere, Category = ShapeSettings, meta = (DisplayName = "UV Scale Relative to World Space"))
+	bool bWorldSpaceUVScale = false;
 };
 
 
