fix witheditor check in edit mesh tool

#rnx

diff --git a/Engine/Plugins/Experimental/MeshModelingToolset/Source/MeshModelingTools/Private/EditMeshPolygonsTool.cpp b/Engine/Plugins/Experimental/MeshModelingToolset/Source/MeshModelingTools/Private/EditMeshPolygonsTool.cpp
--- a/Engine/Plugins/Experimental/MeshModelingToolset/Source/MeshModelingTools/Private/EditMeshPolygonsTool.cpp
+++ b/Engine/Plugins/Experimental/MeshModelingToolset/Source/MeshModelingTools/Private/EditMeshPolygonsTool.cpp
@@ -56,7 +56,7 @@ UPolyEditTransformProperties::UPolyEditTransformProperties()
 	PolygonGroupingAngleThreshold = .5;
 }
 
-#ifdef WITH_EDITOR
+#if WITH_EDITOR
 void UPolyEditTransformProperties::PostEditChangeProperty(FPropertyChangedEvent& PropertyChangedEvent)
 {
 	// skip interactive updates for PolygonGroupingAngleThreshold
diff --git a/Engine/Plugins/Experimental/MeshModelingToolset/Source/MeshModelingTools/Public/EditMeshPolygonsTool.h b/Engine/Plugins/Experimental/MeshModelingToolset/Source/MeshModelingTools/Public/EditMeshPolygonsTool.h
--- a/Engine/Plugins/Experimental/MeshModelingToolset/Source/MeshModelingTools/Public/EditMeshPolygonsTool.h
+++ b/Engine/Plugins/Experimental/MeshModelingToolset/Source/MeshModelingTools/Public/EditMeshPolygonsTool.h
@@ -100,8 +100,10 @@ class MESHMODELINGTOOLS_API UPolyEditTransformProperties : public UInteractiveTo
 public:
 	UPolyEditTransformProperties();
 
-#ifdef WITH_EDITOR
+#if WITH_EDITOR
+	// UObject interface
 	virtual void PostEditChangeProperty(FPropertyChangedEvent& PropertyChangedEvent) override;
+	// End of UObject interface
 #endif
 
 	//Options
