fix loctext key for add primitive transaction

#rnx

diff --git a/Engine/Plugins/Experimental/MeshModelingToolset/Source/MeshModelingTools/Private/AddPrimitiveTool.cpp b/Engine/Plugins/Experimental/MeshModelingToolset/Source/MeshModelingTools/Private/AddPrimitiveTool.cpp
--- a/Engine/Plugins/Experimental/MeshModelingToolset/Source/MeshModelingTools/Private/AddPrimitiveTool.cpp
+++ b/Engine/Plugins/Experimental/MeshModelingToolset/Source/MeshModelingTools/Private/AddPrimitiveTool.cpp
@@ -265,7 +265,7 @@ void UAddPrimitiveTool::OnClicked(const FInputDeviceRay& DeviceClickPos)
 	const FDynamicMesh3* CurMesh = PreviewMesh->GetPreviewDynamicMesh();
 	FTransform3d CurTransform(PreviewMesh->GetTransform());
 	UMaterialInterface* Material = PreviewMesh->GetMaterial();
-	GetToolManager()->BeginUndoTransaction(LOCTEXT("UAddPrimitiveTool", "Add Primitive Mesh"));
+	GetToolManager()->BeginUndoTransaction(LOCTEXT("AddPrimitiveToolTransactionName", "Add Primitive Mesh"));
 
 	const UEnum* const MakeMeshShapeTypeEnum = StaticEnum<EMakeMeshShapeType>();		
 	FString ShapeTypeName = MakeMeshShapeTypeEnum->GetNameStringByValue((uint8)ShapeSettings->Shape);
