fix invert normals option for normal edit tool to also invert face orientations
also general cleanup of normal code + tool options

#rnx

diff --git a/Engine/Plugins/Experimental/GeometryProcessing/Source/DynamicMesh/Private/DynamicMesh3_Edits.cpp b/Engine/Plugins/Experimental/GeometryProcessing/Source/DynamicMesh/Private/DynamicMesh3_Edits.cpp
--- a/Engine/Plugins/Experimental/GeometryProcessing/Source/DynamicMesh/Private/DynamicMesh3_Edits.cpp
+++ b/Engine/Plugins/Experimental/GeometryProcessing/Source/DynamicMesh/Private/DynamicMesh3_Edits.cpp
@@ -1545,8 +1545,8 @@ EMeshResult FDynamicMesh3::MergeEdges(int eKeep, int eDiscard, FMergeEdgesInfo&
 	// alternative that detects normal flip of triangle tcd. This is a more 
 	// robust geometric test, but fails if tri is degenerate...also more expensive
 	//FVector3d otherv = GetVertex(tcd_otherv);
-	//FVector3d Ncd = TMathUtil.FastNormalDirection(GetVertex(c), GetVertex(d), otherv);
-	//FVector3d Nab = TMathUtil.FastNormalDirection(GetVertex(a), GetVertex(b), otherv);
+	//FVector3d Ncd = VectorUtil::NormalDirection(GetVertex(c), GetVertex(d), otherv);
+	//FVector3d Nab = VectorUtil::NormalDirection(GetVertex(a), GetVertex(b), otherv);
 	//if (Ncd.Dot(Nab) < 0)
 	//return EMeshResult::Failed_SameOrientation;
 
diff --git a/Engine/Plugins/Experimental/GeometryProcessing/Source/DynamicMesh/Private/DynamicMesh3_Queries.cpp b/Engine/Plugins/Experimental/GeometryProcessing/Source/DynamicMesh/Private/DynamicMesh3_Queries.cpp
--- a/Engine/Plugins/Experimental/GeometryProcessing/Source/DynamicMesh/Private/DynamicMesh3_Queries.cpp
+++ b/Engine/Plugins/Experimental/GeometryProcessing/Source/DynamicMesh/Private/DynamicMesh3_Queries.cpp
@@ -764,9 +764,7 @@ void FDynamicMesh3::GetTriInfo(int tID, FVector3d& Normal, double& Area, FVector
 	FVector3d v0, v1, v2;
 	GetTriVertices(tID, v0, v1, v2);
 	Centroid = (v0 + v1 + v2) * (1.0 / 3.0);
-	Area = VectorUtil::Area(v0, v1, v2);
-	Normal = VectorUtil::Normal(v0, v1, v2);
-	//normal = FastNormalArea(ref v0, ref v1, ref v2, out fArea);
+	Normal = VectorUtil::NormalArea(v0, v1, v2, Area);
 }
 
 
diff --git a/Engine/Plugins/Experimental/GeometryProcessing/Source/DynamicMesh/Private/MeshRefinerBase.cpp b/Engine/Plugins/Experimental/GeometryProcessing/Source/DynamicMesh/Private/MeshRefinerBase.cpp
--- a/Engine/Plugins/Experimental/GeometryProcessing/Source/DynamicMesh/Private/MeshRefinerBase.cpp
+++ b/Engine/Plugins/Experimental/GeometryProcessing/Source/DynamicMesh/Private/MeshRefinerBase.cpp
@@ -71,14 +71,14 @@ bool FMeshRefinerBase::CheckIfFlipInvertsNormals(int a, int b, int c, int d, int
 	int oa = a, ob = b;
 	IndexUtil::OrientTriEdge(oa, ob, tri_v);
 	FVector3d vOA = Mesh->GetVertex(oa), vOB = Mesh->GetVertex(ob);
-	FVector3d n0 = VectorUtil::FastNormalDirection(vOA, vOB, vC);
-	FVector3d n1 = VectorUtil::FastNormalDirection(vOB, vOA, vD);
-	FVector3d f0 = VectorUtil::FastNormalDirection(vC, vD, vOB);
+	FVector3d n0 = VectorUtil::NormalDirection(vOA, vOB, vC);
+	FVector3d n1 = VectorUtil::NormalDirection(vOB, vOA, vD);
+	FVector3d f0 = VectorUtil::NormalDirection(vC, vD, vOB);
 	if (ComputeEdgeFlipMetric(n0, f0) <= EdgeFlipTolerance || ComputeEdgeFlipMetric(n1, f0) <= EdgeFlipTolerance)
 	{
 		return true;
 	}
-	FVector3d f1 = VectorUtil::FastNormalDirection(vD, vC, vOA);
+	FVector3d f1 = VectorUtil::NormalDirection(vD, vC, vOA);
 	if (ComputeEdgeFlipMetric(n0, f1) <= EdgeFlipTolerance || ComputeEdgeFlipMetric(n1, f1) <= EdgeFlipTolerance)
 	{
 		return true;
diff --git a/Engine/Plugins/Experimental/GeometryProcessing/Source/GeometricObjects/Public/MeshQueries.h b/Engine/Plugins/Experimental/GeometryProcessing/Source/GeometricObjects/Public/MeshQueries.h
--- a/Engine/Plugins/Experimental/GeometryProcessing/Source/GeometricObjects/Public/MeshQueries.h
+++ b/Engine/Plugins/Experimental/GeometryProcessing/Source/GeometricObjects/Public/MeshQueries.h
@@ -65,7 +65,7 @@ public:
 		FTriangle3d Triangle;
 		Mesh.GetTriVertices(TriIdx, Triangle.V[0], Triangle.V[1], Triangle.V[2]);
 		Centroid = Triangle.Centroid();
-		Normal = VectorUtil::FastNormalArea(Triangle.V[0], Triangle.V[1], Triangle.V[2], Area);
+		Normal = VectorUtil::NormalArea(Triangle.V[0], Triangle.V[1], Triangle.V[2], Area);
 	}
 
 	static FAxisAlignedBox3d GetTriBounds(const TriangleMeshType& Mesh, int TID)
diff --git a/Engine/Plugins/Experimental/GeometryProcessing/Source/GeometricObjects/Public/Spatial/FastWinding.h b/Engine/Plugins/Experimental/GeometryProcessing/Source/GeometricObjects/Public/Spatial/FastWinding.h
--- a/Engine/Plugins/Experimental/GeometryProcessing/Source/GeometricObjects/Public/Spatial/FastWinding.h
+++ b/Engine/Plugins/Experimental/GeometryProcessing/Source/GeometricObjects/Public/Spatial/FastWinding.h
@@ -65,7 +65,7 @@ namespace FastTriWinding
 			//if (TriCache == null)
 			//{
 			//	c = (1.0 / 3.0) * (P0 + P1 + P2);
-			//	n = FMathd::FastNormalArea(P0, P1, P2, out a);
+			//	n = VectorUtil::NormalArea(P0, P1, P2, out a);
 			//}
 			//else
 			{
diff --git a/Engine/Plugins/Experimental/GeometryProcessing/Source/GeometricObjects/Public/VectorUtil.h b/Engine/Plugins/Experimental/GeometryProcessing/Source/GeometricObjects/Public/VectorUtil.h
--- a/Engine/Plugins/Experimental/GeometryProcessing/Source/GeometricObjects/Public/VectorUtil.h
+++ b/Engine/Plugins/Experimental/GeometryProcessing/Source/GeometricObjects/Public/VectorUtil.h
@@ -75,7 +75,7 @@ namespace VectorUtil
 	 * @return un-normalized direction that is parallel to normal of triangle V0,V1,V2
 	 */
 	template <typename RealType>
-	inline FVector3<RealType> FastNormalDirection(const FVector3<RealType>& V0, const FVector3<RealType>& V1, const FVector3<RealType>& V2)
+	inline FVector3<RealType> NormalDirection(const FVector3<RealType>& V0, const FVector3<RealType>& V1, const FVector3<RealType>& V2)
 	{
 		// Unreal has Left-Hand Coordinate System so we need to reverse this cross-product to get proper triangle normal
 		return (V2 - V0).Cross(V1 - V0);
@@ -123,7 +123,7 @@ namespace VectorUtil
 	 * @return triangle normal
 	 */
 	template <typename RealType>
-	inline FVector3<RealType> FastNormalArea(const FVector3<RealType>& V0, const FVector3<RealType>& V1, const FVector3<RealType>& V2, RealType& AreaOut)
+	inline FVector3<RealType> NormalArea(const FVector3<RealType>& V0, const FVector3<RealType>& V1, const FVector3<RealType>& V2, RealType& AreaOut)
 	{
 		FVector3<RealType> edge1(V1 - V0);
 		FVector3<RealType> edge2(V2 - V0);
diff --git a/Engine/Plugins/Experimental/GeometryProcessing/Source/MeshConversion/Private/MeshDescriptionBuilder.cpp b/Engine/Plugins/Experimental/GeometryProcessing/Source/MeshConversion/Private/MeshDescriptionBuilder.cpp
--- a/Engine/Plugins/Experimental/GeometryProcessing/Source/MeshConversion/Private/MeshDescriptionBuilder.cpp
+++ b/Engine/Plugins/Experimental/GeometryProcessing/Source/MeshConversion/Private/MeshDescriptionBuilder.cpp
@@ -332,7 +332,7 @@ void FMeshDescriptionBuilder::RecalculateInstanceNormals()
 			FVector3d B = GetPosition(TriangleVertexInstanceIDs[1]);
 			FVector3d C = GetPosition(TriangleVertexInstanceIDs[2]);
 			double Area = 1.0;
-			FVector FaceNormal = VectorUtil::FastNormalArea(A, B, C, Area);
+			FVector FaceNormal = VectorUtil::NormalArea(A, B, C, Area);
 			if (Area > FMathf::ZeroTolerance)
 			{
 				FaceNormal *= Area;
diff --git a/Engine/Plugins/Experimental/MeshModelingToolset/Source/MeshModelingTools/Private/EditNormalsTool.cpp b/Engine/Plugins/Experimental/MeshModelingToolset/Source/MeshModelingTools/Private/EditNormalsTool.cpp
--- a/Engine/Plugins/Experimental/MeshModelingToolset/Source/MeshModelingTools/Private/EditNormalsTool.cpp
+++ b/Engine/Plugins/Experimental/MeshModelingToolset/Source/MeshModelingTools/Private/EditNormalsTool.cpp
@@ -263,9 +263,9 @@ void UEditNormalsTool::GenerateAsset(const TArray<TUniquePtr<FDynamicMeshOpResul
 		ComponentTargets[ComponentIdx]->CommitMesh([&Results, &ComponentIdx, this](FMeshDescription* MeshDescription)
 		{
 			FDynamicMeshToMeshDescription Converter;
-			if (BasicProperties->bRecomputeNormalTopologyAndEdgeSharpness)
+			if (BasicProperties->bRecomputeNormalTopologyAndEdgeSharpness || BasicProperties->bInvertNormals)
 			{
-				// full conversion if normal topology changed
+				// full conversion if normal topology changed or faces were inverted
 				Converter.Convert(Results[ComponentIdx]->Mesh.Get(), *MeshDescription);
 			}
 			else
diff --git a/Engine/Plugins/Experimental/MeshModelingToolset/Source/MeshModelingTools/Public/EditNormalsTool.h b/Engine/Plugins/Experimental/MeshModelingToolset/Source/MeshModelingTools/Public/EditNormalsTool.h
--- a/Engine/Plugins/Experimental/MeshModelingToolset/Source/MeshModelingTools/Public/EditNormalsTool.h
+++ b/Engine/Plugins/Experimental/MeshModelingToolset/Source/MeshModelingTools/Public/EditNormalsTool.h
@@ -53,27 +53,27 @@ public:
 
 
 	/** Recompute all mesh normals */
-	UPROPERTY(EditAnywhere, Category = Options)
+	UPROPERTY(EditAnywhere, Category = Options, meta = (EditCondition = "!bRecomputeNormalTopologyAndEdgeSharpness"))
 	bool bRecomputeNormals;
 
 	/** Choose the method for computing vertex normals */
 	UPROPERTY(EditAnywhere, Category = Options)
 	ENormalCalculationMethod NormalCalculationMethod;
 
-	/** Invert (flip) all mesh normals */
+	/** Invert (flip) all mesh normals and associated triangle orientations */
 	UPROPERTY(EditAnywhere, Category = Options)
 	bool bInvertNormals;
 
-	/** Choose where to create sharp edges by having multiple normals on the same vertices, using the change in normals of adjacent faces */
+	/** Choose where to create sharp edges by having multiple normals on the same vertices, using the change in normals of adjacent faces.  If enabled, normals will always be recomputed. */
 	UPROPERTY(EditAnywhere, Category = NormalsTopology)
 	bool bRecomputeNormalTopologyAndEdgeSharpness;
 
 	/** Threshold on angle of change in face normals across an edge, above which we create a sharp edge if bSplitNormals is true */
-	UPROPERTY(EditAnywhere, Category = NormalsTopology, meta = (UIMin = "0.0", UIMax = "180.0", ClampMin = "0.0", ClampMax = "180.0"))
+	UPROPERTY(EditAnywhere, Category = NormalsTopology, meta = (UIMin = "0.0", UIMax = "180.0", ClampMin = "0.0", ClampMax = "180.0", EditCondition = "bRecomputeNormalTopologyAndEdgeSharpness"))
 	float SharpEdgeAngleThreshold;
 
-	/** If split normals is enabl*/
-	UPROPERTY(EditAnywhere, Category = NormalsTopology)
+	/** Give separate normals to each triangle at 'sharp' vertices where the coincident triangle normals change (across multiple triangles) by more than the angle threshold (e.g. a cone tip) */
+	UPROPERTY(EditAnywhere, Category = NormalsTopology, meta = (EditCondition = "bRecomputeNormalTopologyAndEdgeSharpness"))
 	bool bAllowSharpVertices;
 };
 
diff --git a/Engine/Plugins/Experimental/MeshModelingToolset/Source/ModelingOperators/Private/CleaningOps/EditNormalsOp.cpp b/Engine/Plugins/Experimental/MeshModelingToolset/Source/ModelingOperators/Private/CleaningOps/EditNormalsOp.cpp
--- a/Engine/Plugins/Experimental/MeshModelingToolset/Source/ModelingOperators/Private/CleaningOps/EditNormalsOp.cpp
+++ b/Engine/Plugins/Experimental/MeshModelingToolset/Source/ModelingOperators/Private/CleaningOps/EditNormalsOp.cpp
@@ -32,6 +32,28 @@ void FEditNormalsOp::CalculateResult(FProgressCancel* Progress)
 		ResultMesh->EnableAttributes();
 	}
 
+	// if you split normals you must always recompute as well
+	bool bNeedsRecompute = bRecomputeNormals || bSplitNormals;
+
+	if (bInvertNormals)
+	{
+		for (int TID : ResultMesh->TriangleIndicesItr())
+		{
+			ResultMesh->ReverseTriOrientation(TID);
+		}
+
+		// also reverse the normal directions (but only if a recompute isn't going to do it for us below)
+		if (!bNeedsRecompute)
+		{
+			FDynamicMeshNormalOverlay* Normals = ResultMesh->Attributes()->PrimaryNormals();
+			for (int ElID : Normals->ElementIndicesItr())
+			{
+				auto El = Normals->GetElement(ElID);
+				Normals->SetElement(ElID, -El);
+			}
+		}
+	}
+
 	if (bSplitNormals)
 	{
 		FMeshNormals FaceNormals(ResultMesh.Get());
@@ -47,21 +69,11 @@ void FEditNormalsOp::CalculateResult(FProgressCancel* Progress)
 	bool bAreaWeight = (NormalCalculationMethod == ENormalCalculationMethod::AreaWeighted || NormalCalculationMethod == ENormalCalculationMethod::AreaAngleWeighting);
 	bool bAngleWeight = (NormalCalculationMethod == ENormalCalculationMethod::AngleWeighted || NormalCalculationMethod == ENormalCalculationMethod::AreaAngleWeighting);
 
-	if (bSplitNormals || bRecomputeNormals) // split normals requires recompute, as it will nuke whatever normals were in the original overlay
+	if (bNeedsRecompute)
 	{
 		FMeshNormals MeshNormals(ResultMesh.Get());
 		MeshNormals.RecomputeOverlayNormals(ResultMesh->Attributes()->PrimaryNormals(), bAreaWeight, bAngleWeight);
-		MeshNormals.CopyToOverlay(ResultMesh->Attributes()->PrimaryNormals(), bInvertNormals);
-	}
-	else if (bInvertNormals)
-	{
-		// just invert with no recompute
-		FDynamicMeshNormalOverlay* Normals = ResultMesh->Attributes()->PrimaryNormals();
-		for (int ElID : Normals->ElementIndicesItr())
-		{
-			auto El = Normals->GetElement(ElID);
-			Normals->SetElement(ElID, -El);
-		}
+		MeshNormals.CopyToOverlay(ResultMesh->Attributes()->PrimaryNormals(), false);
 	}
 
 }
\ No newline at end of file
