ModelingTools: set display name of generated Actors
#rnx

diff --git a/Engine/Plugins/Experimental/MeshModelingToolset/Source/ModelingComponents/Private/AssetGenerationUtil.cpp b/Engine/Plugins/Experimental/MeshModelingToolset/Source/ModelingComponents/Private/AssetGenerationUtil.cpp
--- a/Engine/Plugins/Experimental/MeshModelingToolset/Source/ModelingComponents/Private/AssetGenerationUtil.cpp
+++ b/Engine/Plugins/Experimental/MeshModelingToolset/Source/ModelingComponents/Private/AssetGenerationUtil.cpp
@@ -51,6 +51,7 @@ AActor* AssetGenerationUtil::GenerateStaticMeshActor(
 	// @todo nothing here is specific to AStaticMeshActor...could we pass in a CDO and clone it instead of using SpawnActor?
 	AStaticMeshActor* NewActor = TargetWorld->SpawnActor<AStaticMeshActor>(FVector::ZeroVector, Rotation, SpawnInfo);
 	NewActor->Rename(*UniqueMeshName);
+	NewActor->SetActorLabel(*UniqueMeshName);
 
 	// construct new static mesh
 	FStaticMeshComponentBuilder Builder;
