Mesh Modeling: Simplfy tool. The UE4 tool splits verts on the UV seams.  Add a weld step when this is used.
#rnx

diff --git a/Engine/Plugins/Experimental/MeshModelingToolset/Source/MeshModelingToolsEditorOnly/Private/SimplifyMeshTool.cpp b/Engine/Plugins/Experimental/MeshModelingToolset/Source/MeshModelingToolsEditorOnly/Private/SimplifyMeshTool.cpp
--- a/Engine/Plugins/Experimental/MeshModelingToolset/Source/MeshModelingToolsEditorOnly/Private/SimplifyMeshTool.cpp
+++ b/Engine/Plugins/Experimental/MeshModelingToolset/Source/MeshModelingToolsEditorOnly/Private/SimplifyMeshTool.cpp
@@ -21,6 +21,7 @@
 #include "IMeshReductionManagerModule.h"
 #include "IMeshReductionInterfaces.h"
 #include "Modules/ModuleManager.h"
+#include "Operations/MergeCoincidentMeshEdges.h"
 
 #define LOCTEXT_NAMESPACE "USimplifyMeshTool"
 
@@ -266,10 +267,27 @@ void USimplifyMeshTool::UpdateResult()
 			MeshReduction->ReduceMeshDescription(DstMeshDescription, Error, *SrcMeshDescription, OverlappingCorners, ReductionSettings);
 		}
 
-		UE_LOG(LogMeshSimplification, Log, TEXT("Mesh simplified to %d triangles"), DstMeshDescription.Polygons().Num());
-
 		// Put the reduced mesh into the target...
 		DynamicMeshComponent->InitializeMesh(&DstMeshDescription);
+
+		// The UE4 tool will split the UV boundaries.  Need to weld this.
+		{
+			FDynamicMesh3* ComponentMesh = DynamicMeshComponent->GetMesh();
+
+			FMergeCoincidentMeshEdges Merger(ComponentMesh);
+			Merger.MergeSearchTolerance = 10.0f * FMathf::ZeroTolerance;
+			Merger.OnlyUniquePairs = false;
+			if (Merger.Apply() == false)
+			{
+				DynamicMeshComponent->InitializeMesh(&DstMeshDescription);
+			}
+			if (ComponentMesh->CheckValidity(true, EValidityCheckFailMode::ReturnOnly) == false)
+			{
+				DynamicMeshComponent->InitializeMesh(&DstMeshDescription);
+			}
+		
+			DynamicMeshComponent->NotifyMeshUpdated();
+		}
 	}
 
 	UE_LOG(LogMeshSimplification, Log, TEXT("Mesh simplified to %d triangles"), TargetMesh->TriangleCount());
diff --git a/Engine/Plugins/Experimental/MeshModelingToolset/Source/MeshModelingToolsEditorOnly/Public/SimplifyMeshTool.h b/Engine/Plugins/Experimental/MeshModelingToolset/Source/MeshModelingToolsEditorOnly/Public/SimplifyMeshTool.h
--- a/Engine/Plugins/Experimental/MeshModelingToolset/Source/MeshModelingToolsEditorOnly/Public/SimplifyMeshTool.h
+++ b/Engine/Plugins/Experimental/MeshModelingToolset/Source/MeshModelingToolsEditorOnly/Public/SimplifyMeshTool.h
@@ -149,5 +149,4 @@ protected:
 	bool bResultValid;
 	void UpdateResult();
 
-	void ComputeSimplifying();
 };
