standardize uv settings code across mesh inspector and uv projector

#rnx

diff --git a/Engine/Plugins/Experimental/MeshModelingToolset/Source/MeshModelingTools/Private/AddPrimitiveTool.cpp b/Engine/Plugins/Experimental/MeshModelingToolset/Source/MeshModelingTools/Private/AddPrimitiveTool.cpp
--- a/Engine/Plugins/Experimental/MeshModelingToolset/Source/MeshModelingTools/Private/AddPrimitiveTool.cpp
+++ b/Engine/Plugins/Experimental/MeshModelingToolset/Source/MeshModelingTools/Private/AddPrimitiveTool.cpp
@@ -51,7 +51,6 @@ UProceduralShapeToolProperties::UProceduralShapeToolProperties()
 	Subdivisions = 0;
 	PivotLocation = EMakeMeshPivotLocation::Base;
 	PlaceMode = EMakeMeshPlacementType::OnScene;
-    Material = CreateDefaultSubobject<UMaterialInterface>(TEXT("MATERIAL"));
 }
 
 namespace
@@ -68,11 +67,7 @@ namespace
 	  { TEXT("PlaceMode"),     EMakeMeshShapeType::All },
 	  { TEXT("PivotLocation"), EMakeMeshShapeType::All },
 	  { TEXT("Slices"),        EMakeMeshShapeType::Cylinder | EMakeMeshShapeType::Cone | EMakeMeshShapeType::Sphere },
-	  { TEXT("Subdivisions"),  EMakeMeshShapeType::Box | EMakeMeshShapeType::Plane },
-	  { TEXT("Material"),      EMakeMeshShapeType::All },
-	  { TEXT("bWorldSpaceUVScale"), EMakeMeshShapeType::All },
-	  { TEXT("UVScale"),	   EMakeMeshShapeType::All },
-	  { TEXT("bWireframe"),	   EMakeMeshShapeType::All }
+	  { TEXT("Subdivisions"),  EMakeMeshShapeType::Box | EMakeMeshShapeType::Plane }
 	};
 };
 
@@ -115,11 +110,14 @@ void UAddPrimitiveTool::Setup()
 	ShapeSettings = NewObject<UProceduralShapeToolProperties>(this);
 	AddToolPropertySource(ShapeSettings);
 
+	MaterialProperties = NewObject<UNewMeshMaterialProperties>(this);
+	AddToolPropertySource(MaterialProperties);
+
 	// create preview mesh object
 	PreviewMesh = NewObject<UPreviewMesh>(this, TEXT("PreviewMesh"));
 	PreviewMesh->CreateInWorld(TargetWorld, FTransform::Identity);
 	PreviewMesh->SetVisible(false);
-	PreviewMesh->SetMaterial(ShapeSettings->Material);
+	PreviewMesh->SetMaterial(MaterialProperties->Material);
 
 	UpdatePreviewMesh();
 
@@ -147,8 +145,8 @@ void UAddPrimitiveTool::Render(IToolsContextRenderAPI* RenderAPI)
 
 void UAddPrimitiveTool::OnPropertyModified(UObject* PropertySet, UProperty* Property)
 {
-	PreviewMesh->EnableWireframe(ShapeSettings->bWireframe);
-	PreviewMesh->SetMaterial(ShapeSettings->Material);
+	PreviewMesh->EnableWireframe(MaterialProperties->bWireframe);
+	PreviewMesh->SetMaterial(MaterialProperties->Material);
 	UpdatePreviewMesh();
 }
 
@@ -251,11 +249,11 @@ void UAddPrimitiveTool::UpdatePreviewMesh()
 		break;
 	}
 
-	if (ShapeSettings->UVScale != 1.0 || ShapeSettings->bWorldSpaceUVScale)
+	if (MaterialProperties->UVScale != 1.0 || MaterialProperties->bWorldSpaceUVScale)
 	{
 		FDynamicMeshEditor Editor(&NewMesh);
-		float WorldUnitsInMetersFactor = ShapeSettings->bWorldSpaceUVScale ? .01f : 1.0f;
-		Editor.RescaleAttributeUVs(ShapeSettings->UVScale * WorldUnitsInMetersFactor, ShapeSettings->bWorldSpaceUVScale);
+		float WorldUnitsInMetersFactor = MaterialProperties->bWorldSpaceUVScale ? .01f : 1.0f;
+		Editor.RescaleAttributeUVs(MaterialProperties->UVScale * WorldUnitsInMetersFactor, MaterialProperties->bWorldSpaceUVScale);
 	}
 
 	// set mesh position
diff --git a/Engine/Plugins/Experimental/MeshModelingToolset/Source/MeshModelingTools/Private/DrawPolygonTool.cpp b/Engine/Plugins/Experimental/MeshModelingToolset/Source/MeshModelingTools/Private/DrawPolygonTool.cpp
--- a/Engine/Plugins/Experimental/MeshModelingToolset/Source/MeshModelingTools/Private/DrawPolygonTool.cpp
+++ b/Engine/Plugins/Experimental/MeshModelingToolset/Source/MeshModelingTools/Private/DrawPolygonTool.cpp
@@ -64,7 +64,6 @@ UInteractiveTool* UDrawPolygonToolBuilder::BuildTool(const FToolBuilderState& Sc
  */
 UDrawPolygonToolStandardProperties::UDrawPolygonToolStandardProperties()
 {
-	Material = CreateDefaultSubobject<UMaterialInterface>(TEXT("MATERIAL"));
 }
 
 /*
@@ -129,11 +128,13 @@ void UDrawPolygonTool::Setup()
 		UpdateDrawPlaneFromGizmo(WorldFrame);
 	};
 
+	MaterialProperties = NewObject<UNewMeshMaterialProperties>(this);
+
 	// create preview mesh object
 	PreviewMesh = NewObject<UPreviewMesh>(this, TEXT("PreviewMesh"));
 	PreviewMesh->CreateInWorld(this->TargetWorld, FTransform::Identity);
 	PreviewMesh->SetVisible(false);
-	PreviewMesh->SetMaterial(PolygonProperties->Material);
+	PreviewMesh->SetMaterial(MaterialProperties->Material);
 	bPreviewUpdatePending = false;
 
 	SnapEngine.SnapMetricTolerance = ToolSceneQueriesUtil::GetDefaultVisualAngleSnapThreshD();
@@ -145,9 +146,12 @@ void UDrawPolygonTool::Setup()
 
 	SnapProperties = NewObject<UDrawPolygonToolSnapProperties>(this, TEXT("Snapping"));
 
+	
+
 	// add self as tool props
 	AddToolPropertySource(PolygonProperties);
 	AddToolPropertySource(SnapProperties);
+	AddToolPropertySource(MaterialProperties);
 
 	GetToolManager()->DisplayMessage(
 		LOCTEXT("OnStartDrawPolygonTool", "Left-click to place points on the Drawing Plane. Ctrl-click on the scene to reposition the Plane, and click-drag on the gizmo to move it."),
@@ -886,7 +890,7 @@ void UDrawPolygonTool::EmitCurrentPolygon()
 		AssetAPI, TargetWorld,
 		&Mesh, PlaneFrameOut.ToTransform(), TEXT("Polygon"),
 		AssetGenerationUtil::GetDefaultAutoGeneratedAssetPath(),
-		PolygonProperties->Material);
+		MaterialProperties->Material);
 
 	// select newly-created object
 	ToolSelectionUtil::SetNewActorSelection(GetToolManager(), NewActor);
@@ -913,7 +917,8 @@ void UDrawPolygonTool::UpdateLivePreview()
 	GeneratePolygonMesh(PolygonVertices, &Mesh, PlaneFrame, false, ExtrudeDist, false);
 
 	PreviewMesh->SetTransform(PlaneFrame.ToFTransform());
-	PreviewMesh->SetMaterial(PolygonProperties->Material);
+	PreviewMesh->SetMaterial(MaterialProperties->Material);
+	PreviewMesh->EnableWireframe(MaterialProperties->bWireframe);
 	PreviewMesh->UpdatePreview(&Mesh);
 }
 
@@ -991,11 +996,11 @@ void UDrawPolygonTool::GeneratePolygonMesh(const TArray<FVector>& Polygon, FDyna
 	FDynamicMeshEditor Editor(ResultMeshOut);
 	float InitialUVScale = 1.0 / PolyMeshGen.Polygon.Bounds().MaxDim(); // this is the UV scale used by both the polymeshgen and the extruder above
 	// default global rescale -- initial scale doesn't factor in extrude distance; rescale so UVScale of 1.0 fits in the unit square texture
-	float GlobalUVRescale = PolygonProperties->UVScale / FMathf::Max(1.0f, ExtrudeDistance * InitialUVScale);
-	if (PolygonProperties->bWorldSpaceUVScale)
+	float GlobalUVRescale = MaterialProperties->UVScale / FMathf::Max(1.0f, ExtrudeDistance * InitialUVScale);
+	if (MaterialProperties->bWorldSpaceUVScale)
 	{
 		// since we know the initial uv scale, directly compute the global scale (relative to 1 meter as a standard scale)
-		GlobalUVRescale = PolygonProperties->UVScale * .01 / InitialUVScale;
+		GlobalUVRescale = MaterialProperties->UVScale * .01 / InitialUVScale;
 	}
 	Editor.RescaleAttributeUVs(GlobalUVRescale, false);
 }
diff --git a/Engine/Plugins/Experimental/MeshModelingToolset/Source/MeshModelingTools/Private/MeshInspectorTool.cpp b/Engine/Plugins/Experimental/MeshModelingToolset/Source/MeshModelingTools/Private/MeshInspectorTool.cpp
--- a/Engine/Plugins/Experimental/MeshModelingToolset/Source/MeshModelingTools/Private/MeshInspectorTool.cpp
+++ b/Engine/Plugins/Experimental/MeshModelingToolset/Source/MeshModelingTools/Private/MeshInspectorTool.cpp
@@ -64,7 +64,6 @@ void UMeshInspectorTool::Setup()
 	{
 		DynamicMeshComponent->SetMaterial(0, DefaultMaterial);
 	}
-	ActiveMaterialMode = EInspectorMaterialMode::Default;
 
 	DynamicMeshComponent->TangentsType = EDynamicMeshTangentCalcType::ExternallyCalculated;
 	DynamicMeshComponent->InitializeMesh(ComponentTarget->GetMesh());
@@ -78,16 +77,10 @@ void UMeshInspectorTool::Setup()
 	Settings = NewObject<UMeshInspectorProperties>(this);
 	AddToolPropertySource(Settings);
 
-	UMaterial* CheckerMaterialBase = LoadObject<UMaterial>(nullptr, TEXT("/MeshModelingToolset/Materials/CheckerMaterial"));
-	if (CheckerMaterialBase != nullptr)
-	{
-		CheckerMaterial = UMaterialInstanceDynamic::Create(CheckerMaterialBase, NULL);
-		if (CheckerMaterial != nullptr)
-		{
-			CheckerMaterial->SetScalarParameterValue("Density", Settings->CheckerDensity);
-		}
-	}
-	ActiveCheckerDensity = Settings->CheckerDensity;
+	MaterialSettings = NewObject<UExistingMeshMaterialProperties>(this);
+	MaterialSettings->Setup();
+	AddToolPropertySource(MaterialSettings);
+
 
 	DynamicMeshComponent->bExplicitShowWireframe = Settings->bWireframe;
 
@@ -252,33 +245,12 @@ void UMeshInspectorTool::OnPropertyModified(UObject* PropertySet, UProperty* Pro
 {
 	GetToolManager()->PostInvalidation();
 	DynamicMeshComponent->bExplicitShowWireframe = Settings->bWireframe;
-
-	if (CheckerMaterial != nullptr)
-	{
-		CheckerMaterial->SetScalarParameterValue("Density", Settings->CheckerDensity);
-	}
-
-	if (Settings->MaterialMode != ActiveMaterialMode)
+	MaterialSettings->UpdateMaterials();
+	MaterialSettings->SetMaterialIfChanged(DefaultMaterial, DynamicMeshComponent->GetMaterial(0),
+		[this](UMaterialInterface* Material)
 	{
-		if (Settings->MaterialMode == EInspectorMaterialMode::Checkerboard && CheckerMaterial != nullptr)
-		{
-			DynamicMeshComponent->SetMaterial(0, CheckerMaterial);
-			ActiveMaterialMode = EInspectorMaterialMode::Checkerboard;
-			return;
-		}
-		if (Settings->MaterialMode == EInspectorMaterialMode::Override && Settings->OverrideMaterial != nullptr)
-		{
-			DynamicMeshComponent->SetMaterial(0, Settings->OverrideMaterial);
-			ActiveMaterialMode = EInspectorMaterialMode::Override;
-			return;
-		}
-		// default or fallback
-		if (DefaultMaterial != nullptr)
-		{
-			DynamicMeshComponent->SetMaterial(0, DefaultMaterial);
-		}
-		ActiveMaterialMode = EInspectorMaterialMode::Default;
-	}
+		DynamicMeshComponent->SetMaterial(0, Material);
+	});
 }
 
 
diff --git a/Engine/Plugins/Experimental/MeshModelingToolset/Source/MeshModelingTools/Private/Properties/MeshMaterialProperties.cpp b/Engine/Plugins/Experimental/MeshModelingToolset/Source/MeshModelingTools/Private/Properties/MeshMaterialProperties.cpp
new file mode 100644
--- /dev/null
+++ b/Engine/Plugins/Experimental/MeshModelingToolset/Source/MeshModelingTools/Private/Properties/MeshMaterialProperties.cpp
@@ -0,0 +1,59 @@
+// Copyright 1998-2019 Epic Games, Inc. All Rights Reserved.
+
+#include "Properties/MeshMaterialProperties.h"
+
+#include "DynamicMesh3.h"
+#include "DynamicMeshAttributeSet.h"
+#include "MeshNormals.h"
+
+#include "SimpleDynamicMeshComponent.h"
+
+#include "Materials/MaterialInstanceDynamic.h"
+
+#define LOCTEXT_NAMESPACE "UMeshMaterialProperites"
+
+UNewMeshMaterialProperties::UNewMeshMaterialProperties()
+{
+	Material = CreateDefaultSubobject<UMaterialInterface>(TEXT("MATERIAL"));
+}
+
+void UExistingMeshMaterialProperties::Setup()
+{
+	UMaterial* CheckerMaterialBase = LoadObject<UMaterial>(nullptr, TEXT("/MeshModelingToolset/Materials/CheckerMaterial"));
+	if (CheckerMaterialBase != nullptr)
+	{
+		CheckerMaterial = UMaterialInstanceDynamic::Create(CheckerMaterialBase, NULL);
+		if (CheckerMaterial != nullptr)
+		{
+			CheckerMaterial->SetScalarParameterValue("Density", CheckerDensity);
+		}
+	}
+}
+
+void UExistingMeshMaterialProperties::UpdateMaterials()
+{
+	if (CheckerMaterial != nullptr)
+	{
+		CheckerMaterial->SetScalarParameterValue("Density", CheckerDensity);
+	}
+}
+
+void UExistingMeshMaterialProperties::SetMaterialIfChanged(UMaterialInterface* OriginalMaterial, UMaterialInterface* CurrentMaterial, TFunctionRef<void(UMaterialInterface* Material)> SetMaterialFn)
+{
+	UMaterialInterface *Material = OriginalMaterial;
+	if (MaterialMode == ESetMeshMaterialMode::Checkerboard && CheckerMaterial != nullptr)
+	{
+		Material = CheckerMaterial;
+	}
+	else if (MaterialMode == ESetMeshMaterialMode::Override && OverrideMaterial != nullptr)
+	{
+		Material = OverrideMaterial;
+	}
+	if (CurrentMaterial != Material)
+	{
+		SetMaterialFn(Material);
+	}
+}
+
+
+#undef LOCTEXT_NAMESPACE
diff --git a/Engine/Plugins/Experimental/MeshModelingToolset/Source/MeshModelingTools/Private/UVProjectionTool.cpp b/Engine/Plugins/Experimental/MeshModelingToolset/Source/MeshModelingTools/Private/UVProjectionTool.cpp
--- a/Engine/Plugins/Experimental/MeshModelingToolset/Source/MeshModelingTools/Private/UVProjectionTool.cpp
+++ b/Engine/Plugins/Experimental/MeshModelingToolset/Source/MeshModelingTools/Private/UVProjectionTool.cpp
@@ -108,10 +108,9 @@ void UUVProjectionTool::Setup()
 	AddToolPropertySource(BasicProperties);
 	AddToolPropertySource(AdvancedProperties);
 
-	// Set up check material
-	UMaterial* CheckerMaterialBase = LoadObject<UMaterial>(nullptr, TEXT("/MeshModelingToolset/Materials/CheckerMaterial"));
-	CheckerMaterial = UMaterialInstanceDynamic::Create(CheckerMaterialBase, NULL);
-	CheckerMaterial->SetScalarParameterValue("Density", 10);
+	MaterialSettings = NewObject<UExistingMeshMaterialProperties>(this);
+	MaterialSettings->Setup();
+	AddToolPropertySource(MaterialSettings);
 
 	// initialize the PreviewMesh+BackgroundCompute object
 	UpdateNumPreviews();
@@ -166,7 +165,7 @@ void UUVProjectionTool::UpdateNumPreviews()
 			UMeshOpPreviewWithBackgroundCompute* Preview = Previews.Add_GetRef(NewObject<UMeshOpPreviewWithBackgroundCompute>(OpFactory, "Preview"));
 			Preview->Setup(this->TargetWorld, OpFactory);
 			Preview->ConfigureMaterials(
-				ToolSetupUtil::GetDefaultMaterial(GetToolManager(), CheckerMaterial),//ComponentTargets[PreviewIdx]->GetMaterial(0)),
+				ToolSetupUtil::GetDefaultMaterial(GetToolManager(), ComponentTargets[PreviewIdx]->GetMaterial(0)),
 				ToolSetupUtil::GetDefaultWorkingMaterial(GetToolManager())
 			);
 			Preview->SetVisibility(true);
@@ -286,10 +285,16 @@ void UUVProjectionTool::PostEditChangeProperty(FPropertyChangedEvent& PropertyCh
 void UUVProjectionTool::OnPropertyModified(UObject* PropertySet, UProperty* Property)
 {
 	// if we don't know what changed, or we know checker density changed, update checker material
-	if (!Property || Property->GetFName() == GET_MEMBER_NAME_CHECKED(UUVProjectionToolProperties, CheckerDensity))
+	MaterialSettings->UpdateMaterials();
+	for (int PreviewIdx = 0; PreviewIdx < Previews.Num(); PreviewIdx++)
 	{
-		CheckerMaterial->SetScalarParameterValue("Density", BasicProperties->CheckerDensity);
+		UMeshOpPreviewWithBackgroundCompute* Preview = Previews[PreviewIdx];
+		MaterialSettings->SetMaterialIfChanged(ComponentTargets[PreviewIdx]->GetMaterial(0), Preview->StandardMaterial, [&Preview, this](UMaterialInterface* Material)
+		{
+			Preview->ConfigureMaterials(ToolSetupUtil::GetDefaultMaterial(GetToolManager(), Material), Preview->WorkingMaterial);
+		});
 	}
+	
 	UpdateNumPreviews();
 	for (UMeshOpPreviewWithBackgroundCompute* Preview : Previews)
 	{
diff --git a/Engine/Plugins/Experimental/MeshModelingToolset/Source/MeshModelingTools/Public/AddPrimitiveTool.h b/Engine/Plugins/Experimental/MeshModelingToolset/Source/MeshModelingTools/Public/AddPrimitiveTool.h
--- a/Engine/Plugins/Experimental/MeshModelingToolset/Source/MeshModelingTools/Public/AddPrimitiveTool.h
+++ b/Engine/Plugins/Experimental/MeshModelingToolset/Source/MeshModelingTools/Public/AddPrimitiveTool.h
@@ -8,6 +8,7 @@
 #include "InteractiveToolBuilder.h"
 #include "BaseTools/SingleClickTool.h"
 #include "PreviewMesh.h"
+#include "Properties/MeshMaterialProperties.h"
 #include "AddPrimitiveTool.generated.h"
 
 class FDynamicMesh3;
@@ -116,20 +117,6 @@ public:
 	/** Subdivisions */
 	UPROPERTY(EditAnywhere, Category = ShapeSettings, meta = (DisplayName = "Subdivisions", UIMin = "0", UIMax = "100", ClampMin = "0", ClampMax = "4000"))
 	int Subdivisions;
-
-	/** Material */
-	UPROPERTY(EditAnywhere, Category = ShapeSettings, meta = (DisplayName = "Material"))
-	UMaterialInterface* Material;
-
-	UPROPERTY(EditAnywhere, Category = ShapeSettings, meta = (DisplayName = "UV Scale"))
-	float UVScale = 1.0;
-	
-	UPROPERTY(EditAnywhere, Category = ShapeSettings, meta = (DisplayName = "UV Scale Relative to World Space"))
-	bool bWorldSpaceUVScale = false;
-
-	/** Wireframe toggle */
-	UPROPERTY(EditAnywhere, Category = ShapeSettings, meta = (DisplayName = "Show Wireframe"))
-	bool bWireframe;
 };
 
 
@@ -175,6 +162,11 @@ protected:
 	UPROPERTY()
 	UProceduralShapeToolProperties* ShapeSettings;
 
+	UPROPERTY()
+	UNewMeshMaterialProperties* MaterialProperties;
+
+
+
 	UPROPERTY()
 	UPreviewMesh* PreviewMesh;
 
diff --git a/Engine/Plugins/Experimental/MeshModelingToolset/Source/MeshModelingTools/Public/DrawPolygonTool.h b/Engine/Plugins/Experimental/MeshModelingToolset/Source/MeshModelingTools/Public/DrawPolygonTool.h
--- a/Engine/Plugins/Experimental/MeshModelingToolset/Source/MeshModelingTools/Public/DrawPolygonTool.h
+++ b/Engine/Plugins/Experimental/MeshModelingToolset/Source/MeshModelingTools/Public/DrawPolygonTool.h
@@ -11,6 +11,7 @@
 #include "DynamicMeshAABBTree3.h"
 #include "Snapping/PointPlanarSnapSolver.h"
 #include "ToolSceneQueriesUtil.h"
+#include "Properties/MeshMaterialProperties.h"
 #include "DrawPolygonTool.generated.h"
 
 
@@ -97,16 +98,6 @@ class MESHMODELINGTOOLS_API UDrawPolygonToolStandardProperties : public UObject
 
 	UPROPERTY(EditAnywhere, NonTransactional, Category = Polygon)
 	bool bAllowSelfIntersections = false;
-
-	/** Material for new mesh*/
-	UPROPERTY(EditAnywhere, NonTransactional, Category = Polygon, meta = (DisplayName = "Material"))
-	UMaterialInterface* Material;
-
-	UPROPERTY(EditAnywhere, Category = ShapeSettings, meta = (DisplayName = "UV Scale"))
-	float UVScale = 1.0;
-
-	UPROPERTY(EditAnywhere, Category = ShapeSettings, meta = (DisplayName = "UV Scale Relative to World Space"))
-	bool bWorldSpaceUVScale = false;
 };
 
 
@@ -213,6 +204,9 @@ protected:
 
 	UPROPERTY()
 	UDrawPolygonToolSnapProperties* SnapProperties;
+
+	UPROPERTY()
+	UNewMeshMaterialProperties* MaterialProperties;
 	
 
 	/** Origin of plane we will draw polygon on */
diff --git a/Engine/Plugins/Experimental/MeshModelingToolset/Source/MeshModelingTools/Public/MeshInspectorTool.h b/Engine/Plugins/Experimental/MeshModelingToolset/Source/MeshModelingTools/Public/MeshInspectorTool.h
--- a/Engine/Plugins/Experimental/MeshModelingToolset/Source/MeshModelingTools/Public/MeshInspectorTool.h
+++ b/Engine/Plugins/Experimental/MeshModelingToolset/Source/MeshModelingTools/Public/MeshInspectorTool.h
@@ -8,6 +8,7 @@
 #include "InteractiveToolBuilder.h"
 #include "DynamicMesh3.h"
 #include "DynamicMeshAABBTree3.h"
+#include "Properties/MeshMaterialProperties.h"
 #include "Properties/MeshStatisticsProperties.h"
 #include "MeshInspectorTool.generated.h"
 
@@ -32,24 +33,6 @@ public:
 
 
 
-/** Material modes for MeshInspectorTool */
-UENUM()
-enum class EInspectorMaterialMode : uint8
-{
-	/** Input material */
-	Default,
-
-	/** Checkerboard material */
-	Checkerboard,
-
-	/** Override material */
-	Override
-};
-
-
-
-
-
 UCLASS()
 class MESHMODELINGTOOLS_API UMeshInspectorProperties : public UInteractiveToolPropertySet
 {
@@ -93,15 +76,6 @@ public:
 	float TangentLength = 5.0f;
 
 
-	/** Material that will be used on the mesh */
-	UPROPERTY(EditAnywhere, Category = Options)
-	EInspectorMaterialMode MaterialMode = EInspectorMaterialMode::Default;
-
-	UPROPERTY(EditAnywhere, Category = Options, meta = (UIMin = "1.0", UIMax = "40.0", ClampMin = "0.01", ClampMax = "1000.0", EditCondition = "MaterialMode == EInspectorMaterialMode::Checkerboard"))
-	float CheckerDensity = 20.0f;
-
-	UPROPERTY(EditAnywhere, Category = Options, meta = (EditCondition = "MaterialMode == EInspectorMaterialMode::Override"))
-	UMaterialInterface* OverrideMaterial = nullptr;
 };
 
 
@@ -142,20 +116,18 @@ protected:
 	UMeshInspectorProperties* Settings = nullptr;
 
 	UPROPERTY()
-	UMaterialInterface* DefaultMaterial = nullptr;
-
-	UPROPERTY()
-	UMaterialInstanceDynamic* CheckerMaterial = nullptr;
+	UExistingMeshMaterialProperties* MaterialSettings = nullptr;
 
 
 	float LineWidthMultiplier = 1.0f;
-	EInspectorMaterialMode ActiveMaterialMode;
-	float ActiveCheckerDensity = 0.0f;
 
 protected:
 	UPROPERTY()
 	USimpleDynamicMeshComponent* DynamicMeshComponent;
 
+	UPROPERTY()
+	UMaterialInterface* DefaultMaterial = nullptr;
+
 	TArray<int> BoundaryEdges;
 	TArray<int> UVSeamEdges;
 	TArray<int> NormalSeamEdges;
diff --git a/Engine/Plugins/Experimental/MeshModelingToolset/Source/MeshModelingTools/Public/Properties/MeshMaterialProperties.h b/Engine/Plugins/Experimental/MeshModelingToolset/Source/MeshModelingTools/Public/Properties/MeshMaterialProperties.h
new file mode 100644
--- /dev/null
+++ b/Engine/Plugins/Experimental/MeshModelingToolset/Source/MeshModelingTools/Public/Properties/MeshMaterialProperties.h
@@ -0,0 +1,84 @@
+// Copyright 1998-2019 Epic Games, Inc. All Rights Reserved.
+
+#pragma once
+
+#include "CoreMinimal.h"
+
+#include "InteractiveToolBuilder.h"
+
+#include "MeshMaterialProperties.generated.h"
+
+
+// predeclarations
+class FDynamicMesh3;
+class UMaterialInterface;
+
+
+
+// Standard material property settings for tools that generate new meshes
+UCLASS()
+class MESHMODELINGTOOLS_API UNewMeshMaterialProperties : public UInteractiveToolPropertySet
+{
+	GENERATED_BODY()
+public:
+	UNewMeshMaterialProperties();
+
+	/** Material for new mesh*/
+	UPROPERTY(EditAnywhere, NonTransactional, Category = Material, meta = (DisplayName = "Material"))
+	UMaterialInterface* Material;
+
+	/** Scale factor for generated UVs */
+	UPROPERTY(EditAnywhere, Category = Material, meta = (DisplayName = "UV Scale"))
+	float UVScale = 1.0;
+
+	/** If set, UV scales will be relative to world space so different objects created with the same UV scale should have the same average texel size */
+	UPROPERTY(EditAnywhere, Category = Material, meta = (DisplayName = "UV Scale Relative to World Space"))
+	bool bWorldSpaceUVScale = false;
+
+	/** Overlay wireframe on preview */
+	UPROPERTY(EditAnywhere, Category = Material, meta = (DisplayName = "Show Wireframe"))
+	bool bWireframe = false;
+};
+
+
+
+/** Standard material modes for tools that need to set custom materials for visualization */
+UENUM()
+enum class ESetMeshMaterialMode : uint8
+{
+	/** Input material */
+	KeepOriginal,
+
+	/** Checkerboard material */
+	Checkerboard,
+
+	/** Override material */
+	Override
+};
+
+// Standard material property settings for tools that visualize materials on existing meshes (e.g. to help show UVs)
+UCLASS()
+class MESHMODELINGTOOLS_API UExistingMeshMaterialProperties : public UInteractiveToolPropertySet
+{
+	GENERATED_BODY()
+public:
+
+	/** Material that will be used on the mesh */
+	UPROPERTY(EditAnywhere, Category = Material)
+	ESetMeshMaterialMode MaterialMode = ESetMeshMaterialMode::KeepOriginal;
+
+	UPROPERTY(EditAnywhere, Category = Material, meta = (UIMin = "1.0", UIMax = "40.0", ClampMin = "0.01", ClampMax = "1000.0", EditCondition = "MaterialMode == ESetMeshMaterialMode::Checkerboard"))
+	float CheckerDensity = 20.0f;
+
+	UPROPERTY(EditAnywhere, Category = Material, meta = (EditCondition = "MaterialMode == ESetMeshMaterialMode::Override"))
+	UMaterialInterface* OverrideMaterial = nullptr;
+
+	UPROPERTY()
+	UMaterialInstanceDynamic* CheckerMaterial = nullptr;
+
+	void Setup();
+
+	void UpdateMaterials();
+	void SetMaterialIfChanged(UMaterialInterface* OriginalMaterial, UMaterialInterface* CurrentMaterial, TFunctionRef<void(UMaterialInterface* Material)> SetMaterialFn);
+};
+
diff --git a/Engine/Plugins/Experimental/MeshModelingToolset/Source/MeshModelingTools/Public/UVProjectionTool.h b/Engine/Plugins/Experimental/MeshModelingToolset/Source/MeshModelingTools/Public/UVProjectionTool.h
--- a/Engine/Plugins/Experimental/MeshModelingToolset/Source/MeshModelingTools/Public/UVProjectionTool.h
+++ b/Engine/Plugins/Experimental/MeshModelingToolset/Source/MeshModelingTools/Public/UVProjectionTool.h
@@ -11,6 +11,7 @@
 #include "ParameterizationOps/UVProjectionOp.h"
 #include "DynamicMesh3.h"
 #include "BaseTools/SingleClickTool.h"
+#include "Properties/MeshMaterialProperties.h"
 
 #include "UVProjectionTool.generated.h"
 
@@ -74,10 +75,6 @@ public:
 	UPROPERTY(EditAnywhere, Category = Options)
 	FVector2D UVOffset;
 
-	/** For display only - the density of the checkerboard pattern in the viewport.*/
-	UPROPERTY(EditAnywhere, Category = Options, meta = (Default = "20", UIMin = "1.0", UIMax = "40.0", ClampMin = "0.01", ClampMax = "1000.0"))
-	float CheckerDensity = 10.f;
-
 };
 
 
@@ -157,6 +154,10 @@ protected:
 	UPROPERTY()
 	UUVProjectionAdvancedProperties* AdvancedProperties;
 
+	UPROPERTY()
+	UExistingMeshMaterialProperties* MaterialSettings = nullptr;
+
+
 
 	UPROPERTY()
 	TArray<UMeshOpPreviewWithBackgroundCompute*> Previews;
diff --git a/Engine/Plugins/Experimental/MeshModelingToolset/Source/MeshModelingToolsEditorOnly/Private/ParameterizeMeshTool.cpp b/Engine/Plugins/Experimental/MeshModelingToolset/Source/MeshModelingToolsEditorOnly/Private/ParameterizeMeshTool.cpp
--- a/Engine/Plugins/Experimental/MeshModelingToolset/Source/MeshModelingToolsEditorOnly/Private/ParameterizeMeshTool.cpp
+++ b/Engine/Plugins/Experimental/MeshModelingToolset/Source/MeshModelingToolsEditorOnly/Private/ParameterizeMeshTool.cpp
@@ -73,14 +73,7 @@ void UParameterizeMeshTool::Setup()
 	InputMesh = MakeShared<FMeshDescription>(*ComponentTarget->GetMesh());
 
 
-	// Set up check material
-
-	UMaterial* CheckerMaterialBase = LoadObject<UMaterial>(nullptr, TEXT("/MeshModelingToolset/Materials/CheckerMaterial"));
-	CheckerMaterial = UMaterialInstanceDynamic::Create(CheckerMaterialBase, NULL);
-	CheckerMaterial->SetScalarParameterValue("Density", CheckerDensity);
-
-	// Copy existing material if there is one
-	
+	// Copy existing material if there is one	
 	DefaultMaterial = ComponentTarget->GetMaterial(0);
 	if (DefaultMaterial == nullptr)
 	{
@@ -115,20 +108,28 @@ void UParameterizeMeshTool::Setup()
 		ToolSetupUtil::GetDefaultWorkingMaterial(GetToolManager())
 	);
 
-	// Choose the material for display
-	if (MaterialMode == EParameterizeMeshMaterialMode::Checkerboard)
-	{
-		Preview->ConfigureMaterials(
-			ToolSetupUtil::GetDefaultMaterial(GetToolManager(), CheckerMaterial),
-			ToolSetupUtil::GetDefaultWorkingMaterial(GetToolManager())
-		);
-	}
+
+	MaterialSettings = NewObject<UExistingMeshMaterialProperties>(this);
+	MaterialSettings->Setup();
+	AddToolPropertySource(MaterialSettings);
 
 
 	Preview->SetVisibility(true);
 	Preview->InvalidateResult();    // start compute
 }
 
+void UParameterizeMeshTool::OnPropertyModified(UObject* PropertySet, UProperty* Property)
+{
+	if (PropertySet == MaterialSettings)
+	{
+		MaterialSettings->UpdateMaterials();
+		MaterialSettings->SetMaterialIfChanged(ComponentTarget->GetMaterial(0), Preview->StandardMaterial, [this](UMaterialInterface* Material)
+		{
+			Preview->ConfigureMaterials(ToolSetupUtil::GetDefaultMaterial(GetToolManager(), Material), Preview->WorkingMaterial);
+		});
+	}
+}
+
 
 void UParameterizeMeshTool::Shutdown(EToolShutdownType ShutdownType)
 {
@@ -174,56 +175,8 @@ bool UParameterizeMeshTool::CanAccept() const
 
 void UParameterizeMeshTool::PostEditChangeProperty(FPropertyChangedEvent& PropertyChangedEvent)
 {
-	UProperty* PropertyThatChanged = PropertyChangedEvent.Property;
-
-	if (PropertyThatChanged->GetFName() == GET_MEMBER_NAME_CHECKED(UParameterizeMeshTool, CheckerDensity))
-	{
-		CheckerMaterial->SetScalarParameterValue("Density", CheckerDensity);
-	}
-	else if (PropertyThatChanged->GetFName() == GET_MEMBER_NAME_CHECKED(UParameterizeMeshTool, MaterialMode))
-	{
-		if (MaterialMode == EParameterizeMeshMaterialMode::Checkerboard && CheckerMaterial != nullptr)
-		{
-			Preview->ConfigureMaterials(
-				ToolSetupUtil::GetDefaultMaterial(GetToolManager(), CheckerMaterial),
-				ToolSetupUtil::GetDefaultWorkingMaterial(GetToolManager())
-			);
-		}
-		if (MaterialMode == EParameterizeMeshMaterialMode::Override && OverrideMaterial != nullptr)
-		{
-			Preview->ConfigureMaterials(
-				ToolSetupUtil::GetDefaultMaterial(GetToolManager(), OverrideMaterial),
-				ToolSetupUtil::GetDefaultWorkingMaterial(GetToolManager())
-			);
-		}
-		if (MaterialMode == EParameterizeMeshMaterialMode::Default && DefaultMaterial != nullptr)
-		{
-		
-			Preview->ConfigureMaterials(
-				ToolSetupUtil::GetDefaultMaterial(GetToolManager(), DefaultMaterial),
-				ToolSetupUtil::GetDefaultWorkingMaterial(GetToolManager())
-			);
-		}
-
-		
-
-	}
-	else if (PropertyThatChanged->GetFName() == GET_MEMBER_NAME_CHECKED(UParameterizeMeshTool, OverrideMaterial))
-	{
-		if (MaterialMode == EParameterizeMeshMaterialMode::Override && OverrideMaterial != nullptr)
-		{
-			
-			Preview->ConfigureMaterials(
-				ToolSetupUtil::GetDefaultMaterial(GetToolManager(), OverrideMaterial),
-				ToolSetupUtil::GetDefaultWorkingMaterial(GetToolManager())
-			);
-		}
-	}
-	else
-	{
-		// One of the UV generation properties must have changed.  Dirty the result to force a recompute
-		Preview->InvalidateResult();
-	}
+	// One of the UV generation properties must have changed.  Dirty the result to force a recompute
+	Preview->InvalidateResult();
 }
 
 TSharedPtr<FDynamicMeshOperator> UParameterizeMeshTool::MakeNewOperator()
diff --git a/Engine/Plugins/Experimental/MeshModelingToolset/Source/MeshModelingToolsEditorOnly/Public/ParameterizeMeshTool.h b/Engine/Plugins/Experimental/MeshModelingToolset/Source/MeshModelingToolsEditorOnly/Public/ParameterizeMeshTool.h
--- a/Engine/Plugins/Experimental/MeshModelingToolset/Source/MeshModelingToolsEditorOnly/Public/ParameterizeMeshTool.h
+++ b/Engine/Plugins/Experimental/MeshModelingToolset/Source/MeshModelingToolsEditorOnly/Public/ParameterizeMeshTool.h
@@ -8,6 +8,8 @@
 #include "InteractiveToolBuilder.h"
 #include "DynamicMesh3.h"
 #include "MeshOpPreviewHelpers.h"
+#include "Properties/MeshMaterialProperties.h"
+
 #include "ParameterizeMeshTool.generated.h"
 
 
@@ -86,6 +88,8 @@ public:
 	virtual void PostEditChangeProperty(FPropertyChangedEvent& PropertyChangedEvent) override;
 #endif
 
+	virtual void OnPropertyModified(UObject* PropertySet, UProperty* Property) override;
+
 	// IDynamicMeshOperatorFactory API
 	virtual TSharedPtr<FDynamicMeshOperator> MakeNewOperator() override;
 
@@ -96,18 +100,9 @@ protected:
 	/** Maximum amount of stretch, from none to any.  If zero stretch is specified each triangle will likey be its own chart */
 	UPROPERTY(EditAnywhere, Category = Options, meta = (Default = "0.166", UIMin = "0", UIMax = "1", ClampMin = "0", ClampMax= "1"))
 	float ChartStretch = 0.11f;
-	
-	/** Material that will be displayed on the mesh */
-	UPROPERTY(EditAnywhere, Category = Options)
-	EParameterizeMeshMaterialMode MaterialMode = EParameterizeMeshMaterialMode::Default;
 
-	/** The material that will be displayed on the mesh when override material mode has been selected */
-	UPROPERTY(EditAnywhere, Category = Options)
-	UMaterialInterface* OverrideMaterial = nullptr;
-	
-	/** For display only - the density of the checkerboard pattern in the viewport.*/
-	UPROPERTY(EditAnywhere, Category = Options, meta = (Default = "20", UIMin = "1.0", UIMax = "40.0", ClampMin = "0.01", ClampMax = "1000.0"))
-	float CheckerDensity = 20.f;
+	UPROPERTY()
+	UExistingMeshMaterialProperties* MaterialSettings = nullptr;
 
 protected:
 	UPROPERTY()
