Mesh Modeling : Simplify PolyEdit UI by removing some of the 'advaced' options from the display.
#rnx

diff --git a/Engine/Plugins/Experimental/MeshModelingToolset/Source/MeshModelingTools/Public/EditMeshPolygonsTool.h b/Engine/Plugins/Experimental/MeshModelingToolset/Source/MeshModelingTools/Public/EditMeshPolygonsTool.h
--- a/Engine/Plugins/Experimental/MeshModelingToolset/Source/MeshModelingTools/Public/EditMeshPolygonsTool.h
+++ b/Engine/Plugins/Experimental/MeshModelingToolset/Source/MeshModelingTools/Public/EditMeshPolygonsTool.h
@@ -135,15 +135,15 @@ public:
 	UPROPERTY(EditAnywhere, Category = Options, meta = (EditCondition="PolygonMode == EPolygonGroupMode::RecomputePolygonsByAngleThreshold"))
 	float PolygonGroupingAngleThreshold;
 
-	//Laplacian Deformation Options
+	//Laplacian Deformation Options, currently not exposed.
 
-	UPROPERTY(EditAnywhere, Category = LaplacianOptions, meta = (EditCondition="DeformationStrategy == EGroupTopologyDeformationStrategy::Laplacian", DisplayName = "Laplacian Weight Scheme", ToolTip = "Laplacian weight schemes determine how the deformer will interpret the curvature at a given vertex in relation to nearby vertices in order to generate a solution."))
+	UPROPERTY() 
 	EWeightScheme SelectedWeightScheme = EWeightScheme::ClampedCotangent;
 
-	UPROPERTY(EditAnywhere, Category = LaplacianOptions, meta = (EditCondition="DeformationStrategy == EGroupTopologyDeformationStrategy::Laplacian", DisplayName = "Handle Weight", UIMin = "0.0001", UIMax = "100.0", ClampMin = "0.0001", ClampMax = "1000.0", ToolTip = "The handle weight describes how effective dragging the selected PolyGroup will be."))
-	double HandleWeight = 100.0;
+	UPROPERTY() 
+	double HandleWeight = 1000.0;
 
-	UPROPERTY(EditAnywhere, Category = LaplacianOptions, meta = (EditCondition="DeformationStrategy == EGroupTopologyDeformationStrategy::Laplacian", DisplayName = "Post-fix Handles", ToolTip = "Post-fixing the handles will attempt to force their vertices to arrive exactly where you wanted them. This may result\nin undesirable or choppy curvature in some scenarios."))
+	UPROPERTY()
 	bool bPostFixHandles = false;
 /**
 	// How to add a weight curve
