fix the UV orientations and scales and such for uv projection and primitive generators

#rnx

diff --git a/Engine/Plugins/Experimental/GeometryProcessing/Source/DynamicMesh/Private/DynamicMeshEditor.cpp b/Engine/Plugins/Experimental/GeometryProcessing/Source/DynamicMesh/Private/DynamicMeshEditor.cpp
--- a/Engine/Plugins/Experimental/GeometryProcessing/Source/DynamicMesh/Private/DynamicMeshEditor.cpp
+++ b/Engine/Plugins/Experimental/GeometryProcessing/Source/DynamicMesh/Private/DynamicMeshEditor.cpp
@@ -576,7 +576,7 @@ void FDynamicMeshEditor::SetQuadUVsFromProjection(const FIndex2i& QuadTris, cons
 }
 
 
-void FDynamicMeshEditor::RescaleAttributeUVs(float UVScale, bool bWorldSpace, int UVLayerIndex)
+void FDynamicMeshEditor::RescaleAttributeUVs(float UVScale, bool bWorldSpace, int UVLayerIndex, TOptional<FTransform3d> ToWorld)
 {
 	check(Mesh->HasAttributes() && Mesh->Attributes()->NumUVLayers() > UVLayerIndex );
 	FDynamicMeshUVOverlay* UVs = Mesh->Attributes()->GetUVLayer(UVLayerIndex);
@@ -591,6 +591,13 @@ void FDynamicMeshEditor::RescaleAttributeUVs(float UVScale, bool bWorldSpace, in
 		{
 			UVs->GetTriElements(TID, TriUVs[0], TriUVs[1], TriUVs[2]);
 			Mesh->GetTriVertices(TID, TriVs[0], TriVs[1], TriVs[2]);
+			if (ToWorld.IsSet())
+			{
+				for (int i = 0; i < 3; i++)
+				{
+					TriVs[i] = ToWorld->TransformPosition(TriVs[i]);
+				}
+			}
 			for (int j = 2, i = 0; i < 3; j = i++)
 			{
 				TotalEdgeUVLen += TriUVs[j].Distance(TriUVs[i]);
diff --git a/Engine/Plugins/Experimental/GeometryProcessing/Source/DynamicMesh/Public/DynamicMeshEditor.h b/Engine/Plugins/Experimental/GeometryProcessing/Source/DynamicMesh/Public/DynamicMeshEditor.h
--- a/Engine/Plugins/Experimental/GeometryProcessing/Source/DynamicMesh/Public/DynamicMeshEditor.h
+++ b/Engine/Plugins/Experimental/GeometryProcessing/Source/DynamicMesh/Public/DynamicMeshEditor.h
@@ -288,8 +288,9 @@ public:
 	 * @param UVScale Scale factor to multiply into UVs.  If in world space, this is in centimeters relative to the average UV scale
 	 * @param bWorldSpace If true, UVs are rescaled relative to an absolute world scale.
 	 * @param UVLayerIndex which UV layer to operate on (must exist)
+	 * @param ToWorld Optionally transform vertices for world space scaling
 	 */
-	void RescaleAttributeUVs(float UVScale = 1.0f, bool bWorldSpace = false, int UVLayerIndex = 0);
+	void RescaleAttributeUVs(float UVScale = 1.0f, bool bWorldSpace = false, int UVLayerIndex = 0, TOptional<FTransform3d> ToWorld = TOptional<FTransform3d>());
 
 
 	//////////////////////////////////////////////////////////////////////////
diff --git a/Engine/Plugins/Experimental/GeometryProcessing/Source/DynamicMesh/Public/Generators/GridBoxMeshGenerator.h b/Engine/Plugins/Experimental/GeometryProcessing/Source/DynamicMesh/Public/Generators/GridBoxMeshGenerator.h
--- a/Engine/Plugins/Experimental/GeometryProcessing/Source/DynamicMesh/Public/Generators/GridBoxMeshGenerator.h
+++ b/Engine/Plugins/Experimental/GeometryProcessing/Source/DynamicMesh/Public/Generators/GridBoxMeshGenerator.h
@@ -165,6 +165,12 @@ public:
 		for (int Dim = 0; Dim < 3; Dim++)
 		{
 			int FaceIdxBase = FaceDimOrder[Dim]*2;
+
+			// UV-specific minor axes + flips; manually set to match default ue4 cube texture arrangement
+			int Minor1Flip[3] = { -1, 1, 1 };
+			int Minor2Flip[3] = { -1, -1, 1 };
+
+			// UV scales for D0, D1
 			double FaceWidth = FMathd::Abs(Box.Extents[D[0][Dim]]) * 2.0;
 			double FaceHeight = FMathd::Abs(Box.Extents[D[1][Dim]]) * 2.0;
 			double WidthUVScale = FaceWidth * UVScale;
@@ -174,6 +180,7 @@ public:
 			for (int Side = 0; Side < 2; Side++)
 			{
 				int SideOpp = 1 - Side;
+				float SideSign = Side * 2 - 1;
 
 				FVector3f Normal(0, 0, 0);
 				Normal[Dim] = (2 * Side - 1);
@@ -181,11 +188,23 @@ public:
 
 				int FaceUVStartInd = CurrentUVIdx;
 				// set all the UVs and normals
+				FVector2f UV;
+				int UVXDim = Dim == 1 ? 1 : 0;	// which dim (of D0,D1) follows the horizontal UV coordinate
+				int UVYDim = 1 - UVXDim;		// which dim (of D0,D1) follows the vertical UV coordinate
 				for (int D0 = 0; D0 < N[D[0][Dim]]; D0++)
 				{
 					for (int D1 = 0; D1 < N[D[1][Dim]]; D1++)
 					{
-						UVs[CurrentUVIdx] = FVector2f(D0*WidthUVScale*Nscale[D[0][Dim]], D1*HeightUVScale*Nscale[D[1][Dim]]);
+						// put the grid coordinates (centered at 0,0) into the UVs
+						UV[UVXDim] = D0 * Nscale[D[0][Dim]] - .5;
+						UV[UVYDim] = D1 * Nscale[D[1][Dim]] - .5;
+						// invert axes to match the desired UV patterns & so the opp faces are not backwards
+						UV.X *= SideSign * Minor1Flip[Dim];
+						UV.Y *= Minor2Flip[Dim];
+						// recenter and scale up
+						UV[UVXDim] = (UV[UVXDim] + .5) * WidthUVScale;
+						UV[UVYDim] = (UV[UVYDim] + .5) * HeightUVScale;
+						UVs[CurrentUVIdx] = UV;
 						Normals[CurrentUVIdx] = Normal;
 						UVParentVertex[CurrentUVIdx] = FaceVertIndices[MajorFaceInd][ToFaceV(Dim, D0, D1)];
 						NormalParentVertex[CurrentUVIdx] = FaceVertIndices[MajorFaceInd][ToFaceV(Dim, D0, D1)];
diff --git a/Engine/Plugins/Experimental/GeometryProcessing/Source/DynamicMesh/Public/Generators/SphereGenerator.h b/Engine/Plugins/Experimental/GeometryProcessing/Source/DynamicMesh/Public/Generators/SphereGenerator.h
--- a/Engine/Plugins/Experimental/GeometryProcessing/Source/DynamicMesh/Public/Generators/SphereGenerator.h
+++ b/Engine/Plugins/Experimental/GeometryProcessing/Source/DynamicMesh/Public/Generators/SphereGenerator.h
@@ -129,7 +129,7 @@ public:
 
 			for (int32 t = 0; t <= NumTheta; ++t)
 			{
-				float Theta = t * DUVtheta;
+				float Theta = 1 - t * DUVtheta;
 				UVs[UVIdx] = FVector2f(Theta, Phi);
 				UVParentVertex[UVIdx] = GridInterface(t, p);
 				ensure(UVIdxLookup(t, p) == UVIdx);
diff --git a/Engine/Plugins/Experimental/GeometryProcessing/Source/DynamicMesh/Public/Generators/SweepGenerator.h b/Engine/Plugins/Experimental/GeometryProcessing/Source/DynamicMesh/Public/Generators/SweepGenerator.h
--- a/Engine/Plugins/Experimental/GeometryProcessing/Source/DynamicMesh/Public/Generators/SweepGenerator.h
+++ b/Engine/Plugins/Experimental/GeometryProcessing/Source/DynamicMesh/Public/Generators/SweepGenerator.h
@@ -125,10 +125,11 @@ protected:
 					SetTrianglePolygon(TriIdx, PolyIdx);
 					TriIdx++;
 				}
+				float SideScale = 2 * CapIdx - 1;
 				for (int32 Idx = 0; Idx < XVerts; Idx++)
 				{
 					FVector2d CenteredVert = CrossSection.GetVertices()[Idx] * UVScale + UVOffset;
-					SetUV(CapUVStart[CapIdx] + Idx, FVector2f(CenteredVert.X, CenteredVert.Y), VertOffset + Idx);
+					SetUV(CapUVStart[CapIdx] + Idx, FVector2f(CenteredVert.X * SideScale, CenteredVert.Y), VertOffset + Idx);
 					SetNormal(CapNormalStart[CapIdx] + Idx, FVector3f::Zero(), VertOffset + Idx);
 				}
 			}
@@ -147,7 +148,7 @@ protected:
 				for (int32 XIdx = 0; XIdx < NumCrossSections; XIdx++)
 				{
 					float UVY = XIdx / float(NumCrossSections - 1);
-					SetUV(XIdx * XUVs + UVSubIdx, FVector2f(UVX, UVY), XIdx * XVerts + VertSubIdx);
+					SetUV(XIdx * XUVs + UVSubIdx, FVector2f(1-UVX, 1-UVY), XIdx * XVerts + VertSubIdx);
 				}
 
 				if (VertSubIdx == NextDupVertIdx)
@@ -179,7 +180,7 @@ protected:
 				for (int32 XIdx = 0; XIdx < NumCrossSections; XIdx++)
 				{
 					float UVY = XIdx / float(NumCrossSections - 1);
-					SetUV(XIdx * XUVs + UVSubIdx, FVector2f(UVX, UVY), XIdx * XVerts + VertSubIdx);
+					SetUV(XIdx * XUVs + UVSubIdx, FVector2f(1-UVX, 1-UVY), XIdx * XVerts + VertSubIdx);
 				}
 			}
 			NumSections = NormalSections.Num();
@@ -246,6 +247,7 @@ public:
 	float Height = 1.0f;
 	int AngleSamples = 16;
 	bool bCapped = false;
+	bool bUVScaleMatchSidesAndCaps = true;
 
 public:
 	/** Generate the mesh */
@@ -262,7 +264,7 @@ public:
 			Caps[0] = ECapType::FlatTriangulation;
 			Caps[1] = ECapType::FlatTriangulation;
 		}
-		ConstructMeshTopology(X, {}, {}, 2, Caps, FVector2d(1, 1), FVector2d(0, 0));
+		ConstructMeshTopology(X, {}, {}, 2, Caps, FVector2d(.5, .5), FVector2d(.5, .5));
 
 		for (int SubIdx = 0; SubIdx < X.VertexCount(); SubIdx++)
 		{
@@ -284,6 +286,32 @@ public:
 			Normals[k].Normalize();
 		}
 
+		if (bUVScaleMatchSidesAndCaps)
+		{
+			float MaxAbsRad = FMathf::Max(FMathf::Abs(Radius[0]), FMathf::Abs(Radius[1]));
+			float AbsHeight = FMathf::Abs(Height);
+			float MaxAbsCircumference = MaxAbsRad * FMathf::TwoPi;
+			
+			// scales to put each differently-scaled UV coordinate into the same space
+			float ThetaScale = MaxAbsCircumference;
+			float HeightScale = AbsHeight;
+			float CapScale = MaxAbsRad*2;
+
+			float MaxScale = FMathf::Max3(ThetaScale, HeightScale, CapScale);
+			ThetaScale /= MaxScale;
+			HeightScale /= MaxScale;
+			CapScale /= MaxScale;
+			for (int UVIdx = 0; UVIdx < CapUVStart[0]; UVIdx++)
+			{
+				UVs[UVIdx].X *= ThetaScale;
+				UVs[UVIdx].Y *= HeightScale;
+			}
+			for (int UVIdx = CapUVStart[0]; UVIdx < UVs.Num(); UVIdx++)
+			{
+				UVs[UVIdx] *= CapScale;
+			}
+		}
+
 		return *this;
 	}
 };
@@ -321,7 +349,7 @@ public:
 			Caps[1] = ECapType::FlatTriangulation;
 		}
 		int PathNum = Path.Num();
-		ConstructMeshTopology(CrossSection, {}, {}, PathNum, Caps, FVector2d(1, 1), FVector2d(0, 0));
+		ConstructMeshTopology(CrossSection, {}, {}, PathNum, Caps, FVector2d(.5, .5), FVector2d(.5, .5));
 
 		int XNum = CrossSection.VertexCount();
 		TArray<FVector2d> XNormals; XNormals.SetNum(XNum);
diff --git a/Engine/Plugins/Experimental/MeshModelingToolset/Source/MeshModelingTools/Private/UVProjectionTool.cpp b/Engine/Plugins/Experimental/MeshModelingToolset/Source/MeshModelingTools/Private/UVProjectionTool.cpp
--- a/Engine/Plugins/Experimental/MeshModelingToolset/Source/MeshModelingTools/Private/UVProjectionTool.cpp
+++ b/Engine/Plugins/Experimental/MeshModelingToolset/Source/MeshModelingTools/Private/UVProjectionTool.cpp
@@ -223,6 +223,7 @@ TSharedPtr<FDynamicMeshOperator> UUVProjectionOperatorFactory::MakeNewOperator()
 	Op->CylinderProjectToTopOrBottomAngleThreshold = Tool->BasicProperties->CylinderProjectToTopOrBottomAngleThreshold;
 	Op->UVScale = Tool->BasicProperties->UVScale;
 	Op->UVOffset = Tool->BasicProperties->UVOffset;
+	Op->bWorldSpaceUVScale = Tool->BasicProperties->bWorldSpaceUVScale;
 
 	FTransform LocalToWorld = Tool->ComponentTargets[ComponentIndex]->GetWorldTransform();
 	Op->OriginalMesh = Tool->OriginalDynamicMeshes[ComponentIndex];
diff --git a/Engine/Plugins/Experimental/MeshModelingToolset/Source/MeshModelingTools/Public/UVProjectionTool.h b/Engine/Plugins/Experimental/MeshModelingToolset/Source/MeshModelingTools/Public/UVProjectionTool.h
--- a/Engine/Plugins/Experimental/MeshModelingToolset/Source/MeshModelingTools/Public/UVProjectionTool.h
+++ b/Engine/Plugins/Experimental/MeshModelingToolset/Source/MeshModelingTools/Public/UVProjectionTool.h
@@ -75,6 +75,10 @@ public:
 	UPROPERTY(EditAnywhere, Category = Options)
 	FVector2D UVOffset;
 
+	/** If set, UV scales will be relative to world space so different objects created with the same UV scale should have the same average texel size */
+	UPROPERTY(EditAnywhere, Category = Options, meta = (DisplayName = "UV Scale Relative to World Space"))
+	bool bWorldSpaceUVScale = false;
+
 };
 
 
diff --git a/Engine/Plugins/Experimental/MeshModelingToolset/Source/ModelingOperators/Private/ParameterizationOps/UVProjectionOp.cpp b/Engine/Plugins/Experimental/MeshModelingToolset/Source/ModelingOperators/Private/ParameterizationOps/UVProjectionOp.cpp
--- a/Engine/Plugins/Experimental/MeshModelingToolset/Source/ModelingOperators/Private/ParameterizationOps/UVProjectionOp.cpp
+++ b/Engine/Plugins/Experimental/MeshModelingToolset/Source/ModelingOperators/Private/ParameterizationOps/UVProjectionOp.cpp
@@ -14,6 +14,7 @@
 #include "Async/ParallelFor.h"
 
 #include "MeshNormals.h"
+#include "DynamicMeshEditor.h"
 
 
 void FUVProjectionOp::SetTransform(const FTransform& Transform) {
@@ -54,6 +55,7 @@ void FUVProjectionOp::CalculateResult(FProgressCancel* Progress)
 	UVLayer->ClearElements();
 
 	FVector2D Scale = FVector2D(.5,.5)*UVScale, Offset = FVector2D(.5,.5)+UVOffset;
+	// project to major axis
 	auto ProjAxis = [&Scale, &Offset](const FVector3d& P, int Ax1, int Ax2, float Ax1Scale, float Ax2Scale)
 	{
 		return FVector2f(float(P[Ax1]) * Ax1Scale * Scale.X + Offset.X, float(P[Ax2]) * Ax2Scale * Scale.Y + Offset.Y);
@@ -81,86 +83,47 @@ void FUVProjectionOp::CalculateResult(FProgressCancel* Progress)
 				UVLayer->SetTriangle(TID, ResultMesh->GetTriangle(TID));
 			}
 		}
-		return;
 	}
-	
-	// compute normals on the transformed vertices
-	TArray<FVector3d> TransformedNormals; TransformedNormals.SetNumUninitialized(ResultMesh->TriangleCount());
-	ParallelFor(ResultMesh->TriangleCount(), [this, &TransformedVertices, &TransformedNormals](int32 TriangleID)
-	{
-		if (ResultMesh->IsTriangle(TriangleID))
-		{
-			FIndex3i Triangle = ResultMesh->GetTriangle(TriangleID);
-			TransformedNormals[TriangleID] = VectorUtil::Normal(TransformedVertices[Triangle.A], TransformedVertices[Triangle.B], TransformedVertices[Triangle.C]);
-		}
-	});
-	TMap<FIndex2i, int> VertIdxAndBucketIDToElementID;
-
-	if (ProjectionMethod == EUVProjectionMethod::Cube)
-	{
-		for (int TID = 0; TID < MaxTriangleID; TID++)
+	else
+	{	
+		// All cases in this branch require normals 
+		// compute normals on the transformed vertices
+		TArray<FVector3d> TransformedNormals; TransformedNormals.SetNumUninitialized(ResultMesh->TriangleCount());
+		ParallelFor(ResultMesh->TriangleCount(), [this, &TransformedVertices, &TransformedNormals](int32 TriangleID)
 		{
-			if (!ResultMesh->IsTriangle(TID))
-			{
-				continue;
-			}
-			const FVector3d& N = TransformedNormals[TID];
-			FVector3d NAbs(FMathd::Abs(N.X), FMathd::Abs(N.Y), FMathd::Abs(N.Z));
-			int MajorAxis = NAbs[0] > NAbs[1] ? (NAbs[0] > NAbs[2] ? 0 : 2) : (NAbs[1] > NAbs[2] ? 1 : 2);
-			double MajorAxisSign = FMathd::Sign(N[MajorAxis]);
-			int Minor1 = (MajorAxis + 1) % 3;
-			int Minor2 = (MajorAxis + 2) % 3;
-			int Bucket = MajorAxis;
-			if (MajorAxisSign > 0)
+			if (ResultMesh->IsTriangle(TriangleID))
 			{
-				Bucket += 3;
+				FIndex3i Triangle = ResultMesh->GetTriangle(TriangleID);
+				TransformedNormals[TriangleID] = VectorUtil::Normal(TransformedVertices[Triangle.A], TransformedVertices[Triangle.B], TransformedVertices[Triangle.C]);
 			}
-			FIndex3i OverlayTri;
-			FIndex3i SourceTri = ResultMesh->GetTriangle(TID);
-			for (int SubIdx = 0; SubIdx < 3; SubIdx++)
+		});
+		TMap<FIndex2i, int> VertIdxAndBucketIDToElementID;
+
+		if (ProjectionMethod == EUVProjectionMethod::Cube)
+		{
+			int Minor1s[3] = { 1, 0, 0 };
+			int Minor2s[3] = { 2, 2, 1 };
+			int Minor1Flip[3] = { -1, 1, 1 };
+			int Minor2Flip[3] = { -1, -1, 1 };
+			for (int TID = 0; TID < MaxTriangleID; TID++)
 			{
-				FIndex2i ElementKey(SourceTri[SubIdx], Bucket);
-				const int* ElementIdx = VertIdxAndBucketIDToElementID.Find(ElementKey);
-				if (!ElementIdx)
+				if (!ResultMesh->IsTriangle(TID))
 				{
-					FVector3d LocalV = TransformedVertices[SourceTri[SubIdx]];
-					int NewElementIdx = UVLayer->AppendElement(ProjAxis(LocalV, Minor1, Minor2, MajorAxisSign, 1), ElementKey.A);
-					VertIdxAndBucketIDToElementID.Add(ElementKey, NewElementIdx);
-					OverlayTri[SubIdx] = NewElementIdx;
+					continue;
 				}
-				else
+				const FVector3d& N = TransformedNormals[TID];
+				FVector3d NAbs(FMathd::Abs(N.X), FMathd::Abs(N.Y), FMathd::Abs(N.Z));
+				int MajorAxis = NAbs[0] > NAbs[1] ? (NAbs[0] > NAbs[2] ? 0 : 2) : (NAbs[1] > NAbs[2] ? 1 : 2);
+				double MajorAxisSign = FMathd::Sign(N[MajorAxis]);
+				int Minor1 = Minor1s[MajorAxis];
+				int Minor2 = Minor2s[MajorAxis];
+				int Bucket = MajorAxis;
+				if (MajorAxisSign > 0)
 				{
-					OverlayTri[SubIdx] = *ElementIdx;
+					Bucket += 3;
 				}
-			}
-			UVLayer->SetTriangle(TID, OverlayTri);
-		}
-
-		return;
-	}
-
-	if (ProjectionMethod == EUVProjectionMethod::Cylinder)
-	{
-		double DotThresholdRejectFromPlane = FMathd::Cos(FMathd::DegToRad * CylinderProjectToTopOrBottomAngleThreshold);
-		for (int TID = 0; TID < MaxTriangleID; TID++)
-		{
-			if (!ResultMesh->IsTriangle(TID))
-			{
-				continue;
-			}
-			const FVector3d& N = TransformedNormals[TID];
-			FIndex3i OverlayTri;
-			FIndex3i SourceTri = ResultMesh->GetTriangle(TID);
-			if (FMathd::Abs(N.Z) > DotThresholdRejectFromPlane)
-			{
-				int MajorAxis = 2;
-				double MajorAxisSign = FMathd::Sign(N[MajorAxis]);
-				// project to +/- Z
-				int Bucket = MajorAxisSign > 0 ? 1 : 0;
-				int Minor1 = 0;
-				int Minor2 = 1;
-
-				// TODO: extract common major axis plane projection sub-function from cube projection & this?
+				FIndex3i OverlayTri;
+				FIndex3i SourceTri = ResultMesh->GetTriangle(TID);
 				for (int SubIdx = 0; SubIdx < 3; SubIdx++)
 				{
 					FIndex2i ElementKey(SourceTri[SubIdx], Bucket);
@@ -168,7 +131,7 @@ void FUVProjectionOp::CalculateResult(FProgressCancel* Progress)
 					if (!ElementIdx)
 					{
 						FVector3d LocalV = TransformedVertices[SourceTri[SubIdx]];
-						int NewElementIdx = UVLayer->AppendElement(ProjAxis(LocalV, Minor1, Minor2, MajorAxisSign, 1), ElementKey.A);
+						int NewElementIdx = UVLayer->AppendElement(ProjAxis(LocalV, Minor1, Minor2, MajorAxisSign*Minor1Flip[MajorAxis], Minor2Flip[MajorAxis]), ElementKey.A);
 						VertIdxAndBucketIDToElementID.Add(ElementKey, NewElementIdx);
 						OverlayTri[SubIdx] = NewElementIdx;
 					}
@@ -177,50 +140,104 @@ void FUVProjectionOp::CalculateResult(FProgressCancel* Progress)
 						OverlayTri[SubIdx] = *ElementIdx;
 					}
 				}
+				UVLayer->SetTriangle(TID, OverlayTri);
 			}
-			else
+		}
+		else if (ProjectionMethod == EUVProjectionMethod::Cylinder)
+		{
+			double DotThresholdRejectFromPlane = FMathd::Cos(FMathd::DegToRad * CylinderProjectToTopOrBottomAngleThreshold);
+			for (int TID = 0; TID < MaxTriangleID; TID++)
 			{
-				// project to cylinder
-				int Bucket = 3;
-				FVector3d Centroid = (TransformedVertices[SourceTri.A] + TransformedVertices[SourceTri.B] + TransformedVertices[SourceTri.C]) * (1.0 / 3.0);
-				double CentroidAngle = FMathd::Atan2Positive(Centroid.Y, Centroid.X);
-				for (int SubIdx = 0; SubIdx < 3; SubIdx++)
+				if (!ResultMesh->IsTriangle(TID))
+				{
+					continue;
+				}
+				const FVector3d& N = TransformedNormals[TID];
+				FIndex3i OverlayTri;
+				FIndex3i SourceTri = ResultMesh->GetTriangle(TID);
+				if (FMathd::Abs(N.Z) > DotThresholdRejectFromPlane)
 				{
-					FVector3d LocalV = TransformedVertices[SourceTri[SubIdx]];
-					double VAngle = FMathd::Atan2Positive(LocalV.Y, LocalV.X);
-					double AngleDiff = VAngle - CentroidAngle;
-					if (FMathd::Abs(AngleDiff) > FMathd::Pi) {
-						if (AngleDiff < 0)
+					int MajorAxis = 2;
+					double MajorAxisSign = FMathd::Sign(N[MajorAxis]);
+					// project to +/- Z
+					int Bucket = MajorAxisSign > 0 ? 1 : 0;
+					int Minor1 = 0;
+					int Minor2 = 1;
+
+					// TODO: extract common major axis plane projection sub-function from cube projection & this?
+					for (int SubIdx = 0; SubIdx < 3; SubIdx++)
+					{
+						FIndex2i ElementKey(SourceTri[SubIdx], Bucket);
+						const int* ElementIdx = VertIdxAndBucketIDToElementID.Find(ElementKey);
+						if (!ElementIdx)
 						{
-							VAngle += FMathd::TwoPi;
-							Bucket--;
+							FVector3d LocalV = TransformedVertices[SourceTri[SubIdx]];
+							int NewElementIdx = UVLayer->AppendElement(ProjAxis(LocalV, Minor1, Minor2, MajorAxisSign, 1), ElementKey.A);
+							VertIdxAndBucketIDToElementID.Add(ElementKey, NewElementIdx);
+							OverlayTri[SubIdx] = NewElementIdx;
 						}
 						else
 						{
-							VAngle -= FMathd::TwoPi;
-							Bucket++;
+							OverlayTri[SubIdx] = *ElementIdx;
 						}
 					}
-					FIndex2i ElementKey(SourceTri[SubIdx], Bucket);
-					const int* ElementIdx = VertIdxAndBucketIDToElementID.Find(ElementKey);
-					if (!ElementIdx)
-					{
-						FVector2f UV((float(VAngle)*FMathf::InvPi-1.0f)*Scale.X+Offset.X, float(LocalV.Z)*Scale.Y + Offset.Y);
-						int NewElementIdx = UVLayer->AppendElement(UV, ElementKey.A);
-						VertIdxAndBucketIDToElementID.Add(ElementKey, NewElementIdx);
-						OverlayTri[SubIdx] = NewElementIdx;
-					}
-					else
+				}
+				else
+				{
+					// project to cylinder
+					int Bucket = 3;
+					FVector3d Centroid = (TransformedVertices[SourceTri.A] + TransformedVertices[SourceTri.B] + TransformedVertices[SourceTri.C]) * (1.0 / 3.0);
+					double CentroidAngle = FMathd::Atan2Positive(Centroid.Y, Centroid.X);
+					for (int SubIdx = 0; SubIdx < 3; SubIdx++)
 					{
-						OverlayTri[SubIdx] = *ElementIdx;
+						FVector3d LocalV = TransformedVertices[SourceTri[SubIdx]];
+						double VAngle = FMathd::Atan2Positive(LocalV.Y, LocalV.X);
+						double AngleDiff = VAngle - CentroidAngle;
+						if (FMathd::Abs(AngleDiff) > FMathd::Pi) {
+							if (AngleDiff < 0)
+							{
+								VAngle += FMathd::TwoPi;
+								Bucket--;
+							}
+							else
+							{
+								VAngle -= FMathd::TwoPi;
+								Bucket++;
+							}
+						}
+						FIndex2i ElementKey(SourceTri[SubIdx], Bucket);
+						const int* ElementIdx = VertIdxAndBucketIDToElementID.Find(ElementKey);
+						if (!ElementIdx)
+						{
+							FVector2f UV(-(float(VAngle)*FMathf::InvPi - 1.0f)*Scale.X + Offset.X, -float(LocalV.Z)*Scale.Y + Offset.Y);
+							int NewElementIdx = UVLayer->AppendElement(UV, ElementKey.A);
+							VertIdxAndBucketIDToElementID.Add(ElementKey, NewElementIdx);
+							OverlayTri[SubIdx] = NewElementIdx;
+						}
+						else
+						{
+							OverlayTri[SubIdx] = *ElementIdx;
+						}
 					}
 				}
-			}
 
-			UVLayer->SetTriangle(TID, OverlayTri);
+				UVLayer->SetTriangle(TID, OverlayTri);
+			}
 		}
 
-		return;
+
+		// rescale UVs to an average world-space target length (and then re-multiply by the UVScale param)
+		if (bWorldSpaceUVScale)
+		{
+			FDynamicMeshEditor Editor(ResultMesh.Get());
+			// global rescale with a 'standard' global scale factor
+			Editor.RescaleAttributeUVs(.01, true, 0, ResultTransform);
+			// apply the 2D uvscale after
+			for (int ElID : UVLayer->ElementIndicesItr())
+			{
+				UVLayer->SetElement(ElID, UVLayer->GetElement(ElID) * UVScale);
+			}
+		}
 	}
 
 }
diff --git a/Engine/Plugins/Experimental/MeshModelingToolset/Source/ModelingOperators/Public/ParameterizationOps/UVProjectionOp.h b/Engine/Plugins/Experimental/MeshModelingToolset/Source/ModelingOperators/Public/ParameterizationOps/UVProjectionOp.h
--- a/Engine/Plugins/Experimental/MeshModelingToolset/Source/ModelingOperators/Public/ParameterizationOps/UVProjectionOp.h
+++ b/Engine/Plugins/Experimental/MeshModelingToolset/Source/ModelingOperators/Public/ParameterizationOps/UVProjectionOp.h
@@ -31,6 +31,7 @@ public:
 	float CylinderProjectToTopOrBottomAngleThreshold;
 	FTransform ProjectionTransform;
 	FVector2D UVScale, UVOffset;
+	bool bWorldSpaceUVScale;
 
 	void SetTransform(const FTransform& Transform);
 
