Mesh Modeling: PolyEdit tool.  Fix laplacian deformation behavior when a face in a low poly object has been selected.
If there are no degrees of freedom for the laplacian solver it now falls back to the linear solve.

#rnx

diff --git a/Engine/Plugins/Experimental/MeshModelingToolset/Source/MeshModelingTools/Private/EditMeshPolygonsTool.cpp b/Engine/Plugins/Experimental/MeshModelingToolset/Source/MeshModelingTools/Private/EditMeshPolygonsTool.cpp
--- a/Engine/Plugins/Experimental/MeshModelingToolset/Source/MeshModelingTools/Private/EditMeshPolygonsTool.cpp
+++ b/Engine/Plugins/Experimental/MeshModelingToolset/Source/MeshModelingTools/Private/EditMeshPolygonsTool.cpp
@@ -177,7 +177,7 @@ void FConstrainedMeshDeformerTask::DoWork()
 	}
 	else
 	{
-		UE_LOG(LogTemp, Warning, TEXT("Laplacian deformation failed"));
+		//UE_LOG(LogTemp, Warning, TEXT("Laplacian deformation failed"));
 	}
 
 }
@@ -392,21 +392,28 @@ void FGroupTopologyLaplacianDeformer::SetActiveHandleFaces(const TArray<int>& Fa
 	Topology->CollectGroupBoundaryVertices(GroupID, HandleBoundaryVertices);
 	ModifiedVertices = HandleVertices;
 
-	// find neighbour group set
-	TArray<int> HandleGroups = FaceGroupIDs;
-	const TArray<int>& GroupNbrGroups = Topology->GetGroupNbrGroups(GroupID);
 
 
-	CalculateROI(HandleGroups, GroupNbrGroups);
-
+	// list of adj groups.  may contain duplicates.
+	TArray<int> AdjGroups;
+	for (int BoundaryVert : HandleBoundaryVertices)
+	{
+		Topology->FindVertexNbrGroups(BoundaryVert, AdjGroups);
+	}
 
-	TArray<int> NeighborGroups;
+	// Local neighborhood - Adjacent groups plus self
+	TArray<int> NeighborhoodGroups;
 
-	//Add to output parameter 
-	NeighborGroups.Add(GroupID);
-	NeighborGroups.Append(GroupNbrGroups);
+	// Collect the rest of the 1-ring groups that are adjacent to the selected one.
+	NeighborhoodGroups.Add(GroupID);
+	for (int AdjGroup : AdjGroups)
+	{
+		NeighborhoodGroups.AddUnique(AdjGroup); // remove duplicates by add unique
+	}
 
-	UpdateSelection(Mesh, NeighborGroups, bLocalize);
+	CalculateROI(FaceGroupIDs, NeighborhoodGroups);
+	
+	UpdateSelection(Mesh, NeighborhoodGroups, bLocalize);
 
 	// Save the positions of the selected region.
 	SaveInitialPositions();
@@ -893,50 +900,76 @@ void UEditMeshPolygonsTool::OnBeginDrag(const FRay& WorldRay)
 	LastBrushPosLocal = Transform.InverseTransformPosition(LastHitPosWorld);
 	StartBrushPosLocal = LastBrushPosLocal;
 
+	// Record the requested deformation strategy - NB: will be forced to linear if there aren't any free points to solve.
+
+	DeformationStrategy = TransformProps->DeformationStrategy;
+	
 	// Capture the part of the mesh that will deform
 
-	if (TransformProps->DeformationStrategy == EGroupTopologyDeformationStrategy::Linear)
+	if (DeformationStrategy == EGroupTopologyDeformationStrategy::Laplacian)
 	{
+		LaplacianDeformer.bLocalize = true; // TransformProps->bLocalizeDeformation;
+
 		//Determine which of the following (corners, edges or faces) has been selected by counting the associated feature's IDs
 		if (Selection.SelectedCornerIDs.Num() > 0)
 		{
 			//Add all the the Corner's adjacent poly-groups (NbrGroups) to the ongoing array of groups.
-			LinearDeformer.SetActiveHandleCorners(Selection.SelectedCornerIDs);
+			LaplacianDeformer.SetActiveHandleCorners(Selection.SelectedCornerIDs);
 		}
 		else if (Selection.SelectedEdgeIDs.Num() > 0)
 		{
 			//Add all the the edge's adjacent poly-groups (NbrGroups) to the ongoing array of groups.
-			LinearDeformer.SetActiveHandleEdges(Selection.SelectedEdgeIDs);
+			LaplacianDeformer.SetActiveHandleEdges(Selection.SelectedEdgeIDs);
 		}
 		else if (Selection.SelectedGroupIDs.Num() > 0)
 		{
-			LinearDeformer.SetActiveHandleFaces(Selection.SelectedGroupIDs);
+			LaplacianDeformer.SetActiveHandleFaces(Selection.SelectedGroupIDs);
 		}
+
+
+		// If there are actually no interior points, then we can't actually use the laplacian deformer. Need to fall back to the linear.
+		bool bHasInteriorVerts = false;
+		const auto& ROIFaces = LaplacianDeformer.GetROIFaces();
+		for (const auto& Face : ROIFaces)
+		{
+			bHasInteriorVerts = bHasInteriorVerts || ( Face.InteriorVerts.Num() != 0);
+		}
+
+		if (!bHasInteriorVerts)
+		{
+			// Change to the linear strategy for this case.
+
+			DeformationStrategy = EGroupTopologyDeformationStrategy::Linear;
+		}
+		else
+		{		
+			// finalize the laplacian deformer : the task will need a new mesh that corresponds to the selected region.
+		
+			LaplacianDeformer.bTaskSubmeshIsDirty = true;
+
+		}
+		
 	}
-	else //(TransformProps->DeformationStrategy == EGroupTopologyDeformationStrategy::Laplacian)
+	 
+	if (DeformationStrategy == EGroupTopologyDeformationStrategy::Linear )
 	{
-		LaplacianDeformer.bLocalize = true; // TransformProps->bLocalizeDeformation;
-
 		//Determine which of the following (corners, edges or faces) has been selected by counting the associated feature's IDs
 		if (Selection.SelectedCornerIDs.Num() > 0)
 		{
 			//Add all the the Corner's adjacent poly-groups (NbrGroups) to the ongoing array of groups.
-			LaplacianDeformer.SetActiveHandleCorners(Selection.SelectedCornerIDs);
+			LinearDeformer.SetActiveHandleCorners(Selection.SelectedCornerIDs);
 		}
 		else if (Selection.SelectedEdgeIDs.Num() > 0)
 		{
 			//Add all the the edge's adjacent poly-groups (NbrGroups) to the ongoing array of groups.
-			LaplacianDeformer.SetActiveHandleEdges(Selection.SelectedEdgeIDs);
+			LinearDeformer.SetActiveHandleEdges(Selection.SelectedEdgeIDs);
 		}
 		else if (Selection.SelectedGroupIDs.Num() > 0)
 		{
-			LaplacianDeformer.SetActiveHandleFaces(Selection.SelectedGroupIDs);
+			LinearDeformer.SetActiveHandleFaces(Selection.SelectedGroupIDs);
 		}
-
-		// the task will need a new mesh that corresponds to the selected region.
-		LaplacianDeformer.bTaskSubmeshIsDirty = true;
 	}
-
+	
 	BeginChange();
 }
 
@@ -1002,7 +1035,7 @@ void UEditMeshPolygonsTool::UpdateChangeFromROI(bool bFinal)
 	{
 		return;
 	}
-	const bool bIsLaplacian = (TransformProps->DeformationStrategy == EGroupTopologyDeformationStrategy::Laplacian);
+	const bool bIsLaplacian = (DeformationStrategy == EGroupTopologyDeformationStrategy::Laplacian);
 
 	FDynamicMesh3* Mesh = DynamicMeshComponent->GetMesh();
 	const TSet<int>& ModifiedVertices = (bIsLaplacian) ? LaplacianDeformer.GetModifiedVertices() : LinearDeformer.GetModifiedVertices();
@@ -1034,7 +1067,7 @@ void UEditMeshPolygonsTool::OnEndDrag(const FRay& Ray)
 	QuickAxisRotator.ClearAxisLock();
 
 	//If it's linear, it's computed real time with no delay. This may need to be restructured for clarity by using the background task for this as well.
-	if (TransformProps->DeformationStrategy == EGroupTopologyDeformationStrategy::Linear)
+	if (DeformationStrategy == EGroupTopologyDeformationStrategy::Linear)
 	{
 		// close change record
 		EndChange();
@@ -1092,7 +1125,7 @@ void UEditMeshPolygonsTool::ComputeUpdate()
 		}
 	}
 
-	if (TransformProps->DeformationStrategy == EGroupTopologyDeformationStrategy::Laplacian)
+	if (DeformationStrategy == EGroupTopologyDeformationStrategy::Laplacian)
 	{
 		bool bIsWorking = LaplacianDeformer.IsTaskInFlight();
 
@@ -1151,7 +1184,7 @@ void UEditMeshPolygonsTool::ComputeUpdate()
 
 void UEditMeshPolygonsTool::ComputeUpdate_Rotate()
 {
-	const bool bIsLaplacian = (TransformProps->DeformationStrategy == EGroupTopologyDeformationStrategy::Laplacian); 
+	const bool bIsLaplacian = (DeformationStrategy == EGroupTopologyDeformationStrategy::Laplacian); 
 	FGroupTopologyDeformer& SelectedDeformer = (bIsLaplacian) ? LaplacianDeformer : LinearDeformer;
 
 	FDynamicMesh3* Mesh = DynamicMeshComponent->GetMesh();
@@ -1189,7 +1222,7 @@ void UEditMeshPolygonsTool::ComputeUpdate_Rotate()
 		SelectedDeformer.ClearSolution(Mesh);
 
 		//TODO: This is unseemly here, need to potentially defer this so that it's handled the same way as laplacian. Placeholder for now.
-		if (TransformProps->DeformationStrategy == EGroupTopologyDeformationStrategy::Linear)
+		if (DeformationStrategy == EGroupTopologyDeformationStrategy::Linear)
 		{
 
 			DynamicMeshComponent->FastNotifyPositionsUpdated();
@@ -1241,7 +1274,7 @@ void UEditMeshPolygonsTool::ComputeUpdate_Rotate()
 
 void UEditMeshPolygonsTool::ComputeUpdate_Translate()
 {
-	const bool bIsLaplacian = (TransformProps->DeformationStrategy == EGroupTopologyDeformationStrategy::Laplacian);
+	const bool bIsLaplacian = (DeformationStrategy == EGroupTopologyDeformationStrategy::Laplacian);
 	FGroupTopologyDeformer& SelectedDeformer = (bIsLaplacian) ? LaplacianDeformer : LinearDeformer;
 
 	FTransform Transform = ComponentTarget->GetWorldTransform();
@@ -1482,7 +1515,7 @@ void UEditMeshPolygonsTool::SetTriangleGroups(const TDynamicVector<int>& Groups)
 
 void UEditMeshPolygonsTool::BeginChange()
 {
-	const bool bIsLaplacian = (TransformProps->DeformationStrategy == EGroupTopologyDeformationStrategy::Laplacian);
+	const bool bIsLaplacian = (DeformationStrategy == EGroupTopologyDeformationStrategy::Laplacian);
 	if (!bIsLaplacian || LaplacianDeformer.IsDone())
 	{
 		if (ActiveVertexChange == nullptr)
diff --git a/Engine/Plugins/Experimental/MeshModelingToolset/Source/MeshModelingTools/Public/EditMeshPolygonsTool.h b/Engine/Plugins/Experimental/MeshModelingToolset/Source/MeshModelingTools/Public/EditMeshPolygonsTool.h
--- a/Engine/Plugins/Experimental/MeshModelingToolset/Source/MeshModelingTools/Public/EditMeshPolygonsTool.h
+++ b/Engine/Plugins/Experimental/MeshModelingToolset/Source/MeshModelingTools/Public/EditMeshPolygonsTool.h
@@ -359,6 +359,8 @@ public:
 	/** Triggers abort on task and passes off ownership to deleter object */
 	inline void Shutdown();
 
+	const TArray<FROIFace>& GetROIFaces() const { return ROIFaces; }
+
 	/** Stores the position of the vertex constraints and corresponding weights for the entire mesh.  This is used as a form of scratch space.*/
 	TArray<FConstraintData> SrcMeshConstraintBuffer;
 
@@ -501,6 +503,8 @@ protected:
 	FGroupTopologyDeformer LinearDeformer;
 	FGroupTopologyLaplacianDeformer LaplacianDeformer;
 
+	EGroupTopologyDeformationStrategy DeformationStrategy;
+
 	// Initial polygon group and mesh info
 	TDynamicVector<int> InitialTriangleGroups;
 	TUniquePtr<FDynamicMesh3> InitialMesh;
