MeshModeling: VoxelCSGMeshesTool and Voxel Mesh Merging - On Accept, leave the source meshes TemporarilyHiddenInEditor.  Might want to change this to removing them.
#rnx

diff --git a/Engine/Plugins/Experimental/MeshModelingToolset/Source/MeshModelingToolsEditorOnly/Private/MergeMeshesTool.cpp b/Engine/Plugins/Experimental/MeshModelingToolset/Source/MeshModelingToolsEditorOnly/Private/MergeMeshesTool.cpp
--- a/Engine/Plugins/Experimental/MeshModelingToolset/Source/MeshModelingToolsEditorOnly/Private/MergeMeshesTool.cpp
+++ b/Engine/Plugins/Experimental/MeshModelingToolset/Source/MeshModelingToolsEditorOnly/Private/MergeMeshesTool.cpp
@@ -108,12 +108,23 @@ void UMergeMeshesTool::Shutdown(EToolShutdownType ShutdownType)
 	if (ShutdownType == EToolShutdownType::Accept)
 	{
 		GenerateAsset(Result);
-	}
+		for (auto& ComponentTarget : ComponentTargets)
+		{
+			ComponentTarget->SetOwnerVisibility(true);
+			AActor* Actor = ComponentTarget->GetOwnerActor();
+			Actor->SetIsTemporarilyHiddenInEditor(true);
+			// NB: We could alternately remove the actors.
+			//Actor->Destroy();
 
-	// Restore (unhide) the source meshes
-	for (auto& ComponentTarget : ComponentTargets)
+		}
+	}
+	else
 	{
-		ComponentTarget->SetOwnerVisibility(true);
+		// Restore (unhide) the source meshes
+		for (auto& ComponentTarget : ComponentTargets)
+		{
+			ComponentTarget->SetOwnerVisibility(true);
+		}
 	}
 }
 
diff --git a/Engine/Plugins/Experimental/MeshModelingToolset/Source/MeshModelingToolsEditorOnly/Private/VoxelCSGMeshesTool.cpp b/Engine/Plugins/Experimental/MeshModelingToolset/Source/MeshModelingToolsEditorOnly/Private/VoxelCSGMeshesTool.cpp
--- a/Engine/Plugins/Experimental/MeshModelingToolset/Source/MeshModelingToolsEditorOnly/Private/VoxelCSGMeshesTool.cpp
+++ b/Engine/Plugins/Experimental/MeshModelingToolset/Source/MeshModelingToolsEditorOnly/Private/VoxelCSGMeshesTool.cpp
@@ -107,12 +107,24 @@ void UVoxelCSGMeshesTool::Shutdown(EToolShutdownType ShutdownType)
 	if (ShutdownType == EToolShutdownType::Accept)
 	{
 		GenerateAsset(Result);
-	}
 
-	// Restore (unhide) the source meshes
-	for (auto& ComponentTarget : ComponentTargets)
+		for (auto& ComponentTarget : ComponentTargets)
+		{
+			ComponentTarget->SetOwnerVisibility(true);
+			AActor* Actor = ComponentTarget->GetOwnerActor();
+			Actor->SetIsTemporarilyHiddenInEditor(true);
+			// NB: We could alternatly remove the actors.
+			//Actor->Destroy();
+
+		}
+	}
+	else
 	{
-		ComponentTarget->SetOwnerVisibility(true);
+		// Restore (unhide) the source meshes
+		for (auto& ComponentTarget : ComponentTargets)
+		{
+			ComponentTarget->SetOwnerVisibility(true);
+		}
 	}
 }
 
