ModelingTools: add snap-drag mode to TransformMeshesTool. When enabled, click-drag on object can be used to reposition it onto scene surfaces. Fixed undo/redo issues.
#rnx

diff --git a/Engine/Plugins/Experimental/MeshModelingToolset/Source/MeshModelingTools/Private/TransformMeshesTool.cpp b/Engine/Plugins/Experimental/MeshModelingToolset/Source/MeshModelingTools/Private/TransformMeshesTool.cpp
--- a/Engine/Plugins/Experimental/MeshModelingToolset/Source/MeshModelingTools/Private/TransformMeshesTool.cpp
+++ b/Engine/Plugins/Experimental/MeshModelingToolset/Source/MeshModelingTools/Private/TransformMeshesTool.cpp
@@ -72,9 +72,9 @@ void UTransformMeshesTool::Setup()
 {
 	UInteractiveTool::Setup();
 
-	//UClickDragInputBehavior* ClickDragBehavior = NewObject<UClickDragInputBehavior>(this);
-	//ClickDragBehavior->Initialize(this);
-	//AddInputBehavior(ClickDragBehavior);
+	UClickDragInputBehavior* ClickDragBehavior = NewObject<UClickDragInputBehavior>(this);
+	ClickDragBehavior->Initialize(this);
+	AddInputBehavior(ClickDragBehavior);
 
 	TransformProps = NewObject<UTransformMeshesToolProperties>();
 	AddToolPropertySource(TransformProps);
@@ -113,8 +113,7 @@ void UTransformMeshesTool::OnPropertyModified(UObject* PropertySet, UProperty* P
 	}
 
 	bool bEnableSetPivot = TransformProps->bSetPivot;
-	if (TransformProps->TransformMode == ETransformMeshesTransformMode::SharedGizmoLocal ||
-		TransformProps->TransformMode == ETransformMeshesTransformMode::QuickTranslate)
+	if (TransformProps->TransformMode == ETransformMeshesTransformMode::SharedGizmoLocal)
 	{
 		bEnableSetPivot = false;
 	}
@@ -201,80 +200,115 @@ void UTransformMeshesTool::ResetActiveGizmos()
 // does not make sense that CanBeginClickDragSequence() returns a RayHit? Needs to be an out-argument...
 FInputRayHit UTransformMeshesTool::CanBeginClickDragSequence(const FInputDeviceRay& PressPos)
 {
-	check(ActiveGizmos.Num() == 1);
+	if (TransformProps->bEnableSnapDragging == false || ActiveGizmos.Num() == 0)
+	{
+		return FInputRayHit();
+	}
+
+	ActiveSnapDragIndex = -1;
 
 	float MinHitDistance = TNumericLimits<float>::Max();
-	for (TUniquePtr<FPrimitiveComponentTarget>& Target : ComponentTargets)
+	FVector HitNormal;
+
+	for ( int k = 0; k < ComponentTargets.Num(); ++k )
 	{
+		TUniquePtr<FPrimitiveComponentTarget>& Target = ComponentTargets[k];
+
 		FHitResult WorldHit;
 		if (Target->HitTest(PressPos.WorldRay, WorldHit))
 		{
 			MinHitDistance = FMath::Min(MinHitDistance, WorldHit.Distance);
+			HitNormal = WorldHit.Normal;
+			ActiveSnapDragIndex = k;
 		}
 	}
-	return (MinHitDistance < TNumericLimits<float>::Max()) ? FInputRayHit(MinHitDistance) : FInputRayHit();
+	return (MinHitDistance < TNumericLimits<float>::Max()) ? FInputRayHit(MinHitDistance, HitNormal) : FInputRayHit();
 }
 
 void UTransformMeshesTool::OnClickPress(const FInputDeviceRay& PressPos)
 {
-	check(ActiveGizmos.Num() == 1);
-
 	FInputRayHit HitPos = CanBeginClickDragSequence(PressPos);
 	check(HitPos.bHit);
 
-	GetToolManager()->BeginUndoTransaction(LOCTEXT("TransformToolTransformTxnName", "Transform"));
+	GetToolManager()->BeginUndoTransaction(LOCTEXT("TransformToolTransformTxnName", "SnapDrag"));
 
-	StartDragFrameWorld = FFrame3d(PressPos.WorldRay.PointAt(HitPos.HitDepth));
-	USceneComponent* GizmoComponent = ActiveGizmos[0].TransformGizmo->GetGizmoActor()->GetRootComponent();
-	StartDragTransform = GizmoComponent->GetComponentToWorld();
-	GizmoComponent->Modify();
+	StartDragFrameWorld = FFrame3d(PressPos.WorldRay.PointAt(HitPos.HitDepth), HitPos.HitNormal);
 
-	for (TUniquePtr<FPrimitiveComponentTarget>& Target : ComponentTargets)
-	{
-		Target->GetOwnerComponent()->Modify();
-	}
+	FTransformMeshesTarget& ActiveTarget =
+		(TransformProps->TransformMode == ETransformMeshesTransformMode::PerObjectGizmo) ?
+		ActiveGizmos[ActiveSnapDragIndex] : ActiveGizmos[0];
+	USceneComponent* GizmoComponent = ActiveTarget.TransformGizmo->GetGizmoActor()->GetRootComponent();
+	StartDragTransform = GizmoComponent->GetComponentToWorld();
 }
 
+
 void UTransformMeshesTool::OnClickDrag(const FInputDeviceRay& DragPos)
 {
-	for (TUniquePtr<FPrimitiveComponentTarget>& Target : ComponentTargets)
+	FCollisionObjectQueryParams ObjectQueryParams(FCollisionObjectQueryParams::AllObjects);
+	FCollisionQueryParams CollisionParams = FCollisionQueryParams::DefaultQueryParam;
+	
+	int IgnoreIndex = (TransformProps->TransformMode == ETransformMeshesTransformMode::PerObjectGizmo) ?
+		ActiveSnapDragIndex : -1;
+	for ( int k = 0; k < ComponentTargets.Num(); ++k )
 	{
-		Target->SetOwnerVisibility(false);
+		if (IgnoreIndex == -1 || k == IgnoreIndex)
+		{
+			CollisionParams.AddIgnoredComponent(ComponentTargets[k]->GetOwnerComponent());
+		}
 	}
 
-	FCollisionObjectQueryParams QueryParams(FCollisionObjectQueryParams::AllObjects);
 	FHitResult Result;
-	bool bWorldHit = TargetWorld->LineTraceSingleByObjectType(Result, DragPos.WorldRay.Origin, DragPos.WorldRay.PointAt(999999), QueryParams);
+	bool bWorldHit = TargetWorld->LineTraceSingleByObjectType(Result, DragPos.WorldRay.Origin, DragPos.WorldRay.PointAt(999999), ObjectQueryParams, CollisionParams);
 	if (bWorldHit)
 	{
 		FVector HitPos = Result.ImpactPoint;
-		FFrame3d NewFrame = StartDragFrameWorld;
-		NewFrame.Origin = (FVector3d)HitPos;
+		FVector TargetNormal = -Result.Normal;
 
-		FTransform RelTransform = FTransform::Identity;
-		RelTransform.SetTranslation(
-			(FVector)(NewFrame.Origin - StartDragFrameWorld.Origin));
+		FFrame3d FromFrameWorld = StartDragFrameWorld;
+		FFrame3d ToFrameWorld(HitPos, TargetNormal);
+		FFrame3d ObjectFrameWorld(StartDragTransform);
+
+		FVector3d CenterShift = FromFrameWorld.Origin - ObjectFrameWorld.Origin;
+		FQuaterniond AlignRotation(FromFrameWorld.Z(), ToFrameWorld.Z());
+		FVector3d AlignTranslate = ToFrameWorld.Origin - FromFrameWorld.Origin;
 
-		USceneComponent* GizmoComponent = ActiveGizmos[0].TransformGizmo->GetGizmoActor()->GetRootComponent();
 		FTransform NewTransform = StartDragTransform;
-		NewTransform.Accumulate(RelTransform);
+		NewTransform.Accumulate( FTransform(CenterShift) );
+		NewTransform.Accumulate( FTransform(AlignRotation) );
+		NewTransform.Accumulate( FTransform(AlignTranslate) );
+		CenterShift = AlignRotation * CenterShift;
+		NewTransform.Accumulate( FTransform(-CenterShift) );
+
+		FTransformMeshesTarget& ActiveTarget =
+			(TransformProps->TransformMode == ETransformMeshesTransformMode::PerObjectGizmo) ?
+			ActiveGizmos[ActiveSnapDragIndex] : ActiveGizmos[0];
+		USceneComponent* GizmoComponent = ActiveTarget.TransformGizmo->GetGizmoActor()->GetRootComponent();
 		GizmoComponent->SetWorldTransform(NewTransform);
 	}
 
-	for (TUniquePtr<FPrimitiveComponentTarget>& Target : ComponentTargets)
-	{
-		Target->SetOwnerVisibility(true);
-	}
 }
 
+
 void UTransformMeshesTool::OnClickRelease(const FInputDeviceRay& ReleasePos)
 {
-	GetToolManager()->EndUndoTransaction();
+	OnTerminateDragSequence();
 }
 
 void UTransformMeshesTool::OnTerminateDragSequence()
 {
+	FTransformMeshesTarget& ActiveTarget =
+		(TransformProps->TransformMode == ETransformMeshesTransformMode::PerObjectGizmo) ?
+		ActiveGizmos[ActiveSnapDragIndex] : ActiveGizmos[0];
+	USceneComponent* GizmoComponent = ActiveTarget.TransformGizmo->GetGizmoActor()->GetRootComponent();
+	FTransform EndDragtransform = GizmoComponent->GetComponentToWorld();
+
+	TUniquePtr<FComponentWorldTransformChange> Change = MakeUnique<FComponentWorldTransformChange>(StartDragTransform, EndDragtransform);
+	GetToolManager()->EmitObjectChange(GizmoComponent, MoveTemp(Change),
+		LOCTEXT("TransformToolTransformTxnName", "SnapDrag"));
+
 	GetToolManager()->EndUndoTransaction();
+
+	ActiveSnapDragIndex = -1;
 }
 
 
diff --git a/Engine/Plugins/Experimental/MeshModelingToolset/Source/MeshModelingTools/Public/TransformMeshesTool.h b/Engine/Plugins/Experimental/MeshModelingToolset/Source/MeshModelingTools/Public/TransformMeshesTool.h
--- a/Engine/Plugins/Experimental/MeshModelingToolset/Source/MeshModelingTools/Public/TransformMeshesTool.h
+++ b/Engine/Plugins/Experimental/MeshModelingToolset/Source/MeshModelingTools/Public/TransformMeshesTool.h
@@ -6,6 +6,7 @@
 #include "MultiSelectionTool.h"
 #include "InteractiveToolBuilder.h"
 #include "BaseBehaviors//BehaviorTargetInterfaces.h"
+#include "Changes/TransformChange.h"
 #include "FrameTypes.h"
 #include "TransformMeshesTool.generated.h"
 
@@ -40,10 +41,7 @@ enum class ETransformMeshesTransformMode : uint8
 	SharedGizmoLocal = 1 UMETA(DisplayName = "Shared Gizmo (Local)"),
 
 	/** Per Object Gizmo */
-	PerObjectGizmo = 2 UMETA(DisplayName = "Multi-Gizmo"),
-
-	/** Quick Translate Transformer */
-	QuickTranslate = 3 UMETA(DisplayName = "QuickTranslate")
+	PerObjectGizmo = 2 UMETA(DisplayName = "Multi-Gizmo")
 };
 
 
@@ -63,6 +61,10 @@ public:
 
 	UPROPERTY(EditAnywhere, Category = Options, meta = (EditCondition = "TransformMode == ETransformMeshesTransformMode::SharedGizmo || TransformMode == ETransformMeshesTransformMode::PerObjectGizmo"))
 	bool bSetPivot = false;
+
+
+	UPROPERTY(EditAnywhere, Category = Options, meta = (EditCondition = "bSetPivot == false"))
+	bool bEnableSnapDragging = false;
 };
 
 
@@ -135,7 +137,7 @@ protected:
 
 	FFrame3d StartDragFrameWorld;
 	FTransform StartDragTransform;
-
+	int ActiveSnapDragIndex = -1;
 
 	void OnParametersUpdated();
 };
