MeshModeling: VoxelCSGMeshesTool and Voxel Mesh Merging - Add option to remove source meshes on Accept.  Default yes.
#rnx

diff --git a/Engine/Plugins/Experimental/MeshModelingToolset/Source/MeshModelingToolsEditorOnly/Private/MergeMeshesTool.cpp b/Engine/Plugins/Experimental/MeshModelingToolset/Source/MeshModelingToolsEditorOnly/Private/MergeMeshesTool.cpp
--- a/Engine/Plugins/Experimental/MeshModelingToolset/Source/MeshModelingToolsEditorOnly/Private/MergeMeshesTool.cpp
+++ b/Engine/Plugins/Experimental/MeshModelingToolset/Source/MeshModelingToolsEditorOnly/Private/MergeMeshesTool.cpp
@@ -107,16 +107,27 @@ void UMergeMeshesTool::Shutdown(EToolShutdownType ShutdownType)
 	TUniquePtr<FDynamicMeshOpResult> Result = Preview->Shutdown();
 	if (ShutdownType == EToolShutdownType::Accept)
 	{
+		GetToolManager()->BeginUndoTransaction(LOCTEXT("MergeMeshes", "Merge Meshes"));
+
 		GenerateAsset(Result);
+
 		for (auto& ComponentTarget : ComponentTargets)
 		{
 			ComponentTarget->SetOwnerVisibility(true);
 			AActor* Actor = ComponentTarget->GetOwnerActor();
-			Actor->SetIsTemporarilyHiddenInEditor(true);
-			// NB: We could alternately remove the actors.
-			//Actor->Destroy();
-
+			if (MergeProps->bRemoveSources)
+			{
+				Actor->Destroy();
+			}
+			else
+			{
+				// just hide the result.
+				Actor->SetIsTemporarilyHiddenInEditor(true);
+			}
 		}
+
+		GetToolManager()->EndUndoTransaction();
+
 	}
 	else
 	{
@@ -192,7 +203,7 @@ void UMergeMeshesTool::GenerateAsset(const TUniquePtr<FDynamicMeshOpResult>& Res
 {
 	check(Result->Mesh.Get() != nullptr);
 
-	GetToolManager()->BeginUndoTransaction(LOCTEXT("MergeMeshes", "Merge Meshes"));
+	
 
 	AActor* NewActor = AssetGenerationUtil::GenerateStaticMeshActor(
 		AssetAPI, TargetWorld,
@@ -202,7 +213,7 @@ void UMergeMeshesTool::GenerateAsset(const TUniquePtr<FDynamicMeshOpResult>& Res
 	// select newly-created object
 	ToolSelectionUtil::SetNewActorSelection(GetToolManager(), NewActor);
 
-	GetToolManager()->EndUndoTransaction();
+	
 }
 
 
diff --git a/Engine/Plugins/Experimental/MeshModelingToolset/Source/MeshModelingToolsEditorOnly/Private/VoxelCSGMeshesTool.cpp b/Engine/Plugins/Experimental/MeshModelingToolset/Source/MeshModelingToolsEditorOnly/Private/VoxelCSGMeshesTool.cpp
--- a/Engine/Plugins/Experimental/MeshModelingToolset/Source/MeshModelingToolsEditorOnly/Private/VoxelCSGMeshesTool.cpp
+++ b/Engine/Plugins/Experimental/MeshModelingToolset/Source/MeshModelingToolsEditorOnly/Private/VoxelCSGMeshesTool.cpp
@@ -106,17 +106,25 @@ void UVoxelCSGMeshesTool::Shutdown(EToolShutdownType ShutdownType)
 	TUniquePtr<FDynamicMeshOpResult> Result = Preview->Shutdown();
 	if (ShutdownType == EToolShutdownType::Accept)
 	{
+		GetToolManager()->BeginUndoTransaction(LOCTEXT("VoxelCSGMeshes", "Boolean Meshes"));
+
 		GenerateAsset(Result);
 
 		for (auto& ComponentTarget : ComponentTargets)
 		{
 			ComponentTarget->SetOwnerVisibility(true);
 			AActor* Actor = ComponentTarget->GetOwnerActor();
-			Actor->SetIsTemporarilyHiddenInEditor(true);
-			// NB: We could alternatly remove the actors.
-			//Actor->Destroy();
-
+			if (CSGProps->bRemoveSources)
+			{
+				Actor->Destroy();
+			}
+			else
+			{
+				Actor->SetIsTemporarilyHiddenInEditor(true);
+			}
 		}
+
+		GetToolManager()->EndUndoTransaction();
 	}
 	else
 	{
@@ -184,7 +192,7 @@ void UVoxelCSGMeshesTool::GenerateAsset(const TUniquePtr<FDynamicMeshOpResult>&
 {
 	check(Result->Mesh.Get() != nullptr);
 
-	GetToolManager()->BeginUndoTransaction(LOCTEXT("VoxelCSGMeshes", "Boolean Meshes"));
+	//GetToolManager()->BeginUndoTransaction(LOCTEXT("VoxelCSGMeshes", "Boolean Meshes"));
 
 	AActor* NewActor = AssetGenerationUtil::GenerateStaticMeshActor(
 		AssetAPI, TargetWorld,
@@ -194,7 +202,7 @@ void UVoxelCSGMeshesTool::GenerateAsset(const TUniquePtr<FDynamicMeshOpResult>&
 	// select newly-created object
 	ToolSelectionUtil::SetNewActorSelection(GetToolManager(), NewActor);
 
-	GetToolManager()->EndUndoTransaction();
+	//GetToolManager()->EndUndoTransaction();
 }
 
 #undef LOCTEXT_NAMESPACE
diff --git a/Engine/Plugins/Experimental/MeshModelingToolset/Source/MeshModelingToolsEditorOnly/Public/MergeMeshesTool.h b/Engine/Plugins/Experimental/MeshModelingToolset/Source/MeshModelingToolsEditorOnly/Public/MergeMeshesTool.h
--- a/Engine/Plugins/Experimental/MeshModelingToolset/Source/MeshModelingToolsEditorOnly/Public/MergeMeshesTool.h
+++ b/Engine/Plugins/Experimental/MeshModelingToolset/Source/MeshModelingToolsEditorOnly/Public/MergeMeshesTool.h
@@ -56,6 +56,10 @@ public:
 	/** Automatically simplify the result of voxel-based merge.*/
 	UPROPERTY(EditAnywhere, Category = Options)
 	bool bAutoSimplify = true;
+
+	/** Remove the source meshes when accepting results of tool.*/
+	UPROPERTY(EditAnywhere, Category = Options)
+	bool bRemoveSources = true;
 };
 
 /**
diff --git a/Engine/Plugins/Experimental/MeshModelingToolset/Source/MeshModelingToolsEditorOnly/Public/VoxelCSGMeshesTool.h b/Engine/Plugins/Experimental/MeshModelingToolset/Source/MeshModelingToolsEditorOnly/Public/VoxelCSGMeshesTool.h
--- a/Engine/Plugins/Experimental/MeshModelingToolset/Source/MeshModelingToolsEditorOnly/Public/VoxelCSGMeshesTool.h
+++ b/Engine/Plugins/Experimental/MeshModelingToolset/Source/MeshModelingToolsEditorOnly/Public/VoxelCSGMeshesTool.h
@@ -78,6 +78,10 @@ public:
 	/** Automatically simplify the result of voxel-based merge.*/
 	UPROPERTY(EditAnywhere, Category = Options)
 	bool bAutoSimplify = true;
+
+	/** Remove the source meshes when accepting results of tool.*/
+	UPROPERTY(EditAnywhere, Category = Options)
+	bool bRemoveSources = true;
 };
 
 
