ModelingTools: MeshSurfacePointTool now tracks ctrl key state, use in SculptTool to invert brush
#rnx

diff --git a/Engine/Plugins/Experimental/MeshModelingToolset/Source/MeshModelingTools/Private/DynamicMeshSculptTool.cpp b/Engine/Plugins/Experimental/MeshModelingToolset/Source/MeshModelingTools/Private/DynamicMeshSculptTool.cpp
--- a/Engine/Plugins/Experimental/MeshModelingToolset/Source/MeshModelingTools/Private/DynamicMeshSculptTool.cpp
+++ b/Engine/Plugins/Experimental/MeshModelingToolset/Source/MeshModelingTools/Private/DynamicMeshSculptTool.cpp
@@ -221,6 +221,7 @@ bool UDynamicMeshSculptTool::HitTest(const FRay& Ray, FHitResult& OutHit)
 void UDynamicMeshSculptTool::OnBeginDrag(const FRay& Ray)
 {
 	bSmoothing = GetShiftToggle();
+	bInvert = GetCtrlToggle();
 
 	FHitResult OutHit;
 	if (HitTest(Ray, OutHit))
@@ -456,6 +457,8 @@ void UDynamicMeshSculptTool::ApplyOffsetBrush(const FRay& WorldRay)
 	FTransform Transform = ComponentTarget->GetWorldTransform();
 	FVector NewBrushPosLocal = Transform.InverseTransformPosition(LastBrushPosWorld);
 
+	double Direction = (bInvert) ? -1.0 : 1.0;
+
 	FDynamicMesh3* Mesh = DynamicMeshComponent->GetMesh();
 	for (int VertIdx : VertexROI)
 	{
@@ -464,7 +467,8 @@ void UDynamicMeshSculptTool::ApplyOffsetBrush(const FRay& WorldRay)
 		FVector3d BasePos, BaseNormal;
 		GetTargetMeshNearest(OrigPos, (double)(2 * CurrentBrushRadius), BasePos, BaseNormal);
 
-		FVector3d MoveVec = SculptProperties->OffsetPower*FMathd::Sqrt(CurrentBrushRadius)*BaseNormal;
+		FVector3d MoveVec = 
+			(Direction * SculptProperties->OffsetPower * FMathd::Sqrt(CurrentBrushRadius)) * BaseNormal;
 
 		double UseDist = (OrigPos - NewBrushPosLocal).Length();
 		double d = UseDist / CurrentBrushRadius;
@@ -503,6 +507,8 @@ void UDynamicMeshSculptTool::ApplyPinchBrush(const FRay& WorldRay)
 	FVector BrushNormalLocal = Transform.InverseTransformVectorNoScale(LastBrushPosNormalWorld);
 	FVector OffsetBrushPosLocal = NewBrushPosLocal - SculptProperties->Depth * CurrentBrushRadius * BrushNormalLocal;
 
+	double Direction = (bInvert) ? -1.0 : 1.0;
+
 	FVector3d MotionVec = NewBrushPosLocal - LastBrushPosLocal;
 	MotionVec.Normalize();
 
@@ -511,7 +517,7 @@ void UDynamicMeshSculptTool::ApplyPinchBrush(const FRay& WorldRay)
 	{
 		FVector3d OrigPos = Mesh->GetVertex(VertIdx);
 		FVector3d Delta = OffsetBrushPosLocal - OrigPos;
-		FVector3d MoveVec = SculptProperties->OffsetPower*Delta;
+		FVector3d MoveVec = (Direction * SculptProperties->OffsetPower)*Delta;
 
 		double UseDist = (OrigPos - NewBrushPosLocal).Length();
 		double d = UseDist / CurrentBrushRadius;
diff --git a/Engine/Plugins/Experimental/MeshModelingToolset/Source/MeshModelingTools/Public/DynamicMeshSculptTool.h b/Engine/Plugins/Experimental/MeshModelingToolset/Source/MeshModelingTools/Public/DynamicMeshSculptTool.h
--- a/Engine/Plugins/Experimental/MeshModelingToolset/Source/MeshModelingTools/Public/DynamicMeshSculptTool.h
+++ b/Engine/Plugins/Experimental/MeshModelingToolset/Source/MeshModelingTools/Public/DynamicMeshSculptTool.h
@@ -268,6 +268,7 @@ protected:
 	bool bTargetDirty;
 
 	bool bSmoothing;
+	bool bInvert;
 
 	bool bHaveRemeshed;
 
diff --git a/Engine/Source/Runtime/Experimental/InteractiveToolsFramework/Private/BaseTools/MeshSurfacePointTool.cpp b/Engine/Source/Runtime/Experimental/InteractiveToolsFramework/Private/BaseTools/MeshSurfacePointTool.cpp
--- a/Engine/Source/Runtime/Experimental/InteractiveToolsFramework/Private/BaseTools/MeshSurfacePointTool.cpp
+++ b/Engine/Source/Runtime/Experimental/InteractiveToolsFramework/Private/BaseTools/MeshSurfacePointTool.cpp
@@ -47,6 +47,7 @@ void UMeshSurfacePointTool::Setup()
 	UInteractiveTool::Setup();
 
 	bShiftToggle = false;
+	bCtrlToggle = false;
 
 	// add input behaviors
 	UMeshSurfacePointToolMouseBehavior* mouseBehavior = NewObject<UMeshSurfacePointToolMouseBehavior>();
@@ -94,7 +95,10 @@ void UMeshSurfacePointTool::SetShiftToggle(bool bShiftDown)
 	bShiftToggle = bShiftDown;
 }
 
-
+void UMeshSurfacePointTool::SetCtrlToggle(bool bCtrlDown)
+{
+	bCtrlToggle = bCtrlDown;
+}
 
 
 FInputRayHit UMeshSurfacePointTool::BeginHoverSequenceHitTest(const FInputDeviceRay& PressPos)
@@ -140,6 +144,7 @@ FInputCaptureRequest UMeshSurfacePointToolMouseBehavior::WantsCapture(const FInp
 FInputCaptureUpdate UMeshSurfacePointToolMouseBehavior::BeginCapture(const FInputDeviceState& input, EInputCaptureSide eSide)
 {
 	Tool->SetShiftToggle(input.bShiftKeyDown);
+	Tool->SetCtrlToggle(input.bCtrlKeyDown);
 	Tool->OnBeginDrag(input.Mouse.WorldRay);
 	LastWorldRay = input.Mouse.WorldRay;
 	bInDragCapture = true;
diff --git a/Engine/Source/Runtime/Experimental/InteractiveToolsFramework/Public/BaseTools/MeshSurfacePointTool.h b/Engine/Source/Runtime/Experimental/InteractiveToolsFramework/Public/BaseTools/MeshSurfacePointTool.h
--- a/Engine/Source/Runtime/Experimental/InteractiveToolsFramework/Public/BaseTools/MeshSurfacePointTool.h
+++ b/Engine/Source/Runtime/Experimental/InteractiveToolsFramework/Public/BaseTools/MeshSurfacePointTool.h
@@ -97,12 +97,21 @@ public:
 	virtual void SetShiftToggle(bool bShiftDown);
 
 	/** @return current state of the shift toggle */
-	virtual bool GetShiftToggle() { return bShiftToggle; }
+	virtual bool GetShiftToggle() const { return bShiftToggle; }
+
+	/** Called by registered InputBehaviors to set the state of the "shift" button (or device equivalent) */
+	virtual void SetCtrlToggle(bool bCtrlDown);
+
+	/** @return current state of the shift toggle */
+	virtual bool GetCtrlToggle() const { return bCtrlToggle; }
 
 
 protected:
 	/** Current state of the shift modifier toggle */
-	bool bShiftToggle;
+	bool bShiftToggle = false;
+
+	/** Current state of the ctrl modifier toggle */
+	bool bCtrlToggle = false;
 };
 
 
