ModelingTools: some cleanup of Sculpt tool stamps, add alternate type of Flatten brush
#rnx

diff --git a/Engine/Plugins/Experimental/MeshModelingToolset/Source/MeshModelingTools/Private/DynamicMeshSculptTool.cpp b/Engine/Plugins/Experimental/MeshModelingToolset/Source/MeshModelingTools/Private/DynamicMeshSculptTool.cpp
--- a/Engine/Plugins/Experimental/MeshModelingToolset/Source/MeshModelingTools/Private/DynamicMeshSculptTool.cpp
+++ b/Engine/Plugins/Experimental/MeshModelingToolset/Source/MeshModelingTools/Private/DynamicMeshSculptTool.cpp
@@ -242,9 +242,9 @@ void UDynamicMeshSculptTool::OnBeginDrag(const FRay& Ray)
 
 		UpdateROI(LastBrushPosLocal);
 
-		if (SculptProperties->PrimaryBrushType == EDynamicMeshSculptBrushType::Flatten)
+		if (SculptProperties->PrimaryBrushType == EDynamicMeshSculptBrushType::Plane)
 		{
-			ComputeFlattenFrame();
+			ActiveFixedBrushPlane = ComputeROIBrushPlane(LastBrushPosLocal);
 		}
 
 		// apply initial stamp
@@ -345,15 +345,39 @@ void UDynamicMeshSculptTool::ApplyStamp(const FRay& WorldRay)
 		case EDynamicMeshSculptBrushType::Pinch:
 			ApplyPinchBrush(WorldRay);
 			break;
+		case EDynamicMeshSculptBrushType::Inflate:
+			ApplyInflateBrush(WorldRay);
+			break;
 		case EDynamicMeshSculptBrushType::Flatten:
 			ApplyFlattenBrush(WorldRay);
 			break;
+		case EDynamicMeshSculptBrushType::Plane:
+			ApplyPlaneBrush(WorldRay);
+			break;
 	}
 }
 
 
 
 
+
+double UDynamicMeshSculptTool::CalculateBrushFalloff(double Distance)
+{
+	double d = Distance / CurrentBrushRadius;
+	double w = 1;
+	if (d > 0.5)
+	{
+		d = VectorUtil::Clamp((d - 0.5) / 0.5, 0.0, 1.0);
+		w = (1 - d * d);
+		w = w * w * w;
+	}
+	return w;
+}
+
+
+
+
+
 void UDynamicMeshSculptTool::ApplySmoothBrush(const FRay& WorldRay)
 {
 	// both these are useful...
@@ -374,20 +398,12 @@ void UDynamicMeshSculptTool::ApplySmoothBrush(const FRay& WorldRay)
 	{
 		FVector3d OrigPos = Mesh->GetVertex(VertIdx);
 
-		double UseDist = (OrigPos - NewBrushPosLocal).Length();
-		double d = UseDist / CurrentBrushRadius;
-		double w = 1;
-		if (d > 0.5) 
-		{
-			d = VectorUtil::Clamp((d - 0.5) / 0.5, 0.0, 1.0);
-			w = (1 - d * d);
-			w = w * w * w;
-		}
+		double Falloff = CalculateBrushFalloff(OrigPos.Distance(NewBrushPosLocal));
 
 		FVector3d SmoothedPos = (SculptProperties->SmoothingType == EMeshSculptToolSmoothType::TexturePreserving) ?
 			FMeshWeights::MeanValueCentroid(*Mesh, VertIdx) : FMeshWeights::UniformCentroid(*Mesh, VertIdx);
 
-		FVector3d NewPos = FVector3d::Lerp(OrigPos, SmoothedPos, w*SculptProperties->SmoothPower);
+		FVector3d NewPos = FVector3d::Lerp(OrigPos, SmoothedPos, Falloff*SculptProperties->SmoothPower);
 
 		Mesh->SetVertex(VertIdx, NewPos);
 		UpdateSavedVertex(VertIdx, OrigPos, NewPos);
@@ -420,16 +436,10 @@ void UDynamicMeshSculptTool::ApplyMoveBrush(const FRay& WorldRay)
 			double PrevDist = (OrigPos - LastBrushPosLocal).Length();
 			double NewDist = (OrigPos - NewBrushPosLocal).Length();
 			double UseDist = FMath::Min(PrevDist, NewDist);
-			double d = UseDist / CurrentBrushRadius;
-			double w = 1;
-			if (d > 0.5) 
-			{
-				d = VectorUtil::Clamp((d - 0.5) / 0.5, 0.0, 1.0);
-				w = (1 - d * d);
-				w = w * w * w;
-			}
 
-			FVector3d NewPos = OrigPos + w * MoveVec;
+			double Falloff = CalculateBrushFalloff(UseDist);
+
+			FVector3d NewPos = OrigPos + Falloff * MoveVec;
 
 			Mesh->SetVertex(VertIdx, NewPos);
 			UpdateSavedVertex(VertIdx, OrigPos, NewPos);
@@ -470,17 +480,9 @@ void UDynamicMeshSculptTool::ApplyOffsetBrush(const FRay& WorldRay)
 		FVector3d MoveVec = 
 			(Direction * SculptProperties->OffsetPower * FMathd::Sqrt(CurrentBrushRadius)) * BaseNormal;
 
-		double UseDist = (OrigPos - NewBrushPosLocal).Length();
-		double d = UseDist / CurrentBrushRadius;
-		double w = 1;
-		if (d > 0.5) 
-		{
-			d = VectorUtil::Clamp((d - 0.5) / 0.5, 0.0, 1.0);
-			w = (1 - d * d);
-			w = w * w * w;
-		}
+		double Falloff = CalculateBrushFalloff(OrigPos.Distance(NewBrushPosLocal));
 
-		FVector3d NewPos = OrigPos + w * MoveVec;
+		FVector3d NewPos = OrigPos + Falloff * MoveVec;
 
 		Mesh->SetVertex(VertIdx, NewPos);
 		UpdateSavedVertex(VertIdx, OrigPos, NewPos);
@@ -519,21 +521,15 @@ void UDynamicMeshSculptTool::ApplyPinchBrush(const FRay& WorldRay)
 		FVector3d Delta = OffsetBrushPosLocal - OrigPos;
 		FVector3d MoveVec = (Direction * SculptProperties->OffsetPower)*Delta;
 
-		double UseDist = (OrigPos - NewBrushPosLocal).Length();
-		double d = UseDist / CurrentBrushRadius;
-		double w = 1;
-		if (d > 0.5)
-		{
-			d = VectorUtil::Clamp((d - 0.5) / 0.5, 0.0, 1.0);
-			w = (1 - d * d);
-			w = w * w * w;
-		}
-		w = w * w * w;
+		double Falloff = CalculateBrushFalloff(OrigPos.Distance(NewBrushPosLocal));
+
+		// make tighter
+		Falloff = Falloff * Falloff * Falloff;
 
 		double AnglePower = 1.0 - FMathd::Abs(MoveVec.Normalized().Dot(MotionVec));
-		w *= AnglePower;
+		Falloff *= AnglePower;
 
-		FVector3d NewPos = OrigPos + w * MoveVec;
+		FVector3d NewPos = OrigPos + Falloff * MoveVec;
 
 		Mesh->SetVertex(VertIdx, NewPos);
 		UpdateSavedVertex(VertIdx, OrigPos, NewPos);
@@ -547,26 +543,72 @@ void UDynamicMeshSculptTool::ApplyPinchBrush(const FRay& WorldRay)
 
 
 
-void UDynamicMeshSculptTool::ComputeFlattenFrame()
+FFrame3d UDynamicMeshSculptTool::ComputeROIBrushPlane(const FVector3d& BrushCenter)
 {
 	FDynamicMesh3* Mesh = DynamicMeshComponent->GetMesh();
 	FVector3d AverageNormal(0, 0, 0);
 	FVector3d AveragePos(0, 0, 0);
+	double WeightSum = 0;
 	for (int TriID : TriangleROI)
 	{
-		AverageNormal += Mesh->GetTriNormal(TriID);
-		AveragePos += Mesh->GetTriCentroid(TriID);
+		FVector3d Centroid = Mesh->GetTriCentroid(TriID);
+		double Weight = CalculateBrushFalloff(BrushCenter.Distance(Centroid));
+
+		AverageNormal += Weight * Mesh->GetTriNormal(TriID);
+		AveragePos += Weight * Centroid;
+		WeightSum += Weight;
 	}
 	AverageNormal.Normalize();
-	AveragePos /= TriangleROI.Num();
-	ActiveFlattenFrame = FFrame3d(AveragePos, AverageNormal);
+	AveragePos /= WeightSum;
 
-	ActiveFlattenFrame.Origin -= SculptProperties->Depth * CurrentBrushRadius * ActiveFlattenFrame.Z();
+	FFrame3d Result = FFrame3d(AveragePos, AverageNormal);
+	Result.Origin -= SculptProperties->Depth * CurrentBrushRadius * Result.Z();
+
+	return Result;
 }
 
+void UDynamicMeshSculptTool::ApplyPlaneBrush(const FRay& WorldRay)
+{
+	bool bHit = UpdateBrushPositionOnSculptMesh(WorldRay);
+	if (bHit == false)
+	{
+		return;
+	}
+
+	FTransform Transform = ComponentTarget->GetWorldTransform();
+	FVector NewBrushPosLocal = Transform.InverseTransformPosition(LastBrushPosWorld);
+	FVector BrushNormalLocal = Transform.InverseTransformVectorNoScale(LastBrushPosNormalWorld);
+	FVector OffsetBrushPosLocal = NewBrushPosLocal - SculptProperties->Depth * CurrentBrushRadius * BrushNormalLocal;
+
+	FDynamicMesh3* Mesh = DynamicMeshComponent->GetMesh();
+
+	for (int VertIdx : VertexROI)
+	{
+		FVector3d OrigPos = Mesh->GetVertex(VertIdx);
+		FVector3d PlanePos = ActiveFixedBrushPlane.ToPlane(OrigPos, 2);
+		FVector3d Delta = PlanePos - OrigPos;
+		FVector3d MoveVec = FMathd::Sqrt(SculptProperties->OffsetPower) * Delta;
+
+		double Falloff = CalculateBrushFalloff(OrigPos.Distance(NewBrushPosLocal));
+
+		FVector3d NewPos = OrigPos + Falloff * MoveVec;
+
+		Mesh->SetVertex(VertIdx, NewPos);
+		UpdateSavedVertex(VertIdx, OrigPos, NewPos);
+	}
+
+	bRemeshPending = bEnableRemeshing;
+
+	LastBrushPosLocal = NewBrushPosLocal;
+}
+
+
+
+
+
 void UDynamicMeshSculptTool::ApplyFlattenBrush(const FRay& WorldRay)
 {
-	bool bHit = UpdateBrushPositionOnTargetMesh(WorldRay);
+	bool bHit = UpdateBrushPositionOnSculptMesh(WorldRay);
 	if (bHit == false)
 	{
 		return;
@@ -579,24 +621,18 @@ void UDynamicMeshSculptTool::ApplyFlattenBrush(const FRay& WorldRay)
 
 	FDynamicMesh3* Mesh = DynamicMeshComponent->GetMesh();
 
+	FFrame3d StampFlattenPlane = ComputeROIBrushPlane(NewBrushPosLocal);
+
 	for (int VertIdx : VertexROI)
 	{
 		FVector3d OrigPos = Mesh->GetVertex(VertIdx);
-		FVector3d PlanePos = ActiveFlattenFrame.ToPlane(OrigPos, 2);
+		FVector3d PlanePos = StampFlattenPlane.ToPlane(OrigPos, 2);
 		FVector3d Delta = PlanePos - OrigPos;
 		FVector3d MoveVec = FMathd::Sqrt(SculptProperties->OffsetPower) * Delta;
 
-		double UseDist = (OrigPos - NewBrushPosLocal).Length();
-		double d = UseDist / CurrentBrushRadius;
-		double w = 1;
-		if (d > 0.5)
-		{
-			d = VectorUtil::Clamp((d - 0.5) / 0.5, 0.0, 1.0);
-			w = (1 - d * d);
-			w = w * w * w;
-		}
+		double Falloff = CalculateBrushFalloff(OrigPos.Distance(NewBrushPosLocal));
 
-		FVector3d NewPos = OrigPos + w * MoveVec;
+		FVector3d NewPos = OrigPos + Falloff * MoveVec;
 
 		Mesh->SetVertex(VertIdx, NewPos);
 		UpdateSavedVertex(VertIdx, OrigPos, NewPos);
@@ -609,6 +645,52 @@ void UDynamicMeshSculptTool::ApplyFlattenBrush(const FRay& WorldRay)
 
 
 
+void UDynamicMeshSculptTool::ApplyInflateBrush(const FRay& WorldRay)
+{
+	bool bHit = UpdateBrushPositionOnSculptMesh(WorldRay);
+	if (bHit == false)
+	{
+		return;
+	}
+
+	FTransform Transform = ComponentTarget->GetWorldTransform();
+	FVector NewBrushPosLocal = Transform.InverseTransformPosition(LastBrushPosWorld);
+
+	double Direction = (bInvert) ? -1.0 : 1.0;
+
+	FDynamicMesh3* Mesh = DynamicMeshComponent->GetMesh();
+
+	// calculate vertex normals
+	ParallelFor(VertexROI.Num(), [this, Mesh](int Index) {
+		int VertIdx = VertexROI[Index];
+		FVector3d Normal = FMeshNormals::ComputeVertexNormal(*Mesh, VertIdx);
+		Mesh->SetVertexNormal(VertIdx, (FVector3f)Normal);
+	});
+
+
+	ParallelFor(VertexROI.Num(), [this, Mesh, Direction, NewBrushPosLocal](int Index) 
+	{
+		int VertIdx = VertexROI[Index];
+		FVector3d OrigPos = Mesh->GetVertex(VertIdx);
+		FVector3d Normal = (FVector3d)Mesh->GetVertexNormal(VertIdx);
+
+		FVector3d MoveVec =
+			(Direction * SculptProperties->OffsetPower * FMathd::Sqrt(CurrentBrushRadius)) * Normal;
+
+		double Falloff = CalculateBrushFalloff(OrigPos.Distance(NewBrushPosLocal));
+
+		FVector3d NewPos = OrigPos + Falloff*MoveVec;
+		Mesh->SetVertex(VertIdx, NewPos);
+
+		UpdateSavedVertex(VertIdx, OrigPos, NewPos);
+	});
+
+	bRemeshPending = bEnableRemeshing;
+
+	LastBrushPosLocal = NewBrushPosLocal;
+}
+
+
 
 
 
@@ -647,6 +729,31 @@ bool UDynamicMeshSculptTool::UpdateBrushPositionOnTargetMesh(const FRay& WorldRa
 }
 
 
+bool UDynamicMeshSculptTool::UpdateBrushPositionOnSculptMesh(const FRay& WorldRay)
+{
+	FTransform Transform = ComponentTarget->GetWorldTransform();
+
+	FRay3d LocalRay(Transform.InverseTransformPosition(WorldRay.Origin),
+		Transform.InverseTransformVectorNoScale(WorldRay.Direction));
+	LocalRay.Direction.Normalize();
+
+	int HitTID = DynamicMeshComponent->GetOctree()->FindNearestHitObject(LocalRay);
+	if (HitTID != IndexConstants::InvalidID)
+	{
+		const FDynamicMesh3* SculptMesh = DynamicMeshComponent->GetMesh();
+
+		FTriangle3d Triangle;
+		SculptMesh->GetTriVertices(HitTID, Triangle.V[0], Triangle.V[1], Triangle.V[2]);
+		FIntrRay3Triangle3d Query(LocalRay, Triangle);
+		Query.Find();
+
+		LastBrushPosNormalWorld = Transform.TransformVectorNoScale(SculptMesh->GetTriNormal(HitTID));
+		LastBrushPosWorld = Transform.TransformPosition(LocalRay.PointAt(Query.RayParameter));
+		return true;
+	}
+	return false;
+}
+
 
 
 void UDynamicMeshSculptTool::OnEndDrag(const FRay& Ray)
@@ -1184,7 +1291,9 @@ void UDynamicMeshSculptTool::UpdateSavedVertex(int vid, const FVector3d& OldPosi
 {
 	if (ActiveVertexChange != nullptr)
 	{
+		UpdateSavedVertexLock.Lock();
 		ActiveVertexChange->UpdateVertex(vid, OldPosition, NewPosition);
+		UpdateSavedVertexLock.Unlock();
 	}
 }
 
diff --git a/Engine/Plugins/Experimental/MeshModelingToolset/Source/MeshModelingTools/Public/DynamicMeshSculptTool.h b/Engine/Plugins/Experimental/MeshModelingToolset/Source/MeshModelingTools/Public/DynamicMeshSculptTool.h
--- a/Engine/Plugins/Experimental/MeshModelingToolset/Source/MeshModelingTools/Public/DynamicMeshSculptTool.h
+++ b/Engine/Plugins/Experimental/MeshModelingToolset/Source/MeshModelingTools/Public/DynamicMeshSculptTool.h
@@ -47,13 +47,20 @@ enum class EDynamicMeshSculptBrushType : uint8
 	Pull UMETA(DisplayName = "Pull"),
 
 	/** Offset brush */
-	Offset UMETA(DisplayName = "Offset"),
+	Offset UMETA(DisplayName = "Sculpt"),
+
+	/** Inflate brush */
+	Inflate UMETA(DisplayName = "Inflate"),
 
 	/** Pinch brush */
 	Pinch UMETA(DisplayName = "Pinch"),
 
 	/** Flatten brush */
 	Flatten UMETA(DisplayName = "Flatten"),
+
+	/** Plane brush */
+	Plane UMETA(DisplayName = "Plane"),
+
 };
 
 
@@ -284,14 +291,20 @@ protected:
 
 	bool UpdateBrushPositionOnActivePlane(const FRay& WorldRay);
 	bool UpdateBrushPositionOnTargetMesh(const FRay& WorldRay);
+	bool UpdateBrushPositionOnSculptMesh(const FRay& WorldRay);
+
 	void ApplySmoothBrush(const FRay& WorldRay);
 	void ApplyMoveBrush(const FRay& WorldRay);
 	void ApplyOffsetBrush(const FRay& WorldRay);
 	void ApplyPinchBrush(const FRay& WorldRay);
+	void ApplyInflateBrush(const FRay& WorldRay);
+	void ApplyPlaneBrush(const FRay& WorldRay);
 	void ApplyFlattenBrush(const FRay& WorldRay);
 
-	FFrame3d ActiveFlattenFrame;
-	void ComputeFlattenFrame();
+	double CalculateBrushFalloff(double Distance);
+
+	FFrame3d ActiveFixedBrushPlane;
+	FFrame3d ComputeROIBrushPlane(const FVector3d& BrushCenter);
 
 	TArray<int> TrianglesBuffer;
 	TArray<int> NormalsBuffer;
@@ -306,7 +319,9 @@ protected:
 	void BeginChange(bool bIsVertexChange);
 	void EndChange();
 	void SaveActiveROI();
+
 	void UpdateSavedVertex(int vid, const FVector3d& OldPosition, const FVector3d& NewPosition);
+	FCriticalSection UpdateSavedVertexLock;
 
 	double EstimateIntialSafeTargetLength(const FDynamicMesh3& Mesh, int MinTargetTriCount);
 };
