ModelingTools: Sculpt tool improvements. New hotkeys to change active brush and size/speed. Renamed various properties and brushes. Added SculptMax brush,  option to freeze target surface.
#rnx

diff --git a/Engine/Plugins/Experimental/MeshModelingToolset/Source/MeshModelingTools/Private/DynamicMeshSculptTool.cpp b/Engine/Plugins/Experimental/MeshModelingToolset/Source/MeshModelingTools/Private/DynamicMeshSculptTool.cpp
--- a/Engine/Plugins/Experimental/MeshModelingToolset/Source/MeshModelingTools/Private/DynamicMeshSculptTool.cpp
+++ b/Engine/Plugins/Experimental/MeshModelingToolset/Source/MeshModelingTools/Private/DynamicMeshSculptTool.cpp
@@ -38,11 +38,12 @@ UMeshSurfacePointTool* UDynamicMeshSculptToolBuilder::CreateNewTool(const FToolB
 
 UBrushSculptProperties::UBrushSculptProperties()
 {
-	SmoothPower = 0.25;
-	OffsetPower = 0.5;
+	SmoothSpeed = 0.25;
+	BrushSpeed = 0.5;
 	PrimaryBrushType = EDynamicMeshSculptBrushType::Pull;
 	SmoothingType = EMeshSculptToolSmoothType::TexturePreserving;
 	Depth = 0;
+	bFreezeTarget = false;
 }
 
 
@@ -91,7 +92,7 @@ void UDynamicMeshSculptTool::Setup()
 	// initialize AABBTree
 	FDynamicMesh3* Mesh = DynamicMeshComponent->GetMesh();
 
-	// do we always want to keep vertex normals updated?
+	// do we always want to keep vertex normals updated? Perhaps should discard vertex normals before baking?
 	FMeshNormals::QuickComputeVertexNormals(*Mesh);
 
 	// switch to vertex normals for testing
@@ -139,6 +140,10 @@ void UDynamicMeshSculptTool::Setup()
 	{
 		AddToolPropertySource(RemeshProperties);
 	}
+
+	GetToolManager()->DisplayMessage(
+		LOCTEXT("OnStartSculptTool", "Shift key toggles to Smoothing Brush, Ctrl inverts Brush (where applicable). Q/A keys cycle through Brush Types, Shift+Q/A for Brush History. S/D change Size (shift to small-step), W/E change Speed."),
+		EToolMessageLevel::UserNotification);
 }
 
 
@@ -324,6 +329,26 @@ void UDynamicMeshSculptTool::ApplyStamp(const FRay& WorldRay)
 {
 	SCOPE_CYCLE_COUNTER(STAT_SculptToolApplyStamp);
 
+	// update brush type history. apologies for convoluted logic. 
+	StampTimestamp++;
+	if (LastStampType != PendingStampType)
+	{
+		if (BrushTypeHistoryIndex != BrushTypeHistory.Num() - 1)
+		{
+			if (LastStampType != EDynamicMeshSculptBrushType::LastValue)
+			{
+				BrushTypeHistory.Add(LastStampType);
+			}
+			BrushTypeHistoryIndex = BrushTypeHistory.Num()-1;
+		}
+		LastStampType = PendingStampType;
+		if (BrushTypeHistory.Num() == 0 || BrushTypeHistory[BrushTypeHistory.Num()-1] != PendingStampType)
+		{
+			BrushTypeHistory.Add(PendingStampType);
+			BrushTypeHistoryIndex = BrushTypeHistory.Num() - 1;
+		}
+	}
+
 	CalculateBrushRadius();
 
 	SaveActiveROI();
@@ -339,6 +364,9 @@ void UDynamicMeshSculptTool::ApplyStamp(const FRay& WorldRay)
 		case EDynamicMeshSculptBrushType::Offset:
 			ApplyOffsetBrush(WorldRay);
 			break;
+		case EDynamicMeshSculptBrushType::SculptMax:
+			ApplySculptMaxBrush(WorldRay);
+			break;
 		case EDynamicMeshSculptBrushType::Pull:
 			ApplyMoveBrush(WorldRay);
 			break;
@@ -403,7 +431,7 @@ void UDynamicMeshSculptTool::ApplySmoothBrush(const FRay& WorldRay)
 		FVector3d SmoothedPos = (SculptProperties->SmoothingType == EMeshSculptToolSmoothType::TexturePreserving) ?
 			FMeshWeights::MeanValueCentroid(*Mesh, VertIdx) : FMeshWeights::UniformCentroid(*Mesh, VertIdx);
 
-		FVector3d NewPos = FVector3d::Lerp(OrigPos, SmoothedPos, Falloff*SculptProperties->SmoothPower);
+		FVector3d NewPos = FVector3d::Lerp(OrigPos, SmoothedPos, Falloff*SculptProperties->SmoothSpeed);
 
 		Mesh->SetVertex(VertIdx, NewPos);
 		UpdateSavedVertex(VertIdx, OrigPos, NewPos);
@@ -468,6 +496,7 @@ void UDynamicMeshSculptTool::ApplyOffsetBrush(const FRay& WorldRay)
 	FVector NewBrushPosLocal = Transform.InverseTransformPosition(LastBrushPosWorld);
 
 	double Direction = (bInvert) ? -1.0 : 1.0;
+	double UseSpeed = Direction * FMathd::Sqrt(CurrentBrushRadius) * (SculptProperties->BrushSpeed*0.25);
 
 	FDynamicMesh3* Mesh = DynamicMeshComponent->GetMesh();
 	for (int VertIdx : VertexROI)
@@ -477,8 +506,7 @@ void UDynamicMeshSculptTool::ApplyOffsetBrush(const FRay& WorldRay)
 		FVector3d BasePos, BaseNormal;
 		GetTargetMeshNearest(OrigPos, (double)(2 * CurrentBrushRadius), BasePos, BaseNormal);
 
-		FVector3d MoveVec = 
-			(Direction * SculptProperties->OffsetPower * FMathd::Sqrt(CurrentBrushRadius)) * BaseNormal;
+		FVector3d MoveVec = UseSpeed * BaseNormal;
 
 		double Falloff = CalculateBrushFalloff(OrigPos.Distance(NewBrushPosLocal));
 
@@ -496,6 +524,56 @@ void UDynamicMeshSculptTool::ApplyOffsetBrush(const FRay& WorldRay)
 
 
 
+
+void UDynamicMeshSculptTool::ApplySculptMaxBrush(const FRay& WorldRay)
+{
+	bool bHit = UpdateBrushPositionOnTargetMesh(WorldRay);
+	if (bHit == false)
+	{
+		return;
+	}
+
+	FTransform Transform = ComponentTarget->GetWorldTransform();
+	FVector NewBrushPosLocal = Transform.InverseTransformPosition(LastBrushPosWorld);
+
+	double Direction = (bInvert) ? -1.0 : 1.0;
+	double UseSpeed = Direction * FMathd::Sqrt(CurrentBrushRadius) * (SculptProperties->BrushSpeed*0.25);
+
+	double MaxOffset = CurrentBrushRadius * 0.5;
+	double MaxOffsetSqr = MaxOffset * MaxOffset;
+
+	FDynamicMesh3* Mesh = DynamicMeshComponent->GetMesh();
+	for (int VertIdx : VertexROI)
+	{
+		FVector3d OrigPos = Mesh->GetVertex(VertIdx);
+
+		FVector3d BasePos, BaseNormal;
+		GetTargetMeshNearest(OrigPos, (double)(2 * CurrentBrushRadius), BasePos, BaseNormal);
+
+		FVector3d MoveVec = UseSpeed * BaseNormal;
+
+		double Falloff = CalculateBrushFalloff(OrigPos.Distance(NewBrushPosLocal));
+
+		FVector3d NewPos = OrigPos + Falloff * MoveVec;
+
+		FVector3d DeltaPos = NewPos - BasePos;
+		if (DeltaPos.SquaredLength() > MaxOffsetSqr)
+		{
+			DeltaPos.Normalize();
+			NewPos = BasePos + MaxOffset * DeltaPos;
+		}
+
+		Mesh->SetVertex(VertIdx, NewPos);
+		UpdateSavedVertex(VertIdx, OrigPos, NewPos);
+	}
+
+	bRemeshPending = bEnableRemeshing;
+
+	LastBrushPosLocal = NewBrushPosLocal;
+}
+
+
+
 void UDynamicMeshSculptTool::ApplyPinchBrush(const FRay& WorldRay)
 {
 	bool bHit = UpdateBrushPositionOnTargetMesh(WorldRay);
@@ -510,6 +588,7 @@ void UDynamicMeshSculptTool::ApplyPinchBrush(const FRay& WorldRay)
 	FVector OffsetBrushPosLocal = NewBrushPosLocal - SculptProperties->Depth * CurrentBrushRadius * BrushNormalLocal;
 
 	double Direction = (bInvert) ? -1.0 : 1.0;
+	double UseSpeed = Direction * FMathd::Sqrt(CurrentBrushRadius) * (SculptProperties->BrushSpeed*0.25);
 
 	FVector3d MotionVec = NewBrushPosLocal - LastBrushPosLocal;
 	MotionVec.Normalize();
@@ -519,7 +598,7 @@ void UDynamicMeshSculptTool::ApplyPinchBrush(const FRay& WorldRay)
 	{
 		FVector3d OrigPos = Mesh->GetVertex(VertIdx);
 		FVector3d Delta = OffsetBrushPosLocal - OrigPos;
-		FVector3d MoveVec = (Direction * SculptProperties->OffsetPower)*Delta;
+		FVector3d MoveVec = UseSpeed * Delta;
 
 		double Falloff = CalculateBrushFalloff(OrigPos.Distance(NewBrushPosLocal));
 
@@ -581,13 +660,14 @@ void UDynamicMeshSculptTool::ApplyPlaneBrush(const FRay& WorldRay)
 	FVector OffsetBrushPosLocal = NewBrushPosLocal - SculptProperties->Depth * CurrentBrushRadius * BrushNormalLocal;
 
 	FDynamicMesh3* Mesh = DynamicMeshComponent->GetMesh();
+	double UseSpeed = FMathd::Sqrt(CurrentBrushRadius) * FMathd::Sqrt(SculptProperties->BrushSpeed) * 0.25;
 
 	for (int VertIdx : VertexROI)
 	{
 		FVector3d OrigPos = Mesh->GetVertex(VertIdx);
 		FVector3d PlanePos = ActiveFixedBrushPlane.ToPlane(OrigPos, 2);
 		FVector3d Delta = PlanePos - OrigPos;
-		FVector3d MoveVec = FMathd::Sqrt(SculptProperties->OffsetPower) * Delta;
+		FVector3d MoveVec = UseSpeed * Delta;
 
 		double Falloff = CalculateBrushFalloff(OrigPos.Distance(NewBrushPosLocal));
 
@@ -620,7 +700,7 @@ void UDynamicMeshSculptTool::ApplyFlattenBrush(const FRay& WorldRay)
 	FVector OffsetBrushPosLocal = NewBrushPosLocal - SculptProperties->Depth * CurrentBrushRadius * BrushNormalLocal;
 
 	FDynamicMesh3* Mesh = DynamicMeshComponent->GetMesh();
-
+	double UseSpeed = FMathd::Sqrt(CurrentBrushRadius) * FMathd::Sqrt(SculptProperties->BrushSpeed) * 0.25;
 	FFrame3d StampFlattenPlane = ComputeROIBrushPlane(NewBrushPosLocal);
 
 	for (int VertIdx : VertexROI)
@@ -628,7 +708,7 @@ void UDynamicMeshSculptTool::ApplyFlattenBrush(const FRay& WorldRay)
 		FVector3d OrigPos = Mesh->GetVertex(VertIdx);
 		FVector3d PlanePos = StampFlattenPlane.ToPlane(OrigPos, 2);
 		FVector3d Delta = PlanePos - OrigPos;
-		FVector3d MoveVec = FMathd::Sqrt(SculptProperties->OffsetPower) * Delta;
+		FVector3d MoveVec = UseSpeed * Delta;
 
 		double Falloff = CalculateBrushFalloff(OrigPos.Distance(NewBrushPosLocal));
 
@@ -657,7 +737,7 @@ void UDynamicMeshSculptTool::ApplyInflateBrush(const FRay& WorldRay)
 	FVector NewBrushPosLocal = Transform.InverseTransformPosition(LastBrushPosWorld);
 
 	double Direction = (bInvert) ? -1.0 : 1.0;
-
+	double UseSpeed = Direction * FMathd::Sqrt(CurrentBrushRadius) * SculptProperties->BrushSpeed * 0.25;
 	FDynamicMesh3* Mesh = DynamicMeshComponent->GetMesh();
 
 	// calculate vertex normals
@@ -668,14 +748,13 @@ void UDynamicMeshSculptTool::ApplyInflateBrush(const FRay& WorldRay)
 	});
 
 
-	ParallelFor(VertexROI.Num(), [this, Mesh, Direction, NewBrushPosLocal](int Index) 
+	ParallelFor(VertexROI.Num(), [this, Mesh, UseSpeed, NewBrushPosLocal](int Index) 
 	{
 		int VertIdx = VertexROI[Index];
 		FVector3d OrigPos = Mesh->GetVertex(VertIdx);
 		FVector3d Normal = (FVector3d)Mesh->GetVertexNormal(VertIdx);
 
-		FVector3d MoveVec =
-			(Direction * SculptProperties->OffsetPower * FMathd::Sqrt(CurrentBrushRadius)) * Normal;
+		FVector3d MoveVec = UseSpeed * Normal;
 
 		double Falloff = CalculateBrushFalloff(OrigPos.Distance(NewBrushPosLocal));
 
@@ -774,6 +853,8 @@ void UDynamicMeshSculptTool::OnEndDrag(const FRay& Ray)
 
 bool UDynamicMeshSculptTool::OnUpdateHover(const FInputDeviceRay& DevicePos)
 {
+	PendingStampType = SculptProperties->PrimaryBrushType;
+
 	if (bInDrag)
 	{
 		FVector NewHitPosWorld = FMath::RayPlaneIntersection(DevicePos.WorldRay.Origin, DevicePos.WorldRay.Direction, ActiveDragPlane);
@@ -812,6 +893,12 @@ void UDynamicMeshSculptTool::Tick(float DeltaTime)
 	UMeshSurfacePointTool::Tick(DeltaTime);
 	Indicators->Tick(DeltaTime);
 
+	// if user changed to not-frozen, we need to update the target
+	if (bCachedFreezeTarget != SculptProperties->bFreezeTarget)
+	{
+		UpdateTarget();
+	}
+
 	bool bMeshModified = false;
 	bool bMeshShapeModified = false;
 
@@ -1143,6 +1230,15 @@ void UDynamicMeshSculptTool::RecalculateNormals_Overlay()
 
 void UDynamicMeshSculptTool::UpdateTarget()
 {
+	if (SculptProperties != nullptr )
+	{
+		bCachedFreezeTarget = SculptProperties->bFreezeTarget;
+		if (SculptProperties->bFreezeTarget)
+		{
+			return;   // do not update frozen target
+		}
+	}
+
 	BrushTargetMesh.Copy(*DynamicMeshComponent->GetMesh(), false, false, false, false);
 	BrushTargetMeshSpatial.SetMesh(&BrushTargetMesh, true);
 
@@ -1211,6 +1307,76 @@ void UDynamicMeshSculptTool::DecreaseBrushRadiusAction()
 	CalculateBrushRadius();
 }
 
+void UDynamicMeshSculptTool::IncreaseBrushRadiusSmallStepAction()
+{
+	BrushProperties->BrushSize = FMath::Clamp(BrushProperties->BrushSize + 0.005f, 0.0f, 1.0f);
+	CalculateBrushRadius();
+}
+
+void UDynamicMeshSculptTool::DecreaseBrushRadiusSmallStepAction()
+{
+	BrushProperties->BrushSize = FMath::Clamp(BrushProperties->BrushSize - 0.005f, 0.0f, 1.0f);
+	CalculateBrushRadius();
+}
+
+
+
+void UDynamicMeshSculptTool::IncreaseBrushSpeedAction()
+{
+	SculptProperties->BrushSpeed = FMath::Clamp(SculptProperties->BrushSpeed + 0.05f, 0.0f, 1.0f);
+}
+
+void UDynamicMeshSculptTool::DecreaseBrushSpeedAction()
+{
+	SculptProperties->BrushSpeed = FMath::Clamp(SculptProperties->BrushSpeed - 0.05f, 0.0f, 1.0f);
+}
+
+
+
+void UDynamicMeshSculptTool::NextBrushModeAction()
+{
+	uint8 LastMode = (uint8)EDynamicMeshSculptBrushType::LastValue;
+	SculptProperties->PrimaryBrushType = (EDynamicMeshSculptBrushType)(((uint8)SculptProperties->PrimaryBrushType + 1) % LastMode);
+}
+
+void UDynamicMeshSculptTool::PreviousBrushModeAction()
+{
+	uint8 LastMode = (uint8)EDynamicMeshSculptBrushType::LastValue;
+	uint8 CurMode = (uint8)SculptProperties->PrimaryBrushType;
+	if (CurMode == 0)
+	{
+		SculptProperties->PrimaryBrushType = (EDynamicMeshSculptBrushType)((uint8)LastMode - 1);
+	} 
+	else
+	{
+		SculptProperties->PrimaryBrushType = (EDynamicMeshSculptBrushType)((uint8)CurMode - 1);
+	}
+}
+
+
+void UDynamicMeshSculptTool::NextHistoryBrushModeAction()
+{
+	int MaxHistory = BrushTypeHistory.Num() - 1;
+	if (BrushTypeHistoryIndex < MaxHistory)
+	{
+		BrushTypeHistoryIndex++;
+		SculptProperties->PrimaryBrushType = BrushTypeHistory[BrushTypeHistoryIndex];
+		LastStampType = SculptProperties->PrimaryBrushType;
+	}
+}
+
+void UDynamicMeshSculptTool::PreviousHistoryBrushModeAction()
+{
+	if (BrushTypeHistoryIndex > 0)
+	{
+		BrushTypeHistoryIndex--;
+		SculptProperties->PrimaryBrushType = BrushTypeHistory[BrushTypeHistoryIndex];
+		LastStampType = SculptProperties->PrimaryBrushType;
+	}
+}
+
+
+
 void UDynamicMeshSculptTool::RegisterActions(FInteractiveToolActionSet& ActionSet)
 {
 	ActionSet.RegisterAction(this, (int32)EStandardToolActions::IncreaseBrushSize,
@@ -1226,6 +1392,84 @@ void UDynamicMeshSculptTool::RegisterActions(FInteractiveToolActionSet& ActionSe
 		LOCTEXT("SculptDecreaseRadiusTooltip", "Decrease radius of sculpting brush"),
 		EModifierKey::None, EKeys::LeftBracket,
 		[this]() { DecreaseBrushRadiusAction(); } );
+
+
+	ActionSet.RegisterAction(this, (int32)EStandardToolActions::BaseClientDefinedActionID+1,
+		TEXT("NextBrushMode"),
+		LOCTEXT("SculptNextBrushMode", "Next Brush Type"),
+		LOCTEXT("SculptNextBrushModeTooltip", "Cycle to next Brush Type"),
+		EModifierKey::None, EKeys::A,
+		[this]() { NextBrushModeAction(); });
+
+	ActionSet.RegisterAction(this, (int32)EStandardToolActions::BaseClientDefinedActionID+2,
+		TEXT("PreviousBrushMode"),
+		LOCTEXT("SculptPreviousBrushMode", "Previous Brush Type"),
+		LOCTEXT("SculptPreviousBrushModeTooltip", "Cycle to previous Brush Type"),
+		EModifierKey::None, EKeys::Q,
+		[this]() { PreviousBrushModeAction(); });
+
+
+	ActionSet.RegisterAction(this, (int32)EStandardToolActions::BaseClientDefinedActionID + 10,
+		TEXT("NextBrushHistoryState"),
+		LOCTEXT("SculptNextBrushHistoryState", "Next Brush History State"),
+		LOCTEXT("SculptSculptNextBrushHistoryStateTooltip", "Cycle to next Brush History state"),
+		EModifierKey::Shift, EKeys::Q,
+		[this]() { NextHistoryBrushModeAction(); });
+
+	ActionSet.RegisterAction(this, (int32)EStandardToolActions::BaseClientDefinedActionID + 11,
+		TEXT("PreviousBrushHistoryState"),
+		LOCTEXT("SculptPreviousBrushHistoryState", "Previous Brush History State"),
+		LOCTEXT("SculptPreviousBrushHistoryStateTooltip", "Cycle to previous Brush History state"),
+		EModifierKey::Shift, EKeys::A,
+		[this]() { PreviousHistoryBrushModeAction(); });
+
+
+
+	ActionSet.RegisterAction(this, (int32)EStandardToolActions::BaseClientDefinedActionID + 50,
+		TEXT("SculptIncreaseSize"),
+		LOCTEXT("SculptIncreaseSize", "Increase Size"),
+		LOCTEXT("SculptIncreaseSizeTooltip", "Increase Brush Size"),
+		EModifierKey::None, EKeys::D,
+		[this]() { IncreaseBrushRadiusAction(); });
+
+	ActionSet.RegisterAction(this, (int32)EStandardToolActions::BaseClientDefinedActionID + 51,
+		TEXT("SculptDecreaseSize"),
+		LOCTEXT("SculptDecreaseSize", "Decrease Size"),
+		LOCTEXT("SculptDecreaseSizeTooltip", "Decrease Brush Size"),
+		EModifierKey::None, EKeys::S,
+		[this]() { DecreaseBrushRadiusAction(); });
+
+	ActionSet.RegisterAction(this, (int32)EStandardToolActions::BaseClientDefinedActionID + 52,
+		TEXT("SculptIncreaseSizeSmallStep"),
+		LOCTEXT("SculptIncreaseSize", "Increase Size"),
+		LOCTEXT("SculptIncreaseSizeTooltip", "Increase Brush Size"),
+		EModifierKey::Shift, EKeys::D,
+		[this]() { IncreaseBrushRadiusSmallStepAction(); });
+
+	ActionSet.RegisterAction(this, (int32)EStandardToolActions::BaseClientDefinedActionID + 53,
+		TEXT("SculptDecreaseSizeSmallStemp"),
+		LOCTEXT("SculptDecreaseSize", "Decrease Size"),
+		LOCTEXT("SculptDecreaseSizeTooltip", "Decrease Brush Size"),
+		EModifierKey::Shift, EKeys::S,
+		[this]() { DecreaseBrushRadiusSmallStepAction(); });
+
+
+
+
+	ActionSet.RegisterAction(this, (int32)EStandardToolActions::BaseClientDefinedActionID + 60,
+		TEXT("SculptIncreaseSpeed"),
+		LOCTEXT("SculptIncreaseSpeed", "Increase Speed"),
+		LOCTEXT("SculptIncreaseSpeedTooltip", "Increase Brush Speed"),
+		EModifierKey::None, EKeys::E,
+		[this]() { IncreaseBrushSpeedAction(); });
+
+	ActionSet.RegisterAction(this, (int32)EStandardToolActions::BaseClientDefinedActionID + 61,
+		TEXT("SculptDecreaseSpeed"),
+		LOCTEXT("SculptDecreaseSpeed", "Decrease Speed"),
+		LOCTEXT("SculptDecreaseSpeedTooltip", "Decrease Brush Speed"),
+		EModifierKey::None, EKeys::W,
+		[this]() { DecreaseBrushSpeedAction(); });
+
 }
 
 
diff --git a/Engine/Plugins/Experimental/MeshModelingToolset/Source/MeshModelingTools/Public/DynamicMeshSculptTool.h b/Engine/Plugins/Experimental/MeshModelingToolset/Source/MeshModelingTools/Public/DynamicMeshSculptTool.h
--- a/Engine/Plugins/Experimental/MeshModelingToolset/Source/MeshModelingTools/Public/DynamicMeshSculptTool.h
+++ b/Engine/Plugins/Experimental/MeshModelingToolset/Source/MeshModelingTools/Public/DynamicMeshSculptTool.h
@@ -46,9 +46,12 @@ enum class EDynamicMeshSculptBrushType : uint8
 	/** Move brush */
 	Pull UMETA(DisplayName = "Pull"),
 
-	/** Offset brush */
+	/** Sculpt brush */
 	Offset UMETA(DisplayName = "Sculpt"),
 
+	/** SculptMax brush */
+	SculptMax UMETA(DisplayName = "Sculpt Max"),
+
 	/** Inflate brush */
 	Inflate UMETA(DisplayName = "Inflate"),
 
@@ -61,6 +64,9 @@ enum class EDynamicMeshSculptBrushType : uint8
 	/** Plane brush */
 	Plane UMETA(DisplayName = "Plane"),
 
+
+	LastValue UMETA(Hidden)
+
 };
 
 
@@ -69,10 +75,10 @@ enum class EDynamicMeshSculptBrushType : uint8
 UENUM()
 enum class EMeshSculptToolSmoothType : uint8
 {
-	/** Uniform smoothing */
+	/** Uniform Smoothing */
 	Uniform UMETA(DisplayName = "Uniform"),
 
-	/** Texture-preserving smoothing */
+	/** Texture-Preserving Smoothing */
 	TexturePreserving UMETA(DisplayName = "Texture Preserving"),
 };
 
@@ -115,25 +121,31 @@ class MESHMODELINGTOOLS_API UBrushSculptProperties : public UInteractiveToolProp
 public:
 	UBrushSculptProperties();
 
-	/** primary brush mode */
+	/** Primary Brush Mode */
 	UPROPERTY(EditAnywhere, Category = Options)
 	EDynamicMeshSculptBrushType PrimaryBrushType;
 
-	/** smoothing type */
-	UPROPERTY(EditAnywhere, Category = Options)
-	EMeshSculptToolSmoothType SmoothingType;
+	/** Power/Speed of brush */
+	UPROPERTY(EditAnywhere, Category = Options, meta = (UIMin = "0.0", UIMax = "1.0", ClampMin = "0.0", ClampMax = "1.0", EditCondition = "PrimaryBrushType != EDynamicMeshSculptBrushType::Pull"))
+	float BrushSpeed;
 
-	/** Smoothing speed of brush */
+	/** Smoothing Speed of Smoothing brush */
 	UPROPERTY(EditAnywhere, Category = Options, meta = (UIMin = "0.0", UIMax = "1.0", ClampMin = "0.0", ClampMax = "1.0"))
-	float SmoothPower;
+	float SmoothSpeed;
 
-	/** Offset speed of brush */
-	UPROPERTY(EditAnywhere, Category = Options, meta = (UIMin = "0.0", UIMax = "1.0", ClampMin = "0.0", ClampMax = "1.0", EditCondition = "PrimaryBrushType != EDynamicMeshSculptBrushType::Pull"))
-	float OffsetPower;
 
-	/** Depth of Brush into surface along ray */
+	/** Depth of Brush into surface along view ray or surface normal, depending on brush */
 	UPROPERTY(EditAnywhere, Category = Options, meta = (UIMin = "-0.5", UIMax = "0.5", ClampMin = "-1.0", ClampMax = "1.0"))
 	float Depth;
+
+	/** Disable updating of the brush Target Surface after each stroke */
+	UPROPERTY(EditAnywhere, Category = Options, meta = (EditCondition = "PrimaryBrushType == EDynamicMeshSculptBrushType::Sculpt || PrimaryBrushType == EDynamicMeshSculptBrushType::SculptMax || PrimaryBrushType == EDynamicMeshSculptBrushType::Pinch" ))
+	bool bFreezeTarget;
+
+	/** Smoothing Type */
+	UPROPERTY(EditAnywhere, Category = Options, AdvancedDisplay)
+	EMeshSculptToolSmoothType SmoothingType;
+
 };
 
 
@@ -233,6 +245,18 @@ public:
 
 	virtual void IncreaseBrushRadiusAction();
 	virtual void DecreaseBrushRadiusAction();
+	virtual void IncreaseBrushRadiusSmallStepAction();
+	virtual void DecreaseBrushRadiusSmallStepAction();
+
+	virtual void IncreaseBrushSpeedAction();
+	virtual void DecreaseBrushSpeedAction();
+
+
+	virtual void NextBrushModeAction();
+	virtual void PreviousBrushModeAction();
+
+	virtual void NextHistoryBrushModeAction();
+	virtual void PreviousHistoryBrushModeAction();
 
 
 protected:
@@ -281,11 +305,15 @@ protected:
 
 	bool bStampPending;
 	FRay PendingStampRay;
+	int StampTimestamp = 0;
+	EDynamicMeshSculptBrushType LastStampType = EDynamicMeshSculptBrushType::LastValue;
+	EDynamicMeshSculptBrushType PendingStampType = LastStampType;
 	void ApplyStamp(const FRay& WorldRay);
 
 	FDynamicMesh3 BrushTargetMesh;
 	FDynamicMeshAABBTree3 BrushTargetMeshSpatial;
 	FMeshNormals BrushTargetNormals;
+	bool bCachedFreezeTarget = false;
 	void UpdateTarget();
 	bool GetTargetMeshNearest(const FVector3d& Position, double SearchRadius, FVector3d& TargetPosOut, FVector3d& TargetNormalOut);
 
@@ -296,6 +324,7 @@ protected:
 	void ApplySmoothBrush(const FRay& WorldRay);
 	void ApplyMoveBrush(const FRay& WorldRay);
 	void ApplyOffsetBrush(const FRay& WorldRay);
+	void ApplySculptMaxBrush(const FRay& WorldRay);
 	void ApplyPinchBrush(const FRay& WorldRay);
 	void ApplyInflateBrush(const FRay& WorldRay);
 	void ApplyPlaneBrush(const FRay& WorldRay);
@@ -324,6 +353,10 @@ protected:
 	FCriticalSection UpdateSavedVertexLock;
 
 	double EstimateIntialSafeTargetLength(const FDynamicMesh3& Mesh, int MinTargetTriCount);
+
+	TArray<EDynamicMeshSculptBrushType> BrushTypeHistory;
+	int BrushTypeHistoryIndex = 0;
+
 };
 
 
