ModelingTools: TransformTool now has more options for snap-drag. Add hotkeys for various settings.
#rnx

diff --git a/Engine/Plugins/Experimental/MeshModelingToolset/Source/MeshModelingTools/Private/TransformMeshesTool.cpp b/Engine/Plugins/Experimental/MeshModelingToolset/Source/MeshModelingTools/Private/TransformMeshesTool.cpp
--- a/Engine/Plugins/Experimental/MeshModelingToolset/Source/MeshModelingTools/Private/TransformMeshesTool.cpp
+++ b/Engine/Plugins/Experimental/MeshModelingToolset/Source/MeshModelingTools/Private/TransformMeshesTool.cpp
@@ -80,6 +80,10 @@ void UTransformMeshesTool::Setup()
 	AddToolPropertySource(TransformProps);
 
 	UpdateTransformMode(TransformProps->TransformMode);
+
+	GetToolManager()->DisplayMessage(
+		LOCTEXT("OnStartTransformMeshesTool", "To transform Objects around points, reposition the Gizmo using Set Pivot mode (S). To quickly position Objects, enable Snap Drag mode (D). A cycles through Transform modes, W and E through SnapDrag Source and Rotation types."),
+		EToolMessageLevel::UserNotification);
 }
 
 
@@ -97,6 +101,23 @@ void UTransformMeshesTool::Shutdown(EToolShutdownType ShutdownType)
 
 void UTransformMeshesTool::Tick(float DeltaTime)
 {
+	// make sure state is up-to-date
+
+	if (CurTransformMode != TransformProps->TransformMode)
+	{
+		UpdateTransformMode(TransformProps->TransformMode);
+	}
+
+	if (bCurSetPivotMode != TransformProps->bSetPivot)
+	{
+		bool bEnableSetPivot = TransformProps->bSetPivot;
+		if (TransformProps->TransformMode == ETransformMeshesTransformMode::SharedGizmoLocal)
+		{
+			bEnableSetPivot = false;
+		}
+		UpdateSetPivotModes(bEnableSetPivot);
+		bCurSetPivotMode = TransformProps->bSetPivot;
+	}
 }
 
 
@@ -107,17 +128,6 @@ void UTransformMeshesTool::Render(IToolsContextRenderAPI* RenderAPI)
 
 void UTransformMeshesTool::OnPropertyModified(UObject* PropertySet, UProperty* Property)
 {
-	if (CurTransformMode != TransformProps->TransformMode)
-	{
-		UpdateTransformMode(TransformProps->TransformMode);
-	}
-
-	bool bEnableSetPivot = TransformProps->bSetPivot;
-	if (TransformProps->TransformMode == ETransformMeshesTransformMode::SharedGizmoLocal)
-	{
-		bEnableSetPivot = false;
-	}
-	UpdateSetPivotModes(bEnableSetPivot);
 }
 
 
@@ -131,6 +141,51 @@ void UTransformMeshesTool::UpdateSetPivotModes(bool bEnableSetPivot)
 
 
 
+void UTransformMeshesTool::RegisterActions(FInteractiveToolActionSet& ActionSet)
+{
+	ActionSet.RegisterAction(this, (int32)EStandardToolActions::BaseClientDefinedActionID + 1,
+		TEXT("ToggleSetPivot"),
+		LOCTEXT("TransformToggleSetPivot", "Toggle Set Pivot"),
+		LOCTEXT("TransformToggleSetPivotTooltip", "Toggle Set Pivot on and off"),
+		EModifierKey::None, EKeys::S,
+		[this]() { this->TransformProps->bSetPivot = !this->TransformProps->bSetPivot; });
+
+	ActionSet.RegisterAction(this, (int32)EStandardToolActions::BaseClientDefinedActionID + 2,
+		TEXT("ToggleSnapDrag"),
+		LOCTEXT("TransformToggleSnapDrag", "Toggle SnapDrag"),
+		LOCTEXT("TransformToggleSnapDragTooltip", "Toggle SnapDrag on and off"),
+		EModifierKey::None, EKeys::D,
+		[this]() { this->TransformProps->bEnableSnapDragging = !this->TransformProps->bEnableSnapDragging; });
+
+	ActionSet.RegisterAction(this, (int32)EStandardToolActions::BaseClientDefinedActionID + 3,
+		TEXT("CycleTransformMode"),
+		LOCTEXT("TransformCycleTransformMode", "Next Transform Mode"),
+		LOCTEXT("TransformCycleTransformModeTooltip", "Cycle through available Transform Modes"),
+		EModifierKey::None, EKeys::A,
+		[this]() { this->TransformProps->TransformMode = (ETransformMeshesTransformMode)(((uint8)TransformProps->TransformMode+1) % (uint8)ETransformMeshesTransformMode::LastValue); });
+
+
+	ActionSet.RegisterAction(this, (int32)EStandardToolActions::BaseClientDefinedActionID + 4,
+		TEXT("CycleSourceMode"),
+		LOCTEXT("TransformCycleSourceMode", "Next SnapDrag Source Mode"),
+		LOCTEXT("TransformCycleSourceModeTooltip", "Cycle through available SnapDrag Source Modes"),
+		EModifierKey::None, EKeys::W,
+		[this]() { this->TransformProps->SnapDragSource = (ETransformMeshesSnapDragSource)(((uint8)TransformProps->SnapDragSource + 1) % (uint8)ETransformMeshesSnapDragSource::LastValue); });
+
+	ActionSet.RegisterAction(this, (int32)EStandardToolActions::BaseClientDefinedActionID + 5,
+		TEXT("CycleRotationMode"),
+		LOCTEXT("TransformCycleRotationMode", "Next SnapDrag Rotation Mode"),
+		LOCTEXT("TransformCycleRotationModeTooltip", "Cycle through available SnapDrag Rotation Modes"),
+		EModifierKey::None, EKeys::E,
+		[this]() { this->TransformProps->RotationMode = (ETransformMeshesSnapDragRotationMode)(((uint8)TransformProps->RotationMode + 1) % (uint8)ETransformMeshesSnapDragRotationMode::LastValue); });
+
+
+}
+
+
+
+
+
 void UTransformMeshesTool::UpdateTransformMode(ETransformMeshesTransformMode NewMode)
 {
 	ResetActiveGizmos();
@@ -232,13 +287,21 @@ void UTransformMeshesTool::OnClickPress(const FInputDeviceRay& PressPos)
 
 	GetToolManager()->BeginUndoTransaction(LOCTEXT("TransformToolTransformTxnName", "SnapDrag"));
 
-	StartDragFrameWorld = FFrame3d(PressPos.WorldRay.PointAt(HitPos.HitDepth), HitPos.HitNormal);
-
 	FTransformMeshesTarget& ActiveTarget =
 		(TransformProps->TransformMode == ETransformMeshesTransformMode::PerObjectGizmo) ?
 		ActiveGizmos[ActiveSnapDragIndex] : ActiveGizmos[0];
 	USceneComponent* GizmoComponent = ActiveTarget.TransformGizmo->GetGizmoActor()->GetRootComponent();
 	StartDragTransform = GizmoComponent->GetComponentToWorld();
+
+	if (TransformProps->SnapDragSource == ETransformMeshesSnapDragSource::ClickPoint)
+	{
+		StartDragFrameWorld = FFrame3d(PressPos.WorldRay.PointAt(HitPos.HitDepth), HitPos.HitNormal);
+	}
+	else
+	{
+		StartDragFrameWorld = FFrame3d(StartDragTransform);
+	}
+
 }
 
 
@@ -246,23 +309,55 @@ void UTransformMeshesTool::OnClickDrag(const FInputDeviceRay& DragPos)
 {
 	FCollisionObjectQueryParams ObjectQueryParams(FCollisionObjectQueryParams::AllObjects);
 	FCollisionQueryParams CollisionParams = FCollisionQueryParams::DefaultQueryParam;
-	
-	int IgnoreIndex = (TransformProps->TransformMode == ETransformMeshesTransformMode::PerObjectGizmo) ?
-		ActiveSnapDragIndex : -1;
-	for ( int k = 0; k < ComponentTargets.Num(); ++k )
+
+	bool bApplyToPivot = TransformProps->bSetPivot;
+
+	if (bApplyToPivot == false)
 	{
-		if (IgnoreIndex == -1 || k == IgnoreIndex)
+		int IgnoreIndex = (TransformProps->TransformMode == ETransformMeshesTransformMode::PerObjectGizmo) ?
+			ActiveSnapDragIndex : -1;
+		for (int k = 0; k < ComponentTargets.Num(); ++k)
 		{
-			CollisionParams.AddIgnoredComponent(ComponentTargets[k]->GetOwnerComponent());
+			if (IgnoreIndex == -1 || k == IgnoreIndex)
+			{
+				CollisionParams.AddIgnoredComponent(ComponentTargets[k]->GetOwnerComponent());
+			}
 		}
 	}
 
+
+	bool bRotate = (TransformProps->RotationMode != ETransformMeshesSnapDragRotationMode::Ignore);
+	float NormalSign = (TransformProps->RotationMode == ETransformMeshesSnapDragRotationMode::AlignFlipped) ? -1.0f : 1.0f;
+
 	FHitResult Result;
 	bool bWorldHit = TargetWorld->LineTraceSingleByObjectType(Result, DragPos.WorldRay.Origin, DragPos.WorldRay.PointAt(999999), ObjectQueryParams, CollisionParams);
-	if (bWorldHit)
+	if (bWorldHit == false)
+	{
+		return;
+	}
+
+	if (bApplyToPivot)
 	{
 		FVector HitPos = Result.ImpactPoint;
-		FVector TargetNormal = -Result.Normal;
+		FVector TargetNormal = (-NormalSign) * Result.Normal;
+
+		FQuaterniond AlignRotation = (bRotate) ?
+			FQuaterniond(FVector3d::UnitZ(), TargetNormal) : FQuaterniond::Identity();
+
+		FTransform NewTransform = StartDragTransform;
+		NewTransform.SetRotation((FQuat)AlignRotation);
+		NewTransform.SetTranslation(HitPos);
+
+		FTransformMeshesTarget& ActiveTarget =
+			(TransformProps->TransformMode == ETransformMeshesTransformMode::PerObjectGizmo) ?
+			ActiveGizmos[ActiveSnapDragIndex] : ActiveGizmos[0];
+		USceneComponent* GizmoComponent = ActiveTarget.TransformGizmo->GetGizmoActor()->GetRootComponent();
+		GizmoComponent->SetWorldTransform(NewTransform);
+	}
+	else
+	{
+		FVector HitPos = Result.ImpactPoint;
+		FVector TargetNormal = NormalSign * Result.Normal;
 
 		FFrame3d FromFrameWorld = StartDragFrameWorld;
 		FFrame3d ToFrameWorld(HitPos, TargetNormal);
@@ -270,6 +365,10 @@ void UTransformMeshesTool::OnClickDrag(const FInputDeviceRay& DragPos)
 
 		FVector3d CenterShift = FromFrameWorld.Origin - ObjectFrameWorld.Origin;
 		FQuaterniond AlignRotation(FromFrameWorld.Z(), ToFrameWorld.Z());
+		if (bRotate == false)
+		{
+			AlignRotation = FQuaterniond::Identity();
+		}
 		FVector3d AlignTranslate = ToFrameWorld.Origin - FromFrameWorld.Origin;
 
 		FTransform NewTransform = StartDragTransform;
diff --git a/Engine/Plugins/Experimental/MeshModelingToolset/Source/MeshModelingTools/Public/TransformMeshesTool.h b/Engine/Plugins/Experimental/MeshModelingToolset/Source/MeshModelingTools/Public/TransformMeshesTool.h
--- a/Engine/Plugins/Experimental/MeshModelingToolset/Source/MeshModelingTools/Public/TransformMeshesTool.h
+++ b/Engine/Plugins/Experimental/MeshModelingToolset/Source/MeshModelingTools/Public/TransformMeshesTool.h
@@ -34,14 +34,51 @@ public:
 UENUM()
 enum class ETransformMeshesTransformMode : uint8
 {
-	/** Shared Gizmo */
+	/** Single Gizmo for all Objects */
 	SharedGizmo = 0 UMETA(DisplayName = "Shared Gizmo"),
 
-	/** Shared Gizmo, Local Transformations*/
+	/** Single Gizmo for all Objects, Rotations applied per-Object */
 	SharedGizmoLocal = 1 UMETA(DisplayName = "Shared Gizmo (Local)"),
 
-	/** Per Object Gizmo */
-	PerObjectGizmo = 2 UMETA(DisplayName = "Multi-Gizmo")
+	/** Separate Gizmo for each Object */
+	PerObjectGizmo = 2 UMETA(DisplayName = "Multi-Gizmo"),
+
+	LastValue UMETA(Hidden)
+};
+
+
+
+/** Snap-Drag Source Point */
+UENUM()
+enum class ETransformMeshesSnapDragSource : uint8
+{
+	/** Snap-Drag moves the Clicked Point to the Target Location */
+	ClickPoint = 0 UMETA(DisplayName = "Click Point"),
+
+	/** Snap-Drag moves the Gizmo/Pivot to the Target Location */
+	Pivot = 1 UMETA(DisplayName = "Pivot"),
+
+
+	LastValue UMETA(Hidden)
+
+};
+
+
+
+/** Snap-Drag Rotation Mode */
+UENUM()
+enum class ETransformMeshesSnapDragRotationMode : uint8
+{
+	/** Snap-Drag only translates, ignoring Normals */
+	Ignore = 0 UMETA(DisplayName = "Ignore"),
+
+	/** Snap-Drag aligns the Source and Target Normals to point in the same direction */
+	Align = 1 UMETA(DisplayName = "Align"),
+
+	/** Snap-Drag aligns the Source Normal to the opposite of the Target Normal direction */
+	AlignFlipped = 2 UMETA(DisplayName = "Align Flipped"),
+
+	LastValue UMETA(Hidden)
 };
 
 
@@ -63,8 +100,17 @@ public:
 	bool bSetPivot = false;
 
 
-	UPROPERTY(EditAnywhere, Category = Options, meta = (EditCondition = "bSetPivot == false"))
+	/** Click-drag starting on the target objects to reposition them on the rest of the scene */
+	UPROPERTY(EditAnywhere, Category = Options)
 	bool bEnableSnapDragging = false;
+
+
+	UPROPERTY(EditAnywhere, Category = Options, meta = (EditCondition = "bEnableSnapDragging == true"))
+	ETransformMeshesSnapDragSource SnapDragSource = ETransformMeshesSnapDragSource::ClickPoint;
+
+	/** When Snap-Dragging, align source and target normals */
+	UPROPERTY(EditAnywhere, Category = Options, meta = (EditCondition = "bEnableSnapDragging == true"))
+	ETransformMeshesSnapDragRotationMode RotationMode = ETransformMeshesSnapDragRotationMode::AlignFlipped;
 };
 
 
@@ -92,6 +138,8 @@ class MESHMODELINGTOOLS_API UTransformMeshesTool : public UMultiSelectionTool, p
 public:
 	UTransformMeshesTool();
 
+	virtual void RegisterActions(FInteractiveToolActionSet& ActionSet) override;
+
 	virtual void SetWorld(UWorld* World, UInteractiveGizmoManager* GizmoManager);
 
 	virtual void Setup() override;
@@ -128,6 +176,7 @@ protected:
 
 	ETransformMeshesTransformMode CurTransformMode;
 	void UpdateTransformMode(ETransformMeshesTransformMode NewMode);
+	bool bCurSetPivotMode;
 	void UpdateSetPivotModes(bool bEnableSetPivot);
 
 	void SetActiveGizmos_Single(bool bLocalRotations);
