ModelingTools: add support for save/restore of Settings in various Tools.
#rnx

diff --git a/Engine/Plugins/Experimental/MeshModelingToolset/Source/MeshModelingTools/Private/AddPrimitiveTool.cpp b/Engine/Plugins/Experimental/MeshModelingToolset/Source/MeshModelingTools/Private/AddPrimitiveTool.cpp
--- a/Engine/Plugins/Experimental/MeshModelingToolset/Source/MeshModelingTools/Private/AddPrimitiveTool.cpp
+++ b/Engine/Plugins/Experimental/MeshModelingToolset/Source/MeshModelingTools/Private/AddPrimitiveTool.cpp
@@ -53,6 +53,30 @@ UProceduralShapeToolProperties::UProceduralShapeToolProperties()
 	PlaceMode = EMakeMeshPlacementType::OnScene;
 }
 
+void UProceduralShapeToolProperties::SaveProperties(UInteractiveTool* SaveFromTool)
+{
+	UProceduralShapeToolProperties* PropertyCache = GetPropertyCache<UProceduralShapeToolProperties>();
+	PropertyCache->Shape = this->Shape;
+	PropertyCache->Width = this->Width;
+	PropertyCache->Height = this->Height;
+	PropertyCache->Slices = this->Slices;
+	PropertyCache->Subdivisions = this->Subdivisions;
+	PropertyCache->PivotLocation = this->PivotLocation;
+	PropertyCache->PlaceMode = this->PlaceMode;
+}
+
+void UProceduralShapeToolProperties::RestoreProperties(UInteractiveTool* RestoreToTool)
+{
+	UProceduralShapeToolProperties* PropertyCache = GetPropertyCache<UProceduralShapeToolProperties>();
+	this->Shape = PropertyCache->Shape;
+	this->Width = PropertyCache->Width;
+	this->Height = PropertyCache->Height;
+	this->Slices = PropertyCache->Slices;
+	this->Subdivisions = PropertyCache->Subdivisions;
+	this->PivotLocation = PropertyCache->PivotLocation;
+	this->PlaceMode = PropertyCache->PlaceMode;
+}
+
 namespace
 {
 	struct {
@@ -109,9 +133,11 @@ void UAddPrimitiveTool::Setup()
 
 	ShapeSettings = NewObject<UProceduralShapeToolProperties>(this);
 	AddToolPropertySource(ShapeSettings);
+	ShapeSettings->RestoreProperties(this);
 
 	MaterialProperties = NewObject<UNewMeshMaterialProperties>(this);
 	AddToolPropertySource(MaterialProperties);
+	MaterialProperties->RestoreProperties(this);
 
 	// create preview mesh object
 	PreviewMesh = NewObject<UPreviewMesh>(this, TEXT("PreviewMesh"));
@@ -132,6 +158,9 @@ void UAddPrimitiveTool::Shutdown(EToolShutdownType ShutdownType)
 	PreviewMesh->SetVisible(false);
 	PreviewMesh->Disconnect();
 	PreviewMesh = nullptr;
+
+	ShapeSettings->SaveProperties(this);
+	MaterialProperties->SaveProperties(this);	
 }
 
 
diff --git a/Engine/Plugins/Experimental/MeshModelingToolset/Source/MeshModelingTools/Private/BaseBrushTool.cpp b/Engine/Plugins/Experimental/MeshModelingToolset/Source/MeshModelingTools/Private/BaseBrushTool.cpp
--- a/Engine/Plugins/Experimental/MeshModelingToolset/Source/MeshModelingTools/Private/BaseBrushTool.cpp
+++ b/Engine/Plugins/Experimental/MeshModelingToolset/Source/MeshModelingTools/Private/BaseBrushTool.cpp
@@ -16,6 +16,23 @@ UBrushBaseProperties::UBrushBaseProperties()
 	BrushRadius = 0.0f;
 }
 
+void UBrushBaseProperties::SaveProperties(UInteractiveTool* SaveFromTool)
+{
+	UBrushBaseProperties* PropertyCache = GetPropertyCache<UBrushBaseProperties>();
+	PropertyCache->BrushSize = this->BrushSize;
+	PropertyCache->bSpecifyRadius = this->bSpecifyRadius;
+	PropertyCache->BrushRadius = this->BrushRadius;
+}
+
+void UBrushBaseProperties::RestoreProperties(UInteractiveTool* RestoreToTool)
+{
+	UBrushBaseProperties* PropertyCache = GetPropertyCache<UBrushBaseProperties>();
+	this->BrushSize = PropertyCache->BrushSize;
+	this->bSpecifyRadius = PropertyCache->bSpecifyRadius;
+	this->BrushRadius = PropertyCache->BrushRadius;
+}
+
+
 
 
 UBaseBrushTool::UBaseBrushTool()
diff --git a/Engine/Plugins/Experimental/MeshModelingToolset/Source/MeshModelingTools/Private/DrawPolygonTool.cpp b/Engine/Plugins/Experimental/MeshModelingToolset/Source/MeshModelingTools/Private/DrawPolygonTool.cpp
--- a/Engine/Plugins/Experimental/MeshModelingToolset/Source/MeshModelingTools/Private/DrawPolygonTool.cpp
+++ b/Engine/Plugins/Experimental/MeshModelingToolset/Source/MeshModelingTools/Private/DrawPolygonTool.cpp
@@ -66,6 +66,58 @@ UDrawPolygonToolStandardProperties::UDrawPolygonToolStandardProperties()
 {
 }
 
+
+void UDrawPolygonToolStandardProperties::SaveProperties(UInteractiveTool* SaveFromTool)
+{
+	UDrawPolygonToolStandardProperties* PropertyCache = GetPropertyCache<UDrawPolygonToolStandardProperties>();
+	PropertyCache->PolygonType = this->PolygonType;
+	PropertyCache->OutputMode = this->OutputMode;
+	PropertyCache->ExtrudeHeight = this->ExtrudeHeight;
+	PropertyCache->Steps = this->Steps;
+	PropertyCache->bAllowSelfIntersections = this->bAllowSelfIntersections;
+}
+
+void UDrawPolygonToolStandardProperties::RestoreProperties(UInteractiveTool* RestoreToTool)
+{
+	UDrawPolygonToolStandardProperties* PropertyCache = GetPropertyCache<UDrawPolygonToolStandardProperties>();
+	this->PolygonType = PropertyCache->PolygonType;
+	this->OutputMode = PropertyCache->OutputMode;
+	this->ExtrudeHeight = PropertyCache->ExtrudeHeight;
+	this->Steps = PropertyCache->Steps;
+	this->bAllowSelfIntersections = PropertyCache->bAllowSelfIntersections;
+}
+
+
+
+void UDrawPolygonToolSnapProperties::SaveProperties(UInteractiveTool* SaveFromTool)
+{
+	UDrawPolygonToolSnapProperties* PropertyCache = GetPropertyCache<UDrawPolygonToolSnapProperties>();
+	PropertyCache->bEnableSnapping = this->bEnableSnapping;
+	PropertyCache->bSnapToVertices = this->bSnapToVertices;
+	PropertyCache->bSnapToEdges = this->bSnapToEdges;
+	PropertyCache->bSnapToAngles = this->bSnapToAngles;
+	PropertyCache->bSnapToLengths = this->bSnapToLengths;
+	PropertyCache->bHitSceneObjects = this->bHitSceneObjects;
+	//PropertyCache->SegmentLength = this->Length;		// this is purely a feedback property
+	PropertyCache->HitNormalOffset = this->HitNormalOffset;
+}
+
+void UDrawPolygonToolSnapProperties::RestoreProperties(UInteractiveTool* RestoreToTool)
+{
+	UDrawPolygonToolSnapProperties* PropertyCache = GetPropertyCache<UDrawPolygonToolSnapProperties>();
+	this->bEnableSnapping = PropertyCache->bEnableSnapping;
+	this->bSnapToVertices = PropertyCache->bSnapToVertices;
+	this->bSnapToEdges = PropertyCache->bSnapToEdges;
+	this->bSnapToAngles = PropertyCache->bSnapToAngles;
+	this->bSnapToLengths = PropertyCache->bSnapToLengths;
+	this->bHitSceneObjects = PropertyCache->bHitSceneObjects;
+	//this->SegmentLength = PropertyCache->Length;
+	this->HitNormalOffset = PropertyCache->HitNormalOffset;
+}
+
+
+
+
 /*
  * Tool
  */
@@ -115,8 +167,7 @@ void UDrawPolygonTool::Setup()
 	AddInputBehavior(AKeyBehavior);
 
 	PolygonProperties = NewObject<UDrawPolygonToolStandardProperties>(this, TEXT("Polygon Settings"));
-	PolygonProperties->OutputMode = EDrawPolygonOutputMode::ExtrudedInteractive;
-	PolygonProperties->ExtrudeHeight = 100.0f;
+	PolygonProperties->RestoreProperties(this);
 
 
 	UPositionPlaneGizmoBuilder* PositionPlaneGizmoBuilder = NewObject<UPositionPlaneGizmoBuilder>();
@@ -129,6 +180,7 @@ void UDrawPolygonTool::Setup()
 	};
 
 	MaterialProperties = NewObject<UNewMeshMaterialProperties>(this);
+	MaterialProperties->RestoreProperties(this);
 
 	// create preview mesh object
 	PreviewMesh = NewObject<UPreviewMesh>(this, TEXT("DrawPolygonPreviewMesh"));
@@ -145,10 +197,9 @@ void UDrawPolygonTool::Setup()
 
 
 	SnapProperties = NewObject<UDrawPolygonToolSnapProperties>(this, TEXT("Snapping"));
+	SnapProperties->RestoreProperties(this);
 
-	
-
-	// add self as tool props
+	// register tool properties
 	AddToolPropertySource(PolygonProperties);
 	AddToolPropertySource(SnapProperties);
 	AddToolPropertySource(MaterialProperties);
@@ -173,6 +224,10 @@ void UDrawPolygonTool::Shutdown(EToolShutdownType ShutdownType)
 	GizmoManager->DestroyGizmo(PositionPlaneGizmo);
 	PositionPlaneGizmo = nullptr;
 	GizmoManager->DeregisterGizmoType(TEXT("DrawPolygonPlaneGizmo"));
+
+	PolygonProperties->SaveProperties(this);
+	SnapProperties->SaveProperties(this);
+	MaterialProperties->SaveProperties(this);
 }
 
 void UDrawPolygonTool::RegisterActions(FInteractiveToolActionSet& ActionSet)
@@ -424,7 +479,7 @@ void UDrawPolygonTool::UpdatePreviewVertex(const FVector& PreviewVertexIn)
 	if (PolygonVertices.Num() > 0)
 	{
 		FVector LastVertex = PolygonVertices[PolygonVertices.Num() - 1];
-		SnapProperties->Length = FVector::Distance(LastVertex, PreviewVertex);
+		SnapProperties->SegmentLength = FVector::Distance(LastVertex, PreviewVertex);
 	}
 }
 
diff --git a/Engine/Plugins/Experimental/MeshModelingToolset/Source/MeshModelingTools/Private/DynamicMeshSculptTool.cpp b/Engine/Plugins/Experimental/MeshModelingToolset/Source/MeshModelingTools/Private/DynamicMeshSculptTool.cpp
--- a/Engine/Plugins/Experimental/MeshModelingToolset/Source/MeshModelingTools/Private/DynamicMeshSculptTool.cpp
+++ b/Engine/Plugins/Experimental/MeshModelingToolset/Source/MeshModelingTools/Private/DynamicMeshSculptTool.cpp
@@ -46,6 +46,26 @@ UBrushSculptProperties::UBrushSculptProperties()
 	bFreezeTarget = false;
 }
 
+void UBrushSculptProperties::SaveProperties(UInteractiveTool* SaveFromTool)
+{
+	UBrushSculptProperties* PropertyCache = GetPropertyCache<UBrushSculptProperties>();
+	PropertyCache->SmoothSpeed = this->SmoothSpeed;
+	PropertyCache->BrushSpeed = this->BrushSpeed;
+	PropertyCache->PrimaryBrushType = this->PrimaryBrushType;
+	PropertyCache->SmoothingType = this->SmoothingType;
+	PropertyCache->Depth = this->Depth;
+}
+
+void UBrushSculptProperties::RestoreProperties(UInteractiveTool* RestoreToTool)
+{
+	UBrushSculptProperties* PropertyCache = GetPropertyCache<UBrushSculptProperties>();
+	this->SmoothSpeed = PropertyCache->SmoothSpeed;
+	this->BrushSpeed = PropertyCache->BrushSpeed;
+	this->PrimaryBrushType = PropertyCache->PrimaryBrushType;
+	this->SmoothingType = PropertyCache->SmoothingType;
+	this->Depth = PropertyCache->Depth;
+}
+
 
 
 UBrushRemeshProperties::UBrushRemeshProperties()
@@ -141,6 +161,10 @@ void UDynamicMeshSculptTool::Setup()
 		AddToolPropertySource(RemeshProperties);
 	}
 
+	BrushProperties->RestoreProperties(this);
+	CalculateBrushRadius();
+	SculptProperties->RestoreProperties(this);
+
 	GetToolManager()->DisplayMessage(
 		LOCTEXT("OnStartSculptTool", "Shift key toggles to Smoothing Brush, Ctrl inverts Brush (where applicable). Q/A keys cycle through Brush Types, Shift+Q/A for Brush History. S/D change Size (shift to small-step), W/E change Speed."),
 		EToolMessageLevel::UserNotification);
@@ -175,6 +199,8 @@ void UDynamicMeshSculptTool::Shutdown(EToolShutdownType ShutdownType)
 		DynamicMeshComponent = nullptr;
 	}
 
+	BrushProperties->SaveProperties(this);
+	SculptProperties->SaveProperties(this);
 }
 
 
diff --git a/Engine/Plugins/Experimental/MeshModelingToolset/Source/MeshModelingTools/Private/Properties/MeshMaterialProperties.cpp b/Engine/Plugins/Experimental/MeshModelingToolset/Source/MeshModelingTools/Private/Properties/MeshMaterialProperties.cpp
--- a/Engine/Plugins/Experimental/MeshModelingToolset/Source/MeshModelingTools/Private/Properties/MeshMaterialProperties.cpp
+++ b/Engine/Plugins/Experimental/MeshModelingToolset/Source/MeshModelingTools/Private/Properties/MeshMaterialProperties.cpp
@@ -17,6 +17,27 @@ UNewMeshMaterialProperties::UNewMeshMaterialProperties()
 	Material = CreateDefaultSubobject<UMaterialInterface>(TEXT("MATERIAL"));
 }
 
+
+void UNewMeshMaterialProperties::SaveProperties(UInteractiveTool* SaveFromTool)
+{
+	UNewMeshMaterialProperties* PropertyCache = GetPropertyCache<UNewMeshMaterialProperties>();
+	PropertyCache->Material = this->Material;
+	PropertyCache->UVScale = this->UVScale;
+	PropertyCache->bWorldSpaceUVScale = this->bWorldSpaceUVScale;
+	// not bWireframe
+}
+
+void UNewMeshMaterialProperties::RestoreProperties(UInteractiveTool* RestoreToTool)
+{
+	UNewMeshMaterialProperties* PropertyCache = GetPropertyCache<UNewMeshMaterialProperties>();
+	this->Material = PropertyCache->Material;
+	this->UVScale = PropertyCache->UVScale;
+	this->bWorldSpaceUVScale = PropertyCache->bWorldSpaceUVScale;
+}
+
+
+
+
 void UExistingMeshMaterialProperties::Setup()
 {
 	UMaterial* CheckerMaterialBase = LoadObject<UMaterial>(nullptr, TEXT("/MeshModelingToolset/Materials/CheckerMaterial"));
diff --git a/Engine/Plugins/Experimental/MeshModelingToolset/Source/MeshModelingTools/Private/Properties/MeshStatisticsProperties.cpp b/Engine/Plugins/Experimental/MeshModelingToolset/Source/MeshModelingTools/Private/Properties/MeshStatisticsProperties.cpp
--- a/Engine/Plugins/Experimental/MeshModelingToolset/Source/MeshModelingTools/Private/Properties/MeshStatisticsProperties.cpp
+++ b/Engine/Plugins/Experimental/MeshModelingToolset/Source/MeshModelingTools/Private/Properties/MeshStatisticsProperties.cpp
@@ -11,10 +11,30 @@
 #define LOCTEXT_NAMESPACE "UMeshStatisticsProperites"
 
 
-void UMeshStatisticsProperties::Update(const FDynamicMesh3& Mesh)
+void UMeshStatisticsProperties::Update(const FDynamicMesh3& MeshIn)
 {
-	TriangleCount = Mesh.TriangleCount();
-	VertexCount = Mesh.VertexCount();
+	this->Mesh = FString::Printf(TEXT("T: %d  V: %d  E: %d"), MeshIn.TriangleCount(), MeshIn.VertexCount(), MeshIn.EdgeCount());
+
+	if (MeshIn.HasAttributes())
+	{
+		const FDynamicMeshAttributeSet* Attribs = MeshIn.Attributes();
+
+		FString UVString;
+		for ( int k = 0; k < Attribs->NumUVLayers(); k++)
+		{ 
+			const FDynamicMeshUVOverlay* UVLayer = Attribs->GetUVLayer(k);
+			UVString += FString::Printf(TEXT("UV%d: %d"), k, UVLayer->ElementCount());
+		}
+		this->UV = UVString;
+
+ 		this->Attributes = FString::Printf(TEXT("Normals: %d"), Attribs->GetNormalLayer(0)->ElementCount());
+	}
+	else
+	{
+		this->UV = FString(TEXT("none"));
+		this->Attributes = FString(TEXT("none"));
+	}
+
 }
 
 
diff --git a/Engine/Plugins/Experimental/MeshModelingToolset/Source/MeshModelingTools/Public/AddPrimitiveTool.h b/Engine/Plugins/Experimental/MeshModelingToolset/Source/MeshModelingTools/Public/AddPrimitiveTool.h
--- a/Engine/Plugins/Experimental/MeshModelingToolset/Source/MeshModelingTools/Public/AddPrimitiveTool.h
+++ b/Engine/Plugins/Experimental/MeshModelingToolset/Source/MeshModelingTools/Public/AddPrimitiveTool.h
@@ -117,6 +117,10 @@ public:
 	/** Subdivisions */
 	UPROPERTY(EditAnywhere, Category = ShapeSettings, meta = (DisplayName = "Subdivisions", UIMin = "0", UIMax = "100", ClampMin = "0", ClampMax = "4000"))
 	int Subdivisions;
+
+
+	virtual void SaveProperties(UInteractiveTool* SaveFromTool) override;
+	virtual void RestoreProperties(UInteractiveTool* RestoreToTool) override;
 };
 
 
diff --git a/Engine/Plugins/Experimental/MeshModelingToolset/Source/MeshModelingTools/Public/BaseBrushTool.h b/Engine/Plugins/Experimental/MeshModelingToolset/Source/MeshModelingTools/Public/BaseBrushTool.h
--- a/Engine/Plugins/Experimental/MeshModelingToolset/Source/MeshModelingTools/Public/BaseBrushTool.h
+++ b/Engine/Plugins/Experimental/MeshModelingToolset/Source/MeshModelingTools/Public/BaseBrushTool.h
@@ -32,6 +32,13 @@ public:
 	/** Radius of brush */
 	UPROPERTY(EditAnywhere, Category = Options, AdvancedDisplay, meta = (DisplayName = "Radius", UIMin = "1.0", UIMax = "1000.0", ClampMin = "0.1", ClampMax = "50000.0"))
 	float BrushRadius;
+
+
+	//
+	// save/restore support
+	//
+	virtual void SaveProperties(UInteractiveTool* SaveFromTool) override;
+	virtual void RestoreProperties(UInteractiveTool* RestoreToTool) override;
 };
 
 
diff --git a/Engine/Plugins/Experimental/MeshModelingToolset/Source/MeshModelingTools/Public/DrawPolygonTool.h b/Engine/Plugins/Experimental/MeshModelingToolset/Source/MeshModelingTools/Public/DrawPolygonTool.h
--- a/Engine/Plugins/Experimental/MeshModelingToolset/Source/MeshModelingTools/Public/DrawPolygonTool.h
+++ b/Engine/Plugins/Experimental/MeshModelingToolset/Source/MeshModelingTools/Public/DrawPolygonTool.h
@@ -60,13 +60,13 @@ enum class EDrawPolygonDrawMode : uint8
 UENUM()
 enum class EDrawPolygonOutputMode : uint8
 {
-	/** Planar polygon mesh */
+	/** Generate a meshed planar polygon */
 	MeshedPolygon UMETA(DisplayName = "Flat Mesh"),
 
-	/** Fixed-distance Extrusion */
-	ExtrudedConstant UMETA(DisplayName = "Extruded To Distance"),
+	/** Extrude closed polygon to constant height determined by Extrude Height Property */
+	ExtrudedConstant UMETA(DisplayName = "Extrude To Height"),
 
-	/** extrusion height is set via additional mouse input after closing polygon */
+	/** Extrusion height is set via additional mouse input after closing polygon */
 	ExtrudedInteractive UMETA(DisplayName = "Interactive Extrude"),
 };
 
@@ -75,18 +75,18 @@ enum class EDrawPolygonOutputMode : uint8
 
 
 UCLASS()
-class MESHMODELINGTOOLS_API UDrawPolygonToolStandardProperties : public UObject
+class MESHMODELINGTOOLS_API UDrawPolygonToolStandardProperties : public UInteractiveToolPropertySet
 {
 	GENERATED_BODY()
 
-	public:
+public:
 	UDrawPolygonToolStandardProperties();
 
 	UPROPERTY(EditAnywhere, NonTransactional, Category = Polygon)
 	EDrawPolygonDrawMode PolygonType = EDrawPolygonDrawMode::Freehand;
 
 	UPROPERTY(EditAnywhere, NonTransactional, Category = Polygon)
-	EDrawPolygonOutputMode OutputMode = EDrawPolygonOutputMode::ExtrudedConstant;
+	EDrawPolygonOutputMode OutputMode = EDrawPolygonOutputMode::ExtrudedInteractive;
 
 	/** Extrusion Distance in non-interactive mode */
 	UPROPERTY(EditAnywhere, NonTransactional, Category = Polygon, meta = (UIMin = "-1000", UIMax = "1000", ClampMin = "-10000", ClampMax = "10000"))
@@ -98,13 +98,19 @@ class MESHMODELINGTOOLS_API UDrawPolygonToolStandardProperties : public UObject
 
 	UPROPERTY(EditAnywhere, NonTransactional, Category = Polygon)
 	bool bAllowSelfIntersections = false;
+
+	//
+	// save/restore support
+	//
+	virtual void SaveProperties(UInteractiveTool* SaveFromTool) override;
+	virtual void RestoreProperties(UInteractiveTool* RestoreToTool) override;
 };
 
 
 
 
 UCLASS()
-class MESHMODELINGTOOLS_API UDrawPolygonToolSnapProperties : public UObject
+class MESHMODELINGTOOLS_API UDrawPolygonToolSnapProperties : public UInteractiveToolPropertySet
 {
 	GENERATED_BODY()
 
@@ -124,14 +130,20 @@ public:
 	UPROPERTY(EditAnywhere, NonTransactional, Category = Snapping, Meta = (EditCondition = "bEnableSnapping"))
 	bool bSnapToLengths = true;
 
-	UPROPERTY(EditAnywhere, NonTransactional, Category = Snapping)
-	float Length = 0.0f;
+	UPROPERTY(VisibleAnywhere, NonTransactional, Category = Snapping)
+	float SegmentLength = 0.0f;
 
 	UPROPERTY(EditAnywhere, NonTransactional, Category = Snapping)
 	bool bHitSceneObjects = false;
 
 	UPROPERTY(EditAnywhere, NonTransactional, Category = Snapping, Meta = (EditCondition = "bHitSceneObjects"))
 	float HitNormalOffset = 0.0f;
+
+	//
+	// save/restore support
+	//
+	virtual void SaveProperties(UInteractiveTool* SaveFromTool) override;
+	virtual void RestoreProperties(UInteractiveTool* RestoreToTool) override;
 };
 
 
diff --git a/Engine/Plugins/Experimental/MeshModelingToolset/Source/MeshModelingTools/Public/DynamicMeshSculptTool.h b/Engine/Plugins/Experimental/MeshModelingToolset/Source/MeshModelingTools/Public/DynamicMeshSculptTool.h
--- a/Engine/Plugins/Experimental/MeshModelingToolset/Source/MeshModelingTools/Public/DynamicMeshSculptTool.h
+++ b/Engine/Plugins/Experimental/MeshModelingToolset/Source/MeshModelingTools/Public/DynamicMeshSculptTool.h
@@ -146,6 +146,10 @@ public:
 	UPROPERTY(EditAnywhere, Category = Options, AdvancedDisplay)
 	EMeshSculptToolSmoothType SmoothingType;
 
+
+	virtual void SaveProperties(UInteractiveTool* SaveFromTool) override;
+	virtual void RestoreProperties(UInteractiveTool* RestoreToTool) override;
+
 };
 
 
diff --git a/Engine/Plugins/Experimental/MeshModelingToolset/Source/MeshModelingTools/Public/MeshInspectorTool.h b/Engine/Plugins/Experimental/MeshModelingToolset/Source/MeshModelingTools/Public/MeshInspectorTool.h
--- a/Engine/Plugins/Experimental/MeshModelingToolset/Source/MeshModelingTools/Public/MeshInspectorTool.h
+++ b/Engine/Plugins/Experimental/MeshModelingToolset/Source/MeshModelingTools/Public/MeshInspectorTool.h
@@ -76,6 +76,11 @@ public:
 	float TangentLength = 5.0f;
 
 
+	//
+	// save/restore support
+	//
+	virtual void SaveProperties(UInteractiveTool* SaveFromTool) override;
+	virtual void RestoreProperties(UInteractiveTool* RestoreToTool) override;
 };
 
 
diff --git a/Engine/Plugins/Experimental/MeshModelingToolset/Source/MeshModelingTools/Public/Properties/MeshMaterialProperties.h b/Engine/Plugins/Experimental/MeshModelingToolset/Source/MeshModelingTools/Public/Properties/MeshMaterialProperties.h
--- a/Engine/Plugins/Experimental/MeshModelingToolset/Source/MeshModelingTools/Public/Properties/MeshMaterialProperties.h
+++ b/Engine/Plugins/Experimental/MeshModelingToolset/Source/MeshModelingTools/Public/Properties/MeshMaterialProperties.h
@@ -38,6 +38,13 @@ public:
 	/** Overlay wireframe on preview */
 	UPROPERTY(EditAnywhere, Category = Material, meta = (DisplayName = "Show Wireframe"))
 	bool bWireframe = false;
+
+
+	//
+	// save/restore support
+	//
+	virtual void SaveProperties(UInteractiveTool* SaveFromTool) override;
+	virtual void RestoreProperties(UInteractiveTool* RestoreToTool) override;
 };
 
 
diff --git a/Engine/Plugins/Experimental/MeshModelingToolset/Source/MeshModelingTools/Public/Properties/MeshStatisticsProperties.h b/Engine/Plugins/Experimental/MeshModelingToolset/Source/MeshModelingTools/Public/Properties/MeshStatisticsProperties.h
--- a/Engine/Plugins/Experimental/MeshModelingToolset/Source/MeshModelingTools/Public/Properties/MeshStatisticsProperties.h
+++ b/Engine/Plugins/Experimental/MeshModelingToolset/Source/MeshModelingTools/Public/Properties/MeshStatisticsProperties.h
@@ -21,10 +21,13 @@ class MESHMODELINGTOOLS_API UMeshStatisticsProperties : public UInteractiveToolP
 public:
 
 	UPROPERTY(VisibleAnywhere, Category = MeshStatistics)
-	int TriangleCount = 0;
+	FString Mesh;
 
-	UPROPERTY(VisibleAnywhere, Category = MeshStatistics)
-	int VertexCount = 0;
+	UPROPERTY(VisibleAnywhere, Category = MeshStatistics, AdvancedDisplay)
+	FString UV;
+
+	UPROPERTY(VisibleAnywhere, Category = MeshStatistics, AdvancedDisplay)
+	FString Attributes;
 
 	void Update(const FDynamicMesh3& Mesh);
 };
