ModelingTools: add file missed in last checkin
#rnx

diff --git a/Engine/Plugins/Experimental/MeshModelingToolset/Source/MeshModelingTools/Private/MeshInspectorTool.cpp b/Engine/Plugins/Experimental/MeshModelingToolset/Source/MeshModelingTools/Private/MeshInspectorTool.cpp
--- a/Engine/Plugins/Experimental/MeshModelingToolset/Source/MeshModelingTools/Private/MeshInspectorTool.cpp
+++ b/Engine/Plugins/Experimental/MeshModelingToolset/Source/MeshModelingTools/Private/MeshInspectorTool.cpp
@@ -40,6 +40,41 @@ UInteractiveTool* UMeshInspectorToolBuilder::BuildTool(const FToolBuilderState&
 
 
 
+/*
+ * Properties
+ */
+
+void UMeshInspectorProperties::SaveProperties(UInteractiveTool* SaveFromTool)
+{
+	UMeshInspectorProperties* PropertyCache = GetPropertyCache<UMeshInspectorProperties>();
+	PropertyCache->bWireframe = this->bWireframe;
+	PropertyCache->bBoundaryEdges = this->bBoundaryEdges;
+	PropertyCache->bPolygonBorders = this->bPolygonBorders;
+	PropertyCache->bUVSeams = this->bUVSeams;
+	PropertyCache->bNormalSeams = this->bNormalSeams;
+	PropertyCache->bNormalVectors = this->bNormalVectors;
+	PropertyCache->bTangentVectors = this->bTangentVectors;
+	PropertyCache->NormalLength = this->NormalLength;
+	PropertyCache->TangentLength = this->TangentLength;
+}
+
+void UMeshInspectorProperties::RestoreProperties(UInteractiveTool* RestoreToTool)
+{
+	UMeshInspectorProperties* PropertyCache = GetPropertyCache<UMeshInspectorProperties>();
+	this->bWireframe = PropertyCache->bWireframe;
+	this->bBoundaryEdges = PropertyCache->bBoundaryEdges;
+	this->bPolygonBorders = PropertyCache->bPolygonBorders;
+	this->bUVSeams = PropertyCache->bUVSeams;
+	this->bNormalSeams = PropertyCache->bNormalSeams;
+	this->bNormalVectors = PropertyCache->bNormalVectors;
+	this->bTangentVectors = PropertyCache->bTangentVectors;
+	this->NormalLength = PropertyCache->NormalLength;
+	this->TangentLength = PropertyCache->TangentLength;
+}
+
+
+
+
 /*
  * Tool
  */
@@ -75,6 +110,7 @@ void UMeshInspectorTool::Setup()
 
 	// initialize our properties
 	Settings = NewObject<UMeshInspectorProperties>(this);
+	Settings->RestoreProperties(this);
 	AddToolPropertySource(Settings);
 
 	MaterialSettings = NewObject<UExistingMeshMaterialProperties>(this);
@@ -100,6 +136,8 @@ void UMeshInspectorTool::Shutdown(EToolShutdownType ShutdownType)
 		DynamicMeshComponent->DestroyComponent();
 		DynamicMeshComponent = nullptr;
 	}
+
+	Settings->SaveProperties(this);
 }
 
 
