Modeling Tools: remove log-spam from the SimplifyTool
#rnx

diff --git a/Engine/Plugins/Experimental/MeshModelingToolset/Source/MeshModelingToolsEditorOnly/Private/SimplifyMeshTool.cpp b/Engine/Plugins/Experimental/MeshModelingToolset/Source/MeshModelingToolsEditorOnly/Private/SimplifyMeshTool.cpp
--- a/Engine/Plugins/Experimental/MeshModelingToolset/Source/MeshModelingToolsEditorOnly/Private/SimplifyMeshTool.cpp
+++ b/Engine/Plugins/Experimental/MeshModelingToolset/Source/MeshModelingToolsEditorOnly/Private/SimplifyMeshTool.cpp
@@ -290,7 +290,7 @@ void USimplifyMeshTool::UpdateResult()
 		}
 	}
 
-	UE_LOG(LogMeshSimplification, Log, TEXT("Mesh simplified to %d triangles"), TargetMesh->TriangleCount());
+	//UE_LOG(LogMeshSimplification, Log, TEXT("Mesh simplified to %d triangles"), TargetMesh->TriangleCount());
 
 	DynamicMeshComponent->NotifyMeshUpdated();
 	GetToolManager()->PostInvalidation();
