Mesh Modeling: Voxel CSG and Voxel Merge.  Generate preview mesh for display while the result is being computed.  Just an experiment at this point..
#rnx

diff --git a/Engine/Plugins/Experimental/MeshModelingToolset/Source/MeshModelingToolsEditorOnly/Private/MergeMeshesTool.cpp b/Engine/Plugins/Experimental/MeshModelingToolset/Source/MeshModelingToolsEditorOnly/Private/MergeMeshesTool.cpp
--- a/Engine/Plugins/Experimental/MeshModelingToolset/Source/MeshModelingToolsEditorOnly/Private/MergeMeshesTool.cpp
+++ b/Engine/Plugins/Experimental/MeshModelingToolset/Source/MeshModelingToolsEditorOnly/Private/MergeMeshesTool.cpp
@@ -92,12 +92,15 @@ void UMergeMeshesTool::Setup()
 	// initialize the PreviewMesh+BackgroundCompute object
 	Preview = NewObject<UMeshOpPreviewWithBackgroundCompute>(this, "Preview");
 	Preview->Setup(this->TargetWorld, this);
+
+	CreateLowQualityPreview(); // update the preview with a low-quality result
+	
 	Preview->ConfigureMaterials(
 		// TODO: This is is very likely a bug waiting to happen
 		ToolSetupUtil::GetDefaultMaterial(GetToolManager(), ComponentTargets[0]->GetMaterial(0)),
 		ToolSetupUtil::GetDefaultWorkingMaterial(GetToolManager())
 	);
-	Preview->SetVisibility(false);
+
 	Preview->InvalidateResult();    // start compute
 }
 
@@ -197,6 +200,25 @@ void UMergeMeshesTool::CacheInputMeshes()
 	}
 }
 
+void UMergeMeshesTool::CreateLowQualityPreview()
+{
+
+	FProgressCancel NullInterrupter;
+	FVoxelMergeMeshesOp MergeMeshesOp;
+
+	MergeMeshesOp.VoxelCount = 12;
+	MergeMeshesOp.AdaptivityD = 0.001;
+	MergeMeshesOp.bAutoSimplify = true;
+	MergeMeshesOp.InputMeshArray = InputMeshes;
+	
+	MergeMeshesOp.CalculateResult(&NullInterrupter);
+	TUniquePtr<FDynamicMesh3> FastPreviewMesh = MergeMeshesOp.ExtractResult();
+
+
+	Preview->PreviewMesh->SetTransform((FTransform)MergeMeshesOp.GetResultTransform());
+	Preview->PreviewMesh->UpdatePreview(FastPreviewMesh.Get());  // copies the mesh @todo we could just give ownership to the Preview!
+	Preview->SetVisibility(true);
+}
 
 
 void UMergeMeshesTool::GenerateAsset(const TUniquePtr<FDynamicMeshOpResult>& Result)
diff --git a/Engine/Plugins/Experimental/MeshModelingToolset/Source/MeshModelingToolsEditorOnly/Private/VoxelCSGMeshesTool.cpp b/Engine/Plugins/Experimental/MeshModelingToolset/Source/MeshModelingToolsEditorOnly/Private/VoxelCSGMeshesTool.cpp
--- a/Engine/Plugins/Experimental/MeshModelingToolset/Source/MeshModelingToolsEditorOnly/Private/VoxelCSGMeshesTool.cpp
+++ b/Engine/Plugins/Experimental/MeshModelingToolset/Source/MeshModelingToolsEditorOnly/Private/VoxelCSGMeshesTool.cpp
@@ -91,12 +91,15 @@ void UVoxelCSGMeshesTool::Setup()
 	// initialize the PreviewMesh+BackgroundCompute object
 	Preview = NewObject<UMeshOpPreviewWithBackgroundCompute>(this, "Preview");
 	Preview->Setup(this->TargetWorld, this);
+
+	CreateLowQualityPreview();
+
 	Preview->ConfigureMaterials(
 		// TODO : Probably multi-select bug - MCD
 		ToolSetupUtil::GetDefaultMaterial(GetToolManager(), ComponentTargets[0]->GetMaterial(0)),
 		ToolSetupUtil::GetDefaultWorkingMaterial(GetToolManager())
 	);
-	Preview->SetVisibility(false);
+	
 	Preview->InvalidateResult();    // start compute
 }
 
@@ -188,6 +191,27 @@ void UVoxelCSGMeshesTool::CacheInputMeshes()
 	}
 }
 
+void UVoxelCSGMeshesTool::CreateLowQualityPreview()
+{
+
+	FProgressCancel NullInterrupter;
+	FVoxelBooleanMeshesOp BooleanOp;
+
+	BooleanOp.Operation = (FVoxelBooleanMeshesOp::EBooleanOperation)(int)CSGProps->Operation;
+	BooleanOp.VoxelCount = 12;
+	BooleanOp.AdaptivityD = 0.01;
+	BooleanOp.bAutoSimplify = true;
+	BooleanOp.InputMeshArray = InputMeshes;
+	
+	BooleanOp.CalculateResult(&NullInterrupter);
+	TUniquePtr<FDynamicMesh3> FastPreviewMesh = BooleanOp.ExtractResult();
+
+
+	Preview->PreviewMesh->SetTransform((FTransform)BooleanOp.GetResultTransform());
+	Preview->PreviewMesh->UpdatePreview(FastPreviewMesh.Get());  // copies the mesh @todo we could just give ownership to the Preview!
+	Preview->SetVisibility(true);
+}
+
 void UVoxelCSGMeshesTool::GenerateAsset(const TUniquePtr<FDynamicMeshOpResult>& Result)
 {
 	check(Result->Mesh.Get() != nullptr);
diff --git a/Engine/Plugins/Experimental/MeshModelingToolset/Source/MeshModelingToolsEditorOnly/Public/MergeMeshesTool.h b/Engine/Plugins/Experimental/MeshModelingToolset/Source/MeshModelingToolsEditorOnly/Public/MergeMeshesTool.h
--- a/Engine/Plugins/Experimental/MeshModelingToolset/Source/MeshModelingToolsEditorOnly/Public/MergeMeshesTool.h
+++ b/Engine/Plugins/Experimental/MeshModelingToolset/Source/MeshModelingToolsEditorOnly/Public/MergeMeshesTool.h
@@ -42,7 +42,7 @@ class MESHMODELINGTOOLSEDITORONLY_API UMergeMeshesToolProperties : public UInter
 public:
 
 	/** The size of the geometry bounding box major axis measured in voxels.*/
-	UPROPERTY(EditAnywhere, Category = Options, meta = (UIMin = "16", UIMax = "2048", ClampMin = "16", ClampMax = "2048"))
+	UPROPERTY(EditAnywhere, Category = Options, meta = (UIMin = "8", UIMax = "512", ClampMin = "8", ClampMax = "512"))
 	int32 VoxelCount = 128;
 
 	/** Remeshing adaptivity */
@@ -103,7 +103,11 @@ protected:
 	IToolsContextAssetAPI* AssetAPI;
 
 	TSharedPtr<TArray<IVoxelBasedCSG::FPlacedMesh>> InputMeshes;
+	/** stash copies of the transforms and pointers to the meshes for consumption by merge Op*/
 	void CacheInputMeshes();
 
+	/** quickly generate a low-quality result for display while the actual result is being computed. */
+	void CreateLowQualityPreview();
+
 	void GenerateAsset(const TUniquePtr<FDynamicMeshOpResult>& Result);
 };
diff --git a/Engine/Plugins/Experimental/MeshModelingToolset/Source/MeshModelingToolsEditorOnly/Public/VoxelCSGMeshesTool.h b/Engine/Plugins/Experimental/MeshModelingToolset/Source/MeshModelingToolsEditorOnly/Public/VoxelCSGMeshesTool.h
--- a/Engine/Plugins/Experimental/MeshModelingToolset/Source/MeshModelingToolsEditorOnly/Public/VoxelCSGMeshesTool.h
+++ b/Engine/Plugins/Experimental/MeshModelingToolset/Source/MeshModelingToolsEditorOnly/Public/VoxelCSGMeshesTool.h
@@ -64,7 +64,7 @@ public:
 	EVoxelCSGOperation Operation;
 
 	/** The size of the geometry bounding box major axis measured in voxels.*/
-	UPROPERTY(EditAnywhere, Category = Options, meta = (UIMin = "16", UIMax = "2048", ClampMin = "16", ClampMax = "2048"))
+	UPROPERTY(EditAnywhere, Category = Options, meta = (UIMin = "8", UIMax = "512", ClampMin = "8", ClampMax = "512"))
 	int32 VoxelCount = 128;
 
 	/** Remeshing adaptivity */
@@ -129,7 +129,11 @@ protected:
 	IToolsContextAssetAPI* AssetAPI;
 
 	TSharedPtr<TArray<IVoxelBasedCSG::FPlacedMesh>> InputMeshes;
+	/** stash copies of the transforms and pointers to the meshes for consumption by the CSG Op*/
 	void CacheInputMeshes();
 
+	/** quickly generate a low-quality result for display while the actual result is being computed. */
+	void CreateLowQualityPreview();
+
 	void GenerateAsset(const TUniquePtr<FDynamicMeshOpResult>& Result);
 };
diff --git a/Engine/Plugins/Experimental/MeshModelingToolset/Source/ModelingOperatorsEditorOnly/Public/CompositionOps/VoxelBooleanMeshesOp.h b/Engine/Plugins/Experimental/MeshModelingToolset/Source/ModelingOperatorsEditorOnly/Public/CompositionOps/VoxelBooleanMeshesOp.h
--- a/Engine/Plugins/Experimental/MeshModelingToolset/Source/ModelingOperatorsEditorOnly/Public/CompositionOps/VoxelBooleanMeshesOp.h
+++ b/Engine/Plugins/Experimental/MeshModelingToolset/Source/ModelingOperatorsEditorOnly/Public/CompositionOps/VoxelBooleanMeshesOp.h
@@ -32,9 +32,6 @@ public:
 	bool   bAutoSimplify = false;
 	EBooleanOperation Operation = EBooleanOperation::DifferenceBA;
 
-	double FastCollapseVoxelSizeMultipler = 1.2;
-	int NumFastCollapsePasses = 10;
-
 	//
 	// FDynamicMeshOperator implementation
 	// 
