ToolsFramework: add FToolCommandChange, replace all uses of FCommandChange with it. Move IToolContextTransactionProvider defn to separate file. Add IToolCommandChangeSource interface for downstream code to use. Move FCommandChangeSequence from inside Selection tool up to Framework level.
#rnx

diff --git a/Engine/Plugins/Experimental/MeshModelingToolset/Source/MeshModelingTools/Private/MeshSelectionTool.cpp b/Engine/Plugins/Experimental/MeshModelingToolset/Source/MeshModelingTools/Private/MeshSelectionTool.cpp
--- a/Engine/Plugins/Experimental/MeshModelingToolset/Source/MeshModelingTools/Private/MeshSelectionTool.cpp
+++ b/Engine/Plugins/Experimental/MeshModelingToolset/Source/MeshModelingTools/Private/MeshSelectionTool.cpp
@@ -6,6 +6,7 @@
 #include "Drawing/MeshDebugDrawing.h"
 #include "DynamicMeshEditor.h"
 #include "DynamicMeshChangeTracker.h"
+#include "Changes/ToolCommandChangeSequence.h"
 #include "Changes/MeshChange.h"
 #include "AssetGenerationUtil.h"
 
@@ -415,75 +416,6 @@ TUniquePtr<FMeshSelectionChange> UMeshSelectionTool::EndChange()
 	return TUniquePtr<FMeshSelectionChange>();
 }
 
-/**
- * FCommandChangeSequence contains a list of FCommandChanges and associated target UObjects.
- * The sequence of changes is applied atomically.
- * @warning if the target weak UObject pointers become invalid, those changes are skipped. This likely leaves this in an undefined state.
- */
-class FCommandChangeSequence : public FCommandChange
-{
-protected:
-	struct FChangeElem
-	{
-		TWeakObjectPtr<UObject> TargetObject;
-		TUniquePtr<FCommandChange> Change;
-	};
-
-	TArray<TSharedPtr<FChangeElem>> Sequence;
-
-public:
-	FCommandChangeSequence()
-	{
-	}
-
-	/** Add a change to the sequence */
-	void AppendChange(UObject* Target, TUniquePtr<FCommandChange> Change)
-	{
-		TSharedPtr<FChangeElem> Elem = MakeShared<FChangeElem>();
-		Elem->TargetObject = Target;
-		Elem->Change = MoveTemp(Change);
-		Sequence.Add(Elem);
-	}
-	
-	/** Apply sequence of changes in-order */
-	virtual void Apply(UObject* Object) override
-	{
-		for (int k = 0; k < Sequence.Num(); ++k)
-		{
-			TSharedPtr<FChangeElem> Elem = Sequence[k];
-			check(Elem->TargetObject.IsValid());
-			if (Elem->TargetObject.IsValid())
-			{
-				Elem->Change->Apply(Elem->TargetObject.Get());
-			}
-		}
-	}
-
-	/** Reverts sequence of changes in reverse-order */
-	virtual void Revert(UObject* Object) override
-	{
-		for (int k = Sequence.Num()-1; k >= 0; --k)
-		{
-			TSharedPtr<FChangeElem> Elem = Sequence[k];
-			check(Elem->TargetObject.IsValid());
-			if (Elem->TargetObject.IsValid())
-			{
-				Elem->Change->Revert(Elem->TargetObject.Get());
-			}
-		}
-	}
-
-	/** @return string describing this change sequenece */
-	virtual FString ToString() const override
-	{
-		FString Result = TEXT("FCommandChangeSequence: ");
-		for (int k = 0; k < Sequence.Num(); ++k)
-		{
-			Result = Result + Sequence[k]->Change->ToString() + TEXT(" ");
-		}
-		return Result;
-	}
-};
 
 
 
@@ -496,7 +428,7 @@ void UMeshSelectionTool::DeleteSelectedTriangles()
 		return;
 	}
 
-	TUniquePtr<FCommandChangeSequence> ChangeSeq = MakeUnique<FCommandChangeSequence>();
+	TUniquePtr<FToolCommandChangeSequence> ChangeSeq = MakeUnique<FToolCommandChangeSequence>();
 
 	// clear current selection
 	BeginChange(false);
diff --git a/Engine/Plugins/Experimental/MeshModelingToolset/Source/ModelingComponents/Public/Changes/MeshChange.h b/Engine/Plugins/Experimental/MeshModelingToolset/Source/ModelingComponents/Public/Changes/MeshChange.h
--- a/Engine/Plugins/Experimental/MeshModelingToolset/Source/ModelingComponents/Public/Changes/MeshChange.h
+++ b/Engine/Plugins/Experimental/MeshModelingToolset/Source/ModelingComponents/Public/Changes/MeshChange.h
@@ -4,7 +4,7 @@
 
 #include "CoreMinimal.h"
 #include "UObject/Interface.h"
-#include "Misc/Change.h"
+#include "InteractiveToolChange.h"
 #include "VectorTypes.h"
 #include "DynamicMeshChangeTracker.h"
 #include "MeshChange.generated.h"
@@ -18,7 +18,7 @@
  * FMeshChange represents an undoable change to a FDynamicMesh3.
  * Currently only valid to call Apply/Revert when the Object is a USimpleDynamicMeshComponent
  */
-class MODELINGCOMPONENTS_API FMeshChange : public FCommandChange
+class MODELINGCOMPONENTS_API FMeshChange : public FToolCommandChange
 {
 public:
 	FMeshChange();
diff --git a/Engine/Plugins/Experimental/MeshModelingToolset/Source/ModelingComponents/Public/MeshSelectionChange.h b/Engine/Plugins/Experimental/MeshModelingToolset/Source/ModelingComponents/Public/MeshSelectionChange.h
--- a/Engine/Plugins/Experimental/MeshModelingToolset/Source/ModelingComponents/Public/MeshSelectionChange.h
+++ b/Engine/Plugins/Experimental/MeshModelingToolset/Source/ModelingComponents/Public/MeshSelectionChange.h
@@ -3,7 +3,7 @@
 #pragma once
 
 #include "CoreMinimal.h"
-#include "Misc/Change.h"
+#include "InteractiveToolChange.h"
 #include "SelectionSet.h"
 
 
@@ -12,7 +12,7 @@
 /**
  * FMeshSelectionChange represents an reversible change to a UMeshSelectionSet
  */
-class MODELINGCOMPONENTS_API FMeshSelectionChange : public FCommandChange
+class MODELINGCOMPONENTS_API FMeshSelectionChange : public FToolCommandChange
 {
 public:
 	EMeshSelectionElementType ElementType = EMeshSelectionElementType::Vertex;
diff --git a/Engine/Plugins/Experimental/MeshModelingToolset/Source/ModelingComponents/Public/MeshVertexChange.h b/Engine/Plugins/Experimental/MeshModelingToolset/Source/ModelingComponents/Public/MeshVertexChange.h
--- a/Engine/Plugins/Experimental/MeshModelingToolset/Source/ModelingComponents/Public/MeshVertexChange.h
+++ b/Engine/Plugins/Experimental/MeshModelingToolset/Source/ModelingComponents/Public/MeshVertexChange.h
@@ -4,7 +4,7 @@
 
 #include "CoreMinimal.h"
 #include "UObject/Interface.h"
-#include "Misc/Change.h"
+#include "InteractiveToolChange.h"
 #include "VectorTypes.h"
 #include "MeshVertexChange.generated.h"
 
@@ -17,7 +17,7 @@ class FDynamicMesh3;
  * @todo support optionally storing old/new normals and tangents
  * @todo support applying to a StaticMeshComponent/MeshDescription ?
  */
-class MODELINGCOMPONENTS_API FMeshVertexChange : public FCommandChange
+class MODELINGCOMPONENTS_API FMeshVertexChange : public FToolCommandChange
 {
 public:
 	TArray<int> Vertices;
diff --git a/Engine/Source/Editor/Experimental/EditorInteractiveToolsFramework/Private/EdModeInteractiveToolsContext.cpp b/Engine/Source/Editor/Experimental/EditorInteractiveToolsFramework/Private/EdModeInteractiveToolsContext.cpp
--- a/Engine/Source/Editor/Experimental/EditorInteractiveToolsFramework/Private/EdModeInteractiveToolsContext.cpp
+++ b/Engine/Source/Editor/Experimental/EditorInteractiveToolsFramework/Private/EdModeInteractiveToolsContext.cpp
@@ -275,7 +275,7 @@ public:
 		GEditor->EndTransaction();
 	}
 
-	virtual void AppendChange(UObject* TargetObject, TUniquePtr<FChange> Change, const FText& Description) override
+	virtual void AppendChange(UObject* TargetObject, TUniquePtr<FToolCommandChange> Change, const FText& Description) override
 	{
 		FScopedTransaction Transaction(Description);
 		check(GUndo != nullptr);
diff --git a/Engine/Source/Editor/UnrealEd/Private/Tools/UEdMode.cpp b/Engine/Source/Editor/UnrealEd/Private/Tools/UEdMode.cpp
--- a/Engine/Source/Editor/UnrealEd/Private/Tools/UEdMode.cpp
+++ b/Engine/Source/Editor/UnrealEd/Private/Tools/UEdMode.cpp
@@ -265,7 +265,7 @@ public:
 		GEditor->EndTransaction();
 	}
 
-	virtual void AppendChange(UObject* TargetObject, TUniquePtr<FChange> Change, const FText& Description) override
+	virtual void AppendChange(UObject* TargetObject, TUniquePtr<FToolCommandChange> Change, const FText& Description) override
 	{
 		FScopedTransaction Transaction(Description);
 		check(GUndo != nullptr);
diff --git a/Engine/Source/Runtime/Experimental/InteractiveToolsFramework/Private/InteractiveGizmoManager.cpp b/Engine/Source/Runtime/Experimental/InteractiveToolsFramework/Private/InteractiveGizmoManager.cpp
--- a/Engine/Source/Runtime/Experimental/InteractiveToolsFramework/Private/InteractiveGizmoManager.cpp
+++ b/Engine/Source/Runtime/Experimental/InteractiveToolsFramework/Private/InteractiveGizmoManager.cpp
@@ -256,7 +256,7 @@ void UInteractiveGizmoManager::EndUndoTransaction()
 
 
 
-void UInteractiveGizmoManager::EmitObjectChange(UObject* TargetObject, TUniquePtr<FChange> Change, const FText& Description)
+void UInteractiveGizmoManager::EmitObjectChange(UObject* TargetObject, TUniquePtr<FToolCommandChange> Change, const FText& Description)
 {
 	TransactionsAPI->AppendChange(TargetObject, MoveTemp(Change), Description );
 }
diff --git a/Engine/Source/Runtime/Experimental/InteractiveToolsFramework/Private/InteractiveToolManager.cpp b/Engine/Source/Runtime/Experimental/InteractiveToolsFramework/Private/InteractiveToolManager.cpp
--- a/Engine/Source/Runtime/Experimental/InteractiveToolsFramework/Private/InteractiveToolManager.cpp
+++ b/Engine/Source/Runtime/Experimental/InteractiveToolsFramework/Private/InteractiveToolManager.cpp
@@ -263,7 +263,7 @@ void UInteractiveToolManager::EndUndoTransaction()
 
 
 
-void UInteractiveToolManager::EmitObjectChange(UObject* TargetObject, TUniquePtr<FChange> Change, const FText& Description)
+void UInteractiveToolManager::EmitObjectChange(UObject* TargetObject, TUniquePtr<FToolCommandChange> Change, const FText& Description)
 {
 	// wrap change 
 	check(HasActiveTool(EToolSide::Left));
diff --git a/Engine/Source/Runtime/Experimental/InteractiveToolsFramework/Public/Changes/ToolCommandChangeSequence.h b/Engine/Source/Runtime/Experimental/InteractiveToolsFramework/Public/Changes/ToolCommandChangeSequence.h
new file mode 100644
--- /dev/null
+++ b/Engine/Source/Runtime/Experimental/InteractiveToolsFramework/Public/Changes/ToolCommandChangeSequence.h
@@ -0,0 +1,80 @@
+// Copyright 1998-2019 Epic Games, Inc. All Rights Reserved.
+
+#pragma once
+
+#include "InteractiveToolChange.h"
+
+
+// This class may be unnecessary, as multiple FChange can be combined in a single transaction.
+// Revisit usage in MeshSelectionTool.
+
+/**
+ * FCommandChangeSequence contains a list of FCommandChanges and associated target UObjects.
+ * The sequence of changes is applied atomically.
+ * @warning if the target weak UObject pointers become invalid, those changes are skipped. Possibly leaves things in an undefined state...
+ */
+class FToolCommandChangeSequence : public FToolCommandChange
+{
+protected:
+	struct FChangeElem
+	{
+		TWeakObjectPtr<UObject> TargetObject;
+		TUniquePtr<FToolCommandChange> Change;
+	};
+
+	TArray<TSharedPtr<FChangeElem>> Sequence;
+
+public:
+	FToolCommandChangeSequence()
+	{
+	}
+
+	/** Add a change to the sequence */
+	void AppendChange(UObject* Target, TUniquePtr<FToolCommandChange> Change)
+	{
+		TSharedPtr<FChangeElem> Elem = MakeShared<FChangeElem>();
+		Elem->TargetObject = Target;
+		Elem->Change = MoveTemp(Change);
+		Sequence.Add(Elem);
+	}
+
+	/** Apply sequence of changes in-order */
+	virtual void Apply(UObject* Object) override
+	{
+		for (int k = 0; k < Sequence.Num(); ++k)
+		{
+			TSharedPtr<FChangeElem> Elem = Sequence[k];
+			check(Elem->TargetObject.IsValid());
+			if (Elem->TargetObject.IsValid())
+			{
+				Elem->Change->Apply(Elem->TargetObject.Get());
+			}
+		}
+	}
+
+	/** Reverts sequence of changes in reverse-order */
+	virtual void Revert(UObject* Object) override
+	{
+		for (int k = Sequence.Num() - 1; k >= 0; --k)
+		{
+			TSharedPtr<FChangeElem> Elem = Sequence[k];
+			check(Elem->TargetObject.IsValid());
+			if (Elem->TargetObject.IsValid())
+			{
+				Elem->Change->Revert(Elem->TargetObject.Get());
+			}
+		}
+	}
+
+	/** @return string describing this change sequenece */
+	virtual FString ToString() const override
+	{
+		FString Result = TEXT("FCommandChangeSequence: ");
+		for (int k = 0; k < Sequence.Num(); ++k)
+		{
+			Result = Result + Sequence[k]->Change->ToString() + TEXT(" ");
+		}
+		return Result;
+	}
+};
+
diff --git a/Engine/Source/Runtime/Experimental/InteractiveToolsFramework/Public/Changes/TransformChange.h b/Engine/Source/Runtime/Experimental/InteractiveToolsFramework/Public/Changes/TransformChange.h
--- a/Engine/Source/Runtime/Experimental/InteractiveToolsFramework/Public/Changes/TransformChange.h
+++ b/Engine/Source/Runtime/Experimental/InteractiveToolsFramework/Public/Changes/TransformChange.h
@@ -3,13 +3,13 @@
 #pragma once
 
 #include "CoreMinimal.h"
-#include "Misc/Change.h"
+#include "InteractiveToolChange.h"
 
 
 /**
  * FComponentWorldTransformChange represents an undoable change to the world transform of a USceneComponent.
  */
-class INTERACTIVETOOLSFRAMEWORK_API FComponentWorldTransformChange : public FCommandChange
+class INTERACTIVETOOLSFRAMEWORK_API FComponentWorldTransformChange : public FToolCommandChange
 {
 public:
 	FComponentWorldTransformChange();
diff --git a/Engine/Source/Runtime/Experimental/InteractiveToolsFramework/Public/InteractiveGizmoManager.h b/Engine/Source/Runtime/Experimental/InteractiveToolsFramework/Public/InteractiveGizmoManager.h
--- a/Engine/Source/Runtime/Experimental/InteractiveToolsFramework/Public/InteractiveGizmoManager.h
+++ b/Engine/Source/Runtime/Experimental/InteractiveToolsFramework/Public/InteractiveGizmoManager.h
@@ -4,10 +4,10 @@
 
 #include "CoreMinimal.h"
 #include "UObject/Object.h"
-#include "Misc/Change.h"
 #include "InteractiveGizmo.h"
 #include "InteractiveGizmoBuilder.h"
 #include "InputRouter.h"
+#include "InteractiveToolChange.h"
 #include "ToolContextInterfaces.h"
 #include "InteractiveGizmoManager.generated.h"
 
@@ -142,7 +142,7 @@ public:
 	 * @param Change the change object that the Context should insert into the transaction history
 	 * @param Description text description of this change (this is the string that appears on undo/redo in the UE Editor)
 	 */
-	virtual void EmitObjectChange(UObject* TargetObject, TUniquePtr<FChange> Change, const FText& Description );
+	virtual void EmitObjectChange(UObject* TargetObject, TUniquePtr<FToolCommandChange> Change, const FText& Description );
 
 
 	//
diff --git a/Engine/Source/Runtime/Experimental/InteractiveToolsFramework/Public/InteractiveToolChange.h b/Engine/Source/Runtime/Experimental/InteractiveToolsFramework/Public/InteractiveToolChange.h
new file mode 100644
--- /dev/null
+++ b/Engine/Source/Runtime/Experimental/InteractiveToolsFramework/Public/InteractiveToolChange.h
@@ -0,0 +1,82 @@
+// Copyright 1998-2019 Epic Games, Inc. All Rights Reserved.
+
+#pragma once
+
+#include "CoreMinimal.h"
+#include "UObject/Interface.h"
+#include "Misc/Change.h"
+#include "InteractiveToolChange.generated.h"
+
+class UInteractiveToolManager;
+class UInteractiveGizmoManager;
+
+
+/**
+ * FToolCommandChange is a base class for command changes used by the Tools Framework.
+ */
+class INTERACTIVETOOLSFRAMEWORK_API FToolCommandChange : public FCommandChange
+{
+public:
+
+	virtual FString ToString() const override { return TEXT("ToolCommandChange"); }
+};
+
+
+
+
+
+UINTERFACE()
+class INTERACTIVETOOLSFRAMEWORK_API UToolContextTransactionProvider : public UInterface
+{
+	GENERATED_BODY()
+};
+/**
+ * IToolContextTransactionProvider is a UInterface that defines several functions that InteractiveTool code
+ * uses to interface with the higher-level transaction system. UInteractiveToolManager and UInteractiveGizmoManager
+ * both implement this interface.
+ */
+class INTERACTIVETOOLSFRAMEWORK_API IToolContextTransactionProvider
+{
+	GENERATED_BODY()
+public:
+	/**
+	 * Request that the Context open a Transaction, whatever that means to the current Context
+	 * @param Description text description of this transaction (this is the string that appears on undo/redo in the UE Editor)
+	 */
+	virtual void BeginUndoTransaction(const FText& Description) = 0;
+
+	/** Request that the Context close and commit the open Transaction */
+	virtual void EndUndoTransaction() = 0;
+
+	/**
+	 * Forward an FChange object to the Context
+	 * @param TargetObject the object that the FChange applies to
+	 * @param Change the change object that the Context should insert into the transaction history
+	 * @param Description text description of this change (this is the string that appears on undo/redo in the UE Editor)
+	 */
+	virtual void EmitObjectChange(UObject* TargetObject, TUniquePtr<FToolCommandChange> Change, const FText& Description) = 0;
+};
+
+
+
+
+/**
+ * IToolCommandChangeSource is a generic interface for generating a change object.
+ * This is useful to 
+ */
+class INTERACTIVETOOLSFRAMEWORK_API IToolCommandChangeSource
+{
+public:
+	virtual ~IToolCommandChangeSource() {}
+
+	virtual void BeginChange() = 0;
+
+	// Generally assuming that user of this class will call something like EmitObjectChange(GetChangeTarget(), MoveTemp(EndChange()), GetChangeDescription())
+
+	virtual TUniquePtr<FToolCommandChange> EndChange() = 0;
+	virtual UObject* GetChangeTarget() = 0;
+	virtual FText GetChangeDescription() = 0;
+};
+
+
+
diff --git a/Engine/Source/Runtime/Experimental/InteractiveToolsFramework/Public/InteractiveToolManager.h b/Engine/Source/Runtime/Experimental/InteractiveToolsFramework/Public/InteractiveToolManager.h
--- a/Engine/Source/Runtime/Experimental/InteractiveToolsFramework/Public/InteractiveToolManager.h
+++ b/Engine/Source/Runtime/Experimental/InteractiveToolsFramework/Public/InteractiveToolManager.h
@@ -4,10 +4,10 @@
 
 #include "CoreMinimal.h"
 #include "UObject/Object.h"
-#include "Misc/Change.h"
 #include "InteractiveTool.h"
 #include "InteractiveToolBuilder.h"
 #include "InputRouter.h"
+#include "InteractiveToolChange.h"
 #include "ToolContextInterfaces.h"
 #include "InteractiveToolManager.generated.h"
 
@@ -161,7 +161,7 @@ public:
 	 * @param Change the change object that the Context should insert into the transaction history
 	 * @param Description text description of this change (this is the string that appears on undo/redo in the UE Editor)
 	 */
-	virtual void EmitObjectChange(UObject* TargetObject, TUniquePtr<FChange> Change, const FText& Description );
+	virtual void EmitObjectChange(UObject* TargetObject, TUniquePtr<FToolCommandChange> Change, const FText& Description );
 
 
 	/**
@@ -239,7 +239,7 @@ protected:
  * FBeginToolChange is used by UInteractiveToolManager to back out of a Tool on Undo.
  * No action is taken on Redo, ie we do not re-start the Tool on Redo.
  */
-class INTERACTIVETOOLSFRAMEWORK_API FBeginToolChange : public FCommandChange
+class INTERACTIVETOOLSFRAMEWORK_API FBeginToolChange : public FToolCommandChange
 {
 public:
 	virtual void Apply(UObject* Object) override;
@@ -257,12 +257,12 @@ public:
  * FToolChangeWrapperChange wraps an FChange emitted by an InteractiveTool, allowing
  * us to Expire the change without each FChange implementation needing to handle this explicitly.
  */
-class INTERACTIVETOOLSFRAMEWORK_API FToolChangeWrapperChange : public FCommandChange
+class INTERACTIVETOOLSFRAMEWORK_API FToolChangeWrapperChange : public FToolCommandChange
 {
 public:
 	TWeakObjectPtr<UInteractiveToolManager> ToolManager;
 	TWeakObjectPtr<UInteractiveTool> ActiveTool;
-	TUniquePtr<FChange> ToolChange;
+	TUniquePtr<FToolCommandChange> ToolChange;
 
 	virtual void Apply(UObject* Object) override;
 
diff --git a/Engine/Source/Runtime/Experimental/InteractiveToolsFramework/Public/ToolContextInterfaces.h b/Engine/Source/Runtime/Experimental/InteractiveToolsFramework/Public/ToolContextInterfaces.h
--- a/Engine/Source/Runtime/Experimental/InteractiveToolsFramework/Public/ToolContextInterfaces.h
+++ b/Engine/Source/Runtime/Experimental/InteractiveToolsFramework/Public/ToolContextInterfaces.h
@@ -5,13 +5,12 @@
 #include "CoreMinimal.h"
 #include "UObject/Interface.h"
 #include "ComponentSourceInterfaces.h"
-#include "ToolContextInterfaces.generated.h"
 
 
 // predeclarations so we don't have to include these in all tools
 class AActor;
 class UActorComponent;
-class FChange;
+class FToolCommandChange;
 class UPackage;
 class FPrimitiveDrawInterface;
 class UInteractiveToolManager;
@@ -19,42 +18,6 @@ class UInteractiveGizmoManager;
 
 
 
-UINTERFACE()
-class INTERACTIVETOOLSFRAMEWORK_API UToolContextTransactionProvider : public UInterface
-{
-	GENERATED_BODY()
-};
-/**
- * IToolContextTransactionProvider is a UInterface that defines several functions that InteractiveTool code
- * uses to interface with the higher-level transaction system. UInteractiveToolManager and UInteractiveGizmoManager
- * both implement this interface.
- */
-class INTERACTIVETOOLSFRAMEWORK_API IToolContextTransactionProvider
-{
-	GENERATED_BODY()
-public:
-	/**
-	 * Request that the Context open a Transaction, whatever that means to the current Context
-	 * @param Description text description of this transaction (this is the string that appears on undo/redo in the UE Editor)
-	 */
-	virtual void BeginUndoTransaction(const FText& Description) = 0;
-
-	/** Request that the Context close and commit the open Transaction */
-	virtual void EndUndoTransaction() = 0;
-
-	/**
-	 * Forward an FChange object to the Context
-	 * @param TargetObject the object that the FChange applies to
-	 * @param Change the change object that the Context should insert into the transaction history
-	 * @param Description text description of this change (this is the string that appears on undo/redo in the UE Editor)
-	 */
-	virtual void EmitObjectChange(UObject* TargetObject, TUniquePtr<FChange> Change, const FText& Description) = 0;
-};
-
-
-
-
-
 /**
  * FToolBuilderState is a bucket of state information that a ToolBuilder might need
  * to construct a Tool. This information comes from a level above the Tools framework,
@@ -324,7 +287,7 @@ public:
 	 * @param Change The Change implementation
 	 * @param Description text description of the transaction that could be shown to user
 	 */
-	virtual void AppendChange(UObject* TargetObject, TUniquePtr<FChange> Change, const FText& Description) = 0;
+	virtual void AppendChange(UObject* TargetObject, TUniquePtr<FToolCommandChange> Change, const FText& Description) = 0;
 
 
 
