ModelingTools: convert DrawPolygonTool to use the new TransformGizmo instead of PositionPlaneGizmo
#rnx

diff --git a/Engine/Plugins/Experimental/MeshModelingToolset/Source/MeshModelingTools/Private/DrawPolygonTool.cpp b/Engine/Plugins/Experimental/MeshModelingToolset/Source/MeshModelingTools/Private/DrawPolygonTool.cpp
--- a/Engine/Plugins/Experimental/MeshModelingToolset/Source/MeshModelingTools/Private/DrawPolygonTool.cpp
+++ b/Engine/Plugins/Experimental/MeshModelingToolset/Source/MeshModelingTools/Private/DrawPolygonTool.cpp
@@ -22,6 +22,8 @@
 
 #include "DynamicMeshEditor.h"
 
+#include "BaseGizmos/GizmoComponents.h"
+#include "BaseGizmos/TransformGizmo.h"
 #include "PositionPlaneGizmo.h"
 #include "Drawing/MeshDebugDrawing.h"
 
@@ -148,7 +150,7 @@ void UDrawPolygonTool::Setup()
 	MouseBehavior->Modifiers.RegisterModifier(IgnoreSnappingModifier, FInputDeviceState::IsShiftKeyDown);
 	AddInputBehavior(MouseBehavior);
 
-	// click to set plane behavior
+	// Register a click behavior/action pair, that sets the draw plane to the clicked world position
 	FSelectClickedAction* SetPlaneAction = new FSelectClickedAction();
 	SetPlaneAction->World = this->TargetWorld;
 	SetPlaneAction->OnClickedPositionFunc = [this](const FHitResult& Hit) {
@@ -162,23 +164,26 @@ void UDrawPolygonTool::Setup()
 	ClickToSetPlaneBehavior->SetDefaultPriority(MouseBehavior->GetPriority().MakeHigher());
 	AddInputBehavior(ClickToSetPlaneBehavior);
 
-	UKeyAsModifierInputBehavior* AKeyBehavior = NewObject<UKeyAsModifierInputBehavior>();
-	AKeyBehavior->Initialize(this, AngleSnapModifier, EKeys::A);
-	AddInputBehavior(AKeyBehavior);
+	// register modifier key behaviors   (disabled because it is not implemented yet)
+	//UKeyAsModifierInputBehavior* AKeyBehavior = NewObject<UKeyAsModifierInputBehavior>();
+	//AKeyBehavior->Initialize(this, AngleSnapModifier, EKeys::A);
+	//AddInputBehavior(AKeyBehavior);
+
 
 	PolygonProperties = NewObject<UDrawPolygonToolStandardProperties>(this, TEXT("Polygon Settings"));
 	PolygonProperties->RestoreProperties(this);
 
 
-	UPositionPlaneGizmoBuilder* PositionPlaneGizmoBuilder = NewObject<UPositionPlaneGizmoBuilder>();
-	UInteractiveGizmoManager* GizmoManager = GetToolManager()->GetPairedGizmoManager();
-	GizmoManager->RegisterGizmoType(TEXT("DrawPolygonPlaneGizmo"), PositionPlaneGizmoBuilder);
-	PositionPlaneGizmo = GizmoManager->CreateGizmo(TEXT("DrawPolygonPlaneGizmo"), TEXT("TestGizmo1"));
-	Cast<UPositionPlaneGizmo>(PositionPlaneGizmo)->OnPositionUpdatedFunc = [this](const FFrame3d& WorldFrame)
-	{
-		UpdateDrawPlaneFromGizmo(WorldFrame);
-	};
+	// Create a new TransformGizmo and associated TransformProxy. The TransformProxy will not be the
+	// parent of any Components in this case, we just use it's transform and change delegate.
+	PlaneTransformProxy = NewObject<UTransformProxy>(this);
+	PlaneTransformProxy->SetTransform(FTransform(DrawPlaneOrientation, DrawPlaneOrigin));
+	PlaneTransformGizmo = GetToolManager()->GetPairedGizmoManager()->Create3AxisTransformGizmo(this);
+	PlaneTransformGizmo->SetActiveTarget(PlaneTransformProxy, GetToolManager());
+	// listen for changes to the proxy and update the plane when that happens
+	PlaneTransformProxy->OnTransformChanged.AddUObject(this, &UDrawPolygonTool::PlaneTransformChanged);
 
+	// initialize material properties for new objects
 	MaterialProperties = NewObject<UNewMeshMaterialProperties>(this);
 	MaterialProperties->RestoreProperties(this);
 
@@ -189,12 +194,12 @@ void UDrawPolygonTool::Setup()
 	PreviewMesh->SetMaterial(MaterialProperties->Material);
 	bPreviewUpdatePending = false;
 
+	// initialize snapping engine and properties
 	SnapEngine.SnapMetricTolerance = ToolSceneQueriesUtil::GetDefaultVisualAngleSnapThreshD();
 	SnapEngine.SnapMetricFunc = [this](const FVector3d& Position1, const FVector3d& Position2) {
 		return ToolSceneQueriesUtil::CalculateViewVisualAngleD(this->CameraState, Position1, Position2);
 	};
-	SnapEngine.Plane = FFrame3d(DrawPlaneOrigin, DrawPlaneOrientation);
-
+	SnapEngine.Plane = FFrame3d((FVector3d)DrawPlaneOrigin, (FQuaterniond)DrawPlaneOrientation);
 
 	SnapProperties = NewObject<UDrawPolygonToolSnapProperties>(this, TEXT("Snapping"));
 	SnapProperties->RestoreProperties(this);
@@ -220,10 +225,7 @@ void UDrawPolygonTool::Shutdown(EToolShutdownType ShutdownType)
 		delete SetPointInWorldConnector;
 	}
 
-	UInteractiveGizmoManager* GizmoManager = GetToolManager()->GetPairedGizmoManager();
-	GizmoManager->DestroyGizmo(PositionPlaneGizmo);
-	PositionPlaneGizmo = nullptr;
-	GizmoManager->DeregisterGizmoType(TEXT("DrawPolygonPlaneGizmo"));
+	GetToolManager()->GetPairedGizmoManager()->DestroyAllGizmosByOwner(this);
 
 	PolygonProperties->SaveProperties(this);
 	SnapProperties->SaveProperties(this);
@@ -907,27 +909,18 @@ void UDrawPolygonTool::SetDrawPlaneFromWorldPos(const FVector& Position, const F
 
 	SnapEngine.Plane = FFrame3d((FVector3d)DrawPlane.Origin, (FQuaterniond)DrawPlane.Rotation);
 
-	UPositionPlaneGizmo* Gizmo = Cast<UPositionPlaneGizmo>(PositionPlaneGizmo);
-	Gizmo->ExternalUpdatePosition(DrawPlaneOrigin, DrawPlaneOrientation, false);
+	PlaneTransformGizmo->SetNewGizmoTransform(FTransform(DrawPlaneOrientation, DrawPlaneOrigin));
 }
 
-void UDrawPolygonTool::UpdateDrawPlaneFromGizmo(const FFrame3d& WorldPosition)
-{
-	if (bIgnoreSnappingToggle == false)
-	{
-		DrawPlaneOrientation = WorldPosition.Rotation;
-	}
-	DrawPlaneOrigin = (FVector)WorldPosition.Origin;
 
+void UDrawPolygonTool::PlaneTransformChanged(UTransformProxy* Proxy, FTransform Transform)
+{
+	DrawPlaneOrientation = Transform.GetRotation();
+	DrawPlaneOrigin = Transform.GetLocation();
 	SnapEngine.Plane = FFrame3d((FVector3d)DrawPlaneOrigin, (FQuaterniond)DrawPlaneOrientation);
-
-	if (bIgnoreSnappingToggle) 
-	{
-		UPositionPlaneGizmo* Gizmo = Cast<UPositionPlaneGizmo>(PositionPlaneGizmo);
-		Gizmo->ExternalUpdatePosition(DrawPlaneOrigin, DrawPlaneOrientation, false);
-	}
 }
 
+
 void UDrawPolygonTool::EmitCurrentPolygon()
 {
 	FString BaseName = (PolygonProperties->OutputMode == EDrawPolygonOutputMode::MeshedPolygon) ?
diff --git a/Engine/Plugins/Experimental/MeshModelingToolset/Source/MeshModelingTools/Public/DrawPolygonTool.h b/Engine/Plugins/Experimental/MeshModelingToolset/Source/MeshModelingTools/Public/DrawPolygonTool.h
--- a/Engine/Plugins/Experimental/MeshModelingToolset/Source/MeshModelingTools/Public/DrawPolygonTool.h
+++ b/Engine/Plugins/Experimental/MeshModelingToolset/Source/MeshModelingTools/Public/DrawPolygonTool.h
@@ -15,6 +15,9 @@
 #include "DrawPolygonTool.generated.h"
 
 
+class UTransformGizmo;
+class UTransformProxy;
+
 /**
  *
  */
@@ -249,9 +252,12 @@ protected:
 	// drawing plane control
 
 	UPROPERTY()
-	UInteractiveGizmo* PositionPlaneGizmo;
+	UTransformGizmo* PlaneTransformGizmo;
+
+	UPROPERTY()
+	UTransformProxy* PlaneTransformProxy;
 
-	void UpdateDrawPlaneFromGizmo(const FFrame3d& WorldPosition);
+	void PlaneTransformChanged(UTransformProxy* Proxy, FTransform Transform);
 
 	IClickBehaviorTarget* SetPointInWorldConnector = nullptr;
 
