ModelingTools: add option to toggle gizmo on/off to DrawPolygonTool. Also connect to A hotkey. Add TValueWatcher to implement common property-polling pattern (necessary since there are no PropertyChange events at runtime).
#rnx

diff --git a/Engine/Plugins/Experimental/MeshModelingToolset/Source/MeshModelingTools/Private/DrawPolygonTool.cpp b/Engine/Plugins/Experimental/MeshModelingToolset/Source/MeshModelingTools/Private/DrawPolygonTool.cpp
--- a/Engine/Plugins/Experimental/MeshModelingToolset/Source/MeshModelingTools/Private/DrawPolygonTool.cpp
+++ b/Engine/Plugins/Experimental/MeshModelingToolset/Source/MeshModelingTools/Private/DrawPolygonTool.cpp
@@ -77,6 +77,7 @@ void UDrawPolygonToolStandardProperties::SaveProperties(UInteractiveTool* SaveFr
 	PropertyCache->ExtrudeHeight = this->ExtrudeHeight;
 	PropertyCache->Steps = this->Steps;
 	PropertyCache->bAllowSelfIntersections = this->bAllowSelfIntersections;
+	PropertyCache->bShowGizmo = this->bShowGizmo;
 }
 
 void UDrawPolygonToolStandardProperties::RestoreProperties(UInteractiveTool* RestoreToTool)
@@ -87,6 +88,7 @@ void UDrawPolygonToolStandardProperties::RestoreProperties(UInteractiveTool* Res
 	this->ExtrudeHeight = PropertyCache->ExtrudeHeight;
 	this->Steps = PropertyCache->Steps;
 	this->bAllowSelfIntersections = PropertyCache->bAllowSelfIntersections;
+	this->bShowGizmo = PropertyCache->bShowGizmo;
 }
 
 
@@ -172,7 +174,10 @@ void UDrawPolygonTool::Setup()
 
 	PolygonProperties = NewObject<UDrawPolygonToolStandardProperties>(this, TEXT("Polygon Settings"));
 	PolygonProperties->RestoreProperties(this);
-
+	ShowGizmoWatcher.Initialize(
+		[this]() { return this->PolygonProperties->bShowGizmo; }, 
+		[this](bool bNewValue) { this->UpdateShowGizmoState(bNewValue); },
+		true);
 
 	// Create a new TransformGizmo and associated TransformProxy. The TransformProxy will not be the
 	// parent of any Components in this case, we just use it's transform and change delegate.
@@ -209,9 +214,7 @@ void UDrawPolygonTool::Setup()
 	AddToolPropertySource(SnapProperties);
 	AddToolPropertySource(MaterialProperties);
 
-	GetToolManager()->DisplayMessage(
-		LOCTEXT("OnStartDrawPolygonTool", "Left-click to place points on the Drawing Plane. Ctrl-click on the scene to reposition the Plane, and click-drag on the gizmo to move it."),
-		EToolMessageLevel::UserNotification);
+	ShowStartupMessage();
 }
 
 
@@ -240,8 +243,19 @@ void UDrawPolygonTool::RegisterActions(FInteractiveToolActionSet& ActionSet)
 		LOCTEXT("PopLastVertexTooltip", "Pop last vertex added to polygon"),
 		EModifierKey::None, EKeys::BackSpace,
 		[this]() { PopLastVertexAction(); });
+
+
+	ActionSet.RegisterAction(this, (int32)EStandardToolActions::BaseClientDefinedActionID + 2,
+		TEXT("ToggleGizmo"),
+		LOCTEXT("ToggleGizmo", "Toggle Gizmo"),
+		LOCTEXT("ToggleGizmoTooltip", "Toggle visibility of the transformation Gizmo"),
+		EModifierKey::None, EKeys::A,
+		[this]() { PolygonProperties->bShowGizmo = !PolygonProperties->bShowGizmo; });
+
 }
 
+
+
 void UDrawPolygonTool::PopLastVertexAction()
 {
 	if (bInInteractiveExtrude || PolygonVertices.Num() == 0)
@@ -463,8 +477,12 @@ void UDrawPolygonTool::Render(IToolsContextRenderAPI* RenderAPI)
 			PreviewOrigin + PolygonProperties->ExtrudeHeight*DrawPlaneNormal, HitPosFrameWorld.Origin,
 			HeightPosColor, 1, 1.0f, 0.0f, true);
 	}
+
+
+	ShowGizmoWatcher.CheckAndUpdate();
 }
 
+
 void UDrawPolygonTool::ResetPolygon()
 {
 	PolygonVertices.Reset();
@@ -828,6 +846,8 @@ void UDrawPolygonTool::BeginInteractiveExtrude()
 	bInInteractiveExtrude = true;
 
 	GeneratePreviewHeightTarget();
+
+	ShowExtrudeMessage();
 }
 
 void UDrawPolygonTool::EndInteractiveExtrude()
@@ -838,6 +858,8 @@ void UDrawPolygonTool::EndInteractiveExtrude()
 	PreviewMesh->SetVisible(false);
 
 	bInInteractiveExtrude = false;
+
+	ShowStartupMessage();
 }
 
 float UDrawPolygonTool::FindInteractiveHeightDistance(const FInputDeviceRay& ClickPos)
@@ -920,6 +942,19 @@ void UDrawPolygonTool::PlaneTransformChanged(UTransformProxy* Proxy, FTransform
 	SnapEngine.Plane = FFrame3d((FVector3d)DrawPlaneOrigin, (FQuaterniond)DrawPlaneOrientation);
 }
 
+void UDrawPolygonTool::UpdateShowGizmoState(bool bNewVisibility)
+{
+	if (bNewVisibility == false)
+	{
+		GetToolManager()->GetPairedGizmoManager()->DestroyAllGizmosByOwner(this);
+	}
+	else
+	{
+		PlaneTransformGizmo = GetToolManager()->GetPairedGizmoManager()->Create3AxisTransformGizmo(this);
+		PlaneTransformGizmo->SetActiveTarget(PlaneTransformProxy, GetToolManager());
+	}
+}
+
 
 void UDrawPolygonTool::EmitCurrentPolygon()
 {
@@ -1061,4 +1096,22 @@ void UDrawPolygonTool::GeneratePreviewHeightTarget()
 	PreviewHeightTargetAABB.SetMesh(&PreviewHeightTarget);
 }
 
+
+
+
+void UDrawPolygonTool::ShowStartupMessage()
+{
+	GetToolManager()->DisplayMessage(
+		LOCTEXT("OnStartDraw", "Left-click to place points on the Drawing Plane. Hold Shift to ignore Snapping. Ctrl-click on the scene to reposition the Plane (Shift+Ctrl-click to only Translate). Backspace to discard last vertex. A key toggles Gizmo."),
+		EToolMessageLevel::UserNotification);
+}
+
+void UDrawPolygonTool::ShowExtrudeMessage()
+{
+	GetToolManager()->DisplayMessage(
+		LOCTEXT("OnStartExtrude", "Set the height of the Extrusion by positioning the mouse over the extrusion volume, or over the scene to snap to relative heights."),
+		EToolMessageLevel::UserNotification);
+}
+
+
 #undef LOCTEXT_NAMESPACE
diff --git a/Engine/Plugins/Experimental/MeshModelingToolset/Source/MeshModelingTools/Public/DrawPolygonTool.h b/Engine/Plugins/Experimental/MeshModelingToolset/Source/MeshModelingTools/Public/DrawPolygonTool.h
--- a/Engine/Plugins/Experimental/MeshModelingToolset/Source/MeshModelingTools/Public/DrawPolygonTool.h
+++ b/Engine/Plugins/Experimental/MeshModelingToolset/Source/MeshModelingTools/Public/DrawPolygonTool.h
@@ -12,6 +12,7 @@
 #include "Snapping/PointPlanarSnapSolver.h"
 #include "ToolSceneQueriesUtil.h"
 #include "Properties/MeshMaterialProperties.h"
+#include "Changes/ValueWatcher.h"
 #include "DrawPolygonTool.generated.h"
 
 
@@ -102,6 +103,9 @@ public:
 	UPROPERTY(EditAnywhere, NonTransactional, Category = Polygon)
 	bool bAllowSelfIntersections = false;
 
+	UPROPERTY(EditAnywhere, NonTransactional, Category = Polygon)
+	bool bShowGizmo = true;
+
 	//
 	// save/restore support
 	//
@@ -152,6 +156,7 @@ public:
 
 
 
+
 /**
  * This tool allows the user to draw and extrude 2D polygons
  */
@@ -199,8 +204,6 @@ public:
 	virtual void EndInteractiveExtrude();
 	virtual float FindInteractiveHeightDistance(const FInputDeviceRay& ClickPos);
 
-	virtual void SetDrawPlaneFromWorldPos(const FVector& Position, const FVector& Normal);
-
 
 public:
 	virtual void PopLastVertexAction();
@@ -249,7 +252,8 @@ protected:
 	UPROPERTY()
 	UPreviewMesh* PreviewMesh;
 
-	// drawing plane control
+	
+	// drawing plane gizmo
 
 	UPROPERTY()
 	UTransformGizmo* PlaneTransformGizmo;
@@ -257,10 +261,18 @@ protected:
 	UPROPERTY()
 	UTransformProxy* PlaneTransformProxy;
 
+	// called on PlaneTransformProxy.OnTransformChanged
 	void PlaneTransformChanged(UTransformProxy* Proxy, FTransform Transform);
 
+	// calls SetDrawPlaneFromWorldPos when user ctrl+clicks on scene
 	IClickBehaviorTarget* SetPointInWorldConnector = nullptr;
 
+	// updates plane and gizmo position
+	virtual void SetDrawPlaneFromWorldPos(const FVector& Position, const FVector& Normal);
+
+	TValueWatcher<bool> ShowGizmoWatcher;
+	void UpdateShowGizmoState(bool bNewVisibility);
+
 
 	// polygon drawing
 
@@ -304,6 +316,11 @@ protected:
 
 	// we use this to generate extruded meshes
 	void GeneratePolygonMesh(const TArray<FVector>& Polygon, FDynamicMesh3* ResultMeshOut, FFrame3d& WorldFrameOut, bool bIncludePreviewVtx, double ExtrudeDistance, bool bExtrudeSymmetric);
+
+
+	// user feedback messages
+	void ShowStartupMessage();
+	void ShowExtrudeMessage();
 };
 
 
diff --git a/Engine/Plugins/Experimental/MeshModelingToolset/Source/ModelingComponents/Public/Changes/ValueWatcher.h b/Engine/Plugins/Experimental/MeshModelingToolset/Source/ModelingComponents/Public/Changes/ValueWatcher.h
new file mode 100644
--- /dev/null
+++ b/Engine/Plugins/Experimental/MeshModelingToolset/Source/ModelingComponents/Public/Changes/ValueWatcher.h
@@ -0,0 +1,44 @@
+// Copyright 1998-2019 Epic Games, Inc. All Rights Reserved.
+
+#pragma once
+
+
+
+/**
+ * TValueWatcher is used to implement a common Tool pattern where it is necessary to
+ * essentially poll a value to see if it has changed, and if it has, call a function.
+ * For example UObject UProperties are frequently directly modified by (eg) DetailsView/etc,
+ * this results in PropertyChange events that can be used in the editor to respond to changes.
+ * However at Runtime no such events are generated and so the Tool has no way of knowing
+ * when UProperties change except by polling and comparing with a cached value.
+ * The purpose of this class is simply to encapsulate this common pattern/logic.
+ */
+template<typename ValueType>
+class TValueWatcher
+{
+private:
+	ValueType CachedValue;
+	TFunction<ValueType(void)> GetValueFunc;
+	TFunction<void(ValueType)> OnValueChangedFunc;
+
+public:
+	void Initialize(
+		TFunction<ValueType(void)> GetValueFuncIn,
+		TFunction<void(ValueType)> OnValueChangedFuncIn,
+		ValueType InitialValue)
+	{
+		GetValueFunc = GetValueFuncIn;
+		OnValueChangedFunc = OnValueChangedFuncIn;
+		CachedValue = InitialValue;
+	}
+
+	void CheckAndUpdate()
+	{
+		bool CurValue = GetValueFunc();
+		if (CurValue != CachedValue)
+		{
+			CachedValue = GetValueFunc();
+			OnValueChangedFunc(CachedValue);
+		}
+	}
+};
diff --git a/Engine/Plugins/Experimental/ModelingToolsEditorMode/Source/ModelingToolsEditorMode/Private/ModelingToolsActions.cpp b/Engine/Plugins/Experimental/ModelingToolsEditorMode/Source/ModelingToolsEditorMode/Private/ModelingToolsActions.cpp
--- a/Engine/Plugins/Experimental/ModelingToolsEditorMode/Source/ModelingToolsEditorMode/Private/ModelingToolsActions.cpp
+++ b/Engine/Plugins/Experimental/ModelingToolsEditorMode/Source/ModelingToolsEditorMode/Private/ModelingToolsActions.cpp
@@ -43,6 +43,7 @@ void FModelingToolActionCommands::RegisterAllToolActions()
 {
 	FModelingToolActionCommands::Register();
 	FSculptToolActionCommands::Register();
+	FDrawPolygonToolActionCommands::Register();
 	FTransformToolActionCommands::Register();
 }
 
@@ -50,6 +51,7 @@ void FModelingToolActionCommands::UnregisterAllToolActions()
 {
 	FModelingToolActionCommands::Unregister();
 	FSculptToolActionCommands::Unregister();
+	FDrawPolygonToolActionCommands::Unregister();
 	FTransformToolActionCommands::Unregister();
 }
 
@@ -68,6 +70,10 @@ void FModelingToolActionCommands::UpdateToolCommandBinding(UInteractiveTool* Too
 	{
 		UPDATE_BINDING(FSculptToolActionCommands);
 	}
+	else if (Cast<UDrawPolygonTool>(Tool) != nullptr)
+	{
+		UPDATE_BINDING(FDrawPolygonToolActionCommands);
+	}
 	else
 	{
 		UPDATE_BINDING(FModelingToolActionCommands);
@@ -112,5 +118,23 @@ void FTransformToolActionCommands::GetToolDefaultObjectList(TArray<UInteractiveT
 
 
 
+FDrawPolygonToolActionCommands::FDrawPolygonToolActionCommands() :
+	TInteractiveToolCommands<FDrawPolygonToolActionCommands>(
+		"ModelingToolsDrawPolygonTool", // Context name for fast lookup
+		NSLOCTEXT("Contexts", "ModelingToolsDrawPolygonTool", "Modeling Tools - Draw Polygon Tool"), // Localized context name for displaying
+		NAME_None, // Parent
+		FEditorStyle::GetStyleSetName() // Icon Style Set
+		)
+{
+}
+
+void FDrawPolygonToolActionCommands::GetToolDefaultObjectList(TArray<UInteractiveTool*>& ToolCDOs)
+{
+	ToolCDOs.Add(GetMutableDefault<UDrawPolygonTool>());
+}
+
+
+
+
 #undef LOCTEXT_NAMESPACE
 
diff --git a/Engine/Plugins/Experimental/ModelingToolsEditorMode/Source/ModelingToolsEditorMode/Public/ModelingToolsActions.h b/Engine/Plugins/Experimental/ModelingToolsEditorMode/Source/ModelingToolsEditorMode/Public/ModelingToolsActions.h
--- a/Engine/Plugins/Experimental/ModelingToolsEditorMode/Source/ModelingToolsEditorMode/Public/ModelingToolsActions.h
+++ b/Engine/Plugins/Experimental/ModelingToolsEditorMode/Source/ModelingToolsEditorMode/Public/ModelingToolsActions.h
@@ -60,3 +60,12 @@ public:
 	FTransformToolActionCommands();
 	virtual void GetToolDefaultObjectList(TArray<UInteractiveTool*>& ToolCDOs) override;
 };
+
+
+
+class FDrawPolygonToolActionCommands : public TInteractiveToolCommands<FDrawPolygonToolActionCommands>
+{
+public:
+	FDrawPolygonToolActionCommands();
+	virtual void GetToolDefaultObjectList(TArray<UInteractiveTool*>& ToolCDOs) override;
+};
