MeshModelingTools: Unconditionally release a texture data lock in the Displacement Tool

diff --git a/Engine/Plugins/Experimental/MeshModelingToolset/Source/MeshModelingTools/Private/DisplaceMeshTool.cpp b/Engine/Plugins/Experimental/MeshModelingToolset/Source/MeshModelingTools/Private/DisplaceMeshTool.cpp
--- a/Engine/Plugins/Experimental/MeshModelingToolset/Source/MeshModelingTools/Private/DisplaceMeshTool.cpp
+++ b/Engine/Plugins/Experimental/MeshModelingToolset/Source/MeshModelingTools/Private/DisplaceMeshTool.cpp
@@ -282,22 +282,23 @@ void UDisplaceMeshTool::UpdateMap(bool bForceUpdate)
 	if (FormatedImageData == nullptr)
 	{
 		DisplaceField = FSampledScalarField2f();
-		return;
 	}
+	else
+	{
+		int TextureWidth = DisplacementMap->GetSizeX();
+		int TextureHeight = DisplacementMap->GetSizeY();
 
-	int TextureWidth = DisplacementMap->GetSizeX();
-	int TextureHeight = DisplacementMap->GetSizeY();
-
-	DisplaceField.Resize(TextureWidth, TextureHeight, 0.0f);
-	DisplaceField.SetCellSize(1.0f / (float)TextureWidth);
+		DisplaceField.Resize(TextureWidth, TextureHeight, 0.0f);
+		DisplaceField.SetCellSize(1.0f / (float)TextureWidth);
 
-	for (int y = 0; y < TextureHeight; ++y)
-	{
-		for (int x = 0; x < TextureWidth; ++x)
+		for (int y = 0; y < TextureHeight; ++y)
 		{
-			FColor PixelColor = FormatedImageData[y*TextureWidth + x];
-			float Value = PixelColor.R / 255.0;
-			DisplaceField.GridValues[y*TextureWidth + x] = Value;
+			for (int x = 0; x < TextureWidth; ++x)
+			{
+				FColor PixelColor = FormatedImageData[y*TextureWidth + x];
+				float Value = PixelColor.R / 255.0;
+				DisplaceField.GridValues[y*TextureWidth + x] = Value;
+			}
 		}
 	}
 
