ModelingTools: add mesh-changed notification to UMeshOpPreviewWithBackgroundCompute
#rnx

diff --git a/Engine/Plugins/Experimental/MeshModelingToolset/Source/ModelingComponents/Private/MeshOpPreviewHelpers.cpp b/Engine/Plugins/Experimental/MeshModelingToolset/Source/ModelingComponents/Private/MeshOpPreviewHelpers.cpp
--- a/Engine/Plugins/Experimental/MeshModelingToolset/Source/ModelingComponents/Private/MeshOpPreviewHelpers.cpp
+++ b/Engine/Plugins/Experimental/MeshModelingToolset/Source/ModelingComponents/Private/MeshOpPreviewHelpers.cpp
@@ -74,6 +74,8 @@ void UMeshOpPreviewWithBackgroundCompute::UpdateResults()
 		PreviewMesh->UpdatePreview(ResultMesh.Get());  // copies the mesh @todo we could just give ownership to the Preview!
 		PreviewMesh->SetVisible(true);
 		bResultValid = true;
+
+		OnMeshUpdated.Broadcast(this);
 	}
 }
 
diff --git a/Engine/Plugins/Experimental/MeshModelingToolset/Source/ModelingComponents/Public/MeshOpPreviewHelpers.h b/Engine/Plugins/Experimental/MeshModelingToolset/Source/ModelingComponents/Public/MeshOpPreviewHelpers.h
--- a/Engine/Plugins/Experimental/MeshModelingToolset/Source/ModelingComponents/Public/MeshOpPreviewHelpers.h
+++ b/Engine/Plugins/Experimental/MeshModelingToolset/Source/ModelingComponents/Public/MeshOpPreviewHelpers.h
@@ -121,6 +121,16 @@ public:
 
 
 
+	//
+	// Change notification
+	//
+
+	DECLARE_MULTICAST_DELEGATE_OneParam(FOnMeshUpdated, UMeshOpPreviewWithBackgroundCompute*);
+	/** This delegate is broadcast whenever the embedded preview mesh is updated */
+	FOnMeshUpdated OnMeshUpdated;
+
+
+
 public:
 	// preview of MeshOperator result
 	UPROPERTY()
