ModelingTools: disable log spam in OctreeDynamicMeshComponent
#rnx

diff --git a/Engine/Plugins/Experimental/MeshModelingToolset/Source/ModelingComponents/Private/OctreeDynamicMeshComponent.cpp b/Engine/Plugins/Experimental/MeshModelingToolset/Source/ModelingComponents/Private/OctreeDynamicMeshComponent.cpp
--- a/Engine/Plugins/Experimental/MeshModelingToolset/Source/ModelingComponents/Private/OctreeDynamicMeshComponent.cpp
+++ b/Engine/Plugins/Experimental/MeshModelingToolset/Source/ModelingComponents/Private/OctreeDynamicMeshComponent.cpp
@@ -63,7 +63,7 @@ void UOctreeDynamicMeshComponent::InitializeMesh(FMeshDescription* MeshDescripti
 	
 	FDynamicMeshOctree3::FStatistics Stats;
 	Octree->ComputeStatistics(Stats);
-	UE_LOG(LogTemp, Warning, TEXT("OctreeStats %s"), *Stats.ToString());
+	//UE_LOG(LogTemp, Warning, TEXT("OctreeStats %s"), *Stats.ToString());
 
 	OctreeCut = MakeUnique<FDynamicMeshOctree3::FTreeCutSet>();
 
@@ -194,7 +194,7 @@ void UOctreeDynamicMeshComponent::NotifyMeshUpdated()
 		}
 
 
-		UE_LOG(LogTemp, Warning, TEXT("Updating %d of %d decomposition sets, %d tris total, spillcount %d"), SetsToUpdate.Num(), CutCellSetMap.Num(), TotalTriangleCount, SpillTriangleCount);
+		//UE_LOG(LogTemp, Warning, TEXT("Updating %d of %d decomposition sets, %d tris total, spillcount %d"), SetsToUpdate.Num(), CutCellSetMap.Num(), TotalTriangleCount, SpillTriangleCount);
 		{
 			SCOPE_CYCLE_COUNTER(STAT_SculptToolOctree_UpdateFromDecomp);
 			CurrentProxy->UpdateFromDecomposition(TriangleDecomposition, SetsToUpdate);
