ModelingTools: Rrename properties and update defaults for Merge & CSG Tools. Add mesh statistics.
#rnx

diff --git a/Engine/Plugins/Experimental/MeshModelingToolset/Source/MeshModelingToolsEditorOnly/Private/MergeMeshesTool.cpp b/Engine/Plugins/Experimental/MeshModelingToolset/Source/MeshModelingToolsEditorOnly/Private/MergeMeshesTool.cpp
--- a/Engine/Plugins/Experimental/MeshModelingToolset/Source/MeshModelingToolsEditorOnly/Private/MergeMeshesTool.cpp
+++ b/Engine/Plugins/Experimental/MeshModelingToolset/Source/MeshModelingToolsEditorOnly/Private/MergeMeshesTool.cpp
@@ -80,6 +80,9 @@ void UMergeMeshesTool::Setup()
 	MergeProps = NewObject<UMergeMeshesToolProperties>();
 	AddToolPropertySource(MergeProps);
 
+	MeshStatisticsProperties = NewObject<UMeshStatisticsProperties>(this);
+	AddToolPropertySource(MeshStatisticsProperties);
+
 	// Hide the source meshes
 	for (auto& ComponentTarget : ComponentTargets)
 	{
@@ -92,6 +95,9 @@ void UMergeMeshesTool::Setup()
 	// initialize the PreviewMesh+BackgroundCompute object
 	Preview = NewObject<UMeshOpPreviewWithBackgroundCompute>(this, "Preview");
 	Preview->Setup(this->TargetWorld, this);
+	Preview->OnMeshUpdated.AddLambda([this](UMeshOpPreviewWithBackgroundCompute* Compute) {
+		MeshStatisticsProperties->Update(*Compute->PreviewMesh->GetPreviewDynamicMesh());
+	});
 
 	CreateLowQualityPreview(); // update the preview with a low-quality result
 	
@@ -102,6 +108,7 @@ void UMergeMeshesTool::Setup()
 	);
 
 	Preview->InvalidateResult();    // start compute
+
 }
 
 
@@ -118,7 +125,7 @@ void UMergeMeshesTool::Shutdown(EToolShutdownType ShutdownType)
 		{
 			ComponentTarget->SetOwnerVisibility(true);
 			AActor* Actor = ComponentTarget->GetOwnerActor();
-			if (MergeProps->bRemoveSources)
+			if (MergeProps->bDeleteInputActors)
 			{
 				Actor->Destroy();
 			}
@@ -177,8 +184,8 @@ TSharedPtr<FDynamicMeshOperator> UMergeMeshesTool::MakeNewOperator()
 {
 	TSharedPtr<FVoxelMergeMeshesOp> MergeOp = MakeShared<FVoxelMergeMeshesOp>();
 	MergeOp->VoxelCount     = MergeProps->VoxelCount;
-	MergeOp->AdaptivityD    = MergeProps->Adaptivity;
-	MergeOp->IsoSurfaceD    = MergeProps->IsoSurface;
+	MergeOp->AdaptivityD    = MergeProps->MeshAdaptivity;
+	MergeOp->IsoSurfaceD    = MergeProps->OffsetDistance;
 	MergeOp->bAutoSimplify  = MergeProps->bAutoSimplify;
 	MergeOp->InputMeshArray = InputMeshes;
 	return MergeOp;
diff --git a/Engine/Plugins/Experimental/MeshModelingToolset/Source/MeshModelingToolsEditorOnly/Private/VoxelCSGMeshesTool.cpp b/Engine/Plugins/Experimental/MeshModelingToolset/Source/MeshModelingToolsEditorOnly/Private/VoxelCSGMeshesTool.cpp
--- a/Engine/Plugins/Experimental/MeshModelingToolset/Source/MeshModelingToolsEditorOnly/Private/VoxelCSGMeshesTool.cpp
+++ b/Engine/Plugins/Experimental/MeshModelingToolset/Source/MeshModelingToolsEditorOnly/Private/VoxelCSGMeshesTool.cpp
@@ -75,10 +75,13 @@ void UVoxelCSGMeshesTool::Setup()
 
 	CSGProps = NewObject<UVoxelCSGMeshesToolProperties>();
 	CSGProps->VoxelCount = 128;
-	CSGProps->Adaptivity = 0.01f;
-	CSGProps->IsoSurface = 0.0f;
+	CSGProps->MeshAdaptivity = 0.01f;
+	CSGProps->OffsetDistance = 0.0f;
 	AddToolPropertySource(CSGProps);
 
+	MeshStatisticsProperties = NewObject<UMeshStatisticsProperties>(this);
+	AddToolPropertySource(MeshStatisticsProperties);
+
 	// Hide the source meshes
 	for (auto& ComponentTarget : ComponentTargets)
 	{
@@ -91,6 +94,9 @@ void UVoxelCSGMeshesTool::Setup()
 	// initialize the PreviewMesh+BackgroundCompute object
 	Preview = NewObject<UMeshOpPreviewWithBackgroundCompute>(this, "Preview");
 	Preview->Setup(this->TargetWorld, this);
+	Preview->OnMeshUpdated.AddLambda([this](UMeshOpPreviewWithBackgroundCompute* Compute) {
+		MeshStatisticsProperties->Update(*Compute->PreviewMesh->GetPreviewDynamicMesh());
+	});
 
 	CreateLowQualityPreview();
 
@@ -117,7 +123,7 @@ void UVoxelCSGMeshesTool::Shutdown(EToolShutdownType ShutdownType)
 		{
 			ComponentTarget->SetOwnerVisibility(true);
 			AActor* Actor = ComponentTarget->GetOwnerActor();
-			if (CSGProps->bRemoveSources)
+			if (CSGProps->bDeleteInputActors)
 			{
 				Actor->Destroy();
 			}
@@ -170,8 +176,8 @@ TSharedPtr<FDynamicMeshOperator> UVoxelCSGMeshesTool::MakeNewOperator()
 	TSharedPtr<FVoxelBooleanMeshesOp> CSGOp = MakeShared<FVoxelBooleanMeshesOp>();
 	CSGOp->Operation      = (FVoxelBooleanMeshesOp::EBooleanOperation)(int)CSGProps->Operation;
 	CSGOp->VoxelCount     = CSGProps->VoxelCount;
-	CSGOp->AdaptivityD    = CSGProps->Adaptivity;
-	CSGOp->IsoSurfaceD    = CSGProps->IsoSurface;
+	CSGOp->AdaptivityD    = CSGProps->MeshAdaptivity;
+	CSGOp->IsoSurfaceD    = CSGProps->OffsetDistance;
 	CSGOp->bAutoSimplify  = CSGProps->bAutoSimplify;
 	CSGOp->InputMeshArray = InputMeshes;
 	return CSGOp;
diff --git a/Engine/Plugins/Experimental/MeshModelingToolset/Source/MeshModelingToolsEditorOnly/Public/MergeMeshesTool.h b/Engine/Plugins/Experimental/MeshModelingToolset/Source/MeshModelingToolsEditorOnly/Public/MergeMeshesTool.h
--- a/Engine/Plugins/Experimental/MeshModelingToolset/Source/MeshModelingToolsEditorOnly/Public/MergeMeshesTool.h
+++ b/Engine/Plugins/Experimental/MeshModelingToolset/Source/MeshModelingToolsEditorOnly/Public/MergeMeshesTool.h
@@ -7,6 +7,7 @@
 #include "InteractiveToolBuilder.h"
 #include "MeshOpPreviewHelpers.h"
 #include "ProxyLODVolume.h"
+#include "Properties/MeshStatisticsProperties.h"
 #include "MergeMeshesTool.generated.h"
 
 
@@ -47,19 +48,19 @@ public:
 
 	/** Remeshing adaptivity */
 	UPROPERTY(EditAnywhere, Category = Options, meta = (UIMin = "0", UIMax = "1", ClampMin = "0", ClampMax = "1"))
-	float Adaptivity = 0.001f;
+	float MeshAdaptivity = 0.001f;
 
 	/** Offset when remeshing, measured in voxels units */
 	UPROPERTY(EditAnywhere, Category = Options, meta = (UIMin = "-2", UIMax = "2", ClampMin = "-2", ClampMax = "2"))
-	float IsoSurface = 0;
+	float OffsetDistance = 0;
 
 	/** Automatically simplify the result of voxel-based merge.*/
 	UPROPERTY(EditAnywhere, Category = Options)
-	bool bAutoSimplify = true;
+	bool bAutoSimplify = false;
 
-	/** Remove the source meshes when accepting results of tool.*/
+	/** Delete the source Actors/Components when accepting results of tool.*/
 	UPROPERTY(EditAnywhere, Category = Options)
-	bool bRemoveSources = true;
+	bool bDeleteInputActors = true;
 };
 
 /**
@@ -95,6 +96,10 @@ protected:
 	UPROPERTY()
 	UMergeMeshesToolProperties* MergeProps;
 
+	UPROPERTY()
+	UMeshStatisticsProperties* MeshStatisticsProperties;
+
+
 	UPROPERTY()
 	UMeshOpPreviewWithBackgroundCompute* Preview;
 
diff --git a/Engine/Plugins/Experimental/MeshModelingToolset/Source/MeshModelingToolsEditorOnly/Public/VoxelCSGMeshesTool.h b/Engine/Plugins/Experimental/MeshModelingToolset/Source/MeshModelingToolsEditorOnly/Public/VoxelCSGMeshesTool.h
--- a/Engine/Plugins/Experimental/MeshModelingToolset/Source/MeshModelingToolsEditorOnly/Public/VoxelCSGMeshesTool.h
+++ b/Engine/Plugins/Experimental/MeshModelingToolset/Source/MeshModelingToolsEditorOnly/Public/VoxelCSGMeshesTool.h
@@ -8,6 +8,7 @@
 #include "MeshOpPreviewHelpers.h"
 #include "MultiSelectionTool.h"
 #include "ProxyLODVolume.h"
+#include "Properties/MeshStatisticsProperties.h"
 #include "VoxelCSGMeshesTool.generated.h"
 
 
@@ -69,19 +70,19 @@ public:
 
 	/** Remeshing adaptivity */
 	UPROPERTY(EditAnywhere, Category = Options, meta = (UIMin = "0", UIMax = "1", ClampMin = "0", ClampMax = "1"))
-	float Adaptivity = 0.001f;
+	float MeshAdaptivity = 0.001f;
 
 	/** Offset when remeshing, measured in voxels units */
 	UPROPERTY(EditAnywhere, Category = Options, meta = (UIMin = "-2", UIMax = "2", ClampMin = "-2", ClampMax = "2"))
-	float IsoSurface = 0;
+	float OffsetDistance = 0;
 
 	/** Automatically simplify the result of voxel-based merge.*/
 	UPROPERTY(EditAnywhere, Category = Options)
-	bool bAutoSimplify = true;
+	bool bAutoSimplify = false;
 
-	/** Remove the source meshes when accepting results of tool.*/
+	/** Remove the source Actors/Components when accepting results of tool.*/
 	UPROPERTY(EditAnywhere, Category = Options)
-	bool bRemoveSources = true;
+	bool bDeleteInputActors = true;
 };
 
 
@@ -121,6 +122,9 @@ protected:
 	UPROPERTY()
 	UVoxelCSGMeshesToolProperties* CSGProps;
 
+	UPROPERTY()
+	UMeshStatisticsProperties* MeshStatisticsProperties;
+
 	UPROPERTY()
 	UMeshOpPreviewWithBackgroundCompute* Preview;
 
