ModelingTools: clean up attributes in Sculpt tool. Add wireframe toggle. Post warning message if input mesh has UV or Normal seams.
#rnx

diff --git a/Engine/Plugins/Experimental/MeshModelingToolset/Source/MeshModelingTools/Private/DynamicMeshSculptTool.cpp b/Engine/Plugins/Experimental/MeshModelingToolset/Source/MeshModelingTools/Private/DynamicMeshSculptTool.cpp
--- a/Engine/Plugins/Experimental/MeshModelingToolset/Source/MeshModelingTools/Private/DynamicMeshSculptTool.cpp
+++ b/Engine/Plugins/Experimental/MeshModelingToolset/Source/MeshModelingTools/Private/DynamicMeshSculptTool.cpp
@@ -40,10 +40,11 @@ UBrushSculptProperties::UBrushSculptProperties()
 {
 	SmoothSpeed = 0.25;
 	BrushSpeed = 0.5;
-	PrimaryBrushType = EDynamicMeshSculptBrushType::Pull;
-	SmoothingType = EMeshSculptToolSmoothType::TexturePreserving;
-	Depth = 0;
+	PrimaryBrushType = EDynamicMeshSculptBrushType::Move;
+	bPreserveUVFlow = false;
+	BrushDepth = 0;
 	bFreezeTarget = false;
+	bShowWireframe = false;
 }
 
 void UBrushSculptProperties::SaveProperties(UInteractiveTool* SaveFromTool)
@@ -52,8 +53,9 @@ void UBrushSculptProperties::SaveProperties(UInteractiveTool* SaveFromTool)
 	PropertyCache->SmoothSpeed = this->SmoothSpeed;
 	PropertyCache->BrushSpeed = this->BrushSpeed;
 	PropertyCache->PrimaryBrushType = this->PrimaryBrushType;
-	PropertyCache->SmoothingType = this->SmoothingType;
-	PropertyCache->Depth = this->Depth;
+	PropertyCache->bPreserveUVFlow = this->bPreserveUVFlow;
+	PropertyCache->BrushDepth = this->BrushDepth;
+	PropertyCache->bShowWireframe = this->bShowWireframe;
 }
 
 void UBrushSculptProperties::RestoreProperties(UInteractiveTool* RestoreToTool)
@@ -62,8 +64,9 @@ void UBrushSculptProperties::RestoreProperties(UInteractiveTool* RestoreToTool)
 	this->SmoothSpeed = PropertyCache->SmoothSpeed;
 	this->BrushSpeed = PropertyCache->BrushSpeed;
 	this->PrimaryBrushType = PropertyCache->PrimaryBrushType;
-	this->SmoothingType = PropertyCache->SmoothingType;
-	this->Depth = PropertyCache->Depth;
+	this->bPreserveUVFlow = PropertyCache->bPreserveUVFlow;
+	this->BrushDepth = PropertyCache->BrushDepth;
+	this->bShowWireframe = PropertyCache->bShowWireframe;
 }
 
 
@@ -127,6 +130,9 @@ void UDynamicMeshSculptTool::Setup()
 	CalculateBrushRadius();
 
 	SculptProperties = NewObject<UBrushSculptProperties>(this, TEXT("Sculpting"));
+	ShowWireframeWatcher.Initialize(
+		[this]() { return SculptProperties->bShowWireframe; },
+		[this](bool bNewValue) { DynamicMeshComponent->bExplicitShowWireframe = bNewValue; }, false);
 
 	RemeshProperties = NewObject<UBrushRemeshProperties>(this, TEXT("Remeshing"));
 	InitialEdgeLength = EstimateIntialSafeTargetLength(*Mesh, 5000);
@@ -166,8 +172,33 @@ void UDynamicMeshSculptTool::Setup()
 	SculptProperties->RestoreProperties(this);
 
 	GetToolManager()->DisplayMessage(
-		LOCTEXT("OnStartSculptTool", "Shift key toggles to Smoothing Brush, Ctrl inverts Brush (where applicable). Q/A keys cycle through Brush Types, Shift+Q/A for Brush History. S/D change Size (shift to small-step), W/E change Speed."),
+		LOCTEXT("OnStartSculptTool", "Hold Shift to Smooth, Ctrl to Invert (where applicable). Q/A keys cycle through Brush Types, Shift+Q/A for Brush History. S/D change Size (shift to small-step), W/E change Speed."),
 		EToolMessageLevel::UserNotification);
+
+	if (bEnableRemeshing && Mesh->HasAttributes())
+	{
+		const FDynamicMeshAttributeSet* Attribs = Mesh->Attributes();
+		bool bHasUVSeams = false;
+		for (int k = 0; k < Attribs->NumUVLayers(); ++k)
+		{
+			bHasUVSeams = bHasUVSeams || Attribs->GetUVLayer(k)->HasInteriorSeamEdges();
+		}
+		bool bHasNormalSeams = Attribs->PrimaryNormals()->HasInteriorSeamEdges();
+
+		if (bHasUVSeams)
+		{
+			GetToolManager()->DisplayMessage(
+				LOCTEXT("UVSeamWarning", "This mesh has UV seams which may limit remeshing. Consider clearing the UV layers using the Remesh Tool."),
+				EToolMessageLevel::UserWarning);
+		} 
+		else if (bHasNormalSeams)
+		{
+			GetToolManager()->DisplayMessage(
+				LOCTEXT("NormalSeamWarning", "This mesh has Hard Normal seams which may limit remeshing. Consider clearing Hard Normals using the Remesh Tool."),
+				EToolMessageLevel::UserWarning);
+		}
+	}
+
 }
 
 
@@ -257,7 +288,7 @@ void UDynamicMeshSculptTool::OnBeginDrag(const FRay& Ray)
 	FHitResult OutHit;
 	if (HitTest(Ray, OutHit))
 	{
-		BrushStartCenterWorld = Ray.PointAt(OutHit.Distance) + SculptProperties->Depth*CurrentBrushRadius*Ray.Direction;
+		BrushStartCenterWorld = Ray.PointAt(OutHit.Distance) + SculptProperties->BrushDepth*CurrentBrushRadius*Ray.Direction;
 
 		bInDrag = true;
 
@@ -393,7 +424,7 @@ void UDynamicMeshSculptTool::ApplyStamp(const FRay& WorldRay)
 		case EDynamicMeshSculptBrushType::SculptMax:
 			ApplySculptMaxBrush(WorldRay);
 			break;
-		case EDynamicMeshSculptBrushType::Pull:
+		case EDynamicMeshSculptBrushType::Move:
 			ApplyMoveBrush(WorldRay);
 			break;
 		case EDynamicMeshSculptBrushType::Pinch:
@@ -454,7 +485,7 @@ void UDynamicMeshSculptTool::ApplySmoothBrush(const FRay& WorldRay)
 
 		double Falloff = CalculateBrushFalloff(OrigPos.Distance(NewBrushPosLocal));
 
-		FVector3d SmoothedPos = (SculptProperties->SmoothingType == EMeshSculptToolSmoothType::TexturePreserving) ?
+		FVector3d SmoothedPos = (SculptProperties->bPreserveUVFlow) ?
 			FMeshWeights::MeanValueCentroid(*Mesh, VertIdx) : FMeshWeights::UniformCentroid(*Mesh, VertIdx);
 
 		FVector3d NewPos = FVector3d::Lerp(OrigPos, SmoothedPos, Falloff*SculptProperties->SmoothSpeed);
@@ -611,7 +642,7 @@ void UDynamicMeshSculptTool::ApplyPinchBrush(const FRay& WorldRay)
 	FTransform Transform = ComponentTarget->GetWorldTransform();
 	FVector NewBrushPosLocal = Transform.InverseTransformPosition(LastBrushPosWorld);
 	FVector BrushNormalLocal = Transform.InverseTransformVectorNoScale(LastBrushPosNormalWorld);
-	FVector OffsetBrushPosLocal = NewBrushPosLocal - SculptProperties->Depth * CurrentBrushRadius * BrushNormalLocal;
+	FVector OffsetBrushPosLocal = NewBrushPosLocal - SculptProperties->BrushDepth * CurrentBrushRadius * BrushNormalLocal;
 
 	double Direction = (bInvert) ? -1.0 : 1.0;
 	double UseSpeed = Direction * FMathd::Sqrt(CurrentBrushRadius) * (SculptProperties->BrushSpeed*0.25);
@@ -667,7 +698,7 @@ FFrame3d UDynamicMeshSculptTool::ComputeROIBrushPlane(const FVector3d& BrushCent
 	AveragePos /= WeightSum;
 
 	FFrame3d Result = FFrame3d(AveragePos, AverageNormal);
-	Result.Origin -= SculptProperties->Depth * CurrentBrushRadius * Result.Z();
+	Result.Origin -= SculptProperties->BrushDepth * CurrentBrushRadius * Result.Z();
 
 	return Result;
 }
@@ -683,7 +714,7 @@ void UDynamicMeshSculptTool::ApplyPlaneBrush(const FRay& WorldRay)
 	FTransform Transform = ComponentTarget->GetWorldTransform();
 	FVector NewBrushPosLocal = Transform.InverseTransformPosition(LastBrushPosWorld);
 	FVector BrushNormalLocal = Transform.InverseTransformVectorNoScale(LastBrushPosNormalWorld);
-	FVector OffsetBrushPosLocal = NewBrushPosLocal - SculptProperties->Depth * CurrentBrushRadius * BrushNormalLocal;
+	FVector OffsetBrushPosLocal = NewBrushPosLocal - SculptProperties->BrushDepth * CurrentBrushRadius * BrushNormalLocal;
 
 	FDynamicMesh3* Mesh = DynamicMeshComponent->GetMesh();
 	double UseSpeed = FMathd::Sqrt(CurrentBrushRadius) * FMathd::Sqrt(SculptProperties->BrushSpeed) * 0.25;
@@ -723,7 +754,7 @@ void UDynamicMeshSculptTool::ApplyFlattenBrush(const FRay& WorldRay)
 	FTransform Transform = ComponentTarget->GetWorldTransform();
 	FVector NewBrushPosLocal = Transform.InverseTransformPosition(LastBrushPosWorld);
 	FVector BrushNormalLocal = Transform.InverseTransformVectorNoScale(LastBrushPosNormalWorld);
-	FVector OffsetBrushPosLocal = NewBrushPosLocal - SculptProperties->Depth * CurrentBrushRadius * BrushNormalLocal;
+	FVector OffsetBrushPosLocal = NewBrushPosLocal - SculptProperties->BrushDepth * CurrentBrushRadius * BrushNormalLocal;
 
 	FDynamicMesh3* Mesh = DynamicMeshComponent->GetMesh();
 	double UseSpeed = FMathd::Sqrt(CurrentBrushRadius) * FMathd::Sqrt(SculptProperties->BrushSpeed) * 0.25;
@@ -892,7 +923,7 @@ bool UDynamicMeshSculptTool::OnUpdateHover(const FInputDeviceRay& DevicePos)
 		FHitResult OutHit;
 		if (HitTest(DevicePos.WorldRay, OutHit))
 		{
-			LastBrushPosWorld = DevicePos.WorldRay.PointAt(OutHit.Distance + SculptProperties->Depth*CurrentBrushRadius);
+			LastBrushPosWorld = DevicePos.WorldRay.PointAt(OutHit.Distance + SculptProperties->BrushDepth*CurrentBrushRadius);
 			LastBrushPosNormalWorld = OutHit.Normal;
 		}
 	}
@@ -907,6 +938,8 @@ void UDynamicMeshSculptTool::Render(IToolsContextRenderAPI* RenderAPI)
 	UMeshSurfacePointTool::Render(RenderAPI);
 	Indicators->Render(RenderAPI);
 
+	ShowWireframeWatcher.CheckAndUpdate();
+
 	//MeshDebugDraw::DrawNormals(DynamicMeshComponent->GetMesh()->Attributes()->PrimaryNormals(),
 	//	25.0f, FColor::Red, 5.0f, true, RenderAPI->GetPrimitiveDrawInterface(), .GetWorldTransform() );
 }
@@ -1058,7 +1091,7 @@ void UDynamicMeshSculptTool::RemeshROIPass()
 	}
 
 	Remesher.SmoothSpeedT = RemeshProperties->Smoothing;
-	Remesher.SmoothType = (SculptProperties->SmoothingType == EMeshSculptToolSmoothType::TexturePreserving) ?
+	Remesher.SmoothType = (SculptProperties->bPreserveUVFlow) ?
 		FRemesher::ESmoothTypes::MeanValue : FRemesher::ESmoothTypes::Uniform;
 	bool bIsUniformSmooth = (Remesher.SmoothType == FRemesher::ESmoothTypes::Uniform);
 
diff --git a/Engine/Plugins/Experimental/MeshModelingToolset/Source/MeshModelingTools/Public/BaseBrushTool.h b/Engine/Plugins/Experimental/MeshModelingToolset/Source/MeshModelingTools/Public/BaseBrushTool.h
--- a/Engine/Plugins/Experimental/MeshModelingToolset/Source/MeshModelingTools/Public/BaseBrushTool.h
+++ b/Engine/Plugins/Experimental/MeshModelingToolset/Source/MeshModelingTools/Public/BaseBrushTool.h
@@ -22,15 +22,15 @@ public:
 	UBrushBaseProperties();
 
 	/** Relative size of brush */
-	UPROPERTY(EditAnywhere, Category = Options, meta = (DisplayName = "Size", UIMin = "0.0", UIMax = "1.0"))
+	UPROPERTY(EditAnywhere, Category = BrushSize, meta = (DisplayName = "Size", UIMin = "0.0", UIMax = "1.0"))
 	float BrushSize;
 
-	/** If true, ignore relative Brush Size and use explicit Radius*/
-	UPROPERTY(EditAnywhere, Category = Options, AdvancedDisplay)
+	/** If true, ignore relative Brush Size and use explicit World Radius */
+	UPROPERTY(EditAnywhere, Category = BrushSize, AdvancedDisplay)
 	bool bSpecifyRadius;
 
 	/** Radius of brush */
-	UPROPERTY(EditAnywhere, Category = Options, AdvancedDisplay, meta = (DisplayName = "Radius", UIMin = "1.0", UIMax = "1000.0", ClampMin = "0.1", ClampMax = "50000.0"))
+	UPROPERTY(EditAnywhere, Category = BrushSize, AdvancedDisplay, meta = (DisplayName = "Radius", UIMin = "1.0", UIMax = "1000.0", ClampMin = "0.1", ClampMax = "50000.0"))
 	float BrushRadius;
 
 
diff --git a/Engine/Plugins/Experimental/MeshModelingToolset/Source/MeshModelingTools/Public/DynamicMeshSculptTool.h b/Engine/Plugins/Experimental/MeshModelingToolset/Source/MeshModelingTools/Public/DynamicMeshSculptTool.h
--- a/Engine/Plugins/Experimental/MeshModelingToolset/Source/MeshModelingTools/Public/DynamicMeshSculptTool.h
+++ b/Engine/Plugins/Experimental/MeshModelingToolset/Source/MeshModelingTools/Public/DynamicMeshSculptTool.h
@@ -12,6 +12,7 @@
 #include "MeshNormals.h"
 #include "BaseBrushTool.h"
 #include "Drawing/ToolDataVisualizer.h"
+#include "Changes/ValueWatcher.h"
 #include "DynamicMeshSculptTool.generated.h"
 
 
@@ -39,29 +40,29 @@ class FMeshVertexChangeBuilder;
 class FDynamicMeshChangeTracker;
 
 
-/** Mesh sculpting brush types */
+/** Mesh Sculpting Brush Types */
 UENUM()
 enum class EDynamicMeshSculptBrushType : uint8
 {
-	/** Move brush */
-	Pull UMETA(DisplayName = "Pull"),
+	/** Move Brush moves vertices parallel to the view plane  */
+	Move UMETA(DisplayName = "Move"),
 
-	/** Sculpt brush */
+	/** Sculpt Brush displaces vertices from the surface */
 	Offset UMETA(DisplayName = "Sculpt"),
 
-	/** SculptMax brush */
+	/** SculptMax Brush displaces vertices to a maximum offset height */
 	SculptMax UMETA(DisplayName = "Sculpt Max"),
 
-	/** Inflate brush */
+	/** Inflate brush displaces vertices along their vertex normals */
 	Inflate UMETA(DisplayName = "Inflate"),
 
-	/** Pinch brush */
+	/** Pinch Brush pulls vertices towards the center of the brush*/
 	Pinch UMETA(DisplayName = "Pinch"),
 
-	/** Flatten brush */
+	/** Flatten Brush pulls vertices towards the average plane of the brush stamp */
 	Flatten UMETA(DisplayName = "Flatten"),
 
-	/** Plane brush */
+	/** Plane Brush pulls vertices to a plane defined by the initial brush position  */
 	Plane UMETA(DisplayName = "Plane"),
 
 
@@ -71,20 +72,6 @@ enum class EDynamicMeshSculptBrushType : uint8
 
 
 
-/** Smooth type for sculpting brush */
-UENUM()
-enum class EMeshSculptToolSmoothType : uint8
-{
-	/** Uniform Smoothing */
-	Uniform UMETA(DisplayName = "Uniform"),
-
-	/** Texture-Preserving Smoothing */
-	TexturePreserving UMETA(DisplayName = "Texture Preserving"),
-};
-
-
-
-
 
 /**
  *
@@ -122,30 +109,33 @@ public:
 	UBrushSculptProperties();
 
 	/** Primary Brush Mode */
-	UPROPERTY(EditAnywhere, Category = Options)
+	UPROPERTY(EditAnywhere, Category = Sculpting)
 	EDynamicMeshSculptBrushType PrimaryBrushType;
 
-	/** Power/Speed of brush */
-	UPROPERTY(EditAnywhere, Category = Options, meta = (UIMin = "0.0", UIMax = "1.0", ClampMin = "0.0", ClampMax = "1.0", EditCondition = "PrimaryBrushType != EDynamicMeshSculptBrushType::Pull"))
+	/** Power/Speed of Brush */
+	UPROPERTY(EditAnywhere, Category = Sculpting, meta = (UIMin = "0.0", UIMax = "1.0", ClampMin = "0.0", ClampMax = "1.0", EditCondition = "PrimaryBrushType != EDynamicMeshSculptBrushType::Pull"))
 	float BrushSpeed;
 
 	/** Smoothing Speed of Smoothing brush */
-	UPROPERTY(EditAnywhere, Category = Options, meta = (UIMin = "0.0", UIMax = "1.0", ClampMin = "0.0", ClampMax = "1.0"))
+	UPROPERTY(EditAnywhere, Category = Sculpting, meta = (UIMin = "0.0", UIMax = "1.0", ClampMin = "0.0", ClampMax = "1.0"))
 	float SmoothSpeed;
 
+	/** If true, try to preserve the shape of the UV/3D mapping. This will prevent smoothing in some cases */
+	UPROPERTY(EditAnywhere, Category = Sculpting)
+	bool bPreserveUVFlow;
+
 
 	/** Depth of Brush into surface along view ray or surface normal, depending on brush */
-	UPROPERTY(EditAnywhere, Category = Options, meta = (UIMin = "-0.5", UIMax = "0.5", ClampMin = "-1.0", ClampMax = "1.0"))
-	float Depth;
+	UPROPERTY(EditAnywhere, Category = Sculpting, meta = (UIMin = "-0.5", UIMax = "0.5", ClampMin = "-1.0", ClampMax = "1.0"))
+	float BrushDepth;
 
 	/** Disable updating of the brush Target Surface after each stroke */
-	UPROPERTY(EditAnywhere, Category = Options, meta = (EditCondition = "PrimaryBrushType == EDynamicMeshSculptBrushType::Sculpt || PrimaryBrushType == EDynamicMeshSculptBrushType::SculptMax || PrimaryBrushType == EDynamicMeshSculptBrushType::Pinch" ))
+	UPROPERTY(EditAnywhere, Category = Sculpting, meta = (EditCondition = "PrimaryBrushType == EDynamicMeshSculptBrushType::Sculpt || PrimaryBrushType == EDynamicMeshSculptBrushType::SculptMax || PrimaryBrushType == EDynamicMeshSculptBrushType::Pinch" ))
 	bool bFreezeTarget;
 
-	/** Smoothing Type */
-	UPROPERTY(EditAnywhere, Category = Options, AdvancedDisplay)
-	EMeshSculptToolSmoothType SmoothingType;
 
+	UPROPERTY(EditAnywhere, Category = Sculpting)
+	bool bShowWireframe;
 
 	virtual void SaveProperties(UInteractiveTool* SaveFromTool) override;
 	virtual void RestoreProperties(UInteractiveTool* RestoreToTool) override;
@@ -165,28 +155,28 @@ public:
 	UBrushRemeshProperties();
 
 	/** Target Relative Triangle Sizefor Dynamic Meshing */
-	UPROPERTY(EditAnywhere, Category = Options, meta = (UIMin = "0.5", UIMax = "2.0", ClampMin = "0.1", ClampMax = "100.0"))
+	UPROPERTY(EditAnywhere, Category = Remeshing, meta = (UIMin = "0.5", UIMax = "2.0", ClampMin = "0.1", ClampMax = "100.0"))
 	float RelativeSize;
 
 	/** Smoothing speed for dynamic meshing */
-	UPROPERTY(EditAnywhere, Category = Options, meta = (UIMin = "0.0", UIMax = "1.0", ClampMin = "0.0", ClampMax = "1.0"))
+	UPROPERTY(EditAnywhere, Category = Remeshing, meta = (UIMin = "0.0", UIMax = "1.0", ClampMin = "0.0", ClampMax = "1.0"))
 	float Smoothing;
 
 
 	/** Enable edge flips */
-	UPROPERTY(EditAnywhere, Category = Options, AdvancedDisplay)
+	UPROPERTY(EditAnywhere, Category = Remeshing, AdvancedDisplay)
 	bool bFlips = true;
 
 	/** Enable edge splits */
-	UPROPERTY(EditAnywhere, Category = Options, AdvancedDisplay)
+	UPROPERTY(EditAnywhere, Category = Remeshing, AdvancedDisplay)
 	bool bSplits = true;
 
 	/** Enable edge collapses */
-	UPROPERTY(EditAnywhere, Category = Options, AdvancedDisplay)
+	UPROPERTY(EditAnywhere, Category = Remeshing, AdvancedDisplay)
 	bool bCollapses = true;
 
 	/** Prevent normal flips */
-	UPROPERTY(EditAnywhere, Category = Options, AdvancedDisplay)
+	UPROPERTY(EditAnywhere, Category = Remeshing, AdvancedDisplay)
 	bool bPreventNormalFlips = false;
 
 };
@@ -275,6 +265,9 @@ protected:
 	void OnDynamicMeshComponentChanged();
 	FDelegateHandle OnDynamicMeshComponentChangedHandle;
 
+	TValueWatcher<bool> ShowWireframeWatcher;
+
+
 	FInterval1d BrushRelativeSizeRange;
 	double CurrentBrushRadius;
 	void CalculateBrushRadius();
