ModelingTools: move FChange implementations to Changes subdirectory
#rnx

diff --git a/Engine/Plugins/Experimental/MeshModelingToolset/Source/MeshModelingTools/Private/DynamicMeshSculptTool.cpp b/Engine/Plugins/Experimental/MeshModelingToolset/Source/MeshModelingTools/Private/DynamicMeshSculptTool.cpp
--- a/Engine/Plugins/Experimental/MeshModelingToolset/Source/MeshModelingTools/Private/DynamicMeshSculptTool.cpp
+++ b/Engine/Plugins/Experimental/MeshModelingToolset/Source/MeshModelingTools/Private/DynamicMeshSculptTool.cpp
@@ -8,13 +8,12 @@
 #include "ProjectionTargets.h"
 #include "MeshConstraints.h"
 #include "MeshConstraintsUtil.h"
-#include "MeshVertexChange.h"
 #include "MeshWeights.h"
 #include "MeshNormals.h"
 #include "MeshIndexUtil.h"
 #include "Drawing/MeshDebugDrawing.h"
 
-#include "MeshVertexChange.h"
+#include "Changes/MeshVertexChange.h"
 #include "Changes/MeshChange.h"
 #include "DynamicMeshChangeTracker.h"
 
diff --git a/Engine/Plugins/Experimental/MeshModelingToolset/Source/MeshModelingTools/Public/EditMeshPolygonsTool.h b/Engine/Plugins/Experimental/MeshModelingToolset/Source/MeshModelingTools/Public/EditMeshPolygonsTool.h
--- a/Engine/Plugins/Experimental/MeshModelingToolset/Source/MeshModelingTools/Public/EditMeshPolygonsTool.h
+++ b/Engine/Plugins/Experimental/MeshModelingToolset/Source/MeshModelingTools/Public/EditMeshPolygonsTool.h
@@ -10,7 +10,7 @@
 #include "Drawing/ToolDataVisualizer.h"
 #include "Transforms/QuickAxisTranslater.h"
 #include "Transforms/QuickAxisRotator.h"
-#include "MeshVertexChange.h"
+#include "Changes/MeshVertexChange.h"
 #include "GroupTopology.h"
 #include "Spatial/GeometrySet3.h"
 #include "Selection/GroupTopologySelector.h"
diff --git a/Engine/Plugins/Experimental/MeshModelingToolset/Source/MeshModelingTools/Public/MeshSelectionTool.h b/Engine/Plugins/Experimental/MeshModelingToolset/Source/MeshModelingTools/Public/MeshSelectionTool.h
--- a/Engine/Plugins/Experimental/MeshModelingToolset/Source/MeshModelingTools/Public/MeshSelectionTool.h
+++ b/Engine/Plugins/Experimental/MeshModelingToolset/Source/MeshModelingTools/Public/MeshSelectionTool.h
@@ -5,7 +5,7 @@
 #include "CoreMinimal.h"
 #include "DynamicMeshBrushTool.h"
 #include "SelectionSet.h"
-#include "MeshSelectionChange.h"
+#include "Changes/MeshSelectionChange.h"
 #include "MeshSelectionTool.generated.h"
 
 
diff --git a/Engine/Plugins/Experimental/MeshModelingToolset/Source/MeshModelingTools/Public/MeshSpaceDeformerTool.h b/Engine/Plugins/Experimental/MeshModelingToolset/Source/MeshModelingTools/Public/MeshSpaceDeformerTool.h
--- a/Engine/Plugins/Experimental/MeshModelingToolset/Source/MeshModelingTools/Public/MeshSpaceDeformerTool.h
+++ b/Engine/Plugins/Experimental/MeshModelingToolset/Source/MeshModelingTools/Public/MeshSpaceDeformerTool.h
@@ -9,7 +9,7 @@
 #include "Drawing/ToolDataVisualizer.h"
 #include "Transforms/QuickAxisTranslater.h"
 #include "Transforms/QuickAxisRotator.h"
-#include "MeshVertexChange.h"
+#include "Changes/MeshVertexChange.h"
 #include "GroupTopology.h"
 #include "Spatial/GeometrySet3.h"
 #include "Selection/GroupTopologySelector.h"
diff --git a/Engine/Plugins/Experimental/MeshModelingToolset/Source/ModelingComponents/Private/Changes/MeshSelectionChange.cpp b/Engine/Plugins/Experimental/MeshModelingToolset/Source/ModelingComponents/Private/Changes/MeshSelectionChange.cpp
new file mode 100644
--- /dev/null
+++ b/Engine/Plugins/Experimental/MeshModelingToolset/Source/ModelingComponents/Private/Changes/MeshSelectionChange.cpp
@@ -0,0 +1,54 @@
+// Copyright 1998-2019 Epic Games, Inc. All Rights Reserved.
+
+#include "Changes/MeshSelectionChange.h"
+
+
+void FMeshSelectionChange::Apply(UObject* Object)
+{
+	UMeshSelectionSet* Selection = CastChecked<UMeshSelectionSet>(Object);
+
+	if (bAdded)
+	{
+		Selection->AddIndices(ElementType, Indices);
+	}
+	else
+	{
+		Selection->RemoveIndices(ElementType, Indices);
+	}
+}
+
+void FMeshSelectionChange::Revert(UObject* Object)
+{
+	UMeshSelectionSet* Selection = CastChecked<UMeshSelectionSet>(Object);
+
+	if (bAdded)
+	{
+		Selection->RemoveIndices(ElementType, Indices);
+	}
+	else
+	{
+		Selection->AddIndices(ElementType, Indices);
+	}
+}
+
+
+FString FMeshSelectionChange::ToString() const
+{
+	return FString(TEXT("Mesh Selection Change"));
+}
+
+
+
+
+
+FMeshSelectionChangeBuilder::FMeshSelectionChangeBuilder(EMeshSelectionElementType ElementType, bool bAdding)
+{
+	Change = MakeUnique<FMeshSelectionChange>();
+	Change->ElementType = ElementType;
+	Change->bAdded = bAdding;
+}
+
+void FMeshSelectionChangeBuilder::Add(int ElementID)
+{
+	Change->Indices.Add(ElementID);
+}
diff --git a/Engine/Plugins/Experimental/MeshModelingToolset/Source/ModelingComponents/Private/Changes/MeshVertexChange.cpp b/Engine/Plugins/Experimental/MeshModelingToolset/Source/ModelingComponents/Private/Changes/MeshVertexChange.cpp
new file mode 100644
--- /dev/null
+++ b/Engine/Plugins/Experimental/MeshModelingToolset/Source/ModelingComponents/Private/Changes/MeshVertexChange.cpp
@@ -0,0 +1,118 @@
+// Copyright 1998-2019 Epic Games, Inc. All Rights Reserved.
+
+#include "Changes/MeshVertexChange.h"
+#include "DynamicMesh3.h"
+
+#include "BaseDynamicMeshComponent.h"
+
+
+void FMeshVertexChange::Apply(UObject* Object)
+{
+	IMeshVertexCommandChangeTarget* ChangeTarget = CastChecked<IMeshVertexCommandChangeTarget>(Object);
+	ChangeTarget->ApplyChange(this, false);
+}
+
+void FMeshVertexChange::Revert(UObject* Object)
+{
+	IMeshVertexCommandChangeTarget* ChangeTarget = CastChecked<IMeshVertexCommandChangeTarget>(Object);
+	ChangeTarget->ApplyChange(this, true);
+}
+
+
+FString FMeshVertexChange::ToString() const
+{
+	return FString(TEXT("Mesh Vertex Change"));
+}
+
+
+
+
+
+FMeshVertexChangeBuilder::FMeshVertexChangeBuilder()
+{
+	Change = MakeUnique<FMeshVertexChange>();
+}
+
+void FMeshVertexChangeBuilder::UpdateVertex(int VertexID, const FVector3d& OldPosition, const FVector3d& NewPosition)
+{
+	const int* FoundIndex = SavedVertices.Find(VertexID);
+	if (FoundIndex == nullptr)
+	{
+		int NewIndex = Change->Vertices.Num();
+		SavedVertices.Add(VertexID, NewIndex);
+		Change->Vertices.Add(VertexID);
+		Change->OldPositions.Add(OldPosition);
+		Change->NewPositions.Add(NewPosition);
+	} 
+	else
+	{
+		Change->NewPositions[*FoundIndex] = NewPosition;
+	}
+}
+
+void FMeshVertexChangeBuilder::UpdateVertexFinal(int VertexID, const FVector3d& NewPosition)
+{
+	check(SavedVertices.Contains(VertexID));
+
+	const int* Index = SavedVertices.Find(VertexID);
+	if ( Index != nullptr )
+	{
+		Change->NewPositions[*Index] = NewPosition;
+	}
+}
+
+
+void FMeshVertexChangeBuilder::SavePosition(const FDynamicMesh3* Mesh, int VertexID, bool bInitial)
+{
+	FVector3d Pos = Mesh->GetVertex(VertexID);
+	if (bInitial)
+	{
+		UpdateVertex(VertexID, Pos, Pos);
+	}
+	else 
+	{
+		UpdateVertexFinal(VertexID, Pos);
+	}
+}
+
+void FMeshVertexChangeBuilder::SavePositions(const FDynamicMesh3* Mesh, const TArray<int>& VertexIDs, bool bInitial)
+{
+	int Num = VertexIDs.Num();
+	if (bInitial)
+	{
+		for (int k = 0; k < Num; ++k)
+		{
+			FVector3d Pos = Mesh->GetVertex(VertexIDs[k]);
+			UpdateVertex(VertexIDs[k], Pos, Pos);
+		}
+	}
+	else
+	{
+		for (int k = 0; k < Num; ++k)
+		{
+			FVector3d Pos = Mesh->GetVertex(VertexIDs[k]);
+			UpdateVertexFinal(VertexIDs[k], Pos);
+		}
+	}
+}
+
+
+void FMeshVertexChangeBuilder::SavePositions(const FDynamicMesh3* Mesh, const TSet<int>& VertexIDs, bool bInitial)
+{
+	if (bInitial)
+	{
+		for (int VertexID : VertexIDs)
+		{
+			FVector3d Pos = Mesh->GetVertex(VertexID);
+			UpdateVertex(VertexID, Pos, Pos);
+		}
+	}
+	else
+	{
+		for (int VertexID : VertexIDs)
+		{
+			FVector3d Pos = Mesh->GetVertex(VertexID);
+			UpdateVertexFinal(VertexID, Pos);
+		}
+	}
+}
\ No newline at end of file
diff --git a/Engine/Plugins/Experimental/MeshModelingToolset/Source/ModelingComponents/Private/MeshSelectionChange.cpp b/Engine/Plugins/Experimental/MeshModelingToolset/Source/ModelingComponents/Private/MeshSelectionChange.cpp
deleted file mode 100644
diff --git a/Engine/Plugins/Experimental/MeshModelingToolset/Source/ModelingComponents/Private/MeshVertexChange.cpp b/Engine/Plugins/Experimental/MeshModelingToolset/Source/ModelingComponents/Private/MeshVertexChange.cpp
deleted file mode 100644
diff --git a/Engine/Plugins/Experimental/MeshModelingToolset/Source/ModelingComponents/Private/OctreeDynamicMeshComponent.cpp b/Engine/Plugins/Experimental/MeshModelingToolset/Source/ModelingComponents/Private/OctreeDynamicMeshComponent.cpp
--- a/Engine/Plugins/Experimental/MeshModelingToolset/Source/ModelingComponents/Private/OctreeDynamicMeshComponent.cpp
+++ b/Engine/Plugins/Experimental/MeshModelingToolset/Source/ModelingComponents/Private/OctreeDynamicMeshComponent.cpp
@@ -24,7 +24,7 @@
 #include "MeshDescriptionToDynamicMesh.h"
 #include "DynamicMeshToMeshDescription.h"
 
-#include "MeshVertexChange.h"
+#include "Changes/MeshVertexChange.h"
 #include "Changes/MeshChange.h"
 #include "DynamicMeshChangeTracker.h"
 
diff --git a/Engine/Plugins/Experimental/MeshModelingToolset/Source/ModelingComponents/Private/SimpleDynamicMeshComponent.cpp b/Engine/Plugins/Experimental/MeshModelingToolset/Source/ModelingComponents/Private/SimpleDynamicMeshComponent.cpp
--- a/Engine/Plugins/Experimental/MeshModelingToolset/Source/ModelingComponents/Private/SimpleDynamicMeshComponent.cpp
+++ b/Engine/Plugins/Experimental/MeshModelingToolset/Source/ModelingComponents/Private/SimpleDynamicMeshComponent.cpp
@@ -24,7 +24,7 @@
 #include "MeshDescriptionToDynamicMesh.h"
 #include "DynamicMeshToMeshDescription.h"
 
-#include "MeshVertexChange.h"
+#include "Changes/MeshVertexChange.h"
 #include "Changes/MeshChange.h"
 #include "DynamicMeshChangeTracker.h"
 
diff --git a/Engine/Plugins/Experimental/MeshModelingToolset/Source/ModelingComponents/Public/BaseDynamicMeshComponent.h b/Engine/Plugins/Experimental/MeshModelingToolset/Source/ModelingComponents/Public/BaseDynamicMeshComponent.h
--- a/Engine/Plugins/Experimental/MeshModelingToolset/Source/ModelingComponents/Public/BaseDynamicMeshComponent.h
+++ b/Engine/Plugins/Experimental/MeshModelingToolset/Source/ModelingComponents/Public/BaseDynamicMeshComponent.h
@@ -5,7 +5,7 @@
 #include "CoreMinimal.h"
 #include "UObject/ObjectMacros.h"
 #include "Components/MeshComponent.h"
-#include "MeshVertexChange.h"
+#include "Changes/MeshVertexChange.h"
 #include "Changes/MeshChange.h"
 
 #include "DynamicMesh3.h"
diff --git a/Engine/Plugins/Experimental/MeshModelingToolset/Source/ModelingComponents/Public/Changes/MeshSelectionChange.h b/Engine/Plugins/Experimental/MeshModelingToolset/Source/ModelingComponents/Public/Changes/MeshSelectionChange.h
new file mode 100644
--- /dev/null
+++ b/Engine/Plugins/Experimental/MeshModelingToolset/Source/ModelingComponents/Public/Changes/MeshSelectionChange.h
@@ -0,0 +1,54 @@
+// Copyright 1998-2019 Epic Games, Inc. All Rights Reserved.
+
+#pragma once
+
+#include "CoreMinimal.h"
+#include "InteractiveToolChange.h"
+#include "SelectionSet.h"
+
+
+
+
+/**
+ * FMeshSelectionChange represents an reversible change to a UMeshSelectionSet
+ */
+class MODELINGCOMPONENTS_API FMeshSelectionChange : public FToolCommandChange
+{
+public:
+	EMeshSelectionElementType ElementType = EMeshSelectionElementType::Vertex;
+	TArray<int> Indices;
+	bool bAdded;
+
+	/** Makes the change to the object */
+	virtual void Apply(UObject* Object) override;
+
+	/** Reverts change to the object */
+	virtual void Revert(UObject* Object) override;
+
+	/** Describes this change (for debugging) */
+	virtual FString ToString() const override;
+};
+
+
+
+
+/**
+ * FMeshSelectionChangeBuilder can be used to construct a FMeshSelectionChange.
+ */
+class MODELINGCOMPONENTS_API FMeshSelectionChangeBuilder
+{
+public:
+	TUniquePtr<FMeshSelectionChange> Change;
+
+	/**
+	 * Initialize of a selection change of given type. 
+	 * @param ElementType type of element (face/edge/vtx) being added/removed
+	 * @param bAdding if true, output change has bAdded=true, otherwise bAdded=false
+	 */
+	FMeshSelectionChangeBuilder(EMeshSelectionElementType ElementType, bool bAdding);
+
+	/**
+	 * Add ElementID to list of changed elements  (which may then be added or removed, depending on bAdding parameter to constructor)
+	 */
+	void Add(int ElementID);
+};
diff --git a/Engine/Plugins/Experimental/MeshModelingToolset/Source/ModelingComponents/Public/Changes/MeshVertexChange.h b/Engine/Plugins/Experimental/MeshModelingToolset/Source/ModelingComponents/Public/Changes/MeshVertexChange.h
new file mode 100644
--- /dev/null
+++ b/Engine/Plugins/Experimental/MeshModelingToolset/Source/ModelingComponents/Public/Changes/MeshVertexChange.h
@@ -0,0 +1,76 @@
+// Copyright 1998-2019 Epic Games, Inc. All Rights Reserved.
+
+#pragma once
+
+#include "CoreMinimal.h"
+#include "UObject/Interface.h"
+#include "InteractiveToolChange.h"
+#include "VectorTypes.h"
+#include "MeshVertexChange.generated.h"
+
+class FDynamicMesh3;
+
+/**
+ * FMeshVertexChange represents an reversible change to a set of vertex positions.
+ * Currently only a USimpleDynamicMeshComponent target is supported.
+ * 
+ * @todo support optionally storing old/new normals and tangents
+ * @todo support applying to a StaticMeshComponent/MeshDescription ?
+ */
+class MODELINGCOMPONENTS_API FMeshVertexChange : public FToolCommandChange
+{
+public:
+	TArray<int> Vertices;
+	TArray<FVector3d> OldPositions;
+	TArray<FVector3d> NewPositions;
+
+
+	/** Makes the change to the object */
+	virtual void Apply(UObject* Object) override;
+
+	/** Reverts change to the object */
+	virtual void Revert(UObject* Object) override;
+
+	/** Describes this change (for debugging) */
+	virtual FString ToString() const override;
+};
+
+
+
+/**
+ * FMeshVertexChangeBuilder can be used to construct a FMeshVertexChange.
+ * Usage is to call UpdateVertex() each time a vertex moves, with the old and new positions.
+ */
+class MODELINGCOMPONENTS_API FMeshVertexChangeBuilder
+{
+public:
+	TUniquePtr<FMeshVertexChange> Change;
+	TMap<int, int> SavedVertices;
+
+	FMeshVertexChangeBuilder();
+
+	void UpdateVertex(int VertexID, const FVector3d& OldPosition, const FVector3d& NewPosition);
+	void UpdateVertexFinal(int VertexID, const FVector3d& NewPosition);
+
+	void SavePosition(const FDynamicMesh3* Mesh, int VertexID, bool bInitial);
+	void SavePositions(const FDynamicMesh3* Mesh, const TArray<int>& VertexIDs, bool bInitial);
+	void SavePositions(const FDynamicMesh3* Mesh, const TSet<int>& VertexIDs, bool bInitial);
+};
+
+
+
+UINTERFACE()
+class MODELINGCOMPONENTS_API UMeshVertexCommandChangeTarget : public UInterface
+{
+	GENERATED_BODY()
+};
+/**
+ * IMeshVertexCommandChangeTarget is an interface which is used to apply a FMeshVertexChange
+ */
+class MODELINGCOMPONENTS_API IMeshVertexCommandChangeTarget
+{
+	GENERATED_BODY()
+public:
+	virtual void ApplyChange(const FMeshVertexChange* Change, bool bRevert) = 0;
+};
+
diff --git a/Engine/Plugins/Experimental/MeshModelingToolset/Source/ModelingComponents/Public/MeshSelectionChange.h b/Engine/Plugins/Experimental/MeshModelingToolset/Source/ModelingComponents/Public/MeshSelectionChange.h
deleted file mode 100644
diff --git a/Engine/Plugins/Experimental/MeshModelingToolset/Source/ModelingComponents/Public/MeshVertexChange.h b/Engine/Plugins/Experimental/MeshModelingToolset/Source/ModelingComponents/Public/MeshVertexChange.h
deleted file mode 100644
diff --git a/Engine/Plugins/Experimental/MeshModelingToolset/Source/ModelingComponents/Public/PreviewMesh.h b/Engine/Plugins/Experimental/MeshModelingToolset/Source/ModelingComponents/Public/PreviewMesh.h
--- a/Engine/Plugins/Experimental/MeshModelingToolset/Source/ModelingComponents/Public/PreviewMesh.h
+++ b/Engine/Plugins/Experimental/MeshModelingToolset/Source/ModelingComponents/Public/PreviewMesh.h
@@ -6,7 +6,7 @@
 #include "SimpleDynamicMeshComponent.h"
 #include "DynamicMeshAABBTree3.h"
 #include "InteractiveToolObjects.h"
-#include "MeshVertexChange.h"
+#include "Changes/MeshVertexChange.h"
 #include "Changes/MeshChange.h"
 #include "PreviewMesh.generated.h"
 
