ModelingTools: disable log spam
#rnx

diff --git a/Engine/Plugins/Experimental/MeshModelingToolset/Source/ModelingComponents/Private/OctreeDynamicMeshSceneProxy.h b/Engine/Plugins/Experimental/MeshModelingToolset/Source/ModelingComponents/Private/OctreeDynamicMeshSceneProxy.h
--- a/Engine/Plugins/Experimental/MeshModelingToolset/Source/ModelingComponents/Private/OctreeDynamicMeshSceneProxy.h
+++ b/Engine/Plugins/Experimental/MeshModelingToolset/Source/ModelingComponents/Private/OctreeDynamicMeshSceneProxy.h
@@ -313,7 +313,7 @@ public:
 
 		}
 
-		UE_LOG(LogTemp, Warning, TEXT("Have %d renderbuffers"), RenderBufferSets.Num());
+		//UE_LOG(LogTemp, Warning, TEXT("Have %d renderbuffers"), RenderBufferSets.Num());
 	}
 
 
