ModelingTools: remove ToolIndicatorSet and BrushToolIndicator, add BrushStampIndicator based on Gizmo framework. Updated Brush tools as necessary.
#rnx

diff --git a/Engine/Plugins/Experimental/MeshModelingToolset/Source/MeshModelingTools/Private/BaseBrushTool.cpp b/Engine/Plugins/Experimental/MeshModelingToolset/Source/MeshModelingTools/Private/BaseBrushTool.cpp
--- a/Engine/Plugins/Experimental/MeshModelingToolset/Source/MeshModelingTools/Private/BaseBrushTool.cpp
+++ b/Engine/Plugins/Experimental/MeshModelingToolset/Source/MeshModelingTools/Private/BaseBrushTool.cpp
@@ -2,8 +2,8 @@
 
 #include "BaseBrushTool.h"
 #include "InteractiveToolManager.h"
+#include "InteractiveGizmoManager.h"
 #include "ToolBuilderUtil.h"
-#include "BrushToolIndicator.h"
 
 #define LOCTEXT_NAMESPACE "UBaseBrushTool"
 
@@ -40,7 +40,6 @@ UBaseBrushTool::UBaseBrushTool()
 }
 
 
-
 void UBaseBrushTool::Setup()
 {
 	UMeshSurfacePointTool::Setup();
@@ -50,21 +49,16 @@ void UBaseBrushTool::Setup()
 	BrushProperties = NewObject<UBrushBaseProperties>(this, TEXT("Brush"));
 	RecalculateBrushRadius();
 
-
-	// create brush stamp indicator
-	Indicators = NewObject<UToolIndicatorSet>(this, "Indicators");
-	Indicators->Connect(this);
-	Indicators->AddIndicator(MakeBrushIndicator());
-
-
 	// initialize our properties
 	AddToolPropertySource(BrushProperties);
+
+	SetupBrushStampIndicator();
 }
 
 
 void UBaseBrushTool::Shutdown(EToolShutdownType ShutdownType)
 {
-	Indicators->Disconnect();
+	ShutdownBrushStampIndicator();
 }
 
 
@@ -77,19 +71,6 @@ void UBaseBrushTool::OnPropertyModified(UObject* PropertySet, UProperty* Propert
 }
 
 
-IToolIndicator* UBaseBrushTool::MakeBrushIndicator()
-{
-	UBrushStampSizeIndicator* StampIndicator = NewObject<UBrushStampSizeIndicator>(this, "Brush Circle");
-	StampIndicator->bDrawSecondaryLines = true;
-	StampIndicator->DepthLayer = 1;
-	StampIndicator->BrushRadius = MakeAttributeLambda([this] { return this->LastBrushStamp.Radius; });
-	StampIndicator->BrushPosition = MakeAttributeLambda([this] { return this->LastBrushStamp.WorldPosition; });
-	StampIndicator->BrushNormal = MakeAttributeLambda([this] { return this->LastBrushStamp.WorldNormal; });
-	return StampIndicator;
-}
-
-
-
 
 void UBaseBrushTool::IncreaseBrushSizeAction()
 {
@@ -187,15 +168,38 @@ bool UBaseBrushTool::OnUpdateHover(const FInputDeviceRay& DevicePos)
 void UBaseBrushTool::Render(IToolsContextRenderAPI* RenderAPI)
 {
 	UMeshSurfacePointTool::Render(RenderAPI);
-	Indicators->Render(RenderAPI);
+
+	UpdateBrushStampIndicator();
 }
 
 void UBaseBrushTool::Tick(float DeltaTime)
 {
 	UMeshSurfacePointTool::Tick(DeltaTime);
-	Indicators->Tick(DeltaTime);
 }
 
 
 
+
+const FString BrushIndicatorGizmoType = TEXT("BrushIndicatorGizmoType");
+
+void UBaseBrushTool::SetupBrushStampIndicator()
+{
+	// register and spawn brush indicator gizmo
+	GetToolManager()->GetPairedGizmoManager()->RegisterGizmoType(BrushIndicatorGizmoType, NewObject<UBrushStampIndicatorBuilder>());
+	BrushStampIndicator = GetToolManager()->GetPairedGizmoManager()->CreateGizmo<UBrushStampIndicator>(BrushIndicatorGizmoType, FString(), this);
+}
+
+void UBaseBrushTool::UpdateBrushStampIndicator()
+{
+	BrushStampIndicator->Update(LastBrushStamp.Radius, LastBrushStamp.WorldPosition, LastBrushStamp.WorldNormal);
+}
+
+void UBaseBrushTool::ShutdownBrushStampIndicator()
+{
+	GetToolManager()->GetPairedGizmoManager()->DestroyAllGizmosByOwner(this);
+	BrushStampIndicator = nullptr;
+	GetToolManager()->GetPairedGizmoManager()->DeregisterGizmoType(BrushIndicatorGizmoType);
+}
+
+
 #undef LOCTEXT_NAMESPACE
diff --git a/Engine/Plugins/Experimental/MeshModelingToolset/Source/MeshModelingTools/Private/BrushToolIndicator.cpp b/Engine/Plugins/Experimental/MeshModelingToolset/Source/MeshModelingTools/Private/BrushToolIndicator.cpp
deleted file mode 100644
diff --git a/Engine/Plugins/Experimental/MeshModelingToolset/Source/MeshModelingTools/Private/DynamicMeshSculptTool.cpp b/Engine/Plugins/Experimental/MeshModelingToolset/Source/MeshModelingTools/Private/DynamicMeshSculptTool.cpp
--- a/Engine/Plugins/Experimental/MeshModelingToolset/Source/MeshModelingTools/Private/DynamicMeshSculptTool.cpp
+++ b/Engine/Plugins/Experimental/MeshModelingToolset/Source/MeshModelingTools/Private/DynamicMeshSculptTool.cpp
@@ -2,6 +2,7 @@
 
 #include "DynamicMeshSculptTool.h"
 #include "InteractiveToolManager.h"
+#include "InteractiveGizmoManager.h"
 #include "ToolBuilderUtil.h"
 
 #include "SubRegionRemesher.h"
@@ -88,6 +89,9 @@ UDynamicMeshSculptTool::UDynamicMeshSculptTool()
 }
 
 
+const FString BrushIndicatorGizmoType = TEXT("BrushIndicatorGizmoType");
+
+
 void UDynamicMeshSculptTool::Setup()
 {
 	UMeshSurfacePointTool::Setup();
@@ -146,17 +150,9 @@ void UDynamicMeshSculptTool::Setup()
 	bStampPending = false;
 	ActiveVertexChange = nullptr;
 
-	// create indicators
-	Indicators = NewObject<UToolIndicatorSet>(this, "Indicators");
-	Indicators->Connect(this);
-
-	UBrushStampSizeIndicator* StampIndicator = NewObject<UBrushStampSizeIndicator>(this, "Brush Circle");
-	StampIndicator->bDrawSecondaryLines = true;
-	StampIndicator->DepthLayer = 1;
-	StampIndicator->BrushRadius = MakeAttributeLambda([this] { return (float)this->CurrentBrushRadius; });
-	StampIndicator->BrushPosition = MakeAttributeLambda([this] { return this->LastBrushPosWorld; });
-	StampIndicator->BrushNormal = MakeAttributeLambda([this] { return this->LastBrushPosNormalWorld; });
-	Indicators->AddIndicator(StampIndicator);
+	// register and spawn brush indicator gizmo
+	GetToolManager()->GetPairedGizmoManager()->RegisterGizmoType(BrushIndicatorGizmoType, NewObject<UBrushStampIndicatorBuilder>());
+	BrushIndicator = GetToolManager()->GetPairedGizmoManager()->CreateGizmo<UBrushStampIndicator>(BrushIndicatorGizmoType, FString(), this);
 
 	// initialize our properties
 	AddToolPropertySource(SculptProperties);
@@ -205,7 +201,9 @@ void UDynamicMeshSculptTool::Setup()
 
 void UDynamicMeshSculptTool::Shutdown(EToolShutdownType ShutdownType)
 {
-	Indicators->Disconnect();
+	GetToolManager()->GetPairedGizmoManager()->DestroyAllGizmosByOwner(this);
+	BrushIndicator = nullptr;
+	GetToolManager()->GetPairedGizmoManager()->DeregisterGizmoType(BrushIndicatorGizmoType);
 
 	if (DynamicMeshComponent != nullptr)
 	{
@@ -935,12 +933,10 @@ bool UDynamicMeshSculptTool::OnUpdateHover(const FInputDeviceRay& DevicePos)
 void UDynamicMeshSculptTool::Render(IToolsContextRenderAPI* RenderAPI)
 {
 	UMeshSurfacePointTool::Render(RenderAPI);
-	Indicators->Render(RenderAPI);
 
-	ShowWireframeWatcher.CheckAndUpdate();
+	BrushIndicator->Update( (float)this->CurrentBrushRadius, this->LastBrushPosWorld, this->LastBrushPosNormalWorld );
 
-	//MeshDebugDraw::DrawNormals(DynamicMeshComponent->GetMesh()->Attributes()->PrimaryNormals(),
-	//	25.0f, FColor::Red, 5.0f, true, RenderAPI->GetPrimitiveDrawInterface(), .GetWorldTransform() );
+	ShowWireframeWatcher.CheckAndUpdate();
 }
 
 
@@ -949,7 +945,6 @@ void UDynamicMeshSculptTool::Tick(float DeltaTime)
 	SCOPE_CYCLE_COUNTER(STAT_SculptToolTick);
 
 	UMeshSurfacePointTool::Tick(DeltaTime);
-	Indicators->Tick(DeltaTime);
 
 	// if user changed to not-frozen, we need to update the target
 	if (bCachedFreezeTarget != SculptProperties->bFreezeTarget)
diff --git a/Engine/Plugins/Experimental/MeshModelingToolset/Source/MeshModelingTools/Private/MeshSelectionTool.cpp b/Engine/Plugins/Experimental/MeshModelingToolset/Source/MeshModelingTools/Private/MeshSelectionTool.cpp
--- a/Engine/Plugins/Experimental/MeshModelingToolset/Source/MeshModelingTools/Private/MeshSelectionTool.cpp
+++ b/Engine/Plugins/Experimental/MeshModelingToolset/Source/MeshModelingTools/Private/MeshSelectionTool.cpp
@@ -375,8 +375,6 @@ void UMeshSelectionTool::Render(IToolsContextRenderAPI* RenderAPI)
 void UMeshSelectionTool::Tick(float DeltaTime)
 {
 	UDynamicMeshBrushTool::Tick(DeltaTime);
-	Indicators->Tick(DeltaTime);
-
 
 	if (bStampPending)
 	{
diff --git a/Engine/Plugins/Experimental/MeshModelingToolset/Source/MeshModelingTools/Public/BaseBrushTool.h b/Engine/Plugins/Experimental/MeshModelingToolset/Source/MeshModelingTools/Public/BaseBrushTool.h
--- a/Engine/Plugins/Experimental/MeshModelingToolset/Source/MeshModelingTools/Public/BaseBrushTool.h
+++ b/Engine/Plugins/Experimental/MeshModelingToolset/Source/MeshModelingTools/Public/BaseBrushTool.h
@@ -5,7 +5,7 @@
 #include "CoreMinimal.h"
 #include "BaseTools/MeshSurfacePointTool.h"
 #include "BoxTypes.h"
-#include "ToolIndicatorSet.h"
+#include "Gizmos/BrushStampIndicator.h"
 #include "BaseBrushTool.generated.h"
 
 
@@ -103,10 +103,6 @@ public:
 	UPROPERTY()
 	UBrushBaseProperties* BrushProperties;
 
-	/** Brush indicator */
-	UPROPERTY()
-	UToolIndicatorSet* Indicators;
-
 	/** Set to true by Tool if user is currently in an active brush stroke*/
 	UPROPERTY()
 	bool bInBrushStroke = false;
@@ -120,7 +116,6 @@ public:
 	virtual void IncreaseBrushSizeAction();
 	virtual void DecreaseBrushSizeAction();
 
-
 	virtual bool IsInBrushStroke() const { return bInBrushStroke; }
 
 protected:
@@ -130,13 +125,22 @@ protected:
 	 */
 	virtual double EstimateMaximumTargetDimension() { return 100.0; }
 
-	// creates brush indicator widget
-	virtual IToolIndicator* MakeBrushIndicator();
-
-
 protected:
 	FInterval1d BrushRelativeSizeRange;
 	double CurrentBrushRadius;
 	void RecalculateBrushRadius();
 
+
+
+	//
+	// Brush Indicator support
+	//
+protected:
+
+	UPROPERTY()
+	UBrushStampIndicator* BrushStampIndicator;
+
+	virtual void SetupBrushStampIndicator();
+	virtual void UpdateBrushStampIndicator();
+	virtual void ShutdownBrushStampIndicator();
 };
diff --git a/Engine/Plugins/Experimental/MeshModelingToolset/Source/MeshModelingTools/Public/BrushToolIndicator.h b/Engine/Plugins/Experimental/MeshModelingToolset/Source/MeshModelingTools/Public/BrushToolIndicator.h
deleted file mode 100644
diff --git a/Engine/Plugins/Experimental/MeshModelingToolset/Source/MeshModelingTools/Public/DynamicMeshSculptTool.h b/Engine/Plugins/Experimental/MeshModelingToolset/Source/MeshModelingTools/Public/DynamicMeshSculptTool.h
--- a/Engine/Plugins/Experimental/MeshModelingToolset/Source/MeshModelingTools/Public/DynamicMeshSculptTool.h
+++ b/Engine/Plugins/Experimental/MeshModelingToolset/Source/MeshModelingTools/Public/DynamicMeshSculptTool.h
@@ -8,11 +8,11 @@
 #include "OctreeDynamicMeshComponent.h"
 #include "DynamicMeshAABBTree3.h"
 #include "DynamicMeshOctree3.h"
-#include "BrushToolIndicator.h"
 #include "MeshNormals.h"
 #include "BaseBrushTool.h"
 #include "Drawing/ToolDataVisualizer.h"
 #include "Changes/ValueWatcher.h"
+#include "Gizmos/BrushStampIndicator.h"
 #include "DynamicMeshSculptTool.generated.h"
 
 
@@ -234,7 +234,6 @@ public:
 	UBrushRemeshProperties* RemeshProperties;
 
 
-
 public:
 
 	virtual void IncreaseBrushRadiusAction();
@@ -256,7 +255,7 @@ public:
 protected:
 
 	UPROPERTY()
-	UToolIndicatorSet* Indicators;
+	UBrushStampIndicator* BrushIndicator;
 
 	UPROPERTY()
 	UOctreeDynamicMeshComponent* DynamicMeshComponent;
diff --git a/Engine/Plugins/Experimental/MeshModelingToolset/Source/ModelingComponents/Private/Gizmos/BrushStampIndicator.cpp b/Engine/Plugins/Experimental/MeshModelingToolset/Source/ModelingComponents/Private/Gizmos/BrushStampIndicator.cpp
new file mode 100644
--- /dev/null
+++ b/Engine/Plugins/Experimental/MeshModelingToolset/Source/ModelingComponents/Private/Gizmos/BrushStampIndicator.cpp
@@ -0,0 +1,54 @@
+// Copyright 1998-2019 Epic Games, Inc. All Rights Reserved.
+
+#include "Gizmos/BrushStampIndicator.h"
+#include "InteractiveGizmoManager.h"
+#include "Drawing/ToolDataVisualizer.h"
+
+
+
+UInteractiveGizmo* UBrushStampIndicatorBuilder::BuildGizmo(const FToolBuilderState& SceneState) const
+{
+	UBrushStampIndicator* NewGizmo = NewObject<UBrushStampIndicator>(SceneState.GizmoManager);
+	return NewGizmo;
+}
+
+
+
+void UBrushStampIndicator::Setup()
+{
+}
+
+void UBrushStampIndicator::Shutdown()
+{
+}
+
+void UBrushStampIndicator::Render(IToolsContextRenderAPI* RenderAPI)
+{
+	FToolDataVisualizer Draw;
+	Draw.BeginFrame(RenderAPI);
+
+	FVector3f Perp1, Perp2;
+	VectorUtil::MakePerpVectors((FVector3f)BrushNormal, Perp1, Perp2);
+
+	Draw.DrawCircle(BrushPosition, BrushNormal, BrushRadius, SampleStepCount, LineColor, LineThickness, bDepthTested);
+
+	if (bDrawSecondaryLines)
+	{
+		Draw.DrawCircle(BrushPosition, BrushNormal, BrushRadius*0.5f, SampleStepCount, SecondaryLineColor, SecondaryLineThickness, bDepthTested);
+		Draw.DrawLine(BrushPosition, BrushPosition + BrushRadius*BrushNormal, SecondaryLineColor, SecondaryLineThickness, bDepthTested);
+	}
+
+	Draw.EndFrame();
+}
+
+void UBrushStampIndicator::Tick(float DeltaTime)
+{
+}
+
+
+void UBrushStampIndicator::Update(float Radius, const FVector& Position, const FVector& Normal)
+{
+	BrushRadius = Radius;
+	BrushPosition = Position;
+	BrushNormal = Normal;
+}
\ No newline at end of file
diff --git a/Engine/Plugins/Experimental/MeshModelingToolset/Source/ModelingComponents/Private/ToolIndicatorSet.cpp b/Engine/Plugins/Experimental/MeshModelingToolset/Source/ModelingComponents/Private/ToolIndicatorSet.cpp
deleted file mode 100644
diff --git a/Engine/Plugins/Experimental/MeshModelingToolset/Source/ModelingComponents/Public/Gizmos/BrushStampIndicator.h b/Engine/Plugins/Experimental/MeshModelingToolset/Source/ModelingComponents/Public/Gizmos/BrushStampIndicator.h
new file mode 100644
--- /dev/null
+++ b/Engine/Plugins/Experimental/MeshModelingToolset/Source/ModelingComponents/Public/Gizmos/BrushStampIndicator.h
@@ -0,0 +1,84 @@
+// Copyright 1998-2019 Epic Games, Inc. All Rights Reserved.
+
+#pragma once
+
+#include "InteractiveGizmo.h"
+#include "InteractiveGizmoBuilder.h"
+#include "BrushStampIndicator.generated.h"
+
+
+
+UCLASS()
+class MODELINGCOMPONENTS_API UBrushStampIndicatorBuilder : public UInteractiveGizmoBuilder
+{
+	GENERATED_BODY()
+public:
+	virtual UInteractiveGizmo* BuildGizmo(const FToolBuilderState& SceneState) const override;
+};
+
+
+
+/*
+ * UBrushStampIndicator is a simple 3D brush indicator. 
+ */
+UCLASS(Transient)
+class MODELINGCOMPONENTS_API UBrushStampIndicator : public UInteractiveGizmo
+{
+	GENERATED_BODY()
+
+public:
+
+	virtual void Setup() override;
+	virtual void Shutdown() override;
+
+	virtual void Render(IToolsContextRenderAPI* RenderAPI) override;
+
+	virtual void Tick(float DeltaTime) override;
+
+
+	virtual void Update(float Radius, const FVector& Position, const FVector& Normal);
+
+
+public:
+
+	UPROPERTY()
+	float BrushRadius = 1.0f;
+
+	UPROPERTY()
+	FVector BrushPosition = FVector::ZeroVector;
+
+	UPROPERTY()
+	FVector BrushNormal = FVector(0, 0, 1);;
+
+
+	UPROPERTY()
+	int SampleStepCount = 32;
+
+	UPROPERTY()
+	FLinearColor LineColor = FLinearColor(0.96f, 0.06f, 0.06f);
+
+	UPROPERTY()
+	float LineThickness = 2.0f;
+
+	UPROPERTY()
+	bool bDepthTested = false;
+
+	UPROPERTY()
+	bool IsPixelThickness = true;
+
+	UPROPERTY()
+	int DepthLayer = 0;
+
+
+	UPROPERTY()
+	bool bDrawSecondaryLines = true;
+
+	UPROPERTY()
+	float SecondaryLineThickness = 0.5f;
+
+	UPROPERTY()
+	FLinearColor SecondaryLineColor = FLinearColor(0.5f, 0.5f, 0.5f, 0.5f);
+
+
+
+};
\ No newline at end of file
diff --git a/Engine/Plugins/Experimental/MeshModelingToolset/Source/ModelingComponents/Public/ToolIndicatorSet.h b/Engine/Plugins/Experimental/MeshModelingToolset/Source/ModelingComponents/Public/ToolIndicatorSet.h
deleted file mode 100644
diff --git a/Engine/Source/Runtime/Experimental/InteractiveToolsFramework/Public/InteractiveGizmoManager.h b/Engine/Source/Runtime/Experimental/InteractiveToolsFramework/Public/InteractiveGizmoManager.h
--- a/Engine/Source/Runtime/Experimental/InteractiveToolsFramework/Public/InteractiveGizmoManager.h
+++ b/Engine/Source/Runtime/Experimental/InteractiveToolsFramework/Public/InteractiveGizmoManager.h
@@ -71,7 +71,7 @@ public:
 
 
 	/**
-	 * Try to activate a new Gizmo instance on the given Side
+	 * Try to activate a new Gizmo instance
 	 * @param BuilderIdentifier string used to identify Builder that should be called
 	 * @param InstanceIdentifier optional client-defined string that can be used to locate this instance (must be unique across all Gizmos)
 	 * @param Owner void pointer to whatever "owns" this Gizmo. Allows Gizmo to later be deleted using DestroyAllGizmosByOwner()
@@ -80,6 +80,21 @@ public:
 	virtual UInteractiveGizmo* CreateGizmo(const FString& BuilderIdentifier, const FString& InstanceIdentifier = FString(), void* Owner = nullptr);
 
 
+	/**
+	 * Try to activate a new Gizmo instance 
+	 * @param BuilderIdentifier string used to identify Builder that should be called
+	 * @param InstanceIdentifier optional client-defined string that can be used to locate this instance (must be unique across all Gizmos)
+	 * @param Owner void pointer to whatever "owns" this Gizmo. Allows Gizmo to later be deleted using DestroyAllGizmosByOwner()
+	 * @return new Gizmo instance that has been created and initialized, and cast to template type
+	 */
+	template<typename GizmoType>
+	GizmoType* CreateGizmo(const FString& BuilderIdentifier, const FString& InstanceIdentifier = FString(), void* Owner = nullptr)
+	{
+		UInteractiveGizmo* NewGizmo = CreateGizmo(BuilderIdentifier, InstanceIdentifier, Owner);
+		return CastChecked<GizmoType>(NewGizmo);
+	}
+
+
 	/**
 	 * Shutdown and remove a Gizmo
 	 * @param Gizmo the Gizmo to shutdown and remove
