add attribute editor tool

#rnx

diff --git a/Engine/Plugins/Experimental/MeshModelingToolset/Source/MeshModelingToolsEditorOnly/Private/AttributeEditorTool.cpp b/Engine/Plugins/Experimental/MeshModelingToolset/Source/MeshModelingToolsEditorOnly/Private/AttributeEditorTool.cpp
new file mode 100644
--- /dev/null
+++ b/Engine/Plugins/Experimental/MeshModelingToolset/Source/MeshModelingToolsEditorOnly/Private/AttributeEditorTool.cpp
@@ -0,0 +1,209 @@
+// Copyright 1998-2019 Epic Games, Inc. All Rights Reserved.
+
+#include "AttributeEditorTool.h"
+#include "InteractiveToolManager.h"
+#include "ToolBuilderUtil.h"
+
+#include "ToolSetupUtil.h"
+#include "MathUtil.h"
+
+#include "MeshDescriptionOperations.h"
+
+
+#define LOCTEXT_NAMESPACE "UAttributeEditorTool"
+
+
+/*
+ * ToolBuilder
+ */
+
+
+bool UAttributeEditorToolBuilder::CanBuildTool(const FToolBuilderState& SceneState) const
+{
+	return ToolBuilderUtil::CountComponents(SceneState, CanMakeComponentTarget) > 0;
+}
+
+UInteractiveTool* UAttributeEditorToolBuilder::BuildTool(const FToolBuilderState& SceneState) const
+{
+	UAttributeEditorTool* NewTool = NewObject<UAttributeEditorTool>(SceneState.ToolManager);
+
+	TArray<UActorComponent*> Components = ToolBuilderUtil::FindAllComponents(SceneState, CanMakeComponentTarget);
+	check(Components.Num() > 0);
+
+	TArray<TUniquePtr<FPrimitiveComponentTarget>> ComponentTargets;
+	for (UActorComponent* ActorComponent : Components)
+	{
+		auto* MeshComponent = Cast<UPrimitiveComponent>(ActorComponent);
+		if ( MeshComponent )
+		{
+			ComponentTargets.Add(MakeComponentTarget(MeshComponent));
+		}
+	}
+
+	NewTool->SetSelection(MoveTemp(ComponentTargets));
+	NewTool->SetWorld(SceneState.World);
+	NewTool->SetAssetAPI(AssetAPI);
+
+	return NewTool;
+}
+
+
+
+
+/*
+ * Tool
+ */
+
+ // UObject interface
+#if WITH_EDITOR
+bool
+UAttributeEditorToolProperties::CanEditChange( const UProperty* InProperty) const
+{
+	static TArray<FString> UVLayers
+	{
+		TEXT("bClearUVLayer0"),
+		TEXT("bClearUVLayer1"),
+		TEXT("bClearUVLayer2"),
+		TEXT("bClearUVLayer3"),
+		TEXT("bClearUVLayer4"),
+		TEXT("bClearUVLayer5"),
+		TEXT("bClearUVLayer6"),
+		TEXT("bClearUVLayer7")
+	};
+
+	int32 Index;
+	if (UVLayers.Find(InProperty->GetFName().ToString(), Index))
+	{
+		return Index < NumUVLayers;
+	}
+	else
+	{
+		return true;
+	}
+}
+#endif // WITH_EDITOR	
+// End of UObject interface
+
+
+UAttributeEditorTool::UAttributeEditorTool()
+{
+}
+
+void UAttributeEditorTool::SetWorld(UWorld* World)
+{
+	this->TargetWorld = World;
+}
+
+void UAttributeEditorTool::Setup()
+{
+	UInteractiveTool::Setup();
+
+	int NumUVLayers = 1;
+	for (TUniquePtr<FPrimitiveComponentTarget>& ComponentTarget : ComponentTargets)
+	{
+		FMeshDescription* MeshDescription = ComponentTarget->GetMesh();
+		TVertexInstanceAttributesRef<FVector2D> InstanceUVs =
+			MeshDescription->VertexInstanceAttributes().GetAttributesRef<FVector2D>(MeshAttribute::VertexInstance::TextureCoordinate);
+		NumUVLayers = FMath::Max(NumUVLayers, InstanceUVs.GetNumIndices());
+	}
+
+	RemovalProperties = NewObject<UAttributeEditorToolProperties>(this);
+	RemovalProperties->NumUVLayers = NumUVLayers;
+	AddToolPropertySource(RemovalProperties);
+}
+
+
+
+
+void UAttributeEditorTool::Shutdown(EToolShutdownType ShutdownType)
+{
+	GetToolManager()->BeginUndoTransaction(LOCTEXT("AttributeEditorToolTransactionName", "Attribute Editor Tool Transaction"));
+
+	for (TUniquePtr<FPrimitiveComponentTarget>& ComponentTarget : ComponentTargets)
+	{
+		ComponentTarget->CommitMesh([this](FMeshDescription* MeshDescription)
+		{
+			if (RemovalProperties->bClearNormals)
+			{
+				TEdgeAttributesRef<bool> EdgeHardnesses = MeshDescription->EdgeAttributes().GetAttributesRef<bool>(MeshAttribute::Edge::IsHard);
+				if (EdgeHardnesses.IsValid())
+				{
+					for (FEdgeID ElID : MeshDescription->Edges().GetElementIDs())
+					{
+						EdgeHardnesses[ElID] = false;
+					}
+				}
+				FMeshDescriptionOperations::CreatePolygonNTB(*MeshDescription, FMathf::Epsilon);
+				FMeshDescriptionOperations::RecomputeNormalsAndTangentsIfNeeded(*MeshDescription, FMeshDescriptionOperations::ETangentOptions::UseWeightedAreaAndAngle, true);
+			}
+			bool RemoveLayers[8] =
+			{
+				RemovalProperties->bClearUVLayer0,
+				RemovalProperties->bClearUVLayer1,
+				RemovalProperties->bClearUVLayer2,
+				RemovalProperties->bClearUVLayer3,
+				RemovalProperties->bClearUVLayer4,
+				RemovalProperties->bClearUVLayer5,
+				RemovalProperties->bClearUVLayer6,
+				RemovalProperties->bClearUVLayer7
+			};
+			TVertexInstanceAttributesRef<FVector2D> InstanceUVs =
+				MeshDescription->VertexInstanceAttributes().GetAttributesRef<FVector2D>(MeshAttribute::VertexInstance::TextureCoordinate);
+			const FVertexInstanceArray& Instances = MeshDescription->VertexInstances();
+			for (int LayerIndex = 7; LayerIndex >= 0; LayerIndex--)
+			{
+				if (RemoveLayers[LayerIndex] && LayerIndex < InstanceUVs.GetNumIndices())
+				{
+					if (!FMeshDescriptionOperations::RemoveUVChannel(*MeshDescription, LayerIndex))
+					{
+						for (const FVertexInstanceID& ElID : Instances.GetElementIDs())
+						{
+							InstanceUVs.Set(ElID, LayerIndex, FVector2D::ZeroVector);
+						}
+					}
+				}
+			}
+		});
+	}
+
+	GetToolManager()->EndUndoTransaction();
+}
+
+
+
+void UAttributeEditorTool::SetAssetAPI(IToolsContextAssetAPI* AssetAPIIn)
+{
+	this->AssetAPI = AssetAPIIn;
+}
+
+
+void UAttributeEditorTool::Render(IToolsContextRenderAPI* RenderAPI)
+{
+}
+
+void UAttributeEditorTool::Tick(float DeltaTime)
+{
+}
+
+
+void UAttributeEditorTool::OnPropertyModified(UObject* PropertySet, UProperty* Property)
+{
+}
+
+
+
+bool UAttributeEditorTool::HasAccept() const
+{
+	return true;
+}
+
+bool UAttributeEditorTool::CanAccept() const
+{
+	return true;
+}
+
+
+
+
+
+#undef LOCTEXT_NAMESPACE
diff --git a/Engine/Plugins/Experimental/MeshModelingToolset/Source/MeshModelingToolsEditorOnly/Public/AttributeEditorTool.h b/Engine/Plugins/Experimental/MeshModelingToolset/Source/MeshModelingToolsEditorOnly/Public/AttributeEditorTool.h
new file mode 100644
--- /dev/null
+++ b/Engine/Plugins/Experimental/MeshModelingToolset/Source/MeshModelingToolsEditorOnly/Public/AttributeEditorTool.h
@@ -0,0 +1,121 @@
+// Copyright 1998-2019 Epic Games, Inc. All Rights Reserved.
+
+#pragma once
+
+#include "CoreMinimal.h"
+#include "UObject/NoExportTypes.h"
+#include "MultiSelectionTool.h"
+#include "InteractiveToolBuilder.h"
+#include "AttributeEditorTool.generated.h"
+
+
+/**
+ *
+ */
+UCLASS()
+class MESHMODELINGTOOLSEDITORONLY_API UAttributeEditorToolBuilder : public UInteractiveToolBuilder
+{
+	GENERATED_BODY()
+
+public:
+	IToolsContextAssetAPI* AssetAPI = nullptr;
+
+	virtual bool CanBuildTool(const FToolBuilderState& SceneState) const override;
+	virtual UInteractiveTool* BuildTool(const FToolBuilderState& SceneState) const override;
+};
+
+
+
+
+
+/**
+ * Standard properties
+ */
+UCLASS()
+class MESHMODELINGTOOLSEDITORONLY_API UAttributeEditorToolProperties : public UInteractiveToolPropertySet
+{
+	GENERATED_BODY()
+
+public:
+	UAttributeEditorToolProperties()
+	{}
+	
+	int NumUVLayers = 8;
+
+	// UObject interface
+#if WITH_EDITOR
+	virtual bool CanEditChange( const UProperty* InProperty ) const override;
+#endif // WITH_EDITOR	
+	// End of UObject interface
+
+	/** Reset all normals to per-vertex smooth normals, removing all hard edges */
+	UPROPERTY(EditAnywhere, Category = Options)
+	bool bClearNormals;
+
+	// below -- just manually made 8 bools for the 8 possible uv layers in the dumbest possible way ~~
+
+	/** Clear data from UV layer */
+	UPROPERTY(EditAnywhere, Category = Options)
+	bool bClearUVLayer0;
+	/** Clear data from UV layer */
+	UPROPERTY(EditAnywhere, Category = Options)
+	bool bClearUVLayer1;
+	/** Clear data from UV layer */
+	UPROPERTY(EditAnywhere, Category = Options)
+	bool bClearUVLayer2;
+	/** Clear data from UV layer */
+	UPROPERTY(EditAnywhere, Category = Options)
+	bool bClearUVLayer3;
+	/** Clear data from UV layer */
+	UPROPERTY(EditAnywhere, Category = Options)
+	bool bClearUVLayer4;
+	/** Clear data from UV layer */
+	UPROPERTY(EditAnywhere, Category = Options)
+	bool bClearUVLayer5;
+	/** Clear data from UV layer */
+	UPROPERTY(EditAnywhere, Category = Options)
+	bool bClearUVLayer6;
+	/** Clear data from UV layer */
+	UPROPERTY(EditAnywhere, Category = Options)
+	bool bClearUVLayer7;
+
+};
+
+/**
+ * Simple Mesh Normal Updating Tool
+ */
+UCLASS()
+class MESHMODELINGTOOLSEDITORONLY_API UAttributeEditorTool : public UMultiSelectionTool
+{
+	GENERATED_BODY()
+
+public:
+
+	UAttributeEditorTool();
+
+	virtual void Setup() override;
+	virtual void Shutdown(EToolShutdownType ShutdownType) override;
+
+	virtual void SetWorld(UWorld* World);
+	virtual void SetAssetAPI(IToolsContextAssetAPI* AssetAPI);
+
+	virtual void Tick(float DeltaTime) override;
+	virtual void Render(IToolsContextRenderAPI* RenderAPI) override;
+
+	virtual bool HasCancel() const override { return true; }
+	virtual bool HasAccept() const override;
+	virtual bool CanAccept() const override;
+
+	virtual void OnPropertyModified(UObject* PropertySet, UProperty* Property) override;
+
+protected:
+
+	UPROPERTY()
+	UAttributeEditorToolProperties* RemovalProperties;
+
+protected:
+	UWorld* TargetWorld;
+	IToolsContextAssetAPI* AssetAPI;
+
+	FViewCameraState CameraState;
+};
diff --git a/Engine/Plugins/Experimental/ModelingToolsEditorMode/Source/ModelingToolsEditorMode/Private/ModelingToolsEditorMode.cpp b/Engine/Plugins/Experimental/ModelingToolsEditorMode/Source/ModelingToolsEditorMode/Private/ModelingToolsEditorMode.cpp
--- a/Engine/Plugins/Experimental/ModelingToolsEditorMode/Source/ModelingToolsEditorMode/Private/ModelingToolsEditorMode.cpp
+++ b/Engine/Plugins/Experimental/ModelingToolsEditorMode/Source/ModelingToolsEditorMode/Private/ModelingToolsEditorMode.cpp
@@ -29,6 +29,7 @@
 #include "DisplaceMeshTool.h"
 #include "MeshSpaceDeformerTool.h"
 #include "EditNormalsTool.h"
+#include "AttributeEditorTool.h"
 #include "TransformMeshesTool.h"
 #include "MeshSelectionTool.h"
 #include "UVProjectionTool.h"
@@ -388,6 +389,7 @@ void FModelingToolsEditorMode::Enter()
 	RegisterToolFunc(ToolManagerCommands.BeginMeshInspectorTool, TEXT("MeshInspectorTool"), NewObject<UMeshInspectorToolBuilder>());
 	RegisterToolFunc(ToolManagerCommands.BeginWeldEdgesTool, TEXT("WeldMeshEdgesTool"), NewObject<UWeldMeshEdgesToolBuilder>());
 	RegisterToolFunc(ToolManagerCommands.BeginPolyGroupsTool, TEXT("ConvertToPolygonsTool"), NewObject<UConvertToPolygonsToolBuilder>());
+	RegisterToolFunc(ToolManagerCommands.BeginAttributeEditorTool, TEXT("AttributeEditorTool"), NewObject<UAttributeEditorToolBuilder>());
 
 
 
diff --git a/Engine/Plugins/Experimental/ModelingToolsEditorMode/Source/ModelingToolsEditorMode/Private/ModelingToolsEditorModeStyle.cpp b/Engine/Plugins/Experimental/ModelingToolsEditorMode/Source/ModelingToolsEditorMode/Private/ModelingToolsEditorModeStyle.cpp
--- a/Engine/Plugins/Experimental/ModelingToolsEditorMode/Source/ModelingToolsEditorMode/Private/ModelingToolsEditorModeStyle.cpp
+++ b/Engine/Plugins/Experimental/ModelingToolsEditorMode/Source/ModelingToolsEditorMode/Private/ModelingToolsEditorModeStyle.cpp
@@ -126,6 +126,9 @@ void FModelingToolsEditorModeStyle::Initialize()
 		StyleSet->Set("ModelingToolsManagerCommands.BeginWeldEdgesTool.Small", new IMAGE_PLUGIN_BRUSH("Icons/icon_Tool_WeldEdges_40x", Icon28x28));
 		StyleSet->Set("ModelingToolsManagerCommands.BeginPolyGroupsTool", new IMAGE_PLUGIN_BRUSH("Icons/icon_Tool_PolyGroups_40x", Icon28x28));
 		StyleSet->Set("ModelingToolsManagerCommands.BeginPolyGroupsTool.Small", new IMAGE_PLUGIN_BRUSH("Icons/icon_Tool_PolyGroups_40x", Icon28x28));
+		StyleSet->Set("ModelingToolsManagerCommands.BeginAttributeEditorTool", new IMAGE_PLUGIN_BRUSH("Icons/icon_Tool_AttributeEditor_40x", Icon28x28));
+		StyleSet->Set("ModelingToolsManagerCommands.BeginAttributeEditorTool.Small", new IMAGE_PLUGIN_BRUSH("Icons/icon_Tool_AttributeEditor_40x", Icon28x28));
+
 
 
 		//StyleSet->Set("ModelingToolsManagerCommands.CancelActiveTool", new IMAGE_PLUGIN_BRUSH("Icons/icon_ActiveTool_Cancel_40x", Icon40x40));
diff --git a/Engine/Plugins/Experimental/ModelingToolsEditorMode/Source/ModelingToolsEditorMode/Private/ModelingToolsEditorModeToolkit.cpp b/Engine/Plugins/Experimental/ModelingToolsEditorMode/Source/ModelingToolsEditorMode/Private/ModelingToolsEditorModeToolkit.cpp
--- a/Engine/Plugins/Experimental/ModelingToolsEditorMode/Source/ModelingToolsEditorMode/Private/ModelingToolsEditorModeToolkit.cpp
+++ b/Engine/Plugins/Experimental/ModelingToolsEditorMode/Source/ModelingToolsEditorMode/Private/ModelingToolsEditorModeToolkit.cpp
@@ -299,6 +299,7 @@ void FModelingToolsEditorModeToolkit::BuildToolPalette_Standard(FName PaletteInd
 		ToolbarBuilder.AddToolBarButton(Commands.BeginEditNormalsTool);
 		ToolbarBuilder.AddToolBarButton(Commands.BeginWeldEdgesTool);
 		ToolbarBuilder.AddToolBarButton(Commands.BeginMeshInspectorTool);
+		ToolbarBuilder.AddToolBarButton(Commands.BeginAttributeEditorTool);
 	}
 }
 
@@ -357,6 +358,7 @@ void FModelingToolsEditorModeToolkit::BuildToolPalette_Experimental(FName Palett
 		ToolbarBuilder.AddToolBarButton(Commands.BeginRemeshMeshTool);
 		ToolbarBuilder.AddToolBarButton(Commands.BeginWeldEdgesTool);
 		ToolbarBuilder.AddToolBarButton(Commands.BeginMeshInspectorTool);
+		ToolbarBuilder.AddToolBarButton(Commands.BeginAttributeEditorTool);
 	}
 }
 
diff --git a/Engine/Plugins/Experimental/ModelingToolsEditorMode/Source/ModelingToolsEditorMode/Private/ModelingToolsManagerActions.cpp b/Engine/Plugins/Experimental/ModelingToolsEditorMode/Source/ModelingToolsEditorMode/Private/ModelingToolsManagerActions.cpp
--- a/Engine/Plugins/Experimental/ModelingToolsEditorMode/Source/ModelingToolsEditorMode/Private/ModelingToolsManagerActions.cpp
+++ b/Engine/Plugins/Experimental/ModelingToolsEditorMode/Source/ModelingToolsEditorMode/Private/ModelingToolsManagerActions.cpp
@@ -47,6 +47,7 @@ void FModelingToolsManagerCommands::RegisterCommands()
 	UI_COMMAND(BeginMeshInspectorTool, "Inspector", "Start the Mesh Inspector Tool", EUserInterfaceActionType::Button, FInputChord());
 	UI_COMMAND(BeginWeldEdgesTool, "Weld Edges", "Start the Weld Edges Tool", EUserInterfaceActionType::Button, FInputChord());
 	UI_COMMAND(BeginPolyGroupsTool, "PolyGroups", "Start the Poly Groups Tool", EUserInterfaceActionType::Button, FInputChord());
+	UI_COMMAND(BeginAttributeEditorTool, "AttributeEditor", "Start the Attribute Editor Tool", EUserInterfaceActionType::Button, FInputChord());
 
 	UI_COMMAND(AcceptActiveTool, "Accept", "Accept the active tool", EUserInterfaceActionType::Button, FInputChord());
 	UI_COMMAND(CancelActiveTool, "Cancel", "Cancel the active tool", EUserInterfaceActionType::Button, FInputChord());
diff --git a/Engine/Plugins/Experimental/ModelingToolsEditorMode/Source/ModelingToolsEditorMode/Public/ModelingToolsManagerActions.h b/Engine/Plugins/Experimental/ModelingToolsEditorMode/Source/ModelingToolsEditorMode/Public/ModelingToolsManagerActions.h
--- a/Engine/Plugins/Experimental/ModelingToolsEditorMode/Source/ModelingToolsEditorMode/Public/ModelingToolsManagerActions.h
+++ b/Engine/Plugins/Experimental/ModelingToolsEditorMode/Source/ModelingToolsEditorMode/Public/ModelingToolsManagerActions.h
@@ -41,6 +41,7 @@ public:
 	TSharedPtr<FUICommandInfo> BeginParameterizeMeshTool;
 	TSharedPtr<FUICommandInfo> BeginWeldEdgesTool;
 	TSharedPtr<FUICommandInfo> BeginPolyGroupsTool;
+	TSharedPtr<FUICommandInfo> BeginAttributeEditorTool;
 
 	TSharedPtr<FUICommandInfo> AcceptActiveTool;
 	TSharedPtr<FUICommandInfo> CancelActiveTool;
