don't do any work in attrib editor unless user clicks accept

#rnx

diff --git a/Engine/Plugins/Experimental/MeshModelingToolset/Source/MeshModelingToolsEditorOnly/Private/AttributeEditorTool.cpp b/Engine/Plugins/Experimental/MeshModelingToolset/Source/MeshModelingToolsEditorOnly/Private/AttributeEditorTool.cpp
--- a/Engine/Plugins/Experimental/MeshModelingToolset/Source/MeshModelingToolsEditorOnly/Private/AttributeEditorTool.cpp
+++ b/Engine/Plugins/Experimental/MeshModelingToolset/Source/MeshModelingToolsEditorOnly/Private/AttributeEditorTool.cpp
@@ -116,6 +116,14 @@ void UAttributeEditorTool::Setup()
 
 
 void UAttributeEditorTool::Shutdown(EToolShutdownType ShutdownType)
+{
+	if (ShutdownType == EToolShutdownType::Accept)
+	{
+		GenerateAssets();
+	}
+}
+
+void UAttributeEditorTool::GenerateAssets()
 {
 	GetToolManager()->BeginUndoTransaction(LOCTEXT("AttributeEditorToolTransactionName", "Attribute Editor Tool Transaction"));
 
diff --git a/Engine/Plugins/Experimental/MeshModelingToolset/Source/MeshModelingToolsEditorOnly/Public/AttributeEditorTool.h b/Engine/Plugins/Experimental/MeshModelingToolset/Source/MeshModelingToolsEditorOnly/Public/AttributeEditorTool.h
--- a/Engine/Plugins/Experimental/MeshModelingToolset/Source/MeshModelingToolsEditorOnly/Public/AttributeEditorTool.h
+++ b/Engine/Plugins/Experimental/MeshModelingToolset/Source/MeshModelingToolsEditorOnly/Public/AttributeEditorTool.h
@@ -118,4 +118,6 @@ protected:
 	IToolsContextAssetAPI* AssetAPI;
 
 	FViewCameraState CameraState;
+
+	void GenerateAssets();
 };
