ModelingTools: rename string constant to avoid name collision
#rnx

diff --git a/Engine/Plugins/Experimental/MeshModelingToolset/Source/MeshModelingTools/Private/BaseBrushTool.cpp b/Engine/Plugins/Experimental/MeshModelingToolset/Source/MeshModelingTools/Private/BaseBrushTool.cpp
--- a/Engine/Plugins/Experimental/MeshModelingToolset/Source/MeshModelingTools/Private/BaseBrushTool.cpp
+++ b/Engine/Plugins/Experimental/MeshModelingToolset/Source/MeshModelingTools/Private/BaseBrushTool.cpp
@@ -180,13 +180,13 @@ void UBaseBrushTool::Tick(float DeltaTime)
 
 
 
-const FString BrushIndicatorGizmoType = TEXT("BrushIndicatorGizmoType");
+const FString BaseBrushIndicatorGizmoType = TEXT("BrushIndicatorGizmoType");
 
 void UBaseBrushTool::SetupBrushStampIndicator()
 {
 	// register and spawn brush indicator gizmo
-	GetToolManager()->GetPairedGizmoManager()->RegisterGizmoType(BrushIndicatorGizmoType, NewObject<UBrushStampIndicatorBuilder>());
-	BrushStampIndicator = GetToolManager()->GetPairedGizmoManager()->CreateGizmo<UBrushStampIndicator>(BrushIndicatorGizmoType, FString(), this);
+	GetToolManager()->GetPairedGizmoManager()->RegisterGizmoType(BaseBrushIndicatorGizmoType, NewObject<UBrushStampIndicatorBuilder>());
+	BrushStampIndicator = GetToolManager()->GetPairedGizmoManager()->CreateGizmo<UBrushStampIndicator>(BaseBrushIndicatorGizmoType, FString(), this);
 }
 
 void UBaseBrushTool::UpdateBrushStampIndicator()
@@ -198,7 +198,7 @@ void UBaseBrushTool::ShutdownBrushStampIndicator()
 {
 	GetToolManager()->GetPairedGizmoManager()->DestroyAllGizmosByOwner(this);
 	BrushStampIndicator = nullptr;
-	GetToolManager()->GetPairedGizmoManager()->DeregisterGizmoType(BrushIndicatorGizmoType);
+	GetToolManager()->GetPairedGizmoManager()->DeregisterGizmoType(BaseBrushIndicatorGizmoType);
 }
 
 
