fix undo for plane cut tool ctrl click by using the gizmo's settransform instead of directly setting the proxy transform

#rnx

diff --git a/Engine/Plugins/Experimental/MeshModelingToolset/Source/MeshModelingTools/Private/PlaneCutTool.cpp b/Engine/Plugins/Experimental/MeshModelingToolset/Source/MeshModelingTools/Private/PlaneCutTool.cpp
--- a/Engine/Plugins/Experimental/MeshModelingToolset/Source/MeshModelingTools/Private/PlaneCutTool.cpp
+++ b/Engine/Plugins/Experimental/MeshModelingToolset/Source/MeshModelingTools/Private/PlaneCutTool.cpp
@@ -295,8 +295,8 @@ void UPlaneCutTool::SetCutPlaneFromWorldPos(const FVector& Position, const FVect
 	FFrame3f CutPlane(Position, Normal);
 	CutPlaneOrientation = CutPlane.Rotation;
 
-	PlaneTransformProxy->SetTransform(CutPlane.ToFTransform());
 	PlaneTransformGizmo->SetActiveTarget(PlaneTransformProxy);
+	PlaneTransformGizmo->SetNewGizmoTransform(CutPlane.ToFTransform());
 }
 
 
