ModelingTools: add new default material for brush indicator (transparent). Add function to ToolSetupUtil to find it. Expose Root Component of PreviewMesh.
#rnx

diff --git a/Engine/Plugins/Experimental/MeshModelingToolset/Source/ModelingComponents/Private/ToolSetupUtil.cpp b/Engine/Plugins/Experimental/MeshModelingToolset/Source/ModelingComponents/Private/ToolSetupUtil.cpp
--- a/Engine/Plugins/Experimental/MeshModelingToolset/Source/ModelingComponents/Private/ToolSetupUtil.cpp
+++ b/Engine/Plugins/Experimental/MeshModelingToolset/Source/ModelingComponents/Private/ToolSetupUtil.cpp
@@ -27,4 +27,16 @@ UMaterialInterface* ToolSetupUtil::GetDefaultWorkingMaterial(UInteractiveToolMan
 		return ToolManager->GetContextQueriesAPI()->GetStandardMaterial(EStandardToolContextMaterials::VertexColorMaterial);
 	}
 	return Material;
-}
\ No newline at end of file
+}
+
+
+UMaterialInterface* ToolSetupUtil::GetDefaultBrushVolumeMaterial(UInteractiveToolManager* ToolManager)
+{
+	UMaterialInterface* Material = LoadObject<UMaterial>(nullptr, TEXT("/MeshModelingToolset/Materials/BrushIndicatorMaterial"));
+	if (Material == nullptr && ToolManager != nullptr)
+	{
+		return ToolManager->GetContextQueriesAPI()->GetStandardMaterial(EStandardToolContextMaterials::VertexColorMaterial);
+	}
+	return Material;
+}
+
diff --git a/Engine/Plugins/Experimental/MeshModelingToolset/Source/ModelingComponents/Public/PreviewMesh.h b/Engine/Plugins/Experimental/MeshModelingToolset/Source/ModelingComponents/Public/PreviewMesh.h
--- a/Engine/Plugins/Experimental/MeshModelingToolset/Source/ModelingComponents/Public/PreviewMesh.h
+++ b/Engine/Plugins/Experimental/MeshModelingToolset/Source/ModelingComponents/Public/PreviewMesh.h
@@ -75,6 +75,12 @@ public:
 	void Disconnect();
 
 
+	/**
+	 * @return internal Root Component of internal Actor
+	 */
+	UPrimitiveComponent* GetRootComponent() { return DynamicMeshComponent; }
+
+
 	//
 	// visualization parameters
 	// 
@@ -234,9 +240,6 @@ public:
 
 
 
-
-
-
 public:
 	/** If true, we build a spatial data structure internally for the preview mesh, which allows for hit-testing */
 	UPROPERTY()
diff --git a/Engine/Plugins/Experimental/MeshModelingToolset/Source/ModelingComponents/Public/ToolSetupUtil.h b/Engine/Plugins/Experimental/MeshModelingToolset/Source/ModelingComponents/Public/ToolSetupUtil.h
--- a/Engine/Plugins/Experimental/MeshModelingToolset/Source/ModelingComponents/Public/ToolSetupUtil.h
+++ b/Engine/Plugins/Experimental/MeshModelingToolset/Source/ModelingComponents/Public/ToolSetupUtil.h
@@ -23,4 +23,9 @@ namespace ToolSetupUtil
 	 * @return default material to use for "Working"/In-Progress animations
 	 */
 	MODELINGCOMPONENTS_API UMaterialInterface* GetDefaultWorkingMaterial(UInteractiveToolManager* ToolManager);
+
+	/**
+	 * @return default material to use for brush volume indicators
+	 */
+	MODELINGCOMPONENTS_API UMaterialInterface* GetDefaultBrushVolumeMaterial(UInteractiveToolManager* ToolManager);
 }
\ No newline at end of file
