ModelingTools: add ability to attach a component to BrushStampIndicator. Use this to attach transparent sphere as brush indicator in DynamicMeshSculptTool.
#rnx

diff --git a/Engine/Plugins/Experimental/MeshModelingToolset/Source/MeshModelingTools/Private/DynamicMeshSculptTool.cpp b/Engine/Plugins/Experimental/MeshModelingToolset/Source/MeshModelingTools/Private/DynamicMeshSculptTool.cpp
--- a/Engine/Plugins/Experimental/MeshModelingToolset/Source/MeshModelingTools/Private/DynamicMeshSculptTool.cpp
+++ b/Engine/Plugins/Experimental/MeshModelingToolset/Source/MeshModelingTools/Private/DynamicMeshSculptTool.cpp
@@ -13,6 +13,7 @@
 #include "MeshNormals.h"
 #include "MeshIndexUtil.h"
 #include "Drawing/MeshDebugDrawing.h"
+#include "PreviewMesh.h"
 
 #include "Changes/MeshVertexChange.h"
 #include "Changes/MeshChange.h"
@@ -31,6 +32,7 @@ UMeshSurfacePointTool* UDynamicMeshSculptToolBuilder::CreateNewTool(const FToolB
 {
 	UDynamicMeshSculptTool* SculptTool = NewObject<UDynamicMeshSculptTool>(SceneState.ToolManager);
 	SculptTool->SetEnableRemeshing(this->bEnableRemeshing);
+	SculptTool->SetWorld(SceneState.World);
 	return SculptTool;
 }
 
@@ -89,9 +91,14 @@ UDynamicMeshSculptTool::UDynamicMeshSculptTool()
 }
 
 
-const FString BrushIndicatorGizmoType = TEXT("BrushIndicatorGizmoType");
+void UDynamicMeshSculptTool::SetWorld(UWorld* World)
+{
+	this->TargetWorld = World;
+}
 
 
+const FString BrushIndicatorGizmoType = TEXT("BrushIndicatorGizmoType");
+
 void UDynamicMeshSculptTool::Setup()
 {
 	UMeshSurfacePointTool::Setup();
@@ -153,6 +160,9 @@ void UDynamicMeshSculptTool::Setup()
 	// register and spawn brush indicator gizmo
 	GetToolManager()->GetPairedGizmoManager()->RegisterGizmoType(BrushIndicatorGizmoType, NewObject<UBrushStampIndicatorBuilder>());
 	BrushIndicator = GetToolManager()->GetPairedGizmoManager()->CreateGizmo<UBrushStampIndicator>(BrushIndicatorGizmoType, FString(), this);
+	BrushIndicatorMesh = UBrushStampIndicator::MakeDefaultSphereMesh(this, TargetWorld);
+	BrushIndicator->AttachedComponent = BrushIndicatorMesh->GetRootComponent();
+	BrushIndicator->bDrawIndicatorLines = false;
 
 	// initialize our properties
 	AddToolPropertySource(SculptProperties);
@@ -201,6 +211,9 @@ void UDynamicMeshSculptTool::Setup()
 
 void UDynamicMeshSculptTool::Shutdown(EToolShutdownType ShutdownType)
 {
+	BrushIndicatorMesh->Disconnect();
+	BrushIndicatorMesh = nullptr;
+
 	GetToolManager()->GetPairedGizmoManager()->DestroyAllGizmosByOwner(this);
 	BrushIndicator = nullptr;
 	GetToolManager()->GetPairedGizmoManager()->DeregisterGizmoType(BrushIndicatorGizmoType);
@@ -462,10 +475,7 @@ double UDynamicMeshSculptTool::CalculateBrushFalloff(double Distance)
 
 void UDynamicMeshSculptTool::ApplySmoothBrush(const FRay& WorldRay)
 {
-	// both these are useful...
-
-	//bool bHit = UpdateBrushPositionOnActivePlane(WorldRay);
-	bool bHit = UpdateBrushPositionOnTargetMesh(WorldRay);
+	bool bHit = UpdateBrushPositionOnSculptMesh(WorldRay);
 	if (bHit == false)
 	{
 		return;
diff --git a/Engine/Plugins/Experimental/MeshModelingToolset/Source/MeshModelingTools/Public/DynamicMeshSculptTool.h b/Engine/Plugins/Experimental/MeshModelingToolset/Source/MeshModelingTools/Public/DynamicMeshSculptTool.h
--- a/Engine/Plugins/Experimental/MeshModelingToolset/Source/MeshModelingTools/Public/DynamicMeshSculptTool.h
+++ b/Engine/Plugins/Experimental/MeshModelingToolset/Source/MeshModelingTools/Public/DynamicMeshSculptTool.h
@@ -38,7 +38,7 @@ DECLARE_CYCLE_STAT(TEXT("SculptTool_Remesh_6Finish"), STAT_SculptTool_Remesh_Fin
 
 class FMeshVertexChangeBuilder;
 class FDynamicMeshChangeTracker;
-
+class UPreviewMesh;
 
 /** Mesh Sculpting Brush Types */
 UENUM()
@@ -194,6 +194,7 @@ class MESHMODELINGTOOLS_API UDynamicMeshSculptTool : public UMeshSurfacePointToo
 public:
 	UDynamicMeshSculptTool();
 
+	virtual void SetWorld(UWorld* World);
 	virtual void RegisterActions(FInteractiveToolActionSet& ActionSet) override;
 
 	virtual void Setup() override;
@@ -254,9 +255,14 @@ public:
 
 protected:
 
+	UWorld* TargetWorld;		// required to spawn UPreviewMesh/etc
+
 	UPROPERTY()
 	UBrushStampIndicator* BrushIndicator;
 
+	UPROPERTY()
+	UPreviewMesh* BrushIndicatorMesh;
+
 	UPROPERTY()
 	UOctreeDynamicMeshComponent* DynamicMeshComponent;
 
diff --git a/Engine/Plugins/Experimental/MeshModelingToolset/Source/ModelingComponents/Private/Gizmos/BrushStampIndicator.cpp b/Engine/Plugins/Experimental/MeshModelingToolset/Source/ModelingComponents/Private/Gizmos/BrushStampIndicator.cpp
--- a/Engine/Plugins/Experimental/MeshModelingToolset/Source/ModelingComponents/Private/Gizmos/BrushStampIndicator.cpp
+++ b/Engine/Plugins/Experimental/MeshModelingToolset/Source/ModelingComponents/Private/Gizmos/BrushStampIndicator.cpp
@@ -3,7 +3,11 @@
 #include "Gizmos/BrushStampIndicator.h"
 #include "InteractiveGizmoManager.h"
 #include "Drawing/ToolDataVisualizer.h"
+#include "Components/PrimitiveComponent.h"
 
+#include "PreviewMesh.h"
+#include "Generators/SphereGenerator.h"
+#include "ToolSetupUtil.h"
 
 
 UInteractiveGizmo* UBrushStampIndicatorBuilder::BuildGizmo(const FToolBuilderState& SceneState) const
@@ -24,21 +28,21 @@ void UBrushStampIndicator::Shutdown()
 
 void UBrushStampIndicator::Render(IToolsContextRenderAPI* RenderAPI)
 {
-	FToolDataVisualizer Draw;
-	Draw.BeginFrame(RenderAPI);
+	if (bDrawIndicatorLines)
+	{
+		FToolDataVisualizer Draw;
+		Draw.BeginFrame(RenderAPI);
 
-	FVector3f Perp1, Perp2;
-	VectorUtil::MakePerpVectors((FVector3f)BrushNormal, Perp1, Perp2);
+		Draw.DrawCircle(BrushPosition, BrushNormal, BrushRadius, SampleStepCount, LineColor, LineThickness, bDepthTested);
 
-	Draw.DrawCircle(BrushPosition, BrushNormal, BrushRadius, SampleStepCount, LineColor, LineThickness, bDepthTested);
+		if (bDrawSecondaryLines)
+		{
+			Draw.DrawCircle(BrushPosition, BrushNormal, BrushRadius*0.5f, SampleStepCount, SecondaryLineColor, SecondaryLineThickness, bDepthTested);
+			Draw.DrawLine(BrushPosition, BrushPosition + BrushRadius * BrushNormal, SecondaryLineColor, SecondaryLineThickness, bDepthTested);
+		}
 
-	if (bDrawSecondaryLines)
-	{
-		Draw.DrawCircle(BrushPosition, BrushNormal, BrushRadius*0.5f, SampleStepCount, SecondaryLineColor, SecondaryLineThickness, bDepthTested);
-		Draw.DrawLine(BrushPosition, BrushPosition + BrushRadius*BrushNormal, SecondaryLineColor, SecondaryLineThickness, bDepthTested);
+		Draw.EndFrame();
 	}
-
-	Draw.EndFrame();
 }
 
 void UBrushStampIndicator::Tick(float DeltaTime)
@@ -51,4 +55,41 @@ void UBrushStampIndicator::Update(float Radius, const FVector& Position, const F
 	BrushRadius = Radius;
 	BrushPosition = Position;
 	BrushNormal = Normal;
+
+	if (AttachedComponent != nullptr)
+	{
+		FTransform Transform = AttachedComponent->GetComponentTransform();
+
+		if (ScaleInitializedComponent != AttachedComponent)
+		{
+			InitialComponentScale = Transform.GetScale3D();
+			InitialComponentScale *= 1.0f / InitialComponentScale.Z;
+			ScaleInitializedComponent = AttachedComponent;
+		}
+
+		Transform.SetTranslation(BrushPosition);
+
+		FQuat CurRotation = Transform.GetRotation();
+		FQuat ApplyRotation = FQuat::FindBetween(CurRotation.GetAxisZ(), BrushNormal);
+		Transform.SetRotation(ApplyRotation * CurRotation);
+
+		Transform.SetScale3D(Radius * InitialComponentScale);
+
+		AttachedComponent->SetWorldTransform(Transform);
+	}
+}
+
+
+
+UPreviewMesh* UBrushStampIndicator::MakeDefaultSphereMesh(UObject* Parent, UWorld* World, int Resolution)
+{
+	UPreviewMesh* SphereMesh = NewObject<UPreviewMesh>(Parent);
+	SphereMesh->CreateInWorld(World, FTransform::Identity);
+	FSphereGenerator SphereGen;
+	SphereGen.NumPhi = SphereGen.NumTheta = Resolution;
+	SphereGen.Generate();
+	FDynamicMesh3 Mesh(&SphereGen);
+	SphereMesh->UpdatePreview(&Mesh);
+	SphereMesh->SetMaterial(ToolSetupUtil::GetDefaultBrushVolumeMaterial(nullptr));
+	return SphereMesh;
 }
\ No newline at end of file
diff --git a/Engine/Plugins/Experimental/MeshModelingToolset/Source/ModelingComponents/Public/Gizmos/BrushStampIndicator.h b/Engine/Plugins/Experimental/MeshModelingToolset/Source/ModelingComponents/Public/Gizmos/BrushStampIndicator.h
--- a/Engine/Plugins/Experimental/MeshModelingToolset/Source/ModelingComponents/Public/Gizmos/BrushStampIndicator.h
+++ b/Engine/Plugins/Experimental/MeshModelingToolset/Source/ModelingComponents/Public/Gizmos/BrushStampIndicator.h
@@ -6,7 +6,8 @@
 #include "InteractiveGizmoBuilder.h"
 #include "BrushStampIndicator.generated.h"
 
-
+class UPrimitiveComponent;
+class UPreviewMesh;
 
 UCLASS()
 class MODELINGCOMPONENTS_API UBrushStampIndicatorBuilder : public UInteractiveGizmoBuilder
@@ -28,17 +29,28 @@ class MODELINGCOMPONENTS_API UBrushStampIndicator : public UInteractiveGizmo
 
 public:
 
+	// UInteractiveGizmo interface/implementation
+
 	virtual void Setup() override;
 	virtual void Shutdown() override;
-
 	virtual void Render(IToolsContextRenderAPI* RenderAPI) override;
-
 	virtual void Tick(float DeltaTime) override;
 
 
+	/**
+	 * Update the Radius, Position, and Normal of the stamp indicator
+	 */
 	virtual void Update(float Radius, const FVector& Position, const FVector& Normal);
 
 
+	/**
+	 * Generate a mesh that is intended to be set as the AttachedComponent of a UBrushStampIndicator
+	 * Material is set to a default transparent material.
+	 * @warning Calling code must manage the returned UPreviewMesh! (keep it alive, Disconnect it, etc)
+	 */
+	static UPreviewMesh* MakeDefaultSphereMesh(UObject* Parent, UWorld* World, int Resolution = 32);
+
+
 public:
 
 	UPROPERTY()
@@ -51,6 +63,11 @@ public:
 	FVector BrushNormal = FVector(0, 0, 1);;
 
 
+
+	UPROPERTY()
+	bool bDrawIndicatorLines = true;
+
+
 	UPROPERTY()
 	int SampleStepCount = 32;
 
@@ -63,11 +80,6 @@ public:
 	UPROPERTY()
 	bool bDepthTested = false;
 
-	UPROPERTY()
-	bool IsPixelThickness = true;
-
-	UPROPERTY()
-	int DepthLayer = 0;
 
 
 	UPROPERTY()
@@ -80,5 +92,13 @@ public:
 	FLinearColor SecondaryLineColor = FLinearColor(0.5f, 0.5f, 0.5f, 0.5f);
 
 
+	/**
+	 * Optional Component that will be transformed such that it tracks the Radius/Position/Normal
+	 */
+	UPROPERTY()
+	UPrimitiveComponent* AttachedComponent;
 
+protected:
+	UPrimitiveComponent* ScaleInitializedComponent = nullptr;		// we are just using this as a key, never calling any functions on it
+	FVector InitialComponentScale;
 };
\ No newline at end of file
