ModelingTools: improvements to brush feel. Disabled flipping, reduce smoothing power, precomputed uv/normal seams information to skip expensive constraints setup where possible.
#rnx

diff --git a/Engine/Plugins/Experimental/MeshModelingToolset/Source/MeshModelingTools/Private/DrawPolygonTool.cpp b/Engine/Plugins/Experimental/MeshModelingToolset/Source/MeshModelingTools/Private/DrawPolygonTool.cpp
--- a/Engine/Plugins/Experimental/MeshModelingToolset/Source/MeshModelingTools/Private/DrawPolygonTool.cpp
+++ b/Engine/Plugins/Experimental/MeshModelingToolset/Source/MeshModelingTools/Private/DrawPolygonTool.cpp
@@ -24,7 +24,6 @@
 
 #include "BaseGizmos/GizmoComponents.h"
 #include "BaseGizmos/TransformGizmo.h"
-#include "PositionPlaneGizmo.h"
 #include "Drawing/MeshDebugDrawing.h"
 
 #include "Selection/SelectClickedAction.h"
diff --git a/Engine/Plugins/Experimental/MeshModelingToolset/Source/MeshModelingTools/Private/DynamicMeshSculptTool.cpp b/Engine/Plugins/Experimental/MeshModelingToolset/Source/MeshModelingTools/Private/DynamicMeshSculptTool.cpp
--- a/Engine/Plugins/Experimental/MeshModelingToolset/Source/MeshModelingTools/Private/DynamicMeshSculptTool.cpp
+++ b/Engine/Plugins/Experimental/MeshModelingToolset/Source/MeshModelingTools/Private/DynamicMeshSculptTool.cpp
@@ -135,7 +135,7 @@ void UDynamicMeshSculptTool::Setup()
 	bTargetDirty = false;
 
 	double MaxDimension = DynamicMeshComponent->GetMesh()->GetCachedBounds().MaxDim();
-	BrushRelativeSizeRange = FInterval1d(MaxDimension*0.01, 2*MaxDimension);
+	BrushRelativeSizeRange = FInterval1d(MaxDimension*0.01, MaxDimension);
 	BrushProperties = NewObject<UBrushBaseProperties>(this, TEXT("Brush"));
 	CalculateBrushRadius();
 
@@ -180,23 +180,16 @@ void UDynamicMeshSculptTool::Setup()
 		LOCTEXT("OnStartSculptTool", "Hold Shift to Smooth, Ctrl to Invert (where applicable). Q/A keys cycle through Brush Types, Shift+Q/A for Brush History. S/D change Size (shift to small-step), W/E change Speed."),
 		EToolMessageLevel::UserNotification);
 
-	if (bEnableRemeshing && Mesh->HasAttributes())
+	if (bEnableRemeshing)
 	{
-		const FDynamicMeshAttributeSet* Attribs = Mesh->Attributes();
-		bool bHasUVSeams = false;
-		for (int k = 0; k < Attribs->NumUVLayers(); ++k)
-		{
-			bHasUVSeams = bHasUVSeams || Attribs->GetUVLayer(k)->HasInteriorSeamEdges();
-		}
-		bool bHasNormalSeams = Attribs->PrimaryNormals()->HasInteriorSeamEdges();
-
-		if (bHasUVSeams)
+		PrecomputeRemeshInfo();
+		if (bHaveUVSeams)
 		{
 			GetToolManager()->DisplayMessage(
 				LOCTEXT("UVSeamWarning", "This mesh has UV seams which may limit remeshing. Consider clearing the UV layers using the Remesh Tool."),
 				EToolMessageLevel::UserWarning);
 		} 
-		else if (bHasNormalSeams)
+		else if (bHaveNormalSeams)
 		{
 			GetToolManager()->DisplayMessage(
 				LOCTEXT("NormalSeamWarning", "This mesh has Hard Normal seams which may limit remeshing. Consider clearing Hard Normals using the Remesh Tool."),
@@ -1080,6 +1073,42 @@ void UDynamicMeshSculptTool::Tick(float DeltaTime)
 
 
 
+
+
+
+void UDynamicMeshSculptTool::PrecomputeRemeshInfo()
+{
+	FDynamicMesh3* Mesh = DynamicMeshComponent->GetMesh();
+
+	// check if we have any open boundary edges
+	bHaveMeshBoundaries = false;
+	for (int eid : Mesh->EdgeIndicesItr())
+	{
+		if (Mesh->IsBoundaryEdge(eid))
+		{
+			bHaveMeshBoundaries = true;
+			break;
+		}
+	}
+
+	// check if we have any UV seams
+	bHaveUVSeams = false;
+	bHaveNormalSeams = false;
+	if (Mesh->HasAttributes())
+	{
+		FDynamicMeshAttributeSet* Attribs = Mesh->Attributes();
+		for (int k = 0; k < Attribs->NumUVLayers(); ++k)
+		{
+			bHaveUVSeams = bHaveUVSeams || Attribs->GetUVLayer(k)->HasInteriorSeamEdges();
+		}
+
+		bHaveNormalSeams = Attribs->PrimaryUV()->HasInteriorSeamEdges();
+	}
+}
+
+
+
+
 void UDynamicMeshSculptTool::RemeshROIPass()
 {
 	FDynamicMesh3* Mesh = DynamicMeshComponent->GetMesh();
@@ -1094,7 +1123,7 @@ void UDynamicMeshSculptTool::RemeshROIPass()
 		Remesher.MinEdgeLength = Remesher.MinEdgeLength / 4.0;
 	}
 
-	Remesher.SmoothSpeedT = RemeshProperties->Smoothing;
+	Remesher.SmoothSpeedT = RemeshProperties->Smoothing * 0.25;
 	Remesher.SmoothType = (SculptProperties->bPreserveUVFlow) ?
 		FRemesher::ESmoothTypes::MeanValue : FRemesher::ESmoothTypes::Uniform;
 	bool bIsUniformSmooth = (Remesher.SmoothType == FRemesher::ESmoothTypes::Uniform);
@@ -1124,15 +1153,18 @@ void UDynamicMeshSculptTool::RemeshROIPass()
 		// TODO: only constrain in ROI. This is quite difficult to do externally because we need to update based on
 		// the changing triangle set in Remesher. Perhaps FSubRegionRemesher should update the constraints itself?
 
-		FMeshConstraintsUtil::ConstrainAllSeams(constraints, *Mesh, bConstraintAllowSplits, bConstraintAllowSmoothing);
-		Remesher.SetExternalConstraints(&constraints);
+		if (bHaveUVSeams || bHaveNormalSeams)
+		{
+			FMeshConstraintsUtil::ConstrainAllSeams(constraints, *Mesh, bConstraintAllowSplits, bConstraintAllowSmoothing);
+			Remesher.SetExternalConstraints(&constraints);
+		}
 	}
 
 	for (int k = 0; k < 5; ++k)
 	{
 		if (bIsUniformSmooth == false)
 		{
-			Remesher.bEnableFlips = (k < 2);
+			Remesher.bEnableFlips = RemeshProperties->bFlips && (k < 2);
 		}
 
 		{
diff --git a/Engine/Plugins/Experimental/MeshModelingToolset/Source/MeshModelingTools/Public/BaseBrushTool.h b/Engine/Plugins/Experimental/MeshModelingToolset/Source/MeshModelingTools/Public/BaseBrushTool.h
--- a/Engine/Plugins/Experimental/MeshModelingToolset/Source/MeshModelingTools/Public/BaseBrushTool.h
+++ b/Engine/Plugins/Experimental/MeshModelingToolset/Source/MeshModelingTools/Public/BaseBrushTool.h
@@ -22,7 +22,7 @@ public:
 	UBrushBaseProperties();
 
 	/** Relative size of brush */
-	UPROPERTY(EditAnywhere, Category = BrushSize, meta = (DisplayName = "Size", UIMin = "0.0", UIMax = "1.0"))
+	UPROPERTY(EditAnywhere, Category = BrushSize, meta = (DisplayName = "Size", UIMin = "0.0", UIMax = "1.0", ClampMin = "0.0", ClampMax = "10.0"))
 	float BrushSize;
 
 	/** If true, ignore relative Brush Size and use explicit World Radius */
diff --git a/Engine/Plugins/Experimental/MeshModelingToolset/Source/MeshModelingTools/Public/DynamicMeshSculptTool.h b/Engine/Plugins/Experimental/MeshModelingToolset/Source/MeshModelingTools/Public/DynamicMeshSculptTool.h
--- a/Engine/Plugins/Experimental/MeshModelingToolset/Source/MeshModelingTools/Public/DynamicMeshSculptTool.h
+++ b/Engine/Plugins/Experimental/MeshModelingToolset/Source/MeshModelingTools/Public/DynamicMeshSculptTool.h
@@ -165,7 +165,7 @@ public:
 
 	/** Enable edge flips */
 	UPROPERTY(EditAnywhere, Category = Remeshing, AdvancedDisplay)
-	bool bFlips = true;
+	bool bFlips = false;
 
 	/** Enable edge splits */
 	UPROPERTY(EditAnywhere, Category = Remeshing, AdvancedDisplay)
@@ -343,6 +343,10 @@ protected:
 	void RecalculateNormals_PerVertex();
 	void RecalculateNormals_Overlay();
 
+	bool bHaveMeshBoundaries;
+	bool bHaveUVSeams;
+	bool bHaveNormalSeams;
+	void PrecomputeRemeshInfo();
 	void RemeshROIPass();
 
 	FMeshVertexChangeBuilder* ActiveVertexChange = nullptr;
