ModelingTools: add hotkey to toggle wireframe overlay in sculpt tool
#rnx

diff --git a/Engine/Plugins/Experimental/MeshModelingToolset/Source/MeshModelingTools/Private/DynamicMeshSculptTool.cpp b/Engine/Plugins/Experimental/MeshModelingToolset/Source/MeshModelingTools/Private/DynamicMeshSculptTool.cpp
--- a/Engine/Plugins/Experimental/MeshModelingToolset/Source/MeshModelingTools/Private/DynamicMeshSculptTool.cpp
+++ b/Engine/Plugins/Experimental/MeshModelingToolset/Source/MeshModelingTools/Private/DynamicMeshSculptTool.cpp
@@ -1565,6 +1565,15 @@ void UDynamicMeshSculptTool::RegisterActions(FInteractiveToolActionSet& ActionSe
 		EModifierKey::None, EKeys::W,
 		[this]() { DecreaseBrushSpeedAction(); });
 
+
+
+	ActionSet.RegisterAction(this, (int32)EStandardToolActions::BaseClientDefinedActionID + 100,
+		TEXT("ToggleWireframe"),
+		LOCTEXT("ToggleWireframe", "Toggle Wireframe"),
+		LOCTEXT("ToggleWireframeTooltip", "Toggle visibility of wireframe overlay"),
+		EModifierKey::Control, EKeys::W,
+		[this]() { SculptProperties->bShowWireframe = !SculptProperties->bShowWireframe; });
+
 }
 
 
