MeshModelingTools: Allow Displace Tool to work with a broad selection of textures

diff --git a/Engine/Plugins/Experimental/MeshModelingToolset/Source/MeshModelingTools/Private/DisplaceMeshTool.cpp b/Engine/Plugins/Experimental/MeshModelingToolset/Source/MeshModelingTools/Private/DisplaceMeshTool.cpp
--- a/Engine/Plugins/Experimental/MeshModelingToolset/Source/MeshModelingTools/Private/DisplaceMeshTool.cpp
+++ b/Engine/Plugins/Experimental/MeshModelingToolset/Source/MeshModelingTools/Private/DisplaceMeshTool.cpp
@@ -18,12 +18,9 @@
 
 #define LOCTEXT_NAMESPACE "UDisplaceMeshTool"
 
-
 /*
  * ToolBuilder
  */
-
-
 bool UDisplaceMeshToolBuilder::CanBuildTool(const FToolBuilderState& SceneState) const
 {
 	return ToolBuilderUtil::CountComponents(SceneState, CanMakeComponentTarget) == 1;
@@ -42,12 +39,9 @@ UInteractiveTool* UDisplaceMeshToolBuilder::BuildTool(const FToolBuilderState& S
 	return NewTool;
 }
 
-
-
 /*
  * Tool
  */
-
 UDisplaceMeshTool::UDisplaceMeshTool()
 {
 	DisplacementType = EDisplaceMeshToolDisplaceType::DisplacementMap;
@@ -87,7 +81,6 @@ void UDisplaceMeshTool::Setup()
 	bResultValid = false;
 }
 
-
 void UDisplaceMeshTool::Shutdown(EToolShutdownType ShutdownType)
 {
 	if (DynamicMeshComponent != nullptr)
@@ -113,7 +106,6 @@ void UDisplaceMeshTool::Shutdown(EToolShutdownType ShutdownType)
 	}
 }
 
-
 void UDisplaceMeshTool::Render(IToolsContextRenderAPI* RenderAPI)
 {
 	UpdateResult();
@@ -189,8 +181,6 @@ void UDisplaceMeshTool::UpdateResult()
 	bResultValid = true;
 }
 
-
-
 void UDisplaceMeshTool::ComputeDisplacement_Constant()
 {
 	double Intensity = DisplaceIntensity;
@@ -204,7 +194,6 @@ void UDisplaceMeshTool::ComputeDisplacement_Constant()
 	}
 }
 
-
 void UDisplaceMeshTool::ComputeDisplacement_RandomNoise()
 {
 	double Intensity = DisplaceIntensity;
@@ -224,10 +213,6 @@ void UDisplaceMeshTool::ComputeDisplacement_RandomNoise()
 	}
 }
 
-
-
-
-
 void UDisplaceMeshTool::ComputeDisplacement_Map()
 {
 	UpdateMap();
@@ -260,9 +245,65 @@ void UDisplaceMeshTool::ComputeDisplacement_Map()
 
 }
 
+namespace {
+	class FTextureAccess
+	{
+	public:
+		FTextureAccess(UTexture2D* DisplacementMap)
+			:DisplacementMap(DisplacementMap)
+		{
+			check(DisplacementMap);
+			OldCompressionSettings = DisplacementMap->CompressionSettings;
+			bOldSRGB = DisplacementMap->SRGB; 
+#if WITH_EDITOR
+			OldMipGenSettings = DisplacementMap->MipGenSettings;
+#endif
+
 
+			DisplacementMap->CompressionSettings = TextureCompressionSettings::TC_VectorDisplacementmap;
+			DisplacementMap->SRGB = false;
+#if WITH_EDITOR
+			DisplacementMap->MipGenSettings = TextureMipGenSettings::TMGS_NoMipmaps;
+#endif
+			DisplacementMap->UpdateResource();
+
+			FormattedImageData = static_cast<const FColor*>(DisplacementMap->PlatformData->Mips[0].BulkData.LockReadOnly());
+		}
+		FTextureAccess(const FTextureAccess&) = delete;
+		FTextureAccess(FTextureAccess&&) = delete;
+		void operator=(const FTextureAccess&) = delete;
+		void operator=(FTextureAccess&&) = delete;
+
+		~FTextureAccess()
+		{
+			DisplacementMap->PlatformData->Mips[0].BulkData.Unlock();
+
+			DisplacementMap->CompressionSettings = OldCompressionSettings;
+			DisplacementMap->SRGB = bOldSRGB;
+#if WITH_EDITOR
+			DisplacementMap->MipGenSettings = OldMipGenSettings;
+#endif
+
+			DisplacementMap->UpdateResource();
+		}
 
+		bool HasData() const
+		{
+			return FormattedImageData != nullptr;
+		}
+		const FColor* GetData() const
+		{
+			return FormattedImageData;
+		}
 
+	private:
+		UTexture2D* DisplacementMap{nullptr};
+		TextureCompressionSettings OldCompressionSettings{};
+		TextureMipGenSettings OldMipGenSettings{};
+		bool bOldSRGB{false};
+		const FColor* FormattedImageData{nullptr};
+	};
+}
 
 void UDisplaceMeshTool::UpdateMap(bool bForceUpdate)
 {
@@ -278,8 +319,8 @@ void UDisplaceMeshTool::UpdateMap(bool bForceUpdate)
 		return;
 	}
 
-	const FColor* FormatedImageData = static_cast<const FColor*>(DisplacementMap->PlatformData->Mips[0].BulkData.LockReadOnly());
-	if (FormatedImageData == nullptr)
+	FTextureAccess TextureAccess( DisplacementMap );
+	if (!TextureAccess.HasData())
 	{
 		DisplaceField = FSampledScalarField2f();
 	}
@@ -287,26 +328,22 @@ void UDisplaceMeshTool::UpdateMap(bool bForceUpdate)
 	{
 		int TextureWidth = DisplacementMap->GetSizeX();
 		int TextureHeight = DisplacementMap->GetSizeY();
-
 		DisplaceField.Resize(TextureWidth, TextureHeight, 0.0f);
 		DisplaceField.SetCellSize(1.0f / (float)TextureWidth);
 
+		const FColor* FormattedData = TextureAccess.GetData();
 		for (int y = 0; y < TextureHeight; ++y)
 		{
 			for (int x = 0; x < TextureWidth; ++x)
 			{
-				FColor PixelColor = FormatedImageData[y*TextureWidth + x];
+				FColor PixelColor = FormattedData[y*TextureWidth + x];
 				float Value = PixelColor.R / 255.0;
 				DisplaceField.GridValues[y*TextureWidth + x] = Value;
 			}
 		}
 	}
-
-	DisplacementMap->PlatformData->Mips[0].BulkData.Unlock();
 }
 
-
-
 void UDisplaceMeshTool::UpdateSubdividedMesh()
 {
 	if (Subdivisions == CachedSubdivisionsCount)
