add preview visibility toggle to plane cut

#rnx

diff --git a/Engine/Plugins/Experimental/MeshModelingToolset/Source/MeshModelingTools/Private/PlaneCutTool.cpp b/Engine/Plugins/Experimental/MeshModelingToolset/Source/MeshModelingTools/Private/PlaneCutTool.cpp
--- a/Engine/Plugins/Experimental/MeshModelingToolset/Source/MeshModelingTools/Private/PlaneCutTool.cpp
+++ b/Engine/Plugins/Experimental/MeshModelingToolset/Source/MeshModelingTools/Private/PlaneCutTool.cpp
@@ -62,6 +62,7 @@ UPlaneCutToolProperties::UPlaneCutToolProperties()
 	bKeepBothHalves = false;
 	bFillCutHole = true;
 	SpacingBetweenHalves = 1;
+	bShowPreview = true;
 }
 
 UPlaneCutAdvancedProperties::UPlaneCutAdvancedProperties()
@@ -268,6 +269,15 @@ void UPlaneCutTool::PostEditChangeProperty(FPropertyChangedEvent& PropertyChange
 
 void UPlaneCutTool::OnPropertyModified(UObject* PropertySet, UProperty* Property)
 {
+	if (Property && (Property->GetFName() == GET_MEMBER_NAME_CHECKED(UPlaneCutToolProperties, bShowPreview)))
+	{
+		ComponentTarget->SetOwnerVisibility(!BasicProperties->bShowPreview);
+		for (UMeshOpPreviewWithBackgroundCompute* Preview : Previews)
+		{
+			Preview->SetVisibility(BasicProperties->bShowPreview);
+		}
+	}
+
 	UpdateNumPreviews();
 	for (UMeshOpPreviewWithBackgroundCompute* Preview : Previews)
 	{
diff --git a/Engine/Plugins/Experimental/MeshModelingToolset/Source/MeshModelingTools/Public/PlaneCutTool.h b/Engine/Plugins/Experimental/MeshModelingToolset/Source/MeshModelingTools/Public/PlaneCutTool.h
--- a/Engine/Plugins/Experimental/MeshModelingToolset/Source/MeshModelingTools/Public/PlaneCutTool.h
+++ b/Engine/Plugins/Experimental/MeshModelingToolset/Source/MeshModelingTools/Public/PlaneCutTool.h
@@ -69,6 +69,9 @@ public:
 	/** If true, the cut surface is filled with simple planar hole fill surface(s) */
 	UPROPERTY(EditAnywhere, Category = Advanced)
 	bool bFillCutHole;
+
+	UPROPERTY(EditAnywhere, Category = Preview)
+	bool bShowPreview;
 };
 
 
diff --git a/Engine/Plugins/Experimental/MeshModelingToolset/Source/ModelingComponents/Private/MeshOpPreviewHelpers.cpp b/Engine/Plugins/Experimental/MeshModelingToolset/Source/ModelingComponents/Private/MeshOpPreviewHelpers.cpp
--- a/Engine/Plugins/Experimental/MeshModelingToolset/Source/ModelingComponents/Private/MeshOpPreviewHelpers.cpp
+++ b/Engine/Plugins/Experimental/MeshModelingToolset/Source/ModelingComponents/Private/MeshOpPreviewHelpers.cpp
@@ -72,7 +72,7 @@ void UMeshOpPreviewWithBackgroundCompute::UpdateResults()
 		TUniquePtr<FDynamicMesh3> ResultMesh = MeshOp->ExtractResult();
 		PreviewMesh->SetTransform((FTransform)MeshOp->GetResultTransform());
 		PreviewMesh->UpdatePreview(ResultMesh.Get());  // copies the mesh @todo we could just give ownership to the Preview!
-		PreviewMesh->SetVisible(true);
+		PreviewMesh->SetVisible(bVisible);
 		bResultValid = true;
 
 		OnMeshUpdated.Broadcast(this);
@@ -99,7 +99,8 @@ void UMeshOpPreviewWithBackgroundCompute::ConfigureMaterials(UMaterialInterface*
 }
 
 
-void UMeshOpPreviewWithBackgroundCompute::SetVisibility(bool bVisible)
+void UMeshOpPreviewWithBackgroundCompute::SetVisibility(bool bVisibleIn)
 {
+	bVisible = bVisibleIn;
 	PreviewMesh->SetVisible(bVisible);
 }
diff --git a/Engine/Plugins/Experimental/MeshModelingToolset/Source/ModelingComponents/Public/MeshOpPreviewHelpers.h b/Engine/Plugins/Experimental/MeshModelingToolset/Source/ModelingComponents/Public/MeshOpPreviewHelpers.h
--- a/Engine/Plugins/Experimental/MeshModelingToolset/Source/ModelingComponents/Public/MeshOpPreviewHelpers.h
+++ b/Engine/Plugins/Experimental/MeshModelingToolset/Source/ModelingComponents/Public/MeshOpPreviewHelpers.h
@@ -148,6 +148,8 @@ protected:
 	// state flag, if true then we have valid result
 	bool bResultValid = false;
 
+	bool bVisible = true;
+
 	// this object manages the background computes
 	TUniquePtr<FBackgroundDynamicMeshComputeSource> BackgroundCompute;
 
