make remesh tool compute normals if it has thrown the attribute ones out and has no vertex normals

#rnx

diff --git a/Engine/Plugins/Experimental/MeshModelingToolset/Source/MeshModelingTools/Private/RemeshMeshTool.cpp b/Engine/Plugins/Experimental/MeshModelingToolset/Source/MeshModelingTools/Private/RemeshMeshTool.cpp
--- a/Engine/Plugins/Experimental/MeshModelingToolset/Source/MeshModelingTools/Private/RemeshMeshTool.cpp
+++ b/Engine/Plugins/Experimental/MeshModelingToolset/Source/MeshModelingTools/Private/RemeshMeshTool.cpp
@@ -260,6 +260,11 @@ void URemeshMeshTool::UpdateResult()
 		Remesher.BasicRemeshPass();
 	}
 
+	if (!TargetMesh->HasAttributes() && !TargetMesh->HasVertexNormals())
+	{
+		FMeshNormals::QuickComputeVertexNormals(*TargetMesh);
+	}
+
 	DynamicMeshComponent->NotifyMeshUpdated();
 	GetToolManager()->PostInvalidation();
 	MeshStatisticsProperties->Update(*DynamicMeshComponent->GetMesh());
