ModelingTools: add various brush and action options to MeshSelection tool.
#rnx

diff --git a/Engine/Plugins/Experimental/MeshModelingToolset/Source/MeshModelingTools/Private/DynamicMeshSculptTool.cpp b/Engine/Plugins/Experimental/MeshModelingToolset/Source/MeshModelingTools/Private/DynamicMeshSculptTool.cpp
--- a/Engine/Plugins/Experimental/MeshModelingToolset/Source/MeshModelingTools/Private/DynamicMeshSculptTool.cpp
+++ b/Engine/Plugins/Experimental/MeshModelingToolset/Source/MeshModelingTools/Private/DynamicMeshSculptTool.cpp
@@ -1567,7 +1567,7 @@ void UDynamicMeshSculptTool::RegisterActions(FInteractiveToolActionSet& ActionSe
 
 
 
-	ActionSet.RegisterAction(this, (int32)EStandardToolActions::BaseClientDefinedActionID + 100,
+	ActionSet.RegisterAction(this, (int32)EStandardToolActions::ToggleWireframe,
 		TEXT("ToggleWireframe"),
 		LOCTEXT("ToggleWireframe", "Toggle Wireframe"),
 		LOCTEXT("ToggleWireframeTooltip", "Toggle visibility of wireframe overlay"),
diff --git a/Engine/Plugins/Experimental/MeshModelingToolset/Source/MeshModelingTools/Private/MeshSelectionTool.cpp b/Engine/Plugins/Experimental/MeshModelingToolset/Source/MeshModelingTools/Private/MeshSelectionTool.cpp
--- a/Engine/Plugins/Experimental/MeshModelingToolset/Source/MeshModelingTools/Private/MeshSelectionTool.cpp
+++ b/Engine/Plugins/Experimental/MeshModelingToolset/Source/MeshModelingTools/Private/MeshSelectionTool.cpp
@@ -8,7 +8,9 @@
 #include "DynamicMeshChangeTracker.h"
 #include "Changes/ToolCommandChangeSequence.h"
 #include "Changes/MeshChange.h"
+#include "MeshIndexUtil.h"
 #include "AssetGenerationUtil.h"
+#include "DynamicMeshOctree3.h"
 
 #define LOCTEXT_NAMESPACE "UMeshSelectionTool"
 
@@ -25,6 +27,39 @@ UMeshSurfacePointTool* UMeshSelectionToolBuilder::CreateNewTool(const FToolBuild
 }
 
 
+
+
+void UMeshSelectionToolActionPropertySet::PostAction(EMeshSelectionToolActions Action)
+{
+	if (ParentTool.IsValid())
+	{
+		ParentTool->RequestAction(Action);
+	}
+}
+
+
+void UMeshSelectionToolProperties::SaveProperties(UInteractiveTool* SaveFromTool)
+{
+	UMeshSelectionToolProperties* PropertyCache = GetPropertyCache<UMeshSelectionToolProperties>();
+	PropertyCache->SelectionMode = this->SelectionMode;
+	PropertyCache->AngleTolerance = this->AngleTolerance;
+	PropertyCache->bVolumetricBrush = this->bVolumetricBrush;
+	PropertyCache->bHitBackFaces = this->bHitBackFaces;
+	PropertyCache->bShowWireframe = this->bShowWireframe;
+}
+
+void UMeshSelectionToolProperties::RestoreProperties(UInteractiveTool* RestoreToTool)
+{
+	UMeshSelectionToolProperties* PropertyCache = GetPropertyCache<UMeshSelectionToolProperties>();
+	this->SelectionMode = PropertyCache->SelectionMode;
+	this->AngleTolerance = PropertyCache->AngleTolerance;
+	this->bVolumetricBrush = PropertyCache->bVolumetricBrush;
+	this->bHitBackFaces = PropertyCache->bHitBackFaces;
+	this->bShowWireframe = PropertyCache->bShowWireframe;
+}
+
+
+
 /*
  * Tool
  */
@@ -49,8 +84,17 @@ void UMeshSelectionTool::Setup()
 {
 	UDynamicMeshBrushTool::Setup();
 
-	// use ourself as tool prop source for actions
-	AddToolPropertySource(this);
+	SelectionProps = NewObject<UMeshSelectionToolProperties>(this);
+	SelectionProps->RestoreProperties(this);
+	AddToolPropertySource(SelectionProps);
+
+	SelectionActions = NewObject<UMeshSelectionEditActions>(this);
+	SelectionActions->Initialize(this);
+	AddToolPropertySource(SelectionActions);
+
+	EditActions = NewObject<UMeshSelectionMeshEditActions>(this);
+	EditActions->Initialize(this);
+	AddToolPropertySource(EditActions);
 
 	// enable wireframe on component
 	PreviewMesh->EnableWireframe(true);
@@ -73,6 +117,12 @@ void UMeshSelectionTool::Setup()
 		OnExternalSelectionChange();
 	});
 
+	// rebuild octree if mesh changes
+	PreviewMesh->GetOnMeshChanged().AddLambda([this]() { bOctreeValid = false; });
+
+	ShowWireframeWatcher.Initialize(
+		[this]() { return SelectionProps->bShowWireframe; },
+		[this](bool bNewValue) { PreviewMesh->EnableWireframe(bNewValue); }, false);
 }
 
 
@@ -80,6 +130,8 @@ void UMeshSelectionTool::Setup()
 
 void UMeshSelectionTool::OnShutdown(EToolShutdownType ShutdownType)
 {
+	SelectionProps->SaveProperties(this);
+
 	if (bHaveModifiedMesh && ShutdownType == EToolShutdownType::Accept)
 	{
 		// this block bakes the modified DynamicMeshComponent back into the StaticMeshComponent inside an undo transaction
@@ -106,6 +158,14 @@ void UMeshSelectionTool::RegisterActions(FInteractiveToolActionSet& ActionSet)
 		LOCTEXT("MeshSelectionToolDeleteTooltip", "Delete Selected Elements"),
 		EModifierKey::None, EKeys::Delete,
 		[this]() { DeleteSelectedTriangles(); });
+
+
+	ActionSet.RegisterAction(this, (int32)EStandardToolActions::ToggleWireframe,
+		TEXT("ToggleWireframe"),
+		LOCTEXT("ToggleWireframe", "Toggle Wireframe"),
+		LOCTEXT("ToggleWireframeTooltip", "Toggle visibility of wireframe overlay"),
+		EModifierKey::Control, EKeys::W,
+		[this]() { SelectionProps->bShowWireframe = !SelectionProps->bShowWireframe; });
 }
 
 
@@ -136,6 +196,28 @@ void UMeshSelectionTool::OnExternalSelectionChange()
 
 
 
+bool UMeshSelectionTool::HitTest(const FRay& Ray, FHitResult& OutHit)
+{
+	bool bHit = UDynamicMeshBrushTool::HitTest(Ray, OutHit);
+	if (bHit && SelectionProps->bHitBackFaces == false)
+	{
+		const FDynamicMesh3* SourceMesh = PreviewMesh->GetPreviewDynamicMesh();
+		FVector3d Normal, Centroid; 
+		double Area;
+		SourceMesh->GetTriInfo(OutHit.FaceIndex, Normal, Area, Centroid);
+		FViewCameraState StateOut;
+		GetToolManager()->GetContextQueriesAPI()->GetCurrentViewState(StateOut);
+		FVector3d LocalEyePosition(ComponentTarget->GetWorldTransform().InverseTransformPosition(StateOut.Position));
+
+		if (Normal.Dot((Centroid - LocalEyePosition)) > 0)
+		{
+			bHit = false;
+		}
+	}
+	return bHit;
+}
+
+
 void UMeshSelectionTool::OnBeginDrag(const FRay& WorldRay)
 {
 	UDynamicMeshBrushTool::OnBeginDrag(WorldRay);
@@ -165,6 +247,18 @@ void UMeshSelectionTool::OnUpdateDrag(const FRay& WorldRay)
 
 
 
+TUniquePtr<FDynamicMeshOctree3>& UMeshSelectionTool::GetOctree()
+{
+	if (bOctreeValid == false)
+	{
+		Octree = MakeUnique<FDynamicMeshOctree3>();
+		Octree->Initialize(PreviewMesh->GetPreviewDynamicMesh());
+		bOctreeValid = true;
+	}
+	return Octree;
+}
+
+
 
 void UMeshSelectionTool::CalculateVertexROI(const FBrushStampData& Stamp, TArray<int>& VertexROI)
 {
@@ -186,6 +280,39 @@ void UMeshSelectionTool::CalculateVertexROI(const FBrushStampData& Stamp, TArray
 }
 
 
+
+
+static void GrowToConnectedTriangles(const TArray<int>& TriangleROI, TArray<int>& GrowROI, const FDynamicMesh3* Mesh,
+	TArray<int32>& QueueBuffer, TSet<int32>& DoneBuffer,
+	TFunctionRef<bool(int32, int32)> CanGrowPredicate = [](int32, int32) { return true; }
+	)
+{
+	QueueBuffer.Reset(); QueueBuffer.Insert(TriangleROI, 0);
+	DoneBuffer.Reset(); DoneBuffer.Append(TriangleROI);
+	//TArray<int32> Queue(TriangleROI);
+	//TSet<int32> Done(TriangleROI);
+
+	while (QueueBuffer.Num() > 0)
+	{
+		int32 CurTri = QueueBuffer.Pop(false);
+		GrowROI.Add(CurTri);
+
+		FIndex3i NbrTris = Mesh->GetTriNeighbourTris(CurTri);
+		for (int j = 0; j < 3; ++j)
+		{
+			int32 tid = NbrTris[j];
+			if (tid != FDynamicMesh3::InvalidID && DoneBuffer.Contains(tid) == false && CanGrowPredicate(CurTri, tid))
+			{
+				QueueBuffer.Add(tid);
+				DoneBuffer.Add(tid);
+			}
+		}
+	}
+
+}
+
+
+
 void UMeshSelectionTool::CalculateTriangleROI(const FBrushStampData& Stamp, TArray<int>& TriangleROI)
 {
 	FTransform Transform = ComponentTarget->GetWorldTransform();
@@ -193,20 +320,38 @@ void UMeshSelectionTool::CalculateTriangleROI(const FBrushStampData& Stamp, TArr
 
 	// always select first triangle
 	const FDynamicMesh3* Mesh = PreviewMesh->GetPreviewDynamicMesh();
-	if (Mesh->IsTriangle(Stamp.HitResult.FaceIndex))
-	{
-		TriangleROI.Add(Stamp.HitResult.FaceIndex);
-	}
 
 	float RadiusSqr = CurrentBrushRadius * CurrentBrushRadius;
-	for (int TriIdx : Mesh->TriangleIndicesItr())
+	if (SelectionProps->bVolumetricBrush)
 	{
-		FVector3d Position = Mesh->GetTriCentroid(TriIdx);
-		if ((Position - StampPosLocal).SquaredLength() < RadiusSqr)
+		if (Mesh->IsTriangle(Stamp.HitResult.FaceIndex))
+		{
+			TriangleROI.Add(Stamp.HitResult.FaceIndex);
+		}
+
+		FAxisAlignedBox3d Bounds(StampPosLocal-CurrentBrushRadius*FVector3d::One(), StampPosLocal+CurrentBrushRadius*FVector3d::One());
+		TemporaryBuffer.Reset();
+		GetOctree()->RangeQuery(Bounds, TemporaryBuffer);
+
+		for (int32 TriIdx : TemporaryBuffer)
 		{
-			TriangleROI.Add(TriIdx);
+			FVector3d Position = Mesh->GetTriCentroid(TriIdx);
+			if ((Position - StampPosLocal).SquaredLength() < RadiusSqr)
+			{
+				TriangleROI.Add(TriIdx);
+			}
 		}
 	}
+	else
+	{
+		TArray<int32> StartROI;
+		StartROI.Add(Stamp.HitResult.FaceIndex);
+		GrowToConnectedTriangles(StartROI, TriangleROI, Mesh, TemporaryBuffer, TemporarySet,
+			[Mesh, RadiusSqr, StampPosLocal](int t1, int t2) { return (Mesh->GetTriCentroid(t2) - StampPosLocal).SquaredLength() < RadiusSqr; });
+
+	}
+
+
 }
 
 
@@ -236,18 +381,7 @@ void UMeshSelectionTool::ApplyStamp(const FBrushStampData& Stamp)
 	if (SelectionType == EMeshSelectionElementType::Face)
 	{
 		CalculateTriangleROI(Stamp, IndexBuf);
-		for (int TriIdx : IndexBuf)
-		{
-			if (SelectedTriangles[TriIdx] != bDesiredValue)
-			{
-				SelectedTriangles[TriIdx] = bDesiredValue;
-				UpdateList(Selection->Faces, TriIdx, bDesiredValue);
-				if (ActiveSelectionChange != nullptr)
-				{
-					ActiveSelectionChange->Add(TriIdx);
-				}
-			}
-		}
+		UpdateFaceSelection(Stamp, IndexBuf);
 	}
 	else
 	{
@@ -272,6 +406,83 @@ void UMeshSelectionTool::ApplyStamp(const FBrushStampData& Stamp)
 
 
 
+
+
+void UMeshSelectionTool::UpdateFaceSelection(const FBrushStampData& Stamp, const TArray<int>& TriangleROI)
+{
+	const FDynamicMesh3* Mesh = PreviewMesh->GetPreviewDynamicMesh();
+	const TArray<int>* UseROI = &TriangleROI;
+
+	TArray<int> LocalROI;
+	if (SelectionProps->SelectionMode == EMeshSelectionToolPrimaryMode::AllConnected)
+	{
+		GrowToConnectedTriangles(TriangleROI, LocalROI, Mesh, TemporaryBuffer, TemporarySet);
+		UseROI = &LocalROI;
+	}
+	else if (SelectionProps->SelectionMode == EMeshSelectionToolPrimaryMode::AllInGroup)
+	{
+		GrowToConnectedTriangles(TriangleROI, LocalROI, Mesh, TemporaryBuffer, TemporarySet,
+			[Mesh](int t1, int t2) { return Mesh->GetTriangleGroup(t1) == Mesh->GetTriangleGroup(t2); } );
+		UseROI = &LocalROI;
+	}
+	else if (SelectionProps->SelectionMode == EMeshSelectionToolPrimaryMode::AllWithinAngle)
+	{
+		TArray<int32> StartROI; 
+		StartROI.Add(Stamp.HitResult.FaceIndex);
+		FVector3d StartNormal = Mesh->GetTriNormal(StartROI[0]);
+		int AngleTol = SelectionProps->AngleTolerance;
+		GrowToConnectedTriangles(StartROI, LocalROI, Mesh, TemporaryBuffer, TemporarySet,
+			[Mesh, AngleTol, StartNormal](int t1, int t2) { return Mesh->GetTriNormal(t2).AngleD(StartNormal) < AngleTol; });
+		UseROI = &LocalROI;
+	}
+	else if (SelectionProps->SelectionMode == EMeshSelectionToolPrimaryMode::AngleFiltered)
+	{
+		TSet<int32> BrushROI(TriangleROI);
+		TArray<int32> StartROI;
+		StartROI.Add(Stamp.HitResult.FaceIndex);
+		FVector3d StartNormal = Mesh->GetTriNormal(StartROI[0]);
+		int AngleTol = SelectionProps->AngleTolerance;
+		GrowToConnectedTriangles(StartROI, LocalROI, Mesh, TemporaryBuffer, TemporarySet,
+			[Mesh, AngleTol, StartNormal, &BrushROI](int t1, int t2) { return BrushROI.Contains(t2) && Mesh->GetTriNormal(t2).AngleD(StartNormal) < AngleTol; });
+		UseROI = &LocalROI;
+	}
+	else if (SelectionProps->SelectionMode == EMeshSelectionToolPrimaryMode::Visible)
+	{
+		FViewCameraState StateOut;
+		GetToolManager()->GetContextQueriesAPI()->GetCurrentViewState(StateOut);
+		FVector3d LocalEyePosition(ComponentTarget->GetWorldTransform().InverseTransformPosition(StateOut.Position));
+
+		for (int tid : TriangleROI)
+		{
+			FVector3d Centroid = Mesh->GetTriCentroid(tid);
+			int HitTID = GetOctree()->FindNearestHitObject(FRay3d(LocalEyePosition, (Centroid - LocalEyePosition).Normalized()));
+			if (HitTID == tid)
+			{
+				LocalROI.Add(HitTID);
+			}
+		}
+		UseROI = &LocalROI;
+	}
+
+	bool bDesiredValue = bInRemoveStroke ? false : true;
+	for (int TriIdx : *UseROI)
+	{
+		if (SelectedTriangles[TriIdx] != bDesiredValue)
+		{
+			SelectedTriangles[TriIdx] = bDesiredValue;
+			UpdateList(Selection->Faces, TriIdx, bDesiredValue);
+			if (ActiveSelectionChange != nullptr)
+			{
+				ActiveSelectionChange->Add(TriIdx);
+			}
+		}
+	}
+
+}
+
+
+
+
 void UMeshSelectionTool::OnEndDrag(const FRay& Ray)
 {
 	UDynamicMeshBrushTool::OnEndDrag(Ray);
@@ -364,10 +575,10 @@ void UMeshSelectionTool::Render(IToolsContextRenderAPI* RenderAPI)
 	else
 	{
 		// drawn via material
-		MeshDebugDraw::DrawTriCentroids(Mesh, Selection->Faces,
-			12.0f, FColor::Green, RenderAPI->GetPrimitiveDrawInterface(), WorldTransform);
+		//MeshDebugDraw::DrawTriCentroids(Mesh, Selection->Faces,
+		//	12.0f, FColor::Green, RenderAPI->GetPrimitiveDrawInterface(), WorldTransform);
 		MeshDebugDraw::DrawTriCentroids(Mesh, PreviewBrushROI,
-			8.0f, FColor(40, 200, 40), RenderAPI->GetPrimitiveDrawInterface(), WorldTransform);
+			4.0f, FColor(40, 200, 40), RenderAPI->GetPrimitiveDrawInterface(), WorldTransform);
 	}
 }
 
@@ -376,11 +587,20 @@ void UMeshSelectionTool::Tick(float DeltaTime)
 {
 	UDynamicMeshBrushTool::Tick(DeltaTime);
 
+	ShowWireframeWatcher.CheckAndUpdate();
+
 	if (bStampPending)
 	{
 		ApplyStamp(LastStamp);
 		bStampPending = false;
 	}
+
+	if (bHavePendingAction)
+	{
+		ApplyAction(PendingAction);
+		bHavePendingAction = false;
+		PendingAction = EMeshSelectionToolActions::NoAction;
+	}
 }
 
 
@@ -417,6 +637,236 @@ TUniquePtr<FMeshSelectionChange> UMeshSelectionTool::EndChange()
 
 
 
+
+
+void UMeshSelectionTool::RequestAction(EMeshSelectionToolActions ActionType)
+{
+	if (bHavePendingAction)
+	{
+		return;
+	}
+
+	PendingAction = ActionType;
+	bHavePendingAction = true;
+}
+
+
+void UMeshSelectionTool::ApplyAction(EMeshSelectionToolActions ActionType)
+{
+	switch (ActionType)
+	{
+		case EMeshSelectionToolActions::ClearSelection:
+			ClearSelection();
+			break;
+
+		case EMeshSelectionToolActions::InvertSelection:
+			InvertSelection();
+			break;
+
+
+		case EMeshSelectionToolActions::GrowSelection:
+			GrowShrinkSelection(true);
+			break;
+
+		case EMeshSelectionToolActions::ShrinkSelection:
+			GrowShrinkSelection(false);
+			break;
+
+
+		case EMeshSelectionToolActions::ExpandToConnected:
+			ExpandToConnected();
+			break;
+
+		case EMeshSelectionToolActions::DeleteSelected:
+			DeleteSelectedTriangles();
+			break;
+
+		case EMeshSelectionToolActions::SeparateSelected:
+			SeparateSelectedTriangles();
+			break;
+	}
+}
+
+
+
+
+
+void UMeshSelectionTool::ClearSelection()
+{
+	TArray<int32> SelectedFaces = Selection->GetElements(EMeshSelectionElementType::Face);
+	if (SelectedFaces.Num() == 0)
+	{
+		return;
+	}
+
+	BeginChange(false);
+	ActiveSelectionChange->Add(SelectedFaces);
+	Selection->RemoveIndices(EMeshSelectionElementType::Face, SelectedFaces);
+
+	TUniquePtr<FMeshSelectionChange> SelectionChange = EndChange();
+
+	GetToolManager()->EmitObjectChange(Selection, MoveTemp(SelectionChange), LOCTEXT("ClearSelection", "Clear Selection"));
+
+	OnExternalSelectionChange();
+}
+
+
+
+
+void UMeshSelectionTool::InvertSelection()
+{
+	check(SelectionType == EMeshSelectionElementType::Face);
+	TArray<int32> SelectedFaces = Selection->GetElements(EMeshSelectionElementType::Face);
+	if (SelectedFaces.Num() == 0)
+	{
+		return;
+	}
+
+	TArray<int32> InvertedFaces;
+	const FDynamicMesh3* Mesh = PreviewMesh->GetPreviewDynamicMesh();
+	for (int tid : Mesh->TriangleIndicesItr())
+	{
+		if (SelectedTriangles[tid] == false)
+		{
+			InvertedFaces.Add(tid);
+		}
+	}
+
+	GetToolManager()->BeginUndoTransaction(LOCTEXT("InvertSelection", "Invert Selection"));
+
+	// clear current selection
+	BeginChange(false);
+	ActiveSelectionChange->Add(SelectedFaces);
+	Selection->RemoveIndices(EMeshSelectionElementType::Face, SelectedFaces);
+	TUniquePtr<FMeshSelectionChange> ClearChange = EndChange();
+
+	GetToolManager()->EmitObjectChange(Selection, MoveTemp(ClearChange), LOCTEXT("InvertSelection", "Invert Selection"));
+
+	// add inverted selection
+	BeginChange(true);
+	ActiveSelectionChange->Add(InvertedFaces);
+	Selection->AddIndices(EMeshSelectionElementType::Face, InvertedFaces);
+	TUniquePtr<FMeshSelectionChange> AddChange = EndChange();
+
+	GetToolManager()->EmitObjectChange(Selection, MoveTemp(AddChange), LOCTEXT("InvertSelection", "Invert Selection"));
+
+	GetToolManager()->EndUndoTransaction();
+
+	OnExternalSelectionChange();
+}
+
+
+
+
+
+void UMeshSelectionTool::GrowShrinkSelection(bool bGrow)
+{
+	check(SelectionType == EMeshSelectionElementType::Face);
+	TArray<int32> SelectedFaces = Selection->GetElements(EMeshSelectionElementType::Face);
+	if (SelectedFaces.Num() == 0)
+	{
+		return;
+	}
+
+	const FDynamicMesh3* Mesh = PreviewMesh->GetPreviewDynamicMesh();
+	TArray<int32> Vertices;
+	MeshIndexUtil::TriangleToVertexIDs(Mesh, SelectedFaces, Vertices);
+
+	TSet<int32> ChangeFaces;
+	for (int vid : Vertices)
+	{
+		int OutCount = 0;
+		for (int tid : Mesh->VtxTrianglesItr(vid))
+		{
+			if (SelectedTriangles[tid] == false)
+			{
+				OutCount++;
+			}
+		}
+		if (OutCount == 0)
+		{
+			continue;
+		}
+
+		for (int tid : Mesh->VtxTrianglesItr(vid))
+		{
+			if ( (bGrow && SelectedTriangles[tid] == false) || (bGrow == false && SelectedTriangles[tid]) )
+			{
+				ChangeFaces.Add(tid);
+			}
+		}
+	}
+	if (ChangeFaces.Num() == 0)
+	{
+		return;
+	}
+
+	BeginChange(bGrow);
+	ActiveSelectionChange->Add(ChangeFaces);
+	if (bGrow)
+	{
+		Selection->AddIndices(EMeshSelectionElementType::Face, ChangeFaces);
+		TUniquePtr<FMeshSelectionChange> SelectionChange = EndChange();
+		GetToolManager()->EmitObjectChange(Selection, MoveTemp(SelectionChange), LOCTEXT("GrowSelection", "Grow Selection"));
+	}
+	else
+	{
+		Selection->RemoveIndices(EMeshSelectionElementType::Face, ChangeFaces);
+		TUniquePtr<FMeshSelectionChange> SelectionChange = EndChange();
+		GetToolManager()->EmitObjectChange(Selection, MoveTemp(SelectionChange), LOCTEXT("ShrinkSelection", "Shrink Selection"));
+	}
+	OnExternalSelectionChange();
+}
+
+
+
+
+
+void UMeshSelectionTool::ExpandToConnected()
+{
+	check(SelectionType == EMeshSelectionElementType::Face);
+	TArray<int32> SelectedFaces = Selection->GetElements(EMeshSelectionElementType::Face);
+	if (SelectedFaces.Num() == 0)
+	{
+		return;
+	}
+
+	const FDynamicMesh3* Mesh = PreviewMesh->GetPreviewDynamicMesh();
+
+	TArray<int32> Queue(SelectedFaces);
+	TSet<int32> AddFaces;
+
+	while (Queue.Num() > 0)
+	{
+		int32 CurTri = Queue.Pop(false);
+		FIndex3i NbrTris = Mesh->GetTriNeighbourTris(CurTri);
+
+		for (int j = 0; j < 3; ++j)
+		{
+			int32 tid = NbrTris[j];
+			if (SelectedTriangles[tid] == false && AddFaces.Contains(tid) == false)
+			{
+				AddFaces.Add(tid);
+				Queue.Add(tid);
+			}
+		}
+	}
+	if (AddFaces.Num() == 0)
+	{
+		return;
+	}
+
+	BeginChange(true);
+	ActiveSelectionChange->Add(AddFaces);
+	Selection->AddIndices(EMeshSelectionElementType::Face, AddFaces);
+	TUniquePtr<FMeshSelectionChange> SelectionChange = EndChange();
+	GetToolManager()->EmitObjectChange(Selection, MoveTemp(SelectionChange), LOCTEXT("ExpandToConnected", "Expand Selection"));
+	OnExternalSelectionChange();
+}
+
+
+
+
 void UMeshSelectionTool::DeleteSelectedTriangles()
 {
 	check(SelectionType == EMeshSelectionElementType::Face);
@@ -450,9 +900,9 @@ void UMeshSelectionTool::DeleteSelectedTriangles()
 	// emit combined change sequence
 	GetToolManager()->EmitObjectChange(this, MoveTemp(ChangeSeq), LOCTEXT("MeshSelectionToolDelete", "Delete Faces"));
 
-	OnSelectionUpdated();
-
+	OnExternalSelectionChange();
 	bHaveModifiedMesh = true;
+	bOctreeValid = false;
 }
 
 
diff --git a/Engine/Plugins/Experimental/MeshModelingToolset/Source/MeshModelingTools/Public/DynamicMeshSculptTool.h b/Engine/Plugins/Experimental/MeshModelingToolset/Source/MeshModelingTools/Public/DynamicMeshSculptTool.h
--- a/Engine/Plugins/Experimental/MeshModelingToolset/Source/MeshModelingTools/Public/DynamicMeshSculptTool.h
+++ b/Engine/Plugins/Experimental/MeshModelingToolset/Source/MeshModelingTools/Public/DynamicMeshSculptTool.h
@@ -133,7 +133,7 @@ public:
 	UPROPERTY(EditAnywhere, Category = Sculpting, meta = (EditCondition = "PrimaryBrushType == EDynamicMeshSculptBrushType::Sculpt || PrimaryBrushType == EDynamicMeshSculptBrushType::SculptMax || PrimaryBrushType == EDynamicMeshSculptBrushType::Pinch" ))
 	bool bFreezeTarget;
 
-
+	/** Toggle drawing of wireframe overlay on/off */
 	UPROPERTY(EditAnywhere, Category = Sculpting)
 	bool bShowWireframe;
 
diff --git a/Engine/Plugins/Experimental/MeshModelingToolset/Source/MeshModelingTools/Public/MeshSelectionTool.h b/Engine/Plugins/Experimental/MeshModelingToolset/Source/MeshModelingTools/Public/MeshSelectionTool.h
--- a/Engine/Plugins/Experimental/MeshModelingToolset/Source/MeshModelingTools/Public/MeshSelectionTool.h
+++ b/Engine/Plugins/Experimental/MeshModelingToolset/Source/MeshModelingTools/Public/MeshSelectionTool.h
@@ -6,9 +6,10 @@
 #include "DynamicMeshBrushTool.h"
 #include "SelectionSet.h"
 #include "Changes/MeshSelectionChange.h"
+#include "Changes/ValueWatcher.h"
 #include "MeshSelectionTool.generated.h"
 
-
+class FDynamicMeshOctree3;
 
 /**
  *
@@ -26,6 +27,162 @@ public:
 
 
 
+
+UENUM()
+enum class EMeshSelectionToolActions
+{
+	NoAction,
+
+	ClearSelection,
+	InvertSelection,
+	GrowSelection,
+	ShrinkSelection,
+	ExpandToConnected,
+
+	DeleteSelected,
+	SeparateSelected
+};
+
+
+
+UCLASS()
+class MESHMODELINGTOOLS_API UMeshSelectionToolActionPropertySet : public UInteractiveToolPropertySet
+{
+	GENERATED_BODY()
+
+public:
+	TWeakObjectPtr<UMeshSelectionTool> ParentTool;
+
+	void Initialize(UMeshSelectionTool* ParentToolIn) { ParentTool = ParentToolIn; }
+
+	void PostAction(EMeshSelectionToolActions Action);
+};
+
+
+
+
+UCLASS()
+class MESHMODELINGTOOLS_API UMeshSelectionEditActions : public UMeshSelectionToolActionPropertySet
+{
+	GENERATED_BODY()
+
+public:
+	UFUNCTION(CallInEditor, Category = Selection, meta = (DisplayName = "Clear", DisplayPriority = 1))
+	void Clear()
+	{
+		PostAction(EMeshSelectionToolActions::ClearSelection);
+	}
+
+	UFUNCTION(CallInEditor, Category = Selection, meta = (DisplayName = "Invert", DisplayPriority = 2))
+	void Invert()
+	{
+		PostAction(EMeshSelectionToolActions::InvertSelection);
+	}
+
+
+	UFUNCTION(CallInEditor, Category = Selection, meta = (DisplayName = "Grow", DisplayPriority = 3))
+	void Grow()
+	{
+		PostAction(EMeshSelectionToolActions::GrowSelection);
+	}
+
+	UFUNCTION(CallInEditor, Category = Selection, meta = (DisplayName = "Shrink", DisplayPriority = 4))
+	void Shrink()
+	{
+		PostAction(EMeshSelectionToolActions::ShrinkSelection);
+	}
+
+	UFUNCTION(CallInEditor, Category = Selection, meta = (DisplayName = "ExpandToConnected", DisplayPriority = 5))
+	void ExpandToConnected()
+	{
+		PostAction(EMeshSelectionToolActions::ExpandToConnected);
+	}
+
+};
+
+
+
+
+UCLASS()
+class MESHMODELINGTOOLS_API UMeshSelectionMeshEditActions : public UMeshSelectionToolActionPropertySet
+{
+	GENERATED_BODY()
+
+public:
+	UFUNCTION(CallInEditor, Category = MeshEdits, meta = (DisplayName = "Delete"))
+	void DeleteTriangles()
+	{
+		PostAction(EMeshSelectionToolActions::DeleteSelected);
+	}
+
+	UFUNCTION(CallInEditor, Category = MeshEdits, meta = (DisplayName = "Separate"))
+	void SeparateTriangles() 
+	{
+		PostAction(EMeshSelectionToolActions::SeparateSelected);
+	}
+};
+
+
+
+
+
+
+UENUM()
+enum class EMeshSelectionToolPrimaryMode
+{
+	/** Select all triangles inside the brush */
+	Brush,
+
+	/** Select all triangles inside brush with normal within angular tolerance of hit triangle */
+	AngleFiltered,
+
+	/** Select all triangles inside brush that are visible from current view */
+	Visible,
+
+	/** Select all triangles connected to any triangle inside the brush */
+	AllConnected,
+
+	/** Select all triangles in groups connected to any triangle inside the brush */
+	AllInGroup,
+
+	/** Select all triangles with normal within angular tolerance of hit triangle */
+	AllWithinAngle
+};
+
+
+
+UCLASS()
+class MESHMODELINGTOOLS_API UMeshSelectionToolProperties : public UInteractiveToolPropertySet
+{
+	GENERATED_BODY()
+public:
+
+	/** The Selection Mode defines the behavior of the selection brush */
+	UPROPERTY(EditAnywhere, Category = Selection)
+	EMeshSelectionToolPrimaryMode SelectionMode = EMeshSelectionToolPrimaryMode::Brush;
+
+	/** Angle in Degrees used for Angle-based Selection Modes */
+	UPROPERTY(EditAnywhere, Category = Selection, meta = (UIMin = "0.0", UIMax = "90.0") )
+	float AngleTolerance = 1.0;
+
+	/** A Volumetric brush selects anything within a 3D sphere, rather than only elements connected to the Brush Position  */
+	UPROPERTY(EditAnywhere, Category = Selection)
+	bool bVolumetricBrush = false;
+
+	/** Allow the brush to hit back-facing parts of the surface  */
+	UPROPERTY(EditAnywhere, Category = Selection)
+	bool bHitBackFaces = true;
+
+	/** Toggle drawing of wireframe overlay on/off */
+	UPROPERTY(EditAnywhere, Category = Selection)
+	bool bShowWireframe = true;
+
+	virtual void SaveProperties(UInteractiveTool* SaveFromTool) override;
+	virtual void RestoreProperties(UInteractiveTool* RestoreToTool) override;
+};
+
+
+
 /**
  *
  */
@@ -52,6 +209,7 @@ public:
 	virtual bool CanAccept() const override { return bHaveModifiedMesh; }
 
 	// UBaseBrushTool overrides
+	virtual bool HitTest(const FRay& Ray, FHitResult& OutHit) override;
 	virtual void OnBeginDrag(const FRay& Ray) override;
 	virtual void OnUpdateDrag(const FRay& Ray) override;
 	virtual void OnEndDrag(const FRay& Ray) override;
@@ -59,19 +217,21 @@ public:
 
 public:
 
+	virtual void RequestAction(EMeshSelectionToolActions ActionType);
 
-	UFUNCTION(CallInEditor, Category = Options, meta = (DisplayName = "Delete"))
-	void DeleteSelectedTriangles();
+	UPROPERTY()
+	UMeshSelectionToolProperties* SelectionProps;
 
-	UFUNCTION(CallInEditor, Category = Options, meta = (DisplayName = "Separate"))
-	void SeparateSelectedTriangles();
+	UPROPERTY()
+	UMeshSelectionEditActions* SelectionActions;
 
+	UPROPERTY()
+	UMeshSelectionMeshEditActions* EditActions;
 
 protected:
 
 	virtual void ApplyStamp(const FBrushStampData& Stamp);
 
-
 	virtual void OnShutdown(EToolShutdownType ShutdownType) override;
 
 protected:
@@ -81,18 +241,28 @@ protected:
 	UWorld* TargetWorld;
 	IToolsContextAssetAPI* AssetAPI;
 
+	// note: ideally this octree would be part of PreviewMesh!
+	TUniquePtr<FDynamicMeshOctree3> Octree;
+	bool bOctreeValid = false;
+	TUniquePtr<FDynamicMeshOctree3>& GetOctree();
+
 	EMeshSelectionElementType SelectionType = EMeshSelectionElementType::Face;
 
-	bool bInRemoveStroke = false;
+	TValueWatcher<bool> ShowWireframeWatcher;
 
+
+	bool bInRemoveStroke = false;
 	FBrushStampData StartStamp;
 	FBrushStampData LastStamp;
-
 	bool bStampPending;
 
+	void UpdateFaceSelection(const FBrushStampData& Stamp, const TArray<int>& BrushROI);
+
 
 	// temp
 	TArray<int> IndexBuf;
+	TArray<int32> TemporaryBuffer;
+	TSet<int32> TemporarySet;
 	void CalculateVertexROI(const FBrushStampData& Stamp, TArray<int>& VertexROI);
 	void CalculateTriangleROI(const FBrushStampData& Stamp, TArray<int>& TriangleROI);
 	TArray<int> PreviewBrushROI;
@@ -110,6 +280,20 @@ protected:
 	void CancelChange();
 
 
+	// actions
+
+	bool bHavePendingAction = false;
+	EMeshSelectionToolActions PendingAction;
+	virtual void ApplyAction(EMeshSelectionToolActions ActionType);
+
+	void ClearSelection();
+	void InvertSelection();
+	void GrowShrinkSelection(bool bGrow);
+	void ExpandToConnected();
+
+	void DeleteSelectedTriangles();
+	void SeparateSelectedTriangles();
+
 	bool bHaveModifiedMesh = false;
 };
 
diff --git a/Engine/Plugins/Experimental/ModelingToolsEditorMode/Source/ModelingToolsEditorMode/Private/ModelingToolsActions.cpp b/Engine/Plugins/Experimental/ModelingToolsEditorMode/Source/ModelingToolsEditorMode/Private/ModelingToolsActions.cpp
--- a/Engine/Plugins/Experimental/ModelingToolsEditorMode/Source/ModelingToolsEditorMode/Private/ModelingToolsActions.cpp
+++ b/Engine/Plugins/Experimental/ModelingToolsEditorMode/Source/ModelingToolsEditorMode/Private/ModelingToolsActions.cpp
@@ -34,7 +34,6 @@ void FModelingToolActionCommands::GetToolDefaultObjectList(TArray<UInteractiveTo
 	ToolCDOs.Add(GetMutableDefault<UEditMeshPolygonsTool>());
 	ToolCDOs.Add(GetMutableDefault<UMeshSpaceDeformerTool>());
 	ToolCDOs.Add(GetMutableDefault<UShapeSprayTool>());
-	ToolCDOs.Add(GetMutableDefault<UMeshSelectionTool>());
 }
 
 
@@ -45,6 +44,7 @@ void FModelingToolActionCommands::RegisterAllToolActions()
 	FSculptToolActionCommands::Register();
 	FDrawPolygonToolActionCommands::Register();
 	FTransformToolActionCommands::Register();
+	FMeshSelectionToolActionCommands::Register();
 }
 
 void FModelingToolActionCommands::UnregisterAllToolActions()
@@ -53,6 +53,7 @@ void FModelingToolActionCommands::UnregisterAllToolActions()
 	FSculptToolActionCommands::Unregister();
 	FDrawPolygonToolActionCommands::Unregister();
 	FTransformToolActionCommands::Unregister();
+	FMeshSelectionToolActionCommands::Unregister();
 }
 
 
@@ -74,6 +75,10 @@ void FModelingToolActionCommands::UpdateToolCommandBinding(UInteractiveTool* Too
 	{
 		UPDATE_BINDING(FDrawPolygonToolActionCommands);
 	}
+	else if (Cast<UMeshSelectionTool>(Tool) != nullptr)
+	{
+		UPDATE_BINDING(FMeshSelectionToolActionCommands);
+	}
 	else
 	{
 		UPDATE_BINDING(FModelingToolActionCommands);
@@ -134,6 +139,22 @@ void FDrawPolygonToolActionCommands::GetToolDefaultObjectList(TArray<UInteractiv
 }
 
 
+FMeshSelectionToolActionCommands::FMeshSelectionToolActionCommands() :
+	TInteractiveToolCommands<FMeshSelectionToolActionCommands>(
+		"ModelingToolsMeshSelectionTool", // Context name for fast lookup
+		NSLOCTEXT("Contexts", "ModelingToolsMeshSelectionTool", "Modeling Tools - Mesh Selection Tool"), // Localized context name for displaying
+		NAME_None, // Parent
+		FEditorStyle::GetStyleSetName() // Icon Style Set
+		)
+{
+}
+
+void FMeshSelectionToolActionCommands::GetToolDefaultObjectList(TArray<UInteractiveTool*>& ToolCDOs)
+{
+	ToolCDOs.Add(GetMutableDefault<UMeshSelectionTool>());
+}
+
+
 
 
 #undef LOCTEXT_NAMESPACE
diff --git a/Engine/Plugins/Experimental/ModelingToolsEditorMode/Source/ModelingToolsEditorMode/Public/ModelingToolsActions.h b/Engine/Plugins/Experimental/ModelingToolsEditorMode/Source/ModelingToolsEditorMode/Public/ModelingToolsActions.h
--- a/Engine/Plugins/Experimental/ModelingToolsEditorMode/Source/ModelingToolsEditorMode/Public/ModelingToolsActions.h
+++ b/Engine/Plugins/Experimental/ModelingToolsEditorMode/Source/ModelingToolsEditorMode/Public/ModelingToolsActions.h
@@ -62,10 +62,17 @@ public:
 };
 
 
-
 class FDrawPolygonToolActionCommands : public TInteractiveToolCommands<FDrawPolygonToolActionCommands>
 {
 public:
 	FDrawPolygonToolActionCommands();
 	virtual void GetToolDefaultObjectList(TArray<UInteractiveTool*>& ToolCDOs) override;
 };
+
+
+class FMeshSelectionToolActionCommands : public TInteractiveToolCommands<FMeshSelectionToolActionCommands>
+{
+public:
+	FMeshSelectionToolActionCommands();
+	virtual void GetToolDefaultObjectList(TArray<UInteractiveTool*>& ToolCDOs) override;
+};
\ No newline at end of file
