cis fix - auto-generated destructor needs to be able to see definition of DynamicMeshOctree
#rnx

diff --git a/Engine/Plugins/Experimental/MeshModelingToolset/Source/MeshModelingTools/Private/MeshSelectionTool.cpp b/Engine/Plugins/Experimental/MeshModelingToolset/Source/MeshModelingTools/Private/MeshSelectionTool.cpp
--- a/Engine/Plugins/Experimental/MeshModelingToolset/Source/MeshModelingTools/Private/MeshSelectionTool.cpp
+++ b/Engine/Plugins/Experimental/MeshModelingToolset/Source/MeshModelingTools/Private/MeshSelectionTool.cpp
@@ -10,7 +10,6 @@
 #include "Changes/MeshChange.h"
 #include "MeshIndexUtil.h"
 #include "AssetGenerationUtil.h"
-#include "DynamicMeshOctree3.h"
 
 #define LOCTEXT_NAMESPACE "UMeshSelectionTool"
 
diff --git a/Engine/Plugins/Experimental/MeshModelingToolset/Source/MeshModelingTools/Public/MeshSelectionTool.h b/Engine/Plugins/Experimental/MeshModelingToolset/Source/MeshModelingTools/Public/MeshSelectionTool.h
--- a/Engine/Plugins/Experimental/MeshModelingToolset/Source/MeshModelingTools/Public/MeshSelectionTool.h
+++ b/Engine/Plugins/Experimental/MeshModelingToolset/Source/MeshModelingTools/Public/MeshSelectionTool.h
@@ -7,10 +7,9 @@
 #include "SelectionSet.h"
 #include "Changes/MeshSelectionChange.h"
 #include "Changes/ValueWatcher.h"
+#include "DynamicMeshOctree3.h"
 #include "MeshSelectionTool.generated.h"
 
-class FDynamicMeshOctree3;
-
 /**
  *
  */
