ModelingTools: cleanup tool property sets
#rnx

diff --git a/Engine/Plugins/Experimental/MeshModelingToolset/Source/MeshModelingTools/Private/PlaneCutTool.cpp b/Engine/Plugins/Experimental/MeshModelingToolset/Source/MeshModelingTools/Private/PlaneCutTool.cpp
--- a/Engine/Plugins/Experimental/MeshModelingToolset/Source/MeshModelingTools/Private/PlaneCutTool.cpp
+++ b/Engine/Plugins/Experimental/MeshModelingToolset/Source/MeshModelingTools/Private/PlaneCutTool.cpp
@@ -63,14 +63,11 @@ UPlaneCutToolProperties::UPlaneCutToolProperties()
 	bFillCutHole = true;
 	SpacingBetweenHalves = 1;
 	bShowPreview = true;
-}
-
-UPlaneCutAdvancedProperties::UPlaneCutAdvancedProperties()
-{
 	bFillSpans = false;
 }
 
 
+
 UPlaneCutTool::UPlaneCutTool()
 {
 	CutPlaneOrigin = FVector::ZeroVector;
@@ -112,13 +109,9 @@ void UPlaneCutTool::Setup()
 	PlaneTransformProxy = NewObject<UTransformProxy>(this);
 	PlaneTransformGizmo = GizmoManager->Create3AxisTransformGizmo(this);
 
-	BasicProperties = NewObject<UPlaneCutToolProperties>(this, TEXT("Plane Cut Settings"));
-	AdvancedProperties = NewObject<UPlaneCutAdvancedProperties>(this, TEXT("Advanced Settings"));
-
 	// initialize our properties
+	BasicProperties = NewObject<UPlaneCutToolProperties>(this, TEXT("Plane Cut Settings"));
 	AddToolPropertySource(BasicProperties);
-	AddToolPropertySource(AdvancedProperties);
-
 
 	// initialize the PreviewMesh+BackgroundCompute object
 	UpdateNumPreviews();
@@ -209,7 +202,7 @@ TSharedPtr<FDynamicMeshOperator> UPlaneCutOperatorFactory::MakeNewOperator()
 	TSharedPtr<FPlaneCutOp> CutOp = MakeShared<FPlaneCutOp>();
 	CutOp->bDiscardAttributes = CutTool->BasicProperties->bDiscardAttributes;
 	CutOp->bFillCutHole = CutTool->BasicProperties->bFillCutHole;
-	CutOp->bFillSpans = CutTool->AdvancedProperties->bFillSpans;
+	CutOp->bFillSpans = CutTool->BasicProperties->bFillSpans;
 
 	FTransform LocalToWorld = CutTool->ComponentTarget->GetWorldTransform();
 	FTransform WorldToLocal = LocalToWorld.Inverse();
diff --git a/Engine/Plugins/Experimental/MeshModelingToolset/Source/MeshModelingTools/Public/EditNormalsTool.h b/Engine/Plugins/Experimental/MeshModelingToolset/Source/MeshModelingTools/Public/EditNormalsTool.h
--- a/Engine/Plugins/Experimental/MeshModelingToolset/Source/MeshModelingTools/Public/EditNormalsTool.h
+++ b/Engine/Plugins/Experimental/MeshModelingToolset/Source/MeshModelingTools/Public/EditNormalsTool.h
@@ -53,15 +53,15 @@ public:
 
 
 	/** Recompute all mesh normals */
-	UPROPERTY(EditAnywhere, Category = Options, meta = (EditCondition = "!bRecomputeNormalTopologyAndEdgeSharpness"))
+	UPROPERTY(EditAnywhere, Category = NormalsCalculation, meta = (EditCondition = "!bRecomputeNormalTopologyAndEdgeSharpness"))
 	bool bRecomputeNormals;
 
 	/** Choose the method for computing vertex normals */
-	UPROPERTY(EditAnywhere, Category = Options)
+	UPROPERTY(EditAnywhere, Category = NormalsCalculation)
 	ENormalCalculationMethod NormalCalculationMethod;
 
 	/** Invert (flip) all mesh normals and associated triangle orientations */
-	UPROPERTY(EditAnywhere, Category = Options)
+	UPROPERTY(EditAnywhere, Category = NormalsCalculation)
 	bool bInvertNormals;
 
 	/** Choose where to create sharp edges by having multiple normals on the same vertices, using the change in normals of adjacent faces.  If enabled, normals will always be recomputed. */
diff --git a/Engine/Plugins/Experimental/MeshModelingToolset/Source/MeshModelingTools/Public/PlaneCutTool.h b/Engine/Plugins/Experimental/MeshModelingToolset/Source/MeshModelingTools/Public/PlaneCutTool.h
--- a/Engine/Plugins/Experimental/MeshModelingToolset/Source/MeshModelingTools/Public/PlaneCutTool.h
+++ b/Engine/Plugins/Experimental/MeshModelingToolset/Source/MeshModelingTools/Public/PlaneCutTool.h
@@ -62,42 +62,23 @@ public:
 	bool bKeepBothHalves;
 
 	/** If keeping both halves, separate the two pieces by this amount */
-	UPROPERTY(EditAnywhere, Category = Options)
+	UPROPERTY(EditAnywhere, Category = Options, meta = (EditCondition = "bKeepBothHalves == true") )
 	float SpacingBetweenHalves;
 
-
 	/** If true, the cut surface is filled with simple planar hole fill surface(s) */
-	UPROPERTY(EditAnywhere, Category = Advanced)
+	UPROPERTY(EditAnywhere, Category = Options)
 	bool bFillCutHole;
 
-	UPROPERTY(EditAnywhere, Category = Preview)
+	UPROPERTY(EditAnywhere, Category = Options)
 	bool bShowPreview;
-};
-
 
-
-
-/**
- * Advanced properties of the plane cut operation
- */
-UCLASS()
-class MESHMODELINGTOOLS_API UPlaneCutAdvancedProperties : public UInteractiveToolPropertySet
-{
-	GENERATED_BODY()
-
-public:
-	UPlaneCutAdvancedProperties();
-
-	/** 
-	 *  If true, try to treat open edge spans as fillable loops
-	 *  TODO: add some tolerance-based span joining code to try to turn open spans into loops instead?
-	 */
-	UPROPERTY(EditAnywhere, Category = Advanced)
+	UPROPERTY(EditAnywhere, Category = Options, AdvancedDisplay)
 	bool bFillSpans;
-
 };
 
 
+
+
 UCLASS()
 class MESHMODELINGTOOLS_API UPlaneCutOperatorFactory : public UObject, public IDynamicMeshOperatorFactory
 {
@@ -152,9 +133,6 @@ protected:
 	UPROPERTY()
 	UPlaneCutToolProperties* BasicProperties;
 
-	UPROPERTY()
-	UPlaneCutAdvancedProperties* AdvancedProperties;
-
 	/** Origin of cutting plane */
 	UPROPERTY()
 	FVector CutPlaneOrigin;
diff --git a/Engine/Plugins/Experimental/MeshModelingToolset/Source/MeshModelingTools/Public/Properties/MeshMaterialProperties.h b/Engine/Plugins/Experimental/MeshModelingToolset/Source/MeshModelingTools/Public/Properties/MeshMaterialProperties.h
--- a/Engine/Plugins/Experimental/MeshModelingToolset/Source/MeshModelingTools/Public/Properties/MeshMaterialProperties.h
+++ b/Engine/Plugins/Experimental/MeshModelingToolset/Source/MeshModelingTools/Public/Properties/MeshMaterialProperties.h
@@ -71,13 +71,13 @@ class MESHMODELINGTOOLS_API UExistingMeshMaterialProperties : public UInteractiv
 public:
 
 	/** Material that will be used on the mesh */
-	UPROPERTY(EditAnywhere, Category = Material)
+	UPROPERTY(EditAnywhere, Category = MaterialPreview)
 	ESetMeshMaterialMode MaterialMode = ESetMeshMaterialMode::KeepOriginal;
 
-	UPROPERTY(EditAnywhere, Category = Material, meta = (UIMin = "1.0", UIMax = "40.0", ClampMin = "0.01", ClampMax = "1000.0", EditCondition = "MaterialMode == ESetMeshMaterialMode::Checkerboard"))
+	UPROPERTY(EditAnywhere, Category = MaterialPreview, meta = (UIMin = "1.0", UIMax = "40.0", ClampMin = "0.01", ClampMax = "1000.0", EditCondition = "MaterialMode == ESetMeshMaterialMode::Checkerboard"))
 	float CheckerDensity = 20.0f;
 
-	UPROPERTY(EditAnywhere, Category = Material, meta = (EditCondition = "MaterialMode == ESetMeshMaterialMode::Override"))
+	UPROPERTY(EditAnywhere, Category = MaterialPreview, meta = (EditCondition = "MaterialMode == ESetMeshMaterialMode::Override"))
 	UMaterialInterface* OverrideMaterial = nullptr;
 
 	UPROPERTY()
diff --git a/Engine/Plugins/Experimental/MeshModelingToolset/Source/MeshModelingTools/Public/UVProjectionTool.h b/Engine/Plugins/Experimental/MeshModelingToolset/Source/MeshModelingTools/Public/UVProjectionTool.h
--- a/Engine/Plugins/Experimental/MeshModelingToolset/Source/MeshModelingTools/Public/UVProjectionTool.h
+++ b/Engine/Plugins/Experimental/MeshModelingToolset/Source/MeshModelingTools/Public/UVProjectionTool.h
@@ -56,27 +56,27 @@ public:
 
 
 	/** Choose the UV projection method (cube, cylinder, plane) */
-	UPROPERTY(EditAnywhere, Category = Options)
+	UPROPERTY(EditAnywhere, Category = ProjectionSettings)
 	EUVProjectionMethod UVProjectionMethod;
 
 	/** Per-axis scaling of projection primitive */
-	UPROPERTY(EditAnywhere, Category = Options)
+	UPROPERTY(EditAnywhere, Category = ProjectionSettings)
 	FVector ProjectionPrimitiveScale;
 
 	/** If triangle normal direction is within this threshold degrees of the cylinder top/bottom plane direction, project UVs to the top/bottom plane instead of the sides */
-	UPROPERTY(EditAnywhere, Category = Options, meta = (Default = "1", UIMin = "0", UIMax = "20", ClampMin = "0.0", ClampMax = "90.0", EditCondition = "UVProjectionMethod == EUVProjectionMethod::Cylinder"))
+	UPROPERTY(EditAnywhere, Category = ProjectionSettings, meta = (Default = "1", UIMin = "0", UIMax = "20", ClampMin = "0.0", ClampMax = "90.0", EditCondition = "UVProjectionMethod == EUVProjectionMethod::Cylinder"))
 	float CylinderProjectToTopOrBottomAngleThreshold;
 
 	/** Choose the UV scale factors */
-	UPROPERTY(EditAnywhere, Category = Options)
+	UPROPERTY(EditAnywhere, Category = ProjectionSettings)
 	FVector2D UVScale;
 
 	/** Choose the UV offsets */
-	UPROPERTY(EditAnywhere, Category = Options)
+	UPROPERTY(EditAnywhere, Category = ProjectionSettings)
 	FVector2D UVOffset;
 
 	/** If set, UV scales will be relative to world space so different objects created with the same UV scale should have the same average texel size */
-	UPROPERTY(EditAnywhere, Category = Options, meta = (DisplayName = "UV Scale Relative to World Space"))
+	UPROPERTY(EditAnywhere, Category = ProjectionSettings, meta = (DisplayName = "UV Scale Relative to World Space"))
 	bool bWorldSpaceUVScale = false;
 
 };
diff --git a/Engine/Plugins/Experimental/MeshModelingToolset/Source/MeshModelingTools/Public/WeldMeshEdgesTool.h b/Engine/Plugins/Experimental/MeshModelingToolset/Source/MeshModelingTools/Public/WeldMeshEdgesTool.h
--- a/Engine/Plugins/Experimental/MeshModelingToolset/Source/MeshModelingTools/Public/WeldMeshEdgesTool.h
+++ b/Engine/Plugins/Experimental/MeshModelingToolset/Source/MeshModelingTools/Public/WeldMeshEdgesTool.h
@@ -60,7 +60,7 @@ protected:
 	float Tolerance;
 
 	/** Enable edge flips */
-	UPROPERTY(EditAnywhere, Category = Advanced)
+	UPROPERTY(EditAnywhere, Category = Options, AdvancedDisplay)
 	bool bOnlyUnique;
 
 protected:
