ModelingTools: Add custom materials for Selection and Sculpting, that show backfaces, quasi-flat normals, etc.
#rnx

diff --git a/Engine/Plugins/Experimental/MeshModelingToolset/Source/MeshModelingTools/Private/DynamicMeshSculptTool.cpp b/Engine/Plugins/Experimental/MeshModelingToolset/Source/MeshModelingTools/Private/DynamicMeshSculptTool.cpp
--- a/Engine/Plugins/Experimental/MeshModelingToolset/Source/MeshModelingTools/Private/DynamicMeshSculptTool.cpp
+++ b/Engine/Plugins/Experimental/MeshModelingToolset/Source/MeshModelingTools/Private/DynamicMeshSculptTool.cpp
@@ -14,6 +14,7 @@
 #include "MeshIndexUtil.h"
 #include "Drawing/MeshDebugDrawing.h"
 #include "PreviewMesh.h"
+#include "ToolSetupUtil.h"
 
 #include "Changes/MeshVertexChange.h"
 #include "Changes/MeshChange.h"
@@ -46,7 +47,6 @@ UBrushSculptProperties::UBrushSculptProperties()
 	bPreserveUVFlow = false;
 	BrushDepth = 0;
 	bFreezeTarget = false;
-	bShowWireframe = false;
 }
 
 void UBrushSculptProperties::SaveProperties(UInteractiveTool* SaveFromTool)
@@ -57,7 +57,6 @@ void UBrushSculptProperties::SaveProperties(UInteractiveTool* SaveFromTool)
 	PropertyCache->PrimaryBrushType = this->PrimaryBrushType;
 	PropertyCache->bPreserveUVFlow = this->bPreserveUVFlow;
 	PropertyCache->BrushDepth = this->BrushDepth;
-	PropertyCache->bShowWireframe = this->bShowWireframe;
 }
 
 void UBrushSculptProperties::RestoreProperties(UInteractiveTool* RestoreToTool)
@@ -68,11 +67,13 @@ void UBrushSculptProperties::RestoreProperties(UInteractiveTool* RestoreToTool)
 	this->PrimaryBrushType = PropertyCache->PrimaryBrushType;
 	this->bPreserveUVFlow = PropertyCache->bPreserveUVFlow;
 	this->BrushDepth = PropertyCache->BrushDepth;
-	this->bShowWireframe = PropertyCache->bShowWireframe;
 }
 
 
 
+
+
+
 UBrushRemeshProperties::UBrushRemeshProperties()
 {
 	RelativeSize = 1.0;
@@ -110,7 +111,7 @@ void UDynamicMeshSculptTool::Setup()
 	DynamicMeshComponent->SetWorldTransform(ComponentTarget->GetWorldTransform());
 
 	// copy material if there is one
-	auto Material = ComponentTarget->GetMaterial(0);
+	UMaterialInterface* Material = ComponentTarget->GetMaterial(0);
 	if (Material != nullptr)
 	{
 		DynamicMeshComponent->SetMaterial(0, Material);
@@ -140,13 +141,10 @@ void UDynamicMeshSculptTool::Setup()
 	CalculateBrushRadius();
 
 	SculptProperties = NewObject<UBrushSculptProperties>(this, TEXT("Sculpting"));
-	ShowWireframeWatcher.Initialize(
-		[this]() { return SculptProperties->bShowWireframe; },
-		[this](bool bNewValue) { DynamicMeshComponent->bExplicitShowWireframe = bNewValue; }, false);
-
 	RemeshProperties = NewObject<UBrushRemeshProperties>(this, TEXT("Remeshing"));
 	InitialEdgeLength = EstimateIntialSafeTargetLength(*Mesh, 5000);
 
+
 	// hide input StaticMeshComponent
 	ComponentTarget->SetOwnerVisibility(false);
 
@@ -176,6 +174,17 @@ void UDynamicMeshSculptTool::Setup()
 	CalculateBrushRadius();
 	SculptProperties->RestoreProperties(this);
 
+	ViewProperties = NewObject<UMeshEditingViewProperties>();
+	ViewProperties->RestoreProperties(this);
+	AddToolPropertySource(ViewProperties);
+	ShowWireframeWatcher.Initialize(
+		[this]() { return ViewProperties->bShowWireframe; },
+		[this](bool bNewValue) { DynamicMeshComponent->bExplicitShowWireframe = bNewValue; }, false);
+	MaterialModeWatcher.Initialize(
+		[this]() { return ViewProperties->MaterialMode; },
+		[this](EMeshEditingMaterialModes NewMode) { UpdateMaterialMode(NewMode); }, EMeshEditingMaterialModes::ExistingMaterial);
+
+
 	GetToolManager()->DisplayMessage(
 		LOCTEXT("OnStartSculptTool", "Hold Shift to Smooth, Ctrl to Invert (where applicable). Q/A keys cycle through Brush Types, Shift+Q/A for Brush History. S/D change Size (shift to small-step), W/E change Speed."),
 		EToolMessageLevel::UserNotification);
@@ -235,6 +244,7 @@ void UDynamicMeshSculptTool::Shutdown(EToolShutdownType ShutdownType)
 
 	BrushProperties->SaveProperties(this);
 	SculptProperties->SaveProperties(this);
+	ViewProperties->SaveProperties(this);
 }
 
 
@@ -938,8 +948,6 @@ void UDynamicMeshSculptTool::Render(IToolsContextRenderAPI* RenderAPI)
 	UMeshSurfacePointTool::Render(RenderAPI);
 
 	BrushIndicator->Update( (float)this->CurrentBrushRadius, this->LastBrushPosWorld, this->LastBrushPosNormalWorld );
-
-	ShowWireframeWatcher.CheckAndUpdate();
 }
 
 
@@ -949,6 +957,9 @@ void UDynamicMeshSculptTool::Tick(float DeltaTime)
 
 	UMeshSurfacePointTool::Tick(DeltaTime);
 
+	ShowWireframeWatcher.CheckAndUpdate();
+	MaterialModeWatcher.CheckAndUpdate();
+
 	// if user changed to not-frozen, we need to update the target
 	if (bCachedFreezeTarget != SculptProperties->bFreezeTarget)
 	{
@@ -1571,8 +1582,8 @@ void UDynamicMeshSculptTool::RegisterActions(FInteractiveToolActionSet& ActionSe
 		TEXT("ToggleWireframe"),
 		LOCTEXT("ToggleWireframe", "Toggle Wireframe"),
 		LOCTEXT("ToggleWireframeTooltip", "Toggle visibility of wireframe overlay"),
-		EModifierKey::Control, EKeys::W,
-		[this]() { SculptProperties->bShowWireframe = !SculptProperties->bShowWireframe; });
+		EModifierKey::Alt, EKeys::W,
+		[this]() { ViewProperties->bShowWireframe = !ViewProperties->bShowWireframe; });
 
 }
 
@@ -1646,4 +1657,28 @@ void UDynamicMeshSculptTool::UpdateSavedVertex(int vid, const FVector3d& OldPosi
 }
 
 
+
+void UDynamicMeshSculptTool::UpdateMaterialMode(EMeshEditingMaterialModes MaterialMode)
+{
+	if (MaterialMode == EMeshEditingMaterialModes::ExistingMaterial)
+	{
+		UMaterialInterface* Material = ComponentTarget->GetMaterial(0);
+		if (Material != nullptr)
+		{
+			DynamicMeshComponent->SetMaterial(0, Material);
+		}
+		DynamicMeshComponent->bCastDynamicShadow = ComponentTarget->GetOwnerComponent()->bCastDynamicShadow;
+	} 
+	else if (MaterialMode == EMeshEditingMaterialModes::MeshFocusMaterial)
+	{ 
+		UMaterialInterface* Material = ToolSetupUtil::GetSculptMaterial1(GetToolManager());
+		if (Material != nullptr)
+		{
+			DynamicMeshComponent->SetMaterial(0, Material);
+		}
+		DynamicMeshComponent->bCastDynamicShadow = false;
+	}
+}
+
+
 #undef LOCTEXT_NAMESPACE
diff --git a/Engine/Plugins/Experimental/MeshModelingToolset/Source/MeshModelingTools/Private/MeshSelectionTool.cpp b/Engine/Plugins/Experimental/MeshModelingToolset/Source/MeshModelingTools/Private/MeshSelectionTool.cpp
--- a/Engine/Plugins/Experimental/MeshModelingToolset/Source/MeshModelingTools/Private/MeshSelectionTool.cpp
+++ b/Engine/Plugins/Experimental/MeshModelingToolset/Source/MeshModelingTools/Private/MeshSelectionTool.cpp
@@ -10,6 +10,7 @@
 #include "Changes/MeshChange.h"
 #include "MeshIndexUtil.h"
 #include "AssetGenerationUtil.h"
+#include "ToolSetupUtil.h"
 
 #define LOCTEXT_NAMESPACE "UMeshSelectionTool"
 
@@ -99,12 +100,12 @@ void UMeshSelectionTool::Setup()
 	PreviewMesh->EnableWireframe(true);
 
 	// set vertex color material on base component so we can see selection
-	UMaterialInterface* VertexColorMat = 
-		GetToolManager()->GetContextQueriesAPI()->GetStandardMaterial(EStandardToolContextMaterials::VertexColorMaterial);
-	if (VertexColorMat != nullptr)
+	UMaterialInterface* SelectionMat = ToolSetupUtil::GetSelectionMaterial(GetToolManager());
+	if (SelectionMat != nullptr)
 	{
-		PreviewMesh->SetMaterial(VertexColorMat);
+		PreviewMesh->SetMaterial(SelectionMat);
 	}
+	PreviewMesh->GetRootComponent()->bCastDynamicShadow = false;
 
 	const FDynamicMesh3* Mesh = PreviewMesh->GetPreviewDynamicMesh();
 	SelectedVertices = TBitArray<>(false, Mesh->MaxVertexID());
@@ -163,7 +164,7 @@ void UMeshSelectionTool::RegisterActions(FInteractiveToolActionSet& ActionSet)
 		TEXT("ToggleWireframe"),
 		LOCTEXT("ToggleWireframe", "Toggle Wireframe"),
 		LOCTEXT("ToggleWireframeTooltip", "Toggle visibility of wireframe overlay"),
-		EModifierKey::Control, EKeys::W,
+		EModifierKey::Alt, EKeys::W,
 		[this]() { SelectionProps->bShowWireframe = !SelectionProps->bShowWireframe; });
 }
 
diff --git a/Engine/Plugins/Experimental/MeshModelingToolset/Source/MeshModelingTools/Private/Properties/MeshMaterialProperties.cpp b/Engine/Plugins/Experimental/MeshModelingToolset/Source/MeshModelingTools/Private/Properties/MeshMaterialProperties.cpp
--- a/Engine/Plugins/Experimental/MeshModelingToolset/Source/MeshModelingTools/Private/Properties/MeshMaterialProperties.cpp
+++ b/Engine/Plugins/Experimental/MeshModelingToolset/Source/MeshModelingTools/Private/Properties/MeshMaterialProperties.cpp
@@ -77,4 +77,26 @@ void UExistingMeshMaterialProperties::SetMaterialIfChanged(UMaterialInterface* O
 }
 
 
+
+
+
+
+
+void UMeshEditingViewProperties::SaveProperties(UInteractiveTool* SaveFromTool)
+{
+	UMeshEditingViewProperties* PropertyCache = GetPropertyCache<UMeshEditingViewProperties>();
+	PropertyCache->bShowWireframe = this->bShowWireframe;
+	PropertyCache->MaterialMode = this->MaterialMode;
+}
+
+void UMeshEditingViewProperties::RestoreProperties(UInteractiveTool* RestoreToTool)
+{
+	UMeshEditingViewProperties* PropertyCache = GetPropertyCache<UMeshEditingViewProperties>();
+	this->bShowWireframe = PropertyCache->bShowWireframe;
+	this->MaterialMode = PropertyCache->MaterialMode;
+}
+
+
+
+
 #undef LOCTEXT_NAMESPACE
diff --git a/Engine/Plugins/Experimental/MeshModelingToolset/Source/MeshModelingTools/Public/DynamicMeshSculptTool.h b/Engine/Plugins/Experimental/MeshModelingToolset/Source/MeshModelingTools/Public/DynamicMeshSculptTool.h
--- a/Engine/Plugins/Experimental/MeshModelingToolset/Source/MeshModelingTools/Public/DynamicMeshSculptTool.h
+++ b/Engine/Plugins/Experimental/MeshModelingToolset/Source/MeshModelingTools/Public/DynamicMeshSculptTool.h
@@ -13,6 +13,7 @@
 #include "Drawing/ToolDataVisualizer.h"
 #include "Changes/ValueWatcher.h"
 #include "Gizmos/BrushStampIndicator.h"
+#include "Properties/MeshMaterialProperties.h"
 #include "DynamicMeshSculptTool.generated.h"
 
 
@@ -133,10 +134,6 @@ public:
 	UPROPERTY(EditAnywhere, Category = Sculpting, meta = (EditCondition = "PrimaryBrushType == EDynamicMeshSculptBrushType::Sculpt || PrimaryBrushType == EDynamicMeshSculptBrushType::SculptMax || PrimaryBrushType == EDynamicMeshSculptBrushType::Pinch" ))
 	bool bFreezeTarget;
 
-	/** Toggle drawing of wireframe overlay on/off */
-	UPROPERTY(EditAnywhere, Category = Sculpting)
-	bool bShowWireframe;
-
 	virtual void SaveProperties(UInteractiveTool* SaveFromTool) override;
 	virtual void RestoreProperties(UInteractiveTool* RestoreToTool) override;
 
@@ -145,7 +142,6 @@ public:
 
 
 
-
 UCLASS()
 class MESHMODELINGTOOLS_API UBrushRemeshProperties : public UInteractiveToolPropertySet
 {
@@ -234,6 +230,9 @@ public:
 	UPROPERTY()
 	UBrushRemeshProperties* RemeshProperties;
 
+	UPROPERTY()
+	UMeshEditingViewProperties* ViewProperties;
+
 
 public:
 
@@ -271,6 +270,8 @@ protected:
 	FDelegateHandle OnDynamicMeshComponentChangedHandle;
 
 	TValueWatcher<bool> ShowWireframeWatcher;
+	TValueWatcher<EMeshEditingMaterialModes> MaterialModeWatcher;
+	void UpdateMaterialMode(EMeshEditingMaterialModes NewMode);
 
 
 	FInterval1d BrushRelativeSizeRange;
diff --git a/Engine/Plugins/Experimental/MeshModelingToolset/Source/MeshModelingTools/Public/MeshSelectionTool.h b/Engine/Plugins/Experimental/MeshModelingToolset/Source/MeshModelingTools/Public/MeshSelectionTool.h
--- a/Engine/Plugins/Experimental/MeshModelingToolset/Source/MeshModelingTools/Public/MeshSelectionTool.h
+++ b/Engine/Plugins/Experimental/MeshModelingToolset/Source/MeshModelingTools/Public/MeshSelectionTool.h
@@ -172,7 +172,7 @@ public:
 	UPROPERTY(EditAnywhere, Category = Selection)
 	bool bHitBackFaces = true;
 
-	/** Toggle drawing of wireframe overlay on/off */
+	/** Toggle drawing of wireframe overlay on/off [Alt+W] */
 	UPROPERTY(EditAnywhere, Category = Selection)
 	bool bShowWireframe = true;
 
diff --git a/Engine/Plugins/Experimental/MeshModelingToolset/Source/MeshModelingTools/Public/Properties/MeshMaterialProperties.h b/Engine/Plugins/Experimental/MeshModelingToolset/Source/MeshModelingTools/Public/Properties/MeshMaterialProperties.h
--- a/Engine/Plugins/Experimental/MeshModelingToolset/Source/MeshModelingTools/Public/Properties/MeshMaterialProperties.h
+++ b/Engine/Plugins/Experimental/MeshModelingToolset/Source/MeshModelingTools/Public/Properties/MeshMaterialProperties.h
@@ -89,3 +89,33 @@ public:
 	void SetMaterialIfChanged(UMaterialInterface* OriginalMaterial, UMaterialInterface* CurrentMaterial, TFunctionRef<void(UMaterialInterface* Material)> SetMaterialFn);
 };
 
+
+
+
+
+
+UENUM()
+enum class EMeshEditingMaterialModes
+{
+	ExistingMaterial = 0,
+	MeshFocusMaterial = 1
+};
+
+
+UCLASS()
+class MESHMODELINGTOOLS_API UMeshEditingViewProperties : public UInteractiveToolPropertySet
+{
+	GENERATED_BODY()
+
+public:
+	/** Toggle drawing of wireframe overlay on/off [Alt+W] */
+	UPROPERTY(EditAnywhere, Category = ViewOptions)
+	bool bShowWireframe = false;
+
+	/** Set which material to use on object */
+	UPROPERTY(EditAnywhere, Category = ViewOptions)
+	EMeshEditingMaterialModes MaterialMode = EMeshEditingMaterialModes::MeshFocusMaterial;
+
+	virtual void SaveProperties(UInteractiveTool* SaveFromTool) override;
+	virtual void RestoreProperties(UInteractiveTool* RestoreToTool) override;
+};
diff --git a/Engine/Plugins/Experimental/MeshModelingToolset/Source/ModelingComponents/Private/ToolSetupUtil.cpp b/Engine/Plugins/Experimental/MeshModelingToolset/Source/ModelingComponents/Private/ToolSetupUtil.cpp
--- a/Engine/Plugins/Experimental/MeshModelingToolset/Source/ModelingComponents/Private/ToolSetupUtil.cpp
+++ b/Engine/Plugins/Experimental/MeshModelingToolset/Source/ModelingComponents/Private/ToolSetupUtil.cpp
@@ -10,7 +10,7 @@
 
 UMaterialInterface* ToolSetupUtil::GetDefaultMaterial(UInteractiveToolManager* ToolManager, UMaterialInterface* SourceMaterial)
 {
-	if (SourceMaterial == nullptr)
+	if (SourceMaterial == nullptr && ToolManager != nullptr)
 	{
 		return ToolManager->GetContextQueriesAPI()->GetStandardMaterial(EStandardToolContextMaterials::VertexColorMaterial);
 	}
@@ -22,7 +22,7 @@ UMaterialInterface* ToolSetupUtil::GetDefaultMaterial(UInteractiveToolManager* T
 UMaterialInterface* ToolSetupUtil::GetDefaultWorkingMaterial(UInteractiveToolManager* ToolManager)
 {
 	UMaterialInterface* Material = LoadObject<UMaterial>(nullptr, TEXT("/MeshModelingToolset/Materials/InProgressMaterial"));
-	if (Material == nullptr)
+	if (Material == nullptr && ToolManager != nullptr)
 	{
 		return ToolManager->GetContextQueriesAPI()->GetStandardMaterial(EStandardToolContextMaterials::VertexColorMaterial);
 	}
@@ -40,3 +40,25 @@ UMaterialInterface* ToolSetupUtil::GetDefaultBrushVolumeMaterial(UInteractiveToo
 	return Material;
 }
 
+
+
+UMaterialInterface* ToolSetupUtil::GetSculptMaterial1(UInteractiveToolManager* ToolManager)
+{
+	UMaterialInterface* Material = LoadObject<UMaterial>(nullptr, TEXT("/MeshModelingToolset/Materials/SculptMaterial"));
+	if (Material == nullptr && ToolManager != nullptr)
+	{
+		return ToolManager->GetContextQueriesAPI()->GetStandardMaterial(EStandardToolContextMaterials::VertexColorMaterial);
+	}
+	return Material;
+}
+
+
+UMaterialInterface* ToolSetupUtil::GetSelectionMaterial(UInteractiveToolManager* ToolManager)
+{
+	UMaterialInterface* Material = LoadObject<UMaterial>(nullptr, TEXT("/MeshModelingToolset/Materials/SelectionMaterial"));
+	if (Material == nullptr && ToolManager != nullptr)
+	{
+		return ToolManager->GetContextQueriesAPI()->GetStandardMaterial(EStandardToolContextMaterials::VertexColorMaterial);
+	}
+	return Material;
+}
diff --git a/Engine/Plugins/Experimental/MeshModelingToolset/Source/ModelingComponents/Public/Changes/ValueWatcher.h b/Engine/Plugins/Experimental/MeshModelingToolset/Source/ModelingComponents/Public/Changes/ValueWatcher.h
--- a/Engine/Plugins/Experimental/MeshModelingToolset/Source/ModelingComponents/Public/Changes/ValueWatcher.h
+++ b/Engine/Plugins/Experimental/MeshModelingToolset/Source/ModelingComponents/Public/Changes/ValueWatcher.h
@@ -34,7 +34,7 @@ public:
 
 	void CheckAndUpdate()
 	{
-		bool CurValue = GetValueFunc();
+		ValueType CurValue = GetValueFunc();
 		if (CurValue != CachedValue)
 		{
 			CachedValue = GetValueFunc();
diff --git a/Engine/Plugins/Experimental/MeshModelingToolset/Source/ModelingComponents/Public/ToolSetupUtil.h b/Engine/Plugins/Experimental/MeshModelingToolset/Source/ModelingComponents/Public/ToolSetupUtil.h
--- a/Engine/Plugins/Experimental/MeshModelingToolset/Source/ModelingComponents/Public/ToolSetupUtil.h
+++ b/Engine/Plugins/Experimental/MeshModelingToolset/Source/ModelingComponents/Public/ToolSetupUtil.h
@@ -28,4 +28,17 @@ namespace ToolSetupUtil
 	 * @return default material to use for brush volume indicators
 	 */
 	MODELINGCOMPONENTS_API UMaterialInterface* GetDefaultBrushVolumeMaterial(UInteractiveToolManager* ToolManager);
+
+
+	/**
+	 * @return Sculpt Material 1
+	 */
+	MODELINGCOMPONENTS_API UMaterialInterface* GetSculptMaterial1(UInteractiveToolManager* ToolManager);
+
+
+	/**
+	 * @return Selection Material 1
+	 */
+	MODELINGCOMPONENTS_API UMaterialInterface* GetSelectionMaterial(UInteractiveToolManager* ToolManager);
+
 }
\ No newline at end of file
diff --git a/Engine/Source/Editor/UnrealEd/Private/Tools/StandardToolModeCommands.cpp b/Engine/Source/Editor/UnrealEd/Private/Tools/StandardToolModeCommands.cpp
--- a/Engine/Source/Editor/UnrealEd/Private/Tools/StandardToolModeCommands.cpp
+++ b/Engine/Source/Editor/UnrealEd/Private/Tools/StandardToolModeCommands.cpp
@@ -27,7 +27,7 @@ void FStandardToolModeCommands::RegisterCommands()
 	Commands.Add(EStandardToolModeCommands::DecreaseBrushSize, DecreaseBrushSize);
 
 	TSharedPtr<FUICommandInfo> ToggleWireframe;
-	UI_COMMAND(ToggleWireframe, "Toggle Wireframe", "Toggle Wireframe overlay on/off", EUserInterfaceActionType::Button, FInputChord(EModifierKey::Control, EKeys::W));
+	UI_COMMAND(ToggleWireframe, "Toggle Wireframe", "Toggle Wireframe overlay on/off", EUserInterfaceActionType::Button, FInputChord(EModifierKey::Alt, EKeys::W));
 	Commands.Add(EStandardToolModeCommands::ToggleWireframe, ToggleWireframe);
 
 }
