fix hard edge detection for dynasculpt warning message

#rnx

diff --git a/Engine/Plugins/Experimental/MeshModelingToolset/Source/MeshModelingTools/Private/DynamicMeshSculptTool.cpp b/Engine/Plugins/Experimental/MeshModelingToolset/Source/MeshModelingTools/Private/DynamicMeshSculptTool.cpp
--- a/Engine/Plugins/Experimental/MeshModelingToolset/Source/MeshModelingTools/Private/DynamicMeshSculptTool.cpp
+++ b/Engine/Plugins/Experimental/MeshModelingToolset/Source/MeshModelingTools/Private/DynamicMeshSculptTool.cpp
@@ -1113,7 +1113,7 @@ void UDynamicMeshSculptTool::PrecomputeRemeshInfo()
 			bHaveUVSeams = bHaveUVSeams || Attribs->GetUVLayer(k)->HasInteriorSeamEdges();
 		}
 
-		bHaveNormalSeams = Attribs->PrimaryUV()->HasInteriorSeamEdges();
+		bHaveNormalSeams = Attribs->PrimaryNormals()->HasInteriorSeamEdges();
 	}
 }
 
