make hole cut tool actually cut a hole by default + change the gizmo

#rnx

diff --git a/Engine/Plugins/Experimental/MeshModelingToolset/Source/MeshModelingToolsEditorOnly/Private/PolygonOnMeshTool.cpp b/Engine/Plugins/Experimental/MeshModelingToolset/Source/MeshModelingToolsEditorOnly/Private/PolygonOnMeshTool.cpp
--- a/Engine/Plugins/Experimental/MeshModelingToolset/Source/MeshModelingToolsEditorOnly/Private/PolygonOnMeshTool.cpp
+++ b/Engine/Plugins/Experimental/MeshModelingToolset/Source/MeshModelingToolsEditorOnly/Private/PolygonOnMeshTool.cpp
@@ -15,7 +15,9 @@
 
 #include "InteractiveGizmoManager.h"
 
-#include "PositionPlaneGizmo.h"
+#include "BaseGizmos/GizmoComponents.h"
+#include "BaseGizmos/TransformGizmo.h"
+
 #include "AssetGenerationUtil.h"
 
 
@@ -55,7 +57,7 @@ UInteractiveTool* UPolygonOnMeshToolBuilder::BuildTool(const FToolBuilderState&
 
 UPolygonOnMeshToolProperties::UPolygonOnMeshToolProperties()
 {
-	PolygonOperation = EEmbeddedPolygonOpMethod::CutAndFill;
+	PolygonOperation = EEmbeddedPolygonOpMethod::CutThrough;
 	PolygonScale = 10;
 	// ExtrudeDistance = 10;
 	bDiscardAttributes = false;
@@ -103,14 +105,10 @@ void UPolygonOnMeshTool::Setup()
 	AddInputBehavior(ClickToSetPlaneBehavior);
 
 
-	UPositionPlaneGizmoBuilder* PositionPlaneGizmoBuilder = NewObject<UPositionPlaneGizmoBuilder>();
+	// create proxy and gizmo (but don't attach yet)
 	UInteractiveGizmoManager* GizmoManager = GetToolManager()->GetPairedGizmoManager();
-	GizmoManager->RegisterGizmoType(TEXT("PolygonPlaneGizmo"), PositionPlaneGizmoBuilder);
-	PositionPlaneGizmo = GizmoManager->CreateGizmo(TEXT("PolygonPlaneGizmo"), TEXT("PolygonPlaneGizmo"));
-	Cast<UPositionPlaneGizmo>(PositionPlaneGizmo)->OnPositionUpdatedFunc = [this](const FFrame3d& WorldFrame)
-	{
-		UpdatePlaneFromGizmo(WorldFrame);
-	};
+	PlaneTransformProxy = NewObject<UTransformProxy>(this);
+	PlaneTransformGizmo = GizmoManager->Create3AxisTransformGizmo(this);
 
 	BasicProperties = NewObject<UPolygonOnMeshToolProperties>(this, TEXT("Polygon On Mesh Settings"));
 	AdvancedProperties = NewObject<UPolygonOnMeshAdvancedProperties>(this, TEXT("Advanced Settings"));
@@ -123,10 +121,12 @@ void UPolygonOnMeshTool::Setup()
 	// initialize the PreviewMesh+BackgroundCompute object
 	UpdateNumPreviews();
 
+	// set initial cut plane (also attaches gizmo/proxy)
 	FVector DefaultOrigin, Extents;
 	ComponentTarget->GetOwnerActor()->GetActorBounds(false, DefaultOrigin, Extents);
 	SetPlaneFromWorldPos(DefaultOrigin, FVector::UpVector);
-
+	// hook up callback so further changes trigger recut
+	PlaneTransformProxy->OnTransformChanged.AddUObject(this, &UPolygonOnMeshTool::TransformChanged);
 
 
 	// Convert input mesh description to dynamic mesh
@@ -194,9 +194,7 @@ void UPolygonOnMeshTool::Shutdown(EToolShutdownType ShutdownType)
 		delete SetPointInWorldConnector;
 	}
 	UInteractiveGizmoManager* GizmoManager = GetToolManager()->GetPairedGizmoManager();
-	GizmoManager->DestroyGizmo(PositionPlaneGizmo);
-	PositionPlaneGizmo = nullptr;
-	GizmoManager->DeregisterGizmoType(TEXT("PolygonPlaneGizmo"));
+	GizmoManager->DestroyAllGizmosByOwner(this);
 }
 
 void UPolygonOnMeshTool::SetAssetAPI(IToolsContextAssetAPI* AssetAPIIn)
@@ -264,22 +262,12 @@ void UPolygonOnMeshTool::OnPropertyModified(UObject* PropertySet, UProperty* Pro
 
 
 
-void UPolygonOnMeshTool::SetPlaneFromWorldPos(const FVector& Position, const FVector& Normal)
-{
-	EmbedPolygonOrigin = Position;
-
-	FFrame3f PlaneFrame(Position, Normal);
-	EmbedPolygonOrientation = PlaneFrame.Rotation;
-
-	UPositionPlaneGizmo* Gizmo = Cast<UPositionPlaneGizmo>(PositionPlaneGizmo);
-	Gizmo->ExternalUpdatePosition(EmbedPolygonOrigin, EmbedPolygonOrientation, false);
-}
 
-
-void UPolygonOnMeshTool::UpdatePlaneFromGizmo(const FFrame3d& WorldPosition)
+void UPolygonOnMeshTool::TransformChanged(UTransformProxy* Proxy, FTransform Transform)
 {
-	EmbedPolygonOrientation = WorldPosition.Rotation;
-	EmbedPolygonOrigin = (FVector)WorldPosition.Origin;
+	// TODO: if multi-select is re-enabled, only invalidate the preview that actually needs it?
+	EmbedPolygonOrientation = Transform.GetRotation();
+	EmbedPolygonOrigin = (FVector)Transform.GetTranslation();
 	for (UMeshOpPreviewWithBackgroundCompute* Preview : Previews)
 	{
 		Preview->InvalidateResult();
@@ -287,6 +275,18 @@ void UPolygonOnMeshTool::UpdatePlaneFromGizmo(const FFrame3d& WorldPosition)
 }
 
 
+void UPolygonOnMeshTool::SetPlaneFromWorldPos(const FVector& Position, const FVector& Normal)
+{
+	EmbedPolygonOrigin = Position;
+
+	FFrame3f CutPlane(Position, Normal);
+	EmbedPolygonOrientation = CutPlane.Rotation;
+
+	PlaneTransformGizmo->SetActiveTarget(PlaneTransformProxy);
+	PlaneTransformGizmo->SetNewGizmoTransform(CutPlane.ToFTransform());
+}
+
+
 bool UPolygonOnMeshTool::HasAccept() const
 {
 	return true;
diff --git a/Engine/Plugins/Experimental/MeshModelingToolset/Source/MeshModelingToolsEditorOnly/Public/PolygonOnMeshTool.h b/Engine/Plugins/Experimental/MeshModelingToolset/Source/MeshModelingToolsEditorOnly/Public/PolygonOnMeshTool.h
--- a/Engine/Plugins/Experimental/MeshModelingToolset/Source/MeshModelingToolsEditorOnly/Public/PolygonOnMeshTool.h
+++ b/Engine/Plugins/Experimental/MeshModelingToolset/Source/MeshModelingToolsEditorOnly/Public/PolygonOnMeshTool.h
@@ -17,7 +17,8 @@
 // predeclarations
 struct FMeshDescription;
 class USimpleDynamicMeshComponent;
-
+class UTransformGizmo;
+class UTransformProxy;
 
 
 
@@ -164,11 +165,13 @@ protected:
 
 	FViewCameraState CameraState;
 
-	// plane control
 	UPROPERTY()
-	UInteractiveGizmo* PositionPlaneGizmo;
+	UTransformGizmo* PlaneTransformGizmo;
+
+	UPROPERTY()
+	UTransformProxy* PlaneTransformProxy;
 
-	void UpdatePlaneFromGizmo(const FFrame3d& WorldPosition);
+	void TransformChanged(UTransformProxy* Proxy, FTransform Transform);
 
 	void UpdateNumPreviews();
 
