transition to new StaticMeshAttributes implementation
#rnx

diff --git a/Engine/Plugins/Experimental/MeshModelingToolset/Source/MeshModelingToolsEditorOnly/Private/AttributeEditorTool.cpp b/Engine/Plugins/Experimental/MeshModelingToolset/Source/MeshModelingToolsEditorOnly/Private/AttributeEditorTool.cpp
--- a/Engine/Plugins/Experimental/MeshModelingToolset/Source/MeshModelingToolsEditorOnly/Private/AttributeEditorTool.cpp
+++ b/Engine/Plugins/Experimental/MeshModelingToolset/Source/MeshModelingToolsEditorOnly/Private/AttributeEditorTool.cpp
@@ -8,7 +8,7 @@
 #include "MathUtil.h"
 
 #include "MeshDescription.h"
-#include "MeshAttributes.h"
+#include "StaticMeshAttributes.h"
 #include "MeshDescriptionOperations.h"
 
 
diff --git a/Engine/Plugins/Experimental/MeshModelingToolset/Source/ModelingOperatorsEditorOnly/Private/CompositionOps/VoxelBooleanMeshesOp.cpp b/Engine/Plugins/Experimental/MeshModelingToolset/Source/ModelingOperatorsEditorOnly/Private/CompositionOps/VoxelBooleanMeshesOp.cpp
--- a/Engine/Plugins/Experimental/MeshModelingToolset/Source/ModelingOperatorsEditorOnly/Private/CompositionOps/VoxelBooleanMeshesOp.cpp
+++ b/Engine/Plugins/Experimental/MeshModelingToolset/Source/ModelingOperatorsEditorOnly/Private/CompositionOps/VoxelBooleanMeshesOp.cpp
@@ -7,6 +7,7 @@
 #include "Engine/StaticMesh.h"
 
 #include "MeshDescription.h"
+#include "StaticMeshAttributes.h"
 #include "MeshDescriptionToDynamicMesh.h"
 #include "MeshSimplification.h"
 #include "MeshConstraintsUtil.h"
@@ -14,7 +15,8 @@
 void FVoxelBooleanMeshesOp::CalculateResult(FProgressCancel* Progress)
 {
 	FMeshDescription ResultMeshDescription;
-	UStaticMesh::RegisterMeshAttributes(ResultMeshDescription);
+	FStaticMeshAttributes Attributes(ResultMeshDescription);
+	Attributes.Register();
 
 	if (VoxelCount < 16)
 	{
diff --git a/Engine/Plugins/Experimental/MeshModelingToolset/Source/ModelingOperatorsEditorOnly/Private/CompositionOps/VoxelMergeMeshesOp.cpp b/Engine/Plugins/Experimental/MeshModelingToolset/Source/ModelingOperatorsEditorOnly/Private/CompositionOps/VoxelMergeMeshesOp.cpp
--- a/Engine/Plugins/Experimental/MeshModelingToolset/Source/ModelingOperatorsEditorOnly/Private/CompositionOps/VoxelMergeMeshesOp.cpp
+++ b/Engine/Plugins/Experimental/MeshModelingToolset/Source/ModelingOperatorsEditorOnly/Private/CompositionOps/VoxelMergeMeshesOp.cpp
@@ -7,6 +7,7 @@
 #include "Engine/StaticMesh.h"
 
 #include "MeshDescription.h"
+#include "StaticMeshAttributes.h"
 #include "MeshDescriptionToDynamicMesh.h"
 #include "MeshSimplification.h"
 #include "MeshConstraintsUtil.h"
@@ -15,7 +16,8 @@
 void FVoxelMergeMeshesOp::CalculateResult(FProgressCancel* Progress)
 {
 	FMeshDescription MergedMeshesDescription;
-	UStaticMesh::RegisterMeshAttributes(MergedMeshesDescription);
+	FStaticMeshAttributes Attributes(MergedMeshesDescription);
+	Attributes.Register();
 
 	// Use the world space bounding box of each mesh to compute the voxel size
 	VoxelSizeD = ComputeVoxelSize();
