Remove some unnecesary use of TUniquePtr

diff --git a/Engine/Plugins/Experimental/MeshModelingToolset/Source/MeshModelingTools/Private/EditNormalsTool.cpp b/Engine/Plugins/Experimental/MeshModelingToolset/Source/MeshModelingTools/Private/EditNormalsTool.cpp
--- a/Engine/Plugins/Experimental/MeshModelingToolset/Source/MeshModelingTools/Private/EditNormalsTool.cpp
+++ b/Engine/Plugins/Experimental/MeshModelingToolset/Source/MeshModelingTools/Private/EditNormalsTool.cpp
@@ -162,10 +162,10 @@ void UEditNormalsTool::Shutdown(EToolShutdownType ShutdownType)
 		ComponentTarget->SetOwnerVisibility(true);
 	}
 
-	TArray<TUniquePtr<FDynamicMeshOpResult>> Results;
+	TArray<FDynamicMeshOpResult> Results;
 	for (UMeshOpPreviewWithBackgroundCompute* Preview : Previews)
 	{
-		Results.Emplace(Preview->Shutdown());
+		Results.Add(Preview->Shutdown());
 	}
 	if (ShutdownType == EToolShutdownType::Accept)
 	{
@@ -253,7 +253,7 @@ bool UEditNormalsTool::CanAccept() const
 }
 
 
-void UEditNormalsTool::GenerateAsset(const TArray<TUniquePtr<FDynamicMeshOpResult>>& Results)
+void UEditNormalsTool::GenerateAsset(const TArray<FDynamicMeshOpResult>& Results)
 {
 	GetToolManager()->BeginUndoTransaction(LOCTEXT("EditNormalsToolTransactionName", "Edit Normals Tool"));
 
@@ -261,7 +261,7 @@ void UEditNormalsTool::GenerateAsset(const TArray<TUniquePtr<FDynamicMeshOpResul
 	
 	for (int32 ComponentIdx = 0; ComponentIdx < ComponentTargets.Num(); ComponentIdx++)
 	{
-		check(Results[ComponentIdx]->Mesh.Get() != nullptr);
+		check(Results[ComponentIdx].Mesh.Get() != nullptr);
 		ComponentTargets[ComponentIdx]->CommitMesh([&Results, &ComponentIdx, this](FMeshDescription* MeshDescription)
 		{
 			FDynamicMeshToMeshDescription Converter;
@@ -269,12 +269,12 @@ void UEditNormalsTool::GenerateAsset(const TArray<TUniquePtr<FDynamicMeshOpResul
 			if (BasicProperties->WillTopologyChange())
 			{
 				// full conversion if normal topology changed or faces were inverted
-				Converter.Convert(Results[ComponentIdx]->Mesh.Get(), *MeshDescription);
+				Converter.Convert(Results[ComponentIdx].Mesh.Get(), *MeshDescription);
 			}
 			else
 			{
 				// otherwise just copy attributes
-				Converter.UpdateAttributes(Results[ComponentIdx]->Mesh.Get(), *MeshDescription, true, false);
+				Converter.UpdateAttributes(Results[ComponentIdx].Mesh.Get(), *MeshDescription, true, false);
 			}
 		});
 	}
diff --git a/Engine/Plugins/Experimental/MeshModelingToolset/Source/MeshModelingTools/Private/PlaneCutTool.cpp b/Engine/Plugins/Experimental/MeshModelingToolset/Source/MeshModelingTools/Private/PlaneCutTool.cpp
--- a/Engine/Plugins/Experimental/MeshModelingToolset/Source/MeshModelingTools/Private/PlaneCutTool.cpp
+++ b/Engine/Plugins/Experimental/MeshModelingToolset/Source/MeshModelingTools/Private/PlaneCutTool.cpp
@@ -174,7 +174,7 @@ void UPlaneCutTool::Shutdown(EToolShutdownType ShutdownType)
 	// Restore (unhide) the source meshes
 	ComponentTarget->SetOwnerVisibility(true);
 
-	TArray<TUniquePtr<FDynamicMeshOpResult>> Results;
+	TArray<FDynamicMeshOpResult> Results;
 	for (UMeshOpPreviewWithBackgroundCompute* Preview : Previews)
 	{
 		Results.Emplace(Preview->Shutdown());
@@ -321,7 +321,7 @@ bool UPlaneCutTool::CanAccept() const
 }
 
 
-void UPlaneCutTool::GenerateAsset(const TArray<TUniquePtr<FDynamicMeshOpResult>>& Results)
+void UPlaneCutTool::GenerateAsset(const TArray<FDynamicMeshOpResult>& Results)
 {
 	GetToolManager()->BeginUndoTransaction(LOCTEXT("PlaneCutToolTransactionName", "Plane Cut Tool"));
 	
@@ -329,11 +329,11 @@ void UPlaneCutTool::GenerateAsset(const TArray<TUniquePtr<FDynamicMeshOpResult>>
 	// currently in-place replaces the first half, and adds a new actor for the second half (if it was generated)
 	// TODO: options to support other choices re what should be a new actor
 	check(Results.Num() > 0);
-	check(Results[0]->Mesh.Get() != nullptr);
+	check(Results[0].Mesh.Get() != nullptr);
 	ComponentTarget->CommitMesh([&Results](FMeshDescription* MeshDescription)
 	{
 		FDynamicMeshToMeshDescription Converter;
-		Converter.Convert(Results[0]->Mesh.Get(), *MeshDescription);
+		Converter.Convert(Results[0].Mesh.Get(), *MeshDescription);
 	});
 
 
@@ -345,11 +345,11 @@ void UPlaneCutTool::GenerateAsset(const TArray<TUniquePtr<FDynamicMeshOpResult>>
 		NewSelection.ModificationType = ESelectedObjectsModificationType::Replace;
 		NewSelection.Actors.Add(ComponentTarget->GetOwnerActor());
 
-		check(Results[1]->Mesh.Get() != nullptr);
+		check(Results[1].Mesh.Get() != nullptr);
 		// TODO: copy over material?
 		AActor* NewActor = AssetGenerationUtil::GenerateStaticMeshActor(
 			AssetAPI, TargetWorld,
-			Results[1]->Mesh.Get(), Results[1]->Transform, TEXT("PlaneCutOtherHalf"),
+			Results[1].Mesh.Get(), Results[1].Transform, TEXT("PlaneCutOtherHalf"),
 			AssetGenerationUtil::GetDefaultAutoGeneratedAssetPath());
 		NewSelection.Actors.Add(NewActor);
 		GetToolManager()->RequestSelectionChange(NewSelection);
diff --git a/Engine/Plugins/Experimental/MeshModelingToolset/Source/MeshModelingTools/Private/SmoothMeshTool.cpp b/Engine/Plugins/Experimental/MeshModelingToolset/Source/MeshModelingTools/Private/SmoothMeshTool.cpp
--- a/Engine/Plugins/Experimental/MeshModelingToolset/Source/MeshModelingTools/Private/SmoothMeshTool.cpp
+++ b/Engine/Plugins/Experimental/MeshModelingToolset/Source/MeshModelingTools/Private/SmoothMeshTool.cpp
@@ -114,13 +114,13 @@ void USmoothMeshTool::Shutdown(EToolShutdownType ShutdownType)
 
 	if (Preview != nullptr)
 	{
-		TUniquePtr<FDynamicMeshOpResult> Result = Preview->Shutdown();
+		FDynamicMeshOpResult Result = Preview->Shutdown();
 		if (ShutdownType == EToolShutdownType::Accept)
 		{
 			
 			GetToolManager()->BeginUndoTransaction(LOCTEXT("SmoothMeshToolTransactionName", "Smooth Mesh"));
 
-			FDynamicMesh3* DynamicMeshResult = Result->Mesh.Get();
+			FDynamicMesh3* DynamicMeshResult = Result.Mesh.Get();
 			check(DynamicMeshResult != nullptr);
 
 			ComponentTarget->CommitMesh([DynamicMeshResult](FMeshDescription* MeshDescription)
diff --git a/Engine/Plugins/Experimental/MeshModelingToolset/Source/MeshModelingTools/Private/UVProjectionTool.cpp b/Engine/Plugins/Experimental/MeshModelingToolset/Source/MeshModelingTools/Private/UVProjectionTool.cpp
--- a/Engine/Plugins/Experimental/MeshModelingToolset/Source/MeshModelingTools/Private/UVProjectionTool.cpp
+++ b/Engine/Plugins/Experimental/MeshModelingToolset/Source/MeshModelingTools/Private/UVProjectionTool.cpp
@@ -192,7 +192,7 @@ void UUVProjectionTool::Shutdown(EToolShutdownType ShutdownType)
 		ComponentTarget->SetOwnerVisibility(true);
 	}
 
-	TArray<TUniquePtr<FDynamicMeshOpResult>> Results;
+	TArray<FDynamicMeshOpResult> Results;
 	for (UMeshOpPreviewWithBackgroundCompute* Preview : Previews)
 	{
 		Results.Emplace(Preview->Shutdown());
@@ -332,7 +332,7 @@ bool UUVProjectionTool::CanAccept() const
 }
 
 
-void UUVProjectionTool::GenerateAsset(const TArray<TUniquePtr<FDynamicMeshOpResult>>& Results)
+void UUVProjectionTool::GenerateAsset(const TArray<FDynamicMeshOpResult>& Results)
 {
 	GetToolManager()->BeginUndoTransaction(LOCTEXT("UVProjectionToolTransactionName", "UV Projection Tool"));
 
@@ -340,14 +340,14 @@ void UUVProjectionTool::GenerateAsset(const TArray<TUniquePtr<FDynamicMeshOpResu
 	
 	for (int32 ComponentIdx = 0; ComponentIdx < ComponentTargets.Num(); ComponentIdx++)
 	{
-		check(Results[ComponentIdx]->Mesh.Get() != nullptr);
+		check(Results[ComponentIdx].Mesh.Get() != nullptr);
 		ComponentTargets[ComponentIdx]->CommitMesh([&Results, &ComponentIdx, this](FMeshDescription* MeshDescription)
 		{
 			FDynamicMeshToMeshDescription Converter;
 			//if (??) // TODO check if UV topology changed?  or remove all traces of this if statement (may be safe to assume it almost always changes w/ a uv projection op)
 			{
 				// full conversion if topology changed
-				Converter.Convert(Results[ComponentIdx]->Mesh.Get(), *MeshDescription);
+				Converter.Convert(Results[ComponentIdx].Mesh.Get(), *MeshDescription);
 			}
 			//else
 			//{
diff --git a/Engine/Plugins/Experimental/MeshModelingToolset/Source/MeshModelingTools/Public/EditNormalsTool.h b/Engine/Plugins/Experimental/MeshModelingToolset/Source/MeshModelingTools/Public/EditNormalsTool.h
--- a/Engine/Plugins/Experimental/MeshModelingToolset/Source/MeshModelingTools/Public/EditNormalsTool.h
+++ b/Engine/Plugins/Experimental/MeshModelingToolset/Source/MeshModelingTools/Public/EditNormalsTool.h
@@ -177,5 +177,5 @@ protected:
 
 	void UpdateNumPreviews();
 
-	void GenerateAsset(const TArray<TUniquePtr<FDynamicMeshOpResult>>& Results);
+	void GenerateAsset(const TArray<FDynamicMeshOpResult>& Results);
 };
diff --git a/Engine/Plugins/Experimental/MeshModelingToolset/Source/MeshModelingTools/Public/PlaneCutTool.h b/Engine/Plugins/Experimental/MeshModelingToolset/Source/MeshModelingTools/Public/PlaneCutTool.h
--- a/Engine/Plugins/Experimental/MeshModelingToolset/Source/MeshModelingTools/Public/PlaneCutTool.h
+++ b/Engine/Plugins/Experimental/MeshModelingToolset/Source/MeshModelingTools/Public/PlaneCutTool.h
@@ -167,5 +167,5 @@ protected:
 
 	virtual void SetCutPlaneFromWorldPos(const FVector& Position, const FVector& Normal);
 
-	void GenerateAsset(const TArray<TUniquePtr<FDynamicMeshOpResult>>& Results);
+	void GenerateAsset(const TArray<FDynamicMeshOpResult>& Results);
 };
diff --git a/Engine/Plugins/Experimental/MeshModelingToolset/Source/MeshModelingTools/Public/UVProjectionTool.h b/Engine/Plugins/Experimental/MeshModelingToolset/Source/MeshModelingTools/Public/UVProjectionTool.h
--- a/Engine/Plugins/Experimental/MeshModelingToolset/Source/MeshModelingTools/Public/UVProjectionTool.h
+++ b/Engine/Plugins/Experimental/MeshModelingToolset/Source/MeshModelingTools/Public/UVProjectionTool.h
@@ -190,7 +190,7 @@ protected:
 
 	void UpdateNumPreviews();
 
-	void GenerateAsset(const TArray<TUniquePtr<FDynamicMeshOpResult>>& Results);
+	void GenerateAsset(const TArray<FDynamicMeshOpResult>& Results);
 
 	void TransformChanged(UTransformProxy* Proxy, FTransform Transform);
 };
diff --git a/Engine/Plugins/Experimental/MeshModelingToolset/Source/MeshModelingToolsEditorOnly/Private/MergeMeshesTool.cpp b/Engine/Plugins/Experimental/MeshModelingToolset/Source/MeshModelingToolsEditorOnly/Private/MergeMeshesTool.cpp
--- a/Engine/Plugins/Experimental/MeshModelingToolset/Source/MeshModelingToolsEditorOnly/Private/MergeMeshesTool.cpp
+++ b/Engine/Plugins/Experimental/MeshModelingToolset/Source/MeshModelingToolsEditorOnly/Private/MergeMeshesTool.cpp
@@ -114,7 +114,7 @@ void UMergeMeshesTool::Setup()
 
 void UMergeMeshesTool::Shutdown(EToolShutdownType ShutdownType)
 {
-	TUniquePtr<FDynamicMeshOpResult> Result = Preview->Shutdown();
+	FDynamicMeshOpResult Result = Preview->Shutdown();
 	if (ShutdownType == EToolShutdownType::Accept)
 	{
 		GetToolManager()->BeginUndoTransaction(LOCTEXT("MergeMeshes", "Merge Meshes"));
@@ -228,15 +228,15 @@ void UMergeMeshesTool::CreateLowQualityPreview()
 }
 
 
-void UMergeMeshesTool::GenerateAsset(const TUniquePtr<FDynamicMeshOpResult>& Result)
+void UMergeMeshesTool::GenerateAsset(const FDynamicMeshOpResult& Result)
 {
-	check(Result->Mesh.Get() != nullptr);
+	check(Result.Mesh.Get() != nullptr);
 
 	
 
 	AActor* NewActor = AssetGenerationUtil::GenerateStaticMeshActor(
 		AssetAPI, TargetWorld,
-		Result->Mesh.Get(), Result->Transform, TEXT("MergedMesh"),
+		Result.Mesh.Get(), Result.Transform, TEXT("MergedMesh"),
 		AssetGenerationUtil::GetDefaultAutoGeneratedAssetPath());
 
 	// select newly-created object
diff --git a/Engine/Plugins/Experimental/MeshModelingToolset/Source/MeshModelingToolsEditorOnly/Private/ParameterizeMeshTool.cpp b/Engine/Plugins/Experimental/MeshModelingToolset/Source/MeshModelingToolsEditorOnly/Private/ParameterizeMeshTool.cpp
--- a/Engine/Plugins/Experimental/MeshModelingToolset/Source/MeshModelingToolsEditorOnly/Private/ParameterizeMeshTool.cpp
+++ b/Engine/Plugins/Experimental/MeshModelingToolset/Source/MeshModelingToolsEditorOnly/Private/ParameterizeMeshTool.cpp
@@ -134,10 +134,10 @@ void UParameterizeMeshTool::OnPropertyModified(UObject* PropertySet, UProperty*
 void UParameterizeMeshTool::Shutdown(EToolShutdownType ShutdownType)
 {
 
-	TUniquePtr<FDynamicMeshOpResult> Result = Preview->Shutdown();
+	FDynamicMeshOpResult Result = Preview->Shutdown();
 	if (ShutdownType == EToolShutdownType::Accept)
 	{
-		FDynamicMesh3* DynamicMeshResult = Result->Mesh.Get();
+		FDynamicMesh3* DynamicMeshResult = Result.Mesh.Get();
 		check(DynamicMeshResult != nullptr);
 		GetToolManager()->BeginUndoTransaction(LOCTEXT("ParameterizeMesh", "Parameterize Mesh"));
 
diff --git a/Engine/Plugins/Experimental/MeshModelingToolset/Source/MeshModelingToolsEditorOnly/Private/PolygonOnMeshTool.cpp b/Engine/Plugins/Experimental/MeshModelingToolset/Source/MeshModelingToolsEditorOnly/Private/PolygonOnMeshTool.cpp
--- a/Engine/Plugins/Experimental/MeshModelingToolset/Source/MeshModelingToolsEditorOnly/Private/PolygonOnMeshTool.cpp
+++ b/Engine/Plugins/Experimental/MeshModelingToolset/Source/MeshModelingToolsEditorOnly/Private/PolygonOnMeshTool.cpp
@@ -179,7 +179,7 @@ void UPolygonOnMeshTool::Shutdown(EToolShutdownType ShutdownType)
 	// Restore (unhide) the source meshes
 	ComponentTarget->SetOwnerVisibility(true);
 
-	TArray<TUniquePtr<FDynamicMeshOpResult>> Results;
+	TArray<FDynamicMeshOpResult> Results;
 	for (UMeshOpPreviewWithBackgroundCompute* Preview : Previews)
 	{
 		Results.Emplace(Preview->Shutdown());
@@ -305,7 +305,7 @@ bool UPolygonOnMeshTool::CanAccept() const
 }
 
 
-void UPolygonOnMeshTool::GenerateAsset(const TArray<TUniquePtr<FDynamicMeshOpResult>>& Results)
+void UPolygonOnMeshTool::GenerateAsset(const TArray<FDynamicMeshOpResult>& Results)
 {
 	GetToolManager()->BeginUndoTransaction(LOCTEXT("PolygonOnMeshToolTransactionName", "Polygon On Mesh Tool"));
 	
@@ -313,11 +313,11 @@ void UPolygonOnMeshTool::GenerateAsset(const TArray<TUniquePtr<FDynamicMeshOpRes
 	// first preview always replaces, and any subsequent previews can add new actors
 	// TODO: options to support other choices re what should be a new actor
 	check(Results.Num() > 0);
-	check(Results[0]->Mesh.Get() != nullptr);
+	check(Results[0].Mesh.Get() != nullptr);
 	ComponentTarget->CommitMesh([&Results](FMeshDescription* MeshDescription)
 	{
 		FDynamicMeshToMeshDescription Converter;
-		Converter.Convert(Results[0]->Mesh.Get(), *MeshDescription);
+		Converter.Convert(Results[0].Mesh.Get(), *MeshDescription);
 	});
 
 	// currently nothing generates more than the single result for this tool
diff --git a/Engine/Plugins/Experimental/MeshModelingToolset/Source/MeshModelingToolsEditorOnly/Private/VoxelCSGMeshesTool.cpp b/Engine/Plugins/Experimental/MeshModelingToolset/Source/MeshModelingToolsEditorOnly/Private/VoxelCSGMeshesTool.cpp
--- a/Engine/Plugins/Experimental/MeshModelingToolset/Source/MeshModelingToolsEditorOnly/Private/VoxelCSGMeshesTool.cpp
+++ b/Engine/Plugins/Experimental/MeshModelingToolset/Source/MeshModelingToolsEditorOnly/Private/VoxelCSGMeshesTool.cpp
@@ -112,7 +112,7 @@ void UVoxelCSGMeshesTool::Setup()
 
 void UVoxelCSGMeshesTool::Shutdown(EToolShutdownType ShutdownType)
 {
-	TUniquePtr<FDynamicMeshOpResult> Result = Preview->Shutdown();
+	FDynamicMeshOpResult Result = Preview->Shutdown();
 	if (ShutdownType == EToolShutdownType::Accept)
 	{
 		GetToolManager()->BeginUndoTransaction(LOCTEXT("VoxelCSGMeshes", "Boolean Meshes"));
@@ -218,15 +218,15 @@ void UVoxelCSGMeshesTool::CreateLowQualityPreview()
 	Preview->SetVisibility(true);
 }
 
-void UVoxelCSGMeshesTool::GenerateAsset(const TUniquePtr<FDynamicMeshOpResult>& Result)
+void UVoxelCSGMeshesTool::GenerateAsset(const FDynamicMeshOpResult& Result)
 {
-	check(Result->Mesh.Get() != nullptr);
+	check(Result.Mesh.Get() != nullptr);
 
 	//GetToolManager()->BeginUndoTransaction(LOCTEXT("VoxelCSGMeshes", "Boolean Meshes"));
 
 	AActor* NewActor = AssetGenerationUtil::GenerateStaticMeshActor(
 		AssetAPI, TargetWorld,
-		Result->Mesh.Get(), Result->Transform, TEXT("CSGMesh"),
+		Result.Mesh.Get(), Result.Transform, TEXT("CSGMesh"),
 		AssetGenerationUtil::GetDefaultAutoGeneratedAssetPath());
 
 	// select newly-created object
diff --git a/Engine/Plugins/Experimental/MeshModelingToolset/Source/MeshModelingToolsEditorOnly/Public/MergeMeshesTool.h b/Engine/Plugins/Experimental/MeshModelingToolset/Source/MeshModelingToolsEditorOnly/Public/MergeMeshesTool.h
--- a/Engine/Plugins/Experimental/MeshModelingToolset/Source/MeshModelingToolsEditorOnly/Public/MergeMeshesTool.h
+++ b/Engine/Plugins/Experimental/MeshModelingToolset/Source/MeshModelingToolsEditorOnly/Public/MergeMeshesTool.h
@@ -114,5 +114,5 @@ protected:
 	/** quickly generate a low-quality result for display while the actual result is being computed. */
 	void CreateLowQualityPreview();
 
-	void GenerateAsset(const TUniquePtr<FDynamicMeshOpResult>& Result);
+	void GenerateAsset(const FDynamicMeshOpResult& Result);
 };
diff --git a/Engine/Plugins/Experimental/MeshModelingToolset/Source/MeshModelingToolsEditorOnly/Public/PolygonOnMeshTool.h b/Engine/Plugins/Experimental/MeshModelingToolset/Source/MeshModelingToolsEditorOnly/Public/PolygonOnMeshTool.h
--- a/Engine/Plugins/Experimental/MeshModelingToolset/Source/MeshModelingToolsEditorOnly/Public/PolygonOnMeshTool.h
+++ b/Engine/Plugins/Experimental/MeshModelingToolset/Source/MeshModelingToolsEditorOnly/Public/PolygonOnMeshTool.h
@@ -179,5 +179,5 @@ protected:
 
 	virtual void SetPlaneFromWorldPos(const FVector& Position, const FVector& Normal);
 
-	void GenerateAsset(const TArray<TUniquePtr<FDynamicMeshOpResult>>& Results);
+	void GenerateAsset(const TArray<FDynamicMeshOpResult>& Results);
 };
diff --git a/Engine/Plugins/Experimental/MeshModelingToolset/Source/MeshModelingToolsEditorOnly/Public/VoxelCSGMeshesTool.h b/Engine/Plugins/Experimental/MeshModelingToolset/Source/MeshModelingToolsEditorOnly/Public/VoxelCSGMeshesTool.h
--- a/Engine/Plugins/Experimental/MeshModelingToolset/Source/MeshModelingToolsEditorOnly/Public/VoxelCSGMeshesTool.h
+++ b/Engine/Plugins/Experimental/MeshModelingToolset/Source/MeshModelingToolsEditorOnly/Public/VoxelCSGMeshesTool.h
@@ -139,5 +139,5 @@ protected:
 	/** quickly generate a low-quality result for display while the actual result is being computed. */
 	void CreateLowQualityPreview();
 
-	void GenerateAsset(const TUniquePtr<FDynamicMeshOpResult>& Result);
+	void GenerateAsset(const FDynamicMeshOpResult& Result);
 };
diff --git a/Engine/Plugins/Experimental/MeshModelingToolset/Source/ModelingComponents/Private/MeshOpPreviewHelpers.cpp b/Engine/Plugins/Experimental/MeshModelingToolset/Source/ModelingComponents/Private/MeshOpPreviewHelpers.cpp
--- a/Engine/Plugins/Experimental/MeshModelingToolset/Source/ModelingComponents/Private/MeshOpPreviewHelpers.cpp
+++ b/Engine/Plugins/Experimental/MeshModelingToolset/Source/ModelingComponents/Private/MeshOpPreviewHelpers.cpp
@@ -12,14 +12,14 @@ void UMeshOpPreviewWithBackgroundCompute::Setup(UWorld* InWorld, IDynamicMeshOpe
 	bResultValid = false;
 }
 
-TUniquePtr<FDynamicMeshOpResult> UMeshOpPreviewWithBackgroundCompute::Shutdown()
+FDynamicMeshOpResult UMeshOpPreviewWithBackgroundCompute::Shutdown()
 {
 	BackgroundCompute->CancelActiveCompute();
 
 
-	TUniquePtr<FDynamicMeshOpResult> Result = MakeUnique<FDynamicMeshOpResult>();
-	Result->Mesh = PreviewMesh->ExtractPreviewMesh();
-	Result->Transform = FTransform3d(PreviewMesh->GetTransform());
+	FDynamicMeshOpResult Result{};
+	Result.Mesh = PreviewMesh->ExtractPreviewMesh();
+	Result.Transform = FTransform3d(PreviewMesh->GetTransform());
 
 	PreviewMesh->SetVisible(false);
 	PreviewMesh->Disconnect();
diff --git a/Engine/Plugins/Experimental/MeshModelingToolset/Source/ModelingComponents/Public/MeshOpPreviewHelpers.h b/Engine/Plugins/Experimental/MeshModelingToolset/Source/ModelingComponents/Public/MeshOpPreviewHelpers.h
--- a/Engine/Plugins/Experimental/MeshModelingToolset/Source/ModelingComponents/Public/MeshOpPreviewHelpers.h
+++ b/Engine/Plugins/Experimental/MeshModelingToolset/Source/ModelingComponents/Public/MeshOpPreviewHelpers.h
@@ -14,18 +14,7 @@
  * FBackgroundDynamicMeshComputeSource is an instantiation of the TBackgroundModelingComputeSource
  * template for FDynamicMeshOperator / IDynamicMeshOperatorFactory
  */
-class MODELINGCOMPONENTS_API FBackgroundDynamicMeshComputeSource
-	: public TBackgroundModelingComputeSource<FDynamicMeshOperator, IDynamicMeshOperatorFactory>
-{
-public:
-	using BaseType = TBackgroundModelingComputeSource<FDynamicMeshOperator, IDynamicMeshOperatorFactory>;
-
-	FBackgroundDynamicMeshComputeSource(IDynamicMeshOperatorFactory* OperatorSourceIn)
-		: BaseType(OperatorSourceIn)
-	{
-	}
-};
-
+using FBackgroundDynamicMeshComputeSource = TBackgroundModelingComputeSource<FDynamicMeshOperator, IDynamicMeshOperatorFactory>;
 
 
 /**
@@ -74,7 +63,7 @@ public:
 	/**
 	 * Terminate any active computation and return the current Preview Mesh/Transform
 	 */
-	TUniquePtr<FDynamicMeshOpResult> Shutdown();
+	FDynamicMeshOpResult Shutdown();
 
 	/**
 	* Terminate any active computation without returning anything
