pass through the DepthBias parameter in ToolDataVisualizer functions instead of 0.0f

#rnx

diff --git a/Engine/Plugins/Experimental/MeshModelingToolset/Source/ModelingComponents/Private/Drawing/ToolDataVisualizer.cpp b/Engine/Plugins/Experimental/MeshModelingToolset/Source/ModelingComponents/Private/Drawing/ToolDataVisualizer.cpp
--- a/Engine/Plugins/Experimental/MeshModelingToolset/Source/ModelingComponents/Private/Drawing/ToolDataVisualizer.cpp
+++ b/Engine/Plugins/Experimental/MeshModelingToolset/Source/ModelingComponents/Private/Drawing/ToolDataVisualizer.cpp
@@ -74,7 +74,7 @@ void FToolDataVisualizer::InternalDrawTransformedLine(const FVector& A, const FV
 {
 	CurrentPDI->DrawLine(A, B, ColorIn,
 		(bDepthTestedIn) ? SDPG_World : SDPG_Foreground,
-		LineThicknessIn, 0.0f, true);
+		LineThicknessIn, DepthBias, true);
 }
 
 
@@ -95,7 +95,7 @@ void FToolDataVisualizer::InternalDrawCircle(const FVector& Position, const FVec
 	::DrawCircle(CurrentPDI, TransformP(Position), (FVector)Tan1, (FVector)Tan2, 
 		Color, Radius, Steps,
 		(bDepthTestedIn) ? SDPG_World : SDPG_Foreground,
-		LineThicknessIn, 0.0f, true);
+		LineThicknessIn, DepthBias, true);
 }
 
 void FToolDataVisualizer::InternalDrawWireBox(const FBox& Box, const FLinearColor& ColorIn, float LineThicknessIn, bool bDepthTestedIn)
@@ -179,5 +179,5 @@ void FToolDataVisualizer::InternalDrawViewFacingCircle(const FVector& Position,
 	::DrawCircle(CurrentPDI, WorldPosition, (FVector)Tan1, (FVector)Tan2,
 		Color, Radius, Steps,
 		(bDepthTestedIn) ? SDPG_World : SDPG_Foreground,
-		LineThicknessIn, 0.0f, true);
+		LineThicknessIn, DepthBias, true);
 }
