MeshOp's use UniquePtr rather than SharedPtr

diff --git a/Engine/Plugins/Experimental/MeshModelingToolset/Source/MeshModelingTools/Private/EditNormalsTool.cpp b/Engine/Plugins/Experimental/MeshModelingToolset/Source/MeshModelingTools/Private/EditNormalsTool.cpp
--- a/Engine/Plugins/Experimental/MeshModelingToolset/Source/MeshModelingTools/Private/EditNormalsTool.cpp
+++ b/Engine/Plugins/Experimental/MeshModelingToolset/Source/MeshModelingTools/Private/EditNormalsTool.cpp
@@ -178,9 +178,9 @@ void UEditNormalsTool::SetAssetAPI(IToolsContextAssetAPI* AssetAPIIn)
 	this->AssetAPI = AssetAPIIn;
 }
 
-TSharedPtr<FDynamicMeshOperator> UEditNormalsOperatorFactory::MakeNewOperator()
+TUniquePtr<FDynamicMeshOperator> UEditNormalsOperatorFactory::MakeNewOperator()
 {
-	TSharedPtr<FEditNormalsOp> NormalsOp = MakeShared<FEditNormalsOp>();
+	TUniquePtr<FEditNormalsOp> NormalsOp = MakeUnique<FEditNormalsOp>();
 	NormalsOp->bFixInconsistentNormals = Tool->BasicProperties->bFixInconsistentNormals;
 	NormalsOp->bInvertNormals = Tool->BasicProperties->bInvertNormals;
 	NormalsOp->bRecomputeNormals = Tool->BasicProperties->bRecomputeNormals;
diff --git a/Engine/Plugins/Experimental/MeshModelingToolset/Source/MeshModelingTools/Private/PlaneCutTool.cpp b/Engine/Plugins/Experimental/MeshModelingToolset/Source/MeshModelingTools/Private/PlaneCutTool.cpp
--- a/Engine/Plugins/Experimental/MeshModelingToolset/Source/MeshModelingTools/Private/PlaneCutTool.cpp
+++ b/Engine/Plugins/Experimental/MeshModelingToolset/Source/MeshModelingTools/Private/PlaneCutTool.cpp
@@ -197,9 +197,9 @@ void UPlaneCutTool::SetAssetAPI(IToolsContextAssetAPI* AssetAPIIn)
 	this->AssetAPI = AssetAPIIn;
 }
 
-TSharedPtr<FDynamicMeshOperator> UPlaneCutOperatorFactory::MakeNewOperator()
+TUniquePtr<FDynamicMeshOperator> UPlaneCutOperatorFactory::MakeNewOperator()
 {
-	TSharedPtr<FPlaneCutOp> CutOp = MakeShared<FPlaneCutOp>();
+	TUniquePtr<FPlaneCutOp> CutOp = MakeUnique<FPlaneCutOp>();
 	CutOp->bDiscardAttributes = CutTool->BasicProperties->bDiscardAttributes;
 	CutOp->bFillCutHole = CutTool->BasicProperties->bFillCutHole;
 	CutOp->bFillSpans = CutTool->BasicProperties->bFillSpans;
diff --git a/Engine/Plugins/Experimental/MeshModelingToolset/Source/MeshModelingTools/Private/SmoothMeshTool.cpp b/Engine/Plugins/Experimental/MeshModelingToolset/Source/MeshModelingTools/Private/SmoothMeshTool.cpp
--- a/Engine/Plugins/Experimental/MeshModelingToolset/Source/MeshModelingTools/Private/SmoothMeshTool.cpp
+++ b/Engine/Plugins/Experimental/MeshModelingToolset/Source/MeshModelingTools/Private/SmoothMeshTool.cpp
@@ -181,18 +181,18 @@ bool USmoothMeshTool::CanAccept() const
 	return bResultValid;
 }
 
-TSharedPtr<FDynamicMeshOperator> USmoothMeshTool::MakeNewOperator()
+TUniquePtr<FDynamicMeshOperator> USmoothMeshTool::MakeNewOperator()
 {
-	TSharedPtr<FSmoothingOpBase> MeshOp;
+	TUniquePtr<FSmoothingOpBase> MeshOp;
 	switch (SmoothType)
 	{
 	default:
 	case ESmoothMeshToolSmoothType::Iterative:
-		MeshOp = MakeShared<FIterativeSmoothingOp>(&SrcDynamicMesh, SmoothSpeed, SmoothIterations);
+		MeshOp = MakeUnique<FIterativeSmoothingOp>(&SrcDynamicMesh, SmoothSpeed, SmoothIterations);
 		break;
 
 	case ESmoothMeshToolSmoothType::BiHarmonic_Cotan:
-		MeshOp = MakeShared<FCotanSmoothingOp>(&SrcDynamicMesh, SmoothSpeed, SmoothIterations);
+		MeshOp = MakeUnique<FCotanSmoothingOp>(&SrcDynamicMesh, SmoothSpeed, SmoothIterations);
 		break;
 /**
 	case ESmoothMeshToolSmoothType::BiHarmonic_MVW:
diff --git a/Engine/Plugins/Experimental/MeshModelingToolset/Source/MeshModelingTools/Private/UVProjectionTool.cpp b/Engine/Plugins/Experimental/MeshModelingToolset/Source/MeshModelingTools/Private/UVProjectionTool.cpp
--- a/Engine/Plugins/Experimental/MeshModelingToolset/Source/MeshModelingTools/Private/UVProjectionTool.cpp
+++ b/Engine/Plugins/Experimental/MeshModelingToolset/Source/MeshModelingTools/Private/UVProjectionTool.cpp
@@ -211,9 +211,9 @@ void UUVProjectionTool::SetAssetAPI(IToolsContextAssetAPI* AssetAPIIn)
 	this->AssetAPI = AssetAPIIn;
 }
 
-TSharedPtr<FDynamicMeshOperator> UUVProjectionOperatorFactory::MakeNewOperator()
+TUniquePtr<FDynamicMeshOperator> UUVProjectionOperatorFactory::MakeNewOperator()
 {
-	TSharedPtr<FUVProjectionOp> Op = MakeShared<FUVProjectionOp>();
+	TUniquePtr<FUVProjectionOp> Op = MakeUnique<FUVProjectionOp>();
 	Op->ProjectionMethod = Tool->BasicProperties->UVProjectionMethod;
 
 	// TODO: de-dupe this logic (it's also in Render, below)
diff --git a/Engine/Plugins/Experimental/MeshModelingToolset/Source/MeshModelingTools/Public/EditNormalsTool.h b/Engine/Plugins/Experimental/MeshModelingToolset/Source/MeshModelingTools/Public/EditNormalsTool.h
--- a/Engine/Plugins/Experimental/MeshModelingToolset/Source/MeshModelingTools/Public/EditNormalsTool.h
+++ b/Engine/Plugins/Experimental/MeshModelingToolset/Source/MeshModelingTools/Public/EditNormalsTool.h
@@ -112,7 +112,7 @@ class MESHMODELINGTOOLS_API UEditNormalsOperatorFactory : public UObject, public
 
 public:
 	// IDynamicMeshOperatorFactory API
-	virtual TSharedPtr<FDynamicMeshOperator> MakeNewOperator() override;
+	virtual TUniquePtr<FDynamicMeshOperator> MakeNewOperator() override;
 
 	UPROPERTY()
 	UEditNormalsTool *Tool;
diff --git a/Engine/Plugins/Experimental/MeshModelingToolset/Source/MeshModelingTools/Public/PlaneCutTool.h b/Engine/Plugins/Experimental/MeshModelingToolset/Source/MeshModelingTools/Public/PlaneCutTool.h
--- a/Engine/Plugins/Experimental/MeshModelingToolset/Source/MeshModelingTools/Public/PlaneCutTool.h
+++ b/Engine/Plugins/Experimental/MeshModelingToolset/Source/MeshModelingTools/Public/PlaneCutTool.h
@@ -86,7 +86,7 @@ class MESHMODELINGTOOLS_API UPlaneCutOperatorFactory : public UObject, public ID
 
 public:
 	// IDynamicMeshOperatorFactory API
-	virtual TSharedPtr<FDynamicMeshOperator> MakeNewOperator() override;
+	virtual TUniquePtr<FDynamicMeshOperator> MakeNewOperator() override;
 
 	UPROPERTY()
 	UPlaneCutTool *CutTool;
diff --git a/Engine/Plugins/Experimental/MeshModelingToolset/Source/MeshModelingTools/Public/SmoothMeshTool.h b/Engine/Plugins/Experimental/MeshModelingToolset/Source/MeshModelingTools/Public/SmoothMeshTool.h
--- a/Engine/Plugins/Experimental/MeshModelingToolset/Source/MeshModelingTools/Public/SmoothMeshTool.h
+++ b/Engine/Plugins/Experimental/MeshModelingToolset/Source/MeshModelingTools/Public/SmoothMeshTool.h
@@ -83,7 +83,7 @@ public:
 
 	virtual void OnPropertyModified(UObject* PropertySet, UProperty* Property) override;
 
-	virtual TSharedPtr<FDynamicMeshOperator> MakeNewOperator() override;
+	virtual TUniquePtr<FDynamicMeshOperator> MakeNewOperator() override;
 
 
 	
diff --git a/Engine/Plugins/Experimental/MeshModelingToolset/Source/MeshModelingTools/Public/UVProjectionTool.h b/Engine/Plugins/Experimental/MeshModelingToolset/Source/MeshModelingTools/Public/UVProjectionTool.h
--- a/Engine/Plugins/Experimental/MeshModelingToolset/Source/MeshModelingTools/Public/UVProjectionTool.h
+++ b/Engine/Plugins/Experimental/MeshModelingToolset/Source/MeshModelingTools/Public/UVProjectionTool.h
@@ -108,7 +108,7 @@ class MESHMODELINGTOOLS_API UUVProjectionOperatorFactory : public UObject, publi
 
 public:
 	// IDynamicMeshOperatorFactory API
-	virtual TSharedPtr<FDynamicMeshOperator> MakeNewOperator() override;
+	virtual TUniquePtr<FDynamicMeshOperator> MakeNewOperator() override;
 
 	UPROPERTY()
 	UUVProjectionTool *Tool;
diff --git a/Engine/Plugins/Experimental/MeshModelingToolset/Source/MeshModelingToolsEditorOnly/Private/MergeMeshesTool.cpp b/Engine/Plugins/Experimental/MeshModelingToolset/Source/MeshModelingToolsEditorOnly/Private/MergeMeshesTool.cpp
--- a/Engine/Plugins/Experimental/MeshModelingToolset/Source/MeshModelingToolsEditorOnly/Private/MergeMeshesTool.cpp
+++ b/Engine/Plugins/Experimental/MeshModelingToolset/Source/MeshModelingToolsEditorOnly/Private/MergeMeshesTool.cpp
@@ -180,9 +180,9 @@ void UMergeMeshesTool::OnPropertyModified(UObject* PropertySet, UProperty* Prope
 
 
 
-TSharedPtr<FDynamicMeshOperator> UMergeMeshesTool::MakeNewOperator()
+TUniquePtr<FDynamicMeshOperator> UMergeMeshesTool::MakeNewOperator()
 {
-	TSharedPtr<FVoxelMergeMeshesOp> MergeOp = MakeShared<FVoxelMergeMeshesOp>();
+	TUniquePtr<FVoxelMergeMeshesOp> MergeOp = MakeUnique<FVoxelMergeMeshesOp>();
 	MergeOp->VoxelCount     = MergeProps->VoxelCount;
 	MergeOp->AdaptivityD    = MergeProps->MeshAdaptivity;
 	MergeOp->IsoSurfaceD    = MergeProps->OffsetDistance;
diff --git a/Engine/Plugins/Experimental/MeshModelingToolset/Source/MeshModelingToolsEditorOnly/Private/ParameterizeMeshTool.cpp b/Engine/Plugins/Experimental/MeshModelingToolset/Source/MeshModelingToolsEditorOnly/Private/ParameterizeMeshTool.cpp
--- a/Engine/Plugins/Experimental/MeshModelingToolset/Source/MeshModelingToolsEditorOnly/Private/ParameterizeMeshTool.cpp
+++ b/Engine/Plugins/Experimental/MeshModelingToolset/Source/MeshModelingToolsEditorOnly/Private/ParameterizeMeshTool.cpp
@@ -179,9 +179,9 @@ void UParameterizeMeshTool::PostEditChangeProperty(FPropertyChangedEvent& Proper
 	Preview->InvalidateResult();
 }
 
-TSharedPtr<FDynamicMeshOperator> UParameterizeMeshTool::MakeNewOperator()
+TUniquePtr<FDynamicMeshOperator> UParameterizeMeshTool::MakeNewOperator()
 {
-	TSharedPtr<FParameterizeMeshOp> ParamertizeMeshOp = MakeShared<FParameterizeMeshOp>();
+	TUniquePtr<FParameterizeMeshOp> ParamertizeMeshOp = MakeUnique<FParameterizeMeshOp>();
 	ParamertizeMeshOp->Stretch   = ChartStretch;
 	ParamertizeMeshOp->NumCharts = 0;
 	ParamertizeMeshOp->InputMesh = InputMesh;
diff --git a/Engine/Plugins/Experimental/MeshModelingToolset/Source/MeshModelingToolsEditorOnly/Private/PolygonOnMeshTool.cpp b/Engine/Plugins/Experimental/MeshModelingToolset/Source/MeshModelingToolsEditorOnly/Private/PolygonOnMeshTool.cpp
--- a/Engine/Plugins/Experimental/MeshModelingToolset/Source/MeshModelingToolsEditorOnly/Private/PolygonOnMeshTool.cpp
+++ b/Engine/Plugins/Experimental/MeshModelingToolset/Source/MeshModelingToolsEditorOnly/Private/PolygonOnMeshTool.cpp
@@ -202,9 +202,9 @@ void UPolygonOnMeshTool::SetAssetAPI(IToolsContextAssetAPI* AssetAPIIn)
 	this->AssetAPI = AssetAPIIn;
 }
 
-TSharedPtr<FDynamicMeshOperator> UPolygonOnMeshOperatorFactory::MakeNewOperator()
+TUniquePtr<FDynamicMeshOperator> UPolygonOnMeshOperatorFactory::MakeNewOperator()
 {
-	TSharedPtr<FEmbedPolygonsOp> EmbedOp = MakeShared<FEmbedPolygonsOp>();
+	TUniquePtr<FEmbedPolygonsOp> EmbedOp = MakeUnique<FEmbedPolygonsOp>();
 	EmbedOp->bDiscardAttributes = Tool->BasicProperties->bDiscardAttributes;
 	EmbedOp->Operation = Tool->BasicProperties->PolygonOperation;
 	EmbedOp->PolygonScale = Tool->BasicProperties->PolygonScale;
diff --git a/Engine/Plugins/Experimental/MeshModelingToolset/Source/MeshModelingToolsEditorOnly/Private/VoxelCSGMeshesTool.cpp b/Engine/Plugins/Experimental/MeshModelingToolset/Source/MeshModelingToolsEditorOnly/Private/VoxelCSGMeshesTool.cpp
--- a/Engine/Plugins/Experimental/MeshModelingToolset/Source/MeshModelingToolsEditorOnly/Private/VoxelCSGMeshesTool.cpp
+++ b/Engine/Plugins/Experimental/MeshModelingToolset/Source/MeshModelingToolsEditorOnly/Private/VoxelCSGMeshesTool.cpp
@@ -171,9 +171,9 @@ void UVoxelCSGMeshesTool::OnPropertyModified(UObject* PropertySet, UProperty* Pr
 	Preview->InvalidateResult();
 }
 
-TSharedPtr<FDynamicMeshOperator> UVoxelCSGMeshesTool::MakeNewOperator()
+TUniquePtr<FDynamicMeshOperator> UVoxelCSGMeshesTool::MakeNewOperator()
 {
-	TSharedPtr<FVoxelBooleanMeshesOp> CSGOp = MakeShared<FVoxelBooleanMeshesOp>();
+	TUniquePtr<FVoxelBooleanMeshesOp> CSGOp = MakeUnique<FVoxelBooleanMeshesOp>();
 	CSGOp->Operation      = (FVoxelBooleanMeshesOp::EBooleanOperation)(int)CSGProps->Operation;
 	CSGOp->VoxelCount     = CSGProps->VoxelCount;
 	CSGOp->AdaptivityD    = CSGProps->MeshAdaptivity;
diff --git a/Engine/Plugins/Experimental/MeshModelingToolset/Source/MeshModelingToolsEditorOnly/Public/MergeMeshesTool.h b/Engine/Plugins/Experimental/MeshModelingToolset/Source/MeshModelingToolsEditorOnly/Public/MergeMeshesTool.h
--- a/Engine/Plugins/Experimental/MeshModelingToolset/Source/MeshModelingToolsEditorOnly/Public/MergeMeshesTool.h
+++ b/Engine/Plugins/Experimental/MeshModelingToolset/Source/MeshModelingToolsEditorOnly/Public/MergeMeshesTool.h
@@ -90,7 +90,7 @@ public:
 	virtual void OnPropertyModified(UObject* PropertySet, UProperty* Property) override;
 
 	// IDynamicMeshOperatorFactory API
-	virtual TSharedPtr<FDynamicMeshOperator> MakeNewOperator() override;
+	virtual TUniquePtr<FDynamicMeshOperator> MakeNewOperator() override;
 
 protected:
 	UPROPERTY()
diff --git a/Engine/Plugins/Experimental/MeshModelingToolset/Source/MeshModelingToolsEditorOnly/Public/ParameterizeMeshTool.h b/Engine/Plugins/Experimental/MeshModelingToolset/Source/MeshModelingToolsEditorOnly/Public/ParameterizeMeshTool.h
--- a/Engine/Plugins/Experimental/MeshModelingToolset/Source/MeshModelingToolsEditorOnly/Public/ParameterizeMeshTool.h
+++ b/Engine/Plugins/Experimental/MeshModelingToolset/Source/MeshModelingToolsEditorOnly/Public/ParameterizeMeshTool.h
@@ -91,7 +91,7 @@ public:
 	virtual void OnPropertyModified(UObject* PropertySet, UProperty* Property) override;
 
 	// IDynamicMeshOperatorFactory API
-	virtual TSharedPtr<FDynamicMeshOperator> MakeNewOperator() override;
+	virtual TUniquePtr<FDynamicMeshOperator> MakeNewOperator() override;
 
 protected:
 	// need to update bResultValid if these are modified, so we don't publicly expose them. 
diff --git a/Engine/Plugins/Experimental/MeshModelingToolset/Source/MeshModelingToolsEditorOnly/Public/PolygonOnMeshTool.h b/Engine/Plugins/Experimental/MeshModelingToolset/Source/MeshModelingToolsEditorOnly/Public/PolygonOnMeshTool.h
--- a/Engine/Plugins/Experimental/MeshModelingToolset/Source/MeshModelingToolsEditorOnly/Public/PolygonOnMeshTool.h
+++ b/Engine/Plugins/Experimental/MeshModelingToolset/Source/MeshModelingToolsEditorOnly/Public/PolygonOnMeshTool.h
@@ -97,7 +97,7 @@ class MESHMODELINGTOOLSEDITORONLY_API UPolygonOnMeshOperatorFactory : public UOb
 
 public:
 	// IDynamicMeshOperatorFactory API
-	virtual TSharedPtr<FDynamicMeshOperator> MakeNewOperator() override;
+	virtual TUniquePtr<FDynamicMeshOperator> MakeNewOperator() override;
 
 	UPROPERTY()
 	UPolygonOnMeshTool *Tool;
diff --git a/Engine/Plugins/Experimental/MeshModelingToolset/Source/MeshModelingToolsEditorOnly/Public/VoxelCSGMeshesTool.h b/Engine/Plugins/Experimental/MeshModelingToolset/Source/MeshModelingToolsEditorOnly/Public/VoxelCSGMeshesTool.h
--- a/Engine/Plugins/Experimental/MeshModelingToolset/Source/MeshModelingToolsEditorOnly/Public/VoxelCSGMeshesTool.h
+++ b/Engine/Plugins/Experimental/MeshModelingToolset/Source/MeshModelingToolsEditorOnly/Public/VoxelCSGMeshesTool.h
@@ -115,7 +115,7 @@ public:
 	virtual void OnPropertyModified(UObject* PropertySet, UProperty* Property) override;
 
 	// IDynamicMeshOperatorFactory API
-	virtual TSharedPtr<FDynamicMeshOperator> MakeNewOperator() override;
+	virtual TUniquePtr<FDynamicMeshOperator> MakeNewOperator() override;
 
 
 protected:
diff --git a/Engine/Plugins/Experimental/MeshModelingToolset/Source/ModelingComponents/Private/MeshOpPreviewHelpers.cpp b/Engine/Plugins/Experimental/MeshModelingToolset/Source/ModelingComponents/Private/MeshOpPreviewHelpers.cpp
--- a/Engine/Plugins/Experimental/MeshModelingToolset/Source/ModelingComponents/Private/MeshOpPreviewHelpers.cpp
+++ b/Engine/Plugins/Experimental/MeshModelingToolset/Source/ModelingComponents/Private/MeshOpPreviewHelpers.cpp
@@ -25,7 +25,7 @@ FDynamicMeshOpResult UMeshOpPreviewWithBackgroundCompute::Shutdown()
 	PreviewMesh->Disconnect();
 	PreviewMesh = nullptr;
 
-	return MoveTemp(Result);
+	return Result;
 }
 
 
@@ -66,9 +66,7 @@ void UMeshOpPreviewWithBackgroundCompute::UpdateResults()
 	EBackgroundComputeTaskStatus Status = BackgroundCompute->CheckStatus();
 	if (Status == EBackgroundComputeTaskStatus::NewResultAvailable)
 	{
-		TSharedPtr<FDynamicMeshOperator> MeshOp =
-			StaticCastSharedPtr<FDynamicMeshOperator>(BackgroundCompute->ExtractResult());
-
+		TUniquePtr<FDynamicMeshOperator> MeshOp  = BackgroundCompute->ExtractResult();
 		TUniquePtr<FDynamicMesh3> ResultMesh = MeshOp->ExtractResult();
 		PreviewMesh->SetTransform((FTransform)MeshOp->GetResultTransform());
 		PreviewMesh->UpdatePreview(ResultMesh.Get());  // copies the mesh @todo we could just give ownership to the Preview!
diff --git a/Engine/Plugins/Experimental/MeshModelingToolset/Source/ModelingOperators/Public/BackgroundModelingComputeSource.h b/Engine/Plugins/Experimental/MeshModelingToolset/Source/ModelingOperators/Public/BackgroundModelingComputeSource.h
--- a/Engine/Plugins/Experimental/MeshModelingToolset/Source/ModelingOperators/Public/BackgroundModelingComputeSource.h
+++ b/Engine/Plugins/Experimental/MeshModelingToolset/Source/ModelingOperators/Public/BackgroundModelingComputeSource.h
@@ -26,20 +26,20 @@ class TModelingOpTask : public FAbortableBackgroundTask
 
 public:
 
-	TModelingOpTask(TSharedPtr<OpType> OperatorIn) :
-		Operator(OperatorIn)
+	TModelingOpTask(TUniquePtr<OpType> OperatorIn) :
+		Operator(MoveTemp(OperatorIn))
 	{}
 
 	/**
 	 * @return the contained computation Operator
 	 */
-	TSharedPtr<OpType> ExtractOperator()
+	TUniquePtr<OpType> ExtractOperator()
 	{
-		return Operator;
+		return MoveTemp(Operator);
 	}
 
 protected:
-	TSharedPtr<OpType> Operator;
+	TUniquePtr<OpType> Operator;
 
 	// 	FAbortableBackgroundTask API
 	void DoWork()
@@ -90,7 +90,7 @@ enum class EBackgroundComputeTaskStatus
  * and a OpTypeFactory, which creates OpType instances on demand. The APIs that must
  * be provided in these types are minimal:
  *   OpTypeFactory:
- *      - TSharedPtr<OpType> MakeNewOperator()
+ *      - TUniquePtr<OpType> MakeNewOperator()
  *   OpType:
  *      - void CalculateResult(FProgressCancel*)
  * 
@@ -167,7 +167,7 @@ public:
 	/**
 	 * @return The last computed Operator. This may only be called once, the caller then owns the Operator.
 	 */
-	TSharedPtr<OpType> ExtractResult();
+	TUniquePtr<OpType> ExtractResult();
 
 	/**
 	 * @return duration in seconds of current computation
@@ -221,8 +221,8 @@ void TBackgroundModelingComputeSource<OpType, OpTypeFactory>::StartNewCompute()
 {
 	check(ActiveBackgroundTask == nullptr);
 
-	TSharedPtr<OpType> NewOp = OperatorSource->MakeNewOperator();
-	ActiveBackgroundTask = new FAsyncTaskExecuterWithAbort<TModelingOpTask<OpType> >(NewOp);
+	TUniquePtr<OpType> NewOp = OperatorSource->MakeNewOperator();
+	ActiveBackgroundTask = new FAsyncTaskExecuterWithAbort<TModelingOpTask<OpType> >(MoveTemp(NewOp));
 	ActiveBackgroundTask->StartBackgroundTask();
 
 	LastStartTime = AccumTime;
@@ -276,10 +276,10 @@ EBackgroundComputeTaskStatus TBackgroundModelingComputeSource<OpType, OpTypeFact
 
 
 template<typename OpType, typename OpTypeFactory>
-TSharedPtr<OpType> TBackgroundModelingComputeSource<OpType, OpTypeFactory>::ExtractResult()
+TUniquePtr<OpType> TBackgroundModelingComputeSource<OpType, OpTypeFactory>::ExtractResult()
 {
 	check(ActiveBackgroundTask != nullptr && ActiveBackgroundTask->IsDone());
-	TSharedPtr<OpType> Result = ActiveBackgroundTask->GetTask().ExtractOperator();
+	TUniquePtr<OpType> Result = ActiveBackgroundTask->GetTask().ExtractOperator();
 
 	delete ActiveBackgroundTask;
 	ActiveBackgroundTask = nullptr;
diff --git a/Engine/Plugins/Experimental/MeshModelingToolset/Source/ModelingOperators/Public/ModelingOperators.h b/Engine/Plugins/Experimental/MeshModelingToolset/Source/ModelingOperators/Public/ModelingOperators.h
--- a/Engine/Plugins/Experimental/MeshModelingToolset/Source/ModelingOperators/Public/ModelingOperators.h
+++ b/Engine/Plugins/Experimental/MeshModelingToolset/Source/ModelingOperators/Public/ModelingOperators.h
@@ -65,7 +65,7 @@ class IDynamicMeshOperatorFactory
 public:
 	virtual ~IDynamicMeshOperatorFactory() {}
 
-	virtual TSharedPtr<FDynamicMeshOperator> MakeNewOperator() = 0;
+	virtual TUniquePtr<FDynamicMeshOperator> MakeNewOperator() = 0;
 };
 
 
