Mesh Modeling: Voxel tools.  Converting the iso-surface to distance offset.  This gives more intuitive results.

Because voxel-based CSG can destroy the montonicity of distance values, extracting a non-zero iso-surface after a CSG opperation can give odd results.  Now we secretly rebuild the distance field by a remesh and voxelize step after CSG when isosurfaces at non-zero distance values are requested.

Also added a stage of recomputing normals after automatic simplification.

#rnx

diff --git a/Engine/Plugins/Experimental/MeshModelingToolset/Source/MeshModelingToolsEditorOnly/Private/MergeMeshesTool.cpp b/Engine/Plugins/Experimental/MeshModelingToolset/Source/MeshModelingToolsEditorOnly/Private/MergeMeshesTool.cpp
--- a/Engine/Plugins/Experimental/MeshModelingToolset/Source/MeshModelingToolsEditorOnly/Private/MergeMeshesTool.cpp
+++ b/Engine/Plugins/Experimental/MeshModelingToolset/Source/MeshModelingToolsEditorOnly/Private/MergeMeshesTool.cpp
@@ -21,8 +21,9 @@
 
 bool UMergeMeshesToolBuilder::CanBuildTool(const FToolBuilderState& SceneState) const
 {
-	bool bHasBuildAPI = (this->AssetAPI != nullptr);
-	bool bHasComponents = ToolBuilderUtil::CountComponents(SceneState, CanMakeComponentTarget) > 1;
+	const bool bHasBuildAPI = (this->AssetAPI != nullptr);
+	const int32 MinRequiredComponents = 1;
+	const bool bHasComponents = ToolBuilderUtil::CountComponents(SceneState, CanMakeComponentTarget) >= MinRequiredComponents;
 	return ( bHasBuildAPI && bHasComponents );
 }
 
diff --git a/Engine/Plugins/Experimental/MeshModelingToolset/Source/MeshModelingToolsEditorOnly/Public/MergeMeshesTool.h b/Engine/Plugins/Experimental/MeshModelingToolset/Source/MeshModelingToolsEditorOnly/Public/MergeMeshesTool.h
--- a/Engine/Plugins/Experimental/MeshModelingToolset/Source/MeshModelingToolsEditorOnly/Public/MergeMeshesTool.h
+++ b/Engine/Plugins/Experimental/MeshModelingToolset/Source/MeshModelingToolsEditorOnly/Public/MergeMeshesTool.h
@@ -43,15 +43,15 @@ class MESHMODELINGTOOLSEDITORONLY_API UMergeMeshesToolProperties : public UInter
 public:
 
 	/** The size of the geometry bounding box major axis measured in voxels.*/
-	UPROPERTY(EditAnywhere, Category = Options, meta = (UIMin = "8", UIMax = "512", ClampMin = "8", ClampMax = "512"))
+	UPROPERTY(EditAnywhere, Category = Options, meta = (UIMin = "8", UIMax = "1024", ClampMin = "8", ClampMax = "1024"))
 	int32 VoxelCount = 128;
 
-	/** Remeshing adaptivity */
+	/** Remeshing adaptivity, prior to optional simplification */
 	UPROPERTY(EditAnywhere, Category = Options, meta = (UIMin = "0", UIMax = "1", ClampMin = "0", ClampMax = "1"))
 	float MeshAdaptivity = 0.001f;
 
-	/** Offset when remeshing, measured in voxels units */
-	UPROPERTY(EditAnywhere, Category = Options, meta = (UIMin = "-2", UIMax = "2", ClampMin = "-2", ClampMax = "2"))
+	/** Offset when remeshing, note large offsets with high voxels counts will be slow */
+	UPROPERTY(EditAnywhere, Category = Options, meta = (UIMin = "-10", UIMax = "10", ClampMin = "-10", ClampMax = "10"))
 	float OffsetDistance = 0;
 
 	/** Automatically simplify the result of voxel-based merge.*/
diff --git a/Engine/Plugins/Experimental/MeshModelingToolset/Source/MeshModelingToolsEditorOnly/Public/VoxelCSGMeshesTool.h b/Engine/Plugins/Experimental/MeshModelingToolset/Source/MeshModelingToolsEditorOnly/Public/VoxelCSGMeshesTool.h
--- a/Engine/Plugins/Experimental/MeshModelingToolset/Source/MeshModelingToolsEditorOnly/Public/VoxelCSGMeshesTool.h
+++ b/Engine/Plugins/Experimental/MeshModelingToolset/Source/MeshModelingToolsEditorOnly/Public/VoxelCSGMeshesTool.h
@@ -65,15 +65,15 @@ public:
 	EVoxelCSGOperation Operation;
 
 	/** The size of the geometry bounding box major axis measured in voxels.*/
-	UPROPERTY(EditAnywhere, Category = Options, meta = (UIMin = "8", UIMax = "512", ClampMin = "8", ClampMax = "512"))
+	UPROPERTY(EditAnywhere, Category = Options, meta = (UIMin = "8", UIMax = "1024", ClampMin = "8", ClampMax = "1024"))
 	int32 VoxelCount = 128;
 
-	/** Remeshing adaptivity */
+	/** Remeshing adaptivity, prior to optional simplification */
 	UPROPERTY(EditAnywhere, Category = Options, meta = (UIMin = "0", UIMax = "1", ClampMin = "0", ClampMax = "1"))
 	float MeshAdaptivity = 0.001f;
 
-	/** Offset when remeshing, measured in voxels units */
-	UPROPERTY(EditAnywhere, Category = Options, meta = (UIMin = "-2", UIMax = "2", ClampMin = "-2", ClampMax = "2"))
+	/** Offset when remeshing, note large offsets with high voxels counts will be slow */
+	UPROPERTY(EditAnywhere, Category = Options, meta = (UIMin = "-10", UIMax = "10", ClampMin = "-10", ClampMax = "10"))
 	float OffsetDistance = 0;
 
 	/** Automatically simplify the result of voxel-based merge.*/
diff --git a/Engine/Plugins/Experimental/MeshModelingToolset/Source/ModelingOperatorsEditorOnly/Private/CompositionOps/VoxelBooleanMeshesOp.cpp b/Engine/Plugins/Experimental/MeshModelingToolset/Source/ModelingOperatorsEditorOnly/Private/CompositionOps/VoxelBooleanMeshesOp.cpp
--- a/Engine/Plugins/Experimental/MeshModelingToolset/Source/ModelingOperatorsEditorOnly/Private/CompositionOps/VoxelBooleanMeshesOp.cpp
+++ b/Engine/Plugins/Experimental/MeshModelingToolset/Source/ModelingOperatorsEditorOnly/Private/CompositionOps/VoxelBooleanMeshesOp.cpp
@@ -11,6 +11,7 @@
 #include "MeshDescriptionToDynamicMesh.h"
 #include "MeshSimplification.h"
 #include "MeshConstraintsUtil.h"
+#include "CleaningOps/EditNormalsOp.h"
 
 void FVoxelBooleanMeshesOp::CalculateResult(FProgressCancel* Progress)
 {
@@ -29,6 +30,20 @@ void FVoxelBooleanMeshesOp::CalculateResult(FProgressCancel* Progress)
 	// give this an absolute min since the user might scale both objects to zero..
 	VoxelSizeD = FMath::Max(VoxelSizeD, 0.001);
 
+	// world space units.
+	const double MaxNarrowBand = 3 * VoxelSizeD;
+	const double MaxIsoOffset = 2 * VoxelSizeD;
+
+	// IsoSurfaceD is in world units
+
+	const bool bShinkResult = (IsoSurfaceD < 0.);
+
+	// Note after the min/max operations of voxel-based CSG the interior distances or exterior 
+	// distances can be pretty f'd up so we need to do the initial post-CSG remesh with a "safe"
+	// value.
+	double CSGIsoSurface = 0.;
+
+	
 	// Create BooleanTool and Boolean the meshes.
 	TUniquePtr<IVoxelBasedCSG> VoxelCSGTool = IVoxelBasedCSG::CreateCSGTool(VoxelSizeD);
 	FVector MergedOrigin;
@@ -39,16 +54,20 @@ void FVoxelBooleanMeshesOp::CalculateResult(FProgressCancel* Progress)
 	switch (Operation)
 	{
 	case EBooleanOperation::DifferenceAB:
-		MergedOrigin = VoxelCSGTool->ComputeDifference(A, B, ResultMeshDescription, AdaptivityD, IsoSurfaceD);
+		CSGIsoSurface = 0.;
+		MergedOrigin = VoxelCSGTool->ComputeDifference(A, B, ResultMeshDescription, AdaptivityD, CSGIsoSurface);
 		break;
 	case EBooleanOperation::DifferenceBA:
-		MergedOrigin = VoxelCSGTool->ComputeDifference(B, A, ResultMeshDescription, AdaptivityD, IsoSurfaceD);
+		CSGIsoSurface = 0.;
+		MergedOrigin = VoxelCSGTool->ComputeDifference(B, A, ResultMeshDescription, AdaptivityD, CSGIsoSurface);
 		break;
 	case EBooleanOperation::Intersect:
-		MergedOrigin = VoxelCSGTool->ComputeIntersection(A, B, ResultMeshDescription, AdaptivityD, IsoSurfaceD);
+		CSGIsoSurface = FMath::Clamp(IsoSurfaceD, -MaxIsoOffset, MaxIsoOffset);
+		MergedOrigin = VoxelCSGTool->ComputeIntersection(A, B, ResultMeshDescription, AdaptivityD, CSGIsoSurface);
 		break;
 	case EBooleanOperation::Union:
-		MergedOrigin = VoxelCSGTool->ComputeUnion(A, B, ResultMeshDescription, AdaptivityD, IsoSurfaceD);
+		CSGIsoSurface = FMath::Clamp(IsoSurfaceD, 0., MaxIsoOffset);
+		MergedOrigin = VoxelCSGTool->ComputeUnion(A, B, ResultMeshDescription, AdaptivityD, CSGIsoSurface);
 		break;
 	default:
 		check(0);
@@ -61,6 +80,38 @@ void FVoxelBooleanMeshesOp::CalculateResult(FProgressCancel* Progress)
 		return;
 	}
 
+	const bool bNeedAdditionalRemeshing = (CSGIsoSurface != IsoSurfaceD);
+
+	// iteratively expand (or contract) the surface
+	if ( bNeedAdditionalRemeshing )
+	{
+		// how additional remeshes do we need?
+		int32 NumRemeshes = FMath::CeilToInt( FMath::Abs(IsoSurfaceD - CSGIsoSurface) / MaxIsoOffset );
+
+		double DeltaIsoSurface = (IsoSurfaceD - CSGIsoSurface) / NumRemeshes;
+
+		for (int32 i = 0; i < NumRemeshes; ++i)
+		{
+			if (Progress->Cancelled())
+			{
+				return;
+			}
+
+			IVoxelBasedCSG::FPlacedMesh PlacedMesh;
+			PlacedMesh.Mesh = &ResultMeshDescription;
+			PlacedMesh.Transform = static_cast<FTransform>(ResultTransform);
+
+			TArray< IVoxelBasedCSG::FPlacedMesh> SourceMeshes;
+			SourceMeshes.Add(PlacedMesh);
+
+			// union of only one mesh. we are using this to voxelize and remesh with and offset.
+			MergedOrigin = VoxelCSGTool->ComputeUnion(SourceMeshes, ResultMeshDescription, AdaptivityD, DeltaIsoSurface);
+
+			ResultTransform = FTransform3d(MergedOrigin);
+		}
+	}
+
+
 	// Convert to dynamic mesh
 	FMeshDescriptionToDynamicMesh Converter;
 	Converter.bPrintDebugMessages = true;
@@ -71,6 +122,9 @@ void FVoxelBooleanMeshesOp::CalculateResult(FProgressCancel* Progress)
 		return;
 	}
 
+	// for now we automatically fix up the normals if simplificaiton was requested.
+	bool bFixNormals = bAutoSimplify;
+
 	if (bAutoSimplify)
 	{
 		FQEMSimplification Reducer(ResultMesh.Get());
@@ -84,6 +138,31 @@ void FVoxelBooleanMeshesOp::CalculateResult(FProgressCancel* Progress)
 		Reducer.SimplifyToMaxError(MaxDisplacementSqr);
 
 	}
+
+
+	if (bFixNormals)
+	{
+
+		TSharedPtr<FDynamicMesh3> OpResultMesh(ExtractResult().Release()); // moved the unique pointer
+
+		// Recompute the normals
+		FEditNormalsOp EditNormalsOp;
+		EditNormalsOp.OriginalMesh = OpResultMesh; // the tool works on a deep copy of this mesh.
+		EditNormalsOp.bFixInconsistentNormals = true;
+		EditNormalsOp.bInvertNormals = false;
+		EditNormalsOp.bRecomputeNormals = true;
+		EditNormalsOp.NormalCalculationMethod = ENormalCalculationMethod::AreaAngleWeighting;
+		EditNormalsOp.bSplitNormals = true;
+		EditNormalsOp.bAllowSharpVertices = true;
+		EditNormalsOp.NormalSplitThreshold = 60.f;
+
+		EditNormalsOp.SetTransform(FTransform(ResultTransform));
+
+		EditNormalsOp.CalculateResult(Progress);
+
+		ResultMesh = EditNormalsOp.ExtractResult(); // return the edit normals operator copy to this tool.
+	}
+
 }
 
 float FVoxelBooleanMeshesOp::ComputeVoxelSize() const 
diff --git a/Engine/Plugins/Experimental/MeshModelingToolset/Source/ModelingOperatorsEditorOnly/Private/CompositionOps/VoxelMergeMeshesOp.cpp b/Engine/Plugins/Experimental/MeshModelingToolset/Source/ModelingOperatorsEditorOnly/Private/CompositionOps/VoxelMergeMeshesOp.cpp
--- a/Engine/Plugins/Experimental/MeshModelingToolset/Source/ModelingOperatorsEditorOnly/Private/CompositionOps/VoxelMergeMeshesOp.cpp
+++ b/Engine/Plugins/Experimental/MeshModelingToolset/Source/ModelingOperatorsEditorOnly/Private/CompositionOps/VoxelMergeMeshesOp.cpp
@@ -11,6 +11,7 @@
 #include "MeshDescriptionToDynamicMesh.h"
 #include "MeshSimplification.h"
 #include "MeshConstraintsUtil.h"
+#include "CleaningOps/EditNormalsOp.h"
 
 
 void FVoxelMergeMeshesOp::CalculateResult(FProgressCancel* Progress)
@@ -27,7 +28,12 @@ void FVoxelMergeMeshesOp::CalculateResult(FProgressCancel* Progress)
 
 	// Create CSGTool and merge the meshes.
 	TUniquePtr<IVoxelBasedCSG> VoxelCSGTool = IVoxelBasedCSG::CreateCSGTool(VoxelSizeD);
-	FVector MergedOrigin = VoxelCSGTool->ComputeUnion(*InputMeshArray, MergedMeshesDescription, AdaptivityD, IsoSurfaceD);
+
+	// world space units.
+	const double MaxNarrowBand = 3 * VoxelSizeD;
+	const double MaxIsoOffset = 2 * VoxelSizeD;
+	const double CSGIsoSurface = FMath::Clamp(IsoSurfaceD, 0., MaxIsoOffset); // the interior distance values maybe messed up when doing a union.
+	FVector MergedOrigin = VoxelCSGTool->ComputeUnion(*InputMeshArray, MergedMeshesDescription, AdaptivityD, CSGIsoSurface);
 	ResultTransform = FTransform3d(MergedOrigin);
 
 	if (Progress->Cancelled())
@@ -35,6 +41,37 @@ void FVoxelMergeMeshesOp::CalculateResult(FProgressCancel* Progress)
 		return;
 	}
 
+	const bool bNeedAdditionalRemeshing = (CSGIsoSurface != IsoSurfaceD);
+
+	// iteratively expand (or contract) the surface
+	if (bNeedAdditionalRemeshing)
+	{
+		// how additional remeshes do we need?
+		int32 NumRemeshes = FMath::CeilToInt(FMath::Abs(IsoSurfaceD - CSGIsoSurface) / MaxIsoOffset);
+
+		double DeltaIsoSurface = (IsoSurfaceD - CSGIsoSurface) / NumRemeshes;
+
+		for (int32 i = 0; i < NumRemeshes; ++i)
+		{
+			if (Progress->Cancelled())
+			{
+				return;
+			}
+
+			IVoxelBasedCSG::FPlacedMesh PlacedMesh;
+			PlacedMesh.Mesh = &MergedMeshesDescription;
+			PlacedMesh.Transform = static_cast<FTransform>(ResultTransform);
+
+			TArray< IVoxelBasedCSG::FPlacedMesh> SourceMeshes;
+			SourceMeshes.Add(PlacedMesh);
+
+			// union of only one mesh. we are using this to voxelize and remesh with and offset.
+			MergedOrigin = VoxelCSGTool->ComputeUnion(SourceMeshes, MergedMeshesDescription, AdaptivityD, DeltaIsoSurface);
+
+			ResultTransform = FTransform3d(MergedOrigin);
+		}
+	}
+
 	// Convert to dynamic mesh
 	FMeshDescriptionToDynamicMesh Converter;
 	Converter.bPrintDebugMessages = true;
@@ -45,6 +82,9 @@ void FVoxelMergeMeshesOp::CalculateResult(FProgressCancel* Progress)
 		return;
 	}
 
+	// for now we automatically fix up the normals if simplificaiton was requested.
+	bool bFixNormals = bAutoSimplify;
+
 	if (bAutoSimplify)
 	{
 		FQEMSimplification Reducer(ResultMesh.Get());
@@ -56,7 +96,33 @@ void FVoxelMergeMeshesOp::CalculateResult(FProgressCancel* Progress)
 		const double MaxDisplacementSqr = 3. *VoxelSizeD * VoxelSizeD;
 
 		Reducer.SimplifyToMaxError(MaxDisplacementSqr);
+
+
+	}
+
+	if (bFixNormals)
+	{
+
+		TSharedPtr<FDynamicMesh3> OpResultMesh(ExtractResult().Release()); // moved the unique pointer
+
+		// Recompute the normals
+		FEditNormalsOp EditNormalsOp;
+		EditNormalsOp.OriginalMesh = OpResultMesh; // the tool works on a deep copy of this mesh.
+		EditNormalsOp.bFixInconsistentNormals = true;
+		EditNormalsOp.bInvertNormals = false;
+		EditNormalsOp.bRecomputeNormals = true;
+		EditNormalsOp.NormalCalculationMethod = ENormalCalculationMethod::AreaAngleWeighting;
+		EditNormalsOp.bSplitNormals = true;
+		EditNormalsOp.bAllowSharpVertices = true;
+		EditNormalsOp.NormalSplitThreshold = 60.f;
+
+		EditNormalsOp.SetTransform(FTransform(ResultTransform));
+
+		EditNormalsOp.CalculateResult(Progress);
+
+		ResultMesh = EditNormalsOp.ExtractResult(); // return the edit normals operator copy to this tool.
 	}
+
 }
 
 
