Mesh Modeling : Voxel Tools. Delete inisde a separate transaction.  Now the first undo wil bring the source meshes back and the second will fully undo the tool.  This is experimental workflow..

#rnx

diff --git a/Engine/Plugins/Experimental/MeshModelingToolset/Source/MeshModelingToolsEditorOnly/Private/MergeMeshesTool.cpp b/Engine/Plugins/Experimental/MeshModelingToolset/Source/MeshModelingToolsEditorOnly/Private/MergeMeshesTool.cpp
--- a/Engine/Plugins/Experimental/MeshModelingToolset/Source/MeshModelingToolsEditorOnly/Private/MergeMeshesTool.cpp
+++ b/Engine/Plugins/Experimental/MeshModelingToolset/Source/MeshModelingToolsEditorOnly/Private/MergeMeshesTool.cpp
@@ -118,26 +118,39 @@ void UMergeMeshesTool::Shutdown(EToolShutdownType ShutdownType)
 	FDynamicMeshOpResult Result = Preview->Shutdown();
 	if (ShutdownType == EToolShutdownType::Accept)
 	{
-		GetToolManager()->BeginUndoTransaction(LOCTEXT("MergeMeshes", "Merge Meshes"));
+		// Generate the result
+		{
+			GetToolManager()->BeginUndoTransaction(LOCTEXT("MergeMeshes", "Merge Meshes"));
 
-		GenerateAsset(Result);
+			GenerateAsset(Result);
 
-		for (auto& ComponentTarget : ComponentTargets)
+			GetToolManager()->EndUndoTransaction();
+		}
+
+		// Hide or destroy the sources
 		{
-			ComponentTarget->SetOwnerVisibility(true);
-			AActor* Actor = ComponentTarget->GetOwnerActor();
-			if (MergeProps->bDeleteInputActors)
+			if (MergeProps->bDeleteInputActors) 
+				GetToolManager()->BeginUndoTransaction(LOCTEXT("MergeMeshes", "Remove Sources"));
+			
+			for (auto& ComponentTarget : ComponentTargets)
 			{
-				Actor->Destroy();
+				ComponentTarget->SetOwnerVisibility(true);
+				AActor* Actor = ComponentTarget->GetOwnerActor();
+				if (MergeProps->bDeleteInputActors)
+				{
+					Actor->Destroy();
+				}
+				else
+				{
+					// just hide the result.
+					Actor->SetIsTemporarilyHiddenInEditor(true);
+				}
 			}
-			else
-			{
-				// just hide the result.
-				Actor->SetIsTemporarilyHiddenInEditor(true);
-			}
-		}
 
-		GetToolManager()->EndUndoTransaction();
+			if (MergeProps->bDeleteInputActors) 
+				GetToolManager()->EndUndoTransaction();
+		}
+		
 
 	}
 	else
diff --git a/Engine/Plugins/Experimental/MeshModelingToolset/Source/MeshModelingToolsEditorOnly/Private/VoxelCSGMeshesTool.cpp b/Engine/Plugins/Experimental/MeshModelingToolset/Source/MeshModelingToolsEditorOnly/Private/VoxelCSGMeshesTool.cpp
--- a/Engine/Plugins/Experimental/MeshModelingToolset/Source/MeshModelingToolsEditorOnly/Private/VoxelCSGMeshesTool.cpp
+++ b/Engine/Plugins/Experimental/MeshModelingToolset/Source/MeshModelingToolsEditorOnly/Private/VoxelCSGMeshesTool.cpp
@@ -115,25 +115,40 @@ void UVoxelCSGMeshesTool::Shutdown(EToolShutdownType ShutdownType)
 	FDynamicMeshOpResult Result = Preview->Shutdown();
 	if (ShutdownType == EToolShutdownType::Accept)
 	{
-		GetToolManager()->BeginUndoTransaction(LOCTEXT("VoxelCSGMeshes", "Boolean Meshes"));
+		// Generate the result
+		{
+			GetToolManager()->BeginUndoTransaction(LOCTEXT("VoxelCSGMeshes", "Boolean Meshes"));
 
-		GenerateAsset(Result);
+			GenerateAsset(Result);
 
-		for (auto& ComponentTarget : ComponentTargets)
+			GetToolManager()->EndUndoTransaction();
+		}
+
+		// Hide or destroy the sources
 		{
-			ComponentTarget->SetOwnerVisibility(true);
-			AActor* Actor = ComponentTarget->GetOwnerActor();
-			if (CSGProps->bDeleteInputActors)
+			
+			if (CSGProps->bDeleteInputActors) 
+				GetToolManager()->BeginUndoTransaction(LOCTEXT("VoxelCSGMeshes", "Remove Sources"));
+			
+			for (auto& ComponentTarget : ComponentTargets)
 			{
-				Actor->Destroy();
+				ComponentTarget->SetOwnerVisibility(true);
+				AActor* Actor = ComponentTarget->GetOwnerActor();
+				if (CSGProps->bDeleteInputActors)
+				{
+					Actor->Destroy();
+				}
+				else
+				{
+					Actor->SetIsTemporarilyHiddenInEditor(true);
+				}
 			}
-			else
-			{
-				Actor->SetIsTemporarilyHiddenInEditor(true);
-			}
-		}
 
-		GetToolManager()->EndUndoTransaction();
+			if (CSGProps->bDeleteInputActors) 
+				GetToolManager()->EndUndoTransaction();
+			
+		}
+		
 	}
 	else
 	{
