ModelingTools: Add UPointSetComponent, derived from UOverylayComponent, draws fixed-size screen-aligned quads based on point position/size/color. Add Custom materials that calculate the point sizing and optional smooth falloff. ToolSetupUtil can look up these materials.
#rnx

diff --git a/Engine/Plugins/Experimental/MeshModelingToolset/Source/ModelingComponents/Private/Drawing/PointSetComponent.cpp b/Engine/Plugins/Experimental/MeshModelingToolset/Source/ModelingComponents/Private/Drawing/PointSetComponent.cpp
new file mode 100644
--- /dev/null
+++ b/Engine/Plugins/Experimental/MeshModelingToolset/Source/ModelingComponents/Private/Drawing/PointSetComponent.cpp
@@ -0,0 +1,343 @@
+// Copyright 1998-2019 Epic Games, Inc. All Rights Reserved.
+
+#include "Drawing/PointSetComponent.h"
+#include "RenderingThread.h"
+#include "PrimitiveViewRelevance.h"
+#include "PrimitiveSceneProxy.h"
+#include "VertexFactory.h"
+#include "MaterialShared.h"
+#include "Engine/CollisionProfile.h"
+#include "Materials/Material.h"
+#include "LocalVertexFactory.h"
+#include "SceneManagement.h"
+#include "DynamicMeshBuilder.h"
+#include "EngineGlobals.h"
+#include "Engine/Engine.h"
+#include "StaticMeshResources.h"
+#include "Algo/Accumulate.h"
+#include "Async/ParallelFor.h"
+
+
+
+struct FPointSetMeshBatchData
+{
+	FPointSetMeshBatchData()
+		: MaterialProxy(nullptr)
+	{}
+
+	FMaterialRenderProxy* MaterialProxy;
+	int32 StartIndex;
+	int32 NumPrimitives;
+	int32 MinVertexIndex;
+	int32 MaxVertexIndex;
+};
+
+
+
+
+
+class FPointSetSceneProxy final : public FPrimitiveSceneProxy
+{
+public:
+
+	FPointSetSceneProxy(UPointSetComponent* Component)
+		: FPrimitiveSceneProxy(Component),
+		MaterialRelevance(Component->GetMaterialRelevance(GetScene().GetFeatureLevel())),
+		VertexFactory(GetScene().GetFeatureLevel(), "FPointSetSceneProxy")
+	{
+		ERHIFeatureLevel::Type FeatureLevel = GetScene().GetFeatureLevel();
+
+		const int32 NumPointVertices = Component->Points.Num() * 4;
+		const int32 NumPointIndices = Component->Points.Num() * 6;
+		const int32 TotalNumVertices = NumPointVertices;
+		const int32 TotalNumIndices = NumPointIndices;
+		const int32 NumTextureCoordinates = 1;
+
+		VertexBuffers.PositionVertexBuffer.Init(TotalNumVertices);
+		VertexBuffers.StaticMeshVertexBuffer.Init(TotalNumVertices, NumTextureCoordinates);
+		VertexBuffers.ColorVertexBuffer.Init(TotalNumVertices);
+		IndexBuffer.Indices.SetNumUninitialized(TotalNumIndices);
+
+		int32 VertexBufferIndex = 0;
+		int32 IndexBufferIndex = 0;
+			
+		// Initialize points.
+		// Points are represented as two tris, all of whose vertices are coincident.
+		// The material then offsets them according to the signs of the vertex normals in a camera facing orientation.
+		// Size of the point is given by U0.
+		if (Component->Points.Num() > 0)
+		{
+			MeshBatchDatas.Emplace();
+			FPointSetMeshBatchData& MeshBatchData = MeshBatchDatas.Last();
+			MeshBatchData.MinVertexIndex = VertexBufferIndex;
+			MeshBatchData.MaxVertexIndex = VertexBufferIndex + NumPointVertices - 1;
+			MeshBatchData.StartIndex = IndexBufferIndex;
+			MeshBatchData.NumPrimitives = Component->Points.Num() * 2;
+			if (Component->GetMaterial(0) != nullptr)
+			{
+				MeshBatchData.MaterialProxy = Component->GetMaterial(0)->GetRenderProxy();
+			}
+			else
+			{
+				MeshBatchData.MaterialProxy = UMaterial::GetDefaultMaterial(MD_Surface)->GetRenderProxy();
+			}
+
+			const FVector TangentVectors[4] = {
+				FVector(1.0f, -1.0f, 0.0f),
+				FVector(1.0f, 1.0f, 0.0f),
+				FVector(-1.0f, 1.0f, 0.0f),
+				FVector(-1.0f, -1.0f, 0.0f)
+			};
+
+			// flatten to list of linear points because the set might be very large
+			// and we will benefit from parallel processing
+			TArray<const FRenderablePoint*> LinearPoints;
+			LinearPoints.Reserve(Component->Points.Num());
+			for (FRenderablePoint& Point : Component->Points)
+			{
+				LinearPoints.Add(&Point);
+			}
+
+			// assemble the render buffers
+			int NumLinearPoints = LinearPoints.Num();
+			ParallelFor(NumLinearPoints, [&](int32 i)
+			{
+				const FRenderablePoint& Point = *LinearPoints[i];
+				int UseVertexBufferIndex = 4 * i;
+				int UseIndexBufferIndex = 6 * i;
+
+				const FVector2D UV(Point.Size, 0.0f);
+				for (int j = 0; j < 4; ++j)
+				{
+					VertexBuffers.PositionVertexBuffer.VertexPosition(UseVertexBufferIndex + j) = Point.Position;
+					VertexBuffers.StaticMeshVertexBuffer.SetVertexUV(UseVertexBufferIndex + j, 0, UV);
+					VertexBuffers.ColorVertexBuffer.VertexColor(UseVertexBufferIndex + j) = Point.Color;
+					VertexBuffers.StaticMeshVertexBuffer.SetVertexTangents(UseVertexBufferIndex + j, FVector::ZeroVector, FVector::ZeroVector, TangentVectors[j]);
+				}
+
+				IndexBuffer.Indices[UseIndexBufferIndex + 0] = UseVertexBufferIndex + 0;
+				IndexBuffer.Indices[UseIndexBufferIndex + 1] = UseVertexBufferIndex + 1;
+				IndexBuffer.Indices[UseIndexBufferIndex + 2] = UseVertexBufferIndex + 2;
+				IndexBuffer.Indices[UseIndexBufferIndex + 3] = UseVertexBufferIndex + 2;
+				IndexBuffer.Indices[UseIndexBufferIndex + 4] = UseVertexBufferIndex + 3;
+				IndexBuffer.Indices[UseIndexBufferIndex + 5] = UseVertexBufferIndex + 0;
+			});
+		}
+
+		ENQUEUE_RENDER_COMMAND(OverlayVertexBuffersInit)(
+			[this](FRHICommandListImmediate& RHICmdList)
+		{
+			VertexBuffers.PositionVertexBuffer.InitResource();
+			VertexBuffers.StaticMeshVertexBuffer.InitResource();
+			VertexBuffers.ColorVertexBuffer.InitResource();
+
+			FLocalVertexFactory::FDataType Data;
+			VertexBuffers.PositionVertexBuffer.BindPositionVertexBuffer(&VertexFactory, Data);
+			VertexBuffers.StaticMeshVertexBuffer.BindTangentVertexBuffer(&VertexFactory, Data);
+			VertexBuffers.StaticMeshVertexBuffer.BindTexCoordVertexBuffer(&VertexFactory, Data);
+			VertexBuffers.ColorVertexBuffer.BindColorVertexBuffer(&VertexFactory, Data);
+			VertexFactory.SetData(Data);
+
+			VertexFactory.InitResource();
+			IndexBuffer.InitResource();
+		});
+	}
+
+	virtual ~FPointSetSceneProxy()
+	{
+		VertexBuffers.PositionVertexBuffer.ReleaseResource();
+		VertexBuffers.StaticMeshVertexBuffer.ReleaseResource();
+		VertexBuffers.ColorVertexBuffer.ReleaseResource();
+		IndexBuffer.ReleaseResource();
+		VertexFactory.ReleaseResource();
+	}
+
+	virtual void GetDynamicMeshElements(const TArray<const FSceneView*>& Views, const FSceneViewFamily& ViewFamily, uint32 VisibilityMap, FMeshElementCollector& Collector) const override
+	{
+		QUICK_SCOPE_CYCLE_COUNTER(STAT_OverlaySceneProxy_GetDynamicMeshElements);
+
+		for (int32 ViewIndex = 0; ViewIndex < Views.Num(); ViewIndex++)
+		{
+			if (VisibilityMap & (1 << ViewIndex))
+			{
+				for (const FPointSetMeshBatchData& MeshBatchData : MeshBatchDatas)
+				{
+					const FSceneView* View = Views[ViewIndex];
+					FMeshBatch& Mesh = Collector.AllocateMesh();
+					FMeshBatchElement& BatchElement = Mesh.Elements[0];
+					BatchElement.IndexBuffer = &IndexBuffer;
+					Mesh.bWireframe = false;
+					Mesh.VertexFactory = &VertexFactory;
+					Mesh.MaterialRenderProxy = MeshBatchData.MaterialProxy;
+
+					FDynamicPrimitiveUniformBuffer& DynamicPrimitiveUniformBuffer = Collector.AllocateOneFrameResource<FDynamicPrimitiveUniformBuffer>();
+					DynamicPrimitiveUniformBuffer.Set(GetLocalToWorld(), GetLocalToWorld(), GetBounds(), GetLocalBounds(), true, false, DrawsVelocity(), false);
+					BatchElement.PrimitiveUniformBufferResource = &DynamicPrimitiveUniformBuffer.UniformBuffer;
+
+					BatchElement.FirstIndex = MeshBatchData.StartIndex;
+					BatchElement.NumPrimitives = MeshBatchData.NumPrimitives;
+					BatchElement.MinVertexIndex = MeshBatchData.MinVertexIndex;
+					BatchElement.MaxVertexIndex = MeshBatchData.MaxVertexIndex;
+					Mesh.ReverseCulling = IsLocalToWorldDeterminantNegative();
+					Mesh.Type = PT_TriangleList;
+					Mesh.DepthPriorityGroup = SDPG_World;
+					Mesh.bCanApplyViewModeOverrides = false;
+					Collector.AddMesh(ViewIndex, Mesh);
+				}
+			}
+		}
+	}
+
+	virtual FPrimitiveViewRelevance GetViewRelevance(const FSceneView* View) const override
+	{
+		FPrimitiveViewRelevance Result;
+		Result.bDrawRelevance = IsShown(View);
+		Result.bShadowRelevance = IsShadowCast(View);
+		Result.bDynamicRelevance = true;
+		Result.bRenderInMainPass = ShouldRenderInMainPass();
+		Result.bUsesLightingChannels = GetLightingChannelMask() != GetDefaultLightingChannelMask();
+		Result.bRenderCustomDepth = ShouldRenderCustomDepth();
+		Result.bTranslucentSelfShadow = bCastVolumetricTranslucentShadow;
+		MaterialRelevance.SetPrimitiveViewRelevance(Result);
+		Result.bVelocityRelevance = IsMovable() && Result.bOpaqueRelevance && Result.bRenderInMainPass;
+		return Result;
+	}
+
+	virtual bool CanBeOccluded() const override
+	{
+		return !MaterialRelevance.bDisableDepthTest;
+	}
+
+	virtual uint32 GetMemoryFootprint() const override { return sizeof(*this) + GetAllocatedSize(); }
+
+	uint32 GetAllocatedSize() const { return FPrimitiveSceneProxy::GetAllocatedSize(); }
+
+	virtual SIZE_T GetTypeHash() const override
+	{
+		static SIZE_T UniquePointer;
+		return reinterpret_cast<SIZE_T>(&UniquePointer);
+	}
+
+private:
+	TArray<FPointSetMeshBatchData> MeshBatchDatas;
+	FMaterialRelevance MaterialRelevance;
+	FLocalVertexFactory VertexFactory;
+	FStaticMeshVertexBuffers VertexBuffers;
+	FDynamicMeshIndexBuffer32 IndexBuffer;
+};
+
+
+
+
+
+UPointSetComponent::UPointSetComponent()
+{
+	CastShadow = false;
+	bSelectable = false;
+	PrimaryComponentTick.bCanEverTick = false;
+	bBoundsDirty = true;
+
+	SetCollisionProfileName(UCollisionProfile::NoCollision_ProfileName);
+}
+
+void UPointSetComponent::SetPointMaterial(UMaterialInterface* InPointMaterial)
+{
+	PointMaterial = InPointMaterial;
+	UMeshComponent::SetMaterial(0, InPointMaterial);
+}
+
+
+void UPointSetComponent::Clear()
+{
+	Points.Reset();
+	MarkRenderStateDirty();
+	bBoundsDirty = true;
+}
+
+
+void UPointSetComponent::ReservePoints(const int32 MaxID)
+{
+	Points.Reserve(MaxID);
+}
+
+
+
+int32 UPointSetComponent::AddPoint(const FRenderablePoint& OverlayPoint)
+{
+	const int32 ID(Points.Add(OverlayPoint));
+	MarkRenderStateDirty();
+	bBoundsDirty = true;
+	return ID;
+}
+
+
+void UPointSetComponent::InsertPoint(const int32 ID, const FRenderablePoint& OverlayPoint)
+{
+	Points.Insert(ID, OverlayPoint);
+	MarkRenderStateDirty();
+	bBoundsDirty = true;
+}
+
+
+void UPointSetComponent::SetPointColor(const int32 ID, const FColor& NewColor)
+{
+	FRenderablePoint& OverlayPoint = Points[ID];
+	OverlayPoint.Color = NewColor;
+	MarkRenderStateDirty();
+}
+
+
+void UPointSetComponent::SetPointSize(const int32 ID, const float NewSize)
+{
+	FRenderablePoint& OverlayPoint = Points[ID];
+	OverlayPoint.Size = NewSize;
+	MarkRenderStateDirty();
+}
+
+
+void UPointSetComponent::RemovePoint(const int32 ID)
+{
+	Points.RemoveAt(ID);
+	MarkRenderStateDirty();
+	bBoundsDirty = true;
+}
+
+
+bool UPointSetComponent::IsPointValid(const int32 ID) const
+{
+	return ID < Points.GetMaxIndex() && Points.IsAllocated(ID);
+}
+
+
+FPrimitiveSceneProxy* UPointSetComponent::CreateSceneProxy()
+{
+	if (Points.Num() > 0)
+	{
+		return new FPointSetSceneProxy(this);
+	}
+
+	return nullptr;
+}
+
+
+int32 UPointSetComponent::GetNumMaterials() const
+{
+	return 1;
+}
+
+
+FBoxSphereBounds UPointSetComponent::CalcBounds(const FTransform& LocalToWorld) const
+{
+	if (bBoundsDirty)
+	{
+		FBox Box(ForceInit);
+		for (const FRenderablePoint& Point : Points)
+		{
+			Box += Point.Position;
+		}
+		Bounds = FBoxSphereBounds(Box);
+		bBoundsDirty = false;
+	}
+	return Bounds.TransformBy(LocalToWorld);
+}
+
diff --git a/Engine/Plugins/Experimental/MeshModelingToolset/Source/ModelingComponents/Private/ToolSetupUtil.cpp b/Engine/Plugins/Experimental/MeshModelingToolset/Source/ModelingComponents/Private/ToolSetupUtil.cpp
--- a/Engine/Plugins/Experimental/MeshModelingToolset/Source/ModelingComponents/Private/ToolSetupUtil.cpp
+++ b/Engine/Plugins/Experimental/MeshModelingToolset/Source/ModelingComponents/Private/ToolSetupUtil.cpp
@@ -62,3 +62,18 @@ UMaterialInterface* ToolSetupUtil::GetSelectionMaterial(UInteractiveToolManager*
 	}
 	return Material;
 }
+
+
+
+
+UMaterialInterface* ToolSetupUtil::GetDefaultPointComponentMaterial(bool bRoundPoints, UInteractiveToolManager* ToolManager)
+{
+	UMaterialInterface* Material = (bRoundPoints) ?
+		LoadObject<UMaterial>(nullptr, TEXT("/MeshModelingToolset/Materials/PointSetComponentMaterialSoft")) :
+		LoadObject<UMaterial>(nullptr, TEXT("/MeshModelingToolset/Materials/PointSetComponentMaterial"));
+	if (Material == nullptr && ToolManager != nullptr)
+	{
+		return ToolManager->GetContextQueriesAPI()->GetStandardMaterial(EStandardToolContextMaterials::VertexColorMaterial);
+	}
+	return Material;
+}
\ No newline at end of file
diff --git a/Engine/Plugins/Experimental/MeshModelingToolset/Source/ModelingComponents/Public/Drawing/PointSetComponent.h b/Engine/Plugins/Experimental/MeshModelingToolset/Source/ModelingComponents/Public/Drawing/PointSetComponent.h
new file mode 100644
--- /dev/null
+++ b/Engine/Plugins/Experimental/MeshModelingToolset/Source/ModelingComponents/Public/Drawing/PointSetComponent.h
@@ -0,0 +1,108 @@
+// Copyright 1998-2019 Epic Games, Inc. All Rights Reserved.
+
+#pragma once
+
+#include "CoreMinimal.h"
+#include "UObject/ObjectMacros.h"
+#include "Components/MeshComponent.h"
+#include "Containers/SparseArray.h"
+#include "RenderResource.h"
+#include "PackedNormal.h"
+#include "PointSetComponent.generated.h"
+
+class FPrimitiveSceneProxy;
+
+
+struct FRenderablePoint
+{
+	FRenderablePoint()
+		: Position(ForceInitToZero)
+		, Color(ForceInitToZero)
+		, Size(0.f)
+	{}
+
+	FRenderablePoint(const FVector& InPosition, const FColor& InColor, const float InSize)
+		: Position(InPosition)
+		, Color(InColor)
+		, Size(InSize)
+	{}
+
+	FVector Position;
+	FColor Color;
+	float Size;
+};
+
+
+/**
+ * UPointSetComponent is a Component that draws a set of points, as small squares.
+ * Per-point Color and (view-space) Size is supported. Normals are not supported.
+ * 
+ * Points are inserted with an externally-defined ID, internally this is done via
+ * a TSparseArray. This class allocates a contiguous TArray large enugh to hold the 
+ * largest ID. Using ReservePoints() may be beneficial for huge arrays.
+ *
+ * The points are drawn as two triangles (ie a square) orthogonal to the view direction. 
+ * The actual point size is calculated in the shader, and so a custom material must be used.
+ */
+UCLASS()
+class MODELINGCOMPONENTS_API UPointSetComponent : public UMeshComponent
+{
+	GENERATED_BODY()
+
+public:
+	UPointSetComponent();
+
+	/** Specify material which handles points */
+	void SetPointMaterial(UMaterialInterface* InPointMaterial);
+
+	/** Clear all primitives */
+	void Clear();
+
+	/** Reserve enough memory for up to the given ID */
+	void ReservePoints(const int32 MaxID);
+
+	/** Add a point to the overlay */
+	int32 AddPoint(const FRenderablePoint& OverlayPoint);
+
+	/** Insert a point with the given ID to the overlay */
+	void InsertPoint(const int32 ID, const FRenderablePoint& OverlayPoint);
+
+	/** Sets the color of a point */
+	void SetPointColor(const int32 ID, const FColor& NewColor);
+
+	/** Sets the size of a point */
+	void SetPointSize(const int32 ID, const float NewSize);
+
+	/** Remove a point from the overlay */
+	void RemovePoint(const int32 ID);
+
+	/** Queries whether a point with the given ID exists */
+	bool IsPointValid(const int32 ID) const;
+
+private:
+
+	//~ Begin UPrimitiveComponent Interface.
+	virtual FPrimitiveSceneProxy* CreateSceneProxy() override;
+	//~ End UPrimitiveComponent Interface.
+
+	//~ Begin UMeshComponent Interface.
+	virtual int32 GetNumMaterials() const override;
+	//~ End UMeshComponent Interface.
+
+	//~ Begin USceneComponent Interface.
+	virtual FBoxSphereBounds CalcBounds(const FTransform& LocalToWorld) const override;
+	//~ Begin USceneComponent Interface.
+
+	UPROPERTY()
+	const UMaterialInterface* PointMaterial;
+
+	UPROPERTY()
+	mutable FBoxSphereBounds Bounds;
+
+	UPROPERTY()
+	mutable bool bBoundsDirty;
+
+	TSparseArray<FRenderablePoint> Points;
+
+	friend class FPointSetSceneProxy;
+};
diff --git a/Engine/Plugins/Experimental/MeshModelingToolset/Source/ModelingComponents/Public/ToolSetupUtil.h b/Engine/Plugins/Experimental/MeshModelingToolset/Source/ModelingComponents/Public/ToolSetupUtil.h
--- a/Engine/Plugins/Experimental/MeshModelingToolset/Source/ModelingComponents/Public/ToolSetupUtil.h
+++ b/Engine/Plugins/Experimental/MeshModelingToolset/Source/ModelingComponents/Public/ToolSetupUtil.h
@@ -41,4 +41,12 @@ namespace ToolSetupUtil
 	 */
 	MODELINGCOMPONENTS_API UMaterialInterface* GetSelectionMaterial(UInteractiveToolManager* ToolManager);
 
+
+
+	/**
+	 * @param bRoundPoints true for round points, false for square
+	 * @return custom material suitable for use with UPointSetComponent
+	 */
+	MODELINGCOMPONENTS_API UMaterialInterface* GetDefaultPointComponentMaterial(bool bRoundPoints, UInteractiveToolManager* ToolManager);
+
 }
\ No newline at end of file
