ModelingTools: prototype Add Patch tool
#rnx

diff --git a/Engine/Plugins/Experimental/MeshModelingToolset/Source/MeshModelingTools/Private/AddPatchTool.cpp b/Engine/Plugins/Experimental/MeshModelingToolset/Source/MeshModelingTools/Private/AddPatchTool.cpp
new file mode 100644
--- /dev/null
+++ b/Engine/Plugins/Experimental/MeshModelingToolset/Source/MeshModelingTools/Private/AddPatchTool.cpp
@@ -0,0 +1,359 @@
+// Copyright 1998-2019 Epic Games, Inc. All Rights Reserved.
+
+#include "AddPatchTool.h"
+#include "ToolBuilderUtil.h"
+#include "InteractiveToolManager.h"
+#include "BaseBehaviors/MouseHoverBehavior.h"
+#include "Selection/ToolSelectionUtil.h"
+#include "AssetGenerationUtil.h"
+
+#include "MeshDescriptionBuilder.h"
+#include "Generators/RectangleMeshGenerator.h"
+#include "DynamicMeshAttributeSet.h"
+#include "StaticMeshComponentBuilder.h"
+#include "Drawing/MeshDebugDrawing.h"
+#include "MeshNormals.h"
+
+#include "Async/ParallelFor.h"
+
+#include "DynamicMeshEditor.h"
+
+#define LOCTEXT_NAMESPACE "UAddPatchTool"
+
+
+/*
+ * ToolBuilder
+ */
+bool UAddPatchToolBuilder::CanBuildTool(const FToolBuilderState& SceneState) const
+{
+	return (this->AssetAPI != nullptr);
+}
+
+UInteractiveTool* UAddPatchToolBuilder::BuildTool(const FToolBuilderState& SceneState) const
+{
+	UAddPatchTool* NewTool = NewObject<UAddPatchTool>(SceneState.ToolManager);
+	NewTool->SetWorld(SceneState.World);
+	NewTool->SetAssetAPI(AssetAPI);
+	return NewTool;
+}
+
+/*
+ * Tool
+ */
+UAddPatchToolProperties::UAddPatchToolProperties()
+{
+	Width = 10000;
+	Subdivisions = 50;
+	Rotation = 0;
+	Shift = 0.0;
+}
+
+void UAddPatchToolProperties::SaveProperties(UInteractiveTool* SaveFromTool)
+{
+	UAddPatchToolProperties* PropertyCache = GetPropertyCache<UAddPatchToolProperties>();
+	PropertyCache->Width = this->Width;
+	PropertyCache->Rotation = this->Rotation;
+	PropertyCache->Subdivisions = this->Subdivisions;
+}
+
+void UAddPatchToolProperties::RestoreProperties(UInteractiveTool* RestoreToTool)
+{
+	UAddPatchToolProperties* PropertyCache = GetPropertyCache<UAddPatchToolProperties>();
+	this->Width = PropertyCache->Width;
+	this->Rotation = PropertyCache->Rotation;
+	this->Subdivisions = PropertyCache->Subdivisions;
+}
+
+
+
+void UAddPatchTool::SetWorld(UWorld* World)
+{
+	this->TargetWorld = World;
+}
+
+void UAddPatchTool::SetAssetAPI(IToolsContextAssetAPI* AssetAPIIn)
+{
+	this->AssetAPI = AssetAPIIn;
+}
+
+
+void UAddPatchTool::Setup()
+{
+	USingleClickTool::Setup();
+
+	UMouseHoverBehavior* HoverBehavior = NewObject<UMouseHoverBehavior>(this);
+	HoverBehavior->Initialize(this);
+	AddInputBehavior(HoverBehavior);
+
+	ShapeSettings = NewObject<UAddPatchToolProperties>(this);
+	AddToolPropertySource(ShapeSettings);
+	ShapeSettings->RestoreProperties(this);
+
+	MaterialProperties = NewObject<UNewMeshMaterialProperties>(this);
+	AddToolPropertySource(MaterialProperties);
+	MaterialProperties->RestoreProperties(this);
+
+	// create preview mesh object
+	PreviewMesh = NewObject<UPreviewMesh>(this, TEXT("PreviewMesh"));
+	PreviewMesh->CreateInWorld(TargetWorld, FTransform::Identity);
+	PreviewMesh->SetVisible(false);
+	PreviewMesh->SetMaterial(MaterialProperties->Material);
+
+	GeneratePreviewBaseMesh();
+
+	WorldBounds = FBox(ForceInit);
+	for (const ULevel* Level : TargetWorld->GetLevels())
+	{
+		for (AActor* Actor : Level->Actors)
+		{
+			if (Actor)
+			{
+				FBox ActorBox = Actor->GetComponentsBoundingBox(true);
+				if (ActorBox.IsValid)
+				{
+					WorldBounds += ActorBox;
+				}
+			}
+		}
+	}
+	float WorldDiag = WorldBounds.GetSize().Size();
+	if (ShapeSettings->Width > WorldDiag * 0.25)
+	{
+		ShapeSettings->Width = WorldDiag * 0.25;
+	}
+
+	GetToolManager()->DisplayMessage(
+		LOCTEXT("OnStartAddPatchTool", "Position the Patch by moving the mouse over the scene. Drop a new instance by Left-clicking."),
+		EToolMessageLevel::UserNotification);
+}
+
+
+void UAddPatchTool::Shutdown(EToolShutdownType ShutdownType)
+{
+	PreviewMesh->SetVisible(false);
+	PreviewMesh->Disconnect();
+	PreviewMesh = nullptr;
+
+	ShapeSettings->SaveProperties(this);
+	MaterialProperties->SaveProperties(this);
+}
+
+
+void UAddPatchTool::Render(IToolsContextRenderAPI* RenderAPI)
+{
+}
+
+
+
+void UAddPatchTool::OnPropertyModified(UObject* PropertySet, UProperty* Property)
+{
+	PreviewMesh->EnableWireframe(MaterialProperties->bWireframe);
+	PreviewMesh->SetMaterial(MaterialProperties->Material);
+	GeneratePreviewBaseMesh();
+}
+
+
+
+FInputRayHit UAddPatchTool::BeginHoverSequenceHitTest(const FInputDeviceRay& PressPos)
+{
+	return FInputRayHit(0.0f);		// always hit in hover 
+}
+
+void UAddPatchTool::OnBeginHover(const FInputDeviceRay& DevicePos)
+{
+	UpdatePreviewPosition(DevicePos);
+}
+
+bool UAddPatchTool::OnUpdateHover(const FInputDeviceRay& DevicePos)
+{
+	UpdatePreviewPosition(DevicePos);
+	return true;
+}
+
+void UAddPatchTool::OnEndHover()
+{
+	// do nothing
+}
+
+
+
+void UAddPatchTool::Tick(float DeltaTime)
+{
+	if (bPreviewValid == false)
+	{
+		UpdatePreviewMesh();
+		PreviewMesh->SetVisible(true);
+		bPreviewValid = true;
+	}
+}
+
+
+
+void UAddPatchTool::UpdatePreviewPosition(const FInputDeviceRay& DeviceClickPos)
+{
+	FRay ClickPosWorldRay = DeviceClickPos.WorldRay;
+
+	// cast ray into scene
+	FVector RayStart = ClickPosWorldRay.Origin;
+	FVector RayEnd = ClickPosWorldRay.PointAt(999999);
+	FCollisionObjectQueryParams QueryParams(FCollisionObjectQueryParams::AllObjects);
+	FHitResult Result;
+	bool bHit = TargetWorld->LineTraceSingleByObjectType(Result, RayStart, RayEnd, QueryParams);
+	if (bHit)
+	{
+		ShapeFrame = FFrame3f(Result.ImpactPoint, Result.ImpactNormal);
+		//ShapeFrame.ConstrainedAlignPerpAxes();
+	}
+
+	// clear rotation
+	ShapeFrame.Rotation = FQuaternionf::Identity();
+
+	// rotate around up axis
+	if (ShapeSettings->Rotation != 0)
+	{
+		ShapeFrame.Rotate(FQuaternionf(ShapeFrame.Z(), ShapeSettings->Rotation, true));
+	}
+
+	if (bHit == false)
+	{
+		PreviewMesh->SetVisible(false);
+	}
+	else 
+	{
+		bPreviewValid = false;
+	}
+}
+
+
+void UAddPatchTool::UpdatePreviewMesh()
+{
+	FVector Direction = FVector(0, 0, 1);
+	float WorldMaxHeight = WorldBounds.Max.Z;
+	float WorldMinHeight = WorldBounds.Min.Z;
+	float WorldHeight = WorldMaxHeight - WorldMinHeight;
+
+	FDynamicMesh3 Projected(*BaseMesh);
+
+	FTransform MoveTransform = ShapeFrame.ToFTransform();
+
+	FFrame3d RayFrame(MoveTransform);
+	RayFrame.Origin.Z = WorldMaxHeight + 100.0f;
+
+	TSet<int> Misses;
+	FCriticalSection MissesLock;
+
+
+	int NumVerts = Projected.MaxVertexID();
+	ParallelFor(NumVerts, [this, &Misses, &MissesLock, &Projected, RayFrame, WorldMinHeight, Direction, MoveTransform](int vid)
+	{
+		FVector3d Pos = Projected.GetVertex(vid);
+		Pos = RayFrame.FromFramePoint(Pos);
+		FVector RayStart = (FVector)Pos;
+		FVector RayEnd = RayStart; RayEnd.Z = WorldMinHeight;
+
+		FCollisionObjectQueryParams QueryParams(FCollisionObjectQueryParams::AllObjects);
+		FHitResult Result;
+		bool bHit = TargetWorld->LineTraceSingleByObjectType(Result, RayStart, RayEnd, QueryParams);
+		if (bHit)
+		{
+			FVector3d HitPoint = Result.ImpactPoint + ShapeSettings->Shift * Direction;
+			FVector HitPosWorld = (FVector)HitPoint;
+			FVector HitPosLocal = MoveTransform.InverseTransformPosition(HitPosWorld);
+			Projected.SetVertex(vid, HitPosLocal);
+		}
+		else
+		{
+			MissesLock.Lock();
+			Misses.Add(vid);
+			MissesLock.Unlock();
+		}
+	});
+
+	TArray<int> RemoveTris;
+	for (int tid : Projected.TriangleIndicesItr())
+	{
+		FIndex3i Tri = Projected.GetTriangle(tid);
+		if (Misses.Contains(Tri.A) || Misses.Contains(Tri.B) || Misses.Contains(Tri.C))
+		{
+			RemoveTris.Add(tid);
+		}
+	}
+
+	FMeshNormals::QuickComputeVertexNormals(Projected);
+
+	FDynamicMeshEditor Editor(&Projected);
+	Editor.RemoveTriangles(RemoveTris, false);
+	PreviewMesh->UpdatePreview(&Projected);
+	PreviewMesh->SetTransform(MoveTransform);
+
+	PreviewMesh->SetVisible(true);
+}
+
+
+
+void UAddPatchTool::GeneratePreviewBaseMesh()
+{
+	BaseMesh = MakeUnique<FDynamicMesh3>();
+	GeneratePlane(BaseMesh.Get());
+
+	if (MaterialProperties->UVScale != 1.0 || MaterialProperties->bWorldSpaceUVScale)
+	{
+		FDynamicMeshEditor Editor(BaseMesh.Get());
+		float WorldUnitsInMetersFactor = MaterialProperties->bWorldSpaceUVScale ? .01f : 1.0f;
+		Editor.RescaleAttributeUVs(MaterialProperties->UVScale * WorldUnitsInMetersFactor, MaterialProperties->bWorldSpaceUVScale);
+	}
+
+	// recenter mesh
+	FAxisAlignedBox3d Bounds = BaseMesh->GetCachedBounds();
+	FVector3d TargetOrigin = Bounds.Center();
+	for (int vid : BaseMesh->VertexIndicesItr())
+	{
+		FVector3d Pos = BaseMesh->GetVertex(vid);
+		Pos -= TargetOrigin;
+		BaseMesh->SetVertex(vid, Pos);
+	}
+
+	PreviewMesh->UpdatePreview(BaseMesh.Get());
+}
+
+
+void UAddPatchTool::OnClicked(const FInputDeviceRay& DeviceClickPos)
+{
+#if WITH_EDITOR
+	const FDynamicMesh3* CurMesh = PreviewMesh->GetPreviewDynamicMesh();
+	FTransform3d CurTransform(PreviewMesh->GetTransform());
+	UMaterialInterface* Material = PreviewMesh->GetMaterial();
+	GetToolManager()->BeginUndoTransaction(LOCTEXT("AddPatchToolTransactionName", "Add Patch Mesh"));
+
+	AActor* NewActor = AssetGenerationUtil::GenerateStaticMeshActor(
+		AssetAPI, TargetWorld,
+		CurMesh, CurTransform, TEXT("Patch"),
+		AssetGenerationUtil::GetDefaultAutoGeneratedAssetPath(),
+		Material
+	);
+
+	// select newly-created object
+	ToolSelectionUtil::SetNewActorSelection(GetToolManager(), NewActor);
+
+	GetToolManager()->EndUndoTransaction();
+#else
+	checkNoEntry();
+#endif
+}
+
+
+
+
+
+void UAddPatchTool::GeneratePlane(FDynamicMesh3* OutMesh)
+{
+	FRectangleMeshGenerator RectGen;
+	RectGen.Width = RectGen.Height = ShapeSettings->Width;
+	RectGen.WidthVertexCount = RectGen.HeightVertexCount = ShapeSettings->Subdivisions + 2;
+	RectGen.Generate();
+	OutMesh->Copy(&RectGen);
+}
+
+
+
+#undef LOCTEXT_NAMESPACE
diff --git a/Engine/Plugins/Experimental/MeshModelingToolset/Source/MeshModelingTools/Public/AddPatchTool.h b/Engine/Plugins/Experimental/MeshModelingToolset/Source/MeshModelingTools/Public/AddPatchTool.h
new file mode 100644
--- /dev/null
+++ b/Engine/Plugins/Experimental/MeshModelingToolset/Source/MeshModelingTools/Public/AddPatchTool.h
@@ -0,0 +1,134 @@
+// Copyright 1998-2019 Epic Games, Inc. All Rights Reserved.
+
+#pragma once
+
+#include "CoreMinimal.h"
+#include "MaterialShared.h"
+#include "UObject/NoExportTypes.h"
+#include "InteractiveToolBuilder.h"
+#include "BaseTools/SingleClickTool.h"
+#include "PreviewMesh.h"
+#include "Properties/MeshMaterialProperties.h"
+#include "AddPatchTool.generated.h"
+
+class FDynamicMesh3;
+
+/**
+ *
+ */
+UCLASS()
+class MESHMODELINGTOOLS_API UAddPatchToolBuilder : public UInteractiveToolBuilder
+{
+	GENERATED_BODY()
+
+public:
+	IToolsContextAssetAPI* AssetAPI;
+
+	UAddPatchToolBuilder() 
+	{
+		AssetAPI = nullptr;
+	}
+
+	virtual bool CanBuildTool(const FToolBuilderState& SceneState) const override;
+	virtual UInteractiveTool* BuildTool(const FToolBuilderState& SceneState) const override;
+};
+
+
+
+UCLASS()
+class MESHMODELINGTOOLS_API UAddPatchToolProperties : public UInteractiveToolPropertySet
+{
+	GENERATED_BODY()
+
+public:
+	UAddPatchToolProperties();
+
+	/** Width of Shape */
+	UPROPERTY(EditAnywhere, Category = PatchSettings, meta = (DisplayName = "Width", UIMin = "1.0", UIMax = "1000.0", ClampMin = "0.0001", ClampMax = "1000000.0"))
+	float Width;
+
+	/** Rotation around up axis */
+	UPROPERTY(EditAnywhere, Category = PatchSettings, meta = (DisplayName = "Rotation", UIMin = "0.0", UIMax = "360.0"))
+	float Rotation;
+
+	/** Subdivisions */
+	UPROPERTY(EditAnywhere, Category = PatchSettings, meta = (DisplayName = "Subdivisions", UIMin = "0", UIMax = "100", ClampMin = "0", ClampMax = "4000"))
+	int Subdivisions;
+
+	/** Rotation around up axis */
+	UPROPERTY(EditAnywhere, Category = PatchSettings, meta = (DisplayName = "Shift", UIMin = "-1000", UIMax = "1000"))
+	float Shift;
+
+	virtual void SaveProperties(UInteractiveTool* SaveFromTool) override;
+	virtual void RestoreProperties(UInteractiveTool* RestoreToTool) override;
+};
+
+
+
+
+
+
+
+/**
+ *
+ */
+UCLASS()
+class MESHMODELINGTOOLS_API UAddPatchTool : public USingleClickTool, public IHoverBehaviorTarget
+{
+	GENERATED_BODY()
+
+public:
+	virtual void SetWorld(UWorld* World);
+	virtual void SetAssetAPI(IToolsContextAssetAPI* AssetAPI);
+
+	virtual void Setup() override;
+	virtual void Shutdown(EToolShutdownType ShutdownType) override;
+
+	virtual void Render(IToolsContextRenderAPI* RenderAPI) override;
+	virtual void Tick(float DeltaTime) override;
+
+	virtual bool HasCancel() const override { return false; }
+	virtual bool HasAccept() const override { return false; }
+	virtual bool CanAccept() const override { return false; }
+
+	virtual void OnPropertyModified(UObject* PropertySet, UProperty* Property) override;
+
+	virtual void OnClicked(const FInputDeviceRay& ClickPos) override;
+
+
+	// IHoverBehaviorTarget interface
+	virtual FInputRayHit BeginHoverSequenceHitTest(const FInputDeviceRay& PressPos) override;
+	virtual void OnBeginHover(const FInputDeviceRay& DevicePos) override;
+	virtual bool OnUpdateHover(const FInputDeviceRay& DevicePos) override;
+	virtual void OnEndHover() override;
+
+
+protected:
+	UPROPERTY()
+	UAddPatchToolProperties* ShapeSettings;
+
+	UPROPERTY()
+	UNewMeshMaterialProperties* MaterialProperties;
+
+
+
+	UPROPERTY()
+	UPreviewMesh* PreviewMesh;
+
+protected:
+	UWorld* TargetWorld;
+	IToolsContextAssetAPI* AssetAPI;
+
+	FBox WorldBounds;
+
+	FFrame3f ShapeFrame;
+	bool bPreviewValid = true;
+
+	void UpdatePreviewPosition(const FInputDeviceRay& ClickPos);
+	void UpdatePreviewMesh();
+
+	TUniquePtr<FDynamicMesh3> BaseMesh;
+	void GeneratePreviewBaseMesh();
+
+	void GeneratePlane(FDynamicMesh3* OutMesh);
+};
diff --git a/Engine/Plugins/Experimental/ModelingToolsEditorMode/Source/ModelingToolsEditorMode/Private/ModelingToolsEditorMode.cpp b/Engine/Plugins/Experimental/ModelingToolsEditorMode/Source/ModelingToolsEditorMode/Private/ModelingToolsEditorMode.cpp
--- a/Engine/Plugins/Experimental/ModelingToolsEditorMode/Source/ModelingToolsEditorMode/Private/ModelingToolsEditorMode.cpp
+++ b/Engine/Plugins/Experimental/ModelingToolsEditorMode/Source/ModelingToolsEditorMode/Private/ModelingToolsEditorMode.cpp
@@ -15,6 +15,7 @@
 #include "EditMeshPolygonsTool.h"
 #include "ConvertToPolygonsTool.h"
 #include "AddPrimitiveTool.h"
+#include "AddPatchTool.h"
 #include "SmoothMeshTool.h"
 #include "RemeshMeshTool.h"
 #include "SimplifyMeshTool.h"
@@ -317,6 +318,10 @@ void FModelingToolsEditorMode::Enter()
 	AddPrimitiveToolBuilder->AssetAPI = ToolsContext->GetAssetAPI();
 	RegisterToolFunc(ToolManagerCommands.BeginAddPrimitiveTool, TEXT("AddPrimitiveTool"), AddPrimitiveToolBuilder);
 
+	auto AddPatchToolBuilder = NewObject<UAddPatchToolBuilder>();
+	AddPatchToolBuilder->AssetAPI = ToolsContext->GetAssetAPI();
+	RegisterToolFunc(ToolManagerCommands.BeginAddPatchTool, TEXT("AddPatchTool"), AddPatchToolBuilder);
+
 	auto DrawPolygonToolBuilder = NewObject<UDrawPolygonToolBuilder>();
 	DrawPolygonToolBuilder->AssetAPI = ToolsContext->GetAssetAPI();
 	RegisterToolFunc(ToolManagerCommands.BeginDrawPolygonTool, TEXT("DrawPolygonTool"), DrawPolygonToolBuilder);
diff --git a/Engine/Plugins/Experimental/ModelingToolsEditorMode/Source/ModelingToolsEditorMode/Private/ModelingToolsEditorModeToolkit.cpp b/Engine/Plugins/Experimental/ModelingToolsEditorMode/Source/ModelingToolsEditorMode/Private/ModelingToolsEditorModeToolkit.cpp
--- a/Engine/Plugins/Experimental/ModelingToolsEditorMode/Source/ModelingToolsEditorMode/Private/ModelingToolsEditorModeToolkit.cpp
+++ b/Engine/Plugins/Experimental/ModelingToolsEditorMode/Source/ModelingToolsEditorMode/Private/ModelingToolsEditorModeToolkit.cpp
@@ -22,7 +22,7 @@
 // if set to 1, then on mode initialization we include buttons for prototype modeling tools
 static TAutoConsoleVariable<int32> CVarEnablePrototypeModelingTools(
 	TEXT("modeling.EnablePrototypes"),
-	0,
+	1,
 	TEXT("Enable unsupported Experimental prototype Modeling Tools"));
 
 
@@ -227,7 +227,7 @@ UEdModeInteractiveToolsContext* FModelingToolsEditorModeToolkit::GetToolsContext
 
 
 const TArray<FName> FModelingToolsEditorModeToolkit::PaletteNames_Standard = { FName(TEXT("Modeling")), FName(TEXT("Utilities")) };
-const TArray<FName> FModelingToolsEditorModeToolkit::PaletteNames_Experimental = { FName(TEXT("Create")), FName(TEXT("Edit")), FName(TEXT("UVs/Normals")), FName(TEXT("Utilities")) };
+const TArray<FName> FModelingToolsEditorModeToolkit::PaletteNames_Experimental = { FName(TEXT("Create")), FName(TEXT("Edit")), FName(TEXT("UVs/Normals")), FName(TEXT("Utilities")), FName(TEXT("Prototypes")) };
 
 
 void FModelingToolsEditorModeToolkit::GetToolPaletteNames(TArray<FName>& InPaletteName) const
@@ -362,6 +362,10 @@ void FModelingToolsEditorModeToolkit::BuildToolPalette_Experimental(FName Palett
 		ToolbarBuilder.AddToolBarButton(Commands.BeginMeshInspectorTool);
 		ToolbarBuilder.AddToolBarButton(Commands.BeginAttributeEditorTool);
 	}
+	else if (PaletteIndex == FName(TEXT("Prototypes")))
+	{
+		ToolbarBuilder.AddToolBarButton(Commands.BeginAddPatchTool);
+	}
 }
 
 void FModelingToolsEditorModeToolkit::OnToolPaletteChanged(FName PaletteName) 
diff --git a/Engine/Plugins/Experimental/ModelingToolsEditorMode/Source/ModelingToolsEditorMode/Private/ModelingToolsManagerActions.cpp b/Engine/Plugins/Experimental/ModelingToolsEditorMode/Source/ModelingToolsEditorMode/Private/ModelingToolsManagerActions.cpp
--- a/Engine/Plugins/Experimental/ModelingToolsEditorMode/Source/ModelingToolsEditorMode/Private/ModelingToolsManagerActions.cpp
+++ b/Engine/Plugins/Experimental/ModelingToolsEditorMode/Source/ModelingToolsEditorMode/Private/ModelingToolsManagerActions.cpp
@@ -22,6 +22,7 @@ FModelingToolsManagerCommands::FModelingToolsManagerCommands() :
 void FModelingToolsManagerCommands::RegisterCommands()
 {
 	UI_COMMAND(BeginAddPrimitiveTool, "Primitive", "Start the Add Primitive Tool", EUserInterfaceActionType::Button, FInputChord());
+	UI_COMMAND(BeginAddPatchTool, "Patch", "Start the Add Patch Tool", EUserInterfaceActionType::Button, FInputChord());
 	UI_COMMAND(BeginDrawPolygonTool, "Polygon", "Start the Draw Polygon Tool", EUserInterfaceActionType::Button, FInputChord());
 	UI_COMMAND(BeginShapeSprayTool, "ShapeSpray", "Start the Shape Spray Tool", EUserInterfaceActionType::Button, FInputChord());
 
diff --git a/Engine/Plugins/Experimental/ModelingToolsEditorMode/Source/ModelingToolsEditorMode/Public/ModelingToolsManagerActions.h b/Engine/Plugins/Experimental/ModelingToolsEditorMode/Source/ModelingToolsEditorMode/Public/ModelingToolsManagerActions.h
--- a/Engine/Plugins/Experimental/ModelingToolsEditorMode/Source/ModelingToolsEditorMode/Public/ModelingToolsManagerActions.h
+++ b/Engine/Plugins/Experimental/ModelingToolsEditorMode/Source/ModelingToolsEditorMode/Public/ModelingToolsManagerActions.h
@@ -16,6 +16,7 @@ public:
 
 
 	TSharedPtr<FUICommandInfo> BeginAddPrimitiveTool;
+	TSharedPtr<FUICommandInfo> BeginAddPatchTool;
 	TSharedPtr<FUICommandInfo> BeginDrawPolygonTool;
 	TSharedPtr<FUICommandInfo> BeginShapeSprayTool;
 
