add jacketing tool

diff --git a/Engine/Plugins/Experimental/GeometryProcessing/Source/DynamicMesh/Private/Operations/RemoveOccludedTriangles.cpp b/Engine/Plugins/Experimental/GeometryProcessing/Source/DynamicMesh/Private/Operations/RemoveOccludedTriangles.cpp
--- a/Engine/Plugins/Experimental/GeometryProcessing/Source/DynamicMesh/Private/Operations/RemoveOccludedTriangles.cpp
+++ b/Engine/Plugins/Experimental/GeometryProcessing/Source/DynamicMesh/Private/Operations/RemoveOccludedTriangles.cpp
@@ -7,21 +7,15 @@
 #include "DynamicMeshEditor.h"
 #include "MeshNormals.h"
 
+#include "MeshAdapter.h"
+
 #include "Async/ParallelFor.h"
 #include "Misc/ScopeLock.h"
 
 
-
-bool FRemoveOccludedTriangles::Apply(FDynamicMeshAABBTree3* Spatial)
+template<typename OccluderTriangleMeshType>
+bool TRemoveOccludedTriangles<OccluderTriangleMeshType>::Apply(FTransform3d MeshLocalToOccluderSpace, TMeshAABBTree3<OccluderTriangleMeshType>* Spatial, TFastWindingTree<OccluderTriangleMeshType>* FastWindingTree)
 {
-	FDynamicMeshAABBTree3 SpatialIfNotPassed;
-	if (!Spatial)
-	{
-		Spatial = &SpatialIfNotPassed;
-		Spatial->SetMesh(Mesh, true);
-	}
-	TFastWindingTree<FDynamicMesh3> FastWinding(Spatial, InsideMode == ECalculationMode::FastWindingNumber);
-
 	if (Cancelled())
 	{
 		return false;
@@ -29,22 +23,29 @@ bool FRemoveOccludedTriangles::Apply(FDynamicMeshAABBTree3* Spatial)
 
 	// ray directions
 	TArray<FVector3d> RayDirs; int NR = 0;
-	if (InsideMode == ECalculationMode::SimpleOcclusionTest)
+	
+	if (InsideMode == EOcclusionCalculationMode::SimpleOcclusionTest)
 	{
 		RayDirs.Add(FVector3d::UnitX()); RayDirs.Add(-FVector3d::UnitX());
 		RayDirs.Add(FVector3d::UnitY()); RayDirs.Add(-FVector3d::UnitY());
 		RayDirs.Add(FVector3d::UnitZ()); RayDirs.Add(-FVector3d::UnitZ());
+
+		for (int AddRayIdx = 0; AddRayIdx < AddRandomRays; AddRayIdx++)
+		{
+			RayDirs.Add(FVector3d(FMath::VRand()));
+		}
 		NR = RayDirs.Num();
 	}
 
-	auto IsOccludedFWN = [this, &FastWinding](const FVector3d& Pt)
+	auto IsOccludedFWN = [this, &FastWindingTree](const FVector3d& Pt)
 	{
-		return FastWinding.FastWindingNumber(Pt) > WindingIsoValue;
+		return FastWindingTree->FastWindingNumber(Pt) > WindingIsoValue;
 	};
 	auto IsOccludedSimple = [this, &Spatial, &RayDirs, &NR](const FVector3d& Pt)
 	{
 		FRay3d Ray;
 		Ray.Origin = Pt;
+		
 		for (int RayIdx = 0; RayIdx < NR; ++RayIdx)
 		{
 			Ray.Direction = RayDirs[RayIdx];
@@ -58,33 +59,31 @@ bool FRemoveOccludedTriangles::Apply(FDynamicMeshAABBTree3* Spatial)
 	};
 
 	TFunctionRef<bool(const FVector3d& Pt)> IsOccludedF =
-			(InsideMode == ECalculationMode::FastWindingNumber) ? 
+			(InsideMode == EOcclusionCalculationMode::FastWindingNumber) ? 
 		(TFunctionRef<bool(const FVector3d& Pt)>)IsOccludedFWN : 
 		(TFunctionRef<bool(const FVector3d& Pt)>)IsOccludedSimple;
 
 	bool bForceSingleThread = false;
 
 	TArray<bool> VertexOccluded;
-	if ( TriangleSamplingMethod == ETriangleSampling::PerVertex )
+	if ( TriangleSamplingMethod == EOcclusionTriangleSampling::Vertices || TriangleSamplingMethod == EOcclusionTriangleSampling::VerticesAndCentroids )
 	{
-		VertexOccluded.SetNum(Mesh->VertexCount());
+		VertexOccluded.SetNum(Mesh->MaxVertexID());
 
+		// do not trust source mesh normals; safer to recompute
 		FMeshNormals Normals(Mesh);
-		if (Mesh->HasVertexNormals() == false)
-		{
-			Normals.ComputeVertexNormals();
-		}
+		Normals.ComputeVertexNormals();
 
-		ParallelFor(Mesh->MaxVertexID(), [this, &Normals, &VertexOccluded, &IsOccludedF](int32 VID)
+		ParallelFor(Mesh->MaxVertexID(), [this, &Normals, &VertexOccluded, &IsOccludedF, &MeshLocalToOccluderSpace](int32 VID)
 		{
 			if (!Mesh->IsVertex(VID))
 			{
 				return;
 			}
 			FVector3d SamplePos = Mesh->GetVertex(VID);
-			FVector3d Normal = Mesh->HasVertexNormals() ? (FVector3d)Mesh->GetVertexNormal(VID) : Normals[VID];
+			FVector3d Normal = Normals[VID];
 			SamplePos += Normal * NormalOffset;
-			VertexOccluded[VID] = IsOccludedF(SamplePos);
+			VertexOccluded[VID] = IsOccludedF(MeshLocalToOccluderSpace.TransformPosition(SamplePos));
 		}, bForceSingleThread);
 	}
 	if (Cancelled())
@@ -94,26 +93,27 @@ bool FRemoveOccludedTriangles::Apply(FDynamicMeshAABBTree3* Spatial)
 
 	RemovedT.Empty();
 	FCriticalSection RemoveTMutex;
-	ParallelFor(Mesh->MaxTriangleID(), [this, &VertexOccluded, &IsOccludedF, &RemoveTMutex](int32 TID)
+	ParallelFor(Mesh->MaxTriangleID(), [this, &VertexOccluded, &IsOccludedF, &RemoveTMutex, &MeshLocalToOccluderSpace](int32 TID)
 	{
 		if (!Mesh->IsTriangle(TID))
 		{
 			return;
 		}
 
-		bool bInside = false;
-		if (TriangleSamplingMethod == ETriangleSampling::PerVertex)
+		bool bInside = true;
+		if (TriangleSamplingMethod == EOcclusionTriangleSampling::Vertices || TriangleSamplingMethod == EOcclusionTriangleSampling::VerticesAndCentroids)
 		{
 			FIndex3i Tri = Mesh->GetTriangle(TID);
-			bInside = VertexOccluded[Tri.A] || VertexOccluded[Tri.B] || VertexOccluded[Tri.C];
+			bInside = VertexOccluded[Tri.A] && VertexOccluded[Tri.B] && VertexOccluded[Tri.C];
 		}
-		else
+		if (bInside && (TriangleSamplingMethod == EOcclusionTriangleSampling::Centroids || TriangleSamplingMethod == EOcclusionTriangleSampling::VerticesAndCentroids))
 		{
 			FVector3d Centroid, Normal;
 			double Area;
 			Mesh->GetTriInfo(TID, Normal, Area, Centroid);
+
 			FVector3d SamplePos = Centroid + Normal * NormalOffset;
-			bInside = IsOccludedF(SamplePos);
+			bInside = IsOccludedF(MeshLocalToOccluderSpace.TransformPosition(SamplePos));
 		}
 		if (bInside)
 		{
@@ -141,4 +141,4 @@ bool FRemoveOccludedTriangles::Apply(FDynamicMeshAABBTree3* Spatial)
 	} 
 
 	return true;
-}
\ No newline at end of file
+}
diff --git a/Engine/Plugins/Experimental/GeometryProcessing/Source/DynamicMesh/Public/MeshAdapterUtil.h b/Engine/Plugins/Experimental/GeometryProcessing/Source/DynamicMesh/Public/MeshAdapterUtil.h
--- a/Engine/Plugins/Experimental/GeometryProcessing/Source/DynamicMesh/Public/MeshAdapterUtil.h
+++ b/Engine/Plugins/Experimental/GeometryProcessing/Source/DynamicMesh/Public/MeshAdapterUtil.h
@@ -23,6 +23,11 @@ namespace MeshAdapterUtil
 	 */
 	FTriangleMeshAdapterd DYNAMICMESH_API MakeDynamicMeshAdapter(const FDynamicMesh3* Mesh);
 
+	/**
+	* @return Transformed adapter of a FDynamicMesh
+	*/
+	FTriangleMeshAdapterd DYNAMICMESH_API MakeTransformedDynamicMeshAdapter(const FDynamicMesh3* Mesh, FTransform Transform);
+
 	/**
 	 * @return Mesh vertices as a point set
 	 */
diff --git a/Engine/Plugins/Experimental/GeometryProcessing/Source/DynamicMesh/Public/Operations/RemoveOccludedTriangles.h b/Engine/Plugins/Experimental/GeometryProcessing/Source/DynamicMesh/Public/Operations/RemoveOccludedTriangles.h
--- a/Engine/Plugins/Experimental/GeometryProcessing/Source/DynamicMesh/Public/Operations/RemoveOccludedTriangles.h
+++ b/Engine/Plugins/Experimental/GeometryProcessing/Source/DynamicMesh/Public/Operations/RemoveOccludedTriangles.h
@@ -7,35 +7,67 @@
 #include "MathUtil.h"
 #include "VectorTypes.h"
 #include "DynamicMesh3.h"
-#include "DynamicMeshAABBTree3.h"
+#include "Spatial/MeshAABBTree3.h"
 #include "Spatial/FastWinding.h"
 
+#include "MeshAdapter.h"
+
 #include "Util/ProgressCancel.h"
 
+enum class EOcclusionTriangleSampling
+{
+	Centroids,
+	Vertices,
+	VerticesAndCentroids
+};
+
+enum class EOcclusionCalculationMode
+{
+	FastWindingNumber,
+	SimpleOcclusionTest
+};
+
 /**
- * Remove "occluded" triangles, i.e. triangles on the "inside" of the mesh. 
+ * Remove "occluded" triangles, i.e. triangles on the "inside" of the mesh(es).
  * This is a fuzzy definition, current implementation has a couple of options
  * including a winding number-based version and an ambient-occlusion-ish version,
  * where if face is occluded for all test rays, then we classify it as inside and remove it.
  */
-class DYNAMICMESH_API FRemoveOccludedTriangles
+template<typename OccluderTriangleMeshType>
+class TRemoveOccludedTriangles
 {
 public:
 
 	FDynamicMesh3* Mesh;
 
-	FRemoveOccludedTriangles(FDynamicMesh3* Mesh) : Mesh(Mesh)
+	TRemoveOccludedTriangles(FDynamicMesh3* Mesh) : Mesh(Mesh)
 	{
 	}
-	virtual ~FRemoveOccludedTriangles() {}
+	virtual ~TRemoveOccludedTriangles() {}
+
+	/**
+	 * Remove the occluded triangles, considering the given occluder AABB tree (which may represent more geometry than a single mesh)
+	 * See alternative versions below if you would like to not precompute acceleration structures; these are mainly provided so the N-mutual-occluding-objects case can run on each mesh separately without 
+	 *
+	 * @param LocalToWorld Transform to take the local mesh into the space of the occluders
+	 * @param Occluders AABB trees for all occluders to consider
+	 * @param FastWindingTrees Precomputed fast winding trees for occluder
+	 * @return true on success
+	 */
+	virtual bool Apply(FTransform3d MeshLocalToOccluderSpace, TMeshAABBTree3<OccluderTriangleMeshType>* Occluder, TFastWindingTree<OccluderTriangleMeshType>* FastWindingTree);
 
 	/**
 	 * Remove the occluded triangles
 	 *
-	 * @param Spatial optional precomputed AABBTree for the mesh; if null it will be computed as needed
+	 * @param LocalToWorld Transform to take the local mesh into the space of the occluder geometry
+	 * @param Occluder AABB tree of occluding geometry
 	 * @return true on success
 	 */
-	virtual bool Apply(FDynamicMeshAABBTree3* Spatial = nullptr);
+	virtual bool Apply(FTransform3d MeshLocalToOccluderSpace, TMeshAABBTree3<OccluderTriangleMeshType>* Occluder)
+	{
+		TFastWindingTree<OccluderTriangleMeshType> FastWindingTree(Occluder, InsideMode == EOcclusionCalculationMode::FastWindingNumber);
+		return Apply(MeshLocalToOccluderSpace, Occluder, &FastWindingTree);
+	}
 
 	/**
 	 * @return EOperationValidationResult::Ok if we can apply operation, or error code if we cannot
@@ -51,12 +83,8 @@ public:
 	//
 
 	// how/where to sample triangles when testing for occlusion
-	enum class ETriangleSampling
-	{
-		PerVertex,
-		PerCentroid
-	};
-	ETriangleSampling TriangleSamplingMethod = ETriangleSampling::PerCentroid;
+
+	EOcclusionTriangleSampling TriangleSamplingMethod = EOcclusionTriangleSampling::Vertices;
 
 	// we nudge points out by this amount to try to counteract numerical issues
 	double NormalOffset = FMathd::ZeroTolerance;
@@ -64,12 +92,11 @@ public:
 	// use this as winding isovalue for WindingNumber mode
 	double WindingIsoValue = 0.5;
 
-	enum class ECalculationMode
-	{
-		FastWindingNumber,
-		SimpleOcclusionTest
-	};
-	ECalculationMode InsideMode = ECalculationMode::FastWindingNumber;
+
+	EOcclusionCalculationMode InsideMode = EOcclusionCalculationMode::FastWindingNumber;
+
+	/** Number of additional ray directions to add to raycast-based occlusion checks, beyond the default +/- major axis directions */
+	int AddRandomRays = 0;
 
 	/**
 	 * Set this to be able to cancel running operation
@@ -99,3 +126,10 @@ protected:
 
 };
 
+
+
+template class DYNAMICMESH_API TRemoveOccludedTriangles<TIndexMeshArrayAdapter<int, double, FVector3d>>;
+
+template class DYNAMICMESH_API TRemoveOccludedTriangles<FDynamicMesh3>;
+
+
diff --git a/Engine/Plugins/Experimental/GeometryProcessing/Source/GeometricObjects/Public/MeshAdapter.h b/Engine/Plugins/Experimental/GeometryProcessing/Source/GeometricObjects/Public/MeshAdapter.h
--- a/Engine/Plugins/Experimental/GeometryProcessing/Source/GeometricObjects/Public/MeshAdapter.h
+++ b/Engine/Plugins/Experimental/GeometryProcessing/Source/GeometricObjects/Public/MeshAdapter.h
@@ -59,13 +59,13 @@ inline FTriangleMeshAdapterd GetArrayMesh(TArray<FVector>& Vertices, TArray<FInt
 /**
  * Faster adapter specifically for the common index mesh case
  */
-template<typename IndexType, typename OutRealType>
+template<typename IndexType, typename OutRealType, typename InVectorType=FVector>
 struct TIndexMeshArrayAdapter
 {
-	const TArray<FVector>* SourceVertices;
+	const TArray<InVectorType>* SourceVertices;
 	const TArray<IndexType>* SourceTriangles;
 
-	void SetSources(const TArray<FVector>* SourceVerticesIn, const TArray<uint32>* SourceTrianglesIn)
+	void SetSources(const TArray<InVectorType>* SourceVerticesIn, const TArray<IndexType>* SourceTrianglesIn)
 	{
 		SourceVertices = SourceVerticesIn;
 		SourceTriangles = SourceTrianglesIn;
@@ -75,7 +75,7 @@ struct TIndexMeshArrayAdapter
 	{
 	}
 
-	TIndexMeshArrayAdapter(const TArray<FVector>* SourceVerticesIn, const TArray<uint32>* SourceTrianglesIn) : SourceVertices(SourceVerticesIn), SourceTriangles(SourceTrianglesIn)
+	TIndexMeshArrayAdapter(const TArray<InVectorType>* SourceVerticesIn, const TArray<IndexType>* SourceTrianglesIn) : SourceVertices(SourceVerticesIn), SourceTriangles(SourceTrianglesIn)
 	{
 	}
 
diff --git a/Engine/Plugins/Experimental/MeshModelingToolset/Source/MeshModelingTools/Private/MeshSelectionTool.cpp b/Engine/Plugins/Experimental/MeshModelingToolset/Source/MeshModelingTools/Private/MeshSelectionTool.cpp
--- a/Engine/Plugins/Experimental/MeshModelingToolset/Source/MeshModelingTools/Private/MeshSelectionTool.cpp
+++ b/Engine/Plugins/Experimental/MeshModelingToolset/Source/MeshModelingTools/Private/MeshSelectionTool.cpp
@@ -844,7 +844,7 @@ void UMeshSelectionTool::ExpandToConnected()
 		for (int j = 0; j < 3; ++j)
 		{
 			int32 tid = NbrTris[j];
-			if (SelectedTriangles[tid] == false && AddFaces.Contains(tid) == false)
+			if (tid != FDynamicMesh3::InvalidID && SelectedTriangles[tid] == false && AddFaces.Contains(tid) == false)
 			{
 				AddFaces.Add(tid);
 				Queue.Add(tid);
diff --git a/Engine/Plugins/Experimental/MeshModelingToolset/Source/MeshModelingTools/Private/RemoveOccludedTrianglesTool.cpp b/Engine/Plugins/Experimental/MeshModelingToolset/Source/MeshModelingTools/Private/RemoveOccludedTrianglesTool.cpp
new file mode 100644
--- /dev/null
+++ b/Engine/Plugins/Experimental/MeshModelingToolset/Source/MeshModelingTools/Private/RemoveOccludedTrianglesTool.cpp
@@ -0,0 +1,318 @@
+// Copyright 1998-2019 Epic Games, Inc. All Rights Reserved.
+
+#include "RemoveOccludedTrianglesTool.h"
+#include "InteractiveToolManager.h"
+#include "ToolBuilderUtil.h"
+
+#include "ToolSetupUtil.h"
+
+#include "DynamicMesh3.h"
+#include "BaseBehaviors/MultiClickSequenceInputBehavior.h"
+#include "Selection/SelectClickedAction.h"
+
+#include "MeshDescriptionToDynamicMesh.h"
+#include "DynamicMeshToMeshDescription.h"
+
+#include "InteractiveGizmoManager.h"
+
+#include "AssetGenerationUtil.h"
+
+
+#if WITH_EDITOR
+#include "Misc/ScopedSlowTask.h"
+#endif
+
+
+
+#define LOCTEXT_NAMESPACE "URemoveOccludedTrianglesTool"
+
+
+/*
+ * ToolBuilder
+ */
+
+
+bool URemoveOccludedTrianglesToolBuilder::CanBuildTool(const FToolBuilderState& SceneState) const
+{
+	return ToolBuilderUtil::CountComponents(SceneState, CanMakeComponentTarget) > 0;
+}
+
+UInteractiveTool* URemoveOccludedTrianglesToolBuilder::BuildTool(const FToolBuilderState& SceneState) const
+{
+	URemoveOccludedTrianglesTool* NewTool = NewObject<URemoveOccludedTrianglesTool>(SceneState.ToolManager);
+
+	TArray<UActorComponent*> Components = ToolBuilderUtil::FindAllComponents(SceneState, CanMakeComponentTarget);
+	check(Components.Num() > 0);
+
+	TArray<TUniquePtr<FPrimitiveComponentTarget>> ComponentTargets;
+	for (UActorComponent* ActorComponent : Components)
+	{
+		auto* MeshComponent = Cast<UPrimitiveComponent>(ActorComponent);
+		if ( MeshComponent )
+		{
+			ComponentTargets.Add(MakeComponentTarget(MeshComponent));
+		}
+	}
+
+	NewTool->SetSelection(MoveTemp(ComponentTargets));
+	NewTool->SetWorld(SceneState.World);
+	NewTool->SetAssetAPI(AssetAPI);
+
+	return NewTool;
+}
+
+
+
+
+/*
+ * Tool
+ */
+
+URemoveOccludedTrianglesToolProperties::URemoveOccludedTrianglesToolProperties()
+{
+
+}
+
+URemoveOccludedTrianglesAdvancedProperties::URemoveOccludedTrianglesAdvancedProperties()
+{
+}
+
+
+URemoveOccludedTrianglesTool::URemoveOccludedTrianglesTool()
+{
+}
+
+void URemoveOccludedTrianglesTool::SetWorld(UWorld* World)
+{
+	this->TargetWorld = World;
+}
+
+void URemoveOccludedTrianglesTool::Setup()
+{
+	UInteractiveTool::Setup();
+
+	// hide input StaticMeshComponent
+	for (TUniquePtr<FPrimitiveComponentTarget>& ComponentTarget : ComponentTargets)
+	{
+		ComponentTarget->SetOwnerVisibility(false);
+	}
+
+	BasicProperties = NewObject<URemoveOccludedTrianglesToolProperties>(this, TEXT("Remove Occluded Triangle Settings"));
+	AdvancedProperties = NewObject<URemoveOccludedTrianglesAdvancedProperties>(this, TEXT("Advanced Settings"));
+
+	// initialize our properties
+	AddToolPropertySource(BasicProperties);
+	AddToolPropertySource(AdvancedProperties);
+
+	// initialize the PreviewMesh+BackgroundCompute object
+	SetupPreviews();
+}
+
+
+void URemoveOccludedTrianglesTool::SetupPreviews()
+{	
+	int32 TargetNumPreview = ComponentTargets.Num();
+
+	CombinedMeshTrees = MakeShared<IndexMeshWithAcceleration>();
+
+#if WITH_EDITOR
+	static const FText SlowTaskText = LOCTEXT("RemoveOccludedTrianglesInit", "Building mesh occlusion data...");
+
+	FScopedSlowTask SlowTask(TargetNumPreview, SlowTaskText);
+	SlowTask.MakeDialog();
+
+	// Declare progress shortcut lambdas
+	auto EnterProgressFrame = [&SlowTask](float Progress)
+	{
+		SlowTask.EnterProgressFrame(Progress);
+	};
+#else
+	auto EnterProgressFrame = [](float Progress) {};
+#endif
+
+	OriginalDynamicMeshes.SetNum(TargetNumPreview);
+	for (int32 PreviewIdx = 0; PreviewIdx < TargetNumPreview; PreviewIdx++)
+	{
+		EnterProgressFrame(.5);
+
+		URemoveOccludedTrianglesOperatorFactory *OpFactory = NewObject<URemoveOccludedTrianglesOperatorFactory>();
+		OpFactory->Tool = this;
+		OpFactory->ComponentIndex = PreviewIdx;
+		OriginalDynamicMeshes[PreviewIdx] = MakeShared<FDynamicMesh3>();
+		FMeshDescriptionToDynamicMesh Converter;
+		Converter.bPrintDebugMessages = true;
+		Converter.Convert(ComponentTargets[PreviewIdx]->GetMesh(), *OriginalDynamicMeshes[PreviewIdx]);
+
+		EnterProgressFrame(.5);
+
+		UMeshOpPreviewWithBackgroundCompute* Preview = Previews.Add_GetRef(NewObject<UMeshOpPreviewWithBackgroundCompute>(OpFactory, "Preview"));
+		Preview->Setup(this->TargetWorld, OpFactory);
+		Preview->ConfigureMaterials(
+			ToolSetupUtil::GetDefaultMaterial(GetToolManager(), ComponentTargets[PreviewIdx]->GetMaterial(0)),
+			ToolSetupUtil::GetDefaultWorkingMaterial(GetToolManager())
+		);
+		Preview->PreviewMesh->SetTransform(ComponentTargets[PreviewIdx]->GetWorldTransform());
+		Preview->PreviewMesh->UpdatePreview(OriginalDynamicMeshes[PreviewIdx].Get());
+		Preview->SetVisibility(true);
+
+		CombinedMeshTrees->AddMesh(*OriginalDynamicMeshes[PreviewIdx], FTransform3d(ComponentTargets[PreviewIdx]->GetWorldTransform()));
+	}
+
+	CombinedMeshTrees->BuildAcceleration();
+
+	for (UMeshOpPreviewWithBackgroundCompute* Preview : Previews)
+	{
+		Preview->InvalidateResult();
+	}
+}
+
+
+
+void URemoveOccludedTrianglesTool::Shutdown(EToolShutdownType ShutdownType)
+{
+	// Restore (unhide) the source meshes
+	for (TUniquePtr<FPrimitiveComponentTarget>& ComponentTarget : ComponentTargets)
+	{
+		ComponentTarget->SetOwnerVisibility(true);
+	}
+
+	TArray<FDynamicMeshOpResult> Results;
+	for (UMeshOpPreviewWithBackgroundCompute* Preview : Previews)
+	{
+		Results.Add(Preview->Shutdown());
+	}
+	if (ShutdownType == EToolShutdownType::Accept)
+	{
+		GenerateAsset(Results);
+	}
+}
+
+void URemoveOccludedTrianglesTool::SetAssetAPI(IToolsContextAssetAPI* AssetAPIIn)
+{
+	this->AssetAPI = AssetAPIIn;
+}
+
+TUniquePtr<FDynamicMeshOperator> URemoveOccludedTrianglesOperatorFactory::MakeNewOperator()
+{
+	TUniquePtr<FRemoveOccludedTrianglesOp> Op = MakeUnique<FRemoveOccludedTrianglesOp>();
+	Op->NormalOffset = Tool->AdvancedProperties->NormalOffset;
+	switch (Tool->BasicProperties->OcclusionTestMethod)
+	{
+	case EOcclusionCalculationUIMode::GeneralizedWindingNumber:
+		Op->InsideMode = EOcclusionCalculationMode::FastWindingNumber;
+		break;
+	case EOcclusionCalculationUIMode::RaycastOcclusionSamples:
+		Op->InsideMode = EOcclusionCalculationMode::SimpleOcclusionTest;
+		break;
+	default:
+		ensure(false); // all cases should be handled
+	}
+	switch (Tool->BasicProperties->TriangleSampling)
+	{
+	case EOcclusionTriangleSamplingUIMode::Vertices:
+		Op->TriangleSamplingMethod = EOcclusionTriangleSampling::Vertices;
+		break;
+// Centroids sampling not exposed in UI for now
+// 	case EOcclusionTriangleSamplingUIMode::Centroids:
+// 		Op->TriangleSamplingMethod = EOcclusionTriangleSampling::Centroids;
+// 		break;
+	case EOcclusionTriangleSamplingUIMode::VerticesAndCentroids:
+		Op->TriangleSamplingMethod = EOcclusionTriangleSampling::VerticesAndCentroids;
+		break;
+	default:
+		ensure(false);
+	}
+	Op->WindingIsoValue = Tool->BasicProperties->WindingIsoValue;
+
+	FTransform LocalToWorld = Tool->ComponentTargets[ComponentIndex]->GetWorldTransform();
+	Op->OriginalMesh = Tool->OriginalDynamicMeshes[ComponentIndex];
+
+	Op->CombinedMeshTrees = Tool->CombinedMeshTrees;
+
+	Op->bOnlySelfOcclude = Tool->BasicProperties->bOnlySelfOcclude;
+
+	Op->AddRandomRays = Tool->BasicProperties->AddRandomRays;
+	
+	Op->SetTransform(LocalToWorld);
+
+	return Op;
+}
+
+
+
+void URemoveOccludedTrianglesTool::Render(IToolsContextRenderAPI* RenderAPI)
+{
+}
+
+void URemoveOccludedTrianglesTool::Tick(float DeltaTime)
+{
+	for (UMeshOpPreviewWithBackgroundCompute* Preview : Previews)
+	{
+		Preview->Tick(DeltaTime);
+	}
+}
+
+
+
+#if WITH_EDITOR
+void URemoveOccludedTrianglesTool::PostEditChangeProperty(FPropertyChangedEvent& PropertyChangedEvent)
+{
+	for (UMeshOpPreviewWithBackgroundCompute* Preview : Previews)
+	{
+		Preview->InvalidateResult();
+	}
+}
+#endif
+
+void URemoveOccludedTrianglesTool::OnPropertyModified(UObject* PropertySet, UProperty* Property)
+{
+	for (UMeshOpPreviewWithBackgroundCompute* Preview : Previews)
+	{
+		Preview->InvalidateResult();
+	}
+}
+
+
+
+bool URemoveOccludedTrianglesTool::HasAccept() const
+{
+	return true;
+}
+
+bool URemoveOccludedTrianglesTool::CanAccept() const
+{
+	for (UMeshOpPreviewWithBackgroundCompute* Preview : Previews)
+	{
+		if (!Preview->HaveValidResult())
+		{
+			return false;
+		}
+	}
+	return true;
+}
+
+
+void URemoveOccludedTrianglesTool::GenerateAsset(const TArray<FDynamicMeshOpResult>& Results)
+{
+	GetToolManager()->BeginUndoTransaction(LOCTEXT("RemoveOccludedTrianglesToolTransactionName", "Remove Occluded Triangles"));
+
+	check(Results.Num() == ComponentTargets.Num());
+	
+	for (int32 ComponentIdx = 0; ComponentIdx < ComponentTargets.Num(); ComponentIdx++)
+	{
+		check(Results[ComponentIdx].Mesh.Get() != nullptr);
+		ComponentTargets[ComponentIdx]->CommitMesh([&Results, &ComponentIdx, this](FMeshDescription* MeshDescription)
+		{
+			FDynamicMeshToMeshDescription Converter;
+			
+			Converter.Convert(Results[ComponentIdx].Mesh.Get(), *MeshDescription);
+		});
+	}
+
+	GetToolManager()->EndUndoTransaction();
+}
+
+
+
+
+#undef LOCTEXT_NAMESPACE
diff --git a/Engine/Plugins/Experimental/MeshModelingToolset/Source/MeshModelingTools/Public/RemoveOccludedTrianglesTool.h b/Engine/Plugins/Experimental/MeshModelingToolset/Source/MeshModelingTools/Public/RemoveOccludedTrianglesTool.h
new file mode 100644
--- /dev/null
+++ b/Engine/Plugins/Experimental/MeshModelingToolset/Source/MeshModelingTools/Public/RemoveOccludedTrianglesTool.h
@@ -0,0 +1,192 @@
+// Copyright 1998-2019 Epic Games, Inc. All Rights Reserved.
+
+#pragma once
+
+#include "CoreMinimal.h"
+#include "UObject/NoExportTypes.h"
+#include "MultiSelectionTool.h"
+#include "InteractiveToolBuilder.h"
+#include "MeshOpPreviewHelpers.h"
+#include "CleaningOps/RemoveOccludedTrianglesOp.h"
+#include "DynamicMesh3.h"
+#include "MeshAdapter.h"
+#include "BaseTools/SingleClickTool.h"
+#include "RemoveOccludedTrianglesTool.generated.h"
+
+
+// predeclarations
+struct FMeshDescription;
+class USimpleDynamicMeshComponent;
+
+
+
+
+
+/**
+ *
+ */
+UCLASS()
+class MESHMODELINGTOOLS_API URemoveOccludedTrianglesToolBuilder : public UInteractiveToolBuilder
+{
+	GENERATED_BODY()
+
+public:
+	IToolsContextAssetAPI* AssetAPI = nullptr;
+
+	virtual bool CanBuildTool(const FToolBuilderState& SceneState) const override;
+	virtual UInteractiveTool* BuildTool(const FToolBuilderState& SceneState) const override;
+};
+
+
+// these UIMode enums are versions of the enums in Operations/RemoveOccludedTriangles.h, w/ some removed & some renamed to be more user friendly
+
+UENUM()
+enum class EOcclusionTriangleSamplingUIMode : uint8
+{
+	// currently do not expose centroid-only option; it almost always looks bad
+	Vertices,
+	VerticesAndCentroids
+};
+
+UENUM()
+enum class EOcclusionCalculationUIMode : uint8
+{
+	GeneralizedWindingNumber, // maps to using fast winding number approximation
+	RaycastOcclusionSamples
+};
+
+
+
+/**
+ * Standard properties
+ */
+UCLASS()
+class MESHMODELINGTOOLS_API URemoveOccludedTrianglesToolProperties : public UInteractiveToolPropertySet
+{
+	GENERATED_BODY()
+
+public:
+	URemoveOccludedTrianglesToolProperties();
+
+	/** The method for deciding whether a triangle is occluded */
+	UPROPERTY(EditAnywhere, Category = OcclusionCalculation)
+	EOcclusionCalculationUIMode OcclusionTestMethod = EOcclusionCalculationUIMode::GeneralizedWindingNumber;
+
+	/** Where to sample triangles to test occlusion */
+	UPROPERTY(EditAnywhere, Category = OcclusionCalculation)
+	EOcclusionTriangleSamplingUIMode TriangleSampling = EOcclusionTriangleSamplingUIMode::VerticesAndCentroids;
+
+	/** The winding isovalue for GeneralizedWindingNumber mode */
+	UPROPERTY(EditAnywhere, Category = OcclusionCalculation, meta = (UIMin = "-1", UIMax = "1", ClampMin = "-2", ClampMax = "2", EditCondition = "OcclusionTestMethod==EOcclusionCalculationUIMode::GeneralizedWindingNumber"))
+	double WindingIsoValue = 0.5;
+
+	/** For raycast-based occlusion tests, optionally add random ray direction to increase the accuracy of the visibility sampling */
+	UPROPERTY(EditAnywhere, Category = OcclusionCalculation, meta = (EditCondition = "OcclusionTestMethod==EOcclusionCalculationUIMode::RaycastOcclusionSamples"))
+	int AddRandomRays = 0;
+
+	/** If false, when multiple meshes are selected the meshes can occlude each other.  When true, we process each selected mesh independently and only consider self-occlusions. */
+	UPROPERTY(EditAnywhere, Category = OcclusionCalculation)
+	bool bOnlySelfOcclude = false;
+};
+
+
+
+
+/**
+ * Advanced properties
+ */
+UCLASS()
+class MESHMODELINGTOOLS_API URemoveOccludedTrianglesAdvancedProperties : public UInteractiveToolPropertySet
+{
+	GENERATED_BODY()
+
+public:
+	URemoveOccludedTrianglesAdvancedProperties();
+
+	/** Amount to numerically 'nudge' occlusion sample query points away from the surface (to avoid e.g. all occlusion sample rays hitting the source triangle) */
+	// probably not actually a good idea to expose this to the user
+	//UPROPERTY(EditAnywhere, Category = NormalsTopology, meta = (UIMin = "0.000000001", UIMax = ".0001", ClampMin = "0.0", ClampMax = "0.01"))
+	double NormalOffset = FMathd::ZeroTolerance;
+};
+
+
+/**
+ * Factory with enough info to spawn the background-thread Operator to do a chunk of work for the tool
+ *  stores a pointer to the tool and enough info to know which specific operator it should spawn
+ */
+UCLASS()
+class MESHMODELINGTOOLS_API URemoveOccludedTrianglesOperatorFactory : public UObject, public IDynamicMeshOperatorFactory
+{
+	GENERATED_BODY()
+
+public:
+	// IDynamicMeshOperatorFactory API
+	virtual TUniquePtr<FDynamicMeshOperator> MakeNewOperator() override;
+
+	UPROPERTY()
+	URemoveOccludedTrianglesTool *Tool;
+
+	int ComponentIndex;
+};
+
+
+/**
+ * Simple Mesh Normal Updating Tool
+ */
+UCLASS()
+class MESHMODELINGTOOLS_API URemoveOccludedTrianglesTool : public UMultiSelectionTool
+{
+	GENERATED_BODY()
+
+public:
+
+	friend URemoveOccludedTrianglesOperatorFactory;
+
+	URemoveOccludedTrianglesTool();
+
+	virtual void Setup() override;
+	virtual void Shutdown(EToolShutdownType ShutdownType) override;
+
+	virtual void SetWorld(UWorld* World);
+	virtual void SetAssetAPI(IToolsContextAssetAPI* AssetAPI);
+
+	virtual void Tick(float DeltaTime) override;
+	virtual void Render(IToolsContextRenderAPI* RenderAPI) override;
+
+	virtual bool HasCancel() const override { return true; }
+	virtual bool HasAccept() const override;
+	virtual bool CanAccept() const override;
+
+#if WITH_EDITOR
+	virtual void PostEditChangeProperty(FPropertyChangedEvent &PropertyChangedEvent) override;
+#endif
+
+	virtual void OnPropertyModified(UObject* PropertySet, UProperty* Property) override;
+
+
+protected:
+	UPROPERTY()
+	URemoveOccludedTrianglesToolProperties* BasicProperties;
+
+	UPROPERTY()
+	URemoveOccludedTrianglesAdvancedProperties* AdvancedProperties;
+
+	UPROPERTY()
+	TArray<UMeshOpPreviewWithBackgroundCompute*> Previews;
+
+
+protected:
+	TArray<TSharedPtr<FDynamicMesh3>> OriginalDynamicMeshes;
+
+	// pre-transformed combined all-mesh data (with mesh geometry data in just raw index buffers as that's all we need)
+	TSharedPtr<IndexMeshWithAcceleration> CombinedMeshTrees;
+
+	UWorld* TargetWorld;
+	IToolsContextAssetAPI* AssetAPI;
+
+	FViewCameraState CameraState;
+
+	void SetupPreviews();
+
+	void GenerateAsset(const TArray<FDynamicMeshOpResult>& Results);
+};
diff --git a/Engine/Plugins/Experimental/MeshModelingToolset/Source/ModelingOperators/Private/CleaningOps/RemoveOccludedTrianglesOp.cpp b/Engine/Plugins/Experimental/MeshModelingToolset/Source/ModelingOperators/Private/CleaningOps/RemoveOccludedTrianglesOp.cpp
new file mode 100644
--- /dev/null
+++ b/Engine/Plugins/Experimental/MeshModelingToolset/Source/ModelingOperators/Private/CleaningOps/RemoveOccludedTrianglesOp.cpp
@@ -0,0 +1,90 @@
+// Copyright 1998-2019 Epic Games, Inc. All Rights Reserved.
+
+#pragma once
+
+#include "CleaningOps/RemoveOccludedTrianglesOp.h"
+
+#include "DynamicMesh3.h"
+#include "DynamicMeshAABBTree3.h"
+
+#include "Operations/RemoveOccludedTriangles.h"
+
+
+#include "Async/ParallelFor.h"
+
+
+void IndexMeshWithAcceleration::AddMesh(const FDynamicMesh3& MeshIn, FTransform3d Transform)
+{
+	int VertexIndexStart = Vertices.Num();
+	Vertices.SetNum(VertexIndexStart + MeshIn.MaxVertexID());
+	ParallelFor(MeshIn.MaxVertexID(), [this, &MeshIn, &Transform, &VertexIndexStart](int VID)
+	{
+		if (MeshIn.IsVertex(VID))
+		{
+			Vertices[VID + VertexIndexStart] = Transform.TransformPosition(MeshIn.GetVertex(VID));
+		}
+	}
+	, false);
+
+	// don't parallelize triangles b/c we want need them compact
+	for (int TID = 0; TID < MeshIn.MaxTriangleID(); TID++)
+	{
+		if (MeshIn.IsTriangle(TID))
+		{
+			FIndex3i Triangle = MeshIn.GetTriangle(TID);
+			Triangles.Add(Triangle.A + VertexIndexStart);
+			Triangles.Add(Triangle.B + VertexIndexStart);
+			Triangles.Add(Triangle.C + VertexIndexStart);
+		}
+	}
+}
+
+
+void FRemoveOccludedTrianglesOp::SetTransform(const FTransform& Transform)
+{
+	ResultTransform = (FTransform3d)Transform;
+}
+
+void FRemoveOccludedTrianglesOp::CalculateResult(FProgressCancel* Progress)
+{
+	if (Progress->Cancelled())
+	{
+		return;
+	}
+
+	bool bDiscardAttributes = false;
+	ResultMesh->Copy(*OriginalMesh, true, true, true, !bDiscardAttributes);
+	
+	if (bOnlySelfOcclude)
+	{
+		TRemoveOccludedTriangles<FDynamicMesh3> SelfJacket(ResultMesh.Get());
+		FDynamicMeshAABBTree3 SelfAABB(ResultMesh.Get());
+		TFastWindingTree<FDynamicMesh3> SelfFWTree(&SelfAABB);
+		if (Progress->Cancelled())
+		{
+			return;
+		}
+		SelfJacket.InsideMode = InsideMode;
+		SelfJacket.TriangleSamplingMethod = TriangleSamplingMethod;
+		SelfJacket.WindingIsoValue = WindingIsoValue;
+		SelfJacket.NormalOffset = NormalOffset;
+		SelfJacket.AddRandomRays = AddRandomRays;
+		SelfJacket.Apply(FTransform3d::Identity(), &SelfAABB, &SelfFWTree);
+	}
+	else
+	{
+		TRemoveOccludedTriangles<TIndexMeshArrayAdapter<int, double, FVector3d>> Jacket(ResultMesh.Get());
+
+		if (Progress->Cancelled())
+		{
+			return;
+		}
+
+		Jacket.InsideMode = InsideMode;
+		Jacket.TriangleSamplingMethod = TriangleSamplingMethod;
+		Jacket.WindingIsoValue = WindingIsoValue;
+		Jacket.NormalOffset = NormalOffset;
+		Jacket.AddRandomRays = AddRandomRays;
+		Jacket.Apply(ResultTransform, &CombinedMeshTrees->AABB, &CombinedMeshTrees->FastWinding);
+	}
+}
\ No newline at end of file
diff --git a/Engine/Plugins/Experimental/MeshModelingToolset/Source/ModelingOperators/Public/CleaningOps/RemoveOccludedTrianglesOp.h b/Engine/Plugins/Experimental/MeshModelingToolset/Source/ModelingOperators/Public/CleaningOps/RemoveOccludedTrianglesOp.h
new file mode 100644
--- /dev/null
+++ b/Engine/Plugins/Experimental/MeshModelingToolset/Source/ModelingOperators/Public/CleaningOps/RemoveOccludedTrianglesOp.h
@@ -0,0 +1,73 @@
+// Copyright 1998-2019 Epic Games, Inc. All Rights Reserved.
+
+#pragma once
+
+#include "CoreMinimal.h"
+#include "Util/ProgressCancel.h"
+#include "ModelingOperators.h"
+
+#include "MeshAdapter.h"
+#include "Operations/RemoveOccludedTriangles.h"
+
+// simple struct to help construct big trees of all the occluders
+struct MODELINGOPERATORS_API IndexMeshWithAcceleration
+{
+	TArray<int> Triangles;
+	TArray<FVector3d> Vertices;
+	TIndexMeshArrayAdapter<int, double, FVector3d> Mesh;
+	TMeshAABBTree3<TIndexMeshArrayAdapter<int, double, FVector3d>> AABB;
+	TFastWindingTree<TIndexMeshArrayAdapter<int, double, FVector3d>> FastWinding;
+
+	IndexMeshWithAcceleration() : AABB(&Mesh, false), FastWinding(&AABB, false)
+	{
+		Mesh.SetSources(&Vertices, &Triangles);
+	}
+
+	void AddMesh(const FDynamicMesh3& MeshIn, FTransform3d Transform);
+
+	void BuildAcceleration()
+	{
+		AABB.Build();
+		FastWinding.Build();
+	}
+};
+
+
+class MODELINGOPERATORS_API FRemoveOccludedTrianglesOp : public FDynamicMeshOperator
+{
+public:
+	virtual ~FRemoveOccludedTrianglesOp() {}
+
+	// inputs
+	TSharedPtr<FDynamicMesh3> OriginalMesh;
+	TSharedPtr<IndexMeshWithAcceleration> CombinedMeshTrees;
+
+	EOcclusionTriangleSampling TriangleSamplingMethod =
+		EOcclusionTriangleSampling::Centroids;
+
+	// we nudge points out by this amount to try to counteract numerical issues
+	double NormalOffset = FMathd::ZeroTolerance;
+
+	// use this as winding isovalue for WindingNumber mode
+	double WindingIsoValue = 0.5;
+
+	// random rays to add beyond +/- major axes, for raycast sampling
+	int AddRandomRays = 0;
+
+	// if true, *ignore* the combined selected occluder data in CombinedMeshTrees, and only consider self-occlusion of this mesh
+	bool bOnlySelfOcclude = false;
+
+	EOcclusionCalculationMode InsideMode =
+		EOcclusionCalculationMode::FastWindingNumber;
+
+
+	void SetTransform(const FTransform& Transform);
+
+	//
+	// FDynamicMeshOperator implementation
+	// 
+
+	virtual void CalculateResult(FProgressCancel* Progress) override;
+};
+
+
diff --git a/Engine/Plugins/Experimental/ModelingToolsEditorMode/Source/ModelingToolsEditorMode/Private/ModelingToolsEditorMode.cpp b/Engine/Plugins/Experimental/ModelingToolsEditorMode/Source/ModelingToolsEditorMode/Private/ModelingToolsEditorMode.cpp
--- a/Engine/Plugins/Experimental/ModelingToolsEditorMode/Source/ModelingToolsEditorMode/Private/ModelingToolsEditorMode.cpp
+++ b/Engine/Plugins/Experimental/ModelingToolsEditorMode/Source/ModelingToolsEditorMode/Private/ModelingToolsEditorMode.cpp
@@ -30,6 +30,7 @@
 #include "DisplaceMeshTool.h"
 #include "MeshSpaceDeformerTool.h"
 #include "EditNormalsTool.h"
+#include "RemoveOccludedTrianglesTool.h"
 #include "AttributeEditorTool.h"
 #include "TransformMeshesTool.h"
 #include "MeshSelectionTool.h"
@@ -360,6 +361,10 @@ void FModelingToolsEditorMode::Enter()
 	EditNormalsToolBuilder->AssetAPI = ToolsContext->GetAssetAPI();
 	RegisterToolFunc(ToolManagerCommands.BeginEditNormalsTool, TEXT("EditNormalsTool"), EditNormalsToolBuilder);
 
+	auto RemoveOccludedTrianglesToolBuilder = NewObject<URemoveOccludedTrianglesToolBuilder>();
+	RemoveOccludedTrianglesToolBuilder->AssetAPI = ToolsContext->GetAssetAPI();
+	RegisterToolFunc(ToolManagerCommands.BeginRemoveOccludedTrianglesTool, TEXT("RemoveOccludedTrianglesTool"), RemoveOccludedTrianglesToolBuilder);
+
 	auto UVProjectionToolBuilder = NewObject<UUVProjectionToolBuilder>();
 	UVProjectionToolBuilder->AssetAPI = ToolsContext->GetAssetAPI();
 	RegisterToolFunc(ToolManagerCommands.BeginUVProjectionTool, TEXT("UVProjectionTool"), UVProjectionToolBuilder);
diff --git a/Engine/Plugins/Experimental/ModelingToolsEditorMode/Source/ModelingToolsEditorMode/Private/ModelingToolsEditorModeStyle.cpp b/Engine/Plugins/Experimental/ModelingToolsEditorMode/Source/ModelingToolsEditorMode/Private/ModelingToolsEditorModeStyle.cpp
--- a/Engine/Plugins/Experimental/ModelingToolsEditorMode/Source/ModelingToolsEditorMode/Private/ModelingToolsEditorModeStyle.cpp
+++ b/Engine/Plugins/Experimental/ModelingToolsEditorMode/Source/ModelingToolsEditorMode/Private/ModelingToolsEditorModeStyle.cpp
@@ -117,6 +117,8 @@ void FModelingToolsEditorModeStyle::Initialize()
 		StyleSet->Set("ModelingToolsManagerCommands.BeginSimplifyMeshTool.Small", 		new IMAGE_PLUGIN_BRUSH("Icons/Simplify_40x", 		Icon20x20));
 		StyleSet->Set("ModelingToolsManagerCommands.BeginEditNormalsTool", 				new IMAGE_PLUGIN_BRUSH("Icons/Normals_40x",			Icon20x20));
 		StyleSet->Set("ModelingToolsManagerCommands.BeginEditNormalsTool.Small", 		new IMAGE_PLUGIN_BRUSH("Icons/Normals_40x",			Icon20x20));
+		StyleSet->Set("ModelingToolsManagerCommands.BeginRemoveOccludedTrianglesTool", 				new IMAGE_PLUGIN_BRUSH("Icons/Jacket_40x",			Icon20x20));
+		StyleSet->Set("ModelingToolsManagerCommands.BeginRemoveOccludedTrianglesTool.Small", 		new IMAGE_PLUGIN_BRUSH("Icons/Jacket_40x",			Icon20x20));
 		StyleSet->Set("ModelingToolsManagerCommands.BeginUVProjectionTool", 			new IMAGE_PLUGIN_BRUSH("Icons/UVProjection_40x", 	Icon20x20));
 		StyleSet->Set("ModelingToolsManagerCommands.BeginUVProjectionTool.Small", 		new IMAGE_PLUGIN_BRUSH("Icons/UVProjection_40x", 	Icon20x20));
 		StyleSet->Set("ModelingToolsManagerCommands.BeginVoxelMergeTool", 				new IMAGE_PLUGIN_BRUSH("Icons/VoxMerge_40x", 		Icon20x20));
diff --git a/Engine/Plugins/Experimental/ModelingToolsEditorMode/Source/ModelingToolsEditorMode/Private/ModelingToolsEditorModeToolkit.cpp b/Engine/Plugins/Experimental/ModelingToolsEditorMode/Source/ModelingToolsEditorMode/Private/ModelingToolsEditorModeToolkit.cpp
--- a/Engine/Plugins/Experimental/ModelingToolsEditorMode/Source/ModelingToolsEditorMode/Private/ModelingToolsEditorModeToolkit.cpp
+++ b/Engine/Plugins/Experimental/ModelingToolsEditorMode/Source/ModelingToolsEditorMode/Private/ModelingToolsEditorModeToolkit.cpp
@@ -359,6 +359,7 @@ void FModelingToolsEditorModeToolkit::BuildToolPalette_Experimental(FName Palett
 		ToolbarBuilder.AddToolBarButton(Commands.BeginSimplifyMeshTool);
 		ToolbarBuilder.AddToolBarButton(Commands.BeginRemeshMeshTool);
 		ToolbarBuilder.AddToolBarButton(Commands.BeginWeldEdgesTool);
+		ToolbarBuilder.AddToolBarButton(Commands.BeginRemoveOccludedTrianglesTool);
 		ToolbarBuilder.AddToolBarButton(Commands.BeginMeshInspectorTool);
 		ToolbarBuilder.AddToolBarButton(Commands.BeginAttributeEditorTool);
 	}
diff --git a/Engine/Plugins/Experimental/ModelingToolsEditorMode/Source/ModelingToolsEditorMode/Private/ModelingToolsManagerActions.cpp b/Engine/Plugins/Experimental/ModelingToolsEditorMode/Source/ModelingToolsEditorMode/Private/ModelingToolsManagerActions.cpp
--- a/Engine/Plugins/Experimental/ModelingToolsEditorMode/Source/ModelingToolsEditorMode/Private/ModelingToolsManagerActions.cpp
+++ b/Engine/Plugins/Experimental/ModelingToolsEditorMode/Source/ModelingToolsEditorMode/Private/ModelingToolsManagerActions.cpp
@@ -37,6 +37,7 @@ void FModelingToolsManagerCommands::RegisterCommands()
 	UI_COMMAND(BeginRemeshMeshTool, "Remesh", "Start the Remesh Tool", EUserInterfaceActionType::Button, FInputChord());
 	UI_COMMAND(BeginSimplifyMeshTool, "Simplify", "Start the Simplify Mesh Tool", EUserInterfaceActionType::Button, FInputChord());
 	UI_COMMAND(BeginEditNormalsTool, "Normals", "Start the Edit Normals Tool", EUserInterfaceActionType::Button, FInputChord());
+	UI_COMMAND(BeginRemoveOccludedTrianglesTool, "Jacketing", "Start the Jacketing Tool", EUserInterfaceActionType::Button, FInputChord());
 	UI_COMMAND(BeginUVProjectionTool, "UVProjection", "Start the UV Projection Tool", EUserInterfaceActionType::Button, FInputChord());
 	UI_COMMAND(BeginVoxelMergeTool, "VoxMerge", "Start the Voxel Merge Tool", EUserInterfaceActionType::Button, FInputChord());
 	UI_COMMAND(BeginVoxelBooleanTool, "VoxBoolean", "Start the Voxel Boolean Tool", EUserInterfaceActionType::Button, FInputChord());
diff --git a/Engine/Plugins/Experimental/ModelingToolsEditorMode/Source/ModelingToolsEditorMode/Public/ModelingToolsManagerActions.h b/Engine/Plugins/Experimental/ModelingToolsEditorMode/Source/ModelingToolsEditorMode/Public/ModelingToolsManagerActions.h
--- a/Engine/Plugins/Experimental/ModelingToolsEditorMode/Source/ModelingToolsEditorMode/Public/ModelingToolsManagerActions.h
+++ b/Engine/Plugins/Experimental/ModelingToolsEditorMode/Source/ModelingToolsEditorMode/Public/ModelingToolsManagerActions.h
@@ -31,6 +31,7 @@ public:
 	TSharedPtr<FUICommandInfo> BeginRemeshMeshTool;
 	TSharedPtr<FUICommandInfo> BeginSimplifyMeshTool;
 	TSharedPtr<FUICommandInfo> BeginEditNormalsTool;
+	TSharedPtr<FUICommandInfo> BeginRemoveOccludedTrianglesTool;
 	TSharedPtr<FUICommandInfo> BeginUVProjectionTool;
 	TSharedPtr<FUICommandInfo> BeginPlaneCutTool;
 	TSharedPtr<FUICommandInfo> BeginPolygonOnMeshTool;
