More placement options for add primitive tool

#rnx

diff --git a/Engine/Plugins/Experimental/MeshModelingToolset/Source/MeshModelingTools/Private/AddPrimitiveTool.cpp b/Engine/Plugins/Experimental/MeshModelingToolset/Source/MeshModelingTools/Private/AddPrimitiveTool.cpp
--- a/Engine/Plugins/Experimental/MeshModelingToolset/Source/MeshModelingTools/Private/AddPrimitiveTool.cpp
+++ b/Engine/Plugins/Experimental/MeshModelingToolset/Source/MeshModelingTools/Private/AddPrimitiveTool.cpp
@@ -53,6 +53,7 @@ UProceduralShapeToolProperties::UProceduralShapeToolProperties()
 	Subdivisions = 0;
 	PivotLocation = EMakeMeshPivotLocation::Base;
 	PlaceMode = EMakeMeshPlacementType::OnScene;
+	bAlignShapeToPlacementSurface = true;
 }
 
 void UProceduralShapeToolProperties::SaveProperties(UInteractiveTool* SaveFromTool)
@@ -65,6 +66,7 @@ void UProceduralShapeToolProperties::SaveProperties(UInteractiveTool* SaveFromTo
 	PropertyCache->Subdivisions = this->Subdivisions;
 	PropertyCache->PivotLocation = this->PivotLocation;
 	PropertyCache->PlaceMode = this->PlaceMode;
+	PropertyCache->bAlignShapeToPlacementSurface = this->bAlignShapeToPlacementSurface;
 }
 
 void UProceduralShapeToolProperties::RestoreProperties(UInteractiveTool* RestoreToTool)
@@ -77,6 +79,7 @@ void UProceduralShapeToolProperties::RestoreProperties(UInteractiveTool* Restore
 	this->Subdivisions = PropertyCache->Subdivisions;
 	this->PivotLocation = PropertyCache->PivotLocation;
 	this->PlaceMode = PropertyCache->PlaceMode;
+	this->bAlignShapeToPlacementSurface = PropertyCache->bAlignShapeToPlacementSurface;
 }
 
 namespace
@@ -92,6 +95,7 @@ namespace
 	  { TEXT("Rotation"),      EMakeMeshShapeType::All },
 	  { TEXT("PlaceMode"),     EMakeMeshShapeType::All },
 	  { TEXT("PivotLocation"), EMakeMeshShapeType::All },
+	  { TEXT("bAlignShapeToPlacementSurface"), EMakeMeshShapeType::All },
 	  { TEXT("Slices"),        EMakeMeshShapeType::Cylinder | EMakeMeshShapeType::Cone | EMakeMeshShapeType::Arrow | EMakeMeshShapeType::RoundedRectangle | EMakeMeshShapeType::Disc | EMakeMeshShapeType::PuncturedDisc | EMakeMeshShapeType::Sphere },
 	  { TEXT("Subdivisions"),  EMakeMeshShapeType::Box | EMakeMeshShapeType::Rectangle | EMakeMeshShapeType::RoundedRectangle | EMakeMeshShapeType::Disc | EMakeMeshShapeType::PuncturedDisc | EMakeMeshShapeType::Cylinder |
 							   EMakeMeshShapeType::Cone | EMakeMeshShapeType::Arrow | EMakeMeshShapeType::SphericalBox }
@@ -232,7 +236,12 @@ void UAddPrimitiveTool::UpdatePreviewPosition(const FInputDeviceRay& DeviceClick
 		bHit = TargetWorld->LineTraceSingleByObjectType(Result, RayStart, RayEnd, QueryParams);
 		if (bHit)
 		{
-			ShapeFrame = FFrame3f(Result.ImpactPoint, Result.ImpactNormal);
+			FVector3f Normal = Result.ImpactNormal;
+			if (!ShapeSettings->bAlignShapeToPlacementSurface)
+			{
+				Normal = FVector3f::UnitZ();
+			}
+			ShapeFrame = FFrame3f(Result.ImpactPoint, Normal);
 			ShapeFrame.ConstrainedAlignPerpAxes();
 		}
 	}
@@ -315,6 +324,10 @@ void UAddPrimitiveTool::UpdatePreviewMesh()
 	{
 		TargetOrigin.Z = Bounds.Min.Z;
 	}
+	else if (ShapeSettings->PivotLocation == EMakeMeshPivotLocation::Top)
+	{
+		TargetOrigin.Z = Bounds.Max.Z;
+	}
 	for (int vid : NewMesh.VertexIndicesItr())
 	{
 		FVector3d Pos = NewMesh.GetVertex(vid);
diff --git a/Engine/Plugins/Experimental/MeshModelingToolset/Source/MeshModelingTools/Public/AddPrimitiveTool.h b/Engine/Plugins/Experimental/MeshModelingToolset/Source/MeshModelingTools/Public/AddPrimitiveTool.h
--- a/Engine/Plugins/Experimental/MeshModelingToolset/Source/MeshModelingTools/Public/AddPrimitiveTool.h
+++ b/Engine/Plugins/Experimental/MeshModelingToolset/Source/MeshModelingTools/Public/AddPrimitiveTool.h
@@ -66,7 +66,8 @@ UENUM()
 enum class EMakeMeshPivotLocation : uint8
 {
 	Base,
-	Centered
+	Centered,
+	Top
 };
 
 UCLASS()
@@ -107,6 +108,10 @@ public:
 	UPROPERTY(EditAnywhere, Category = ShapeSettings)
     EMakeMeshPivotLocation PivotLocation;
 
+	/** Align shape to placement surface */
+	UPROPERTY(EditAnywhere, Category = ShapeSettings, meta = (EditCondition = "PlaceMode == EMakeMeshPlacementType::OnScene"))
+	bool bAlignShapeToPlacementSurface = true;
+
 	///** Start Angle of Shape */
 	//UPROPERTY(EditAnywhere, Category = ShapeSettings, meta = (DisplayName = "Start Angle", UIMin = "0.0", UIMax = "360.0", ClampMin = "-10000", ClampMax = "10000.0"))
 	//float StartAngle;
