Add feature radius parameter to add primitive tool, to capture circular features e.g. on rounded rectangle corners or the inside of a punctured disc

#rnx

diff --git a/Engine/Plugins/Experimental/MeshModelingToolset/Source/MeshModelingTools/Private/AddPrimitiveTool.cpp b/Engine/Plugins/Experimental/MeshModelingToolset/Source/MeshModelingTools/Private/AddPrimitiveTool.cpp
--- a/Engine/Plugins/Experimental/MeshModelingToolset/Source/MeshModelingTools/Private/AddPrimitiveTool.cpp
+++ b/Engine/Plugins/Experimental/MeshModelingToolset/Source/MeshModelingTools/Private/AddPrimitiveTool.cpp
@@ -47,6 +47,7 @@ UProceduralShapeToolProperties::UProceduralShapeToolProperties()
 	Shape = EMakeMeshShapeType::Box;
 	Width = 100;
 	Height = 200;
+	FeatureRadius = 25;
 	//StartAngle = 0;
 	//EndAngle = 360;
 	Slices = 16;
@@ -62,6 +63,7 @@ void UProceduralShapeToolProperties::SaveProperties(UInteractiveTool* SaveFromTo
 	PropertyCache->Shape = this->Shape;
 	PropertyCache->Width = this->Width;
 	PropertyCache->Height = this->Height;
+	PropertyCache->FeatureRadius = this->FeatureRadius;
 	PropertyCache->Slices = this->Slices;
 	PropertyCache->Subdivisions = this->Subdivisions;
 	PropertyCache->PivotLocation = this->PivotLocation;
@@ -75,6 +77,7 @@ void UProceduralShapeToolProperties::RestoreProperties(UInteractiveTool* Restore
 	this->Shape = PropertyCache->Shape;
 	this->Width = PropertyCache->Width;
 	this->Height = PropertyCache->Height;
+	this->FeatureRadius = PropertyCache->FeatureRadius;
 	this->Slices = PropertyCache->Slices;
 	this->Subdivisions = PropertyCache->Subdivisions;
 	this->PivotLocation = PropertyCache->PivotLocation;
@@ -92,6 +95,7 @@ namespace
 	  { TEXT("Shape"),         EMakeMeshShapeType::All },
 	  { TEXT("Width"),         EMakeMeshShapeType::All },
 	  { TEXT("Height"),        EMakeMeshShapeType::Box | EMakeMeshShapeType::Cylinder | EMakeMeshShapeType::Cone | EMakeMeshShapeType::Arrow | EMakeMeshShapeType::Rectangle | EMakeMeshShapeType::RoundedRectangle },
+	  { TEXT("FeatureRadius"),		   EMakeMeshShapeType::Arrow | EMakeMeshShapeType::RoundedRectangle | EMakeMeshShapeType::PuncturedDisc },
 	  { TEXT("Rotation"),      EMakeMeshShapeType::All },
 	  { TEXT("PlaceMode"),     EMakeMeshShapeType::All },
 	  { TEXT("PivotLocation"), EMakeMeshShapeType::All },
@@ -395,7 +399,7 @@ void UAddPrimitiveTool::GenerateRoundedRectangle(FDynamicMesh3* OutMesh)
 	RectGen.Width = ShapeSettings->Width;
 	RectGen.Height = ShapeSettings->Height;
 	RectGen.WidthVertexCount = RectGen.HeightVertexCount = ShapeSettings->Subdivisions + 2;
-	RectGen.Radius = ShapeSettings->Width * .5; // TODO replace with a radius parameter
+	RectGen.Radius = ShapeSettings->FeatureRadius;
 	RectGen.AngleSamples = ShapeSettings->Slices;
 	RectGen.Generate();
 	OutMesh->Copy(&RectGen);
@@ -417,7 +421,7 @@ void UAddPrimitiveTool::GeneratePuncturedDisc(FDynamicMesh3* OutMesh)
 {
 	FPuncturedDiscMeshGenerator Gen;
 	Gen.Radius = ShapeSettings->Width * 0.5f;
-	Gen.HoleRadius = Gen.Radius * .5;
+	Gen.HoleRadius = FMath::Min(ShapeSettings->FeatureRadius, Gen.Radius * .999f); // hole cannot be bigger than outer radius
 	Gen.AngleSamples = ShapeSettings->Slices;
 	Gen.RadialSamples = ShapeSettings->Subdivisions;
 	Gen.Generate();
@@ -457,7 +461,7 @@ void UAddPrimitiveTool::GenerateCone(FDynamicMesh3* OutMesh)
 void UAddPrimitiveTool::GenerateArrow(FDynamicMesh3* OutMesh)
 {
 	FArrowGenerator ArrowGen;
-	ArrowGen.StickRadius = ShapeSettings->Width * .25f;
+	ArrowGen.StickRadius = ShapeSettings->FeatureRadius;
 	ArrowGen.StickLength = ShapeSettings->Height * .25f;
 	ArrowGen.HeadBaseRadius = ShapeSettings->Width * .5f;
 	ArrowGen.TipRadius = .01f;
diff --git a/Engine/Plugins/Experimental/MeshModelingToolset/Source/MeshModelingTools/Public/AddPrimitiveTool.h b/Engine/Plugins/Experimental/MeshModelingToolset/Source/MeshModelingTools/Public/AddPrimitiveTool.h
--- a/Engine/Plugins/Experimental/MeshModelingToolset/Source/MeshModelingTools/Public/AddPrimitiveTool.h
+++ b/Engine/Plugins/Experimental/MeshModelingToolset/Source/MeshModelingTools/Public/AddPrimitiveTool.h
@@ -96,6 +96,10 @@ public:
 	UPROPERTY(EditAnywhere, Category = ShapeSettings, meta = (DisplayName = "Height", UIMin = "1.0", UIMax = "1000.0", ClampMin = "0.0001", ClampMax = "1000000.0"))
 	float Height;
 
+	/** Radius of additional circular features of the shape (not implicitly defined by the width of the shape) */
+	UPROPERTY(EditAnywhere, Category = ShapeSettings, meta = (DisplayName = "Feature Radius", UIMin = "1.0", UIMax = "1000.0", ClampMin = "0.0001", ClampMax = "1000000.0"))
+	float FeatureRadius;
+
 	/** Rotation around up axis */
 	UPROPERTY(EditAnywhere, Category = ShapeSettings, meta = (DisplayName = "Rotation", UIMin = "0.0", UIMax = "360.0"))
 	float Rotation;
