Fix for correct header location
#rnx

diff --git a/Engine/Plugins/Experimental/MeshModelingToolset/Source/MeshModelingTools/Private/DynamicMeshSculptTool.cpp b/Engine/Plugins/Experimental/MeshModelingToolset/Source/MeshModelingTools/Private/DynamicMeshSculptTool.cpp
--- a/Engine/Plugins/Experimental/MeshModelingToolset/Source/MeshModelingTools/Private/DynamicMeshSculptTool.cpp
+++ b/Engine/Plugins/Experimental/MeshModelingToolset/Source/MeshModelingTools/Private/DynamicMeshSculptTool.cpp
@@ -14,7 +14,7 @@
 #include "MeshIndexUtil.h"
 #include "Drawing/MeshDebugDrawing.h"
 #include "PreviewMesh.h"
-#include "BaseTools/ToolSetupUtil.h"
+#include "ToolSetupUtil.h"
 
 #include "Changes/MeshVertexChange.h"
 #include "Changes/MeshChange.h"
diff --git a/Engine/Plugins/Experimental/MeshModelingToolset/Source/MeshModelingTools/Private/EditNormalsTool.cpp b/Engine/Plugins/Experimental/MeshModelingToolset/Source/MeshModelingTools/Private/EditNormalsTool.cpp
--- a/Engine/Plugins/Experimental/MeshModelingToolset/Source/MeshModelingTools/Private/EditNormalsTool.cpp
+++ b/Engine/Plugins/Experimental/MeshModelingToolset/Source/MeshModelingTools/Private/EditNormalsTool.cpp
@@ -4,7 +4,7 @@
 #include "InteractiveToolManager.h"
 #include "ToolBuilderUtil.h"
 
-#include "BaseTools/ToolSetupUtil.h"
+#include "ToolSetupUtil.h"
 
 #include "DynamicMesh3.h"
 #include "BaseBehaviors/MultiClickSequenceInputBehavior.h"
diff --git a/Engine/Plugins/Experimental/MeshModelingToolset/Source/MeshModelingTools/Private/MeshSelectionTool.cpp b/Engine/Plugins/Experimental/MeshModelingToolset/Source/MeshModelingTools/Private/MeshSelectionTool.cpp
--- a/Engine/Plugins/Experimental/MeshModelingToolset/Source/MeshModelingTools/Private/MeshSelectionTool.cpp
+++ b/Engine/Plugins/Experimental/MeshModelingToolset/Source/MeshModelingTools/Private/MeshSelectionTool.cpp
@@ -10,7 +10,7 @@
 #include "Changes/MeshChange.h"
 #include "MeshIndexUtil.h"
 #include "AssetGenerationUtil.h"
-#include "BaseTools/ToolSetupUtil.h"
+#include "ToolSetupUtil.h"
 
 #define LOCTEXT_NAMESPACE "UMeshSelectionTool"
 
diff --git a/Engine/Plugins/Experimental/MeshModelingToolset/Source/MeshModelingTools/Private/PlaneCutTool.cpp b/Engine/Plugins/Experimental/MeshModelingToolset/Source/MeshModelingTools/Private/PlaneCutTool.cpp
--- a/Engine/Plugins/Experimental/MeshModelingToolset/Source/MeshModelingTools/Private/PlaneCutTool.cpp
+++ b/Engine/Plugins/Experimental/MeshModelingToolset/Source/MeshModelingTools/Private/PlaneCutTool.cpp
@@ -4,7 +4,7 @@
 #include "InteractiveToolManager.h"
 #include "ToolBuilderUtil.h"
 
-#include "BaseTools/ToolSetupUtil.h"
+#include "ToolSetupUtil.h"
 
 #include "DynamicMesh3.h"
 #include "BaseBehaviors/MultiClickSequenceInputBehavior.h"
diff --git a/Engine/Plugins/Experimental/MeshModelingToolset/Source/MeshModelingTools/Private/RemoveOccludedTrianglesTool.cpp b/Engine/Plugins/Experimental/MeshModelingToolset/Source/MeshModelingTools/Private/RemoveOccludedTrianglesTool.cpp
--- a/Engine/Plugins/Experimental/MeshModelingToolset/Source/MeshModelingTools/Private/RemoveOccludedTrianglesTool.cpp
+++ b/Engine/Plugins/Experimental/MeshModelingToolset/Source/MeshModelingTools/Private/RemoveOccludedTrianglesTool.cpp
@@ -4,7 +4,7 @@
 #include "InteractiveToolManager.h"
 #include "ToolBuilderUtil.h"
 
-#include "BaseTools/ToolSetupUtil.h"
+#include "ToolSetupUtil.h"
 
 #include "DynamicMesh3.h"
 #include "BaseBehaviors/MultiClickSequenceInputBehavior.h"
diff --git a/Engine/Plugins/Experimental/MeshModelingToolset/Source/MeshModelingTools/Private/SmoothMeshTool.cpp b/Engine/Plugins/Experimental/MeshModelingToolset/Source/MeshModelingTools/Private/SmoothMeshTool.cpp
--- a/Engine/Plugins/Experimental/MeshModelingToolset/Source/MeshModelingTools/Private/SmoothMeshTool.cpp
+++ b/Engine/Plugins/Experimental/MeshModelingToolset/Source/MeshModelingTools/Private/SmoothMeshTool.cpp
@@ -11,7 +11,7 @@
 #include "SimpleDynamicMeshComponent.h"
 
 #include "SceneManagement.h" // for FPrimitiveDrawInterface
-#include "BaseTools/ToolSetupUtil.h"
+#include "ToolSetupUtil.h"
 
 // Smoothing operators
 #include "SmoothingOps/IterativeSmoothingOp.h"
diff --git a/Engine/Plugins/Experimental/MeshModelingToolset/Source/MeshModelingTools/Private/TransformMeshesTool.cpp b/Engine/Plugins/Experimental/MeshModelingToolset/Source/MeshModelingTools/Private/TransformMeshesTool.cpp
--- a/Engine/Plugins/Experimental/MeshModelingToolset/Source/MeshModelingTools/Private/TransformMeshesTool.cpp
+++ b/Engine/Plugins/Experimental/MeshModelingToolset/Source/MeshModelingTools/Private/TransformMeshesTool.cpp
@@ -4,7 +4,7 @@
 #include "InteractiveToolManager.h"
 #include "InteractiveGizmoManager.h"
 #include "ToolBuilderUtil.h"
-#include "BaseTools/ToolSetupUtil.h"
+#include "ToolSetupUtil.h"
 #include "DynamicMesh3.h"
 #include "BaseBehaviors/ClickDragBehavior.h"
 
diff --git a/Engine/Plugins/Experimental/MeshModelingToolset/Source/MeshModelingTools/Private/UVProjectionTool.cpp b/Engine/Plugins/Experimental/MeshModelingToolset/Source/MeshModelingTools/Private/UVProjectionTool.cpp
--- a/Engine/Plugins/Experimental/MeshModelingToolset/Source/MeshModelingTools/Private/UVProjectionTool.cpp
+++ b/Engine/Plugins/Experimental/MeshModelingToolset/Source/MeshModelingTools/Private/UVProjectionTool.cpp
@@ -4,7 +4,7 @@
 #include "InteractiveToolManager.h"
 #include "ToolBuilderUtil.h"
 
-#include "BaseTools/ToolSetupUtil.h"
+#include "ToolSetupUtil.h"
 
 #include "DynamicMesh3.h"
 #include "BaseBehaviors/MultiClickSequenceInputBehavior.h"
diff --git a/Engine/Plugins/Experimental/MeshModelingToolset/Source/MeshModelingToolsEditorOnly/Private/AttributeEditorTool.cpp b/Engine/Plugins/Experimental/MeshModelingToolset/Source/MeshModelingToolsEditorOnly/Private/AttributeEditorTool.cpp
--- a/Engine/Plugins/Experimental/MeshModelingToolset/Source/MeshModelingToolsEditorOnly/Private/AttributeEditorTool.cpp
+++ b/Engine/Plugins/Experimental/MeshModelingToolset/Source/MeshModelingToolsEditorOnly/Private/AttributeEditorTool.cpp
@@ -4,7 +4,7 @@
 #include "InteractiveToolManager.h"
 #include "ToolBuilderUtil.h"
 
-#include "BaseTools/ToolSetupUtil.h"
+#include "ToolSetupUtil.h"
 #include "MathUtil.h"
 
 #include "MeshDescription.h"
diff --git a/Engine/Plugins/Experimental/MeshModelingToolset/Source/MeshModelingToolsEditorOnly/Private/MergeMeshesTool.cpp b/Engine/Plugins/Experimental/MeshModelingToolset/Source/MeshModelingToolsEditorOnly/Private/MergeMeshesTool.cpp
--- a/Engine/Plugins/Experimental/MeshModelingToolset/Source/MeshModelingToolsEditorOnly/Private/MergeMeshesTool.cpp
+++ b/Engine/Plugins/Experimental/MeshModelingToolset/Source/MeshModelingToolsEditorOnly/Private/MergeMeshesTool.cpp
@@ -3,7 +3,7 @@
 #include "MergeMeshesTool.h"
 #include "InteractiveToolManager.h"
 #include "ToolBuilderUtil.h"
-#include "BaseTools/ToolSetupUtil.h"
+#include "ToolSetupUtil.h"
 #include "AssetGenerationUtil.h"
 #include "Selection/ToolSelectionUtil.h"
 
diff --git a/Engine/Plugins/Experimental/MeshModelingToolset/Source/MeshModelingToolsEditorOnly/Private/ParameterizeMeshTool.cpp b/Engine/Plugins/Experimental/MeshModelingToolset/Source/MeshModelingToolsEditorOnly/Private/ParameterizeMeshTool.cpp
--- a/Engine/Plugins/Experimental/MeshModelingToolset/Source/MeshModelingToolsEditorOnly/Private/ParameterizeMeshTool.cpp
+++ b/Engine/Plugins/Experimental/MeshModelingToolset/Source/MeshModelingToolsEditorOnly/Private/ParameterizeMeshTool.cpp
@@ -15,7 +15,7 @@
 #include "MeshDescriptionToDynamicMesh.h"
 
 #include "ParameterizationOps/ParameterizeMeshOp.h"
-#include "BaseTools/ToolSetupUtil.h"
+#include "ToolSetupUtil.h"
 
 #define LOCTEXT_NAMESPACE "UParameterizeMeshTool"
 
diff --git a/Engine/Plugins/Experimental/MeshModelingToolset/Source/MeshModelingToolsEditorOnly/Private/PolygonOnMeshTool.cpp b/Engine/Plugins/Experimental/MeshModelingToolset/Source/MeshModelingToolsEditorOnly/Private/PolygonOnMeshTool.cpp
--- a/Engine/Plugins/Experimental/MeshModelingToolset/Source/MeshModelingToolsEditorOnly/Private/PolygonOnMeshTool.cpp
+++ b/Engine/Plugins/Experimental/MeshModelingToolset/Source/MeshModelingToolsEditorOnly/Private/PolygonOnMeshTool.cpp
@@ -4,7 +4,7 @@
 #include "InteractiveToolManager.h"
 #include "ToolBuilderUtil.h"
 
-#include "BaseTools/ToolSetupUtil.h"
+#include "ToolSetupUtil.h"
 
 #include "DynamicMesh3.h"
 #include "BaseBehaviors/MultiClickSequenceInputBehavior.h"
diff --git a/Engine/Plugins/Experimental/MeshModelingToolset/Source/MeshModelingToolsEditorOnly/Private/VoxelCSGMeshesTool.cpp b/Engine/Plugins/Experimental/MeshModelingToolset/Source/MeshModelingToolsEditorOnly/Private/VoxelCSGMeshesTool.cpp
--- a/Engine/Plugins/Experimental/MeshModelingToolset/Source/MeshModelingToolsEditorOnly/Private/VoxelCSGMeshesTool.cpp
+++ b/Engine/Plugins/Experimental/MeshModelingToolset/Source/MeshModelingToolsEditorOnly/Private/VoxelCSGMeshesTool.cpp
@@ -3,7 +3,7 @@
 #include "VoxelCSGMeshesTool.h"
 #include "InteractiveToolManager.h"
 #include "ToolBuilderUtil.h"
-#include "BaseTools/ToolSetupUtil.h"
+#include "ToolSetupUtil.h"
 #include "AssetGenerationUtil.h"
 #include "Selection/ToolSelectionUtil.h"
 
