Flip Normals edit action in mesh selection tool
#rnx

diff --git a/Engine/Plugins/Experimental/MeshModelingToolset/Source/MeshModelingTools/Private/MeshSelectionTool.cpp b/Engine/Plugins/Experimental/MeshModelingToolset/Source/MeshModelingTools/Private/MeshSelectionTool.cpp
--- a/Engine/Plugins/Experimental/MeshModelingToolset/Source/MeshModelingTools/Private/MeshSelectionTool.cpp
+++ b/Engine/Plugins/Experimental/MeshModelingToolset/Source/MeshModelingTools/Private/MeshSelectionTool.cpp
@@ -684,6 +684,10 @@ void UMeshSelectionTool::ApplyAction(EMeshSelectionToolActions ActionType)
 		case EMeshSelectionToolActions::SeparateSelected:
 			SeparateSelectedTriangles();
 			break;
+
+		case EMeshSelectionToolActions::FlipSelected:
+			FlipSelectedTriangles();
+			break;
 	}
 }
 
@@ -978,5 +982,39 @@ void UMeshSelectionTool::SeparateSelectedTriangles()
 
 
 
+void UMeshSelectionTool::FlipSelectedTriangles()
+{
+	check(SelectionType == EMeshSelectionElementType::Face);
+	TArray<int32> SelectedFaces = Selection->GetElements(EMeshSelectionElementType::Face);
+	if (SelectedFaces.Num() == 0)
+	{
+		return;
+	}
+
+	TUniquePtr<FToolCommandChangeSequence> ChangeSeq = MakeUnique<FToolCommandChangeSequence>();
+
+	// invert triangles and emit triangles change
+	TUniquePtr<FMeshChange> MeshChange = PreviewMesh->TrackedEditMesh(
+		[&SelectedFaces](FDynamicMesh3& Mesh, FDynamicMeshChangeTracker& ChangeTracker)
+	{
+		for (int TID : SelectedFaces)
+		{
+			ChangeTracker.SaveTriangle(TID, true);
+		}
+		FDynamicMeshEditor Editor(&Mesh);
+		Editor.ReverseTriangleOrientations(SelectedFaces, true);
+	});
+	ChangeSeq->AppendChange(PreviewMesh, MoveTemp(MeshChange));
+
+	// emit combined change sequence
+	GetToolManager()->EmitObjectChange(this, MoveTemp(ChangeSeq), LOCTEXT("MeshSelectionToolFlipFaces", "Flip Face Orientations"));
+
+	bHaveModifiedMesh = true;
+}
+
+
+
+
+
 
 #undef LOCTEXT_NAMESPACE
\ No newline at end of file
diff --git a/Engine/Plugins/Experimental/MeshModelingToolset/Source/MeshModelingTools/Public/MeshSelectionTool.h b/Engine/Plugins/Experimental/MeshModelingToolset/Source/MeshModelingTools/Public/MeshSelectionTool.h
--- a/Engine/Plugins/Experimental/MeshModelingToolset/Source/MeshModelingTools/Public/MeshSelectionTool.h
+++ b/Engine/Plugins/Experimental/MeshModelingToolset/Source/MeshModelingTools/Public/MeshSelectionTool.h
@@ -39,7 +39,8 @@ enum class EMeshSelectionToolActions
 	ExpandToConnected,
 
 	DeleteSelected,
-	SeparateSelected
+	SeparateSelected,
+	FlipSelected
 };
 
 
@@ -119,6 +120,12 @@ public:
 	{
 		PostAction(EMeshSelectionToolActions::SeparateSelected);
 	}
+
+	UFUNCTION(CallInEditor, Category = MeshEdits, meta = (DisplayName = "Flip Orientations"))
+		void FlipOrientations() 
+	{
+		PostAction(EMeshSelectionToolActions::FlipSelected);
+	}
 };
 
 
@@ -292,6 +299,7 @@ protected:
 
 	void DeleteSelectedTriangles();
 	void SeparateSelectedTriangles();
+	void FlipSelectedTriangles();
 
 	bool bHaveModifiedMesh = false;
 };
