fix convert to polygons tool transaction name key
#rnx

diff --git a/Engine/Plugins/Experimental/MeshModelingToolset/Source/MeshModelingTools/Private/ConvertToPolygonsTool.cpp b/Engine/Plugins/Experimental/MeshModelingToolset/Source/MeshModelingTools/Private/ConvertToPolygonsTool.cpp
--- a/Engine/Plugins/Experimental/MeshModelingToolset/Source/MeshModelingTools/Private/ConvertToPolygonsTool.cpp
+++ b/Engine/Plugins/Experimental/MeshModelingToolset/Source/MeshModelingTools/Private/ConvertToPolygonsTool.cpp
@@ -89,7 +89,7 @@ void UConvertToPolygonsTool::Shutdown(EToolShutdownType ShutdownType)
 
 	if (ShutdownType == EToolShutdownType::Accept)
 	{
-		GetToolManager()->BeginUndoTransaction(LOCTEXT("DeformMeshToolTransactionName", "Convert to Polygons"));
+		GetToolManager()->BeginUndoTransaction(LOCTEXT("ConvertToPolygonsToolTransactionName", "Convert to Polygons"));
 		ComponentTarget->CommitMesh([=](FMeshDescription* MeshDescription)
 		{
 			ConvertToPolygons(MeshDescription);
