Add Instance mode in Add Primitive Tool

#rnx

diff --git a/Engine/Plugins/Experimental/MeshModelingToolset/Source/MeshModelingTools/Private/AddPrimitiveTool.cpp b/Engine/Plugins/Experimental/MeshModelingToolset/Source/MeshModelingTools/Private/AddPrimitiveTool.cpp
--- a/Engine/Plugins/Experimental/MeshModelingToolset/Source/MeshModelingTools/Private/AddPrimitiveTool.cpp
+++ b/Engine/Plugins/Experimental/MeshModelingToolset/Source/MeshModelingTools/Private/AddPrimitiveTool.cpp
@@ -7,6 +7,11 @@
 #include "Selection/ToolSelectionUtil.h"
 #include "AssetGenerationUtil.h"
 
+#include "Components/StaticMeshComponent.h"
+#include "Engine/StaticMesh.h"
+#include "Engine/StaticMeshActor.h"
+
+
 #include "MeshDescriptionBuilder.h"
 #include "Generators/SweepGenerator.h"
 #include "Generators/GridBoxMeshGenerator.h"
@@ -100,6 +105,7 @@ namespace
 	  { TEXT("PlaceMode"),     EMakeMeshShapeType::All },
 	  { TEXT("PivotLocation"), EMakeMeshShapeType::All },
 	  { TEXT("bAlignShapeToPlacementSurface"), EMakeMeshShapeType::All },
+	  { TEXT("bInstanceLastCreatedAssetIfPossible"), EMakeMeshShapeType::All },
 	  { TEXT("Slices"),        EMakeMeshShapeType::Cylinder | EMakeMeshShapeType::Cone | EMakeMeshShapeType::Arrow | EMakeMeshShapeType::RoundedRectangle | EMakeMeshShapeType::Disc | EMakeMeshShapeType::PuncturedDisc | EMakeMeshShapeType::Sphere | EMakeMeshShapeType::Torus },
 	  { TEXT("Subdivisions"),  EMakeMeshShapeType::Box | EMakeMeshShapeType::Rectangle | EMakeMeshShapeType::RoundedRectangle | EMakeMeshShapeType::Disc | EMakeMeshShapeType::PuncturedDisc | EMakeMeshShapeType::Cylinder |
 							   EMakeMeshShapeType::Cone | EMakeMeshShapeType::Arrow | EMakeMeshShapeType::SphericalBox | EMakeMeshShapeType::Torus }
@@ -349,24 +355,51 @@ void UAddPrimitiveTool::UpdatePreviewMesh()
 
 void UAddPrimitiveTool::OnClicked(const FInputDeviceRay& DeviceClickPos)
 {
+	const UEnum* const MakeMeshShapeTypeEnum = StaticEnum<EMakeMeshShapeType>();		
+	FString ShapeTypeName = MakeMeshShapeTypeEnum->GetNameStringByValue((int64)ShapeSettings->Shape);
+	UMaterialInterface* Material = PreviewMesh->GetMaterial();
+
+	if (ShapeSettings->bInstanceLastCreatedAssetIfPossible && IsEquivalentLastGeneratedAsset())
+	{
+		GetToolManager()->BeginUndoTransaction(LOCTEXT("AddPrimitiveToolTransactionName", "Add Primitive Mesh"));
+		FActorSpawnParameters SpawnParameters;
+		SpawnParameters.Template = LastGenerated->Actor;
+		FTransform3d CurTransform(PreviewMesh->GetTransform());
+		FRotator Rotation(0.0f, 0.0f, 0.0f);
+		AStaticMeshActor* CloneActor = TargetWorld->SpawnActor<AStaticMeshActor>(FVector::ZeroVector, Rotation, SpawnParameters);
+		// some properties must be manually set on the component because they will not persist reliably through the spawn template (especially if the actor creation was undone)
+		CloneActor->GetStaticMeshComponent()->SetWorldTransform(PreviewMesh->GetTransform());
+		CloneActor->GetStaticMeshComponent()->SetStaticMesh(LastGenerated->StaticMesh);
+		CloneActor->GetStaticMeshComponent()->SetMaterial(0, Material);
+		CloneActor->SetActorLabel(LastGenerated->Label);
+		// select newly-created object
+		ToolSelectionUtil::SetNewActorSelection(GetToolManager(), CloneActor);
+		GetToolManager()->EndUndoTransaction();
+
+		return;
+	}
+
+
 #if WITH_EDITOR
 	const FDynamicMesh3* CurMesh = PreviewMesh->GetPreviewDynamicMesh();
 	FTransform3d CurTransform(PreviewMesh->GetTransform());
-	UMaterialInterface* Material = PreviewMesh->GetMaterial();
+	
 	GetToolManager()->BeginUndoTransaction(LOCTEXT("AddPrimitiveToolTransactionName", "Add Primitive Mesh"));
 
-	const UEnum* const MakeMeshShapeTypeEnum = StaticEnum<EMakeMeshShapeType>();		
-	FString ShapeTypeName = MakeMeshShapeTypeEnum->GetNameStringByValue((uint32)ShapeSettings->Shape);
-
-	AActor* NewActor = AssetGenerationUtil::GenerateStaticMeshActor(
+	LastGenerated = NewObject<ULastActorInfo>(this);
+	LastGenerated->ShapeSettings = DuplicateObject(ShapeSettings, nullptr);
+	LastGenerated->MaterialProperties = DuplicateObject(MaterialProperties, nullptr);
+	LastGenerated->Actor = AssetGenerationUtil::GenerateStaticMeshActor(
 		AssetAPI, TargetWorld,
 		CurMesh, CurTransform, ShapeTypeName,
 		AssetGenerationUtil::GetDefaultAutoGeneratedAssetPath(),
 		Material
 	);
+	LastGenerated->StaticMesh = CastChecked<AStaticMeshActor>(LastGenerated->Actor)->GetStaticMeshComponent()->GetStaticMesh();
+	LastGenerated->Label = LastGenerated->Actor->GetActorLabel();
 
 	// select newly-created object
-	ToolSelectionUtil::SetNewActorSelection(GetToolManager(), NewActor);
+	ToolSelectionUtil::SetNewActorSelection(GetToolManager(), LastGenerated->Actor);
 
 	GetToolManager()->EndUndoTransaction();
 #else
diff --git a/Engine/Plugins/Experimental/MeshModelingToolset/Source/MeshModelingTools/Public/AddPrimitiveTool.h b/Engine/Plugins/Experimental/MeshModelingToolset/Source/MeshModelingTools/Public/AddPrimitiveTool.h
--- a/Engine/Plugins/Experimental/MeshModelingToolset/Source/MeshModelingTools/Public/AddPrimitiveTool.h
+++ b/Engine/Plugins/Experimental/MeshModelingToolset/Source/MeshModelingTools/Public/AddPrimitiveTool.h
@@ -117,6 +117,10 @@ public:
 	UPROPERTY(EditAnywhere, Category = ShapeSettings, meta = (EditCondition = "PlaceMode == EMakeMeshPlacementType::OnScene"))
 	bool bAlignShapeToPlacementSurface = true;
 
+	/** If the shape settings haven't changed, create instances of the last created asset rather than creating a whole new asset.  If false, all created actors will have separate underlying mesh assets. */
+	UPROPERTY(EditAnywhere, Category = ShapeSettings)
+	bool bInstanceLastCreatedAssetIfPossible = true;
+
 	///** Start Angle of Shape */
 	//UPROPERTY(EditAnywhere, Category = ShapeSettings, meta = (DisplayName = "Start Angle", UIMin = "0.0", UIMax = "360.0", ClampMin = "-10000", ClampMax = "10000.0"))
 	//float StartAngle;
@@ -134,14 +138,37 @@ public:
 	UPROPERTY(EditAnywhere, Category = ShapeSettings, meta = (DisplayName = "Subdivisions", UIMin = "0", UIMax = "100", ClampMin = "0", ClampMax = "4000"))
 	int Subdivisions;
 
-
 	virtual void SaveProperties(UInteractiveTool* SaveFromTool) override;
 	virtual void RestoreProperties(UInteractiveTool* RestoreToTool) override;
 };
 
 
 
+UCLASS(Transient)
+class MESHMODELINGTOOLS_API ULastActorInfo : public UObject
+{
+	GENERATED_BODY()
+
+public:
+	FString Label = "";
 
+	UPROPERTY()
+	AActor* Actor = nullptr;
+	
+	UPROPERTY()
+	UStaticMesh* StaticMesh = nullptr;
+
+	UPROPERTY()
+	UProceduralShapeToolProperties* ShapeSettings;
+
+	UPROPERTY()
+	UNewMeshMaterialProperties* MaterialProperties;
+
+	bool IsInvalid()
+	{
+		return Actor == nullptr || StaticMesh == nullptr || ShapeSettings == nullptr || MaterialProperties == nullptr;
+	}
+};
 
 
 
@@ -186,10 +213,38 @@ protected:
 	UNewMeshMaterialProperties* MaterialProperties;
 
 
+	/**
+	 * Checks if the passed-in settings would create the same asset as the current settings
+	 */
+	bool IsEquivalentLastGeneratedAsset()
+	{
+		if (LastGenerated == nullptr || LastGenerated->IsInvalid())
+		{
+			return false;
+		}
+		// manual diff because not all shape setting changes result in a different asset (e.g. some just affect the transform)
+		return
+			LastGenerated->ShapeSettings->Subdivisions == ShapeSettings->Subdivisions &&
+			LastGenerated->ShapeSettings->Slices == ShapeSettings->Slices &&
+			LastGenerated->ShapeSettings->PivotLocation == ShapeSettings->PivotLocation &&
+			LastGenerated->ShapeSettings->FeatureRadius == ShapeSettings->FeatureRadius &&
+			LastGenerated->ShapeSettings->Height == ShapeSettings->Height &&
+			LastGenerated->ShapeSettings->Width == ShapeSettings->Width &&
+			LastGenerated->ShapeSettings->Shape == ShapeSettings->Shape &&
+			LastGenerated->MaterialProperties->UVScale == MaterialProperties->UVScale &&
+			LastGenerated->MaterialProperties->bWorldSpaceUVScale == MaterialProperties->bWorldSpaceUVScale
+			;
+	}
+
 
 	UPROPERTY()
 	UPreviewMesh* PreviewMesh;
 
+	UPROPERTY()
+	ULastActorInfo* LastGenerated;
+
+
+
 protected:
 	UWorld* TargetWorld;
 	IToolsContextAssetAPI* AssetAPI;
