ModelingTools: Handle case where user disables Gizmo in DrawPolygon, resulting in it being nullptr in ctrl+click handler. Also update Gizmo position after it is re-shown.
#rnx

diff --git a/Engine/Plugins/Experimental/MeshModelingToolset/Source/MeshModelingTools/Private/DrawPolygonTool.cpp b/Engine/Plugins/Experimental/MeshModelingToolset/Source/MeshModelingTools/Private/DrawPolygonTool.cpp
--- a/Engine/Plugins/Experimental/MeshModelingToolset/Source/MeshModelingTools/Private/DrawPolygonTool.cpp
+++ b/Engine/Plugins/Experimental/MeshModelingToolset/Source/MeshModelingTools/Private/DrawPolygonTool.cpp
@@ -930,7 +930,10 @@ void UDrawPolygonTool::SetDrawPlaneFromWorldPos(const FVector& Position, const F
 
 	SnapEngine.Plane = FFrame3d((FVector3d)DrawPlane.Origin, (FQuaterniond)DrawPlane.Rotation);
 
-	PlaneTransformGizmo->SetNewGizmoTransform(FTransform(DrawPlaneOrientation, DrawPlaneOrigin));
+	if (PlaneTransformGizmo != nullptr)
+	{
+		PlaneTransformGizmo->SetNewGizmoTransform(FTransform(DrawPlaneOrientation, DrawPlaneOrigin));
+	}
 }
 
 
@@ -946,11 +949,13 @@ void UDrawPolygonTool::UpdateShowGizmoState(bool bNewVisibility)
 	if (bNewVisibility == false)
 	{
 		GetToolManager()->GetPairedGizmoManager()->DestroyAllGizmosByOwner(this);
+		PlaneTransformGizmo = nullptr;
 	}
 	else
 	{
 		PlaneTransformGizmo = GetToolManager()->GetPairedGizmoManager()->Create3AxisTransformGizmo(this);
 		PlaneTransformGizmo->SetActiveTarget(PlaneTransformProxy, GetToolManager());
+		PlaneTransformGizmo->SetNewGizmoTransform(FTransform(DrawPlaneOrientation, DrawPlaneOrigin));
 	}
 }
 
