ModelingTools: snap test in PolyEdit needs to be performed on world-space points, not local-space
#rnx

diff --git a/Engine/Plugins/Experimental/MeshModelingToolset/Source/MeshModelingTools/Private/EditMeshPolygonsTool.cpp b/Engine/Plugins/Experimental/MeshModelingToolset/Source/MeshModelingTools/Private/EditMeshPolygonsTool.cpp
--- a/Engine/Plugins/Experimental/MeshModelingToolset/Source/MeshModelingTools/Private/EditMeshPolygonsTool.cpp
+++ b/Engine/Plugins/Experimental/MeshModelingToolset/Source/MeshModelingTools/Private/EditMeshPolygonsTool.cpp
@@ -657,7 +657,9 @@ void UEditMeshPolygonsTool::Setup()
 	TopoSelector.Initialize(DynamicMeshComponent->GetMesh(), &Topology);
 	TopoSelector.SetSpatialSource([this]() {return &GetSpatial(); });
 	TopoSelector.PointsWithinToleranceTest = [this](const FVector3d& Position1, const FVector3d& Position2) {
-		return ToolSceneQueriesUtil::PointSnapQuery(this->CameraState, Position1, Position2, VisualAngleSnapThreshold);
+		FTransform Transform = ComponentTarget->GetWorldTransform();
+		return ToolSceneQueriesUtil::PointSnapQuery(this->CameraState, 
+			Transform.TransformPosition(Position1), Transform.TransformPosition(Position2), VisualAngleSnapThreshold);
 	};
 
 	// hide input StaticMeshComponent
