multi-select for plane tool

diff --git a/Engine/Plugins/Experimental/MeshModelingToolset/Source/MeshModelingTools/Private/PlaneCutTool.cpp b/Engine/Plugins/Experimental/MeshModelingToolset/Source/MeshModelingTools/Private/PlaneCutTool.cpp
--- a/Engine/Plugins/Experimental/MeshModelingToolset/Source/MeshModelingTools/Private/PlaneCutTool.cpp
+++ b/Engine/Plugins/Experimental/MeshModelingToolset/Source/MeshModelingTools/Private/PlaneCutTool.cpp
@@ -33,20 +33,30 @@
 
 bool UPlaneCutToolBuilder::CanBuildTool(const FToolBuilderState& SceneState) const
 {
-	return ToolBuilderUtil::CountComponents(SceneState, CanMakeComponentTarget) == 1;
+	return AssetAPI != nullptr && ToolBuilderUtil::CountComponents(SceneState, CanMakeComponentTarget) > 0;
 }
 
 UInteractiveTool* UPlaneCutToolBuilder::BuildTool(const FToolBuilderState& SceneState) const
 {
 	UPlaneCutTool* NewTool = NewObject<UPlaneCutTool>(SceneState.ToolManager);
 
-	UActorComponent* ActorComponent = ToolBuilderUtil::FindFirstComponent(SceneState, CanMakeComponentTarget);
-	auto* MeshComponent = Cast<UPrimitiveComponent>(ActorComponent);
-	check(MeshComponent != nullptr);
-	NewTool->SetSelection(MakeComponentTarget(MeshComponent));
+	TArray<UActorComponent*> Components = ToolBuilderUtil::FindAllComponents(SceneState, CanMakeComponentTarget);
+	check(Components.Num() > 0);
 
+	TArray<TUniquePtr<FPrimitiveComponentTarget>> ComponentTargets;
+	for (UActorComponent* ActorComponent : Components)
+	{
+		auto* MeshComponent = Cast<UPrimitiveComponent>(ActorComponent);
+		if ( MeshComponent )
+		{
+			ComponentTargets.Add(MakeComponentTarget(MeshComponent));
+		}
+	}
+
+	NewTool->SetSelection(MoveTemp(ComponentTargets));
 	NewTool->SetWorld(SceneState.World);
 	NewTool->SetAssetAPI(AssetAPI);
+
 	return NewTool;
 }
 
@@ -83,8 +93,21 @@ void UPlaneCutTool::Setup()
 {
 	UInteractiveTool::Setup();
 
-	// hide input StaticMeshComponent
-	ComponentTarget->SetOwnerVisibility(false);
+	// hide input StaticMeshComponents
+	for (TUniquePtr<FPrimitiveComponentTarget>& ComponentTarget : ComponentTargets)
+	{
+		ComponentTarget->SetOwnerVisibility(false);
+	}
+
+	// Convert input mesh descriptions to dynamic mesh
+	for (TUniquePtr<FPrimitiveComponentTarget>& ComponentTarget : ComponentTargets)
+	{
+		TSharedPtr<FDynamicMesh3> OriginalDynamicMesh = MakeShared<FDynamicMesh3>();
+		FMeshDescriptionToDynamicMesh Converter;
+		Converter.bPrintDebugMessages = true;
+		Converter.Convert(ComponentTarget->GetMesh(), *OriginalDynamicMesh);
+		OriginalDynamicMeshes.Add(OriginalDynamicMesh);
+	}
 
 	// click to set plane behavior
 	FSelectClickedAction* SetPlaneAction = new FSelectClickedAction();
@@ -117,20 +140,17 @@ void UPlaneCutTool::Setup()
 	UpdateNumPreviews();
 
 	// set initial cut plane (also attaches gizmo/proxy)
-	FVector DefaultOrigin, Extents;
-	ComponentTarget->GetOwnerActor()->GetActorBounds(false, DefaultOrigin, Extents);
-	SetCutPlaneFromWorldPos(DefaultOrigin, FVector::UpVector);
+	FBox CombinedBounds; CombinedBounds.Init();
+	for (TUniquePtr<FPrimitiveComponentTarget>& ComponentTarget : ComponentTargets)
+	{
+		FVector ComponentOrigin, ComponentExtents;
+		ComponentTarget->GetOwnerActor()->GetActorBounds(false, ComponentOrigin, ComponentExtents);
+		CombinedBounds += FBox::BuildAABB(ComponentOrigin, ComponentExtents);
+	}
+	SetCutPlaneFromWorldPos(CombinedBounds.GetCenter(), FVector::UpVector);
 	// hook up callback so further changes trigger recut
 	PlaneTransformProxy->OnTransformChanged.AddUObject(this, &UPlaneCutTool::TransformChanged);
 
-
-
-	// Convert input mesh description to dynamic mesh
-	OriginalDynamicMesh = MakeShared<FDynamicMesh3>();
-	FMeshDescriptionToDynamicMesh Converter;
-	Converter.bPrintDebugMessages = true;
-	Converter.Convert(ComponentTarget->GetMesh(), *OriginalDynamicMesh);
-
 	for (UMeshOpPreviewWithBackgroundCompute* Preview : Previews)
 	{
 		Preview->InvalidateResult();
@@ -141,7 +161,8 @@ void UPlaneCutTool::Setup()
 void UPlaneCutTool::UpdateNumPreviews()
 {
 	int32 CurrentNumPreview = Previews.Num();
-	int32 TargetNumPreview = BasicProperties->bKeepBothHalves ? 2 : 1;
+	int32 NumSourceMeshes = OriginalDynamicMeshes.Num();
+	int32 TargetNumPreview = (BasicProperties->bKeepBothHalves ? 2 : 1) * NumSourceMeshes;
 	if (TargetNumPreview < CurrentNumPreview)
 	{
 		for (int32 PreviewIdx = CurrentNumPreview - 1; PreviewIdx >= TargetNumPreview; PreviewIdx--)
@@ -156,14 +177,20 @@ void UPlaneCutTool::UpdateNumPreviews()
 		{
 			UPlaneCutOperatorFactory *CutSide = NewObject<UPlaneCutOperatorFactory>();
 			CutSide->CutTool = this;
-			CutSide->bCutBackSide = PreviewIdx == 1;
+			CutSide->bCutBackSide = PreviewIdx >= NumSourceMeshes;
+			int32 SrcIdx = PreviewIdx % NumSourceMeshes;
+			CutSide->ComponentIndex = SrcIdx;
 			UMeshOpPreviewWithBackgroundCompute* Preview = Previews.Add_GetRef(NewObject<UMeshOpPreviewWithBackgroundCompute>(CutSide, "Preview"));
 			Preview->Setup(this->TargetWorld, CutSide);
 			Preview->ConfigureMaterials(
-				ToolSetupUtil::GetDefaultMaterial(GetToolManager(), ComponentTarget->GetMaterial(0)),
+				ToolSetupUtil::GetDefaultMaterial(GetToolManager(), ComponentTargets[SrcIdx]->GetMaterial(0)),
 				ToolSetupUtil::GetDefaultWorkingMaterial(GetToolManager())
 			);
-			Preview->SetVisibility(true);
+			// set initial preview to un-processed mesh, so stuff doesn't just disappear if the first cut takes a while
+			Preview->PreviewMesh->UpdatePreview(OriginalDynamicMeshes[SrcIdx].Get());
+			Preview->PreviewMesh->SetTransform(ComponentTargets[SrcIdx]->GetWorldTransform());
+			Preview->SetVisibility(BasicProperties->bShowPreview);
+			
 		}
 	}
 }
@@ -172,7 +199,10 @@ void UPlaneCutTool::UpdateNumPreviews()
 void UPlaneCutTool::Shutdown(EToolShutdownType ShutdownType)
 {
 	// Restore (unhide) the source meshes
-	ComponentTarget->SetOwnerVisibility(true);
+	for (TUniquePtr<FPrimitiveComponentTarget>& ComponentTarget : ComponentTargets)
+	{
+		ComponentTarget->SetOwnerVisibility(true);
+	}
 
 	TArray<FDynamicMeshOpResult> Results;
 	for (UMeshOpPreviewWithBackgroundCompute* Preview : Previews)
@@ -204,18 +234,18 @@ TUniquePtr<FDynamicMeshOperator> UPlaneCutOperatorFactory::MakeNewOperator()
 	CutOp->bFillCutHole = CutTool->BasicProperties->bFillCutHole;
 	CutOp->bFillSpans = CutTool->BasicProperties->bFillSpans;
 
-	FTransform LocalToWorld = CutTool->ComponentTarget->GetWorldTransform();
+	FTransform LocalToWorld = CutTool->ComponentTargets[ComponentIndex]->GetWorldTransform();
 	FTransform WorldToLocal = LocalToWorld.Inverse();
 	FVector LocalOrigin = WorldToLocal.TransformPosition(CutTool->CutPlaneOrigin);
 	FVector WorldNormal = CutTool->CutPlaneOrientation.GetAxisZ();
-	FVector LocalNormal = WorldToLocal.TransformVectorNoScale(WorldNormal);
+	FVector LocalNormal = WorldToLocal.TransformVectorNoScale(WorldNormal); // TODO: correct for nonuniform scaling?
 	if (bCutBackSide)
 	{
 		LocalNormal = -LocalNormal;
 	}
 	CutOp->LocalPlaneOrigin = LocalOrigin;
 	CutOp->LocalPlaneNormal = LocalNormal;
-	CutOp->OriginalMesh = CutTool->OriginalDynamicMesh;
+	CutOp->OriginalMesh = CutTool->OriginalDynamicMeshes[ComponentIndex];
 	
 	if (bCutBackSide)
 	{
@@ -264,7 +294,10 @@ void UPlaneCutTool::OnPropertyModified(UObject* PropertySet, UProperty* Property
 {
 	if (Property && (Property->GetFName() == GET_MEMBER_NAME_CHECKED(UPlaneCutToolProperties, bShowPreview)))
 	{
-		ComponentTarget->SetOwnerVisibility(!BasicProperties->bShowPreview);
+		for (TUniquePtr<FPrimitiveComponentTarget>& ComponentTarget : ComponentTargets)
+		{
+			ComponentTarget->SetOwnerVisibility(!BasicProperties->bShowPreview);
+		}
 		for (UMeshOpPreviewWithBackgroundCompute* Preview : Previews)
 		{
 			Preview->SetVisibility(BasicProperties->bShowPreview);
@@ -328,30 +361,42 @@ void UPlaneCutTool::GenerateAsset(const TArray<FDynamicMeshOpResult>& Results)
 
 	// currently in-place replaces the first half, and adds a new actor for the second half (if it was generated)
 	// TODO: options to support other choices re what should be a new actor
-	check(Results.Num() > 0);
-	check(Results[0].Mesh.Get() != nullptr);
-	ComponentTarget->CommitMesh([&Results](FMeshDescription* MeshDescription)
+	ensure(Results.Num() > 0);
+	int32 NumSourceMeshes = OriginalDynamicMeshes.Num();
+	for (int OrigMeshIdx = 0; OrigMeshIdx < NumSourceMeshes; OrigMeshIdx++)
 	{
-		FDynamicMeshToMeshDescription Converter;
-		Converter.Convert(Results[0].Mesh.Get(), *MeshDescription);
-	});
-
+		check(Results[OrigMeshIdx].Mesh.Get() != nullptr);
+		ComponentTargets[OrigMeshIdx]->CommitMesh([&Results, OrigMeshIdx](FMeshDescription* MeshDescription)
+		{
+			FDynamicMeshToMeshDescription Converter;
+			Converter.Convert(Results[OrigMeshIdx].Mesh.Get(), *MeshDescription);
+		});
+	}
 
 	// The method for creating a new mesh (AssetGenerationUtil::GenerateStaticMeshActor) is editor-only; just creating the other half if not in editor
 #if WITH_EDITOR
-	if (Results.Num() == 2)
+	
+	if (Results.Num() > NumSourceMeshes)
 	{
+		ensure(Results.Num() == NumSourceMeshes * 2);
+
 		FSelectedOjectsChangeList NewSelection;
 		NewSelection.ModificationType = ESelectedObjectsModificationType::Replace;
-		NewSelection.Actors.Add(ComponentTarget->GetOwnerActor());
-
-		check(Results[1].Mesh.Get() != nullptr);
-		// TODO: copy over material?
-		AActor* NewActor = AssetGenerationUtil::GenerateStaticMeshActor(
-			AssetAPI, TargetWorld,
-			Results[1].Mesh.Get(), Results[1].Transform, TEXT("PlaneCutOtherHalf"),
-			AssetGenerationUtil::GetDefaultAutoGeneratedAssetPath());
-		NewSelection.Actors.Add(NewActor);
+		for (TUniquePtr<FPrimitiveComponentTarget>& ComponentTarget : ComponentTargets)
+		{
+			NewSelection.Actors.Add(ComponentTarget->GetOwnerActor());
+		}
+		
+		for (int32 AddedMeshIdx = NumSourceMeshes; AddedMeshIdx < Results.Num(); AddedMeshIdx++)
+		{
+			check(Results[AddedMeshIdx].Mesh.Get() != nullptr);
+			// TODO: copy over material?
+			AActor* NewActor = AssetGenerationUtil::GenerateStaticMeshActor(
+				AssetAPI, TargetWorld,
+				Results[AddedMeshIdx].Mesh.Get(), Results[AddedMeshIdx].Transform, TEXT("PlaneCutOtherHalf"),
+				AssetGenerationUtil::GetDefaultAutoGeneratedAssetPath());
+			NewSelection.Actors.Add(NewActor);
+		}
 		GetToolManager()->RequestSelectionChange(NewSelection);
 	}
 #endif
diff --git a/Engine/Plugins/Experimental/MeshModelingToolset/Source/MeshModelingTools/Public/PlaneCutTool.h b/Engine/Plugins/Experimental/MeshModelingToolset/Source/MeshModelingTools/Public/PlaneCutTool.h
--- a/Engine/Plugins/Experimental/MeshModelingToolset/Source/MeshModelingTools/Public/PlaneCutTool.h
+++ b/Engine/Plugins/Experimental/MeshModelingToolset/Source/MeshModelingTools/Public/PlaneCutTool.h
@@ -5,7 +5,7 @@
 #include "CoreMinimal.h"
 #include "UObject/NoExportTypes.h"
 #include "InteractiveGizmo.h"
-#include "SingleSelectionTool.h"
+#include "MultiSelectionTool.h"
 #include "InteractiveToolBuilder.h"
 #include "MeshOpPreviewHelpers.h"
 #include "DynamicMesh3.h"
@@ -93,13 +93,15 @@ public:
 
 	UPROPERTY()
 	bool bCutBackSide = false;
+
+	int ComponentIndex;
 };
 
 /**
  * Simple Mesh Plane Cutting Tool
  */
 UCLASS()
-class MESHMODELINGTOOLS_API UPlaneCutTool : public USingleSelectionTool
+class MESHMODELINGTOOLS_API UPlaneCutTool : public UMultiSelectionTool
 {
 	GENERATED_BODY()
 
@@ -146,7 +148,7 @@ protected:
 
 
 protected:
-	TSharedPtr<FDynamicMesh3> OriginalDynamicMesh;
+	TArray<TSharedPtr<FDynamicMesh3>> OriginalDynamicMeshes;
 
 	UWorld* TargetWorld;
 	IToolsContextAssetAPI* AssetAPI;
