ModelingTools: Only bake scale & rotation transforms into sculpt mesh, keep translation on Component. Avoids issues with float imprecision far from origin.
#rnx

diff --git a/Engine/Plugins/Experimental/MeshModelingToolset/Source/MeshModelingTools/Private/DynamicMeshSculptTool.cpp b/Engine/Plugins/Experimental/MeshModelingToolset/Source/MeshModelingTools/Private/DynamicMeshSculptTool.cpp
--- a/Engine/Plugins/Experimental/MeshModelingToolset/Source/MeshModelingTools/Private/DynamicMeshSculptTool.cpp
+++ b/Engine/Plugins/Experimental/MeshModelingToolset/Source/MeshModelingTools/Private/DynamicMeshSculptTool.cpp
@@ -119,10 +119,11 @@ void UDynamicMeshSculptTool::Setup()
 	InitialTargetTransform = FTransform3d(ComponentTarget->GetWorldTransform());
 	// clamp scaling because if we allow zero-scale we cannot invert this transform on Accept
 	InitialTargetTransform.ClampMinimumScale(0.01);
+	FVector3d Translation = InitialTargetTransform.GetTranslation();
+	InitialTargetTransform.SetTranslation(FVector3d::Zero());
 	DynamicMeshComponent->ApplyTransform(InitialTargetTransform, false);
 	// since we moved to World coords there is not a current transform anymore.
-	// @todo: only remove scaling, keep rotation and translation in CurTargetTransform. This will improve numerical precision.
-	CurTargetTransform = FTransform3d::Identity();
+	CurTargetTransform = FTransform3d(Translation);
 	DynamicMeshComponent->SetWorldTransform((FTransform)CurTargetTransform);
 
 	// copy material if there is one
