add holey circle and rounded rect to draw polygon tool

diff --git a/Engine/Plugins/Experimental/GeometryProcessing/Source/GeometricObjects/Private/Generators/FlatTriangulationMeshGenerator.cpp b/Engine/Plugins/Experimental/GeometryProcessing/Source/GeometricObjects/Private/Generators/FlatTriangulationMeshGenerator.cpp
new file mode 100644
--- /dev/null
+++ b/Engine/Plugins/Experimental/GeometryProcessing/Source/GeometricObjects/Private/Generators/FlatTriangulationMeshGenerator.cpp
@@ -0,0 +1,57 @@
+// Copyright 1998-2019 Epic Games, Inc. All Rights Reserved.
+
+#include "Generators/FlatTriangulationMeshGenerator.h"
+#include "BoxTypes.h"
+
+FFlatTriangulationMeshGenerator::FFlatTriangulationMeshGenerator()
+{
+	Normal = FVector3f::UnitZ();
+	IndicesMap = FIndex2i(0, 1);
+}
+
+
+
+FMeshShapeGenerator& FFlatTriangulationMeshGenerator::Generate()
+{
+	int NumVertices = Vertices2D.Num();
+	ensure(NumVertices >= 3);
+	if (NumVertices < 3)
+	{
+		return *this;
+	}
+
+	int NumTriangles = Triangles2D.Num();
+	SetBufferSizes(NumVertices, NumTriangles, NumVertices, NumVertices);
+
+	FAxisAlignedBox2d BoundingBox(Vertices2D);
+	double Width = BoundingBox.Width(), Height = BoundingBox.Height();
+	double UVScale = FMath::Max(Width, Height);
+
+	for (int VertIdx = 0; VertIdx < NumVertices; ++VertIdx)
+	{
+		const FVector2d& Pos = Vertices2D[VertIdx];
+
+		UVs[VertIdx] = FVector2f(
+			(float)((Pos.X - BoundingBox.Min.X) / UVScale),
+			(float)((Pos.Y - BoundingBox.Min.Y) / UVScale));
+		UVParentVertex[VertIdx] = VertIdx;
+
+		Normals[VertIdx] = Normal;
+		NormalParentVertex[VertIdx] = VertIdx;
+
+		Vertices[VertIdx] = MakeVertex(Pos.X, Pos.Y);
+	}
+
+
+	for (int TriIdx = 0; TriIdx < NumTriangles; TriIdx++)
+	{
+		FIndex3i Tri = Triangles2D[TriIdx];
+		SetTriangle(TriIdx, Tri);
+		SetTriangleUVs(TriIdx, Tri);
+		SetTriangleNormals(TriIdx, Tri);
+		SetTrianglePolygon(TriIdx, 0);
+	}
+
+
+	return *this;
+}
\ No newline at end of file
diff --git a/Engine/Plugins/Experimental/GeometryProcessing/Source/GeometricObjects/Public/BoxTypes.h b/Engine/Plugins/Experimental/GeometryProcessing/Source/GeometricObjects/Public/BoxTypes.h
--- a/Engine/Plugins/Experimental/GeometryProcessing/Source/GeometricObjects/Public/BoxTypes.h
+++ b/Engine/Plugins/Experimental/GeometryProcessing/Source/GeometricObjects/Public/BoxTypes.h
@@ -425,6 +425,12 @@ struct TAxisAlignedBox2
 	{
 	}
 
+	TAxisAlignedBox2(const TArray<FVector2<RealType>>& Pts)
+	{
+		*this = Empty();
+		Contain(Pts);
+	}
+
 	TAxisAlignedBox2(const FVector2<RealType>& Center, RealType HalfWidth)
 	{
 		this->Min = FVector2<RealType>(Center.X - HalfWidth, Center.Y - HalfWidth);
diff --git a/Engine/Plugins/Experimental/GeometryProcessing/Source/GeometricObjects/Public/Generators/FlatTriangulationMeshGenerator.h b/Engine/Plugins/Experimental/GeometryProcessing/Source/GeometricObjects/Public/Generators/FlatTriangulationMeshGenerator.h
new file mode 100644
--- /dev/null
+++ b/Engine/Plugins/Experimental/GeometryProcessing/Source/GeometricObjects/Public/Generators/FlatTriangulationMeshGenerator.h
@@ -0,0 +1,40 @@
+// Copyright 1998-2019 Epic Games, Inc. All Rights Reserved.
+
+#pragma once
+
+#include "MeshShapeGenerator.h"
+
+/**
+ * Generate planar triangulation of a Polygon.
+ */
+class GEOMETRICOBJECTS_API FFlatTriangulationMeshGenerator : public FMeshShapeGenerator
+{
+public:
+	/** Vertices of 2D triangulation. */
+	TArray<FVector2d> Vertices2D;
+
+	/** Source triangle indices */
+	TArray<FIndex3i> Triangles2D;
+
+	/** Normal vector of all vertices will be set to this value. Default is +Z axis. */
+	FVector3f Normal;
+
+	/** How to map 2D indices to 3D. Default is (0,1) = (x,y,0). */
+	FIndex2i IndicesMap;
+
+public:
+	FFlatTriangulationMeshGenerator();
+
+	/** Generate the triangulation */
+	virtual FMeshShapeGenerator& Generate() override;
+
+
+	/** Create vertex at position under IndicesMap, shifted to Origin*/
+	inline FVector3d MakeVertex(double x, double y)
+	{
+		FVector3d v(0, 0, 0);
+		v[IndicesMap.A] = x;
+		v[IndicesMap.B] = y;
+		return v;
+	}
+};
\ No newline at end of file
diff --git a/Engine/Plugins/Experimental/GeometryProcessing/Source/GeometricObjects/Public/Polygon2.h b/Engine/Plugins/Experimental/GeometryProcessing/Source/GeometricObjects/Public/Polygon2.h
--- a/Engine/Plugins/Experimental/GeometryProcessing/Source/GeometricObjects/Public/Polygon2.h
+++ b/Engine/Plugins/Experimental/GeometryProcessing/Source/GeometricObjects/Public/Polygon2.h
@@ -1059,7 +1059,7 @@ public:
 
 
 	/**
-	 * Construct a four-vertex rectangle Polygon
+	 * Construct a four-vertex rectangle polygon
 	 */
 	static TPolygon2<T> MakeRectangle(const FVector2<T>& Center, T Width, T Height)
 	{
@@ -1072,9 +1072,41 @@ public:
 		return Rectangle;
 	}
 
+	/**
+	* Construct a rounded rectangle polygon
+	*/
+	static TPolygon2<T> MakeRoundedRectangle(const FVector2<T>& Center, T Width, T Height, T Corner, int CornerSteps)
+	{
+		TPolygon2<T> RRectangle;
+		CornerSteps = FMath::Max(2, CornerSteps);
+		RRectangle.Vertices.SetNumUninitialized(CornerSteps * 4);
+		
+		FVector2<T> InnerExtents(Width*.5 - Corner, Height*.5 - Corner);
+		FVector2<T> InnerCorners[4]
+		{
+			Center - InnerExtents,
+			FVector2<T>(Center.X + InnerExtents.X, Center.Y - InnerExtents.Y),
+			Center + InnerExtents,
+			FVector2<T>(Center.X - InnerExtents.X, Center.Y + InnerExtents.Y),
+		};
+
+		T IdxToAng = TMathUtil<T>::HalfPi / T(CornerSteps - 1);
+		for (int StepIdx = 0; StepIdx < CornerSteps; StepIdx++)
+		{
+			T Angle = StepIdx * IdxToAng;
+			T OffC = TMathUtil<T>::Cos(Angle) * Corner;
+			T OffS = TMathUtil<T>::Sin(Angle) * Corner;
+			RRectangle.Set(StepIdx + CornerSteps * 0, InnerCorners[0] + FVector2<T>(-OffC, -OffS));
+			RRectangle.Set(StepIdx + CornerSteps * 1, InnerCorners[1] + FVector2<T>(OffS, -OffC));
+			RRectangle.Set(StepIdx + CornerSteps * 2, InnerCorners[2] + FVector2<T>(OffC, OffS));
+			RRectangle.Set(StepIdx + CornerSteps * 3, InnerCorners[3] + FVector2<T>(-OffS, OffC));
+		}
+		return RRectangle;
+	}
+
 
 	/**
-	 * Construct a circular Polygon
+	 * Construct a circular polygon
 	 */
 	static TPolygon2<T> MakeCircle(T Radius, int Steps, T AngleShiftRadians = 0)
 	{
diff --git a/Engine/Plugins/Experimental/GeometryProcessing/Source/GeometryAlgorithms/Private/ConstrainedDelaunay2.cpp b/Engine/Plugins/Experimental/GeometryProcessing/Source/GeometryAlgorithms/Private/ConstrainedDelaunay2.cpp
--- a/Engine/Plugins/Experimental/GeometryProcessing/Source/GeometryAlgorithms/Private/ConstrainedDelaunay2.cpp
+++ b/Engine/Plugins/Experimental/GeometryProcessing/Source/GeometryAlgorithms/Private/ConstrainedDelaunay2.cpp
@@ -295,17 +295,7 @@ TArray<FIndex3i> GEOMETRYALGORITHMS_API ConstrainedDelaunayTriangulate(const TGe
 template TArray<FIndex3i> GEOMETRYALGORITHMS_API ConstrainedDelaunayTriangulate(const TGeneralPolygon2<double>& GeneralPolygon);
 template TArray<FIndex3i> GEOMETRYALGORITHMS_API ConstrainedDelaunayTriangulate(const TGeneralPolygon2<float>& GeneralPolygon);
 
-template void TConstrainedDelaunay2<double>::Add(const FDynamicGraph2<double>& Graph);
-template void TConstrainedDelaunay2<double>::Add(const TGeneralPolygon2<double>& Polygon);
-template void TConstrainedDelaunay2<double>::Add(const TGeneralPolygon2<float>& Polygon);
-template void TConstrainedDelaunay2<double>::Add(const TPolygon2<double>& Polygon, bool bIsHole);
-template void TConstrainedDelaunay2<double>::Add(const TPolygon2<float>& Polygon, bool bIsHole);
 
-template void TConstrainedDelaunay2<float>::Add(const FDynamicGraph2<double>& Graph);
-template void TConstrainedDelaunay2<float>::Add(const TGeneralPolygon2<double>& Polygon);
-template void TConstrainedDelaunay2<float>::Add(const TGeneralPolygon2<float>& Polygon);
-template void TConstrainedDelaunay2<float>::Add(const TPolygon2<double>& Polygon, bool bIsHole);
-template void TConstrainedDelaunay2<float>::Add(const TPolygon2<float>& Polygon, bool bIsHole);
 
 template struct TConstrainedDelaunay2<float>;
 template struct TConstrainedDelaunay2<double>;
diff --git a/Engine/Plugins/Experimental/MeshModelingToolset/Source/MeshModelingTools/MeshModelingTools.Build.cs b/Engine/Plugins/Experimental/MeshModelingToolset/Source/MeshModelingTools/MeshModelingTools.Build.cs
--- a/Engine/Plugins/Experimental/MeshModelingToolset/Source/MeshModelingTools/MeshModelingTools.Build.cs
+++ b/Engine/Plugins/Experimental/MeshModelingToolset/Source/MeshModelingTools/MeshModelingTools.Build.cs
@@ -37,6 +37,7 @@ public class MeshModelingTools : ModuleRules
 				"Eigen",
                 "InteractiveToolsFramework",
 				"GeometricObjects",
+				"GeometryAlgorithms",
 				"DynamicMesh",
 				"MeshConversion",
                 "MeshSolverUtilities",
diff --git a/Engine/Plugins/Experimental/MeshModelingToolset/Source/MeshModelingTools/Private/DrawPolygonTool.cpp b/Engine/Plugins/Experimental/MeshModelingToolset/Source/MeshModelingTools/Private/DrawPolygonTool.cpp
--- a/Engine/Plugins/Experimental/MeshModelingToolset/Source/MeshModelingTools/Private/DrawPolygonTool.cpp
+++ b/Engine/Plugins/Experimental/MeshModelingToolset/Source/MeshModelingTools/Private/DrawPolygonTool.cpp
@@ -8,17 +8,19 @@
 #include "BaseBehaviors/KeyAsModifierInputBehavior.h"
 
 #include "Polygon2.h"
+#include "Curve/GeneralPolygon2.h"
 #include "FrameTypes.h"
 #include "MatrixTypes.h"
 #include "DynamicMeshAttributeSet.h"
 
 #include "MeshDescriptionBuilder.h"
-#include "Generators/PlanarPolygonMeshGenerator.h"
+#include "Generators/FlatTriangulationMeshGenerator.h"
 #include "Operations/ExtrudeMesh.h"
 #include "Distance/DistLine3Ray3.h"
 #include "Intersection/IntrSegment2Segment2.h"
 #include "MeshQueries.h"
 #include "ToolSceneQueriesUtil.h"
+#include "ConstrainedDelaunay2.h"
 
 #include "DynamicMeshEditor.h"
 
@@ -346,7 +348,7 @@ void UDrawPolygonTool::Render(IToolsContextRenderAPI* RenderAPI)
 		if (bInInteractiveExtrude == false)		// once we are in extrude, polygon is done
 		{
 			FixedPolygonClickPoints.Add(PreviewVertex);
-			GenerateFixedPolygon(FixedPolygonClickPoints, PolygonVertices);
+			GenerateFixedPolygon(FixedPolygonClickPoints, PolygonVertices, PolygonHolesVertices);
 			FixedPolygonClickPoints.Pop(false);
 		}
 		bIsClosed = true;
@@ -421,24 +423,27 @@ void UDrawPolygonTool::Render(IToolsContextRenderAPI* RenderAPI)
 			UseThickness = SelfIntersectThickness;
 		}
 
+		auto DrawVertices = [&PDI, &UseColor](const TArray<FVector>& Vertices, ESceneDepthPriorityGroup Group, float Thickness)
+		{
+			for (int lasti = Vertices.Num() - 1, i = 0, NumVertices = Vertices.Num(); i < NumVertices; lasti = i++)
+			{
+				PDI->DrawLine(Vertices[lasti], Vertices[i], UseColor, Group, Thickness, 0.0f, true);
+			}
+		};
+
 		// draw thin no-depth
-		for (int i = 0; i < NumVerts - 1; ++i)
+		DrawVertices(PolygonVertices, SDPG_Foreground, HiddenLineThickness);
+		for (int HoleIdx = 0; HoleIdx < PolygonHolesVertices.Num(); HoleIdx++)
 		{
-			PDI->DrawLine(PolygonVertices[i], PolygonVertices[i + 1],
-				UseColor, SDPG_Foreground, HiddenLineThickness, 0.0f, true);
+			DrawVertices(PolygonHolesVertices[HoleIdx], SDPG_Foreground, HiddenLineThickness);
 		}
-		PDI->DrawLine(PolygonVertices[NumVerts - 1], UseLastVertex,
-			UseColor, SDPG_Foreground, HiddenLineThickness, 0.0f, true);
-
 
 		// draw thick depth-tested
-		for (int i = 0; i < NumVerts-1; ++i)
+		DrawVertices(PolygonVertices, SDPG_World, LineThickness);
+		for (int HoleIdx = 0; HoleIdx < PolygonHolesVertices.Num(); HoleIdx++)
 		{
-			PDI->DrawLine(PolygonVertices[i], PolygonVertices[i+1],
-				UseColor, SDPG_World, LineThickness, 0.0f, true);
+			DrawVertices(PolygonHolesVertices[HoleIdx], SDPG_World, LineThickness);
 		}
-		PDI->DrawLine(PolygonVertices[NumVerts-1], UseLastVertex,
-			UseColor, SDPG_World, LineThickness, 0.0f, true);
 
 		if (bHaveSelfIntersection)
 		{
@@ -485,6 +490,7 @@ void UDrawPolygonTool::Render(IToolsContextRenderAPI* RenderAPI)
 void UDrawPolygonTool::ResetPolygon()
 {
 	PolygonVertices.Reset();
+	PolygonHolesVertices.Reset();
 	SnapEngine.Reset();
 	bHaveSurfaceHit = false;
 	bInFixedPolygonMode = false;
@@ -679,11 +685,11 @@ bool UDrawPolygonTool::OnNextSequenceClick(const FInputDeviceRay& ClickPos)
 		}
 
 		FixedPolygonClickPoints.Add(HitPos);
-		int NumTargetPoints = (PolygonProperties->PolygonType == EDrawPolygonDrawMode::Rectangle) ? 3 : 2;
+		int NumTargetPoints = (PolygonProperties->PolygonType == EDrawPolygonDrawMode::Rectangle || PolygonProperties->PolygonType == EDrawPolygonDrawMode::RoundedRectangle) ? 3 : 2;
 		bDonePolygon = (FixedPolygonClickPoints.Num() == NumTargetPoints);
 		if (bDonePolygon)
 		{
-			GenerateFixedPolygon(FixedPolygonClickPoints, PolygonVertices);
+			GenerateFixedPolygon(FixedPolygonClickPoints, PolygonVertices, PolygonHolesVertices);
 		}
 	} 
 	else
@@ -800,8 +806,9 @@ bool UDrawPolygonTool::UpdateSelfIntersection()
 	return false;
 }
 
-void UDrawPolygonTool::GenerateFixedPolygon(TArray<FVector>& FixedPoints, TArray<FVector>& VerticesOut)
+void UDrawPolygonTool::GenerateFixedPolygon(const TArray<FVector>& FixedPoints, TArray<FVector>& VerticesOut, TArray<TArray<FVector>>& HolesVerticesOut)
 {
+
 	FFrame3f DrawFrame(DrawPlaneOrigin, DrawPlaneOrientation);
 	FVector2f CenterPt = DrawFrame.ToPlaneUV(FixedPoints[0], 2 );
 	FVector2f EdgePt = DrawFrame.ToPlaneUV(FixedPoints[1], 2);
@@ -813,30 +820,58 @@ void UDrawPolygonTool::GenerateFixedPolygon(TArray<FVector>& FixedPoints, TArray
 	float Width = FMath::Abs(Delta.Dot(RotAxisX));
 
 	FPolygon2f Polygon;
+	TArray<FPolygon2f> PolygonHoles;
 	if (PolygonProperties->PolygonType == EDrawPolygonDrawMode::Square)
 	{
 		Polygon = FPolygon2f::MakeRectangle(FVector2f::Zero(), 2*Width, 2*Width);
 	}
-	else if (PolygonProperties->PolygonType == EDrawPolygonDrawMode::Rectangle)
+	else if (PolygonProperties->PolygonType == EDrawPolygonDrawMode::Rectangle || PolygonProperties->PolygonType == EDrawPolygonDrawMode::RoundedRectangle)
 	{
 		FVector2f HeightPt = DrawFrame.ToPlaneUV((FixedPoints.Num() == 3) ? FixedPoints[2] : FixedPoints[1], 2);
 		FVector2f HeightDelta = HeightPt - CenterPt;
 		FVector2f RotAxisY = RotationMat * FVector2f::UnitY();
 		float YSign = FMath::Sign(HeightDelta.Dot(RotAxisY));
 		float Height = FMath::Abs(HeightDelta.Dot(RotAxisY));
-		Polygon = FPolygon2f::MakeRectangle(FVector2f(Width/2, YSign*Height/2), Width, Height);
+		if (PolygonProperties->PolygonType == EDrawPolygonDrawMode::Rectangle)
+		{
+			Polygon = FPolygon2f::MakeRectangle(FVector2f(Width / 2, YSign*Height / 2), Width, Height);
+		}
+		else // PolygonProperties->PolygonType == EDrawPolygonDrawMode::RoundedRectangle
+		{
+			Polygon = FPolygon2f::MakeRoundedRectangle(FVector2f(Width / 2, YSign*Height / 2), Width, Height, FMath::Min(Width,Height) * FMathf::Clamp(PolygonProperties->FeatureSizeRatio, .01, .99) * .5, PolygonProperties->Steps);
+		}
 	}
-	else
+	else // Circle or HoleyCircle
 	{
 		Polygon = FPolygon2f::MakeCircle(Dist, PolygonProperties->Steps, 0);
+		if (PolygonProperties->PolygonType == EDrawPolygonDrawMode::HoleyCircle)
+		{
+			PolygonHoles.Add(FPolygon2f::MakeCircle(Dist * FMathd::Clamp(PolygonProperties->FeatureSizeRatio, .01, .99), PolygonProperties->Steps, 0));
+		}
 	}
 	Polygon.Transform([RotationMat](const FVector2f& Pt) { return RotationMat * Pt; });
+	for (FPolygon2f& Hole : PolygonHoles)
+	{
+		Hole.Transform([RotationMat](const FVector2f& Pt) { return RotationMat * Pt; });
+	}
 
-	VerticesOut.Reset();
+	VerticesOut.SetNum(Polygon.VertexCount());
 	for (int k = 0; k < Polygon.VertexCount(); ++k)
 	{
 		FVector2f NewPt = CenterPt + Polygon[k];
-		VerticesOut.Add(DrawFrame.FromPlaneUV(NewPt, 2));
+		VerticesOut[k] = DrawFrame.FromPlaneUV(NewPt, 2);
+	}
+
+	HolesVerticesOut.SetNum(PolygonHoles.Num());
+	for (int HoleIdx = 0; HoleIdx < PolygonHoles.Num(); HoleIdx++)
+	{
+		int NumHoleVerts = PolygonHoles[HoleIdx].VertexCount();
+		HolesVerticesOut[HoleIdx].SetNum(NumHoleVerts);
+		for (int k = 0; k < NumHoleVerts; ++k)
+		{
+			FVector2f NewPt = CenterPt + PolygonHoles[HoleIdx][k];
+			HolesVerticesOut[HoleIdx][k] = DrawFrame.FromPlaneUV(NewPt, 2);
+		}
 	}
 }
 
@@ -971,7 +1006,7 @@ void UDrawPolygonTool::EmitCurrentPolygon()
 	FDynamicMesh3 Mesh;
 	double ExtrudeDist = (PolygonProperties->OutputMode == EDrawPolygonOutputMode::MeshedPolygon) ?
 		0 : PolygonProperties->ExtrudeHeight;
-	GeneratePolygonMesh(PolygonVertices, &Mesh, PlaneFrameOut, false, ExtrudeDist, false);
+	GeneratePolygonMesh(PolygonVertices, PolygonHolesVertices, &Mesh, PlaneFrameOut, false, ExtrudeDist, false);
 
 	GetToolManager()->BeginUndoTransaction(LOCTEXT("CreatePolygon", "Create Polygon"));
 
@@ -1003,7 +1038,7 @@ void UDrawPolygonTool::UpdateLivePreview()
 	FDynamicMesh3 Mesh;
 	double ExtrudeDist = (PolygonProperties->OutputMode == EDrawPolygonOutputMode::MeshedPolygon) ?
 		0 : PolygonProperties->ExtrudeHeight;
-	GeneratePolygonMesh(PolygonVertices, &Mesh, PlaneFrame, false, ExtrudeDist, false);
+	GeneratePolygonMesh(PolygonVertices, PolygonHolesVertices, &Mesh, PlaneFrame, false, ExtrudeDist, false);
 
 	PreviewMesh->SetTransform(PlaneFrame.ToFTransform());
 	PreviewMesh->SetMaterial(MaterialProperties->Material);
@@ -1011,7 +1046,7 @@ void UDrawPolygonTool::UpdateLivePreview()
 	PreviewMesh->UpdatePreview(&Mesh);
 }
 
-void UDrawPolygonTool::GeneratePolygonMesh(const TArray<FVector>& Polygon, FDynamicMesh3* ResultMeshOut, FFrame3d& WorldFrameOut, bool bIncludePreviewVtx, double ExtrudeDistance, bool bExtrudeSymmetric)
+void UDrawPolygonTool::GeneratePolygonMesh(const TArray<FVector>& Polygon, const TArray<TArray<FVector>>& PolygonHoles, FDynamicMesh3* ResultMeshOut, FFrame3d& WorldFrameOut, bool bIncludePreviewVtx, double ExtrudeDistance, bool bExtrudeSymmetric)
 {
 	// construct centered frame for polygon
 	WorldFrameOut = FFrame3d(DrawPlaneOrigin, DrawPlaneOrientation);
@@ -1027,27 +1062,51 @@ void UDrawPolygonTool::GeneratePolygonMesh(const TArray<FVector>& Polygon, FDyna
 
 
 	// triangulate polygon into the MeshDescription
-	FPlanarPolygonMeshGenerator PolyMeshGen;
-	for (int k = 0; k < NumVerts; ++k)
+	FGeneralPolygon2d GeneralPolygon;
+	FFlatTriangulationMeshGenerator TriangulationMeshGen;
+
+	auto VertexArrayToPolygon = [&WorldFrameOut](const TArray<FVector>& Vertices)
 	{
-		PolyMeshGen.Polygon.AppendVertex(WorldFrameOut.ToPlaneUV(Polygon[k], 2));
-	}
+		FPolygon2d OutPolygon;
+		for (int k = 0, N = Vertices.Num(); k < N; ++k)
+		{
+			OutPolygon.AppendVertex(WorldFrameOut.ToPlaneUV(Vertices[k], 2));
+		}
+		return OutPolygon;
+	};
+
+	FPolygon2d OuterPolygon = VertexArrayToPolygon(Polygon);
 
 	// add preview vertex
 	if (bIncludePreviewVtx)
 	{
 		if (FVector::Dist(PreviewVertex, Polygon[NumVerts - 1]) > 0.1)
 		{
-			PolyMeshGen.Polygon.AppendVertex(WorldFrameOut.ToPlaneUV(PreviewVertex, 2));
+			OuterPolygon.AppendVertex(WorldFrameOut.ToPlaneUV(PreviewVertex, 2));
 		}
 	}
 
-	if (PolyMeshGen.Polygon.IsClockwise() == false)
+	if (OuterPolygon.IsClockwise() == false)
+	{
+		OuterPolygon.Reverse();
+	}
+
+	GeneralPolygon.SetOuter(OuterPolygon);
+
+	for (int HoleIdx = 0; HoleIdx < PolygonHoles.Num(); HoleIdx++)
 	{
-		PolyMeshGen.Polygon.Reverse();
+		// attempt to add holes (skipping if safety checks fail)
+		GeneralPolygon.AddHole(VertexArrayToPolygon(PolygonHoles[HoleIdx]), true, false /*currently don't care about hole orientation; we'll just set the triangulation algo not to care*/);
 	}
 
-	ResultMeshOut->Copy(&PolyMeshGen.Generate());
+	FConstrainedDelaunay2d Triangulator;
+	Triangulator.Add(GeneralPolygon);
+	bool bTriangulationSuccess = Triangulator.Triangulate();
+	// TODO: actually do something different if triangulation fails?  note that failure doesn't mean that it didn't generate triangles ... it will try to give triangles back even if input was not fully valid
+	TriangulationMeshGen.Vertices2D = Triangulator.Vertices;
+	TriangulationMeshGen.Triangles2D = Triangulator.Triangles;
+
+	ResultMeshOut->Copy(&TriangulationMeshGen.Generate());
 
 	// for symmetric extrude we translate the first poly by -dist along axis
 	if (bExtrudeSymmetric)
@@ -1066,7 +1125,7 @@ void UDrawPolygonTool::GeneratePolygonMesh(const TArray<FVector>& Polygon, FDyna
 	{
 		FExtrudeMesh Extruder(ResultMeshOut);
 		Extruder.DefaultExtrudeDistance = ExtrudeDistance;
-		FAxisAlignedBox2d bounds = PolyMeshGen.Polygon.Bounds();
+		FAxisAlignedBox2d bounds = GeneralPolygon.Bounds();
 		Extruder.UVScaleFactor = 1.0 / bounds.MaxDim();
 		if (ExtrudeDistance < 0)
 		{
@@ -1083,7 +1142,7 @@ void UDrawPolygonTool::GeneratePolygonMesh(const TArray<FVector>& Polygon, FDyna
 	}
 
 	FDynamicMeshEditor Editor(ResultMeshOut);
-	float InitialUVScale = 1.0 / PolyMeshGen.Polygon.Bounds().MaxDim(); // this is the UV scale used by both the polymeshgen and the extruder above
+	float InitialUVScale = 1.0 / GeneralPolygon.Bounds().MaxDim(); // this is the UV scale used by both the polymeshgen and the extruder above
 	// default global rescale -- initial scale doesn't factor in extrude distance; rescale so UVScale of 1.0 fits in the unit square texture
 	float GlobalUVRescale = MaterialProperties->UVScale / FMathf::Max(1.0f, ExtrudeDistance * InitialUVScale);
 	if (MaterialProperties->bWorldSpaceUVScale)
@@ -1096,7 +1155,7 @@ void UDrawPolygonTool::GeneratePolygonMesh(const TArray<FVector>& Polygon, FDyna
 
 void UDrawPolygonTool::GeneratePreviewHeightTarget()
 {
-	GeneratePolygonMesh(PolygonVertices, &PreviewHeightTarget, PreviewHeightFrame, false, 99999, true);
+	GeneratePolygonMesh(PolygonVertices, PolygonHolesVertices, &PreviewHeightTarget, PreviewHeightFrame, false, 99999, true);
 	PreviewHeightTargetAABB.SetMesh(&PreviewHeightTarget);
 }
 
diff --git a/Engine/Plugins/Experimental/MeshModelingToolset/Source/MeshModelingTools/Public/DrawPolygonTool.h b/Engine/Plugins/Experimental/MeshModelingToolset/Source/MeshModelingTools/Public/DrawPolygonTool.h
--- a/Engine/Plugins/Experimental/MeshModelingToolset/Source/MeshModelingTools/Public/DrawPolygonTool.h
+++ b/Engine/Plugins/Experimental/MeshModelingToolset/Source/MeshModelingTools/Public/DrawPolygonTool.h
@@ -55,7 +55,13 @@ enum class EDrawPolygonDrawMode : uint8
 	Square UMETA(DisplayName = "Square"),
 
 	/** Rectangle */
-	Rectangle UMETA(DisplayName = "Rectangle")
+	Rectangle UMETA(DisplayName = "Rectangle"),
+
+	/** Rounded Rectangle */
+	RoundedRectangle UMETA(DisplayName = "Rounded Rectangle"),
+
+	/** Circle w/ Hole */
+	HoleyCircle UMETA(DisplayName = "Circle w/ Hole")
 
 };
 
@@ -92,6 +98,11 @@ public:
 	UPROPERTY(EditAnywhere, NonTransactional, Category = Polygon)
 	EDrawPolygonOutputMode OutputMode = EDrawPolygonOutputMode::ExtrudedInteractive;
 
+	/** Feature size as fraction of overall shape size, for shapes with secondary features like the rounded corners of a Rounded Rectangle */
+	UPROPERTY(EditAnywhere, NonTransactional, Category = Polygon, meta = (UIMin = "0", UIMax = "1", ClampMin = "0", ClampMax = "1", 
+									EditCondition = "PolygonType == EDrawPolygonDrawMode::RoundedRectangle || PolygonType == EDrawPolygonDrawMode::HoleyCircle"))
+	float FeatureSizeRatio = .25;
+
 	/** Extrusion Distance in non-interactive mode */
 	UPROPERTY(EditAnywhere, NonTransactional, Category = Polygon, meta = (UIMin = "-1000", UIMax = "1000", ClampMin = "-10000", ClampMax = "10000"))
 	float ExtrudeHeight = 100.0f;
@@ -239,6 +250,9 @@ protected:
 	UPROPERTY()
 	TArray<FVector> PolygonVertices;
 
+	/** Vertices of holes in current preview polygon */
+	TArray<TArray<FVector>> PolygonHolesVertices;
+
 	/** last vertex of polygon that is actively being updated as input device is moved */
 	UPROPERTY()
 	FVector PreviewVertex;
@@ -296,7 +310,7 @@ protected:
 	FPointPlanarSnapSolver SnapEngine;
 	ToolSceneQueriesUtil::FSnapGeometry LastSnapGeometry;
 
-	void GenerateFixedPolygon(TArray<FVector>& FixedPoints, TArray<FVector>& VerticesOut);
+	void GenerateFixedPolygon(const TArray<FVector>& FixedPoints, TArray<FVector>& VerticesOut, TArray<TArray<FVector>>& HolesVerticesOut);
 
 
 	// extrusion control
@@ -315,7 +329,7 @@ protected:
 	FFrame3f HitPosFrameWorld;
 
 	// we use this to generate extruded meshes
-	void GeneratePolygonMesh(const TArray<FVector>& Polygon, FDynamicMesh3* ResultMeshOut, FFrame3d& WorldFrameOut, bool bIncludePreviewVtx, double ExtrudeDistance, bool bExtrudeSymmetric);
+	void GeneratePolygonMesh(const TArray<FVector>& Polygon, const TArray<TArray<FVector>>& PolygonHoles, FDynamicMesh3* ResultMeshOut, FFrame3d& WorldFrameOut, bool bIncludePreviewVtx, double ExtrudeDistance, bool bExtrudeSymmetric);
 
 
 	// user feedback messages
