ModelingTools: revert changes to DrawPolygon line rendering that broke Freehand mode
#rnx

diff --git a/Engine/Plugins/Experimental/MeshModelingToolset/Source/MeshModelingTools/Private/DrawPolygonTool.cpp b/Engine/Plugins/Experimental/MeshModelingToolset/Source/MeshModelingTools/Private/DrawPolygonTool.cpp
--- a/Engine/Plugins/Experimental/MeshModelingToolset/Source/MeshModelingTools/Private/DrawPolygonTool.cpp
+++ b/Engine/Plugins/Experimental/MeshModelingToolset/Source/MeshModelingTools/Private/DrawPolygonTool.cpp
@@ -432,18 +432,32 @@ void UDrawPolygonTool::Render(IToolsContextRenderAPI* RenderAPI)
 		};
 
 		// draw thin no-depth
-		DrawVertices(PolygonVertices, SDPG_Foreground, HiddenLineThickness);
+		//DrawVertices(PolygonVertices, SDPG_Foreground, HiddenLineThickness);
+		for (int i = 0; i < NumVerts - 1; ++i)
+		{
+			PDI->DrawLine(PolygonVertices[i], PolygonVertices[i + 1],
+				UseColor, SDPG_Foreground, HiddenLineThickness, 0.0f, true);
+		}
 		for (int HoleIdx = 0; HoleIdx < PolygonHolesVertices.Num(); HoleIdx++)
 		{
 			DrawVertices(PolygonHolesVertices[HoleIdx], SDPG_Foreground, HiddenLineThickness);
 		}
+		PDI->DrawLine(PolygonVertices[NumVerts - 1], UseLastVertex,
+			UseColor, SDPG_Foreground, HiddenLineThickness, 0.0f, true);
 
 		// draw thick depth-tested
-		DrawVertices(PolygonVertices, SDPG_World, LineThickness);
+		//DrawVertices(PolygonVertices, SDPG_World, LineThickness);
+		for (int i = 0; i < NumVerts - 1; ++i)
+		{
+			PDI->DrawLine(PolygonVertices[i], PolygonVertices[i + 1],
+				UseColor, SDPG_World, LineThickness, 0.0f, true);
+		}
 		for (int HoleIdx = 0; HoleIdx < PolygonHolesVertices.Num(); HoleIdx++)
 		{
 			DrawVertices(PolygonHolesVertices[HoleIdx], SDPG_World, LineThickness);
 		}
+		PDI->DrawLine(PolygonVertices[NumVerts - 1], UseLastVertex,
+			UseColor, SDPG_World, LineThickness, 0.0f, true);
 
 		if (bHaveSelfIntersection)
 		{
