ModelingTools: add option to ignore backface hits in Sculpt Tool. Allows sculpting "inside" of surfaces like caves/etc
#rnx

diff --git a/Engine/Plugins/Experimental/MeshModelingToolset/Source/MeshModelingTools/Private/DynamicMeshSculptTool.cpp b/Engine/Plugins/Experimental/MeshModelingToolset/Source/MeshModelingTools/Private/DynamicMeshSculptTool.cpp
--- a/Engine/Plugins/Experimental/MeshModelingToolset/Source/MeshModelingTools/Private/DynamicMeshSculptTool.cpp
+++ b/Engine/Plugins/Experimental/MeshModelingToolset/Source/MeshModelingTools/Private/DynamicMeshSculptTool.cpp
@@ -50,6 +50,7 @@ UBrushSculptProperties::UBrushSculptProperties()
 	bPreserveUVFlow = false;
 	BrushDepth = 0;
 	bFreezeTarget = false;
+	bHitBackFaces = true;
 }
 
 void UBrushSculptProperties::SaveProperties(UInteractiveTool* SaveFromTool)
@@ -60,6 +61,7 @@ void UBrushSculptProperties::SaveProperties(UInteractiveTool* SaveFromTool)
 	PropertyCache->PrimaryBrushType = this->PrimaryBrushType;
 	PropertyCache->bPreserveUVFlow = this->bPreserveUVFlow;
 	PropertyCache->BrushDepth = this->BrushDepth;
+	PropertyCache->bHitBackFaces = this->bHitBackFaces;
 }
 
 void UBrushSculptProperties::RestoreProperties(UInteractiveTool* RestoreToTool)
@@ -70,6 +72,7 @@ void UBrushSculptProperties::RestoreProperties(UInteractiveTool* RestoreToTool)
 	this->PrimaryBrushType = PropertyCache->PrimaryBrushType;
 	this->bPreserveUVFlow = PropertyCache->bPreserveUVFlow;
 	this->BrushDepth = PropertyCache->BrushDepth;
+	this->bHitBackFaces = PropertyCache->bHitBackFaces;
 }
 
 
@@ -278,16 +281,14 @@ void UDynamicMeshSculptTool::OnPropertyModified(UObject* PropertySet, UProperty*
 }
 
 
-
 bool UDynamicMeshSculptTool::HitTest(const FRay& Ray, FHitResult& OutHit)
 {
 	FRay3d LocalRay(CurTargetTransform.InverseTransformPosition(Ray.Origin),
 		CurTargetTransform.InverseTransformVector(Ray.Direction));
 	LocalRay.Direction.Normalize();
-
 	FDynamicMesh3* Mesh = DynamicMeshComponent->GetMesh();
-	int HitTID = DynamicMeshComponent->GetOctree()->FindNearestHitObject(LocalRay);
-
+	
+	int HitTID = FindHitSculptMeshTriangle(LocalRay);
 	if (HitTID != IndexConstants::InvalidID)
 	{
 		FTriangle3d Triangle;
@@ -965,6 +966,64 @@ void UDynamicMeshSculptTool::ApplyInflateBrush(const FRay& WorldRay)
 
 
 
+int UDynamicMeshSculptTool::FindHitSculptMeshTriangle(const FRay3d& LocalRay)
+{
+	if (SculptProperties->bHitBackFaces)
+	{
+		return DynamicMeshComponent->GetOctree()->FindNearestHitObject(LocalRay);
+	}
+	else
+	{
+		FDynamicMesh3* Mesh = DynamicMeshComponent->GetMesh();
+
+		FViewCameraState StateOut;
+		GetToolManager()->GetContextQueriesAPI()->GetCurrentViewState(StateOut);
+		FVector3d LocalEyePosition(CurTargetTransform.InverseTransformPosition(StateOut.Position));
+		int HitTID = DynamicMeshComponent->GetOctree()->FindNearestHitObject(LocalRay,
+			[this, Mesh, &LocalEyePosition](int TriangleID) {
+			FVector3d Normal, Centroid;
+			double Area;
+			Mesh->GetTriInfo(TriangleID, Normal, Area, Centroid);
+			return Normal.Dot((Centroid - LocalEyePosition)) < 0;
+		});
+		return HitTID;
+	}
+}
+
+
+int UDynamicMeshSculptTool::FindHitTargetMeshTriangle(const FRay3d& LocalRay)
+{
+	if (SculptProperties->bHitBackFaces == false)
+	{
+		FDynamicMesh3* Mesh = DynamicMeshComponent->GetMesh();
+
+		FViewCameraState StateOut;
+		GetToolManager()->GetContextQueriesAPI()->GetCurrentViewState(StateOut);
+		FVector3d LocalEyePosition(CurTargetTransform.InverseTransformPosition(StateOut.Position));
+
+		BrushTargetMeshSpatial.TriangleFilterF = [this, Mesh, &LocalEyePosition](int TriangleID) {
+			FVector3d Normal, Centroid;
+			double Area;
+			Mesh->GetTriInfo(TriangleID, Normal, Area, Centroid);
+			return Normal.Dot((Centroid - LocalEyePosition)) < 0;
+		};
+	}
+
+	int HitTID =  BrushTargetMeshSpatial.FindNearestHitTriangle(LocalRay);
+
+	if (SculptProperties->bHitBackFaces == false)
+	{
+		BrushTargetMeshSpatial.TriangleFilterF = nullptr;
+	}
+
+	return HitTID;
+}
+
+
+
+
+
+
 
 bool UDynamicMeshSculptTool::UpdateBrushPositionOnActivePlane(const FRay& WorldRay)
 {
@@ -981,7 +1040,7 @@ bool UDynamicMeshSculptTool::UpdateBrushPositionOnTargetMesh(const FRay& WorldRa
 		CurTargetTransform.InverseTransformVector(WorldRay.Direction));
 	LocalRay.Direction.Normalize();
 
-	int HitTID = BrushTargetMeshSpatial.FindNearestHitTriangle(LocalRay);
+	int HitTID = FindHitTargetMeshTriangle(LocalRay);
 	if (HitTID != IndexConstants::InvalidID)
 	{
 		const FDynamicMesh3* TargetMesh = BrushTargetMeshSpatial.GetMesh();
@@ -1005,7 +1064,7 @@ bool UDynamicMeshSculptTool::UpdateBrushPositionOnSculptMesh(const FRay& WorldRa
 		CurTargetTransform.InverseTransformVector(WorldRay.Direction));
 	LocalRay.Direction.Normalize();
 
-	int HitTID = DynamicMeshComponent->GetOctree()->FindNearestHitObject(LocalRay);
+	int HitTID = FindHitSculptMeshTriangle(LocalRay);
 	if (HitTID != IndexConstants::InvalidID)
 	{
 		const FDynamicMesh3* SculptMesh = DynamicMeshComponent->GetMesh();
diff --git a/Engine/Plugins/Experimental/MeshModelingToolset/Source/MeshModelingTools/Public/DynamicMeshSculptTool.h b/Engine/Plugins/Experimental/MeshModelingToolset/Source/MeshModelingTools/Public/DynamicMeshSculptTool.h
--- a/Engine/Plugins/Experimental/MeshModelingToolset/Source/MeshModelingTools/Public/DynamicMeshSculptTool.h
+++ b/Engine/Plugins/Experimental/MeshModelingToolset/Source/MeshModelingTools/Public/DynamicMeshSculptTool.h
@@ -138,6 +138,10 @@ public:
 	UPROPERTY(EditAnywhere, Category = Sculpting, meta = (EditCondition = "PrimaryBrushType == EDynamicMeshSculptBrushType::Sculpt || PrimaryBrushType == EDynamicMeshSculptBrushType::SculptMax || PrimaryBrushType == EDynamicMeshSculptBrushType::Pinch" ))
 	bool bFreezeTarget;
 
+	/** Hit back sides of triangles */
+	UPROPERTY(EditAnywhere, Category = Sculpting)
+	bool bHitBackFaces;
+
 	virtual void SaveProperties(UInteractiveTool* SaveFromTool) override;
 	virtual void RestoreProperties(UInteractiveTool* RestoreToTool) override;
 
@@ -332,6 +336,9 @@ protected:
 	void UpdateTarget();
 	bool GetTargetMeshNearest(const FVector3d& Position, double SearchRadius, FVector3d& TargetPosOut, FVector3d& TargetNormalOut);
 
+	int FindHitSculptMeshTriangle(const FRay3d& LocalRay);
+	int FindHitTargetMeshTriangle(const FRay3d& LocalRay);
+
 	bool UpdateBrushPositionOnActivePlane(const FRay& WorldRay);
 	bool UpdateBrushPositionOnTargetMesh(const FRay& WorldRay);
 	bool UpdateBrushPositionOnSculptMesh(const FRay& WorldRay);
