ModelingTools: handle case where FixedPolygonClickPoints is empty
#rnx

diff --git a/Engine/Plugins/Experimental/MeshModelingToolset/Source/MeshModelingTools/Private/DrawPolygonTool.cpp b/Engine/Plugins/Experimental/MeshModelingToolset/Source/MeshModelingTools/Private/DrawPolygonTool.cpp
--- a/Engine/Plugins/Experimental/MeshModelingToolset/Source/MeshModelingTools/Private/DrawPolygonTool.cpp
+++ b/Engine/Plugins/Experimental/MeshModelingToolset/Source/MeshModelingTools/Private/DrawPolygonTool.cpp
@@ -345,7 +345,7 @@ void UDrawPolygonTool::Render(IToolsContextRenderAPI* RenderAPI)
 
 	if (bInFixedPolygonMode)
 	{
-		if (bInInteractiveExtrude == false)		// once we are in extrude, polygon is done
+		if (FixedPolygonClickPoints.Num() > 0 && bInInteractiveExtrude == false)		// once we are in extrude, polygon is done
 		{
 			FixedPolygonClickPoints.Add(PreviewVertex);
 			GenerateFixedPolygon(FixedPolygonClickPoints, PolygonVertices, PolygonHolesVertices);
@@ -438,12 +438,12 @@ void UDrawPolygonTool::Render(IToolsContextRenderAPI* RenderAPI)
 			PDI->DrawLine(PolygonVertices[i], PolygonVertices[i + 1],
 				UseColor, SDPG_Foreground, HiddenLineThickness, 0.0f, true);
 		}
+		PDI->DrawLine(PolygonVertices[NumVerts - 1], UseLastVertex,
+			UseColor, SDPG_Foreground, HiddenLineThickness, 0.0f, true);
 		for (int HoleIdx = 0; HoleIdx < PolygonHolesVertices.Num(); HoleIdx++)
 		{
 			DrawVertices(PolygonHolesVertices[HoleIdx], SDPG_Foreground, HiddenLineThickness);
 		}
-		PDI->DrawLine(PolygonVertices[NumVerts - 1], UseLastVertex,
-			UseColor, SDPG_Foreground, HiddenLineThickness, 0.0f, true);
 
 		// draw thick depth-tested
 		//DrawVertices(PolygonVertices, SDPG_World, LineThickness);
@@ -452,12 +452,12 @@ void UDrawPolygonTool::Render(IToolsContextRenderAPI* RenderAPI)
 			PDI->DrawLine(PolygonVertices[i], PolygonVertices[i + 1],
 				UseColor, SDPG_World, LineThickness, 0.0f, true);
 		}
+		PDI->DrawLine(PolygonVertices[NumVerts - 1], UseLastVertex,
+			UseColor, SDPG_World, LineThickness, 0.0f, true);
 		for (int HoleIdx = 0; HoleIdx < PolygonHolesVertices.Num(); HoleIdx++)
 		{
 			DrawVertices(PolygonHolesVertices[HoleIdx], SDPG_World, LineThickness);
 		}
-		PDI->DrawLine(PolygonVertices[NumVerts - 1], UseLastVertex,
-			UseColor, SDPG_World, LineThickness, 0.0f, true);
 
 		if (bHaveSelfIntersection)
 		{
