ModelingTools: add Fixed Plane mode to Sculpt brush. Allows user to set a specific plane for brush to pull towards. Plane can be snapped to world grid locations. Hotkey p places grid.
#rnx

diff --git a/Engine/Plugins/Experimental/MeshModelingToolset/Source/MeshModelingTools/Private/DynamicMeshSculptTool.cpp b/Engine/Plugins/Experimental/MeshModelingToolset/Source/MeshModelingTools/Private/DynamicMeshSculptTool.cpp
--- a/Engine/Plugins/Experimental/MeshModelingToolset/Source/MeshModelingTools/Private/DynamicMeshSculptTool.cpp
+++ b/Engine/Plugins/Experimental/MeshModelingToolset/Source/MeshModelingTools/Private/DynamicMeshSculptTool.cpp
@@ -25,6 +25,11 @@
 #include "Components/PrimitiveComponent.h"
 #include "Generators/SphereGenerator.h"
 
+#include "InteractiveGizmoManager.h"
+#include "BaseGizmos/GizmoComponents.h"
+#include "BaseGizmos/TransformGizmo.h"
+
+
 #define LOCTEXT_NAMESPACE "UDynamicMeshSculptTool"
 
 
@@ -77,6 +82,27 @@ void UBrushSculptProperties::RestoreProperties(UInteractiveTool* RestoreToTool)
 
 
 
+UFixedPlaneBrushProperties::UFixedPlaneBrushProperties()
+{
+	bPropertySetEnabled = true;
+	bSnapToGrid = true;
+	bShowGizmo = true;
+}
+
+void UFixedPlaneBrushProperties::SaveProperties(UInteractiveTool* SaveFromTool)
+{
+	UFixedPlaneBrushProperties* PropertyCache = GetPropertyCache<UFixedPlaneBrushProperties>();
+	PropertyCache->bShowGizmo = this->bShowGizmo;
+	PropertyCache->bSnapToGrid = this->bSnapToGrid;
+}
+void UFixedPlaneBrushProperties::RestoreProperties(UInteractiveTool* RestoreToTool)
+{
+	UFixedPlaneBrushProperties* PropertyCache = GetPropertyCache<UFixedPlaneBrushProperties>();
+	this->bShowGizmo = PropertyCache->bShowGizmo;
+	this->bSnapToGrid = PropertyCache->bSnapToGrid;
+}
+
+
 
 
 
@@ -190,6 +216,10 @@ void UDynamicMeshSculptTool::Setup()
 	CalculateBrushRadius();
 	SculptProperties->RestoreProperties(this);
 
+	GizmoProperties = NewObject<UFixedPlaneBrushProperties>();
+	GizmoProperties->RestoreProperties(this);
+	AddToolPropertySource(GizmoProperties);
+
 	ViewProperties = NewObject<UMeshEditingViewProperties>();
 	ViewProperties->RestoreProperties(this);
 	AddToolPropertySource(ViewProperties);
@@ -200,6 +230,10 @@ void UDynamicMeshSculptTool::Setup()
 		[this]() { return ViewProperties->MaterialMode; },
 		[this](EMeshEditingMaterialModes NewMode) { UpdateMaterialMode(NewMode); }, EMeshEditingMaterialModes::ExistingMaterial);
 
+	// create proxy for plane gizmo, but not gizmo itself, as it only appears in FixedPlane brush mode
+	// listen for changes to the proxy and update the plane when that happens
+	PlaneTransformProxy = NewObject<UTransformProxy>(this);
+	PlaneTransformProxy->OnTransformChanged.AddUObject(this, &UDynamicMeshSculptTool::PlaneTransformChanged);
 
 	GetToolManager()->DisplayMessage(
 		LOCTEXT("OnStartSculptTool", "Hold Shift to Smooth, Ctrl to Invert (where applicable). Shift+Q/A keys cycle through Brush Types. Shift+S/D change Size (Ctrl+Shift to small-step), Shift+W/E change Speed."),
@@ -470,6 +504,9 @@ void UDynamicMeshSculptTool::ApplyStamp(const FRay& WorldRay)
 		case EDynamicMeshSculptBrushType::Plane:
 			ApplyPlaneBrush(WorldRay);
 			break;
+		case EDynamicMeshSculptBrushType::FixedPlane:
+			ApplyFixedPlaneBrush(WorldRay);
+			break;
 	}
 }
 
@@ -853,6 +890,58 @@ void UDynamicMeshSculptTool::ApplyPlaneBrush(const FRay& WorldRay)
 
 
 
+void UDynamicMeshSculptTool::ApplyFixedPlaneBrush(const FRay& WorldRay)
+{
+	bool bHit = UpdateBrushPositionOnSculptMesh(WorldRay);
+	if (bHit == false)
+	{
+		return;
+	}
+
+	FVector3d NewBrushPosLocal = CurTargetTransform.InverseTransformPosition(LastBrushPosWorld);
+	FVector3d BrushNormalLocal = CurTargetTransform.InverseTransformNormal(LastBrushPosNormalWorld);
+	double UseSpeed = CurrentBrushRadius * FMathd::Sqrt(SculptProperties->BrushSpeed) * 0.1;
+
+	FFrame3d FixedPlaneLocal(
+		CurTargetTransform.InverseTransformPosition(DrawPlaneOrigin),
+		CurTargetTransform.GetRotation().Inverse() * (FQuaterniond)DrawPlaneOrientation);
+
+	FDynamicMesh3* Mesh = DynamicMeshComponent->GetMesh();
+	int NumV = VertexROI.Num();
+	ROIPositionBuffer.SetNum(NumV, false);
+
+	ParallelFor(NumV, [this, Mesh, NewBrushPosLocal, UseSpeed, FixedPlaneLocal](int k)
+	{
+		int VertIdx = VertexROI[k];
+		FVector3d OrigPos = Mesh->GetVertex(VertIdx);
+		FVector3d PlanePos = FixedPlaneLocal.ToPlane(OrigPos, 2);
+		FVector3d Delta = PlanePos - OrigPos;
+		double MaxDist = Delta.Normalize();
+		double Falloff = CalculateBrushFalloff(OrigPos.Distance(NewBrushPosLocal));
+		FVector3d MoveVec = Falloff * UseSpeed * Delta;
+		FVector3d NewPos = (MoveVec.SquaredLength() > MaxDist*MaxDist) ? 
+			PlanePos : OrigPos + Falloff * MoveVec;
+
+		ROIPositionBuffer[k] = NewPos;
+	});
+
+	for (int k = 0; k < NumV; ++k)
+	{
+		int VertIdx = VertexROI[k];
+		const FVector3d& NewPos = ROIPositionBuffer[k];
+		FVector3d OrigPos = Mesh->GetVertex(VertIdx);
+		Mesh->SetVertex(VertIdx, NewPos);
+		UpdateSavedVertex(VertIdx, OrigPos, NewPos);
+	}
+
+	bRemeshPending = bEnableRemeshing;
+
+	LastBrushPosLocal = NewBrushPosLocal;
+}
+
+
+
+
 void UDynamicMeshSculptTool::ApplyFlattenBrush(const FRay& WorldRay)
 {
 	bool bHit = UpdateBrushPositionOnSculptMesh(WorldRay);
@@ -1123,6 +1212,17 @@ void UDynamicMeshSculptTool::Render(IToolsContextRenderAPI* RenderAPI)
 	UMeshSurfacePointTool::Render(RenderAPI);
 
 	BrushIndicator->Update( (float)this->CurrentBrushRadius, this->LastBrushPosWorld, this->LastBrushPosNormalWorld );
+
+	if (SculptProperties->PrimaryBrushType == EDynamicMeshSculptBrushType::FixedPlane)
+	{
+		FPrimitiveDrawInterface* PDI = RenderAPI->GetPrimitiveDrawInterface();
+		FColor GridColor(128, 128, 128, 32);
+		float GridThickness = 0.5f;
+		float GridLineSpacing = 25.0f;   // @todo should be relative to view
+		int NumGridLines = 10;
+		FFrame3f DrawFrame(DrawPlaneOrigin, DrawPlaneOrientation);
+		MeshDebugDraw::DrawSimpleGrid(DrawFrame, NumGridLines, GridLineSpacing, GridThickness, GridColor, false, PDI, FTransform::Identity);
+	}
 }
 
 
@@ -1143,6 +1243,17 @@ void UDynamicMeshSculptTool::Tick(float DeltaTime)
 	ShowWireframeWatcher.CheckAndUpdate();
 	MaterialModeWatcher.CheckAndUpdate();
 
+	bool bGizmoVisible = (SculptProperties->PrimaryBrushType == EDynamicMeshSculptBrushType::FixedPlane)
+		&& (GizmoProperties->bShowGizmo);
+	UpdateFixedPlaneGizmoVisibility(bGizmoVisible);
+	GizmoProperties->bPropertySetEnabled = (SculptProperties->PrimaryBrushType == EDynamicMeshSculptBrushType::FixedPlane);
+
+	if (bPendingSetFixedPlanePosition)
+	{
+		SetFixedSculptPlaneFromWorldPos(LastBrushPosWorld);
+		bPendingSetFixedPlanePosition = false;
+	}
+
 	// if user changed to not-frozen, we need to update the target
 	if (bCachedFreezeTarget != SculptProperties->bFreezeTarget)
 	{
@@ -1796,6 +1907,14 @@ void UDynamicMeshSculptTool::RegisterActions(FInteractiveToolActionSet& ActionSe
 		EModifierKey::Alt, EKeys::W,
 		[this]() { ViewProperties->bShowWireframe = !ViewProperties->bShowWireframe; });
 
+
+	ActionSet.RegisterAction(this, (int32)EStandardToolActions::BaseClientDefinedActionID + 100,
+		TEXT("SetFixedSculptPlane"),
+		LOCTEXT("SetFixedSculptPlane", "Set Fixed Sculpt Plane"),
+		LOCTEXT("SetFixedSculptPlaneTooltip", "Set position of fixed sculpt plane"),
+		EModifierKey::None, EKeys::P,
+		[this]() { bPendingSetFixedPlanePosition = true; });
+
 }
 
 
@@ -1892,4 +2011,66 @@ void UDynamicMeshSculptTool::UpdateMaterialMode(EMeshEditingMaterialModes Materi
 }
 
 
+
+
+void UDynamicMeshSculptTool::SetFixedSculptPlaneFromWorldPos(const FVector& Position)
+{
+	UpdateFixedSculptPlanePosition(Position);
+	if (PlaneTransformGizmo != nullptr)
+	{
+		PlaneTransformGizmo->SetNewGizmoTransform(FTransform(DrawPlaneOrientation, DrawPlaneOrigin));
+	}
+}
+
+
+void UDynamicMeshSculptTool::PlaneTransformChanged(UTransformProxy* Proxy, FTransform Transform)
+{
+	DrawPlaneOrientation = Transform.GetRotation();
+	UpdateFixedSculptPlanePosition(Transform.GetLocation());
+}
+
+
+void UDynamicMeshSculptTool::UpdateFixedSculptPlanePosition(const FVector& Position)
+{
+	DrawPlaneOrigin = Position;
+
+	if (GizmoProperties->bSnapToGrid)
+	{
+		FSceneSnapQueryRequest Request;
+		Request.RequestType = ESceneSnapQueryType::Position;
+		Request.TargetTypes = ESceneSnapQueryTargetType::Grid;
+		Request.Position = DrawPlaneOrigin;
+		TArray<FSceneSnapQueryResult> Results;
+		if (GetToolManager()->GetContextQueriesAPI()->ExecuteSceneSnapQuery(Request, Results) && Results.Num() > 0)
+		{
+			DrawPlaneOrigin = Results[0].Position;
+		}
+	}
+}
+
+
+void UDynamicMeshSculptTool::UpdateFixedPlaneGizmoVisibility(bool bVisible)
+{
+	if (bVisible == false)
+	{
+		if (PlaneTransformGizmo != nullptr)
+		{
+			GetToolManager()->GetPairedGizmoManager()->DestroyGizmo(PlaneTransformGizmo);
+			PlaneTransformGizmo = nullptr;
+		}
+	}
+	else
+	{
+		if (PlaneTransformGizmo == nullptr)
+		{
+			PlaneTransformGizmo = GetToolManager()->GetPairedGizmoManager()->Create3AxisTransformGizmo(this);
+			PlaneTransformGizmo->SetActiveTarget(PlaneTransformProxy, GetToolManager());
+			PlaneTransformGizmo->SetNewGizmoTransform(FTransform(DrawPlaneOrientation, DrawPlaneOrigin));
+		}
+	}
+}
+
+
+
+
 #undef LOCTEXT_NAMESPACE
diff --git a/Engine/Plugins/Experimental/MeshModelingToolset/Source/MeshModelingTools/Public/DynamicMeshSculptTool.h b/Engine/Plugins/Experimental/MeshModelingToolset/Source/MeshModelingTools/Public/DynamicMeshSculptTool.h
--- a/Engine/Plugins/Experimental/MeshModelingToolset/Source/MeshModelingTools/Public/DynamicMeshSculptTool.h
+++ b/Engine/Plugins/Experimental/MeshModelingToolset/Source/MeshModelingTools/Public/DynamicMeshSculptTool.h
@@ -17,6 +17,9 @@
 #include "TransformTypes.h"
 #include "DynamicMeshSculptTool.generated.h"
 
+class UTransformGizmo;
+class UTransformProxy;
+
 
 DECLARE_STATS_GROUP(TEXT("SculptTool"), STATGROUP_SculptTool, STATCAT_Advanced);
 DECLARE_CYCLE_STAT(TEXT("SculptTool_UpdateROI"), SculptTool_UpdateROI, STATGROUP_SculptTool);
@@ -70,6 +73,8 @@ enum class EDynamicMeshSculptBrushType : uint8
 	/** Plane Brush pulls vertices to a plane defined by the initial brush position  */
 	Plane UMETA(DisplayName = "Plane"),
 
+	/** FixedPlane Brush pulls vertices to a fixed plane defined by the gizmo */
+	FixedPlane UMETA(DisplayName = "FixedPlane"),
 
 	LastValue UMETA(Hidden)
 
@@ -144,7 +149,31 @@ public:
 
 	virtual void SaveProperties(UInteractiveTool* SaveFromTool) override;
 	virtual void RestoreProperties(UInteractiveTool* RestoreToTool) override;
+};
+
+
+
+UCLASS()
+class MESHMODELINGTOOLS_API UFixedPlaneBrushProperties : public UInteractiveToolPropertySet
+{
+	GENERATED_BODY()
 
+public:
+	UFixedPlaneBrushProperties();
+
+	UPROPERTY()
+	bool bPropertySetEnabled = true;
+
+	/** Toggle whether interactive plane gizmo is visible */
+	UPROPERTY(EditAnywhere, Category = TargetPlane, meta = (EditCondition = "bPropertySetEnabled == true"))
+	bool bShowGizmo;
+
+	/** Toggle whether fixed plane snaps to grid */
+	UPROPERTY(EditAnywhere, Category = TargetPlane, meta = (EditCondition = "bPropertySetEnabled == true"))
+	bool bSnapToGrid;
+
+	virtual void SaveProperties(UInteractiveTool* SaveFromTool) override;
+	virtual void RestoreProperties(UInteractiveTool* RestoreToTool) override;
 };
 
 
@@ -245,6 +274,8 @@ public:
 	UPROPERTY()
 	UMeshEditingViewProperties* ViewProperties;
 
+	UPROPERTY()
+	UFixedPlaneBrushProperties* GizmoProperties;
 
 public:
 
@@ -350,6 +381,7 @@ protected:
 	void ApplyPinchBrush(const FRay& WorldRay);
 	void ApplyInflateBrush(const FRay& WorldRay);
 	void ApplyPlaneBrush(const FRay& WorldRay);
+	void ApplyFixedPlaneBrush(const FRay& WorldRay);
 	void ApplyFlattenBrush(const FRay& WorldRay);
 
 	double CalculateBrushFalloff(double Distance);
@@ -385,6 +417,34 @@ protected:
 	int BrushTypeHistoryIndex = 0;
 
 	UPreviewMesh* MakeDefaultSphereMesh(UObject* Parent, UWorld* World, int Resolution = 32);
+
+
+	//
+	// support for gizmo in FixedPlane mode
+	//
+
+	// plane gizmo
+	UPROPERTY()
+	UTransformGizmo* PlaneTransformGizmo;
+
+	UPROPERTY()
+	UTransformProxy* PlaneTransformProxy;
+
+	UPROPERTY()
+	FVector DrawPlaneOrigin;
+
+	/** Orientation of plane we will draw polygon on */
+	UPROPERTY()
+	FQuat DrawPlaneOrientation;
+
+	void PlaneTransformChanged(UTransformProxy* Proxy, FTransform Transform);
+
+	bool bPendingSetFixedPlanePosition = false;
+	void SetFixedSculptPlaneFromWorldPos(const FVector& Position);
+
+	void UpdateFixedSculptPlanePosition(const FVector& Position);
+
+	void UpdateFixedPlaneGizmoVisibility(bool bVisible);
 };
 
 
