ModelingTools: expose fixed plane position in properties, save/restore it properly
#rnx

diff --git a/Engine/Plugins/Experimental/MeshModelingToolset/Source/MeshModelingTools/Private/DynamicMeshSculptTool.cpp b/Engine/Plugins/Experimental/MeshModelingToolset/Source/MeshModelingTools/Private/DynamicMeshSculptTool.cpp
--- a/Engine/Plugins/Experimental/MeshModelingToolset/Source/MeshModelingTools/Private/DynamicMeshSculptTool.cpp
+++ b/Engine/Plugins/Experimental/MeshModelingToolset/Source/MeshModelingTools/Private/DynamicMeshSculptTool.cpp
@@ -87,6 +87,7 @@ UFixedPlaneBrushProperties::UFixedPlaneBrushProperties()
 	bPropertySetEnabled = true;
 	bSnapToGrid = true;
 	bShowGizmo = true;
+	Position = FVector::ZeroVector;
 }
 
 void UFixedPlaneBrushProperties::SaveProperties(UInteractiveTool* SaveFromTool)
@@ -94,12 +95,14 @@ void UFixedPlaneBrushProperties::SaveProperties(UInteractiveTool* SaveFromTool)
 	UFixedPlaneBrushProperties* PropertyCache = GetPropertyCache<UFixedPlaneBrushProperties>();
 	PropertyCache->bShowGizmo = this->bShowGizmo;
 	PropertyCache->bSnapToGrid = this->bSnapToGrid;
+	PropertyCache->Position = this->Position;
 }
 void UFixedPlaneBrushProperties::RestoreProperties(UInteractiveTool* RestoreToTool)
 {
 	UFixedPlaneBrushProperties* PropertyCache = GetPropertyCache<UFixedPlaneBrushProperties>();
 	this->bShowGizmo = PropertyCache->bShowGizmo;
 	this->bSnapToGrid = PropertyCache->bSnapToGrid;
+	this->Position = PropertyCache->Position;
 }
 
 
@@ -298,6 +301,7 @@ void UDynamicMeshSculptTool::Shutdown(EToolShutdownType ShutdownType)
 	BrushProperties->SaveProperties(this);
 	SculptProperties->SaveProperties(this);
 	ViewProperties->SaveProperties(this);
+	GizmoProperties->SaveProperties(this);
 }
 
 
@@ -903,7 +907,7 @@ void UDynamicMeshSculptTool::ApplyFixedPlaneBrush(const FRay& WorldRay)
 	double UseSpeed = CurrentBrushRadius * FMathd::Sqrt(SculptProperties->BrushSpeed) * 0.1;
 
 	FFrame3d FixedPlaneLocal(
-		CurTargetTransform.InverseTransformPosition(DrawPlaneOrigin),
+		CurTargetTransform.InverseTransformPosition(GizmoProperties->Position),
 		CurTargetTransform.GetRotation().Inverse() * (FQuaterniond)DrawPlaneOrientation);
 
 	FDynamicMesh3* Mesh = DynamicMeshComponent->GetMesh();
@@ -1220,7 +1224,7 @@ void UDynamicMeshSculptTool::Render(IToolsContextRenderAPI* RenderAPI)
 		float GridThickness = 0.5f;
 		float GridLineSpacing = 25.0f;   // @todo should be relative to view
 		int NumGridLines = 10;
-		FFrame3f DrawFrame(DrawPlaneOrigin, DrawPlaneOrientation);
+		FFrame3f DrawFrame(GizmoProperties->Position, DrawPlaneOrientation);
 		MeshDebugDraw::DrawSimpleGrid(DrawFrame, NumGridLines, GridLineSpacing, GridThickness, GridColor, false, PDI, FTransform::Identity);
 	}
 }
@@ -2018,7 +2022,7 @@ void UDynamicMeshSculptTool::SetFixedSculptPlaneFromWorldPos(const FVector& Posi
 	UpdateFixedSculptPlanePosition(Position);
 	if (PlaneTransformGizmo != nullptr)
 	{
-		PlaneTransformGizmo->SetNewGizmoTransform(FTransform(DrawPlaneOrientation, DrawPlaneOrigin));
+		PlaneTransformGizmo->SetNewGizmoTransform(FTransform(DrawPlaneOrientation, GizmoProperties->Position));
 	}
 }
 
@@ -2032,18 +2036,18 @@ void UDynamicMeshSculptTool::PlaneTransformChanged(UTransformProxy* Proxy, FTran
 
 void UDynamicMeshSculptTool::UpdateFixedSculptPlanePosition(const FVector& Position)
 {
-	DrawPlaneOrigin = Position;
+	GizmoProperties->Position = Position;
 
 	if (GizmoProperties->bSnapToGrid)
 	{
 		FSceneSnapQueryRequest Request;
 		Request.RequestType = ESceneSnapQueryType::Position;
 		Request.TargetTypes = ESceneSnapQueryTargetType::Grid;
-		Request.Position = DrawPlaneOrigin;
+		Request.Position = GizmoProperties->Position;
 		TArray<FSceneSnapQueryResult> Results;
 		if (GetToolManager()->GetContextQueriesAPI()->ExecuteSceneSnapQuery(Request, Results) && Results.Num() > 0)
 		{
-			DrawPlaneOrigin = Results[0].Position;
+			GizmoProperties->Position = Results[0].Position;
 		}
 	}
 }
@@ -2065,7 +2069,7 @@ void UDynamicMeshSculptTool::UpdateFixedPlaneGizmoVisibility(bool bVisible)
 		{
 			PlaneTransformGizmo = GetToolManager()->GetPairedGizmoManager()->Create3AxisTransformGizmo(this);
 			PlaneTransformGizmo->SetActiveTarget(PlaneTransformProxy, GetToolManager());
-			PlaneTransformGizmo->SetNewGizmoTransform(FTransform(DrawPlaneOrientation, DrawPlaneOrigin));
+			PlaneTransformGizmo->SetNewGizmoTransform(FTransform(DrawPlaneOrientation, GizmoProperties->Position));
 		}
 	}
 }
diff --git a/Engine/Plugins/Experimental/MeshModelingToolset/Source/MeshModelingTools/Public/DynamicMeshSculptTool.h b/Engine/Plugins/Experimental/MeshModelingToolset/Source/MeshModelingTools/Public/DynamicMeshSculptTool.h
--- a/Engine/Plugins/Experimental/MeshModelingToolset/Source/MeshModelingTools/Public/DynamicMeshSculptTool.h
+++ b/Engine/Plugins/Experimental/MeshModelingToolset/Source/MeshModelingTools/Public/DynamicMeshSculptTool.h
@@ -172,6 +172,9 @@ public:
 	UPROPERTY(EditAnywhere, Category = TargetPlane, meta = (EditCondition = "bPropertySetEnabled == true"))
 	bool bSnapToGrid;
 
+	UPROPERTY(EditAnywhere, Category = TargetPlane, AdvancedDisplay, meta = (EditCondition = "bPropertySetEnabled == true"))
+	FVector Position;
+
 	virtual void SaveProperties(UInteractiveTool* SaveFromTool) override;
 	virtual void RestoreProperties(UInteractiveTool* RestoreToTool) override;
 };
@@ -430,9 +433,6 @@ protected:
 	UPROPERTY()
 	UTransformProxy* PlaneTransformProxy;
 
-	UPROPERTY()
-	FVector DrawPlaneOrigin;
-
 	/** Orientation of plane we will draw polygon on */
 	UPROPERTY()
 	FQuat DrawPlaneOrientation;
