handle self-intersection cases better for draw polygon tool

#rnx

diff --git a/Engine/Plugins/Experimental/GeometryProcessing/Source/DynamicMesh/Private/Operations/ExtrudeMesh.cpp b/Engine/Plugins/Experimental/GeometryProcessing/Source/DynamicMesh/Private/Operations/ExtrudeMesh.cpp
--- a/Engine/Plugins/Experimental/GeometryProcessing/Source/DynamicMesh/Private/Operations/ExtrudeMesh.cpp
+++ b/Engine/Plugins/Experimental/GeometryProcessing/Source/DynamicMesh/Private/Operations/ExtrudeMesh.cpp
@@ -47,6 +47,10 @@ bool FExtrudeMesh::Apply()
 	// set vertices to new positions
 	for (int vid : InitialVertices)
 	{
+		if (!InitialToOffsetMapV.Contains(vid))
+		{
+			continue;
+		}
 		int newvid = InitialToOffsetMapV[vid];
 		if ( ! Mesh->IsVertex(newvid) )
 		{
diff --git a/Engine/Plugins/Experimental/GeometryProcessing/Source/GeometryAlgorithms/Private/ConstrainedDelaunay2.cpp b/Engine/Plugins/Experimental/GeometryProcessing/Source/GeometryAlgorithms/Private/ConstrainedDelaunay2.cpp
--- a/Engine/Plugins/Experimental/GeometryProcessing/Source/GeometryAlgorithms/Private/ConstrainedDelaunay2.cpp
+++ b/Engine/Plugins/Experimental/GeometryProcessing/Source/GeometryAlgorithms/Private/ConstrainedDelaunay2.cpp
@@ -68,8 +68,25 @@ template<class RealType>
 template<class InputRealType>
 void TConstrainedDelaunay2<RealType>::Add(const FDynamicGraph2<InputRealType>& Graph)
 {
-	// TODO: implement
-	check(false);
+	int32 VertexStart = Vertices.Num();
+	int32 GMaxVertID = Graph.MaxVertexID();
+	TArray<int32> GraphToDTVertIdxMap; GraphToDTVertIdxMap.SetNum(GMaxVertID);
+	for (int32 Idx = 0; Idx < GMaxVertID; Idx++)
+	{
+		if (Graph.IsVertex(Idx))
+		{
+			GraphToDTVertIdxMap[Idx] = Vertices.Num();
+			Vertices.Add((FVector2<RealType>)Graph.GetVertex(Idx));
+		}
+		else
+		{
+			GraphToDTVertIdxMap[Idx] = -1;
+		}
+	}
+	for (const FDynamicGraph::FEdge& Edge : Graph.Edges())
+	{
+		Edges.Add(FIndex2i(GraphToDTVertIdxMap[Edge.A], GraphToDTVertIdxMap[Edge.B]));
+	}
 }
 
 template<class RealType>
@@ -208,9 +225,7 @@ bool TConstrainedDelaunay2<RealType>::Triangulate()
 		}
 	}
 
-	
 	const std::vector<int>& Indices = Delaunay.GetIndices();
-	
 	const std::vector<int>& Adj = Delaunay.GetAdjacencies();
 	int TriNum = int(Adj.size() / 3);
 	TArray<int8> Keep;  // values: 0->unprocessed (delete), 1->yes keep, -1->processed, delete
@@ -266,18 +281,163 @@ bool TConstrainedDelaunay2<RealType>::Triangulate()
 		}
 	}
 
-	for (int i = 0; i < TriNum; i++)
+	// function to build output triangles out of an indices array
+	// normally called directly on the const indices from the CDT, but will be called on an updated copy if bowtie splits happen
+	auto BuildTriangles = [this, &TriNum, &Keep](const std::vector<int>& IndicesIn)
 	{
-		if (Keep[i] > 0)
+		for (int i = 0; i < TriNum; i++)
+		{
+			if (Keep[i] > 0)
+			{
+				FIndex3i& Tri = Triangles.Emplace_GetRef(IndicesIn[i * 3], IndicesIn[i * 3 + 1], IndicesIn[i * 3 + 2]);
+				if (!bOutputCCW)
+				{
+					Swap(Tri.B, Tri.C);
+				}
+			}
+		}
+	};
+
+
+	TArray<TPair<int, int>> NeedUpdates; // stores all wedge indices and the corresponding new vertices they require
+	int32 OrigNumVertices = Vertices.Num();
+	if (bSplitBowties)
+	{
+		// track all wedge verts that are seen by walking local tris
+		auto OtherEdgeOnTri = [&Indices](int VertID, int TriID, int EdgeIdx)
+		{
+			int StepNext = 1;
+			if (Indices[TriID*3+EdgeIdx] == VertID) {
+				StepNext = 2;
+			}
+			return (EdgeIdx + StepNext) % 3;
+		};
+		// helper to find new edge idx of the edge you crossed to go from FromTriID over to ToTriID
+		auto CrossEdge = [&Adj](int FromTriID, int ToTriID)
+		{
+			for (int AdjEdgeIdx = 0; AdjEdgeIdx < 3; AdjEdgeIdx++)
+			{
+				if (Adj[ToTriID * 3 + AdjEdgeIdx] == FromTriID)
+				{
+					return AdjEdgeIdx;
+				}
+			}
+			ensure(false);
+			return -1;
+		};
+		auto GetVertSubIdx = [&Indices](int VertID, int TriID)
+		{
+			for (int VertSubIdx = 0; VertSubIdx < 3; VertSubIdx++)
+			{
+				if (Indices[TriID * 3 + VertSubIdx] == VertID)
+				{
+					return VertSubIdx;
+				}
+			}
+			ensure(false);
+			return -1;
+		};
+		auto Walk = [&Adj, &OtherEdgeOnTri](int VertID, int TriID, int EdgeSubIdx)
+		{
+			int OtherEdgeSubIdx = OtherEdgeOnTri(VertID, TriID, EdgeSubIdx);
+			int AdjTri = Adj[TriID * 3 + OtherEdgeSubIdx];
+			return AdjTri;
+		};
+		TArray<bool> Seen, SeenSource; Seen.SetNumZeroed(TriNum * 3); SeenSource.SetNumZeroed(Vertices.Num());
+		for (int TriID = 0; TriID < TriNum; TriID++)
 		{
-			FIndex3i& Tri = Triangles.Emplace_GetRef(Indices[i * 3], Indices[i * 3 + 1], Indices[i * 3 + 2]);
-			if (!bOutputCCW)
+			if (Keep[TriID] != 1)
 			{
-				Swap(Tri.B, Tri.C);
+				continue;
 			}
+			for (int SubIdx = 0, OtherSubIdx=2; SubIdx < 3; OtherSubIdx=SubIdx++)
+			{
+				int WedgeIdx = TriID * 3 + SubIdx;
+				int VertID = Indices[WedgeIdx];
+				
+				if (Seen[WedgeIdx]) // already been walked over & therefore covered by previous pass
+				{
+					continue;
+				}
+
+				// if seen source but haven't seen specific wedge, then we need to duplicate the vertex and re-link
+				bool bSeenSource = SeenSource[VertID];
+
+				int NewVertID = -1;
+				if (bSeenSource)
+				{
+					FVector2<RealType> VertexToCopy = Vertices[VertID];
+					NewVertID = Vertices.Add(VertexToCopy);
+				}
+				
+				// process all triangles starting from the given tri ID and edge idx; return true if looped, false otherwise
+				auto WalkAll = [&TriNum, &Indices, &Seen, &SeenSource, &bSeenSource, &NeedUpdates, &NewVertID, &Keep,
+								&CrossEdge, &GetVertSubIdx, &Walk](int WalkVertID, int WalkTriID, int WalkSubIdx)
+				{
+					int StartTriID = WalkTriID;
+					int SafetyCounter = 0;
+					while (true)
+					{
+						int VertSubIdx = GetVertSubIdx(WalkVertID, WalkTriID);
+						int WalkWedgeIdx = WalkTriID * 3 + VertSubIdx;
+						ensure(!Seen[WalkWedgeIdx]);
+						checkSlow(Indices[WalkWedgeIdx] == WalkVertID);
+						Seen[WalkWedgeIdx] = true;
+						if (bSeenSource)
+						{
+							NeedUpdates.Add(TPair<int, int>(WalkWedgeIdx, NewVertID));
+						}
+
+						int NextTriID = Walk(WalkVertID, WalkTriID, WalkSubIdx);
+						if (NextTriID < 0 || Keep[NextTriID] != 1)
+						{
+							return false;
+						}
+						WalkSubIdx = CrossEdge(WalkTriID, NextTriID);
+						WalkTriID = NextTriID;
+						if (WalkTriID == StartTriID)
+						{
+							return true;
+						}
+						check(SafetyCounter++ < TriNum*2); // infinite loop catcher
+					}
+				};
+				bool bLooped = WalkAll(VertID, TriID, SubIdx);
+				if (!bLooped)
+				{
+					// if it didn't loop around, also walk the other direction
+					int OtherWayTriID = Walk(VertID, TriID, OtherSubIdx);
+					if (OtherWayTriID >= 0 && Keep[OtherWayTriID] == 1)
+					{
+						int OtherWayTriSubIdx = CrossEdge(TriID, OtherWayTriID);
+						ensure(!WalkAll(VertID, OtherWayTriID, OtherWayTriSubIdx));
+					}
+				}
+
+				SeenSource[VertID] = true;
+			}
+		}
+	}
+
+	if (NeedUpdates.Num() > 0)
+	{
+		std::vector<int> UpdatedIndices = Indices;
+		for (const TPair<int, int>& Update : NeedUpdates)
+		{
+			UpdatedIndices[Update.Key] = Update.Value;
 		}
+
+		AddedVerticesStartIndex = OrigNumVertices;
+
+		BuildTriangles(UpdatedIndices);
 	}
+	else
+	{
+		BuildTriangles(Indices);
+	}
+
 	
+
 	return !bBoundaryTrackingFailure;
 }
 //
diff --git a/Engine/Plugins/Experimental/GeometryProcessing/Source/GeometryAlgorithms/Public/ConstrainedDelaunay2.h b/Engine/Plugins/Experimental/GeometryProcessing/Source/GeometryAlgorithms/Public/ConstrainedDelaunay2.h
--- a/Engine/Plugins/Experimental/GeometryProcessing/Source/GeometryAlgorithms/Public/ConstrainedDelaunay2.h
+++ b/Engine/Plugins/Experimental/GeometryProcessing/Source/GeometryAlgorithms/Public/ConstrainedDelaunay2.h
@@ -28,6 +28,7 @@ struct TConstrainedDelaunay2
 
 	bool bOrientedEdges = true;
 	bool bOutputCCW = false;
+	bool bSplitBowties = false;
 
 	enum class EFillRule {
 		Odd = 0,
@@ -68,6 +69,9 @@ struct TConstrainedDelaunay2
 	//
 	TArray<FIndex3i> Triangles;
 
+	/** If vertices were added to output (e.g. to split bowties), this is set to the index of the first added vertex */
+	int AddedVerticesStartIndex = -1;
+
 	/**
 	 * Populate Triangles
 	 *
diff --git a/Engine/Plugins/Experimental/MeshModelingToolset/Source/MeshModelingTools/Private/DrawPolygonTool.cpp b/Engine/Plugins/Experimental/MeshModelingToolset/Source/MeshModelingTools/Private/DrawPolygonTool.cpp
--- a/Engine/Plugins/Experimental/MeshModelingToolset/Source/MeshModelingTools/Private/DrawPolygonTool.cpp
+++ b/Engine/Plugins/Experimental/MeshModelingToolset/Source/MeshModelingTools/Private/DrawPolygonTool.cpp
@@ -21,6 +21,7 @@
 #include "MeshQueries.h"
 #include "ToolSceneQueriesUtil.h"
 #include "ConstrainedDelaunay2.h"
+#include "Arrangement2d.h"
 
 #include "DynamicMeshEditor.h"
 
@@ -805,6 +806,7 @@ bool UDrawPolygonTool::UpdateSelfIntersection()
 	FFrame3f DrawFrame(DrawPlaneOrigin, DrawPlaneOrientation);
 	FSegment2f PreviewSegment(DrawFrame.ToPlaneUV(PolygonVertices[NumVertices - 1],2), DrawFrame.ToPlaneUV(PreviewVertex,2));
 
+	float BestIntersectionParameter = FMathf::MaxReal;
 	for (int k = 0; k < NumVertices - 2; ++k) 
 	{
 		FSegment2f Segment(DrawFrame.ToPlaneUV(PolygonVertices[k],2), DrawFrame.ToPlaneUV(PolygonVertices[k + 1],2));
@@ -812,12 +814,15 @@ bool UDrawPolygonTool::UpdateSelfIntersection()
 		if (Intersection.Find()) 
 		{
 			bHaveSelfIntersection = true;
-			SelfIntersectSegmentIdx = k;
-			SelfIntersectionPoint = DrawFrame.FromPlaneUV(Intersection.Point0,2);
-			return true;
+			if (Intersection.Parameter0 < BestIntersectionParameter)
+			{
+				BestIntersectionParameter = Intersection.Parameter0;
+				SelfIntersectSegmentIdx = k;
+				SelfIntersectionPoint = DrawFrame.FromPlaneUV(Intersection.Point0, 2);
+			}
 		}
 	}
-	return false;
+	return bHaveSelfIntersection;
 }
 
 void UDrawPolygonTool::GenerateFixedPolygon(const TArray<FVector>& FixedPoints, TArray<FVector>& VerticesOut, TArray<TArray<FVector>>& HolesVerticesOut)
@@ -1020,7 +1025,12 @@ void UDrawPolygonTool::EmitCurrentPolygon()
 	FDynamicMesh3 Mesh;
 	double ExtrudeDist = (PolygonProperties->OutputMode == EDrawPolygonOutputMode::MeshedPolygon) ?
 		0 : PolygonProperties->ExtrudeHeight;
-	GeneratePolygonMesh(PolygonVertices, PolygonHolesVertices, &Mesh, PlaneFrameOut, false, ExtrudeDist, false);
+	bool bSucceeded = GeneratePolygonMesh(PolygonVertices, PolygonHolesVertices, &Mesh, PlaneFrameOut, false, ExtrudeDist, false);
+	if (!bSucceeded) // somehow made a polygon with no valid triangulation; just throw it away ...
+	{
+		ResetPolygon();
+		return;
+	}
 
 	GetToolManager()->BeginUndoTransaction(LOCTEXT("CreatePolygon", "Create Polygon"));
 
@@ -1043,7 +1053,7 @@ void UDrawPolygonTool::EmitCurrentPolygon()
 void UDrawPolygonTool::UpdateLivePreview()
 {
 	int NumVerts = PolygonVertices.Num();
-	if (NumVerts < 2 || PreviewMesh == nullptr || PreviewMesh->IsVisible() == false )
+	if (NumVerts < 2 || PreviewMesh == nullptr || PreviewMesh->IsVisible() == false)
 	{
 		return;
 	}
@@ -1052,15 +1062,16 @@ void UDrawPolygonTool::UpdateLivePreview()
 	FDynamicMesh3 Mesh;
 	double ExtrudeDist = (PolygonProperties->OutputMode == EDrawPolygonOutputMode::MeshedPolygon) ?
 		0 : PolygonProperties->ExtrudeHeight;
-	GeneratePolygonMesh(PolygonVertices, PolygonHolesVertices, &Mesh, PlaneFrame, false, ExtrudeDist, false);
-
-	PreviewMesh->SetTransform(PlaneFrame.ToFTransform());
-	PreviewMesh->SetMaterial(MaterialProperties->Material);
-	PreviewMesh->EnableWireframe(MaterialProperties->bWireframe);
-	PreviewMesh->UpdatePreview(&Mesh);
+	if (GeneratePolygonMesh(PolygonVertices, PolygonHolesVertices, &Mesh, PlaneFrame, false, ExtrudeDist, false))
+	{
+		PreviewMesh->SetTransform(PlaneFrame.ToFTransform());
+		PreviewMesh->SetMaterial(MaterialProperties->Material);
+		PreviewMesh->EnableWireframe(MaterialProperties->bWireframe);
+		PreviewMesh->UpdatePreview(&Mesh);
+	}
 }
 
-void UDrawPolygonTool::GeneratePolygonMesh(const TArray<FVector>& Polygon, const TArray<TArray<FVector>>& PolygonHoles, FDynamicMesh3* ResultMeshOut, FFrame3d& WorldFrameOut, bool bIncludePreviewVtx, double ExtrudeDistance, bool bExtrudeSymmetric)
+bool UDrawPolygonTool::GeneratePolygonMesh(const TArray<FVector>& Polygon, const TArray<TArray<FVector>>& PolygonHoles, FDynamicMesh3* ResultMeshOut, FFrame3d& WorldFrameOut, bool bIncludePreviewVtx, double ExtrudeDistance, bool bExtrudeSymmetric)
 {
 	// construct centered frame for polygon
 	WorldFrameOut = FFrame3d(DrawPlaneOrigin, DrawPlaneOrientation);
@@ -1114,12 +1125,41 @@ void UDrawPolygonTool::GeneratePolygonMesh(const TArray<FVector>& Polygon, const
 	}
 
 	FConstrainedDelaunay2d Triangulator;
-	Triangulator.Add(GeneralPolygon);
+	if (PolygonProperties->bAllowSelfIntersections)
+	{
+		FArrangement2d Arrangement(OuterPolygon.Bounds());
+		// arrangement2d builds a general 2d graph that discards orientation info ...
+		Triangulator.FillRule = FConstrainedDelaunay2d::EFillRule::Odd;
+		Triangulator.bOrientedEdges = false;
+		Triangulator.bSplitBowties = true;
+		for (FSegment2d Seg : GeneralPolygon.GetOuter().Segments())
+		{
+			Arrangement.Insert(Seg);
+		}
+		Triangulator.Add(Arrangement.Graph);
+		for (const FPolygon2d& Hole : GeneralPolygon.GetHoles())
+		{
+			Triangulator.Add(Hole, true);
+		}
+	}
+	else
+	{
+		Triangulator.Add(GeneralPolygon);
+	}
+
+	
 	bool bTriangulationSuccess = Triangulator.Triangulate();
-	// TODO: actually do something different if triangulation fails?  note that failure doesn't mean that it didn't generate triangles ... it will try to give triangles back even if input was not fully valid
+	// only truly fail if we got zero triangles back from the triangulator; if it just returned false it may still have managed to partially generate something
+	if (Triangulator.Triangles.Num() == 0)
+	{
+		return false;
+	}
+
 	TriangulationMeshGen.Vertices2D = Triangulator.Vertices;
 	TriangulationMeshGen.Triangles2D = Triangulator.Triangles;
 
+	
+
 	ResultMeshOut->Copy(&TriangulationMeshGen.Generate());
 
 	// for symmetric extrude we translate the first poly by -dist along axis
@@ -1165,12 +1205,16 @@ void UDrawPolygonTool::GeneratePolygonMesh(const TArray<FVector>& Polygon, const
 		GlobalUVRescale = MaterialProperties->UVScale * .01 / InitialUVScale;
 	}
 	Editor.RescaleAttributeUVs(GlobalUVRescale, false);
+
+	return true;
 }
 
 void UDrawPolygonTool::GeneratePreviewHeightTarget()
 {
-	GeneratePolygonMesh(PolygonVertices, PolygonHolesVertices, &PreviewHeightTarget, PreviewHeightFrame, false, 99999, true);
-	PreviewHeightTargetAABB.SetMesh(&PreviewHeightTarget);
+	if (GeneratePolygonMesh(PolygonVertices, PolygonHolesVertices, &PreviewHeightTarget, PreviewHeightFrame, false, 99999, true))
+	{
+		PreviewHeightTargetAABB.SetMesh(&PreviewHeightTarget);
+	}
 }
 
 
diff --git a/Engine/Plugins/Experimental/MeshModelingToolset/Source/MeshModelingTools/Public/DrawPolygonTool.h b/Engine/Plugins/Experimental/MeshModelingToolset/Source/MeshModelingTools/Public/DrawPolygonTool.h
--- a/Engine/Plugins/Experimental/MeshModelingToolset/Source/MeshModelingTools/Public/DrawPolygonTool.h
+++ b/Engine/Plugins/Experimental/MeshModelingToolset/Source/MeshModelingTools/Public/DrawPolygonTool.h
@@ -328,8 +328,8 @@ protected:
 
 	FFrame3f HitPosFrameWorld;
 
-	// we use this to generate extruded meshes
-	void GeneratePolygonMesh(const TArray<FVector>& Polygon, const TArray<TArray<FVector>>& PolygonHoles, FDynamicMesh3* ResultMeshOut, FFrame3d& WorldFrameOut, bool bIncludePreviewVtx, double ExtrudeDistance, bool bExtrudeSymmetric);
+	/** Generate extruded meshes.  Returns true on success. */
+	bool GeneratePolygonMesh(const TArray<FVector>& Polygon, const TArray<TArray<FVector>>& PolygonHoles, FDynamicMesh3* ResultMeshOut, FFrame3d& WorldFrameOut, bool bIncludePreviewVtx, double ExtrudeDistance, bool bExtrudeSymmetric);
 
 
 	// user feedback messages
