custom triangulations for select polygon types on the draw polygon tool

#rnx

diff --git a/Engine/Plugins/Experimental/GeometryProcessing/Source/GeometricObjects/Private/Generators/RectangleMeshGenerator.cpp b/Engine/Plugins/Experimental/GeometryProcessing/Source/GeometricObjects/Private/Generators/RectangleMeshGenerator.cpp
--- a/Engine/Plugins/Experimental/GeometryProcessing/Source/GeometricObjects/Private/Generators/RectangleMeshGenerator.cpp
+++ b/Engine/Plugins/Experimental/GeometryProcessing/Source/GeometricObjects/Private/Generators/RectangleMeshGenerator.cpp
@@ -139,8 +139,9 @@ FMeshShapeGenerator& FRoundedRectangleMeshGenerator::Generate()
 		Normals[VertIdx] = Normal;
 	}
 
-	float TotWidth = Radius * 2 + Width;
-	float TotHeight = Radius * 2 + Height;
+	float TotWidth = FMathf::Max(FMathf::ZeroTolerance, Radius * 2 + Width);
+	float TotHeight = FMathf::Max(FMathf::ZeroTolerance, Radius * 2 + Height);
+	
 
 	// corner vertices
 	FVector3d v00 = MakeVertex(-TotWidth / 2.0f, -TotHeight / 2.0f);
diff --git a/Engine/Plugins/Experimental/MeshModelingToolset/Source/MeshModelingTools/Private/DrawPolygonTool.cpp b/Engine/Plugins/Experimental/MeshModelingToolset/Source/MeshModelingTools/Private/DrawPolygonTool.cpp
--- a/Engine/Plugins/Experimental/MeshModelingToolset/Source/MeshModelingTools/Private/DrawPolygonTool.cpp
+++ b/Engine/Plugins/Experimental/MeshModelingToolset/Source/MeshModelingTools/Private/DrawPolygonTool.cpp
@@ -15,6 +15,8 @@
 
 #include "MeshDescriptionBuilder.h"
 #include "Generators/FlatTriangulationMeshGenerator.h"
+#include "Generators/DiscMeshGenerator.h"
+#include "Generators/RectangleMeshGenerator.h"
 #include "Operations/ExtrudeMesh.h"
 #include "Distance/DistLine3Ray3.h"
 #include "Intersection/IntrSegment2Segment2.h"
@@ -825,18 +827,37 @@ bool UDrawPolygonTool::UpdateSelfIntersection()
 	return bHaveSelfIntersection;
 }
 
-void UDrawPolygonTool::GenerateFixedPolygon(const TArray<FVector>& FixedPoints, TArray<FVector>& VerticesOut, TArray<TArray<FVector>>& HolesVerticesOut)
+void UDrawPolygonTool::GetPolygonParametersFromFixedPoints(const TArray<FVector>& FixedPoints, FVector2f& FirstReferencePt, FVector2f& BoxSize, float& YSign, float& AngleRad)
 {
+	if (FixedPoints.Num() < 2)
+	{
+		return;
+	}
 
 	FFrame3f DrawFrame(DrawPlaneOrigin, DrawPlaneOrientation);
-	FVector2f CenterPt = DrawFrame.ToPlaneUV(FixedPoints[0], 2 );
+	FirstReferencePt = DrawFrame.ToPlaneUV(FixedPoints[0], 2);
+
 	FVector2f EdgePt = DrawFrame.ToPlaneUV(FixedPoints[1], 2);
-	FVector2f Delta = EdgePt - CenterPt;
-	float AngleRad = FMath::Atan2(Delta.Y, Delta.X);
+	FVector2f Delta = EdgePt - FirstReferencePt;
+	AngleRad = FMath::Atan2(Delta.Y, Delta.X);
+
+	float Radius = Delta.Length();
+	FVector2f AxisX = Delta / Radius;
+	FVector2f AxisY = -AxisX.Perp();
+	FVector2f HeightPt = DrawFrame.ToPlaneUV((FixedPoints.Num() == 3) ? FixedPoints[2] : FixedPoints[1], 2);
+	FVector2f HeightDelta = HeightPt - FirstReferencePt;
+	YSign = FMath::Sign(HeightDelta.Dot(AxisY));
+	BoxSize.X = Radius;
+	BoxSize.Y = FMath::Abs(HeightDelta.Dot(AxisY));
+}
+
+void UDrawPolygonTool::GenerateFixedPolygon(const TArray<FVector>& FixedPoints, TArray<FVector>& VerticesOut, TArray<TArray<FVector>>& HolesVerticesOut)
+{
+	FVector2f FirstReferencePt, BoxSize;
+	float YSign, AngleRad;
+	GetPolygonParametersFromFixedPoints(FixedPoints, FirstReferencePt, BoxSize, YSign, AngleRad);
+	float Width = BoxSize.X, Height = BoxSize.Y;
 	FMatrix2f RotationMat = FMatrix2f::RotationRad(AngleRad);
-	FVector2f RotAxisX = RotationMat * FVector2f::UnitX();
-	float Dist = Delta.Length();
-	float Width = FMath::Abs(Delta.Dot(RotAxisX));
 
 	FPolygon2f Polygon;
 	TArray<FPolygon2f> PolygonHoles;
@@ -846,11 +867,6 @@ void UDrawPolygonTool::GenerateFixedPolygon(const TArray<FVector>& FixedPoints,
 	}
 	else if (PolygonProperties->PolygonType == EDrawPolygonDrawMode::Rectangle || PolygonProperties->PolygonType == EDrawPolygonDrawMode::RoundedRectangle)
 	{
-		FVector2f HeightPt = DrawFrame.ToPlaneUV((FixedPoints.Num() == 3) ? FixedPoints[2] : FixedPoints[1], 2);
-		FVector2f HeightDelta = HeightPt - CenterPt;
-		FVector2f RotAxisY = RotationMat * FVector2f::UnitY();
-		float YSign = FMath::Sign(HeightDelta.Dot(RotAxisY));
-		float Height = FMath::Abs(HeightDelta.Dot(RotAxisY));
 		if (PolygonProperties->PolygonType == EDrawPolygonDrawMode::Rectangle)
 		{
 			Polygon = FPolygon2f::MakeRectangle(FVector2f(Width / 2, YSign*Height / 2), Width, Height);
@@ -862,10 +878,10 @@ void UDrawPolygonTool::GenerateFixedPolygon(const TArray<FVector>& FixedPoints,
 	}
 	else // Circle or HoleyCircle
 	{
-		Polygon = FPolygon2f::MakeCircle(Dist, PolygonProperties->Steps, 0);
+		Polygon = FPolygon2f::MakeCircle(Width, PolygonProperties->Steps, 0);
 		if (PolygonProperties->PolygonType == EDrawPolygonDrawMode::HoleyCircle)
 		{
-			PolygonHoles.Add(FPolygon2f::MakeCircle(Dist * FMathd::Clamp(PolygonProperties->FeatureSizeRatio, .01, .99), PolygonProperties->Steps, 0));
+			PolygonHoles.Add(FPolygon2f::MakeCircle(Width * FMathd::Clamp(PolygonProperties->FeatureSizeRatio, .01, .99), PolygonProperties->Steps, 0));
 		}
 	}
 	Polygon.Transform([RotationMat](const FVector2f& Pt) { return RotationMat * Pt; });
@@ -874,10 +890,11 @@ void UDrawPolygonTool::GenerateFixedPolygon(const TArray<FVector>& FixedPoints,
 		Hole.Transform([RotationMat](const FVector2f& Pt) { return RotationMat * Pt; });
 	}
 
+	FFrame3f DrawFrame(DrawPlaneOrigin, DrawPlaneOrientation);
 	VerticesOut.SetNum(Polygon.VertexCount());
 	for (int k = 0; k < Polygon.VertexCount(); ++k)
 	{
-		FVector2f NewPt = CenterPt + Polygon[k];
+		FVector2f NewPt = FirstReferencePt + Polygon[k];
 		VerticesOut[k] = DrawFrame.FromPlaneUV(NewPt, 2);
 	}
 
@@ -888,7 +905,7 @@ void UDrawPolygonTool::GenerateFixedPolygon(const TArray<FVector>& FixedPoints,
 		HolesVerticesOut[HoleIdx].SetNum(NumHoleVerts);
 		for (int k = 0; k < NumHoleVerts; ++k)
 		{
-			FVector2f NewPt = CenterPt + PolygonHoles[HoleIdx][k];
+			FVector2f NewPt = FirstReferencePt + PolygonHoles[HoleIdx][k];
 			HolesVerticesOut[HoleIdx][k] = DrawFrame.FromPlaneUV(NewPt, 2);
 		}
 	}
@@ -1085,11 +1102,7 @@ bool UDrawPolygonTool::GeneratePolygonMesh(const TArray<FVector>& Polygon, const
 	Centroid /= (double)NumVerts;
 	WorldFrameOut.Origin = Centroid;
 
-
-	// triangulate polygon into the MeshDescription
-	FGeneralPolygon2d GeneralPolygon;
-	FFlatTriangulationMeshGenerator TriangulationMeshGen;
-
+	// Compute outer polygon & bounds
 	auto VertexArrayToPolygon = [&WorldFrameOut](const TArray<FVector>& Vertices)
 	{
 		FPolygon2d OutPolygon;
@@ -1099,9 +1112,7 @@ bool UDrawPolygonTool::GeneratePolygonMesh(const TArray<FVector>& Polygon, const
 		}
 		return OutPolygon;
 	};
-
 	FPolygon2d OuterPolygon = VertexArrayToPolygon(Polygon);
-
 	// add preview vertex
 	if (bIncludePreviewVtx)
 	{
@@ -1110,57 +1121,112 @@ bool UDrawPolygonTool::GeneratePolygonMesh(const TArray<FVector>& Polygon, const
 			OuterPolygon.AppendVertex(WorldFrameOut.ToPlaneUV(PreviewVertex, 2));
 		}
 	}
+	FAxisAlignedBox2d Bounds(OuterPolygon.Bounds());
 
-	if (OuterPolygon.IsClockwise() == false)
+	// special case paths
+	if (PolygonProperties->PolygonType == EDrawPolygonDrawMode::HoleyCircle || PolygonProperties->PolygonType == EDrawPolygonDrawMode::Circle || PolygonProperties->PolygonType == EDrawPolygonDrawMode::RoundedRectangle)
 	{
-		OuterPolygon.Reverse();
-	}
+		// get polygon parameters
+		FVector2f FirstReferencePt, BoxSize;
+		float YSign, AngleRad;
+		GetPolygonParametersFromFixedPoints(FixedPolygonClickPoints, FirstReferencePt, BoxSize, YSign, AngleRad);
+		FirstReferencePt -= FVector2f(Centroid.X, Centroid.Y);
+		FMatrix2f RotationMat = FMatrix2f::RotationRad(AngleRad);
 
-	GeneralPolygon.SetOuter(OuterPolygon);
+		// translate general polygon parameters to specific mesh generator parameters, and generate mesh
+		if (PolygonProperties->PolygonType == EDrawPolygonDrawMode::HoleyCircle)
+		{
+			FPuncturedDiscMeshGenerator HCGen;
+			HCGen.AngleSamples = PolygonProperties->Steps;
+			HCGen.RadialSamples = 1;
+			HCGen.Radius = BoxSize.X;
+			HCGen.HoleRadius = BoxSize.X * FMath::Clamp(PolygonProperties->FeatureSizeRatio, .01f, .99f);
+			ResultMeshOut->Copy(&HCGen.Generate());
+		}
+		else if (PolygonProperties->PolygonType == EDrawPolygonDrawMode::Circle)
+		{
+			FDiscMeshGenerator CGen;
+			CGen.AngleSamples = PolygonProperties->Steps;
+			CGen.RadialSamples = 1;
+			CGen.Radius = BoxSize.X;
+			ResultMeshOut->Copy(&CGen.Generate());
+		}
+		else if (PolygonProperties->PolygonType == EDrawPolygonDrawMode::RoundedRectangle)
+		{
+			FRoundedRectangleMeshGenerator RRGen;
+			FirstReferencePt += RotationMat * (FVector2f(BoxSize.X, BoxSize.Y * YSign)*.5f);
+			RRGen.AngleSamples = PolygonProperties->Steps;
+			RRGen.Radius = .5 * FMath::Min(BoxSize.X, BoxSize.Y) * FMath::Clamp(PolygonProperties->FeatureSizeRatio, .01f, .99f);
+			RRGen.Height = BoxSize.Y - RRGen.Radius * 2.;
+			RRGen.Width = BoxSize.X - RRGen.Radius * 2.;
+			RRGen.WidthVertexCount = 1;
+			RRGen.HeightVertexCount = 1;
+			ResultMeshOut->Copy(&RRGen.Generate());
+		}
 
-	for (int HoleIdx = 0; HoleIdx < PolygonHoles.Num(); HoleIdx++)
-	{
-		// attempt to add holes (skipping if safety checks fail)
-		GeneralPolygon.AddHole(VertexArrayToPolygon(PolygonHoles[HoleIdx]), true, false /*currently don't care about hole orientation; we'll just set the triangulation algo not to care*/);
+		// transform generated mesh
+		for (int VertIdx : ResultMeshOut->VertexIndicesItr())
+		{
+			FVector3d V = ResultMeshOut->GetVertex(VertIdx);
+			FVector2f VTransformed = RotationMat * FVector2f(V.X, V.Y) + FirstReferencePt;
+			ResultMeshOut->SetVertex(VertIdx, FVector3d(VTransformed.X, VTransformed.Y, 0));
+		}
 	}
-
-	FConstrainedDelaunay2d Triangulator;
-	if (PolygonProperties->bAllowSelfIntersections)
+	else // generic path: triangulate using polygon vertices
 	{
-		FArrangement2d Arrangement(OuterPolygon.Bounds());
-		// arrangement2d builds a general 2d graph that discards orientation info ...
-		Triangulator.FillRule = FConstrainedDelaunay2d::EFillRule::Odd;
-		Triangulator.bOrientedEdges = false;
-		Triangulator.bSplitBowties = true;
-		for (FSegment2d Seg : GeneralPolygon.GetOuter().Segments())
+		// triangulate polygon into the MeshDescription
+		FGeneralPolygon2d GeneralPolygon;
+		FFlatTriangulationMeshGenerator TriangulationMeshGen;
+
+		if (OuterPolygon.IsClockwise() == false)
 		{
-			Arrangement.Insert(Seg);
+			OuterPolygon.Reverse();
 		}
-		Triangulator.Add(Arrangement.Graph);
-		for (const FPolygon2d& Hole : GeneralPolygon.GetHoles())
+
+		GeneralPolygon.SetOuter(OuterPolygon);
+
+		for (int HoleIdx = 0; HoleIdx < PolygonHoles.Num(); HoleIdx++)
 		{
-			Triangulator.Add(Hole, true);
+			// attempt to add holes (skipping if safety checks fail)
+			GeneralPolygon.AddHole(VertexArrayToPolygon(PolygonHoles[HoleIdx]), true, false /*currently don't care about hole orientation; we'll just set the triangulation algo not to care*/);
 		}
-	}
-	else
-	{
-		Triangulator.Add(GeneralPolygon);
-	}
 
-	
-	bool bTriangulationSuccess = Triangulator.Triangulate();
-	// only truly fail if we got zero triangles back from the triangulator; if it just returned false it may still have managed to partially generate something
-	if (Triangulator.Triangles.Num() == 0)
-	{
-		return false;
-	}
+		FConstrainedDelaunay2d Triangulator;
+		if (PolygonProperties->bAllowSelfIntersections)
+		{
+			FArrangement2d Arrangement(OuterPolygon.Bounds());
+			// arrangement2d builds a general 2d graph that discards orientation info ...
+			Triangulator.FillRule = FConstrainedDelaunay2d::EFillRule::Odd;
+			Triangulator.bOrientedEdges = false;
+			Triangulator.bSplitBowties = true;
+			for (FSegment2d Seg : GeneralPolygon.GetOuter().Segments())
+			{
+				Arrangement.Insert(Seg);
+			}
+			Triangulator.Add(Arrangement.Graph);
+			for (const FPolygon2d& Hole : GeneralPolygon.GetHoles())
+			{
+				Triangulator.Add(Hole, true);
+			}
+		}
+		else
+		{
+			Triangulator.Add(GeneralPolygon);
+		}
 
-	TriangulationMeshGen.Vertices2D = Triangulator.Vertices;
-	TriangulationMeshGen.Triangles2D = Triangulator.Triangles;
 
-	
+		bool bTriangulationSuccess = Triangulator.Triangulate();
+		// only truly fail if we got zero triangles back from the triangulator; if it just returned false it may still have managed to partially generate something
+		if (Triangulator.Triangles.Num() == 0)
+		{
+			return false;
+		}
 
-	ResultMeshOut->Copy(&TriangulationMeshGen.Generate());
+		TriangulationMeshGen.Vertices2D = Triangulator.Vertices;
+		TriangulationMeshGen.Triangles2D = Triangulator.Triangles;
+
+		ResultMeshOut->Copy(&TriangulationMeshGen.Generate());
+	}
 
 	// for symmetric extrude we translate the first poly by -dist along axis
 	if (bExtrudeSymmetric)
@@ -1179,8 +1245,8 @@ bool UDrawPolygonTool::GeneratePolygonMesh(const TArray<FVector>& Polygon, const
 	{
 		FExtrudeMesh Extruder(ResultMeshOut);
 		Extruder.DefaultExtrudeDistance = ExtrudeDistance;
-		FAxisAlignedBox2d bounds = GeneralPolygon.Bounds();
-		Extruder.UVScaleFactor = 1.0 / bounds.MaxDim();
+		
+		Extruder.UVScaleFactor = 1.0 / Bounds.MaxDim();
 		if (ExtrudeDistance < 0)
 		{
 			Extruder.IsPositiveOffset = false;
@@ -1196,7 +1262,7 @@ bool UDrawPolygonTool::GeneratePolygonMesh(const TArray<FVector>& Polygon, const
 	}
 
 	FDynamicMeshEditor Editor(ResultMeshOut);
-	float InitialUVScale = 1.0 / GeneralPolygon.Bounds().MaxDim(); // this is the UV scale used by both the polymeshgen and the extruder above
+	float InitialUVScale = 1.0 / Bounds.MaxDim(); // this is the UV scale used by both the polymeshgen and the extruder above
 	// default global rescale -- initial scale doesn't factor in extrude distance; rescale so UVScale of 1.0 fits in the unit square texture
 	float GlobalUVRescale = MaterialProperties->UVScale / FMathf::Max(1.0f, ExtrudeDistance * InitialUVScale);
 	if (MaterialProperties->bWorldSpaceUVScale)
diff --git a/Engine/Plugins/Experimental/MeshModelingToolset/Source/MeshModelingTools/Public/DrawPolygonTool.h b/Engine/Plugins/Experimental/MeshModelingToolset/Source/MeshModelingTools/Public/DrawPolygonTool.h
--- a/Engine/Plugins/Experimental/MeshModelingToolset/Source/MeshModelingTools/Public/DrawPolygonTool.h
+++ b/Engine/Plugins/Experimental/MeshModelingToolset/Source/MeshModelingTools/Public/DrawPolygonTool.h
@@ -310,6 +310,7 @@ protected:
 	FPointPlanarSnapSolver SnapEngine;
 	ToolSceneQueriesUtil::FSnapGeometry LastSnapGeometry;
 
+	void GetPolygonParametersFromFixedPoints(const TArray<FVector>& FixedPoints, FVector2f& FirstReferencePt, FVector2f& BoxSize, float& YSign, float& AngleRad);
 	void GenerateFixedPolygon(const TArray<FVector>& FixedPoints, TArray<FVector>& VerticesOut, TArray<TArray<FVector>>& HolesVerticesOut);
 
 
