Mesh Modeling:  Mesh Space Deformer tools.   Update to use interval gizmo and respect relative location of gizmo to geometry.

diff --git a/Engine/Plugins/Experimental/MeshModelingToolset/Source/MeshModelingTools/Private/MeshSpaceDeformerTool.cpp b/Engine/Plugins/Experimental/MeshModelingToolset/Source/MeshModelingTools/Private/MeshSpaceDeformerTool.cpp
--- a/Engine/Plugins/Experimental/MeshModelingToolset/Source/MeshModelingTools/Private/MeshSpaceDeformerTool.cpp
+++ b/Engine/Plugins/Experimental/MeshModelingToolset/Source/MeshModelingTools/Private/MeshSpaceDeformerTool.cpp
@@ -2,521 +2,497 @@
 
 #include "MeshSpaceDeformerTool.h"
 #include "InteractiveToolManager.h"
+#include "InteractiveGizmoManager.h"
 #include "ToolBuilderUtil.h"
 
 #include "SegmentTypes.h"
 #include "DynamicMeshAttributeSet.h"
 #include "MeshNormals.h"
+#include "MeshOpPreviewHelpers.h"
 #include "ToolSceneQueriesUtil.h"
+#include "ToolSetupUtil.h"
 #include "Intersection/IntersectionUtil.h"
+#include "PreviewMesh.h"
 
-#include "Async/ParallelFor.h"
-#include "Containers/BitArray.h"
+#include "BaseGizmos/GizmoComponents.h"
+#include "BaseGizmos/TransformGizmo.h"
+#include "BaseGizmos/IntervalGizmo.h"
+#include "PositionPlaneGizmo.h"
+#include "MeshDescriptionToDynamicMesh.h"
+#include "DynamicMeshToMeshDescription.h"
+#include "CoreMinimal.h"
+#include "Math/Matrix.h"
 
-///////////////////////////////
-// EIGEN
-
-#if defined(_MSC_VER) && USING_CODE_ANALYSIS
-#pragma warning(push)
-#pragma warning(disable : 6011)
-#pragma warning(disable : 6387)
-#pragma warning(disable : 6313)
-#pragma warning(disable : 6294)
-#endif
-THIRD_PARTY_INCLUDES_START
-#include <Eigen/Core>
-#include <Eigen/Dense> 
-#include <Eigen/Eigenvalues>
-
-#ifndef EIGEN_MPL2_ONLY
-#endif
-THIRD_PARTY_INCLUDES_END
-#if defined(_MSC_VER) && USING_CODE_ANALYSIS
-#pragma warning(pop)
-#endif
-
-///////////////////////////////
 
 
 #define LOCTEXT_NAMESPACE "MeshSpaceDeformerTool"
 
 
-//////////////////////////////
-// DEBUG_SETTINGS
+/*
+ * ToolBuilder
+ */
+UMeshSurfacePointTool* UMeshSpaceDeformerToolBuilder::CreateNewTool(const FToolBuilderState& SceneState) const
+{
+	UMeshSpaceDeformerTool* MeshSpaceDeformerTool = NewObject<UMeshSpaceDeformerTool>(SceneState.ToolManager);
+	
+	UActorComponent* ActorComponent = ToolBuilderUtil::FindFirstComponent(SceneState, CanMakeComponentTarget);
+	auto* MeshComponent = Cast<UPrimitiveComponent>(ActorComponent);
+	check(MeshComponent != nullptr);
 
-#define DEBUG_ASYNC_TASK
+	MeshSpaceDeformerTool->SetSelection(MakeComponentTarget(MeshComponent));
+	MeshSpaceDeformerTool->SetWorld(SceneState.World);
 
-#define DEBUG_DRAW_AXES
-//////////////////////////////
+	return MeshSpaceDeformerTool;
+}
 
 
-#ifdef DEBUG_ASYNC_TASK
-#define print(text,...) UE_LOG(LogTemp, Warning, TEXT(text), __VA_ARGS__)
-#else
-#define print(x)
-#endif
+bool UMeshSpaceDeformerToolBuilder::CanBuildTool(const FToolBuilderState& SceneState) const
+{
+	return ToolBuilderUtil::CountComponents(SceneState, CanMakeComponentTarget) == 1;
+}
 
 
 
-double RayToSegmentSquareDist(const FRay3d& Ray, const FVector3d& V0, const FVector3d& V1, FVector3d* ClosestPtOnSegment, FVector3d* ClosestPtOnRay)
+TUniquePtr<FDynamicMeshOperator> USpaceDeformerOperatorFactory::MakeNewOperator()
 {
-	if (V0 == V1)
-	{
-		if (ClosestPtOnRay) 
-		{
-			*ClosestPtOnRay = V0;
-		}
 
-		if (ClosestPtOnSegment)
+	check(SpaceDeformerTool);
+
+	const ENonlinearOperationType OperationType = SpaceDeformerTool->SelectedOperationType;
+	
+	// Create the actual operator type based on the requested operation
+	TUniquePtr<FMeshSpaceDeformerOp>  DeformerOp;
+	
+	switch (OperationType)
+	{
+		case ENonlinearOperationType::Bend: 
 		{
-			*ClosestPtOnSegment = V0;
+			DeformerOp = MakeUnique<FBendMeshOp>();
+			break;
 		}
-		return 0.;
-	}
-	FVector3d SegCenter = (V0 + V1) * 0.5;
-	FVector3d SegDir = (V1 - V0).Normalized();
-	FVector3d Diff = Ray.Origin - SegCenter;
-
-	double SegExtent = V0.Distance(V1) * 0.5;
-	double A01 = -Ray.Direction.Dot(SegDir);
-	double B0 = Diff.Dot(Ray.Direction);
-	double B1 = -Diff.Dot(SegDir);
-	double C = (Diff).SquaredLength();
-	double det = FMath::Abs(1. - A01 * A01);
-	double S0, S1, SqrDist, ExtDet;
-
-	if (det > 0.)
-	{
-		// The ray and segment are not parallel.
-
-		S0 = A01 * B1 - B0;
-		S1 = A01 * B0 - B1;
-		ExtDet = SegExtent * det;
-
-		if (S0 >= 0.)
+		case ENonlinearOperationType::Flare:
 		{
-			if (S1 >= -ExtDet)
-			{
-				if (S1 <= ExtDet)
-				{
-					// region 0
-					// Minimum at interior points of ray and segment.
-					float invDet = 1. / det;
-					S0 *= invDet;
-					S1 *= invDet;
-					SqrDist = S0 * (S0 + A01 * S1 + 2. * B0) + S1 * (A01 * S0 + S1 + 2. * B1) + C;
-				}
-				else
-				{
-					// region 1
-					S1 = SegExtent;
-					S0 = std::max(0., -(A01 * S1 + B0));
-					SqrDist = -S0 * S0 + S1 * (S1 + 2. * B1) + C;
-				}
-			}
-			else
-			{
-				// region 5
-				S1 = -SegExtent;
-				S0 = std::max(0., -(A01 * S1 + B0));
-				SqrDist = -S0 * S0 + S1 * (S1 + 2. * B1) + C;
-			}
+			DeformerOp = MakeUnique<FFlareMeshOp>();
+			break;
 		}
-		else
+		case ENonlinearOperationType::Twist:
 		{
-			if (S1 <= -ExtDet)
-			{
-				// region 4
-				S0 = std::max(0., -(-A01 * SegExtent + B0));
-				S1 = (S0 > 0.) ? -SegExtent : std::min(std::max(-SegExtent, -B1), SegExtent);
-				SqrDist = -S0 * S0 + S1 * (S1 + 2. * B1) + C;
-			}
-			else if (S1 <= ExtDet)
-			{
-				// region 3
-				S0 = 0;
-				S1 = std::min(std::max(-SegExtent, -B1), SegExtent);
-				SqrDist = S1 * (S1 + 2. * B1) + C;
-			}
-			else
-			{
-				// region 2
-				S0 = std::max(0., -(A01 * SegExtent + B0));
-				S1 = (S0 > 0.) ? SegExtent : std::min(std::max(-SegExtent, -B1), SegExtent);
-				SqrDist = -S0 * S0 + S1 * (S1 + 2. * B1) + C;
-			}
+			DeformerOp = MakeUnique<FTwistMeshOp>();
+			break;
 		}
-	}
-	else
-	{
-		// Ray and segment are parallel.
-
-		S1 = (A01 > 0.) ? -SegExtent : SegExtent;
-		S0 = std::max(0., -(A01 * S1 + B0));
-		SqrDist = -S0 * S0 + S1 * (S1 + 2. * B1) + C;
+	default:
+		check(0);
 	}
 
-	if (ClosestPtOnSegment)
-	{
-		*ClosestPtOnSegment = SegDir * S1 + SegCenter;
-	}
-
-	if (ClosestPtOnRay)
-	{
-		*ClosestPtOnRay = Ray.Direction * S0 + Ray.Origin;
-	}
-	return SqrDist;
-}
-
-FMatrix3d SolveSampleCovariance(double* Array, const int NCols, const int NRows, FVector3d& ExtentsOut, FVector3d& OriginOut)
-{
-	using RowMajorVec3dMatrix = Eigen::Matrix<double, Eigen::Dynamic, 3, Eigen::RowMajor>;
-	using RowMajorVec3d = Eigen::Matrix<double, 1, 3, Eigen::RowMajor>;
-	RowMajorVec3dMatrix VertexMatrix = Eigen::Map<RowMajorVec3dMatrix>(Array, NRows, NCols);
-
-	check(VertexMatrix(0, 0) == *Array); // Ensure the mapping is in expected format
-	check(VertexMatrix(0, 1) == Array[1]);
-	check(VertexMatrix(0, 2) == Array[2]);
-
-	//Using adjoint eigen solver
-	auto Origin = VertexMatrix.colwise().mean();															 // Get the centroid of the point cloud
-	Eigen::MatrixXd CenteredVertices = VertexMatrix.rowwise() - Origin;										 // Translate each point so centroid is (0,0,0)
-	Eigen::MatrixXd CovarianceMatrix = (CenteredVertices.adjoint() * CenteredVertices) / double(NRows - 1);	 // note: adjoint ~= conjugate transpose; formula from https://en.wikipedia.org/wiki/Sample_mean_and_covariance#Sample_covariance
+	// Operator runs on another thread - copy data over that it needs.
+	SpaceDeformerTool->UpdateOpParameters(*DeformerOp);
 	
-	// note: SelfAdjointEigenSolver automatically sorts the eigen vectors by the order of corresponding eigenvalues in increasing order
-	Eigen::SelfAdjointEigenSolver<Eigen::MatrixXd> EigenSolver(CovarianceMatrix);							 
-	Eigen::Matrix3d EigenVectors = EigenSolver.eigenvectors();
 
-	//Align object with the standard basis by rotating all the points by the inverse of EigenVectors.
-	//EigenVectors is an orthonormal basis (3x3) so transpose is the inverse; 
-	Eigen::Matrix3Xd RotatedVertices = EigenVectors.transpose() * CenteredVertices.transpose();
+	// give the operator
+	return DeformerOp;
+}
 
-	//Since our point cloud is now aligned with the standard axes, the dimensions of its AABB should be the dimensions of its 
-	RowMajorVec3d Extents = RotatedVertices.rowwise().maxCoeff() - RotatedVertices.rowwise().minCoeff();
 
-	std::sort(Extents.data(), Extents.data() +3);						//Sort the extents to match our sorted eigenvectors
-	ExtentsOut = FVector3d{ FMath::Abs(Extents[0]), FMath::Abs(Extents[1]), FMath::Abs(Extents[2]) } / 2;	//Divide by two because extent represents distance from center
-	OriginOut = FVector3d{ Origin(0,0), Origin(0,1), Origin(0,2) };		//Origin in world space
+/*
+ * Tool
+ */
+UMeshSpaceDeformerTool::UMeshSpaceDeformerTool()
+{
+	GizmoCenter = FVector::ZeroVector;
+	GizmoOrientation = FQuat::Identity;
 
-	return FMatrix3d{ FVector3d{ EigenVectors(0,0), EigenVectors(1,0), EigenVectors(2,0) },
-					  FVector3d{ EigenVectors(0,1), EigenVectors(1,1), EigenVectors(2,1) },
-					  FVector3d{ EigenVectors(0,2), EigenVectors(1,2), EigenVectors(2,2) }, 
-					  false };
 }
 
-FMatrix3d UMeshSpaceDeformerTool::CalculateBestAxis(FDynamicMesh3* Mesh, const TArray<FVector3d>& PositionBuffer, FTransform WorldTransform, FVector3d * ExtentsOut, FVector3d * OriginOut)
+void UMeshSpaceDeformerTool::SetWorld(UWorld* World)
 {
-	PCAVertexPositionBuffer.Reset(3 * Mesh->VertexCount());
-	double* VectorData = PCAVertexPositionBuffer.GetData();
-	FVector3d* InitializePtr = reinterpret_cast<FVector3d*>(VectorData);
-	for (const int& VertexID : Mesh->VertexIndicesItr())
-	{
-		const FVector3d& Point = PositionBuffer[VertexID];
-		*InitializePtr = WorldTransform.TransformPosition(Point);
-		++InitializePtr;
-	}
+	this->TargetWorld = World;
+}
 
-	const int NCols = 3; // R^3;
-	const int NRows = Mesh->VertexCount();
 
-	return SolveSampleCovariance(VectorData, NCols, NRows, *ExtentsOut, *OriginOut);
+void UMeshSpaceDeformerTool::SetAssetAPI(IToolsContextAssetAPI* AssetAPIIn)
+{
+	this->AssetAPI = AssetAPIIn;
 }
 
-/*
- * ToolBuilder
- */
-UMeshSurfacePointTool* UMeshSpaceDeformerToolBuilder::CreateNewTool(const FToolBuilderState& SceneState) const
+
+bool UMeshSpaceDeformerTool::CanAccept() const 
 {
-	UMeshSpaceDeformerTool* NonlinearTool = NewObject<UMeshSpaceDeformerTool>(SceneState.ToolManager);
-	return NonlinearTool;
+	return (Preview == nullptr || Preview->HaveValidResult());
 }
 
-UMeshSpaceDeformerTool::UMeshSpaceDeformerTool()
+void UMeshSpaceDeformerTool::ComputeAABB(const FDynamicMesh3& MeshIn, const FTransform& XFormIn, FVector& BBoxMin, FVector& BBoxMax) const 
 {
+	
+
+	BBoxMin = FVector(FLT_MAX, FLT_MAX, FLT_MAX);
+	BBoxMax = -BBoxMin;
+	
+	for (const auto VertID : MeshIn.VertexIndicesItr())
+	{
+		const FVector3d LocalPosition = MeshIn.GetVertex(VertID);
+		const FVector WorldPosition = XFormIn.TransformPosition(FVector(LocalPosition.X, LocalPosition.Y, LocalPosition.Z));
+		for (int i = 0; i < 3; ++i) BBoxMin[i] = FMath::Min(BBoxMin[i], WorldPosition[i]);
+		for (int i = 0; i < 3; ++i) BBoxMax[i] = FMath::Max(BBoxMax[i], WorldPosition[i]);
+	}
+
 }
 
 void UMeshSpaceDeformerTool::Setup()
 {
 	UMeshSurfacePointTool::Setup();
 
-	//// create dynamic mesh component to use for live preview
-	DynamicMeshComponent = NewObject<USimpleDynamicMeshComponent>(ComponentTarget->GetOwnerActor(), "Dynamic Mesh");
-	DynamicMeshComponent->SetupAttachment(ComponentTarget->GetOwnerActor()->GetRootComponent());
-	DynamicMeshComponent->RegisterComponent();
-	DynamicMeshComponent->SetWorldTransform(ComponentTarget->GetWorldTransform());
-
 
-	//// copy material if there is one
-	auto Material = ComponentTarget->GetMaterial(0);
-	if (Material != nullptr)
+	// populate the OriginalDynamicMesh with a conversion of the input mesh.
 	{
-		DynamicMeshComponent->SetMaterial(0, Material);
+		OriginalDynamicMesh = MakeShared< FDynamicMesh3 >();
+		FMeshDescriptionToDynamicMesh Converter;
+		Converter.bPrintDebugMessages = false;
+		Converter.Convert(ComponentTarget->GetMesh(), *OriginalDynamicMesh);
 	}
-	//// dynamic mesh configuration settings
-	ComponentTarget->SetOwnerVisibility(false);
-	DynamicMeshComponent->TangentsType = EDynamicMeshTangentCalcType::AutoCalculated;
-	DynamicMeshComponent->InitializeMesh(ComponentTarget->GetMesh());
-	OnDynamicMeshComponentChangedHandle = DynamicMeshComponent->OnMeshChanged.Add(
-		FSimpleMulticastDelegate::FDelegate::CreateUObject(this, &UMeshSpaceDeformerTool::OnDynamicMeshComponentChanged));
 
-	FDynamicMesh3* TargetMesh = DynamicMeshComponent->GetMesh();
+	// hide input StaticMeshComponent
+	ComponentTarget->SetOwnerVisibility(false);
 
-	OriginalPositions.SetNumUninitialized(TargetMesh->MaxVertexID());
-	for (int32 VertexID : TargetMesh->VertexIndicesItr())
-	{
-		OriginalPositions[VertexID] = TargetMesh->GetVertex(VertexID);
-	}
 
-	bInDrag = false;
 
 	//// add properties
-	ToolPropertyObjects.Add(this);
 
-	//// set up visualizers
-	AxisRenderer.LineColor = FLinearColor::Red;
-	AxisRenderer.LineThickness = 2.0;
-	AxisRenderer.bDepthTested = false;
+	AddToolPropertySource(this);
 
-	for (int8 i = 0; i < 3; ++i)
-	{
-		FMeshSpaceDeformerOp* Op = Operators[i];
-		Op->UpdateMesh(DynamicMeshComponent->GetMesh());
-	}
 
-	AutoDetectAxes();
+	// compute the bounding box for the input mesh
+	// AABB for the source geometry
+	FVector AABBMin;
+	FVector AABBMax;
+	ComputeAABB(*OriginalDynamicMesh, ComponentTarget->GetWorldTransform(), AABBMin, AABBMax);
+	FVector Extents = AABBMax - AABBMin;
 
-}
 
-void UMeshSpaceDeformerTool::Shutdown(EToolShutdownType ShutdownType)
-{
+	AABBHalfExtents = 0.5 * FVector3d(Extents.X, Extents.Y, Extents.Z);
+	GizmoCenter = 0.5 * (AABBMin + AABBMax);
+
+
+	// Create a new TransformGizmo and associated TransformProxy. The TransformProxy will not be the
+	// parent of any Components in this case, we just use it's transform and change delegate.
+	TransformProxy = NewObject<UTransformProxy>(this);
+	TransformProxy->SetTransform(FTransform(GizmoOrientation, GizmoCenter));
+	TransformGizmo = GetToolManager()->GetPairedGizmoManager()->Create3AxisTransformGizmo(this);
+
+	TransformGizmo->SetActiveTarget(TransformProxy, GetToolManager());
+	// listen for changes to the proxy and update the preview when that happens
+	TransformProxy->OnTransformChanged.AddUObject(this, &UMeshSpaceDeformerTool::TransformProxyChanged);
+
+	// wire these to the input 
+
+	// The initial range for the intervals.
+	const float InitialVerticalIntervalSize = AABBHalfExtents.Z;
+
+	// create sources for the interval parameters
+
+	UpIntervalSource      = NewObject< UGizmoLocalFloatParameterSource >(this);
+	DownIntervalSource    = NewObject< UGizmoLocalFloatParameterSource >(this);
+	ForwardIntervalSource = NewObject< UGizmoLocalFloatParameterSource >(this);
+
+	// Initial Lengths for the interval handles
+
+	UpIntervalSource->Value      = InitialVerticalIntervalSize;
+	DownIntervalSource->Value    = -InitialVerticalIntervalSize;
+	ForwardIntervalSource->Value = 20.f;
+
+	// Sync the properties panel to the interval handles.
+
+	this->UpperBoundsInterval = UpIntervalSource->Value;
+	this->LowerBoundsInterval = DownIntervalSource->Value;
+	this->ModifierPercent     = ForwardIntervalSource->Value;
+
+	// Wire up callbacks to update result mesh and the properties panel when these parameters are changed (by gizmo manipulation in viewport).  Note this is just a one-way
+	// coupling (Sources to Properties). The OnPropertyModified() method provides the Properties to Souces coupling 
 
-	if (DynamicMeshComponent != nullptr)
+	UpIntervalSource->OnParameterChanged.AddLambda([this](IGizmoFloatParameterSource* ParamSource, FGizmoFloatParameterChange Change)->void
 	{
-		DynamicMeshComponent->OnMeshChanged.Remove(OnDynamicMeshComponentChangedHandle);
+		this->UpperBoundsInterval = Change.CurrentValue;
 
-		ComponentTarget->SetOwnerVisibility(true);
+		if (this->Preview != nullptr)
+		{
+			this->Preview->InvalidateResult();
+		}
+	});
 
-		if (ShutdownType == EToolShutdownType::Accept)
+	DownIntervalSource->OnParameterChanged.AddLambda([this](IGizmoFloatParameterSource* ParamSource, FGizmoFloatParameterChange Change)->void
+	{
+		this->LowerBoundsInterval = Change.CurrentValue;
+
+		if (this->Preview != nullptr)
 		{
-			// this block bakes the modified DynamicMeshComponent back into the StaticMeshComponent inside an undo transaction
-			GetToolManager()->BeginUndoTransaction(LOCTEXT("MeshSpaceDeformerToolTransactionName", "Deform Mesh"));
-			ComponentTarget->CommitMesh([&](FMeshDescription* MeshDescription)
-			{
-				DynamicMeshComponent->Bake(MeshDescription, false);
-			});
-			GetToolManager()->EndUndoTransaction();
+			this->Preview->InvalidateResult();
 		}
-			DynamicMeshComponent->UnregisterComponent();
-			DynamicMeshComponent->DestroyComponent();
-			DynamicMeshComponent = nullptr;
-	}
-}
+	});
 
+	ForwardIntervalSource->OnParameterChanged.AddLambda([this](IGizmoFloatParameterSource* ParamSource, FGizmoFloatParameterChange Change)->void
+	{
+		this->ModifierPercent = Change.CurrentValue;
 
-void UMeshSpaceDeformerTool::RegisterActions(FInteractiveToolActionSet& ActionSet)
-{
-}
+		if (this->Preview != nullptr)
+		{
+			this->Preview->InvalidateResult();
+		}
+	});
 
-void UMeshSpaceDeformerTool::OnDynamicMeshComponentChanged()
-{
-}
 
-bool UMeshSpaceDeformerTool::HitTest(const FRay & Ray, FHitResult & OutHit)
-{
-	return Handle.SelectedVertexID != -1;
-}
 
-void UMeshSpaceDeformerTool::OnBeginDrag(const FRay& WorldRay)
-{
-	if (Handle.SelectedVertexID != -1)
-	{
-		bInDrag = true;
-		Handle.DragUpdateBounds(WorldRay);
-		UpdateIntervalsFromDrag();
-		UpdateOp();
-	}
-}
+	// add the interval gizmo
+	IntervalGizmo = GetToolManager()->GetPairedGizmoManager()->CreateGizmo<UIntervalGizmo>(UIntervalGizmo::GizmoName, TEXT("MeshSpaceDefomerInterval"), this);
 
+	// wire in the transform and the interval sources.
+	IntervalGizmo->SetActiveTarget(TransformProxy, UpIntervalSource, DownIntervalSource, ForwardIntervalSource, GetToolManager());
 
-void UMeshSpaceDeformerTool::OnUpdateDrag(const FRay& Ray)
-{
-	if (bInDrag)
+	// use the statetarget to track details changes
+	StateTarget = IntervalGizmo->StateTarget;
+	// Set up the preview object
 	{
-		Handle.DragUpdateBounds(Ray);
-		UpdateIntervalsFromDrag();
-		UpdateOp();
+		// create the operator factory
+		USpaceDeformerOperatorFactory* DeformerOperatorFactory = NewObject<USpaceDeformerOperatorFactory>(this);
+		DeformerOperatorFactory->SpaceDeformerTool = this; // set the back pointer
+
+
+		Preview = NewObject<UMeshOpPreviewWithBackgroundCompute>(DeformerOperatorFactory, "Preview");
+		Preview->Setup(this->TargetWorld, DeformerOperatorFactory);
+
+		// Give the preview something to display
+		Preview->PreviewMesh->UpdatePreview(OriginalDynamicMesh.Get());
+		Preview->PreviewMesh->SetTransform(ComponentTarget->GetWorldTransform());
+
+		Preview->ConfigureMaterials(
+			ToolSetupUtil::GetDefaultMaterial(GetToolManager(), ComponentTarget->GetMaterial(0)),
+			ToolSetupUtil::GetDefaultWorkingMaterial(GetToolManager())
+		);
+		// show the preview mesh
+		Preview->SetVisibility(true);
+
+		// start the compute
+		Preview->InvalidateResult();
 	}
+
 }
 
-void UMeshSpaceDeformerTool::OnEndDrag(const FRay& Ray)
+void UMeshSpaceDeformerTool::Shutdown(EToolShutdownType ShutdownType)
 {
-	bInDrag = false;
+	
+	// Restore (unhide) the source meshes
+	ComponentTarget->SetOwnerVisibility(true);
 
-}
+	if (Preview != nullptr)
+	{
+		FDynamicMeshOpResult Result = Preview->Shutdown();
+		if (ShutdownType == EToolShutdownType::Accept)
+		{
 
-inline void UMeshSpaceDeformerTool::AutoDetectAxes()
-{
-	FTransform WorldTransform = ComponentTarget->GetWorldTransform();
-	FTransform WorldToObject = WorldTransform.Inverse();
+			GetToolManager()->BeginUndoTransaction(LOCTEXT("MeshSpaceDeformer", "Space Deformer"));
 
-	//Retrieves the best axis in world space via principal component analysis.
-	FMatrix3d Basis = CalculateBestAxis(DynamicMeshComponent->GetMesh(), OriginalPositions, WorldTransform, &PrincipalAxesHalfExtentCoeff, &AxisCentroidWorldSpace);
+			FDynamicMesh3* DynamicMeshResult = Result.Mesh.Get();
+			check(DynamicMeshResult != nullptr);
 
-	PrincipalAxesWorldSpace[0] = FVector3d{ Basis(0,0),Basis(1,0),Basis(2,0) };
-	PrincipalAxesWorldSpace[1] = FVector3d{ Basis(0,1),Basis(1,1),Basis(2,1) };
-	PrincipalAxesWorldSpace[2] = FVector3d{ Basis(0,2),Basis(1,2),Basis(2,2) };
+			ComponentTarget->CommitMesh([DynamicMeshResult](FMeshDescription* MeshDescription)
+			{
+				FDynamicMeshToMeshDescription Converter;
+				Converter.Convert(DynamicMeshResult, *MeshDescription);
+			});
 
-	PrincipalAxesObjectSpace[0] = WorldToObject.TransformVectorNoScale(PrincipalAxesWorldSpace[0]);
-	PrincipalAxesObjectSpace[1] = WorldToObject.TransformVectorNoScale(PrincipalAxesWorldSpace[1]);
-	PrincipalAxesObjectSpace[2] = WorldToObject.TransformVectorNoScale(PrincipalAxesWorldSpace[2]);
 
-	UpdateObjectSpaceAxisCentroid();
-}
+			GetToolManager()->EndUndoTransaction();
+		}
+	}
 
-void UMeshSpaceDeformerTool::SwapSecondaryAxis()
-{
-	std::swap(SecondaryAxis,ThirdAxis);
-}
+	UInteractiveGizmoManager* GizmoManager = GetToolManager()->GetPairedGizmoManager();
+	GizmoManager->DestroyAllGizmosByOwner(this);
 
-inline void UMeshSpaceDeformerTool::UpdateObjectSpaceAxisCentroid()
-{
-	AxisCentroidObjectSpace = ComponentTarget->GetWorldTransform().InverseTransformPositionNoScale(AxisCentroidWorldSpace);
 }
 
-void UMeshSpaceDeformerTool::UpdateIntervalsFromDrag()
-{
-	FHandleWidgetVertex* SelectedVertex = Handle.GetSelectedVertex();
-	double CurrentSegmentLength = (SelectedVertex->WorldCenter - AxisCentroidWorldSpace).Length();
-	double MaxLength = PrincipalAxesHalfExtentCoeff[SelectedAxis];
-	double NewInterval = FMath::Clamp(CurrentSegmentLength / MaxLength, 0.0, 1.5);
-	if (Handle.SelectedVertexID == 0)
-	{
-		LowerBoundsInterval = -NewInterval;
-	}
-	else
-	{
-		UpperBoundsInterval = NewInterval;
-	}
-}
 
-bool UMeshSpaceDeformerTool::OnUpdateHover(const FInputDeviceRay& DevicePos)
+
+void  UMeshSpaceDeformerTool::TransformProxyChanged(UTransformProxy* Proxy, FTransform Transform)
 {
-	if (!bInDrag)
+	GizmoOrientation = Transform.GetRotation();
+	GizmoCenter = Transform.GetLocation();
+	const FVector3d Center = GizmoCenter;
+	const FQuaterniond Orientation = GizmoOrientation;
+	GizmoFrame = FFrame3d(Center, Orientation);
+
+	if (Preview != nullptr)
 	{
-		const auto ToleranceTest = [this](const FVector3d& Position1, const FVector3d& Position2) {
-			return ToolSceneQueriesUtil::PointSnapQuery(this->CameraState, Position1, Position2, VisualAngleSnapThreshold);
-		};
-		Handle.UpdateHover(DevicePos, ComponentTarget->GetWorldTransform(), TFunction<bool(const FVector3d&, const FVector3d&)>{ToleranceTest});
+		Preview->InvalidateResult();
 	}
-	return true;
 }
 
+
+
 #if WITH_EDITOR
 void UMeshSpaceDeformerTool::PostEditChangeProperty(FPropertyChangedEvent & PropertyChangedEvent)
 {
-	//UpdateOp();
+	if (PropertyChangedEvent.ChangeType & EPropertyChangeType::Interactive)
+	{
+		if (!bHasBegin)
+		{
+			StateTarget->BeginUpdate();
+			bHasBegin = true;
+		}
+
+		// Propagates change to the ParameterSource by triggering UMeshSpaceDeformerTool::OnPropertyModified
+		OnPropertyModified(this, PropertyChangedEvent.Property);
+	}
+
+	if (PropertyChangedEvent.ChangeType & EPropertyChangeType::ValueSet)
+	{
+		if (!bHasBegin)
+		{
+			StateTarget->BeginUpdate();
+		}
+
+		// Propagates change to the ParameterSource by triggering UMeshSpaceDeformerTool::OnPropertyModified
+		OnPropertyModified(this, PropertyChangedEvent.Property);
+
+		StateTarget->EndUpdate();
+		// reset the bHasBegin
+		bHasBegin = false;
+	}
+		
 }
+
+
 #endif
 
-void UMeshSpaceDeformerTool::UpdateOp()
+void  UMeshSpaceDeformerTool::OnPropertyModified(UObject* PropertySet, UProperty* Property)
 {
-	FMeshSpaceDeformerOp* Op = Operators[static_cast<int8>(SelectedOperationType)];
+	// Update the property sources to reflect any changes.
+	// NB: due to callback, the SetParameter will trigger a recompute of the result via  Preview->InvalidateResult();
 
-	//Provide the basis so that the selected axis is in the Y direction (this will NEED to be rotated based on the orientation the operator expects)
-	// i.e. the Bend operator expects the bend along the Y axis, the twist expects along the Z axis. Just swap the axes as needed inside the operator
-	FMatrix3d Basis
+	if (PropertySet == this)
 	{
-		PrincipalAxesObjectSpace[SecondaryAxis],
-		PrincipalAxesObjectSpace[SelectedAxis],
-		PrincipalAxesObjectSpace[ThirdAxis],
-		false
-	};
-
-	//This allows the object to move, or the handle position to move and have the changes reflected based on their relative world-space positions
-	UpdateObjectSpaceAxisCentroid();
+		if (Property && (Property->GetFName() == GET_MEMBER_NAME_CHECKED(UMeshSpaceDeformerTool, UpperBoundsInterval)))
+		{
+			UpIntervalSource->BeginModify();
+			UpIntervalSource->SetParameter(this->UpperBoundsInterval);
+			UpIntervalSource->EndModify();
+		}
+		if (Property && (Property->GetFName() == GET_MEMBER_NAME_CHECKED(UMeshSpaceDeformerTool, LowerBoundsInterval)))
+		{
+			DownIntervalSource->BeginModify();
+			DownIntervalSource->SetParameter(this->LowerBoundsInterval);
+			DownIntervalSource->EndModify();
+		}
+		if (Property && (Property->GetFName() == GET_MEMBER_NAME_CHECKED(UMeshSpaceDeformerTool, ModifierPercent)))
+		{
+			ForwardIntervalSource->BeginModify();
+			ForwardIntervalSource->SetParameter(this->ModifierPercent);
+			ForwardIntervalSource->EndModify();
+		}
 
-	//Pass the relevant updated data to the operator then run the operation.
-	Op->UpdateAxisData(Basis, AxisCentroidObjectSpace, PrincipalAxesHalfExtentCoeff, LowerBoundsInterval, UpperBoundsInterval, ModifierPercent);
-	Op->CalculateResult(nullptr);
 
-	DynamicMeshComponent->FastNotifyPositionsUpdated();
+		if (Property && (Property->GetFName() == GET_MEMBER_NAME_CHECKED(UMeshSpaceDeformerTool, SelectedOperationType)))
+		{
+			if (Preview != nullptr)
+			{
+				Preview->InvalidateResult();
+			}
+		}
+	}
 }
 
-void UMeshSpaceDeformerTool::Tick(float DeltaTime)
+void UMeshSpaceDeformerTool::UpdateOpParameters(FMeshSpaceDeformerOp& MeshSpaceDeformerOp) const
 {
-	UMeshSurfacePointTool::Tick(DeltaTime);
-}
+	
+	
 
+	
+	FMatrix WorldToGizmo(EForceInit::ForceInitToZero);
+	{
 
-void UMeshSpaceDeformerTool::Render(IToolsContextRenderAPI* RenderAPI)
-{
-	GetToolManager()->GetContextQueriesAPI()->GetCurrentViewState(CameraState);
+		// Matrix that rotates to the gizmo frame 
+		const FMatrix3d ToGizmo(GizmoFrame.GetAxis(0),
+		                    	GizmoFrame.GetAxis(1),
+			                    GizmoFrame.GetAxis(2),
+			                    true /* row constructor*/);
 
-	FDynamicMesh3* TargetMesh = DynamicMeshComponent->GetMesh();
+		// Gizmo To World is rotation followed by translation of the form
+		// world_vec = FromGizmo * local_vec + GizmoCenter
+		// so the inverse will be
+		// local_vec = FromGizmo^(-1) * world_vec - FromGizmo^(-1) * GizmoCenter
+		// 
 
-	AxisRenderer.BeginFrame(RenderAPI, CameraState);
+		// Copy the rotation
+		for (int i = 0; i < 3; ++i)
+		{
+			for (int j = 0; j < 3; ++j)
+			{
+				WorldToGizmo.M[i][j] = ToGizmo(i, j);
+			}
+		}
+		// Add the translation 
+		FVector3d Translation = ToGizmo * GizmoCenter;
+		WorldToGizmo.M[0][3] = -Translation[0];
+		WorldToGizmo.M[1][3] = -Translation[1];
+		WorldToGizmo.M[2][3] = -Translation[2];
+		WorldToGizmo.M[3][3] = 1.f;
+	}
+	// Get the transform as a matrix.  Note may need to transpose this due
+	// to the vec * Matrix  vs Matrix * vec usage in engine.
 
-	//Update the operator (pass it relevant data)
-	UpdateOp();
+	FMatrix ObjectToWorld = ComponentTarget->GetWorldTransform().ToMatrixWithScale().GetTransposed();
 
-	//Update the handle visualization and render it.
-	Handle.SetBasis(PrincipalAxesObjectSpace[0], PrincipalAxesObjectSpace[1], PrincipalAxesObjectSpace[2]);
-	Handle.UpdateDisplayData(SelectedAxis, LowerBoundsInterval, UpperBoundsInterval, PrincipalAxesHalfExtentCoeff);
-	Handle.SetPosition(AxisCentroidWorldSpace);
-	Handle.Render();
+	FMatrix ObjectToGizmo(EForceInit::ForceInitToZero);
+	{
+		// ObjectToGizmo = WorldToGizmo * ObjectToWorld
+		for (int i = 0; i < 4; ++i)
+		{
+			for (int j = 0; j < 4; ++j)
+			{
+				for (int k = 0; k < 4; ++k)
+				{
+					ObjectToGizmo.M[i][j] += WorldToGizmo.M[i][k] * ObjectToWorld.M[k][j];
+				}
+			}
+		}
+	}
 
+	//This allows the object to move, or the handle position to move and have the changes reflected based on their relative world-space positions
+	const FVector3d AxisCentroidObjectSpace = ComponentTarget->GetWorldTransform().InverseTransformPositionNoScale(GizmoCenter);
 
-	AxisRenderer.EndFrame();
-	GetToolManager()->PostInvalidation();
-}
+	//Pass the relevant updated data to the operator then run the operation.
+	/**
+	const double LowerBoundsInterval = this->LowerBoundsInterval;
+	const double UpperBoundsInterval = this->UpperBoundsInterval;
+	const double ModifierPercent     = this->ModifierPercent;
+	*/
 
+	// Set half axis length to be 1/2 the major axis of the bbox.
+	double LengthScale = FMath::Max( AABBHalfExtents.MaxAbs(), 1.e-3);
+	MeshSpaceDeformerOp.AxesHalfLength = LengthScale;
 
+	// set the bound range
+	MeshSpaceDeformerOp.UpperBoundsInterval =  UpperBoundsInterval / LengthScale;
+	MeshSpaceDeformerOp.LowerBoundsInterval = -LowerBoundsInterval / LengthScale;
 
-void FHandleWidget::UpdateHover(const FInputDeviceRay & DevicePos , const FTransform& WorldTransform, TFunction<bool(const FVector3d&, const FVector3d&)> ToleranceFunction)
-{
-	FVector3d LocalPosition, LocalNormal;
-	FGeometrySet3::FNearest Nearest;
+	// percent to apply.
+	MeshSpaceDeformerOp.ModifierPercent = ModifierPercent;
 
-	//Reset all handle axes and vertex states
-	for (unsigned k = 0; k < 3; ++k)
-	{
-		Axes[k].State = IHandleWidgetSelectable::Default;
-		Axes[k].Vertices[0].State = IHandleWidgetSelectable::Default;
-		Axes[k].Vertices[1].State = IHandleWidgetSelectable::Default;
-	}
 
-	//If we're hovering over a vertex, change its state
-	bool bFound = GeometrySet.FindNearestPointToRay(DevicePos.WorldRay, Nearest, ToleranceFunction);
-	if (bFound)
-	{ // Found a hit point
-		int j = Nearest.ID % 2;
-		int i = (Nearest.ID) / 2;
-		SelectedVertexID = j;
-		Axes[i].Vertices[j].State = IHandleWidgetSelectable::Hovering;
-		Axes[i].Vertices[(j + 1) % 2].State = IHandleWidgetSelectable::Default;
-	}else
-	{
-		SelectedVertexID = -1;
-	}
-	
-	//TODO: Do the same for the axes. The geometry set should already contain the axes, just need to decide how the UI will work.
+	// the transform to gizmo space
+	MeshSpaceDeformerOp.ObjectToGizmo = ObjectToGizmo;
 
+	//MeshSpaceDeformerOp.AxisOriginObjectSpace = AxisCentroidObjectSpace; // gizmo origin in object space.
 
+	// copy the mesh
+	MeshSpaceDeformerOp.CopySource(*OriginalDynamicMesh, ComponentTarget->GetWorldTransform());
 }
 
-void FHandleWidget::DragUpdateBounds(const FRay3d& WorldRay)
+void UMeshSpaceDeformerTool::Tick(float DeltaTime)
 {
-	FVector3d ClosestPtOnSegment;
-	FHandleWidgetAxis* SelectedAxis = GetSelectedAxis();
-
-	FVector3d V0 = SelectedAxis->WorldCenter + SelectedAxis->Axis * 99999.0;
-	FVector3d V1 = SelectedAxis->WorldCenter - SelectedAxis->Axis * 99999.0;
-
-	double SqDist = RayToSegmentSquareDist(WorldRay, V0, V1, &ClosestPtOnSegment, nullptr);
+	UMeshSurfacePointTool::Tick(DeltaTime);
 
-	FHandleWidgetVertex* SelectedVertex = GetSelectedVertex();
-	SelectedVertex->State = IHandleWidgetSelectable::Selected;
-	SelectedVertex->SetPosition(ClosestPtOnSegment);
+	if (Preview != nullptr)
+	{
+		Preview->Tick(DeltaTime);
+	}
 }
 
 #undef LOCTEXT_NAMESPACE
diff --git a/Engine/Plugins/Experimental/MeshModelingToolset/Source/MeshModelingTools/Public/MeshSpaceDeformerTool.h b/Engine/Plugins/Experimental/MeshModelingToolset/Source/MeshModelingTools/Public/MeshSpaceDeformerTool.h
--- a/Engine/Plugins/Experimental/MeshModelingToolset/Source/MeshModelingTools/Public/MeshSpaceDeformerTool.h
+++ b/Engine/Plugins/Experimental/MeshModelingToolset/Source/MeshModelingTools/Public/MeshSpaceDeformerTool.h
@@ -22,11 +22,20 @@
 #include "SimpleDynamicMeshComponent.h"
 
 #include "ModelingOperators/Public/ModelingTaskTypes.h"
+#include "BaseGizmos/GizmoInterfaces.h"
 
 
 #include "MeshSpaceDeformerTool.generated.h"
 
 class FMeshVertexChangeBuilder;
+class UPreviewMesh;
+class UTransformGizmo;
+class UTransformProxy;
+class UIntervalGizmo;
+class UMeshOpPreviewWithBackgroundCompute;
+class UGizmoLocalFloatParameterSource;
+class UGizmoTransformChangeStateTarget;
+
 
 /**
  * ToolBuilder
@@ -41,6 +50,8 @@ public:
 	{
 	}
 
+	virtual bool CanBuildTool(const FToolBuilderState& SceneState) const override;
+
 	virtual UMeshSurfacePointTool* CreateNewTool(const FToolBuilderState& SceneState) const override;
 };
 
@@ -56,272 +67,19 @@ enum class  ENonlinearOperationType : int8
 	//Squish	UMETA(DisplayName = "Squish")
 };
 
-/**
-* \brief IHandleWidgetSelectable is a pure virtual base class for elements of the HandleWidget which may be selectable (i.e. the end points of line segments or the axes themselves)
-*/
-class IHandleWidgetSelectable
-{
-public:
-	IHandleWidgetSelectable() = delete;
-	IHandleWidgetSelectable(const FLinearColor& DefaultColor, const FLinearColor& HoveringColor, const FLinearColor& SelectedColor) : StateColors{ DefaultColor,HoveringColor,SelectedColor } 
-	{
-	};
-	virtual ~IHandleWidgetSelectable() 
-	{
-	};
-
-	/**
-	* \brief Sets the renderer to be used by the tool
-	*/
-	virtual inline void SetRenderer(FToolDataVisualizer* Target) 
-	{ 
-		Renderer = Target;
-	};
-
-	/**
-	* \brief Sets the world position of the tool
-	*/ 
-	virtual inline void SetPosition(const FVector3d& Pos) 
-	{
-		WorldCenter = Pos;
-	};
-
-	/**
-	* \brief Gets the color depnding on the state of the selectable element
-	*/
-	virtual inline const FLinearColor& GetColor() const 
-	{
-		return StateColors[static_cast<int8>(State)];
-	}
-
-	/**
-	* \brief Sets the visibility of the specified widget element
-	*/
-	virtual inline void SetVisible(bool Enable) 
-	{
-		Visible = Enable;
-	};
-
-	FToolDataVisualizer* Renderer{nullptr};
-
-	enum EDisplayState : int8
-	{
-		Default,  // Unselected, the default display color
-		Hovering, // The element is being hovered over
-		Selected  // The element is selected
-	};
-
-	EDisplayState State = EDisplayState::Default;
-	FLinearColor StateColors[3];
-
-	FVector3d WorldCenter;
-	bool Visible{false};
-	virtual void Render() = 0;
-};
-
-/**
-* \brief This object is located on the ends of the line segments which extend from the center of the widget. 
-*/
-class FHandleWidgetVertex : public IHandleWidgetSelectable
-{
-public:
-	FHandleWidgetVertex() : IHandleWidgetSelectable
-	{
-		FLinearColor{1.f,1.f,1.f,1.f},	// White
-		FLinearColor{0.f,1.f,1.f,1.f},	// Yellow
-		FLinearColor{1.f,0.0f,1.f,1.f}	// Pink
-	} {};
-	virtual void Render() override
-	{
-		if (Visible)
-		{
-			float Width = State != Default ? 4.0 : 2.0;
-			Renderer->DrawViewFacingCircle(WorldCenter,4.0, 16, GetColor(), Width,false);
-		}
-	};
-};
 
 
-/**
-* \brief Represents an axis of the handle widget, each axis is selectable and has two vertices located at the end points
-*/
-class FHandleWidgetAxis : public IHandleWidgetSelectable
+UCLASS()
+class MESHMODELINGTOOLS_API USpaceDeformerOperatorFactory : public UObject, public IDynamicMeshOperatorFactory
 {
-public:
-	FHandleWidgetAxis() : IHandleWidgetSelectable
-	{
-		FLinearColor{0.2f,0.2f,1.f,1.f}, // Light-blue
-		FLinearColor{0.f,1.f,1.f,1.f},	 // Yellow
-		FLinearColor{1.f,0.f,1.f,1.f}	 // Pink
-	} {};
-
-	double LowerExtentLength;
-	double UpperExtentLength;
-	
-	//Identity
-	FMatrix3d Basis{ 
-		FVector3d{1.0,0.0,0.0}, 
-		FVector3d{0.0,1.0,0.0}, 
-		FVector3d{0.0,0.0,1.0}, 
-		false };
-
-	//The direction vector representing the local space direction of the line segment this handle axis represents
-	FVector3d Axis;
-
-	void SetAxis(const FVector3d& TargetAxis) {Axis = TargetAxis;};
-
-	void SetBasis(const FVector3d& V0, const FVector3d& V1, const FVector3d& V2 ) { Basis = FMatrix3d{ V0, V1, V2, false }; };
-
-	FHandleWidgetVertex Vertices[2];
-	double SelectedWidth = 2.5;
-	double DefaultWidth = 1.5;
-
-	/**
-	* \brief Sets the center position from the center of the handle. This, in turn, sets the position of its associated vertices
-	*/
-	virtual inline void SetPosition(const FVector3d& Center) override
-	{
-		WorldCenter = Center;
-
-		const FVector3d Positions[2]
-		{
-			WorldCenter + Axis * LowerExtentLength,
-			WorldCenter + Axis * UpperExtentLength
-		};
-
-		for (uint8 i = 0; i < 2; ++i)
-		{
-			Vertices[i].SetPosition(Positions[i]);
-		}
-
-	};
-
-	virtual void Render() override
-	{
-		double Width = State == EDisplayState::Default ? DefaultWidth : SelectedWidth;
-
-		for (uint8 i = 0; i < 2; ++i)
-		{
-			Vertices[i].Render();
-		}
-
-		if (Visible)
-		{
-			Renderer->DrawLine(Vertices[0].WorldCenter, Vertices[1].WorldCenter, GetColor(), Width, false);
-		}
-	};
-
-};
+	GENERATED_BODY()
 
-/**
-* \brief This purpose of this "widget" is to visualize the world space center of the transformation being applied by a Mesh Space Deformer operator. 
-  \par This is meant to be an analogous visualizer to the "handle" widgets used by popular modeling tools i.e. maya
-*/
-class FHandleWidget
-{
 public:
-	FHandleWidget(FToolDataVisualizer* Renderer) 
-	{
-		for (unsigned i=0;i<3;++i)
-		{
-			Axes[i].SetRenderer(Renderer);
-
-
-			for (unsigned j=0;j<2;++j)
-			{
-				Axes[i].Vertices[j].SetRenderer(Renderer);
-			}
-			const FVector3d& A = Axes[i].Vertices[0].WorldCenter;
-			const FVector3d& B = Axes[i].Vertices[1].WorldCenter;
-			GeometrySet.AddPoint(i * 2 , A);
-			GeometrySet.AddPoint(i * 2 + 1, B);
-			GeometrySet.AddCurve(i, FPolyline3d{ TArray<FVector3d>{A,B} });
-
-		}
-
-
-	};
-	virtual ~FHandleWidget() {};
-
-	/**
-	* \brief Returns a pointer to the Handle Widget Axis related object as specified by the Selection
-	*/
-	FHandleWidgetAxis* GetSelectedAxis() { return &Axes[SelectedAxisID];};
-
-	/**
-	* \brief Returns a pointer to the Handle Widget Axis related object as specified by the Selection
-	*/
-	FHandleWidgetVertex* GetSelectedVertex() { return SelectedVertexID == -1 ? nullptr : &Axes[SelectedAxisID].Vertices[SelectedVertexID]; };
-	int SelectedAxisID;
-	int SelectedVertexID;
-
-	FHandleWidgetAxis Axes[3];
-	FVector3d WorldCenter;
-
-	/** Given a ray, tests the geometry collection and updates the corresponding selectable handle elements */
-	void UpdateHover(const FInputDeviceRay& DevicePos, const FTransform& WorldTransform, TFunction<bool(const FVector3d&, const FVector3d&)> ToleranceFunction);
-	void SetPosition(const FVector3d& Center)
-	{
-		WorldCenter = Center;
-
-
-		for (unsigned i = 0; i < 3; ++i)
-		{
-			Axes[i].SetPosition(Center);
-			const FVector3d& A = Axes[i].Vertices[0].WorldCenter;
-			const FVector3d& B = Axes[i].Vertices[1].WorldCenter;
-			GeometrySet.UpdatePoint(i * 2 , A);
-			GeometrySet.UpdatePoint(i * 2 + 1, B);
-			GeometrySet.UpdateCurve(i, FPolyline3d{ TArray<FVector3d>{A,B} });
-		}
-	};
-
-	void SetBasis(const FVector3d& Axis0, const FVector3d& Axis1, const FVector3d& Axis2 )
-	{
-		Axes[0].SetAxis(Axis0);
-		Axes[1].SetAxis(Axis1);
-		Axes[2].SetAxis(Axis2);
-	};
-	void UpdateDisplayData(int SelectedAxis, double LowerBoundsInterval, double UpperBoundsInterval, const FVector3d& HalfExtents)
-	{
-
-		SelectedAxisID = SelectedAxis;
-		for (unsigned i = 0; i < 3; ++i)
-		{
-			if ( i == SelectedAxis )
-			{
-				Axes[i].State = IHandleWidgetSelectable::EDisplayState::Selected;
-				Axes[i].SetVisible(true);
-				Axes[i].Vertices[0].SetVisible(true);
-				Axes[i].Vertices[1].SetVisible(true);
-				Axes[i].LowerExtentLength = HalfExtents[i] * LowerBoundsInterval;
-				Axes[i].UpperExtentLength = HalfExtents[i] * UpperBoundsInterval;
-			}else
-			{
-				Axes[i].SetVisible(true);
-				Axes[i].Vertices[0].SetVisible(false);
-				Axes[i].Vertices[1].SetVisible(false);
-				Axes[i].LowerExtentLength = -HalfExtents[i] ;
-				Axes[i].UpperExtentLength = HalfExtents[i] ;
-			}
-
-		}
-	}
-
-	/** Called when updating a click-and-drag transaction, updates the positions of the selected vertices based on the new mouse input position*/
-	void DragUpdateBounds(const FRay3d& WorldRay);
-
-	void Render() 
-	{
-		for (unsigned i=0;i<3;++i)
-		{
-			Axes[i].Render();
-		}
-	};
-
-	/** The set of geometry containing the vertices which are draggable and the axes objects which are selectable */
-	FGeometrySet3 GeometrySet;
+	// IDynamicMeshOperatorFactory API
+	virtual TUniquePtr<FDynamicMeshOperator> MakeNewOperator() override;
 
+	UPROPERTY()
+	UMeshSpaceDeformerTool* SpaceDeformerTool;  // back pointer
 };
 
 
@@ -336,129 +94,116 @@ class MESHMODELINGTOOLS_API UMeshSpaceDeformerTool : public UMeshSurfacePointToo
 public:
 	UMeshSpaceDeformerTool();
 
-	virtual void RegisterActions(FInteractiveToolActionSet& ActionSet) override;
-
-	FMatrix3d CalculateBestAxis(FDynamicMesh3 * Mesh, const TArray<FVector3d>& PositionBuffer, FTransform WorldTransform, FVector3d * ExtentsOut, FVector3d * OriginOut );
 
 	virtual void Setup() override;
 	virtual void Shutdown(EToolShutdownType ShutdownType) override;
 
+	virtual void SetWorld(UWorld* World);
+	virtual void SetAssetAPI(IToolsContextAssetAPI* AssetAPI);
+
+
 	virtual void Tick(float DeltaTime) override;
-	virtual void Render(IToolsContextRenderAPI* RenderAPI) override;
+	virtual void Render(IToolsContextRenderAPI* RenderAPI) override {};
 
 	virtual bool HasCancel() const override { return true; }
 	virtual bool HasAccept() const override { return true; }
-	virtual bool CanAccept() const override { return true; }
+	virtual bool CanAccept() const override;
 
-	virtual bool HitTest(const FRay& Ray, FHitResult& OutHit) override;
-	
-	//
-	// Input behavior
-	//
-	virtual void OnBeginDrag(const FRay& Ray) override;
-	virtual void OnUpdateDrag(const FRay& Ray) override;
-	virtual void OnEndDrag(const FRay& Ray) override;
-
-	virtual bool OnUpdateHover(const FInputDeviceRay& DevicePos) override;
 
 #if WITH_EDITOR
 	virtual void PostEditChangeProperty(FPropertyChangedEvent& PropertyChangedEvent) override;
 #endif
 
-	void UpdateOp();
-public:
+	virtual void OnPropertyModified(UObject* PropertySet, UProperty* Property) override;
 
-	float VisualAngleSnapThreshold = 1.5f;
+	// sync the parameters owned by the MeshSpaceDeformerOp 
+	void UpdateOpParameters(FMeshSpaceDeformerOp& MeshSpaceDeformerOp) const;
 
-protected:
+public:
 
-	//Options
 	UPROPERTY(EditAnywhere, Category = Options, meta = (DisplayName = "Operation Type"))
-	ENonlinearOperationType SelectedOperationType;
-
-	/** The upper bounds interval corresponds to the region of space which the selected operator will affect. A setting of -1.0 should envelope all points in the "lower" half of the mesh given the axis has been auto-detected. The corresponding upper value of 1 will cover the entire mesh. */
-	UPROPERTY(EditAnywhere, Category = Options, meta = (DisplayName = "Lower Bound", UIMin = "-1.5", UIMax = "0.0", ClampMin = "-1.5", ClampMax = "0.0"))
-	double LowerBoundsInterval{ -1.0 };
+	ENonlinearOperationType SelectedOperationType = ENonlinearOperationType::Bend;
 
 	/** The upper bounds interval corresponds to the region of space which the selected operator will affect. A setting of 1.0 should envelope all points in the "upper" half of the mesh given the axis has been auto-detected. The corresponding lower value of -1 will cover the entire mesh. */
-	UPROPERTY(EditAnywhere, Category = Options, meta = (DisplayName = "Upper Bound", UIMin = "0.0", UIMax = "1.5", ClampMin = "0.0", ClampMax = "1.5"))
-	double UpperBoundsInterval{ 1.0 };
+	UPROPERTY(EditAnywhere, Category = Options, meta = (DisplayName = "Upper Bound", UIMin = "0.0", ClampMin = "0.0"))
+	float UpperBoundsInterval = 10.0;
+
+	/** The upper bounds interval corresponds to the region of space which the selected operator will affect. A setting of -1.0 should envelope all points in the "lower" half of the mesh given the axis has been auto-detected. The corresponding upper value of 1 will cover the entire mesh. */
+	UPROPERTY(EditAnywhere, Category = Options, meta = (DisplayName = "Lower Bound", UIMax = "0", ClampMax = "0"))
+	float LowerBoundsInterval = -10.0;
 
 	/** As each operator has a range of values (i.e. curvature, angle of twist, scale), this represents the percentage passed to the operator as a parameter. In the future, for more control, this should be separated into individual settings for each operator for more precise control */
-	UPROPERTY(EditAnywhere, Category = Options, meta = (DisplayName = "Modifier Percent", UIMin = "-100.0", UIMax = "100.0", ClampMin = "-100.0", ClampMax = "100.0"))
-	double ModifierPercent{ 0.0 };
+	UPROPERTY(EditAnywhere, Category = Options, meta = (DisplayName = "Modifier Percent"))
+	float ModifierPercent = 20.0;
+
+	//UPROPERTY()
+	//UMeshSpaceDeformerToolStandardProperties* SpaceDeformerProperties;
+
+protected:
+
+	UPROPERTY()
+	UGizmoTransformChangeStateTarget* StateTarget = nullptr;
+	
+	// used to coordinate undo for the detail panel. 
+
+	bool bHasBegin = false;
+protected:
 
-	/** Automatically detects the principal axes of the point cloud of the mesh, as well as determining the lengths of the maximal elements in those directions, setting the intervals accordingly. */
-	UFUNCTION(CallInEditor, Category = Options, meta = (DisplayName = "Auto-detect Axis", ToolTip = "Attempts to automatically align the handle with principal axis."))
-	void AutoDetectAxes();
+	UPROPERTY()
+	UMeshOpPreviewWithBackgroundCompute* Preview = nullptr;   
 
-	/** Swaps the secondary axis upon which the operator will work. I.e. bending */
-	UFUNCTION(CallInEditor, Category = Options, meta = (DisplayName = "Swap secondary axis"))
-	void SwapSecondaryAxis();
 
-	//Calculate the updated object-space position which gets passed to the handle widget.
-	void UpdateObjectSpaceAxisCentroid();
+protected:	
 
-	//Updates the upper and lower intervals from the drag update
-	void UpdateIntervalsFromDrag();
+	TSharedPtr<FDynamicMesh3> OriginalDynamicMesh;
 
+	UWorld* TargetWorld;
+	IToolsContextAssetAPI* AssetAPI;
 
+
+	// drawing plane control
+
+	/** offset to center of gizmo*/
 	UPROPERTY()
-	USimpleDynamicMeshComponent* DynamicMeshComponent;
+	FVector GizmoCenter;
 
-	FViewCameraState CameraState;
+	/** Gizmo Plane Orientation */
+	UPROPERTY()
+	FQuat GizmoOrientation;
 
-	//For buffer re-use and avoiding recurrent dynamic allocation
-	TArray<double> PCAVertexPositionBuffer;
 
-	//Array of the original positions of each vertex
-	TArray<FVector3d> OriginalPositions;
+	UPROPERTY()
+	UIntervalGizmo* IntervalGizmo;
 
-	//Index of the axes being used. 
-	uint8 SelectedAxis{ 2 };
-	uint8 SecondaryAxis{ 1 };
-	uint8 ThirdAxis{0};
+	UPROPERTY()
+	UTransformGizmo* TransformGizmo;
 
-	//Provided PCA, these are the direction vectors representing the principal axes to best align the mesh
-	FVector3d PrincipalAxesWorldSpace[3];
-	FVector3d PrincipalAxesObjectSpace[3];
+	UPROPERTY()
+	UTransformProxy* TransformProxy;
 
-	// These represent the distance from the centroid to the farthest elements in the direction (respectively) of the PCA bases. 
-	// Specifically, the X element is the shortest of the PCA axes, the Y element is the middle axis, and the Z element represents the axis with the largest degree of variance, meaning Z is the greatest half extent.
-	FVector3d PrincipalAxesHalfExtentCoeff;
+	/** Interval Parameter sources that reflect UI settings. */
 
-	//Where the handle widget will be rendered
-	FVector3d AxisCentroidWorldSpace;
-	FVector3d AxisCentroidObjectSpace;
+	UPROPERTY()
+	UGizmoLocalFloatParameterSource* UpIntervalSource;
 
-	FToolDataVisualizer AxisRenderer;
-	FHandleWidget Handle{&AxisRenderer};
+	UPROPERTY()
+	UGizmoLocalFloatParameterSource* DownIntervalSource;
 
-	FBendMeshOp BendOp;
-	FTwistMeshOp TwistOp;
-	FFlareMeshOp FlareOp;
+	UPROPERTY()
+	UGizmoLocalFloatParameterSource* ForwardIntervalSource;
 
-	TArray<FMeshSpaceDeformerOp*> Operators
-	{
-		&BendOp,
-		&FlareOp,
-		&TwistOp
-		//Sinusoid
-		//Wave
-		//Squish
-	};
-
-
-	// realtime visualization
-	void OnDynamicMeshComponentChanged();
-	FDelegateHandle OnDynamicMeshComponentChangedHandle;
+	FFrame3d GizmoFrame;
+
+
+
+	FVector3d AABBHalfExtents; // 1/2 the extents of the bbox
+
+	void TransformProxyChanged(UTransformProxy* Proxy, FTransform Transform);
+
+
+	/** Compute the axis aligned abounding box for the source mesh */
+	void ComputeAABB(const FDynamicMesh3& MeshIn, const FTransform& XFormIn, FVector& BBoxMin, FVector& BBoxMax) const;
 	
-	//
-	// Update	
-	//
-	/** True for the duration of UI click+drag */
-	bool bInDrag;
 	
-private:
 };
 
diff --git a/Engine/Plugins/Experimental/MeshModelingToolset/Source/ModelingOperators/Private/MeshSpaceDeformerOps/BendMeshOp.cpp b/Engine/Plugins/Experimental/MeshModelingToolset/Source/ModelingOperators/Private/MeshSpaceDeformerOps/BendMeshOp.cpp
--- a/Engine/Plugins/Experimental/MeshModelingToolset/Source/ModelingOperators/Private/MeshSpaceDeformerOps/BendMeshOp.cpp
+++ b/Engine/Plugins/Experimental/MeshModelingToolset/Source/ModelingOperators/Private/MeshSpaceDeformerOps/BendMeshOp.cpp
@@ -1,59 +1,131 @@
 // Copyright 1998-2019 Epic Games, Inc. All Rights Reserved.
 #include "ModelingOperators\Public\SpaceDeformerOps\BendMeshOp.h"
-#include "ProfilingDebugging/ScopedTimers.h"
 
 
 //Bends along the Y-axis
 void FBendMeshOp::CalculateResult(FProgressCancel* Progress)
 {
 
-	//FString DebugLogString = FString::Printf(TEXT(" Bend timer "));
-	//FScopedDurationTimeLogger Timer(DebugLogString);
+
+	// Matrix from gizmo space (z-up) to a y-up space
+	FMatrix ToYUp(EForceInit::ForceInitToZero);
+	ToYUp.M[0][0] =  1.f;
+	ToYUp.M[1][2] =  1.f;
+	ToYUp.M[2][1] = -1.f;
+	ToYUp.M[3][3] =  1.f; 
+	
+	// Full transform to y-up in gizmo space
+	FMatrix ObjectToYUpGizmo(EForceInit::ForceInitToZero);
+	{
+		for (int i = 0; i < 4; ++i)
+		{
+			for (int j = 0; j < 4; ++j)
+			{
+				for (int k = 0; k < 4; ++k)
+				{
+					ObjectToYUpGizmo.M[i][j] += ToYUp.M[i][k] * ObjectToGizmo.M[k][j];
+				}
+			}
+		}
+	}
+
+
+	float Det = ObjectToYUpGizmo.Determinant();
+
+	// Check if the transform is nearly singular
+	// this could happen if the scale on the object to world transform has a very small component.
+	if (FMath::Abs(Det) < 1.e-4)
+	{
+		return;
+	}
+
+	FMatrix GizmoToObject = ObjectToYUpGizmo.Inverse();
+	
+	
 	const double Curvature = GetModifierValue();
-	const double ToRadians = 0.017453292519943295;
-	const double ThetaRadians = ToRadians * Curvature;
-	const FVector3d O = ObjectSpaceToOpSpace * AxisOriginObjectSpace;
-	const double& Y0 = O.Y;
-	double YMin = Y0 + LowerBoundsInterval * AxesHalfLengths[2];
-	double YMax = Y0 + UpperBoundsInterval * AxesHalfLengths[2];
-	const double K = ThetaRadians / AxesHalfLengths[2];
+	
+	// early out if nothing has been requested.
+	if (FMath::Abs(Curvature) < 0.001)
+	{
+		return;
+	}
+
+
+
+	const double DegreesToRadians = 0.017453292519943295769236907684886; // Pi / 180
+	const double ThetaRadians = DegreesToRadians * Curvature;
+	//const FVector3d OpOrigin = ToOpSpace * AxisOriginObjectSpace;
+	
+	
+	// bounds in Op space
+	const double YMin =  -LowerBoundsInterval * AxesHalfLength;
+	const double YMax =   UpperBoundsInterval * AxesHalfLength;
+
+	const double Y0 = 0.;
+	const double K = ThetaRadians / AxesHalfLength;
 	const double Ik = 1.0 / K;
+
+	
+
 	for (int VertexID : TargetMesh->VertexIndicesItr())
 	{
-		const FVector3d& OriginalPositionObjectSpace = OriginalPositions[VertexID];
-		if (Curvature == 0.0)
+		//const FVector3d& SrcPos = OriginalPositions[VertexID];
+		const FVector3d SrcPos = TargetMesh->GetVertex(VertexID);
+		const double SrcPos4[4] = { SrcPos[0], SrcPos[1], SrcPos[2], 1.0};
+
+		// Position in gizmo space
+		double GizmoPos4[4] = { 0., 0., 0., 0. };
+		for (int i = 0; i < 4; ++i)
 		{
-			TargetMesh->SetVertex(VertexID, OriginalPositionObjectSpace);
-			continue;
+			for (int j = 0; j < 4; ++j)
+			{
+				GizmoPos4[i] += ObjectToYUpGizmo.M[i][j] * SrcPos4[j];
+			}
 		}
-		const FVector3d Pos = ObjectSpaceToOpSpace * (OriginalPositionObjectSpace);
 
-		const double YHat = FMath::Clamp(Pos.Y, YMin, YMax);
+
+		//FVector3d DstPos = ToOpSpace * (SrcPos) - OpOrigin;
+
+		const double YHat = FMath::Clamp(GizmoPos4[1], YMin, YMax);
 
 		const double Theta = K * (YHat - Y0);
 
-		const double S0 = FMath::Sin(Theta);
-		const double C0 = FMath::Cos(Theta);
+		const double S0 = TMathUtil<double>::Sin(Theta);
+		const double C0 = TMathUtil<double>::Cos(Theta);
 
-		const double ZP = Pos.Z - Ik;
+		const double ZP = GizmoPos4[2] - Ik;
 
-		const double X = Pos.X;
-		double Y = -S0 * (ZP) + Y0;
-		double Z =  C0 * (ZP) + Ik;
+		
+		double Y = -S0 * ZP + Y0;
+		double Z =  C0 * ZP + Ik;
 
-		if (Pos.Y > YMax)
+		if (GizmoPos4[1] > YMax)
 		{
-			const double YDiff = Pos.Y - YMax;
+			const double YDiff = GizmoPos4[1] - YMax;
 			Y += C0 * YDiff;
 			Z += S0 * YDiff;
 		}
-		else if (Pos.Y < YMin)
+		else if (GizmoPos4[1] < YMin)
 		{
-			const double YDiff = Pos.Y - YMin;
+			const double YDiff = GizmoPos4[1] - YMin;
 			Y += C0 * YDiff;
 			Z += S0 * YDiff;
 		}
-		TargetMesh->SetVertex(VertexID, OpSpaceToObjectSpace * FVector3d{ X,Y,Z });
+
+		GizmoPos4[1] = Y;
+		GizmoPos4[2] = Z;
+
+		// Position in Obj Space
+		double DstPos4[4] = { 0., 0., 0., 0. };
+		for (int i = 0; i < 4; ++i)
+		{
+			for (int j = 0; j < 4; ++j)
+			{
+				DstPos4[i] += GizmoToObject.M[i][j] * GizmoPos4[j];
+			}
+		}
+
+		TargetMesh->SetVertex(VertexID, FVector3d(DstPos4[0], DstPos4[1], DstPos4[2]));
 	}
 
 }
\ No newline at end of file
diff --git a/Engine/Plugins/Experimental/MeshModelingToolset/Source/ModelingOperators/Private/MeshSpaceDeformerOps/FlareMeshOp.cpp b/Engine/Plugins/Experimental/MeshModelingToolset/Source/ModelingOperators/Private/MeshSpaceDeformerOps/FlareMeshOp.cpp
--- a/Engine/Plugins/Experimental/MeshModelingToolset/Source/ModelingOperators/Private/MeshSpaceDeformerOps/FlareMeshOp.cpp
+++ b/Engine/Plugins/Experimental/MeshModelingToolset/Source/ModelingOperators/Private/MeshSpaceDeformerOps/FlareMeshOp.cpp
@@ -33,29 +33,60 @@ double  inverseSinerp(const T& valueBetween, const T& value1, const T& value2)
 //Flares along the Z axis
 void FFlareMeshOp::CalculateResult(FProgressCancel* Progress)
 {
-	FMatrix3d Rotate{ FVector3d{0,1,0}, FVector3d{0,0,1}, FVector3d{1,0,0}, false };
-	FMatrix3d ToOpSpace = Rotate * ObjectSpaceToOpSpace;
 
-	const FVector3d O = ToOpSpace * AxisOriginObjectSpace;
-	const double& Z0 = O.Z;
-	const double ZMin = Z0 + LowerBoundsInterval * AxesHalfLengths[2];
-	const double ZMax = Z0 + UpperBoundsInterval * AxesHalfLengths[2];
+	float Det = ObjectToGizmo.Determinant();
 
-	FMatrix3d FromOpSpace = ToOpSpace.Inverse();
-	for (int VertexID : TargetMesh->VertexIndicesItr())
+	// Check if the transform is nearly singular
+	// this could happen if the scale on the object to world transform has a very small component.
+	if (FMath::Abs(Det) < 1.e-4)
 	{
-		const FVector3d& OriginalPositionObjectSpace = OriginalPositions[VertexID];
-
-		const FVector3d Pos = ToOpSpace * (OriginalPositionObjectSpace);
+		return;
+	}
 
-		double T = FMath::Clamp(FMath::Abs((Z0 - Pos.Z)) / (ZMax - ZMin) * 0.5, 0.0, 1.0);
-		double R = coserp((1.0-T), 1.0, GetModifierValue());
+	FMatrix GizmoToObject = ObjectToGizmo.Inverse();
 
-		double X = Pos.X * R;
-		double Y = Pos.Y * R;
-		double Z = Pos.Z;
+	const double ZMin = -LowerBoundsInterval * AxesHalfLength;
+	const double ZMax =  UpperBoundsInterval * AxesHalfLength;
 
-		TargetMesh->SetVertex(VertexID, FromOpSpace * (FVector3d{ X,Y,Z }));
+	for (int VertexID : TargetMesh->VertexIndicesItr())
+	{
+		
+		const FVector3d SrcPos = TargetMesh->GetVertex(VertexID);
+
+		const double SrcPos4[4] = { SrcPos[0], SrcPos[1], SrcPos[2], 1.0 };
+
+		// Position in gizmo space
+		double GizmoPos4[4] = { 0., 0., 0., 0. };
+		for (int i = 0; i < 4; ++i)
+		{
+			for (int j = 0; j < 4; ++j)
+			{
+				GizmoPos4[i] += ObjectToGizmo.M[i][j] * SrcPos4[j];
+			}
+		}
+
+		
+		// Parameterize curve between ZMin and ZMax to go between 0, 1
+		double T = FMath::Clamp( (GizmoPos4[2]- ZMin) / (ZMax - ZMin), 0.0, 1.0);
+
+		double R = 1. + FMath::Sin(PI * T) * (ModifierPercent / 100.f);
+
+		// 2d scale x,y values.
+		GizmoPos4[0] *= R;
+		GizmoPos4[1] *= R;
+
+
+		double DstPos4[4] = { 0., 0., 0., 0. };
+		for (int i = 0; i < 4; ++i)
+		{
+			for (int j = 0; j < 4; ++j)
+			{
+				DstPos4[i] += GizmoToObject.M[i][j] * GizmoPos4[j];
+			}
+		}
+
+		// set the position
+		TargetMesh->SetVertex(VertexID, FVector3d(DstPos4[0], DstPos4[1], DstPos4[2]));
 	}
 
 }
\ No newline at end of file
diff --git a/Engine/Plugins/Experimental/MeshModelingToolset/Source/ModelingOperators/Private/MeshSpaceDeformerOps/MeshSpaceDeformerOp.cpp b/Engine/Plugins/Experimental/MeshModelingToolset/Source/ModelingOperators/Private/MeshSpaceDeformerOps/MeshSpaceDeformerOp.cpp
--- a/Engine/Plugins/Experimental/MeshModelingToolset/Source/ModelingOperators/Private/MeshSpaceDeformerOps/MeshSpaceDeformerOp.cpp
+++ b/Engine/Plugins/Experimental/MeshModelingToolset/Source/ModelingOperators/Private/MeshSpaceDeformerOps/MeshSpaceDeformerOp.cpp
@@ -7,24 +7,11 @@ void FMeshSpaceDeformerOp::CalculateResult(FProgressCancel* Progress)
 {
 }
 
-void FMeshSpaceDeformerOp::UpdateMesh(FDynamicMesh3* TargetMeshComponent)
+void FMeshSpaceDeformerOp::CopySource(const FDynamicMesh3& MeshIn, const FTransform& XForm)
 {
-	TargetMesh = TargetMeshComponent;
+	ResultMesh = MakeUnique < FDynamicMesh3>(MeshIn);
+	ResultTransform = FTransform3d(XForm);
 
-	OriginalPositions.SetNumUninitialized(TargetMesh->MaxVertexID());
-	for (int32 VertexID : TargetMesh->VertexIndicesItr())
-	{
-		OriginalPositions[VertexID] = TargetMesh->GetVertex(VertexID);
-	}
-}
+	TargetMesh = ResultMesh.Get();
 
-void FMeshSpaceDeformerOp::UpdateAxisData(const FMatrix3d & ObjectSpaceToOpSpaceTransform, const FVector3d & ObjectSpaceOrigin, const FVector3d & AxesHalfExtents, double LowerBounds, double UpperBounds, double Modifier)
-{
-	ObjectSpaceToOpSpace = ObjectSpaceToOpSpaceTransform;  // Rotation matrix to rotate from object space to the space expected by the operator
-	OpSpaceToObjectSpace = ObjectSpaceToOpSpace.Inverse(); // Pre-computation for use in the operator
-	AxisOriginObjectSpace = ObjectSpaceOrigin;			   // The location of the handle's origin, transformed into the object space of the target mesh
-	AxesHalfLengths = AxesHalfExtents;					   // Each element of this vector is a scalar corresponding to the distance from the centroid to the farthest vertex in that respective direction
- 	ModifierPercent = Modifier;							   // Input from the tool, as each operator will have a min/max value i.e. the angle of curvature, twist or scale. This is the percent used to interpolate between the min/max value
-	LowerBoundsInterval = LowerBounds;					   // The interval [0.0-1.0] representing the range of effect in space this operator will have
-	UpperBoundsInterval = UpperBounds;				
 }
diff --git a/Engine/Plugins/Experimental/MeshModelingToolset/Source/ModelingOperators/Private/MeshSpaceDeformerOps/TwistMeshOp.cpp b/Engine/Plugins/Experimental/MeshModelingToolset/Source/ModelingOperators/Private/MeshSpaceDeformerOps/TwistMeshOp.cpp
--- a/Engine/Plugins/Experimental/MeshModelingToolset/Source/ModelingOperators/Private/MeshSpaceDeformerOps/TwistMeshOp.cpp
+++ b/Engine/Plugins/Experimental/MeshModelingToolset/Source/ModelingOperators/Private/MeshSpaceDeformerOps/TwistMeshOp.cpp
@@ -5,34 +5,66 @@
 //Twists along the Z axis
 void FTwistMeshOp::CalculateResult(FProgressCancel* Progress)
 {
-	FMatrix3d Rotate{FVector3d{0,1,0}, FVector3d{0,0,1}, FVector3d{1,0,0}, false};
-	FMatrix3d ToOpSpace = Rotate * ObjectSpaceToOpSpace;
+	float Det = ObjectToGizmo.Determinant();
+
+	// Check if the transform is nearly singular
+	// this could happen if the scale on the object to world transform has a very small component.
+	if (FMath::Abs(Det) < 1.e-4)
+	{
+		return;
+	}
+
+	FMatrix GizmoToObject = ObjectToGizmo.Inverse();
 
-	const FVector3d O = ToOpSpace * AxisOriginObjectSpace;
-	const double& Z0 = O.Z;
-	const double ZMin = Z0 + LowerBoundsInterval * AxesHalfLengths[2];
-	const double ZMax = Z0 + UpperBoundsInterval * AxesHalfLengths[2];
 
-	const double ToRadians = 0.017453292519943295;
-	const double ThetaRadians = ToRadians * GetModifierValue();
+	const double ZMin = -LowerBoundsInterval * AxesHalfLength;
+	const double ZMax =  UpperBoundsInterval * AxesHalfLength;
+
+
+	const double DegreesToRadians = 0.017453292519943295769236907684886; // Pi / 180
+	const double ThetaRadians = DegreesToRadians * GetModifierValue();
 
-	FMatrix3d FromOpSpace = ToOpSpace.Inverse();
 	for (int VertexID : TargetMesh->VertexIndicesItr())
 	{
-		const FVector3d& OriginalPositionObjectSpace = OriginalPositions[VertexID];
 
-		const FVector3d Pos = ToOpSpace * (OriginalPositionObjectSpace);
+		const FVector3d SrcPos = TargetMesh->GetVertex(VertexID);
+		const double SrcPos4[4] = { SrcPos[0], SrcPos[1], SrcPos[2], 1.0 };
+
+		// Position in gizmo space
+		double GizmoPos4[4] = { 0., 0., 0., 0. };
+		for (int i = 0; i < 4; ++i)
+		{
+			for (int j = 0; j < 4; ++j)
+			{
+				GizmoPos4[i] += ObjectToGizmo.M[i][j] * SrcPos4[j];
+			}
+		}
 
-		double T = FMath::Clamp(( ZMax - Pos.Z ) / ( ZMax - ZMin ), 0.0, 1.0);
+		double T = FMath::Clamp( (GizmoPos4[2] - ZMin) / ( ZMax - ZMin ), 0.0, 1.0) - 0.5;
 		double Theta = ThetaRadians * T;
 
 		const double C0 = FMath::Cos(Theta);
 		const double S0 = FMath::Sin(Theta);
+		
+		// apply 2d rotation.
+		const double X =  C0 * GizmoPos4[0]  + S0 * GizmoPos4[1];
+		const double Y = -S0 * GizmoPos4[0]  + C0 * GizmoPos4[1];
+		
+		GizmoPos4[0] = X;
+		GizmoPos4[1] = Y;
 
-		double X = Pos.X * C0 + Pos.Y * S0;
-		double Y = Pos.X * (-S0) + Pos.Y * C0;
-		double Z = Pos.Z;
+		// Position in Obj Space
+		double DstPos4[4] = { 0., 0., 0., 0. };
+		for (int i = 0; i < 4; ++i)
+		{
+			for (int j = 0; j < 4; ++j)
+			{
+				DstPos4[i] += GizmoToObject.M[i][j] * GizmoPos4[j];
+			}
+		}
 
-		TargetMesh->SetVertex(VertexID, FromOpSpace * FVector3d{ X,Y,Z });
+		// set the position
+		TargetMesh->SetVertex(VertexID, FVector3d(DstPos4[0], DstPos4[1], DstPos4[2]));
+		
 	}
 }
\ No newline at end of file
diff --git a/Engine/Plugins/Experimental/MeshModelingToolset/Source/ModelingOperators/Public/SpaceDeformerOps/BendMeshOp.h b/Engine/Plugins/Experimental/MeshModelingToolset/Source/ModelingOperators/Public/SpaceDeformerOps/BendMeshOp.h
--- a/Engine/Plugins/Experimental/MeshModelingToolset/Source/ModelingOperators/Public/SpaceDeformerOps/BendMeshOp.h
+++ b/Engine/Plugins/Experimental/MeshModelingToolset/Source/ModelingOperators/Public/SpaceDeformerOps/BendMeshOp.h
@@ -10,7 +10,7 @@ class MODELINGOPERATORS_API FBendMeshOp : public FMeshSpaceDeformerOp
 {
 public:
 	virtual ~FBendMeshOp() {}
-	FBendMeshOp() : FMeshSpaceDeformerOp(0.0, 180.0) {};
+	FBendMeshOp() : FMeshSpaceDeformerOp(0.0, 180.1) {};
 	virtual void CalculateResult(FProgressCancel* Progress) override;
 
 protected:
diff --git a/Engine/Plugins/Experimental/MeshModelingToolset/Source/ModelingOperators/Public/SpaceDeformerOps/FlareMeshOp.h b/Engine/Plugins/Experimental/MeshModelingToolset/Source/ModelingOperators/Public/SpaceDeformerOps/FlareMeshOp.h
--- a/Engine/Plugins/Experimental/MeshModelingToolset/Source/ModelingOperators/Public/SpaceDeformerOps/FlareMeshOp.h
+++ b/Engine/Plugins/Experimental/MeshModelingToolset/Source/ModelingOperators/Public/SpaceDeformerOps/FlareMeshOp.h
@@ -11,7 +11,7 @@ class MODELINGOPERATORS_API FFlareMeshOp : public FMeshSpaceDeformerOp
 {
 public:
 	virtual ~FFlareMeshOp() {}
-	FFlareMeshOp() : FMeshSpaceDeformerOp(0.1, 2.0) {};
+	FFlareMeshOp() : FMeshSpaceDeformerOp(-2.0, 2.0) {};
 	virtual void CalculateResult(FProgressCancel* Progress) override;
 
 
diff --git a/Engine/Plugins/Experimental/MeshModelingToolset/Source/ModelingOperators/Public/SpaceDeformerOps/MeshSpaceDeformerOp.h b/Engine/Plugins/Experimental/MeshModelingToolset/Source/ModelingOperators/Public/SpaceDeformerOps/MeshSpaceDeformerOp.h
--- a/Engine/Plugins/Experimental/MeshModelingToolset/Source/ModelingOperators/Public/SpaceDeformerOps/MeshSpaceDeformerOp.h
+++ b/Engine/Plugins/Experimental/MeshModelingToolset/Source/ModelingOperators/Public/SpaceDeformerOps/MeshSpaceDeformerOp.h
@@ -22,34 +22,14 @@ public:
 	/**
 	* \brief Updates the pointer to the target mesh
 	*/
-	virtual void UpdateMesh(FDynamicMesh3* Mesh);
+	virtual void CopySource(const FDynamicMesh3& Mesh, const FTransform& XFrom);
 
 
-	/**
-	* \brief Updates spatial data i.e. origin, orientation and other pertaining inputs needed by the operators
-	* \param ObjectSpaceToOpSpaceTransform Rotation matrix to rotate from object space to the space expected by the operator
-	* \param ObjectSpaceOrigin The location of the handle's origin, transformed from world space into the local coordinate system of the mesh
-	* \param AxesHalfExtents Each element of this vector is a scalar corresponding to the distance from the ObjectSpaceOrigin (centroid) to the farthest vertex in that respective direction
-	* \param LowerBounds The negative interval (only negative numbers) representing the range of effect in space this operator will have where -1 will utilize the entire half extent
-	* \param UpperBounds The positive interval (only positive numbers) representing the range of effect in space this operator will have where 1 will utilize the entire half extent
-	* \param ModifierValue Percent input from the tool (between -100 and 100), as each operator will have a min/max value i.e. the angle of curvature, twist or scale. This is the percent used to interpolate between the min/max value
-	*/
-	virtual void UpdateAxisData(const FMatrix3d& ObjectSpaceToOpSpaceTransform, const FVector3d& ObjectSpaceOrigin, const FVector3d& AxesHalfExtents, double LowerBounds, double UpperBounds, double ModifierValue);
-
-	//Array of the original positions of each vertex of the mesh
-	TArray<FVector3d> OriginalPositions;
-
-	//Origin of the transformation, given in the local coordinate system of the mesh
-	FVector3d AxisOriginObjectSpace;
-	
-	//Each element of this vector is a scalar corresponding to the distance from the ObjectSpaceOrigin (centroid) to the farthest vertex in that respective direction
-	FVector3d AxesHalfLengths;
-
-	//Rotation matrix to rotate from object space to the space expected by the operator
-	FMatrix3d ObjectSpaceToOpSpace;
+	// half the major axis of the bounding box for the source geometry
+	double AxesHalfLength;
 
-	//Precomputed inversion from the provided ObjectSpaceToOpSpace, used to convert back to the object's coordinate system
-	FMatrix3d OpSpaceToObjectSpace;
+	// transform, including translation, to gizmo space
+	FMatrix ObjectToGizmo;
 
 	//Percent input from the tool (between -100 and 100), as each operator will have a min/max value i.e. the angle of curvature, twist or scale. This is the percent used to interpolate between the min/max value
 	double ModifierPercent;
