Mesh Modeling: Mesh Space Deformer tool.  Update to transform the normals while deforming the mesh, uses the inverse jacobian transpose.

diff --git a/Engine/Plugins/Experimental/MeshModelingToolset/Source/ModelingOperators/Private/MeshSpaceDeformerOps/BendMeshOp.cpp b/Engine/Plugins/Experimental/MeshModelingToolset/Source/ModelingOperators/Private/MeshSpaceDeformerOps/BendMeshOp.cpp
--- a/Engine/Plugins/Experimental/MeshModelingToolset/Source/ModelingOperators/Private/MeshSpaceDeformerOps/BendMeshOp.cpp
+++ b/Engine/Plugins/Experimental/MeshModelingToolset/Source/ModelingOperators/Private/MeshSpaceDeformerOps/BendMeshOp.cpp
@@ -1,5 +1,6 @@
 // Copyright 1998-2019 Epic Games, Inc. All Rights Reserved.
 #include "ModelingOperators\Public\SpaceDeformerOps\BendMeshOp.h"
+#include "DynamicMeshAttributeSet.h"
 
 
 //Bends along the Y-axis
@@ -39,7 +40,7 @@ void FBendMeshOp::CalculateResult(FProgressCancel* Progress)
 		return;
 	}
 
-	FMatrix GizmoToObject = ObjectToYUpGizmo.Inverse();
+	FMatrix YUpGizmoToObject = ObjectToYUpGizmo.Inverse();
 	
 	
 	const double Curvature = GetModifierValue();
@@ -65,7 +66,87 @@ void FBendMeshOp::CalculateResult(FProgressCancel* Progress)
 	const double K = ThetaRadians / AxesHalfLength;
 	const double Ik = 1.0 / K;
 
-	
+
+	if (ResultMesh->HasAttributes())
+	{
+		// Fix the normals first if they exist.
+
+		FDynamicMeshNormalOverlay* Normals = ResultMesh->Attributes()->PrimaryNormals();
+		for (int ElID : Normals->ElementIndicesItr())
+		{
+			// get the vertex
+			auto VertexID = Normals->GetParentVertex(ElID);
+			const FVector3d& SrcPos = TargetMesh->GetVertex(VertexID);
+
+			FVector3d SrcNormal = Normals->GetElement(ElID);
+
+
+
+			const double SrcPos4[4] = { SrcPos[0], SrcPos[1], SrcPos[2], 1.0 };
+
+			// Position in gizmo space
+			double GizmoPos4[4] = { 0., 0., 0., 0. };
+			for (int i = 0; i < 4; ++i)
+			{
+				for (int j = 0; j < 4; ++j)
+				{
+					GizmoPos4[i] += ObjectToYUpGizmo.M[i][j] * SrcPos4[j];
+				}
+			}
+
+			FVector3d RotatedNormal(0, 0, 0);
+			{
+				// Rotate normal to gizmo space.
+				for (int i = 0; i < 3; ++i)
+				{
+					for (int j = 0; j < 3; ++j)
+					{
+						RotatedNormal[i] += YUpGizmoToObject.M[j][i] * SrcNormal[j];
+					}
+				}
+			}
+
+			const double YHat = FMath::Clamp(GizmoPos4[1], YMin, YMax);
+
+			const double Theta = K * (YHat - Y0);
+
+			const double S0 = TMathUtil<double>::Sin(Theta);
+			const double C0 = TMathUtil<double>::Cos(Theta);
+
+			const double ZP = GizmoPos4[2] - Ik;
+
+
+
+			// transform normal.  Do this before changing GizmoPos
+
+			{
+
+				double DTheta = K;
+				if (GizmoPos4[1] < YMin || GizmoPos4[1] > YMax) DTheta = 0.;
+
+				double tmp = (1. - DTheta * GizmoPos4[2]);
+				FVector3d DstNormal(0., 0., 0.);
+				DstNormal[0] = tmp * RotatedNormal[0];
+				DstNormal[1] =                          C0 * RotatedNormal[1] - S0 * tmp * RotatedNormal[2];
+				DstNormal[2] =                          S0 * RotatedNormal[1] + C0 * tmp * RotatedNormal[2];
+
+			
+				// rotate back to mesh space.
+				RotatedNormal = FVector3d(0, 0, 0);
+				for (int i = 0; i < 3; ++i)
+				{
+					for (int j = 0; j < 3; ++j)
+					{
+						RotatedNormal[i] += ObjectToYUpGizmo.M[j][i] * DstNormal[j];
+					}
+				}
+
+			}
+
+			FVector3f RotatedNoramlF(RotatedNormal);
+			Normals->SetElement(ElID, RotatedNoramlF);
+		}
+	}
 
 	for (int VertexID : TargetMesh->VertexIndicesItr())
 	{
@@ -121,7 +202,7 @@ void FBendMeshOp::CalculateResult(FProgressCancel* Progress)
 		{
 			for (int j = 0; j < 4; ++j)
 			{
-				DstPos4[i] += GizmoToObject.M[i][j] * GizmoPos4[j];
+				DstPos4[i] += YUpGizmoToObject.M[i][j] * GizmoPos4[j];
 			}
 		}
 
diff --git a/Engine/Plugins/Experimental/MeshModelingToolset/Source/ModelingOperators/Private/MeshSpaceDeformerOps/FlareMeshOp.cpp b/Engine/Plugins/Experimental/MeshModelingToolset/Source/ModelingOperators/Private/MeshSpaceDeformerOps/FlareMeshOp.cpp
--- a/Engine/Plugins/Experimental/MeshModelingToolset/Source/ModelingOperators/Private/MeshSpaceDeformerOps/FlareMeshOp.cpp
+++ b/Engine/Plugins/Experimental/MeshModelingToolset/Source/ModelingOperators/Private/MeshSpaceDeformerOps/FlareMeshOp.cpp
@@ -1,5 +1,6 @@
 // Copyright 1998-2019 Epic Games, Inc. All Rights Reserved.
 #include "ModelingOperators\Public\SpaceDeformerOps\FlareMeshOp.h"
+#include "DynamicMeshAttributeSet.h"
 
 
 
@@ -48,6 +49,82 @@ void FFlareMeshOp::CalculateResult(FProgressCancel* Progress)
 	const double ZMin = -LowerBoundsInterval * AxesHalfLength;
 	const double ZMax =  UpperBoundsInterval * AxesHalfLength;
 
+
+
+	if (ResultMesh->HasAttributes())
+	{
+		// Fix the normals first if they exist.
+
+		FDynamicMeshNormalOverlay* Normals = ResultMesh->Attributes()->PrimaryNormals();
+		for (int ElID : Normals->ElementIndicesItr())
+		{
+			// get the vertex
+			auto VertexID = Normals->GetParentVertex(ElID);
+			const FVector3d& SrcPos = TargetMesh->GetVertex(VertexID);
+
+			FVector3d SrcNormal = Normals->GetElement(ElID);
+
+
+
+			const double SrcPos4[4] = { SrcPos[0], SrcPos[1], SrcPos[2], 1.0 };
+
+			// Position in gizmo space
+			double GizmoPos4[4] = { 0., 0., 0., 0. };
+			for (int i = 0; i < 4; ++i)
+			{
+				for (int j = 0; j < 4; ++j)
+				{
+					GizmoPos4[i] += ObjectToGizmo.M[i][j] * SrcPos4[j];
+				}
+			}
+
+			FVector3d RotatedNormal(0, 0, 0);
+			{
+				// Rotate normal to gizmo space.
+				for (int i = 0; i < 3; ++i)
+				{
+					for (int j = 0; j < 3; ++j)
+					{
+						RotatedNormal[i] += GizmoToObject.M[j][i] * SrcNormal[j];
+					}
+				}
+			}
+
+
+			const double T = FMath::Clamp((GizmoPos4[2] - ZMin) / (ZMax - ZMin), 0.0, 1.0);
+
+			const double R = 1. + FMath::Sin(PI * T) * (ModifierPercent / 100.f);
+			
+			// transform normal.  Do this before changing GizmoPos
+
+			{
+
+				double DR = FMath::Cos(PI * T) * (PI / (ZMax - ZMin)) *  (ModifierPercent / 100.f);
+				if (GizmoPos4[2] > ZMax || GizmoPos4[2] < ZMin)  DR = 0.f;
+			
+
+				FVector3d DstNormal(0., 0., 0.);
+				DstNormal[0] = R * RotatedNormal[0];
+				DstNormal[1] = R * RotatedNormal[1];
+				DstNormal[2] = -R * DR * (RotatedNormal[0] * GizmoPos4[0] + RotatedNormal[1] * GizmoPos4[1]) + R * R * RotatedNormal[2];
+
+				// rotate back to mesh space.
+				RotatedNormal = FVector3d(0, 0, 0);
+				for (int i = 0; i < 3; ++i)
+				{
+					for (int j = 0; j < 3; ++j)
+					{
+						RotatedNormal[i] += ObjectToGizmo.M[j][i] * DstNormal[j];
+					}
+				}
+
+			}
+
+			FVector3f RotatedNoramlF(RotatedNormal);
+			Normals->SetElement(ElID, RotatedNoramlF);
+		}
+	}
+
 	for (int VertexID : TargetMesh->VertexIndicesItr())
 	{
 		
diff --git a/Engine/Plugins/Experimental/MeshModelingToolset/Source/ModelingOperators/Private/MeshSpaceDeformerOps/TwistMeshOp.cpp b/Engine/Plugins/Experimental/MeshModelingToolset/Source/ModelingOperators/Private/MeshSpaceDeformerOps/TwistMeshOp.cpp
--- a/Engine/Plugins/Experimental/MeshModelingToolset/Source/ModelingOperators/Private/MeshSpaceDeformerOps/TwistMeshOp.cpp
+++ b/Engine/Plugins/Experimental/MeshModelingToolset/Source/ModelingOperators/Private/MeshSpaceDeformerOps/TwistMeshOp.cpp
@@ -1,6 +1,7 @@
 // Copyright 1998-2019 Epic Games, Inc. All Rights Reserved.
 
 #include "ModelingOperators\Public\SpaceDeformerOps\TwistMeshOp.h"
+#include "DynamicMeshAttributeSet.h"
 
 //Twists along the Z axis
 void FTwistMeshOp::CalculateResult(FProgressCancel* Progress)
@@ -14,8 +15,13 @@ void FTwistMeshOp::CalculateResult(FProgressCancel* Progress)
 		return;
 	}
 
+	// note the transpose of GizmoToObject should be used to transform normals into Gizmo Space.
+	// the transpose of OjbectToGizmo should be used to transform normals back to Object Space.  
+
 	FMatrix GizmoToObject = ObjectToGizmo.Inverse();
 
+	
+
 
 	const double ZMin = -LowerBoundsInterval * AxesHalfLength;
 	const double ZMax =  UpperBoundsInterval * AxesHalfLength;
@@ -24,6 +30,86 @@ void FTwistMeshOp::CalculateResult(FProgressCancel* Progress)
 	const double DegreesToRadians = 0.017453292519943295769236907684886; // Pi / 180
 	const double ThetaRadians = DegreesToRadians * GetModifierValue();
 
+
+	if (ResultMesh->HasAttributes())
+	{
+		// Fix the normals first if they exist.
+
+		FDynamicMeshNormalOverlay* Normals = ResultMesh->Attributes()->PrimaryNormals();
+		for (int ElID : Normals->ElementIndicesItr())
+		{
+			// get the vertex
+			auto VertexID = Normals->GetParentVertex(ElID);
+			const FVector3d& SrcPos = TargetMesh->GetVertex(VertexID);
+
+			FVector3d SrcNormal = Normals->GetElement(ElID);
+
+
+
+			const double SrcPos4[4] = { SrcPos[0], SrcPos[1], SrcPos[2], 1.0 };
+
+			// Position in gizmo space
+			double GizmoPos4[4] = { 0., 0., 0., 0. };
+			for (int i = 0; i < 4; ++i)
+			{
+				for (int j = 0; j < 4; ++j)
+				{
+					GizmoPos4[i] += ObjectToGizmo.M[i][j] * SrcPos4[j];
+				}
+			}
+
+			FVector3d RotatedNormal(0, 0, 0);
+			{
+				// Rotate normal to gizmo space.
+				for (int i = 0; i < 3; ++i)
+				{
+					for (int j = 0; j < 3; ++j)
+					{
+						RotatedNormal[i] += GizmoToObject.M[j][i] * SrcNormal[j];
+					}
+				}
+			}
+
+
+			double T = FMath::Clamp((GizmoPos4[2] - ZMin) / (ZMax - ZMin), 0.0, 1.0) - 0.5;
+			double Theta = ThetaRadians * T;
+
+
+
+			const double C0 = FMath::Cos(Theta);
+			const double S0 = FMath::Sin(Theta);
+
+			// transform normal.  Do this before changing GizmoPos
+
+			{
+
+				// rotate with twist deformation
+				double DthetaDZ = ThetaRadians / (ZMax - ZMin);
+				if (GizmoPos4[2] > ZMax || GizmoPos4[2] < ZMin) DthetaDZ = 0.;
+
+				FVector3d DstNormal(0., 0., 0.);
+				DstNormal[0] = C0 * RotatedNormal[0] - S0 * RotatedNormal[1];
+				DstNormal[1] = S0 * RotatedNormal[0] + C0 * RotatedNormal[1];
+				DstNormal[2] = DthetaDZ * (GizmoPos4[1] * RotatedNormal[0] - GizmoPos4[0] * RotatedNormal[1]) + RotatedNormal[2];
+
+				// rotate back to mesh space.
+				RotatedNormal = FVector3d(0, 0, 0);
+				for (int i = 0; i < 3; ++i)
+				{
+					for (int j = 0; j < 3; ++j)
+					{
+						RotatedNormal[i] += ObjectToGizmo.M[j][i] * DstNormal[j];
+					}
+				}
+				
+			}
+
+			FVector3f RotatedNoramlF(RotatedNormal);
+			Normals->SetElement(ElID, RotatedNoramlF);
+		}
+	}
+
+	// now fix the vertex positions
 	for (int VertexID : TargetMesh->VertexIndicesItr())
 	{
 
@@ -43,16 +129,21 @@ void FTwistMeshOp::CalculateResult(FProgressCancel* Progress)
 		double T = FMath::Clamp( (GizmoPos4[2] - ZMin) / ( ZMax - ZMin ), 0.0, 1.0) - 0.5;
 		double Theta = ThetaRadians * T;
 
+		
+
 		const double C0 = FMath::Cos(Theta);
 		const double S0 = FMath::Sin(Theta);
-		
+	
+
 		// apply 2d rotation.
-		const double X =  C0 * GizmoPos4[0]  + S0 * GizmoPos4[1];
-		const double Y = -S0 * GizmoPos4[0]  + C0 * GizmoPos4[1];
+		const double X =  C0 * GizmoPos4[0]  - S0 * GizmoPos4[1];
+		const double Y =  S0 * GizmoPos4[0]  + C0 * GizmoPos4[1];
 		
 		GizmoPos4[0] = X;
 		GizmoPos4[1] = Y;
 
+		
+
 		// Position in Obj Space
 		double DstPos4[4] = { 0., 0., 0., 0. };
 		for (int i = 0; i < 4; ++i)
