Fixing compile error with the Mesh Space Deformer Tool.

diff --git a/Engine/Plugins/Experimental/MeshModelingToolset/Source/ModelingOperators/Private/MeshSpaceDeformerOps/BendMeshOp.cpp b/Engine/Plugins/Experimental/MeshModelingToolset/Source/ModelingOperators/Private/MeshSpaceDeformerOps/BendMeshOp.cpp
--- a/Engine/Plugins/Experimental/MeshModelingToolset/Source/ModelingOperators/Private/MeshSpaceDeformerOps/BendMeshOp.cpp
+++ b/Engine/Plugins/Experimental/MeshModelingToolset/Source/ModelingOperators/Private/MeshSpaceDeformerOps/BendMeshOp.cpp
@@ -78,8 +78,9 @@ void FBendMeshOp::CalculateResult(FProgressCancel* Progress)
 			auto VertexID = Normals->GetParentVertex(ElID);
 			const FVector3d& SrcPos = TargetMesh->GetVertex(VertexID);
 
-			FVector3d SrcNormal = Normals->GetElement(ElID);
-
+			FVector3f SrcNormalF = Normals->GetElement(ElID);
+			FVector3d SrcNormal;
+			SrcNormal[0] = SrcNormalF[0]; SrcNormal[1] = SrcNormalF[1]; SrcNormal[2] = SrcNormalF[2];
 
 
 			const double SrcPos4[4] = { SrcPos[0], SrcPos[1], SrcPos[2], 1.0 };
@@ -143,7 +144,8 @@ void FBendMeshOp::CalculateResult(FProgressCancel* Progress)
 
 			}
 
-			FVector3f RotatedNoramlF(RotatedNormal);
+			FVector3f RotatedNoramlF;
+			RotatedNoramlF[0] = RotatedNormal[0]; RotatedNoramlF[1] = RotatedNormal[1]; RotatedNoramlF[2] = RotatedNormal[2];
 			Normals->SetElement(ElID, RotatedNoramlF);
 		}
 	}
diff --git a/Engine/Plugins/Experimental/MeshModelingToolset/Source/ModelingOperators/Private/MeshSpaceDeformerOps/FlareMeshOp.cpp b/Engine/Plugins/Experimental/MeshModelingToolset/Source/ModelingOperators/Private/MeshSpaceDeformerOps/FlareMeshOp.cpp
--- a/Engine/Plugins/Experimental/MeshModelingToolset/Source/ModelingOperators/Private/MeshSpaceDeformerOps/FlareMeshOp.cpp
+++ b/Engine/Plugins/Experimental/MeshModelingToolset/Source/ModelingOperators/Private/MeshSpaceDeformerOps/FlareMeshOp.cpp
@@ -62,8 +62,9 @@ void FFlareMeshOp::CalculateResult(FProgressCancel* Progress)
 			auto VertexID = Normals->GetParentVertex(ElID);
 			const FVector3d& SrcPos = TargetMesh->GetVertex(VertexID);
 
-			FVector3d SrcNormal = Normals->GetElement(ElID);
-
+			FVector3f SrcNormalF = Normals->GetElement(ElID);
+			FVector3d SrcNormal; 
+			SrcNormal[0] = SrcNormalF[0]; SrcNormal[1] = SrcNormalF[1]; SrcNormal[2] = SrcNormalF[2];
 
 
 			const double SrcPos4[4] = { SrcPos[0], SrcPos[1], SrcPos[2], 1.0 };
@@ -120,7 +121,8 @@ void FFlareMeshOp::CalculateResult(FProgressCancel* Progress)
 
 			}
 
-			FVector3f RotatedNoramlF(RotatedNormal);
+			FVector3f RotatedNoramlF;
+			RotatedNoramlF[0] = RotatedNormal[0]; RotatedNoramlF[1] = RotatedNormal[1]; RotatedNoramlF[2] = RotatedNormal[2];
 			Normals->SetElement(ElID, RotatedNoramlF);
 		}
 	}
diff --git a/Engine/Plugins/Experimental/MeshModelingToolset/Source/ModelingOperators/Private/MeshSpaceDeformerOps/TwistMeshOp.cpp b/Engine/Plugins/Experimental/MeshModelingToolset/Source/ModelingOperators/Private/MeshSpaceDeformerOps/TwistMeshOp.cpp
--- a/Engine/Plugins/Experimental/MeshModelingToolset/Source/ModelingOperators/Private/MeshSpaceDeformerOps/TwistMeshOp.cpp
+++ b/Engine/Plugins/Experimental/MeshModelingToolset/Source/ModelingOperators/Private/MeshSpaceDeformerOps/TwistMeshOp.cpp
@@ -42,7 +42,9 @@ void FTwistMeshOp::CalculateResult(FProgressCancel* Progress)
 			auto VertexID = Normals->GetParentVertex(ElID);
 			const FVector3d& SrcPos = TargetMesh->GetVertex(VertexID);
 
-			FVector3d SrcNormal = Normals->GetElement(ElID);
+			FVector3f SrcNormalF = Normals->GetElement(ElID);
+			FVector3d SrcNormal;
+			SrcNormal[0] = SrcNormalF[0]; SrcNormal[1] = SrcNormalF[1]; SrcNormal[2] = SrcNormalF[2];
 
 
 
@@ -104,7 +106,8 @@ void FTwistMeshOp::CalculateResult(FProgressCancel* Progress)
 				
 			}
 
-			FVector3f RotatedNoramlF(RotatedNormal);
+			FVector3f RotatedNoramlF;
+			RotatedNoramlF[0] = RotatedNormal[0]; RotatedNoramlF[1] = RotatedNormal[1]; RotatedNoramlF[2] = RotatedNormal[2];
 			Normals->SetElement(ElID, RotatedNoramlF);
 		}
 	}
