ModelingTools: add Create Group option to Selection Tool. Add ability to show PolyGroup Colors in Selection Tool. Custom TriangleColor TFunction of DMesh-backed Components now takes Mesh as parameter (to be able to look up group w/o complex captures)
#rnx

diff --git a/Engine/Plugins/Experimental/MeshModelingToolset/Source/MeshModelingTools/Private/MeshSelectionTool.cpp b/Engine/Plugins/Experimental/MeshModelingToolset/Source/MeshModelingTools/Private/MeshSelectionTool.cpp
--- a/Engine/Plugins/Experimental/MeshModelingToolset/Source/MeshModelingTools/Private/MeshSelectionTool.cpp
+++ b/Engine/Plugins/Experimental/MeshModelingToolset/Source/MeshModelingTools/Private/MeshSelectionTool.cpp
@@ -8,6 +8,8 @@
 #include "DynamicMeshChangeTracker.h"
 #include "Changes/ToolCommandChangeSequence.h"
 #include "Changes/MeshChange.h"
+#include "Util/ColorConstants.h"
+#include "Selections/MeshConnectedComponents.h"
 #include "MeshIndexUtil.h"
 #include "AssetGenerationUtil.h"
 #include "ToolSetupUtil.h"
@@ -46,6 +48,7 @@ void UMeshSelectionToolProperties::SaveProperties(UInteractiveTool* SaveFromTool
 	PropertyCache->bVolumetricBrush = this->bVolumetricBrush;
 	PropertyCache->bHitBackFaces = this->bHitBackFaces;
 	PropertyCache->bShowWireframe = this->bShowWireframe;
+	PropertyCache->bShowPolygroups = this->bShowPolygroups;
 }
 
 void UMeshSelectionToolProperties::RestoreProperties(UInteractiveTool* RestoreToTool)
@@ -56,6 +59,7 @@ void UMeshSelectionToolProperties::RestoreProperties(UInteractiveTool* RestoreTo
 	this->bVolumetricBrush = PropertyCache->bVolumetricBrush;
 	this->bHitBackFaces = PropertyCache->bHitBackFaces;
 	this->bShowWireframe = PropertyCache->bShowWireframe;
+	this->bShowPolygroups = PropertyCache->bShowPolygroups;
 }
 
 
@@ -123,6 +127,10 @@ void UMeshSelectionTool::Setup()
 	ShowWireframeWatcher.Initialize(
 		[this]() { return SelectionProps->bShowWireframe; },
 		[this](bool bNewValue) { PreviewMesh->EnableWireframe(bNewValue); }, false);
+
+	ShowGroupsWatcher.Initialize(
+		[this]() { return SelectionProps->bShowPolygroups; },
+		[this](bool bNewValue) { UpdateVisualization(); }, false);
 }
 
 
@@ -282,37 +290,6 @@ void UMeshSelectionTool::CalculateVertexROI(const FBrushStampData& Stamp, TArray
 
 
 
-static void GrowToConnectedTriangles(const TArray<int>& TriangleROI, TArray<int>& GrowROI, const FDynamicMesh3* Mesh,
-	TArray<int32>& QueueBuffer, TSet<int32>& DoneBuffer,
-	TFunctionRef<bool(int32, int32)> CanGrowPredicate = [](int32, int32) { return true; }
-	)
-{
-	QueueBuffer.Reset(); QueueBuffer.Insert(TriangleROI, 0);
-	DoneBuffer.Reset(); DoneBuffer.Append(TriangleROI);
-	//TArray<int32> Queue(TriangleROI);
-	//TSet<int32> Done(TriangleROI);
-
-	while (QueueBuffer.Num() > 0)
-	{
-		int32 CurTri = QueueBuffer.Pop(false);
-		GrowROI.Add(CurTri);
-
-		FIndex3i NbrTris = Mesh->GetTriNeighbourTris(CurTri);
-		for (int j = 0; j < 3; ++j)
-		{
-			int32 tid = NbrTris[j];
-			if (tid != FDynamicMesh3::InvalidID && DoneBuffer.Contains(tid) == false && CanGrowPredicate(CurTri, tid))
-			{
-				QueueBuffer.Add(tid);
-				DoneBuffer.Add(tid);
-			}
-		}
-	}
-
-}
-
-
-
 void UMeshSelectionTool::CalculateTriangleROI(const FBrushStampData& Stamp, TArray<int>& TriangleROI)
 {
 	FTransform Transform = ComponentTarget->GetWorldTransform();
@@ -346,7 +323,7 @@ void UMeshSelectionTool::CalculateTriangleROI(const FBrushStampData& Stamp, TArr
 	{
 		TArray<int32> StartROI;
 		StartROI.Add(Stamp.HitResult.FaceIndex);
-		GrowToConnectedTriangles(StartROI, TriangleROI, Mesh, TemporaryBuffer, TemporarySet,
+		FMeshConnectedComponents::GrowToConnectedTriangles(Mesh, StartROI, TriangleROI,  &TemporaryBuffer, &TemporarySet,
 			[Mesh, RadiusSqr, StampPosLocal](int t1, int t2) { return (Mesh->GetTriCentroid(t2) - StampPosLocal).SquaredLength() < RadiusSqr; });
 
 	}
@@ -416,12 +393,12 @@ void UMeshSelectionTool::UpdateFaceSelection(const FBrushStampData& Stamp, const
 	TArray<int> LocalROI;
 	if (SelectionProps->SelectionMode == EMeshSelectionToolPrimaryMode::AllConnected)
 	{
-		GrowToConnectedTriangles(TriangleROI, LocalROI, Mesh, TemporaryBuffer, TemporarySet);
+		FMeshConnectedComponents::GrowToConnectedTriangles(Mesh, TriangleROI, LocalROI, &TemporaryBuffer, &TemporarySet);
 		UseROI = &LocalROI;
 	}
 	else if (SelectionProps->SelectionMode == EMeshSelectionToolPrimaryMode::AllInGroup)
 	{
-		GrowToConnectedTriangles(TriangleROI, LocalROI, Mesh, TemporaryBuffer, TemporarySet,
+		FMeshConnectedComponents::GrowToConnectedTriangles(Mesh, TriangleROI, LocalROI, &TemporaryBuffer, &TemporarySet,
 			[Mesh](int t1, int t2) { return Mesh->GetTriangleGroup(t1) == Mesh->GetTriangleGroup(t2); } );
 		UseROI = &LocalROI;
 	}
@@ -431,7 +408,7 @@ void UMeshSelectionTool::UpdateFaceSelection(const FBrushStampData& Stamp, const
 		StartROI.Add(Stamp.HitResult.FaceIndex);
 		FVector3d StartNormal = Mesh->GetTriNormal(StartROI[0]);
 		int AngleTol = SelectionProps->AngleTolerance;
-		GrowToConnectedTriangles(StartROI, LocalROI, Mesh, TemporaryBuffer, TemporarySet,
+		FMeshConnectedComponents::GrowToConnectedTriangles(Mesh, StartROI, LocalROI, &TemporaryBuffer, &TemporarySet,
 			[Mesh, AngleTol, StartNormal](int t1, int t2) { return Mesh->GetTriNormal(t2).AngleD(StartNormal) < AngleTol; });
 		UseROI = &LocalROI;
 	}
@@ -442,7 +419,7 @@ void UMeshSelectionTool::UpdateFaceSelection(const FBrushStampData& Stamp, const
 		StartROI.Add(Stamp.HitResult.FaceIndex);
 		FVector3d StartNormal = Mesh->GetTriNormal(StartROI[0]);
 		int AngleTol = SelectionProps->AngleTolerance;
-		GrowToConnectedTriangles(StartROI, LocalROI, Mesh, TemporaryBuffer, TemporarySet,
+		FMeshConnectedComponents::GrowToConnectedTriangles(Mesh, StartROI, LocalROI, &TemporaryBuffer, &TemporarySet,
 			[Mesh, AngleTol, StartNormal, &BrushROI](int t1, int t2) { return BrushROI.Contains(t2) && Mesh->GetTriNormal(t2).AngleD(StartNormal) < AngleTol; });
 		UseROI = &LocalROI;
 	}
@@ -532,22 +509,22 @@ void UMeshSelectionTool::UpdateVisualization()
 {
 	if (SelectionType == EMeshSelectionElementType::Face)
 	{
-		PreviewMesh->SetTriangleColorFunction([this](int TriangleID)
+		if (SelectionProps->bShowPolygroups)
+		{
+			PreviewMesh->SetTriangleColorFunction([this](const FDynamicMesh3* Mesh, int TriangleID)
+			{
+				return SelectedTriangles[TriangleID] ? LinearColors::VideoRed3b() : LinearColors::SelectFColor(Mesh->GetTriangleGroup(TriangleID));
+
+			}, UPreviewMesh::ERenderUpdateMode::FullUpdate);
+		}
+		else
 		{
-			return SelectedTriangles[TriangleID] ? FColor::Red : FColor::White;
-		}, UPreviewMesh::ERenderUpdateMode::FullUpdate);
-
-		//bool bIsFirst = (! BaseMeshComponent->TriangleColorFunc);
-		//BaseMeshComponent->TriangleColorFunc =
-		//if (bIsFirst)
-		//{
-		//	BaseMeshComponent->NotifyMeshUpdated();
-		//}
-		//else
-		//{
-		//	BaseMeshComponent->NotifyMeshUpdated();
-		//	//BaseMeshComponent->FastNotifyColorsUpdated();
-		//}
+			PreviewMesh->SetTriangleColorFunction([this](const FDynamicMesh3* Mesh, int TriangleID)
+			{
+				return SelectedTriangles[TriangleID] ? LinearColors::VideoRed3b() : LinearColors::VideoWhite3b();
+
+			}, UPreviewMesh::ERenderUpdateMode::FullUpdate);
+		}
 	}
 	else
 	{
@@ -588,6 +565,7 @@ void UMeshSelectionTool::Tick(float DeltaTime)
 	UDynamicMeshBrushTool::Tick(DeltaTime);
 
 	ShowWireframeWatcher.CheckAndUpdate();
+	ShowGroupsWatcher.CheckAndUpdate();
 
 	if (bStampPending)
 	{
@@ -688,6 +666,10 @@ void UMeshSelectionTool::ApplyAction(EMeshSelectionToolActions ActionType)
 		case EMeshSelectionToolActions::FlipSelected:
 			FlipSelectedTriangles();
 			break;
+
+		case EMeshSelectionToolActions::CreateGroup:
+			AssignNewGroupToSelectedTriangles();
+			break;
 	}
 }
 
@@ -993,7 +975,17 @@ void UMeshSelectionTool::FlipSelectedTriangles()
 
 	TUniquePtr<FToolCommandChangeSequence> ChangeSeq = MakeUnique<FToolCommandChangeSequence>();
 
-	// invert triangles and emit triangles change
+	// clear current selection
+	BeginChange(false);
+	for (int tid : SelectedFaces)
+	{
+		ActiveSelectionChange->Add(tid);
+	}
+	Selection->RemoveIndices(EMeshSelectionElementType::Face, SelectedFaces);
+	TUniquePtr<FMeshSelectionChange> SelectionChange = EndChange();
+	ChangeSeq->AppendChange(Selection, MoveTemp(SelectionChange));
+
+	// flip normals
 	TUniquePtr<FMeshChange> MeshChange = PreviewMesh->TrackedEditMesh(
 		[&SelectedFaces](FDynamicMesh3& Mesh, FDynamicMeshChangeTracker& ChangeTracker)
 	{
@@ -1013,6 +1005,56 @@ void UMeshSelectionTool::FlipSelectedTriangles()
 }
 
 
+void UMeshSelectionTool::AssignNewGroupToSelectedTriangles()
+{
+	check(SelectionType == EMeshSelectionElementType::Face);
+	TArray<int32> SelectedFaces = Selection->GetElements(EMeshSelectionElementType::Face);
+	if (SelectedFaces.Num() == 0)
+	{
+		return;
+	}
+
+	TUniquePtr<FToolCommandChangeSequence> ChangeSeq = MakeUnique<FToolCommandChangeSequence>();
+
+	// clear current selection
+	BeginChange(false);
+	for (int tid : SelectedFaces)
+	{
+		ActiveSelectionChange->Add(tid);
+	}
+	Selection->RemoveIndices(EMeshSelectionElementType::Face, SelectedFaces);
+	TUniquePtr<FMeshSelectionChange> SelectionChange = EndChange();
+	ChangeSeq->AppendChange(Selection, MoveTemp(SelectionChange));
+
+	// assign new groups to triangles
+	// note: using an FMeshChange is kind of overkill here
+	TUniquePtr<FMeshChange> MeshChange = PreviewMesh->TrackedEditMesh(
+		[&SelectedFaces](FDynamicMesh3& Mesh, FDynamicMeshChangeTracker& ChangeTracker)
+	{
+		// each component gets its own group id
+		FMeshConnectedComponents Components(&Mesh);
+		Components.FindConnectedTriangles(SelectedFaces);
+
+		for (FMeshConnectedComponents::FComponent& Component : Components)
+		{
+			int NewGroupID = Mesh.AllocateTriangleGroup();
+			for (int tid : Component.Indices)
+			{
+				ChangeTracker.SaveTriangle(tid, true);
+				Mesh.SetTriangleGroup(tid, NewGroupID);
+			}
+		}
+	});
+	ChangeSeq->AppendChange(PreviewMesh, MoveTemp(MeshChange));
+
+	// emit combined change sequence
+	GetToolManager()->EmitObjectChange(this, MoveTemp(ChangeSeq), LOCTEXT("MeshSelectionToolCreateGroup", "Create Polygroup"));
+
+	OnExternalSelectionChange();
+	bHaveModifiedMesh = true;
+
+}
+
 
 
 
diff --git a/Engine/Plugins/Experimental/MeshModelingToolset/Source/MeshModelingTools/Public/MeshSelectionTool.h b/Engine/Plugins/Experimental/MeshModelingToolset/Source/MeshModelingTools/Public/MeshSelectionTool.h
--- a/Engine/Plugins/Experimental/MeshModelingToolset/Source/MeshModelingTools/Public/MeshSelectionTool.h
+++ b/Engine/Plugins/Experimental/MeshModelingToolset/Source/MeshModelingTools/Public/MeshSelectionTool.h
@@ -40,7 +40,8 @@ enum class EMeshSelectionToolActions
 
 	DeleteSelected,
 	SeparateSelected,
-	FlipSelected
+	FlipSelected,
+	CreateGroup
 };
 
 
@@ -109,23 +110,29 @@ class MESHMODELINGTOOLS_API UMeshSelectionMeshEditActions : public UMeshSelectio
 	GENERATED_BODY()
 
 public:
-	UFUNCTION(CallInEditor, Category = MeshEdits, meta = (DisplayName = "Delete"))
+	UFUNCTION(CallInEditor, Category = MeshEdits, meta = (DisplayName = "Delete", DisplayPriority = 1))
 	void DeleteTriangles()
 	{
 		PostAction(EMeshSelectionToolActions::DeleteSelected);
 	}
 
-	UFUNCTION(CallInEditor, Category = MeshEdits, meta = (DisplayName = "Separate"))
+	UFUNCTION(CallInEditor, Category = MeshEdits, meta = (DisplayName = "Separate", DisplayPriority = 2))
 	void SeparateTriangles() 
 	{
 		PostAction(EMeshSelectionToolActions::SeparateSelected);
 	}
 
-	UFUNCTION(CallInEditor, Category = MeshEdits, meta = (DisplayName = "Flip Orientations"))
-		void FlipOrientations() 
+	UFUNCTION(CallInEditor, Category = MeshEdits, meta = (DisplayName = "Flip Normals", DisplayPriority = 3))
+	void FlipNormals() 
 	{
 		PostAction(EMeshSelectionToolActions::FlipSelected);
 	}
+
+	UFUNCTION(CallInEditor, Category = MeshEdits, meta = (DisplayName = "Create Polygroup", DisplayPriority = 4))
+	void CreatePolygroup()
+	{
+		PostAction(EMeshSelectionToolActions::CreateGroup);
+	}
 };
 
 
@@ -183,6 +190,10 @@ public:
 	UPROPERTY(EditAnywhere, Category = Selection)
 	bool bShowWireframe = true;
 
+	/** Toggle drawing of groups */
+	UPROPERTY(EditAnywhere, Category = Selection)
+	bool bShowPolygroups = false;
+
 	virtual void SaveProperties(UInteractiveTool* SaveFromTool) override;
 	virtual void RestoreProperties(UInteractiveTool* RestoreToTool) override;
 };
@@ -255,6 +266,7 @@ protected:
 	EMeshSelectionElementType SelectionType = EMeshSelectionElementType::Face;
 
 	TValueWatcher<bool> ShowWireframeWatcher;
+	TValueWatcher<bool> ShowGroupsWatcher;
 
 
 	bool bInRemoveStroke = false;
@@ -300,6 +312,7 @@ protected:
 	void DeleteSelectedTriangles();
 	void SeparateSelectedTriangles();
 	void FlipSelectedTriangles();
+	void AssignNewGroupToSelectedTriangles();
 
 	bool bHaveModifiedMesh = false;
 };
diff --git a/Engine/Plugins/Experimental/MeshModelingToolset/Source/ModelingComponents/Private/OctreeDynamicMeshComponent.cpp b/Engine/Plugins/Experimental/MeshModelingToolset/Source/ModelingComponents/Private/OctreeDynamicMeshComponent.cpp
--- a/Engine/Plugins/Experimental/MeshModelingToolset/Source/ModelingComponents/Private/OctreeDynamicMeshComponent.cpp
+++ b/Engine/Plugins/Experimental/MeshModelingToolset/Source/ModelingComponents/Private/OctreeDynamicMeshComponent.cpp
@@ -275,7 +275,7 @@ FPrimitiveSceneProxy* UOctreeDynamicMeshComponent::CreateSceneProxy()
 		if (TriangleColorFunc != nullptr)
 		{
 			CurrentProxy->bUsePerTriangleColor = true;
-			CurrentProxy->PerTriangleColorFunc = [this](int TriangleID) { return GetTriangleColor(TriangleID); };
+			CurrentProxy->PerTriangleColorFunc = [this](const FDynamicMesh3* Mesh, int TriangleID) { return GetTriangleColor(TriangleID); };
 		}
 
 
@@ -327,7 +327,7 @@ FColor UOctreeDynamicMeshComponent::GetTriangleColor(int TriangleID)
 {
 	if (TriangleColorFunc != nullptr)
 	{
-		return TriangleColorFunc(TriangleID);
+		return TriangleColorFunc(Mesh.Get(), TriangleID);
 	}
 	else
 	{
diff --git a/Engine/Plugins/Experimental/MeshModelingToolset/Source/ModelingComponents/Private/OctreeDynamicMeshSceneProxy.h b/Engine/Plugins/Experimental/MeshModelingToolset/Source/ModelingComponents/Private/OctreeDynamicMeshSceneProxy.h
--- a/Engine/Plugins/Experimental/MeshModelingToolset/Source/ModelingComponents/Private/OctreeDynamicMeshSceneProxy.h
+++ b/Engine/Plugins/Experimental/MeshModelingToolset/Source/ModelingComponents/Private/OctreeDynamicMeshSceneProxy.h
@@ -142,7 +142,7 @@ public:
 
 
 	bool bUsePerTriangleColor = false;
-	TFunction<FColor(int)> PerTriangleColorFunc = nullptr;
+	TFunction<FColor(const FDynamicMesh3*, int)> PerTriangleColorFunc = nullptr;
 
 
 	FOctreeDynamicMeshSceneProxy(UOctreeDynamicMeshComponent* Component)
@@ -366,7 +366,7 @@ protected:
 			FColor TriColor = ConstantVertexColor;
 			if (bUsePerTriangleColor && PerTriangleColorFunc != nullptr)
 			{
-				TriColor = PerTriangleColorFunc(TriangleID);
+				TriColor = PerTriangleColorFunc(Mesh, TriangleID);
 				bHaveColors = false;
 			}
 
diff --git a/Engine/Plugins/Experimental/MeshModelingToolset/Source/ModelingComponents/Private/PreviewMesh.cpp b/Engine/Plugins/Experimental/MeshModelingToolset/Source/ModelingComponents/Private/PreviewMesh.cpp
--- a/Engine/Plugins/Experimental/MeshModelingToolset/Source/ModelingComponents/Private/PreviewMesh.cpp
+++ b/Engine/Plugins/Experimental/MeshModelingToolset/Source/ModelingComponents/Private/PreviewMesh.cpp
@@ -293,7 +293,7 @@ void UPreviewMesh::Bake(FMeshDescription* MeshDescription, bool bHaveModifiedTop
 	DynamicMeshComponent->Bake(MeshDescription, bHaveModifiedToplogy);
 }
 
-void UPreviewMesh::SetTriangleColorFunction(TFunction<FColor(int)> TriangleColorFunc, ERenderUpdateMode UpdateMode)
+void UPreviewMesh::SetTriangleColorFunction(TFunction<FColor(const FDynamicMesh3*, int)> TriangleColorFunc, ERenderUpdateMode UpdateMode)
 {
 	DynamicMeshComponent->TriangleColorFunc = TriangleColorFunc;
 	if (UpdateMode == ERenderUpdateMode::FastUpdate)
diff --git a/Engine/Plugins/Experimental/MeshModelingToolset/Source/ModelingComponents/Private/SimpleDynamicMeshComponent.cpp b/Engine/Plugins/Experimental/MeshModelingToolset/Source/ModelingComponents/Private/SimpleDynamicMeshComponent.cpp
--- a/Engine/Plugins/Experimental/MeshModelingToolset/Source/ModelingComponents/Private/SimpleDynamicMeshComponent.cpp
+++ b/Engine/Plugins/Experimental/MeshModelingToolset/Source/ModelingComponents/Private/SimpleDynamicMeshComponent.cpp
@@ -195,7 +195,7 @@ FPrimitiveSceneProxy* USimpleDynamicMeshComponent::CreateSceneProxy()
 		if (TriangleColorFunc)
 		{
 			CurrentProxy->bUsePerTriangleColor = true;
-			CurrentProxy->PerTriangleColorFunc = [this](int TriangleID) { return GetTriangleColor(TriangleID); };
+			CurrentProxy->PerTriangleColorFunc = [this](const FDynamicMesh3* MeshIn, int TriangleID) { return GetTriangleColor(MeshIn, TriangleID); };
 		}
 
 		CurrentProxy->Initialize();
@@ -209,11 +209,11 @@ int32 USimpleDynamicMeshComponent::GetNumMaterials() const
 }
 
 
-FColor USimpleDynamicMeshComponent::GetTriangleColor(int TriangleID)
+FColor USimpleDynamicMeshComponent::GetTriangleColor(const FDynamicMesh3* MeshIn, int TriangleID)
 {
 	if (TriangleColorFunc)
 	{
-		return TriangleColorFunc(TriangleID);
+		return TriangleColorFunc(MeshIn, TriangleID);
 	}
 	else
 	{
diff --git a/Engine/Plugins/Experimental/MeshModelingToolset/Source/ModelingComponents/Private/SimpleDynamicMeshSceneProxy.h b/Engine/Plugins/Experimental/MeshModelingToolset/Source/ModelingComponents/Private/SimpleDynamicMeshSceneProxy.h
--- a/Engine/Plugins/Experimental/MeshModelingToolset/Source/ModelingComponents/Private/SimpleDynamicMeshSceneProxy.h
+++ b/Engine/Plugins/Experimental/MeshModelingToolset/Source/ModelingComponents/Private/SimpleDynamicMeshSceneProxy.h
@@ -43,7 +43,7 @@ public:
 
 
 	bool bUsePerTriangleColor = false;
-	TFunction<FColor(int)> PerTriangleColorFunc;
+	TFunction<FColor(const FDynamicMesh3*, int)> PerTriangleColorFunc;
 
 
 	FSimpleDynamicMeshSceneProxy(USimpleDynamicMeshComponent* Component)
@@ -194,7 +194,7 @@ public:
 			FColor TriColor = ConstantVertexColor;
 			if (bUsePerTriangleColor && PerTriangleColorFunc)
 			{
-				TriColor = PerTriangleColorFunc(TriangleID);
+				TriColor = PerTriangleColorFunc(Mesh, TriangleID);
 				bHaveColors = false;
 			}
 
@@ -272,7 +272,7 @@ protected:
 			FColor TriColor = ConstantVertexColor;
 			if (bUsePerTriangleColor && PerTriangleColorFunc)
 			{
-				TriColor = PerTriangleColorFunc(TriangleID);
+				TriColor = PerTriangleColorFunc(Mesh, TriangleID);
 				bHaveColors = false;
 			}
 
diff --git a/Engine/Plugins/Experimental/MeshModelingToolset/Source/ModelingComponents/Public/OctreeDynamicMeshComponent.h b/Engine/Plugins/Experimental/MeshModelingToolset/Source/ModelingComponents/Public/OctreeDynamicMeshComponent.h
--- a/Engine/Plugins/Experimental/MeshModelingToolset/Source/ModelingComponents/Public/OctreeDynamicMeshComponent.h
+++ b/Engine/Plugins/Experimental/MeshModelingToolset/Source/ModelingComponents/Public/OctreeDynamicMeshComponent.h
@@ -111,7 +111,7 @@ public:
 	UPROPERTY()
 	bool bExplicitShowWireframe = false;
 
-	TFunction<FColor(int)> TriangleColorFunc = nullptr;
+	TFunction<FColor(const FDynamicMesh3*, int)> TriangleColorFunc = nullptr;
 
 private:
 
diff --git a/Engine/Plugins/Experimental/MeshModelingToolset/Source/ModelingComponents/Public/PreviewMesh.h b/Engine/Plugins/Experimental/MeshModelingToolset/Source/ModelingComponents/Public/PreviewMesh.h
--- a/Engine/Plugins/Experimental/MeshModelingToolset/Source/ModelingComponents/Public/PreviewMesh.h
+++ b/Engine/Plugins/Experimental/MeshModelingToolset/Source/ModelingComponents/Public/PreviewMesh.h
@@ -134,7 +134,7 @@ public:
 	/**
 	 * Set the triangle color function for rendering / render data construction
 	 */
-	void SetTriangleColorFunction(TFunction<FColor(int)> TriangleColorFunc, ERenderUpdateMode UpdateMode = ERenderUpdateMode::FullUpdate);
+	void SetTriangleColorFunction(TFunction<FColor(const FDynamicMesh3*, int)> TriangleColorFunc, ERenderUpdateMode UpdateMode = ERenderUpdateMode::FullUpdate);
 
 	/**
 	 * Clear the triangle color function for rendering / render data construction
diff --git a/Engine/Plugins/Experimental/MeshModelingToolset/Source/ModelingComponents/Public/SimpleDynamicMeshComponent.h b/Engine/Plugins/Experimental/MeshModelingToolset/Source/ModelingComponents/Public/SimpleDynamicMeshComponent.h
--- a/Engine/Plugins/Experimental/MeshModelingToolset/Source/ModelingComponents/Public/SimpleDynamicMeshComponent.h
+++ b/Engine/Plugins/Experimental/MeshModelingToolset/Source/ModelingComponents/Public/SimpleDynamicMeshComponent.h
@@ -121,7 +121,7 @@ public:
 
 
 
-	TFunction<FColor(int)> TriangleColorFunc = nullptr;
+	TFunction<FColor(const FDynamicMesh3*, int)> TriangleColorFunc = nullptr;
 	void FastNotifyColorsUpdated();
 
 	void FastNotifyPositionsUpdated();
@@ -148,7 +148,7 @@ private:
 	bool bTangentsValid = false;
 	FMeshTangentsf Tangents;
 	
-	FColor GetTriangleColor(int TriangleID);
+	FColor GetTriangleColor(const FDynamicMesh3* Mesh, int TriangleID);
 
 	//friend class FCustomMeshSceneProxy;
 };
