ModelingTools: improve snapping support
- add ToolSceneQueriesUtil::FindWorldGridSnapPoint
- replace OverrideMetric in FBasePositionSnapSolver3 with more general CustomMetric struct, can be interpreted in various ways, update OverrideMetric usage sites
- Add RaySpatialSnapSolver::PointConstraintFunc, allows snapping to constraints (eg to world grid) after snapping to lines
- expose above in QuickAxisTranslater::UpdateSnap
#rnx

diff --git a/Engine/Plugins/Experimental/MeshModelingToolset/Source/ModelingComponents/Private/Snapping/BasePositionSnapSolver3.cpp b/Engine/Plugins/Experimental/MeshModelingToolset/Source/ModelingComponents/Private/Snapping/BasePositionSnapSolver3.cpp
--- a/Engine/Plugins/Experimental/MeshModelingToolset/Source/ModelingComponents/Private/Snapping/BasePositionSnapSolver3.cpp
+++ b/Engine/Plugins/Experimental/MeshModelingToolset/Source/ModelingComponents/Private/Snapping/BasePositionSnapSolver3.cpp
@@ -29,10 +29,25 @@ void FBasePositionSnapSolver3::ResetActiveSnap()
 	ClearActiveSnapData();
 }
 
+void FBasePositionSnapSolver3::AddPointTarget(const FVector3d& Position, int TargetID, int Priority)
+{
+	FSnapTargetPoint Point;
+	Point.Position = Position;
+	Point.TargetID = TargetID;
+	Point.Priority = Priority;
+	Point.bHaveCustomMetric = false;
+	TargetPoints.Add(Point);
+}
 
-void FBasePositionSnapSolver3::AddPointTarget(const FVector3d& Position, int TargetID, int Priority, double OverrideMetric)
+
+void FBasePositionSnapSolver3::AddPointTarget(const FVector3d& Position, int TargetID, const FCustomMetric& CustomMetric, int Priority)
 {
-	FSnapTargetPoint Point = { Position, TargetID, Priority, OverrideMetric };
+	FSnapTargetPoint Point;
+	Point.Position = Position;
+	Point.TargetID = TargetID;
+	Point.Priority = Priority;
+	Point.bHaveCustomMetric = true;
+	Point.CustomMetric = CustomMetric;
 	TargetPoints.Add(Point);
 }
 
@@ -138,9 +153,27 @@ const FBasePositionSnapSolver3::FSnapTargetPoint* FBasePositionSnapSolver3::Find
 			continue;
 		}
 
+		double TestMetric = SnapMetricTolerance;
+		if (TestTargets[k].bHaveCustomMetric)
+		{
+			switch (TestTargets[k].CustomMetric.Type)
+			{
+			case ECustomMetricType::MinValue:
+				TestMetric = FMathd::Min(SnapMetricTolerance, TestTargets[k].CustomMetric.Value);
+				break;
+			case ECustomMetricType::Multiplier:
+				TestMetric = TestTargets[k].CustomMetric.Value * SnapMetricTolerance;
+				break;
+			case ECustomMetricType::ReplaceValue:
+				TestMetric = TestTargets[k].CustomMetric.Value;
+				break;
+			}
+		}
+
 		FVector3d SnapPoint = GetSnapPointFromFunc(TestTargets[k].Position);
 		double Metric = SnapMetricFunc(SnapPoint, TestTargets[k].Position);
-		if (Metric < SnapMetricTolerance && Metric < TestTargets[k].OverrideMetric)
+
+		if (Metric < TestMetric)
 		{
 			if (Metric < MinMetric || TestTargets[k].Priority < MinPriority)
 			{
diff --git a/Engine/Plugins/Experimental/MeshModelingToolset/Source/ModelingComponents/Private/Snapping/RaySpatialSnapSolver.cpp b/Engine/Plugins/Experimental/MeshModelingToolset/Source/ModelingComponents/Private/Snapping/RaySpatialSnapSolver.cpp
--- a/Engine/Plugins/Experimental/MeshModelingToolset/Source/ModelingComponents/Private/Snapping/RaySpatialSnapSolver.cpp
+++ b/Engine/Plugins/Experimental/MeshModelingToolset/Source/ModelingComponents/Private/Snapping/RaySpatialSnapSolver.cpp
@@ -28,10 +28,28 @@ void FRaySpatialSnapSolver::GenerateTargetPoints(const FRay3d& Ray)
 		FDistLine3Ray3d DistQuery(TargetLines[j].Line, Ray);
 		DistQuery.GetSquared();
 		Target.Position = DistQuery.LineClosestPoint;
+
 		Target.TargetID = TargetLines[j].TargetID;
 		Target.Priority = TargetLines[j].Priority;
 		Target.bIsSnapLine = true;
 		Target.SnapLine = TargetLines[j].Line;
+
+		// if we have a point constraint function, find constraint point and then project back onto line
+		if (PointConstraintFunc != nullptr)
+		{
+			FVector3d SnapPos = PointConstraintFunc(Target.Position);
+			Target.ConstrainedPosition = TargetLines[j].Line.NearestPoint(SnapPos);
+			double PosT = TargetLines[j].Line.Project(Target.Position);
+			double PosConstrainT = TargetLines[j].Line.Project(Target.ConstrainedPosition);
+			// its usually weird to snap from A to B if they are on opposite sides of the origin, so clamp to origin in that case
+			// (todo: maybe this should be optional...makes sense when using this for interactive gizmo snapping but maybe not in other cases)
+			if (PosT * PosConstrainT < 0)
+			{
+				Target.ConstrainedPosition = TargetLines[j].Line.Origin;
+			}
+			Target.bHaveConstrainedPosition = true;
+		}
+
 		GeneratedTargetPoints.Add(Target);
 	}
 
@@ -98,7 +116,11 @@ void FRaySpatialSnapSolver::UpdateSnappedPoint(const FRay3d& RayIn)
 	// if we found a best-target, update our snap details
 	if (BestSnapTarget != nullptr)
 	{
-		SetActiveSnapData(*BestSnapTarget, RayIn.NearestPoint(BestSnapTarget->Position), BestSnapTarget->Position, MinMetric);
+		// if we have a constrained position for a given snap we will use that instead. This is 
+		// used for grid+line snapping for example. We want to snap to the grid *after* we snap to the
+		// line, otherwise things snap to weird unexpected places.
+		FVector3d UseSnapPosition = (BestSnapTarget->bHaveConstrainedPosition) ? BestSnapTarget->ConstrainedPosition : BestSnapTarget->Position;
+		SetActiveSnapData(*BestSnapTarget, RayIn.NearestPoint(BestSnapTarget->Position), UseSnapPosition, MinMetric);
 	}
 	else
 	{
diff --git a/Engine/Plugins/Experimental/MeshModelingToolset/Source/ModelingComponents/Private/ToolSceneQueriesUtil.cpp b/Engine/Plugins/Experimental/MeshModelingToolset/Source/ModelingComponents/Private/ToolSceneQueriesUtil.cpp
--- a/Engine/Plugins/Experimental/MeshModelingToolset/Source/ModelingComponents/Private/ToolSceneQueriesUtil.cpp
+++ b/Engine/Plugins/Experimental/MeshModelingToolset/Source/ModelingComponents/Private/ToolSceneQueriesUtil.cpp
@@ -132,4 +132,21 @@ bool ToolSceneQueriesUtil::FindSceneSnapPoint(const UInteractiveTool* Tool, cons
 		return true;
 	}
 	return false;
+}
+
+
+bool ToolSceneQueriesUtil::FindWorldGridSnapPoint(const UInteractiveTool* Tool, const FVector3d& Point, FVector3d& GridSnapPointOut)
+{
+	IToolsContextQueriesAPI* QueryAPI = Tool->GetToolManager()->GetContextQueriesAPI();
+	FSceneSnapQueryRequest Request;
+	Request.RequestType = ESceneSnapQueryType::Position;
+	Request.TargetTypes = ESceneSnapQueryTargetType::Grid;
+	Request.Position = Point;
+	TArray<FSceneSnapQueryResult> Results;
+	if ( QueryAPI->ExecuteSceneSnapQuery(Request, Results) )
+	{
+		GridSnapPointOut = Results[0].Position;
+		return true;
+	};
+	return false;
 }
\ No newline at end of file
diff --git a/Engine/Plugins/Experimental/MeshModelingToolset/Source/ModelingComponents/Private/Transforms/QuickAxisRotator.cpp b/Engine/Plugins/Experimental/MeshModelingToolset/Source/ModelingComponents/Private/Transforms/QuickAxisRotator.cpp
--- a/Engine/Plugins/Experimental/MeshModelingToolset/Source/ModelingComponents/Private/Transforms/QuickAxisRotator.cpp
+++ b/Engine/Plugins/Experimental/MeshModelingToolset/Source/ModelingComponents/Private/Transforms/QuickAxisRotator.cpp
@@ -98,7 +98,8 @@ void FQuickAxisRotator::UpdateSnapAxes()
 	FFrame3d RotZFrame(AxisFrameWorld.Origin, AxisFrameWorld.GetAxis(2));
 	MoveAxisSolver.AddCircleTarget(FCircle3d(RotZFrame, UseRadius), RotateZSnapID, 100);
 
-	MoveAxisSolver.AddPointTarget(AxisFrameWorld.Origin, RotateOriginSnapID, 50, 0.5*ToolSceneQueriesUtil::GetDefaultVisualAngleSnapThreshD());
+	MoveAxisSolver.AddPointTarget(AxisFrameWorld.Origin, RotateOriginSnapID,
+		FRaySpatialSnapSolver::FCustomMetric::Replace(0.5*ToolSceneQueriesUtil::GetDefaultVisualAngleSnapThreshD()), 50);
 }
 
 
diff --git a/Engine/Plugins/Experimental/MeshModelingToolset/Source/ModelingComponents/Private/Transforms/QuickAxisTranslater.cpp b/Engine/Plugins/Experimental/MeshModelingToolset/Source/ModelingComponents/Private/Transforms/QuickAxisTranslater.cpp
--- a/Engine/Plugins/Experimental/MeshModelingToolset/Source/ModelingComponents/Private/Transforms/QuickAxisTranslater.cpp
+++ b/Engine/Plugins/Experimental/MeshModelingToolset/Source/ModelingComponents/Private/Transforms/QuickAxisTranslater.cpp
@@ -178,16 +178,32 @@ void FQuickAxisTranslater::UpdateSnapAxes()
 		}
 	}
 
-	MoveAxisSolver.AddPointTarget(AxisFrameWorld.Origin, TranslateOriginSnapID, 50, 0.5*ToolSceneQueriesUtil::GetDefaultVisualAngleSnapThreshD());
+	MoveAxisSolver.AddPointTarget(AxisFrameWorld.Origin, TranslateOriginSnapID,
+		FRaySpatialSnapSolver::FCustomMetric::Replace(0.5*ToolSceneQueriesUtil::GetDefaultVisualAngleSnapThreshD()), 50);
 }
 
 
 
 
 
-bool FQuickAxisTranslater::UpdateSnap(const FRay3d& Ray, FVector3d& SnapPointOut)
+bool FQuickAxisTranslater::UpdateSnap(const FRay3d& Ray, FVector3d& SnapPointOut,
+	TFunction<FVector3d(const FVector3d&)> PositionConstraintFunc)
 {
+	bool bSetConstraintFunc;
+	if (PositionConstraintFunc != nullptr)
+	{
+		check(MoveAxisSolver.PointConstraintFunc == nullptr);
+		MoveAxisSolver.PointConstraintFunc = PositionConstraintFunc;
+		bSetConstraintFunc = true;
+	}
+
 	MoveAxisSolver.UpdateSnappedPoint(Ray);
+
+	if (bSetConstraintFunc)
+	{
+		MoveAxisSolver.PointConstraintFunc = nullptr;
+	}
+
 	if (MoveAxisSolver.HaveActiveSnap())
 	{
 		SnapPointOut = MoveAxisSolver.GetActiveSnapToPoint();
diff --git a/Engine/Plugins/Experimental/MeshModelingToolset/Source/ModelingComponents/Public/Snapping/BasePositionSnapSolver3.h b/Engine/Plugins/Experimental/MeshModelingToolset/Source/ModelingComponents/Public/Snapping/BasePositionSnapSolver3.h
--- a/Engine/Plugins/Experimental/MeshModelingToolset/Source/ModelingComponents/Public/Snapping/BasePositionSnapSolver3.h
+++ b/Engine/Plugins/Experimental/MeshModelingToolset/Source/ModelingComponents/Public/Snapping/BasePositionSnapSolver3.h
@@ -54,13 +54,41 @@ public:
 	// setup
 	// 
 
+	/** ECustomMetricType defines how a FCustomMetric should be interpreted  */
+	enum class ECustomMetricType
+	{
+		ReplaceValue = 0,		/** Use given parameter instead of default metric value */
+		Multiplier = 1,			/** Multiply default metric value by given parameter */
+		MinValue = 2			/** Use minimum of default metric value and given parameter */
+	};
+	/** FCustomMetric overrides/modifies the default should-snap-happen metric */
+	struct FCustomMetric
+	{
+		/** What kind of custom metric is this */
+		ECustomMetricType Type;
+		/** parameter for custom metric */
+		double Value;
+		static FCustomMetric Minimum(double Val) { return FCustomMetric{ ECustomMetricType::MinValue, Val }; }
+		static FCustomMetric Replace(double Val) { return FCustomMetric{ ECustomMetricType::ReplaceValue, Val }; }
+		static FCustomMetric Multiplier(double Val) { return FCustomMetric{ ECustomMetricType::Multiplier, Val }; }
+	};
+
 	/**
 	 * Add a snap target point at the given Position
+	 * @param Position snap target point location
 	 * @param TargetID identifier for this target. Can be shared with other targets
 	 * @param Priority importance of this snap point. Lower priority is more important.
-	 * @param OverrideMetric if provided, then we only snap to this point if Metric < OverrideMetric, as well as the global tolerance
 	 */
-	virtual void AddPointTarget(const FVector3d& Position, int TargetID, int Priority = 100, double OverrideMetric = TNumericLimits<double>::Max());
+	virtual void AddPointTarget(const FVector3d& Position, int TargetID, int Priority = 100);
+
+	/**
+	 * Add a snap target point at the given Position
+	 * @param Position snap target point location
+	 * @param TargetID identifier for this target. Can be shared with other targets
+	 * @param CustomMetric Use the given custom metric instead of the default metric
+	 * @param Priority importance of this snap point. Lower priority is more important.
+	 */
+	virtual void AddPointTarget(const FVector3d& Position, int TargetID, const FCustomMetric& CustomMetric, int Priority = 100);
 
 	/** Remove any point targets with this TargetID */
 	virtual bool RemovePointTargetsByID(int TargetID);
@@ -141,12 +169,20 @@ protected:
 
 		int TargetID;
 		int Priority;
-		double OverrideMetric = TNumericLimits<double>::Max();
+
+		FCustomMetric CustomMetric;
+		bool bHaveCustomMetric = false;
 
 		bool bIsSnapLine = false;
 		FLine3d SnapLine;
 		bool bIsSnapDistance = false;
 		int SnapDistanceID = -1;
+
+		// if true, then Position is the snap target point but ConstrainedPosition is the one we should return
+		// Sounds weird but allows for things like axis-constrained grid snapping to work. We have to snap to
+		// the line first, based on user interaction, and then to the grid once the line is selected
+		bool bHaveConstrainedPosition = false;
+		FVector3d ConstrainedPosition;
 	};
 	TArray<FSnapTargetPoint> TargetPoints;
 
diff --git a/Engine/Plugins/Experimental/MeshModelingToolset/Source/ModelingComponents/Public/Snapping/RaySpatialSnapSolver.h b/Engine/Plugins/Experimental/MeshModelingToolset/Source/ModelingComponents/Public/Snapping/RaySpatialSnapSolver.h
--- a/Engine/Plugins/Experimental/MeshModelingToolset/Source/ModelingComponents/Public/Snapping/RaySpatialSnapSolver.h
+++ b/Engine/Plugins/Experimental/MeshModelingToolset/Source/ModelingComponents/Public/Snapping/RaySpatialSnapSolver.h
@@ -18,6 +18,14 @@ class MODELINGCOMPONENTS_API FRaySpatialSnapSolver : public FBasePositionSnapSol
 public:
 	FRaySpatialSnapSolver();
 
+	/**
+	 * Optional function that will be used to project potential snap points onto constraints.
+	 * Note that Line/Curve constraints are still respected, so eg if this projects to a 3D grid,
+	 * then when calculating possible line-snap positions, the 3D grid point will be projected back 
+	 * onto the line targets.
+	 */
+	TFunction<FVector3d(const FVector3d&)> PointConstraintFunc = nullptr;
+
 	//
 	// solving
 	//
diff --git a/Engine/Plugins/Experimental/MeshModelingToolset/Source/ModelingComponents/Public/ToolSceneQueriesUtil.h b/Engine/Plugins/Experimental/MeshModelingToolset/Source/ModelingComponents/Public/ToolSceneQueriesUtil.h
--- a/Engine/Plugins/Experimental/MeshModelingToolset/Source/ModelingComponents/Public/ToolSceneQueriesUtil.h
+++ b/Engine/Plugins/Experimental/MeshModelingToolset/Source/ModelingComponents/Public/ToolSceneQueriesUtil.h
@@ -80,10 +80,18 @@ namespace ToolSceneQueriesUtil
 	 * @param bVertices if true, try snapping to mesh vertices in the scene
 	 * @param bEdges if true, try snapping to mesh triangle edges in the scene
 	 * @param VisualAngleThreshold visual angle threshold to use. If 0, GetDefaultVisualAngleSnapThresh() is used
+	 * @param SnapGeometry world-space position of the snap geometry (point/line/polygon)
 	 * @param DebugTriangleOut if non-null, triangle containing snap is returned if a snap is found
 	 * @return true if a valid snap point was found
 	 */
 	MODELINGCOMPONENTS_API bool FindSceneSnapPoint(const UInteractiveTool* Tool, const FVector3d& Point, FVector3d& SnapPointOut,
 		bool bVertices = true, bool bEdges = false, double VisualAngleThreshold = 0, 
 		FSnapGeometry* SnapGeometry = nullptr, FVector* DebugTriangleOut = nullptr);
+
+
+	/**
+	 * Run a query against the scene to find the nearest WorldGrid snap point
+	 */
+	MODELINGCOMPONENTS_API bool FindWorldGridSnapPoint(const UInteractiveTool* Tool, const FVector3d& QueryPoint, FVector3d& GridSnapPointOut);
+
 }
\ No newline at end of file
diff --git a/Engine/Plugins/Experimental/MeshModelingToolset/Source/ModelingComponents/Public/Transforms/QuickAxisTranslater.h b/Engine/Plugins/Experimental/MeshModelingToolset/Source/ModelingComponents/Public/Transforms/QuickAxisTranslater.h
--- a/Engine/Plugins/Experimental/MeshModelingToolset/Source/ModelingComponents/Public/Transforms/QuickAxisTranslater.h
+++ b/Engine/Plugins/Experimental/MeshModelingToolset/Source/ModelingComponents/Public/Transforms/QuickAxisTranslater.h
@@ -56,9 +56,13 @@ public:
 
 	/**
 	 * Try to find the best snap point for the given Ray, and store in SnapPointOut if found
+	 * @param Ray 3D ray in space of snap points
+	 * @param SnapPointOut found snap point
+	 * @param PositionConstraintFunc Function that projects potential snap points onto constraint surfaces (eg grid points)
 	 * @return true if found
 	 */
-	virtual bool UpdateSnap(const FRay3d& Ray, FVector3d& SnapPointOut);
+	virtual bool UpdateSnap(const FRay3d& Ray, FVector3d& SnapPointOut,
+		TFunction<FVector3d(const FVector3d&)> PositionConstraintFunc = nullptr);
 
 	/** @return true if there is an active snap */
 	virtual bool HaveActiveSnap() const;
