ModelingTools: Add WorldGrid Snap Support in DrawPolygonTool and EditMeshPolygonsTool
#rnx

diff --git a/Engine/Plugins/Experimental/MeshModelingToolset/Source/MeshModelingTools/Private/DrawPolygonTool.cpp b/Engine/Plugins/Experimental/MeshModelingToolset/Source/MeshModelingTools/Private/DrawPolygonTool.cpp
--- a/Engine/Plugins/Experimental/MeshModelingToolset/Source/MeshModelingTools/Private/DrawPolygonTool.cpp
+++ b/Engine/Plugins/Experimental/MeshModelingToolset/Source/MeshModelingTools/Private/DrawPolygonTool.cpp
@@ -51,6 +51,7 @@
  */
 constexpr int StartPointSnapID = FPointPlanarSnapSolver::BaseExternalPointID + 1;
 constexpr int CurrentSceneSnapID = FPointPlanarSnapSolver::BaseExternalPointID + 2;
+constexpr int CurrentGridSnapID = FPointPlanarSnapSolver::BaseExternalPointID + 3;
 
 bool UDrawPolygonToolBuilder::CanBuildTool(const FToolBuilderState& SceneState) const
 {
@@ -101,6 +102,7 @@ void UDrawPolygonToolSnapProperties::SaveProperties(UInteractiveTool* SaveFromTo
 {
 	UDrawPolygonToolSnapProperties* PropertyCache = GetPropertyCache<UDrawPolygonToolSnapProperties>();
 	PropertyCache->bEnableSnapping = this->bEnableSnapping;
+	PropertyCache->bSnapToWorldGrid = this->bSnapToWorldGrid;
 	PropertyCache->bSnapToVertices = this->bSnapToVertices;
 	PropertyCache->bSnapToEdges = this->bSnapToEdges;
 	PropertyCache->bSnapToAngles = this->bSnapToAngles;
@@ -114,6 +116,7 @@ void UDrawPolygonToolSnapProperties::RestoreProperties(UInteractiveTool* Restore
 {
 	UDrawPolygonToolSnapProperties* PropertyCache = GetPropertyCache<UDrawPolygonToolSnapProperties>();
 	this->bEnableSnapping = PropertyCache->bEnableSnapping;
+	this->bSnapToWorldGrid = PropertyCache->bSnapToWorldGrid;
 	this->bSnapToVertices = PropertyCache->bSnapToVertices;
 	this->bSnapToEdges = PropertyCache->bSnapToEdges;
 	this->bSnapToAngles = PropertyCache->bSnapToAngles;
@@ -378,6 +381,11 @@ void UDrawPolygonTool::Render(IToolsContextRenderAPI* RenderAPI)
 					SnapHighlightColor, SDPG_Foreground, 1.0f, 0.0f, true);
 			}
 		}
+		else if (SnapEngine.GetActiveSnapTargetID() == CurrentGridSnapID)
+		{
+			DrawCircle(PDI, (FVector)LastGridSnapPoint, CameraState.Right(), CameraState.Up(),
+				SnapHighlightColor, ElementSize, 4, SDPG_Foreground, 1.0f, 0.0f, true);
+		}
 
 		if (SnapEngine.HaveActiveSnapLine())
 		{
@@ -400,7 +408,6 @@ void UDrawPolygonTool::Render(IToolsContextRenderAPI* RenderAPI)
 		}
 	}
 
-
 	if (bHaveSurfaceHit)
 	{
 		PDI->DrawPoint((FVector)SurfaceHitPoint, ClosedPolygonColor, 10, SDPG_Foreground);
@@ -543,7 +550,6 @@ bool UDrawPolygonTool::FindDrawPlaneHitPoint(const FInputDeviceRay& ClickPos, FV
 	}
 
 	// if we found a scene snap point, add to snap set
-	FVector3d SnapPos;
 	if (bIgnoreSnappingToggle || SnapProperties->bEnableSnapping == false)
 	{
 		SnapEngine.ResetActiveSnap();
@@ -551,9 +557,25 @@ bool UDrawPolygonTool::FindDrawPlaneHitPoint(const FInputDeviceRay& ClickPos, FV
 	}
 	else 
 	{
-		if (ToolSceneQueriesUtil::FindSceneSnapPoint(this, HitPos, SnapPos, SnapProperties->bSnapToVertices, SnapProperties->bSnapToEdges, 0, &LastSnapGeometry))
+		if (SnapProperties->bSnapToWorldGrid)
 		{
-			SnapEngine.AddPointTarget(SnapPos, CurrentSceneSnapID, SnapEngine.MinInternalPriority()-1 );
+			FVector3d WorldGridSnapPos;
+			if (ToolSceneQueriesUtil::FindWorldGridSnapPoint(this, HitPos, WorldGridSnapPos))
+			{
+				WorldGridSnapPos = Frame.ToPlane(WorldGridSnapPos, 2);
+				SnapEngine.AddPointTarget(WorldGridSnapPos, CurrentGridSnapID, 
+					FBasePositionSnapSolver3::FCustomMetric::Replace(999), SnapEngine.MinInternalPriority() - 5);
+				LastGridSnapPoint = WorldGridSnapPos;
+			}
+		}
+
+		if (SnapProperties->bSnapToVertices || SnapProperties->bSnapToEdges)
+		{
+			FVector3d SceneSnapPos;
+			if (ToolSceneQueriesUtil::FindSceneSnapPoint(this, HitPos, SceneSnapPos, SnapProperties->bSnapToVertices, SnapProperties->bSnapToEdges, 0, &LastSnapGeometry))
+			{
+				SnapEngine.AddPointTarget(SceneSnapPos, CurrentSceneSnapID, SnapEngine.MinInternalPriority() - 10);
+			}
 		}
 
 		TArray<FVector>& HistoryPoints = (bInFixedPolygonMode) ? FixedPolygonClickPoints : PolygonVertices;
@@ -569,6 +591,7 @@ bool UDrawPolygonTool::FindDrawPlaneHitPoint(const FInputDeviceRay& ClickPos, FV
 
 	// remove scene snap point
 	SnapEngine.RemovePointTargetsByID(CurrentSceneSnapID);
+	SnapEngine.RemovePointTargetsByID(CurrentGridSnapID);
 
 	if (SnapEngine.HaveActiveSnap())
 	{
@@ -977,6 +1000,17 @@ float UDrawPolygonTool::FindInteractiveHeightDistance(const FInputDeviceRay& Cli
 
 	if (NearestHitDist < TNumericLimits<float>::Max())
 	{
+		if ( bIgnoreSnappingToggle == false && SnapProperties->bEnableSnapping && SnapProperties->bSnapToWorldGrid )
+		{
+			FVector3d GridPosWorld;
+			if (ToolSceneQueriesUtil::FindWorldGridSnapPoint(this, (FVector3d)NearestHitFrameWorld.Origin, GridPosWorld))
+			{
+				NearestHitFrameWorld.Origin = (FVector3f)GridPosWorld;
+				FVector3d LocalPos = PreviewHeightFrame.ToFramePoint((FVector3d)NearestHitFrameWorld.Origin);
+				NearestHitHeight = (float)LocalPos.Z;
+			}
+		}
+
 		this->HitPosFrameWorld = NearestHitFrameWorld;
 		return NearestHitHeight;
 	}
@@ -1011,7 +1045,18 @@ void UDrawPolygonTool::SetDrawPlaneFromWorldPos(const FVector& Position, const F
 void UDrawPolygonTool::PlaneTransformChanged(UTransformProxy* Proxy, FTransform Transform)
 {
 	DrawPlaneOrientation = Transform.GetRotation();
+	
 	DrawPlaneOrigin = Transform.GetLocation();
+
+	if (bIgnoreSnappingToggle == false && SnapProperties->bEnableSnapping && SnapProperties->bSnapToWorldGrid)
+	{
+		FVector3d GridSnapPos;
+		if (ToolSceneQueriesUtil::FindWorldGridSnapPoint(this, DrawPlaneOrigin, GridSnapPos))
+		{
+			DrawPlaneOrigin = GridSnapPos;
+		}
+	}
+
 	SnapEngine.Plane = FFrame3d((FVector3d)DrawPlaneOrigin, (FQuaterniond)DrawPlaneOrientation);
 }
 
diff --git a/Engine/Plugins/Experimental/MeshModelingToolset/Source/MeshModelingTools/Private/DynamicMeshSculptTool.cpp b/Engine/Plugins/Experimental/MeshModelingToolset/Source/MeshModelingTools/Private/DynamicMeshSculptTool.cpp
--- a/Engine/Plugins/Experimental/MeshModelingToolset/Source/MeshModelingTools/Private/DynamicMeshSculptTool.cpp
+++ b/Engine/Plugins/Experimental/MeshModelingToolset/Source/MeshModelingTools/Private/DynamicMeshSculptTool.cpp
@@ -15,6 +15,7 @@
 #include "Drawing/MeshDebugDrawing.h"
 #include "PreviewMesh.h"
 #include "ToolSetupUtil.h"
+#include "ToolSceneQueriesUtil.h"
 
 #include "Changes/MeshVertexChange.h"
 #include "Changes/MeshChange.h"
@@ -2040,14 +2041,10 @@ void UDynamicMeshSculptTool::UpdateFixedSculptPlanePosition(const FVector& Posit
 
 	if (GizmoProperties->bSnapToGrid)
 	{
-		FSceneSnapQueryRequest Request;
-		Request.RequestType = ESceneSnapQueryType::Position;
-		Request.TargetTypes = ESceneSnapQueryTargetType::Grid;
-		Request.Position = GizmoProperties->Position;
-		TArray<FSceneSnapQueryResult> Results;
-		if (GetToolManager()->GetContextQueriesAPI()->ExecuteSceneSnapQuery(Request, Results) && Results.Num() > 0)
+		FVector3d GridSnapPos;
+		if (ToolSceneQueriesUtil::FindWorldGridSnapPoint(this, GizmoProperties->Position, GridSnapPos))
 		{
-			GizmoProperties->Position = Results[0].Position;
+			GizmoProperties->Position = GridSnapPos;
 		}
 	}
 }
diff --git a/Engine/Plugins/Experimental/MeshModelingToolset/Source/MeshModelingTools/Private/EditMeshPolygonsTool.cpp b/Engine/Plugins/Experimental/MeshModelingToolset/Source/MeshModelingTools/Private/EditMeshPolygonsTool.cpp
--- a/Engine/Plugins/Experimental/MeshModelingToolset/Source/MeshModelingTools/Private/EditMeshPolygonsTool.cpp
+++ b/Engine/Plugins/Experimental/MeshModelingToolset/Source/MeshModelingTools/Private/EditMeshPolygonsTool.cpp
@@ -51,6 +51,8 @@ UPolyEditTransformProperties::UPolyEditTransformProperties()
 	bSelectVertices = true;
 	bSelectFaces = true;
 	bSelectEdges = true;
+	bShowWireframe = false;
+	bSnapToWorldGrid = false;
 	PolygonMode = EPolygonGroupMode::KeepInputPolygons;
 	PolygonGroupingAngleThreshold = .5;
 }
@@ -1265,10 +1267,21 @@ void UEditMeshPolygonsTool::ComputeUpdate_Translate()
 	const bool bIsLaplacian = (DeformationStrategy == EGroupTopologyDeformationStrategy::Laplacian);
 	FGroupTopologyDeformer& SelectedDeformer = (bIsLaplacian) ? LaplacianDeformer : LinearDeformer;
 
+	TFunction<FVector3d(const FVector3d&)> PointConstraintFunc = nullptr;
+	if (TransformProps->bSnapToWorldGrid 
+		&& GetToolManager()->GetContextQueriesAPI()->GetCurrentCoordinateSystem() == EToolContextCoordinateSystem::World)
+	{
+		PointConstraintFunc = [&](const FVector3d& Pos)
+		{
+			FVector3d GridSnapPos;
+			return ToolSceneQueriesUtil::FindWorldGridSnapPoint(this, Pos, GridSnapPos) ? GridSnapPos : Pos;
+		};
+	}
+
 	FTransform Transform = ComponentTarget->GetWorldTransform();
 	FVector NewHitPosWorld = LastHitPosWorld;
 	FVector3d SnappedPoint;
-	if (QuickAxisTranslater.UpdateSnap(UpdateRay, SnappedPoint))
+	if (QuickAxisTranslater.UpdateSnap(UpdateRay, SnappedPoint, PointConstraintFunc))
 	{
 		NewHitPosWorld = SnappedPoint;
 	}
@@ -1386,7 +1399,7 @@ void UEditMeshPolygonsTool::Render(IToolsContextRenderAPI* RenderAPI)
 	GetToolManager()->GetContextQueriesAPI()->GetCurrentViewState(CameraState);
 	GetActiveQuickTransformer()->UpdateCameraState(CameraState);
 
-	DynamicMeshComponent->bExplicitShowWireframe = TransformProps->bWireframe;
+	DynamicMeshComponent->bExplicitShowWireframe = TransformProps->bShowWireframe;
 	FDynamicMesh3* TargetMesh = DynamicMeshComponent->GetMesh();
 
 	PolyEdgesRenderer.BeginFrame(RenderAPI, CameraState);
diff --git a/Engine/Plugins/Experimental/MeshModelingToolset/Source/MeshModelingTools/Public/DrawPolygonTool.h b/Engine/Plugins/Experimental/MeshModelingToolset/Source/MeshModelingTools/Public/DrawPolygonTool.h
--- a/Engine/Plugins/Experimental/MeshModelingToolset/Source/MeshModelingTools/Public/DrawPolygonTool.h
+++ b/Engine/Plugins/Experimental/MeshModelingToolset/Source/MeshModelingTools/Public/DrawPolygonTool.h
@@ -136,6 +136,9 @@ public:
 	UPROPERTY(EditAnywhere, NonTransactional, Category = Snapping)
 	bool bEnableSnapping = true;
 
+	UPROPERTY(EditAnywhere, NonTransactional, Category = Snapping, Meta = (EditCondition = "bEnableSnapping"))
+	bool bSnapToWorldGrid = false;
+
 	UPROPERTY(EditAnywhere, NonTransactional, Category = Snapping, Meta = (EditCondition = "bEnableSnapping"))
 	bool bSnapToVertices = true;
 
@@ -309,6 +312,7 @@ protected:
 	bool bIgnoreSnappingToggle = false;		// toggled by hotkey (shift)
 	FPointPlanarSnapSolver SnapEngine;
 	ToolSceneQueriesUtil::FSnapGeometry LastSnapGeometry;
+	FVector3d LastGridSnapPoint;
 
 	void GetPolygonParametersFromFixedPoints(const TArray<FVector>& FixedPoints, FVector2f& FirstReferencePt, FVector2f& BoxSize, float& YSign, float& AngleRad);
 	void GenerateFixedPolygon(const TArray<FVector>& FixedPoints, TArray<FVector>& VerticesOut, TArray<TArray<FVector>>& HolesVerticesOut);
diff --git a/Engine/Plugins/Experimental/MeshModelingToolset/Source/MeshModelingTools/Public/EditMeshPolygonsTool.h b/Engine/Plugins/Experimental/MeshModelingToolset/Source/MeshModelingTools/Public/EditMeshPolygonsTool.h
--- a/Engine/Plugins/Experimental/MeshModelingToolset/Source/MeshModelingTools/Public/EditMeshPolygonsTool.h
+++ b/Engine/Plugins/Experimental/MeshModelingToolset/Source/MeshModelingTools/Public/EditMeshPolygonsTool.h
@@ -123,8 +123,11 @@ public:
 	UPROPERTY(EditAnywhere, Category = Options)
 	bool bSelectVertices;
 
-	UPROPERTY(EditAnywhere, Category = Options, meta = (DisplayName = "Show Triangle Mesh Wireframe"))
-	bool bWireframe;
+	UPROPERTY(EditAnywhere, Category = Options)
+	bool bShowWireframe;
+
+	UPROPERTY(EditAnywhere, Category = Options)
+	bool bSnapToWorldGrid;
 
 	UPROPERTY(EditAnywhere, Category = Options)
 	EPolygonGroupMode PolygonMode;
diff --git a/Engine/Plugins/Experimental/MeshModelingToolset/Source/ModelingComponents/Private/PreviewMesh.cpp b/Engine/Plugins/Experimental/MeshModelingToolset/Source/ModelingComponents/Private/PreviewMesh.cpp
--- a/Engine/Plugins/Experimental/MeshModelingToolset/Source/ModelingComponents/Private/PreviewMesh.cpp
+++ b/Engine/Plugins/Experimental/MeshModelingToolset/Source/ModelingComponents/Private/PreviewMesh.cpp
@@ -200,8 +200,6 @@ bool UPreviewMesh::FindRayIntersection(const FRay3d& WorldRay, FHitResult& HitOu
 			FIntrRay3Triangle3d Query(LocalRay, Triangle);
 			Query.Find();
 
-			//auto Query = MeshQueries::RayTriangleIntersection(UseMesh, HitTID, LocalRay);
-
 			HitOut.FaceIndex = HitTriID;
 			HitOut.Distance = Query.RayParameter;
 			HitOut.Normal = Transform.TransformVectorNoScale(UseMesh->GetTriNormal(HitTriID));
