ModelingTools: fix uninitialized variable
#rnx

diff --git a/Engine/Plugins/Experimental/MeshModelingToolset/Source/ModelingComponents/Private/Transforms/QuickAxisTranslater.cpp b/Engine/Plugins/Experimental/MeshModelingToolset/Source/ModelingComponents/Private/Transforms/QuickAxisTranslater.cpp
--- a/Engine/Plugins/Experimental/MeshModelingToolset/Source/ModelingComponents/Private/Transforms/QuickAxisTranslater.cpp
+++ b/Engine/Plugins/Experimental/MeshModelingToolset/Source/ModelingComponents/Private/Transforms/QuickAxisTranslater.cpp
@@ -189,7 +189,7 @@ void FQuickAxisTranslater::UpdateSnapAxes()
 bool FQuickAxisTranslater::UpdateSnap(const FRay3d& Ray, FVector3d& SnapPointOut,
 	TFunction<FVector3d(const FVector3d&)> PositionConstraintFunc)
 {
-	bool bSetConstraintFunc;
+	bool bSetConstraintFunc = false;
 	if (PositionConstraintFunc != nullptr)
 	{
 		check(MoveAxisSolver.PointConstraintFunc == nullptr);
