MeshModelingTools : Flesh out interface for commiting meshes

#rnx

diff --git a/Engine/Plugins/Experimental/MeshModelingToolset/Source/MeshModelingTools/Private/ConvertToPolygonsTool.cpp b/Engine/Plugins/Experimental/MeshModelingToolset/Source/MeshModelingTools/Private/ConvertToPolygonsTool.cpp
--- a/Engine/Plugins/Experimental/MeshModelingToolset/Source/MeshModelingTools/Private/ConvertToPolygonsTool.cpp
+++ b/Engine/Plugins/Experimental/MeshModelingToolset/Source/MeshModelingTools/Private/ConvertToPolygonsTool.cpp
@@ -90,9 +90,9 @@ void UConvertToPolygonsTool::Shutdown(EToolShutdownType ShutdownType)
 	if (ShutdownType == EToolShutdownType::Accept)
 	{
 		GetToolManager()->BeginUndoTransaction(LOCTEXT("ConvertToPolygonsToolTransactionName", "Convert to Polygons"));
-		ComponentTarget->CommitMesh([=](FMeshDescription* MeshDescription)
+		ComponentTarget->CommitMesh([=](const FPrimitiveComponentTarget::FCommitParams& CommitParams)
 		{
-			ConvertToPolygons(MeshDescription);
+			ConvertToPolygons(CommitParams.MeshDescription);
 		});
 		GetToolManager()->EndUndoTransaction();
 	}
diff --git a/Engine/Plugins/Experimental/MeshModelingToolset/Source/MeshModelingTools/Private/DisplaceMeshTool.cpp b/Engine/Plugins/Experimental/MeshModelingToolset/Source/MeshModelingTools/Private/DisplaceMeshTool.cpp
--- a/Engine/Plugins/Experimental/MeshModelingToolset/Source/MeshModelingTools/Private/DisplaceMeshTool.cpp
+++ b/Engine/Plugins/Experimental/MeshModelingToolset/Source/MeshModelingTools/Private/DisplaceMeshTool.cpp
@@ -475,9 +475,9 @@ void UDisplaceMeshTool::Shutdown(EToolShutdownType ShutdownType)
 		{
 			// this block bakes the modified DynamicMeshComponent back into the StaticMeshComponent inside an undo transaction
 			GetToolManager()->BeginUndoTransaction(LOCTEXT("DisplaceMeshToolTransactionName", "Displace Mesh"));
-			ComponentTarget->CommitMesh([=](FMeshDescription* MeshDescription)
+			ComponentTarget->CommitMesh([=](const FPrimitiveComponentTarget::FCommitParams& CommitParams)
 			{
-				DynamicMeshComponent->Bake(MeshDescription, Subdivisions > 0);
+				DynamicMeshComponent->Bake(CommitParams.MeshDescription, Subdivisions > 0);
 			});
 			GetToolManager()->EndUndoTransaction();
 		}
diff --git a/Engine/Plugins/Experimental/MeshModelingToolset/Source/MeshModelingTools/Private/DynamicMeshSculptTool.cpp b/Engine/Plugins/Experimental/MeshModelingToolset/Source/MeshModelingTools/Private/DynamicMeshSculptTool.cpp
--- a/Engine/Plugins/Experimental/MeshModelingToolset/Source/MeshModelingTools/Private/DynamicMeshSculptTool.cpp
+++ b/Engine/Plugins/Experimental/MeshModelingToolset/Source/MeshModelingTools/Private/DynamicMeshSculptTool.cpp
@@ -287,9 +287,9 @@ void UDynamicMeshSculptTool::Shutdown(EToolShutdownType ShutdownType)
 
 			// this block bakes the modified DynamicMeshComponent back into the StaticMeshComponent inside an undo transaction
 			GetToolManager()->BeginUndoTransaction(LOCTEXT("SculptMeshToolTransactionName", "Sculpt Mesh"));
-			ComponentTarget->CommitMesh([=](FMeshDescription* MeshDescription)
+			ComponentTarget->CommitMesh([=](const FPrimitiveComponentTarget::FCommitParams& CommitParams)
 			{
-				DynamicMeshComponent->Bake(MeshDescription, bHaveRemeshed);
+				DynamicMeshComponent->Bake(CommitParams.MeshDescription, bHaveRemeshed);
 			});
 			GetToolManager()->EndUndoTransaction();
 		}
diff --git a/Engine/Plugins/Experimental/MeshModelingToolset/Source/MeshModelingTools/Private/EditMeshPolygonsTool.cpp b/Engine/Plugins/Experimental/MeshModelingToolset/Source/MeshModelingTools/Private/EditMeshPolygonsTool.cpp
--- a/Engine/Plugins/Experimental/MeshModelingToolset/Source/MeshModelingTools/Private/EditMeshPolygonsTool.cpp
+++ b/Engine/Plugins/Experimental/MeshModelingToolset/Source/MeshModelingTools/Private/EditMeshPolygonsTool.cpp
@@ -718,11 +718,11 @@ void UEditMeshPolygonsTool::Shutdown(EToolShutdownType ShutdownType)
 		{
 			// this block bakes the modified DynamicMeshComponent back into the StaticMeshComponent inside an undo transaction
 			GetToolManager()->BeginUndoTransaction(LOCTEXT("EditMeshPolygonsToolTransactionName", "Deform Mesh"));
-			ComponentTarget->CommitMesh([=](FMeshDescription* MeshDescription)
+			ComponentTarget->CommitMesh([=](const FPrimitiveComponentTarget::FCommitParams& CommitParams)
 			{
 				FConversionToMeshDescriptionOptions ConversionOptions;
 				ConversionOptions.bSetPolyGroups = false; // don't save polygroups, as we may change these temporarily in this tool just to get a different edit effect
-				DynamicMeshComponent->Bake(MeshDescription, false, ConversionOptions);
+				DynamicMeshComponent->Bake(CommitParams.MeshDescription, false, ConversionOptions);
 			});
 			GetToolManager()->EndUndoTransaction();
 		}
diff --git a/Engine/Plugins/Experimental/MeshModelingToolset/Source/MeshModelingTools/Private/EditNormalsTool.cpp b/Engine/Plugins/Experimental/MeshModelingToolset/Source/MeshModelingTools/Private/EditNormalsTool.cpp
--- a/Engine/Plugins/Experimental/MeshModelingToolset/Source/MeshModelingTools/Private/EditNormalsTool.cpp
+++ b/Engine/Plugins/Experimental/MeshModelingToolset/Source/MeshModelingTools/Private/EditNormalsTool.cpp
@@ -262,19 +262,19 @@ void UEditNormalsTool::GenerateAsset(const TArray<FDynamicMeshOpResult>& Results
 	for (int32 ComponentIdx = 0; ComponentIdx < ComponentTargets.Num(); ComponentIdx++)
 	{
 		check(Results[ComponentIdx].Mesh.Get() != nullptr);
-		ComponentTargets[ComponentIdx]->CommitMesh([&Results, &ComponentIdx, this](FMeshDescription* MeshDescription)
+		ComponentTargets[ComponentIdx]->CommitMesh([&Results, &ComponentIdx, this](const FPrimitiveComponentTarget::FCommitParams& CommitParams)
 		{
 			FDynamicMeshToMeshDescription Converter;
 			
 			if (BasicProperties->WillTopologyChange())
 			{
 				// full conversion if normal topology changed or faces were inverted
-				Converter.Convert(Results[ComponentIdx].Mesh.Get(), *MeshDescription);
+				Converter.Convert(Results[ComponentIdx].Mesh.Get(), *CommitParams.MeshDescription);
 			}
 			else
 			{
 				// otherwise just copy attributes
-				Converter.UpdateAttributes(Results[ComponentIdx].Mesh.Get(), *MeshDescription, true, false);
+				Converter.UpdateAttributes(Results[ComponentIdx].Mesh.Get(), *CommitParams.MeshDescription, true, false);
 			}
 		});
 	}
diff --git a/Engine/Plugins/Experimental/MeshModelingToolset/Source/MeshModelingTools/Private/MeshSelectionTool.cpp b/Engine/Plugins/Experimental/MeshModelingToolset/Source/MeshModelingTools/Private/MeshSelectionTool.cpp
--- a/Engine/Plugins/Experimental/MeshModelingToolset/Source/MeshModelingTools/Private/MeshSelectionTool.cpp
+++ b/Engine/Plugins/Experimental/MeshModelingToolset/Source/MeshModelingTools/Private/MeshSelectionTool.cpp
@@ -145,9 +145,9 @@
 		// this block bakes the modified DynamicMeshComponent back into the StaticMeshComponent inside an undo transaction
 		GetToolManager()->BeginUndoTransaction(LOCTEXT("MeshSelectionToolTransactionName", "Edit Mesh"));
 
-		ComponentTarget->CommitMesh([=](FMeshDescription* MeshDescription)
+		ComponentTarget->CommitMesh([=](const FPrimitiveComponentTarget::FCommitParams& CommitParams)
 		{
-			PreviewMesh->Bake(MeshDescription, true);
+			PreviewMesh->Bake(CommitParams.MeshDescription, true);
 		});
 		GetToolManager()->EndUndoTransaction();
 	}
@@ -560,7 +560,7 @@
 }
 
 
-void UMeshSelectionTool::Tick(float DeltaTime)
+void UMeshSelectionTool::Tick(float DeltaTime)
 {
 	UDynamicMeshBrushTool::Tick(DeltaTime);
 
diff --git a/Engine/Plugins/Experimental/MeshModelingToolset/Source/MeshModelingTools/Private/MeshSpaceDeformerTool.cpp b/Engine/Plugins/Experimental/MeshModelingToolset/Source/MeshModelingTools/Private/MeshSpaceDeformerTool.cpp
--- a/Engine/Plugins/Experimental/MeshModelingToolset/Source/MeshModelingTools/Private/MeshSpaceDeformerTool.cpp
+++ b/Engine/Plugins/Experimental/MeshModelingToolset/Source/MeshModelingTools/Private/MeshSpaceDeformerTool.cpp
@@ -292,10 +292,10 @@ void UMeshSpaceDeformerTool::Shutdown(EToolShutdownType ShutdownType)
 			FDynamicMesh3* DynamicMeshResult = Result.Mesh.Get();
 			check(DynamicMeshResult != nullptr);
 
-			ComponentTarget->CommitMesh([DynamicMeshResult](FMeshDescription* MeshDescription)
+			ComponentTarget->CommitMesh([DynamicMeshResult](const FPrimitiveComponentTarget::FCommitParams& CommitParams)
 			{
 				FDynamicMeshToMeshDescription Converter;
-				Converter.Convert(DynamicMeshResult, *MeshDescription);
+				Converter.Convert(DynamicMeshResult, *CommitParams.MeshDescription);
 			});
 
 
diff --git a/Engine/Plugins/Experimental/MeshModelingToolset/Source/MeshModelingTools/Private/PlaneCutTool.cpp b/Engine/Plugins/Experimental/MeshModelingToolset/Source/MeshModelingTools/Private/PlaneCutTool.cpp
--- a/Engine/Plugins/Experimental/MeshModelingToolset/Source/MeshModelingTools/Private/PlaneCutTool.cpp
+++ b/Engine/Plugins/Experimental/MeshModelingToolset/Source/MeshModelingTools/Private/PlaneCutTool.cpp
@@ -373,10 +373,10 @@ void UPlaneCutTool::GenerateAsset(const TArray<FDynamicMeshOpResult>& Results)
 	for (int OrigMeshIdx = 0; OrigMeshIdx < NumSourceMeshes; OrigMeshIdx++)
 	{
 		check(Results[OrigMeshIdx].Mesh.Get() != nullptr);
-		ComponentTargets[OrigMeshIdx]->CommitMesh([&Results, OrigMeshIdx](FMeshDescription* MeshDescription)
+		ComponentTargets[OrigMeshIdx]->CommitMesh([&Results, OrigMeshIdx](const FPrimitiveComponentTarget::FCommitParams& CommitParams)
 		{
 			FDynamicMeshToMeshDescription Converter;
-			Converter.Convert(Results[OrigMeshIdx].Mesh.Get(), *MeshDescription);
+			Converter.Convert(Results[OrigMeshIdx].Mesh.Get(), *CommitParams.MeshDescription);
 		});
 	}
 
diff --git a/Engine/Plugins/Experimental/MeshModelingToolset/Source/MeshModelingTools/Private/RemeshMeshTool.cpp b/Engine/Plugins/Experimental/MeshModelingToolset/Source/MeshModelingTools/Private/RemeshMeshTool.cpp
--- a/Engine/Plugins/Experimental/MeshModelingToolset/Source/MeshModelingTools/Private/RemeshMeshTool.cpp
+++ b/Engine/Plugins/Experimental/MeshModelingToolset/Source/MeshModelingTools/Private/RemeshMeshTool.cpp
@@ -243,12 +243,12 @@ void URemeshMeshTool::GenerateAsset(const FDynamicMeshOpResult& Result)
 	GetToolManager()->BeginUndoTransaction(LOCTEXT("RemeshMeshToolTransactionName", "Remesh Mesh"));
 
 	check(Result.Mesh.Get() != nullptr);
-	ComponentTarget->CommitMesh([&Result](FMeshDescription* MeshDescription)
+	ComponentTarget->CommitMesh([&Result](const FPrimitiveComponentTarget::FCommitParams& CommitParams)
 	{
 		FDynamicMeshToMeshDescription Converter;
 
 		// full conversion if normal topology changed or faces were inverted
-		Converter.Convert(Result.Mesh.Get(), *MeshDescription);
+		Converter.Convert(Result.Mesh.Get(), *CommitParams.MeshDescription);
 	});
 
 	GetToolManager()->EndUndoTransaction();
diff --git a/Engine/Plugins/Experimental/MeshModelingToolset/Source/MeshModelingTools/Private/RemoveOccludedTrianglesTool.cpp b/Engine/Plugins/Experimental/MeshModelingToolset/Source/MeshModelingTools/Private/RemoveOccludedTrianglesTool.cpp
--- a/Engine/Plugins/Experimental/MeshModelingToolset/Source/MeshModelingTools/Private/RemoveOccludedTrianglesTool.cpp
+++ b/Engine/Plugins/Experimental/MeshModelingToolset/Source/MeshModelingTools/Private/RemoveOccludedTrianglesTool.cpp
@@ -303,11 +303,11 @@ void URemoveOccludedTrianglesTool::GenerateAsset(const TArray<FDynamicMeshOpResu
 	for (int32 ComponentIdx = 0; ComponentIdx < ComponentTargets.Num(); ComponentIdx++)
 	{
 		check(Results[ComponentIdx].Mesh.Get() != nullptr);
-		ComponentTargets[ComponentIdx]->CommitMesh([&Results, &ComponentIdx, this](FMeshDescription* MeshDescription)
+		ComponentTargets[ComponentIdx]->CommitMesh([&Results, &ComponentIdx, this](const FPrimitiveComponentTarget::FCommitParams& CommitParams)
 		{
 			FDynamicMeshToMeshDescription Converter;
 			
-			Converter.Convert(Results[ComponentIdx].Mesh.Get(), *MeshDescription);
+			Converter.Convert(Results[ComponentIdx].Mesh.Get(), *CommitParams.MeshDescription);
 		});
 	}
 
diff --git a/Engine/Plugins/Experimental/MeshModelingToolset/Source/MeshModelingTools/Private/SmoothMeshTool.cpp b/Engine/Plugins/Experimental/MeshModelingToolset/Source/MeshModelingTools/Private/SmoothMeshTool.cpp
--- a/Engine/Plugins/Experimental/MeshModelingToolset/Source/MeshModelingTools/Private/SmoothMeshTool.cpp
+++ b/Engine/Plugins/Experimental/MeshModelingToolset/Source/MeshModelingTools/Private/SmoothMeshTool.cpp
@@ -123,10 +123,10 @@ void USmoothMeshTool::Shutdown(EToolShutdownType ShutdownType)
 			FDynamicMesh3* DynamicMeshResult = Result.Mesh.Get();
 			check(DynamicMeshResult != nullptr);
 
-			ComponentTarget->CommitMesh([DynamicMeshResult](FMeshDescription* MeshDescription)
+			ComponentTarget->CommitMesh([DynamicMeshResult](const FPrimitiveComponentTarget::FCommitParams& CommitParams)
 			{
 				FDynamicMeshToMeshDescription Converter;
-				Converter.Convert(DynamicMeshResult, *MeshDescription);
+				Converter.Convert(DynamicMeshResult, *CommitParams.MeshDescription);
 			});
 
 
diff --git a/Engine/Plugins/Experimental/MeshModelingToolset/Source/MeshModelingTools/Private/UVProjectionTool.cpp b/Engine/Plugins/Experimental/MeshModelingToolset/Source/MeshModelingTools/Private/UVProjectionTool.cpp
--- a/Engine/Plugins/Experimental/MeshModelingToolset/Source/MeshModelingTools/Private/UVProjectionTool.cpp
+++ b/Engine/Plugins/Experimental/MeshModelingToolset/Source/MeshModelingTools/Private/UVProjectionTool.cpp
@@ -341,13 +341,13 @@ void UUVProjectionTool::GenerateAsset(const TArray<FDynamicMeshOpResult>& Result
 	for (int32 ComponentIdx = 0; ComponentIdx < ComponentTargets.Num(); ComponentIdx++)
 	{
 		check(Results[ComponentIdx].Mesh.Get() != nullptr);
-		ComponentTargets[ComponentIdx]->CommitMesh([&Results, &ComponentIdx, this](FMeshDescription* MeshDescription)
+		ComponentTargets[ComponentIdx]->CommitMesh([&Results, &ComponentIdx, this](const FPrimitiveComponentTarget::FCommitParams& CommitParams)
 		{
 			FDynamicMeshToMeshDescription Converter;
 			//if (??) // TODO check if UV topology changed?  or remove all traces of this if statement (may be safe to assume it almost always changes w/ a uv projection op)
 			{
 				// full conversion if topology changed
-				Converter.Convert(Results[ComponentIdx].Mesh.Get(), *MeshDescription);
+				Converter.Convert(Results[ComponentIdx].Mesh.Get(), *CommitParams.MeshDescription);
 			}
 			//else
 			//{
diff --git a/Engine/Plugins/Experimental/MeshModelingToolset/Source/MeshModelingTools/Private/WeldMeshEdgesTool.cpp b/Engine/Plugins/Experimental/MeshModelingToolset/Source/MeshModelingTools/Private/WeldMeshEdgesTool.cpp
--- a/Engine/Plugins/Experimental/MeshModelingToolset/Source/MeshModelingTools/Private/WeldMeshEdgesTool.cpp
+++ b/Engine/Plugins/Experimental/MeshModelingToolset/Source/MeshModelingTools/Private/WeldMeshEdgesTool.cpp
@@ -91,9 +91,9 @@ void UWeldMeshEdgesTool::Shutdown(EToolShutdownType ShutdownType)
 		{
 			// this block bakes the modified DynamicMeshComponent back into the StaticMeshComponent inside an undo transaction
 			GetToolManager()->BeginUndoTransaction(LOCTEXT("WeldMeshEdgesToolTransactionName", "Remesh Mesh"));
-			ComponentTarget->CommitMesh([=](FMeshDescription* MeshDescription)
+			ComponentTarget->CommitMesh([=](const FPrimitiveComponentTarget::FCommitParams& CommitParams)
 			{
-				DynamicMeshComponent->Bake(MeshDescription, true);
+				DynamicMeshComponent->Bake(CommitParams.MeshDescription, true);
 			});
 			GetToolManager()->EndUndoTransaction();
 		}
diff --git a/Engine/Plugins/Experimental/MeshModelingToolset/Source/MeshModelingToolsEditorOnly/Private/AttributeEditorTool.cpp b/Engine/Plugins/Experimental/MeshModelingToolset/Source/MeshModelingToolsEditorOnly/Private/AttributeEditorTool.cpp
--- a/Engine/Plugins/Experimental/MeshModelingToolset/Source/MeshModelingToolsEditorOnly/Private/AttributeEditorTool.cpp
+++ b/Engine/Plugins/Experimental/MeshModelingToolset/Source/MeshModelingToolsEditorOnly/Private/AttributeEditorTool.cpp
@@ -131,20 +131,20 @@ void UAttributeEditorTool::GenerateAssets()
 
 	for (TUniquePtr<FPrimitiveComponentTarget>& ComponentTarget : ComponentTargets)
 	{
-		ComponentTarget->CommitMesh([this](FMeshDescription* MeshDescription)
+		ComponentTarget->CommitMesh([this](const FPrimitiveComponentTarget::FCommitParams& CommitParams)
 		{
 			if (RemovalProperties->bClearNormals)
 			{
-				TEdgeAttributesRef<bool> EdgeHardnesses = MeshDescription->EdgeAttributes().GetAttributesRef<bool>(MeshAttribute::Edge::IsHard);
+				TEdgeAttributesRef<bool> EdgeHardnesses = CommitParams.MeshDescription->EdgeAttributes().GetAttributesRef<bool>(MeshAttribute::Edge::IsHard);
 				if (EdgeHardnesses.IsValid())
 				{
-					for (FEdgeID ElID : MeshDescription->Edges().GetElementIDs())
+					for (FEdgeID ElID : CommitParams.MeshDescription->Edges().GetElementIDs())
 					{
 						EdgeHardnesses[ElID] = false;
 					}
 				}
-				FMeshDescriptionOperations::CreatePolygonNTB(*MeshDescription, FMathf::Epsilon);
-				FMeshDescriptionOperations::RecomputeNormalsAndTangentsIfNeeded(*MeshDescription, FMeshDescriptionOperations::ETangentOptions::UseWeightedAreaAndAngle, true);
+				FMeshDescriptionOperations::CreatePolygonNTB(*CommitParams.MeshDescription, FMathf::Epsilon);
+				FMeshDescriptionOperations::RecomputeNormalsAndTangentsIfNeeded(*CommitParams.MeshDescription, FMeshDescriptionOperations::ETangentOptions::UseWeightedAreaAndAngle, true);
 			}
 			bool RemoveLayers[8] =
 			{
@@ -158,13 +158,13 @@ void UAttributeEditorTool::GenerateAssets()
 				RemovalProperties->bClearUVLayer7
 			};
 			TVertexInstanceAttributesRef<FVector2D> InstanceUVs =
-				MeshDescription->VertexInstanceAttributes().GetAttributesRef<FVector2D>(MeshAttribute::VertexInstance::TextureCoordinate);
-			const FVertexInstanceArray& Instances = MeshDescription->VertexInstances();
+				CommitParams.MeshDescription->VertexInstanceAttributes().GetAttributesRef<FVector2D>(MeshAttribute::VertexInstance::TextureCoordinate);
+			const FVertexInstanceArray& Instances = CommitParams.MeshDescription->VertexInstances();
 			for (int LayerIndex = 7; LayerIndex >= 0; LayerIndex--)
 			{
 				if (RemoveLayers[LayerIndex] && LayerIndex < InstanceUVs.GetNumIndices())
 				{
-					if (!FMeshDescriptionOperations::RemoveUVChannel(*MeshDescription, LayerIndex))
+					if (!FMeshDescriptionOperations::RemoveUVChannel(*CommitParams.MeshDescription, LayerIndex))
 					{
 						for (const FVertexInstanceID& ElID : Instances.GetElementIDs())
 						{
diff --git a/Engine/Plugins/Experimental/MeshModelingToolset/Source/MeshModelingToolsEditorOnly/Private/ParameterizeMeshTool.cpp b/Engine/Plugins/Experimental/MeshModelingToolset/Source/MeshModelingToolsEditorOnly/Private/ParameterizeMeshTool.cpp
--- a/Engine/Plugins/Experimental/MeshModelingToolset/Source/MeshModelingToolsEditorOnly/Private/ParameterizeMeshTool.cpp
+++ b/Engine/Plugins/Experimental/MeshModelingToolset/Source/MeshModelingToolsEditorOnly/Private/ParameterizeMeshTool.cpp
@@ -141,10 +141,10 @@ void UParameterizeMeshTool::Shutdown(EToolShutdownType ShutdownType)
 		check(DynamicMeshResult != nullptr);
 		GetToolManager()->BeginUndoTransaction(LOCTEXT("ParameterizeMesh", "Parameterize Mesh"));
 
-		ComponentTarget->CommitMesh([DynamicMeshResult](FMeshDescription* MeshDescription)
+		ComponentTarget->CommitMesh([DynamicMeshResult](const FPrimitiveComponentTarget::FCommitParams& CommitParams)
 		{
 			FDynamicMeshToMeshDescription Converter;
-			Converter.Convert(DynamicMeshResult, *MeshDescription);
+			Converter.Convert(DynamicMeshResult, *CommitParams.MeshDescription);
 		});
 		GetToolManager()->EndUndoTransaction();
 	}
diff --git a/Engine/Plugins/Experimental/MeshModelingToolset/Source/MeshModelingToolsEditorOnly/Private/PolygonOnMeshTool.cpp b/Engine/Plugins/Experimental/MeshModelingToolset/Source/MeshModelingToolsEditorOnly/Private/PolygonOnMeshTool.cpp
--- a/Engine/Plugins/Experimental/MeshModelingToolset/Source/MeshModelingToolsEditorOnly/Private/PolygonOnMeshTool.cpp
+++ b/Engine/Plugins/Experimental/MeshModelingToolset/Source/MeshModelingToolsEditorOnly/Private/PolygonOnMeshTool.cpp
@@ -314,10 +314,10 @@ void UPolygonOnMeshTool::GenerateAsset(const TArray<FDynamicMeshOpResult>& Resul
 	// TODO: options to support other choices re what should be a new actor
 	check(Results.Num() > 0);
 	check(Results[0].Mesh.Get() != nullptr);
-	ComponentTarget->CommitMesh([&Results](FMeshDescription* MeshDescription)
+	ComponentTarget->CommitMesh([&Results](const FPrimitiveComponentTarget::FCommitParams& CommitParams)
 	{
 		FDynamicMeshToMeshDescription Converter;
-		Converter.Convert(Results[0].Mesh.Get(), *MeshDescription);
+		Converter.Convert(Results[0].Mesh.Get(), *CommitParams.MeshDescription);
 	});
 
 	// currently nothing generates more than the single result for this tool
diff --git a/Engine/Plugins/Experimental/MeshModelingToolset/Source/MeshModelingToolsEditorOnly/Private/SimplifyMeshTool.cpp b/Engine/Plugins/Experimental/MeshModelingToolset/Source/MeshModelingToolsEditorOnly/Private/SimplifyMeshTool.cpp
--- a/Engine/Plugins/Experimental/MeshModelingToolset/Source/MeshModelingToolsEditorOnly/Private/SimplifyMeshTool.cpp
+++ b/Engine/Plugins/Experimental/MeshModelingToolset/Source/MeshModelingToolsEditorOnly/Private/SimplifyMeshTool.cpp
@@ -240,12 +240,12 @@ void USimplifyMeshTool::GenerateAsset(const FDynamicMeshOpResult& Result)
 	GetToolManager()->BeginUndoTransaction(LOCTEXT("SimplifyMeshToolTransactionName", "Simplify Mesh"));
 
 	check(Result.Mesh.Get() != nullptr);
-	ComponentTarget->CommitMesh([&Result](FMeshDescription* MeshDescription)
+	ComponentTarget->CommitMesh([&Result](const FPrimitiveComponentTarget::FCommitParams& CommitParams)
 	{
 		FDynamicMeshToMeshDescription Converter;
 
 		// full conversion if normal topology changed or faces were inverted
-		Converter.Convert(Result.Mesh.Get(), *MeshDescription);
+		Converter.Convert(Result.Mesh.Get(), *CommitParams.MeshDescription);
 	});
 
 	GetToolManager()->EndUndoTransaction();
diff --git a/Engine/Plugins/Runtime/ProceduralMeshComponent/Source/ProceduralMeshComponent/Private/ProceduralMeshBridge.cpp b/Engine/Plugins/Runtime/ProceduralMeshComponent/Source/ProceduralMeshComponent/Private/ProceduralMeshBridge.cpp
--- a/Engine/Plugins/Runtime/ProceduralMeshComponent/Source/ProceduralMeshComponent/Private/ProceduralMeshBridge.cpp
+++ b/Engine/Plugins/Runtime/ProceduralMeshComponent/Source/ProceduralMeshComponent/Private/ProceduralMeshBridge.cpp
@@ -40,6 +40,6 @@ FProceduralMeshComponentTarget::GetMesh()
 void
 FProceduralMeshComponentTarget::CommitMesh( const FCommitter& ModifyFunc )
 {
-	ModifyFunc(MeshDescription.Get());
+	ModifyFunc({MeshDescription.Get()});
 	MeshDescriptionToProcMesh( *MeshDescription, Cast<UProceduralMeshComponent>(Component) );
 }
diff --git a/Engine/Source/Editor/UnrealEd/Private/Tools/EditorComponentSourceFactory.cpp b/Engine/Source/Editor/UnrealEd/Private/Tools/EditorComponentSourceFactory.cpp
--- a/Engine/Source/Editor/UnrealEd/Private/Tools/EditorComponentSourceFactory.cpp
+++ b/Engine/Source/Editor/UnrealEd/Private/Tools/EditorComponentSourceFactory.cpp
@@ -28,7 +28,7 @@ FStaticMeshComponentTarget::CommitMesh( const FCommitter& Committer )
 	check(bSavedToTransactionBuffer);
 	FMeshDescription* MeshDescription = StaticMesh->GetMeshDescription(LODIndex);
 
-	Committer(MeshDescription);
+	Committer({MeshDescription});
 
 	StaticMesh->CommitMeshDescription(LODIndex);
 	StaticMesh->PostEditChange();
diff --git a/Engine/Source/Runtime/Experimental/InteractiveToolsFramework/Public/ComponentSourceInterfaces.h b/Engine/Source/Runtime/Experimental/InteractiveToolsFramework/Public/ComponentSourceInterfaces.h
--- a/Engine/Source/Runtime/Experimental/InteractiveToolsFramework/Public/ComponentSourceInterfaces.h
+++ b/Engine/Source/Runtime/Experimental/InteractiveToolsFramework/Public/ComponentSourceInterfaces.h
@@ -69,7 +69,11 @@ public:
 	 */
 	void SetOwnerVisibility(bool bVisible) const;
 
-	using  FCommitter  = TFunction< void( FMeshDescription* ) >;
+	struct FCommitParams
+	{
+		FMeshDescription* MeshDescription{ nullptr };
+	};
+	using  FCommitter  = TFunction< void( const FCommitParams& ) >;
 	virtual FMeshDescription* GetMesh() = 0;
 	virtual void CommitMesh( const FCommitter& ) = 0;
 
