ModelingTools: use new gizmo snapping support for Sculpt and DrawPolygon plane-position gizmos
#rnx

diff --git a/Engine/Plugins/Experimental/MeshModelingToolset/Source/MeshModelingTools/Private/DrawPolygonTool.cpp b/Engine/Plugins/Experimental/MeshModelingToolset/Source/MeshModelingTools/Private/DrawPolygonTool.cpp
--- a/Engine/Plugins/Experimental/MeshModelingToolset/Source/MeshModelingTools/Private/DrawPolygonTool.cpp
+++ b/Engine/Plugins/Experimental/MeshModelingToolset/Source/MeshModelingTools/Private/DrawPolygonTool.cpp
@@ -298,6 +298,20 @@ void UDrawPolygonTool::PopLastVertexAction()
 	}
 }
 
+
+
+
+void UDrawPolygonTool::Tick(float DeltaTime)
+{
+	if (PlaneTransformGizmo != nullptr)
+	{
+		PlaneTransformGizmo->bSnapToWorldGrid =
+			SnapProperties->bEnableSnapping && SnapProperties->bSnapToWorldGrid && bIgnoreSnappingToggle == false;
+	}
+}
+
+
+
 void DrawEdgeTicks(FPrimitiveDrawInterface* PDI, 
 	const FSegment3f& Segment, float Height,
 	const FVector3f& PlaneNormal, 
@@ -1045,18 +1059,7 @@ void UDrawPolygonTool::SetDrawPlaneFromWorldPos(const FVector& Position, const F
 void UDrawPolygonTool::PlaneTransformChanged(UTransformProxy* Proxy, FTransform Transform)
 {
 	DrawPlaneOrientation = Transform.GetRotation();
-	
 	DrawPlaneOrigin = Transform.GetLocation();
-
-	if (bIgnoreSnappingToggle == false && SnapProperties->bEnableSnapping && SnapProperties->bSnapToWorldGrid)
-	{
-		FVector3d GridSnapPos;
-		if (ToolSceneQueriesUtil::FindWorldGridSnapPoint(this, DrawPlaneOrigin, GridSnapPos))
-		{
-			DrawPlaneOrigin = GridSnapPos;
-		}
-	}
-
 	SnapEngine.Plane = FFrame3d((FVector3d)DrawPlaneOrigin, (FQuaterniond)DrawPlaneOrientation);
 }
 
diff --git a/Engine/Plugins/Experimental/MeshModelingToolset/Source/MeshModelingTools/Private/DynamicMeshSculptTool.cpp b/Engine/Plugins/Experimental/MeshModelingToolset/Source/MeshModelingTools/Private/DynamicMeshSculptTool.cpp
--- a/Engine/Plugins/Experimental/MeshModelingToolset/Source/MeshModelingTools/Private/DynamicMeshSculptTool.cpp
+++ b/Engine/Plugins/Experimental/MeshModelingToolset/Source/MeshModelingTools/Private/DynamicMeshSculptTool.cpp
@@ -2038,15 +2038,6 @@ void UDynamicMeshSculptTool::PlaneTransformChanged(UTransformProxy* Proxy, FTran
 void UDynamicMeshSculptTool::UpdateFixedSculptPlanePosition(const FVector& Position)
 {
 	GizmoProperties->Position = Position;
-
-	if (GizmoProperties->bSnapToGrid)
-	{
-		FVector3d GridSnapPos;
-		if (ToolSceneQueriesUtil::FindWorldGridSnapPoint(this, GizmoProperties->Position, GridSnapPos))
-		{
-			GizmoProperties->Position = GridSnapPos;
-		}
-	}
 }
 
 
@@ -2068,6 +2059,8 @@ void UDynamicMeshSculptTool::UpdateFixedPlaneGizmoVisibility(bool bVisible)
 			PlaneTransformGizmo->SetActiveTarget(PlaneTransformProxy, GetToolManager());
 			PlaneTransformGizmo->SetNewGizmoTransform(FTransform(DrawPlaneOrientation, GizmoProperties->Position));
 		}
+
+		PlaneTransformGizmo->bSnapToWorldGrid = GizmoProperties->bSnapToGrid;
 	}
 }
 
diff --git a/Engine/Plugins/Experimental/MeshModelingToolset/Source/MeshModelingTools/Public/DrawPolygonTool.h b/Engine/Plugins/Experimental/MeshModelingToolset/Source/MeshModelingTools/Public/DrawPolygonTool.h
--- a/Engine/Plugins/Experimental/MeshModelingToolset/Source/MeshModelingTools/Public/DrawPolygonTool.h
+++ b/Engine/Plugins/Experimental/MeshModelingToolset/Source/MeshModelingTools/Public/DrawPolygonTool.h
@@ -190,6 +190,7 @@ public:
 	virtual void Setup() override;
 	virtual void Shutdown(EToolShutdownType ShutdownType) override;
 
+	virtual void Tick(float DeltaTime) override;
 	virtual void Render(IToolsContextRenderAPI* RenderAPI) override;
 
 	virtual bool HasCancel() const override { return false; }
