Copy up from Dev-Editor [REDACTED_HANDLE]

diff --git a/Engine/Plugins/Experimental/GeometryProcessing/Source/GeometricObjects/Public/Util/SparseIndexCollectionTypes.h b/Engine/Plugins/Experimental/GeometryProcessing/Source/GeometricObjects/Public/Util/SparseIndexCollectionTypes.h
--- a/Engine/Plugins/Experimental/GeometryProcessing/Source/GeometricObjects/Public/Util/SparseIndexCollectionTypes.h
+++ b/Engine/Plugins/Experimental/GeometryProcessing/Source/GeometricObjects/Public/Util/SparseIndexCollectionTypes.h
@@ -4,6 +4,8 @@
 
 #pragma once
 
+#include "IndexTypes.h"
+
 #include <CoreMinimal.h>
 
 
@@ -28,6 +30,9 @@ public:
 		InitAuto(MaxIndex, SubsetCountEst);
 	}
 
+	/**
+	 * Initialize to either sparse or dense automatically, based on max index and estimated number of set indices
+	 */
 	void InitAuto(int MaxIndex, int SubsetCountEst)
 	{
 		bool bSmall = MaxIndex < 32000;
@@ -36,6 +41,11 @@ public:
 		InitManual(!bSmall && ((float)SubsetCountEst / (float)MaxIndex < PercentThresh), MaxIndex);
 	}
 
+	/**
+	 * Initialize to sparse or dense based on the explicit caller choice
+	 * @param bSetSparse if true, will use sparse storage; otherwise will use dense
+	 * @param MaxIndex maximum valid index in set; must be non-negative for dense case
+	 */
 	void InitManual(bool bSetSparse, int MaxIndex = -1)
 	{
 		if (bSetSparse)
@@ -52,29 +62,30 @@ public:
 	}
 
 	/**
-	 *  checks if value i is true
+	 * @param Index index to check for set membership
+	 * @return true if Index is in set
 	 */
-	FORCEINLINE bool Contains(int i) const
+	FORCEINLINE bool Contains(int Index) const
 	{
 		if (Dense.IsSet())
 		{
-			return Dense.GetValue()[i];
+			return Dense.GetValue()[Index];
 		}
 		else
 		{
 			check(Sparse.IsSet());
-			return Sparse->Contains(i);
+			return Sparse->Contains(Index);
 		}
 	}
 
 	/**
-	 *  sets value i to true
+	 * @param Index index to add to set
 	 */
-	FORCEINLINE void Add(int i)
+	FORCEINLINE void Add(int Index)
 	{
 		if (Dense.IsSet())
 		{
-			bool &Value = Dense.GetValue()[i];
+			bool &Value = Dense.GetValue()[Index];
 			if (!Value)
 			{
 				Value = true;
@@ -84,10 +95,13 @@ public:
 		else
 		{
 			check(Sparse.IsSet());
-			Sparse->Add(i);
+			Sparse->Add(Index);
 		}
 	}
 
+	/**
+	 * @param Index index to remove from set
+	 */
 	FORCEINLINE void Remove(int Index)
 	{
 		if (Dense.IsSet())
@@ -107,7 +121,7 @@ public:
 	}
 
 	/**
-	 *  Returns number of true values in set
+	 *  Returns number of values in set
 	 */
 	FORCEINLINE int Count() const
 	{
@@ -121,7 +135,9 @@ public:
 		}
 	}
 
-
+	/**
+	 * Array-style accessor that returns true if Index is in set, false otherwise
+	 */
 	inline const bool operator[](unsigned int Index) const
 	{
 		return Contains(Index);
@@ -132,13 +148,221 @@ public:
 
 };
 
+/**
+ * Index map that supports dense or sparse storage, or a simple formula-based map (e.g. constant, identity, shift)
+ * For dense and sparse, the formula can be used to set default values.
+ */
+struct FOptionallySparseIndexMap
+{
+	// storage used when map is dense
+	TArray<int> Dense;
+	// storage used when map is sparse
+	TMap<int, int> Sparse;
+
+	// mapping to assign defaults for indices that are not explicitly assigned; computed as Index*DefaultScale + DefaultOffset
+	int DefaultOffset = -1, DefaultScale = 0;
+
+	// max index in map; -1 will leave the map unbounded (invalid for Dense maps)
+	int MaxIndex = -1;
+
+	// different kinds of storage that could be used for the map
+	enum class EMapType : uint8
+	{
+		Dense, Sparse, ScaleAndOffset
+	};
 
+	// choice of storage backing the map
+	EMapType MapType = EMapType::Sparse;
+
+	/** default constructor sets map to sparse, w/ constant -1 default value */
+	FOptionallySparseIndexMap()
+	{}
+
+	/** construct map with explicit choice of storage type */
+	FOptionallySparseIndexMap(EMapType MapType, int MaxIndex = -1) : MaxIndex(MaxIndex), MapType(MapType)
+	{
+		if (MapType == EMapType::Dense)
+		{
+			check(MaxIndex > -1);
+			Dense.SetNum(MaxIndex);
+			InitDefaults();
+		}
+	}
 
+	/** construct map with dense storage copied from an existing array */
+	FOptionallySparseIndexMap(TArray<int> Dense, int MaxIndex = -1) : Dense(Dense), MaxIndex(MaxIndex), MapType(EMapType::Dense)
+	{
+	}
 
-// TODO implement a similar class but as a template map
-//template<typename T>
-//class TOptionallySparseMap
-//{
-// etc
-//
-//};
+	/** construct map with automatically-chosen dense or sparse storage, based on max index and estimated element count */
+	FOptionallySparseIndexMap(int MaxIndex, int SubsetCountEst)
+	{
+		Initialize(MaxIndex, SubsetCountEst);
+	}
+
+	/**
+	 * Automatically choose sparse or dense storage based on use estimate
+	 */
+	void Initialize(int MaxIndexIn, int SubsetCountEst)
+	{
+		Dense.Reset();
+		Sparse.Reset();
+		MaxIndex = MaxIndexIn;
+
+		// if buffer is less than threshold, just use dense map
+		bool bSmall = MaxIndex < 32000;
+		float fPercent = (float)SubsetCountEst / (float)MaxIndex;
+		float fPercentThresh = 0.1f;
+
+		if (bSmall || fPercent > fPercentThresh)
+		{
+			Dense.SetNum(MaxIndex);
+			MapType = EMapType::Dense;
+		}
+		else
+		{
+			MapType = EMapType::Sparse;
+		}
+		
+		InitDefaults();
+	}
+
+	/**
+	 * @return an identity map
+	 */
+	static FOptionallySparseIndexMap IdentityMap(int MaxIndex = -1)
+	{
+		FOptionallySparseIndexMap ToRet(EMapType::ScaleAndOffset, MaxIndex);
+		ToRet.DefaultScale = 1;
+		ToRet.DefaultOffset = 0;
+		return ToRet;
+	}
+	
+	/**
+	 * @return a constant map
+	 */
+	static FOptionallySparseIndexMap ConstantMap(int ConstantValue, int MaxIndex = -1)
+	{
+		FOptionallySparseIndexMap ToRet(EMapType::ScaleAndOffset, MaxIndex);
+		ToRet.DefaultScale = 0;
+		ToRet.DefaultOffset = ConstantValue;
+		return ToRet;
+	}
+
+	/**
+	 * @return a sparsely-overriden an identity map; elements not explicitly set are identity
+	 */
+	static FOptionallySparseIndexMap SparseIdentityMap(int MaxIndex = -1)
+	{
+		FOptionallySparseIndexMap ToRet(EMapType::Sparse, MaxIndex);
+		ToRet.DefaultScale = 1;
+		ToRet.DefaultOffset = 0;
+		return ToRet;
+	}
+
+
+	/* Initialize default values for dense map (sparse defaults are computed on the fly) */
+	void InitDefaults()
+	{
+		if (MapType == EMapType::Dense)
+		{
+			for (int i = 0; i < Dense.Num(); ++i)
+			{
+				Dense[i] = i*DefaultScale + DefaultOffset;
+			}
+		}
+	}
+
+	/** @return true if a given Index is invalid (below zero or above MaxIndex if MaxIndex is valid) */
+	inline bool BadIndex(int Index) const
+	{
+		return (Index < 0) || (MaxIndex >= 0 && Index >= MaxIndex);
+	}
+
+	/**
+	 * dense variant: returns true unless you have set index to InvalidIndex (eg via SetToInvalid)
+	 * sparse variant: returns true if index is in map or default index is >= 0
+	 * scaleandoffset variant: returns true if default index is >= 0
+	 * all return false if index is out-of-bounds
+	 */
+	bool Contains(int Index)
+	{
+		if (BadIndex(Index))
+		{
+			return false;
+		}
+		switch (MapType)
+		{
+		case EMapType::Dense:
+			return Dense[Index] >= 0;
+		case EMapType::Sparse:
+			return Sparse.Contains(Index) || (Index * DefaultScale + DefaultOffset >= 0);
+		case EMapType::ScaleAndOffset:
+			return Index * DefaultScale + DefaultOffset >= 0;
+		}
+		check(false);
+		return false;
+	}
+
+	/**
+	 * Array-style accessor to the map
+	 */
+	inline const int operator[](int Index) const
+	{
+		if (BadIndex(Index))
+		{
+			return IndexConstants::InvalidID;
+		}
+		switch (MapType)
+		{
+		case EMapType::Dense:
+			return Dense[Index];
+		case EMapType::Sparse:
+		{
+			const int *Value = Sparse.Find(Index);
+			return Value ? (*Value) : (Index * DefaultScale + DefaultOffset);
+		}
+		case EMapType::ScaleAndOffset:
+			return Index * DefaultScale + DefaultOffset;
+		}
+		check(false);
+		return IndexConstants::InvalidID;
+	}
+
+	/**
+	 * @param Index this index will be unset in map, restoring the default mapping (Index -> Index*DefaultScale + DefaultOffset)
+	 */
+	void Unset(int Index)
+	{
+		check(!BadIndex(Index));
+		if (MapType == EMapType::Dense)
+		{
+			Dense[Index] = Index*DefaultScale + DefaultOffset;
+		}
+		else
+		{
+			Sparse.Remove(Index);
+		}
+	}
+
+	/** Sets the map at a given index.  Do not call if MapType is ScaleAndOffset. */
+	void Set(int Index, int Value)
+	{
+		check(!BadIndex(Index));
+		check(MapType != EMapType::ScaleAndOffset);
+		if (MapType == EMapType::Dense)
+		{
+			Dense[Index] = Value;
+		}
+		else
+		{
+			Sparse[Index] = Value;
+		}
+	}
+
+	/** @param Index Explicitly sets the given index to Invalid (*not* default).  Do not call if MapType is ScaleAndOffset. */
+	inline void SetInvalid(int Index)
+	{
+		Set(Index, IndexConstants::InvalidID);
+	}
+};
diff --git a/Engine/Plugins/Experimental/GeometryProcessing/Source/GeometricObjects/Public/VectorTypes.h b/Engine/Plugins/Experimental/GeometryProcessing/Source/GeometricObjects/Public/VectorTypes.h
--- a/Engine/Plugins/Experimental/GeometryProcessing/Source/GeometricObjects/Public/VectorTypes.h
+++ b/Engine/Plugins/Experimental/GeometryProcessing/Source/GeometricObjects/Public/VectorTypes.h
@@ -17,81 +17,90 @@
 * Convenience typedefs for FVector2f/FVector2d/FVector2i are defined, and
 * should be preferentially used over the base template type
 *
-* Implicit casts to/from FVector2D are defined, for all types.
-*
 * @todo Possibly can be replaced/merged with Chaos TVector<T,N>
 */
 template <typename T>
 struct FVector2
 {
-	T X, Y;
+	T X{}, Y{};
 
-	FVector2()
+	constexpr FVector2()
+		: X(0), Y(0)
 	{
 	}
-	FVector2(T ValX, T ValY)
+
+	constexpr FVector2(T ValX, T ValY)
 		: X(ValX), Y(ValY)
 	{
 	}
-	FVector2(const T* Data)
+
+	constexpr FVector2(const T* Data)
+		: X(Data[0]), Y(Data[1])
+	{
+	}
+
+	constexpr FVector2(const FVector2& Vec) = default;
+
+	constexpr FVector2(const FVector2D& Vec)
+		: X((T)Vec.X), Y((T)Vec.Y)
+	{
+	}
+
+	template<typename RealType2>
+	explicit constexpr FVector2(const FVector2<RealType2>& Vec)
+		: X((T)Vec.X), Y((T)Vec.Y)
 	{
-		X = Data[0];
-		Y = Data[1];
 	}
 
-	inline operator const T*() const
+	explicit constexpr operator const T*() const
 	{
 		return &X;
 	};
-	inline operator T*()
+
+	explicit constexpr operator T*()
 	{
 		return &X;
 	}
 
-	operator FVector2D() const
+	explicit operator FVector2D() const
 	{
 		return FVector2D((float)X, (float)Y);
 	}
-	FVector2(const FVector2D& Vec)
-	{
-		X = (T)Vec.X;
-		Y = (T)Vec.Y;
-	}
 
-	static FVector2<T> Zero()
+	constexpr static FVector2<T> Zero()
 	{
 		return FVector2<T>((T)0, (T)0);
 	}
-	static FVector2<T> One()
+	constexpr static FVector2<T> One()
 	{
 		return FVector2<T>((T)1, (T)1);
 	}
-	static FVector2<T> UnitX()
+	constexpr static FVector2<T> UnitX()
 	{
 		return FVector2<T>((T)1, (T)0);
 	}
-	static FVector2<T> UnitY()
+	constexpr static FVector2<T> UnitY()
 	{
 		return FVector2<T>((T)0, (T)1);
 	}
 
-	FVector2<T>& operator=(const FVector2<T>& V2)
+	constexpr FVector2<T>& operator=(const FVector2<T>& V2)
 	{
 		X = V2.X;
 		Y = V2.Y;
 		return *this;
 	}
 
-	T& operator[](int Idx)
+	constexpr T& operator[](int Idx)
 	{
 		return (&X)[Idx];
 	}
-	const T& operator[](int Idx) const
+	constexpr const T& operator[](int Idx) const
 	{
 		return (&X)[Idx];
 	}
 
-	T Length() const
+	constexpr T Length() const
 	{
 		return TMathUtil<T>::Sqrt(X * X + Y * Y);
 	}
@@ -100,113 +109,119 @@ struct FVector2
 		return X * X + Y * Y;
 	}
 
-	inline T Distance(const FVector2<T>& V2) const
+	constexpr T Distance(const FVector2<T>& V2) const
 	{
 		T dx = V2.X - X;
 		T dy = V2.Y - Y;
 		return TMathUtil<T>::Sqrt(dx * dx + dy * dy);
 	}
-	inline T DistanceSquared(const FVector2<T>& V2) const
+	constexpr T DistanceSquared(const FVector2<T>& V2) const
 	{
 		T dx = V2.X - X;
 		T dy = V2.Y - Y;
 		return dx * dx + dy * dy;
 	}
 
-	T Dot(const FVector2<T>& V2) const
+	constexpr T Dot(const FVector2<T>& V2) const
 	{
 		return X * V2.X + Y * V2.Y;
 	}
 
 	// Note: DotPerp and FVector2's version of Cross are the same function
-	inline T DotPerp(const FVector2<T>& V2) const
+	constexpr T DotPerp(const FVector2<T>& V2) const
 	{
 		return X * V2.Y - Y * V2.X;
 	}
-	inline T Cross(const FVector2<T>& V2) const
+	constexpr T Cross(const FVector2<T>& V2) const
 	{
 		return X * V2.Y - Y * V2.X;
 	}
 
-	FVector2<T> Perp() const
+	constexpr FVector2<T> Perp() const
 	{
 		return FVector2<T>(Y, -X);
 	}
 
-	inline FVector2 operator-() const
+	constexpr  FVector2 operator-() const
 	{
 		return FVector2(-X, -Y);
 	}
 
-	inline FVector2 operator+(const FVector2& V2) const
+	constexpr FVector2 operator+(const FVector2& V2) const
 	{
 		return FVector2(X + V2.X, Y + V2.Y);
 	}
 
-	inline FVector2 operator-(const FVector2& V2) const
+	constexpr FVector2 operator-(const FVector2& V2) const
 	{
 		return FVector2(X - V2.X, Y - V2.Y);
 	}
 
-	inline FVector2<T> operator+(const T& Scalar) const
+	constexpr FVector2<T> operator+(const T& Scalar) const
 	{
 		return FVector2<T>(X + Scalar, Y + Scalar);
 	}
 
-	inline FVector2<T> operator-(const T& Scalar) const
+	constexpr FVector2<T> operator-(const T& Scalar) const
 	{
 		return FVector2<T>(X - Scalar, Y - Scalar);
 	}
 
-	inline FVector2 operator*(const T& Scalar) const
+	constexpr FVector2<T> operator*(const T& Scalar) const
+	{
+		return FVector2<T>(X * Scalar, Y * Scalar);
+	}
+
+	template<typename RealType2>
+	constexpr FVector2<T> operator*(const RealType2& Scalar) const
 	{
-		return FVector2(X * Scalar, Y * Scalar);
+		return FVector2<T>(X * (T)Scalar, Y * (T)Scalar);
 	}
 
-	inline FVector2<T> operator*(const FVector2<T>& V2) const // component-wise
+	constexpr FVector2<T> operator*(const FVector2<T>& V2) const // component-wise
 	{
 		return FVector2<T>(X * V2.X, Y * V2.Y);
 	}
 
-	inline FVector2 operator/(const T& Scalar) const
+	constexpr FVector2<T> operator/(const T& Scalar) const
 	{
-		return FVector2(X / Scalar, Y / Scalar);
+		return FVector2<T>(X / Scalar, Y / Scalar);
 	}
 
-	inline FVector2<T> operator/(const FVector2<T>& V2) const // component-wise
+	constexpr FVector2<T> operator/(const FVector2<T>& V2) const // component-wise
 	{
 		return FVector2<T>(X / V2.X, Y / V2.Y);
 	}
 
-	inline FVector2<T>& operator+=(const FVector2<T>& V2)
+	constexpr FVector2<T>& operator+=(const FVector2<T>& V2)
 	{
 		X += V2.X;
 		Y += V2.Y;
 		return *this;
 	}
 
-	inline FVector2<T>& operator-=(const FVector2<T>& V2)
+	constexpr FVector2<T>& operator-=(const FVector2<T>& V2)
 	{
 		X -= V2.X;
 		Y -= V2.Y;
 		return *this;
 	}
 
-	inline FVector2<T>& operator*=(const T& Scalar)
+	constexpr FVector2<T>& operator*=(const T& Scalar)
 	{
 		X *= Scalar;
 		Y *= Scalar;
 		return *this;
 	}
 
-	inline FVector2<T>& operator/=(const T& Scalar)
+	constexpr FVector2<T>& operator/=(const T& Scalar)
 	{
 		X /= Scalar;
 		Y /= Scalar;
 		return *this;
 	}
 
-	inline bool IsNormalized()
+	constexpr bool IsNormalized()
 	{
 		return TMathUtil<T>::Abs((X * X + Y * Y) - 1) < TMathUtil<T>::ZeroTolerance;
 	}
@@ -258,7 +273,7 @@ struct FVector2
 		return 0;
 	}
 
-	inline FVector2<T> Normalized(const T Epsilon = 0) const
+	constexpr FVector2<T> Normalized(const T Epsilon = 0) const
 	{
 		T length = Length();
 		if (length > Epsilon)
@@ -301,68 +316,71 @@ struct FVector2
  * Convenience typedefs for FVector3f/FVector3d/FVector3i are defined, and
  * should be preferentially used over the base template type
  * 
- * Implicit casts to/from FVector are defined, for all types.
- *  
  * @todo Possibly can be replaced/merged with Chaos TVector<T,N>
  */
 template <typename T>
 struct FVector3
 {
-	T X, Y, Z;
+	T X{}, Y{}, Z{};
 
-	FVector3()
+	constexpr FVector3()
+		: X(0), Y(0), Z(0)
 	{
 	}
-	FVector3(T ValX, T ValY, T ValZ)
+
+	constexpr FVector3(T ValX, T ValY, T ValZ)
 		: X(ValX), Y(ValY), Z(ValZ)
 	{
 	}
-	FVector3(const T* Data)
+
+	constexpr FVector3(const T* Data)
+		: X(Data[0]), Y(Data[1]), Z(Data[2])
 	{
-		X = Data[0];
-		Y = Data[1];
-		Z = Data[2];
 	}
-	FVector3(const FVector2<T>& Data)
+
+	constexpr FVector3(const FVector2<T>& Data)
+		: X(Data.X), Y(Data.Y), Z((T)0)
 	{
-		X = Data.X;
-		Y = Data.Y;
-		Z = (T)0;
 	}
 
-	inline operator const T*() const
+	constexpr FVector3(const FVector3& Vec) = default;
+
+	template<typename RealType2>
+	explicit constexpr FVector3(const FVector3<RealType2>& Vec)
+		: X((T)Vec.X), Y((T)Vec.Y), Z((T)Vec.Z)
+	{
+	}
+
+	explicit constexpr operator const T*() const
 	{
 		return &X;
 	};
-	inline operator T*()
+	explicit constexpr operator T*()
 	{
 		return &X;
 	}
 
-	inline operator FVector() const
+	explicit operator FVector() const
 	{
 		return FVector((float)X, (float)Y, (float)Z);
 	}
-	inline FVector3(const FVector& Vec)
+
+	FVector3(const FVector& Vec)
+		: X((T)Vec.X), Y((T)Vec.Y), Z((T)Vec.Z)
 	{
-		X = (T)Vec.X;
-		Y = (T)Vec.Y;
-		Z = (T)Vec.Z;
 	}
 
 
-	explicit inline operator FLinearColor() const
+	explicit constexpr operator FLinearColor() const
 	{
 		return FLinearColor((float)X, (float)Y, (float)Z);
 	}
-	inline FVector3(const FLinearColor& Color)
+	constexpr FVector3(const FLinearColor& Color)
+		: X((T)Color.R), Y((T)Color.G), Z((T)Color.B)
 	{
-		X = (T)Color.R;
-		Y = (T)Color.G;
-		Z = (T)Color.B;
 	}
 
-	explicit inline operator FColor() const
+	explicit constexpr operator FColor() const
 	{
 		return FColor(
 			FMathf::Clamp((int)((float)X*255.0f), 0, 255),
@@ -422,14 +440,14 @@ struct FVector3
 		return X * X + Y * Y + Z * Z;
 	}
 
-	inline T Distance(const FVector3<T>& V2) const
+	constexpr T Distance(const FVector3<T>& V2) const
 	{
 		T dx = V2.X - X;
 		T dy = V2.Y - Y;
 		T dz = V2.Z - Z;
 		return TMathUtil<T>::Sqrt(dx * dx + dy * dy + dz * dz);
 	}
-	inline T DistanceSquared(const FVector3<T>& V2) const
+	constexpr T DistanceSquared(const FVector3<T>& V2) const
 	{
 		T dx = V2.X - X;
 		T dy = V2.Y - Y;
@@ -437,52 +455,58 @@ struct FVector3
 		return dx * dx + dy * dy + dz * dz;
 	}
 
-	inline FVector3<T> operator-() const
+	constexpr FVector3<T> operator-() const
 	{
 		return FVector3<T>(-X, -Y, -Z);
 	}
 
-	inline FVector3<T> operator+(const FVector3<T>& V2) const
+	constexpr FVector3<T> operator+(const FVector3<T>& V2) const
 	{
 		return FVector3<T>(X + V2.X, Y + V2.Y, Z + V2.Z);
 	}
 
-	inline FVector3<T> operator-(const FVector3<T>& V2) const
+	constexpr FVector3<T> operator-(const FVector3<T>& V2) const
 	{
 		return FVector3<T>(X - V2.X, Y - V2.Y, Z - V2.Z);
 	}
 
-	inline FVector3<T> operator+(const T& Scalar) const
+	constexpr FVector3<T> operator+(const T& Scalar) const
 	{
 		return FVector3<T>(X + Scalar, Y + Scalar, Z + Scalar);
 	}
 
-	inline FVector3<T> operator-(const T& Scalar) const
+	constexpr FVector3<T> operator-(const T& Scalar) const
 	{
 		return FVector3<T>(X - Scalar, Y - Scalar, Z - Scalar);
 	}
 
-	inline FVector3<T> operator*(const T& Scalar) const
+	constexpr FVector3<T> operator*(const T& Scalar) const
 	{
 		return FVector3<T>(X * Scalar, Y * Scalar, Z * Scalar);
 	}
 
-	inline FVector3<T> operator*(const FVector3<T>& V2) const // component-wise
+	template<typename RealType2>
+	constexpr FVector3<T> operator*(const RealType2& Scalar) const
+	{
+		return FVector3<T>(X * (T)Scalar, Y * (T)Scalar, Z * (T)Scalar);
+	}
+
+	constexpr FVector3<T> operator*(const FVector3<T>& V2) const // component-wise
 	{
 		return FVector3<T>(X * V2.X, Y * V2.Y, Z * V2.Z);
 	}
 
-	inline FVector3<T> operator/(const T& Scalar) const
+	constexpr FVector3<T> operator/(const T& Scalar) const
 	{
 		return FVector3<T>(X / Scalar, Y / Scalar, Z / Scalar);
 	}
 
-	inline FVector3<T> operator/(const FVector3<T>& V2) const // component-wise
+	constexpr FVector3<T> operator/(const FVector3<T>& V2) const // component-wise
 	{
 		return FVector3<T>(X / V2.X, Y / V2.Y, Z / V2.Z);
 	}
 
-	inline FVector3<T>& operator+=(const FVector3<T>& V2)
+	constexpr FVector3<T>& operator+=(const FVector3<T>& V2)
 	{
 		X += V2.X;
 		Y += V2.Y;
@@ -490,7 +514,7 @@ struct FVector3
 		return *this;
 	}
 
-	inline FVector3<T>& operator-=(const FVector3<T>& V2)
+	constexpr FVector3<T>& operator-=(const FVector3<T>& V2)
 	{
 		X -= V2.X;
 		Y -= V2.Y;
@@ -498,7 +522,7 @@ struct FVector3
 		return *this;
 	}
 
-	inline FVector3<T>& operator*=(const T& Scalar)
+	constexpr FVector3<T>& operator*=(const T& Scalar)
 	{
 		X *= Scalar;
 		Y *= Scalar;
@@ -506,7 +530,7 @@ struct FVector3
 		return *this;
 	}
 
-	inline FVector3<T>& operator/=(const T& Scalar)
+	constexpr FVector3<T>& operator/=(const T& Scalar)
 	{
 		X /= Scalar;
 		Y /= Scalar;
@@ -543,7 +567,7 @@ struct FVector3
 		return (T)acos(ClampedDot);
 	}
 
-	inline bool IsNormalized()
+	constexpr bool IsNormalized()
 	{
 		return TMathUtil<T>::Abs((X * X + Y * Y + Z * Z) - 1) < TMathUtil<T>::ZeroTolerance;
 	}
@@ -563,7 +587,7 @@ struct FVector3
 		return 0;
 	}
 
-	inline FVector3<T> Normalized(const T Epsilon = 0) const
+	constexpr FVector3<T> Normalized(const T Epsilon = 0) const
 	{
 		T length = Length();
 		if (length > Epsilon)
@@ -574,12 +598,12 @@ struct FVector3
 		return FVector3<T>((T)0, (T)0, (T)0);
 	}
 
-	inline T MaxAbs() const
+	constexpr T MaxAbs() const
 	{
 		return TMathUtil<T>::Max3(TMathUtil<T>::Abs(X), TMathUtil<T>::Abs(Y), TMathUtil<T>::Abs(Z));
 	}
 
-	inline T MinAbs() const
+	constexpr T MinAbs() const
 	{
 		return TMathUtil<T>::Min3(TMathUtil<T>::Abs(X), TMathUtil<T>::Abs(Y), TMathUtil<T>::Abs(Z));
 	}
@@ -592,7 +616,7 @@ struct FVector3
 						   OneMinusAlpha * A.Z + Alpha * B.Z);
 	}
 
-	inline bool operator==(const FVector3<T>& Other) const
+	constexpr bool operator==(const FVector3<T>& Other) const
 	{
 		return X == Other.X && Y == Other.Y && Z == Other.Z;
 	}
@@ -604,6 +628,13 @@ inline FVector3<RealType> operator*(RealType Scalar, const FVector3<RealType>& V
 	return FVector3<RealType>(Scalar * V.X, Scalar * V.Y, Scalar * V.Z);
 }
 
+// allow float*Vector3<double> and double*Vector3<float>
+template <typename RealType, typename RealType2>
+inline FVector3<RealType> operator*(RealType2 Scalar, const FVector3<RealType>& V)
+{
+	return FVector3<RealType>((RealType)Scalar * V.X, (RealType)Scalar * V.Y, (RealType)Scalar * V.Z);
+}
+
 typedef FVector3<float> FVector3f;
 typedef FVector3<double> FVector3d;
 typedef FVector3<int> FVector3i;
@@ -634,6 +665,13 @@ inline FVector2<RealType> operator*(RealType Scalar, const FVector2<RealType>& V
 	return FVector2<RealType>(Scalar * V.X, Scalar * V.Y);
 }
 
+// allow float*Vector2<double> and double*Vector2<float>
+template <typename RealType, typename RealType2>
+inline FVector2<RealType> operator*(RealType2 Scalar, const FVector2<RealType>& V)
+{
+	return FVector2<RealType>((RealType)Scalar * V.X, (RealType)Scalar * V.Y);
+}
+
 typedef FVector2<float> FVector2f;
 typedef FVector2<double> FVector2d;
 typedef FVector2<int> FVector2i;
diff --git a/Engine/Plugins/Experimental/GeometryProcessing/Source/MeshConversion/Private/DynamicMeshToMeshDescription.cpp b/Engine/Plugins/Experimental/GeometryProcessing/Source/MeshConversion/Private/DynamicMeshToMeshDescription.cpp
--- a/Engine/Plugins/Experimental/GeometryProcessing/Source/MeshConversion/Private/DynamicMeshToMeshDescription.cpp
+++ b/Engine/Plugins/Experimental/GeometryProcessing/Source/MeshConversion/Private/DynamicMeshToMeshDescription.cpp
@@ -42,7 +42,7 @@ namespace DynamicMeshToMeshDescriptionConversionHelper
 }
 
 
-void FDynamicMeshToMeshDescription::Update(const FDynamicMesh3* MeshIn, FMeshDescription& MeshOut)
+void FDynamicMeshToMeshDescription::Update(const FDynamicMesh3* MeshIn, FMeshDescription& MeshOut, bool bUpdateNormals, bool bUpdateUVs)
 {
 	FMeshDescriptionBuilder Builder;
 	Builder.SetMeshDescription(&MeshOut);
@@ -53,10 +53,10 @@ void FDynamicMeshToMeshDescription::Update(const FDynamicMesh3* MeshIn, FMeshDes
 	// update positions
 	for (int VertID : MeshIn->VertexIndicesItr())
 	{
-		Builder.SetPosition(FVertexID(VertID), MeshIn->GetVertex(VertID));
+		Builder.SetPosition(FVertexID(VertID), (FVector)MeshIn->GetVertex(VertID));
 	}
 
-	UpdateAttributes(MeshIn, MeshOut, true, false);
+	UpdateAttributes(MeshIn, MeshOut, bUpdateNormals, bUpdateUVs);
 }
 
 
@@ -157,7 +157,7 @@ void FDynamicMeshToMeshDescription::Convert_NoAttributes(const FDynamicMesh3* Me
 	MapV.SetNum(MeshIn->MaxVertexID());
 	for (int VertID : MeshIn->VertexIndicesItr())
 	{
-		MapV[VertID] = Builder.AppendVertex(MeshIn->GetVertex(VertID));
+		MapV[VertID] = Builder.AppendVertex((FVector)MeshIn->GetVertex(VertID));
 	}
 
 	FPolygonGroupID AllGroupID = Builder.AppendPolygonGroup();
@@ -218,7 +218,7 @@ void FDynamicMeshToMeshDescription::Convert_SharedInstances(const FDynamicMesh3*
 	TArray<FVertexID> MapV; MapV.SetNum(MeshIn->MaxVertexID());
 	for (int VertID : MeshIn->VertexIndicesItr())
 	{
-		MapV[VertID] = Builder.AppendVertex(MeshIn->GetVertex(VertID));
+		MapV[VertID] = Builder.AppendVertex((FVector)MeshIn->GetVertex(VertID));
 	}
 
 
@@ -311,7 +311,8 @@ void FDynamicMeshToMeshDescription::Convert_SharedInstances(const FDynamicMesh3*
 				{
 					int ElID = NormalOverlay->GetTriangle(TriID)[SubIdx];
 					KnownInstanceIDs.Add(int32(ElID));
-					Builder.SetInstanceNormal(NewInstanceID, ElID != -1 ? NormalOverlay->GetElement(ElID) : FVector3f::UnitZ());
+					FVector3f ElementNormal = ElID != -1 ? NormalOverlay->GetElement(ElID) : FVector3f::UnitZ();
+					Builder.SetInstanceNormal(NewInstanceID, (FVector)ElementNormal);
 				}
 				else
 				{
@@ -323,7 +324,8 @@ void FDynamicMeshToMeshDescription::Convert_SharedInstances(const FDynamicMesh3*
 					int ElID = Overlay->GetTriangle(TriID)[SubIdx];
 					KnownInstanceIDs.Add(int32(ElID));
 
-					Builder.SetInstanceUV(NewInstanceID, ElID != -1 ? Overlay->GetElement(ElID) : FVector2f::Zero(), UVLayerIndex);
+					FVector2f ElementUV = ElID != -1 ? Overlay->GetElement(ElID) : FVector2f::Zero();
+					Builder.SetInstanceUV(NewInstanceID, (FVector2D)ElementUV, UVLayerIndex);
 				}
 				FoundInstance = NewInstanceID.GetValue();
 				KnownInstanceIDs.Add(FoundInstance);
@@ -379,7 +381,7 @@ void FDynamicMeshToMeshDescription::Convert_NoSharedInstances(const FDynamicMesh
 	TArray<FVertexID> MapV; MapV.SetNum(MeshIn->MaxVertexID());
 	for (int VertID : MeshIn->VertexIndicesItr())
 	{
-		MapV[VertID] = Builder.AppendVertex(MeshIn->GetVertex(VertID));
+		MapV[VertID] = Builder.AppendVertex((FVector)MeshIn->GetVertex(VertID));
 	}
 
 	FPolygonGroupID ZeroPolygonGroupID = Builder.AppendPolygonGroup();
@@ -424,9 +426,9 @@ void FDynamicMeshToMeshDescription::Convert_NoSharedInstances(const FDynamicMesh
 		if (UVOverlay != nullptr)
 		{
 			FIndex3i UVTriangle = UVOverlay->GetTriangle(TriID);
-			TriUVs[0] = UVOverlay->GetElement(UVTriangle[0]);
-			TriUVs[1] = UVOverlay->GetElement(UVTriangle[1]);
-			TriUVs[2] = UVOverlay->GetElement(UVTriangle[2]);
+			TriUVs[0] = (FVector2D)UVOverlay->GetElement(UVTriangle[0]);
+			TriUVs[1] = (FVector2D)UVOverlay->GetElement(UVTriangle[1]);
+			TriUVs[2] = (FVector2D)UVOverlay->GetElement(UVTriangle[2]);
 			UseUVs = TriUVs;
 		}
 
@@ -434,9 +436,9 @@ void FDynamicMeshToMeshDescription::Convert_NoSharedInstances(const FDynamicMesh
 		if (NormalOverlay != nullptr)
 		{
 			FIndex3i NormalTriangle = NormalOverlay->GetTriangle(TriID);
-			TriNormals[0] = NormalOverlay->GetElement(NormalTriangle[0]);
-			TriNormals[1] = NormalOverlay->GetElement(NormalTriangle[1]);
-			TriNormals[2] = NormalOverlay->GetElement(NormalTriangle[2]);
+			TriNormals[0] = (FVector)NormalOverlay->GetElement(NormalTriangle[0]);
+			TriNormals[1] = (FVector)NormalOverlay->GetElement(NormalTriangle[1]);
+			TriNormals[2] = (FVector)NormalOverlay->GetElement(NormalTriangle[2]);
 			UseNormals = TriNormals;
 		}
 
diff --git a/Engine/Plugins/Experimental/GeometryProcessing/Source/MeshConversion/Private/MeshDescriptionBuilder.cpp b/Engine/Plugins/Experimental/GeometryProcessing/Source/MeshConversion/Private/MeshDescriptionBuilder.cpp
--- a/Engine/Plugins/Experimental/GeometryProcessing/Source/MeshConversion/Private/MeshDescriptionBuilder.cpp
+++ b/Engine/Plugins/Experimental/GeometryProcessing/Source/MeshConversion/Private/MeshDescriptionBuilder.cpp
@@ -267,5 +267,5 @@ FBox FMeshDescriptionBuilder::ComputeBoundingBox() const
 	{
 		bounds.Contain(VertexPositions.Get(VertexID));
 	}
-	return bounds;
+	return (FBox)bounds;
 }
\ No newline at end of file
diff --git a/Engine/Plugins/Experimental/GeometryProcessing/Source/MeshConversion/Private/MeshDescriptionToDynamicMesh.cpp b/Engine/Plugins/Experimental/GeometryProcessing/Source/MeshConversion/Private/MeshDescriptionToDynamicMesh.cpp
--- a/Engine/Plugins/Experimental/GeometryProcessing/Source/MeshConversion/Private/MeshDescriptionToDynamicMesh.cpp
+++ b/Engine/Plugins/Experimental/GeometryProcessing/Source/MeshConversion/Private/MeshDescriptionToDynamicMesh.cpp
@@ -36,7 +36,7 @@ public:
 		UVOverlay = UVOverlayIn;
 	}
 
-	int FindOrAddUnique(const FVector2D & UV, int VertexID)
+	int FindOrAddUnique(const FVector2D& UV, int VertexID)
 	{
 		FVertexUV VertUV = { VertexID, UV.X, UV.Y };
 
@@ -47,7 +47,7 @@ public:
 		}
 		else
 		{
-			NewIndex = UVOverlay->AppendElement(UV, VertexID);
+			NewIndex = UVOverlay->AppendElement(FVector2f(UV), VertexID);
 			UniqueVertexUVs.Add(VertUV, NewIndex);
 		}
 		return NewIndex;
@@ -173,8 +173,9 @@ void FMeshDescriptionToDynamicMesh::Convert(const FMeshDescription* MeshIn, FDyn
 	}
 	FNormalWelder NormalWelder(NormalOverlay);
 
-	// only enable material ID if we have more than one material
-	FDynamicMeshMaterialAttribute* MaterialIDAttrib = nullptr;
+	// always enable material ID
+	MeshOut.Attributes()->EnableMaterialID();
+	FDynamicMeshMaterialAttribute* MaterialIDAttrib = MeshOut.Attributes()->GetMaterialID();
 
 	// NOTE: If you change the iteration order here, please update the corresponding iteration in FDynamicMeshToMeshDescription::UpdateAttributes, 
 	//	which assumes the iteration order here is polygons -> triangles, to correspond the triangles when writing updated attributes back!
@@ -182,11 +183,6 @@ void FMeshDescriptionToDynamicMesh::Convert(const FMeshDescription* MeshIn, FDyn
 	for (const FPolygonID PolygonID : Polygons.GetElementIDs())
 	{
 		int32 PolygonGroupID = MeshIn->GetPolygonPolygonGroup(PolygonID).GetValue();
-		if (MaterialIDAttrib == nullptr && PolygonGroupID != 0)
-		{
-			MeshOut.Attributes()->EnableMaterialID();		// initializes all existing values to 0, which is fine since this is first non-zero
-			MaterialIDAttrib = MeshOut.Attributes()->GetMaterialID();
-		}
 
 		const TArray<FTriangleID>& TriangleIDs = MeshIn->GetPolygonTriangleIDs(PolygonID);
 		int NumTriangles = TriangleIDs.Num();
diff --git a/Engine/Plugins/Experimental/GeometryProcessing/Source/MeshConversion/Public/DynamicMeshToMeshDescription.h b/Engine/Plugins/Experimental/GeometryProcessing/Source/MeshConversion/Public/DynamicMeshToMeshDescription.h
--- a/Engine/Plugins/Experimental/GeometryProcessing/Source/MeshConversion/Public/DynamicMeshToMeshDescription.h
+++ b/Engine/Plugins/Experimental/GeometryProcessing/Source/MeshConversion/Public/DynamicMeshToMeshDescription.h
@@ -38,7 +38,7 @@ public:
 	 * Update existing MeshDescription based on DynamicMesh. Assumes mesh topology has not changed.
 	 * Copies positions and normals; does not update UVs
 	 */
-	void Update(const FDynamicMesh3* MeshIn, FMeshDescription& MeshOut);
+	void Update(const FDynamicMesh3* MeshIn, FMeshDescription& MeshOut, bool bUpdateNormals = true, bool bUpdateUVs = false);
 
 
 	/**
diff --git a/Engine/Plugins/Experimental/MeshModelingToolset/Source/MeshModelingTools/Private/BakeTransformTool.cpp b/Engine/Plugins/Experimental/MeshModelingToolset/Source/MeshModelingTools/Private/BakeTransformTool.cpp
new file mode 100644
--- /dev/null
+++ b/Engine/Plugins/Experimental/MeshModelingToolset/Source/MeshModelingTools/Private/BakeTransformTool.cpp
@@ -0,0 +1,205 @@
+// Copyright 1998-2019 Epic Games, Inc. All Rights Reserved.
+
+#include "BakeTransformTool.h"
+#include "InteractiveToolManager.h"
+#include "ToolBuilderUtil.h"
+
+#include "ToolSetupUtil.h"
+
+#include "DynamicMesh3.h"
+#include "BaseBehaviors/MultiClickSequenceInputBehavior.h"
+#include "Selection/SelectClickedAction.h"
+#include "MeshTransforms.h"
+
+#include "MeshDescriptionToDynamicMesh.h"
+#include "DynamicMeshToMeshDescription.h"
+
+#include "InteractiveGizmoManager.h"
+
+#include "AssetGenerationUtil.h"
+
+
+#define LOCTEXT_NAMESPACE "UBakeTransformTool"
+
+
+/*
+ * ToolBuilder
+ */
+
+
+bool UBakeTransformToolBuilder::CanBuildTool(const FToolBuilderState& SceneState) const
+{
+	return ToolBuilderUtil::CountComponents(SceneState, CanMakeComponentTarget) > 0;
+}
+
+UInteractiveTool* UBakeTransformToolBuilder::BuildTool(const FToolBuilderState& SceneState) const
+{
+	UBakeTransformTool* NewTool = NewObject<UBakeTransformTool>(SceneState.ToolManager);
+
+	TArray<UActorComponent*> Components = ToolBuilderUtil::FindAllComponents(SceneState, CanMakeComponentTarget);
+	check(Components.Num() > 0);
+
+	TArray<TUniquePtr<FPrimitiveComponentTarget>> ComponentTargets;
+	for (UActorComponent* ActorComponent : Components)
+	{
+		auto* MeshComponent = Cast<UPrimitiveComponent>(ActorComponent);
+		if (MeshComponent)
+		{
+			ComponentTargets.Add(MakeComponentTarget(MeshComponent));
+		}
+	}
+
+	NewTool->SetSelection(MoveTemp(ComponentTargets));
+	NewTool->SetWorld(SceneState.World);
+	NewTool->SetAssetAPI(AssetAPI);
+
+	return NewTool;
+}
+
+
+
+
+/*
+ * Tool
+ */
+
+UBakeTransformToolProperties::UBakeTransformToolProperties()
+{
+	bRecomputeNormals = true;
+}
+
+
+
+UBakeTransformTool::UBakeTransformTool()
+{
+}
+
+void UBakeTransformTool::SetWorld(UWorld* World)
+{
+	this->TargetWorld = World;
+}
+
+void UBakeTransformTool::Setup()
+{
+	UInteractiveTool::Setup();
+
+	// hide input StaticMeshComponent
+	for (int k = 0; k < ComponentTargets.Num(); ++k)
+	{
+		TUniquePtr<FPrimitiveComponentTarget>& ComponentTarget = ComponentTargets[k];
+		ComponentTarget->SetOwnerVisibility(false);
+
+		FMeshDescriptionToDynamicMesh Converter;
+		TSharedPtr<FDynamicMesh3> Mesh = MakeShared<FDynamicMesh3>();
+		Converter.Convert(ComponentTarget->GetMesh(), *Mesh);
+		OriginalDynamicMeshes.Add(Mesh);
+
+		FTransform3d CurTransform = FTransform3d(ComponentTarget->GetWorldTransform());
+		FTransform3d BakeScaleTransform;
+		BakeScaleTransform.SetScale(CurTransform.GetScale());
+		CurTransform.SetScale(FVector3d::One());
+		FDynamicMesh3 BakeScaleMesh(*Mesh);
+		MeshTransforms::ApplyTransform(BakeScaleMesh, BakeScaleTransform);
+
+		UPreviewMesh* Preview = NewObject<UPreviewMesh>(this);
+		Preview->CreateInWorld(this->TargetWorld, FTransform::Identity);
+		Preview->UpdatePreview(&BakeScaleMesh);
+		Preview->SetTransform((FTransform)CurTransform);
+		UMaterialInterface* Material = ComponentTarget->GetMaterial(0);
+		if (Material != nullptr)
+		{
+			Preview->SetMaterial(Material);
+		}
+		Previews.Add(Preview);
+	}
+
+	BasicProperties = NewObject<UBakeTransformToolProperties>(this, TEXT("Mesh Normals Settings"));
+	AddToolPropertySource(BasicProperties);
+
+	GetToolManager()->DisplayMessage(
+		LOCTEXT("BakeTransformWarning", "WARNING: This Tool will Modify the selected StaticMesh Assets! If you do not wish to modify the original Assets, please make copies in the Content Browser first!"),
+		EToolMessageLevel::UserWarning);
+}
+
+
+
+void UBakeTransformTool::Shutdown(EToolShutdownType ShutdownType)
+{
+	// Restore (unhide) the source meshes
+	for (TUniquePtr<FPrimitiveComponentTarget>& ComponentTarget : ComponentTargets)
+	{
+		ComponentTarget->SetOwnerVisibility(true);
+	}
+
+	TArray<TUniquePtr<FDynamicMesh3>> ResultMeshes;
+	for (int k = 0; k < Previews.Num(); ++k)
+	{
+		ResultMeshes.Add(Previews[k]->ExtractPreviewMesh());
+		Previews[k]->Disconnect();
+	}
+	Previews.SetNum(0);
+
+	if (ShutdownType == EToolShutdownType::Accept)
+	{
+		UpdateAssets(ResultMeshes);
+	}
+}
+
+void UBakeTransformTool::SetAssetAPI(IToolsContextAssetAPI* AssetAPIIn)
+{
+	this->AssetAPI = AssetAPIIn;
+}
+
+
+
+void UBakeTransformTool::Render(IToolsContextRenderAPI* RenderAPI)
+{
+}
+
+void UBakeTransformTool::Tick(float DeltaTime)
+{
+}
+
+
+
+bool UBakeTransformTool::HasAccept() const
+{
+	return true;
+}
+
+bool UBakeTransformTool::CanAccept() const
+{
+	return true;
+}
+
+
+void UBakeTransformTool::UpdateAssets(const TArray<TUniquePtr<FDynamicMesh3>>& Results)
+{
+	GetToolManager()->BeginUndoTransaction(LOCTEXT("BakeTransformToolTransactionName", "Bake Transforms"));
+
+	check(Results.Num() == ComponentTargets.Num());
+
+	for (int32 ComponentIdx = 0; ComponentIdx < ComponentTargets.Num(); ComponentIdx++)
+	{
+		ComponentTargets[ComponentIdx]->CommitMesh([&Results, &ComponentIdx, this](const FPrimitiveComponentTarget::FCommitParams& CommitParams)
+		{
+			FDynamicMeshToMeshDescription Converter;
+
+			Converter.Update(Results[ComponentIdx].Get(), *CommitParams.MeshDescription, true, false);
+
+			// otherwise just copy attributes
+			Converter.UpdateAttributes(Results[ComponentIdx].Get(), *CommitParams.MeshDescription, true, false);
+		});
+
+		UPrimitiveComponent* Component = ComponentTargets[ComponentIdx]->GetOwnerComponent();
+		Component->Modify();
+		Component->SetRelativeScale3D(FVector(1, 1, 1));
+	}
+
+	GetToolManager()->EndUndoTransaction();
+}
+
+
+
+
+#undef LOCTEXT_NAMESPACE
diff --git a/Engine/Plugins/Experimental/MeshModelingToolset/Source/MeshModelingTools/Private/ConvertToPolygonsTool.cpp b/Engine/Plugins/Experimental/MeshModelingToolset/Source/MeshModelingTools/Private/ConvertToPolygonsTool.cpp
--- a/Engine/Plugins/Experimental/MeshModelingToolset/Source/MeshModelingTools/Private/ConvertToPolygonsTool.cpp
+++ b/Engine/Plugins/Experimental/MeshModelingToolset/Source/MeshModelingTools/Private/ConvertToPolygonsTool.cpp
@@ -70,10 +70,13 @@ void UConvertToPolygonsTool::Setup()
 	PreviewMesh->SetVisible(false);
 	PreviewMesh->SetTransform(ComponentTarget->GetWorldTransform());
 
-	if (ComponentTarget->GetMaterial(0) != nullptr)
-	{
-		PreviewMesh->SetMaterial(ComponentTarget->GetMaterial(0));
-	}
+	FComponentMaterialSet MaterialSet;
+	ComponentTarget->GetMaterialSet(MaterialSet);
+	PreviewMesh->SetMaterials(MaterialSet.Materials);
+
+	ConvertModeWatcher.Initialize(
+		[this]() { return Settings->ConversionMode; },
+		[this](EConvertToPolygonsMode NewMode) { bPolygonsValid = false; }, Settings->ConversionMode);
 
 	UpdatePolygons();
 }
@@ -89,10 +92,10 @@ void UConvertToPolygonsTool::Shutdown(EToolShutdownType ShutdownType)
 
 	if (ShutdownType == EToolShutdownType::Accept)
 	{
-		GetToolManager()->BeginUndoTransaction(LOCTEXT("DeformMeshToolTransactionName", "Convert to Polygons"));
-		ComponentTarget->CommitMesh([=](FMeshDescription* MeshDescription)
+		GetToolManager()->BeginUndoTransaction(LOCTEXT("ConvertToPolygonsToolTransactionName", "Convert to Polygons"));
+		ComponentTarget->CommitMesh([=](const FPrimitiveComponentTarget::FCommitParams& CommitParams)
 		{
-			ConvertToPolygons(MeshDescription);
+			ConvertToPolygons(CommitParams.MeshDescription);
 		});
 		GetToolManager()->EndUndoTransaction();
 	}
@@ -106,14 +109,19 @@ void UConvertToPolygonsTool::OnPropertyModified(UObject* PropertySet, UProperty*
 }
 
 
-void UConvertToPolygonsTool::Render(IToolsContextRenderAPI* RenderAPI)
+void UConvertToPolygonsTool::Tick(float DeltaTime)
 {
-	FColor LineColor(255, 0, 0);
+	ConvertModeWatcher.CheckAndUpdate();
 
 	if (bPolygonsValid == false)
 	{
 		UpdatePolygons();
 	}
+}
+
+void UConvertToPolygonsTool::Render(IToolsContextRenderAPI* RenderAPI)
+{
+	FColor LineColor(255, 0, 0);
 
 	if (true)
 	{
@@ -125,7 +133,7 @@ void UConvertToPolygonsTool::Render(IToolsContextRenderAPI* RenderAPI)
 		{
 			FVector3d A, B;
 			Polygons.Mesh->GetEdgeV(eid, A, B);
-			PDI->DrawLine(Transform.TransformPosition(A), Transform.TransformPosition(B),
+			PDI->DrawLine(Transform.TransformPosition((FVector)A), Transform.TransformPosition((FVector)B),
 				LineColor, 0, 2.0, 1.0f, true);
 		}
 	}
@@ -146,13 +154,25 @@ bool UConvertToPolygonsTool::CanAccept() const
 void UConvertToPolygonsTool::UpdatePolygons()
 {
 	Polygons = FFindPolygonsAlgorithm(&SearchMesh);
-	double DotTolerance = 1.0 - FMathd::Cos(Settings->AngleTolerance * FMathd::DegToRad);
-	Polygons.FindPolygons(DotTolerance);
+	if (Settings->ConversionMode == EConvertToPolygonsMode::FromUVISlands)
+	{
+		Polygons.FindPolygonsFromUVIslands();
+	}
+	else if (Settings->ConversionMode == EConvertToPolygonsMode::FaceNormalDeviation)
+	{
+		double DotTolerance = 1.0 - FMathd::Cos(Settings->AngleTolerance * FMathd::DegToRad);
+		Polygons.FindPolygonsFromFaceNormals(DotTolerance);
+	}
+	else
+	{
+		check(false);
+	}
+	
 	Polygons.FindPolygonEdges();
 
 	GetToolManager()->DisplayMessage(
-		FText::Format(LOCTEXT("UpdatePolygonsMessage", "ConvertToPolygons - found {0} polys in {1} triangles"), 
-			FText::AsNumber(Polygons.FoundPolygons.Num()), FText::AsNumber(SearchMesh.TriangleCount()) ), EToolMessageLevel::Internal);
+		FText::Format(LOCTEXT("UpdatePolygonsMessage", "ConvertToPolygons - found {0} polys in {1} triangles"),
+			FText::AsNumber(Polygons.FoundPolygons.Num()), FText::AsNumber(SearchMesh.TriangleCount())), EToolMessageLevel::Internal);
 
 	if (Settings->bCalculateNormals)
 	{
diff --git a/Engine/Plugins/Experimental/MeshModelingToolset/Source/MeshModelingTools/Private/DeformMeshPolygonsTool.cpp b/Engine/Plugins/Experimental/MeshModelingToolset/Source/MeshModelingTools/Private/DeformMeshPolygonsTool.cpp
new file mode 100644
--- /dev/null
+++ b/Engine/Plugins/Experimental/MeshModelingToolset/Source/MeshModelingTools/Private/DeformMeshPolygonsTool.cpp
@@ -0,0 +1,1472 @@
+// Copyright 1998-2019 Epic Games, Inc. All Rights Reserved.
+
+#include "DeformMeshPolygonsTool.h"
+#include "InteractiveToolManager.h"
+#include "ToolBuilderUtil.h"
+
+#include "SegmentTypes.h"
+#include "DynamicMeshAttributeSet.h"
+#include "MeshNormals.h"
+#include "ToolSceneQueriesUtil.h"
+#include "Intersection/IntersectionUtil.h"
+#include "FindPolygonsAlgorithm.h"
+
+#include "Async/ParallelFor.h"
+#include "Containers/BitArray.h"
+
+#define LOCTEXT_NAMESPACE "UDeformMeshPolygonsTool"
+
+
+
+//////////////////////////////
+// DEBUG_SETTINGS
+
+//Draw white triangles defining the selection subset
+//#define DEBUG_ROI_TRIANGLES
+
+//Draw pink circles around the handles
+//#define DEBUG_ROI_HANDLES
+
+//Draw points on the ROI vertices, White => Weight == 0, Black => Weight == 1
+//#define DEBUG_ROI_WEIGHTS
+
+//////////////////////////////
+
+
+/*
+ * ToolBuilder
+ */
+UMeshSurfacePointTool* UDeformMeshPolygonsToolBuilder::CreateNewTool(const FToolBuilderState& SceneState) const
+{
+	UDeformMeshPolygonsTool* DeformTool = NewObject<UDeformMeshPolygonsTool>(SceneState.ToolManager);
+	return DeformTool;
+}
+
+/*
+ * Tool
+ */
+UDeformMeshPolygonsTransformProperties::UDeformMeshPolygonsTransformProperties()
+{
+	DeformationStrategy = EGroupTopologyDeformationStrategy::Laplacian;
+	TransformMode = EQuickTransformerMode::AxisTranslation;
+	bSelectVertices = true;
+	bSelectFaces = true;
+	bSelectEdges = true;
+	bShowWireframe = false;
+	bSnapToWorldGrid = false;
+}
+
+#if WITH_EDITOR
+void UDeformMeshPolygonsTransformProperties::PostEditChangeProperty(FPropertyChangedEvent& PropertyChangedEvent)
+{
+	Super::PostEditChangeProperty(PropertyChangedEvent);
+}
+#endif
+
+/*
+ * Asynchronous Task
+ */
+
+
+void FConstrainedMeshDeformerTask::UpdateDeformer(const ELaplacianWeightScheme SelectedWeightScheme, const FDynamicMesh3& SrcMesh, const TArray<FConstraintData>& ConstraintArray, const TArray<int32>& SrcIDBufferSubset, bool bNewTransaction, const FRichCurve* Curve)
+{
+	bIsNewTransaction = bNewTransaction;
+	SrcMeshMaxVertexID = SrcMesh.MaxVertexID();
+
+	LaplacianWeightScheme = SelectedWeightScheme;
+	
+	bAttenuateWeights = (Curve != nullptr);
+	if (bAttenuateWeights)
+	{
+		WeightAttenuationCurve = *Curve;
+	}
+
+	// Set-up the subset mesh.
+	if (bIsNewTransaction)
+	{
+		//Copy the part of the mesh we want to deform into the SubsetMesh and create map from Src Mesh to the SubsetMesh.
+		InitializeSubsetMesh(SrcMesh, SrcIDBufferSubset);
+	}
+
+	// only want the subset of constraints that correspond to our subset mesh.
+	{
+		const int32 NumSubsetVerts = SubsetVertexIDToSrcVertexIDMap.Num();
+		SubsetConstraintBuffer.Empty(NumSubsetVerts);
+		SubsetConstraintBuffer.AddUninitialized(NumSubsetVerts);
+
+		for (int32 SubVertexID = 0; SubVertexID < SubsetVertexIDToSrcVertexIDMap.Num(); ++SubVertexID)
+		{
+			int32 SrcVtxID = SubsetVertexIDToSrcVertexIDMap[SubVertexID];
+			SubsetConstraintBuffer[SubVertexID] = ConstraintArray[SrcVtxID];
+		}
+	}
+
+	check(bIsNewTransaction || ConstrainedDeformer.IsValid());
+	
+}
+
+
+void FConstrainedMeshDeformerTask::DoWork()
+{
+
+	//TODO: (simple optimization) -
+	//	Instead of SrcVertexIDtoSubsetVertexIDMap, use SubsetVertexIDToSetVertexIDMap - then we can use the VertexIndicesItr()
+	//	on the SubsetMesh to minimize the quantity of vertex indices we need to iterate at every following step.
+
+	if (*bAbortSource == true)
+	{
+		return;
+	}
+
+	if (bIsNewTransaction) //Will only be true once per input transaction (click+drag)
+	{
+		
+		// Create a new deformation solver.
+		ConstrainedDeformer = MakeUnique<FConstrainedMeshDeformer>(SubsetMesh, LaplacianWeightScheme);
+
+		if (bAttenuateWeights)
+		{
+			ApplyAttenuation();
+		}
+
+		//Update our deformer's constraints before deforming using the copy of the constraint buffer
+		for (int32 SubsetVertexID = 0; SubsetVertexID < SubsetConstraintBuffer.Num(); ++SubsetVertexID)
+		{
+			FConstraintData& CData = SubsetConstraintBuffer[SubsetVertexID];
+			ConstrainedDeformer->AddConstraint(SubsetVertexID, CData.Weight, CData.Position, CData.bPostFix);
+		}
+
+		bIsNewTransaction = false;
+
+	}
+	else
+	{
+		//This else block is run every consecutive frame after the start of the input transaction because UpdateConstraintPosition() is very cheap (no factorizing or rebuilding)
+		//Update only the positions of the constraints, as the weights cannot change mid-transaction
+
+		for (int32 SubsetVertexID = 0; SubsetVertexID < SubsetConstraintBuffer.Num(); ++SubsetVertexID)
+		{
+			FConstraintData& CData = SubsetConstraintBuffer[SubsetVertexID];
+			ConstrainedDeformer->UpdateConstraintPosition(SubsetVertexID, CData.Position, CData.bPostFix);
+		}
+
+	}
+
+	if (*bAbortSource == true)
+	{
+		return;
+	}
+
+	//Run the deformation process
+
+	const bool bSuccessfulSolve = ConstrainedDeformer->Deform(SubsetPositionBuffer);
+	
+	
+	if (bSuccessfulSolve)
+	{
+		if (*bAbortSource == true)
+		{
+			return;
+		}
+	}
+	else
+	{
+		//UE_LOG(LogTemp, Warning, TEXT("Laplacian deformation failed"));
+	}
+
+}
+
+inline void FConstrainedMeshDeformerTask::InitializeSubsetMesh(const FDynamicMesh3& SrcMesh, const TArray<int32>& SrcIDBufferSubset)
+{
+	//These can be re-used until the user stops dragging
+	SubsetMesh.Clear();
+	SubsetPositionBuffer.Reset();
+
+	//Initialize every element to -1, helps us keep track of vertices we've already added while iterating the triangles
+	TArray<int32> SrcVertexIDToSubsetVertexIDMap;
+	SrcVertexIDToSubsetVertexIDMap.Init(INACTIVE_SUBSET_ID, SrcMeshMaxVertexID);
+
+
+	//Iterate the triangle array to append vertices, and then triangles to the temporary subset mesh all at once
+	for (int32 i = 0; i < SrcIDBufferSubset.Num(); i += 3)
+	{
+
+		// Build the triangle
+		FIndex3i Triangle;
+		for (int32 v = 0; v < 3; ++v)
+		{
+			//It's the SrcVertexID because every element in the SrcIDBufferSubset is the Vertex ID of a vertex in the original mesh.
+			const int32 SrcVertexID = SrcIDBufferSubset[i + v];
+			int32& SubsetID = SrcVertexIDToSubsetVertexIDMap[SrcVertexID];
+			
+			if (SubsetID == INACTIVE_SUBSET_ID) // we haven't already visited this vertex
+			{
+				const FVector3d Vertex = SrcMesh.GetVertex(SrcVertexID);
+				SubsetID = SubsetMesh.AppendVertex(Vertex);
+			}
+
+			Triangle[v] = SubsetID;
+		}
+		SubsetMesh.AppendTriangle(Triangle);
+	}
+
+	// create a mapping back to the original vertex IDs from the subset mesh
+	int32 MaxSubMeshVertexID = SubsetMesh.MaxVertexID(); // Really MaxID + 1
+	SubsetVertexIDToSrcVertexIDMap.Reset(MaxSubMeshVertexID);
+	SubsetVertexIDToSrcVertexIDMap.AddUninitialized(MaxSubMeshVertexID);
+
+	for (int32 SrcID = 0; SrcID < SrcVertexIDToSubsetVertexIDMap.Num(); ++SrcID)
+	{
+		const int32 SubsetVertexID = SrcVertexIDToSubsetVertexIDMap[SrcID];
+		if (SubsetVertexID != INACTIVE_SUBSET_ID)
+		{
+			SubsetVertexIDToSrcVertexIDMap[SubsetVertexID] = SrcID;
+		}
+	}
+
+}
+
+void FConstrainedMeshDeformerTask::ExportResults(FDynamicMesh3& TargetMesh) const
+{
+
+	//Update the position buffer result
+	for (int32 SubsetVertexID = 0; SubsetVertexID < SubsetVertexIDToSrcVertexIDMap.Num(); ++SubsetVertexID)
+	{
+		const int32 SrcVertexID = SubsetVertexIDToSrcVertexIDMap[SubsetVertexID];
+		const FVector3d Position = SubsetPositionBuffer[SubsetVertexID];
+		
+		TargetMesh.SetVertex(SrcVertexID, Position);
+	}
+}
+
+void FConstrainedMeshDeformerTask::ApplyAttenuation()
+{
+	size_t Size = SrcMeshMaxVertexID;
+	TSet<int> Handles;
+
+	auto InPlaceMinMaxElements = [](FVector3d& Min, FVector3d& Max, const FVector3d Test)
+	{
+		for (uint8 i = 0; i < 3; ++i)
+		{
+			Min[i] = Test[i] < Min[i] ? Test[i] : Min[i];
+			Max[i] = Test[i] > Max[i] ? Test[i] : Max[i];
+		}
+	};
+
+	//Experimental approach: Just going to try grabbing the bounding box of the entire mesh, then the bounding box of the handles as a point cloud.
+	//						 We need a T value to pass to the Weights curve, so let's try finding the distance of each vertex V from line segment formed by the min/max handles
+	//						 Divide the distance from the handles to vertex V by the length of the mesh's bounding box extent, 
+	//						 and that will provide a **ROUGH** approximation of the time value for our curve.
+	//
+	//											  Distance( LineSegment(MaxHandle,MinHandle) , V )
+	// where T(V) is time value at V     T(V) = -----------------------------------------------------
+	//   and V is the position								Length(MeshMin - MeshMax)   
+	//      of each vertex
+
+	FVector3d Min{ std::numeric_limits<double>::max(),std::numeric_limits<double>::max(),std::numeric_limits<double>::max() };
+	FVector3d Max{ std::numeric_limits<double>::min(),std::numeric_limits<double>::min(),std::numeric_limits<double>::min() };
+	FVector3d MinHandles = Min;
+	FVector3d MaxHandles = Max;
+	double LeastWeight = std::numeric_limits<double>::max();
+
+
+
+	for (int32 SubVertexID = 0; SubVertexID < SubsetConstraintBuffer.Num(); ++SubVertexID)
+	{
+
+		FConstraintData& CData = SubsetConstraintBuffer[SubVertexID];
+		// Update bounding box
+		InPlaceMinMaxElements(Min, Max, CData.Position);
+		
+		if (CData.Weight > 0.0)
+		{
+			LeastWeight = CData.Weight < LeastWeight ? CData.Weight : LeastWeight;
+
+			// update bounding box
+			InPlaceMinMaxElements(MinHandles, MaxHandles, CData.Position);
+			Handles.Add(SubVertexID);
+		}
+	}
+
+	double ExtentLength = Min.Distance(Max);
+
+	// Is this why the system has memory?
+	for (int32 SubVertexID = 0; SubVertexID < SubsetConstraintBuffer.Num(); ++SubVertexID)
+	{
+		if (!Handles.Contains(SubVertexID))
+		{
+			FConstraintData& CData = SubsetConstraintBuffer[SubVertexID];
+			double T = CData.Position.Distance(FMath::ClosestPointOnSegment((FVector)CData.Position, (FVector)MinHandles, (FVector)MaxHandles)) / ExtentLength;
+			CData.Weight = WeightAttenuationCurve.Eval(T) * LeastWeight;
+		}
+
+	}
+}
+
+
+/*
+*	FGroupTopologyLaplacianDeformer methods
+*/
+
+
+void FGroupTopologyLaplacianDeformer::InitBackgroundWorker(const ELaplacianWeightScheme WeightScheme)
+{
+
+	//Initialize asynchronous deformation objects
+	if (AsyncMeshDeformTask == nullptr)
+	{
+		AsyncMeshDeformTask = new FAsyncTaskExecuterWithAbort<FConstrainedMeshDeformerTask>(WeightScheme);
+	}
+}
+
+void FGroupTopologyLaplacianDeformer::InitializeConstraintBuffer()
+{
+	//MaxVertexID is used because the array is potentially sparse.
+	int MaxVertexID = Mesh->MaxVertexID();
+
+	SrcMeshConstraintBuffer.SetNum(MaxVertexID);
+	
+	for (int32 VertexID : Mesh->VertexIndicesItr())
+	{
+		FConstraintData& CD = SrcMeshConstraintBuffer[VertexID];
+		CD.Position = Mesh->GetVertex(VertexID); 
+		CD.Weight   = 0.0;
+		CD.bPostFix = false;
+	}
+}
+
+bool FGroupTopologyLaplacianDeformer::IsTaskInFlight() const
+{
+	return ( AsyncMeshDeformTask != nullptr && !AsyncMeshDeformTask->IsDone() );
+}
+
+bool FGroupTopologyLaplacianDeformer::UpdateAndLaunchdWorker(const ELaplacianWeightScheme SelectedWeightScheme, const FRichCurve*  Curve)
+{
+	/* Deformer needs to run if we've modified the constraints since the last time it finished. */
+	if (AsyncMeshDeformTask == nullptr)
+	{
+		InitBackgroundWorker(SelectedWeightScheme);
+	}
+
+	if (bDeformerNeedsToRun && AsyncMeshDeformTask->IsDone())
+	{		
+		bool bRebuildSubsetMesh = bTaskSubmeshIsDirty;
+
+		FConstrainedMeshDeformerTask& Task = AsyncMeshDeformTask->GetTask();
+
+		// Update the deformer's buffers and weight scheme
+		// this creates the subset mesh if needed.
+		Task.UpdateDeformer(SelectedWeightScheme, *Mesh, SrcMeshConstraintBuffer, SubsetIDBuffer, bRebuildSubsetMesh, Curve);
+
+		// task now has valid submesh
+
+		bTaskSubmeshIsDirty = false;
+
+		//Launch second thread
+		AsyncMeshDeformTask->StartBackgroundTask();
+
+		bDeformerNeedsToRun = false;	 // This was set to true above in UpdateSolution()
+		bVertexPositionsNeedSync = true; // The task will generate new vertex positions. 
+
+		return true;
+
+	}
+	return false;
+}
+
+void FGroupTopologyLaplacianDeformer::SetActiveHandleFaces(const TArray<int>& FaceGroupIDs)
+{
+	Reset();
+
+	check(FaceGroupIDs.Num() == 1);   // multi-face not supported yet
+	int GroupID = FaceGroupIDs[0];
+
+	// find set of vertices in handle 
+	Topology->CollectGroupVertices(GroupID, HandleVertices);
+	Topology->CollectGroupBoundaryVertices(GroupID, HandleBoundaryVertices);
+	ModifiedVertices = HandleVertices;
+
+
+
+	// list of adj groups.  may contain duplicates.
+	TArray<int> AdjGroups;
+	for (int BoundaryVert : HandleBoundaryVertices)
+	{
+		Topology->FindVertexNbrGroups(BoundaryVert, AdjGroups);
+	}
+
+	// Local neighborhood - Adjacent groups plus self
+	TArray<int> NeighborhoodGroups;
+
+	// Collect the rest of the 1-ring groups that are adjacent to the selected one.
+	NeighborhoodGroups.Add(GroupID);
+	for (int AdjGroup : AdjGroups)
+	{
+		NeighborhoodGroups.AddUnique(AdjGroup); // remove duplicates by add unique
+	}
+
+	CalculateROI(FaceGroupIDs, NeighborhoodGroups);
+	
+	UpdateSelection(Mesh, NeighborhoodGroups, bLocalize);
+
+	// Save the positions of the selected region.
+	SaveInitialPositions();
+
+}
+
+
+void FGroupTopologyLaplacianDeformer::SetActiveHandleEdges(const TArray<int>& TopologyEdgeIDs)
+{
+	Reset();
+
+	for (int EdgeID : TopologyEdgeIDs)
+	{
+		const TArray<int>& EdgeVerts = Topology->GetGroupEdgeVertices(EdgeID);
+		for (int VertID : EdgeVerts)
+		{
+			HandleVertices.Add(VertID);
+		}
+	}
+	HandleBoundaryVertices = HandleVertices;
+	ModifiedVertices = HandleVertices;
+
+	TArray<int> HandleGroups;
+	TArray<int> NbrGroups;
+	Topology->FindEdgeNbrGroups(TopologyEdgeIDs, NbrGroups);
+
+	CalculateROI(HandleGroups, NbrGroups);
+
+	UpdateSelection(Mesh, NbrGroups, bLocalize);
+
+	// Save the positions of the selected region.
+	SaveInitialPositions();
+}
+
+void FGroupTopologyLaplacianDeformer::SetActiveHandleCorners(const TArray<int>& CornerIDs)
+{
+	Reset();
+
+	for (int CornerID : CornerIDs)
+	{
+		int VertID = Topology->GetCornerVertexID(CornerID);
+		if (VertID >= 0)
+		{
+			HandleVertices.Add(VertID);
+		}
+	}
+	HandleBoundaryVertices = HandleVertices;
+	ModifiedVertices = HandleVertices;
+
+	TArray<int> HandleGroups;
+	TArray<int> NbrGroups;
+
+	Topology->FindCornerNbrGroups(CornerIDs, NbrGroups);
+
+
+	CalculateROI(HandleGroups, NbrGroups);
+
+	UpdateSelection(Mesh, NbrGroups, bLocalize);
+
+	// Save the positions of the selected region.
+	SaveInitialPositions();
+}
+
+
+
+void FGroupTopologyLaplacianDeformer::UpdateSelection(const FDynamicMesh3* TargetMesh, const TArray<int>& Groups, bool bLocalizeDeformation)
+{
+
+	// Build an index buffer (SubsetIdBuffer) and a vertexId buffer (ModifidedVertices) for the region we want to change
+
+	if (bLocalizeDeformation)
+	{
+		//For each group ID, retrieve the array of all TriangleIDs associated with that GroupID and append that array to the end of the TriSet to remove duplicates
+		TSet<int> TriSet;
+		for (const int32& GroupID : Groups)
+		{
+			TriSet.Append(Topology->GetGroupFaces(GroupID));
+		}//Now we have every triangle ID involved in the transaction 
+
+		//Since we are flattening the Face to a set of 3 indices, we do 3 * number of triangles though it is too many.
+		SubsetIDBuffer.Reset(3 * TriSet.Num());
+		//Add each triangle's A,B, and C indices to the subset triangle array.
+		for (const int& Tri : TriSet)
+		{
+			FIndex3i Triple = TargetMesh->GetTriangle(Tri);
+			SubsetIDBuffer.Add(Triple.A);
+			SubsetIDBuffer.Add(Triple.B);
+			SubsetIDBuffer.Add(Triple.C);
+		
+		}
+	}
+	else
+	{
+		// the entire mesh.
+		const int32 NumTris = TargetMesh->TriangleCount();
+		SubsetIDBuffer.Reset(3 * NumTris);
+		for (int TriId : TargetMesh->TriangleIndicesItr()) 
+		{ 
+			FIndex3i Triple = TargetMesh->GetTriangle(TriId);
+			SubsetIDBuffer.Add(Triple.A);
+			SubsetIDBuffer.Add(Triple.B);
+			SubsetIDBuffer.Add(Triple.C);
+		}
+	}
+
+	// Add the vertices to the set (eliminates duplicates.)  Todo: don't use a set.
+	ResetModifiedVertices();
+	for (int32 VertexID : SubsetIDBuffer)
+	{
+		RecordModifiedVertex(VertexID);
+	}
+
+}
+
+// This actually updates constraints that correspond to the handle vertices.
+void FGroupTopologyLaplacianDeformer::UpdateSolution(FDynamicMesh3 * TargetMesh, const TFunction<FVector3d(FDynamicMesh3* Mesh, int)>& HandleVertexDeformFunc)
+{
+	// copy the current positions.
+	FVertexPositionCache CurrentPositions;
+	for (int VertexID : InitialPositions.Vertices)
+	{
+		CurrentPositions.AddVertex(TargetMesh, VertexID);
+	}
+
+	// Set the target mesh to the initial positions.
+	// Note: this only updates the vertices in the selected region.
+	InitialPositions.SetPositions(TargetMesh);
+
+	//Reset the constraints
+	for (int32 VertexID : ModifiedVertices)
+	{
+		//Get the vertex's data from the constraint buffer
+		FConstraintData& CData = SrcMeshConstraintBuffer[VertexID];
+
+		CData.Position = TargetMesh->GetVertex(VertexID);
+		CData.Weight   = 0.0; //A weight of zero is used to allow this point to move freely when moving the handles
+		CData.bPostFix = false;
+	}
+
+	//Actually deform the handles and add a constraint.
+	for (int VertexID : HandleVertices)
+	{
+		const FVector3d DeformPos = HandleVertexDeformFunc(TargetMesh, VertexID);
+
+		//Get the vertex's data from the constraint buffer
+		FConstraintData& CData = SrcMeshConstraintBuffer[VertexID];
+
+		//Set the new vertex data
+		CData.Position = DeformPos;
+		CData.Weight   = HandleWeights;
+		CData.bPostFix = bPostfixHandles;
+	}
+
+	// Restore Current Positions.  This is done because the target mesh is being used to define the highlight region.
+	// if we don't reset the positions the highlight mesh will appear to reset momentarily until the first laplacian solver result is available 
+	CurrentPositions.SetPositions(TargetMesh);
+
+	bDeformerNeedsToRun = true;
+
+}
+
+void FGroupTopologyLaplacianDeformer::ExportDeformedPositions(FDynamicMesh3* TargetMesh)
+{
+	bool bIsWorking = IsTaskInFlight();
+	if (AsyncMeshDeformTask != nullptr && !bIsWorking)
+	{
+		const FConstrainedMeshDeformerTask& Task = AsyncMeshDeformTask->GetTask();
+		Task.ExportResults(*TargetMesh);
+	}
+}
+
+inline FGroupTopologyLaplacianDeformer::~FGroupTopologyLaplacianDeformer()
+{
+	Shutdown();
+}
+
+inline void FGroupTopologyLaplacianDeformer::Shutdown()
+{	
+	if (AsyncMeshDeformTask != nullptr)
+	{
+		if (AsyncMeshDeformTask->IsDone())
+		{
+			delete AsyncMeshDeformTask;
+		}
+		else
+		{
+			AsyncMeshDeformTask->CancelAndDelete();
+		}
+
+		AsyncMeshDeformTask = nullptr;
+	}
+}
+
+/*
+* Tool methods
+*/
+
+UDeformMeshPolygonsTool::UDeformMeshPolygonsTool()
+{
+}
+
+void UDeformMeshPolygonsTool::Setup()
+{
+	UMeshSurfacePointTool::Setup();
+
+	// create dynamic mesh component to use for live preview
+	DynamicMeshComponent = NewObject<USimpleDynamicMeshComponent>(ComponentTarget->GetOwnerActor(), "DynamicMesh");
+	DynamicMeshComponent->SetupAttachment(ComponentTarget->GetOwnerActor()->GetRootComponent());
+	DynamicMeshComponent->RegisterComponent();
+	DynamicMeshComponent->SetWorldTransform(ComponentTarget->GetWorldTransform());
+
+	// set materials
+	FComponentMaterialSet MaterialSet;
+	ComponentTarget->GetMaterialSet(MaterialSet);
+	for (int k = 0; k < MaterialSet.Materials.Num(); ++k)
+	{
+		DynamicMeshComponent->SetMaterial(k, MaterialSet.Materials[k]);
+	}
+
+	// dynamic mesh configuration settings
+	DynamicMeshComponent->TangentsType = EDynamicMeshTangentCalcType::AutoCalculated;
+	DynamicMeshComponent->InitializeMesh(ComponentTarget->GetMesh());
+	OnDynamicMeshComponentChangedHandle = DynamicMeshComponent->OnMeshChanged.Add(
+		FSimpleMulticastDelegate::FDelegate::CreateUObject(this, &UDeformMeshPolygonsTool::OnDynamicMeshComponentChanged));
+
+
+	// add properties
+	TransformProps = NewObject<UDeformMeshPolygonsTransformProperties>(this);
+	AddToolPropertySource(TransformProps);
+
+	// initialize AABBTree
+	MeshSpatial.SetMesh(DynamicMeshComponent->GetMesh());
+	PrecomputeTopology();
+
+	//initialize topology selector
+	TopoSelector.Initialize(DynamicMeshComponent->GetMesh(), &Topology);
+	TopoSelector.SetSpatialSource([this]() {return &GetSpatial(); });
+	TopoSelector.PointsWithinToleranceTest = [this](const FVector3d& Position1, const FVector3d& Position2) {
+		FTransform Transform = ComponentTarget->GetWorldTransform();
+		return ToolSceneQueriesUtil::PointSnapQuery(this->CameraState, 
+			Transform.TransformPosition((FVector)Position1), Transform.TransformPosition((FVector)Position2), VisualAngleSnapThreshold);
+	};
+
+	// hide input StaticMeshComponent
+	ComponentTarget->SetOwnerVisibility(false);
+
+	// init state flags flags
+	bInDrag = false;
+
+	// initialize snap solver
+	QuickAxisTranslater.Initialize();
+	QuickAxisRotator.Initialize();
+
+	// set up visualizers
+	PolyEdgesRenderer.LineColor = FLinearColor::Red;
+	PolyEdgesRenderer.LineThickness = 2.0;
+	HilightRenderer.LineColor = FLinearColor::Green;
+	HilightRenderer.LineThickness = 4.0f;
+
+	// Allocates buffers, sets up the asynchronous task
+	// Copies the source mesh positions.
+	const ELaplacianWeightScheme LaplacianWeightScheme = ConvertToLaplacianWeightScheme(TransformProps->SelectedWeightScheme);
+	LaplacianDeformer.InitBackgroundWorker(LaplacianWeightScheme);
+	
+
+	/**
+	// How to add a curve for the weights.
+	//Add a default curve for falloff
+	FKeyHandle Keys[5];
+	Keys[0] = TransformProps->DefaultFalloffCurve.UpdateOrAddKey(0.f, 1.f);
+	Keys[1] = TransformProps->DefaultFalloffCurve.UpdateOrAddKey(0.25f, 0.25f);
+	Keys[2] = TransformProps->DefaultFalloffCurve.UpdateOrAddKey(0.3333333f, 0.25f);
+	Keys[3] = TransformProps->DefaultFalloffCurve.UpdateOrAddKey(0.6666667f, 1.25f);
+	Keys[4] = TransformProps->DefaultFalloffCurve.UpdateOrAddKey(1.f, 1.4f);
+	for (uint8 i = 0; i < 5; ++i)
+	{
+		TransformProps->DefaultFalloffCurve.SetKeyInterpMode(Keys[i], ERichCurveInterpMode::RCIM_Cubic);
+	}
+	TransformProps->WeightAttenuationCurve.EditorCurveData = TransformProps->DefaultFalloffCurve;
+	*/
+
+	if (Topology.Groups.Num() < 2)
+	{
+		GetToolManager()->DisplayMessage(
+			LOCTEXT("NoGroupsWarning", "This object has no PolyGroups to deform. Use the PolyGroups Tool to create some."),
+			EToolMessageLevel::UserWarning);
+	}
+
+}
+
+void UDeformMeshPolygonsTool::Shutdown(EToolShutdownType ShutdownType)
+{
+	//Tell the background thread to cancel the rest of its jobs before we close;
+	LaplacianDeformer.Shutdown();
+
+	if (DynamicMeshComponent != nullptr)
+	{
+		DynamicMeshComponent->OnMeshChanged.Remove(OnDynamicMeshComponentChangedHandle);
+
+		ComponentTarget->SetOwnerVisibility(true);
+
+		if (ShutdownType == EToolShutdownType::Accept)
+		{
+			// this block bakes the modified DynamicMeshComponent back into the StaticMeshComponent inside an undo transaction
+			GetToolManager()->BeginUndoTransaction(LOCTEXT("DeformMeshPolygonsToolTransactionName", "Deform Mesh"));
+			ComponentTarget->CommitMesh([=](const FPrimitiveComponentTarget::FCommitParams& CommitParams)
+			{
+				FConversionToMeshDescriptionOptions ConversionOptions;
+				ConversionOptions.bSetPolyGroups = false; // don't save polygroups, as we may change these temporarily in this tool just to get a different edit effect
+				DynamicMeshComponent->Bake(CommitParams.MeshDescription, false, ConversionOptions);
+			});
+			GetToolManager()->EndUndoTransaction();
+		}
+
+		DynamicMeshComponent->UnregisterComponent();
+		DynamicMeshComponent->DestroyComponent();
+		DynamicMeshComponent = nullptr;
+	}
+}
+
+
+
+
+void UDeformMeshPolygonsTool::NextTransformTypeAction()
+{
+	if (bInDrag == false)
+	{
+		if (TransformProps->TransformMode == EQuickTransformerMode::AxisRotation)
+		{
+			TransformProps->TransformMode = EQuickTransformerMode::AxisTranslation;
+		}
+		else
+		{
+			TransformProps->TransformMode = EQuickTransformerMode::AxisRotation;
+		}
+		UpdateQuickTransformer();
+	}
+}
+
+
+
+void UDeformMeshPolygonsTool::RegisterActions(FInteractiveToolActionSet& ActionSet)
+{
+	ActionSet.RegisterAction(this, (int32)EStandardToolActions::BaseClientDefinedActionID + 2,
+		TEXT("DeformNextTransformType"),
+		LOCTEXT("DeformNextTransformType", "Next Transform Type"),
+		LOCTEXT("DeformNextTransformTypeTooltip", "Cycle to next transform type"),
+		EModifierKey::None, EKeys::Q,
+		[this]() { NextTransformTypeAction(); });
+}
+
+
+
+
+void UDeformMeshPolygonsTool::OnDynamicMeshComponentChanged()
+{
+	bSpatialDirty = true;
+	TopoSelector.Invalidate(true, false);
+
+
+	//Makes sure the constraint buffer and position buffers reflect Undo/Redo changes
+	FDynamicMesh3* Mesh = DynamicMeshComponent->GetMesh();
+	
+	
+	//Apply Undo/redo
+	for (int VertexID : Mesh->VertexIndicesItr())
+	{
+		const FVector3d Position = Mesh->GetVertex(VertexID);
+		LaplacianDeformer.SrcMeshConstraintBuffer[VertexID].Position = Position;
+	}
+
+	// a deform task could still be in flight.
+	if (LaplacianDeformer.AsyncMeshDeformTask != nullptr)
+	{
+		LaplacianDeformer.AsyncMeshDeformTask->CancelAndDelete();
+		LaplacianDeformer.AsyncMeshDeformTask = nullptr;
+		LaplacianDeformer.bTaskSubmeshIsDirty = true;
+	}
+}
+
+FDynamicMeshAABBTree3& UDeformMeshPolygonsTool::GetSpatial()
+{
+	if (bSpatialDirty)
+	{
+		MeshSpatial.Build();
+		bSpatialDirty = false;
+	}
+	return MeshSpatial;
+}
+
+
+bool UDeformMeshPolygonsTool::HitTest(const FRay& WorldRay, FHitResult& OutHit)
+{
+	FTransform3d Transform(ComponentTarget->GetWorldTransform());
+	FRay3d LocalRay(Transform.InverseTransformPosition(WorldRay.Origin),
+		Transform.InverseTransformVector(WorldRay.Direction));
+	LocalRay.Direction.Normalize();
+
+	TopoSelector.UpdateEnableFlags(TransformProps->bSelectFaces, TransformProps->bSelectEdges, TransformProps->bSelectVertices);
+	FGroupTopologySelection Selection;
+	FVector3d LocalPosition, LocalNormal;
+	if (TopoSelector.FindSelectedElement(LocalRay, Selection, LocalPosition, LocalNormal) == false)
+	{
+		return false;
+	}
+
+	if (Selection.SelectedCornerIDs.Num() > 0)
+	{
+		OutHit.FaceIndex = Selection.SelectedCornerIDs[0];
+		OutHit.Distance = LocalRay.Project(LocalPosition);
+		OutHit.ImpactPoint = (FVector)Transform.TransformPosition(LocalRay.PointAt(OutHit.Distance));
+	}
+	else if (Selection.SelectedEdgeIDs.Num() > 0)
+	{
+		OutHit.FaceIndex = Selection.SelectedEdgeIDs[0];
+		OutHit.Distance = LocalRay.Project(LocalPosition);
+		OutHit.ImpactPoint = (FVector)Transform.TransformPosition(LocalRay.PointAt(OutHit.Distance));
+	}
+	else
+	{
+		int HitTID = GetSpatial().FindNearestHitTriangle(LocalRay);
+		if (HitTID != IndexConstants::InvalidID)
+		{
+			FTriangle3d Triangle;
+			GetSpatial().GetMesh()->GetTriVertices(HitTID, Triangle.V[0], Triangle.V[1], Triangle.V[2]);
+			FIntrRay3Triangle3d Query(LocalRay, Triangle);
+			Query.Find();
+			OutHit.FaceIndex = HitTID;
+			OutHit.Distance = Query.RayParameter;
+			OutHit.Normal = (FVector)Transform.TransformVectorNoScale(GetSpatial().GetMesh()->GetTriNormal(HitTID));
+			OutHit.ImpactPoint = (FVector)Transform.TransformPosition(LocalRay.PointAt(Query.RayParameter));
+		}
+	}
+	return true;
+}
+
+
+void UDeformMeshPolygonsTool::OnBeginDrag(const FRay& WorldRay)
+{
+	FTransform3d Transform(ComponentTarget->GetWorldTransform());
+	FRay3d LocalRay(Transform.InverseTransformPosition(WorldRay.Origin),
+		Transform.InverseTransformVector(WorldRay.Direction));
+	LocalRay.Direction.Normalize();
+
+	HilightSelection.Clear();
+
+	TopoSelector.UpdateEnableFlags(TransformProps->bSelectFaces, TransformProps->bSelectEdges, TransformProps->bSelectVertices);
+	FGroupTopologySelection Selection;
+	FVector3d LocalPosition, LocalNormal;
+	bool bHit = TopoSelector.FindSelectedElement(LocalRay, Selection, LocalPosition, LocalNormal);
+
+	if (bHit == false)
+	{
+		bInDrag = false;
+		return;
+	}
+
+	HilightSelection = Selection;
+
+	FVector3d WorldHitPos = Transform.TransformPosition(LocalPosition);
+	FVector3d WorldHitNormal = Transform.TransformVector(LocalNormal);
+
+	bInDrag = true;
+	StartHitPosWorld = (FVector)WorldHitPos;
+	LastHitPosWorld = StartHitPosWorld;
+	StartHitNormalWorld = (FVector)WorldHitNormal;
+
+	QuickAxisRotator.ClearAxisLock();
+	UpdateActiveSurfaceFrame(HilightSelection);
+	UpdateQuickTransformer();
+
+	LastBrushPosLocal = (FVector)Transform.InverseTransformPosition(LastHitPosWorld);
+	StartBrushPosLocal = LastBrushPosLocal;
+
+	// Record the requested deformation strategy - NB: will be forced to linear if there aren't any free points to solve.
+
+	DeformationStrategy = TransformProps->DeformationStrategy;
+	
+	// Capture the part of the mesh that will deform
+
+	if (DeformationStrategy == EGroupTopologyDeformationStrategy::Laplacian)
+	{
+		LaplacianDeformer.bLocalize = true; // TransformProps->bLocalizeDeformation;
+
+		//Determine which of the following (corners, edges or faces) has been selected by counting the associated feature's IDs
+		if (Selection.SelectedCornerIDs.Num() > 0)
+		{
+			//Add all the the Corner's adjacent poly-groups (NbrGroups) to the ongoing array of groups.
+			LaplacianDeformer.SetActiveHandleCorners(Selection.SelectedCornerIDs);
+		}
+		else if (Selection.SelectedEdgeIDs.Num() > 0)
+		{
+			//Add all the the edge's adjacent poly-groups (NbrGroups) to the ongoing array of groups.
+			LaplacianDeformer.SetActiveHandleEdges(Selection.SelectedEdgeIDs);
+		}
+		else if (Selection.SelectedGroupIDs.Num() > 0)
+		{
+			LaplacianDeformer.SetActiveHandleFaces(Selection.SelectedGroupIDs);
+		}
+
+
+		// If there are actually no interior points, then we can't actually use the laplacian deformer. Need to fall back to the linear.
+		bool bHasInteriorVerts = false;
+		const auto& ROIFaces = LaplacianDeformer.GetROIFaces();
+		for (const auto& Face : ROIFaces)
+		{
+			bHasInteriorVerts = bHasInteriorVerts || ( Face.InteriorVerts.Num() != 0);
+		}
+
+		if (!bHasInteriorVerts)
+		{
+			// Change to the linear strategy for this case.
+
+			DeformationStrategy = EGroupTopologyDeformationStrategy::Linear;
+		}
+		else
+		{		
+			// finalize the laplacian deformer : the task will need a new mesh that corresponds to the selected region.
+		
+			LaplacianDeformer.bTaskSubmeshIsDirty = true;
+
+		}
+		
+	}
+	 
+	if (DeformationStrategy == EGroupTopologyDeformationStrategy::Linear )
+	{
+		//Determine which of the following (corners, edges or faces) has been selected by counting the associated feature's IDs
+		if (Selection.SelectedCornerIDs.Num() > 0)
+		{
+			//Add all the the Corner's adjacent poly-groups (NbrGroups) to the ongoing array of groups.
+			LinearDeformer.SetActiveHandleCorners(Selection.SelectedCornerIDs);
+		}
+		else if (Selection.SelectedEdgeIDs.Num() > 0)
+		{
+			//Add all the the edge's adjacent poly-groups (NbrGroups) to the ongoing array of groups.
+			LinearDeformer.SetActiveHandleEdges(Selection.SelectedEdgeIDs);
+		}
+		else if (Selection.SelectedGroupIDs.Num() > 0)
+		{
+			LinearDeformer.SetActiveHandleFaces(Selection.SelectedGroupIDs);
+		}
+	}
+	
+	BeginChange();
+}
+
+
+void UDeformMeshPolygonsTool::UpdateActiveSurfaceFrame(FGroupTopologySelection& Selection)
+{
+	FTransform3d Transform(ComponentTarget->GetWorldTransform());
+
+	// update surface frame
+	ActiveSurfaceFrame.Origin = StartHitPosWorld;
+	if (HilightSelection.SelectedCornerIDs.Num() == 1)
+	{
+		// just keeping existing axes...we don't have enough info to do something smarter
+	}
+	else
+	{
+		ActiveSurfaceFrame.AlignAxis(2, StartHitNormalWorld);
+		if (HilightSelection.SelectedEdgeIDs.Num() == 1)
+		{
+			FVector3d Tangent;
+			if (Topology.GetGroupEdgeTangent(HilightSelection.SelectedEdgeIDs[0], Tangent))
+			{
+				Tangent = Transform.TransformVector(Tangent);
+				ActiveSurfaceFrame.ConstrainedAlignAxis(0, Tangent, ActiveSurfaceFrame.Z());
+			}
+		}
+	}
+}
+
+
+FQuickTransformer* UDeformMeshPolygonsTool::GetActiveQuickTransformer()
+{
+	if (TransformProps->TransformMode == EQuickTransformerMode::AxisRotation)
+	{
+		return &QuickAxisRotator;
+	}
+	else
+	{
+		return &QuickAxisTranslater;
+	}
+}
+
+
+void UDeformMeshPolygonsTool::UpdateQuickTransformer()
+{
+	bool bUseLocalAxes =
+		(GetToolManager()->GetContextQueriesAPI()->GetCurrentCoordinateSystem() == EToolContextCoordinateSystem::Local);
+	if (bUseLocalAxes)
+	{
+		GetActiveQuickTransformer()->SetActiveWorldFrame(ActiveSurfaceFrame);
+	}
+	else
+	{
+		GetActiveQuickTransformer()->SetActiveFrameFromWorldAxes(StartHitPosWorld);
+	}
+}
+
+
+
+
+
+void UDeformMeshPolygonsTool::UpdateChangeFromROI(bool bFinal)
+{
+	if (ActiveVertexChange == nullptr)
+	{
+		return;
+	}
+	const bool bIsLaplacian = (DeformationStrategy == EGroupTopologyDeformationStrategy::Laplacian);
+
+	FDynamicMesh3* Mesh = DynamicMeshComponent->GetMesh();
+	const TSet<int>& ModifiedVertices = (bIsLaplacian) ? LaplacianDeformer.GetModifiedVertices() : LinearDeformer.GetModifiedVertices();
+	
+	ActiveVertexChange->SavePositions(Mesh, ModifiedVertices, !bFinal);
+
+}
+
+
+void UDeformMeshPolygonsTool::OnUpdateDrag(const FRay& Ray)
+{
+	if (bInDrag)
+	{
+		bUpdatePending = true;
+		UpdateRay = Ray;
+	}
+}
+
+void UDeformMeshPolygonsTool::OnEndDrag(const FRay& Ray)
+{
+	bInDrag = false;
+	bUpdatePending = false;
+
+	// update spatial
+	bSpatialDirty = true;
+
+	HilightSelection.Clear(); 
+	TopoSelector.Invalidate(true, false);
+	QuickAxisRotator.Reset();
+	QuickAxisTranslater.Reset();
+
+	//If it's linear, it's computed real time with no delay. This may need to be restructured for clarity by using the background task for this as well.
+	if (DeformationStrategy == EGroupTopologyDeformationStrategy::Linear)
+	{
+		// close change record
+		EndChange();
+	}
+}
+
+
+
+bool UDeformMeshPolygonsTool::OnUpdateHover(const FInputDeviceRay& DevicePos)
+{
+	//if (!bNeedEmitEndChange)
+	if (ActiveVertexChange == nullptr)
+	{
+		FTransform3d Transform(ComponentTarget->GetWorldTransform());
+		FRay3d LocalRay(Transform.InverseTransformPosition(DevicePos.WorldRay.Origin),
+		Transform.InverseTransformVector(DevicePos.WorldRay.Direction));
+		LocalRay.Direction.Normalize();
+
+		HilightSelection.Clear();
+		TopoSelector.UpdateEnableFlags(TransformProps->bSelectFaces, TransformProps->bSelectEdges, TransformProps->bSelectVertices);
+		FVector3d LocalPosition, LocalNormal;
+		bool bHit = TopoSelector.FindSelectedElement(LocalRay, HilightSelection, LocalPosition, LocalNormal);
+
+		if (bHit)
+		{
+			StartHitPosWorld = (FVector)Transform.TransformPosition(LocalPosition);
+			StartHitNormalWorld = (FVector)Transform.TransformVector(LocalNormal);
+
+			UpdateActiveSurfaceFrame(HilightSelection);
+			UpdateQuickTransformer();
+		}
+	}
+	return true;
+}
+
+
+
+
+void UDeformMeshPolygonsTool::ComputeUpdate()
+{
+
+	if (bUpdatePending == true)
+	{
+		// Linear Deformer : Update the solution
+		// Laplacain Deformer : Update the constraints (positions and weights) - the region was identified in onBeginDrag
+
+
+		if (TransformProps->TransformMode == EQuickTransformerMode::AxisRotation)
+		{
+			ComputeUpdate_Rotate();
+		}
+		else
+		{
+			ComputeUpdate_Translate();
+		}
+	}
+
+	if (DeformationStrategy == EGroupTopologyDeformationStrategy::Laplacian)
+	{
+		bool bIsWorking = LaplacianDeformer.IsTaskInFlight();
+
+		if (!bIsWorking)
+		{
+			// Sync update if we have new results.
+			if (LaplacianDeformer.bVertexPositionsNeedSync)
+			{
+				
+				//Update the mesh with the provided solutions.
+				LaplacianDeformer.ExportDeformedPositions(DynamicMeshComponent->GetMesh());
+
+				LaplacianDeformer.bVertexPositionsNeedSync = false;
+
+				//Re-sync mesh, and flag the spatial data struct & topology for re-evaluation
+				DynamicMeshComponent->FastNotifyPositionsUpdated();
+
+				GetToolManager()->PostInvalidation();
+				bSpatialDirty = true;
+				TopoSelector.Invalidate(true, false);
+			}
+
+			// emit end change if we are done with the drag
+			if (!LaplacianDeformer.bDeformerNeedsToRun && !bInDrag)
+			{
+				EndChange();
+			}
+
+			// Not working but we have more work for it to do..
+
+			if (LaplacianDeformer.bDeformerNeedsToRun)
+			{
+			
+				FRichCurve* Curve = NULL;
+
+				/**
+				// How to add a deformation curve
+				const bool bApplyAttenuationCurve = TransformProps->bApplyAttenuationCurve
+				if (bApplyAttenuationCurve)
+				{
+					Curve = TransformProps->WeightAttenuationCurve.GetRichCurve();
+				}
+				*/
+				const ELaplacianWeightScheme LaplacianWeightScheme = ConvertToLaplacianWeightScheme(TransformProps->SelectedWeightScheme);
+				LaplacianDeformer.UpdateAndLaunchdWorker(LaplacianWeightScheme, Curve);
+
+			}
+		}
+	}
+}
+
+
+
+
+
+
+void UDeformMeshPolygonsTool::ComputeUpdate_Rotate()
+{
+	const bool bIsLaplacian = (DeformationStrategy == EGroupTopologyDeformationStrategy::Laplacian); 
+	FGroupTopologyDeformer& SelectedDeformer = (bIsLaplacian) ? LaplacianDeformer : LinearDeformer;
+
+	FDynamicMesh3* Mesh = DynamicMeshComponent->GetMesh();
+	FTransform3d Transform = FTransform3d(ComponentTarget->GetWorldTransform());
+	FVector NewHitPosWorld = LastHitPosWorld;
+
+	FVector3d SnappedPoint;
+	if (QuickAxisRotator.UpdateSnap(FRay3d(UpdateRay), SnappedPoint))
+	{
+		NewHitPosWorld = (FVector)SnappedPoint;
+	}
+	else
+	{
+		return;
+	}
+
+	// check if we are on back-facing part of rotation in which case we ignore...
+	FVector3d SphereCenter = QuickAxisRotator.GetActiveWorldFrame().Origin;
+	if (QuickAxisRotator.HaveActiveSnapRotation() && QuickAxisRotator.GetHaveLockedToAxis() == false)
+	{
+		FVector3d ToSnapPointVec = (SnappedPoint - SphereCenter);
+		FVector3d ToEyeVec = (SnappedPoint - (FVector3d)CameraState.Position);
+		if (ToSnapPointVec.Dot(ToEyeVec) > 0)
+		{
+			return;
+		}
+	}
+
+
+	// if we haven't snapped to a rotation we can exit
+	if (QuickAxisRotator.HaveActiveSnapRotation() == false)
+	{
+		QuickAxisRotator.ClearAxisLock();
+
+		SelectedDeformer.ClearSolution(Mesh);
+
+		//TODO: This is unseemly here, need to potentially defer this so that it's handled the same way as laplacian. Placeholder for now.
+		if (DeformationStrategy == EGroupTopologyDeformationStrategy::Linear)
+		{
+
+			DynamicMeshComponent->FastNotifyPositionsUpdated();
+			GetToolManager()->PostInvalidation();
+		}
+		bUpdatePending = false;
+		return;
+	}
+
+	// ok we have an axis...
+	if (QuickAxisRotator.GetHaveLockedToAxis() == false)
+	{
+		QuickAxisRotator.SetAxisLock();
+		RotationStartPointWorld = SnappedPoint;
+		RotationStartFrame = QuickAxisRotator.GetActiveRotationFrame();
+	}
+
+	FVector2d RotateStartVec = RotationStartFrame.ToPlaneUV(RotationStartPointWorld, 2);
+	RotateStartVec.Normalize();
+	FVector2d RotateToVec = RotationStartFrame.ToPlaneUV(NewHitPosWorld, 2);
+	RotateToVec.Normalize();
+	double AngleRad = RotateStartVec.SignedAngleR(RotateToVec);
+	FQuaterniond Rotation(
+		Transform.InverseTransformVectorNoScale(RotationStartFrame.Z()), AngleRad, false);
+	FVector3d LocalOrigin = Transform.InverseTransformPosition(RotationStartFrame.Origin);
+
+	// Linear Deformer: Update Mesh the rotation,
+	// Laplacian Deformer:  Update handles constraints with the rotation and set bDeformerNeedsToRun = true;.
+	SelectedDeformer.UpdateSolution(Mesh, [this, LocalOrigin, Rotation](FDynamicMesh3* TargetMesh, int VertIdx)
+	{
+		FVector3d V = TargetMesh->GetVertex(VertIdx);
+		V -= LocalOrigin;
+		V = Rotation * V;
+		V += LocalOrigin;
+		return V;
+	});
+
+	//TODO: This is unseemly here, need to potentially defer this so that it's handled the same way as laplacian. Placeholder for now.
+	if (!bIsLaplacian)
+	{
+		DynamicMeshComponent->FastNotifyPositionsUpdated();
+		GetToolManager()->PostInvalidation();
+	}
+	bUpdatePending = false;
+}
+
+
+
+
+void UDeformMeshPolygonsTool::ComputeUpdate_Translate()
+{
+	const bool bIsLaplacian = (DeformationStrategy == EGroupTopologyDeformationStrategy::Laplacian);
+	FGroupTopologyDeformer& SelectedDeformer = (bIsLaplacian) ? LaplacianDeformer : LinearDeformer;
+
+	TFunction<FVector3d(const FVector3d&)> PointConstraintFunc = nullptr;
+	if (TransformProps->bSnapToWorldGrid 
+		&& GetToolManager()->GetContextQueriesAPI()->GetCurrentCoordinateSystem() == EToolContextCoordinateSystem::World)
+	{
+		PointConstraintFunc = [&](const FVector3d& Pos)
+		{
+			FVector3d GridSnapPos;
+			return ToolSceneQueriesUtil::FindWorldGridSnapPoint(this, Pos, GridSnapPos) ? GridSnapPos : Pos;
+		};
+	}
+
+	FTransform Transform = ComponentTarget->GetWorldTransform();
+	FVector NewHitPosWorld = LastHitPosWorld;
+	FVector3d SnappedPoint;
+	if (QuickAxisTranslater.UpdateSnap(FRay3d(UpdateRay), SnappedPoint, PointConstraintFunc))
+	{
+		NewHitPosWorld = (FVector)SnappedPoint;
+	}
+	else
+	{
+		return;
+	}
+
+	FVector NewBrushPosLocal = Transform.InverseTransformPosition(NewHitPosWorld);
+	FVector3d NewMoveDelta = NewBrushPosLocal - StartBrushPosLocal;
+
+	FDynamicMesh3* Mesh = DynamicMeshComponent->GetMesh();
+	if (LastMoveDelta.SquaredLength() > 0.)
+	{
+		if (NewMoveDelta.SquaredLength() > 0.)
+		{
+			// Linear Deformer: Update Mesh with the translation,
+			// Laplacian Deformer:  Update handles constraints and set bDeformerNeedsToRun = true;.
+
+			SelectedDeformer.UpdateSolution(Mesh, [this, NewMoveDelta](FDynamicMesh3* TargetMesh, int VertIdx)
+			{
+				return TargetMesh->GetVertex(VertIdx) + NewMoveDelta;
+			});
+		}
+		else
+		{
+			// Reset mesh to initial positions.
+
+			SelectedDeformer.ClearSolution(Mesh);
+
+		}
+		//TODO: This is unseemly here, need to potentially defer this so that it's handled the same way as laplacian. Placeholder for now.
+		if (!bIsLaplacian)
+		{
+			DynamicMeshComponent->FastNotifyPositionsUpdated();
+			GetToolManager()->PostInvalidation();
+		}
+	}
+
+	LastMoveDelta = NewMoveDelta;
+	LastBrushPosLocal = NewBrushPosLocal;
+
+	bUpdatePending = false;
+}
+
+
+
+void UDeformMeshPolygonsTool::Tick(float DeltaTime)
+{
+	UMeshSurfacePointTool::Tick(DeltaTime);
+	
+	LaplacianDeformer.HandleWeights   = TransformProps->HandleWeight;
+	LaplacianDeformer.bPostfixHandles = TransformProps->bPostFixHandles;
+}
+
+
+
+void UDeformMeshPolygonsTool::PrecomputeTopology()
+{
+	FDynamicMesh3* Mesh = DynamicMeshComponent->GetMesh();
+	Topology = FGroupTopology(Mesh, true);
+
+	LinearDeformer.Initialize(Mesh, &Topology);
+	LaplacianDeformer.Initialize(Mesh, &Topology);
+
+	// Make the Constraint Buffer, zero weights, but current pos
+	LaplacianDeformer.InitializeConstraintBuffer();
+}
+
+
+
+
+void UDeformMeshPolygonsTool::Render(IToolsContextRenderAPI* RenderAPI)
+{
+
+	ComputeUpdate();
+		
+	GetToolManager()->GetContextQueriesAPI()->GetCurrentViewState(CameraState);
+	GetActiveQuickTransformer()->UpdateCameraState(CameraState);
+
+	DynamicMeshComponent->bExplicitShowWireframe = TransformProps->bShowWireframe;
+	FDynamicMesh3* TargetMesh = DynamicMeshComponent->GetMesh();
+
+	PolyEdgesRenderer.BeginFrame(RenderAPI, CameraState);
+	PolyEdgesRenderer.SetTransform(ComponentTarget->GetWorldTransform());
+
+
+	for (FGroupTopology::FGroupEdge& Edge : Topology.Edges)
+	{
+		FVector3d A, B;
+		for (int eid : Edge.Span.Edges)
+		{
+			TargetMesh->GetEdgeV(eid, A, B);
+			PolyEdgesRenderer.DrawLine(A, B);
+		}
+	}
+
+	PolyEdgesRenderer.EndFrame();
+
+
+	HilightRenderer.BeginFrame(RenderAPI, CameraState);
+	HilightRenderer.SetTransform(ComponentTarget->GetWorldTransform());
+
+#ifdef DEBUG_ROI_WEIGHTS
+	FDynamicMesh3* MeshPtr = DynamicMeshComponent->GetMesh();
+	for (int32 VertexID : DynamicMeshComponent->GetMesh()->VertexIndicesItr())
+	{
+		float Color = 1.f - SrcMeshConstraintBuffer[VertexID].Weight;
+		HilightRenderer.DrawPoint(MeshPtr->GetVertex(VertexID), FLinearColor(Color, Color, Color, 1.f), 8, true);
+	}
+
+#endif
+
+#ifdef DEBUG_ROI_HANDLES
+	const FLinearColor FOOF{ 1.f,0.f,1.f,1.f };
+	for (int VertIdx : HandleVertices)
+	{
+		HilightRenderer.DrawViewFacingCircle(TargetMesh->GetVertex(VertIdx), 0.8f, 8, FOOF, 3, false);
+	}
+#endif 
+
+#ifdef DEBUG_ROI_TRIANGLES
+	const FLinearColor Whiteish{ 0.67f,0.67f,0.67f,1.f };
+	for (int32 i = 0; i < SubsetIDBuffer.Num(); i += 3)
+	{
+		FVector3d A = TargetMesh->GetVertex(SubsetIDBuffer[i]);
+		FVector3d B = TargetMesh->GetVertex(SubsetIDBuffer[i + 1]);
+		FVector3d C = TargetMesh->GetVertex(SubsetIDBuffer[i + 2]);
+		HilightRenderer.DrawLine(A, B, Whiteish, 2.7f, true);
+		HilightRenderer.DrawLine(B, C, Whiteish, 2.7f, true);
+		HilightRenderer.DrawLine(C, A, Whiteish, 2.7f, true);
+	}
+#endif
+
+	TopoSelector.VisualAngleSnapThreshold = this->VisualAngleSnapThreshold;
+	TopoSelector.DrawSelection(HilightSelection, &HilightRenderer, &CameraState);
+	HilightRenderer.EndFrame();
+
+
+	if (bInDrag)
+	{
+
+		GetActiveQuickTransformer()->Render(RenderAPI);
+	}
+	else
+	{
+		GetActiveQuickTransformer()->PreviewRender(RenderAPI);
+	}
+}
+
+
+
+void UDeformMeshPolygonsTool::OnPropertyModified(UObject* PropertySet, UProperty* Property)
+{
+}
+
+
+
+//
+// Change Tracking
+//
+
+
+void UDeformMeshPolygonsTool::BeginChange()
+{
+	const bool bIsLaplacian = (DeformationStrategy == EGroupTopologyDeformationStrategy::Laplacian);
+	if (!bIsLaplacian || LaplacianDeformer.IsDone())
+	{
+		if (ActiveVertexChange == nullptr)
+		{
+			ActiveVertexChange = new FMeshVertexChangeBuilder();
+			UpdateChangeFromROI(false);
+		}
+	}
+}
+
+
+void UDeformMeshPolygonsTool::EndChange()
+{
+
+	if (ActiveVertexChange != nullptr)
+	{
+		UpdateChangeFromROI(true);
+		GetToolManager()->EmitObjectChange(DynamicMeshComponent, MoveTemp(ActiveVertexChange->Change), LOCTEXT("PolyMeshDeformationChange", "PolyMesh Edit"));
+	}
+
+	delete ActiveVertexChange;
+	ActiveVertexChange = nullptr;
+}
+
+
+
+#undef LOCTEXT_NAMESPACE
diff --git a/Engine/Plugins/Experimental/MeshModelingToolset/Source/MeshModelingTools/Private/DisplaceMeshTool.cpp b/Engine/Plugins/Experimental/MeshModelingToolset/Source/MeshModelingTools/Private/DisplaceMeshTool.cpp
--- a/Engine/Plugins/Experimental/MeshModelingToolset/Source/MeshModelingTools/Private/DisplaceMeshTool.cpp
+++ b/Engine/Plugins/Experimental/MeshModelingToolset/Source/MeshModelingTools/Private/DisplaceMeshTool.cpp
@@ -196,9 +196,9 @@ namespace {
 		void SetSubdivisionsCount(int SubdivisionsCountIn);
 		int  GetSubdivisionsCount();
 
-		TSharedPtr<FDynamicMeshOperator> MakeNewOperator() final
+		TUniquePtr<FDynamicMeshOperator> MakeNewOperator() final
 		{
-			return MakeShared<FSubdivideMeshOp>(SourceMesh, SubdivisionsCount);
+			return MakeUnique<FSubdivideMeshOp>(SourceMesh, SubdivisionsCount);
 		}
 	private:
 		const FDynamicMesh3& SourceMesh;
@@ -324,9 +324,9 @@ namespace {
 		void SetDisplacementMap(UTexture2D* DisplacementMapIn);
 		void SetDisplacementType(EDisplaceMeshToolDisplaceType TypeIn);
 
-		TSharedPtr<FDynamicMeshOperator> MakeNewOperator() final
+		TUniquePtr<FDynamicMeshOperator> MakeNewOperator() final
 		{
-			return MakeShared<FDisplaceMeshOp>(SourceMesh, DisplaceField, DisplaceIntensity, RandomSeed, DisplacementType);
+			return MakeUnique<FDisplaceMeshOp>(SourceMesh, DisplaceField, DisplaceIntensity, RandomSeed, DisplacementType);
 		}
 	private:
 		void UpdateMap();
@@ -442,11 +442,12 @@ void UDisplaceMeshTool::Setup()
 	DynamicMeshComponent->RegisterComponent();
 	DynamicMeshComponent->SetWorldTransform(ComponentTarget->GetWorldTransform());
 
-	// copy material if there is one
-	auto Material = ComponentTarget->GetMaterial(0);
-	if (Material != nullptr)
+	// transfer materials
+	FComponentMaterialSet MaterialSet;
+	ComponentTarget->GetMaterialSet(MaterialSet);
+	for (int k = 0; k < MaterialSet.Materials.Num(); ++k)
 	{
-		DynamicMeshComponent->SetMaterial(0, Material);
+		DynamicMeshComponent->SetMaterial(k, MaterialSet.Materials[k]);
 	}
 
 	DynamicMeshComponent->InitializeMesh(ComponentTarget->GetMesh());
@@ -475,9 +476,9 @@ void UDisplaceMeshTool::Shutdown(EToolShutdownType ShutdownType)
 		{
 			// this block bakes the modified DynamicMeshComponent back into the StaticMeshComponent inside an undo transaction
 			GetToolManager()->BeginUndoTransaction(LOCTEXT("DisplaceMeshToolTransactionName", "Displace Mesh"));
-			ComponentTarget->CommitMesh([=](FMeshDescription* MeshDescription)
+			ComponentTarget->CommitMesh([=](const FPrimitiveComponentTarget::FCommitParams& CommitParams)
 			{
-				DynamicMeshComponent->Bake(MeshDescription, Subdivisions > 0);
+				DynamicMeshComponent->Bake(CommitParams.MeshDescription, Subdivisions > 0);
 			});
 			GetToolManager()->EndUndoTransaction();
 		}
@@ -573,13 +574,13 @@ void UDisplaceMeshTool::StartComputation()
 		SubdivideTask = new FAsyncTaskExecuterWithAbort<TModelingOpTask<FDynamicMeshOperator>>(Subdivider->MakeNewOperator());
 		SubdivideTask->StartBackgroundTask();
 		bNeedsSubdivided = false;
-		DynamicMeshComponent->SetMaterial(0, ToolSetupUtil::GetDefaultWorkingMaterial(GetToolManager()));
+		DynamicMeshComponent->SetOverrideRenderMaterial(ToolSetupUtil::GetDefaultWorkingMaterial(GetToolManager()));
 	}
 	if (bNeedsDisplaced && DisplaceTask)
 	{
 		DisplaceTask->CancelAndDelete();
 		DisplaceTask = nullptr;
-		DynamicMeshComponent->SetMaterial(0, ToolSetupUtil::GetDefaultWorkingMaterial(GetToolManager()));
+		DynamicMeshComponent->SetOverrideRenderMaterial(ToolSetupUtil::GetDefaultWorkingMaterial(GetToolManager()));
 	}
 	AdvanceComputation();
 }
@@ -603,7 +604,7 @@ void UDisplaceMeshTool::AdvanceComputation()
 		TUniquePtr<FDynamicMesh3> DisplacedMesh = DisplaceTask->GetTask().ExtractOperator()->ExtractResult();
 		delete DisplaceTask;
 		DisplaceTask = nullptr;
-		DynamicMeshComponent->SetMaterial(0,ToolSetupUtil::GetDefaultMaterial(GetToolManager(), ComponentTarget->GetMaterial(0)));
+		DynamicMeshComponent->ClearOverrideRenderMaterial();
 		DynamicMeshComponent->GetMesh()->Copy(*DisplacedMesh);
 		DynamicMeshComponent->NotifyMeshUpdated();
 		GetToolManager()->PostInvalidation();
