Copy up from Dev-Editor [REDACTED_HANDLE]

diff --git a/Engine/Plugins/Experimental/MeshModelingToolset/Source/MeshModelingTools/Private/EditMeshPolygonsTool.cpp b/Engine/Plugins/Experimental/MeshModelingToolset/Source/MeshModelingTools/Private/EditMeshPolygonsTool.cpp
--- a/Engine/Plugins/Experimental/MeshModelingToolset/Source/MeshModelingTools/Private/EditMeshPolygonsTool.cpp
+++ b/Engine/Plugins/Experimental/MeshModelingToolset/Source/MeshModelingTools/Private/EditMeshPolygonsTool.cpp
@@ -10,6 +10,9 @@
 #include "ToolSceneQueriesUtil.h"
 #include "Intersection/IntersectionUtil.h"
 #include "FindPolygonsAlgorithm.h"
+#include "Transforms/MultiTransformer.h"
+#include "BaseBehaviors/SingleClickBehavior.h"
+#include "Util/ColorConstants.h"
 
 #include "Async/ParallelFor.h"
 #include "Containers/BitArray.h"
@@ -18,21 +21,6 @@
 
 
 
-//////////////////////////////
-// DEBUG_SETTINGS
-
-//Draw white triangles defining the selection subset
-//#define DEBUG_ROI_TRIANGLES
-
-//Draw pink circles around the handles
-//#define DEBUG_ROI_HANDLES
-
-//Draw points on the ROI vertices, White => Weight == 0, Black => Weight == 1
-//#define DEBUG_ROI_WEIGHTS
-
-//////////////////////////////
-
-
 /*
  * ToolBuilder
  */
@@ -47,10 +35,12 @@ UMeshSurfacePointTool* UEditMeshPolygonsToolBuilder::CreateNewTool(const FToolBu
  */
 UPolyEditTransformProperties::UPolyEditTransformProperties()
 {
-	TransformMode = EQuickTransformerMode::AxisTranslation;
+	TransformMode = EMultiTransformerMode::DefaultGizmo;
 	bSelectVertices = true;
 	bSelectFaces = true;
 	bSelectEdges = true;
+	bShowWireframe = false;
+	bSnapToWorldGrid = false;
 	PolygonMode = EPolygonGroupMode::KeepInputPolygons;
 	PolygonGroupingAngleThreshold = .5;
 }
@@ -68,543 +58,6 @@ void UPolyEditTransformProperties::PostEditChangeProperty(FPropertyChangedEvent&
 }
 #endif
 
-/*
- * Asynchronous Task
- */
-
-
-void FConstrainedMeshDeformerTask::UpdateDeformer(const ELaplacianWeightScheme SelectedWeightScheme, const FDynamicMesh3& SrcMesh, const TArray<FConstraintData>& ConstraintArray, const TArray<int32>& SrcIDBufferSubset, bool bNewTransaction, const FRichCurve* Curve)
-{
-	bIsNewTransaction = bNewTransaction;
-	SrcMeshMaxVertexID = SrcMesh.MaxVertexID();
-
-	LaplacianWeightScheme = SelectedWeightScheme;
-	
-	bAttenuateWeights = (Curve != nullptr);
-	if (bAttenuateWeights)
-	{
-		WeightAttenuationCurve = *Curve;
-	}
-
-	// Set-up the subset mesh.
-	if (bIsNewTransaction)
-	{
-		//Copy the part of the mesh we want to deform into the SubsetMesh and create map from Src Mesh to the SubsetMesh.
-		InitializeSubsetMesh(SrcMesh, SrcIDBufferSubset);
-	}
-
-	// only want the subset of constraints that correspond to our subset mesh.
-	{
-		const int32 NumSubsetVerts = SubsetVertexIDToSrcVertexIDMap.Num();
-		SubsetConstraintBuffer.Empty(NumSubsetVerts);
-		SubsetConstraintBuffer.AddUninitialized(NumSubsetVerts);
-
-		for (int32 SubVertexID = 0; SubVertexID < SubsetVertexIDToSrcVertexIDMap.Num(); ++SubVertexID)
-		{
-			int32 SrcVtxID = SubsetVertexIDToSrcVertexIDMap[SubVertexID];
-			SubsetConstraintBuffer[SubVertexID] = ConstraintArray[SrcVtxID];
-		}
-	}
-
-	check(bIsNewTransaction || ConstrainedDeformer.IsValid());
-	
-}
-
-
-void FConstrainedMeshDeformerTask::DoWork()
-{
-
-	//TODO: (simple optimization) -
-	//	Instead of SrcVertexIDtoSubsetVertexIDMap, use SubsetVertexIDToSetVertexIDMap - then we can use the VertexIndicesItr()
-	//	on the SubsetMesh to minimize the quantity of vertex indices we need to iterate at every following step.
-
-	if (*bAbortSource == true)
-	{
-		return;
-	}
-
-	if (bIsNewTransaction) //Will only be true once per input transaction (click+drag)
-	{
-		
-		// Create a new deformation solver.
-		ConstrainedDeformer = MakeUnique<FConstrainedMeshDeformer>(SubsetMesh, LaplacianWeightScheme);
-
-		if (bAttenuateWeights)
-		{
-			ApplyAttenuation();
-		}
-
-		//Update our deformer's constraints before deforming using the copy of the constraint buffer
-		for (int32 SubsetVertexID = 0; SubsetVertexID < SubsetConstraintBuffer.Num(); ++SubsetVertexID)
-		{
-			FConstraintData& CData = SubsetConstraintBuffer[SubsetVertexID];
-			ConstrainedDeformer->AddConstraint(SubsetVertexID, CData.Weight, CData.Position, CData.bPostFix);
-		}
-
-		bIsNewTransaction = false;
-
-	}
-	else
-	{
-		//This else block is run every consecutive frame after the start of the input transaction because UpdateConstraintPosition() is very cheap (no factorizing or rebuilding)
-		//Update only the positions of the constraints, as the weights cannot change mid-transaction
-
-		for (int32 SubsetVertexID = 0; SubsetVertexID < SubsetConstraintBuffer.Num(); ++SubsetVertexID)
-		{
-			FConstraintData& CData = SubsetConstraintBuffer[SubsetVertexID];
-			ConstrainedDeformer->UpdateConstraintPosition(SubsetVertexID, CData.Position, CData.bPostFix);
-		}
-
-	}
-
-	if (*bAbortSource == true)
-	{
-		return;
-	}
-
-	//Run the deformation process
-
-	const bool bSuccessfulSolve = ConstrainedDeformer->Deform(SubsetPositionBuffer);
-	
-	
-	if (bSuccessfulSolve)
-	{
-		if (*bAbortSource == true)
-		{
-			return;
-		}
-	}
-	else
-	{
-		//UE_LOG(LogTemp, Warning, TEXT("Laplacian deformation failed"));
-	}
-
-}
-
-inline void FConstrainedMeshDeformerTask::InitializeSubsetMesh(const FDynamicMesh3& SrcMesh, const TArray<int32>& SrcIDBufferSubset)
-{
-	//These can be re-used until the user stops dragging
-	SubsetMesh.Clear();
-	SubsetPositionBuffer.Reset();
-
-	//Initialize every element to -1, helps us keep track of vertices we've already added while iterating the triangles
-	TArray<int32> SrcVertexIDToSubsetVertexIDMap;
-	SrcVertexIDToSubsetVertexIDMap.Init(INACTIVE_SUBSET_ID, SrcMeshMaxVertexID);
-
-
-	//Iterate the triangle array to append vertices, and then triangles to the temporary subset mesh all at once
-	for (int32 i = 0; i < SrcIDBufferSubset.Num(); i += 3)
-	{
-
-		// Build the triangle
-		FIndex3i Triangle;
-		for (int32 v = 0; v < 3; ++v)
-		{
-			//It's the SrcVertexID because every element in the SrcIDBufferSubset is the Vertex ID of a vertex in the original mesh.
-			const int32 SrcVertexID = SrcIDBufferSubset[i + v];
-			int32& SubsetID = SrcVertexIDToSubsetVertexIDMap[SrcVertexID];
-			
-			if (SubsetID == INACTIVE_SUBSET_ID) // we haven't already visited this vertex
-			{
-				const FVector3d Vertex = SrcMesh.GetVertex(SrcVertexID);
-				SubsetID = SubsetMesh.AppendVertex(Vertex);
-			}
-
-			Triangle[v] = SubsetID;
-		}
-		SubsetMesh.AppendTriangle(Triangle);
-	}
-
-	// create a mapping back to the original vertex IDs from the subset mesh
-	int32 MaxSubMeshVertexID = SubsetMesh.MaxVertexID(); // Really MaxID + 1
-	SubsetVertexIDToSrcVertexIDMap.Reset(MaxSubMeshVertexID);
-	SubsetVertexIDToSrcVertexIDMap.AddUninitialized(MaxSubMeshVertexID);
-
-	for (int32 SrcID = 0; SrcID < SrcVertexIDToSubsetVertexIDMap.Num(); ++SrcID)
-	{
-		const int32 SubsetVertexID = SrcVertexIDToSubsetVertexIDMap[SrcID];
-		if (SubsetVertexID != INACTIVE_SUBSET_ID)
-		{
-			SubsetVertexIDToSrcVertexIDMap[SubsetVertexID] = SrcID;
-		}
-	}
-
-}
-
-void FConstrainedMeshDeformerTask::ExportResults(FDynamicMesh3& TargetMesh) const
-{
-
-	//Update the position buffer result
-	for (int32 SubsetVertexID = 0; SubsetVertexID < SubsetVertexIDToSrcVertexIDMap.Num(); ++SubsetVertexID)
-	{
-		const int32 SrcVertexID = SubsetVertexIDToSrcVertexIDMap[SubsetVertexID];
-		const FVector3d Position = SubsetPositionBuffer[SubsetVertexID];
-		
-		TargetMesh.SetVertex(SrcVertexID, Position);
-	}
-}
-
-void FConstrainedMeshDeformerTask::ApplyAttenuation()
-{
-	size_t Size = SrcMeshMaxVertexID;
-	TSet<int> Handles;
-
-	auto InPlaceMinMaxElements = [](FVector3d& Min, FVector3d& Max, const FVector3d Test)
-	{
-		for (uint8 i = 0; i < 3; ++i)
-		{
-			Min[i] = Test[i] < Min[i] ? Test[i] : Min[i];
-			Max[i] = Test[i] > Max[i] ? Test[i] : Max[i];
-		}
-	};
-
-	//Experimental approach: Just going to try grabbing the bounding box of the entire mesh, then the bounding box of the handles as a point cloud.
-	//						 We need a T value to pass to the Weights curve, so let's try finding the distance of each vertex V from line segment formed by the min/max handles
-	//						 Divide the distance from the handles to vertex V by the length of the mesh's bounding box extent, 
-	//						 and that will provide a **ROUGH** approximation of the time value for our curve.
-	//
-	//											  Distance( LineSegment(MaxHandle,MinHandle) , V )
-	// where T(V) is time value at V     T(V) = -----------------------------------------------------
-	//   and V is the position								Length(MeshMin - MeshMax)   
-	//      of each vertex
-
-	FVector3d Min{ std::numeric_limits<double>::max(),std::numeric_limits<double>::max(),std::numeric_limits<double>::max() };
-	FVector3d Max{ std::numeric_limits<double>::min(),std::numeric_limits<double>::min(),std::numeric_limits<double>::min() };
-	FVector3d MinHandles = Min;
-	FVector3d MaxHandles = Max;
-	double LeastWeight = std::numeric_limits<double>::max();
-
-
-
-	for (int32 SubVertexID = 0; SubVertexID < SubsetConstraintBuffer.Num(); ++SubVertexID)
-	{
-
-		FConstraintData& CData = SubsetConstraintBuffer[SubVertexID];
-		// Update bounding box
-		InPlaceMinMaxElements(Min, Max, CData.Position);
-		
-		if (CData.Weight > 0.0)
-		{
-			LeastWeight = CData.Weight < LeastWeight ? CData.Weight : LeastWeight;
-
-			// update bounding box
-			InPlaceMinMaxElements(MinHandles, MaxHandles, CData.Position);
-			Handles.Add(SubVertexID);
-		}
-	}
-
-	double ExtentLength = Min.Distance(Max);
-
-	// Is this why the system has memory?
-	for (int32 SubVertexID = 0; SubVertexID < SubsetConstraintBuffer.Num(); ++SubVertexID)
-	{
-		if (!Handles.Contains(SubVertexID))
-		{
-			FConstraintData& CData = SubsetConstraintBuffer[SubVertexID];
-			double T = CData.Position.Distance(FMath::ClosestPointOnSegment(CData.Position, MinHandles, MaxHandles)) / ExtentLength;
-			CData.Weight = WeightAttenuationCurve.Eval(T) * LeastWeight;
-		}
-
-	}
-}
-
-
-/*
-*	FGroupTopologyLaplacianDeformer methods
-*/
-
-
-void FGroupTopologyLaplacianDeformer::InitBackgroundWorker(const ELaplacianWeightScheme WeightScheme)
-{
-
-	//Initialize asynchronous deformation objects
-	if (AsyncMeshDeformTask == nullptr)
-	{
-		AsyncMeshDeformTask = new FAsyncTaskExecuterWithAbort<FConstrainedMeshDeformerTask>(WeightScheme);
-	}
-}
-
-void FGroupTopologyLaplacianDeformer::InitializeConstraintBuffer()
-{
-	//MaxVertexID is used because the array is potentially sparse.
-	int MaxVertexID = Mesh->MaxVertexID();
-
-	SrcMeshConstraintBuffer.SetNum(MaxVertexID);
-	
-	for (int32 VertexID : Mesh->VertexIndicesItr())
-	{
-		FConstraintData& CD = SrcMeshConstraintBuffer[VertexID];
-		CD.Position = Mesh->GetVertex(VertexID); 
-		CD.Weight   = 0.0;
-		CD.bPostFix = false;
-	}
-}
-
-bool FGroupTopologyLaplacianDeformer::IsTaskInFlight() const
-{
-	return ( AsyncMeshDeformTask != nullptr && !AsyncMeshDeformTask->IsDone() );
-}
-
-bool FGroupTopologyLaplacianDeformer::UpdateAndLaunchdWorker(const ELaplacianWeightScheme SelectedWeightScheme, const FRichCurve*  Curve)
-{
-	/* Deformer needs to run if we've modified the constraints since the last time it finished. */
-	if (AsyncMeshDeformTask == nullptr)
-	{
-		InitBackgroundWorker(SelectedWeightScheme);
-	}
-
-	if (bDeformerNeedsToRun && AsyncMeshDeformTask->IsDone())
-	{		
-		bool bRebuildSubsetMesh = bTaskSubmeshIsDirty;
-
-		FConstrainedMeshDeformerTask& Task = AsyncMeshDeformTask->GetTask();
-
-		// Update the deformer's buffers and weight scheme
-		// this creates the subset mesh if needed.
-		Task.UpdateDeformer(SelectedWeightScheme, *Mesh, SrcMeshConstraintBuffer, SubsetIDBuffer, bRebuildSubsetMesh, Curve);
-
-		// task now has valid submesh
-
-		bTaskSubmeshIsDirty = false;
-
-		//Launch second thread
-		AsyncMeshDeformTask->StartBackgroundTask();
-
-		bDeformerNeedsToRun = false;	 // This was set to true above in UpdateSolution()
-		bVertexPositionsNeedSync = true; // The task will generate new vertex positions. 
-
-		return true;
-
-	}
-	return false;
-}
-
-void FGroupTopologyLaplacianDeformer::SetActiveHandleFaces(const TArray<int>& FaceGroupIDs)
-{
-	Reset();
-
-	check(FaceGroupIDs.Num() == 1);   // multi-face not supported yet
-	int GroupID = FaceGroupIDs[0];
-
-	// find set of vertices in handle 
-	Topology->CollectGroupVertices(GroupID, HandleVertices);
-	Topology->CollectGroupBoundaryVertices(GroupID, HandleBoundaryVertices);
-	ModifiedVertices = HandleVertices;
-
-
-
-	// list of adj groups.  may contain duplicates.
-	TArray<int> AdjGroups;
-	for (int BoundaryVert : HandleBoundaryVertices)
-	{
-		Topology->FindVertexNbrGroups(BoundaryVert, AdjGroups);
-	}
-
-	// Local neighborhood - Adjacent groups plus self
-	TArray<int> NeighborhoodGroups;
-
-	// Collect the rest of the 1-ring groups that are adjacent to the selected one.
-	NeighborhoodGroups.Add(GroupID);
-	for (int AdjGroup : AdjGroups)
-	{
-		NeighborhoodGroups.AddUnique(AdjGroup); // remove duplicates by add unique
-	}
-
-	CalculateROI(FaceGroupIDs, NeighborhoodGroups);
-	
-	UpdateSelection(Mesh, NeighborhoodGroups, bLocalize);
-
-	// Save the positions of the selected region.
-	SaveInitialPositions();
-
-}
-
-
-void FGroupTopologyLaplacianDeformer::SetActiveHandleEdges(const TArray<int>& TopologyEdgeIDs)
-{
-	Reset();
-
-	for (int EdgeID : TopologyEdgeIDs)
-	{
-		const TArray<int>& EdgeVerts = Topology->GetGroupEdgeVertices(EdgeID);
-		for (int VertID : EdgeVerts)
-		{
-			HandleVertices.Add(VertID);
-		}
-	}
-	HandleBoundaryVertices = HandleVertices;
-	ModifiedVertices = HandleVertices;
-
-	TArray<int> HandleGroups;
-	TArray<int> NbrGroups;
-	Topology->FindEdgeNbrGroups(TopologyEdgeIDs, NbrGroups);
-
-	CalculateROI(HandleGroups, NbrGroups);
-
-	UpdateSelection(Mesh, NbrGroups, bLocalize);
-
-	// Save the positions of the selected region.
-	SaveInitialPositions();
-}
-
-void FGroupTopologyLaplacianDeformer::SetActiveHandleCorners(const TArray<int>& CornerIDs)
-{
-	Reset();
-
-	for (int CornerID : CornerIDs)
-	{
-		int VertID = Topology->GetCornerVertexID(CornerID);
-		if (VertID >= 0)
-		{
-			HandleVertices.Add(VertID);
-		}
-	}
-	HandleBoundaryVertices = HandleVertices;
-	ModifiedVertices = HandleVertices;
-
-	TArray<int> HandleGroups;
-	TArray<int> NbrGroups;
-
-	Topology->FindCornerNbrGroups(CornerIDs, NbrGroups);
-
-
-	CalculateROI(HandleGroups, NbrGroups);
-
-	UpdateSelection(Mesh, NbrGroups, bLocalize);
-
-	// Save the positions of the selected region.
-	SaveInitialPositions();
-}
-
-
-
-void FGroupTopologyLaplacianDeformer::UpdateSelection(const FDynamicMesh3* TargetMesh, const TArray<int>& Groups, bool bLocalizeDeformation)
-{
-
-	// Build an index buffer (SubsetIdBuffer) and a vertexId buffer (ModifidedVertices) for the region we want to change
-
-	if (bLocalizeDeformation)
-	{
-		//For each group ID, retrieve the array of all TriangleIDs associated with that GroupID and append that array to the end of the TriSet to remove duplicates
-		TSet<int> TriSet;
-		for (const int32& GroupID : Groups)
-		{
-			TriSet.Append(Topology->GetGroupFaces(GroupID));
-		}//Now we have every triangle ID involved in the transaction 
-
-		//Since we are flattening the Face to a set of 3 indices, we do 3 * number of triangles though it is too many.
-		SubsetIDBuffer.Reset(3 * TriSet.Num());
-		//Add each triangle's A,B, and C indices to the subset triangle array.
-		for (const int& Tri : TriSet)
-		{
-			FIndex3i Triple = TargetMesh->GetTriangle(Tri);
-			SubsetIDBuffer.Add(Triple.A);
-			SubsetIDBuffer.Add(Triple.B);
-			SubsetIDBuffer.Add(Triple.C);
-		
-		}
-	}
-	else
-	{
-		// the entire mesh.
-		const int32 NumTris = TargetMesh->TriangleCount();
-		SubsetIDBuffer.Reset(3 * NumTris);
-		for (int TriId : TargetMesh->TriangleIndicesItr()) 
-		{ 
-			FIndex3i Triple = TargetMesh->GetTriangle(TriId);
-			SubsetIDBuffer.Add(Triple.A);
-			SubsetIDBuffer.Add(Triple.B);
-			SubsetIDBuffer.Add(Triple.C);
-		}
-	}
-
-	// Add the vertices to the set (eliminates duplicates.)  Todo: don't use a set.
-	ResetModifiedVertices();
-	for (int32 VertexID : SubsetIDBuffer)
-	{
-		RecordModifiedVertex(VertexID);
-	}
-
-}
-
-// This actually updates constraints that correspond to the handle vertices.
-void FGroupTopologyLaplacianDeformer::UpdateSolution(FDynamicMesh3 * TargetMesh, const TFunction<FVector3d(FDynamicMesh3* Mesh, int)>& HandleVertexDeformFunc)
-{
-	// copy the current positions.
-	FVertexPositionCache CurrentPositions;
-	for (int VertexID : InitialPositions.Vertices)
-	{
-		CurrentPositions.AddVertex(TargetMesh, VertexID);
-	}
-
-	// Set the target mesh to the initial positions.
-	// Note: this only updates the vertices in the selected region.
-	InitialPositions.SetPositions(TargetMesh);
-
-	//Reset the constraints
-	for (int32 VertexID : ModifiedVertices)
-	{
-		//Get the vertex's data from the constraint buffer
-		FConstraintData& CData = SrcMeshConstraintBuffer[VertexID];
-
-		CData.Position = TargetMesh->GetVertex(VertexID);
-		CData.Weight   = 0.0; //A weight of zero is used to allow this point to move freely when moving the handles
-		CData.bPostFix = false;
-	}
-
-	//Actually deform the handles and add a constraint.
-	for (int VertexID : HandleVertices)
-	{
-		const FVector3d DeformPos = HandleVertexDeformFunc(TargetMesh, VertexID);
-
-		//Get the vertex's data from the constraint buffer
-		FConstraintData& CData = SrcMeshConstraintBuffer[VertexID];
-
-		//Set the new vertex data
-		CData.Position = DeformPos;
-		CData.Weight   = HandleWeights;
-		CData.bPostFix = bPostfixHandles;
-	}
-
-	// Restore Current Positions.  This is done because the target mesh is being used to define the highlight region.
-	// if we don't reset the positions the highlight mesh will appear to reset momentarily until the first laplacian solver result is available 
-	CurrentPositions.SetPositions(TargetMesh);
-
-	bDeformerNeedsToRun = true;
-
-}
-
-void FGroupTopologyLaplacianDeformer::ExportDeformedPositions(FDynamicMesh3* TargetMesh)
-{
-	bool bIsWorking = IsTaskInFlight();
-	if (AsyncMeshDeformTask != nullptr && !bIsWorking)
-	{
-		const FConstrainedMeshDeformerTask& Task = AsyncMeshDeformTask->GetTask();
-		Task.ExportResults(*TargetMesh);
-	}
-}
-
-inline FGroupTopologyLaplacianDeformer::~FGroupTopologyLaplacianDeformer()
-{
-	Shutdown();
-}
-
-inline void FGroupTopologyLaplacianDeformer::Shutdown()
-{	
-	if (AsyncMeshDeformTask != nullptr)
-	{
-		if (AsyncMeshDeformTask->IsDone())
-		{
-			delete AsyncMeshDeformTask;
-		}
-		else
-		{
-			AsyncMeshDeformTask->CancelAndDelete();
-		}
-
-		AsyncMeshDeformTask = nullptr;
-	}
-}
 
 /*
 * Tool methods
@@ -618,17 +71,23 @@ void UEditMeshPolygonsTool::Setup()
 {
 	UMeshSurfacePointTool::Setup();
 
+	// register click behavior
+	USingleClickInputBehavior* ClickBehavior = NewObject<USingleClickInputBehavior>();
+	ClickBehavior->Initialize(this);
+	AddInputBehavior(ClickBehavior);
+
 	// create dynamic mesh component to use for live preview
 	DynamicMeshComponent = NewObject<USimpleDynamicMeshComponent>(ComponentTarget->GetOwnerActor(), "DynamicMesh");
 	DynamicMeshComponent->SetupAttachment(ComponentTarget->GetOwnerActor()->GetRootComponent());
 	DynamicMeshComponent->RegisterComponent();
 	DynamicMeshComponent->SetWorldTransform(ComponentTarget->GetWorldTransform());
 
-	// copy material if there is one
-	auto Material = ComponentTarget->GetMaterial(0);
-	if (Material != nullptr)
+	// set materials
+	FComponentMaterialSet MaterialSet;
+	ComponentTarget->GetMaterialSet(MaterialSet);
+	for (int k = 0; k < MaterialSet.Materials.Num(); ++k)
 	{
-		DynamicMeshComponent->SetMaterial(0, Material);
+		DynamicMeshComponent->SetMaterial(k, MaterialSet.Materials[k]);
 	}
 
 	// dynamic mesh configuration settings
@@ -659,7 +118,7 @@ void UEditMeshPolygonsTool::Setup()
 	TopoSelector.PointsWithinToleranceTest = [this](const FVector3d& Position1, const FVector3d& Position2) {
 		FTransform Transform = ComponentTarget->GetWorldTransform();
 		return ToolSceneQueriesUtil::PointSnapQuery(this->CameraState, 
-			Transform.TransformPosition(Position1), Transform.TransformPosition(Position2), VisualAngleSnapThreshold);
+			Transform.TransformPosition((FVector)Position1), Transform.TransformPosition((FVector)Position2), VisualAngleSnapThreshold);
 	};
 
 	// hide input StaticMeshComponent
@@ -677,34 +136,30 @@ void UEditMeshPolygonsTool::Setup()
 	PolyEdgesRenderer.LineThickness = 2.0;
 	HilightRenderer.LineColor = FLinearColor::Green;
 	HilightRenderer.LineThickness = 4.0f;
+	SelectionRenderer.LineColor = LinearColors::Gold3f<FLinearColor>();
+	SelectionRenderer.LineThickness = 4.0f;
+
+	MultiTransformer = NewObject<UMultiTransformer>(this);
+	MultiTransformer->Setup(GetToolManager()->GetPairedGizmoManager());
+	MultiTransformer->OnTransformStarted.AddUObject(this, &UEditMeshPolygonsTool::OnMultiTransformerTransformBegin);
+	MultiTransformer->OnTransformUpdated.AddUObject(this, &UEditMeshPolygonsTool::OnMultiTransformerTransformUpdate);
+	MultiTransformer->OnTransformCompleted.AddUObject(this, &UEditMeshPolygonsTool::OnMultiTransformerTransformEnd);
+	MultiTransformer->SetSnapToWorldGridSourceFunc([this]() {
+		return TransformProps->bSnapToWorldGrid
+			&& GetToolManager()->GetContextQueriesAPI()->GetCurrentCoordinateSystem() == EToolContextCoordinateSystem::World;
+	});
+	MultiTransformer->SetGizmoVisibility(false);
 
-	// Allocates buffers, sets up the asynchronous task
-	// Copies the source mesh positions.
-	const ELaplacianWeightScheme LaplacianWeightScheme = ConvertToLaplacianWeightScheme(TransformProps->SelectedWeightScheme);
-	LaplacianDeformer.InitBackgroundWorker(LaplacianWeightScheme);
-	
+	TransformerModeWatcher.Initialize(
+		[this]() { return this->TransformProps->TransformMode; },
+		[this](EMultiTransformerMode NewMode) { this->UpdateTransformerMode(NewMode); },
+		TransformProps->TransformMode);
 
-	/**
-	// How to add a curve for the weights.
-	//Add a default curve for falloff
-	FKeyHandle Keys[5];
-	Keys[0] = TransformProps->DefaultFalloffCurve.UpdateOrAddKey(0.f, 1.f);
-	Keys[1] = TransformProps->DefaultFalloffCurve.UpdateOrAddKey(0.25f, 0.25f);
-	Keys[2] = TransformProps->DefaultFalloffCurve.UpdateOrAddKey(0.3333333f, 0.25f);
-	Keys[3] = TransformProps->DefaultFalloffCurve.UpdateOrAddKey(0.6666667f, 1.25f);
-	Keys[4] = TransformProps->DefaultFalloffCurve.UpdateOrAddKey(1.f, 1.4f);
-	for (uint8 i = 0; i < 5; ++i)
-	{
-		TransformProps->DefaultFalloffCurve.SetKeyInterpMode(Keys[i], ERichCurveInterpMode::RCIM_Cubic);
-	}
-	TransformProps->WeightAttenuationCurve.EditorCurveData = TransformProps->DefaultFalloffCurve;
-	*/
 }
 
 void UEditMeshPolygonsTool::Shutdown(EToolShutdownType ShutdownType)
 {
-	//Tell the background thread to cancel the rest of its jobs before we close;
-	LaplacianDeformer.Shutdown();
+	MultiTransformer->Shutdown();
 
 	if (DynamicMeshComponent != nullptr)
 	{
@@ -716,11 +171,11 @@ void UEditMeshPolygonsTool::Shutdown(EToolShutdownType ShutdownType)
 		{
 			// this block bakes the modified DynamicMeshComponent back into the StaticMeshComponent inside an undo transaction
 			GetToolManager()->BeginUndoTransaction(LOCTEXT("EditMeshPolygonsToolTransactionName", "Deform Mesh"));
-			ComponentTarget->CommitMesh([=](FMeshDescription* MeshDescription)
+			ComponentTarget->CommitMesh([=](const FPrimitiveComponentTarget::FCommitParams& CommitParams)
 			{
 				FConversionToMeshDescriptionOptions ConversionOptions;
 				ConversionOptions.bSetPolyGroups = false; // don't save polygroups, as we may change these temporarily in this tool just to get a different edit effect
-				DynamicMeshComponent->Bake(MeshDescription, false, ConversionOptions);
+				DynamicMeshComponent->Bake(CommitParams.MeshDescription, false, ConversionOptions);
 			});
 			GetToolManager()->EndUndoTransaction();
 		}
@@ -733,18 +188,17 @@ void UEditMeshPolygonsTool::Shutdown(EToolShutdownType ShutdownType)
 
 
 
-
 void UEditMeshPolygonsTool::NextTransformTypeAction()
 {
 	if (bInDrag == false)
 	{
-		if (TransformProps->TransformMode == EQuickTransformerMode::AxisRotation)
+		if (TransformProps->TransformMode == EMultiTransformerMode::DefaultGizmo)
 		{
-			TransformProps->TransformMode = EQuickTransformerMode::AxisTranslation;
+			TransformProps->TransformMode = EMultiTransformerMode::QuickAxisTranslation;
 		}
 		else
 		{
-			TransformProps->TransformMode = EQuickTransformerMode::AxisRotation;
+			TransformProps->TransformMode = EMultiTransformerMode::DefaultGizmo;
 		}
 		UpdateQuickTransformer();
 	}
@@ -769,26 +223,6 @@ void UEditMeshPolygonsTool::OnDynamicMeshComponentChanged()
 {
 	bSpatialDirty = true;
 	TopoSelector.Invalidate(true, false);
-
-
-	//Makes sure the constraint buffer and position buffers reflect Undo/Redo changes
-	FDynamicMesh3* Mesh = DynamicMeshComponent->GetMesh();
-	
-	
-	//Apply Undo/redo
-	for (int VertexID : Mesh->VertexIndicesItr())
-	{
-		const FVector3d Position = Mesh->GetVertex(VertexID);
-		LaplacianDeformer.SrcMeshConstraintBuffer[VertexID].Position = Position;
-	}
-
-	// a deform task could still be in flight.
-	if (LaplacianDeformer.AsyncMeshDeformTask != nullptr)
-	{
-		LaplacianDeformer.AsyncMeshDeformTask->CancelAndDelete();
-		LaplacianDeformer.AsyncMeshDeformTask = nullptr;
-		LaplacianDeformer.bTaskSubmeshIsDirty = true;
-	}
 }
 
 FDynamicMeshAABBTree3& UEditMeshPolygonsTool::GetSpatial()
@@ -802,32 +236,46 @@ FDynamicMeshAABBTree3& UEditMeshPolygonsTool::GetSpatial()
 }
 
 
+
+void UEditMeshPolygonsTool::UpdateTransformerMode(EMultiTransformerMode NewMode)
+{
+	MultiTransformer->SetMode(NewMode);
+}
+
+
+
 bool UEditMeshPolygonsTool::HitTest(const FRay& WorldRay, FHitResult& OutHit)
+{
+	FGroupTopologySelection Selection;
+	return TopologyHitTest(WorldRay, OutHit, Selection);
+}
+
+
+bool UEditMeshPolygonsTool::TopologyHitTest(const FRay& WorldRay, FHitResult& OutHit, FGroupTopologySelection& OutSelection)
 {
 	FTransform Transform = ComponentTarget->GetWorldTransform();
 	FRay3d LocalRay(Transform.InverseTransformPosition(WorldRay.Origin),
 		Transform.InverseTransformVector(WorldRay.Direction));
 	LocalRay.Direction.Normalize();
 
-	TopoSelector.UpdateEnableFlags(TransformProps->bSelectFaces, TransformProps->bSelectEdges, TransformProps->bSelectVertices);
-	FGroupTopologySelection Selection;
+	UpdateTopoSelector();
 	FVector3d LocalPosition, LocalNormal;
-	if (TopoSelector.FindSelectedElement(LocalRay, Selection, LocalPosition, LocalNormal) == false)
+	if (TopoSelector.FindSelectedElement(LocalRay, OutSelection, LocalPosition, LocalNormal) == false)
 	{
 		return false;
 	}
 
-	if (Selection.SelectedCornerIDs.Num() > 0)
+	if (OutSelection.SelectedCornerIDs.Num() > 0)
 	{
-		OutHit.FaceIndex = Selection.SelectedCornerIDs[0];
+		OutHit.FaceIndex = OutSelection.SelectedCornerIDs[0];
 		OutHit.Distance = LocalRay.Project(LocalPosition);
-		OutHit.ImpactPoint = Transform.TransformPosition(LocalRay.PointAt(OutHit.Distance));
+		OutHit.ImpactPoint = Transform.TransformPosition((FVector)LocalRay.PointAt(OutHit.Distance));
 	}
-	else if (Selection.SelectedEdgeIDs.Num() > 0)
+	else if (OutSelection.SelectedEdgeIDs.Num() > 0)
 	{
-		OutHit.FaceIndex = Selection.SelectedEdgeIDs[0];
+		OutHit.FaceIndex = OutSelection.SelectedEdgeIDs[0];
 		OutHit.Distance = LocalRay.Project(LocalPosition);
-		OutHit.ImpactPoint = Transform.TransformPosition(LocalRay.PointAt(OutHit.Distance));
+		OutHit.ImpactPoint = Transform.TransformPosition((FVector)LocalRay.PointAt(OutHit.Distance));
 	}
 	else
 	{
@@ -840,14 +288,106 @@ bool UEditMeshPolygonsTool::HitTest(const FRay& WorldRay, FHitResult& OutHit)
 			Query.Find();
 			OutHit.FaceIndex = HitTID;
 			OutHit.Distance = Query.RayParameter;
-			OutHit.Normal = Transform.TransformVectorNoScale(GetSpatial().GetMesh()->GetTriNormal(HitTID));
-			OutHit.ImpactPoint = Transform.TransformPosition(LocalRay.PointAt(Query.RayParameter));
+			OutHit.Normal = Transform.TransformVectorNoScale((FVector)GetSpatial().GetMesh()->GetTriNormal(HitTID));
+			OutHit.ImpactPoint = Transform.TransformPosition((FVector)LocalRay.PointAt(Query.RayParameter));
 		}
 	}
 	return true;
 }
 
 
+
+
+FInputRayHit UEditMeshPolygonsTool::CanBeginClickDragSequence(const FInputDeviceRay& PressPos)
+{
+	if (TransformProps->TransformMode == EMultiTransformerMode::DefaultGizmo)
+	{
+		return FInputRayHit();
+	}
+	return UMeshSurfacePointTool::CanBeginClickDragSequence(PressPos);
+}
+
+
+void UEditMeshPolygonsTool::UpdateTopoSelector()
+{
+	bool bFaces = TransformProps->bSelectFaces;
+	bool bEdges = TransformProps->bSelectEdges;
+	bool bVertices = TransformProps->bSelectVertices;
+
+	if (PersistentSelection.IsEmpty() == false)
+	{
+		bFaces = bFaces && PersistentSelection.SelectedGroupIDs.Num() > 0;
+		bEdges = bEdges && PersistentSelection.SelectedEdgeIDs.Num() > 0;
+		bVertices = bVertices && PersistentSelection.SelectedCornerIDs.Num() > 0;
+	}
+	
+	TopoSelector.UpdateEnableFlags(bFaces, bEdges, bVertices);
+}
+
+
+FInputRayHit UEditMeshPolygonsTool::IsHitByClick(const FInputDeviceRay& ClickPos)
+{
+	if (TransformProps->TransformMode != EMultiTransformerMode::DefaultGizmo)
+	{
+		return FInputRayHit();
+	}
+
+	FHitResult OutHit;
+	if (HitTest(ClickPos.WorldRay, OutHit))
+	{
+		return FInputRayHit(OutHit.Distance);
+	}
+
+	// background capture, if nothing else is hit
+	return FInputRayHit(TNumericLimits<float>::Max());
+}
+
+void UEditMeshPolygonsTool::OnClicked(const FInputDeviceRay& ClickPos)
+{
+	FTransform Transform = ComponentTarget->GetWorldTransform();
+	FRay3d LocalRay(Transform.InverseTransformPosition(ClickPos.WorldRay.Origin),
+		Transform.InverseTransformVector(ClickPos.WorldRay.Direction));
+	LocalRay.Direction.Normalize();
+
+	UpdateTopoSelector();
+
+	FVector3d LocalPosition, LocalNormal;
+	FGroupTopologySelection Selection;
+	if (TopoSelector.FindSelectedElement(LocalRay, Selection, LocalPosition, LocalNormal))
+	{
+		if (GetShiftToggle())
+		{
+			PersistentSelection.Toggle(Selection);
+		}
+		else
+		{
+			PersistentSelection = Selection;
+		}
+	}
+	else
+	{
+		PersistentSelection.Clear();
+	}
+
+	// really just want show hide here....
+	if (PersistentSelection.IsEmpty())
+	{
+		MultiTransformer->SetGizmoVisibility(false);
+		return;
+	}
+	else
+	{
+		MultiTransformer->SetGizmoVisibility(true);
+	}
+
+	// update selection
+	FFrame3d SelectionFrame = Topology.GetSelectionFrame(PersistentSelection);
+	SelectionFrame.Transform(Transform);
+	MultiTransformer->SetGizmoPositionFromWorldFrame(SelectionFrame);
+}
+
+
+
 void UEditMeshPolygonsTool::OnBeginDrag(const FRay& WorldRay)
 {
 	FTransform Transform = ComponentTarget->GetWorldTransform();
@@ -870,8 +410,8 @@ void UEditMeshPolygonsTool::OnBeginDrag(const FRay& WorldRay)
 
 	HilightSelection = Selection;
 
-	FVector3d WorldHitPos = Transform.TransformPosition(LocalPosition);
-	FVector3d WorldHitNormal = Transform.TransformVector(LocalNormal);
+	FVector WorldHitPos = Transform.TransformPosition((FVector)LocalPosition);
+	FVector WorldHitNormal = Transform.TransformVector((FVector)LocalNormal);
 
 	bInDrag = true;
 	StartHitPosWorld = WorldHitPos;
@@ -885,83 +425,16 @@ void UEditMeshPolygonsTool::OnBeginDrag(const FRay& WorldRay)
 	LastBrushPosLocal = Transform.InverseTransformPosition(LastHitPosWorld);
 	StartBrushPosLocal = LastBrushPosLocal;
 
-	// Record the requested deformation strategy - NB: will be forced to linear if there aren't any free points to solve.
-
-	DeformationStrategy = TransformProps->DeformationStrategy;
-	
-	// Capture the part of the mesh that will deform
-
-	if (DeformationStrategy == EGroupTopologyDeformationStrategy::Laplacian)
-	{
-		LaplacianDeformer.bLocalize = true; // TransformProps->bLocalizeDeformation;
-
-		//Determine which of the following (corners, edges or faces) has been selected by counting the associated feature's IDs
-		if (Selection.SelectedCornerIDs.Num() > 0)
-		{
-			//Add all the the Corner's adjacent poly-groups (NbrGroups) to the ongoing array of groups.
-			LaplacianDeformer.SetActiveHandleCorners(Selection.SelectedCornerIDs);
-		}
-		else if (Selection.SelectedEdgeIDs.Num() > 0)
-		{
-			//Add all the the edge's adjacent poly-groups (NbrGroups) to the ongoing array of groups.
-			LaplacianDeformer.SetActiveHandleEdges(Selection.SelectedEdgeIDs);
-		}
-		else if (Selection.SelectedGroupIDs.Num() > 0)
-		{
-			LaplacianDeformer.SetActiveHandleFaces(Selection.SelectedGroupIDs);
-		}
-
-
-		// If there are actually no interior points, then we can't actually use the laplacian deformer. Need to fall back to the linear.
-		bool bHasInteriorVerts = false;
-		const auto& ROIFaces = LaplacianDeformer.GetROIFaces();
-		for (const auto& Face : ROIFaces)
-		{
-			bHasInteriorVerts = bHasInteriorVerts || ( Face.InteriorVerts.Num() != 0);
-		}
-
-		if (!bHasInteriorVerts)
-		{
-			// Change to the linear strategy for this case.
-
-			DeformationStrategy = EGroupTopologyDeformationStrategy::Linear;
-		}
-		else
-		{		
-			// finalize the laplacian deformer : the task will need a new mesh that corresponds to the selected region.
-		
-			LaplacianDeformer.bTaskSubmeshIsDirty = true;
-
-		}
-		
-	}
-	 
-	if (DeformationStrategy == EGroupTopologyDeformationStrategy::Linear )
-	{
-		//Determine which of the following (corners, edges or faces) has been selected by counting the associated feature's IDs
-		if (Selection.SelectedCornerIDs.Num() > 0)
-		{
-			//Add all the the Corner's adjacent poly-groups (NbrGroups) to the ongoing array of groups.
-			LinearDeformer.SetActiveHandleCorners(Selection.SelectedCornerIDs);
-		}
-		else if (Selection.SelectedEdgeIDs.Num() > 0)
-		{
-			//Add all the the edge's adjacent poly-groups (NbrGroups) to the ongoing array of groups.
-			LinearDeformer.SetActiveHandleEdges(Selection.SelectedEdgeIDs);
-		}
-		else if (Selection.SelectedGroupIDs.Num() > 0)
-		{
-			LinearDeformer.SetActiveHandleFaces(Selection.SelectedGroupIDs);
-		}
-	}
+	UpdateDeformerFromSelection(Selection);
 	
 	BeginChange();
 }
 
 
+
 void UEditMeshPolygonsTool::UpdateActiveSurfaceFrame(FGroupTopologySelection& Selection)
 {
-	FTransform Transform = ComponentTarget->GetWorldTransform();
+	FTransform3d Transform(ComponentTarget->GetWorldTransform());
 
 	// update surface frame
 	ActiveSurfaceFrame.Origin = StartHitPosWorld;
@@ -987,19 +460,29 @@ void UEditMeshPolygonsTool::UpdateActiveSurfaceFrame(FGroupTopologySelection& Se
 
 FQuickTransformer* UEditMeshPolygonsTool::GetActiveQuickTransformer()
 {
-	if (TransformProps->TransformMode == EQuickTransformerMode::AxisRotation)
-	{
-		return &QuickAxisRotator;
-	}
-	else
+	//if (TransformProps->TransformMode == EQuickTransformerMode::AxisRotation)
+	//{
+	//	return &QuickAxisRotator;
+	//}
+	//else
+	//{
+	//	return &QuickAxisTranslater;
+	//}
+	if (TransformProps->TransformMode == EMultiTransformerMode::QuickAxisTranslation)
 	{
 		return &QuickAxisTranslater;
 	}
+	return nullptr;
 }
 
 
 void UEditMeshPolygonsTool::UpdateQuickTransformer()
 {
+	if (GetActiveQuickTransformer() == nullptr)
+	{
+		return;
+	}
+
 	bool bUseLocalAxes =
 		(GetToolManager()->GetContextQueriesAPI()->GetCurrentCoordinateSystem() == EToolContextCoordinateSystem::Local);
 	if (bUseLocalAxes)
@@ -1022,13 +505,10 @@ void UEditMeshPolygonsTool::UpdateChangeFromROI(bool bFinal)
 	{
 		return;
 	}
-	const bool bIsLaplacian = (DeformationStrategy == EGroupTopologyDeformationStrategy::Laplacian);
 
 	FDynamicMesh3* Mesh = DynamicMeshComponent->GetMesh();
-	const TSet<int>& ModifiedVertices = (bIsLaplacian) ? LaplacianDeformer.GetModifiedVertices() : LinearDeformer.GetModifiedVertices();
-	
+	const TSet<int>& ModifiedVertices = LinearDeformer.GetModifiedVertices();
 	ActiveVertexChange->SavePositions(Mesh, ModifiedVertices, !bFinal);
-
 }
 
 
@@ -1054,35 +534,61 @@ void UEditMeshPolygonsTool::OnEndDrag(const FRay& Ray)
 	QuickAxisRotator.Reset();
 	QuickAxisTranslater.Reset();
 
-	//If it's linear, it's computed real time with no delay. This may need to be restructured for clarity by using the background task for this as well.
-	if (DeformationStrategy == EGroupTopologyDeformationStrategy::Linear)
+	// close change record
+	EndChange();
+}
+
+
+
+void UEditMeshPolygonsTool::OnMultiTransformerTransformBegin()
+{
+	HilightSelection.Clear();
+
+	UpdateDeformerFromSelection(PersistentSelection);
+
+	InitialGizmoFrame = MultiTransformer->GetCurrentGizmoFrame();
+
+	BeginChange();
+}
+
+void UEditMeshPolygonsTool::OnMultiTransformerTransformUpdate()
+{
+	if (MultiTransformer->InGizmoEdit())
 	{
-		// close change record
-		EndChange();
+		ComputeUpdate_Gizmo();
 	}
 }
 
+void UEditMeshPolygonsTool::OnMultiTransformerTransformEnd()
+{
+	bSpatialDirty = true;
+
+	TopoSelector.Invalidate(true, false);
+
+	// close change record
+	EndChange();
+}
+
 
 
 bool UEditMeshPolygonsTool::OnUpdateHover(const FInputDeviceRay& DevicePos)
 {
-	//if (!bNeedEmitEndChange)
-	if (ActiveVertexChange == nullptr)
+	if (ActiveVertexChange == nullptr && MultiTransformer->InGizmoEdit() == false )
 	{
-		FTransform Transform = ComponentTarget->GetWorldTransform();
+		FTransform3d Transform(ComponentTarget->GetWorldTransform());
 		FRay3d LocalRay(Transform.InverseTransformPosition(DevicePos.WorldRay.Origin),
 		Transform.InverseTransformVector(DevicePos.WorldRay.Direction));
 		LocalRay.Direction.Normalize();
 
 		HilightSelection.Clear();
-		TopoSelector.UpdateEnableFlags(TransformProps->bSelectFaces, TransformProps->bSelectEdges, TransformProps->bSelectVertices);
+		UpdateTopoSelector();
 		FVector3d LocalPosition, LocalNormal;
 		bool bHit = TopoSelector.FindSelectedElement(LocalRay, HilightSelection, LocalPosition, LocalNormal);
 
 		if (bHit)
 		{
-			StartHitPosWorld = Transform.TransformPosition(LocalPosition);
-			StartHitNormalWorld = Transform.TransformVector(LocalNormal);
+			StartHitPosWorld = (FVector)Transform.TransformPosition(LocalPosition);
+			StartHitNormalWorld = (FVector)Transform.TransformVector(LocalNormal);
 
 			UpdateActiveSurfaceFrame(HilightSelection);
 			UpdateQuickTransformer();
@@ -1094,75 +600,53 @@ bool UEditMeshPolygonsTool::OnUpdateHover(const FInputDeviceRay& DevicePos)
 
 
 
-void UEditMeshPolygonsTool::ComputeUpdate()
-{
 
-	if (bUpdatePending == true)
-	{
-		// Linear Deformer : Update the solution
-		// Laplacain Deformer : Update the constraints (positions and weights) - the region was identified in onBeginDrag
 
 
-		if (TransformProps->TransformMode == EQuickTransformerMode::AxisRotation)
-		{
-			ComputeUpdate_Rotate();
-		}
-		else
-		{
-			ComputeUpdate_Translate();
-		}
+void UEditMeshPolygonsTool::UpdateDeformerFromSelection(const FGroupTopologySelection& Selection)
+{
+	//Determine which of the following (corners, edges or faces) has been selected by counting the associated feature's IDs
+	if (Selection.SelectedCornerIDs.Num() > 0)
+	{
+		//Add all the the Corner's adjacent poly-groups (NbrGroups) to the ongoing array of groups.
+		LinearDeformer.SetActiveHandleCorners(Selection.SelectedCornerIDs);
 	}
-
-	if (DeformationStrategy == EGroupTopologyDeformationStrategy::Laplacian)
+	else if (Selection.SelectedEdgeIDs.Num() > 0)
 	{
-		bool bIsWorking = LaplacianDeformer.IsTaskInFlight();
-
-		if (!bIsWorking)
-		{
-			// Sync update if we have new results.
-			if (LaplacianDeformer.bVertexPositionsNeedSync)
-			{
-				
-				//Update the mesh with the provided solutions.
-				LaplacianDeformer.ExportDeformedPositions(DynamicMeshComponent->GetMesh());
+		//Add all the the edge's adjacent poly-groups (NbrGroups) to the ongoing array of groups.
+		LinearDeformer.SetActiveHandleEdges(Selection.SelectedEdgeIDs);
+	}
+	else if (Selection.SelectedGroupIDs.Num() > 0)
+	{
+		LinearDeformer.SetActiveHandleFaces(Selection.SelectedGroupIDs);
+	}
+}
 
-				LaplacianDeformer.bVertexPositionsNeedSync = false;
 
-				//Re-sync mesh, and flag the spatial data struct & topology for re-evaluation
-				DynamicMeshComponent->FastNotifyPositionsUpdated();
 
-				GetToolManager()->PostInvalidation();
-				bSpatialDirty = true;
-				TopoSelector.Invalidate(true, false);
-			}
 
-			// emit end change if we are done with the drag
-			if (!LaplacianDeformer.bDeformerNeedsToRun && !bInDrag)
-			{
-				EndChange();
-			}
 
-			// Not working but we have more work for it to do..
+void UEditMeshPolygonsTool::ComputeUpdate()
+{
 
-			if (LaplacianDeformer.bDeformerNeedsToRun)
-			{
-			
-				FRichCurve* Curve = NULL;
-
-				/**
-				// How to add a deformation curve
-				const bool bApplyAttenuationCurve = TransformProps->bApplyAttenuationCurve
-				if (bApplyAttenuationCurve)
-				{
-					Curve = TransformProps->WeightAttenuationCurve.GetRichCurve();
-				}
-				*/
-				const ELaplacianWeightScheme LaplacianWeightScheme = ConvertToLaplacianWeightScheme(TransformProps->SelectedWeightScheme);
-				LaplacianDeformer.UpdateAndLaunchdWorker(LaplacianWeightScheme, Curve);
+	if (bUpdatePending == true)
+	{
+		//if (TransformProps->TransformMode == EQuickTransformerMode::AxisRotation)
+		//{
+		//	ComputeUpdate_Rotate();
+		//}
+		//else
+		//{
+		//	ComputeUpdate_Translate();
+		//}
 
-			}
+		if (TransformProps->TransformMode == EMultiTransformerMode::QuickAxisTranslation)
+		{
+			ComputeUpdate_Translate();
 		}
+
 	}
+
 }
 
 
@@ -1172,17 +656,14 @@ void UEditMeshPolygonsTool::ComputeUpdate()
 
 void UEditMeshPolygonsTool::ComputeUpdate_Rotate()
 {
-	const bool bIsLaplacian = (DeformationStrategy == EGroupTopologyDeformationStrategy::Laplacian); 
-	FGroupTopologyDeformer& SelectedDeformer = (bIsLaplacian) ? LaplacianDeformer : LinearDeformer;
-
 	FDynamicMesh3* Mesh = DynamicMeshComponent->GetMesh();
-	FTransform Transform = ComponentTarget->GetWorldTransform();
+	FTransform3d Transform(ComponentTarget->GetWorldTransform());
 	FVector NewHitPosWorld = LastHitPosWorld;
 
 	FVector3d SnappedPoint;
-	if (QuickAxisRotator.UpdateSnap(UpdateRay, SnappedPoint))
+	if (QuickAxisRotator.UpdateSnap(FRay3d(UpdateRay), SnappedPoint))
 	{
-		NewHitPosWorld = SnappedPoint;
+		NewHitPosWorld = (FVector)SnappedPoint;
 	}
 	else
 	{
@@ -1207,15 +688,11 @@ void UEditMeshPolygonsTool::ComputeUpdate_Rotate()
 	{
 		QuickAxisRotator.ClearAxisLock();
 
-		SelectedDeformer.ClearSolution(Mesh);
+		LinearDeformer.ClearSolution(Mesh);
 
-		//TODO: This is unseemly here, need to potentially defer this so that it's handled the same way as laplacian. Placeholder for now.
-		if (DeformationStrategy == EGroupTopologyDeformationStrategy::Linear)
-		{
+		DynamicMeshComponent->FastNotifyPositionsUpdated();
+		GetToolManager()->PostInvalidation();
 
-			DynamicMeshComponent->FastNotifyPositionsUpdated();
-			GetToolManager()->PostInvalidation();
-		}
 		bUpdatePending = false;
 		return;
 	}
@@ -1237,9 +714,8 @@ void UEditMeshPolygonsTool::ComputeUpdate_Rotate()
 		Transform.InverseTransformVectorNoScale(RotationStartFrame.Z()), AngleRad, false);
 	FVector3d LocalOrigin = Transform.InverseTransformPosition(RotationStartFrame.Origin);
 
-	// Linear Deformer: Update Mesh the rotation,
-	// Laplacian Deformer:  Update handles constraints with the rotation and set bDeformerNeedsToRun = true;.
-	SelectedDeformer.UpdateSolution(Mesh, [this, LocalOrigin, Rotation](FDynamicMesh3* TargetMesh, int VertIdx)
+	// Update Mesh the rotation,
+	LinearDeformer.UpdateSolution(Mesh, [this, LocalOrigin, Rotation](FDynamicMesh3* TargetMesh, int VertIdx)
 	{
 		FVector3d V = TargetMesh->GetVertex(VertIdx);
 		V -= LocalOrigin;
@@ -1248,12 +724,8 @@ void UEditMeshPolygonsTool::ComputeUpdate_Rotate()
 		return V;
 	});
 
-	//TODO: This is unseemly here, need to potentially defer this so that it's handled the same way as laplacian. Placeholder for now.
-	if (!bIsLaplacian)
-	{
-		DynamicMeshComponent->FastNotifyPositionsUpdated();
-		GetToolManager()->PostInvalidation();
-	}
+	DynamicMeshComponent->FastNotifyPositionsUpdated();
+	GetToolManager()->PostInvalidation();
 	bUpdatePending = false;
 }
 
@@ -1262,22 +734,30 @@ void UEditMeshPolygonsTool::ComputeUpdate_Rotate()
 
 void UEditMeshPolygonsTool::ComputeUpdate_Translate()
 {
-	const bool bIsLaplacian = (DeformationStrategy == EGroupTopologyDeformationStrategy::Laplacian);
-	FGroupTopologyDeformer& SelectedDeformer = (bIsLaplacian) ? LaplacianDeformer : LinearDeformer;
+	TFunction<FVector3d(const FVector3d&)> PointConstraintFunc = nullptr;
+	if (TransformProps->bSnapToWorldGrid 
+		&& GetToolManager()->GetContextQueriesAPI()->GetCurrentCoordinateSystem() == EToolContextCoordinateSystem::World)
+	{
+		PointConstraintFunc = [&](const FVector3d& Pos)
+		{
+			FVector3d GridSnapPos;
+			return ToolSceneQueriesUtil::FindWorldGridSnapPoint(this, Pos, GridSnapPos) ? GridSnapPos : Pos;
+		};
+	}
 
-	FTransform Transform = ComponentTarget->GetWorldTransform();
+	FTransform3d Transform(ComponentTarget->GetWorldTransform());
 	FVector NewHitPosWorld = LastHitPosWorld;
 	FVector3d SnappedPoint;
-	if (QuickAxisTranslater.UpdateSnap(UpdateRay, SnappedPoint))
+	if (QuickAxisTranslater.UpdateSnap(FRay3d(UpdateRay), SnappedPoint, PointConstraintFunc))
 	{
-		NewHitPosWorld = SnappedPoint;
+		NewHitPosWorld = (FVector)SnappedPoint;
 	}
 	else
 	{
 		return;
 	}
 
-	FVector NewBrushPosLocal = Transform.InverseTransformPosition(NewHitPosWorld);
+	FVector3d NewBrushPosLocal = Transform.InverseTransformPosition(NewHitPosWorld);
 	FVector3d NewMoveDelta = NewBrushPosLocal - StartBrushPosLocal;
 
 	FDynamicMesh3* Mesh = DynamicMeshComponent->GetMesh();
@@ -1285,10 +765,7 @@ void UEditMeshPolygonsTool::ComputeUpdate_Translate()
 	{
 		if (NewMoveDelta.SquaredLength() > 0.)
 		{
-			// Linear Deformer: Update Mesh with the translation,
-			// Laplacian Deformer:  Update handles constraints and set bDeformerNeedsToRun = true;.
-
-			SelectedDeformer.UpdateSolution(Mesh, [this, NewMoveDelta](FDynamicMesh3* TargetMesh, int VertIdx)
+			LinearDeformer.UpdateSolution(Mesh, [this, NewMoveDelta](FDynamicMesh3* TargetMesh, int VertIdx)
 			{
 				return TargetMesh->GetVertex(VertIdx) + NewMoveDelta;
 			});
@@ -1296,32 +773,60 @@ void UEditMeshPolygonsTool::ComputeUpdate_Translate()
 		else
 		{
 			// Reset mesh to initial positions.
-
-			SelectedDeformer.ClearSolution(Mesh);
-
-		}
-		//TODO: This is unseemly here, need to potentially defer this so that it's handled the same way as laplacian. Placeholder for now.
-		if (!bIsLaplacian)
-		{
-			DynamicMeshComponent->FastNotifyPositionsUpdated();
-			GetToolManager()->PostInvalidation();
+			LinearDeformer.ClearSolution(Mesh);
 		}
+		DynamicMeshComponent->FastNotifyPositionsUpdated();
+		GetToolManager()->PostInvalidation();
 	}
 
 	LastMoveDelta = NewMoveDelta;
-	LastBrushPosLocal = NewBrushPosLocal;
+	LastBrushPosLocal = (FVector)NewBrushPosLocal;
 
 	bUpdatePending = false;
 }
 
 
 
+
+
+void UEditMeshPolygonsTool::ComputeUpdate_Gizmo()
+{
+	if (PersistentSelection.IsEmpty())
+	{
+		return;
+	}
+
+	FFrame3d CurFrame = MultiTransformer->GetCurrentGizmoFrame();
+	FVector3d Translation = CurFrame.Origin - InitialGizmoFrame.Origin;
+	FTransform Transform = ComponentTarget->GetWorldTransform();
+	FVector3d LocalTranslation = (FVector3d)Transform.InverseTransformVector((FVector)Translation);
+
+
+	FDynamicMesh3* Mesh = DynamicMeshComponent->GetMesh();
+	if (LocalTranslation.SquaredLength() > 0.)
+	{
+		LinearDeformer.UpdateSolution(Mesh, [this, LocalTranslation](FDynamicMesh3* TargetMesh, int VertIdx)
+		{
+			return TargetMesh->GetVertex(VertIdx) + LocalTranslation;
+		});
+	}
+	else
+	{
+		// Reset mesh to initial positions.
+		LinearDeformer.ClearSolution(Mesh);
+	}
+	DynamicMeshComponent->FastNotifyPositionsUpdated();
+	GetToolManager()->PostInvalidation();
+}
+
+
+
 void UEditMeshPolygonsTool::Tick(float DeltaTime)
 {
 	UMeshSurfacePointTool::Tick(DeltaTime);
-	
-	LaplacianDeformer.HandleWeights   = TransformProps->HandleWeight;
-	LaplacianDeformer.bPostfixHandles = TransformProps->bPostFixHandles;
+
+	TransformerModeWatcher.CheckAndUpdate();
+	MultiTransformer->Tick(DeltaTime);
 }
 
 
@@ -1336,7 +841,7 @@ void UEditMeshPolygonsTool::ComputePolygons(bool RecomputeTopology)
 	{
 		FFindPolygonsAlgorithm Polygons = FFindPolygonsAlgorithm(InitialMesh.Get());
 		double DotTolerance = 1.0 - FMathd::Cos(TransformProps->PolygonGroupingAngleThreshold * FMathd::DegToRad);
-		Polygons.FindPolygons(DotTolerance);
+		Polygons.FindPolygonsFromFaceNormals(DotTolerance);
 		Polygons.FindPolygonEdges();
 		SetTriangleGroups(*InitialMesh->GetTriangleGroupsBuffer());
 	}
@@ -1369,10 +874,6 @@ void UEditMeshPolygonsTool::PrecomputeTopology()
 	Topology = FGroupTopology(Mesh, true);
 
 	LinearDeformer.Initialize(Mesh, &Topology);
-	LaplacianDeformer.Initialize(Mesh, &Topology);
-
-	// Make the Constraint Buffer, zero weights, but current pos
-	LaplacianDeformer.InitializeConstraintBuffer();
 }
 
 
@@ -1384,9 +885,8 @@ void UEditMeshPolygonsTool::Render(IToolsContextRenderAPI* RenderAPI)
 	ComputeUpdate();
 		
 	GetToolManager()->GetContextQueriesAPI()->GetCurrentViewState(CameraState);
-	GetActiveQuickTransformer()->UpdateCameraState(CameraState);
 
-	DynamicMeshComponent->bExplicitShowWireframe = TransformProps->bWireframe;
+	DynamicMeshComponent->bExplicitShowWireframe = TransformProps->bShowWireframe;
 	FDynamicMesh3* TargetMesh = DynamicMeshComponent->GetMesh();
 
 	PolyEdgesRenderer.BeginFrame(RenderAPI, CameraState);
@@ -1409,50 +909,30 @@ void UEditMeshPolygonsTool::Render(IToolsContextRenderAPI* RenderAPI)
 	HilightRenderer.BeginFrame(RenderAPI, CameraState);
 	HilightRenderer.SetTransform(ComponentTarget->GetWorldTransform());
 
-#ifdef DEBUG_ROI_WEIGHTS
-	FDynamicMesh3* MeshPtr = DynamicMeshComponent->GetMesh();
-	for (int32 VertexID : DynamicMeshComponent->GetMesh()->VertexIndicesItr())
-	{
-		float Color = 1.f - SrcMeshConstraintBuffer[VertexID].Weight;
-		HilightRenderer.DrawPoint(MeshPtr->GetVertex(VertexID), FLinearColor(Color, Color, Color, 1.f), 8, true);
-	}
-
-#endif
-
-#ifdef DEBUG_ROI_HANDLES
-	const FLinearColor FOOF{ 1.f,0.f,1.f,1.f };
-	for (int VertIdx : HandleVertices)
-	{
-		HilightRenderer.DrawViewFacingCircle(TargetMesh->GetVertex(VertIdx), 0.8f, 8, FOOF, 3, false);
-	}
-#endif 
-
-#ifdef DEBUG_ROI_TRIANGLES
-	const FLinearColor Whiteish{ 0.67f,0.67f,0.67f,1.f };
-	for (int32 i = 0; i < SubsetIDBuffer.Num(); i += 3)
-	{
-		FVector3d A = TargetMesh->GetVertex(SubsetIDBuffer[i]);
-		FVector3d B = TargetMesh->GetVertex(SubsetIDBuffer[i + 1]);
-		FVector3d C = TargetMesh->GetVertex(SubsetIDBuffer[i + 2]);
-		HilightRenderer.DrawLine(A, B, Whiteish, 2.7f, true);
-		HilightRenderer.DrawLine(B, C, Whiteish, 2.7f, true);
-		HilightRenderer.DrawLine(C, A, Whiteish, 2.7f, true);
-	}
-#endif
-
 	TopoSelector.VisualAngleSnapThreshold = this->VisualAngleSnapThreshold;
 	TopoSelector.DrawSelection(HilightSelection, &HilightRenderer, &CameraState);
 	HilightRenderer.EndFrame();
 
-
-	if (bInDrag)
+	if (PersistentSelection.IsEmpty() == false)
 	{
-
-		GetActiveQuickTransformer()->Render(RenderAPI);
+		SelectionRenderer.BeginFrame(RenderAPI, CameraState);
+		SelectionRenderer.SetTransform(ComponentTarget->GetWorldTransform());
+		SelectionRenderer.SetTransform(ComponentTarget->GetWorldTransform());
+		TopoSelector.DrawSelection(PersistentSelection, &SelectionRenderer, &CameraState);
+		SelectionRenderer.EndFrame();
 	}
-	else
+
+	if (GetActiveQuickTransformer() != nullptr)
 	{
-		GetActiveQuickTransformer()->PreviewRender(RenderAPI);
+		GetActiveQuickTransformer()->UpdateCameraState(CameraState);
+		if (bInDrag)
+		{
+			GetActiveQuickTransformer()->Render(RenderAPI);
+		}
+		else
+		{
+			GetActiveQuickTransformer()->PreviewRender(RenderAPI);
+		}
 	}
 }
 
@@ -1503,21 +983,16 @@ void UEditMeshPolygonsTool::SetTriangleGroups(const TDynamicVector<int>& Groups)
 
 void UEditMeshPolygonsTool::BeginChange()
 {
-	const bool bIsLaplacian = (DeformationStrategy == EGroupTopologyDeformationStrategy::Laplacian);
-	if (!bIsLaplacian || LaplacianDeformer.IsDone())
+	if (ActiveVertexChange == nullptr)
 	{
-		if (ActiveVertexChange == nullptr)
-		{
-			ActiveVertexChange = new FMeshVertexChangeBuilder();
-			UpdateChangeFromROI(false);
-		}
+		ActiveVertexChange = new FMeshVertexChangeBuilder();
+		UpdateChangeFromROI(false);
 	}
 }
 
 
 void UEditMeshPolygonsTool::EndChange()
 {
-
 	if (ActiveVertexChange != nullptr)
 	{
 		UpdateChangeFromROI(true);
diff --git a/Engine/Plugins/Experimental/MeshModelingToolset/Source/MeshModelingTools/Private/EditPivotTool.cpp b/Engine/Plugins/Experimental/MeshModelingToolset/Source/MeshModelingTools/Private/EditPivotTool.cpp
new file mode 100644
--- /dev/null
+++ b/Engine/Plugins/Experimental/MeshModelingToolset/Source/MeshModelingTools/Private/EditPivotTool.cpp
@@ -0,0 +1,364 @@
+// Copyright 1998-2019 Epic Games, Inc. All Rights Reserved.
+
+#include "EditPivotTool.h"
+#include "InteractiveToolManager.h"
+#include "InteractiveGizmoManager.h"
+#include "ToolBuilderUtil.h"
+#include "ToolSetupUtil.h"
+#include "DynamicMesh3.h"
+#include "DynamicMeshToMeshDescription.h"
+#include "MeshDescriptionToDynamicMesh.h"
+#include "MeshTransforms.h"
+#include "BaseBehaviors/ClickDragBehavior.h"
+
+#include "BaseGizmos/GizmoComponents.h"
+#include "BaseGizmos/TransformGizmo.h"
+
+#include "Components/PrimitiveComponent.h"
+#include "CollisionQueryParams.h"
+#include "Engine/World.h"
+
+#define LOCTEXT_NAMESPACE "UEditPivotTool"
+
+/*
+ * ToolBuilder
+ */
+
+
+bool UEditPivotToolBuilder::CanBuildTool(const FToolBuilderState& SceneState) const
+{
+	return ToolBuilderUtil::CountComponents(SceneState, CanMakeComponentTarget) >= 1;
+}
+
+UInteractiveTool* UEditPivotToolBuilder::BuildTool(const FToolBuilderState& SceneState) const
+{
+	UEditPivotTool* NewTool = NewObject<UEditPivotTool>(SceneState.ToolManager);
+
+	TArray<UActorComponent*> Components = ToolBuilderUtil::FindAllComponents(SceneState, CanMakeComponentTarget);
+	check(Components.Num() > 0);
+
+	TArray<TUniquePtr<FPrimitiveComponentTarget>> ComponentTargets;
+	for (UActorComponent* ActorComponent : Components)
+	{
+		auto* MeshComponent = Cast<UPrimitiveComponent>(ActorComponent);
+		if ( MeshComponent )
+		{
+			ComponentTargets.Add(MakeComponentTarget(MeshComponent));
+		}
+	}
+
+	NewTool->SetSelection(MoveTemp(ComponentTargets));
+	NewTool->SetWorld(SceneState.World, SceneState.GizmoManager);
+
+	return NewTool;
+}
+
+
+
+
+/*
+ * Tool
+ */
+
+UEditPivotTool::UEditPivotTool()
+{
+}
+
+void UEditPivotTool::SetWorld(UWorld* World, UInteractiveGizmoManager* GizmoManagerIn)
+{
+	this->TargetWorld = World;
+	this->GizmoManager = GizmoManagerIn;
+}
+
+
+void UEditPivotTool::Setup()
+{
+	UInteractiveTool::Setup();
+
+	UClickDragInputBehavior* ClickDragBehavior = NewObject<UClickDragInputBehavior>(this);
+	ClickDragBehavior->Initialize(this);
+	AddInputBehavior(ClickDragBehavior);
+
+	TransformProps = NewObject<UEditPivotToolProperties>();
+	AddToolPropertySource(TransformProps);
+
+	ResetActiveGizmos();
+	SetActiveGizmos_Single(false);
+	UpdateSetPivotModes(true);
+
+	GetToolManager()->DisplayMessage(
+		LOCTEXT("OnStartEditPivotTool", "To transform Objects around points, reposition the Gizmo using Set Pivot mode (S). To quickly position Objects, enable Snap Drag mode (D). A cycles through Transform modes, W and E through SnapDrag Source and Rotation types."),
+		EToolMessageLevel::UserNotification);
+
+	GetToolManager()->DisplayMessage(
+		LOCTEXT("EditPivotWarning", "WARNING: This Tool will Modify the selected StaticMesh Assets! If you do not wish to modify the original Assets, please make copies in the Content Browser first!"),
+		EToolMessageLevel::UserWarning);
+}
+
+
+
+
+
+
+
+void UEditPivotTool::Shutdown(EToolShutdownType ShutdownType)
+{
+	FFrame3d CurPivotFrame(ActiveGizmos[0].TransformProxy->GetTransform());
+
+	GizmoManager->DestroyAllGizmosByOwner(this);
+
+	if (ShutdownType == EToolShutdownType::Accept)
+	{
+		UpdateAssets(CurPivotFrame);
+	}
+}
+
+
+
+void UEditPivotTool::Tick(float DeltaTime)
+{
+}
+
+void UEditPivotTool::Render(IToolsContextRenderAPI* RenderAPI)
+{
+}
+
+void UEditPivotTool::OnPropertyModified(UObject* PropertySet, UProperty* Property)
+{
+}
+
+
+void UEditPivotTool::UpdateSetPivotModes(bool bEnableSetPivot)
+{
+	for (FEditPivotTarget& Target : ActiveGizmos)
+	{
+		Target.TransformProxy->bSetPivotMode = bEnableSetPivot;
+	}
+}
+
+
+
+void UEditPivotTool::RegisterActions(FInteractiveToolActionSet& ActionSet)
+{
+	//ActionSet.RegisterAction(this, (int32)EStandardToolActions::BaseClientDefinedActionID + 2,
+	//	TEXT("ToggleSnapDrag"),
+	//	LOCTEXT("TransformToggleSnapDrag", "Toggle SnapDrag"),
+	//	LOCTEXT("TransformToggleSnapDragTooltip", "Toggle SnapDrag on and off"),
+	//	EModifierKey::None, EKeys::D,
+	//	[this]() { this->TransformProps->bEnableSnapDragging = !this->TransformProps->bEnableSnapDragging; });
+
+	//ActionSet.RegisterAction(this, (int32)EStandardToolActions::BaseClientDefinedActionID + 3,
+	//	TEXT("CycleTransformMode"),
+	//	LOCTEXT("TransformCycleTransformMode", "Next Transform Mode"),
+	//	LOCTEXT("TransformCycleTransformModeTooltip", "Cycle through available Transform Modes"),
+	//	EModifierKey::None, EKeys::A,
+	//	[this]() { this->TransformProps->TransformMode = (EEditPivotTransformMode)(((uint8)TransformProps->TransformMode+1) % (uint8)EEditPivotTransformMode::LastValue); });
+
+
+	//ActionSet.RegisterAction(this, (int32)EStandardToolActions::BaseClientDefinedActionID + 4,
+	//	TEXT("CycleSourceMode"),
+	//	LOCTEXT("TransformCycleSourceMode", "Next SnapDrag Source Mode"),
+	//	LOCTEXT("TransformCycleSourceModeTooltip", "Cycle through available SnapDrag Source Modes"),
+	//	EModifierKey::None, EKeys::W,
+	//	[this]() { this->TransformProps->SnapDragSource = (EEditPivotSnapDragSource)(((uint8)TransformProps->SnapDragSource + 1) % (uint8)EEditPivotSnapDragSource::LastValue); });
+
+	//ActionSet.RegisterAction(this, (int32)EStandardToolActions::BaseClientDefinedActionID + 5,
+	//	TEXT("CycleRotationMode"),
+	//	LOCTEXT("TransformCycleRotationMode", "Next SnapDrag Rotation Mode"),
+	//	LOCTEXT("TransformCycleRotationModeTooltip", "Cycle through available SnapDrag Rotation Modes"),
+	//	EModifierKey::None, EKeys::E,
+	//	[this]() { this->TransformProps->RotationMode = (EEditPivotSnapDragRotationMode)(((uint8)TransformProps->RotationMode + 1) % (uint8)EEditPivotSnapDragRotationMode::LastValue); });
+
+
+}
+
+
+
+
+void UEditPivotTool::SetActiveGizmos_Single(bool bLocalRotations)
+{
+	check(ActiveGizmos.Num() == 0);
+
+	FEditPivotTarget Transformable;
+	Transformable.TransformProxy = NewObject<UTransformProxy>(this);
+	Transformable.TransformProxy->bRotatePerObject = bLocalRotations;
+
+	for (TUniquePtr<FPrimitiveComponentTarget>& Target : ComponentTargets)
+	{
+		Transformable.TransformProxy->AddComponent(Target->GetOwnerComponent());
+	}
+	Transformable.TransformGizmo = GizmoManager->Create3AxisTransformGizmo(this);
+	Transformable.TransformGizmo->SetActiveTarget(Transformable.TransformProxy);
+
+	ActiveGizmos.Add(Transformable);
+}
+
+
+void UEditPivotTool::ResetActiveGizmos()
+{
+	GizmoManager->DestroyAllGizmosByOwner(this);
+	ActiveGizmos.Reset();
+}
+
+
+
+// does not make sense that CanBeginClickDragSequence() returns a RayHit? Needs to be an out-argument...
+FInputRayHit UEditPivotTool::CanBeginClickDragSequence(const FInputDeviceRay& PressPos)
+{
+	if (TransformProps->bEnableSnapDragging == false || ActiveGizmos.Num() == 0)
+	{
+		return FInputRayHit();
+	}
+
+	ActiveSnapDragIndex = -1;
+
+	float MinHitDistance = TNumericLimits<float>::Max();
+	FVector HitNormal;
+
+	for ( int k = 0; k < ComponentTargets.Num(); ++k )
+	{
+		TUniquePtr<FPrimitiveComponentTarget>& Target = ComponentTargets[k];
+
+		FHitResult WorldHit;
+		if (Target->HitTest(PressPos.WorldRay, WorldHit))
+		{
+			MinHitDistance = FMath::Min(MinHitDistance, WorldHit.Distance);
+			HitNormal = WorldHit.Normal;
+			ActiveSnapDragIndex = k;
+		}
+	}
+	return (MinHitDistance < TNumericLimits<float>::Max()) ? FInputRayHit(MinHitDistance, HitNormal) : FInputRayHit();
+}
+
+void UEditPivotTool::OnClickPress(const FInputDeviceRay& PressPos)
+{
+	FInputRayHit HitPos = CanBeginClickDragSequence(PressPos);
+	check(HitPos.bHit);
+
+	GetToolManager()->BeginUndoTransaction(LOCTEXT("TransformToolTransformTxnName", "SnapDrag"));
+
+	FEditPivotTarget& ActiveTarget = ActiveGizmos[0];
+	USceneComponent* GizmoComponent = ActiveTarget.TransformGizmo->GetGizmoActor()->GetRootComponent();
+	StartDragTransform = GizmoComponent->GetComponentToWorld();
+
+	if (TransformProps->SnapDragSource == EEditPivotSnapDragSource::ClickPoint)
+	{
+		StartDragFrameWorld = FFrame3d(PressPos.WorldRay.PointAt(HitPos.HitDepth), HitPos.HitNormal);
+	}
+	else
+	{
+		StartDragFrameWorld = FFrame3d(StartDragTransform);
+	}
+
+}
+
+
+void UEditPivotTool::OnClickDrag(const FInputDeviceRay& DragPos)
+{
+	FCollisionObjectQueryParams ObjectQueryParams(FCollisionObjectQueryParams::AllObjects);
+	FCollisionQueryParams CollisionParams = FCollisionQueryParams::DefaultQueryParam;
+
+	//bool bApplyToPivot = true;
+	//if (bApplyToPivot == false)
+	//{
+	//	int IgnoreIndex = -1;
+	//	for (int k = 0; k < ComponentTargets.Num(); ++k)
+	//	{
+	//		if (IgnoreIndex == -1 || k == IgnoreIndex)
+	//		{
+	//			CollisionParams.AddIgnoredComponent(ComponentTargets[k]->GetOwnerComponent());
+	//		}
+	//	}
+	//}
+
+
+	bool bRotate = (TransformProps->RotationMode != EEditPivotSnapDragRotationMode::Ignore);
+	float NormalSign = (TransformProps->RotationMode == EEditPivotSnapDragRotationMode::AlignFlipped) ? -1.0f : 1.0f;
+
+	FHitResult Result;
+	bool bWorldHit = TargetWorld->LineTraceSingleByObjectType(Result, DragPos.WorldRay.Origin, DragPos.WorldRay.PointAt(999999), ObjectQueryParams, CollisionParams);
+	if (bWorldHit == false)
+	{
+		return;
+	}
+
+
+	FVector HitPos = Result.ImpactPoint;
+	FVector TargetNormal = (-NormalSign) * Result.Normal;
+
+	FQuaterniond AlignRotation = (bRotate) ?
+		FQuaterniond(FVector3d::UnitZ(), TargetNormal) : FQuaterniond::Identity();
+
+	FTransform NewTransform = StartDragTransform;
+	NewTransform.SetRotation((FQuat)AlignRotation);
+	NewTransform.SetTranslation(HitPos);
+
+	FEditPivotTarget& ActiveTarget = ActiveGizmos[0];
+	USceneComponent* GizmoComponent = ActiveTarget.TransformGizmo->GetGizmoActor()->GetRootComponent();
+	GizmoComponent->SetWorldTransform(NewTransform);
+
+
+}
+
+
+void UEditPivotTool::OnClickRelease(const FInputDeviceRay& ReleasePos)
+{
+	OnTerminateDragSequence();
+}
+
+void UEditPivotTool::OnTerminateDragSequence()
+{
+	FEditPivotTarget& ActiveTarget = ActiveGizmos[0];
+	USceneComponent* GizmoComponent = ActiveTarget.TransformGizmo->GetGizmoActor()->GetRootComponent();
+	FTransform EndDragtransform = GizmoComponent->GetComponentToWorld();
+
+	TUniquePtr<FComponentWorldTransformChange> Change = MakeUnique<FComponentWorldTransformChange>(StartDragTransform, EndDragtransform);
+	GetToolManager()->EmitObjectChange(GizmoComponent, MoveTemp(Change),
+		LOCTEXT("TransformToolTransformTxnName", "SnapDrag"));
+
+	GetToolManager()->EndUndoTransaction();
+
+	ActiveSnapDragIndex = -1;
+}
+
+
+
+
+void UEditPivotTool::UpdateAssets(const FFrame3d& NewPivotWorldFrame)
+{
+	GetToolManager()->BeginUndoTransaction(LOCTEXT("EditPivotToolTransactionName", "Edit Pivot"));
+
+	FTransform NewWorldTransform = NewPivotWorldFrame.ToFTransform();
+
+	for (int32 ComponentIdx = 0; ComponentIdx < ComponentTargets.Num(); ComponentIdx++)
+	{
+		FTransform3d ComponentToWorld(ComponentTargets[ComponentIdx]->GetWorldTransform());
+
+		ComponentTargets[ComponentIdx]->CommitMesh([&](const FPrimitiveComponentTarget::FCommitParams& CommitParams)
+		{
+			// AAAAHHHH we can apply this xform directly to MeshDescription if we had the suitable function!
+
+			FMeshDescriptionToDynamicMesh ToDynamicMesh;
+			FDynamicMesh3 Mesh;
+			ToDynamicMesh.Convert(CommitParams.MeshDescription, Mesh);
+
+			MeshTransforms::ApplyTransform(Mesh, ComponentToWorld);
+			MeshTransforms::WorldToFrameCoords(Mesh, NewPivotWorldFrame);
+
+			FDynamicMeshToMeshDescription Converter;
+			Converter.Update(&Mesh, *CommitParams.MeshDescription, true, false);
+		});
+
+		UPrimitiveComponent* Component = ComponentTargets[ComponentIdx]->GetOwnerComponent();
+		Component->Modify();
+		Component->SetWorldTransform(NewWorldTransform);
+	}
+
+	GetToolManager()->EndUndoTransaction();
+}
+
+
+
+
+
+#undef LOCTEXT_NAMESPACE
diff --git a/Engine/Plugins/Experimental/MeshModelingToolset/Source/MeshModelingTools/Private/FindPolygonsAlgorithm.cpp b/Engine/Plugins/Experimental/MeshModelingToolset/Source/MeshModelingTools/Private/FindPolygonsAlgorithm.cpp
--- a/Engine/Plugins/Experimental/MeshModelingToolset/Source/MeshModelingTools/Private/FindPolygonsAlgorithm.cpp
+++ b/Engine/Plugins/Experimental/MeshModelingToolset/Source/MeshModelingTools/Private/FindPolygonsAlgorithm.cpp
@@ -2,6 +2,7 @@
 
 #include "FindPolygonsAlgorithm.h"
 #include "MeshNormals.h"
+#include "Selections/MeshConnectedComponents.h"
 
 
 
@@ -11,7 +12,32 @@ FFindPolygonsAlgorithm::FFindPolygonsAlgorithm(FDynamicMesh3* MeshIn)
 }
 
 
-bool FFindPolygonsAlgorithm::FindPolygons(double DotTolerance)
+
+bool FFindPolygonsAlgorithm::FindPolygonsFromUVIslands()
+{
+	PolygonGroupIDs.SetNum(Mesh->MaxTriangleID());
+	const FDynamicMeshUVOverlay* UV = Mesh->Attributes()->GetUVLayer(0);
+
+	FMeshConnectedComponents Components(Mesh);
+	Components.FindConnectedTriangles([&](int32 TriIdx0, int32 TriIdx1)
+	{
+		return UV->AreTrianglesConnected(TriIdx0, TriIdx1);
+	});
+
+	int32 NumComponents = Components.Num();
+	for (int32 ci = 0; ci < NumComponents; ++ci)
+	{
+		FoundPolygons.Add(Components.GetComponent(ci).Indices);
+	}
+
+	SetGroupsFromPolygons();
+
+	return (FoundPolygons.Num() > 0);
+}
+
+
+
+bool FFindPolygonsAlgorithm::FindPolygonsFromFaceNormals(double DotTolerance)
 {
 	DotTolerance = 1.0 - DotTolerance;
 
@@ -63,26 +89,49 @@ bool FFindPolygonsAlgorithm::FindPolygons(double DotTolerance)
 		FoundPolygons.Add(Polygon);
 	}
 
+	SetGroupsFromPolygons();
+
+	return (FoundPolygons.Num() > 0);
+}
+
+
+
+void FFindPolygonsAlgorithm::SetGroupsFromPolygons()
+{
 	Mesh->EnableTriangleGroups(0);
 
 	// set groups from polygons
 	int NumPolygons = FoundPolygons.Num();
 	PolygonTags.SetNum(NumPolygons);
 	PolygonNormals.SetNum(NumPolygons);
+	// can be parallel for
 	for (int PolyIdx = 0; PolyIdx < NumPolygons; PolyIdx++)
 	{
-		int Count = FoundPolygons[PolyIdx].Num();
-		for (int k = 0; k < Count; ++k)
+		const TArray<int>& Polygon = FoundPolygons[PolyIdx];
+		FVector3d AccumNormal(0, 0, 0);
+		int NumTriangles = Polygon.Num();
+		for (int k = 0; k < NumTriangles; ++k)
 		{
-			Mesh->SetTriangleGroup(FoundPolygons[PolyIdx][k], (PolyIdx + 1));
+			Mesh->SetTriangleGroup(Polygon[k], (PolyIdx + 1));
+			AccumNormal += Mesh->GetTriArea(k) * Mesh->GetTriNormal(Polygon[k]);
 		}
 		PolygonTags[PolyIdx] = (PolyIdx + 1);
-		PolygonNormals[PolyIdx] = Normals[FoundPolygons[PolyIdx][0]];
-	}
 
-	return (NumPolygons > 0);
-}
+		// find a normal if the average failed
+		AccumNormal.Normalize();
+		int pi = 0;
+		while (AccumNormal.Length() < 0.9 && pi < NumTriangles)
+		{
+			AccumNormal = Mesh->GetTriNormal(Polygon[pi++]);
+		}
+		if (AccumNormal.Length() < 0.9)
+		{
+			AccumNormal = FVector3d::UnitY();
+		}
 
+		PolygonNormals[PolyIdx] = AccumNormal;
+	}
+}
 
 
 bool FFindPolygonsAlgorithm::FindPolygonEdges()
diff --git a/Engine/Plugins/Experimental/MeshModelingToolset/Source/MeshModelingTools/Private/MeshInspectorTool.cpp b/Engine/Plugins/Experimental/MeshModelingToolset/Source/MeshModelingTools/Private/MeshInspectorTool.cpp
--- a/Engine/Plugins/Experimental/MeshModelingToolset/Source/MeshModelingTools/Private/MeshInspectorTool.cpp
+++ b/Engine/Plugins/Experimental/MeshModelingToolset/Source/MeshModelingTools/Private/MeshInspectorTool.cpp
@@ -90,21 +90,18 @@ void UMeshInspectorTool::Setup()
 {
 	UInteractiveTool::Setup();
 
-	// create dynamic mesh component to use for live preview
-	DynamicMeshComponent = NewObject<USimpleDynamicMeshComponent>(ComponentTarget->GetOwnerActor(), "DynamicMesh");
-	DynamicMeshComponent->SetupAttachment(ComponentTarget->GetOwnerActor()->GetRootComponent());
-	DynamicMeshComponent->RegisterComponent();
-	DynamicMeshComponent->SetWorldTransform(ComponentTarget->GetWorldTransform());
-
-	// copy material if there is one
-	DefaultMaterial = ComponentTarget->GetMaterial(0);
-	if (DefaultMaterial != nullptr)
-	{
-		DynamicMeshComponent->SetMaterial(0, DefaultMaterial);
-	}
+	PreviewMesh = NewObject<UPreviewMesh>(this);
+	PreviewMesh->bBuildSpatialDataStructure = false;
+	PreviewMesh->CreateInWorld(ComponentTarget->GetOwnerActor()->GetWorld(), FTransform::Identity);
+	PreviewMesh->SetTransform(ComponentTarget->GetWorldTransform());
+
+	FComponentMaterialSet MaterialSet;
+	ComponentTarget->GetMaterialSet(MaterialSet);
+	PreviewMesh->SetMaterials(MaterialSet.Materials);
+	DefaultMaterial = PreviewMesh->GetMaterial(0);
 
-	DynamicMeshComponent->TangentsType = EDynamicMeshTangentCalcType::ExternallyCalculated;
-	DynamicMeshComponent->InitializeMesh(ComponentTarget->GetMesh());
+	PreviewMesh->SetTangentsMode(EDynamicMeshTangentCalcType::ExternallyCalculated);
+	PreviewMesh->InitializeMesh(ComponentTarget->GetMesh());
 
 	Precompute();
 
@@ -120,28 +117,27 @@ void UMeshInspectorTool::Setup()
 	MaterialSettings->Setup();
 	AddToolPropertySource(MaterialSettings);
 
-
-	DynamicMeshComponent->bExplicitShowWireframe = Settings->bWireframe;
+	PreviewMesh->EnableWireframe(Settings->bWireframe);
 
 	UMeshStatisticsProperties* Statistics = NewObject<UMeshStatisticsProperties>(this);
-	Statistics->Update(*DynamicMeshComponent->GetMesh());
+	Statistics->Update(*PreviewMesh->GetPreviewDynamicMesh());
 	AddToolPropertySource(Statistics);
 
 	UMeshAnalysisProperties* MeshAnalysis = NewObject<UMeshAnalysisProperties>(this);
-	MeshAnalysis->Update(*DynamicMeshComponent->GetMesh(), ComponentTarget->GetWorldTransform());
+	MeshAnalysis->Update(*PreviewMesh->GetPreviewDynamicMesh(), ComponentTarget->GetWorldTransform());
 	AddToolPropertySource(MeshAnalysis);
 }
 
 
 void UMeshInspectorTool::Shutdown(EToolShutdownType ShutdownType)
 {
-	if (DynamicMeshComponent != nullptr)
-	{
-		ComponentTarget->SetOwnerVisibility(true);
+	ComponentTarget->SetOwnerVisibility(true);
 
-		DynamicMeshComponent->UnregisterComponent();
-		DynamicMeshComponent->DestroyComponent();
-		DynamicMeshComponent = nullptr;
+	if (PreviewMesh != nullptr)
+	{
+		PreviewMesh->SetVisible(false);
+		PreviewMesh->Disconnect();
+		PreviewMesh = nullptr;
 	}
 
 	Settings->SaveProperties(this);
@@ -154,10 +150,10 @@ void UMeshInspectorTool::Precompute()
 	UVSeamEdges.Reset();
 	NormalSeamEdges.Reset();
 
-	FDynamicMesh3* TargetMesh = DynamicMeshComponent->GetMesh();
-	FDynamicMeshUVOverlay* UVOverlay =
+	const FDynamicMesh3* TargetMesh = PreviewMesh->GetPreviewDynamicMesh();
+	const FDynamicMeshUVOverlay* UVOverlay =
 		TargetMesh->HasAttributes() ? TargetMesh->Attributes()->PrimaryUV() : nullptr;
-	FDynamicMeshNormalOverlay* NormalOverlay =
+	const FDynamicMeshNormalOverlay* NormalOverlay =
 		TargetMesh->HasAttributes() ? TargetMesh->Attributes()->PrimaryNormals() : nullptr;
 
 	for (int eid : TargetMesh->EdgeIndicesItr())
@@ -185,7 +181,7 @@ void UMeshInspectorTool::Precompute()
 void UMeshInspectorTool::Render(IToolsContextRenderAPI* RenderAPI)
 {
 	FPrimitiveDrawInterface* PDI = RenderAPI->GetPrimitiveDrawInterface();
-	FTransform Transform = ComponentTarget->GetWorldTransform();
+	FTransform3d Transform(ComponentTarget->GetWorldTransform());
 
 	FColor BoundaryEdgeColor(240, 15, 15);
 	float BoundaryEdgeThickness = LineWidthMultiplier * 4.0;
@@ -201,7 +197,7 @@ void UMeshInspectorTool::Render(IToolsContextRenderAPI* RenderAPI)
 	FColor BinormalColor(15, 240, 15);
 	float TangentThickness = LineWidthMultiplier * 2.0f;
 
-	FDynamicMesh3* TargetMesh = DynamicMeshComponent->GetMesh();
+	const FDynamicMesh3* TargetMesh = PreviewMesh->GetPreviewDynamicMesh();
 	FVector3d A, B;
 
 	if (Settings->bBoundaryEdges)
@@ -209,7 +205,7 @@ void UMeshInspectorTool::Render(IToolsContextRenderAPI* RenderAPI)
 		for (int eid : BoundaryEdges)
 		{
 			TargetMesh->GetEdgeV(eid, A, B);
-			PDI->DrawLine(Transform.TransformPosition(A), Transform.TransformPosition(B),
+			PDI->DrawLine((FVector)Transform.TransformPosition(A), (FVector)Transform.TransformPosition(B),
 				BoundaryEdgeColor, 0, BoundaryEdgeThickness, 2.0f, true);
 		}
 	}
@@ -219,7 +215,7 @@ void UMeshInspectorTool::Render(IToolsContextRenderAPI* RenderAPI)
 		for (int eid : UVSeamEdges)
 		{
 			TargetMesh->GetEdgeV(eid, A, B);
-			PDI->DrawLine(Transform.TransformPosition(A), Transform.TransformPosition(B),
+			PDI->DrawLine((FVector)Transform.TransformPosition(A), (FVector)Transform.TransformPosition(B),
 				UVSeamColor, 0, UVSeamThickness, 3.0f, true);
 		}
 	}
@@ -229,7 +225,7 @@ void UMeshInspectorTool::Render(IToolsContextRenderAPI* RenderAPI)
 		for (int eid : NormalSeamEdges)
 		{
 			TargetMesh->GetEdgeV(eid, A, B);
-			PDI->DrawLine(Transform.TransformPosition(A), Transform.TransformPosition(B),
+			PDI->DrawLine((FVector)Transform.TransformPosition(A), (FVector)Transform.TransformPosition(B),
 				NormalSeamColor, 0, NormalSeamThickness, 3.0f, true);
 		}
 	}
@@ -239,7 +235,7 @@ void UMeshInspectorTool::Render(IToolsContextRenderAPI* RenderAPI)
 		for (int eid : GroupBoundaryEdges)
 		{
 			TargetMesh->GetEdgeV(eid, A, B);
-			PDI->DrawLine(Transform.TransformPosition(A), Transform.TransformPosition(B),
+			PDI->DrawLine((FVector)Transform.TransformPosition(A), (FVector)Transform.TransformPosition(B),
 				PolygonBorderColor, 0, PolygonBorderThickness, 2.0f, true);
 		}
 	}
@@ -256,27 +252,27 @@ void UMeshInspectorTool::Render(IToolsContextRenderAPI* RenderAPI)
 			for (int j = 0; j < 3; ++j)
 			{
 				TriV[j] = Transform.TransformPosition(TriV[j]);
-				PDI->DrawLine(TriV[j], TriV[j] + Settings->NormalLength * Transform.TransformVectorNoScale(TriN[j]),
+				PDI->DrawLine((FVector)TriV[j], (FVector)((FVector3f)TriV[j] + Settings->NormalLength * Transform.TransformVectorNoScale(TriN[j])),
 					NormalColor, SDPG_World, NormalThickness, 0.0f, true);
 			}
 		}
 	}
 
-	if (Settings->bTangentVectors)
+	if (Settings->bTangentVectors && PreviewMesh->GetTangents() != nullptr)
 	{
-		const FMeshTangentsf* Tangents = DynamicMeshComponent->GetTangents();
+		const FMeshTangentsf* Tangents = PreviewMesh->GetTangents();
 		for (int TID : TargetMesh->TriangleIndicesItr())
 		{
 			FVector3d TriV[3];
 			TargetMesh->GetTriVertices(TID, TriV[0], TriV[1], TriV[2]);
 			for (int SubIdx = 0; SubIdx < 3; SubIdx++)
 			{
-				FVector Vert = Transform.TransformPosition(TriV[SubIdx]);
+				FVector3f Vert(Transform.TransformPosition(TriV[SubIdx]));
 				FVector3f Tangent, Bitangent;
 				Tangents->GetPerTriangleTangent(TID, SubIdx, Tangent, Bitangent);
-				PDI->DrawLine(Vert, Vert + Settings->TangentLength * Transform.TransformVectorNoScale(Tangent),
+				PDI->DrawLine((FVector)Vert, (FVector)(Vert + Settings->TangentLength * Transform.TransformVectorNoScale(Tangent)),
 					TangentColor, SDPG_World, TangentThickness, 3.5f, true);
-				PDI->DrawLine(Vert, Vert + Settings->TangentLength * Transform.TransformVectorNoScale(Bitangent),
+				PDI->DrawLine((FVector)Vert, (FVector)(Vert + Settings->TangentLength * Transform.TransformVectorNoScale(Bitangent)),
 					BinormalColor, SDPG_World, TangentThickness, 3.5f, true);
 			}
 		}
@@ -289,13 +285,18 @@ void UMeshInspectorTool::Render(IToolsContextRenderAPI* RenderAPI)
 void UMeshInspectorTool::OnPropertyModified(UObject* PropertySet, UProperty* Property)
 {
 	GetToolManager()->PostInvalidation();
-	DynamicMeshComponent->bExplicitShowWireframe = Settings->bWireframe;
+	PreviewMesh->EnableWireframe(Settings->bWireframe);
+	
 	MaterialSettings->UpdateMaterials();
-	MaterialSettings->SetMaterialIfChanged(DefaultMaterial, DynamicMeshComponent->GetMaterial(0),
-		[this](UMaterialInterface* Material)
+	UMaterialInterface* OverrideMaterial = MaterialSettings->GetActiveOverrideMaterial();
+	if (OverrideMaterial == nullptr)
+	{
+		PreviewMesh->ClearOverrideRenderMaterial();
+	}
+	else
 	{
-		DynamicMeshComponent->SetMaterial(0, Material);
-	});
+		PreviewMesh->SetOverrideRenderMaterial(OverrideMaterial);
+	}
 }
 
 
diff --git a/Engine/Plugins/Experimental/MeshModelingToolset/Source/MeshModelingTools/Private/MeshSelectionTool.cpp b/Engine/Plugins/Experimental/MeshModelingToolset/Source/MeshModelingTools/Private/MeshSelectionTool.cpp
--- a/Engine/Plugins/Experimental/MeshModelingToolset/Source/MeshModelingTools/Private/MeshSelectionTool.cpp
+++ b/Engine/Plugins/Experimental/MeshModelingToolset/Source/MeshModelingTools/Private/MeshSelectionTool.cpp
@@ -8,9 +8,13 @@
 #include "DynamicMeshChangeTracker.h"
 #include "Changes/ToolCommandChangeSequence.h"
 #include "Changes/MeshChange.h"
+#include "Util/ColorConstants.h"
+#include "Selections/MeshConnectedComponents.h"
+#include "MeshRegionBoundaryLoops.h"
 #include "MeshIndexUtil.h"
 #include "AssetGenerationUtil.h"
 #include "ToolSetupUtil.h"
+#include "Selections/MeshConnectedComponents.h"
 
 #define LOCTEXT_NAMESPACE "UMeshSelectionTool"
 
@@ -46,6 +50,7 @@ void UMeshSelectionToolProperties::SaveProperties(UInteractiveTool* SaveFromTool
 	PropertyCache->bVolumetricBrush = this->bVolumetricBrush;
 	PropertyCache->bHitBackFaces = this->bHitBackFaces;
 	PropertyCache->bShowWireframe = this->bShowWireframe;
+	PropertyCache->FaceColorMode = this->FaceColorMode;
 }
 
 void UMeshSelectionToolProperties::RestoreProperties(UInteractiveTool* RestoreToTool)
@@ -56,6 +61,7 @@ void UMeshSelectionToolProperties::RestoreProperties(UInteractiveTool* RestoreTo
 	this->bVolumetricBrush = PropertyCache->bVolumetricBrush;
 	this->bHitBackFaces = PropertyCache->bHitBackFaces;
 	this->bShowWireframe = PropertyCache->bShowWireframe;
+	this->FaceColorMode = PropertyCache->FaceColorMode;
 }
 
 
@@ -84,33 +90,49 @@ void UMeshSelectionTool::Setup()
 {
 	UDynamicMeshBrushTool::Setup();
 
+	// hide strength and falloff
+	BrushProperties->bShowFullSettings = false;
+
 	SelectionProps = NewObject<UMeshSelectionToolProperties>(this);
 	SelectionProps->RestoreProperties(this);
 	AddToolPropertySource(SelectionProps);
 
+	AddSubclassPropertySets();
+
 	SelectionActions = NewObject<UMeshSelectionEditActions>(this);
 	SelectionActions->Initialize(this);
 	AddToolPropertySource(SelectionActions);
 
-	EditActions = NewObject<UMeshSelectionMeshEditActions>(this);
-	EditActions->Initialize(this);
+	EditActions = CreateEditActions();
 	AddToolPropertySource(EditActions);
 
 	// enable wireframe on component
 	PreviewMesh->EnableWireframe(true);
 
-	// set vertex color material on base component so we can see selection
-	UMaterialInterface* SelectionMat = ToolSetupUtil::GetSelectionMaterial(GetToolManager());
-	if (SelectionMat != nullptr)
+	// disable shadows
+	PreviewMesh->GetRootComponent()->bCastDynamicShadow = false;
+
+	// configure secondary render material
+	UMaterialInterface* SelectionMaterial = ToolSetupUtil::GetSelectionMaterial(FLinearColor(0.9f, 0.1f, 0.1f), GetToolManager());
+	if (SelectionMaterial != nullptr)
 	{
-		PreviewMesh->SetMaterial(SelectionMat);
+		PreviewMesh->SetSecondaryRenderMaterial(SelectionMaterial);
 	}
-	PreviewMesh->GetRootComponent()->bCastDynamicShadow = false;
+
+	// enable secondary triangle buffers
+	PreviewMesh->EnableSecondaryTriangleBuffers(
+		[this](const FDynamicMesh3* Mesh, int32 TriangleID)
+	{
+		return SelectedTriangles[TriangleID] ? true : false;
+	});
 
 	const FDynamicMesh3* Mesh = PreviewMesh->GetPreviewDynamicMesh();
 	SelectedVertices = TBitArray<>(false, Mesh->MaxVertexID());
 	SelectedTriangles = TBitArray<>(false, Mesh->MaxTriangleID());
 
+	// we could probably calculate this on-demand but we need to do it before making any mesh changes? or update?
+	CacheUVIslandIDs();
+
 	this->Selection = NewObject<UMeshSelectionSet>(this);
 	Selection->GetOnModified().AddLambda([this](USelectionSet* SelectionObj)
 	{
@@ -118,15 +140,29 @@ void UMeshSelectionTool::Setup()
 	});
 
 	// rebuild octree if mesh changes
-	PreviewMesh->GetOnMeshChanged().AddLambda([this]() { bOctreeValid = false; });
+	PreviewMesh->GetOnMeshChanged().AddLambda([this]() { bOctreeValid = false; bFullMeshInvalidationPending = true; });
 
 	ShowWireframeWatcher.Initialize(
 		[this]() { return SelectionProps->bShowWireframe; },
-		[this](bool bNewValue) { PreviewMesh->EnableWireframe(bNewValue); }, false);
+		[this](bool bNewValue) { PreviewMesh->EnableWireframe(bNewValue); }, SelectionProps->bShowWireframe);
+
+	ColorModeWatcher.Initialize(
+		[this]() { return SelectionProps->FaceColorMode; },
+		[this](EMeshFacesColorMode NewValue) { bColorsUpdatePending = true; UpdateVisualization(false); }, SelectionProps->FaceColorMode );
+	bColorsUpdatePending = (SelectionProps->FaceColorMode != EMeshFacesColorMode::None);
+	
 }
 
 
 
+UMeshSelectionToolActionPropertySet* UMeshSelectionTool::CreateEditActions()
+{
+	UMeshSelectionMeshEditActions* Actions = NewObject<UMeshSelectionMeshEditActions>(this);
+	Actions->Initialize(this);
+	return Actions;
+}
+
+
 
 void UMeshSelectionTool::OnShutdown(EToolShutdownType ShutdownType)
 {
@@ -137,9 +173,9 @@ void UMeshSelectionTool::OnShutdown(EToolShutdownType ShutdownType)
 		// this block bakes the modified DynamicMeshComponent back into the StaticMeshComponent inside an undo transaction
 		GetToolManager()->BeginUndoTransaction(LOCTEXT("MeshSelectionToolTransactionName", "Edit Mesh"));
 
-		ComponentTarget->CommitMesh([=](FMeshDescription* MeshDescription)
+		ComponentTarget->CommitMesh([=](const FPrimitiveComponentTarget::FCommitParams& CommitParams)
 		{
-			PreviewMesh->Bake(MeshDescription, true);
+			PreviewMesh->Bake(CommitParams.MeshDescription, true);
 		});
 		GetToolManager()->EndUndoTransaction();
 	}
@@ -282,41 +318,10 @@ void UMeshSelectionTool::CalculateVertexROI(const FBrushStampData& Stamp, TArray
 
 
 
-static void GrowToConnectedTriangles(const TArray<int>& TriangleROI, TArray<int>& GrowROI, const FDynamicMesh3* Mesh,
-	TArray<int32>& QueueBuffer, TSet<int32>& DoneBuffer,
-	TFunctionRef<bool(int32, int32)> CanGrowPredicate = [](int32, int32) { return true; }
-	)
-{
-	QueueBuffer.Reset(); QueueBuffer.Insert(TriangleROI, 0);
-	DoneBuffer.Reset(); DoneBuffer.Append(TriangleROI);
-	//TArray<int32> Queue(TriangleROI);
-	//TSet<int32> Done(TriangleROI);
-
-	while (QueueBuffer.Num() > 0)
-	{
-		int32 CurTri = QueueBuffer.Pop(false);
-		GrowROI.Add(CurTri);
-
-		FIndex3i NbrTris = Mesh->GetTriNeighbourTris(CurTri);
-		for (int j = 0; j < 3; ++j)
-		{
-			int32 tid = NbrTris[j];
-			if (tid != FDynamicMesh3::InvalidID && DoneBuffer.Contains(tid) == false && CanGrowPredicate(CurTri, tid))
-			{
-				QueueBuffer.Add(tid);
-				DoneBuffer.Add(tid);
-			}
-		}
-	}
-
-}
-
-
-
 void UMeshSelectionTool::CalculateTriangleROI(const FBrushStampData& Stamp, TArray<int>& TriangleROI)
 {
-	FTransform Transform = ComponentTarget->GetWorldTransform();
-	FVector StampPosLocal = Transform.InverseTransformPosition(Stamp.WorldPosition);
+	FTransform3d Transform(ComponentTarget->GetWorldTransform());
+	FVector3d StampPosLocal = Transform.InverseTransformPosition(Stamp.WorldPosition);
 
 	// always select first triangle
 	const FDynamicMesh3* Mesh = PreviewMesh->GetPreviewDynamicMesh();
@@ -346,7 +351,7 @@ void UMeshSelectionTool::CalculateTriangleROI(const FBrushStampData& Stamp, TArr
 	{
 		TArray<int32> StartROI;
 		StartROI.Add(Stamp.HitResult.FaceIndex);
-		GrowToConnectedTriangles(StartROI, TriangleROI, Mesh, TemporaryBuffer, TemporarySet,
+		FMeshConnectedComponents::GrowToConnectedTriangles(Mesh, StartROI, TriangleROI,  &TemporaryBuffer, &TemporarySet,
 			[Mesh, RadiusSqr, StampPosLocal](int t1, int t2) { return (Mesh->GetTriCentroid(t2) - StampPosLocal).SquaredLength() < RadiusSqr; });
 
 	}
@@ -416,22 +421,40 @@ void UMeshSelectionTool::UpdateFaceSelection(const FBrushStampData& Stamp, const
 	TArray<int> LocalROI;
 	if (SelectionProps->SelectionMode == EMeshSelectionToolPrimaryMode::AllConnected)
 	{
-		GrowToConnectedTriangles(TriangleROI, LocalROI, Mesh, TemporaryBuffer, TemporarySet);
+		FMeshConnectedComponents::GrowToConnectedTriangles(Mesh, TriangleROI, LocalROI, &TemporaryBuffer, &TemporarySet);
 		UseROI = &LocalROI;
 	}
 	else if (SelectionProps->SelectionMode == EMeshSelectionToolPrimaryMode::AllInGroup)
 	{
-		GrowToConnectedTriangles(TriangleROI, LocalROI, Mesh, TemporaryBuffer, TemporarySet,
+		FMeshConnectedComponents::GrowToConnectedTriangles(Mesh, TriangleROI, LocalROI, &TemporaryBuffer, &TemporarySet,
 			[Mesh](int t1, int t2) { return Mesh->GetTriangleGroup(t1) == Mesh->GetTriangleGroup(t2); } );
 		UseROI = &LocalROI;
 	}
+	else if (SelectionProps->SelectionMode == EMeshSelectionToolPrimaryMode::ByMaterial)
+	{
+		const FDynamicMeshMaterialAttribute* MaterialIDs = Mesh->Attributes()->GetMaterialID();
+		TArray<int32> StartROI;
+		StartROI.Add(Stamp.HitResult.FaceIndex);
+		FMeshConnectedComponents::GrowToConnectedTriangles(Mesh, StartROI, LocalROI, &TemporaryBuffer, &TemporarySet,
+			[Mesh, MaterialIDs](int t1, int t2) { return MaterialIDs->GetValue(t1) == MaterialIDs->GetValue(t2); });
+		UseROI = &LocalROI;
+	}
+	else if (SelectionProps->SelectionMode == EMeshSelectionToolPrimaryMode::ByUVIsland)
+	{
+		const FDynamicMeshMaterialAttribute* MaterialIDs = Mesh->Attributes()->GetMaterialID();
+		TArray<int32> StartROI;
+		StartROI.Add(Stamp.HitResult.FaceIndex);
+		FMeshConnectedComponents::GrowToConnectedTriangles(Mesh, StartROI, LocalROI, &TemporaryBuffer, &TemporarySet,
+			[&](int t1, int t2) { return TriangleToUVIsland[t1] == TriangleToUVIsland[t2]; });
+		UseROI = &LocalROI;
+	}
 	else if (SelectionProps->SelectionMode == EMeshSelectionToolPrimaryMode::AllWithinAngle)
 	{
 		TArray<int32> StartROI; 
 		StartROI.Add(Stamp.HitResult.FaceIndex);
 		FVector3d StartNormal = Mesh->GetTriNormal(StartROI[0]);
 		int AngleTol = SelectionProps->AngleTolerance;
-		GrowToConnectedTriangles(StartROI, LocalROI, Mesh, TemporaryBuffer, TemporarySet,
+		FMeshConnectedComponents::GrowToConnectedTriangles(Mesh, StartROI, LocalROI, &TemporaryBuffer, &TemporarySet,
 			[Mesh, AngleTol, StartNormal](int t1, int t2) { return Mesh->GetTriNormal(t2).AngleD(StartNormal) < AngleTol; });
 		UseROI = &LocalROI;
 	}
@@ -442,7 +465,7 @@ void UMeshSelectionTool::UpdateFaceSelection(const FBrushStampData& Stamp, const
 		StartROI.Add(Stamp.HitResult.FaceIndex);
 		FVector3d StartNormal = Mesh->GetTriNormal(StartROI[0]);
 		int AngleTol = SelectionProps->AngleTolerance;
-		GrowToConnectedTriangles(StartROI, LocalROI, Mesh, TemporaryBuffer, TemporarySet,
+		FMeshConnectedComponents::GrowToConnectedTriangles(Mesh, StartROI, LocalROI, &TemporaryBuffer, &TemporarySet,
 			[Mesh, AngleTol, StartNormal, &BrushROI](int t1, int t2) { return BrushROI.Contains(t2) && Mesh->GetTriNormal(t2).AngleD(StartNormal) < AngleTol; });
 		UseROI = &LocalROI;
 	}
@@ -525,37 +548,86 @@ bool UMeshSelectionTool::OnUpdateHover(const FInputDeviceRay& DevicePos)
 
 void UMeshSelectionTool::OnSelectionUpdated()
 {
-	UpdateVisualization();
+	UpdateVisualization(true);
 }
 
-void UMeshSelectionTool::UpdateVisualization()
+void UMeshSelectionTool::UpdateVisualization(bool bSelectionModified)
 {
-	if (SelectionType == EMeshSelectionElementType::Face)
+	check(SelectionType == EMeshSelectionElementType::Face);  // only face selection supported so far
+
+	bFullMeshInvalidationPending = false;
+
+	// force an update of renderbuffers
+	if (bSelectionModified)
+	{
+		PreviewMesh->NotifyDeferredEditCompleted(UPreviewMesh::ERenderUpdateMode::FullUpdate, true);
+	}
+
+	if (bColorsUpdatePending)
 	{
-		PreviewMesh->SetTriangleColorFunction([this](int TriangleID)
+		if (SelectionProps->FaceColorMode != EMeshFacesColorMode::None)
+		{
+			PreviewMesh->SetOverrideRenderMaterial(ToolSetupUtil::GetSelectionMaterial(GetToolManager()));
+			PreviewMesh->SetTriangleColorFunction([this](const FDynamicMesh3* Mesh, int TriangleID)
+			{
+				return GetCurrentFaceColor(Mesh, TriangleID);
+			}, 
+			UPreviewMesh::ERenderUpdateMode::FastUpdate);
+		}
+		else
 		{
-			return SelectedTriangles[TriangleID] ? FColor::Red : FColor::White;
-		}, UPreviewMesh::ERenderUpdateMode::FullUpdate);
-
-		//bool bIsFirst = (! BaseMeshComponent->TriangleColorFunc);
-		//BaseMeshComponent->TriangleColorFunc =
-		//if (bIsFirst)
-		//{
-		//	BaseMeshComponent->NotifyMeshUpdated();
-		//}
-		//else
-		//{
-		//	BaseMeshComponent->NotifyMeshUpdated();
-		//	//BaseMeshComponent->FastNotifyColorsUpdated();
-		//}
+			PreviewMesh->ClearOverrideRenderMaterial();
+			PreviewMesh->ClearTriangleColorFunction(UPreviewMesh::ERenderUpdateMode::FastUpdate);
+		}
+
+		bColorsUpdatePending = false;
 	}
-	else
+}
+
+
+
+FColor UMeshSelectionTool::GetCurrentFaceColor(const FDynamicMesh3* Mesh, int TriangleID)
+{
+	if (SelectionProps->FaceColorMode == EMeshFacesColorMode::ByGroup)
+	{
+		return LinearColors::SelectFColor(Mesh->GetTriangleGroup(TriangleID));
+	}
+	else if (SelectionProps->FaceColorMode == EMeshFacesColorMode::ByMaterialID)
 	{
-		PreviewMesh->ClearTriangleColorFunction(UPreviewMesh::ERenderUpdateMode::FullUpdate);
+		return LinearColors::SelectFColor( Mesh->Attributes()->GetMaterialID()->GetValue(TriangleID) );
 	}
+	else if (SelectionProps->FaceColorMode == EMeshFacesColorMode::ByUVIsland)
+	{
+		return LinearColors::SelectFColor(TriangleToUVIsland[TriangleID]);
+	}
+	return FColor::Red;
+
 }
 
 
+void UMeshSelectionTool::CacheUVIslandIDs()
+{
+	const FDynamicMesh3* Mesh = PreviewMesh->GetMesh();
+	FMeshConnectedComponents Components(Mesh);
+
+	TriangleToUVIsland.SetNum(Mesh->MaxTriangleID());
+
+	const FDynamicMeshUVOverlay* UV = Mesh->Attributes()->GetUVLayer(0);
+
+	Components.FindConnectedTriangles([&](int32 TriIdx0, int32 TriIdx1)
+	{
+		return UV->AreTrianglesConnected(TriIdx0, TriIdx1);
+	});
+
+	int32 NumComponents = Components.Num();
+	for (int32 ci = 0; ci < NumComponents; ++ci)
+	{
+		for (int32 TriIdx : Components.GetComponent(ci).Indices)
+		{
+			TriangleToUVIsland[TriIdx] = ci;
+		}
+	}
+}
 
 
 void UMeshSelectionTool::Render(IToolsContextRenderAPI* RenderAPI)
@@ -563,7 +635,7 @@ void UMeshSelectionTool::Render(IToolsContextRenderAPI* RenderAPI)
 	UDynamicMeshBrushTool::Render(RenderAPI);
 
 	FTransform WorldTransform = ComponentTarget->GetWorldTransform();
-	const FDynamicMesh3* Mesh = PreviewMesh->GetPreviewDynamicMesh();
+	const FDynamicMesh3* Mesh = PreviewMesh->GetMesh();
 
 	if (SelectionType == EMeshSelectionElementType::Vertex)
 	{
@@ -584,10 +656,12 @@ void UMeshSelectionTool::Render(IToolsContextRenderAPI* RenderAPI)
 
 
 void UMeshSelectionTool::Tick(float DeltaTime)
+
 {
 	UDynamicMeshBrushTool::Tick(DeltaTime);
 
 	ShowWireframeWatcher.CheckAndUpdate();
+	ColorModeWatcher.CheckAndUpdate();
 
 	if (bStampPending)
 	{
@@ -681,9 +755,21 @@ void UMeshSelectionTool::ApplyAction(EMeshSelectionToolActions ActionType)
 			DeleteSelectedTriangles();
 			break;
 
+		case EMeshSelectionToolActions::DisconnectSelected:
+			DisconnectSelectedTriangles();
+			break;
+
 		case EMeshSelectionToolActions::SeparateSelected:
 			SeparateSelectedTriangles();
 			break;
+
+		case EMeshSelectionToolActions::FlipSelected:
+			FlipSelectedTriangles();
+			break;
+
+		case EMeshSelectionToolActions::CreateGroup:
+			AssignNewGroupToSelectedTriangles();
+			break;
 	}
 }
 
@@ -900,6 +986,7 @@ void UMeshSelectionTool::DeleteSelectedTriangles()
 	// emit combined change sequence
 	GetToolManager()->EmitObjectChange(this, MoveTemp(ChangeSeq), LOCTEXT("MeshSelectionToolDeleteFaces", "Delete Faces"));
 
+	bFullMeshInvalidationPending = true;
 	OnExternalSelectionChange();
 	bHaveModifiedMesh = true;
 	bOctreeValid = false;
@@ -934,6 +1021,50 @@ void AssignMaterial(AActor* ToActor, const TUniquePtr<FPrimitiveComponentTarget>
 }
 
 
+void UMeshSelectionTool::DisconnectSelectedTriangles()
+{
+	check(SelectionType == EMeshSelectionElementType::Face);
+	TArray<int32> SelectedFaces = Selection->GetElements(EMeshSelectionElementType::Face);
+	if (SelectedFaces.Num() == 0)
+	{
+		return;
+	}
+
+	TUniquePtr<FToolCommandChangeSequence> ChangeSeq = MakeUnique<FToolCommandChangeSequence>();
+
+	// split out selected triangles and emit triangle change
+	TUniquePtr<FMeshChange> MeshChange = PreviewMesh->TrackedEditMesh(
+		[&SelectedFaces](FDynamicMesh3& Mesh, FDynamicMeshChangeTracker& ChangeTracker)
+	{
+		// save vertices and triangles that are on the boundary of the selection
+		FMeshRegionBoundaryLoops BoundaryLoops(&Mesh, SelectedFaces);
+		for (const FEdgeLoop& Loop : BoundaryLoops.Loops)
+		{
+			for (int VID : Loop.Vertices)
+			{
+				ChangeTracker.SaveVertex(VID);
+				// include the whole one-ring in case the disconnect creates bowties that need to be split
+				for (int TID : Mesh.VtxTrianglesItr(VID))
+				{
+					ChangeTracker.SaveTriangle(TID, true);
+				}
+			}
+		}
+
+		FDynamicMeshEditor Editor(&Mesh);
+		Editor.DisconnectTriangles(SelectedFaces);
+	});
+	ChangeSeq->AppendChange(PreviewMesh, MoveTemp(MeshChange));
+
+	// emit combined change sequence
+	GetToolManager()->EmitObjectChange(this, MoveTemp(ChangeSeq), LOCTEXT("MeshSelectionToolDisconnectFaces", "Disconnect Faces"));
+
+	bFullMeshInvalidationPending = true;
+	bHaveModifiedMesh = true;
+}
+
+
+
 void UMeshSelectionTool::SeparateSelectedTriangles()
 {
 #if WITH_EDITOR
@@ -956,6 +1087,7 @@ void UMeshSelectionTool::SeparateSelectedTriangles()
 	// extract copy of triangles
 	FDynamicMesh3 SeparatedMesh;
 	SeparatedMesh.EnableAttributes();
+	SeparatedMesh.Attributes()->EnableMatchingAttributes(*SourceMesh->Attributes());
 	FDynamicMeshEditor Editor(&SeparatedMesh);
 	FMeshIndexMappings Mappings; FDynamicMeshEditResult EditResult;
 	Editor.AppendTriangles(SourceMesh, SelectedFaces, Mappings, EditResult);
@@ -978,5 +1110,98 @@ void UMeshSelectionTool::SeparateSelectedTriangles()
 
 
 
+void UMeshSelectionTool::FlipSelectedTriangles()
+{
+	check(SelectionType == EMeshSelectionElementType::Face);
+	TArray<int32> SelectedFaces = Selection->GetElements(EMeshSelectionElementType::Face);
+	if (SelectedFaces.Num() == 0)
+	{
+		return;
+	}
+
+	TUniquePtr<FToolCommandChangeSequence> ChangeSeq = MakeUnique<FToolCommandChangeSequence>();
+
+	// clear current selection
+	BeginChange(false);
+	for (int tid : SelectedFaces)
+	{
+		ActiveSelectionChange->Add(tid);
+	}
+	Selection->RemoveIndices(EMeshSelectionElementType::Face, SelectedFaces);
+	TUniquePtr<FMeshSelectionChange> SelectionChange = EndChange();
+	ChangeSeq->AppendChange(Selection, MoveTemp(SelectionChange));
+
+	// flip normals
+	TUniquePtr<FMeshChange> MeshChange = PreviewMesh->TrackedEditMesh(
+		[&SelectedFaces](FDynamicMesh3& Mesh, FDynamicMeshChangeTracker& ChangeTracker)
+	{
+		for (int TID : SelectedFaces)
+		{
+			ChangeTracker.SaveTriangle(TID, true);
+		}
+		FDynamicMeshEditor Editor(&Mesh);
+		Editor.ReverseTriangleOrientations(SelectedFaces, true);
+	});
+	ChangeSeq->AppendChange(PreviewMesh, MoveTemp(MeshChange));
+
+	// emit combined change sequence
+	GetToolManager()->EmitObjectChange(this, MoveTemp(ChangeSeq), LOCTEXT("MeshSelectionToolFlipFaces", "Flip Face Orientations"));
+
+	bHaveModifiedMesh = true;
+}
+
+
+void UMeshSelectionTool::AssignNewGroupToSelectedTriangles()
+{
+	check(SelectionType == EMeshSelectionElementType::Face);
+	TArray<int32> SelectedFaces = Selection->GetElements(EMeshSelectionElementType::Face);
+	if (SelectedFaces.Num() == 0)
+	{
+		return;
+	}
+
+	TUniquePtr<FToolCommandChangeSequence> ChangeSeq = MakeUnique<FToolCommandChangeSequence>();
+
+	// clear current selection
+	BeginChange(false);
+	for (int tid : SelectedFaces)
+	{
+		ActiveSelectionChange->Add(tid);
+	}
+	Selection->RemoveIndices(EMeshSelectionElementType::Face, SelectedFaces);
+	TUniquePtr<FMeshSelectionChange> SelectionChange = EndChange();
+	ChangeSeq->AppendChange(Selection, MoveTemp(SelectionChange));
+
+	// assign new groups to triangles
+	// note: using an FMeshChange is kind of overkill here
+	TUniquePtr<FMeshChange> MeshChange = PreviewMesh->TrackedEditMesh(
+		[&SelectedFaces](FDynamicMesh3& Mesh, FDynamicMeshChangeTracker& ChangeTracker)
+	{
+		// each component gets its own group id
+		FMeshConnectedComponents Components(&Mesh);
+		Components.FindConnectedTriangles(SelectedFaces);
+
+		for (FMeshConnectedComponents::FComponent& Component : Components)
+		{
+			int NewGroupID = Mesh.AllocateTriangleGroup();
+			for (int tid : Component.Indices)
+			{
+				ChangeTracker.SaveTriangle(tid, true);
+				Mesh.SetTriangleGroup(tid, NewGroupID);
+			}
+		}
+	});
+	ChangeSeq->AppendChange(PreviewMesh, MoveTemp(MeshChange));
+
+	// emit combined change sequence
+	GetToolManager()->EmitObjectChange(this, MoveTemp(ChangeSeq), LOCTEXT("MeshSelectionToolCreateGroup", "Create Polygroup"));
+
+	OnExternalSelectionChange();
+	bHaveModifiedMesh = true;
+}
+
+
+
+
 
-#undef LOCTEXT_NAMESPACE
\ No newline at end of file
+#undef LOCTEXT_NAMESPACE
diff --git a/Engine/Plugins/Experimental/MeshModelingToolset/Source/MeshModelingTools/Private/PositionPlaneGizmo.cpp b/Engine/Plugins/Experimental/MeshModelingToolset/Source/MeshModelingTools/Private/PositionPlaneGizmo.cpp
--- a/Engine/Plugins/Experimental/MeshModelingToolset/Source/MeshModelingTools/Private/PositionPlaneGizmo.cpp
+++ b/Engine/Plugins/Experimental/MeshModelingToolset/Source/MeshModelingTools/Private/PositionPlaneGizmo.cpp
@@ -114,7 +114,7 @@ void UPositionPlaneGizmo::Render(IToolsContextRenderAPI* RenderAPI)
 
 bool UPositionPlaneGizmo::HitTest(const FRay& Ray, FHitResult& OutHit)
 {
-	if (CenterBallShape->TestRayIntersection(Ray))
+	if (CenterBallShape->TestRayIntersection(FRay3d(Ray)))
 	{
 		OutHit.Distance = 0.1;
 		return true;
@@ -136,7 +136,7 @@ void UPositionPlaneGizmo::OnUpdateDrag(const FRay& Ray)
 	if (bInTransformDrag)
 	{
 		FVector3d SnapPos;
-		if (QuickTransformer.UpdateSnap(Ray, SnapPos))
+		if (QuickTransformer.UpdateSnap(FRay3d(Ray), SnapPos))
 		{
 			FTransform CurTransform = CenterBallShape->GetTransform();
 			CurTransform.SetTranslation((FVector)SnapPos);
