Copy up from Dev-Editor [REDACTED_HANDLE]

diff --git a/Engine/Plugins/Experimental/MeshModelingToolset/Source/MeshModelingTools/Private/RemeshMeshTool.cpp b/Engine/Plugins/Experimental/MeshModelingToolset/Source/MeshModelingTools/Private/RemeshMeshTool.cpp
--- a/Engine/Plugins/Experimental/MeshModelingToolset/Source/MeshModelingTools/Private/RemeshMeshTool.cpp
+++ b/Engine/Plugins/Experimental/MeshModelingToolset/Source/MeshModelingTools/Private/RemeshMeshTool.cpp
@@ -100,8 +100,9 @@ void URemeshMeshTool::Setup()
 	{
 		MeshStatisticsProperties->Update(*Compute->PreviewMesh->GetPreviewDynamicMesh());
 	});
-	Preview->ConfigureMaterials(
-		ToolSetupUtil::GetDefaultMaterial(GetToolManager(), ComponentTarget->GetMaterial(0)),
+	FComponentMaterialSet MaterialSet;
+	ComponentTarget->GetMaterialSet(MaterialSet);
+	Preview->ConfigureMaterials( MaterialSet.Materials,
 		ToolSetupUtil::GetDefaultWorkingMaterial(GetToolManager())
 	);
 	Preview->PreviewMesh->EnableWireframe(true);
@@ -141,10 +142,10 @@ void URemeshMeshTool::Setup()
 void URemeshMeshTool::Shutdown(EToolShutdownType ShutdownType)
 {
 	ComponentTarget->SetOwnerVisibility(true);
-	TUniquePtr<FDynamicMeshOpResult> Result = Preview->Shutdown();
+	FDynamicMeshOpResult Result = Preview->Shutdown();
 	if (ShutdownType == EToolShutdownType::Accept)
 	{
-		GenerateAsset(*Result);
+		GenerateAsset(Result);
 	}
 }
 
@@ -155,9 +156,9 @@ void URemeshMeshTool::Tick(float DeltaTime)
 }
 
 
-TSharedPtr<FDynamicMeshOperator> URemeshMeshTool::MakeNewOperator()
+TUniquePtr<FDynamicMeshOperator> URemeshMeshTool::MakeNewOperator()
 {
-	TSharedPtr<FRemeshMeshOp> Op = MakeShared<FRemeshMeshOp>();
+	TUniquePtr<FRemeshMeshOp> Op = MakeUnique<FRemeshMeshOp>();
 
 	if (!BasicProperties->bUseTargetEdgeLength)
 	{
@@ -206,7 +207,7 @@ void URemeshMeshTool::Render(IToolsContextRenderAPI* RenderAPI)
 			{
 				FVector3d A, B;
 				TargetMesh->GetEdgeV(eid, A, B);
-				PDI->DrawLine(Transform.TransformPosition(A), Transform.TransformPosition(B),
+				PDI->DrawLine(Transform.TransformPosition((FVector)A), Transform.TransformPosition((FVector)B),
 					LineColor, 0, 2.0, 1.0f, true);
 			}
 		}
@@ -243,12 +244,12 @@ void URemeshMeshTool::GenerateAsset(const FDynamicMeshOpResult& Result)
 	GetToolManager()->BeginUndoTransaction(LOCTEXT("RemeshMeshToolTransactionName", "Remesh Mesh"));
 
 	check(Result.Mesh.Get() != nullptr);
-	ComponentTarget->CommitMesh([&Result](FMeshDescription* MeshDescription)
+	ComponentTarget->CommitMesh([&Result](const FPrimitiveComponentTarget::FCommitParams& CommitParams)
 	{
 		FDynamicMeshToMeshDescription Converter;
 
 		// full conversion if normal topology changed or faces were inverted
-		Converter.Convert(Result.Mesh.Get(), *MeshDescription);
+		Converter.Convert(Result.Mesh.Get(), *CommitParams.MeshDescription);
 	});
 
 	GetToolManager()->EndUndoTransaction();
diff --git a/Engine/Plugins/Experimental/MeshModelingToolset/Source/MeshModelingTools/Private/TransformMeshesTool.cpp b/Engine/Plugins/Experimental/MeshModelingToolset/Source/MeshModelingTools/Private/TransformMeshesTool.cpp
--- a/Engine/Plugins/Experimental/MeshModelingToolset/Source/MeshModelingTools/Private/TransformMeshesTool.cpp
+++ b/Engine/Plugins/Experimental/MeshModelingToolset/Source/MeshModelingTools/Private/TransformMeshesTool.cpp
@@ -372,11 +372,11 @@ void UTransformMeshesTool::OnClickDrag(const FInputDeviceRay& DragPos)
 		FVector3d AlignTranslate = ToFrameWorld.Origin - FromFrameWorld.Origin;
 
 		FTransform NewTransform = StartDragTransform;
-		NewTransform.Accumulate( FTransform(CenterShift) );
-		NewTransform.Accumulate( FTransform(AlignRotation) );
-		NewTransform.Accumulate( FTransform(AlignTranslate) );
+		NewTransform.Accumulate( FTransform((FVector)CenterShift) );
+		NewTransform.Accumulate( FTransform((FQuat)AlignRotation) );
+		NewTransform.Accumulate( FTransform((FVector)AlignTranslate) );
 		CenterShift = AlignRotation * CenterShift;
-		NewTransform.Accumulate( FTransform(-CenterShift) );
+		NewTransform.Accumulate( FTransform((FVector)-CenterShift) );
 
 		FTransformMeshesTarget& ActiveTarget =
 			(TransformProps->TransformMode == ETransformMeshesTransformMode::PerObjectGizmo) ?
diff --git a/Engine/Plugins/Experimental/MeshModelingToolset/Source/MeshModelingTools/Private/WeldMeshEdgesTool.cpp b/Engine/Plugins/Experimental/MeshModelingToolset/Source/MeshModelingTools/Private/WeldMeshEdgesTool.cpp
--- a/Engine/Plugins/Experimental/MeshModelingToolset/Source/MeshModelingTools/Private/WeldMeshEdgesTool.cpp
+++ b/Engine/Plugins/Experimental/MeshModelingToolset/Source/MeshModelingTools/Private/WeldMeshEdgesTool.cpp
@@ -61,11 +61,12 @@ void UWeldMeshEdgesTool::Setup()
 	DynamicMeshComponent->RegisterComponent();
 	DynamicMeshComponent->SetWorldTransform(ComponentTarget->GetWorldTransform());
 
-	// copy material if there is one
-	auto Material = ComponentTarget->GetMaterial(0);
-	if (Material != nullptr)
+	// transfer materials
+	FComponentMaterialSet MaterialSet;
+	ComponentTarget->GetMaterialSet(MaterialSet);
+	for (int k = 0; k < MaterialSet.Materials.Num(); ++k)
 	{
-		DynamicMeshComponent->SetMaterial(0, Material);
+		DynamicMeshComponent->SetMaterial(k, MaterialSet.Materials[k]);
 	}
 
 	DynamicMeshComponent->InitializeMesh(ComponentTarget->GetMesh());
@@ -91,9 +92,9 @@ void UWeldMeshEdgesTool::Shutdown(EToolShutdownType ShutdownType)
 		{
 			// this block bakes the modified DynamicMeshComponent back into the StaticMeshComponent inside an undo transaction
 			GetToolManager()->BeginUndoTransaction(LOCTEXT("WeldMeshEdgesToolTransactionName", "Remesh Mesh"));
-			ComponentTarget->CommitMesh([=](FMeshDescription* MeshDescription)
+			ComponentTarget->CommitMesh([=](const FPrimitiveComponentTarget::FCommitParams& CommitParams)
 			{
-				DynamicMeshComponent->Bake(MeshDescription, true);
+				DynamicMeshComponent->Bake(CommitParams.MeshDescription, true);
 			});
 			GetToolManager()->EndUndoTransaction();
 		}
@@ -109,7 +110,6 @@ void UWeldMeshEdgesTool::Render(IToolsContextRenderAPI* RenderAPI)
 {
 	UpdateResult();
 
-
 	FPrimitiveDrawInterface* PDI = RenderAPI->GetPrimitiveDrawInterface();
 	FTransform Transform = ComponentTarget->GetWorldTransform(); //Actor->GetTransform();
 
@@ -122,7 +122,7 @@ void UWeldMeshEdgesTool::Render(IToolsContextRenderAPI* RenderAPI)
 		{
 			FVector3d A, B;
 			TargetMesh->GetEdgeV(eid, A, B);
-			PDI->DrawLine(Transform.TransformPosition(A), Transform.TransformPosition(B),
+			PDI->DrawLine(Transform.TransformPosition((FVector)A), Transform.TransformPosition((FVector)B),
 				LineColor, 0, 1.0, 1.0f, true);
 		}
 	}
@@ -133,7 +133,7 @@ void UWeldMeshEdgesTool::Render(IToolsContextRenderAPI* RenderAPI)
 	{
 		FVector3d A, B;
 		TargetMesh->GetEdgeV(eid, A, B);
-		PDI->DrawLine(Transform.TransformPosition(A), Transform.TransformPosition(B),
+		PDI->DrawLine(Transform.TransformPosition((FVector)A), Transform.TransformPosition((FVector)B),
 			LineColor2, 0, 2.0, 1.0f, true);
 	}
 
diff --git a/Engine/Plugins/Experimental/MeshModelingToolset/Source/MeshModelingTools/Public/BakeTransformTool.h b/Engine/Plugins/Experimental/MeshModelingToolset/Source/MeshModelingTools/Public/BakeTransformTool.h
new file mode 100644
--- /dev/null
+++ b/Engine/Plugins/Experimental/MeshModelingToolset/Source/MeshModelingTools/Public/BakeTransformTool.h
@@ -0,0 +1,99 @@
+// Copyright 1998-2019 Epic Games, Inc. All Rights Reserved.
+
+#pragma once
+
+#include "CoreMinimal.h"
+#include "UObject/NoExportTypes.h"
+#include "MultiSelectionTool.h"
+#include "InteractiveToolBuilder.h"
+#include "MeshOpPreviewHelpers.h"
+#include "DynamicMesh3.h"
+#include "BaseTools/SingleClickTool.h"
+#include "BakeTransformTool.generated.h"
+
+
+// predeclarations
+struct FMeshDescription;
+class USimpleDynamicMeshComponent;
+
+
+
+
+
+/**
+ *
+ */
+UCLASS()
+class MESHMODELINGTOOLS_API UBakeTransformToolBuilder : public UInteractiveToolBuilder
+{
+	GENERATED_BODY()
+
+public:
+	IToolsContextAssetAPI* AssetAPI = nullptr;
+
+	virtual bool CanBuildTool(const FToolBuilderState& SceneState) const override;
+	virtual UInteractiveTool* BuildTool(const FToolBuilderState& SceneState) const override;
+};
+
+
+
+
+
+/**
+ * Standard properties
+ */
+UCLASS()
+class MESHMODELINGTOOLS_API UBakeTransformToolProperties : public UInteractiveToolPropertySet
+{
+	GENERATED_BODY()
+
+public:
+	UBakeTransformToolProperties();
+
+	/** Recompute all mesh normals */
+	UPROPERTY(EditAnywhere, Category = Options)
+	bool bRecomputeNormals;
+};
+
+
+
+/**
+ * Simple Mesh Normal Updating Tool
+ */
+UCLASS()
+class MESHMODELINGTOOLS_API UBakeTransformTool : public UMultiSelectionTool
+{
+	GENERATED_BODY()
+
+public:
+	UBakeTransformTool();
+
+	virtual void Setup() override;
+	virtual void Shutdown(EToolShutdownType ShutdownType) override;
+
+	virtual void SetWorld(UWorld* World);
+	virtual void SetAssetAPI(IToolsContextAssetAPI* AssetAPI);
+
+	virtual void Tick(float DeltaTime) override;
+	virtual void Render(IToolsContextRenderAPI* RenderAPI) override;
+
+	virtual bool HasCancel() const override { return true; }
+	virtual bool HasAccept() const override;
+	virtual bool CanAccept() const override;
+
+protected:
+
+	UPROPERTY()
+		UBakeTransformToolProperties* BasicProperties;
+
+	UPROPERTY()
+		TArray<UPreviewMesh*> Previews;
+
+protected:
+	TArray<TSharedPtr<FDynamicMesh3>> OriginalDynamicMeshes;
+
+	UWorld* TargetWorld;
+	IToolsContextAssetAPI* AssetAPI;
+
+	void UpdateAssets(const TArray<TUniquePtr<FDynamicMesh3>>& Results);
+};
diff --git a/Engine/Plugins/Experimental/MeshModelingToolset/Source/MeshModelingTools/Public/ConvertToPolygonsTool.h b/Engine/Plugins/Experimental/MeshModelingToolset/Source/MeshModelingTools/Public/ConvertToPolygonsTool.h
--- a/Engine/Plugins/Experimental/MeshModelingToolset/Source/MeshModelingTools/Public/ConvertToPolygonsTool.h
+++ b/Engine/Plugins/Experimental/MeshModelingToolset/Source/MeshModelingTools/Public/ConvertToPolygonsTool.h
@@ -9,6 +9,7 @@
 #include "DynamicMesh3.h"
 #include "FindPolygonsAlgorithm.h"
 #include "PreviewMesh.h"
+#include "Changes/ValueWatcher.h"
 #include "ConvertToPolygonsTool.generated.h"
 
 // predeclaration
@@ -30,6 +31,17 @@ public:
 
 
 
+UENUM()
+enum class EConvertToPolygonsMode
+{
+	/** Convert based on Angle Tolerance between Face Normals */
+	FaceNormalDeviation,
+	/** Create PolyGroups based on UV Islands */
+	FromUVISlands
+};
+
+
+
 
 UCLASS()
 class MESHMODELINGTOOLS_API UConvertToPolygonsToolProperties : public UInteractiveToolPropertySet
@@ -37,8 +49,11 @@ class MESHMODELINGTOOLS_API UConvertToPolygonsToolProperties : public UInteracti
 	GENERATED_BODY()
 
 public:
+	UPROPERTY(EditAnywhere, Category = PolyGroups)
+	EConvertToPolygonsMode ConversionMode = EConvertToPolygonsMode::FaceNormalDeviation;
+
 	/** Tolerance for planarity */
-	UPROPERTY(EditAnywhere, Category = PolyGroups, meta = (UIMin = "0.001", UIMax = "20.0", ClampMin = "0.0", ClampMax = "90.0"))
+	UPROPERTY(EditAnywhere, Category = PolyGroups, meta = (UIMin = "0.001", UIMax = "20.0", ClampMin = "0.0", ClampMax = "90.0", EditCondition = "ConversionMode == EConvertToPolygonsMode::FaceNormalDeviation"))
 	float AngleTolerance = 0.1f;
 
 	UPROPERTY(EditAnywhere, Category = PolyGroups)
@@ -58,6 +73,7 @@ public:
 	virtual void Shutdown(EToolShutdownType ShutdownType) override;
 
 	virtual void Render(IToolsContextRenderAPI* RenderAPI) override;
+	virtual void Tick(float DeltaTime) override;
 
 	virtual bool HasCancel() const override { return true; }
 	virtual bool HasAccept() const override;
@@ -75,6 +91,10 @@ protected:
 protected:
 	FDynamicMesh3 SearchMesh;
 	FDynamicMeshNormalOverlay InitialNormals;
+
+	TValueWatcher<EConvertToPolygonsMode> ConvertModeWatcher;
+
+
 	FFindPolygonsAlgorithm Polygons;
 	bool bPolygonsValid = false;
 	void UpdatePolygons();
diff --git a/Engine/Plugins/Experimental/MeshModelingToolset/Source/MeshModelingTools/Public/DeformMeshPolygonsTool.h b/Engine/Plugins/Experimental/MeshModelingToolset/Source/MeshModelingTools/Public/DeformMeshPolygonsTool.h
new file mode 100644
--- /dev/null
+++ b/Engine/Plugins/Experimental/MeshModelingToolset/Source/MeshModelingTools/Public/DeformMeshPolygonsTool.h
@@ -0,0 +1,490 @@
+// Copyright 1998-2019 Epic Games, Inc. All Rights Reserved.
+
+#pragma once
+
+#include "CoreMinimal.h"
+#include "UObject/NoExportTypes.h"
+#include "BaseTools/MeshSurfacePointTool.h"
+#include "SimpleDynamicMeshComponent.h"
+#include "DynamicMeshAABBTree3.h"
+#include "ToolDataVisualizer.h"
+#include "Transforms/QuickAxisTranslater.h"
+#include "Transforms/QuickAxisRotator.h"
+#include "Changes/MeshVertexChange.h"
+#include "GroupTopology.h"
+#include "Spatial/GeometrySet3.h"
+#include "Selection/GroupTopologySelector.h"
+#include "Operations/GroupTopologyDeformer.h"
+#include "MeshSolverUtilities/Private/LaplacianMeshSmoother.h"
+#include "Curves/CurveFloat.h"
+#include "ModelingOperators/Public/ModelingTaskTypes.h"
+#include "EditMeshPolygonsTool.h"		// for a few enums
+#include "DeformMeshPolygonsTool.generated.h"
+
+class FMeshVertexChangeBuilder;
+class FGroupTopologyLaplacianDeformer;
+class FDeformTask;
+struct FConstraintData;
+
+/**
+ * ToolBuilder
+ */
+UCLASS()
+class MESHMODELINGTOOLS_API UDeformMeshPolygonsToolBuilder : public UMeshSurfacePointToolBuilder
+{
+	GENERATED_BODY()
+
+public:
+	UDeformMeshPolygonsToolBuilder()
+	{
+	}
+
+	virtual UMeshSurfacePointTool* CreateNewTool(const FToolBuilderState& SceneState) const override;
+};
+
+/** Deformation strategies */
+UENUM()
+enum class EGroupTopologyDeformationStrategy : uint8
+{
+	/** Deforms the mesh using linear translations*/
+	Linear UMETA(DisplayName = "Linear"),
+
+	/** Deforms the mesh using laplacian deformation*/
+	Laplacian UMETA(DisplayName = "Laplacian")
+};
+
+/** Laplacian weight schemes determine how we will look at the curvature at a given vertex in relation to its neighborhood*/
+UENUM()
+enum class EWeightScheme 
+{
+	Uniform				UMETA(DisplayName = "Uniform"),
+	Umbrella			UMETA(DisplayName = "Umbrella"),
+	Valence				UMETA(DisplayName = "Valence"),
+	MeanValue			UMETA(DisplayName = "MeanValue"),
+	Cotangent			UMETA(DisplayName = "Cotangent"),
+	ClampedCotangent	UMETA(DisplayName = "ClampedCotangent")
+};
+
+/** The ELaplacianWeightScheme enum is the same..*/
+static ELaplacianWeightScheme ConvertToLaplacianWeightScheme(const EWeightScheme WeightScheme)
+{
+	return static_cast<ELaplacianWeightScheme>(WeightScheme);
+}
+
+
+
+UCLASS()
+class MESHMODELINGTOOLS_API UDeformMeshPolygonsTransformProperties : public UInteractiveToolPropertySet
+{
+	GENERATED_BODY()
+
+public:
+	UDeformMeshPolygonsTransformProperties();
+
+#if WITH_EDITOR
+	// UObject interface
+	virtual void PostEditChangeProperty(FPropertyChangedEvent& PropertyChangedEvent) override;
+	// End of UObject interface
+#endif
+
+	//Options
+
+	UPROPERTY(EditAnywhere, Category = Options, meta = (DisplayName = "Deformation Type", ToolTip = "Select the type of deformation you wish to employ on a polygroup."))
+	EGroupTopologyDeformationStrategy DeformationStrategy;
+
+	UPROPERTY(EditAnywhere, Category = Options)
+	EQuickTransformerMode TransformMode;
+
+	UPROPERTY(EditAnywhere, Category = Options)
+	bool bSelectFaces;
+
+	UPROPERTY(EditAnywhere, Category = Options)
+	bool bSelectEdges;
+
+	UPROPERTY(EditAnywhere, Category = Options)
+	bool bSelectVertices;
+
+	UPROPERTY(EditAnywhere, Category = Options)
+	bool bShowWireframe;
+
+	UPROPERTY(EditAnywhere, Category = Options)
+	bool bSnapToWorldGrid;
+
+	//Laplacian Deformation Options, currently not exposed.
+
+	UPROPERTY() 
+	EWeightScheme SelectedWeightScheme = EWeightScheme::ClampedCotangent;
+
+	UPROPERTY() 
+	double HandleWeight = 1000.0;
+
+	UPROPERTY()
+	bool bPostFixHandles = false;
+/**
+	// How to add a weight curve
+	UPROPERTY(EditAnywhere, Category = LaplacianOptions, meta = (EditCondition="DeformationStrategy == EGroupTopologyDeformationStrategy::Laplacian", DisplayName = "Localize Deformation", ToolTip = "When enabled, only the vertices in the polygroups immediately adjacent to the selected group will be affected by the deformation.\nWhen disabled, the deformer will solve for the curvature of the entire mesh (slower)"))
+	bool bLocalizeDeformation{ true};
+
+	UPROPERTY(EditAnywhere, Category = LaplacianOptions, meta = (EditCondition="DeformationStrategy == EGroupTopologyDeformationStrategy::Laplacian", DisplayName = "Apply Weight Attenuation Curve", ToolTip = "When enabled, the curve below will be used to calculate the weight at a given vertex based on distance from the handles"))
+	bool bApplyAttenuationCurve{ false };
+	
+	FRichCurve DefaultFalloffCurve;
+
+	UPROPERTY(EditAnywhere,  Category = LaplacianOptions, meta = ( EditCondition="DeformationStrategy == EGroupTopologyDeformationStrategy::Laplacian && bApplyAttenuationCurve", DisplayName = "Distance-Weight Attenuation Curve",UIMin = "1.0", UIMax = "1.0", ClampMin = "1.0", ClampMax = "1.0", ToolTip = "This curve determines the weight attenuation over the distance of the mesh.\nThe selected polygroup handle is t=0.0, and t=1.0 is roughly the farthest vertices from the handles.\nThe value of the curve at each time interval represents the weight of the vertices at that distance from the selection."))
+	FRuntimeFloatCurve WeightAttenuationCurve;
+*/ 
+};
+
+//Stores per-vertex data needed by the laplacian deformer object
+//TODO: May be a candidate for a subclass of the FGroupTopologyLaplacianDeformer
+struct FConstraintData
+{
+	FConstraintData& operator=(const FConstraintData& other)
+	{
+		Position = other.Position;
+		Weight   = other.Weight;
+		bPostFix = other.bPostFix;
+		return *this;
+	}
+
+	FVector3d Position;
+	double Weight{ 0.0 };
+	bool bPostFix{ false };
+};
+
+/**
+*	FDeformTask is an object which wraps an asynchronous task to be run multiple times on a separate thread.
+*	The Laplacian deformation process requires the use of potentially large sparse matrices and sparse multiplication.
+*	
+*   Expected usage:
+*
+*   
+*   // define constraints.  Need Constraints[VertID] to hold the constraints for the corresponding vertex.
+*   TArray<FConstraintData> Constraints;
+*   ....
+*
+*   // populate with the VertexIDs of the vertices that are in the region you wish to deform.  
+*   TArray<int32> SrcVertIDs;  //Basically a mini-index buffer.
+*   ...
+* 
+*   // Create or reuse a laplacian deformation task.
+*   FDeformTask*   DeformTask = New FDeformTask(WeightScheme);
+*
+*   // the deformer will have to build a new mesh that represents the regions in SrcVertIDs;
+*   // but set this to false on subsequent calls to UpdateDeformer if the SrcVertIDs array hasn't changed.
+*   bool bRequiresRegion = true;
+
+*   DefTask->UpdateDeformer(WeightScheme, Mesh, Constraints, SrcVertIDs, bRequiresRegion);
+*
+*   DeformTask->DoWork();  or DeformTask->StartBackgroundTask(); //which calls DoWork on background thread.
+*
+*  // wheh DeformTask->IsDone == true; you can copy the results back to the mesh
+*  DeformTask->ExportResults(Mesh);
+*
+* Note: if only the positions in the Constraints change (e.g. handle positions) then subsequent calls 
+* to UpdateDeformer() and DoWork() will be much faster as the matrix system will not be rebuilt or re-factored
+*/
+class FConstrainedMeshDeformerTask : public FNonAbandonableTask
+{
+	friend class FAsyncTask<FDeformTask>;
+public:
+
+	enum
+	{
+		INACTIVE_SUBSET_ID = -1
+	};
+
+	FConstrainedMeshDeformerTask(const ELaplacianWeightScheme SelectedWeightScheme)
+	{
+		LaplacianWeightScheme = SelectedWeightScheme;
+	}
+
+	virtual ~FConstrainedMeshDeformerTask() {};
+
+	//NO idea what this is meant to do. Performance analysis maybe? Scheduling data?
+	FORCEINLINE TStatId GetStatId() const
+	{
+		RETURN_QUICK_DECLARE_CYCLE_STAT(FConstrainedMeshDeformerTask, STATGROUP_ThreadPoolAsyncTasks);
+	}
+
+	
+	/** Called by the main thread in the tool, this copies the Constraint buffer right before the task begins on another thread.
+	  * Ensures the FConstrainedMeshDeformer is using correct mesh subset and the selected settings, then updates on change in properties, i.e. weight scheme */
+	void UpdateDeformer(const ELaplacianWeightScheme SelectedWeightScheme, const FDynamicMesh3& Mesh, const TArray<FConstraintData>& ConstraintArray, const TArray<int32>& SrcIDBufferSubset, bool bNewTransaction, const FRichCurve* Curve);
+
+	/** Required by the FAsyncTaskExecutor */
+	void SetAbortSource(bool* bAbort) { bAbortSource = bAbort; };
+
+	/** Called by the FAsyncTask<FDeformTask> object for background computation. */
+	void DoWork();
+
+
+	/** Updates the positions in the target mesh for regions that correspond to the subset mesh */
+	void ExportResults(FDynamicMesh3& TargetMesh) const;
+
+private:
+
+	/** Creates the mesh (i.e. SubsetMesh) that corresponds to the region of the SrcMesh defined by the partial index buffer SrcIDBufferSubset */
+	void InitializeSubsetMesh(const FDynamicMesh3& SrcMesh, const TArray<int32>& SrcIDBufferSubset);
+
+	/** Attenuates the weights of the constraints using the selected curve */
+	void ApplyAttenuation();
+
+	/** Denotes the weight scheme being used by the running background task. Changes when selected property changes in editor. */
+	ELaplacianWeightScheme LaplacianWeightScheme;
+
+	/** positions for each vertex in the subset mesh - for use in the deformer */
+	TArray<FVector3d> SubsetPositionBuffer;
+
+	/** constraint data for each vertex in subset mesh - for use by the deformer*/
+	TArray<FConstraintData> SubsetConstraintBuffer;
+
+	FRichCurve WeightAttenuationCurve;
+
+	/** True only for the first update, and then false for the duration of the Input transaction
+	  * It's passed in and copied in UpdateDeformer() */
+	bool bIsNewTransaction = true;
+
+	/** When true, the constraint weights will be attenuated based on distance using the provided curve object*/
+	bool bAttenuateWeights = false;
+
+	/** The abort bool used by the Task Deleter */
+	bool* bAbortSource = nullptr;
+
+	/** Used to initialize the array mapping, updated during the UpdateDeformer() function */
+	int SrcMeshMaxVertexID;
+
+	/** A subset of the original mesh */
+	FDynamicMesh3 SubsetMesh;
+
+	/** Maps Subset Mesh VertexID to Src Mesh VertexID */
+	TArray<int32> SubsetVertexIDToSrcVertexIDMap;
+
+	/** Laplacian deformer object gets rebuilt each new transaction */
+	TUniquePtr<FConstrainedMeshDeformer> ConstrainedDeformer;
+
+private :
+	FConstrainedMeshDeformerTask();
+
+};
+
+class FGroupTopologyLaplacianDeformer : public FGroupTopologyDeformer
+{
+
+public:
+
+	FGroupTopologyLaplacianDeformer() = default;
+
+	virtual ~FGroupTopologyLaplacianDeformer();
+
+	/** Used to begin a procedural addition of modified vertices */
+	inline void ResetModifiedVertices()
+	{
+		ModifiedVertices.Empty();
+	};
+
+	/** Change tracking */
+	template <typename ValidSetAppendContainerType>
+	void RecordModifiedVertices(const ValidSetAppendContainerType& Container)
+	{
+		ModifiedVertices.Empty();
+		ModifiedVertices.Append(Container);
+	}
+
+	/** Used to iteratively add to the active change set (TSet<>)*/
+	inline void RecordModifiedVertex(int32 VertexID)
+	{
+		ModifiedVertices.Add(VertexID);
+	};
+
+
+	void SetActiveHandleFaces(const TArray<int>& FaceGroupIDs) override;
+	void SetActiveHandleEdges(const TArray<int>& TopologyEdgeIDs) override;
+	void SetActiveHandleCorners(const TArray<int>& TopologyCornerIDs) override;
+
+
+
+	/** Allocates shared storage for use in task synchronization */
+	void InitBackgroundWorker(const ELaplacianWeightScheme WeightScheme);
+
+
+	/** Coordinates the background tasks. Returns false if the worker was already running */
+	bool UpdateAndLaunchdWorker(const ELaplacianWeightScheme WeightScheme, const FRichCurve* Curve = nullptr);
+
+	/** Capture data about background task state.*/
+	bool IsTaskInFlight() const;
+
+
+
+	/** Sets the SrcMeshConstraintBuffer to have a size of MaxVertexID, and initializes with the current mesh positions, but weight zero*/
+	void InitializeConstraintBuffer();
+
+	/** Given an array of Group IDs, update the selection and record vertices */
+	void UpdateSelection(const FDynamicMesh3* TargetMesh, const TArray<int>& Groups, bool bLocalizeDeformation);
+
+	/** Updates the mesh preview and/or solvers upon user input, provided a deformation strategy */
+	void UpdateSolution(FDynamicMesh3* TargetMesh, const TFunction<FVector3d(FDynamicMesh3* Mesh, int)>& HandleVertexDeformFunc) override;
+
+	/** Updates the vertex positions of the mesh with the result from the last deformation solve. */
+	void ExportDeformedPositions(FDynamicMesh3* TargetMesh);
+
+	/** Returns true if the asynchronous task has finished. */
+	inline bool IsDone() { return AsyncMeshDeformTask == nullptr || AsyncMeshDeformTask->IsDone(); };
+
+	/** Triggers abort on task and passes off ownership to deleter object */
+	inline void Shutdown();
+
+	const TArray<FROIFace>& GetROIFaces() const { return ROIFaces; }
+
+	/** Stores the position of the vertex constraints and corresponding weights for the entire mesh.  This is used as a form of scratch space.*/
+	TArray<FConstraintData> SrcMeshConstraintBuffer;
+
+	/** Array of vertex indices organized in groups of three - basically an index buffer - that defines the subset of the mesh that the deformation task will work on.*/
+	TArray<int32> SubsetIDBuffer;
+
+	/** Need to update the task with the current submesh */
+	bool bTaskSubmeshIsDirty = true;
+
+	/** Asynchronous task object. This object deals with expensive matrix functionality that computes the deformation of a local mesh. */
+	FAsyncTaskExecuterWithAbort<FConstrainedMeshDeformerTask>* AsyncMeshDeformTask = nullptr;
+
+
+	/** The weight which will be applied to the constraints corresponding to the handle vertices. */
+	double HandleWeights = 1.0;
+
+	/** This is set to true whenever the user interacts with the tool under laplacian deformation mode.
+	  * It is set to false immediately before beginning a background task and cannot be set to false again until the work is done. */
+	bool bDeformerNeedsToRun = false;
+
+
+	/** When true, tells the solver to attempt to postfix the actual position of the handles to the constrained position */
+	bool bPostfixHandles = false;
+
+	//This is set to false only after 
+	//	1) the asynchronous deformation task is complete
+	//	2) the main thread has seen it complete, and
+	//	3) the main thread updates the vertex positions of the mesh one last time
+	bool bVertexPositionsNeedSync = false;
+
+	bool bLocalize =  true;
+
+};
+
+
+/**
+ *
+ */
+UCLASS()
+class MESHMODELINGTOOLS_API UDeformMeshPolygonsTool : public UMeshSurfacePointTool
+{
+	GENERATED_BODY()
+
+public:
+	UDeformMeshPolygonsTool();
+
+	virtual void RegisterActions(FInteractiveToolActionSet& ActionSet) override;
+
+	virtual void Setup() override;
+	virtual void Shutdown(EToolShutdownType ShutdownType) override;
+
+	virtual void Tick(float DeltaTime) override;
+	virtual void Render(IToolsContextRenderAPI* RenderAPI) override;
+
+	virtual bool HasCancel() const override { return true; }
+	virtual bool HasAccept() const override { return true; }
+	virtual bool CanAccept() const override { return true; }
+
+
+	virtual bool HitTest(const FRay& Ray, FHitResult& OutHit) override;
+	
+	virtual void OnBeginDrag(const FRay& Ray) override;
+	virtual void OnUpdateDrag(const FRay& Ray) override;
+	virtual void OnEndDrag(const FRay& Ray) override;
+	virtual bool OnUpdateHover(const FInputDeviceRay& DevicePos) override;
+
+public:
+	virtual void NextTransformTypeAction();
+
+	//
+	float VisualAngleSnapThreshold = 0.5;
+
+protected:
+	UPROPERTY()
+	USimpleDynamicMeshComponent* DynamicMeshComponent;
+
+	UPROPERTY()
+	UDeformMeshPolygonsTransformProperties* TransformProps;
+
+	// realtime visualization
+	void OnDynamicMeshComponentChanged();
+	FDelegateHandle OnDynamicMeshComponentChangedHandle;
+
+	virtual void OnPropertyModified(UObject* PropertySet, UProperty* Property);
+	
+	// camera state at last render
+	FViewCameraState CameraState;
+
+	FToolDataVisualizer PolyEdgesRenderer;
+
+	// True for the duration of UI click+drag
+	bool bInDrag;
+
+	FPlane ActiveDragPlane;
+	FVector StartHitPosWorld;
+	FVector StartHitNormalWorld;
+	FVector LastHitPosWorld;
+	FVector LastBrushPosLocal;
+	FVector StartBrushPosLocal;
+
+	FFrame3d ActiveSurfaceFrame;
+	FQuickTransformer* GetActiveQuickTransformer();
+	void UpdateActiveSurfaceFrame(FGroupTopologySelection& Selection);
+	void UpdateQuickTransformer();
+
+	FRay UpdateRay;
+	bool bUpdatePending = false;
+	void ComputeUpdate();
+
+	FVector3d LastMoveDelta;
+	FQuickAxisTranslater QuickAxisTranslater;
+	void ComputeUpdate_Translate();
+
+	FQuickAxisRotator QuickAxisRotator;
+	FVector3d RotationStartPointWorld;
+	FFrame3d RotationStartFrame;
+	void ComputeUpdate_Rotate();
+
+	FGroupTopology Topology;
+	void PrecomputeTopology();
+
+	FGroupTopologySelector TopoSelector;
+	
+	//
+	// data for current drag
+	//
+
+	FGroupTopologySelection HilightSelection;
+	FToolDataVisualizer HilightRenderer;
+
+	FDynamicMeshAABBTree3 MeshSpatial;
+	FDynamicMeshAABBTree3& GetSpatial();
+
+	FMeshVertexChangeBuilder* ActiveVertexChange;
+
+	EGroupTopologyDeformationStrategy DeformationStrategy;
+
+	// The two deformer type options.
+	FGroupTopologyDeformer LinearDeformer;
+	FGroupTopologyLaplacianDeformer LaplacianDeformer;
+
+
+
+	// This is true when the spatial index needs to reflect a modification
+	bool bSpatialDirty; 
+
+	void BeginChange();
+	void EndChange();
+	void UpdateChangeFromROI(bool bFinal);
+
+};
+
diff --git a/Engine/Plugins/Experimental/MeshModelingToolset/Source/MeshModelingTools/Public/EditMeshPolygonsTool.h b/Engine/Plugins/Experimental/MeshModelingToolset/Source/MeshModelingTools/Public/EditMeshPolygonsTool.h
--- a/Engine/Plugins/Experimental/MeshModelingToolset/Source/MeshModelingTools/Public/EditMeshPolygonsTool.h
+++ b/Engine/Plugins/Experimental/MeshModelingToolset/Source/MeshModelingTools/Public/EditMeshPolygonsTool.h
@@ -7,7 +7,7 @@
 #include "BaseTools/MeshSurfacePointTool.h"
 #include "SimpleDynamicMeshComponent.h"
 #include "DynamicMeshAABBTree3.h"
-#include "Drawing/ToolDataVisualizer.h"
+#include "ToolDataVisualizer.h"
 #include "Transforms/QuickAxisTranslater.h"
 #include "Transforms/QuickAxisRotator.h"
 #include "Changes/MeshVertexChange.h"
@@ -15,15 +15,13 @@
 #include "Spatial/GeometrySet3.h"
 #include "Selection/GroupTopologySelector.h"
 #include "Operations/GroupTopologyDeformer.h"
-#include "MeshSolverUtilities/Private/LaplacianMeshSmoother.h"
-#include "Curves/CurveFloat.h"
 #include "ModelingOperators/Public/ModelingTaskTypes.h"
+#include "Transforms/MultiTransformer.h"
+#include "Changes/ValueWatcher.h"
 #include "EditMeshPolygonsTool.generated.h"
 
 class FMeshVertexChangeBuilder;
-class FGroupTopologyLaplacianDeformer;
-class FDeformTask;
-struct FConstraintData;
+
 
 /**
  * ToolBuilder
@@ -32,44 +30,10 @@ UCLASS()
 class MESHMODELINGTOOLS_API UEditMeshPolygonsToolBuilder : public UMeshSurfacePointToolBuilder
 {
 	GENERATED_BODY()
-
 public:
-	UEditMeshPolygonsToolBuilder()
-	{
-	}
-
 	virtual UMeshSurfacePointTool* CreateNewTool(const FToolBuilderState& SceneState) const override;
 };
 
-/** Deformation strategies */
-UENUM()
-enum class EGroupTopologyDeformationStrategy : uint8
-{
-	/** Deforms the mesh using linear translations*/
-	Linear UMETA(DisplayName = "Linear"),
-
-	/** Deforms the mesh using laplacian deformation*/
-	Laplacian UMETA(DisplayName = "Laplacian")
-};
-
-/** Laplacian weight schemes determine how we will look at the curvature at a given vertex in relation to its neighborhood*/
-UENUM()
-enum class EWeightScheme 
-{
-	Uniform				UMETA(DisplayName = "Uniform"),
-	Umbrella			UMETA(DisplayName = "Umbrella"),
-	Valence				UMETA(DisplayName = "Valence"),
-	MeanValue			UMETA(DisplayName = "MeanValue"),
-	Cotangent			UMETA(DisplayName = "Cotangent"),
-	ClampedCotangent	UMETA(DisplayName = "ClampedCotangent")
-};
-
-/** The ELaplacianWeightScheme enum is the same..*/
-static ELaplacianWeightScheme ConvertToLaplacianWeightScheme(const EWeightScheme WeightScheme)
-{
-	return static_cast<ELaplacianWeightScheme>(WeightScheme);
-}
-
 /** Modes for quick transformer */
 UENUM()
 enum class EQuickTransformerMode : uint8
@@ -107,12 +71,8 @@ public:
 #endif
 
 	//Options
-
-	UPROPERTY(EditAnywhere, Category = Options, meta = (DisplayName = "Deformation Type", ToolTip = "Select the type of deformation you wish to employ on a polygroup."))
-	EGroupTopologyDeformationStrategy DeformationStrategy;
-
 	UPROPERTY(EditAnywhere, Category = Options)
-	EQuickTransformerMode TransformMode;
+	EMultiTransformerMode TransformMode;
 
 	UPROPERTY(EditAnywhere, Category = Options)
 	bool bSelectFaces;
@@ -123,281 +83,24 @@ public:
 	UPROPERTY(EditAnywhere, Category = Options)
 	bool bSelectVertices;
 
-	UPROPERTY(EditAnywhere, Category = Options, meta = (DisplayName = "Show Triangle Mesh Wireframe"))
-	bool bWireframe;
+	UPROPERTY(EditAnywhere, Category = Options)
+	bool bShowWireframe;
+
+	UPROPERTY(EditAnywhere, Category = Options)
+	bool bSnapToWorldGrid;
 
 	UPROPERTY(EditAnywhere, Category = Options)
 	EPolygonGroupMode PolygonMode;
 
 	UPROPERTY(EditAnywhere, Category = Options, meta = (EditCondition="PolygonMode == EPolygonGroupMode::RecomputePolygonsByAngleThreshold"))
 	float PolygonGroupingAngleThreshold;
-
-	//Laplacian Deformation Options, currently not exposed.
-
-	UPROPERTY() 
-	EWeightScheme SelectedWeightScheme = EWeightScheme::ClampedCotangent;
-
-	UPROPERTY() 
-	double HandleWeight = 1000.0;
-
-	UPROPERTY()
-	bool bPostFixHandles = false;
-/**
-	// How to add a weight curve
-	UPROPERTY(EditAnywhere, Category = LaplacianOptions, meta = (EditCondition="DeformationStrategy == EGroupTopologyDeformationStrategy::Laplacian", DisplayName = "Localize Deformation", ToolTip = "When enabled, only the vertices in the polygroups immediately adjacent to the selected group will be affected by the deformation.\nWhen disabled, the deformer will solve for the curvature of the entire mesh (slower)"))
-	bool bLocalizeDeformation{ true};
-
-	UPROPERTY(EditAnywhere, Category = LaplacianOptions, meta = (EditCondition="DeformationStrategy == EGroupTopologyDeformationStrategy::Laplacian", DisplayName = "Apply Weight Attenuation Curve", ToolTip = "When enabled, the curve below will be used to calculate the weight at a given vertex based on distance from the handles"))
-	bool bApplyAttenuationCurve{ false };
-	
-	FRichCurve DefaultFalloffCurve;
-
-	UPROPERTY(EditAnywhere,  Category = LaplacianOptions, meta = ( EditCondition="DeformationStrategy == EGroupTopologyDeformationStrategy::Laplacian && bApplyAttenuationCurve", DisplayName = "Distance-Weight Attenuation Curve",UIMin = "1.0", UIMax = "1.0", ClampMin = "1.0", ClampMax = "1.0", ToolTip = "This curve determines the weight attenuation over the distance of the mesh.\nThe selected polygroup handle is t=0.0, and t=1.0 is roughly the farthest vertices from the handles.\nThe value of the curve at each time interval represents the weight of the vertices at that distance from the selection."))
-	FRuntimeFloatCurve WeightAttenuationCurve;
-*/ 
-};
-
-//Stores per-vertex data needed by the laplacian deformer object
-//TODO: May be a candidate for a subclass of the FGroupTopologyLaplacianDeformer
-struct FConstraintData
-{
-	FConstraintData& operator=(const FConstraintData& other)
-	{
-		Position = other.Position;
-		Weight   = other.Weight;
-		bPostFix = other.bPostFix;
-		return *this;
-	}
-
-	FVector3d Position;
-	double Weight{ 0.0 };
-	bool bPostFix{ false };
 };
 
-/**
-*	FDeformTask is an object which wraps an asynchronous task to be run multiple times on a separate thread.
-*	The Laplacian deformation process requires the use of potentially large sparse matrices and sparse multiplication.
-*	
-*   Expected usage:
-*
-*   
-*   // define constraints.  Need Constraints[VertID] to hold the constraints for the corresponding vertex.
-*   TArray<FConstraintData> Constraints;
-*   ....
-*
-*   // populate with the VertexIDs of the vertices that are in the region you wish to deform.  
-*   TArray<int32> SrcVertIDs;  //Basically a mini-index buffer.
-*   ...
-* 
-*   // Create or reuse a laplacian deformation task.
-*   FDeformTask*   DeformTask = New FDeformTask(WeightScheme);
-*
-*   // the deformer will have to build a new mesh that represents the regions in SrcVertIDs;
-*   // but set this to false on subsequent calls to UpdateDeformer if the SrcVertIDs array hasn't changed.
-*   bool bRequiresRegion = true;
-
-*   DefTask->UpdateDeformer(WeightScheme, Mesh, Constraints, SrcVertIDs, bRequiresRegion);
-*
-*   DeformTask->DoWork();  or DeformTask->StartBackgroundTask(); //which calls DoWork on background thread.
-*
-*  // wheh DeformTask->IsDone == true; you can copy the results back to the mesh
-*  DeformTask->ExportResults(Mesh);
-*
-* Note: if only the positions in the Constraints change (e.g. handle positions) then subsequent calls 
-* to UpdateDeformer() and DoWork() will be much faster as the matrix system will not be rebuilt or re-factored
-*/
-class FConstrainedMeshDeformerTask : public FNonAbandonableTask
-{
-	friend class FAsyncTask<FDeformTask>;
-public:
-
-	enum
-	{
-		INACTIVE_SUBSET_ID = -1
-	};
-
-	FConstrainedMeshDeformerTask(const ELaplacianWeightScheme SelectedWeightScheme)
-	{
-		LaplacianWeightScheme = SelectedWeightScheme;
-	}
-
-	virtual ~FConstrainedMeshDeformerTask() {};
-
-	//NO idea what this is meant to do. Performance analysis maybe? Scheduling data?
-	FORCEINLINE TStatId GetStatId() const
-	{
-		RETURN_QUICK_DECLARE_CYCLE_STAT(FConstrainedMeshDeformerTask, STATGROUP_ThreadPoolAsyncTasks);
-	}
-
-	
-	/** Called by the main thread in the tool, this copies the Constraint buffer right before the task begins on another thread.
-	  * Ensures the FConstrainedMeshDeformer is using correct mesh subset and the selected settings, then updates on change in properties, i.e. weight scheme */
-	void UpdateDeformer(const ELaplacianWeightScheme SelectedWeightScheme, const FDynamicMesh3& Mesh, const TArray<FConstraintData>& ConstraintArray, const TArray<int32>& SrcIDBufferSubset, bool bNewTransaction, const FRichCurve* Curve);
-
-	/** Required by the FAsyncTaskExecutor */
-	void SetAbortSource(bool* bAbort) { bAbortSource = bAbort; };
-
-	/** Called by the FAsyncTask<FDeformTask> object for background computation. */
-	void DoWork();
-
-
-	/** Updates the positions in the target mesh for regions that correspond to the subset mesh */
-	void ExportResults(FDynamicMesh3& TargetMesh) const;
-
-private:
-
-	/** Creates the mesh (i.e. SubsetMesh) that corresponds to the region of the SrcMesh defined by the partial index buffer SrcIDBufferSubset */
-	void InitializeSubsetMesh(const FDynamicMesh3& SrcMesh, const TArray<int32>& SrcIDBufferSubset);
-
-	/** Attenuates the weights of the constraints using the selected curve */
-	void ApplyAttenuation();
-
-	/** Denotes the weight scheme being used by the running background task. Changes when selected property changes in editor. */
-	ELaplacianWeightScheme LaplacianWeightScheme;
-
-	/** positions for each vertex in the subset mesh - for use in the deformer */
-	TArray<FVector3d> SubsetPositionBuffer;
-
-	/** constraint data for each vertex in subset mesh - for use by the deformer*/
-	TArray<FConstraintData> SubsetConstraintBuffer;
-
-	FRichCurve WeightAttenuationCurve;
-
-	/** True only for the first update, and then false for the duration of the Input transaction
-	  * It's passed in and copied in UpdateDeformer() */
-	bool bIsNewTransaction = true;
-
-	/** When true, the constraint weights will be attenuated based on distance using the provided curve object*/
-	bool bAttenuateWeights = false;
-
-	/** The abort bool used by the Task Deleter */
-	bool* bAbortSource = nullptr;
-
-	/** Used to initialize the array mapping, updated during the UpdateDeformer() function */
-	int SrcMeshMaxVertexID;
-
-	/** A subset of the original mesh */
-	FDynamicMesh3 SubsetMesh;
-
-	/** Maps Subset Mesh VertexID to Src Mesh VertexID */
-	TArray<int32> SubsetVertexIDToSrcVertexIDMap;
-
-	/** Laplacian deformer object gets rebuilt each new transaction */
-	TUniquePtr<FConstrainedMeshDeformer> ConstrainedDeformer;
-
-private :
-	FConstrainedMeshDeformerTask();
-
-};
-
-class FGroupTopologyLaplacianDeformer : public FGroupTopologyDeformer
-{
-
-public:
-
-	FGroupTopologyLaplacianDeformer() = default;
-
-	virtual ~FGroupTopologyLaplacianDeformer();
-
-	/** Used to begin a procedural addition of modified vertices */
-	inline void ResetModifiedVertices()
-	{
-		ModifiedVertices.Empty();
-	};
-
-	/** Change tracking */
-	template <typename ValidSetAppendContainerType>
-	void RecordModifiedVertices(const ValidSetAppendContainerType& Container)
-	{
-		ModifiedVertices.Empty();
-		ModifiedVertices.Append(Container);
-	}
-
-	/** Used to iteratively add to the active change set (TSet<>)*/
-	inline void RecordModifiedVertex(int32 VertexID)
-	{
-		ModifiedVertices.Add(VertexID);
-	};
-
-
-	void SetActiveHandleFaces(const TArray<int>& FaceGroupIDs) override;
-	void SetActiveHandleEdges(const TArray<int>& TopologyEdgeIDs) override;
-	void SetActiveHandleCorners(const TArray<int>& TopologyCornerIDs) override;
-
-
-
-	/** Allocates shared storage for use in task synchronization */
-	void InitBackgroundWorker(const ELaplacianWeightScheme WeightScheme);
-
-
-	/** Coordinates the background tasks. Returns false if the worker was already running */
-	bool UpdateAndLaunchdWorker(const ELaplacianWeightScheme WeightScheme, const FRichCurve* Curve = nullptr);
-
-	/** Capture data about background task state.*/
-	bool IsTaskInFlight() const;
-
-
-
-	/** Sets the SrcMeshConstraintBuffer to have a size of MaxVertexID, and initializes with the current mesh positions, but weight zero*/
-	void InitializeConstraintBuffer();
-
-	/** Given an array of Group IDs, update the selection and record vertices */
-	void UpdateSelection(const FDynamicMesh3* TargetMesh, const TArray<int>& Groups, bool bLocalizeDeformation);
-
-	/** Updates the mesh preview and/or solvers upon user input, provided a deformation strategy */
-	void UpdateSolution(FDynamicMesh3* TargetMesh, const TFunction<FVector3d(FDynamicMesh3* Mesh, int)>& HandleVertexDeformFunc) override;
-
-	/** Updates the vertex positions of the mesh with the result from the last deformation solve. */
-	void ExportDeformedPositions(FDynamicMesh3* TargetMesh);
-
-	/** Returns true if the asynchronous task has finished. */
-	inline bool IsDone() { return AsyncMeshDeformTask == nullptr || AsyncMeshDeformTask->IsDone(); };
-
-	/** Triggers abort on task and passes off ownership to deleter object */
-	inline void Shutdown();
-
-	const TArray<FROIFace>& GetROIFaces() const { return ROIFaces; }
-
-	/** Stores the position of the vertex constraints and corresponding weights for the entire mesh.  This is used as a form of scratch space.*/
-	TArray<FConstraintData> SrcMeshConstraintBuffer;
-
-	/** Array of vertex indices organized in groups of three - basically an index buffer - that defines the subset of the mesh that the deformation task will work on.*/
-	TArray<int32> SubsetIDBuffer;
-
-	/** Need to update the task with the current submesh */
-	bool bTaskSubmeshIsDirty = true;
-
-	/** Asynchronous task object. This object deals with expensive matrix functionality that computes the deformation of a local mesh. */
-	FAsyncTaskExecuterWithAbort<FConstrainedMeshDeformerTask>* AsyncMeshDeformTask = nullptr;
-
-
-	/** The weight which will be applied to the constraints corresponding to the handle vertices. */
-	double HandleWeights = 1.0;
-
-	/** This is set to true whenever the user interacts with the tool under laplacian deformation mode.
-	  * It is set to false immediately before beginning a background task and cannot be set to false again until the work is done. */
-	bool bDeformerNeedsToRun = false;
-
-
-	/** When true, tells the solver to attempt to postfix the actual position of the handles to the constrained position */
-	bool bPostfixHandles = false;
-
-	//This is set to false only after 
-	//	1) the asynchronous deformation task is complete
-	//	2) the main thread has seen it complete, and
-	//	3) the main thread updates the vertex positions of the mesh one last time
-	bool bVertexPositionsNeedSync = false;
-
-	bool bLocalize =  true;
-
-};
-
-
 /**
  *
  */
 UCLASS()
-class MESHMODELINGTOOLS_API UEditMeshPolygonsTool : public UMeshSurfacePointTool
+class MESHMODELINGTOOLS_API UEditMeshPolygonsTool : public UMeshSurfacePointTool, public IClickBehaviorTarget
 {
 	GENERATED_BODY()
 
@@ -417,13 +120,20 @@ public:
 	virtual bool CanAccept() const override { return true; }
 
 
+	// UMeshSurfacePointTool API
 	virtual bool HitTest(const FRay& Ray, FHitResult& OutHit) override;
-	
 	virtual void OnBeginDrag(const FRay& Ray) override;
 	virtual void OnUpdateDrag(const FRay& Ray) override;
 	virtual void OnEndDrag(const FRay& Ray) override;
 	virtual bool OnUpdateHover(const FInputDeviceRay& DevicePos) override;
 
+	// IClickDragBehaviorTarget API
+	virtual FInputRayHit CanBeginClickDragSequence(const FInputDeviceRay& PressPos) override;
+
+	// IClickBehaviorTarget API
+	virtual FInputRayHit IsHitByClick(const FInputDeviceRay& ClickPos) override;
+	virtual void OnClicked(const FInputDeviceRay& ClickPos) override;
+
 public:
 	virtual void NextTransformTypeAction();
 
@@ -437,12 +147,24 @@ protected:
 	UPROPERTY()
 	UPolyEditTransformProperties* TransformProps;
 
+	UPROPERTY()
+	UMultiTransformer* MultiTransformer = nullptr;
+
+	TValueWatcher<EMultiTransformerMode> TransformerModeWatcher;
+	void UpdateTransformerMode(EMultiTransformerMode NewMode);
+
+	void OnMultiTransformerTransformBegin();
+	void OnMultiTransformerTransformUpdate();
+	void OnMultiTransformerTransformEnd();
+
 	// realtime visualization
 	void OnDynamicMeshComponentChanged();
 	FDelegateHandle OnDynamicMeshComponentChangedHandle;
 
 	virtual void OnPropertyModified(UObject* PropertySet, UProperty* Property);
 	
+
+
 	// camera state at last render
 	FViewCameraState CameraState;
 
@@ -476,13 +198,17 @@ protected:
 	FFrame3d RotationStartFrame;
 	void ComputeUpdate_Rotate();
 
+	FFrame3d InitialGizmoFrame;
+	void ComputeUpdate_Gizmo();
 
 	FGroupTopology Topology;
 	void PrecomputeTopology();
 	void ComputePolygons(bool RecomputeTopology = true);
 
 	FGroupTopologySelector TopoSelector;
-	
+	bool TopologyHitTest(const FRay& WorldRay, FHitResult& OutHit, FGroupTopologySelection& OutSelection);
+	void UpdateTopoSelector();
+
 	//
 	// data for current drag
 	//
@@ -490,16 +216,20 @@ protected:
 	FGroupTopologySelection HilightSelection;
 	FToolDataVisualizer HilightRenderer;
 
+	FGroupTopologySelection PersistentSelection;
+	FToolDataVisualizer SelectionRenderer;
+
+
 	FDynamicMeshAABBTree3 MeshSpatial;
 	FDynamicMeshAABBTree3& GetSpatial();
 
 	FMeshVertexChangeBuilder* ActiveVertexChange;
 
-	// The two deformer type options.
 	FGroupTopologyDeformer LinearDeformer;
-	FGroupTopologyLaplacianDeformer LaplacianDeformer;
+	void UpdateDeformerFromSelection(const FGroupTopologySelection& Selection);
+
+
 
-	EGroupTopologyDeformationStrategy DeformationStrategy;
 
 	// Initial polygon group and mesh info
 	TDynamicVector<int> InitialTriangleGroups;
diff --git a/Engine/Plugins/Experimental/MeshModelingToolset/Source/MeshModelingTools/Public/EditNormalsTool.h b/Engine/Plugins/Experimental/MeshModelingToolset/Source/MeshModelingTools/Public/EditNormalsTool.h
--- a/Engine/Plugins/Experimental/MeshModelingToolset/Source/MeshModelingTools/Public/EditNormalsTool.h
+++ b/Engine/Plugins/Experimental/MeshModelingToolset/Source/MeshModelingTools/Public/EditNormalsTool.h
@@ -53,11 +53,11 @@ public:
 
 	bool WillTopologyChange()
 	{
-		return bFixInconsistentNormals || bInvertNormals || bRecomputeNormalTopologyAndEdgeSharpness;
+		return bFixInconsistentNormals || bInvertNormals || SplitNormalMethod != ESplitNormalMethod::UseExistingTopology;
 	}
 
 	/** Recompute all mesh normals */
-	UPROPERTY(EditAnywhere, Category = NormalsCalculation, meta = (EditCondition = "!bRecomputeNormalTopologyAndEdgeSharpness"))
+	UPROPERTY(EditAnywhere, Category = NormalsCalculation, meta = (EditCondition = "SplitNormalMethod == ESplitNormalMethod::UseExistingTopology"))
 	bool bRecomputeNormals;
 
 	/** Choose the method for computing vertex normals */
@@ -72,16 +72,16 @@ public:
 	UPROPERTY(EditAnywhere, Category = NormalsCalculation)
 	bool bInvertNormals;
 
-	/** Choose where to create sharp edges by having multiple normals on the same vertices, using the change in normals of adjacent faces.  If enabled, normals will always be recomputed. */
+	/** Control whether and how the topology of the normals is recomputed, e.g. to create sharp edges where face normals change by a large amount or where face group IDs change.  Normals will always be recomputed unless SplitNormal Method is UseExistingTopology. */
 	UPROPERTY(EditAnywhere, Category = NormalsTopology)
-	bool bRecomputeNormalTopologyAndEdgeSharpness;
+	ESplitNormalMethod SplitNormalMethod;
 
 	/** Threshold on angle of change in face normals across an edge, above which we create a sharp edge if bSplitNormals is true */
-	UPROPERTY(EditAnywhere, Category = NormalsTopology, meta = (UIMin = "0.0", UIMax = "180.0", ClampMin = "0.0", ClampMax = "180.0", EditCondition = "bRecomputeNormalTopologyAndEdgeSharpness"))
+	UPROPERTY(EditAnywhere, Category = NormalsTopology, meta = (UIMin = "0.0", UIMax = "180.0", ClampMin = "0.0", ClampMax = "180.0", EditCondition = "SplitNormalMethod == ESplitNormalMethod::FaceNormalThreshold"))
 	float SharpEdgeAngleThreshold;
 
 	/** Assign separate normals at 'sharp' vertices -- for example, at the tip of a cone */
-	UPROPERTY(EditAnywhere, Category = NormalsTopology, meta = (EditCondition = "bRecomputeNormalTopologyAndEdgeSharpness"))
+	UPROPERTY(EditAnywhere, Category = NormalsTopology, meta = (EditCondition = "SplitNormalMethod == ESplitNormalMethod::FaceNormalThreshold"))
 	bool bAllowSharpVertices;
 };
 
@@ -112,7 +112,7 @@ class MESHMODELINGTOOLS_API UEditNormalsOperatorFactory : public UObject, public
 
 public:
 	// IDynamicMeshOperatorFactory API
-	virtual TSharedPtr<FDynamicMeshOperator> MakeNewOperator() override;
+	virtual TUniquePtr<FDynamicMeshOperator> MakeNewOperator() override;
 
 	UPROPERTY()
 	UEditNormalsTool *Tool;
@@ -177,5 +177,5 @@ protected:
 
 	void UpdateNumPreviews();
 
-	void GenerateAsset(const TArray<TUniquePtr<FDynamicMeshOpResult>>& Results);
+	void GenerateAsset(const TArray<FDynamicMeshOpResult>& Results);
 };
diff --git a/Engine/Plugins/Experimental/MeshModelingToolset/Source/MeshModelingTools/Public/EditPivotTool.h b/Engine/Plugins/Experimental/MeshModelingToolset/Source/MeshModelingTools/Public/EditPivotTool.h
new file mode 100644
--- /dev/null
+++ b/Engine/Plugins/Experimental/MeshModelingToolset/Source/MeshModelingTools/Public/EditPivotTool.h
@@ -0,0 +1,160 @@
+// Copyright 1998-2019 Epic Games, Inc. All Rights Reserved.
+
+#pragma once
+
+#include "CoreMinimal.h"
+#include "MultiSelectionTool.h"
+#include "InteractiveToolBuilder.h"
+#include "BaseBehaviors//BehaviorTargetInterfaces.h"
+#include "Changes/TransformChange.h"
+#include "FrameTypes.h"
+#include "EditPivotTool.generated.h"
+
+class UBaseAxisTranslationGizmo;
+class UAxisAngleGizmo;
+class UTransformGizmo;
+class UTransformProxy;
+
+/**
+ *
+ */
+UCLASS()
+class MESHMODELINGTOOLS_API UEditPivotToolBuilder : public UInteractiveToolBuilder
+{
+	GENERATED_BODY()
+
+public:
+	virtual bool CanBuildTool(const FToolBuilderState& SceneState) const override;
+	virtual UInteractiveTool* BuildTool(const FToolBuilderState& SceneState) const override;
+};
+
+
+
+
+/** Snap-Drag Source Point */
+UENUM()
+enum class EEditPivotSnapDragSource : uint8
+{
+	/** Snap-Drag moves the Clicked Point to the Target Location */
+	ClickPoint = 0 UMETA(DisplayName = "Click Point"),
+
+	/** Snap-Drag moves the Gizmo/Pivot to the Target Location */
+	Pivot = 1 UMETA(DisplayName = "Pivot"),
+
+
+	LastValue UMETA(Hidden)
+
+};
+
+
+
+/** Snap-Drag Rotation Mode */
+UENUM()
+enum class EEditPivotSnapDragRotationMode : uint8
+{
+	/** Snap-Drag only translates, ignoring Normals */
+	Ignore = 0 UMETA(DisplayName = "Ignore"),
+
+	/** Snap-Drag aligns the Source and Target Normals to point in the same direction */
+	Align = 1 UMETA(DisplayName = "Align"),
+
+	/** Snap-Drag aligns the Source Normal to the opposite of the Target Normal direction */
+	AlignFlipped = 2 UMETA(DisplayName = "Align Flipped"),
+
+	LastValue UMETA(Hidden)
+};
+
+
+
+/**
+ * Standard properties of the Transform Meshes operation
+ */
+UCLASS()
+class MESHMODELINGTOOLS_API UEditPivotToolProperties : public UInteractiveToolPropertySet
+{
+	GENERATED_BODY()
+public:
+
+	/** Click-drag starting on the target objects to reposition them on the rest of the scene */
+	UPROPERTY(EditAnywhere, Category = Options)
+	bool bEnableSnapDragging = false;
+
+
+	UPROPERTY(EditAnywhere, Category = Options, meta = (EditCondition = "bEnableSnapDragging == true"))
+	EEditPivotSnapDragSource SnapDragSource = EEditPivotSnapDragSource::ClickPoint;
+
+	/** When Snap-Dragging, align source and target normals */
+	UPROPERTY(EditAnywhere, Category = Options, meta = (EditCondition = "bEnableSnapDragging == true"))
+	EEditPivotSnapDragRotationMode RotationMode = EEditPivotSnapDragRotationMode::AlignFlipped;
+};
+
+
+USTRUCT()
+struct FEditPivotTarget
+{
+	GENERATED_BODY()
+
+	UPROPERTY()
+	UTransformProxy* TransformProxy;
+
+	UPROPERTY()
+	UTransformGizmo* TransformGizmo;
+};
+
+
+/**
+ *
+ */
+UCLASS()
+class MESHMODELINGTOOLS_API UEditPivotTool : public UMultiSelectionTool, public IClickDragBehaviorTarget
+{
+	GENERATED_BODY()
+
+public:
+	UEditPivotTool();
+
+	virtual void RegisterActions(FInteractiveToolActionSet& ActionSet) override;
+
+	virtual void SetWorld(UWorld* World, UInteractiveGizmoManager* GizmoManager);
+
+	virtual void Setup() override;
+	virtual void Shutdown(EToolShutdownType ShutdownType) override;
+
+	virtual void Tick(float DeltaTime) override;
+	virtual void Render(IToolsContextRenderAPI* RenderAPI) override;
+
+	virtual bool HasCancel() const override { return true; }
+	virtual bool HasAccept() const override { return true; }
+	virtual bool CanAccept() const override { return true; }
+
+	virtual void OnPropertyModified(UObject* PropertySet, UProperty* Property) override;
+
+
+	// ICLickDragBehaviorTarget interface
+	virtual FInputRayHit CanBeginClickDragSequence(const FInputDeviceRay& PressPos) override;
+	virtual void OnClickPress(const FInputDeviceRay& PressPos) override;
+	virtual void OnClickDrag(const FInputDeviceRay& DragPos) override;
+	virtual void OnClickRelease(const FInputDeviceRay& ReleasePos) override;
+	virtual void OnTerminateDragSequence() override;
+
+protected:
+	UPROPERTY()
+	UEditPivotToolProperties* TransformProps;
+
+protected:
+	UWorld* TargetWorld;
+	UInteractiveGizmoManager* GizmoManager;
+
+	UPROPERTY()
+	TArray<FEditPivotTarget> ActiveGizmos;
+
+	void UpdateSetPivotModes(bool bEnableSetPivot);
+	void SetActiveGizmos_Single(bool bLocalRotations);
+	void ResetActiveGizmos();
+
+	FFrame3d StartDragFrameWorld;
+	FTransform StartDragTransform;
+	int ActiveSnapDragIndex = -1;
+
+	void UpdateAssets(const FFrame3d& NewPivotWorldFrame);
+};
diff --git a/Engine/Plugins/Experimental/MeshModelingToolset/Source/MeshModelingTools/Public/FindPolygonsAlgorithm.h b/Engine/Plugins/Experimental/MeshModelingToolset/Source/MeshModelingTools/Public/FindPolygonsAlgorithm.h
--- a/Engine/Plugins/Experimental/MeshModelingToolset/Source/MeshModelingTools/Public/FindPolygonsAlgorithm.h
+++ b/Engine/Plugins/Experimental/MeshModelingToolset/Source/MeshModelingTools/Public/FindPolygonsAlgorithm.h
@@ -13,14 +13,21 @@ public:
 	FFindPolygonsAlgorithm(FDynamicMesh3* MeshIn);
 
 	FDynamicMesh3* Mesh = nullptr;
-	TArray<int> PolygonGroupIDs;
 	TArray<TArray<int>> FoundPolygons;
 	TArray<int> PolygonTags;
-	TArray<FVector> PolygonNormals;
+	TArray<FVector3d> PolygonNormals;
 
 	TArray<int> PolygonEdges;
 
-	bool FindPolygons(double DotTolerance = 0.0001);
+	bool FindPolygonsFromFaceNormals(double DotTolerance = 0.0001);
+	bool FindPolygonsFromUVIslands();
 
 	bool FindPolygonEdges();
+
+
+protected:
+
+	TArray<int> PolygonGroupIDs;
+
+	void SetGroupsFromPolygons();
 };
\ No newline at end of file
diff --git a/Engine/Plugins/Experimental/MeshModelingToolset/Source/MeshModelingTools/Public/MeshSelectionTool.h b/Engine/Plugins/Experimental/MeshModelingToolset/Source/MeshModelingTools/Public/MeshSelectionTool.h
--- a/Engine/Plugins/Experimental/MeshModelingToolset/Source/MeshModelingTools/Public/MeshSelectionTool.h
+++ b/Engine/Plugins/Experimental/MeshModelingToolset/Source/MeshModelingTools/Public/MeshSelectionTool.h
@@ -39,7 +39,11 @@ enum class EMeshSelectionToolActions
 	ExpandToConnected,
 
 	DeleteSelected,
-	SeparateSelected
+	DisconnectSelected,
+	SeparateSelected,
+	FlipSelected,
+	CreateGroup,
+	AssignMaterial
 };
 
 
@@ -66,32 +70,32 @@ class MESHMODELINGTOOLS_API UMeshSelectionEditActions : public UMeshSelectionToo
 	GENERATED_BODY()
 
 public:
-	UFUNCTION(CallInEditor, Category = Selection, meta = (DisplayName = "Clear", DisplayPriority = 1))
+	UFUNCTION(CallInEditor, Category = SelectionEdits, meta = (DisplayName = "Clear", DisplayPriority = 1))
 	void Clear()
 	{
 		PostAction(EMeshSelectionToolActions::ClearSelection);
 	}
 
-	UFUNCTION(CallInEditor, Category = Selection, meta = (DisplayName = "Invert", DisplayPriority = 2))
+	UFUNCTION(CallInEditor, Category = SelectionEdits, meta = (DisplayName = "Invert", DisplayPriority = 2))
 	void Invert()
 	{
 		PostAction(EMeshSelectionToolActions::InvertSelection);
 	}
 
 
-	UFUNCTION(CallInEditor, Category = Selection, meta = (DisplayName = "Grow", DisplayPriority = 3))
+	UFUNCTION(CallInEditor, Category = SelectionEdits, meta = (DisplayName = "Grow", DisplayPriority = 3))
 	void Grow()
 	{
 		PostAction(EMeshSelectionToolActions::GrowSelection);
 	}
 
-	UFUNCTION(CallInEditor, Category = Selection, meta = (DisplayName = "Shrink", DisplayPriority = 4))
+	UFUNCTION(CallInEditor, Category = SelectionEdits, meta = (DisplayName = "Shrink", DisplayPriority = 4))
 	void Shrink()
 	{
 		PostAction(EMeshSelectionToolActions::ShrinkSelection);
 	}
 
-	UFUNCTION(CallInEditor, Category = Selection, meta = (DisplayName = "ExpandToConnected", DisplayPriority = 5))
+	UFUNCTION(CallInEditor, Category = SelectionEdits, meta = (DisplayName = "ExpandToConnected", DisplayPriority = 5))
 	void ExpandToConnected()
 	{
 		PostAction(EMeshSelectionToolActions::ExpandToConnected);
@@ -108,17 +112,35 @@ class MESHMODELINGTOOLS_API UMeshSelectionMeshEditActions : public UMeshSelectio
 	GENERATED_BODY()
 
 public:
-	UFUNCTION(CallInEditor, Category = MeshEdits, meta = (DisplayName = "Delete"))
+	UFUNCTION(CallInEditor, Category = MeshEdits, meta = (DisplayName = "Delete", DisplayPriority = 1))
 	void DeleteTriangles()
 	{
 		PostAction(EMeshSelectionToolActions::DeleteSelected);
 	}
 
-	UFUNCTION(CallInEditor, Category = MeshEdits, meta = (DisplayName = "Separate"))
+	UFUNCTION(CallInEditor, Category = MeshEdits, meta = (DisplayName = "Separate", DisplayPriority = 2))
 	void SeparateTriangles() 
 	{
 		PostAction(EMeshSelectionToolActions::SeparateSelected);
 	}
+
+	UFUNCTION(CallInEditor, Category = MeshEdits, meta = (DisplayName = "Disconnect", DisplayPriority = 3))
+	void DisconnectTriangles() 
+	{
+		PostAction(EMeshSelectionToolActions::DisconnectSelected);
+	}
+
+	UFUNCTION(CallInEditor, Category = MeshEdits, meta = (DisplayName = "Flip Normals", DisplayPriority = 4))
+	void FlipNormals() 
+	{
+		PostAction(EMeshSelectionToolActions::FlipSelected);
+	}
+
+	UFUNCTION(CallInEditor, Category = MeshEdits, meta = (DisplayName = "Create Polygroup", DisplayPriority = 5))
+	void CreatePolygroup()
+	{
+		PostAction(EMeshSelectionToolActions::CreateGroup);
+	}
 };
 
 
@@ -144,12 +166,32 @@ enum class EMeshSelectionToolPrimaryMode
 	/** Select all triangles in groups connected to any triangle inside the brush */
 	AllInGroup,
 
+	/** Select all triangles with same material as hit triangle */
+	ByMaterial,
+
+	/** Select all triangles in same UV island as hit triangle */
+	ByUVIsland,
+
 	/** Select all triangles with normal within angular tolerance of hit triangle */
 	AllWithinAngle
 };
 
 
 
+UENUM()
+enum class EMeshFacesColorMode
+{
+	/** Display original mesh materials */
+	None,
+	/** Color mesh triangles by PolyGroup Color */
+	ByGroup,
+	/** Color mesh triangles by Material ID */
+	ByMaterialID,
+	/** Color mesh triangles by UV Island */
+	ByUVIsland
+};
+
+
 UCLASS()
 class MESHMODELINGTOOLS_API UMeshSelectionToolProperties : public UInteractiveToolPropertySet
 {
@@ -176,6 +218,9 @@ public:
 	UPROPERTY(EditAnywhere, Category = Selection)
 	bool bShowWireframe = true;
 
+	UPROPERTY(EditAnywhere, Category = Selection)
+	EMeshFacesColorMode FaceColorMode = EMeshFacesColorMode::None;
+
 	virtual void SaveProperties(UInteractiveTool* SaveFromTool) override;
 	virtual void RestoreProperties(UInteractiveTool* RestoreToTool) override;
 };
@@ -225,7 +270,12 @@ public:
 	UMeshSelectionEditActions* SelectionActions;
 
 	UPROPERTY()
-	UMeshSelectionMeshEditActions* EditActions;
+	UMeshSelectionToolActionPropertySet* EditActions;
+
+
+protected:
+	virtual UMeshSelectionToolActionPropertySet* CreateEditActions();
+	virtual void AddSubclassPropertySets() {}
 
 protected:
 
@@ -248,6 +298,7 @@ protected:
 	EMeshSelectionElementType SelectionType = EMeshSelectionElementType::Face;
 
 	TValueWatcher<bool> ShowWireframeWatcher;
+	TValueWatcher<EMeshFacesColorMode> ColorModeWatcher;
 
 
 	bool bInRemoveStroke = false;
@@ -270,7 +321,12 @@ protected:
 	void OnExternalSelectionChange();
 
 	void OnSelectionUpdated();
-	void UpdateVisualization();
+	void UpdateVisualization(bool bSelectionModified);
+	bool bFullMeshInvalidationPending = false;
+	bool bColorsUpdatePending = false;
+	FColor GetCurrentFaceColor(const FDynamicMesh3* Mesh, int TriangleID);
+	void CacheUVIslandIDs();
+	TArray<int> TriangleToUVIsland;
 
 	// selection change
 	FMeshSelectionChangeBuilder* ActiveSelectionChange = nullptr;
@@ -291,8 +347,12 @@ protected:
 	void ExpandToConnected();
 
 	void DeleteSelectedTriangles();
-	void SeparateSelectedTriangles();
+	void DisconnectSelectedTriangles(); // disconnects edges between selected and not-selected triangles; keeps all triangles in the same mesh
+	void SeparateSelectedTriangles(); // separates out selected triangles to a new mesh, removing them from the current mesh
+	void FlipSelectedTriangles();
+	void AssignNewGroupToSelectedTriangles();
 
+	// if true, mesh has been edited
 	bool bHaveModifiedMesh = false;
 };
 
diff --git a/Engine/Plugins/Experimental/MeshModelingToolset/Source/MeshModelingTools/Public/PlaneCutTool.h b/Engine/Plugins/Experimental/MeshModelingToolset/Source/MeshModelingTools/Public/PlaneCutTool.h
--- a/Engine/Plugins/Experimental/MeshModelingToolset/Source/MeshModelingTools/Public/PlaneCutTool.h
+++ b/Engine/Plugins/Experimental/MeshModelingToolset/Source/MeshModelingTools/Public/PlaneCutTool.h
@@ -5,10 +5,11 @@
 #include "CoreMinimal.h"
 #include "UObject/NoExportTypes.h"
 #include "InteractiveGizmo.h"
-#include "SingleSelectionTool.h"
+#include "MultiSelectionTool.h"
 #include "InteractiveToolBuilder.h"
 #include "MeshOpPreviewHelpers.h"
 #include "DynamicMesh3.h"
+#include "Changes/DynamicMeshChangeTarget.h"
 #include "BaseTools/SingleClickTool.h"
 #include "PlaneCutTool.generated.h"
 
@@ -40,6 +41,24 @@ public:
 
 
 
+/**
+* Properties controlling how changes are baked out to static meshes on tool accept
+*/
+UCLASS()
+class MESHMODELINGTOOLS_API UAcceptOutputProperties : public UInteractiveToolPropertySet
+{
+	GENERATED_BODY()
+
+public:
+
+	/** If false, tool cannot 'split' any input mesh into multiple meshes, and will bake all results that came from one mesh into that mesh.  If true, the tool can generate additional assets when meshes are logically split into multiple parts. */
+	UPROPERTY(EditAnywhere, Category = ToolOutputOptions)
+	bool bAllowLogicalMeshSplitsToExportAsNewMeshAssets = true;
+};
+
+
+
+
 
 
 /**
@@ -53,16 +72,12 @@ class MESHMODELINGTOOLS_API UPlaneCutToolProperties : public UInteractiveToolPro
 public:
 	UPlaneCutToolProperties();
 
-	/** If true, UVs and Normals are discarded  */
-	UPROPERTY(EditAnywhere, Category = Options)
-	bool bDiscardAttributes;
-
 	/** If true, both halves of the cut are computed */
 	UPROPERTY(EditAnywhere, Category = Options)
 	bool bKeepBothHalves;
 
 	/** If keeping both halves, separate the two pieces by this amount */
-	UPROPERTY(EditAnywhere, Category = Options, meta = (EditCondition = "bKeepBothHalves == true") )
+	UPROPERTY(EditAnywhere, Category = Options, meta = (EditCondition = "bKeepBothHalves == true", UIMin = "0", ClampMin = "0") )
 	float SpacingBetweenHalves;
 
 	/** If true, the cut surface is filled with simple planar hole fill surface(s) */
@@ -86,20 +101,19 @@ class MESHMODELINGTOOLS_API UPlaneCutOperatorFactory : public UObject, public ID
 
 public:
 	// IDynamicMeshOperatorFactory API
-	virtual TSharedPtr<FDynamicMeshOperator> MakeNewOperator() override;
+	virtual TUniquePtr<FDynamicMeshOperator> MakeNewOperator() override;
 
 	UPROPERTY()
 	UPlaneCutTool *CutTool;
 
-	UPROPERTY()
-	bool bCutBackSide = false;
+	int ComponentIndex;
 };
 
 /**
  * Simple Mesh Plane Cutting Tool
  */
 UCLASS()
-class MESHMODELINGTOOLS_API UPlaneCutTool : public USingleSelectionTool
+class MESHMODELINGTOOLS_API UPlaneCutTool : public UMultiSelectionTool
 {
 	GENERATED_BODY()
 
@@ -128,11 +142,15 @@ public:
 
 	virtual void OnPropertyModified(UObject* PropertySet, UProperty* Property) override;
 
+
 protected:
 
 	UPROPERTY()
 	UPlaneCutToolProperties* BasicProperties;
 
+	UPROPERTY()
+	UAcceptOutputProperties* AcceptProperties;
+
 	/** Origin of cutting plane */
 	UPROPERTY()
 	FVector CutPlaneOrigin;
@@ -144,9 +162,19 @@ protected:
 	UPROPERTY()
 	TArray<UMeshOpPreviewWithBackgroundCompute*> Previews;
 
+	/** Do the plane cut without exiting the tool, useful for doing a lot of cuts quickly */
+	UFUNCTION(CallInEditor, Category = Options, meta = (DisplayName = "Cut"))
+	void DoCut();
 
 protected:
-	TSharedPtr<FDynamicMesh3> OriginalDynamicMesh;
+
+	UPROPERTY()
+	TArray<UDynamicMeshReplacementChangeTarget*> MeshesToCut;
+
+	// for each mesh in MeshesToCut, the index of the attached generic triangle attribute tracking the object index
+	TArray<int> MeshSubObjectAttribIndices;
+	// UV Scale factor is cached based on the bounding box of the mesh before any cuts are performed, so you don't get inconsistent UVs if you multi-cut the object to smaller sizes
+	TArray<float> MeshUVScaleFactor;
 
 	UWorld* TargetWorld;
 	IToolsContextAssetAPI* AssetAPI;
@@ -160,12 +188,13 @@ protected:
 	UTransformProxy* PlaneTransformProxy;
 
 	void TransformChanged(UTransformProxy* Proxy, FTransform Transform);
+	void MeshChanged();
 
-	void UpdateNumPreviews();
+	void SetupPreviews();
 
 	IClickBehaviorTarget* SetPointInWorldConnector = nullptr;
 
 	virtual void SetCutPlaneFromWorldPos(const FVector& Position, const FVector& Normal);
 
-	void GenerateAsset(const TArray<TUniquePtr<FDynamicMeshOpResult>>& Results);
+	void GenerateAsset(const TArray<FDynamicMeshOpResult>& Results);
 };
diff --git a/Engine/Plugins/Experimental/MeshModelingToolset/Source/MeshModelingTools/Public/RemeshMeshTool.h b/Engine/Plugins/Experimental/MeshModelingToolset/Source/MeshModelingTools/Public/RemeshMeshTool.h
--- a/Engine/Plugins/Experimental/MeshModelingToolset/Source/MeshModelingTools/Public/RemeshMeshTool.h
+++ b/Engine/Plugins/Experimental/MeshModelingToolset/Source/MeshModelingTools/Public/RemeshMeshTool.h
@@ -128,7 +128,7 @@ public:
 	virtual void OnPropertyModified(UObject* PropertySet, UProperty* Property) override;
 
 	// IDynamicMeshOperatorFactory API
-	virtual TSharedPtr<FDynamicMeshOperator> MakeNewOperator() override;
+	virtual TUniquePtr<FDynamicMeshOperator> MakeNewOperator() override;
 
 protected:
 	UPROPERTY()
diff --git a/Engine/Plugins/Experimental/MeshModelingToolset/Source/MeshModelingToolsEditorOnly/Private/ParameterizeMeshTool.cpp b/Engine/Plugins/Experimental/MeshModelingToolset/Source/MeshModelingToolsEditorOnly/Private/ParameterizeMeshTool.cpp
--- a/Engine/Plugins/Experimental/MeshModelingToolset/Source/MeshModelingToolsEditorOnly/Private/ParameterizeMeshTool.cpp
+++ b/Engine/Plugins/Experimental/MeshModelingToolset/Source/MeshModelingToolsEditorOnly/Private/ParameterizeMeshTool.cpp
@@ -46,13 +46,36 @@ UInteractiveTool* UParameterizeMeshToolBuilder::BuildTool(const FToolBuilderStat
 }
 
 
+
+
+
+void UParameterizeMeshToolProperties::SaveProperties(UInteractiveTool* SaveFromTool)
+{
+	UParameterizeMeshToolProperties* PropertyCache = GetPropertyCache<UParameterizeMeshToolProperties>();
+	PropertyCache->ChartStretch = this->ChartStretch;
+	PropertyCache->bRespectPolygroups = this->bRespectPolygroups;
+	PropertyCache->UVScaleMode = this->UVScaleMode;
+	PropertyCache->UVScale = this->UVScale;
+}
+
+void UParameterizeMeshToolProperties::RestoreProperties(UInteractiveTool* RestoreToTool)
+{
+	UParameterizeMeshToolProperties* PropertyCache = GetPropertyCache<UParameterizeMeshToolProperties>();
+	this->ChartStretch = PropertyCache->ChartStretch;
+	this->bRespectPolygroups = PropertyCache->bRespectPolygroups;
+	this->UVScaleMode = PropertyCache->UVScaleMode;
+	this->UVScale = PropertyCache->UVScale;
+}
+
+
+
+
 /*
  * Tool
  */
 
 UParameterizeMeshTool::UParameterizeMeshTool()
 {
-	ChartStretch = 0.11f;
 }
 
 void UParameterizeMeshTool::SetWorld(UWorld* World)
@@ -80,14 +103,13 @@ void UParameterizeMeshTool::Setup()
 		DefaultMaterial = LoadObject<UMaterial>(nullptr, TEXT("/Engine/EngineMaterials/DefaultMaterial"));  
 	}
 
-	
-
 	// hide input StaticMeshComponent
 	ComponentTarget->SetOwnerVisibility(false);
 
 	// initialize our properties
-	ToolPropertyObjects.Add(this);
-
+	Settings = NewObject<UParameterizeMeshToolProperties>(this);
+	Settings->RestoreProperties(this);
+	AddToolPropertySource(Settings);
 	
 	// Construct the preview object and set the material on it
 	Preview = NewObject<UMeshOpPreviewWithBackgroundCompute>(this, "Preview");
@@ -100,18 +122,24 @@ void UParameterizeMeshTool::Setup()
 		Converter.bPrintDebugMessages = false;
 		Converter.Convert(InputMesh.Get(), Mesh);
 
+		FComponentMaterialSet MaterialSet;
+		ComponentTarget->GetMaterialSet(MaterialSet);
+		Preview->ConfigureMaterials(MaterialSet.Materials,
+			ToolSetupUtil::GetDefaultWorkingMaterial(GetToolManager())
+		);
+
 		Preview->PreviewMesh->UpdatePreview(&Mesh);
 		Preview->PreviewMesh->SetTransform(ComponentTarget->GetWorldTransform());
 	}
-	Preview->ConfigureMaterials(
-		ToolSetupUtil::GetDefaultMaterial(GetToolManager(), DefaultMaterial),
-		ToolSetupUtil::GetDefaultWorkingMaterial(GetToolManager())
-	);
 
 
 	MaterialSettings = NewObject<UExistingMeshMaterialProperties>(this);
 	MaterialSettings->Setup();
+	MaterialSettings->MaterialMode = ESetMeshMaterialMode::Checkerboard;
 	AddToolPropertySource(MaterialSettings);
+	// force update
+	MaterialSettings->UpdateMaterials();
+	Preview->OverrideMaterial = MaterialSettings->GetActiveOverrideMaterial();
 
 
 	Preview->SetVisibility(true);
@@ -123,28 +151,31 @@ void UParameterizeMeshTool::OnPropertyModified(UObject* PropertySet, UProperty*
 	if (PropertySet == MaterialSettings)
 	{
 		MaterialSettings->UpdateMaterials();
-		MaterialSettings->SetMaterialIfChanged(ComponentTarget->GetMaterial(0), Preview->StandardMaterial, [this](UMaterialInterface* Material)
-		{
-			Preview->ConfigureMaterials(ToolSetupUtil::GetDefaultMaterial(GetToolManager(), Material), Preview->WorkingMaterial);
-		});
+		Preview->OverrideMaterial = MaterialSettings->GetActiveOverrideMaterial();
+	}
+	if (PropertySet == Settings)
+	{
+		// One of the UV generation properties must have changed.  Dirty the result to force a recompute
+		Preview->InvalidateResult();
 	}
 }
 
 
 void UParameterizeMeshTool::Shutdown(EToolShutdownType ShutdownType)
 {
+	Settings->SaveProperties(this);
 
-	TUniquePtr<FDynamicMeshOpResult> Result = Preview->Shutdown();
+	FDynamicMeshOpResult Result = Preview->Shutdown();
 	if (ShutdownType == EToolShutdownType::Accept)
 	{
-		FDynamicMesh3* DynamicMeshResult = Result->Mesh.Get();
+		FDynamicMesh3* DynamicMeshResult = Result.Mesh.Get();
 		check(DynamicMeshResult != nullptr);
 		GetToolManager()->BeginUndoTransaction(LOCTEXT("ParameterizeMesh", "Parameterize Mesh"));
 
-		ComponentTarget->CommitMesh([DynamicMeshResult](FMeshDescription* MeshDescription)
+		ComponentTarget->CommitMesh([DynamicMeshResult](const FPrimitiveComponentTarget::FCommitParams& CommitParams)
 		{
 			FDynamicMeshToMeshDescription Converter;
-			Converter.Convert(DynamicMeshResult, *MeshDescription);
+			Converter.Convert(DynamicMeshResult, *CommitParams.MeshDescription);
 		});
 		GetToolManager()->EndUndoTransaction();
 	}
@@ -173,18 +204,30 @@ bool UParameterizeMeshTool::CanAccept() const
 	return Preview->HaveValidResult();
 }
 
-void UParameterizeMeshTool::PostEditChangeProperty(FPropertyChangedEvent& PropertyChangedEvent)
+TUniquePtr<FDynamicMeshOperator> UParameterizeMeshTool::MakeNewOperator()
 {
-	// One of the UV generation properties must have changed.  Dirty the result to force a recompute
-	Preview->InvalidateResult();
-}
-
-TSharedPtr<FDynamicMeshOperator> UParameterizeMeshTool::MakeNewOperator()
-{
-	TSharedPtr<FParameterizeMeshOp> ParamertizeMeshOp = MakeShared<FParameterizeMeshOp>();
-	ParamertizeMeshOp->Stretch   = ChartStretch;
+	FAxisAlignedBox3d MeshBounds = Preview->PreviewMesh->GetMesh()->GetBounds();
+	TUniquePtr<FParameterizeMeshOp> ParamertizeMeshOp = MakeUnique<FParameterizeMeshOp>();
+	ParamertizeMeshOp->Stretch   = Settings->ChartStretch;
 	ParamertizeMeshOp->NumCharts = 0;
 	ParamertizeMeshOp->InputMesh = InputMesh;
+	ParamertizeMeshOp->bRespectPolygroups = Settings->bRespectPolygroups;
+
+	switch (Settings->UVScaleMode)
+	{
+		case EParameterizeMeshToolUVScaleMode::NoScaling:
+			ParamertizeMeshOp->bNormalizeAreas = false;
+			ParamertizeMeshOp->AreaScaling = 1.0;
+			break;
+		case EParameterizeMeshToolUVScaleMode::NormalizeToBounds:
+			ParamertizeMeshOp->bNormalizeAreas = true;
+			ParamertizeMeshOp->AreaScaling = Settings->UVScale / MeshBounds.MaxDim();
+			break;
+		case EParameterizeMeshToolUVScaleMode::NormalizeToWorld:
+			ParamertizeMeshOp->bNormalizeAreas = true;
+			ParamertizeMeshOp->AreaScaling = Settings->UVScale;
+			break;
+	}
 
 	const FTransform XForm = ComponentTarget->GetWorldTransform();
 	FTransform3d XForm3d(XForm);
diff --git a/Engine/Plugins/Experimental/MeshModelingToolset/Source/MeshModelingToolsEditorOnly/Private/ShapeSprayTool.cpp b/Engine/Plugins/Experimental/MeshModelingToolset/Source/MeshModelingToolsEditorOnly/Private/ShapeSprayTool.cpp
--- a/Engine/Plugins/Experimental/MeshModelingToolset/Source/MeshModelingToolsEditorOnly/Private/ShapeSprayTool.cpp
+++ b/Engine/Plugins/Experimental/MeshModelingToolset/Source/MeshModelingToolsEditorOnly/Private/ShapeSprayTool.cpp
@@ -155,7 +155,7 @@ void UShapeSprayTool::OnUpdateDrag(const FRay& Ray)
 		FVector3f WorldPt = WorldFrame.PointAt(PlanePt);
 		FVector3f SampleRayDir = WorldPt - (FVector3f)Ray.Origin; 
 		SampleRayDir.Normalize();
-		FRay WorldRay(Ray.Origin, SampleRayDir);
+		FRay WorldRay((FVector)Ray.Origin, (FVector)SampleRayDir);
 		FHitResult Hit;
 		if (HitTest(WorldRay, Hit))
 		{
diff --git a/Engine/Plugins/Experimental/MeshModelingToolset/Source/MeshModelingToolsEditorOnly/Private/SimplifyMeshTool.cpp b/Engine/Plugins/Experimental/MeshModelingToolset/Source/MeshModelingToolsEditorOnly/Private/SimplifyMeshTool.cpp
--- a/Engine/Plugins/Experimental/MeshModelingToolset/Source/MeshModelingToolsEditorOnly/Private/SimplifyMeshTool.cpp
+++ b/Engine/Plugins/Experimental/MeshModelingToolset/Source/MeshModelingToolsEditorOnly/Private/SimplifyMeshTool.cpp
@@ -99,8 +99,9 @@ void USimplifyMeshTool::Setup()
 	ComponentTarget->SetOwnerVisibility(false);
 	Preview = NewObject<UMeshOpPreviewWithBackgroundCompute>(this, "Preview");
 	Preview->Setup(this->TargetWorld, this);
-	Preview->ConfigureMaterials(
-		ToolSetupUtil::GetDefaultMaterial(GetToolManager(), ComponentTarget->GetMaterial(0)),
+	FComponentMaterialSet MaterialSet;
+	ComponentTarget->GetMaterialSet(MaterialSet);
+	Preview->ConfigureMaterials( MaterialSet.Materials,
 		ToolSetupUtil::GetDefaultWorkingMaterial(GetToolManager())
 	);
 	Preview->PreviewMesh->EnableWireframe(true);
@@ -154,10 +155,10 @@ void USimplifyMeshTool::Setup()
 void USimplifyMeshTool::Shutdown(EToolShutdownType ShutdownType)
 {
 	ComponentTarget->SetOwnerVisibility(true);
-	TUniquePtr<FDynamicMeshOpResult> Result = Preview->Shutdown();
+	FDynamicMeshOpResult Result = Preview->Shutdown();
 	if (ShutdownType == EToolShutdownType::Accept)
 	{
-		GenerateAsset(*Result);
+		GenerateAsset(Result);
 	}
 }
 
@@ -167,9 +168,9 @@ void USimplifyMeshTool::Tick(float DeltaTime)
 	Preview->Tick(DeltaTime);
 }
 
-TSharedPtr<FDynamicMeshOperator> USimplifyMeshTool::MakeNewOperator()
+TUniquePtr<FDynamicMeshOperator> USimplifyMeshTool::MakeNewOperator()
 {
-	TSharedPtr<FSimplifyMeshOp> Op = MakeShared<FSimplifyMeshOp>();
+	TUniquePtr<FSimplifyMeshOp> Op = MakeUnique<FSimplifyMeshOp>();
 
 	Op->bDiscardAttributes = SimplifyProperties->bDiscardAttributes;
 	Op->bPreventNormalFlips = SimplifyProperties->bPreventNormalFlips;
@@ -212,7 +213,7 @@ void USimplifyMeshTool::Render(IToolsContextRenderAPI* RenderAPI)
 			{
 				FVector3d A, B;
 				TargetMesh->GetEdgeV(eid, A, B);
-				PDI->DrawLine(Transform.TransformPosition(A), Transform.TransformPosition(B),
+				PDI->DrawLine(Transform.TransformPosition((FVector)A), Transform.TransformPosition((FVector)B),
 					LineColor, 0, 2.0, 1.0f, true);
 			}
 		}
@@ -240,12 +241,12 @@ void USimplifyMeshTool::GenerateAsset(const FDynamicMeshOpResult& Result)
 	GetToolManager()->BeginUndoTransaction(LOCTEXT("SimplifyMeshToolTransactionName", "Simplify Mesh"));
 
 	check(Result.Mesh.Get() != nullptr);
-	ComponentTarget->CommitMesh([&Result](FMeshDescription* MeshDescription)
+	ComponentTarget->CommitMesh([&Result](const FPrimitiveComponentTarget::FCommitParams& CommitParams)
 	{
 		FDynamicMeshToMeshDescription Converter;
 
 		// full conversion if normal topology changed or faces were inverted
-		Converter.Convert(Result.Mesh.Get(), *MeshDescription);
+		Converter.Convert(Result.Mesh.Get(), *CommitParams.MeshDescription);
 	});
 
 	GetToolManager()->EndUndoTransaction();
diff --git a/Engine/Plugins/Experimental/MeshModelingToolset/Source/MeshModelingToolsEditorOnly/Public/EditMeshMaterialsTool.h b/Engine/Plugins/Experimental/MeshModelingToolset/Source/MeshModelingToolsEditorOnly/Public/EditMeshMaterialsTool.h
new file mode 100644
--- /dev/null
+++ b/Engine/Plugins/Experimental/MeshModelingToolset/Source/MeshModelingToolsEditorOnly/Public/EditMeshMaterialsTool.h
@@ -0,0 +1,141 @@
+// Copyright 1998-2019 Epic Games, Inc. All Rights Reserved.
+
+#pragma once
+
+#include "MeshSelectionTool.h"
+#include "EditMeshMaterialsTool.generated.h"
+
+
+/**
+ *
+ */
+UCLASS()
+class MESHMODELINGTOOLSEDITORONLY_API UEditMeshMaterialsToolBuilder : public UMeshSelectionToolBuilder
+{
+	GENERATED_BODY()
+
+public:
+	virtual UMeshSurfacePointTool* CreateNewTool(const FToolBuilderState& SceneState) const override;
+};
+
+
+
+
+
+
+UCLASS()
+class MESHMODELINGTOOLSEDITORONLY_API UEditMeshMaterialsToolProperties : public UInteractiveToolPropertySet
+{
+	GENERATED_BODY()
+public:
+
+	UPROPERTY(EditAnywhere, Category = Materials, meta = (ArrayClamp = "Materials"))
+	int SelectedMaterial = 0;
+
+	UPROPERTY(EditAnywhere, Category=Materials)
+	TArray<UMaterialInterface*> Materials;
+};
+
+
+
+
+
+UENUM()
+enum class EEditMeshMaterialsToolActions
+{
+	NoAction,
+	AssignMaterial
+};
+
+
+
+UCLASS()
+class MESHMODELINGTOOLSEDITORONLY_API UEditMeshMaterialsEditActions : public UMeshSelectionToolActionPropertySet
+{
+	GENERATED_BODY()
+
+public:
+	UFUNCTION(CallInEditor, Category = MaterialEdits, meta = (DisplayName = "Assign Selected Material", DisplayPriority = 1))
+	void AssignSelectedMaterial()
+	{
+		PostMaterialAction(EEditMeshMaterialsToolActions::AssignMaterial);
+	}
+
+	void PostMaterialAction(EEditMeshMaterialsToolActions Action);
+};
+
+
+
+
+
+
+/**
+ *
+ */
+UCLASS()
+class MESHMODELINGTOOLSEDITORONLY_API UEditMeshMaterialsTool : public UMeshSelectionTool
+{
+	GENERATED_BODY()
+
+public:
+	virtual void Setup() override;
+	virtual void Tick(float DeltaTime) override;
+
+	virtual bool CanAccept() const override { return UMeshSelectionTool::CanAccept() || bHaveModifiedMaterials; }
+
+	virtual void OnShutdown(EToolShutdownType ShutdownType) override;
+
+	void RequestMaterialAction(EEditMeshMaterialsToolActions ActionType);
+
+protected:
+	UPROPERTY()
+	UEditMeshMaterialsToolProperties* MaterialProps;
+
+	virtual UMeshSelectionToolActionPropertySet* CreateEditActions() override;
+	virtual void AddSubclassPropertySets() override;
+
+	bool bHavePendingSubAction = false;
+	EEditMeshMaterialsToolActions PendingSubAction = EEditMeshMaterialsToolActions::NoAction;
+
+	void ApplyMaterialAction(EEditMeshMaterialsToolActions ActionType);
+	void AssignMaterialToSelectedTriangles();
+
+	TArray<UMaterialInterface*> CurrentMaterials;
+	void OnMaterialSetChanged();
+
+	struct FMaterialSetKey
+	{
+		TArray<void*> Values;
+		bool operator!=(const FMaterialSetKey& Key2) const;
+	};
+	FMaterialSetKey GetMaterialKey();
+
+	FMaterialSetKey InitialMaterialKey;
+	TValueWatcher<FMaterialSetKey> MaterialSetWatcher;
+	bool bHaveModifiedMaterials = false;
+
+	void ExternalUpdateMaterialSet(const TArray<UMaterialInterface*>& NewMaterialSet);
+	friend class FEditMeshMaterials_MaterialSetChange;
+};
+
+
+
+
+/**
+ */
+class MESHMODELINGTOOLSEDITORONLY_API FEditMeshMaterials_MaterialSetChange : public FToolCommandChange
+{
+public:
+	TArray<UMaterialInterface*> MaterialsBefore;
+	TArray<UMaterialInterface*> MaterialsAfter;
+
+	/** Makes the change to the object */
+	virtual void Apply(UObject* Object) override;
+
+	/** Reverts change to the object */
+	virtual void Revert(UObject* Object) override;
+
+	/** Describes this change (for debugging) */
+	virtual FString ToString() const override;
+};
+
diff --git a/Engine/Plugins/Experimental/MeshModelingToolset/Source/MeshModelingToolsEditorOnly/Public/ParameterizeMeshTool.h b/Engine/Plugins/Experimental/MeshModelingToolset/Source/MeshModelingToolsEditorOnly/Public/ParameterizeMeshTool.h
--- a/Engine/Plugins/Experimental/MeshModelingToolset/Source/MeshModelingToolsEditorOnly/Public/ParameterizeMeshTool.h
+++ b/Engine/Plugins/Experimental/MeshModelingToolset/Source/MeshModelingToolsEditorOnly/Public/ParameterizeMeshTool.h
@@ -58,6 +58,48 @@ public:
 
 
 
+UENUM()
+enum class EParameterizeMeshToolUVScaleMode
+{
+	/** No scaling is applied to UV islands */
+	NoScaling,
+	/** Scale UV islands such that they have constant relative area, relative to object bounds */
+	NormalizeToBounds,
+	/** Scale UV islands such that they have constant relative area, relative to world space */
+	NormalizeToWorld
+};
+
+
+
+
+UCLASS()
+class MESHMODELINGTOOLSEDITORONLY_API UParameterizeMeshToolProperties : public UInteractiveToolPropertySet
+{
+	GENERATED_BODY()
+public:
+
+	/** Prevent UVs from crossing poly group boundaries boundaries */
+	UPROPERTY(EditAnywhere, Category = Options, meta = (DisplayName = "Respect Group Boundaries"))
+	bool bRespectPolygroups = true;
+
+	/** Maximum amount of stretch, from none to any.  If zero stretch is specified each triangle will likey be its own chart */
+	UPROPERTY(EditAnywhere, Category = Options, meta = (Default = "0.166", UIMin = "0", UIMax = "1", ClampMin = "0", ClampMax = "1"))
+	float ChartStretch = 0.11f;
+
+	/** Scaling applied to UV islands */
+	UPROPERTY(EditAnywhere, Category = Options)
+	EParameterizeMeshToolUVScaleMode UVScaleMode = EParameterizeMeshToolUVScaleMode::NormalizeToBounds;
+
+	/** Scaling factor used for UV island normalization/scaling */
+	UPROPERTY(EditAnywhere, Category = Options, meta = (EditCondition = "UVScaleMode!=EParameterizeMeshToolUVScaleMode::NoScaling", UIMin = "0.001", UIMax = "10", ClampMin = "0.00001", ClampMax = "1000000.0") )
+	float UVScale = 1.0;
+
+	virtual void SaveProperties(UInteractiveTool* SaveFromTool) override;
+	virtual void RestoreProperties(UInteractiveTool* RestoreToTool) override;
+};
+
+
+
 
 
 /**
@@ -84,22 +126,14 @@ public:
 	virtual bool HasAccept() const override;
 	virtual bool CanAccept() const override;
 
-#if WITH_EDITOR
-	virtual void PostEditChangeProperty(FPropertyChangedEvent& PropertyChangedEvent) override;
-#endif
-
 	virtual void OnPropertyModified(UObject* PropertySet, UProperty* Property) override;
 
 	// IDynamicMeshOperatorFactory API
-	virtual TSharedPtr<FDynamicMeshOperator> MakeNewOperator() override;
+	virtual TUniquePtr<FDynamicMeshOperator> MakeNewOperator() override;
 
 protected:
-	// need to update bResultValid if these are modified, so we don't publicly expose them. 
-	// @todo setters/getters for these
-
-	/** Maximum amount of stretch, from none to any.  If zero stretch is specified each triangle will likey be its own chart */
-	UPROPERTY(EditAnywhere, Category = Options, meta = (Default = "0.166", UIMin = "0", UIMax = "1", ClampMin = "0", ClampMax= "1"))
-	float ChartStretch = 0.11f;
+	UPROPERTY()
+	UParameterizeMeshToolProperties* Settings = nullptr;
 
 	UPROPERTY()
 	UExistingMeshMaterialProperties* MaterialSettings = nullptr;
diff --git a/Engine/Plugins/Experimental/MeshModelingToolset/Source/ModelingComponents/Private/Drawing/MeshDebugDrawing.cpp b/Engine/Plugins/Experimental/MeshModelingToolset/Source/ModelingComponents/Private/Drawing/MeshDebugDrawing.cpp
--- a/Engine/Plugins/Experimental/MeshModelingToolset/Source/ModelingComponents/Private/Drawing/MeshDebugDrawing.cpp
+++ b/Engine/Plugins/Experimental/MeshModelingToolset/Source/ModelingComponents/Private/Drawing/MeshDebugDrawing.cpp
@@ -19,7 +19,7 @@ void MeshDebugDraw::DrawNormals(
 		int ParentVID = Overlay->GetParentVertex(ElementID);
 		FVector3f ParentPos = (FVector3f)Mesh->GetVertex(ParentVID);
 
-		FVector A = ParentPos, B = ParentPos + Length * Normal;
+		FVector A = (FVector)ParentPos, B = (FVector)(ParentPos + Length * Normal);
 		PDI->DrawLine(Transform.TransformPosition(A), Transform.TransformPosition(B),
 			Color, 0, Thickness, 0, bScreenSpace);
 	}
@@ -37,7 +37,7 @@ void MeshDebugDraw::DrawVertices(
 	for (int VertID : Indices)
 	{
 		FVector3d Pos = Mesh->GetVertex(VertID);
-		PDI->DrawPoint(Transform.TransformPosition(Pos), Color, PointSize, SDPG_World);
+		PDI->DrawPoint(Transform.TransformPosition((FVector)Pos), Color, PointSize, SDPG_World);
 	}
 }
 
@@ -49,7 +49,7 @@ void MeshDebugDraw::DrawVertices(
 	for (int VertID : Indices)
 	{
 		FVector3d Pos = Mesh->GetVertex(VertID);
-		PDI->DrawPoint(Transform.TransformPosition(Pos), Color, PointSize, SDPG_World);
+		PDI->DrawPoint(Transform.TransformPosition((FVector)Pos), Color, PointSize, SDPG_World);
 	}
 }
 
@@ -63,7 +63,7 @@ void MeshDebugDraw::DrawTriCentroids(
 	for (int TriID : Indices)
 	{
 		FVector3d Pos = Mesh->GetTriCentroid(TriID);
-		PDI->DrawPoint(Transform.TransformPosition(Pos), Color, PointSize, SDPG_World);
+		PDI->DrawPoint(Transform.TransformPosition((FVector)Pos), Color, PointSize, SDPG_World);
 	}
 }
 
@@ -85,7 +85,7 @@ void MeshDebugDraw::DrawSimpleGrid(
 	FVector3f Origin = WorldFrame.Origin;
 	FVector3f X = WorldFrame.X();
 	FVector3f Y = WorldFrame.Y();
-	FVector A, B;
+	FVector3f A, B;
 
 	int LineSteps = GridLines / 2;
 	for (int i = 0; i < LineSteps; i++)
@@ -93,17 +93,17 @@ void MeshDebugDraw::DrawSimpleGrid(
 		float dx = (float)i * GridLineSpacing;
 		A = Origin - Extent * Y - dx * X;
 		B = Origin + Extent * Y - dx * X;
-		PDI->DrawLine(A, B, Color, DepthPriority, LineWidth, 0, true);
+		PDI->DrawLine((FVector)A, (FVector)B, Color, DepthPriority, LineWidth, 0, true);
 		A = Origin - Extent * Y + dx * X;
 		B = Origin + Extent * Y + dx * X;
-		PDI->DrawLine(A, B, Color, DepthPriority, LineWidth, 0, true);
+		PDI->DrawLine((FVector)A, (FVector)B, Color, DepthPriority, LineWidth, 0, true);
 
 		A = Origin - Extent * X - dx * Y;
 		B = Origin + Extent * X - dx * Y;
-		PDI->DrawLine(A, B, Color, DepthPriority, LineWidth, 0, true);
+		PDI->DrawLine((FVector)A, (FVector)B, Color, DepthPriority, LineWidth, 0, true);
 		A = Origin - Extent * X + dx * Y;
 		B = Origin + Extent * X + dx * Y;
-		PDI->DrawLine(A, B, Color, DepthPriority, LineWidth, 0, true);
+		PDI->DrawLine((FVector)A, (FVector)B, Color, DepthPriority, LineWidth, 0, true);
 	}
 
 }
diff --git a/Engine/Plugins/Experimental/MeshModelingToolset/Source/ModelingComponents/Private/PreviewMesh.cpp b/Engine/Plugins/Experimental/MeshModelingToolset/Source/ModelingComponents/Private/PreviewMesh.cpp
--- a/Engine/Plugins/Experimental/MeshModelingToolset/Source/ModelingComponents/Private/PreviewMesh.cpp
+++ b/Engine/Plugins/Experimental/MeshModelingToolset/Source/ModelingComponents/Private/PreviewMesh.cpp
@@ -3,6 +3,7 @@
 #include "PreviewMesh.h"
 
 #include "Containers/StaticArray.h"
+#include "Engine/Classes/Materials/Material.h"
 
 
 APreviewMeshActor::APreviewMeshActor()
@@ -61,18 +62,96 @@ void UPreviewMesh::Disconnect()
 
 
 void UPreviewMesh::SetMaterial(UMaterialInterface* Material)
+{
+	SetMaterial(0, Material);
+}
+
+void UPreviewMesh::SetMaterial(int MaterialIndex, UMaterialInterface* Material)
+{
+	check(DynamicMeshComponent);
+	DynamicMeshComponent->SetMaterial(MaterialIndex, Material);
+
+	// force rebuild because we can't change materials yet - surprisingly complicated
+	DynamicMeshComponent->NotifyMeshUpdated();
+}
+
+void UPreviewMesh::SetMaterials(const TArray<UMaterialInterface*>& Materials)
+{
+	check(DynamicMeshComponent);
+	for (int k = 0; k < Materials.Num(); ++k)
+	{
+		DynamicMeshComponent->SetMaterial(k, Materials[k]);
+	}
+
+	// force rebuild because we can't change materials yet - surprisingly complicated
+	DynamicMeshComponent->NotifyMeshUpdated();
+}
+
+UMaterialInterface* UPreviewMesh::GetMaterial(int MaterialIndex) const
+{
+	check(DynamicMeshComponent);
+	return DynamicMeshComponent->GetMaterial(MaterialIndex);
+}
+
+
+void UPreviewMesh::SetOverrideRenderMaterial(UMaterialInterface* Material)
+{
+	check(DynamicMeshComponent);
+	return DynamicMeshComponent->SetOverrideRenderMaterial(Material);
+}
+
+void UPreviewMesh::ClearOverrideRenderMaterial()
+{
+	check(DynamicMeshComponent);
+	return DynamicMeshComponent->ClearOverrideRenderMaterial();
+}
+
+
+UMaterialInterface* UPreviewMesh::GetActiveMaterial(int MaterialIndex) const
+{
+	return DynamicMeshComponent->HasOverrideRenderMaterial(MaterialIndex) ?
+		DynamicMeshComponent->GetOverrideRenderMaterial(MaterialIndex) : GetMaterial(MaterialIndex);
+}
+
+
+
+void UPreviewMesh::SetSecondaryRenderMaterial(UMaterialInterface* Material)
+{
+	check(DynamicMeshComponent);
+	return DynamicMeshComponent->SetSecondaryRenderMaterial(Material);
+}
+
+void UPreviewMesh::ClearSecondaryRenderMaterial()
+{
+	check(DynamicMeshComponent);
+	return DynamicMeshComponent->ClearSecondaryRenderMaterial();
+}
+
+
+
+void UPreviewMesh::EnableSecondaryTriangleBuffers(TUniqueFunction<bool(const FDynamicMesh3*, int32)> TriangleFilterFunc)
+{
+	check(DynamicMeshComponent);
+	DynamicMeshComponent->EnableSecondaryTriangleBuffers(MoveTemp(TriangleFilterFunc));
+}
+
+void UPreviewMesh::DisableSecondaryTriangleBuffers()
 {
 	check(DynamicMeshComponent);
-	DynamicMeshComponent->SetMaterial(0, Material);
+	DynamicMeshComponent->DisableSecondaryTriangleBuffers();
 }
 
-UMaterialInterface*
-UPreviewMesh::GetMaterial() const
+
+void UPreviewMesh::SetTangentsMode(EDynamicMeshTangentCalcType TangentsType)
 {
 	check(DynamicMeshComponent);
-	return DynamicMeshComponent->GetMaterial(0);
+	DynamicMeshComponent->TangentsType = TangentsType;
 }
 
+FMeshTangentsf* UPreviewMesh::GetTangents() const
+{
+	return DynamicMeshComponent->GetTangents();
+}
 
 void UPreviewMesh::EnableWireframe(bool bEnable)
 {
@@ -145,7 +224,7 @@ void UPreviewMesh::UpdatePreview(const FDynamicMesh3* Mesh)
 }
 
 
-const FDynamicMesh3* UPreviewMesh::GetPreviewDynamicMesh() const
+const FDynamicMesh3* UPreviewMesh::GetMesh() const
 {
 	if (DynamicMeshComponent != nullptr)
 	{
@@ -155,6 +234,7 @@ const FDynamicMesh3* UPreviewMesh::GetPreviewDynamicMesh() const
 }
 
 
+
 TUniquePtr<FDynamicMesh3> UPreviewMesh::ExtractPreviewMesh() const
 {
 	if (DynamicMeshComponent != nullptr)
@@ -187,7 +267,7 @@ bool UPreviewMesh::FindRayIntersection(const FRay3d& WorldRay, FHitResult& HitOu
 {
 	if (IsVisible() && TemporaryParentActor != nullptr && bBuildSpatialDataStructure)
 	{
-		FTransform Transform = TemporaryParentActor->GetActorTransform();
+		FTransform3d Transform(TemporaryParentActor->GetActorTransform());
 		FRay3d LocalRay(Transform.InverseTransformPosition(WorldRay.Origin),
 			Transform.InverseTransformVector(WorldRay.Direction));
 		LocalRay.Direction.Normalize();
@@ -200,13 +280,11 @@ bool UPreviewMesh::FindRayIntersection(const FRay3d& WorldRay, FHitResult& HitOu
 			FIntrRay3Triangle3d Query(LocalRay, Triangle);
 			Query.Find();
 
-			//auto Query = MeshQueries::RayTriangleIntersection(UseMesh, HitTID, LocalRay);
-
 			HitOut.FaceIndex = HitTriID;
 			HitOut.Distance = Query.RayParameter;
-			HitOut.Normal = Transform.TransformVectorNoScale(UseMesh->GetTriNormal(HitTriID));
+			HitOut.Normal = (FVector)Transform.TransformVectorNoScale(UseMesh->GetTriNormal(HitTriID));
 			HitOut.ImpactNormal = HitOut.Normal;
-			HitOut.ImpactPoint = Transform.TransformPosition(LocalRay.PointAt(Query.RayParameter));
+			HitOut.ImpactPoint = (FVector)Transform.TransformPosition(LocalRay.PointAt(Query.RayParameter));
 			return true;
 		}
 	}
@@ -242,6 +320,45 @@ void UPreviewMesh::EditMesh(TFunctionRef<void(FDynamicMesh3&)> EditFunc)
 }
 
 
+void UPreviewMesh::DeferredEditMesh(TFunctionRef<void(FDynamicMesh3&)> EditFunc, bool bRebuildSpatial)
+{
+	FDynamicMesh3* Mesh = DynamicMeshComponent->GetMesh();
+	EditFunc(*Mesh);
+
+	if (bBuildSpatialDataStructure && bRebuildSpatial)
+	{
+		MeshAABBTree.SetMesh(DynamicMeshComponent->GetMesh(), true);
+	}
+}
+
+
+void UPreviewMesh::ForceRebuildSpatial()
+{
+	if (bBuildSpatialDataStructure)
+	{
+		MeshAABBTree.SetMesh(DynamicMeshComponent->GetMesh(), true);
+	}
+}
+
+
+void UPreviewMesh::NotifyDeferredEditCompleted(ERenderUpdateMode UpdateMode, bool bRebuildSpatial)
+{
+	if (bBuildSpatialDataStructure && bRebuildSpatial)
+	{
+		MeshAABBTree.SetMesh(DynamicMeshComponent->GetMesh(), true);
+	}
+
+	if (UpdateMode == ERenderUpdateMode::FullUpdate)
+	{
+		DynamicMeshComponent->NotifyMeshUpdated();
+	}
+	else
+	{
+		DynamicMeshComponent->FastNotifyPositionsUpdated();
+	}
+}
+
+
 TUniquePtr<FMeshChange> UPreviewMesh::TrackedEditMesh(TFunctionRef<void(FDynamicMesh3&, FDynamicMeshChangeTracker&)> EditFunc)
 {
 	FDynamicMesh3* Mesh = DynamicMeshComponent->GetMesh();
@@ -279,6 +396,15 @@ void UPreviewMesh::ApplyChange(const FMeshChange* Change, bool bRevert)
 		MeshAABBTree.SetMesh(DynamicMeshComponent->GetMesh(), true);
 	}
 }
+void UPreviewMesh::ApplyChange(const FMeshReplacementChange* Change, bool bRevert)
+{
+	check(DynamicMeshComponent != nullptr);
+	DynamicMeshComponent->ApplyChange(Change, bRevert);
+	if (bBuildSpatialDataStructure)
+	{
+		MeshAABBTree.SetMesh(DynamicMeshComponent->GetMesh(), true);
+	}
+}
 
 FSimpleMulticastDelegate& UPreviewMesh::GetOnMeshChanged()
 {
@@ -293,7 +419,7 @@ void UPreviewMesh::Bake(FMeshDescription* MeshDescription, bool bHaveModifiedTop
 	DynamicMeshComponent->Bake(MeshDescription, bHaveModifiedToplogy);
 }
 
-void UPreviewMesh::SetTriangleColorFunction(TFunction<FColor(int)> TriangleColorFunc, ERenderUpdateMode UpdateMode)
+void UPreviewMesh::SetTriangleColorFunction(TFunction<FColor(const FDynamicMesh3*, int)> TriangleColorFunc, ERenderUpdateMode UpdateMode)
 {
 	DynamicMeshComponent->TriangleColorFunc = TriangleColorFunc;
 	if (UpdateMode == ERenderUpdateMode::FastUpdate)
diff --git a/Engine/Plugins/Experimental/MeshModelingToolset/Source/ModelingComponents/Private/Selection/GroupTopologySelector.cpp b/Engine/Plugins/Experimental/MeshModelingToolset/Source/ModelingComponents/Private/Selection/GroupTopologySelector.cpp
--- a/Engine/Plugins/Experimental/MeshModelingToolset/Source/ModelingComponents/Private/Selection/GroupTopologySelector.cpp
+++ b/Engine/Plugins/Experimental/MeshModelingToolset/Source/ModelingComponents/Private/Selection/GroupTopologySelector.cpp
@@ -3,19 +3,11 @@
 
 #include "Selection/GroupTopologySelector.h"
 #include "MeshQueries.h"
-#include "Drawing/ToolDataVisualizer.h"
+#include "ToolDataVisualizer.h"
 #include "ToolSceneQueriesUtil.h"
 
 
 
-void FGroupTopologySelection::Clear()
-{
-	SelectedGroupIDs.Reset();
-	SelectedCornerIDs.Reset();
-	SelectedEdgeIDs.Reset();
-}
-
-
 
 FGroupTopologySelector::FGroupTopologySelector()
 {
