Fixed CIS

#rnx

diff --git a/Engine/Plugins/Experimental/MeshModelingToolset/Source/ModelingComponents/Private/Drawing/PointSetComponent.cpp b/Engine/Plugins/Experimental/MeshModelingToolset/Source/ModelingComponents/Private/Drawing/PointSetComponent.cpp
--- a/Engine/Plugins/Experimental/MeshModelingToolset/Source/ModelingComponents/Private/Drawing/PointSetComponent.cpp
+++ b/Engine/Plugins/Experimental/MeshModelingToolset/Source/ModelingComponents/Private/Drawing/PointSetComponent.cpp
@@ -199,7 +199,7 @@ public:
 		Result.bRenderCustomDepth = ShouldRenderCustomDepth();
 		Result.bTranslucentSelfShadow = bCastVolumetricTranslucentShadow;
 		MaterialRelevance.SetPrimitiveViewRelevance(Result);
-		Result.bVelocityRelevance = IsMovable() && Result.bOpaqueRelevance && Result.bRenderInMainPass;
+		Result.bVelocityRelevance = IsMovable() && Result.bOpaque && Result.bRenderInMainPass;
 		return Result;
 	}
 
