uvlayout tool initial version + disconnect triangles button on select tool + a lot of related support code to handle bow ties and breaking stuff apart

#rnx

diff --git a/Engine/Plugins/Experimental/GeometryProcessing/Source/DynamicMesh/Private/DynamicMesh3.cpp b/Engine/Plugins/Experimental/GeometryProcessing/Source/DynamicMesh/Private/DynamicMesh3.cpp
--- a/Engine/Plugins/Experimental/GeometryProcessing/Source/DynamicMesh/Private/DynamicMesh3.cpp
+++ b/Engine/Plugins/Experimental/GeometryProcessing/Source/DynamicMesh/Private/DynamicMesh3.cpp
@@ -1012,3 +1012,24 @@ int FDynamicMesh3::FindEdgeFromTri(int vA, int vB, int tID) const
 	}
 	return InvalidID;
 }
+
+int FDynamicMesh3::FindEdgeFromTriPair(int TriA, int TriB) const
+{
+	if (TriangleRefCounts.IsValid(TriA) && TriangleRefCounts.IsValid(TriB))
+	{
+		int AI = 3 * TriA, BI = 3 * TriB;
+		for (int j = 0; j < 3; ++j)
+		{
+			int EdgeID = TriangleEdges[AI + j];
+			int EdgeIndex = 4 * EdgeID;
+			int NbrT = (Edges[EdgeIndex + 2] == TriA) ? Edges[EdgeIndex + 3] : Edges[EdgeIndex + 2];
+			if (NbrT == TriB)
+			{
+				return EdgeID;
+			}
+		}
+	}
+	return InvalidID;
+}
+
+
diff --git a/Engine/Plugins/Experimental/GeometryProcessing/Source/DynamicMesh/Private/DynamicMesh3_Edits.cpp b/Engine/Plugins/Experimental/GeometryProcessing/Source/DynamicMesh/Private/DynamicMesh3_Edits.cpp
--- a/Engine/Plugins/Experimental/GeometryProcessing/Source/DynamicMesh/Private/DynamicMesh3_Edits.cpp
+++ b/Engine/Plugins/Experimental/GeometryProcessing/Source/DynamicMesh/Private/DynamicMesh3_Edits.cpp
@@ -1187,7 +1187,77 @@ EMeshResult FDynamicMesh3::FlipEdge(int vA, int vB, FEdgeFlipInfo& FlipInfo)
 
 
 
+EMeshResult FDynamicMesh3::SplitVertex(int VertexID, const TArrayView<const int>& TrianglesToUpdate, FVertexSplitInfo& SplitInfo)
+{
+	if (!ensure(IsVertex(VertexID)))
+	{
+		return EMeshResult::Failed_NotAVertex;
+	}
+
+	SplitInfo.OriginalVertex = VertexID;
+	SplitInfo.NewVertex = AppendVertex(*this, VertexID);
+
+	// TODO: consider making a TSet copy of TrianglesToUpdate for membership tests, if TrianglesToUpdate is large
+	auto ProcessEdge = [this, &TrianglesToUpdate, &SplitInfo](int TriID, FIndex3i& UpdatedTri, FIndex3i& TriEdges, int SubIdx)
+	{
+		int EdgeID = TriEdges[SubIdx];
+		int OtherTri = GetOtherEdgeTriangle(EdgeID, TriID);
+		bool bNewBoundary = OtherTri >= 0 && !TrianglesToUpdate.Contains(OtherTri); // processing this edge will create a new boundary
+		if (bNewBoundary) // there *is* a triangle across from this edge and we do need to separate from it
+		{
+			ReplaceEdgeTriangle(EdgeID, TriID, InvalidID); // remove TriID from original edge, disconnecting it from OtherTri
+			// add a new edge for TriID connecting the updated tri vertices
+			AddTriangleEdge(TriID, UpdatedTri[SubIdx], UpdatedTri[(SubIdx+1)%3], SubIdx, InvalidID); // adds to vertexedgelists
+		}
+		else // othertri invalid or also in set
+		{
+			// if OtherTri already was processed and replaced edge, ReplaceEdgeVertex will return InvalidID and do nothing
+			if (ReplaceEdgeVertex(EdgeID, SplitInfo.OriginalVertex, SplitInfo.NewVertex) != InvalidID)
+			{
+				// if replace edge actually happened, also update VertexEdgeLists accordingly
+				ensure(VertexEdgeLists.Remove(SplitInfo.OriginalVertex, EdgeID));
+				VertexEdgeLists.Insert(SplitInfo.NewVertex, EdgeID);
+			}
+		}
+	};
+	for (int TriID : TrianglesToUpdate)
+	{
+		FIndex3i Triangle = GetTriangle(TriID);
+		int SubIdx = Triangle.IndexOf(VertexID);
+		if (SubIdx < 0)
+		{
+			continue;
+		}
+		Triangle[SubIdx] = SplitInfo.NewVertex; // update local copy w/ new vertex, for use by ProcessEdge helper
+		FIndex3i TriEdges = GetTriEdges(TriID);
+		ProcessEdge(TriID, Triangle, TriEdges, SubIdx);
+		ProcessEdge(TriID, Triangle, TriEdges, (SubIdx+2)%3);
 
+		Triangles[TriID * 3 + SubIdx] = SplitInfo.NewVertex;
+		VertexRefCounts.Decrement(SplitInfo.OriginalVertex); // remove the triangle from the original vertex
+		VertexRefCounts.Increment(SplitInfo.NewVertex);
+	}
+
+	if (HasAttributes())
+	{
+		Attributes()->OnSplitVertex(SplitInfo, TrianglesToUpdate);
+	}
+	UpdateTimeStamp(true, true);
+	return EMeshResult::Ok;
+}
+
+
+bool FDynamicMesh3::SplitVertexWouldLeaveIsolated(int VertexID, const TArrayView<const int>& TrianglesToUpdate)
+{
+	for (int TID : VtxTrianglesItr(VertexID))
+	{
+		if (!TrianglesToUpdate.Contains(TID))
+		{
+			return false; // at least one triangle will keep the old VertexID
+		}
+	}
+	return true; // no triangles founds that keep old VertexID
+}
 
 
 
diff --git a/Engine/Plugins/Experimental/GeometryProcessing/Source/DynamicMesh/Private/DynamicMeshAttributeSet.cpp b/Engine/Plugins/Experimental/GeometryProcessing/Source/DynamicMesh/Private/DynamicMeshAttributeSet.cpp
--- a/Engine/Plugins/Experimental/GeometryProcessing/Source/DynamicMesh/Private/DynamicMeshAttributeSet.cpp
+++ b/Engine/Plugins/Experimental/GeometryProcessing/Source/DynamicMesh/Private/DynamicMeshAttributeSet.cpp
@@ -179,3 +179,19 @@ void FDynamicMeshAttributeSet::OnMergeEdges(const FDynamicMesh3::FMergeEdgesInfo
 	}
 }
 
+void FDynamicMeshAttributeSet::OnSplitVertex(const DynamicMeshInfo::FVertexSplitInfo& SplitInfo, const TArrayView<const int>& TrianglesToUpdate)
+{
+	FDynamicMeshAttributeSetBase::OnSplitVertex(SplitInfo, TrianglesToUpdate);
+
+	for (FDynamicMeshUVOverlay& UVLayer : UVLayers)
+	{
+		UVLayer.OnSplitVertex(SplitInfo, TrianglesToUpdate);
+	}
+	Normals0.OnSplitVertex(SplitInfo, TrianglesToUpdate);
+
+	if (MaterialIDAttrib)
+	{
+		MaterialIDAttrib->OnSplitVertex(SplitInfo, TrianglesToUpdate);
+	}
+}
+
diff --git a/Engine/Plugins/Experimental/GeometryProcessing/Source/DynamicMesh/Private/DynamicMeshEditor.cpp b/Engine/Plugins/Experimental/GeometryProcessing/Source/DynamicMesh/Private/DynamicMeshEditor.cpp
--- a/Engine/Plugins/Experimental/GeometryProcessing/Source/DynamicMesh/Private/DynamicMeshEditor.cpp
+++ b/Engine/Plugins/Experimental/GeometryProcessing/Source/DynamicMesh/Private/DynamicMeshEditor.cpp
@@ -306,8 +306,7 @@ bool FDynamicMeshEditor::DisconnectTriangles(const TArray<int>& Triangles, TArra
 	// find the region boundary loops
 	FMeshRegionBoundaryLoops RegionLoops(Mesh, Triangles, false);
 	bool bOK = RegionLoops.Compute();
-	check(bOK);
-	if (!bOK)
+	if (!ensure(bOK))
 	{
 		return false;
 	}
@@ -379,6 +378,115 @@ bool FDynamicMeshEditor::DisconnectTriangles(const TArray<int>& Triangles, TArra
 }
 
 
+void FDynamicMeshEditor::DisconnectTriangles(const TArray<int>& Triangles, bool bPreventBowties)
+{
+	TSet<int> TriSet, BoundaryVerts;
+	TArray<int> NewVerts, OldVertsThatSplit;
+	TArray<int> FilteredTriangles;
+	DynamicMeshInfo::FVertexSplitInfo SplitInfo;
+
+	TriSet.Append(Triangles);
+	for (int TID : Triangles)
+	{
+		FIndex3i Nbrs = Mesh->GetTriNeighbourTris(TID);
+		FIndex3i Tri = Mesh->GetTriangle(TID);
+		for (int SubIdx = 0; SubIdx < 3; SubIdx++)
+		{
+			int NeighborTID = Nbrs[SubIdx];
+			if (!TriSet.Contains(NeighborTID))
+			{
+				BoundaryVerts.Add(Tri[SubIdx]);
+				BoundaryVerts.Add(Tri[(SubIdx + 1) % 3]);
+			}
+		}
+	}
+	for (int VID : BoundaryVerts)
+	{
+		FilteredTriangles.Reset();
+		int TriRingCount = 0;
+		for (int RingTID : Mesh->VtxTrianglesItr(VID))
+		{
+			if (TriSet.Contains(RingTID))
+			{
+				FilteredTriangles.Add(RingTID);
+			}
+			TriRingCount++;
+		}
+
+		if (FilteredTriangles.Num() < TriRingCount)
+		{
+			checkSlow(!Mesh->SplitVertexWouldLeaveIsolated(VID, FilteredTriangles));
+			ensure(EMeshResult::Ok == Mesh->SplitVertex(VID, FilteredTriangles, SplitInfo));
+			NewVerts.Add(SplitInfo.NewVertex);
+			OldVertsThatSplit.Add(SplitInfo.OriginalVertex);
+		}
+	}
+	if (bPreventBowties)
+	{
+		FDynamicMeshEditResult Result;
+		for (int VID : OldVertsThatSplit)
+		{
+			SplitBowties(VID, Result);
+			Result.Reset(); // don't actually keep results; they are not used in this fn
+		}
+		for (int VID : NewVerts)
+		{
+			SplitBowties(VID, Result);
+			Result.Reset(); // don't actually keep results; they are not used in this fn
+		}
+	}
+}
+
+
+
+
+void FDynamicMeshEditor::SplitBowties(FDynamicMeshEditResult& ResultOut)
+{
+	ResultOut.Reset();
+	TSet<int> AddedVerticesWithIDLessThanMax; // added vertices that we can't filter just by checking against original max id; this will be empty for compact meshes
+	for (int VertexID = 0, OriginalMaxID = Mesh->MaxVertexID(); VertexID < OriginalMaxID; VertexID++)
+	{
+		if (!Mesh->IsVertex(VertexID) || AddedVerticesWithIDLessThanMax.Contains(VertexID))
+		{
+			continue;
+		}
+		int32 NumVertsBefore = ResultOut.NewVertices.Num();
+		// TODO: may be faster to inline this call to reuse the contiguous triangle arrays?
+		SplitBowties(VertexID, ResultOut);
+		for (int Idx = NumVertsBefore; Idx < ResultOut.NewVertices.Num(); Idx++)
+		{
+			if (ResultOut.NewVertices[Idx] < OriginalMaxID)
+			{
+				AddedVerticesWithIDLessThanMax.Add(ResultOut.NewVertices[Idx]);
+			}
+		}
+	}
+}
+
+
+
+void FDynamicMeshEditor::SplitBowties(int VertexID, FDynamicMeshEditResult& ResultOut)
+{
+	TArray<int> TrianglesOut, ContiguousGroupLengths;
+	TArray<bool> GroupIsLoop;
+	DynamicMeshInfo::FVertexSplitInfo SplitInfo;
+	check(Mesh->IsVertex(VertexID));
+	if (ensure(EMeshResult::Ok == Mesh->GetVtxContiguousTriangles(VertexID, TrianglesOut, ContiguousGroupLengths, GroupIsLoop)))
+	{
+		if (ContiguousGroupLengths.Num() > 1)
+		{
+			// is bowtie
+			for (int GroupIdx = 1, GroupStartIdx = ContiguousGroupLengths[0]; GroupIdx < ContiguousGroupLengths.Num(); GroupStartIdx += ContiguousGroupLengths[GroupIdx++])
+			{
+				ensure(EMeshResult::Ok == Mesh->SplitVertex(VertexID, TArrayView<const int>(TrianglesOut.GetData() + GroupStartIdx, ContiguousGroupLengths[GroupIdx]), SplitInfo));
+				ResultOut.NewVertices.Add(SplitInfo.NewVertex);
+			}
+		}
+	}
+}
+
+
+
 
 
 FVector3f FDynamicMeshEditor::ComputeAndSetQuadNormal(const FIndex2i& QuadTris, bool bIsPlanar)
diff --git a/Engine/Plugins/Experimental/GeometryProcessing/Source/DynamicMesh/Private/DynamicMeshOverlay.cpp b/Engine/Plugins/Experimental/GeometryProcessing/Source/DynamicMesh/Private/DynamicMeshOverlay.cpp
--- a/Engine/Plugins/Experimental/GeometryProcessing/Source/DynamicMesh/Private/DynamicMeshOverlay.cpp
+++ b/Engine/Plugins/Experimental/GeometryProcessing/Source/DynamicMesh/Private/DynamicMeshOverlay.cpp
@@ -203,6 +203,135 @@ void TDynamicMeshOverlay<RealType, ElementSize>::SplitVerticesWithPredicate(TFun
 }
 
 
+template<typename RealType, int ElementSize>
+int TDynamicMeshOverlay<RealType, ElementSize>::SplitElement(int ElementID, const TArrayView<const int>& TrianglesToUpdate)
+{
+	int ParentID = ParentVertices[ElementID];
+	return SplitElementWithNewParent(ElementID, ParentID, TrianglesToUpdate);
+}
+
+
+template<typename RealType, int ElementSize>
+int TDynamicMeshOverlay<RealType, ElementSize>::SplitElementWithNewParent(int ElementID, int NewParentID, const TArrayView<const int>& TrianglesToUpdate)
+{
+	RealType SourceData[ElementSize];
+	GetElement(ElementID, SourceData);
+	int NewElID = AppendElement(SourceData, NewParentID);
+	for (int TriID : TrianglesToUpdate)
+	{
+		int ElementTriStart = TriID * 3;
+		for (int SubIdx = 0; SubIdx < 3; SubIdx++)
+		{
+			int CurElID = ElementTriangles[ElementTriStart + SubIdx];
+			if (CurElID == ElementID)
+			{
+				ElementsRefCounts.Decrement(ElementID);
+				ElementsRefCounts.Increment(NewElID);
+				ElementTriangles[ElementTriStart + SubIdx] = NewElID;
+			}
+		}
+	}
+	checkSlow(ElementsRefCounts.IsValid(ElementID)); // simple sanity check; shouldn't completely remove an element with SplitElement
+	return NewElID;
+}
+
+
+template<typename RealType, int ElementSize>
+void TDynamicMeshOverlay<RealType, ElementSize>::SplitBowties()
+{
+	// arrays for storing contiguous triangle groups from parentmesh
+	TArray<int> TrianglesOut, ContiguousGroupLengths;
+	TArray<bool> GroupIsLoop;
+
+	// per-vertex element group tracking data, reused in loop below
+	TSet<int> ElementIDSeen;
+	TArray<int> ElementIDs, SubGroupID, SubGroupElementIDs;
+	
+	for (int VertexID : ParentMesh->VertexIndicesItr())
+	{
+		ensure(EMeshResult::Ok == ParentMesh->GetVtxContiguousTriangles(VertexID, TrianglesOut, ContiguousGroupLengths, GroupIsLoop));
+		int32 NumTris = TrianglesOut.Num();
+
+		ElementIDSeen.Reset();
+		// per contiguous group of triangles around vertex in ParentMesh, find contiguous sub-groups in overlay
+		for (int32 GroupIdx = 0, NumGroups = ContiguousGroupLengths.Num(), TriSubStart = 0; GroupIdx < NumGroups; GroupIdx++)
+		{
+			bool bIsLoop = GroupIsLoop[GroupIdx];
+			int TriInGroupNum = ContiguousGroupLengths[GroupIdx];
+			check(TriInGroupNum > 0);
+			int TriSubEnd = TriSubStart + TriInGroupNum;
+			
+			ElementIDs.Reset();
+			for (int TriSubIdx = TriSubStart; TriSubIdx < TriSubEnd; TriSubIdx++)
+			{
+				int TriID = TrianglesOut[TriSubIdx];
+				FIndex3i TriVIDs = ParentMesh->GetTriangle(TriID);
+				FIndex3i TriEIDs = GetTriangle(TriID);
+				int SubIdx = TriVIDs.IndexOf(VertexID);
+				ElementIDs.Add(TriEIDs[SubIdx]);
+			}
+
+			auto IsConnected = [this, &ElementIDs, &TrianglesOut](int TriOutIdxA, int TriOutIdxB)
+			{
+				if (ElementIDs[TriOutIdxA] != ElementIDs[TriOutIdxB])
+				{
+					return false;
+				}
+				int EdgeID = ParentMesh->FindEdgeFromTriPair(TrianglesOut[TriOutIdxA], TrianglesOut[TriOutIdxB]);
+				return EdgeID >= 0 && !IsSeamEdge(EdgeID);
+			};
+
+			SubGroupID.Reset(); SubGroupID.SetNum(TriInGroupNum);
+			SubGroupElementIDs.Reset();
+			int MaxSubID = 0;
+			SubGroupID[0] = 0;
+			SubGroupElementIDs.Add(ElementIDs[0]);
+			for (int TriSubIdx = TriSubStart; TriSubIdx+1 < TriSubEnd; TriSubIdx++)
+			{
+				if (!IsConnected(TriSubIdx, TriSubIdx+1))
+				{
+					SubGroupElementIDs.Add(ElementIDs[TriSubIdx + 1]);
+					MaxSubID++;
+				}
+				SubGroupID[TriSubIdx - TriSubStart + 1] = MaxSubID;
+			}
+			// if group was a loop, need to check if the last sub-group and first sub-group were actually the same group
+			if (bIsLoop && MaxSubID > 0 && IsConnected(TriSubStart, TriSubStart + TriInGroupNum - 1))
+			{
+				int LastGroupID = SubGroupID.Last();
+				for (int32 Idx = SubGroupID.Num() - 1; Idx >= 0 && SubGroupID[Idx] == LastGroupID; Idx--)
+				{
+					SubGroupID[Idx] = 0;
+				}
+				MaxSubID--;
+				SubGroupElementIDs.Pop(false);
+			}
+
+			for (int SubID = 0; SubID < SubGroupElementIDs.Num(); SubID++)
+			{
+				int ElementID = SubGroupElementIDs[SubID];
+				// split needed the *second* time we see a sub-group using a given ElementID
+				if (ElementIDSeen.Contains(ElementID))
+				{
+					TArray<int> ConnectedTris;
+					for (int TriSubIdx = TriSubStart; TriSubIdx < TriSubEnd; TriSubIdx++)
+					{
+						if (SubID == SubGroupID[TriSubIdx - TriSubStart])
+						{
+							ConnectedTris.Add(TrianglesOut[TriSubIdx]);
+						}
+					}
+					SplitElement(ElementID, ConnectedTris);
+				}
+				ElementIDSeen.Add(ElementID);
+			}
+
+			TriSubStart = TriSubEnd;
+		}
+	}
+}
+
+
 template<typename RealType, int ElementSize>
 void TDynamicMeshOverlay<RealType, ElementSize>::InitializeTriangles(int MaxTriangleID)
 {
@@ -777,6 +906,33 @@ void TDynamicMeshOverlay<RealType, ElementSize>::OnMergeEdges(const FDynamicMesh
 
 
 
+template<typename RealType, int ElementSize>
+void TDynamicMeshOverlay<RealType, ElementSize>::OnSplitVertex(const DynamicMeshInfo::FVertexSplitInfo& SplitInfo, const TArrayView<const int>& TrianglesToUpdate)
+{
+	TArray<int> OutElements;
+	// custom reimplementation of GetVertexElements() accounting for the base mesh already being updated
+	for (int tid : ParentMesh->VtxTrianglesItr(SplitInfo.NewVertex)) // only care about triangles connected to the *new* vertex; these are updated
+	{
+		FIndex3i Triangle = GetTriangle(tid);
+		for (int j = 0; j < 3; ++j)
+		{
+			// though the ParentMesh vertex is NewVertex in the source mesh, it is still OriginalVertex in the ParentVertices array (since that hasn't been updated yet)
+			if (ParentVertices[Triangle[j]] == SplitInfo.OriginalVertex)
+			{
+				OutElements.AddUnique(Triangle[j]);
+			}
+		}
+	}
+
+	for (int ElementID : OutElements)
+	{
+		// Note: TrianglesToUpdate will include triangles that don't include the element, but that's ok; it just won't find any elements to update for those
+		//			(and this should be cheaper than constructing a new array for every element)
+		SplitElementWithNewParent(ElementID, SplitInfo.NewVertex, TrianglesToUpdate);
+	}
+}
+
+
 
 template<typename RealType, int ElementSize>
 void TDynamicMeshOverlay<RealType, ElementSize>::SetElementFromLerp(int SetElement, int ElementA, int ElementB, RealType Alpha)
diff --git a/Engine/Plugins/Experimental/GeometryProcessing/Source/DynamicMesh/Public/DynamicAttribute.h b/Engine/Plugins/Experimental/GeometryProcessing/Source/DynamicMesh/Public/DynamicAttribute.h
--- a/Engine/Plugins/Experimental/GeometryProcessing/Source/DynamicMesh/Public/DynamicAttribute.h
+++ b/Engine/Plugins/Experimental/GeometryProcessing/Source/DynamicMesh/Public/DynamicAttribute.h
@@ -125,6 +125,11 @@ public:
 	{
 	}
 
+	/** Update to reflect an edge merge in the parent mesh */
+	virtual void OnSplitVertex(const DynamicMeshInfo::FVertexSplitInfo& SplitInfo, const TArrayView<const int>& TrianglesToUpdate)
+	{
+	}
+
 };
 
 
@@ -262,6 +267,13 @@ public:
 			A->OnMergeEdges(MergeInfo);
 		}
 	}
+	virtual void OnSplitVertex(const DynamicMeshInfo::FVertexSplitInfo& SplitInfo, const TArrayView<const int>& TrianglesToUpdate)
+	{
+		for (TDynamicAttributeBase<ParentType>* A : RegisteredAttributes)
+		{
+			A->OnSplitVertex(SplitInfo, TrianglesToUpdate);
+		}
+	}
 };
 
 typedef TDynamicAttributeSetBase<FDynamicMesh3> FDynamicMeshAttributeSetBase;
diff --git a/Engine/Plugins/Experimental/GeometryProcessing/Source/DynamicMesh/Public/DynamicMesh3.h b/Engine/Plugins/Experimental/GeometryProcessing/Source/DynamicMesh/Public/DynamicMesh3.h
--- a/Engine/Plugins/Experimental/GeometryProcessing/Source/DynamicMesh/Public/DynamicMesh3.h
+++ b/Engine/Plugins/Experimental/GeometryProcessing/Source/DynamicMesh/Public/DynamicMesh3.h
@@ -89,6 +89,7 @@ using FEdgeSplitInfo = DynamicMeshInfo::FEdgeSplitInfo;
 using FEdgeCollapseInfo = DynamicMeshInfo::FEdgeCollapseInfo;
 using FMergeEdgesInfo = DynamicMeshInfo::FMergeEdgesInfo;
 using FPokeTriangleInfo = DynamicMeshInfo::FPokeTriangleInfo;
+using FVertexSplitInfo = DynamicMeshInfo::FVertexSplitInfo;
 
 	/** InvalidID indicates that a vertex/edge/triangle ID is invalid */
 	constexpr static int InvalidID = IndexConstants::InvalidID;
@@ -104,7 +105,7 @@ using FPokeTriangleInfo = DynamicMeshInfo::FPokeTriangleInfo;
 protected:
 	/** List of vertex positions */
 	TDynamicVector<double> Vertices{};
-	/** Reference counts of vertex indices. Iterate over this to find out which vertex indices are valid. */
+	/** Reference counts of vertex indices. For vertices that exist, the count is 1 + num_triangle_using_vertex. Iterate over this to find out which vertex indices are valid. */
 	FRefCountVector VertexRefCounts{};
 	/** (optional) List of per-vertex normals */
 	TOptional<TDynamicVector<float>> VertexNormals{};
@@ -117,7 +118,7 @@ protected:
 
 	/** List of triangle vertex-index triplets [Vert0 Vert1 Vert2]*/
 	TDynamicVector<int> Triangles;
-	/** Reference counts of triangle indices. Iterate over this to find out which triangle indices are valid. */
+	/** Reference counts of triangle indices. Ref count is always 1 if the triangle exists. Iterate over this to find out which triangle indices are valid. */
 	FRefCountVector TriangleRefCounts;
 	/** List of triangle edge triplets [Edge0 Edge1 Edge2] */
 	TDynamicVector<int> TriangleEdges;
@@ -128,7 +129,7 @@ protected:
 
 	/** List of edge elements. An edge is four elements [VertA, VertB, Tri0, Tri1], where VertA < VertB, and Tri1 may be InvalidID (if the edge is a boundary edge) */
 	TDynamicVector<int> Edges;
-	/** Reference counts of edge indices. Iterate over this to find out which edge indices are valid. */
+	/** Reference counts of edge indices. Ref count is always 1 if the edge exists. Iterate over this to find out which edge indices are valid. */
 	FRefCountVector EdgeRefCounts;
 
 	TUniquePtr<FDynamicMeshAttributeSet> AttributeSet{};
@@ -730,6 +731,9 @@ public:
 	/** Find edgeid for edge [a,b] from triangle that contains the edge. Faster than FindEdge() because it is constant-time. */
 	int FindEdgeFromTri(int VertexA, int VertexB, int TriangleID) const;
 
+	/** Find edgeid for edge connecting two triangles */
+	int FindEdgeFromTriPair(int TriangleA, int TriangleB) const;
+
 	/** Find triangle made up of any permutation of vertices [a,b,c] */
 	int FindTriangle(int A, int B, int C) const;
 
@@ -1062,6 +1066,24 @@ public:
 	virtual EMeshResult FlipEdge(int EdgeVertA, int EdgeVertB, FEdgeFlipInfo& FlipInfo);
 
 
+	/**
+	 * Clones the given vertex and updates any provided triangles to use the new vertex if/where they used the old one.
+	 * @param VertexID the vertex to split
+	 * @param TrianglesToUpdate triangles that should be updated to use the new vertex anywhere they previously had the old one
+	 * @param SplitInfo returned info about the new and modified mesh elements
+	 * @return Ok on success, or enum value indicates why operation cannot be applied. Mesh remains unmodified on error.
+	 */
+	virtual EMeshResult SplitVertex(int VertexID, const TArrayView<const int>& TrianglesToUpdate, FVertexSplitInfo& SplitInfo);
+
+	/**
+	 * Tests whether splitting the given vertex with the given triangles would leave no triangles attached to the original vertex (creating an isolated vertex)
+	 * @param VertexID the vertex to split
+	 * @param TrianglesToUpdate triangles that should be updated to use the new vertex anywhere they previously had the old one
+	 * @return true if calling SplitVertex with these arguments would leave an isolated vertex at the original VertexID
+	 */
+	virtual bool SplitVertexWouldLeaveIsolated(int VertexID, const TArrayView<const int>& TrianglesToUpdate);
+
+
 
 	/**
 	 * Collapse the edge between the two vertices, if topologically possible.
diff --git a/Engine/Plugins/Experimental/GeometryProcessing/Source/DynamicMesh/Public/DynamicMeshAttributeSet.h b/Engine/Plugins/Experimental/GeometryProcessing/Source/DynamicMesh/Public/DynamicMeshAttributeSet.h
--- a/Engine/Plugins/Experimental/GeometryProcessing/Source/DynamicMesh/Public/DynamicMeshAttributeSet.h
+++ b/Engine/Plugins/Experimental/GeometryProcessing/Source/DynamicMesh/Public/DynamicMeshAttributeSet.h
@@ -278,6 +278,7 @@ protected:
 	virtual void OnCollapseEdge(const DynamicMeshInfo::FEdgeCollapseInfo & collapseInfo);
 	virtual void OnPokeTriangle(const DynamicMeshInfo::FPokeTriangleInfo & pokeInfo);
 	virtual void OnMergeEdges(const DynamicMeshInfo::FMergeEdgesInfo & mergeInfo);
+	virtual void OnSplitVertex(const DynamicMeshInfo::FVertexSplitInfo& SplitInfo, const TArrayView<const int>& TrianglesToUpdate);
 
 	/**
 	 * Check validity of attributes
diff --git a/Engine/Plugins/Experimental/GeometryProcessing/Source/DynamicMesh/Public/DynamicMeshEditor.h b/Engine/Plugins/Experimental/GeometryProcessing/Source/DynamicMesh/Public/DynamicMeshEditor.h
--- a/Engine/Plugins/Experimental/GeometryProcessing/Source/DynamicMesh/Public/DynamicMeshEditor.h
+++ b/Engine/Plugins/Experimental/GeometryProcessing/Source/DynamicMesh/Public/DynamicMeshEditor.h
@@ -186,12 +186,32 @@ public:
 	/**
 	 * Finds boundary loops of connected components of a set of triangles, and duplicates the vertices
 	 * along the boundary, such that the triangles become disconnected.
+	 * Note: This can create a mesh with bowties if you pass in a subset of triangles that would have bowtie connectivity.
+	 *		 (it is impossible to create the paired LoopSetOut with 1:1 vertex pairs without leaving in the bowties?)
 	 * @param Triangles set of triangles
 	 * @param LoopSetOut set of boundary loops. LoopA is original loop which remains with "outer" triangles, and LoopB is new boundary loop of triangle set
 	 * @return true on success
 	 */
 	bool DisconnectTriangles(const TArray<int>& Triangles, TArray<FLoopPairSet>& LoopSetOut);
 
+	/**
+	* Disconnects triangles (without constructing boundary loops) such that the input Triangles are not connected to any other triangles in the mesh
+	* @param Triangles set of triangles
+	* @param bPreventBowties do some additional processing and vertex splitting as needed to prevent the creation of any new bowties
+	*/
+	void DisconnectTriangles(const TArray<int>& Triangles, bool bPreventBowties = true);
+
+
+	/**
+	 * Splits all bowties across the whole mesh
+	 */
+	void SplitBowties(FDynamicMeshEditResult& ResultOut);
+
+	/**
+	 * Splits any bowties specifically on the given vertex, and updates (does not reset!) ResultOut with any added vertices
+	 */
+	void SplitBowties(int VertexID, FDynamicMeshEditResult& ResultOut);
+
 
 	/**
 	 * Remove a list of triangles from the mesh, and optionally any vertices that are now orphaned
diff --git a/Engine/Plugins/Experimental/GeometryProcessing/Source/DynamicMesh/Public/DynamicMeshOverlay.h b/Engine/Plugins/Experimental/GeometryProcessing/Source/DynamicMesh/Public/DynamicMeshOverlay.h
--- a/Engine/Plugins/Experimental/GeometryProcessing/Source/DynamicMesh/Public/DynamicMeshOverlay.h
+++ b/Engine/Plugins/Experimental/GeometryProcessing/Source/DynamicMesh/Public/DynamicMeshOverlay.h
@@ -132,6 +132,33 @@ public:
 	void SplitVerticesWithPredicate(TFunctionRef<bool(int ElementIdx, int TriID)> ShouldSplitOutVertex, TFunctionRef<void(int ElementIdx, int TriID, RealType* FillVect)> GetNewElementValue);
 
 
+	/**
+	 * Create a new copy of ElementID, and update connected triangles in the TrianglesToUpdate array to reference the copy of ElementID where they used to reference ElementID
+	 * (Note: This just calls "SplitElementWithNewParent" with the existing element's parent id.)
+	 *
+	 * @param ElementID the element to copy
+	 * @param TrianglesToUpdate the triangles that should now reference the new element
+	 * @return the ID of the newly created element
+	 */
+	int SplitElement(int ElementID, const TArrayView<const int>& TrianglesToUpdate);
+
+	/**
+	* Create a new copy of ElementID, and update connected triangles in the TrianglesToUpdate array to reference the copy of ElementID where they used to reference ElementID.  The new element will have the given parent vertex ID.
+	*
+	* @param ElementID the element to copy
+	* @param SplitParentVertexID the new parent vertex for copied elements
+	* @param TrianglesToUpdate the triangles that should now reference the new element.  Note: this is allowed to include triangles that do not have the element at all; sometimes you may want to do so to avoid creating a new array for each call.
+	* @return the ID of the newly created element
+	*/
+	int SplitElementWithNewParent(int ElementID, int SplitParentVertexID, const TArrayView<const int>& TrianglesToUpdate);
+
+
+	/**
+	* Refine an existing overlay topology by splitting any bow ties
+	*/
+	void SplitBowties();
+
+
 	//
 	// Support for inserting element at specific ID. This is a bit tricky
 	// because we likely will need to update the free list in the RefCountVector, which
@@ -270,6 +297,8 @@ public:
 	void OnPokeTriangle(const DynamicMeshInfo::FPokeTriangleInfo& PokeInfo);
 	/** Update the overlay to reflect an edge merge in the parent mesh */
 	void OnMergeEdges(const DynamicMeshInfo::FMergeEdgesInfo& MergeInfo);
+	/** Update the overlay to reflect a vertex split in the parent mesh */
+	void OnSplitVertex(const DynamicMeshInfo::FVertexSplitInfo& SplitInfo, const TArrayView<const int>& TrianglesToUpdate);
 
 protected:
 	/** Set the value at an Element to be a linear interpolation of two other Elements */
diff --git a/Engine/Plugins/Experimental/GeometryProcessing/Source/DynamicMesh/Public/DynamicMeshTriangleAttribute.h b/Engine/Plugins/Experimental/GeometryProcessing/Source/DynamicMesh/Public/DynamicMeshTriangleAttribute.h
--- a/Engine/Plugins/Experimental/GeometryProcessing/Source/DynamicMesh/Public/DynamicMeshTriangleAttribute.h
+++ b/Engine/Plugins/Experimental/GeometryProcessing/Source/DynamicMesh/Public/DynamicMeshTriangleAttribute.h
@@ -207,7 +207,7 @@ public:
 public:
 
 	/** Update the overlay to reflect an edge split in the parent mesh */
-	void OnSplitEdge(const DynamicMeshInfo::FEdgeSplitInfo& SplitInfo)
+	void OnSplitEdge(const DynamicMeshInfo::FEdgeSplitInfo& SplitInfo) override
 	{
 		CopyValue(SplitInfo.OriginalTriangles.A, SplitInfo.NewTriangles.A);
 		if (SplitInfo.OriginalTriangles.B != FDynamicMesh3::InvalidID)
@@ -217,30 +217,37 @@ public:
 	}
 
 	/** Update the overlay to reflect an edge flip in the parent mesh */
-	void OnFlipEdge(const DynamicMeshInfo::FEdgeFlipInfo& FlipInfo)
+	void OnFlipEdge(const DynamicMeshInfo::FEdgeFlipInfo& FlipInfo) override
 	{
 		// yikes! triangles did not actually change so we will leave attrib unmodified
 	}
 
 	/** Update the overlay to reflect an edge collapse in the parent mesh */
-	void OnCollapseEdge(const DynamicMeshInfo::FEdgeCollapseInfo& CollapseInfo)
+	void OnCollapseEdge(const DynamicMeshInfo::FEdgeCollapseInfo& CollapseInfo) override
 	{
 		// nothing to do here, triangles were only deleted
 	}
 
 	/** Update the overlay to reflect a face poke in the parent mesh */
-	void OnPokeTriangle(const DynamicMeshInfo::FPokeTriangleInfo& PokeInfo)
+	void OnPokeTriangle(const DynamicMeshInfo::FPokeTriangleInfo& PokeInfo) override
 	{
 		CopyValue(PokeInfo.OriginalTriangle, PokeInfo.NewTriangles.A);
 		CopyValue(PokeInfo.OriginalTriangle, PokeInfo.NewTriangles.B);
 	}
 
 	/** Update the overlay to reflect an edge merge in the parent mesh */
-	void OnMergeEdges(const DynamicMeshInfo::FMergeEdgesInfo& MergeInfo)
+	void OnMergeEdges(const DynamicMeshInfo::FMergeEdgesInfo& MergeInfo) override
 	{
 		// nothing to do here because triangles did not change
 	}
 
+	/** Update the overlay to reflect a vertex split in the parent */
+	void OnSplitVertex(const DynamicMeshInfo::FVertexSplitInfo& SplitInfo, const TArrayView<const int>& TrianglesToUpdate) override
+	{
+		// nothing to do here because triangles did not change
+	}
+
+
 };
 
 
diff --git a/Engine/Plugins/Experimental/GeometryProcessing/Source/DynamicMesh/Public/DynamicVertexAttribute.h b/Engine/Plugins/Experimental/GeometryProcessing/Source/DynamicMesh/Public/DynamicVertexAttribute.h
--- a/Engine/Plugins/Experimental/GeometryProcessing/Source/DynamicMesh/Public/DynamicVertexAttribute.h
+++ b/Engine/Plugins/Experimental/GeometryProcessing/Source/DynamicMesh/Public/DynamicVertexAttribute.h
@@ -226,6 +226,12 @@ public:
 		SetAttributeFromLerp(MergeInfo.KeptVerts.B, MergeInfo.KeptVerts.B, MergeInfo.RemovedVerts.B, .5);
 	}
 
+	/** Update the overlay to reflect a vertex split in the parent */
+	void OnSplitVertex(const FDynamicMesh3::FVertexSplitInfo& SplitInfo, const TArrayView<const int>& TrianglesToUpdate) override
+	{
+		CopyValue(SplitInfo.OriginalVertex, SplitInfo.NewVertex);
+	}
+
 	virtual TUniquePtr<TDynamicAttributeChangeBase<ParentType>> NewBlankChange() override
 	{
 		return MakeUnique<TDynamicVertexAttributeChange<AttribValueType, AttribDimension, ParentType>>();
diff --git a/Engine/Plugins/Experimental/GeometryProcessing/Source/DynamicMesh/Public/InfoTypes.h b/Engine/Plugins/Experimental/GeometryProcessing/Source/DynamicMesh/Public/InfoTypes.h
--- a/Engine/Plugins/Experimental/GeometryProcessing/Source/DynamicMesh/Public/InfoTypes.h
+++ b/Engine/Plugins/Experimental/GeometryProcessing/Source/DynamicMesh/Public/InfoTypes.h
@@ -97,4 +97,13 @@ struct FPokeTriangleInfo
 
 	FVector3d BaryCoords;				// barycentric coords that NewVertex was inserted at
 };
+
+/** Information about mesh elements modified/created by SplitVertex() */
+struct FVertexSplitInfo
+{
+	int OriginalVertex;
+	int NewVertex;
+	// if needed could possibly add information about added edges?  but it would be a dynamic array, and there is no use for it yet.
+	// modified triangles are passed as input to the function, no need to store those here.
+};
 }
diff --git a/Engine/Plugins/Experimental/GeometryProcessing/Source/DynamicMesh/Public/MeshRegionBoundaryLoops.h b/Engine/Plugins/Experimental/GeometryProcessing/Source/DynamicMesh/Public/MeshRegionBoundaryLoops.h
--- a/Engine/Plugins/Experimental/GeometryProcessing/Source/DynamicMesh/Public/MeshRegionBoundaryLoops.h
+++ b/Engine/Plugins/Experimental/GeometryProcessing/Source/DynamicMesh/Public/MeshRegionBoundaryLoops.h
@@ -51,9 +51,7 @@ public:
 
 protected:
 
-
-	// TODO: C# code used IndexFlagSet here, which is sparse if region size << mesh size
-	// list of included triangles and edges
+	// sets of included triangles and edges
 	FIndexFlagSet Triangles;
 	FIndexFlagSet Edges;
 	TArray<int> edges_roi;
diff --git a/Engine/Plugins/Experimental/MeshModelingToolset/Source/MeshModelingTools/Private/MeshSelectionTool.cpp b/Engine/Plugins/Experimental/MeshModelingToolset/Source/MeshModelingTools/Private/MeshSelectionTool.cpp
--- a/Engine/Plugins/Experimental/MeshModelingToolset/Source/MeshModelingTools/Private/MeshSelectionTool.cpp
+++ b/Engine/Plugins/Experimental/MeshModelingToolset/Source/MeshModelingTools/Private/MeshSelectionTool.cpp
@@ -10,6 +10,7 @@
 #include "Changes/MeshChange.h"
 #include "Util/ColorConstants.h"
 #include "Selections/MeshConnectedComponents.h"
+#include "MeshRegionBoundaryLoops.h"
 #include "MeshIndexUtil.h"
 #include "AssetGenerationUtil.h"
 #include "ToolSetupUtil.h"
@@ -560,7 +561,7 @@
 }
 
 
-void UMeshSelectionTool::Tick(float DeltaTime)
+void UMeshSelectionTool::Tick(float DeltaTime)
 {
 	UDynamicMeshBrushTool::Tick(DeltaTime);
 
@@ -659,6 +660,10 @@
 			DeleteSelectedTriangles();
 			break;
 
+		case EMeshSelectionToolActions::DisconnectSelected:
+			DisconnectSelectedTriangles();
+			break;
+
 		case EMeshSelectionToolActions::SeparateSelected:
 			SeparateSelectedTriangles();
 			break;
@@ -920,6 +925,49 @@
 }
 
 
+void UMeshSelectionTool::DisconnectSelectedTriangles()
+{
+	check(SelectionType == EMeshSelectionElementType::Face);
+	TArray<int32> SelectedFaces = Selection->GetElements(EMeshSelectionElementType::Face);
+	if (SelectedFaces.Num() == 0)
+	{
+		return;
+	}
+
+	TUniquePtr<FToolCommandChangeSequence> ChangeSeq = MakeUnique<FToolCommandChangeSequence>();
+
+	// split out selected triangles and emit triangle change
+	TUniquePtr<FMeshChange> MeshChange = PreviewMesh->TrackedEditMesh(
+		[&SelectedFaces](FDynamicMesh3& Mesh, FDynamicMeshChangeTracker& ChangeTracker)
+	{
+		// save vertices and triangles that are on the boundary of the selection
+		FMeshRegionBoundaryLoops BoundaryLoops(&Mesh, SelectedFaces);
+		for (const FEdgeLoop& Loop : BoundaryLoops.Loops)
+		{
+			for (int VID : Loop.Vertices)
+			{
+				ChangeTracker.SaveVertex(VID);
+				// include the whole one-ring in case the disconnect creates bowties that need to be split
+				for (int TID : Mesh.VtxTrianglesItr(VID))
+				{
+					ChangeTracker.SaveTriangle(TID, false);
+				}
+			}
+		}
+
+		FDynamicMeshEditor Editor(&Mesh);
+		Editor.DisconnectTriangles(SelectedFaces);
+	});
+	ChangeSeq->AppendChange(PreviewMesh, MoveTemp(MeshChange));
+
+	// emit combined change sequence
+	GetToolManager()->EmitObjectChange(this, MoveTemp(ChangeSeq), LOCTEXT("MeshSelectionToolDisconnectFaces", "Disconnect Faces"));
+
+	bHaveModifiedMesh = true;
+}
+
+
+
 void UMeshSelectionTool::SeparateSelectedTriangles()
 {
 #if WITH_EDITOR
diff --git a/Engine/Plugins/Experimental/MeshModelingToolset/Source/MeshModelingTools/Private/UVLayoutTool.cpp b/Engine/Plugins/Experimental/MeshModelingToolset/Source/MeshModelingTools/Private/UVLayoutTool.cpp
new file mode 100644
--- /dev/null
+++ b/Engine/Plugins/Experimental/MeshModelingToolset/Source/MeshModelingTools/Private/UVLayoutTool.cpp
@@ -0,0 +1,275 @@
+// Copyright 1998-2019 Epic Games, Inc. All Rights Reserved.
+
+#include "UVLayoutTool.h"
+#include "InteractiveToolManager.h"
+#include "ToolBuilderUtil.h"
+
+#include "ToolSetupUtil.h"
+
+#include "DynamicMesh3.h"
+#include "BaseBehaviors/MultiClickSequenceInputBehavior.h"
+#include "Selection/SelectClickedAction.h"
+
+#include "MeshDescriptionToDynamicMesh.h"
+#include "DynamicMeshToMeshDescription.h"
+#include "Materials/MaterialInstanceDynamic.h"
+
+#include "ParameterizationOps/UVLayoutOp.h"
+
+#include "AssetGenerationUtil.h"
+
+#define LOCTEXT_NAMESPACE "UUVLayoutTool"
+
+
+/*
+ * ToolBuilder
+ */
+
+
+bool UUVLayoutToolBuilder::CanBuildTool(const FToolBuilderState& SceneState) const
+{
+	return ToolBuilderUtil::CountComponents(SceneState, CanMakeComponentTarget) >= 1;
+}
+
+UInteractiveTool* UUVLayoutToolBuilder::BuildTool(const FToolBuilderState& SceneState) const
+{
+	UUVLayoutTool* NewTool = NewObject<UUVLayoutTool>(SceneState.ToolManager);
+
+	TArray<UActorComponent*> Components = ToolBuilderUtil::FindAllComponents(SceneState, CanMakeComponentTarget);
+	check(Components.Num() > 0);
+
+	TArray<TUniquePtr<FPrimitiveComponentTarget>> ComponentTargets;
+	for (UActorComponent* ActorComponent : Components)
+	{
+		auto* MeshComponent = Cast<UPrimitiveComponent>(ActorComponent);
+		if ( MeshComponent )
+		{
+			ComponentTargets.Add(MakeComponentTarget(MeshComponent));
+		}
+	}
+
+	NewTool->SetSelection(MoveTemp(ComponentTargets));
+	NewTool->SetWorld(SceneState.World, SceneState.GizmoManager);
+	NewTool->SetAssetAPI(AssetAPI);
+
+	return NewTool;
+}
+
+
+
+/*
+ * Tool
+ */
+
+UUVLayoutToolProperties::UUVLayoutToolProperties()
+{
+
+}
+
+UUVLayoutAdvancedProperties::UUVLayoutAdvancedProperties()
+{
+}
+
+
+UUVLayoutTool::UUVLayoutTool()
+{
+}
+
+void UUVLayoutTool::SetWorld(UWorld* World, UInteractiveGizmoManager* GizmoManagerIn)
+{
+	this->TargetWorld = World;
+}
+
+void UUVLayoutTool::Setup()
+{
+	UInteractiveTool::Setup();
+
+	// hide input StaticMeshComponent
+	for (TUniquePtr<FPrimitiveComponentTarget>& ComponentTarget : ComponentTargets)
+	{
+		ComponentTarget->SetOwnerVisibility(false);
+	}
+
+	BasicProperties = NewObject<UUVLayoutToolProperties>(this, TEXT("UV Projection Settings"));
+	AdvancedProperties = NewObject<UUVLayoutAdvancedProperties>(this, TEXT("Advanced Settings"));
+
+	// initialize our properties
+	AddToolPropertySource(BasicProperties);
+	AddToolPropertySource(AdvancedProperties);
+
+	MaterialSettings = NewObject<UExistingMeshMaterialProperties>(this);
+	MaterialSettings->Setup();
+	AddToolPropertySource(MaterialSettings);
+
+	// initialize the PreviewMesh+BackgroundCompute object
+	UpdateNumPreviews();
+	
+	for (UMeshOpPreviewWithBackgroundCompute* Preview : Previews)
+	{
+		Preview->InvalidateResult();
+	}
+}
+
+
+void UUVLayoutTool::UpdateNumPreviews()
+{
+	int32 CurrentNumPreview = Previews.Num();
+	int32 TargetNumPreview = ComponentTargets.Num();
+	if (TargetNumPreview < CurrentNumPreview)
+	{
+		for (int32 PreviewIdx = CurrentNumPreview - 1; PreviewIdx >= TargetNumPreview; PreviewIdx--)
+		{
+			Previews[PreviewIdx]->Cancel();
+		}
+		Previews.SetNum(TargetNumPreview);
+		OriginalDynamicMeshes.SetNum(TargetNumPreview);
+	}
+	else
+	{
+		OriginalDynamicMeshes.SetNum(TargetNumPreview);
+		for (int32 PreviewIdx = CurrentNumPreview; PreviewIdx < TargetNumPreview; PreviewIdx++)
+		{
+			UUVLayoutOperatorFactory *OpFactory = NewObject<UUVLayoutOperatorFactory>();
+			OpFactory->Tool = this;
+			OpFactory->ComponentIndex = PreviewIdx;
+			OriginalDynamicMeshes[PreviewIdx] = MakeShared<FDynamicMesh3>();
+			FMeshDescriptionToDynamicMesh Converter;
+			Converter.bPrintDebugMessages = true;
+			Converter.Convert(ComponentTargets[PreviewIdx]->GetMesh(), *OriginalDynamicMeshes[PreviewIdx]);
+
+			UMeshOpPreviewWithBackgroundCompute* Preview = Previews.Add_GetRef(NewObject<UMeshOpPreviewWithBackgroundCompute>(OpFactory, "Preview"));
+			Preview->Setup(this->TargetWorld, OpFactory);
+			Preview->ConfigureMaterials(
+				ToolSetupUtil::GetDefaultMaterial(GetToolManager(), ComponentTargets[PreviewIdx]->GetMaterial(0)),
+				ToolSetupUtil::GetDefaultWorkingMaterial(GetToolManager())
+			);
+			Preview->SetVisibility(true);
+		}
+	}
+}
+
+
+void UUVLayoutTool::Shutdown(EToolShutdownType ShutdownType)
+{
+	// Restore (unhide) the source meshes
+	for (TUniquePtr<FPrimitiveComponentTarget>& ComponentTarget : ComponentTargets)
+	{
+		ComponentTarget->SetOwnerVisibility(true);
+	}
+
+	TArray<FDynamicMeshOpResult> Results;
+	for (UMeshOpPreviewWithBackgroundCompute* Preview : Previews)
+	{
+		Results.Emplace(Preview->Shutdown());
+	}
+	if (ShutdownType == EToolShutdownType::Accept)
+	{
+		GenerateAsset(Results);
+	}
+}
+
+void UUVLayoutTool::SetAssetAPI(IToolsContextAssetAPI* AssetAPIIn)
+{
+	this->AssetAPI = AssetAPIIn;
+}
+
+TUniquePtr<FDynamicMeshOperator> UUVLayoutOperatorFactory::MakeNewOperator()
+{
+	TUniquePtr<FUVLayoutOp> Op = MakeUnique<FUVLayoutOp>();
+
+	FTransform LocalToWorld = Tool->ComponentTargets[ComponentIndex]->GetWorldTransform();
+	Op->OriginalMesh = Tool->OriginalDynamicMeshes[ComponentIndex];
+	
+	Op->SetTransform(LocalToWorld);
+
+	return Op;
+}
+
+
+
+void UUVLayoutTool::Render(IToolsContextRenderAPI* RenderAPI)
+{
+}
+
+void UUVLayoutTool::Tick(float DeltaTime)
+{
+	for (UMeshOpPreviewWithBackgroundCompute* Preview : Previews)
+	{
+		Preview->Tick(DeltaTime);
+	}
+}
+
+
+
+#if WITH_EDITOR
+void UUVLayoutTool::PostEditChangeProperty(FPropertyChangedEvent& PropertyChangedEvent)
+{
+	UpdateNumPreviews();
+	for (UMeshOpPreviewWithBackgroundCompute* Preview : Previews)
+	{
+		Preview->InvalidateResult();
+	}
+}
+#endif
+
+void UUVLayoutTool::OnPropertyModified(UObject* PropertySet, UProperty* Property)
+{
+	// if we don't know what changed, or we know checker density changed, update checker material
+	MaterialSettings->UpdateMaterials();
+	for (int PreviewIdx = 0; PreviewIdx < Previews.Num(); PreviewIdx++)
+	{
+		UMeshOpPreviewWithBackgroundCompute* Preview = Previews[PreviewIdx];
+		MaterialSettings->SetMaterialIfChanged(ComponentTargets[PreviewIdx]->GetMaterial(0), Preview->StandardMaterial, [&Preview, this](UMaterialInterface* Material)
+		{
+			Preview->ConfigureMaterials(ToolSetupUtil::GetDefaultMaterial(GetToolManager(), Material), Preview->WorkingMaterial);
+		});
+	}
+	
+	UpdateNumPreviews();
+	for (UMeshOpPreviewWithBackgroundCompute* Preview : Previews)
+	{
+		Preview->InvalidateResult();
+	}
+}
+
+bool UUVLayoutTool::HasAccept() const
+{
+	return true;
+}
+
+bool UUVLayoutTool::CanAccept() const
+{
+	for (UMeshOpPreviewWithBackgroundCompute* Preview : Previews)
+	{
+		if (!Preview->HaveValidResult())
+		{
+			return false;
+		}
+	}
+	return true;
+}
+
+
+void UUVLayoutTool::GenerateAsset(const TArray<FDynamicMeshOpResult>& Results)
+{
+	GetToolManager()->BeginUndoTransaction(LOCTEXT("UVLayoutToolTransactionName", "UV Layout Tool"));
+
+	check(Results.Num() == ComponentTargets.Num());
+	
+	for (int32 ComponentIdx = 0; ComponentIdx < ComponentTargets.Num(); ComponentIdx++)
+	{
+		check(Results[ComponentIdx].Mesh.Get() != nullptr);
+		ComponentTargets[ComponentIdx]->CommitMesh([&Results, &ComponentIdx, this](const FPrimitiveComponentTarget::FCommitParams& CommitParams)
+		{
+			FDynamicMeshToMeshDescription Converter;
+			Converter.Convert(Results[ComponentIdx].Mesh.Get(), *CommitParams.MeshDescription);
+		});
+	}
+
+	GetToolManager()->EndUndoTransaction();
+}
+
+
+
+
+#undef LOCTEXT_NAMESPACE
diff --git a/Engine/Plugins/Experimental/MeshModelingToolset/Source/MeshModelingTools/Public/MeshSelectionTool.h b/Engine/Plugins/Experimental/MeshModelingToolset/Source/MeshModelingTools/Public/MeshSelectionTool.h
--- a/Engine/Plugins/Experimental/MeshModelingToolset/Source/MeshModelingTools/Public/MeshSelectionTool.h
+++ b/Engine/Plugins/Experimental/MeshModelingToolset/Source/MeshModelingTools/Public/MeshSelectionTool.h
@@ -39,6 +39,7 @@ enum class EMeshSelectionToolActions
 	ExpandToConnected,
 
 	DeleteSelected,
+	DisconnectSelected,
 	SeparateSelected,
 	FlipSelected,
 	CreateGroup
@@ -122,13 +123,19 @@ public:
 		PostAction(EMeshSelectionToolActions::SeparateSelected);
 	}
 
-	UFUNCTION(CallInEditor, Category = MeshEdits, meta = (DisplayName = "Flip Normals", DisplayPriority = 3))
+	UFUNCTION(CallInEditor, Category = MeshEdits, meta = (DisplayName = "Disconnect", DisplayPriority = 3))
+	void DisconnectTriangles() 
+	{
+		PostAction(EMeshSelectionToolActions::DisconnectSelected);
+	}
+
+	UFUNCTION(CallInEditor, Category = MeshEdits, meta = (DisplayName = "Flip Normals", DisplayPriority = 4))
 	void FlipNormals() 
 	{
 		PostAction(EMeshSelectionToolActions::FlipSelected);
 	}
 
-	UFUNCTION(CallInEditor, Category = MeshEdits, meta = (DisplayName = "Create Polygroup", DisplayPriority = 4))
+	UFUNCTION(CallInEditor, Category = MeshEdits, meta = (DisplayName = "Create Polygroup", DisplayPriority = 5))
 	void CreatePolygroup()
 	{
 		PostAction(EMeshSelectionToolActions::CreateGroup);
@@ -310,7 +317,8 @@ protected:
 	void ExpandToConnected();
 
 	void DeleteSelectedTriangles();
-	void SeparateSelectedTriangles();
+	void DisconnectSelectedTriangles(); // disconnects edges between selected and not-selected triangles; keeps all triangles in the same mesh
+	void SeparateSelectedTriangles(); // separates out selected triangles to a new mesh, removing them from the current mesh
 	void FlipSelectedTriangles();
 	void AssignNewGroupToSelectedTriangles();
 
diff --git a/Engine/Plugins/Experimental/MeshModelingToolset/Source/MeshModelingTools/Public/UVLayoutTool.h b/Engine/Plugins/Experimental/MeshModelingToolset/Source/MeshModelingTools/Public/UVLayoutTool.h
new file mode 100644
--- /dev/null
+++ b/Engine/Plugins/Experimental/MeshModelingToolset/Source/MeshModelingTools/Public/UVLayoutTool.h
@@ -0,0 +1,158 @@
+// Copyright 1998-2019 Epic Games, Inc. All Rights Reserved.
+
+#pragma once
+
+#include "CoreMinimal.h"
+#include "UObject/NoExportTypes.h"
+#include "MultiSelectionTool.h"
+#include "InteractiveToolBuilder.h"
+#include "MeshOpPreviewHelpers.h"
+#include "ToolDataVisualizer.h"
+#include "DynamicMesh3.h"
+#include "BaseTools/SingleClickTool.h"
+#include "Properties/MeshMaterialProperties.h"
+
+#include "UVLayoutTool.generated.h"
+
+
+// predeclarations
+struct FMeshDescription;
+class USimpleDynamicMeshComponent;
+
+
+/**
+ *
+ */
+UCLASS()
+class MESHMODELINGTOOLS_API UUVLayoutToolBuilder : public UInteractiveToolBuilder
+{
+	GENERATED_BODY()
+
+public:
+	IToolsContextAssetAPI* AssetAPI = nullptr;
+
+	virtual bool CanBuildTool(const FToolBuilderState& SceneState) const override;
+	virtual UInteractiveTool* BuildTool(const FToolBuilderState& SceneState) const override;
+};
+
+
+
+
+
+/**
+ * Standard properties
+ */
+UCLASS()
+class MESHMODELINGTOOLS_API UUVLayoutToolProperties : public UInteractiveToolPropertySet
+{
+	GENERATED_BODY()
+
+public:
+	UUVLayoutToolProperties();
+
+
+// TODO: add config settings for uvlayout (currently nothing really in the actual algorithm is exposed as a usable parameter)
+
+};
+
+
+
+
+/**
+ * Advanced properties
+ */
+UCLASS()
+class MESHMODELINGTOOLS_API UUVLayoutAdvancedProperties : public UInteractiveToolPropertySet
+{
+	GENERATED_BODY()
+
+public:
+	UUVLayoutAdvancedProperties();
+};
+
+
+/**
+ * Factory with enough info to spawn the background-thread Operator to do a chunk of work for the tool
+ *  stores a pointer to the tool and enough info to know which specific operator it should spawn
+ */
+UCLASS()
+class MESHMODELINGTOOLS_API UUVLayoutOperatorFactory : public UObject, public IDynamicMeshOperatorFactory
+{
+	GENERATED_BODY()
+
+public:
+	// IDynamicMeshOperatorFactory API
+	virtual TUniquePtr<FDynamicMeshOperator> MakeNewOperator() override;
+
+	UPROPERTY()
+	UUVLayoutTool *Tool;
+
+	int ComponentIndex;
+
+};
+
+/**
+ * Simple Mesh Normal Updating Tool
+ */
+UCLASS()
+class MESHMODELINGTOOLS_API UUVLayoutTool : public UMultiSelectionTool
+{
+	GENERATED_BODY()
+
+public:
+
+	friend UUVLayoutOperatorFactory;
+
+	UUVLayoutTool();
+
+	virtual void Setup() override;
+	virtual void Shutdown(EToolShutdownType ShutdownType) override;
+
+	virtual void SetWorld(UWorld* World, UInteractiveGizmoManager* GizmoManagerIn);
+	virtual void SetAssetAPI(IToolsContextAssetAPI* AssetAPI);
+
+	virtual void Tick(float DeltaTime) override;
+	virtual void Render(IToolsContextRenderAPI* RenderAPI) override;
+
+	virtual bool HasCancel() const override { return true; }
+	virtual bool HasAccept() const override;
+	virtual bool CanAccept() const override;
+
+#if WITH_EDITOR
+	virtual void PostEditChangeProperty(FPropertyChangedEvent &PropertyChangedEvent) override;
+#endif
+
+	virtual void OnPropertyModified(UObject* PropertySet, UProperty* Property) override;
+
+protected:
+
+	UPROPERTY()
+	UUVLayoutToolProperties* BasicProperties;
+
+	UPROPERTY()
+	UUVLayoutAdvancedProperties* AdvancedProperties;
+
+	UPROPERTY()
+	UExistingMeshMaterialProperties* MaterialSettings = nullptr;
+
+
+
+	UPROPERTY()
+	TArray<UMeshOpPreviewWithBackgroundCompute*> Previews;
+
+
+	UPROPERTY()
+	UMaterialInstanceDynamic* CheckerMaterial = nullptr;
+
+protected:
+	TArray<TSharedPtr<FDynamicMesh3>> OriginalDynamicMeshes;
+
+	UWorld* TargetWorld;
+	IToolsContextAssetAPI* AssetAPI;
+
+	FViewCameraState CameraState;
+
+	void UpdateNumPreviews();
+
+	void GenerateAsset(const TArray<FDynamicMeshOpResult>& Results);
+};
diff --git a/Engine/Plugins/Experimental/MeshModelingToolset/Source/ModelingOperators/ModelingOperators.Build.cs b/Engine/Plugins/Experimental/MeshModelingToolset/Source/ModelingOperators/ModelingOperators.Build.cs
--- a/Engine/Plugins/Experimental/MeshModelingToolset/Source/ModelingOperators/ModelingOperators.Build.cs
+++ b/Engine/Plugins/Experimental/MeshModelingToolset/Source/ModelingOperators/ModelingOperators.Build.cs
@@ -41,6 +41,7 @@ public class ModelingOperators : ModuleRules
 				"MeshConversion",
                 "GeometryAlgorithms", // required for constrained Delaunay triangulation
                 "MeshSolverUtilities", // required by the smoothing operators
+				"MeshUtilitiesCommon", // required by uvlayoutop
 				// ... add other public dependencies that you statically link with here ...
 			}
             );
diff --git a/Engine/Plugins/Experimental/MeshModelingToolset/Source/ModelingOperators/Private/ParameterizationOps/UVLayoutOp.cpp b/Engine/Plugins/Experimental/MeshModelingToolset/Source/ModelingOperators/Private/ParameterizationOps/UVLayoutOp.cpp
new file mode 100644
--- /dev/null
+++ b/Engine/Plugins/Experimental/MeshModelingToolset/Source/ModelingOperators/Private/ParameterizationOps/UVLayoutOp.cpp
@@ -0,0 +1,161 @@
+// Copyright 1998-2019 Epic Games, Inc. All Rights Reserved.
+
+#include "ParameterizationOps/UVLayoutOp.h"
+
+#include "OverlappingCorners.h"
+#include "LayoutUV.h"
+
+#include "DynamicMesh3.h"
+#include "DynamicMeshAttributeSet.h"
+
+
+void FUVLayoutOp::SetTransform(const FTransform& Transform) {
+	ResultTransform = (FTransform3d)Transform;
+}
+
+// very task-specific lightweight view of a FDynamicMesh3; ONLY for compact meshes with attributes
+// not intended for use outside of this narrow context
+class FCompactDynamicMeshWithAttributesLayoutView : public FLayoutUV::IMeshView
+{
+	FDynamicMesh3* Mesh;
+	int UVLayerInput, UVLayerOutput;
+
+public:
+	FCompactDynamicMeshWithAttributesLayoutView(FDynamicMesh3* Mesh, int UVLayerIn=0, int UVLayerOut=0) : Mesh(Mesh), UVLayerInput(UVLayerIn), UVLayerOutput(UVLayerOut)
+	{
+		check(Mesh->HasAttributes());
+		check(Mesh->IsCompact());
+		check(Mesh->Attributes()->NumUVLayers() > UVLayerIn);
+		check(Mesh->Attributes()->NumUVLayers() > UVLayerOut);
+	}
+
+	uint32 GetNumIndices() const override
+	{
+		return Mesh->TriangleCount() * 3;
+	}
+	FVector GetPosition(uint32 Index) const override
+	{
+		return FVector(Mesh->GetVertex(Mesh->GetTriangle(Index / 3)[Index % 3]));
+	}
+	FVector GetNormal(uint32 Index) const override
+	{
+		const FDynamicMeshNormalOverlay* NormalOverlay = Mesh->Attributes()->PrimaryNormals();
+		FIndex3i NormalEl = NormalOverlay->GetTriangle(Index / 3);
+		FVector3f Normal;
+		NormalOverlay->GetElement(NormalEl[Index % 3], Normal);
+		return FVector(Normal);
+	}
+	FVector2D GetInputTexcoord(uint32 Index) const override
+	{
+		const FDynamicMeshUVOverlay* UVOverlay = Mesh->Attributes()->GetUVLayer(UVLayerInput);
+		FIndex3i UVEl = UVOverlay->GetTriangle(Index / 3);
+		FVector2f UV;
+		UVOverlay->GetElement(UVEl[Index % 3], UV);
+		return FVector2D(UV);
+	}
+
+	/**
+	 * This function is kind of nonsense for our use case as we cannot initialize a UV overlay from a single number
+	 * (unless we make every triangle disconnected in a fully raw wedge thing, which we never ever want to do)
+	 * So we assume the calling code will only call this w/ NumIndices matching the input layer & the intent of making the UV layers the same.
+	 */
+	void InitOutputTexcoords(uint32 Num) override
+	{
+		if (UVLayerInput != UVLayerOutput)
+		{
+			Mesh->Attributes()->GetUVLayer(UVLayerOutput)->Copy(*Mesh->Attributes()->GetUVLayer(UVLayerInput));
+		}
+		check(Num == Mesh->TriangleCount() * 3);
+	}
+
+	void SetOutputTexcoord(uint32 Index, const FVector2D& Value) override
+	{
+		FDynamicMeshUVOverlay* UVOverlay = Mesh->Attributes()->GetUVLayer(UVLayerOutput);
+		FIndex3i UVEl = UVOverlay->GetTriangle(Index / 3);
+		UVOverlay->SetElement(UVEl[Index % 3], FVector2f(Value));
+	}
+};
+
+/**
+ * Create an overlapping corner map to identify wedge indices that share the same UV element index for a given UV layer
+ */
+FOverlappingCorners OverlappingCornersFromUVs(const FDynamicMesh3* Mesh, int UVLayerIndex)
+{
+	// track all wedge indices that map to the same UV element
+	const FDynamicMeshUVOverlay* UVOverlay = Mesh->Attributes()->GetUVLayer(UVLayerIndex);
+	FOverlappingCorners Overlaps;
+	Overlaps.Init(Mesh->TriangleCount()*3);
+	for (int ElementID : UVOverlay->ElementIndicesItr())
+	{
+		int VertID = UVOverlay->GetParentVertex(ElementID);
+		int LastWedgeIdx = -1;
+		for (int TriID : Mesh->VtxTrianglesItr(VertID))
+		{
+			FIndex3i ElTri = UVOverlay->GetTriangle(TriID);
+			for (int SubIdx = 0; SubIdx < 3; SubIdx++)
+			{
+				if (ElTri[SubIdx] == ElementID)
+				{
+					int WedgeIdx = TriID * 3 + SubIdx;
+					if (LastWedgeIdx != -1)
+					{
+						Overlaps.Add(LastWedgeIdx, WedgeIdx);
+					}
+					LastWedgeIdx = WedgeIdx;
+				}
+			}
+		}
+	}
+	Overlaps.FinishAdding();
+	return Overlaps;
+}
+
+void FUVLayoutOp::CalculateResult(FProgressCancel* Progress)
+{
+	if (Progress->Cancelled())
+	{
+		return;
+	}
+	bool bDiscardAttributes = false;
+	ResultMesh->Copy(*OriginalMesh, true, true, true, !bDiscardAttributes);
+
+	if (!ensureMsgf(ResultMesh->HasAttributes(), TEXT("Attributes not found on mesh? Conversion should always create them, so this operator should not need to do so.")))
+	{
+		ResultMesh->EnableAttributes();
+	}
+
+	if (Progress->Cancelled())
+	{
+		return;
+	}
+
+	int UVLayerInput = 0, UVLayerOutput = 0;
+	ResultMesh->Attributes()->GetUVLayer(UVLayerInput)->SplitBowties();
+
+	if (Progress->Cancelled())
+	{
+		return;
+	}
+
+	FCompactDynamicMeshWithAttributesLayoutView MeshView(ResultMesh.Get(), UVLayerInput, UVLayerOutput);
+	FLayoutUV LayoutUV(MeshView);
+	FOverlappingCorners Overlaps = OverlappingCornersFromUVs(ResultMesh.Get(), UVLayerInput);
+	if (Progress->Cancelled())
+	{
+		return;
+	}
+
+	LayoutUV.FindCharts(Overlaps);
+	if (Progress->Cancelled())
+	{
+		return;
+	}
+
+	LayoutUV.FindBestPacking(1024); // TODO: make this a parameter
+	if (Progress->Cancelled())
+	{
+		return;
+	}
+
+	LayoutUV.CommitPackedUVs();
+}
diff --git a/Engine/Plugins/Experimental/MeshModelingToolset/Source/ModelingOperators/Public/ParameterizationOps/UVLayoutOp.h b/Engine/Plugins/Experimental/MeshModelingToolset/Source/ModelingOperators/Public/ParameterizationOps/UVLayoutOp.h
new file mode 100644
--- /dev/null
+++ b/Engine/Plugins/Experimental/MeshModelingToolset/Source/ModelingOperators/Public/ParameterizationOps/UVLayoutOp.h
@@ -0,0 +1,29 @@
+// Copyright 1998-2019 Epic Games, Inc. All Rights Reserved.
+
+#pragma once
+
+#include "CoreMinimal.h"
+#include "Util/ProgressCancel.h"
+#include "ModelingOperators.h"
+
+
+
+class MODELINGOPERATORS_API FUVLayoutOp : public FDynamicMeshOperator
+{
+public:
+	virtual ~FUVLayoutOp() {}
+
+	// inputs
+	TSharedPtr<FDynamicMesh3> OriginalMesh;
+
+
+	void SetTransform(const FTransform& Transform);
+
+	//
+	// FDynamicMeshOperator implementation
+	// 
+
+	virtual void CalculateResult(FProgressCancel* Progress) override;
+};
+
+
diff --git a/Engine/Plugins/Experimental/ModelingToolsEditorMode/Source/ModelingToolsEditorMode/Private/ModelingToolsEditorMode.cpp b/Engine/Plugins/Experimental/ModelingToolsEditorMode/Source/ModelingToolsEditorMode/Private/ModelingToolsEditorMode.cpp
--- a/Engine/Plugins/Experimental/ModelingToolsEditorMode/Source/ModelingToolsEditorMode/Private/ModelingToolsEditorMode.cpp
+++ b/Engine/Plugins/Experimental/ModelingToolsEditorMode/Source/ModelingToolsEditorMode/Private/ModelingToolsEditorMode.cpp
@@ -35,6 +35,7 @@
 #include "TransformMeshesTool.h"
 #include "MeshSelectionTool.h"
 #include "UVProjectionTool.h"
+#include "UVLayoutTool.h"
 
 #include "ParameterizeMeshTool.h"
 
@@ -381,6 +382,10 @@ void FModelingToolsEditorMode::Enter()
 	UVProjectionToolBuilder->AssetAPI = ToolsContext->GetAssetAPI();
 	RegisterToolFunc(ToolManagerCommands.BeginUVProjectionTool, TEXT("UVProjectionTool"), UVProjectionToolBuilder);
 
+	auto UVLayoutToolBuilder = NewObject<UUVLayoutToolBuilder>();
+	UVLayoutToolBuilder->AssetAPI = ToolsContext->GetAssetAPI();
+	RegisterToolFunc(ToolManagerCommands.BeginUVLayoutTool, TEXT("UVLayoutTool"), UVLayoutToolBuilder);
+
 	auto MergeMeshesToolBuilder = NewObject<UMergeMeshesToolBuilder>();
 	MergeMeshesToolBuilder->AssetAPI = ToolsContext->GetAssetAPI();
 	RegisterToolFunc(ToolManagerCommands.BeginVoxelMergeTool, TEXT("MergeMeshesTool"), MergeMeshesToolBuilder);
diff --git a/Engine/Plugins/Experimental/ModelingToolsEditorMode/Source/ModelingToolsEditorMode/Private/ModelingToolsEditorModeStyle.cpp b/Engine/Plugins/Experimental/ModelingToolsEditorMode/Source/ModelingToolsEditorMode/Private/ModelingToolsEditorModeStyle.cpp
--- a/Engine/Plugins/Experimental/ModelingToolsEditorMode/Source/ModelingToolsEditorMode/Private/ModelingToolsEditorModeStyle.cpp
+++ b/Engine/Plugins/Experimental/ModelingToolsEditorMode/Source/ModelingToolsEditorMode/Private/ModelingToolsEditorModeStyle.cpp
@@ -121,6 +121,8 @@ void FModelingToolsEditorModeStyle::Initialize()
 		StyleSet->Set("ModelingToolsManagerCommands.BeginRemoveOccludedTrianglesTool.Small", 		new IMAGE_PLUGIN_BRUSH("Icons/Jacket_40x",			Icon20x20));
 		StyleSet->Set("ModelingToolsManagerCommands.BeginUVProjectionTool", 			new IMAGE_PLUGIN_BRUSH("Icons/UVProjection_40x", 	Icon20x20));
 		StyleSet->Set("ModelingToolsManagerCommands.BeginUVProjectionTool.Small", 		new IMAGE_PLUGIN_BRUSH("Icons/UVProjection_40x", 	Icon20x20));
+		StyleSet->Set("ModelingToolsManagerCommands.BeginUVLayoutTool", 			new IMAGE_PLUGIN_BRUSH("Icons/UVLayout_40x", 	Icon20x20));
+		StyleSet->Set("ModelingToolsManagerCommands.BeginUVLayoutTool.Small", 		new IMAGE_PLUGIN_BRUSH("Icons/UVLayout_40x", 	Icon20x20));
 		StyleSet->Set("ModelingToolsManagerCommands.BeginVoxelMergeTool", 				new IMAGE_PLUGIN_BRUSH("Icons/VoxMerge_40x", 		Icon20x20));
 		StyleSet->Set("ModelingToolsManagerCommands.BeginVoxelMergeTool.Small", 		new IMAGE_PLUGIN_BRUSH("Icons/VoxMerge_40x", 		Icon20x20));
 		StyleSet->Set("ModelingToolsManagerCommands.BeginVoxelBooleanTool", 			new IMAGE_PLUGIN_BRUSH("Icons/VoxBoolean_40x", 		Icon20x20));
diff --git a/Engine/Plugins/Experimental/ModelingToolsEditorMode/Source/ModelingToolsEditorMode/Private/ModelingToolsEditorModeToolkit.cpp b/Engine/Plugins/Experimental/ModelingToolsEditorMode/Source/ModelingToolsEditorMode/Private/ModelingToolsEditorModeToolkit.cpp
--- a/Engine/Plugins/Experimental/ModelingToolsEditorMode/Source/ModelingToolsEditorMode/Private/ModelingToolsEditorModeToolkit.cpp
+++ b/Engine/Plugins/Experimental/ModelingToolsEditorMode/Source/ModelingToolsEditorMode/Private/ModelingToolsEditorModeToolkit.cpp
@@ -298,6 +298,7 @@ void FModelingToolsEditorModeToolkit::BuildToolPalette_Standard(FName PaletteInd
 	else if (PaletteIndex == FName(TEXT("Utilities")))
 	{
 		ToolbarBuilder.AddToolBarButton(Commands.BeginUVProjectionTool);
+		ToolbarBuilder.AddToolBarButton(Commands.BeginUVLayoutTool);
 		ToolbarBuilder.AddToolBarButton(Commands.BeginEditNormalsTool);
 		ToolbarBuilder.AddToolBarButton(Commands.BeginWeldEdgesTool);
 		ToolbarBuilder.AddToolBarButton(Commands.BeginMeshInspectorTool);
@@ -351,6 +352,7 @@ void FModelingToolsEditorModeToolkit::BuildToolPalette_Experimental(FName Palett
 	{
 		ToolbarBuilder.AddToolBarButton(Commands.BeginUVProjectionTool);
 		ToolbarBuilder.AddToolBarButton(Commands.BeginParameterizeMeshTool);
+		ToolbarBuilder.AddToolBarButton(Commands.BeginUVLayoutTool);
 		ToolbarBuilder.AddToolBarButton(Commands.BeginEditNormalsTool);
 	}
 	else if (PaletteIndex == FName(TEXT("Utilities")))
diff --git a/Engine/Plugins/Experimental/ModelingToolsEditorMode/Source/ModelingToolsEditorMode/Private/ModelingToolsManagerActions.cpp b/Engine/Plugins/Experimental/ModelingToolsEditorMode/Source/ModelingToolsEditorMode/Private/ModelingToolsManagerActions.cpp
--- a/Engine/Plugins/Experimental/ModelingToolsEditorMode/Source/ModelingToolsEditorMode/Private/ModelingToolsManagerActions.cpp
+++ b/Engine/Plugins/Experimental/ModelingToolsEditorMode/Source/ModelingToolsEditorMode/Private/ModelingToolsManagerActions.cpp
@@ -39,6 +39,7 @@ void FModelingToolsManagerCommands::RegisterCommands()
 	UI_COMMAND(BeginEditNormalsTool, "Normals", "Start the Edit Normals Tool", EUserInterfaceActionType::Button, FInputChord());
 	UI_COMMAND(BeginRemoveOccludedTrianglesTool, "Jacketing", "Start the Jacketing Tool", EUserInterfaceActionType::Button, FInputChord());
 	UI_COMMAND(BeginUVProjectionTool, "UVProjection", "Start the UV Projection Tool", EUserInterfaceActionType::Button, FInputChord());
+	UI_COMMAND(BeginUVLayoutTool, "UVLayout", "Start the UV Layout Tool", EUserInterfaceActionType::Button, FInputChord());
 	UI_COMMAND(BeginVoxelMergeTool, "VoxMerge", "Start the Voxel Merge Tool", EUserInterfaceActionType::Button, FInputChord());
 	UI_COMMAND(BeginVoxelBooleanTool, "VoxBoolean", "Start the Voxel Boolean Tool", EUserInterfaceActionType::Button, FInputChord());
 	UI_COMMAND(BeginPlaneCutTool, "PlaneCut", "Start the Plane Cut Tool", EUserInterfaceActionType::Button, FInputChord());
diff --git a/Engine/Plugins/Experimental/ModelingToolsEditorMode/Source/ModelingToolsEditorMode/Public/ModelingToolsManagerActions.h b/Engine/Plugins/Experimental/ModelingToolsEditorMode/Source/ModelingToolsEditorMode/Public/ModelingToolsManagerActions.h
--- a/Engine/Plugins/Experimental/ModelingToolsEditorMode/Source/ModelingToolsEditorMode/Public/ModelingToolsManagerActions.h
+++ b/Engine/Plugins/Experimental/ModelingToolsEditorMode/Source/ModelingToolsEditorMode/Public/ModelingToolsManagerActions.h
@@ -33,6 +33,7 @@ public:
 	TSharedPtr<FUICommandInfo> BeginEditNormalsTool;
 	TSharedPtr<FUICommandInfo> BeginRemoveOccludedTrianglesTool;
 	TSharedPtr<FUICommandInfo> BeginUVProjectionTool;
+	TSharedPtr<FUICommandInfo> BeginUVLayoutTool;
 	TSharedPtr<FUICommandInfo> BeginPlaneCutTool;
 	TSharedPtr<FUICommandInfo> BeginPolygonOnMeshTool;
 	TSharedPtr<FUICommandInfo> BeginVoxelMergeTool;
