fix undo bug for disconnect triangle feature
#rnx

diff --git a/Engine/Plugins/Experimental/MeshModelingToolset/Source/MeshModelingTools/Private/MeshSelectionTool.cpp b/Engine/Plugins/Experimental/MeshModelingToolset/Source/MeshModelingTools/Private/MeshSelectionTool.cpp
--- a/Engine/Plugins/Experimental/MeshModelingToolset/Source/MeshModelingTools/Private/MeshSelectionTool.cpp
+++ b/Engine/Plugins/Experimental/MeshModelingToolset/Source/MeshModelingTools/Private/MeshSelectionTool.cpp
@@ -950,7 +950,7 @@ void UMeshSelectionTool::DisconnectSelectedTriangles()
 				// include the whole one-ring in case the disconnect creates bowties that need to be split
 				for (int TID : Mesh.VtxTrianglesItr(VID))
 				{
-					ChangeTracker.SaveTriangle(TID, false);
+					ChangeTracker.SaveTriangle(TID, true);
 				}
 			}
 		}
