fix multi-material support in plane cut

#rnx

diff --git a/Engine/Plugins/Experimental/MeshModelingToolset/Source/MeshModelingTools/Private/PlaneCutTool.cpp b/Engine/Plugins/Experimental/MeshModelingToolset/Source/MeshModelingTools/Private/PlaneCutTool.cpp
--- a/Engine/Plugins/Experimental/MeshModelingToolset/Source/MeshModelingTools/Private/PlaneCutTool.cpp
+++ b/Engine/Plugins/Experimental/MeshModelingToolset/Source/MeshModelingTools/Private/PlaneCutTool.cpp
@@ -397,11 +397,18 @@ void UPlaneCutTool::GenerateAsset(const TArray<FDynamicMeshOpResult>& Results)
 		for (int32 AddedMeshIdx = NumSourceMeshes; AddedMeshIdx < Results.Num(); AddedMeshIdx++)
 		{
 			check(Results[AddedMeshIdx].Mesh.Get() != nullptr);
-			// TODO: copy over material?
+
+			TArray<UMaterialInterface*> Materials;
+			TUniquePtr<FPrimitiveComponentTarget>& ComponentTarget = ComponentTargets[AddedMeshIdx - NumSourceMeshes];
+			for (int MaterialIdx = 0, NumMaterials = ComponentTarget->GetNumMaterials(); MaterialIdx < NumMaterials; MaterialIdx++)
+			{
+				Materials.Add(ComponentTarget->GetMaterial(MaterialIdx));
+			}
+
 			AActor* NewActor = AssetGenerationUtil::GenerateStaticMeshActor(
 				AssetAPI, TargetWorld,
 				Results[AddedMeshIdx].Mesh.Get(), Results[AddedMeshIdx].Transform, TEXT("PlaneCutOtherHalf"),
-				AssetGenerationUtil::GetDefaultAutoGeneratedAssetPath());
+				AssetGenerationUtil::GetDefaultAutoGeneratedAssetPath(), Materials);
 			NewSelection.Actors.Add(NewActor);
 		}
 		GetToolManager()->RequestSelectionChange(NewSelection);
