uv layout needs to be editor only
#rnx

diff --git a/Engine/Plugins/Experimental/MeshModelingToolset/Source/MeshModelingTools/Private/UVLayoutTool.cpp b/Engine/Plugins/Experimental/MeshModelingToolset/Source/MeshModelingTools/Private/UVLayoutTool.cpp
deleted file mode 100644
diff --git a/Engine/Plugins/Experimental/MeshModelingToolset/Source/MeshModelingTools/Public/UVLayoutTool.h b/Engine/Plugins/Experimental/MeshModelingToolset/Source/MeshModelingTools/Public/UVLayoutTool.h
deleted file mode 100644
diff --git a/Engine/Plugins/Experimental/MeshModelingToolset/Source/MeshModelingToolsEditorOnly/Private/UVLayoutTool.cpp b/Engine/Plugins/Experimental/MeshModelingToolset/Source/MeshModelingToolsEditorOnly/Private/UVLayoutTool.cpp
new file mode 100644
--- /dev/null
+++ b/Engine/Plugins/Experimental/MeshModelingToolset/Source/MeshModelingToolsEditorOnly/Private/UVLayoutTool.cpp
@@ -0,0 +1,275 @@
+// Copyright 1998-2019 Epic Games, Inc. All Rights Reserved.
+
+#include "UVLayoutTool.h"
+#include "InteractiveToolManager.h"
+#include "ToolBuilderUtil.h"
+
+#include "ToolSetupUtil.h"
+
+#include "DynamicMesh3.h"
+#include "BaseBehaviors/MultiClickSequenceInputBehavior.h"
+#include "Selection/SelectClickedAction.h"
+
+#include "MeshDescriptionToDynamicMesh.h"
+#include "DynamicMeshToMeshDescription.h"
+#include "Materials/MaterialInstanceDynamic.h"
+
+#include "ParameterizationOps/UVLayoutOp.h"
+
+#include "AssetGenerationUtil.h"
+
+#define LOCTEXT_NAMESPACE "UUVLayoutTool"
+
+
+/*
+ * ToolBuilder
+ */
+
+
+bool UUVLayoutToolBuilder::CanBuildTool(const FToolBuilderState& SceneState) const
+{
+	return ToolBuilderUtil::CountComponents(SceneState, CanMakeComponentTarget) >= 1;
+}
+
+UInteractiveTool* UUVLayoutToolBuilder::BuildTool(const FToolBuilderState& SceneState) const
+{
+	UUVLayoutTool* NewTool = NewObject<UUVLayoutTool>(SceneState.ToolManager);
+
+	TArray<UActorComponent*> Components = ToolBuilderUtil::FindAllComponents(SceneState, CanMakeComponentTarget);
+	check(Components.Num() > 0);
+
+	TArray<TUniquePtr<FPrimitiveComponentTarget>> ComponentTargets;
+	for (UActorComponent* ActorComponent : Components)
+	{
+		auto* MeshComponent = Cast<UPrimitiveComponent>(ActorComponent);
+		if ( MeshComponent )
+		{
+			ComponentTargets.Add(MakeComponentTarget(MeshComponent));
+		}
+	}
+
+	NewTool->SetSelection(MoveTemp(ComponentTargets));
+	NewTool->SetWorld(SceneState.World, SceneState.GizmoManager);
+	NewTool->SetAssetAPI(AssetAPI);
+
+	return NewTool;
+}
+
+
+
+/*
+ * Tool
+ */
+
+UUVLayoutToolProperties::UUVLayoutToolProperties()
+{
+
+}
+
+UUVLayoutAdvancedProperties::UUVLayoutAdvancedProperties()
+{
+}
+
+
+UUVLayoutTool::UUVLayoutTool()
+{
+}
+
+void UUVLayoutTool::SetWorld(UWorld* World, UInteractiveGizmoManager* GizmoManagerIn)
+{
+	this->TargetWorld = World;
+}
+
+void UUVLayoutTool::Setup()
+{
+	UInteractiveTool::Setup();
+
+	// hide input StaticMeshComponent
+	for (TUniquePtr<FPrimitiveComponentTarget>& ComponentTarget : ComponentTargets)
+	{
+		ComponentTarget->SetOwnerVisibility(false);
+	}
+
+	BasicProperties = NewObject<UUVLayoutToolProperties>(this, TEXT("UV Projection Settings"));
+	AdvancedProperties = NewObject<UUVLayoutAdvancedProperties>(this, TEXT("Advanced Settings"));
+
+	// initialize our properties
+	AddToolPropertySource(BasicProperties);
+	AddToolPropertySource(AdvancedProperties);
+
+	MaterialSettings = NewObject<UExistingMeshMaterialProperties>(this);
+	MaterialSettings->Setup();
+	AddToolPropertySource(MaterialSettings);
+
+	// initialize the PreviewMesh+BackgroundCompute object
+	UpdateNumPreviews();
+	
+	for (UMeshOpPreviewWithBackgroundCompute* Preview : Previews)
+	{
+		Preview->InvalidateResult();
+	}
+}
+
+
+void UUVLayoutTool::UpdateNumPreviews()
+{
+	int32 CurrentNumPreview = Previews.Num();
+	int32 TargetNumPreview = ComponentTargets.Num();
+	if (TargetNumPreview < CurrentNumPreview)
+	{
+		for (int32 PreviewIdx = CurrentNumPreview - 1; PreviewIdx >= TargetNumPreview; PreviewIdx--)
+		{
+			Previews[PreviewIdx]->Cancel();
+		}
+		Previews.SetNum(TargetNumPreview);
+		OriginalDynamicMeshes.SetNum(TargetNumPreview);
+	}
+	else
+	{
+		OriginalDynamicMeshes.SetNum(TargetNumPreview);
+		for (int32 PreviewIdx = CurrentNumPreview; PreviewIdx < TargetNumPreview; PreviewIdx++)
+		{
+			UUVLayoutOperatorFactory *OpFactory = NewObject<UUVLayoutOperatorFactory>();
+			OpFactory->Tool = this;
+			OpFactory->ComponentIndex = PreviewIdx;
+			OriginalDynamicMeshes[PreviewIdx] = MakeShared<FDynamicMesh3>();
+			FMeshDescriptionToDynamicMesh Converter;
+			Converter.bPrintDebugMessages = true;
+			Converter.Convert(ComponentTargets[PreviewIdx]->GetMesh(), *OriginalDynamicMeshes[PreviewIdx]);
+
+			UMeshOpPreviewWithBackgroundCompute* Preview = Previews.Add_GetRef(NewObject<UMeshOpPreviewWithBackgroundCompute>(OpFactory, "Preview"));
+			Preview->Setup(this->TargetWorld, OpFactory);
+			Preview->ConfigureMaterials(
+				ToolSetupUtil::GetDefaultMaterial(GetToolManager(), ComponentTargets[PreviewIdx]->GetMaterial(0)),
+				ToolSetupUtil::GetDefaultWorkingMaterial(GetToolManager())
+			);
+			Preview->SetVisibility(true);
+		}
+	}
+}
+
+
+void UUVLayoutTool::Shutdown(EToolShutdownType ShutdownType)
+{
+	// Restore (unhide) the source meshes
+	for (TUniquePtr<FPrimitiveComponentTarget>& ComponentTarget : ComponentTargets)
+	{
+		ComponentTarget->SetOwnerVisibility(true);
+	}
+
+	TArray<FDynamicMeshOpResult> Results;
+	for (UMeshOpPreviewWithBackgroundCompute* Preview : Previews)
+	{
+		Results.Emplace(Preview->Shutdown());
+	}
+	if (ShutdownType == EToolShutdownType::Accept)
+	{
+		GenerateAsset(Results);
+	}
+}
+
+void UUVLayoutTool::SetAssetAPI(IToolsContextAssetAPI* AssetAPIIn)
+{
+	this->AssetAPI = AssetAPIIn;
+}
+
+TUniquePtr<FDynamicMeshOperator> UUVLayoutOperatorFactory::MakeNewOperator()
+{
+	TUniquePtr<FUVLayoutOp> Op = MakeUnique<FUVLayoutOp>();
+
+	FTransform LocalToWorld = Tool->ComponentTargets[ComponentIndex]->GetWorldTransform();
+	Op->OriginalMesh = Tool->OriginalDynamicMeshes[ComponentIndex];
+	
+	Op->SetTransform(LocalToWorld);
+
+	return Op;
+}
+
+
+
+void UUVLayoutTool::Render(IToolsContextRenderAPI* RenderAPI)
+{
+}
+
+void UUVLayoutTool::Tick(float DeltaTime)
+{
+	for (UMeshOpPreviewWithBackgroundCompute* Preview : Previews)
+	{
+		Preview->Tick(DeltaTime);
+	}
+}
+
+
+
+#if WITH_EDITOR
+void UUVLayoutTool::PostEditChangeProperty(FPropertyChangedEvent& PropertyChangedEvent)
+{
+	UpdateNumPreviews();
+	for (UMeshOpPreviewWithBackgroundCompute* Preview : Previews)
+	{
+		Preview->InvalidateResult();
+	}
+}
+#endif
+
+void UUVLayoutTool::OnPropertyModified(UObject* PropertySet, UProperty* Property)
+{
+	// if we don't know what changed, or we know checker density changed, update checker material
+	MaterialSettings->UpdateMaterials();
+	for (int PreviewIdx = 0; PreviewIdx < Previews.Num(); PreviewIdx++)
+	{
+		UMeshOpPreviewWithBackgroundCompute* Preview = Previews[PreviewIdx];
+		MaterialSettings->SetMaterialIfChanged(ComponentTargets[PreviewIdx]->GetMaterial(0), Preview->StandardMaterial, [&Preview, this](UMaterialInterface* Material)
+		{
+			Preview->ConfigureMaterials(ToolSetupUtil::GetDefaultMaterial(GetToolManager(), Material), Preview->WorkingMaterial);
+		});
+	}
+	
+	UpdateNumPreviews();
+	for (UMeshOpPreviewWithBackgroundCompute* Preview : Previews)
+	{
+		Preview->InvalidateResult();
+	}
+}
+
+bool UUVLayoutTool::HasAccept() const
+{
+	return true;
+}
+
+bool UUVLayoutTool::CanAccept() const
+{
+	for (UMeshOpPreviewWithBackgroundCompute* Preview : Previews)
+	{
+		if (!Preview->HaveValidResult())
+		{
+			return false;
+		}
+	}
+	return true;
+}
+
+
+void UUVLayoutTool::GenerateAsset(const TArray<FDynamicMeshOpResult>& Results)
+{
+	GetToolManager()->BeginUndoTransaction(LOCTEXT("UVLayoutToolTransactionName", "UV Layout Tool"));
+
+	check(Results.Num() == ComponentTargets.Num());
+	
+	for (int32 ComponentIdx = 0; ComponentIdx < ComponentTargets.Num(); ComponentIdx++)
+	{
+		check(Results[ComponentIdx].Mesh.Get() != nullptr);
+		ComponentTargets[ComponentIdx]->CommitMesh([&Results, &ComponentIdx, this](const FPrimitiveComponentTarget::FCommitParams& CommitParams)
+		{
+			FDynamicMeshToMeshDescription Converter;
+			Converter.Convert(Results[ComponentIdx].Mesh.Get(), *CommitParams.MeshDescription);
+		});
+	}
+
+	GetToolManager()->EndUndoTransaction();
+}
+
+
+
+
+#undef LOCTEXT_NAMESPACE
diff --git a/Engine/Plugins/Experimental/MeshModelingToolset/Source/MeshModelingToolsEditorOnly/Public/UVLayoutTool.h b/Engine/Plugins/Experimental/MeshModelingToolset/Source/MeshModelingToolsEditorOnly/Public/UVLayoutTool.h
new file mode 100644
--- /dev/null
+++ b/Engine/Plugins/Experimental/MeshModelingToolset/Source/MeshModelingToolsEditorOnly/Public/UVLayoutTool.h
@@ -0,0 +1,158 @@
+// Copyright 1998-2019 Epic Games, Inc. All Rights Reserved.
+
+#pragma once
+
+#include "CoreMinimal.h"
+#include "UObject/NoExportTypes.h"
+#include "MultiSelectionTool.h"
+#include "InteractiveToolBuilder.h"
+#include "MeshOpPreviewHelpers.h"
+#include "ToolDataVisualizer.h"
+#include "DynamicMesh3.h"
+#include "BaseTools/SingleClickTool.h"
+#include "Properties/MeshMaterialProperties.h"
+
+#include "UVLayoutTool.generated.h"
+
+
+// predeclarations
+struct FMeshDescription;
+class USimpleDynamicMeshComponent;
+
+
+/**
+ *
+ */
+UCLASS()
+class MESHMODELINGTOOLSEDITORONLY_API UUVLayoutToolBuilder : public UInteractiveToolBuilder
+{
+	GENERATED_BODY()
+
+public:
+	IToolsContextAssetAPI* AssetAPI = nullptr;
+
+	virtual bool CanBuildTool(const FToolBuilderState& SceneState) const override;
+	virtual UInteractiveTool* BuildTool(const FToolBuilderState& SceneState) const override;
+};
+
+
+
+
+
+/**
+ * Standard properties
+ */
+UCLASS()
+class MESHMODELINGTOOLSEDITORONLY_API UUVLayoutToolProperties : public UInteractiveToolPropertySet
+{
+	GENERATED_BODY()
+
+public:
+	UUVLayoutToolProperties();
+
+
+// TODO: add config settings for uvlayout (currently nothing really in the actual algorithm is exposed as a usable parameter)
+
+};
+
+
+
+
+/**
+ * Advanced properties
+ */
+UCLASS()
+class MESHMODELINGTOOLSEDITORONLY_API UUVLayoutAdvancedProperties : public UInteractiveToolPropertySet
+{
+	GENERATED_BODY()
+
+public:
+	UUVLayoutAdvancedProperties();
+};
+
+
+/**
+ * Factory with enough info to spawn the background-thread Operator to do a chunk of work for the tool
+ *  stores a pointer to the tool and enough info to know which specific operator it should spawn
+ */
+UCLASS()
+class MESHMODELINGTOOLSEDITORONLY_API UUVLayoutOperatorFactory : public UObject, public IDynamicMeshOperatorFactory
+{
+	GENERATED_BODY()
+
+public:
+	// IDynamicMeshOperatorFactory API
+	virtual TUniquePtr<FDynamicMeshOperator> MakeNewOperator() override;
+
+	UPROPERTY()
+	UUVLayoutTool *Tool;
+
+	int ComponentIndex;
+
+};
+
+/**
+ * Simple Mesh Normal Updating Tool
+ */
+UCLASS()
+class MESHMODELINGTOOLSEDITORONLY_API UUVLayoutTool : public UMultiSelectionTool
+{
+	GENERATED_BODY()
+
+public:
+
+	friend UUVLayoutOperatorFactory;
+
+	UUVLayoutTool();
+
+	virtual void Setup() override;
+	virtual void Shutdown(EToolShutdownType ShutdownType) override;
+
+	virtual void SetWorld(UWorld* World, UInteractiveGizmoManager* GizmoManagerIn);
+	virtual void SetAssetAPI(IToolsContextAssetAPI* AssetAPI);
+
+	virtual void Tick(float DeltaTime) override;
+	virtual void Render(IToolsContextRenderAPI* RenderAPI) override;
+
+	virtual bool HasCancel() const override { return true; }
+	virtual bool HasAccept() const override;
+	virtual bool CanAccept() const override;
+
+#if WITH_EDITOR
+	virtual void PostEditChangeProperty(FPropertyChangedEvent &PropertyChangedEvent) override;
+#endif
+
+	virtual void OnPropertyModified(UObject* PropertySet, UProperty* Property) override;
+
+protected:
+
+	UPROPERTY()
+	UUVLayoutToolProperties* BasicProperties;
+
+	UPROPERTY()
+	UUVLayoutAdvancedProperties* AdvancedProperties;
+
+	UPROPERTY()
+	UExistingMeshMaterialProperties* MaterialSettings = nullptr;
+
+
+
+	UPROPERTY()
+	TArray<UMeshOpPreviewWithBackgroundCompute*> Previews;
+
+
+	UPROPERTY()
+	UMaterialInstanceDynamic* CheckerMaterial = nullptr;
+
+protected:
+	TArray<TSharedPtr<FDynamicMesh3>> OriginalDynamicMeshes;
+
+	UWorld* TargetWorld;
+	IToolsContextAssetAPI* AssetAPI;
+
+	FViewCameraState CameraState;
+
+	void UpdateNumPreviews();
+
+	void GenerateAsset(const TArray<FDynamicMeshOpResult>& Results);
+};
diff --git a/Engine/Plugins/Experimental/MeshModelingToolset/Source/ModelingOperators/Private/ParameterizationOps/UVLayoutOp.cpp b/Engine/Plugins/Experimental/MeshModelingToolset/Source/ModelingOperators/Private/ParameterizationOps/UVLayoutOp.cpp
deleted file mode 100644
diff --git a/Engine/Plugins/Experimental/MeshModelingToolset/Source/ModelingOperators/Public/ParameterizationOps/UVLayoutOp.h b/Engine/Plugins/Experimental/MeshModelingToolset/Source/ModelingOperators/Public/ParameterizationOps/UVLayoutOp.h
deleted file mode 100644
diff --git a/Engine/Plugins/Experimental/MeshModelingToolset/Source/ModelingOperatorsEditorOnly/Private/ParameterizationOps/UVLayoutOp.cpp b/Engine/Plugins/Experimental/MeshModelingToolset/Source/ModelingOperatorsEditorOnly/Private/ParameterizationOps/UVLayoutOp.cpp
new file mode 100644
--- /dev/null
+++ b/Engine/Plugins/Experimental/MeshModelingToolset/Source/ModelingOperatorsEditorOnly/Private/ParameterizationOps/UVLayoutOp.cpp
@@ -0,0 +1,161 @@
+// Copyright 1998-2019 Epic Games, Inc. All Rights Reserved.
+
+#include "ParameterizationOps/UVLayoutOp.h"
+
+#include "OverlappingCorners.h"
+#include "LayoutUV.h"
+
+#include "DynamicMesh3.h"
+#include "DynamicMeshAttributeSet.h"
+
+
+void FUVLayoutOp::SetTransform(const FTransform& Transform) {
+	ResultTransform = (FTransform3d)Transform;
+}
+
+// very task-specific lightweight view of a FDynamicMesh3; ONLY for compact meshes with attributes
+// not intended for use outside of this narrow context
+class FCompactDynamicMeshWithAttributesLayoutView : public FLayoutUV::IMeshView
+{
+	FDynamicMesh3* Mesh;
+	int UVLayerInput, UVLayerOutput;
+
+public:
+	FCompactDynamicMeshWithAttributesLayoutView(FDynamicMesh3* Mesh, int UVLayerIn=0, int UVLayerOut=0) : Mesh(Mesh), UVLayerInput(UVLayerIn), UVLayerOutput(UVLayerOut)
+	{
+		check(Mesh->HasAttributes());
+		check(Mesh->IsCompact());
+		check(Mesh->Attributes()->NumUVLayers() > UVLayerIn);
+		check(Mesh->Attributes()->NumUVLayers() > UVLayerOut);
+	}
+
+	uint32 GetNumIndices() const override
+	{
+		return Mesh->TriangleCount() * 3;
+	}
+	FVector GetPosition(uint32 Index) const override
+	{
+		return FVector(Mesh->GetVertex(Mesh->GetTriangle(Index / 3)[Index % 3]));
+	}
+	FVector GetNormal(uint32 Index) const override
+	{
+		const FDynamicMeshNormalOverlay* NormalOverlay = Mesh->Attributes()->PrimaryNormals();
+		FIndex3i NormalEl = NormalOverlay->GetTriangle(Index / 3);
+		FVector3f Normal;
+		NormalOverlay->GetElement(NormalEl[Index % 3], Normal);
+		return FVector(Normal);
+	}
+	FVector2D GetInputTexcoord(uint32 Index) const override
+	{
+		const FDynamicMeshUVOverlay* UVOverlay = Mesh->Attributes()->GetUVLayer(UVLayerInput);
+		FIndex3i UVEl = UVOverlay->GetTriangle(Index / 3);
+		FVector2f UV;
+		UVOverlay->GetElement(UVEl[Index % 3], UV);
+		return FVector2D(UV);
+	}
+
+	/**
+	 * This function is kind of nonsense for our use case as we cannot initialize a UV overlay from a single number
+	 * (unless we make every triangle disconnected in a fully raw wedge thing, which we never ever want to do)
+	 * So we assume the calling code will only call this w/ NumIndices matching the input layer & the intent of making the UV layers the same.
+	 */
+	void InitOutputTexcoords(uint32 Num) override
+	{
+		if (UVLayerInput != UVLayerOutput)
+		{
+			Mesh->Attributes()->GetUVLayer(UVLayerOutput)->Copy(*Mesh->Attributes()->GetUVLayer(UVLayerInput));
+		}
+		check(Num == Mesh->TriangleCount() * 3);
+	}
+
+	void SetOutputTexcoord(uint32 Index, const FVector2D& Value) override
+	{
+		FDynamicMeshUVOverlay* UVOverlay = Mesh->Attributes()->GetUVLayer(UVLayerOutput);
+		FIndex3i UVEl = UVOverlay->GetTriangle(Index / 3);
+		UVOverlay->SetElement(UVEl[Index % 3], FVector2f(Value));
+	}
+};
+
+/**
+ * Create an overlapping corner map to identify wedge indices that share the same UV element index for a given UV layer
+ */
+FOverlappingCorners OverlappingCornersFromUVs(const FDynamicMesh3* Mesh, int UVLayerIndex)
+{
+	// track all wedge indices that map to the same UV element
+	const FDynamicMeshUVOverlay* UVOverlay = Mesh->Attributes()->GetUVLayer(UVLayerIndex);
+	FOverlappingCorners Overlaps;
+	Overlaps.Init(Mesh->TriangleCount()*3);
+	for (int ElementID : UVOverlay->ElementIndicesItr())
+	{
+		int VertID = UVOverlay->GetParentVertex(ElementID);
+		int LastWedgeIdx = -1;
+		for (int TriID : Mesh->VtxTrianglesItr(VertID))
+		{
+			FIndex3i ElTri = UVOverlay->GetTriangle(TriID);
+			for (int SubIdx = 0; SubIdx < 3; SubIdx++)
+			{
+				if (ElTri[SubIdx] == ElementID)
+				{
+					int WedgeIdx = TriID * 3 + SubIdx;
+					if (LastWedgeIdx != -1)
+					{
+						Overlaps.Add(LastWedgeIdx, WedgeIdx);
+					}
+					LastWedgeIdx = WedgeIdx;
+				}
+			}
+		}
+	}
+	Overlaps.FinishAdding();
+	return Overlaps;
+}
+
+void FUVLayoutOp::CalculateResult(FProgressCancel* Progress)
+{
+	if (Progress->Cancelled())
+	{
+		return;
+	}
+	bool bDiscardAttributes = false;
+	ResultMesh->Copy(*OriginalMesh, true, true, true, !bDiscardAttributes);
+
+	if (!ensureMsgf(ResultMesh->HasAttributes(), TEXT("Attributes not found on mesh? Conversion should always create them, so this operator should not need to do so.")))
+	{
+		ResultMesh->EnableAttributes();
+	}
+
+	if (Progress->Cancelled())
+	{
+		return;
+	}
+
+	int UVLayerInput = 0, UVLayerOutput = 0;
+	ResultMesh->Attributes()->GetUVLayer(UVLayerInput)->SplitBowties();
+
+	if (Progress->Cancelled())
+	{
+		return;
+	}
+
+	FCompactDynamicMeshWithAttributesLayoutView MeshView(ResultMesh.Get(), UVLayerInput, UVLayerOutput);
+	FLayoutUV LayoutUV(MeshView);
+	FOverlappingCorners Overlaps = OverlappingCornersFromUVs(ResultMesh.Get(), UVLayerInput);
+	if (Progress->Cancelled())
+	{
+		return;
+	}
+
+	LayoutUV.FindCharts(Overlaps);
+	if (Progress->Cancelled())
+	{
+		return;
+	}
+
+	LayoutUV.FindBestPacking(1024); // TODO: make this a parameter
+	if (Progress->Cancelled())
+	{
+		return;
+	}
+
+	LayoutUV.CommitPackedUVs();
+}
diff --git a/Engine/Plugins/Experimental/MeshModelingToolset/Source/ModelingOperatorsEditorOnly/Public/ParameterizationOps/UVLayoutOp.h b/Engine/Plugins/Experimental/MeshModelingToolset/Source/ModelingOperatorsEditorOnly/Public/ParameterizationOps/UVLayoutOp.h
new file mode 100644
--- /dev/null
+++ b/Engine/Plugins/Experimental/MeshModelingToolset/Source/ModelingOperatorsEditorOnly/Public/ParameterizationOps/UVLayoutOp.h
@@ -0,0 +1,29 @@
+// Copyright 1998-2019 Epic Games, Inc. All Rights Reserved.
+
+#pragma once
+
+#include "CoreMinimal.h"
+#include "Util/ProgressCancel.h"
+#include "ModelingOperators.h"
+
+
+
+class MODELINGOPERATORSEDITORONLY_API FUVLayoutOp : public FDynamicMeshOperator
+{
+public:
+	virtual ~FUVLayoutOp() {}
+
+	// inputs
+	TSharedPtr<FDynamicMesh3> OriginalMesh;
+
+
+	void SetTransform(const FTransform& Transform);
+
+	//
+	// FDynamicMeshOperator implementation
+	// 
+
+	virtual void CalculateResult(FProgressCancel* Progress) override;
+};
+
+
