ModelingTools: Multiple Material Support on DynamicMesh Components
- MeshDescriptionToDynamicMesh now always enables MaterialID attribute
- add FComponentMaterialSet and FComponentTarget::GetMaterialSet() that returns for current object (eg set of Material Slot materials on StaticMesh)
- move FMeshRenderBufferSet to BaseMeshSceneProxy.h, add new FBaseDynamicMeshSceneProxy, centralizes management/initialization of RenderBuffers, default GetDynamicMeshElements implementation draws buffer set returned by subclasses
- FSimpleDynamicMeshSceneProxy and FOctreeDynamicMeshSceneProxy now subclass FBaseDynamicMeshSceneProxy, implement necessary functions/etc
- UBaseDynamicMeshComponent now manages its own set of Materials. New OverrideMaterial can be set to replace existing materials with override material during render calls (does not affect material set)
- FSimpleDynamicMeshSceneProxy now builds per-material RenderBuffers, new FastUpdateVertices function to update positions/normals/colors. Removed shared-vertex construction, now only builds split-triangle buffers.
- UPreviewMesh gets wrapper functions to manage materials, override material, tangents of contained Component
#rnx

diff --git a/Engine/Plugins/Experimental/GeometryProcessing/Source/MeshConversion/Private/MeshDescriptionToDynamicMesh.cpp b/Engine/Plugins/Experimental/GeometryProcessing/Source/MeshConversion/Private/MeshDescriptionToDynamicMesh.cpp
--- a/Engine/Plugins/Experimental/GeometryProcessing/Source/MeshConversion/Private/MeshDescriptionToDynamicMesh.cpp
+++ b/Engine/Plugins/Experimental/GeometryProcessing/Source/MeshConversion/Private/MeshDescriptionToDynamicMesh.cpp
@@ -173,8 +173,9 @@ void FMeshDescriptionToDynamicMesh::Convert(const FMeshDescription* MeshIn, FDyn
 	}
 	FNormalWelder NormalWelder(NormalOverlay);
 
-	// only enable material ID if we have more than one material
-	FDynamicMeshMaterialAttribute* MaterialIDAttrib = nullptr;
+	// always enable material ID
+	MeshOut.Attributes()->EnableMaterialID();
+	FDynamicMeshMaterialAttribute* MaterialIDAttrib = MeshOut.Attributes()->GetMaterialID();
 
 	// NOTE: If you change the iteration order here, please update the corresponding iteration in FDynamicMeshToMeshDescription::UpdateAttributes, 
 	//	which assumes the iteration order here is polygons -> triangles, to correspond the triangles when writing updated attributes back!
@@ -182,11 +183,6 @@ void FMeshDescriptionToDynamicMesh::Convert(const FMeshDescription* MeshIn, FDyn
 	for (const FPolygonID PolygonID : Polygons.GetElementIDs())
 	{
 		int32 PolygonGroupID = MeshIn->GetPolygonPolygonGroup(PolygonID).GetValue();
-		if (MaterialIDAttrib == nullptr && PolygonGroupID != 0)
-		{
-			MeshOut.Attributes()->EnableMaterialID();		// initializes all existing values to 0, which is fine since this is first non-zero
-			MaterialIDAttrib = MeshOut.Attributes()->GetMaterialID();
-		}
 
 		const TArray<FTriangleID>& TriangleIDs = MeshIn->GetPolygonTriangleIDs(PolygonID);
 		int NumTriangles = TriangleIDs.Num();
diff --git a/Engine/Plugins/Experimental/MeshModelingToolset/Source/ModelingComponents/Private/BaseDynamicMeshComponent.cpp b/Engine/Plugins/Experimental/MeshModelingToolset/Source/ModelingComponents/Private/BaseDynamicMeshComponent.cpp
--- a/Engine/Plugins/Experimental/MeshModelingToolset/Source/ModelingComponents/Private/BaseDynamicMeshComponent.cpp
+++ b/Engine/Plugins/Experimental/MeshModelingToolset/Source/ModelingComponents/Private/BaseDynamicMeshComponent.cpp
@@ -8,3 +8,50 @@ UBaseDynamicMeshComponent::UBaseDynamicMeshComponent(const FObjectInitializer& O
 	: Super(ObjectInitializer)
 {
 }
+
+
+
+void UBaseDynamicMeshComponent::SetOverrideRenderMaterial(UMaterialInterface* Material)
+{
+	OverrideRenderMaterial = Material;
+	NotifyMaterialSetUpdated();
+}
+
+void UBaseDynamicMeshComponent::ClearOverrideRenderMaterial()
+{
+	OverrideRenderMaterial = nullptr;
+	NotifyMaterialSetUpdated();
+}
+
+
+
+int32 UBaseDynamicMeshComponent::GetNumMaterials() const
+{
+	return BaseMaterials.Num();
+}
+
+UMaterialInterface* UBaseDynamicMeshComponent::GetMaterial(int32 ElementIndex) const 
+{
+	return (ElementIndex >= 0 && ElementIndex < BaseMaterials.Num()) ? BaseMaterials[ElementIndex] : nullptr;
+}
+
+void UBaseDynamicMeshComponent::SetMaterial(int32 ElementIndex, UMaterialInterface* Material)
+{
+	check(ElementIndex >= 0);
+	if (ElementIndex >= BaseMaterials.Num())
+	{
+		BaseMaterials.SetNum(ElementIndex + 1, false);
+	}
+	BaseMaterials[ElementIndex] = Material;
+}
+
+
+void UBaseDynamicMeshComponent::GetUsedMaterials(TArray<UMaterialInterface*>& OutMaterials, bool bGetDebugMaterials) const
+{
+	UMeshComponent::GetUsedMaterials(OutMaterials, bGetDebugMaterials);
+	if (OverrideRenderMaterial != nullptr)
+	{
+		OutMaterials.Add(OverrideRenderMaterial);
+	}
+}
+
diff --git a/Engine/Plugins/Experimental/MeshModelingToolset/Source/ModelingComponents/Private/BaseDynamicMeshSceneProxy.h b/Engine/Plugins/Experimental/MeshModelingToolset/Source/ModelingComponents/Private/BaseDynamicMeshSceneProxy.h
new file mode 100644
--- /dev/null
+++ b/Engine/Plugins/Experimental/MeshModelingToolset/Source/ModelingComponents/Private/BaseDynamicMeshSceneProxy.h
@@ -0,0 +1,617 @@
+// Copyright 1998-2019 Epic Games, Inc. All Rights Reserved. 
+
+#pragma once
+
+#include "CoreMinimal.h"
+#include "PrimitiveSceneProxy.h"
+#include "VertexFactory.h"
+#include "LocalVertexFactory.h"
+#include "PrimitiveViewRelevance.h"
+#include "BaseDynamicMeshComponent.h"
+
+
+/**
+ * FMeshRenderBufferSet stores a set of RenderBuffers for a mesh
+ */
+class FMeshRenderBufferSet
+{
+public:
+	int TriangleCount = 0;
+
+	/** The buffer containing vertex data. */
+	FStaticMeshVertexBuffer StaticMeshVertexBuffer;
+	/** The buffer containing the position vertex data. */
+	FPositionVertexBuffer PositionVertexBuffer;
+	/** The buffer containing the vertex color data. */
+	FColorVertexBuffer ColorVertexBuffer;
+
+	/** triangle indices */
+	FDynamicMeshIndexBuffer32 IndexBuffer;
+
+	/** vertex factory */
+	FLocalVertexFactory VertexFactory;
+
+	/** Material to draw this mesh with */
+	UMaterialInterface* Material = nullptr;
+
+	/**
+	 * Optional list of triangles stored in this buffer. Storing this allows us
+	 * to rebuild the buffers if vertex data changes.
+	 */
+	TOptional<TArray<int>> Triangles;
+
+
+	FMeshRenderBufferSet(ERHIFeatureLevel::Type FeatureLevelType)
+		: VertexFactory(FeatureLevelType, "FMeshRenderBufferSet")
+	{
+
+	}
+
+
+	virtual ~FMeshRenderBufferSet()
+	{
+		check(IsInRenderingThread());
+
+		if (TriangleCount > 0)
+		{
+			PositionVertexBuffer.ReleaseResource();
+			StaticMeshVertexBuffer.ReleaseResource();
+			ColorVertexBuffer.ReleaseResource();
+			IndexBuffer.ReleaseResource();
+			VertexFactory.ReleaseResource();
+		}
+	}
+
+
+	/**
+	 * Upload initialized mesh buffers. 
+	 * @warning This can only be called on the Rendering Thread.
+	 */
+	void Upload()
+	{
+		check(IsInRenderingThread());
+
+		if (TriangleCount == 0)
+		{
+			return;
+		}
+
+		InitOrUpdateResource(&this->PositionVertexBuffer);
+		InitOrUpdateResource(&this->StaticMeshVertexBuffer);
+		InitOrUpdateResource(&this->ColorVertexBuffer);
+
+		FLocalVertexFactory::FDataType Data;
+		this->PositionVertexBuffer.BindPositionVertexBuffer(&this->VertexFactory, Data);
+		this->StaticMeshVertexBuffer.BindTangentVertexBuffer(&this->VertexFactory, Data);
+		this->StaticMeshVertexBuffer.BindPackedTexCoordVertexBuffer(&this->VertexFactory, Data);
+		// currently no lightmaps support
+		//this->StaticMeshVertexBuffer.BindLightMapVertexBuffer(&this->VertexFactory, Data, LightMapIndex);
+		this->ColorVertexBuffer.BindColorVertexBuffer(&this->VertexFactory, Data);
+		this->VertexFactory.SetData(Data);
+
+		InitOrUpdateResource(&this->VertexFactory);
+		PositionVertexBuffer.InitResource();
+		StaticMeshVertexBuffer.InitResource();
+		ColorVertexBuffer.InitResource();
+		IndexBuffer.InitResource();
+		VertexFactory.InitResource();
+	}
+
+
+
+	void UploadUpdate(bool bPositions, bool bNormals, bool bColors)
+	{
+		check(IsInRenderingThread());
+
+		if (TriangleCount == 0)
+		{
+			return;
+		}
+
+		if (bPositions)
+		{
+			InitOrUpdateResource(&this->PositionVertexBuffer);
+		}
+		if (bNormals)
+		{
+			InitOrUpdateResource(&this->StaticMeshVertexBuffer);
+		}
+		if (bColors)
+		{
+			InitOrUpdateResource(&this->ColorVertexBuffer);
+		}
+
+		FLocalVertexFactory::FDataType Data;
+		this->PositionVertexBuffer.BindPositionVertexBuffer(&this->VertexFactory, Data);
+		this->StaticMeshVertexBuffer.BindTangentVertexBuffer(&this->VertexFactory, Data);
+		this->StaticMeshVertexBuffer.BindPackedTexCoordVertexBuffer(&this->VertexFactory, Data);
+		// currently no lightmaps support
+		//this->StaticMeshVertexBuffer.BindLightMapVertexBuffer(&this->VertexFactory, Data, LightMapIndex);
+		this->ColorVertexBuffer.BindColorVertexBuffer(&this->VertexFactory, Data);
+		this->VertexFactory.SetData(Data);
+
+		InitOrUpdateResource(&this->VertexFactory);
+
+		//PositionVertexBuffer.InitResource();
+		//StaticMeshVertexBuffer.InitResource();
+		//ColorVertexBuffer.InitResource();
+		//IndexBuffer.InitResource();
+		//VertexFactory.InitResource();
+	}
+
+
+
+
+	/**
+	 * Initializes a render resource, or update it if already initialized.
+	 * @wearning This function can only be called on the Render Thread
+	 */
+	void InitOrUpdateResource(FRenderResource* Resource)
+	{
+		check(IsInRenderingThread());
+
+		if (!Resource->IsInitialized())
+		{
+			Resource->InitResource();
+		}
+		else
+		{
+			Resource->UpdateRHI();
+		}
+	}
+
+
+
+
+protected:
+	friend class FBaseDynamicMeshSceneProxy;
+
+	/**
+	 * Enqueue a command on the Render Thread to destroy the passed in buffer set.
+	 * At this point the buffer set should be considered invalid.
+	 */
+	static void DestroyRenderBufferSet(FMeshRenderBufferSet* BufferSet)
+	{
+		if (BufferSet->TriangleCount == 0)
+		{
+			return;
+		}
+
+		ENQUEUE_RENDER_COMMAND(FMeshRenderBufferSetDestroy)(
+			[BufferSet](FRHICommandListImmediate& RHICmdList)
+		{
+			delete BufferSet;
+		});
+	}
+
+
+};
+
+
+
+
+/**
+ * FBaseDynamicMeshSceneProxy is an abstract base class for a Render Proxy
+ * for a UBaseDynamicMeshComponent, where the assumption is that mesh data
+ * will be stored in FMeshRenderBufferSet instances
+ */
+class FBaseDynamicMeshSceneProxy : public FPrimitiveSceneProxy
+{
+public:
+	UBaseDynamicMeshComponent* ParentBaseComponent;
+
+	/**
+	 * Constant color assigned to vertices if no other vertex color is specified
+	 */
+	FColor ConstantVertexColor = FColor::White;
+
+	/**
+	 * If true, vertex colors on the FDynamicMesh will be ignored
+	 */
+	bool bIgnoreVertexColors = false;
+
+	/**
+	 * If true, a per-triangle color is used to set vertex colors
+	 */
+	bool bUsePerTriangleColor = false;
+
+	/**
+	 * Per-triangle color function. Only called if bUsePerTriangleColor=true
+	 */
+	TFunction<FColor(const FDynamicMesh3*, int)> PerTriangleColorFunc = nullptr;
+
+
+protected:
+	// Set of currently-allocated RenderBuffers. We own these pointers and must clean them up.
+	// Must guard access with AllocatedSetsLock!!
+	TSet<FMeshRenderBufferSet*> AllocatedBufferSets;
+
+	// use to control access to AllocatedBufferSets 
+	FCriticalSection AllocatedSetsLock;
+
+public:
+	FBaseDynamicMeshSceneProxy(UBaseDynamicMeshComponent* Component)
+		: FPrimitiveSceneProxy(Component),
+		ParentBaseComponent(Component)
+	{
+	}
+
+
+	virtual ~FBaseDynamicMeshSceneProxy()
+	{
+		// we are assuming in code below that this is always called from the rendering thread
+		check(IsInRenderingThread());
+
+		// destroy all existing renderbuffers
+		for (FMeshRenderBufferSet* BufferSet : AllocatedBufferSets)
+		{
+			FMeshRenderBufferSet::DestroyRenderBufferSet(BufferSet);
+		}
+	}
+
+
+	//
+	// FBaseDynamicMeshSceneProxy API - subclasses must implement these functions
+	//
+
+
+	/**
+	 * Return set of active renderbuffers. Must be implemented by subclass.
+	 * This is the set of render buffers that will be drawn by GetDynamicMeshElements
+	 */
+	virtual void GetActiveRenderBufferSets(TArray<FMeshRenderBufferSet*>& Buffers) const = 0;
+
+
+
+	//
+	// RenderBuffer management
+	//
+
+
+	/**
+	 * Allocates a set of render buffers. FPrimitiveSceneProxy will keep track of these
+	 * buffers and destroy them on destruction.
+	 */
+	virtual FMeshRenderBufferSet* AllocateNewRenderBufferSet()
+	{
+		// should we hang onto these and destroy them in constructor? leaving to subclass seems risky?
+		FMeshRenderBufferSet* RenderBufferSet = new FMeshRenderBufferSet(GetScene().GetFeatureLevel());
+
+		RenderBufferSet->Material = UMaterial::GetDefaultMaterial(MD_Surface);
+
+		AllocatedSetsLock.Lock();
+		AllocatedBufferSets.Add(RenderBufferSet);
+		AllocatedSetsLock.Unlock();
+
+		return RenderBufferSet;
+	}
+
+	/**
+	 * Explicitly release a set of RenderBuffers
+	 */
+	virtual void ReleaseRenderBufferSet(FMeshRenderBufferSet* BufferSet)
+	{
+		AllocatedSetsLock.Lock();
+		check(AllocatedBufferSets.Contains(BufferSet));
+		AllocatedBufferSets.Remove(BufferSet);
+		AllocatedSetsLock.Unlock();
+
+		FMeshRenderBufferSet::DestroyRenderBufferSet(BufferSet);
+	}
+
+
+	/**
+	 * Initialize rendering buffers from given attribute overlays.
+	 * Creates three vertices per triangle, IE no shared vertices in buffers.
+	 */
+	template<typename TriangleEnumerable>
+	void InitializeBuffersFromOverlays(
+		FMeshRenderBufferSet* RenderBuffers,
+		const FDynamicMesh3* Mesh,
+		int NumTriangles, TriangleEnumerable Enumerable,
+		FDynamicMeshUVOverlay* UVOverlay,
+		FDynamicMeshNormalOverlay* NormalOverlay,
+		TFunction<void(int, int, int, FVector3f&, FVector3f&)> TangentsFunc = nullptr)
+	{
+		//SCOPE_CYCLE_COUNTER(STAT_SculptToolOctree_InitializeBufferFromOverlay);
+
+		RenderBuffers->TriangleCount = NumTriangles;
+		if (NumTriangles == 0)
+		{
+			return;
+		}
+
+		bool bHaveColors = Mesh->HasVertexColors() && (bIgnoreVertexColors == false);
+
+		int NumVertices = NumTriangles * 3;
+		int NumTexCoords = 1;		// no! zero!
+
+		{
+			RenderBuffers->PositionVertexBuffer.Init(NumVertices);
+			RenderBuffers->StaticMeshVertexBuffer.Init(NumVertices, NumTexCoords);
+			RenderBuffers->ColorVertexBuffer.Init(NumVertices);
+			RenderBuffers->IndexBuffer.Indices.AddUninitialized(NumTriangles * 3);
+		}
+
+		int TriIdx = 0, VertIdx = 0;
+		FVector3f TangentX, TangentY;
+		for (int TriangleID : Enumerable)
+		{
+			FIndex3i Tri = Mesh->GetTriangle(TriangleID);
+			FIndex3i TriUV = (UVOverlay != nullptr) ? UVOverlay->GetTriangle(TriangleID) : FIndex3i::Zero();
+			FIndex3i TriNormal = (NormalOverlay != nullptr) ? NormalOverlay->GetTriangle(TriangleID) : FIndex3i::Zero();
+
+			FColor TriColor = ConstantVertexColor;
+			if (bUsePerTriangleColor && PerTriangleColorFunc != nullptr)
+			{
+				TriColor = PerTriangleColorFunc(Mesh, TriangleID);
+				bHaveColors = false;
+			}
+
+			for (int j = 0; j < 3; ++j)
+			{
+				RenderBuffers->PositionVertexBuffer.VertexPosition(VertIdx) = (FVector)Mesh->GetVertex(Tri[j]);
+
+				FVector3f Normal = (NormalOverlay != nullptr && TriNormal[j] != FDynamicMesh3::InvalidID) ?
+					NormalOverlay->GetElement(TriNormal[j]) : Mesh->GetVertexNormal(Tri[j]);
+
+				// either request known tangent, or calculate a nonsense one
+				if (TangentsFunc != nullptr)
+				{
+					TangentsFunc(Tri[j], TriangleID, j, TangentX, TangentY);
+				}
+				else
+				{
+					VectorUtil::MakePerpVectors(Normal, TangentX, TangentY);
+				}
+				RenderBuffers->StaticMeshVertexBuffer.SetVertexTangents(VertIdx, TangentX, TangentY, Normal);
+
+				FVector2f UV = (UVOverlay != nullptr && TriUV[j] != FDynamicMesh3::InvalidID) ?
+					UVOverlay->GetElement(TriUV[j]) : FVector2f::Zero();
+				RenderBuffers->StaticMeshVertexBuffer.SetVertexUV(VertIdx, 0, UV);
+
+				RenderBuffers->ColorVertexBuffer.VertexColor(VertIdx) = (bHaveColors) ?
+					(FColor)Mesh->GetVertexColor(Tri[j]) : TriColor;
+
+				RenderBuffers->IndexBuffer.Indices[TriIdx++] = VertIdx;
+				VertIdx++;
+			}
+		}
+	}
+
+
+
+
+
+	/**
+	 * Update vertex positions/normals/colors of an existing set of render buffers.
+	 * Assumes that buffers were created with unshared vertices, ie three vertices per triangle, eg by InitializeBuffersFromOverlays()
+	 */
+	template<typename TriangleEnumerable>
+	void UpdateVertexBuffersFromOverlays(
+		FMeshRenderBufferSet* RenderBuffers,
+		const FDynamicMesh3* Mesh,
+		int NumTriangles, TriangleEnumerable Enumerable,
+		FDynamicMeshNormalOverlay* NormalOverlay,
+		bool bUpdatePositions = true,
+		bool bUpdateNormals = false,
+		bool bUpdateColors = false)
+	{
+		if (RenderBuffers->TriangleCount == 0)
+		{
+			return;
+		}
+
+		bool bHaveColors = Mesh->HasVertexColors() && (bIgnoreVertexColors == false);
+
+		int NumVertices = NumTriangles * 3;
+		check(RenderBuffers->PositionVertexBuffer.GetNumVertices() == NumVertices);
+		if (bUpdateNormals)
+		{
+			check(RenderBuffers->StaticMeshVertexBuffer.GetNumVertices() == NumVertices);
+		}
+		if (bUpdateColors)
+		{
+			check(RenderBuffers->ColorVertexBuffer.GetNumVertices() == NumVertices);
+		}
+
+		int VertIdx = 0;
+		FVector3f TangentX, TangentY;
+		for (int TriangleID : Enumerable)
+		{
+			FIndex3i Tri = Mesh->GetTriangle(TriangleID);
+
+			FIndex3i TriNormal = (bUpdateColors && NormalOverlay != nullptr) ? NormalOverlay->GetTriangle(TriangleID) : FIndex3i::Zero();
+
+			FColor TriColor = ConstantVertexColor;
+			if (bUpdateColors && bUsePerTriangleColor && PerTriangleColorFunc != nullptr)
+			{
+				TriColor = PerTriangleColorFunc(Mesh, TriangleID);
+				bHaveColors = false;
+			}
+
+			for (int j = 0; j < 3; ++j)
+			{
+				if (bUpdatePositions)
+				{
+					RenderBuffers->PositionVertexBuffer.VertexPosition(VertIdx) = (FVector)Mesh->GetVertex(Tri[j]);
+				}
+
+				if (bUpdateNormals)
+				{
+					FVector3f Normal = (NormalOverlay != nullptr && TriNormal[j] != FDynamicMesh3::InvalidID) ?
+						NormalOverlay->GetElement(TriNormal[j]) : Mesh->GetVertexNormal(Tri[j]);
+
+					// calculate a nonsense tangent
+					VectorUtil::MakePerpVectors(Normal, TangentX, TangentY);
+					RenderBuffers->StaticMeshVertexBuffer.SetVertexTangents(VertIdx, TangentX, TangentY, Normal);
+				}
+
+				if (bUpdateColors)
+				{
+					RenderBuffers->ColorVertexBuffer.VertexColor(VertIdx) = (bHaveColors) ?
+						(FColor)Mesh->GetVertexColor(Tri[j]) : TriColor;
+				}
+
+				VertIdx++;
+			}
+		}
+	}
+
+
+
+	/**
+	 * @return number of active materials
+	 */
+	virtual int32 GetNumMaterials() const
+	{
+		return ParentBaseComponent->GetNumMaterials();
+	}
+
+	/**
+	 * Safe GetMaterial function that will never return nullptr
+	 */
+	virtual UMaterialInterface* GetMaterial(int32 k) const
+	{
+		UMaterialInterface* Material = ParentBaseComponent->GetMaterial(k);
+		return (Material != nullptr) ? Material : UMaterial::GetDefaultMaterial(MD_Surface);
+	}
+
+
+	/**
+	 * This needs to be called if the set of active materials changes, otherwise
+	 * the check in FPrimitiveSceneProxy::VerifyUsedMaterial() will fail if an override
+	 * material is set, if materials change, etc, etc
+	 */
+	virtual void UpdatedReferencedMaterials()
+	{
+		// copied from FPrimitiveSceneProxy::FPrimitiveSceneProxy()
+#if WITH_EDITOR
+		TArray<UMaterialInterface*> Materials;
+		ParentBaseComponent->GetUsedMaterials(Materials, true);
+		ENQUEUE_RENDER_COMMAND(FMeshRenderBufferSetDestroy)(
+			[this, Materials](FRHICommandListImmediate& RHICmdList)
+		{
+			this->SetUsedMaterialForVerification(Materials);
+		});
+#endif
+	}
+
+
+	//
+	// FBaseDynamicMeshSceneProxy implementation
+	//
+
+
+	/**
+	 * Render set of active RenderBuffers returned by GetActiveRenderBufferSets
+	 */
+	virtual void GetDynamicMeshElements(
+		const TArray<const FSceneView*>& Views, 
+		const FSceneViewFamily& ViewFamily, 
+		uint32 VisibilityMap, 
+		FMeshElementCollector& Collector) const override
+	{
+		QUICK_SCOPE_CYCLE_COUNTER(STAT_OctreeDynamicMeshSceneProxy_GetDynamicMeshElements);
+
+		const bool bWireframe = (AllowDebugViewmodes() && ViewFamily.EngineShowFlags.Wireframe)
+			|| ParentBaseComponent->EnableWireframeRenderPass();
+
+		// set up wireframe material. Probably bad to reference GEngine here...also this material is very bad?
+		FMaterialRenderProxy* WireframeMaterialProxy = nullptr;
+		if (bWireframe)
+		{
+			FColoredMaterialRenderProxy* WireframeMaterialInstance = new FColoredMaterialRenderProxy(
+				GEngine->WireframeMaterial ? GEngine->WireframeMaterial->GetRenderProxy() : nullptr,
+				FLinearColor(0, 0.5f, 1.f)
+			);
+			Collector.RegisterOneFrameMaterialProxy(WireframeMaterialInstance);
+			WireframeMaterialProxy = WireframeMaterialInstance;
+		}
+
+		ESceneDepthPriorityGroup DepthPriority = SDPG_World;
+
+		TArray<FMeshRenderBufferSet*> Buffers;
+		GetActiveRenderBufferSets(Buffers);
+
+		for (int32 ViewIndex = 0; ViewIndex < Views.Num(); ViewIndex++)
+		{
+			if (VisibilityMap & (1 << ViewIndex))
+			{
+				const FSceneView* View = Views[ViewIndex];
+
+				bool bHasPrecomputedVolumetricLightmap;
+				FMatrix PreviousLocalToWorld;
+				int32 SingleCaptureIndex;
+				bool bOutputVelocity;
+				GetScene().GetPrimitiveUniformShaderParameters_RenderThread(GetPrimitiveSceneInfo(), bHasPrecomputedVolumetricLightmap, PreviousLocalToWorld, SingleCaptureIndex, bOutputVelocity);
+
+				// Draw the mesh.
+				for (FMeshRenderBufferSet* BufferSet : Buffers)
+				{
+					UMaterialInterface* UseMaterial = BufferSet->Material;
+					if (ParentBaseComponent->HasOverrideRenderMaterial(0))
+					{
+						UseMaterial = ParentBaseComponent->GetOverrideRenderMaterial(0);
+					}
+					FMaterialRenderProxy* MaterialProxy = UseMaterial->GetRenderProxy();
+
+					if (BufferSet->TriangleCount == 0)
+					{
+						continue;
+					}
+
+					// do we need separate one of these for each MeshRenderBufferSet?
+					FDynamicPrimitiveUniformBuffer& DynamicPrimitiveUniformBuffer = Collector.AllocateOneFrameResource<FDynamicPrimitiveUniformBuffer>();
+					DynamicPrimitiveUniformBuffer.Set(GetLocalToWorld(), PreviousLocalToWorld, GetBounds(), GetLocalBounds(), true, bHasPrecomputedVolumetricLightmap, DrawsVelocity(), bOutputVelocity);
+
+					DrawBatch(Collector, *BufferSet, MaterialProxy, false, DepthPriority, ViewIndex, DynamicPrimitiveUniformBuffer);
+					if (bWireframe)
+					{
+						DrawBatch(Collector, *BufferSet, WireframeMaterialProxy, true, DepthPriority, ViewIndex, DynamicPrimitiveUniformBuffer);
+					}
+				}
+			}
+		}
+	}
+
+
+
+	/**
+	 * Draw a single-frame FMeshBatch for a FMeshRenderBufferSet
+	 */
+	virtual void DrawBatch(FMeshElementCollector& Collector,
+		const FMeshRenderBufferSet& RenderBuffers,
+		FMaterialRenderProxy* UseMaterial,
+		bool bWireframe,
+		ESceneDepthPriorityGroup DepthPriority,
+		int ViewIndex,
+		FDynamicPrimitiveUniformBuffer& DynamicPrimitiveUniformBuffer) const
+	{
+		FMeshBatch& Mesh = Collector.AllocateMesh();
+		FMeshBatchElement& BatchElement = Mesh.Elements[0];
+		BatchElement.IndexBuffer = &RenderBuffers.IndexBuffer;
+		Mesh.bWireframe = bWireframe;
+		Mesh.VertexFactory = &RenderBuffers.VertexFactory;
+		Mesh.MaterialRenderProxy = UseMaterial;
+
+		BatchElement.PrimitiveUniformBufferResource = &DynamicPrimitiveUniformBuffer.UniformBuffer;
+
+		BatchElement.FirstIndex = 0;
+		BatchElement.NumPrimitives = RenderBuffers.IndexBuffer.Indices.Num() / 3;
+		BatchElement.MinVertexIndex = 0;
+		BatchElement.MaxVertexIndex = RenderBuffers.PositionVertexBuffer.GetNumVertices() - 1;
+		Mesh.ReverseCulling = IsLocalToWorldDeterminantNegative();
+		Mesh.Type = PT_TriangleList;
+		Mesh.DepthPriorityGroup = DepthPriority;
+		Mesh.bCanApplyViewModeOverrides = false;
+		Collector.AddMesh(ViewIndex, Mesh);
+	}
+
+
+
+
+
+
+
+};
\ No newline at end of file
diff --git a/Engine/Plugins/Experimental/MeshModelingToolset/Source/ModelingComponents/Private/OctreeDynamicMeshComponent.cpp b/Engine/Plugins/Experimental/MeshModelingToolset/Source/ModelingComponents/Private/OctreeDynamicMeshComponent.cpp
--- a/Engine/Plugins/Experimental/MeshModelingToolset/Source/ModelingComponents/Private/OctreeDynamicMeshComponent.cpp
+++ b/Engine/Plugins/Experimental/MeshModelingToolset/Source/ModelingComponents/Private/OctreeDynamicMeshComponent.cpp
@@ -317,12 +317,19 @@ FPrimitiveSceneProxy* UOctreeDynamicMeshComponent::CreateSceneProxy()
 	return CurrentProxy;
 }
 
-int32 UOctreeDynamicMeshComponent::GetNumMaterials() const
+
+
+void UOctreeDynamicMeshComponent::NotifyMaterialSetUpdated()
 {
-	return 1;
+	if (CurrentProxy != nullptr)
+	{
+		CurrentProxy->UpdatedReferencedMaterials();
+	}
 }
 
 
+
+
 FColor UOctreeDynamicMeshComponent::GetTriangleColor(int TriangleID)
 {
 	if (TriangleColorFunc != nullptr)
diff --git a/Engine/Plugins/Experimental/MeshModelingToolset/Source/ModelingComponents/Private/OctreeDynamicMeshSceneProxy.h b/Engine/Plugins/Experimental/MeshModelingToolset/Source/ModelingComponents/Private/OctreeDynamicMeshSceneProxy.h
--- a/Engine/Plugins/Experimental/MeshModelingToolset/Source/ModelingComponents/Private/OctreeDynamicMeshSceneProxy.h
+++ b/Engine/Plugins/Experimental/MeshModelingToolset/Source/ModelingComponents/Private/OctreeDynamicMeshSceneProxy.h
@@ -4,6 +4,7 @@
 
 #include "OctreeDynamicMeshComponent.h"
 #include "Util/IndexSetDecompositions.h"
+#include "BaseDynamicMeshSceneProxy.h"
 
 
 DECLARE_STATS_GROUP(TEXT("SculptToolOctree"), STATGROUP_SculptToolOctree, STATCAT_Advanced);
@@ -17,102 +18,6 @@ DECLARE_CYCLE_STAT(TEXT("SculptToolOctree_UpdateDecompCreate"), STAT_SculptToolO
 DECLARE_CYCLE_STAT(TEXT("SculptToolOctree_InitializeBufferFromOverlay"), STAT_SculptToolOctree_InitializeBufferFromOverlay, STATGROUP_SculptToolOctree);
 DECLARE_CYCLE_STAT(TEXT("SculptToolOctree_BufferUpload"), STAT_SculptToolOctree_BufferUpload, STATGROUP_SculptToolOctree);
 
-class FMeshRenderBufferSet
-{
-public:
-	int TriangleCount = 0;
-
-	/** The buffer containing vertex data. */
-	FStaticMeshVertexBuffer StaticMeshVertexBuffer;
-	/** The buffer containing the position vertex data. */
-	FPositionVertexBuffer PositionVertexBuffer;
-	/** The buffer containing the vertex color data. */
-	FColorVertexBuffer ColorVertexBuffer;
-
-	FDynamicMeshIndexBuffer32 IndexBuffer;
-	FLocalVertexFactory VertexFactory;
-
-	FMeshRenderBufferSet(ERHIFeatureLevel::Type FeatureLevelType)
-		: VertexFactory(FeatureLevelType, "FMeshRenderBufferSet")
-	{
-
-	}
-
-
-	virtual ~FMeshRenderBufferSet()
-	{
-		check(IsInRenderingThread());
-
-		if (TriangleCount > 0)
-		{
-			PositionVertexBuffer.ReleaseResource();
-			StaticMeshVertexBuffer.ReleaseResource();
-			ColorVertexBuffer.ReleaseResource();
-			IndexBuffer.ReleaseResource();
-			VertexFactory.ReleaseResource();
-		}
-	}
-
-
-	static void DestroyRenderBufferSet(FMeshRenderBufferSet* BufferSet)
-	{
-		if (BufferSet->TriangleCount == 0)
-		{
-			return;
-		}
-
-		ENQUEUE_RENDER_COMMAND(FMeshRenderBufferSetDestroy)(
-			[BufferSet](FRHICommandListImmediate& RHICmdList)
-		{
-			delete BufferSet;
-		});
-	}
-
-	void Upload()
-	{
-		check(IsInRenderingThread());
-
-		if (TriangleCount == 0)
-		{
-			return;
-		}
-
-		InitOrUpdateResource(&this->PositionVertexBuffer);
-		InitOrUpdateResource(&this->StaticMeshVertexBuffer);
-		InitOrUpdateResource(&this->ColorVertexBuffer);
-
-		FLocalVertexFactory::FDataType Data;
-		this->PositionVertexBuffer.BindPositionVertexBuffer(&this->VertexFactory, Data);
-		this->StaticMeshVertexBuffer.BindTangentVertexBuffer(&this->VertexFactory, Data);
-		this->StaticMeshVertexBuffer.BindPackedTexCoordVertexBuffer(&this->VertexFactory, Data);
-		//this->StaticMeshVertexBuffer.BindLightMapVertexBuffer(&this->VertexFactory, Data, LightMapIndex);
-		this->ColorVertexBuffer.BindColorVertexBuffer(&this->VertexFactory, Data);
-		this->VertexFactory.SetData(Data);
-
-		InitOrUpdateResource(&this->VertexFactory);
-		PositionVertexBuffer.InitResource();
-		StaticMeshVertexBuffer.InitResource();
-		ColorVertexBuffer.InitResource();
-		IndexBuffer.InitResource();
-		VertexFactory.InitResource();
-	}
-
-	// copied from StaticMesh.cpp
-	void InitOrUpdateResource(FRenderResource* Resource)
-	{
-		check(IsInRenderingThread());
-
-		if (!Resource->IsInitialized())
-		{
-			Resource->InitResource();
-		}
-		else
-		{
-			Resource->UpdateRHI();
-		}
-	}
-
-};
 
 /**
  * Scene Proxy for a mesh buffer.
@@ -122,13 +27,12 @@ public:
  * Supports wireframe-on-shaded rendering.
  * 
  */
-class FOctreeDynamicMeshSceneProxy final : public FPrimitiveSceneProxy
+class FOctreeDynamicMeshSceneProxy final : public FBaseDynamicMeshSceneProxy
 {
 private:
-	UMaterialInterface* Material;
-
 	FMaterialRelevance MaterialRelevance;
 
+	// note: FBaseDynamicMeshSceneProxy owns and will destroy these
 	TMap<int32, FMeshRenderBufferSet*> RenderBufferSets;
 
 public:
@@ -136,17 +40,8 @@ public:
 	UOctreeDynamicMeshComponent* ParentComponent;
 
 
-	FColor ConstantVertexColor = FColor::White;
-	bool bIgnoreVertexColors = false;
-	bool bIgnoreVertexNormals = false;
-
-
-	bool bUsePerTriangleColor = false;
-	TFunction<FColor(const FDynamicMesh3*, int)> PerTriangleColorFunc = nullptr;
-
-
 	FOctreeDynamicMeshSceneProxy(UOctreeDynamicMeshComponent* Component)
-		: FPrimitiveSceneProxy(Component)
+		: FBaseDynamicMeshSceneProxy(Component)
 		, MaterialRelevance(Component->GetMaterialRelevance(GetScene().GetFeatureLevel()))
 	{
 		// This is an assumption we are currently making. We do not necessarily require this
@@ -154,65 +49,49 @@ public:
 		check(IsInGameThread());
 
 		ParentComponent = Component;
-
-		// Grab material
-		Material = Component->GetMaterial(0);
-		if (Material == NULL)
-		{
-			Material = UMaterial::GetDefaultMaterial(MD_Surface);
-		}
 	}
 
 
-	virtual ~FOctreeDynamicMeshSceneProxy()
-	{
-		// we are assuming in code below that this is always called from the rendering thread
-		check(IsInRenderingThread());
 
+	virtual void GetActiveRenderBufferSets(TArray<FMeshRenderBufferSet*>& Buffers) const override
+	{
 		for (auto MapPair : RenderBufferSets)
 		{
-			FMeshRenderBufferSet* BufferSet = MapPair.Value;
-			FMeshRenderBufferSet::DestroyRenderBufferSet(BufferSet);
+			Buffers.Add(MapPair.Value);
 		}
-		RenderBufferSets.Reset();
 	}
 
 
-	FMeshRenderBufferSet* AllocateNewRenderBufferSet()
-	{
-		FMeshRenderBufferSet* RenderBufferSet = new FMeshRenderBufferSet(GetScene().GetFeatureLevel());
-		return RenderBufferSet;
-	}
-
 
 
 	virtual void InitializeSingleBuffer()
 	{
-			check(RenderBufferSets.Num() == 0);
+		check(RenderBufferSets.Num() == 0);
 
-			FDynamicMesh3* Mesh = ParentComponent->GetMesh();
+		FDynamicMesh3* Mesh = ParentComponent->GetMesh();
 
-			FMeshRenderBufferSet* RenderBuffers = AllocateNewRenderBufferSet();
+		FMeshRenderBufferSet* RenderBuffers = AllocateNewRenderBufferSet();
+		RenderBuffers->Material = GetMaterial(0);
 
-			// find suitable overlays
-			FDynamicMeshUVOverlay* UVOverlay = nullptr;
-			FDynamicMeshNormalOverlay* NormalOverlay = nullptr;
-			if (Mesh->HasAttributes())
-			{
-				UVOverlay = Mesh->Attributes()->PrimaryUV();
-				NormalOverlay = Mesh->Attributes()->PrimaryNormals();
-			}
+		// find suitable overlays
+		FDynamicMeshUVOverlay* UVOverlay = nullptr;
+		FDynamicMeshNormalOverlay* NormalOverlay = nullptr;
+		if (Mesh->HasAttributes())
+		{
+			UVOverlay = Mesh->Attributes()->PrimaryUV();
+			NormalOverlay = Mesh->Attributes()->PrimaryNormals();
+		}
 
-			InitializeBuffersFromOverlays(Mesh,
-				Mesh->TriangleCount(), Mesh->TriangleIndicesItr(),
-				UVOverlay, NormalOverlay, RenderBuffers);
+		InitializeBuffersFromOverlays(RenderBuffers, Mesh,
+			Mesh->TriangleCount(), Mesh->TriangleIndicesItr(),
+			UVOverlay, NormalOverlay);
 
-			ENQUEUE_RENDER_COMMAND(FOctreeDynamicMeshSceneProxyInitializeSingle)(
-				[this, RenderBuffers](FRHICommandListImmediate& RHICmdList)
-			{
-				RenderBuffers->Upload();
-				RenderBufferSets.Add(0, RenderBuffers);
-			});
+		ENQUEUE_RENDER_COMMAND(FOctreeDynamicMeshSceneProxyInitializeSingle)(
+			[this, RenderBuffers](FRHICommandListImmediate& RHICmdList)
+		{
+			RenderBuffers->Upload();
+			RenderBufferSets.Add(0, RenderBuffers);
+		});
 	}
 
 
@@ -238,10 +117,11 @@ public:
 			const TArray<int32>& Tris = Decomposition.GetIndexSetArray(SetID);
 			
 			FMeshRenderBufferSet* RenderBuffers = AllocateNewRenderBufferSet();
+			RenderBuffers->Material = GetMaterial(0);
 
-			InitializeBuffersFromOverlays(Mesh,
+			InitializeBuffersFromOverlays(RenderBuffers, Mesh,
 				Tris.Num(), Tris,
-				UVOverlay, NormalOverlay, RenderBuffers);
+				UVOverlay, NormalOverlay);
 
 			ENQUEUE_RENDER_COMMAND(FOctreeDynamicMeshSceneProxyInitializeFromDecomposition)(
 				[this, SetID, RenderBuffers](FRHICommandListImmediate& RHICmdList)
@@ -272,7 +152,7 @@ public:
 				if (RenderBufferSets.Contains(SetID))
 				{
 					FMeshRenderBufferSet* BufferSet = RenderBufferSets.FindAndRemoveChecked(SetID);
-					FMeshRenderBufferSet::DestroyRenderBufferSet(BufferSet);
+					ReleaseRenderBufferSet(BufferSet);
 				}
 			}
 		});
@@ -289,8 +169,6 @@ public:
 		}
 
 		{
-			FCriticalSection SectionLock;
-
 			SCOPE_CYCLE_COUNTER(STAT_SculptToolOctree_UpdateDecompCreate);
 			int NumSets = SetsToUpdate.Num();
 			ParallelFor(NumSets, [&](int k)
@@ -299,10 +177,11 @@ public:
 				const TArray<int32>& Tris = Decomposition.GetIndexSetArray(SetID);
 
 				FMeshRenderBufferSet* RenderBuffers = AllocateNewRenderBufferSet();
+				RenderBuffers->Material = GetMaterial(0);
 
-				InitializeBuffersFromOverlays(Mesh,
+				InitializeBuffersFromOverlays(RenderBuffers, Mesh,
 					Tris.Num(), Tris,
-					UVOverlay, NormalOverlay, RenderBuffers);
+					UVOverlay, NormalOverlay);
 
 				ENQUEUE_RENDER_COMMAND(FOctreeDynamicMeshSceneProxyUpdateAddOne)(
 					[this, SetID, RenderBuffers](FRHICommandListImmediate& RHICmdList)
@@ -321,170 +200,8 @@ public:
 
 
 
-protected:
-
-
-	/**
-	 * Initialize rendering buffers from given attribute overlays.
-	 * Creates three vertices per triangle, IE no shared vertices in buffers.
-	 */
-	template<typename TriangleEnumerable>
-	void InitializeBuffersFromOverlays(const FDynamicMesh3* Mesh, 
-		int NumTriangles, TriangleEnumerable Enumerable,
-		FDynamicMeshUVOverlay* UVOverlay, 
-		FDynamicMeshNormalOverlay* NormalOverlay,
-		FMeshRenderBufferSet* RenderBuffers)
-	{
-		SCOPE_CYCLE_COUNTER(STAT_SculptToolOctree_InitializeBufferFromOverlay);
-
-		RenderBuffers->TriangleCount = NumTriangles;
-		if (NumTriangles == 0)
-		{
-			return;
-		}
-
-		bool bHaveColors = Mesh->HasVertexColors() && (bIgnoreVertexColors == false);
-
-		int NumVertices = NumTriangles * 3;
-		int NumTexCoords = 1;		// no! zero!
-
-		{
-			RenderBuffers->PositionVertexBuffer.Init(NumVertices);
-			RenderBuffers->StaticMeshVertexBuffer.Init(NumVertices, NumTexCoords);
-			RenderBuffers->ColorVertexBuffer.Init(NumVertices);
-			RenderBuffers->IndexBuffer.Indices.AddUninitialized(NumTriangles * 3);
-		}
-
-		int TriIdx = 0, VertIdx = 0;
-		FVector3f TangentX, TangentY;
-		for (int TriangleID : Enumerable)
-		{
-			FIndex3i Tri = Mesh->GetTriangle(TriangleID);
-			FIndex3i TriUV = (UVOverlay != nullptr) ? UVOverlay->GetTriangle(TriangleID) : FIndex3i::Zero();
-			FIndex3i TriNormal = (NormalOverlay != nullptr) ? NormalOverlay->GetTriangle(TriangleID) : FIndex3i::Zero();
-
-			FColor TriColor = ConstantVertexColor;
-			if (bUsePerTriangleColor && PerTriangleColorFunc != nullptr)
-			{
-				TriColor = PerTriangleColorFunc(Mesh, TriangleID);
-				bHaveColors = false;
-			}
-
-			for (int j = 0; j < 3; ++j)
-			{
-				RenderBuffers->PositionVertexBuffer.VertexPosition(VertIdx) = (FVector)Mesh->GetVertex(Tri[j]);
-
-				FVector3f Normal = (NormalOverlay != nullptr && TriNormal[j] != FDynamicMesh3::InvalidID) ? 
-					NormalOverlay->GetElement(TriNormal[j]) : Mesh->GetVertexNormal(Tri[j]);
-
-				// calculate a nonsense tangent
-				VectorUtil::MakePerpVectors(Normal, TangentX, TangentY);
-				RenderBuffers->StaticMeshVertexBuffer.SetVertexTangents(VertIdx, TangentX, TangentY, Normal);
-
-				FVector2f UV = (UVOverlay != nullptr && TriUV[j] != FDynamicMesh3::InvalidID) ? 
-					UVOverlay->GetElement(TriUV[j]) : FVector2f::Zero();
-				RenderBuffers->StaticMeshVertexBuffer.SetVertexUV(VertIdx, 0, UV);
-
-				RenderBuffers->ColorVertexBuffer.VertexColor(VertIdx) = (bHaveColors) ?
-					(FColor)Mesh->GetVertexColor(Tri[j]) : TriColor;
-
-				RenderBuffers->IndexBuffer.Indices[TriIdx++] = VertIdx;
-				VertIdx++;
-			}
-		}
-	}
-
-
-
-
-
 public:
 
-	virtual void GetDynamicMeshElements(const TArray<const FSceneView*>& Views, const FSceneViewFamily& ViewFamily, uint32 VisibilityMap, FMeshElementCollector& Collector) const override
-	{
-		QUICK_SCOPE_CYCLE_COUNTER(STAT_OctreeDynamicMeshSceneProxy_GetDynamicMeshElements);
-
-		const bool bWireframe = (AllowDebugViewmodes() && ViewFamily.EngineShowFlags.Wireframe)
-			|| ParentComponent->bExplicitShowWireframe;
-
-		FColoredMaterialRenderProxy* WireframeMaterialInstance = new FColoredMaterialRenderProxy(
-			GEngine->WireframeMaterial ? GEngine->WireframeMaterial->GetRenderProxy() : nullptr,
-			FLinearColor(0, 0.5f, 1.f)
-		);
-
-		Collector.RegisterOneFrameMaterialProxy(WireframeMaterialInstance);
-
-		FMaterialRenderProxy* MaterialProxy = Material->GetRenderProxy();
-		FMaterialRenderProxy* WireframeMaterialProxy = WireframeMaterialInstance;
-
-		//ESceneDepthPriorityGroup DepthPriority = (ParentComponent->bDrawOnTop) ? SDPG_Foreground : SDPG_World;
-		ESceneDepthPriorityGroup DepthPriority = SDPG_World;
-
-
-		for (int32 ViewIndex = 0; ViewIndex < Views.Num(); ViewIndex++)
-		{
-			if (VisibilityMap & (1 << ViewIndex))
-			{
-				const FSceneView* View = Views[ViewIndex];
-
-				bool bHasPrecomputedVolumetricLightmap;
-				FMatrix PreviousLocalToWorld;
-				int32 SingleCaptureIndex;
-				bool bOutputVelocity;
-				GetScene().GetPrimitiveUniformShaderParameters_RenderThread(GetPrimitiveSceneInfo(), bHasPrecomputedVolumetricLightmap, PreviousLocalToWorld, SingleCaptureIndex, bOutputVelocity);
-
-				// Draw the mesh.
-				for (auto MapPair : RenderBufferSets)
-				{
-					const FMeshRenderBufferSet& RenderBuffers = *MapPair.Value;
-					if (RenderBuffers.TriangleCount == 0)
-					{
-						continue;
-					}
-
-					// do we need separate one of these for each MeshRenderBufferSet?
-					FDynamicPrimitiveUniformBuffer& DynamicPrimitiveUniformBuffer = Collector.AllocateOneFrameResource<FDynamicPrimitiveUniformBuffer>();
-					DynamicPrimitiveUniformBuffer.Set(GetLocalToWorld(), PreviousLocalToWorld, GetBounds(), GetLocalBounds(), true, bHasPrecomputedVolumetricLightmap, DrawsVelocity(), bOutputVelocity);
-
-					DrawBatch(Collector, RenderBuffers, MaterialProxy, false, DepthPriority, ViewIndex, DynamicPrimitiveUniformBuffer);
-					if (bWireframe)
-					{
-						DrawBatch(Collector, RenderBuffers, WireframeMaterialProxy, true, DepthPriority, ViewIndex, DynamicPrimitiveUniformBuffer);
-					}
-				}
-			}
-		}
-	}
-
-
-
-	virtual void DrawBatch(FMeshElementCollector& Collector, 
-		const FMeshRenderBufferSet& RenderBuffers,
-		FMaterialRenderProxy* UseMaterial, 
-		bool bWireframe,
-		ESceneDepthPriorityGroup DepthPriority,
-		int ViewIndex,
-		FDynamicPrimitiveUniformBuffer& DynamicPrimitiveUniformBuffer) const
-	{
-		FMeshBatch& Mesh = Collector.AllocateMesh();
-		FMeshBatchElement& BatchElement = Mesh.Elements[0];
-		BatchElement.IndexBuffer = &RenderBuffers.IndexBuffer;
-		Mesh.bWireframe = bWireframe;
-		Mesh.VertexFactory = &RenderBuffers.VertexFactory;
-		Mesh.MaterialRenderProxy = UseMaterial;
-
-		BatchElement.PrimitiveUniformBufferResource = &DynamicPrimitiveUniformBuffer.UniformBuffer;
-
-		BatchElement.FirstIndex = 0;
-		BatchElement.NumPrimitives = RenderBuffers.IndexBuffer.Indices.Num() / 3;
-		BatchElement.MinVertexIndex = 0;
-		BatchElement.MaxVertexIndex = RenderBuffers.PositionVertexBuffer.GetNumVertices() - 1;
-		Mesh.ReverseCulling = IsLocalToWorldDeterminantNegative();
-		Mesh.Type = PT_TriangleList;
-		Mesh.DepthPriorityGroup = DepthPriority;
-		Mesh.bCanApplyViewModeOverrides = false;
-		Collector.AddMesh(ViewIndex, Mesh);
-	}
 
 
 
@@ -522,11 +239,6 @@ public:
 	uint32 GetAllocatedSize(void) const { return(FPrimitiveSceneProxy::GetAllocatedSize()); }
 
 
-	virtual void SetMaterial(UMaterialInterface* MaterialIn)
-	{
-		this->Material = MaterialIn;
-	}
-
 
 	SIZE_T GetTypeHash() const override
 	{
diff --git a/Engine/Plugins/Experimental/MeshModelingToolset/Source/ModelingComponents/Private/PreviewMesh.cpp b/Engine/Plugins/Experimental/MeshModelingToolset/Source/ModelingComponents/Private/PreviewMesh.cpp
--- a/Engine/Plugins/Experimental/MeshModelingToolset/Source/ModelingComponents/Private/PreviewMesh.cpp
+++ b/Engine/Plugins/Experimental/MeshModelingToolset/Source/ModelingComponents/Private/PreviewMesh.cpp
@@ -3,6 +3,7 @@
 #include "PreviewMesh.h"
 
 #include "Containers/StaticArray.h"
+#include "Engine/Classes/Materials/Material.h"
 
 
 APreviewMeshActor::APreviewMeshActor()
@@ -61,18 +62,61 @@ void UPreviewMesh::Disconnect()
 
 
 void UPreviewMesh::SetMaterial(UMaterialInterface* Material)
+{
+	SetMaterial(0, Material);
+}
+
+void UPreviewMesh::SetMaterial(int MaterialIndex, UMaterialInterface* Material)
+{
+	check(DynamicMeshComponent);
+	DynamicMeshComponent->SetMaterial(MaterialIndex, Material);
+}
+
+void UPreviewMesh::SetMaterials(const TArray<UMaterialInterface*>& Materials)
+{
+	check(DynamicMeshComponent);
+	for (int k = 0; k < Materials.Num(); ++k)
+	{
+		DynamicMeshComponent->SetMaterial(k, Materials[k]);
+	}
+}
+
+UMaterialInterface* UPreviewMesh::GetMaterial(int MaterialIndex) const
+{
+	check(DynamicMeshComponent);
+	return DynamicMeshComponent->GetMaterial(MaterialIndex);
+}
+
+
+void UPreviewMesh::SetOverrideRenderMaterial(UMaterialInterface* Material)
+{
+	check(DynamicMeshComponent);
+	return DynamicMeshComponent->SetOverrideRenderMaterial(Material);
+}
+
+void UPreviewMesh::ClearOverrideRenderMaterial()
 {
 	check(DynamicMeshComponent);
-	DynamicMeshComponent->SetMaterial(0, Material);
+	return DynamicMeshComponent->ClearOverrideRenderMaterial();
 }
 
-UMaterialInterface*
-UPreviewMesh::GetMaterial() const
+
+UMaterialInterface* UPreviewMesh::GetActiveMaterial(int MaterialIndex) const
+{
+	return DynamicMeshComponent->HasOverrideRenderMaterial(MaterialIndex) ?
+		DynamicMeshComponent->GetOverrideRenderMaterial(MaterialIndex) : GetMaterial(MaterialIndex);
+}
+
+void UPreviewMesh::SetTangentsMode(EDynamicMeshTangentCalcType TangentsType)
 {
 	check(DynamicMeshComponent);
-	return DynamicMeshComponent->GetMaterial(0);
+	DynamicMeshComponent->TangentsType = TangentsType;
 }
 
+FMeshTangentsf* UPreviewMesh::GetTangents() const
+{
+	return DynamicMeshComponent->GetTangents();
+}
 
 void UPreviewMesh::EnableWireframe(bool bEnable)
 {
diff --git a/Engine/Plugins/Experimental/MeshModelingToolset/Source/ModelingComponents/Private/SimpleDynamicMeshComponent.cpp b/Engine/Plugins/Experimental/MeshModelingToolset/Source/ModelingComponents/Private/SimpleDynamicMeshComponent.cpp
--- a/Engine/Plugins/Experimental/MeshModelingToolset/Source/ModelingComponents/Private/SimpleDynamicMeshComponent.cpp
+++ b/Engine/Plugins/Experimental/MeshModelingToolset/Source/ModelingComponents/Private/SimpleDynamicMeshComponent.cpp
@@ -129,6 +129,12 @@ FMeshTangentsf* USimpleDynamicMeshComponent::GetTangents()
 
 	// in this mode we assume the tangents are valid
 	check(TangentsType == EDynamicMeshTangentCalcType::ExternallyCalculated);
+	if (TangentsType == EDynamicMeshTangentCalcType::ExternallyCalculated)
+	{
+		// if you hit this, you did not request ExternallyCalculated tangents before initializing this PreviewMesh
+		check(Tangents.GetTangents().Num() > 0)
+	}
+
 	return &Tangents;
 }
 
@@ -159,7 +165,7 @@ void USimpleDynamicMeshComponent::FastNotifyColorsUpdated()
 {
 	if (CurrentProxy != nullptr)
 	{
-		CurrentProxy->FastUpdateColors();
+		CurrentProxy->FastUpdateVertices(false, false, true);
 	}
 	else
 	{
@@ -173,13 +179,13 @@ void USimpleDynamicMeshComponent::FastNotifyPositionsUpdated()
 {
 	if (CurrentProxy != nullptr)
 	{
-		bTangentsValid = false;
-		CurrentProxy->FastUpdatePositions(false);
+		CurrentProxy->FastUpdateVertices(true, false, false);
 	}
 	else
 	{
 		NotifyMeshUpdated();
 	}
+
 }
 
 
@@ -203,12 +209,18 @@ FPrimitiveSceneProxy* USimpleDynamicMeshComponent::CreateSceneProxy()
 	return CurrentProxy;
 }
 
-int32 USimpleDynamicMeshComponent::GetNumMaterials() const
+
+void USimpleDynamicMeshComponent::NotifyMaterialSetUpdated()
 {
-	return 1;
+	if (CurrentProxy != nullptr)
+	{
+		CurrentProxy->UpdatedReferencedMaterials();
+	}
 }
 
 
+
+
 FColor USimpleDynamicMeshComponent::GetTriangleColor(const FDynamicMesh3* MeshIn, int TriangleID)
 {
 	if (TriangleColorFunc)
diff --git a/Engine/Plugins/Experimental/MeshModelingToolset/Source/ModelingComponents/Private/SimpleDynamicMeshSceneProxy.h b/Engine/Plugins/Experimental/MeshModelingToolset/Source/ModelingComponents/Private/SimpleDynamicMeshSceneProxy.h
--- a/Engine/Plugins/Experimental/MeshModelingToolset/Source/ModelingComponents/Private/SimpleDynamicMeshSceneProxy.h
+++ b/Engine/Plugins/Experimental/MeshModelingToolset/Source/ModelingComponents/Private/SimpleDynamicMeshSceneProxy.h
@@ -3,6 +3,8 @@
 #pragma once
 
 #include "SimpleDynamicMeshComponent.h"
+#include "BaseDynamicMeshSceneProxy.h"
+#include "Async/ParallelFor.h"
 
 
 /**
@@ -13,297 +15,189 @@
  * Supports wireframe-on-shaded rendering.
  * 
  */
-class FSimpleDynamicMeshSceneProxy final : public FPrimitiveSceneProxy
+class FSimpleDynamicMeshSceneProxy final : public FBaseDynamicMeshSceneProxy
 {
 private:
-	UMaterialInterface* Material;
-
-	/** The buffer containing vertex data. */
-	FStaticMeshVertexBuffer StaticMeshVertexBuffer;
-	/** The buffer containing the position vertex data. */
-	FPositionVertexBuffer PositionVertexBuffer;
-	/** The buffer containing the vertex color data. */
-	FColorVertexBuffer ColorVertexBuffer;
-
-	FDynamicMeshIndexBuffer32 IndexBuffer;
-	FLocalVertexFactory VertexFactory;
-
 	FMaterialRelevance MaterialRelevance;
 
+	// note: FBaseDynamicMeshSceneProxy owns and will destroy these
+	TArray<FMeshRenderBufferSet*> RenderBufferSets;
 
+	// if true, we store entire mesh in single RenderBuffers and we can do some optimizations
+	bool bIsSingleBuffer = false;
 
 public:
 	/** Component that created this proxy (is there a way to look this up?) */
 	USimpleDynamicMeshComponent* ParentComponent;
 
 
-	FColor ConstantVertexColor = FColor::White;
-	bool bIgnoreVertexColors = false;
-	bool bIgnoreVertexNormals = false;
-
-
-	bool bUsePerTriangleColor = false;
-	TFunction<FColor(const FDynamicMesh3*, int)> PerTriangleColorFunc;
-
-
 	FSimpleDynamicMeshSceneProxy(USimpleDynamicMeshComponent* Component)
-		: FPrimitiveSceneProxy(Component)
-		, VertexFactory(GetScene().GetFeatureLevel(), "FSimpleDynamicMeshSceneProxy")
+		: FBaseDynamicMeshSceneProxy(Component)
 		, MaterialRelevance(Component->GetMaterialRelevance(GetScene().GetFeatureLevel()))
 	{
 		ParentComponent = Component;
-
-		// Grab material
-		Material = Component->GetMaterial(0);
-		if (Material == NULL)
-		{
-			Material = UMaterial::GetDefaultMaterial(MD_Surface);
-		}
 	}
 
 
-	virtual ~FSimpleDynamicMeshSceneProxy()
+	virtual void GetActiveRenderBufferSets(TArray<FMeshRenderBufferSet*>& Buffers) const override
 	{
-		PositionVertexBuffer.ReleaseResource();
-		StaticMeshVertexBuffer.ReleaseResource();
-		ColorVertexBuffer.ReleaseResource();
-		IndexBuffer.ReleaseResource();
-		VertexFactory.ReleaseResource();
+		Buffers = RenderBufferSets;
 	}
 
 
 
-
-
 	virtual void Initialize()
 	{
-		FDynamicMesh3* Mesh = ParentComponent->GetMesh();
-
-		// find suitable overlays
-		FDynamicMeshUVOverlay* UVOverlay = nullptr;
-		FDynamicMeshNormalOverlay* NormalOverlay = nullptr;
-		if (Mesh->HasAttributes())
+		// allocate buffer sets based on materials
+		check(RenderBufferSets.Num() == 0);
+		int32 NumMaterials = GetNumMaterials();
+		if (NumMaterials == 0)
 		{
-			UVOverlay = Mesh->Attributes()->PrimaryUV();
-			NormalOverlay = Mesh->Attributes()->PrimaryNormals();
-		}
-
-		// fill our buffers
-		if (UVOverlay != nullptr || NormalOverlay != nullptr)
-		{
-			InitializeBuffersFromOverlays(Mesh, UVOverlay, NormalOverlay);
+			RenderBufferSets.SetNum(1);
+			RenderBufferSets[0] = AllocateNewRenderBufferSet();
+			RenderBufferSets[0]->Material = UMaterial::GetDefaultMaterial(MD_Surface);
 		}
 		else
 		{
-			InitializeBuffersFromMesh_SharedVertices(Mesh);
+			RenderBufferSets.SetNum(NumMaterials);
+			for (int32 k = 0; k < NumMaterials; ++k)
+			{
+				RenderBufferSets[k] = AllocateNewRenderBufferSet();
+				RenderBufferSets[k]->Material = GetMaterial(k);
+			}
 		}
 
-		// transfer to render thread?
-		// copied from FStaticMeshVertexBuffers::InitFromDynamicVertex
-		int LightMapIndex = 0;
-		ENQUEUE_RENDER_COMMAND(DynamicMeshComponentVertexBuffersInit)(
-			[this, LightMapIndex](FRHICommandListImmediate& RHICmdList)
-		{
-			InitOrUpdateResource(&this->PositionVertexBuffer);
-			InitOrUpdateResource(&this->StaticMeshVertexBuffer);
-			InitOrUpdateResource(&this->ColorVertexBuffer);
-
-			FLocalVertexFactory::FDataType Data;
-			this->PositionVertexBuffer.BindPositionVertexBuffer(&this->VertexFactory, Data);
-			this->StaticMeshVertexBuffer.BindTangentVertexBuffer(&this->VertexFactory, Data);
-			this->StaticMeshVertexBuffer.BindPackedTexCoordVertexBuffer(&this->VertexFactory, Data);
-			this->StaticMeshVertexBuffer.BindLightMapVertexBuffer(&this->VertexFactory, Data, LightMapIndex);
-			this->ColorVertexBuffer.BindColorVertexBuffer(&this->VertexFactory, Data);
-			this->VertexFactory.SetData(Data);
-
-			InitOrUpdateResource(&this->VertexFactory);
-		});
-
-		// Enqueue initialization of render resource
-		BeginInitResource(&PositionVertexBuffer);
-		BeginInitResource(&StaticMeshVertexBuffer);
-		BeginInitResource(&ColorVertexBuffer);
-		BeginInitResource(&IndexBuffer);
-		BeginInitResource(&VertexFactory);
-	}
-
-
-
-
-
-
-	virtual void FastUpdatePositions(bool bNormals)
-	{
-		check(bNormals == false);
 
 		FDynamicMesh3* Mesh = ParentComponent->GetMesh();
-		int NumVertices = Mesh->TriangleCount() * 3;
-		check(PositionVertexBuffer.GetNumVertices() == NumVertices);
-
-		int VertIdx = 0;
-		for (int TriangleID : Mesh->TriangleIndicesItr())
+		if (Mesh->HasAttributes() && Mesh->Attributes()->HasMaterialID() && NumMaterials > 1)
 		{
-			FIndex3i Tri = Mesh->GetTriangle(TriangleID);
-			for (int j = 0; j < 3; ++j)
-			{
-				PositionVertexBuffer.VertexPosition(VertIdx++) = Mesh->GetVertex(Tri[j]);
-			}
+			bIsSingleBuffer = false;
+			InitializeByMaterial(RenderBufferSets);
 		}
-
-		// @todo is it necessary to rebind everything? does this re-upload the buffers?
-		int LightMapIndex = 0;
-		ENQUEUE_RENDER_COMMAND(DynamicMeshComponentFastUpdatePositions)(
-			[this, LightMapIndex](FRHICommandListImmediate& RHICmdList)
+		else
 		{
-			InitOrUpdateResource(&this->PositionVertexBuffer);
+			bIsSingleBuffer = true;
+			InitializeSingleBufferSet(RenderBufferSets[0]);
+		}
+	}
 
-			FLocalVertexFactory::FDataType Data;
-			this->PositionVertexBuffer.BindPositionVertexBuffer(&this->VertexFactory, Data);
-			this->StaticMeshVertexBuffer.BindTangentVertexBuffer(&this->VertexFactory, Data);
-			this->StaticMeshVertexBuffer.BindPackedTexCoordVertexBuffer(&this->VertexFactory, Data);
-			this->StaticMeshVertexBuffer.BindLightMapVertexBuffer(&this->VertexFactory, Data, LightMapIndex);
-			this->ColorVertexBuffer.BindColorVertexBuffer(&this->VertexFactory, Data);
-			this->VertexFactory.SetData(Data);
 
-			InitOrUpdateResource(&this->VertexFactory);
-		});
 
-		// don't need to do this because they aren't new buffers??
-		//BeginInitResource(&PositionVertexBuffer);
-		//BeginInitResource(&VertexFactory);
-	}
 
 
 
 
 
 
-	virtual void FastUpdateColors()
+	/**
+	 * Initialize a single set of mesh buffers for the entire mesh
+	 */
+	virtual void InitializeSingleBufferSet(FMeshRenderBufferSet* RenderBuffers)
 	{
 		FDynamicMesh3* Mesh = ParentComponent->GetMesh();
-		int NumVertices = Mesh->TriangleCount() * 3;
-		check(ColorVertexBuffer.GetNumVertices() == NumVertices);
-
-		bool bHaveColors = Mesh->HasVertexColors() && (bIgnoreVertexColors == false);
 
-		int VertIdx = 0;
-		for (int TriangleID : Mesh->TriangleIndicesItr())
+		// find suitable overlays
+		FDynamicMeshUVOverlay* UVOverlay = nullptr;
+		FDynamicMeshNormalOverlay* NormalOverlay = nullptr;
+		if (Mesh->HasAttributes())
 		{
-			FIndex3i Tri = Mesh->GetTriangle(TriangleID);
-
-			FColor TriColor = ConstantVertexColor;
-			if (bUsePerTriangleColor && PerTriangleColorFunc)
-			{
-				TriColor = PerTriangleColorFunc(Mesh, TriangleID);
-				bHaveColors = false;
-			}
-
-			for (int j = 0; j < 3; ++j)
-			{
-				ColorVertexBuffer.VertexColor(VertIdx) = (bHaveColors) ?
-					(FColor)Mesh->GetVertexColor(Tri[j]) : TriColor;
-				VertIdx++;
-			}
+			UVOverlay = Mesh->Attributes()->PrimaryUV();
+			NormalOverlay = Mesh->Attributes()->PrimaryNormals();
 		}
 
-		// @todo is it necessary to rebind everything? does this re-upload the buffers?
-		int LightMapIndex = 0;
-		ENQUEUE_RENDER_COMMAND(DynamicMeshComponentFastUpdateColors)(
-			[this, LightMapIndex](FRHICommandListImmediate& RHICmdList)
-		{
-			InitOrUpdateResource(&this->ColorVertexBuffer);
-
-			FLocalVertexFactory::FDataType Data;
-			this->PositionVertexBuffer.BindPositionVertexBuffer(&this->VertexFactory, Data);
-			this->StaticMeshVertexBuffer.BindTangentVertexBuffer(&this->VertexFactory, Data);
-			this->StaticMeshVertexBuffer.BindPackedTexCoordVertexBuffer(&this->VertexFactory, Data);
-			this->StaticMeshVertexBuffer.BindLightMapVertexBuffer(&this->VertexFactory, Data, LightMapIndex);
-			this->ColorVertexBuffer.BindColorVertexBuffer(&this->VertexFactory, Data);
-			this->VertexFactory.SetData(Data);
+		InitializeBuffersFromOverlays(RenderBuffers, Mesh,
+			Mesh->TriangleCount(), Mesh->TriangleIndicesItr(),
+			UVOverlay, NormalOverlay);
 
-			InitOrUpdateResource(&this->VertexFactory);
+		ENQUEUE_RENDER_COMMAND(FOctreeDynamicMeshSceneProxyInitializeSingle)(
+			[RenderBuffers](FRHICommandListImmediate& RHICmdList)
+		{
+			RenderBuffers->Upload();
 		});
-
-		BeginInitResource(&ColorVertexBuffer);		// calls ENQUEUE_RENDER_COMMAND
-		//BeginInitResource(&VertexFactory);
 	}
 
 
 
-protected:
-
-
 	/**
-	 * Initialize rendering buffers from given attribute overlays.
-	 * Creates three vertices per triangle, IE no shared vertices in buffers.
+	 * Initialize the mesh buffers, one per material
 	 */
-	void InitializeBuffersFromOverlays(FDynamicMesh3* Mesh, FDynamicMeshUVOverlay* UVOverlay, FDynamicMeshNormalOverlay* NormalOverlay)
+	virtual void InitializeByMaterial(TArray<FMeshRenderBufferSet*>& BufferSets)
 	{
-		// compute normals if we need to
-		FMeshNormals Normals(Mesh);
-		bool bUseComputedNormals = false;
-		if (NormalOverlay == nullptr && (Mesh->HasVertexNormals() == false || bIgnoreVertexNormals) )
-		{
-			Normals.ComputeVertexNormals();
-			bUseComputedNormals = true;
-		}
+		FDynamicMesh3* Mesh = ParentComponent->GetMesh();
+		check(Mesh->HasAttributes() && Mesh->Attributes()->HasMaterialID());
 
-		bool bHaveColors = Mesh->HasVertexColors() && (bIgnoreVertexColors == false);
+		// find suitable overlays
+		FDynamicMeshUVOverlay* UVOverlay = Mesh->Attributes()->PrimaryUV();
+		FDynamicMeshNormalOverlay* NormalOverlay = Mesh->Attributes()->PrimaryNormals();
+		FDynamicMeshMaterialAttribute* MaterialID = Mesh->Attributes()->GetMaterialID();
 
-		// this may be null
+		TFunction<void(int, int, int, FVector3f&, FVector3f&)> TangentsFunc = nullptr;
 		FMeshTangentsf* Tangents = ParentComponent->GetTangents();
-
-
-		int NumVertices = Mesh->TriangleCount() * 3;
-		int NumTexCoords = 1;		// no! zero!
-		PositionVertexBuffer.Init(NumVertices);
-		StaticMeshVertexBuffer.Init(NumVertices, NumTexCoords);
-		ColorVertexBuffer.Init(NumVertices);
-
-		IndexBuffer.Indices.AddUninitialized(Mesh->TriangleCount() * 3);
-		int TriIdx = 0, VertIdx = 0;
-		FVector3f TangentX, TangentY;
-		for (int TriangleID : Mesh->TriangleIndicesItr())
+		if (Tangents != nullptr)
 		{
-			FIndex3i Tri = Mesh->GetTriangle(TriangleID);
-			FIndex3i TriUV = (UVOverlay != nullptr) ? UVOverlay->GetTriangle(TriangleID) : FIndex3i::Zero();
-			FIndex3i TriNormal = (NormalOverlay != nullptr) ? NormalOverlay->GetTriangle(TriangleID) : FIndex3i::Zero();
+			TangentsFunc = [Tangents](int VertexID, int TriangleID, int TriVtxIdx, FVector3f& TangentX, FVector3f& TangentY)
+			{
+				return Tangents->GetPerTriangleTangent(TriangleID, TriVtxIdx, TangentX, TangentY);
+			};
+		}
 
-			FColor TriColor = ConstantVertexColor;
-			if (bUsePerTriangleColor && PerTriangleColorFunc)
+		// count number of triangles for each material (could do in parallel?)
+		int NumMaterials = BufferSets.Num();
+		TArray<FThreadSafeCounter> Counts;
+		Counts.SetNum(NumMaterials);
+		for (int k = 0; k < NumMaterials; ++k)
+		{
+			Counts[k].Reset();
+		}
+		ParallelFor(Mesh->MaxTriangleID(), [&](int tid)
+		{
+			int MatIdx;
+			MaterialID->GetValue(tid, &MatIdx);
+			if (MatIdx >= 0 && MatIdx < NumMaterials)
 			{
-				TriColor = PerTriangleColorFunc(Mesh, TriangleID);
-				bHaveColors = false;
+				Counts[MatIdx].Increment();
 			}
+		});
 
-			for (int j = 0; j < 3; ++j)
-			{
-				PositionVertexBuffer.VertexPosition(VertIdx) = Mesh->GetVertex(Tri[j]);
+		// find max count
+		int32 MaxCount = 0;
+		for (FThreadSafeCounter& Count : Counts)
+		{
+			MaxCount = FMath::Max(Count.GetValue(), MaxCount);
+		}
 
-				FVector3f Normal = (NormalOverlay != nullptr && TriNormal[j] != FDynamicMesh3::InvalidID) ? NormalOverlay->GetElement(TriNormal[j])
-					: ((bUseComputedNormals) ? (FVector3f)Normals[Tri[j]] : Mesh->GetVertexNormal(Tri[j]));
+		// init renderbuffers for each material
+		// could do this in parallel but then we need to allocate separate triangle arrays...is it worth it?
+		TArray<int> Triangles;
+		Triangles.Reserve(MaxCount);
+		for (int k = 0; k < NumMaterials; ++k)
+		{
+			if (Counts[k].GetValue() > 0)
+			{
+				FMeshRenderBufferSet* RenderBuffers = BufferSets[k];
 
-				if (Tangents == nullptr)
-				{
-					// calculate a nonsense tangent
-					VectorUtil::MakePerpVectors(Normal, TangentX, TangentY);
-					StaticMeshVertexBuffer.SetVertexTangents(VertIdx, TangentX, TangentY, Normal);
-				}
-				else
+				Triangles.Reset();
+				for (int tid : Mesh->TriangleIndicesItr())
 				{
-					Tangents->GetPerTriangleTangent(TriangleID, j, TangentX, TangentY);
+					int MatIdx;
+					MaterialID->GetValue(tid, &MatIdx);
+					if (MatIdx == k)
+					{
+						Triangles.Add(tid);
+					}
 				}
 
-				StaticMeshVertexBuffer.SetVertexTangents(VertIdx, TangentX, TangentY, Normal);
+				InitializeBuffersFromOverlays(RenderBuffers, Mesh,
+					Triangles.Num(), Triangles,
+					UVOverlay, NormalOverlay);
 
-				FVector2f UV = (UVOverlay != nullptr && TriUV[j] != FDynamicMesh3::InvalidID) ? UVOverlay->GetElement(TriUV[j]) : FVector2f::Zero();
-				StaticMeshVertexBuffer.SetVertexUV(VertIdx, 0, UV);
+				RenderBuffers->Triangles = Triangles;
 
-				ColorVertexBuffer.VertexColor(VertIdx) = (bHaveColors) ?
-					(FColor)Mesh->GetVertexColor(Tri[j]) : TriColor;
-
-				IndexBuffer.Indices[TriIdx++] = VertIdx;
-				VertIdx++;
+				ENQUEUE_RENDER_COMMAND(FOctreeDynamicMeshSceneProxyInitializeSingle)(
+					[RenderBuffers](FRHICommandListImmediate& RHICmdList)
+				{
+					RenderBuffers->Upload();
+				});
 			}
 		}
 	}
@@ -312,134 +206,65 @@ protected:
 
 
 
+
 	/**
-	 * Initialize rendering buffers from DMesh using shared vertices
-	 * (ie no attribute overlays)
+	 * Update the vertex position/normal/color buffers
 	 */
-	void InitializeBuffersFromMesh_SharedVertices(FDynamicMesh3* Mesh)
+	virtual void FastUpdateVertices(bool bPositions, bool bNormals, bool bColors)
 	{
-		// compute normals if we need to
-		FMeshNormals Normals(Mesh);
-		bool bUseComputedNormals = false;
-		if (Mesh->HasVertexNormals() == false || bIgnoreVertexNormals)
-		{
-			Normals.ComputeVertexNormals();
-			bUseComputedNormals = true;
-		}
-
-		bool bHaveColors = Mesh->HasVertexColors() && (bIgnoreVertexColors == false);
+		// This needs to be rewritten for split-by-material buffers.
+		// Could store triangle set with each buffer, and then rebuild vtx buffer(s) as needed?
 
-		int NumVertices = Mesh->MaxVertexID();
-		int NumTexCoords = 1;		// no! zero!
-		PositionVertexBuffer.Init(NumVertices);
-		StaticMeshVertexBuffer.Init(NumVertices, NumTexCoords);
-		ColorVertexBuffer.Init(NumVertices);
+		FDynamicMesh3* Mesh = ParentComponent->GetMesh();
 
-		for (int VertIdx : Mesh->VertexIndicesItr())
+		// find suitable overlays
+		FDynamicMeshNormalOverlay* NormalOverlay = nullptr;
+		if (bNormals)
 		{
-			PositionVertexBuffer.VertexPosition(VertIdx) = Mesh->GetVertex(VertIdx);
-
-			FVector3d Normal = (bUseComputedNormals) ? Normals[VertIdx] : (FVector3d)Mesh->GetVertexNormal(VertIdx);
-			FVector3d TangentX, TangentY;
-			VectorUtil::MakePerpVectors(Normal, TangentX, TangentY);
-
-			StaticMeshVertexBuffer.SetVertexTangents(VertIdx, TangentX, TangentY, Normal);
-			StaticMeshVertexBuffer.SetVertexUV(VertIdx, 0, FVector2D::ZeroVector);
-
-			ColorVertexBuffer.VertexColor(VertIdx) = (bHaveColors) ?
-				(FColor)Mesh->GetVertexColor(VertIdx) : ConstantVertexColor;
+			check(Mesh->HasAttributes());
+			NormalOverlay = Mesh->Attributes()->PrimaryNormals();
 		}
 
-		IndexBuffer.Indices.AddUninitialized(Mesh->TriangleCount() * 3);
-		int TriIndex = 0;
-		for (FIndex3i Tri : Mesh->TrianglesItr())
+		if (bIsSingleBuffer)
 		{
-			IndexBuffer.Indices[TriIndex++] = Tri.A;
-			IndexBuffer.Indices[TriIndex++] = Tri.B;
-			IndexBuffer.Indices[TriIndex++] = Tri.C;
+			check(RenderBufferSets.Num() == 1);
+			FMeshRenderBufferSet* Buffers = RenderBufferSets[0];
+			UpdateVertexBuffersFromOverlays(Buffers, Mesh,
+				Mesh->TriangleCount(), Mesh->TriangleIndicesItr(),
+				NormalOverlay,
+				bPositions, bNormals, bColors);
+
+			ENQUEUE_RENDER_COMMAND(FOctreeDynamicMeshSceneProxyInitializeSingle)(
+				[Buffers, bPositions, bNormals, bColors](FRHICommandListImmediate& RHICmdList)
+			{
+				Buffers->UploadUpdate(bPositions, bNormals, bColors);
+			});
 		}
-	}
-
-
-
-
-public:
-
-
-
-	virtual void GetDynamicMeshElements(const TArray<const FSceneView*>& Views, const FSceneViewFamily& ViewFamily, uint32 VisibilityMap, FMeshElementCollector& Collector) const override
-	{
-		QUICK_SCOPE_CYCLE_COUNTER(STAT_SimpleDynamicMeshSceneProxy_GetDynamicMeshElements);
-
-		const bool bWireframe = (AllowDebugViewmodes() && ViewFamily.EngineShowFlags.Wireframe)
-			|| ParentComponent->bExplicitShowWireframe;
-
-		FColoredMaterialRenderProxy* WireframeMaterialInstance = new FColoredMaterialRenderProxy(
-			GEngine->WireframeMaterial ? GEngine->WireframeMaterial->GetRenderProxy() : nullptr,
-			FLinearColor(0, 0.5f, 1.f)
-		);
-
-		Collector.RegisterOneFrameMaterialProxy(WireframeMaterialInstance);
-
-		FMaterialRenderProxy* MaterialProxy = Material->GetRenderProxy();
-		FMaterialRenderProxy* WireframeMaterialProxy = WireframeMaterialInstance;
-
-		ESceneDepthPriorityGroup DepthPriority = (ParentComponent->bDrawOnTop) ? SDPG_Foreground : SDPG_World;
-
-
-		for (int32 ViewIndex = 0; ViewIndex < Views.Num(); ViewIndex++)
+		else
 		{
-			if (VisibilityMap & (1 << ViewIndex))
+			ParallelFor(RenderBufferSets.Num(), [&](int i)
 			{
-				const FSceneView* View = Views[ViewIndex];
-
-				bool bHasPrecomputedVolumetricLightmap;
-				FMatrix PreviousLocalToWorld;
-				int32 SingleCaptureIndex;
-				bool bOutputVelocity;
-				GetScene().GetPrimitiveUniformShaderParameters_RenderThread(GetPrimitiveSceneInfo(), bHasPrecomputedVolumetricLightmap, PreviousLocalToWorld, SingleCaptureIndex, bOutputVelocity);
-
-				FDynamicPrimitiveUniformBuffer& DynamicPrimitiveUniformBuffer = Collector.AllocateOneFrameResource<FDynamicPrimitiveUniformBuffer>();
-				DynamicPrimitiveUniformBuffer.Set(GetLocalToWorld(), PreviousLocalToWorld, GetBounds(), GetLocalBounds(), true, bHasPrecomputedVolumetricLightmap, DrawsVelocity(), bOutputVelocity);
-
-				// Draw the mesh.
-				DrawBatch(Collector, MaterialProxy, false, DepthPriority, ViewIndex, DynamicPrimitiveUniformBuffer);
-				if (bWireframe)
+				FMeshRenderBufferSet* Buffers = RenderBufferSets[i];
+				check(Buffers->Triangles.IsSet());
+				UpdateVertexBuffersFromOverlays(Buffers, Mesh,
+					Buffers->Triangles->Num(), Buffers->Triangles.GetValue(),
+					NormalOverlay,
+					bPositions, bNormals, bColors);
+
+				ENQUEUE_RENDER_COMMAND(FOctreeDynamicMeshSceneProxyInitializeSingle)(
+					[Buffers, bPositions, bNormals, bColors](FRHICommandListImmediate& RHICmdList)
 				{
-					DrawBatch(Collector, WireframeMaterialProxy, true, DepthPriority, ViewIndex, DynamicPrimitiveUniformBuffer);
-				}
-			}
+					Buffers->UploadUpdate(bPositions, bNormals, bColors);
+				});
+			});
 		}
+
 	}
 
 
 
-	virtual void DrawBatch(FMeshElementCollector& Collector, 
-		FMaterialRenderProxy* UseMaterial, 
-		bool bWireframe,
-		ESceneDepthPriorityGroup DepthPriority,
-		int ViewIndex,
-		FDynamicPrimitiveUniformBuffer& DynamicPrimitiveUniformBuffer) const
-	{
-		FMeshBatch& Mesh = Collector.AllocateMesh();
-		FMeshBatchElement& BatchElement = Mesh.Elements[0];
-		BatchElement.IndexBuffer = &IndexBuffer;
-		Mesh.bWireframe = bWireframe;
-		Mesh.VertexFactory = &VertexFactory;
-		Mesh.MaterialRenderProxy = UseMaterial;
-
-		BatchElement.PrimitiveUniformBufferResource = &DynamicPrimitiveUniformBuffer.UniformBuffer;
-
-		BatchElement.FirstIndex = 0;
-		BatchElement.NumPrimitives = IndexBuffer.Indices.Num() / 3;
-		BatchElement.MinVertexIndex = 0;
-		BatchElement.MaxVertexIndex = PositionVertexBuffer.GetNumVertices() - 1;
-		Mesh.ReverseCulling = IsLocalToWorldDeterminantNegative();
-		Mesh.Type = PT_TriangleList;
-		Mesh.DepthPriorityGroup = DepthPriority;
-		Mesh.bCanApplyViewModeOverrides = false;
-		Collector.AddMesh(ViewIndex, Mesh);
-	}
+
+public:
 
 
 
@@ -447,7 +272,6 @@ public:
 	{
 		FPrimitiveViewRelevance Result;
 
-
 		if (ParentComponent->bDrawOnTop)
 		{
 			Result.bDrawRelevance = IsShown(View);
@@ -491,10 +315,6 @@ public:
 
 	virtual bool CanBeOccluded() const override
 	{
-		//if (ParentComponent->bDrawOnTop)
-		//{
-		//	return false;
-		//}
 		return !MaterialRelevance.bDisableDepthTest;
 	}
 
@@ -503,27 +323,6 @@ public:
 	uint32 GetAllocatedSize(void) const { return(FPrimitiveSceneProxy::GetAllocatedSize()); }
 
 
-
-	virtual void SetMaterial(UMaterialInterface* MaterialIn)
-	{
-		this->Material = MaterialIn;
-	}
-
-
-	// copied from StaticMesh.cpp
-	void InitOrUpdateResource(FRenderResource* Resource)
-	{
-		if (!Resource->IsInitialized())
-		{
-			Resource->InitResource();
-		}
-		else
-		{
-			Resource->UpdateRHI();
-		}
-	}
-
-
 	SIZE_T GetTypeHash() const override
 	{
 		static size_t UniquePointer;
diff --git a/Engine/Plugins/Experimental/MeshModelingToolset/Source/ModelingComponents/Public/BaseDynamicMeshComponent.h b/Engine/Plugins/Experimental/MeshModelingToolset/Source/ModelingComponents/Public/BaseDynamicMeshComponent.h
--- a/Engine/Plugins/Experimental/MeshModelingToolset/Source/ModelingComponents/Public/BaseDynamicMeshComponent.h
+++ b/Engine/Plugins/Experimental/MeshModelingToolset/Source/ModelingComponents/Public/BaseDynamicMeshComponent.h
@@ -72,4 +72,71 @@ public:
 		unimplemented();
 	}
 
+
+protected:
+	/**
+	 * Subclass must implement this to notify allocated proxies of updated materials
+	 */
+	virtual void NotifyMaterialSetUpdated()
+	{
+		unimplemented();
+	}
+
+
+
+public:
+
+	/**
+	 * @return true if wireframe rendering pass is enabled (default false)
+	 */
+	virtual bool EnableWireframeRenderPass() const { return false; }
+
+
+	//
+	// Override rendering material support
+	//
+
+	/**
+	 * Set an active override render material. This should replace all materials during rendering.
+	 */
+	virtual void SetOverrideRenderMaterial(UMaterialInterface* Material);
+
+	/**
+	 * Clear any active override render material
+	 */
+	virtual void ClearOverrideRenderMaterial();
+	
+	/**
+	 * @return true if an override render material is currently enabled for the given MaterialIndex
+	 */
+	virtual bool HasOverrideRenderMaterial(int k) const
+	{
+		return OverrideRenderMaterial != nullptr;
+	}
+
+	/**
+	 * @return active override render material for the given MaterialIndex
+	 */
+	virtual UMaterialInterface* GetOverrideRenderMaterial(int MaterialIndex) const
+	{
+		return OverrideRenderMaterial;
+	}
+
+
+protected:
+
+	TArray<UMaterialInterface*> BaseMaterials;
+
+	UMaterialInterface* OverrideRenderMaterial = nullptr;
+
+public:
+
+
+	//~ Begin UMeshComponent Interface.
+	virtual int32 GetNumMaterials() const override;
+	virtual UMaterialInterface* GetMaterial(int32 ElementIndex) const override;
+	virtual void SetMaterial(int32 ElementIndex, UMaterialInterface* Material) override;
+	virtual void GetUsedMaterials(TArray<UMaterialInterface*>& OutMaterials, bool bGetDebugMaterials = false) const override;
+	//~ End UMeshComponent Interface.
+
 };
diff --git a/Engine/Plugins/Experimental/MeshModelingToolset/Source/ModelingComponents/Public/OctreeDynamicMeshComponent.h b/Engine/Plugins/Experimental/MeshModelingToolset/Source/ModelingComponents/Public/OctreeDynamicMeshComponent.h
--- a/Engine/Plugins/Experimental/MeshModelingToolset/Source/ModelingComponents/Public/OctreeDynamicMeshComponent.h
+++ b/Engine/Plugins/Experimental/MeshModelingToolset/Source/ModelingComponents/Public/OctreeDynamicMeshComponent.h
@@ -111,8 +111,23 @@ public:
 	UPROPERTY()
 	bool bExplicitShowWireframe = false;
 
+	/**
+	 * @return true if wireframe rendering pass is enabled
+	 */
+	virtual bool EnableWireframeRenderPass() const override { return bExplicitShowWireframe; }
+
+	/**
+	 * If this function is set, we will use these colors instead of vertex colors
+	 */
 	TFunction<FColor(const FDynamicMesh3*, int)> TriangleColorFunc = nullptr;
 
+protected:
+	/**
+	 * This is called to tell our RenderProxy about modifications to the material set.
+	 * We need to pass this on for things like material validation in the Editor.
+	 */
+	virtual void NotifyMaterialSetUpdated();
+
 private:
 
 	FOctreeDynamicMeshSceneProxy* CurrentProxy = nullptr;
@@ -121,10 +136,6 @@ private:
 	virtual FPrimitiveSceneProxy* CreateSceneProxy() override;
 	//~ End UPrimitiveComponent Interface.
 
-	//~ Begin UMeshComponent Interface.
-	virtual int32 GetNumMaterials() const override;
-	//~ End UMeshComponent Interface.
-
 	//~ Begin USceneComponent Interface.
 	virtual FBoxSphereBounds CalcBounds(const FTransform& LocalToWorld) const override;
 	//~ Begin USceneComponent Interface.
diff --git a/Engine/Plugins/Experimental/MeshModelingToolset/Source/ModelingComponents/Public/PreviewMesh.h b/Engine/Plugins/Experimental/MeshModelingToolset/Source/ModelingComponents/Public/PreviewMesh.h
--- a/Engine/Plugins/Experimental/MeshModelingToolset/Source/ModelingComponents/Public/PreviewMesh.h
+++ b/Engine/Plugins/Experimental/MeshModelingToolset/Source/ModelingComponents/Public/PreviewMesh.h
@@ -12,6 +12,7 @@
 
 class FDynamicMesh3;
 struct FMeshDescription;
+class FTangentsf;
 
 
 
@@ -95,10 +96,50 @@ public:
 	 */
 	void SetMaterial(UMaterialInterface* Material);
 
+	/**
+	 * Set material on the preview mesh
+	 */
+	void SetMaterial(int MaterialIndex, UMaterialInterface* Material);
+
+	/**
+	 * Set the entire material set on the preview mesh
+	 */
+	void SetMaterials(const TArray<UMaterialInterface*>& Materials);
+
 	/**
 	* Get material from the preview mesh
 	*/
-	UMaterialInterface* GetMaterial() const;
+	UMaterialInterface* GetMaterial(int MaterialIndex = 0) const;
+
+	/**
+	 * Set an override material for the preview mesh. This material will override all the given materials.
+	 */
+	void SetOverrideRenderMaterial(UMaterialInterface* Material);
+
+	/**
+	 * Clear the override material for the preview mesh.
+	 */
+	void ClearOverrideRenderMaterial();
+
+	/**
+	 * @return the actual material that will be used for rendering for the given MaterialIndex. Will return override material if set.
+	 * 
+	 */
+	UMaterialInterface* GetActiveMaterial(int MaterialIndex = 0) const;
+
+
+	/**
+	 * Set the tangents mode for the underlying component, if available. 
+	 * Note that this function may need to be called before the mesh is initialized.
+	 */
+	void SetTangentsMode(EDynamicMeshTangentCalcType TangentsType);
+
+	/**
+	 * @return a MeshTangents data structure for the underlying component, if available, otherwise nullptr
+	 */
+	FMeshTangentsf* GetTangents() const;
+
+
 	/**
 	 * Get the current transform on the preview mesh
 	 */
diff --git a/Engine/Plugins/Experimental/MeshModelingToolset/Source/ModelingComponents/Public/SimpleDynamicMeshComponent.h b/Engine/Plugins/Experimental/MeshModelingToolset/Source/ModelingComponents/Public/SimpleDynamicMeshComponent.h
--- a/Engine/Plugins/Experimental/MeshModelingToolset/Source/ModelingComponents/Public/SimpleDynamicMeshComponent.h
+++ b/Engine/Plugins/Experimental/MeshModelingToolset/Source/ModelingComponents/Public/SimpleDynamicMeshComponent.h
@@ -61,6 +61,8 @@ public:
 	FMeshTangentsf* GetTangents();
 
 
+
+
 	/**
 	 * Write the internal mesh to a MeshDescription
 	 * @param bHaveModifiedTopology if false, we only update the vertex positions in the MeshDescription, otherwise it is Empty()'d and regenerated entirely
@@ -84,11 +86,23 @@ public:
 	//
 
 	/**
-	 * Call this if you update the mesh via GetMesh()
+	 * Call this if you update the mesh via GetMesh(). This will destroy the existing RenderProxy and create a new one.
 	 * @todo should provide a function that calls a lambda to modify the mesh, and only return const mesh pointer
 	 */
 	virtual void NotifyMeshUpdated() override;
 
+	/**
+	 * Call this instead of NotifyMeshUpdated() if you have only updated the vertex colors (or triangle color function).
+	 * This function will update the existing RenderProxy buffers if possible
+	 */
+	void FastNotifyColorsUpdated();
+
+	/**
+	 * Call this instead of NotifyMeshUpdated() if you have only updated the vertex positions.
+	 * This function will update the existing RenderProxy buffers if possible
+	 */
+	void FastNotifyPositionsUpdated();
+
 	/**
 	 * Apply a vertex deformation change to the internal mesh
 	 */
@@ -107,24 +121,40 @@ public:
 	FSimpleMulticastDelegate OnMeshChanged;
 
 
-	void SetDrawOnTop(bool bSet);
-
 	/**
 	 * if true, we always show the wireframe on top of the shaded mesh, even when not in wireframe mode
 	 */
 	UPROPERTY()
 	bool bExplicitShowWireframe = false;
 
+	/**
+	 * @return true if wireframe rendering pass is enabled
+	 */
+	virtual bool EnableWireframeRenderPass() const override { return bExplicitShowWireframe; }
+
 
+	/**
+	 * If this function is set, we will use these colors instead of vertex colors
+	 */
+	TFunction<FColor(const FDynamicMesh3*, int)> TriangleColorFunc = nullptr;
+
+
+public:
+
+	// do not use this
 	UPROPERTY()
 	bool bDrawOnTop = false;
 
+	// do not use this
+	void SetDrawOnTop(bool bSet);
 
 
-	TFunction<FColor(const FDynamicMesh3*, int)> TriangleColorFunc = nullptr;
-	void FastNotifyColorsUpdated();
-
-	void FastNotifyPositionsUpdated();
+protected:
+	/**
+	 * This is called to tell our RenderProxy about modifications to the material set.
+	 * We need to pass this on for things like material validation in the Editor.
+	 */
+	virtual void NotifyMaterialSetUpdated();
 
 private:
 
@@ -134,10 +164,6 @@ private:
 	virtual FPrimitiveSceneProxy* CreateSceneProxy() override;
 	//~ End UPrimitiveComponent Interface.
 
-	//~ Begin UMeshComponent Interface.
-	virtual int32 GetNumMaterials() const override;
-	//~ End UMeshComponent Interface.
-
 	//~ Begin USceneComponent Interface.
 	virtual FBoxSphereBounds CalcBounds(const FTransform& LocalToWorld) const override;
 	//~ Begin USceneComponent Interface.
diff --git a/Engine/Source/Runtime/Experimental/InteractiveToolsFramework/Public/ComponentSourceInterfaces.h b/Engine/Source/Runtime/Experimental/InteractiveToolsFramework/Public/ComponentSourceInterfaces.h
--- a/Engine/Source/Runtime/Experimental/InteractiveToolsFramework/Public/ComponentSourceInterfaces.h
+++ b/Engine/Source/Runtime/Experimental/InteractiveToolsFramework/Public/ComponentSourceInterfaces.h
@@ -15,6 +15,15 @@ class UPrimitiveComponent;
 struct FMeshDescription;
 class UMaterialInterface;
 
+/**
+ * FComponentMaterialSet is the set of materials assigned to a component (ie Material Slots on a StaticMesh)
+ */
+struct INTERACTIVETOOLSFRAMEWORK_API FComponentMaterialSet
+{
+	TArray<UMaterialInterface*> Materials;
+};
+
+
 /**
  * Wrapper around a UObject Component that can provide a MeshDescription, and
  * (optionally) bake a modified MeshDescription back to this Component.
@@ -49,6 +58,12 @@ public:
 	 */
 	UMaterialInterface* GetMaterial(int32 MaterialIndex) const;
 
+	/**
+	 * Get material set provided by this source
+	 * @param MaterialSetOut returned material set
+	 */
+	void GetMaterialSet(FComponentMaterialSet& MaterialSetOut) const;
+
 	/**
 	 * @return the transform on this component
 	 * @todo Do we need to return a list of transforms here?
