ModelingTools: update various tools with new multi-material and override-material support - MeshSculpt, base MeshBrush, MeshInspector, MeshSelection (currently does not display materials). Material-assignment code in MeshSelectionTool (disabled)
#rnx

diff --git a/Engine/Plugins/Experimental/MeshModelingToolset/Source/MeshModelingTools/Private/DynamicMeshBrushTool.cpp b/Engine/Plugins/Experimental/MeshModelingToolset/Source/MeshModelingTools/Private/DynamicMeshBrushTool.cpp
--- a/Engine/Plugins/Experimental/MeshModelingToolset/Source/MeshModelingTools/Private/DynamicMeshBrushTool.cpp
+++ b/Engine/Plugins/Experimental/MeshModelingToolset/Source/MeshModelingTools/Private/DynamicMeshBrushTool.cpp
@@ -23,10 +23,10 @@ void UDynamicMeshBrushTool::Setup()
 	PreviewMesh->bBuildSpatialDataStructure = true;
 	PreviewMesh->CreateInWorld(ComponentTarget->GetOwnerActor()->GetWorld(), FTransform::Identity);
 	PreviewMesh->SetTransform(ComponentTarget->GetWorldTransform());
-	if (ComponentTarget->GetMaterial(0) != nullptr)
-	{
-		PreviewMesh->SetMaterial(ComponentTarget->GetMaterial(0));
-	}
+
+	FComponentMaterialSet MaterialSet;
+	ComponentTarget->GetMaterialSet(MaterialSet);
+	PreviewMesh->SetMaterials(MaterialSet.Materials);
 
 	// initialize from LOD-0 MeshDescription
 	PreviewMesh->InitializeMesh(ComponentTarget->GetMesh());
diff --git a/Engine/Plugins/Experimental/MeshModelingToolset/Source/MeshModelingTools/Private/DynamicMeshSculptTool.cpp b/Engine/Plugins/Experimental/MeshModelingToolset/Source/MeshModelingTools/Private/DynamicMeshSculptTool.cpp
--- a/Engine/Plugins/Experimental/MeshModelingToolset/Source/MeshModelingTools/Private/DynamicMeshSculptTool.cpp
+++ b/Engine/Plugins/Experimental/MeshModelingToolset/Source/MeshModelingTools/Private/DynamicMeshSculptTool.cpp
@@ -1997,19 +1997,15 @@ void UDynamicMeshSculptTool::UpdateMaterialMode(EMeshEditingMaterialModes Materi
 {
 	if (MaterialMode == EMeshEditingMaterialModes::ExistingMaterial)
 	{
-		UMaterialInterface* Material = ComponentTarget->GetMaterial(0);
-		if (Material != nullptr)
-		{
-			DynamicMeshComponent->SetMaterial(0, Material);
-		}
+		DynamicMeshComponent->ClearOverrideRenderMaterial();
 		DynamicMeshComponent->bCastDynamicShadow = ComponentTarget->GetOwnerComponent()->bCastDynamicShadow;
 	} 
 	else if (MaterialMode == EMeshEditingMaterialModes::MeshFocusMaterial)
 	{ 
-		UMaterialInterface* Material = ToolSetupUtil::GetSculptMaterial1(GetToolManager());
-		if (Material != nullptr)
+		UMaterialInterface* SculptMaterial = ToolSetupUtil::GetSculptMaterial1(GetToolManager());
+		if (SculptMaterial != nullptr)
 		{
-			DynamicMeshComponent->SetMaterial(0, Material);
+			DynamicMeshComponent->SetOverrideRenderMaterial(SculptMaterial);
 		}
 		DynamicMeshComponent->bCastDynamicShadow = false;
 	}
diff --git a/Engine/Plugins/Experimental/MeshModelingToolset/Source/MeshModelingTools/Private/EditMeshPolygonsTool.cpp b/Engine/Plugins/Experimental/MeshModelingToolset/Source/MeshModelingTools/Private/EditMeshPolygonsTool.cpp
--- a/Engine/Plugins/Experimental/MeshModelingToolset/Source/MeshModelingTools/Private/EditMeshPolygonsTool.cpp
+++ b/Engine/Plugins/Experimental/MeshModelingToolset/Source/MeshModelingTools/Private/EditMeshPolygonsTool.cpp
@@ -626,11 +626,12 @@ void UEditMeshPolygonsTool::Setup()
 	DynamicMeshComponent->RegisterComponent();
 	DynamicMeshComponent->SetWorldTransform(ComponentTarget->GetWorldTransform());
 
-	// copy material if there is one
-	auto Material = ComponentTarget->GetMaterial(0);
-	if (Material != nullptr)
+	// set materials
+	FComponentMaterialSet MaterialSet;
+	ComponentTarget->GetMaterialSet(MaterialSet);
+	for (int k = 0; k < MaterialSet.Materials.Num(); ++k)
 	{
-		DynamicMeshComponent->SetMaterial(0, Material);
+		DynamicMeshComponent->SetMaterial(k, MaterialSet.Materials[k]);
 	}
 
 	// dynamic mesh configuration settings
diff --git a/Engine/Plugins/Experimental/MeshModelingToolset/Source/MeshModelingTools/Private/MeshInspectorTool.cpp b/Engine/Plugins/Experimental/MeshModelingToolset/Source/MeshModelingTools/Private/MeshInspectorTool.cpp
--- a/Engine/Plugins/Experimental/MeshModelingToolset/Source/MeshModelingTools/Private/MeshInspectorTool.cpp
+++ b/Engine/Plugins/Experimental/MeshModelingToolset/Source/MeshModelingTools/Private/MeshInspectorTool.cpp
@@ -90,21 +90,18 @@ void UMeshInspectorTool::Setup()
 {
 	UInteractiveTool::Setup();
 
-	// create dynamic mesh component to use for live preview
-	DynamicMeshComponent = NewObject<USimpleDynamicMeshComponent>(ComponentTarget->GetOwnerActor(), "DynamicMesh");
-	DynamicMeshComponent->SetupAttachment(ComponentTarget->GetOwnerActor()->GetRootComponent());
-	DynamicMeshComponent->RegisterComponent();
-	DynamicMeshComponent->SetWorldTransform(ComponentTarget->GetWorldTransform());
-
-	// copy material if there is one
-	DefaultMaterial = ComponentTarget->GetMaterial(0);
-	if (DefaultMaterial != nullptr)
-	{
-		DynamicMeshComponent->SetMaterial(0, DefaultMaterial);
-	}
+	PreviewMesh = NewObject<UPreviewMesh>(this);
+	PreviewMesh->bBuildSpatialDataStructure = false;
+	PreviewMesh->CreateInWorld(ComponentTarget->GetOwnerActor()->GetWorld(), FTransform::Identity);
+	PreviewMesh->SetTransform(ComponentTarget->GetWorldTransform());
+
+	FComponentMaterialSet MaterialSet;
+	ComponentTarget->GetMaterialSet(MaterialSet);
+	PreviewMesh->SetMaterials(MaterialSet.Materials);
+	DefaultMaterial = PreviewMesh->GetMaterial(0);
 
-	DynamicMeshComponent->TangentsType = EDynamicMeshTangentCalcType::ExternallyCalculated;
-	DynamicMeshComponent->InitializeMesh(ComponentTarget->GetMesh());
+	PreviewMesh->SetTangentsMode(EDynamicMeshTangentCalcType::ExternallyCalculated);
+	PreviewMesh->InitializeMesh(ComponentTarget->GetMesh());
 
 	Precompute();
 
@@ -120,28 +117,27 @@ void UMeshInspectorTool::Setup()
 	MaterialSettings->Setup();
 	AddToolPropertySource(MaterialSettings);
 
-
-	DynamicMeshComponent->bExplicitShowWireframe = Settings->bWireframe;
+	PreviewMesh->EnableWireframe(Settings->bWireframe);
 
 	UMeshStatisticsProperties* Statistics = NewObject<UMeshStatisticsProperties>(this);
-	Statistics->Update(*DynamicMeshComponent->GetMesh());
+	Statistics->Update(*PreviewMesh->GetPreviewDynamicMesh());
 	AddToolPropertySource(Statistics);
 
 	UMeshAnalysisProperties* MeshAnalysis = NewObject<UMeshAnalysisProperties>(this);
-	MeshAnalysis->Update(*DynamicMeshComponent->GetMesh(), ComponentTarget->GetWorldTransform());
+	MeshAnalysis->Update(*PreviewMesh->GetPreviewDynamicMesh(), ComponentTarget->GetWorldTransform());
 	AddToolPropertySource(MeshAnalysis);
 }
 
 
 void UMeshInspectorTool::Shutdown(EToolShutdownType ShutdownType)
 {
-	if (DynamicMeshComponent != nullptr)
-	{
-		ComponentTarget->SetOwnerVisibility(true);
+	ComponentTarget->SetOwnerVisibility(true);
 
-		DynamicMeshComponent->UnregisterComponent();
-		DynamicMeshComponent->DestroyComponent();
-		DynamicMeshComponent = nullptr;
+	if (PreviewMesh != nullptr)
+	{
+		PreviewMesh->SetVisible(false);
+		PreviewMesh->Disconnect();
+		PreviewMesh = nullptr;
 	}
 
 	Settings->SaveProperties(this);
@@ -154,10 +150,10 @@ void UMeshInspectorTool::Precompute()
 	UVSeamEdges.Reset();
 	NormalSeamEdges.Reset();
 
-	FDynamicMesh3* TargetMesh = DynamicMeshComponent->GetMesh();
-	FDynamicMeshUVOverlay* UVOverlay =
+	const FDynamicMesh3* TargetMesh = PreviewMesh->GetPreviewDynamicMesh();
+	const FDynamicMeshUVOverlay* UVOverlay =
 		TargetMesh->HasAttributes() ? TargetMesh->Attributes()->PrimaryUV() : nullptr;
-	FDynamicMeshNormalOverlay* NormalOverlay =
+	const FDynamicMeshNormalOverlay* NormalOverlay =
 		TargetMesh->HasAttributes() ? TargetMesh->Attributes()->PrimaryNormals() : nullptr;
 
 	for (int eid : TargetMesh->EdgeIndicesItr())
@@ -201,7 +197,7 @@ void UMeshInspectorTool::Render(IToolsContextRenderAPI* RenderAPI)
 	FColor BinormalColor(15, 240, 15);
 	float TangentThickness = LineWidthMultiplier * 2.0f;
 
-	FDynamicMesh3* TargetMesh = DynamicMeshComponent->GetMesh();
+	const FDynamicMesh3* TargetMesh = PreviewMesh->GetPreviewDynamicMesh();
 	FVector3d A, B;
 
 	if (Settings->bBoundaryEdges)
@@ -262,9 +258,9 @@ void UMeshInspectorTool::Render(IToolsContextRenderAPI* RenderAPI)
 		}
 	}
 
-	if (Settings->bTangentVectors)
+	if (Settings->bTangentVectors && PreviewMesh->GetTangents() != nullptr)
 	{
-		const FMeshTangentsf* Tangents = DynamicMeshComponent->GetTangents();
+		const FMeshTangentsf* Tangents = PreviewMesh->GetTangents();
 		for (int TID : TargetMesh->TriangleIndicesItr())
 		{
 			FVector3d TriV[3];
@@ -289,12 +285,20 @@ void UMeshInspectorTool::Render(IToolsContextRenderAPI* RenderAPI)
 void UMeshInspectorTool::OnPropertyModified(UObject* PropertySet, UProperty* Property)
 {
 	GetToolManager()->PostInvalidation();
-	DynamicMeshComponent->bExplicitShowWireframe = Settings->bWireframe;
+	PreviewMesh->EnableWireframe(Settings->bWireframe);
+	
 	MaterialSettings->UpdateMaterials();
-	MaterialSettings->SetMaterialIfChanged(DefaultMaterial, DynamicMeshComponent->GetMaterial(0),
+	MaterialSettings->SetMaterialIfChanged(DefaultMaterial, PreviewMesh->GetActiveMaterial(0),
 		[this](UMaterialInterface* Material)
 	{
-		DynamicMeshComponent->SetMaterial(0, Material);
+		if (Material == DefaultMaterial)
+		{
+			PreviewMesh->ClearOverrideRenderMaterial();
+		}
+		else
+		{
+			PreviewMesh->SetOverrideRenderMaterial(Material);
+		}
 	});
 }
 
diff --git a/Engine/Plugins/Experimental/MeshModelingToolset/Source/MeshModelingTools/Private/MeshSelectionTool.cpp b/Engine/Plugins/Experimental/MeshModelingToolset/Source/MeshModelingTools/Private/MeshSelectionTool.cpp
--- a/Engine/Plugins/Experimental/MeshModelingToolset/Source/MeshModelingTools/Private/MeshSelectionTool.cpp
+++ b/Engine/Plugins/Experimental/MeshModelingToolset/Source/MeshModelingTools/Private/MeshSelectionTool.cpp
@@ -104,12 +104,12 @@ void UMeshSelectionTool::Setup()
 	// enable wireframe on component
 	PreviewMesh->EnableWireframe(true);
 
-	// set vertex color material on base component so we can see selection
 	UMaterialInterface* SelectionMat = ToolSetupUtil::GetSelectionMaterial(GetToolManager());
 	if (SelectionMat != nullptr)
 	{
-		PreviewMesh->SetMaterial(SelectionMat);
+		PreviewMesh->SetOverrideRenderMaterial(SelectionMat);
 	}
+
 	PreviewMesh->GetRootComponent()->bCastDynamicShadow = false;
 
 	const FDynamicMesh3* Mesh = PreviewMesh->GetPreviewDynamicMesh();
@@ -562,6 +562,7 @@ void UMeshSelectionTool::Render(IToolsContextRenderAPI* RenderAPI)
 
 
 void UMeshSelectionTool::Tick(float DeltaTime)
+
 {
 	UDynamicMeshBrushTool::Tick(DeltaTime);
 
@@ -675,6 +676,10 @@ void UMeshSelectionTool::ApplyAction(EMeshSelectionToolActions ActionType)
 		case EMeshSelectionToolActions::CreateGroup:
 			AssignNewGroupToSelectedTriangles();
 			break;
+
+		case EMeshSelectionToolActions::AssignMaterial:
+			AssignMaterialToSelectedTriangles();
+			break;
 	}
 }
 
@@ -1100,7 +1105,55 @@ void UMeshSelectionTool::AssignNewGroupToSelectedTriangles()
 
 	OnExternalSelectionChange();
 	bHaveModifiedMesh = true;
+}
+
+
+
+
+void UMeshSelectionTool::AssignMaterialToSelectedTriangles()
+{
+	check(SelectionType == EMeshSelectionElementType::Face);
+	TArray<int32> SelectedFaces = Selection->GetElements(EMeshSelectionElementType::Face);
+	if (SelectedFaces.Num() == 0 || SelectionProps->Material == nullptr)
+	{
+		return;
+	}
+
+	TUniquePtr<FToolCommandChangeSequence> ChangeSeq = MakeUnique<FToolCommandChangeSequence>();
+
+	// clear current selection
+	BeginChange(false);
+	for (int tid : SelectedFaces)
+	{
+		ActiveSelectionChange->Add(tid);
+	}
+	Selection->RemoveIndices(EMeshSelectionElementType::Face, SelectedFaces);
+	TUniquePtr<FMeshSelectionChange> SelectionChange = EndChange();
+	ChangeSeq->AppendChange(Selection, MoveTemp(SelectionChange));
+
+	// assign new groups to triangles
+	// note: using an FMeshChange is kind of overkill here
+	TUniquePtr<FMeshChange> MeshChange = PreviewMesh->TrackedEditMesh(
+		[&SelectedFaces](FDynamicMesh3& Mesh, FDynamicMeshChangeTracker& ChangeTracker)
+	{
+		int32 SetID = 1;
+		if (Mesh.Attributes() && Mesh.Attributes()->HasMaterialID())
+		{
+			FDynamicMeshMaterialAttribute* MaterialIDAttrib = Mesh.Attributes()->GetMaterialID();
+			for (int tid : SelectedFaces)
+			{
+				ChangeTracker.SaveTriangle(tid, true);
+				MaterialIDAttrib->SetValue(tid, &SetID);
+			}
+		}
+	});
+	ChangeSeq->AppendChange(PreviewMesh, MoveTemp(MeshChange));
+
+	// emit combined change sequence
+	GetToolManager()->EmitObjectChange(this, MoveTemp(ChangeSeq), LOCTEXT("MeshSelectionToolAssignMaterial", "Assign Material"));
 
+	OnExternalSelectionChange();
+	bHaveModifiedMesh = true;
 }
 
 
diff --git a/Engine/Plugins/Experimental/MeshModelingToolset/Source/MeshModelingTools/Public/MeshInspectorTool.h b/Engine/Plugins/Experimental/MeshModelingToolset/Source/MeshModelingTools/Public/MeshInspectorTool.h
--- a/Engine/Plugins/Experimental/MeshModelingToolset/Source/MeshModelingTools/Public/MeshInspectorTool.h
+++ b/Engine/Plugins/Experimental/MeshModelingToolset/Source/MeshModelingTools/Public/MeshInspectorTool.h
@@ -9,12 +9,12 @@
 #include "DynamicMesh3.h"
 #include "DynamicMeshAABBTree3.h"
 #include "Properties/MeshMaterialProperties.h"
+#include "PreviewMesh.h"
 #include "MeshInspectorTool.generated.h"
 
 
 // predeclarations
 struct FMeshDescription;
-class USimpleDynamicMeshComponent;
 class UMaterialInstanceDynamic;
 
 /**
@@ -127,7 +127,7 @@ protected:
 
 protected:
 	UPROPERTY()
-	USimpleDynamicMeshComponent* DynamicMeshComponent;
+	UPreviewMesh* PreviewMesh;
 
 	UPROPERTY()
 	UMaterialInterface* DefaultMaterial = nullptr;
diff --git a/Engine/Plugins/Experimental/MeshModelingToolset/Source/MeshModelingTools/Public/MeshSelectionTool.h b/Engine/Plugins/Experimental/MeshModelingToolset/Source/MeshModelingTools/Public/MeshSelectionTool.h
--- a/Engine/Plugins/Experimental/MeshModelingToolset/Source/MeshModelingTools/Public/MeshSelectionTool.h
+++ b/Engine/Plugins/Experimental/MeshModelingToolset/Source/MeshModelingTools/Public/MeshSelectionTool.h
@@ -42,7 +42,8 @@ enum class EMeshSelectionToolActions
 	DisconnectSelected,
 	SeparateSelected,
 	FlipSelected,
-	CreateGroup
+	CreateGroup,
+	AssignMaterial
 };
 
 
@@ -140,6 +141,12 @@ public:
 	{
 		PostAction(EMeshSelectionToolActions::CreateGroup);
 	}
+
+	//UFUNCTION(CallInEditor, Category = MeshEdits, meta = (DisplayName = "Assign Material", DisplayPriority = 5))
+	//void AssignMaterial()
+	//{
+	//	PostAction(EMeshSelectionToolActions::AssignMaterial);
+	//}
 };
 
 
@@ -201,6 +208,10 @@ public:
 	UPROPERTY(EditAnywhere, Category = Selection)
 	bool bShowPolygroups = false;
 
+	UPROPERTY(EditAnywhere, Category = Selection)
+	UMaterialInterface* Material = nullptr;
+
+
 	virtual void SaveProperties(UInteractiveTool* SaveFromTool) override;
 	virtual void RestoreProperties(UInteractiveTool* RestoreToTool) override;
 };
@@ -321,6 +332,7 @@ protected:
 	void SeparateSelectedTriangles(); // separates out selected triangles to a new mesh, removing them from the current mesh
 	void FlipSelectedTriangles();
 	void AssignNewGroupToSelectedTriangles();
+	void AssignMaterialToSelectedTriangles();
 
 	bool bHaveModifiedMesh = false;
 };
