ModelingTools: update remaining tools for multiple material support.
- Removed UExistingMeshMaterialProperties::SetMaterialIfChanged, the logic was too convoluted now that we don't actually set new materials on the PreviewMesh. Now just forcing a material update whenever any material setting changes.
- UMeshOpPreviewWithBackgroundCompute now takes material set, added ::ConfigureMaterials() that takes array  (previous API remains)
- UMeshOpPreviewWithBackgroundCompute now has additional OverrideMaterial, if set then this is forwarded to UPreviewMesh as override material
#rnx

diff --git a/Engine/Plugins/Experimental/MeshModelingToolset/Source/MeshModelingTools/Private/ConvertToPolygonsTool.cpp b/Engine/Plugins/Experimental/MeshModelingToolset/Source/MeshModelingTools/Private/ConvertToPolygonsTool.cpp
--- a/Engine/Plugins/Experimental/MeshModelingToolset/Source/MeshModelingTools/Private/ConvertToPolygonsTool.cpp
+++ b/Engine/Plugins/Experimental/MeshModelingToolset/Source/MeshModelingTools/Private/ConvertToPolygonsTool.cpp
@@ -70,10 +70,9 @@ void UConvertToPolygonsTool::Setup()
 	PreviewMesh->SetVisible(false);
 	PreviewMesh->SetTransform(ComponentTarget->GetWorldTransform());
 
-	if (ComponentTarget->GetMaterial(0) != nullptr)
-	{
-		PreviewMesh->SetMaterial(ComponentTarget->GetMaterial(0));
-	}
+	FComponentMaterialSet MaterialSet;
+	ComponentTarget->GetMaterialSet(MaterialSet);
+	PreviewMesh->SetMaterials(MaterialSet.Materials);
 
 	UpdatePolygons();
 }
diff --git a/Engine/Plugins/Experimental/MeshModelingToolset/Source/MeshModelingTools/Private/DisplaceMeshTool.cpp b/Engine/Plugins/Experimental/MeshModelingToolset/Source/MeshModelingTools/Private/DisplaceMeshTool.cpp
--- a/Engine/Plugins/Experimental/MeshModelingToolset/Source/MeshModelingTools/Private/DisplaceMeshTool.cpp
+++ b/Engine/Plugins/Experimental/MeshModelingToolset/Source/MeshModelingTools/Private/DisplaceMeshTool.cpp
@@ -442,11 +442,12 @@ void UDisplaceMeshTool::Setup()
 	DynamicMeshComponent->RegisterComponent();
 	DynamicMeshComponent->SetWorldTransform(ComponentTarget->GetWorldTransform());
 
-	// copy material if there is one
-	auto Material = ComponentTarget->GetMaterial(0);
-	if (Material != nullptr)
+	// transfer materials
+	FComponentMaterialSet MaterialSet;
+	ComponentTarget->GetMaterialSet(MaterialSet);
+	for (int k = 0; k < MaterialSet.Materials.Num(); ++k)
 	{
-		DynamicMeshComponent->SetMaterial(0, Material);
+		DynamicMeshComponent->SetMaterial(k, MaterialSet.Materials[k]);
 	}
 
 	DynamicMeshComponent->InitializeMesh(ComponentTarget->GetMesh());
@@ -573,13 +574,13 @@ void UDisplaceMeshTool::StartComputation()
 		SubdivideTask = new FAsyncTaskExecuterWithAbort<TModelingOpTask<FDynamicMeshOperator>>(Subdivider->MakeNewOperator());
 		SubdivideTask->StartBackgroundTask();
 		bNeedsSubdivided = false;
-		DynamicMeshComponent->SetMaterial(0, ToolSetupUtil::GetDefaultWorkingMaterial(GetToolManager()));
+		DynamicMeshComponent->SetOverrideRenderMaterial(ToolSetupUtil::GetDefaultWorkingMaterial(GetToolManager()));
 	}
 	if (bNeedsDisplaced && DisplaceTask)
 	{
 		DisplaceTask->CancelAndDelete();
 		DisplaceTask = nullptr;
-		DynamicMeshComponent->SetMaterial(0, ToolSetupUtil::GetDefaultWorkingMaterial(GetToolManager()));
+		DynamicMeshComponent->SetOverrideRenderMaterial(ToolSetupUtil::GetDefaultWorkingMaterial(GetToolManager()));
 	}
 	AdvanceComputation();
 }
@@ -603,7 +604,7 @@ void UDisplaceMeshTool::AdvanceComputation()
 		TUniquePtr<FDynamicMesh3> DisplacedMesh = DisplaceTask->GetTask().ExtractOperator()->ExtractResult();
 		delete DisplaceTask;
 		DisplaceTask = nullptr;
-		DynamicMeshComponent->SetMaterial(0,ToolSetupUtil::GetDefaultMaterial(GetToolManager(), ComponentTarget->GetMaterial(0)));
+		DynamicMeshComponent->ClearOverrideRenderMaterial();
 		DynamicMeshComponent->GetMesh()->Copy(*DisplacedMesh);
 		DynamicMeshComponent->NotifyMeshUpdated();
 		GetToolManager()->PostInvalidation();
diff --git a/Engine/Plugins/Experimental/MeshModelingToolset/Source/MeshModelingTools/Private/EditNormalsTool.cpp b/Engine/Plugins/Experimental/MeshModelingToolset/Source/MeshModelingTools/Private/EditNormalsTool.cpp
--- a/Engine/Plugins/Experimental/MeshModelingToolset/Source/MeshModelingTools/Private/EditNormalsTool.cpp
+++ b/Engine/Plugins/Experimental/MeshModelingToolset/Source/MeshModelingTools/Private/EditNormalsTool.cpp
@@ -144,10 +144,13 @@ void UEditNormalsTool::UpdateNumPreviews()
 
 			UMeshOpPreviewWithBackgroundCompute* Preview = Previews.Add_GetRef(NewObject<UMeshOpPreviewWithBackgroundCompute>(OpFactory, "Preview"));
 			Preview->Setup(this->TargetWorld, OpFactory);
-			Preview->ConfigureMaterials(
-				ToolSetupUtil::GetDefaultMaterial(GetToolManager(), ComponentTargets[PreviewIdx]->GetMaterial(0)),
+
+			FComponentMaterialSet MaterialSet;
+			ComponentTargets[PreviewIdx]->GetMaterialSet(MaterialSet);
+			Preview->ConfigureMaterials(MaterialSet.Materials,
 				ToolSetupUtil::GetDefaultWorkingMaterial(GetToolManager())
 			);
+
 			Preview->SetVisibility(true);
 		}
 	}
diff --git a/Engine/Plugins/Experimental/MeshModelingToolset/Source/MeshModelingTools/Private/MeshInspectorTool.cpp b/Engine/Plugins/Experimental/MeshModelingToolset/Source/MeshModelingTools/Private/MeshInspectorTool.cpp
--- a/Engine/Plugins/Experimental/MeshModelingToolset/Source/MeshModelingTools/Private/MeshInspectorTool.cpp
+++ b/Engine/Plugins/Experimental/MeshModelingToolset/Source/MeshModelingTools/Private/MeshInspectorTool.cpp
@@ -288,18 +288,15 @@ void UMeshInspectorTool::OnPropertyModified(UObject* PropertySet, UProperty* Pro
 	PreviewMesh->EnableWireframe(Settings->bWireframe);
 	
 	MaterialSettings->UpdateMaterials();
-	MaterialSettings->SetMaterialIfChanged(DefaultMaterial, PreviewMesh->GetActiveMaterial(0),
-		[this](UMaterialInterface* Material)
+	UMaterialInterface* OverrideMaterial = MaterialSettings->GetActiveOverrideMaterial();
+	if (OverrideMaterial == nullptr)
 	{
-		if (Material == DefaultMaterial)
-		{
-			PreviewMesh->ClearOverrideRenderMaterial();
-		}
-		else
-		{
-			PreviewMesh->SetOverrideRenderMaterial(Material);
-		}
-	});
+		PreviewMesh->ClearOverrideRenderMaterial();
+	}
+	else
+	{
+		PreviewMesh->SetOverrideRenderMaterial(OverrideMaterial);
+	}
 }
 
 
diff --git a/Engine/Plugins/Experimental/MeshModelingToolset/Source/MeshModelingTools/Private/MeshSpaceDeformerTool.cpp b/Engine/Plugins/Experimental/MeshModelingToolset/Source/MeshModelingTools/Private/MeshSpaceDeformerTool.cpp
--- a/Engine/Plugins/Experimental/MeshModelingToolset/Source/MeshModelingTools/Private/MeshSpaceDeformerTool.cpp
+++ b/Engine/Plugins/Experimental/MeshModelingToolset/Source/MeshModelingTools/Private/MeshSpaceDeformerTool.cpp
@@ -262,10 +262,12 @@ void UMeshSpaceDeformerTool::Setup()
 		Preview->PreviewMesh->UpdatePreview(OriginalDynamicMesh.Get());
 		Preview->PreviewMesh->SetTransform(ComponentTarget->GetWorldTransform());
 
-		Preview->ConfigureMaterials(
-			ToolSetupUtil::GetDefaultMaterial(GetToolManager(), ComponentTarget->GetMaterial(0)),
+		FComponentMaterialSet MaterialSet;
+		ComponentTarget->GetMaterialSet(MaterialSet);
+		Preview->ConfigureMaterials(MaterialSet.Materials,
 			ToolSetupUtil::GetDefaultWorkingMaterial(GetToolManager())
 		);
+
 		// show the preview mesh
 		Preview->SetVisibility(true);
 
diff --git a/Engine/Plugins/Experimental/MeshModelingToolset/Source/MeshModelingTools/Private/PlaneCutTool.cpp b/Engine/Plugins/Experimental/MeshModelingToolset/Source/MeshModelingTools/Private/PlaneCutTool.cpp
--- a/Engine/Plugins/Experimental/MeshModelingToolset/Source/MeshModelingTools/Private/PlaneCutTool.cpp
+++ b/Engine/Plugins/Experimental/MeshModelingToolset/Source/MeshModelingTools/Private/PlaneCutTool.cpp
@@ -182,10 +182,13 @@ void UPlaneCutTool::UpdateNumPreviews()
 			CutSide->ComponentIndex = SrcIdx;
 			UMeshOpPreviewWithBackgroundCompute* Preview = Previews.Add_GetRef(NewObject<UMeshOpPreviewWithBackgroundCompute>(CutSide, "Preview"));
 			Preview->Setup(this->TargetWorld, CutSide);
-			Preview->ConfigureMaterials(
-				ToolSetupUtil::GetDefaultMaterial(GetToolManager(), ComponentTargets[SrcIdx]->GetMaterial(0)),
+
+			FComponentMaterialSet MaterialSet;
+			ComponentTargets[SrcIdx]->GetMaterialSet(MaterialSet);
+			Preview->ConfigureMaterials(MaterialSet.Materials,
 				ToolSetupUtil::GetDefaultWorkingMaterial(GetToolManager())
 			);
+
 			// set initial preview to un-processed mesh, so stuff doesn't just disappear if the first cut takes a while
 			Preview->PreviewMesh->UpdatePreview(OriginalDynamicMeshes[SrcIdx].Get());
 			Preview->PreviewMesh->SetTransform(ComponentTargets[SrcIdx]->GetWorldTransform());
diff --git a/Engine/Plugins/Experimental/MeshModelingToolset/Source/MeshModelingTools/Private/Properties/MeshMaterialProperties.cpp b/Engine/Plugins/Experimental/MeshModelingToolset/Source/MeshModelingTools/Private/Properties/MeshMaterialProperties.cpp
--- a/Engine/Plugins/Experimental/MeshModelingToolset/Source/MeshModelingTools/Private/Properties/MeshMaterialProperties.cpp
+++ b/Engine/Plugins/Experimental/MeshModelingToolset/Source/MeshModelingTools/Private/Properties/MeshMaterialProperties.cpp
@@ -59,21 +59,18 @@ void UExistingMeshMaterialProperties::UpdateMaterials()
 	}
 }
 
-void UExistingMeshMaterialProperties::SetMaterialIfChanged(UMaterialInterface* OriginalMaterial, UMaterialInterface* CurrentMaterial, TFunctionRef<void(UMaterialInterface* Material)> SetMaterialFn)
+
+UMaterialInterface* UExistingMeshMaterialProperties::GetActiveOverrideMaterial() const
 {
-	UMaterialInterface *Material = OriginalMaterial;
 	if (MaterialMode == ESetMeshMaterialMode::Checkerboard && CheckerMaterial != nullptr)
 	{
-		Material = CheckerMaterial;
-	}
-	else if (MaterialMode == ESetMeshMaterialMode::Override && OverrideMaterial != nullptr)
-	{
-		Material = OverrideMaterial;
+		return CheckerMaterial;
 	}
-	if (CurrentMaterial != Material)
+	if (MaterialMode == ESetMeshMaterialMode::Override && OverrideMaterial != nullptr)
 	{
-		SetMaterialFn(Material);
+		return OverrideMaterial;
 	}
+	return nullptr;
 }
 
 
@@ -81,7 +78,6 @@ void UExistingMeshMaterialProperties::SetMaterialIfChanged(UMaterialInterface* O
 
 
 
-
 void UMeshEditingViewProperties::SaveProperties(UInteractiveTool* SaveFromTool)
 {
 	UMeshEditingViewProperties* PropertyCache = GetPropertyCache<UMeshEditingViewProperties>();
diff --git a/Engine/Plugins/Experimental/MeshModelingToolset/Source/MeshModelingTools/Private/RemeshMeshTool.cpp b/Engine/Plugins/Experimental/MeshModelingToolset/Source/MeshModelingTools/Private/RemeshMeshTool.cpp
--- a/Engine/Plugins/Experimental/MeshModelingToolset/Source/MeshModelingTools/Private/RemeshMeshTool.cpp
+++ b/Engine/Plugins/Experimental/MeshModelingToolset/Source/MeshModelingTools/Private/RemeshMeshTool.cpp
@@ -100,8 +100,9 @@ void URemeshMeshTool::Setup()
 	{
 		MeshStatisticsProperties->Update(*Compute->PreviewMesh->GetPreviewDynamicMesh());
 	});
-	Preview->ConfigureMaterials(
-		ToolSetupUtil::GetDefaultMaterial(GetToolManager(), ComponentTarget->GetMaterial(0)),
+	FComponentMaterialSet MaterialSet;
+	ComponentTarget->GetMaterialSet(MaterialSet);
+	Preview->ConfigureMaterials( MaterialSet.Materials,
 		ToolSetupUtil::GetDefaultWorkingMaterial(GetToolManager())
 	);
 	Preview->PreviewMesh->EnableWireframe(true);
diff --git a/Engine/Plugins/Experimental/MeshModelingToolset/Source/MeshModelingTools/Private/RemoveOccludedTrianglesTool.cpp b/Engine/Plugins/Experimental/MeshModelingToolset/Source/MeshModelingTools/Private/RemoveOccludedTrianglesTool.cpp
--- a/Engine/Plugins/Experimental/MeshModelingToolset/Source/MeshModelingTools/Private/RemoveOccludedTrianglesTool.cpp
+++ b/Engine/Plugins/Experimental/MeshModelingToolset/Source/MeshModelingTools/Private/RemoveOccludedTrianglesTool.cpp
@@ -147,10 +147,13 @@ void URemoveOccludedTrianglesTool::SetupPreviews()
 
 		UMeshOpPreviewWithBackgroundCompute* Preview = Previews.Add_GetRef(NewObject<UMeshOpPreviewWithBackgroundCompute>(OpFactory, "Preview"));
 		Preview->Setup(this->TargetWorld, OpFactory);
-		Preview->ConfigureMaterials(
-			ToolSetupUtil::GetDefaultMaterial(GetToolManager(), ComponentTargets[PreviewIdx]->GetMaterial(0)),
+
+		FComponentMaterialSet MaterialSet;
+		ComponentTargets[PreviewIdx]->GetMaterialSet(MaterialSet);
+		Preview->ConfigureMaterials(MaterialSet.Materials,
 			ToolSetupUtil::GetDefaultWorkingMaterial(GetToolManager())
 		);
+
 		Preview->PreviewMesh->SetTransform(ComponentTargets[PreviewIdx]->GetWorldTransform());
 		Preview->PreviewMesh->UpdatePreview(OriginalDynamicMeshes[PreviewIdx].Get());
 		Preview->SetVisibility(true);
diff --git a/Engine/Plugins/Experimental/MeshModelingToolset/Source/MeshModelingTools/Private/SmoothMeshTool.cpp b/Engine/Plugins/Experimental/MeshModelingToolset/Source/MeshModelingTools/Private/SmoothMeshTool.cpp
--- a/Engine/Plugins/Experimental/MeshModelingToolset/Source/MeshModelingTools/Private/SmoothMeshTool.cpp
+++ b/Engine/Plugins/Experimental/MeshModelingToolset/Source/MeshModelingTools/Private/SmoothMeshTool.cpp
@@ -85,15 +85,14 @@ void USmoothMeshTool::Setup()
 		Preview = NewObject<UMeshOpPreviewWithBackgroundCompute>(this, "Preview");
 		Preview->Setup(this->TargetWorld, this); // Adds the actual functional tool in the Preview object
 
-		Preview->PreviewMesh->UpdatePreview(&SrcDynamicMesh);
-		Preview->PreviewMesh->SetTransform(ComponentTarget->GetWorldTransform());
-
-	auto Material = ComponentTarget->GetMaterial(0);
-
-		Preview->ConfigureMaterials(
-			ToolSetupUtil::GetDefaultMaterial(GetToolManager(), Material),
+		FComponentMaterialSet MaterialSet;
+		ComponentTarget->GetMaterialSet(MaterialSet);
+		Preview->ConfigureMaterials(MaterialSet.Materials,
 			ToolSetupUtil::GetDefaultWorkingMaterial(GetToolManager())
 		);
+
+		Preview->PreviewMesh->SetTransform(ComponentTarget->GetWorldTransform());
+		Preview->PreviewMesh->UpdatePreview(&SrcDynamicMesh);
 	}
 
 	// show the preview mesh
diff --git a/Engine/Plugins/Experimental/MeshModelingToolset/Source/MeshModelingTools/Private/UVProjectionTool.cpp b/Engine/Plugins/Experimental/MeshModelingToolset/Source/MeshModelingTools/Private/UVProjectionTool.cpp
--- a/Engine/Plugins/Experimental/MeshModelingToolset/Source/MeshModelingTools/Private/UVProjectionTool.cpp
+++ b/Engine/Plugins/Experimental/MeshModelingToolset/Source/MeshModelingTools/Private/UVProjectionTool.cpp
@@ -164,12 +164,14 @@ void UUVProjectionTool::UpdateNumPreviews()
 
 			UMeshOpPreviewWithBackgroundCompute* Preview = Previews.Add_GetRef(NewObject<UMeshOpPreviewWithBackgroundCompute>(OpFactory, "Preview"));
 			Preview->Setup(this->TargetWorld, OpFactory);
-			Preview->ConfigureMaterials(
-				ToolSetupUtil::GetDefaultMaterial(GetToolManager(), ComponentTargets[PreviewIdx]->GetMaterial(0)),
+
+			FComponentMaterialSet MaterialSet;
+			ComponentTargets[PreviewIdx]->GetMaterialSet(MaterialSet);
+			Preview->ConfigureMaterials(MaterialSet.Materials,
 				ToolSetupUtil::GetDefaultWorkingMaterial(GetToolManager())
 			);
-			Preview->SetVisibility(true);
 
+			Preview->SetVisibility(true);
 
 			UTransformProxy* TransformProxy = TransformProxies.Add_GetRef(NewObject<UTransformProxy>(this));
 			TransformProxy->SetTransform(LocalXF * ComponentTargets[PreviewIdx]->GetWorldTransform());
@@ -290,10 +292,7 @@ void UUVProjectionTool::OnPropertyModified(UObject* PropertySet, UProperty* Prop
 	for (int PreviewIdx = 0; PreviewIdx < Previews.Num(); PreviewIdx++)
 	{
 		UMeshOpPreviewWithBackgroundCompute* Preview = Previews[PreviewIdx];
-		MaterialSettings->SetMaterialIfChanged(ComponentTargets[PreviewIdx]->GetMaterial(0), Preview->StandardMaterial, [&Preview, this](UMaterialInterface* Material)
-		{
-			Preview->ConfigureMaterials(ToolSetupUtil::GetDefaultMaterial(GetToolManager(), Material), Preview->WorkingMaterial);
-		});
+		Preview->OverrideMaterial = MaterialSettings->GetActiveOverrideMaterial();
 	}
 	
 	UpdateNumPreviews();
diff --git a/Engine/Plugins/Experimental/MeshModelingToolset/Source/MeshModelingTools/Private/WeldMeshEdgesTool.cpp b/Engine/Plugins/Experimental/MeshModelingToolset/Source/MeshModelingTools/Private/WeldMeshEdgesTool.cpp
--- a/Engine/Plugins/Experimental/MeshModelingToolset/Source/MeshModelingTools/Private/WeldMeshEdgesTool.cpp
+++ b/Engine/Plugins/Experimental/MeshModelingToolset/Source/MeshModelingTools/Private/WeldMeshEdgesTool.cpp
@@ -61,11 +61,12 @@ void UWeldMeshEdgesTool::Setup()
 	DynamicMeshComponent->RegisterComponent();
 	DynamicMeshComponent->SetWorldTransform(ComponentTarget->GetWorldTransform());
 
-	// copy material if there is one
-	auto Material = ComponentTarget->GetMaterial(0);
-	if (Material != nullptr)
+	// transfer materials
+	FComponentMaterialSet MaterialSet;
+	ComponentTarget->GetMaterialSet(MaterialSet);
+	for (int k = 0; k < MaterialSet.Materials.Num(); ++k)
 	{
-		DynamicMeshComponent->SetMaterial(0, Material);
+		DynamicMeshComponent->SetMaterial(k, MaterialSet.Materials[k]);
 	}
 
 	DynamicMeshComponent->InitializeMesh(ComponentTarget->GetMesh());
diff --git a/Engine/Plugins/Experimental/MeshModelingToolset/Source/MeshModelingTools/Public/Properties/MeshMaterialProperties.h b/Engine/Plugins/Experimental/MeshModelingToolset/Source/MeshModelingTools/Public/Properties/MeshMaterialProperties.h
--- a/Engine/Plugins/Experimental/MeshModelingToolset/Source/MeshModelingTools/Public/Properties/MeshMaterialProperties.h
+++ b/Engine/Plugins/Experimental/MeshModelingToolset/Source/MeshModelingTools/Public/Properties/MeshMaterialProperties.h
@@ -86,7 +86,7 @@ public:
 	void Setup();
 
 	void UpdateMaterials();
-	void SetMaterialIfChanged(UMaterialInterface* OriginalMaterial, UMaterialInterface* CurrentMaterial, TFunctionRef<void(UMaterialInterface* Material)> SetMaterialFn);
+	UMaterialInterface* GetActiveOverrideMaterial() const;
 };
 
 
diff --git a/Engine/Plugins/Experimental/MeshModelingToolset/Source/MeshModelingToolsEditorOnly/Private/ParameterizeMeshTool.cpp b/Engine/Plugins/Experimental/MeshModelingToolset/Source/MeshModelingToolsEditorOnly/Private/ParameterizeMeshTool.cpp
--- a/Engine/Plugins/Experimental/MeshModelingToolset/Source/MeshModelingToolsEditorOnly/Private/ParameterizeMeshTool.cpp
+++ b/Engine/Plugins/Experimental/MeshModelingToolset/Source/MeshModelingToolsEditorOnly/Private/ParameterizeMeshTool.cpp
@@ -100,13 +100,15 @@ void UParameterizeMeshTool::Setup()
 		Converter.bPrintDebugMessages = false;
 		Converter.Convert(InputMesh.Get(), Mesh);
 
+		FComponentMaterialSet MaterialSet;
+		ComponentTarget->GetMaterialSet(MaterialSet);
+		Preview->ConfigureMaterials(MaterialSet.Materials,
+			ToolSetupUtil::GetDefaultWorkingMaterial(GetToolManager())
+		);
+
 		Preview->PreviewMesh->UpdatePreview(&Mesh);
 		Preview->PreviewMesh->SetTransform(ComponentTarget->GetWorldTransform());
 	}
-	Preview->ConfigureMaterials(
-		ToolSetupUtil::GetDefaultMaterial(GetToolManager(), DefaultMaterial),
-		ToolSetupUtil::GetDefaultWorkingMaterial(GetToolManager())
-	);
 
 
 	MaterialSettings = NewObject<UExistingMeshMaterialProperties>(this);
@@ -123,10 +125,7 @@ void UParameterizeMeshTool::OnPropertyModified(UObject* PropertySet, UProperty*
 	if (PropertySet == MaterialSettings)
 	{
 		MaterialSettings->UpdateMaterials();
-		MaterialSettings->SetMaterialIfChanged(ComponentTarget->GetMaterial(0), Preview->StandardMaterial, [this](UMaterialInterface* Material)
-		{
-			Preview->ConfigureMaterials(ToolSetupUtil::GetDefaultMaterial(GetToolManager(), Material), Preview->WorkingMaterial);
-		});
+		Preview->OverrideMaterial = MaterialSettings->GetActiveOverrideMaterial();
 	}
 }
 
diff --git a/Engine/Plugins/Experimental/MeshModelingToolset/Source/MeshModelingToolsEditorOnly/Private/PolygonOnMeshTool.cpp b/Engine/Plugins/Experimental/MeshModelingToolset/Source/MeshModelingToolsEditorOnly/Private/PolygonOnMeshTool.cpp
--- a/Engine/Plugins/Experimental/MeshModelingToolset/Source/MeshModelingToolsEditorOnly/Private/PolygonOnMeshTool.cpp
+++ b/Engine/Plugins/Experimental/MeshModelingToolset/Source/MeshModelingToolsEditorOnly/Private/PolygonOnMeshTool.cpp
@@ -164,10 +164,13 @@ void UPolygonOnMeshTool::UpdateNumPreviews()
 			// TODO: give the OpFactory more info about what specifically it is doing differently vs the other previews
 			UMeshOpPreviewWithBackgroundCompute* Preview = Previews.Add_GetRef(NewObject<UMeshOpPreviewWithBackgroundCompute>(OpFactory, "Preview"));
 			Preview->Setup(this->TargetWorld, OpFactory);
-			Preview->ConfigureMaterials(
-				ToolSetupUtil::GetDefaultMaterial(GetToolManager(), ComponentTarget->GetMaterial(0)),
+
+			FComponentMaterialSet MaterialSet;
+			ComponentTarget->GetMaterialSet(MaterialSet);
+			Preview->ConfigureMaterials(MaterialSet.Materials,
 				ToolSetupUtil::GetDefaultWorkingMaterial(GetToolManager())
 			);
+
 			Preview->SetVisibility(true);
 		}
 	}
diff --git a/Engine/Plugins/Experimental/MeshModelingToolset/Source/MeshModelingToolsEditorOnly/Private/SimplifyMeshTool.cpp b/Engine/Plugins/Experimental/MeshModelingToolset/Source/MeshModelingToolsEditorOnly/Private/SimplifyMeshTool.cpp
--- a/Engine/Plugins/Experimental/MeshModelingToolset/Source/MeshModelingToolsEditorOnly/Private/SimplifyMeshTool.cpp
+++ b/Engine/Plugins/Experimental/MeshModelingToolset/Source/MeshModelingToolsEditorOnly/Private/SimplifyMeshTool.cpp
@@ -99,8 +99,9 @@ void USimplifyMeshTool::Setup()
 	ComponentTarget->SetOwnerVisibility(false);
 	Preview = NewObject<UMeshOpPreviewWithBackgroundCompute>(this, "Preview");
 	Preview->Setup(this->TargetWorld, this);
-	Preview->ConfigureMaterials(
-		ToolSetupUtil::GetDefaultMaterial(GetToolManager(), ComponentTarget->GetMaterial(0)),
+	FComponentMaterialSet MaterialSet;
+	ComponentTarget->GetMaterialSet(MaterialSet);
+	Preview->ConfigureMaterials( MaterialSet.Materials,
 		ToolSetupUtil::GetDefaultWorkingMaterial(GetToolManager())
 	);
 	Preview->PreviewMesh->EnableWireframe(true);
diff --git a/Engine/Plugins/Experimental/MeshModelingToolset/Source/MeshModelingToolsEditorOnly/Private/UVLayoutTool.cpp b/Engine/Plugins/Experimental/MeshModelingToolset/Source/MeshModelingToolsEditorOnly/Private/UVLayoutTool.cpp
--- a/Engine/Plugins/Experimental/MeshModelingToolset/Source/MeshModelingToolsEditorOnly/Private/UVLayoutTool.cpp
+++ b/Engine/Plugins/Experimental/MeshModelingToolset/Source/MeshModelingToolsEditorOnly/Private/UVLayoutTool.cpp
@@ -139,10 +139,13 @@ void UUVLayoutTool::UpdateNumPreviews()
 
 			UMeshOpPreviewWithBackgroundCompute* Preview = Previews.Add_GetRef(NewObject<UMeshOpPreviewWithBackgroundCompute>(OpFactory, "Preview"));
 			Preview->Setup(this->TargetWorld, OpFactory);
-			Preview->ConfigureMaterials(
-				ToolSetupUtil::GetDefaultMaterial(GetToolManager(), ComponentTargets[PreviewIdx]->GetMaterial(0)),
+
+			FComponentMaterialSet MaterialSet;
+			ComponentTargets[PreviewIdx]->GetMaterialSet(MaterialSet);
+			Preview->ConfigureMaterials(MaterialSet.Materials,
 				ToolSetupUtil::GetDefaultWorkingMaterial(GetToolManager())
 			);
+
 			Preview->SetVisibility(true);
 		}
 	}
@@ -219,10 +222,7 @@ void UUVLayoutTool::OnPropertyModified(UObject* PropertySet, UProperty* Property
 	for (int PreviewIdx = 0; PreviewIdx < Previews.Num(); PreviewIdx++)
 	{
 		UMeshOpPreviewWithBackgroundCompute* Preview = Previews[PreviewIdx];
-		MaterialSettings->SetMaterialIfChanged(ComponentTargets[PreviewIdx]->GetMaterial(0), Preview->StandardMaterial, [&Preview, this](UMaterialInterface* Material)
-		{
-			Preview->ConfigureMaterials(ToolSetupUtil::GetDefaultMaterial(GetToolManager(), Material), Preview->WorkingMaterial);
-		});
+		Preview->OverrideMaterial = MaterialSettings->GetActiveOverrideMaterial();
 	}
 	
 	UpdateNumPreviews();
diff --git a/Engine/Plugins/Experimental/MeshModelingToolset/Source/ModelingComponents/Private/MeshOpPreviewHelpers.cpp b/Engine/Plugins/Experimental/MeshModelingToolset/Source/ModelingComponents/Private/MeshOpPreviewHelpers.cpp
--- a/Engine/Plugins/Experimental/MeshModelingToolset/Source/ModelingComponents/Private/MeshOpPreviewHelpers.cpp
+++ b/Engine/Plugins/Experimental/MeshModelingToolset/Source/ModelingComponents/Private/MeshOpPreviewHelpers.cpp
@@ -52,11 +52,18 @@ void UMeshOpPreviewWithBackgroundCompute::Tick(float DeltaTime)
 
 	if (bResultValid || WorkingMaterial == nullptr || bIsLongDelay == false)
 	{
-		PreviewMesh->SetMaterial(StandardMaterial);
+		if (OverrideMaterial != nullptr)
+		{
+			PreviewMesh->SetOverrideRenderMaterial(OverrideMaterial);
+		}
+		else
+		{
+			PreviewMesh->ClearOverrideRenderMaterial();
+		}
 	}
 	else
 	{
-		PreviewMesh->SetMaterial(WorkingMaterial);
+		PreviewMesh->SetOverrideRenderMaterial(WorkingMaterial);
 	}
 }
 
@@ -87,12 +94,20 @@ void UMeshOpPreviewWithBackgroundCompute::InvalidateResult()
 
 void UMeshOpPreviewWithBackgroundCompute::ConfigureMaterials(UMaterialInterface* StandardMaterialIn, UMaterialInterface* WorkingMaterialIn)
 {
-	this->StandardMaterial = StandardMaterialIn;
+	TArray<UMaterialInterface*> Materials;
+	Materials.Add(StandardMaterialIn);
+	ConfigureMaterials(Materials, WorkingMaterialIn);
+}
+
+
+void UMeshOpPreviewWithBackgroundCompute::ConfigureMaterials(TArray<UMaterialInterface*> StandardMaterialsIn, UMaterialInterface* WorkingMaterialIn)
+{
+	this->StandardMaterials = StandardMaterialsIn;
 	this->WorkingMaterial = WorkingMaterialIn;
 
 	if (PreviewMesh != nullptr)
 	{
-		PreviewMesh->SetMaterial(this->StandardMaterial);
+		PreviewMesh->SetMaterials(this->StandardMaterials);
 	}
 }
 
diff --git a/Engine/Plugins/Experimental/MeshModelingToolset/Source/ModelingComponents/Public/MeshOpPreviewHelpers.h b/Engine/Plugins/Experimental/MeshModelingToolset/Source/ModelingComponents/Public/MeshOpPreviewHelpers.h
--- a/Engine/Plugins/Experimental/MeshModelingToolset/Source/ModelingComponents/Public/MeshOpPreviewHelpers.h
+++ b/Engine/Plugins/Experimental/MeshModelingToolset/Source/ModelingComponents/Public/MeshOpPreviewHelpers.h
@@ -103,6 +103,12 @@ public:
 	 */
 	void ConfigureMaterials(UMaterialInterface* StandardMaterial, UMaterialInterface* InProgressMaterial);
 
+	/**
+	 * Configure the Standard and In-Progress materials
+	 */
+	void ConfigureMaterials(TArray<UMaterialInterface*> StandardMaterials, UMaterialInterface* InProgressMaterial);
+
+
 	/**
 	 * Set the visibility of the Preview mesh
 	 */
@@ -125,9 +131,13 @@ public:
 	UPROPERTY()
 	UPreviewMesh* PreviewMesh = nullptr;
 
-	// material used on PreviewMesh
+	// input set of materials to assign to PreviewMesh
+	UPROPERTY()
+	TArray<UMaterialInterface*> StandardMaterials;
+
+	// override material to forward to PreviewMesh if set
 	UPROPERTY()
-	UMaterialInterface* StandardMaterial = nullptr;
+	UMaterialInterface* OverrideMaterial = nullptr;
 
 	// if non-null, this material is swapped in when a background compute is active
 	UPROPERTY()
