cis fix -- move MeshUtilitiesCommon dependency to editor only modeling operators' build.cs file

#rnx

diff --git a/Engine/Plugins/Experimental/MeshModelingToolset/Source/ModelingOperators/ModelingOperators.Build.cs b/Engine/Plugins/Experimental/MeshModelingToolset/Source/ModelingOperators/ModelingOperators.Build.cs
--- a/Engine/Plugins/Experimental/MeshModelingToolset/Source/ModelingOperators/ModelingOperators.Build.cs
+++ b/Engine/Plugins/Experimental/MeshModelingToolset/Source/ModelingOperators/ModelingOperators.Build.cs
@@ -41,7 +41,6 @@ public class ModelingOperators : ModuleRules
 				"MeshConversion",
                 "GeometryAlgorithms", // required for constrained Delaunay triangulation
                 "MeshSolverUtilities", // required by the smoothing operators
-				"MeshUtilitiesCommon", // required by uvlayoutop
 				// ... add other public dependencies that you statically link with here ...
 			}
             );
diff --git a/Engine/Plugins/Experimental/MeshModelingToolset/Source/ModelingOperatorsEditorOnly/ModelingOperatorsEditorOnly.Build.cs b/Engine/Plugins/Experimental/MeshModelingToolset/Source/ModelingOperatorsEditorOnly/ModelingOperatorsEditorOnly.Build.cs
--- a/Engine/Plugins/Experimental/MeshModelingToolset/Source/ModelingOperatorsEditorOnly/ModelingOperatorsEditorOnly.Build.cs
+++ b/Engine/Plugins/Experimental/MeshModelingToolset/Source/ModelingOperatorsEditorOnly/ModelingOperatorsEditorOnly.Build.cs
@@ -36,6 +36,7 @@ public class ModelingOperatorsEditorOnly : ModuleRules
 				"MeshSolverUtilities",
 				"GeometryAlgorithms", // required for constrained Delaunay triangulation
                 "ProxyLODMeshReduction", // currently required to be public due to IVoxelBasedCSG API
+				"MeshUtilitiesCommon", // required by uvlayoutop
 				// ... add other public dependencies that you statically link with here ...
 			}
             );
