ModelingTools: Add Select By Material option in Select Tool
#rnx

diff --git a/Engine/Plugins/Experimental/MeshModelingToolset/Source/MeshModelingTools/Private/MeshSelectionTool.cpp b/Engine/Plugins/Experimental/MeshModelingToolset/Source/MeshModelingTools/Private/MeshSelectionTool.cpp
--- a/Engine/Plugins/Experimental/MeshModelingToolset/Source/MeshModelingTools/Private/MeshSelectionTool.cpp
+++ b/Engine/Plugins/Experimental/MeshModelingToolset/Source/MeshModelingTools/Private/MeshSelectionTool.cpp
@@ -403,6 +403,15 @@ void UMeshSelectionTool::UpdateFaceSelection(const FBrushStampData& Stamp, const
 			[Mesh](int t1, int t2) { return Mesh->GetTriangleGroup(t1) == Mesh->GetTriangleGroup(t2); } );
 		UseROI = &LocalROI;
 	}
+	else if (SelectionProps->SelectionMode == EMeshSelectionToolPrimaryMode::ByMaterial)
+	{
+		const FDynamicMeshMaterialAttribute* MaterialIDs = Mesh->Attributes()->GetMaterialID();
+		TArray<int32> StartROI;
+		StartROI.Add(Stamp.HitResult.FaceIndex);
+		FMeshConnectedComponents::GrowToConnectedTriangles(Mesh, StartROI, LocalROI, &TemporaryBuffer, &TemporarySet,
+			[Mesh, MaterialIDs](int t1, int t2) { return MaterialIDs->GetValue(t1) == MaterialIDs->GetValue(t2); });
+		UseROI = &LocalROI;
+	}
 	else if (SelectionProps->SelectionMode == EMeshSelectionToolPrimaryMode::AllWithinAngle)
 	{
 		TArray<int32> StartROI; 
diff --git a/Engine/Plugins/Experimental/MeshModelingToolset/Source/MeshModelingTools/Public/MeshSelectionTool.h b/Engine/Plugins/Experimental/MeshModelingToolset/Source/MeshModelingTools/Public/MeshSelectionTool.h
--- a/Engine/Plugins/Experimental/MeshModelingToolset/Source/MeshModelingTools/Public/MeshSelectionTool.h
+++ b/Engine/Plugins/Experimental/MeshModelingToolset/Source/MeshModelingTools/Public/MeshSelectionTool.h
@@ -172,6 +172,9 @@ enum class EMeshSelectionToolPrimaryMode
 	/** Select all triangles in groups connected to any triangle inside the brush */
 	AllInGroup,
 
+	/** Select all triangles with same material as hit triangle */
+	ByMaterial,
+
 	/** Select all triangles with normal within angular tolerance of hit triangle */
 	AllWithinAngle
 };
