ModelingTools: Fix incorrect world-bounds computation in DynamicMeshComponents
#rnx

diff --git a/Engine/Plugins/Experimental/MeshModelingToolset/Source/ModelingComponents/Private/OctreeDynamicMeshComponent.cpp b/Engine/Plugins/Experimental/MeshModelingToolset/Source/ModelingComponents/Private/OctreeDynamicMeshComponent.cpp
--- a/Engine/Plugins/Experimental/MeshModelingToolset/Source/ModelingComponents/Private/OctreeDynamicMeshComponent.cpp
+++ b/Engine/Plugins/Experimental/MeshModelingToolset/Source/ModelingComponents/Private/OctreeDynamicMeshComponent.cpp
@@ -346,19 +346,13 @@ FColor UOctreeDynamicMeshComponent::GetTriangleColor(int TriangleID)
 
 FBoxSphereBounds UOctreeDynamicMeshComponent::CalcBounds(const FTransform& LocalToWorld) const
 {
-	// Bounds are tighter if the box is generated from pre-transformed vertices.
+	// Bounds are tighter if the box is generated from transformed vertices.
 	FBox BoundingBox(ForceInit);
 	for ( FVector3d Vertex : Mesh->VerticesItr() ) 
 	{
 		BoundingBox += LocalToWorld.TransformPosition(Vertex);
 	}
-
-	FBoxSphereBounds NewBounds;
-	NewBounds.BoxExtent = BoundingBox.GetExtent();
-	NewBounds.Origin = BoundingBox.GetCenter();
-	NewBounds.SphereRadius = NewBounds.BoxExtent.Size();
-
-	return NewBounds;
+	return FBoxSphereBounds(BoundingBox);
 }
 
 
diff --git a/Engine/Plugins/Experimental/MeshModelingToolset/Source/ModelingComponents/Private/SimpleDynamicMeshComponent.cpp b/Engine/Plugins/Experimental/MeshModelingToolset/Source/ModelingComponents/Private/SimpleDynamicMeshComponent.cpp
--- a/Engine/Plugins/Experimental/MeshModelingToolset/Source/ModelingComponents/Private/SimpleDynamicMeshComponent.cpp
+++ b/Engine/Plugins/Experimental/MeshModelingToolset/Source/ModelingComponents/Private/SimpleDynamicMeshComponent.cpp
@@ -242,16 +242,13 @@ FColor USimpleDynamicMeshComponent::GetTriangleColor(const FDynamicMesh3* MeshIn
 
 FBoxSphereBounds USimpleDynamicMeshComponent::CalcBounds(const FTransform& LocalToWorld) const
 {
-	// Bounds are tighter if the box is generated from pre-transformed vertices.
-	FAxisAlignedBox3d MeshAABB = Mesh->GetCachedBounds();
-	FBox BoundingBox = (FBox)MeshAABB;
-
-	FBoxSphereBounds NewBounds;
-	NewBounds.BoxExtent = BoundingBox.GetExtent();
-	NewBounds.Origin = BoundingBox.GetCenter();
-	NewBounds.SphereRadius = NewBounds.BoxExtent.Size();
-
-	return NewBounds;
+	// Bounds are tighter if the box is generated from transformed vertices.
+	FBox BoundingBox(ForceInit);
+	for ( FVector3d Vertex : Mesh->VerticesItr() )
+	{
+		BoundingBox += LocalToWorld.TransformPosition(Vertex);
+	}
+	return FBoxSphereBounds(BoundingBox);
 }
 
 
