make plane cut allow multiple cuts without leaving tool
to support that: adds some support for packing multiple meshes into one dynamic mesh via a triangle attribute, a meshreplacementchange, and some improvements to the plane cut algo

#rnx

diff --git a/Engine/Plugins/Experimental/GeometryProcessing/Source/DynamicMesh/Private/DynamicMeshAttributeSet.cpp b/Engine/Plugins/Experimental/GeometryProcessing/Source/DynamicMesh/Private/DynamicMeshAttributeSet.cpp
--- a/Engine/Plugins/Experimental/GeometryProcessing/Source/DynamicMesh/Private/DynamicMeshAttributeSet.cpp
+++ b/Engine/Plugins/Experimental/GeometryProcessing/Source/DynamicMesh/Private/DynamicMeshAttributeSet.cpp
@@ -14,6 +14,11 @@ void FDynamicMeshAttributeSet::EnableMaterialID()
 	}
 }
 
+void FDynamicMeshAttributeSet::DisableMaterialID()
+{
+	MaterialIDAttrib.Reset();
+}
+
 
 
 bool FDynamicMeshAttributeSet::IsSeamEdge(int eid) const
diff --git a/Engine/Plugins/Experimental/GeometryProcessing/Source/DynamicMesh/Private/DynamicMeshEditor.cpp b/Engine/Plugins/Experimental/GeometryProcessing/Source/DynamicMesh/Private/DynamicMeshEditor.cpp
--- a/Engine/Plugins/Experimental/GeometryProcessing/Source/DynamicMesh/Private/DynamicMeshEditor.cpp
+++ b/Engine/Plugins/Experimental/GeometryProcessing/Source/DynamicMesh/Private/DynamicMeshEditor.cpp
@@ -833,6 +833,13 @@ void FDynamicMeshEditor::CopyAttributes(int FromTriangleID, int ToTriangleID, FM
 		NormalOverlay->SetTriangle(ToTriangleID, ToElemTri);
 		NormalLayerIndex++;
 	}
+
+	if (Mesh->Attributes()->HasMaterialID())
+	{
+		FDynamicMeshMaterialAttribute* MaterialIDs = Mesh->Attributes()->GetMaterialID();
+		MaterialIDs->SetValue(ToTriangleID, MaterialIDs->GetValue(FromTriangleID));
+	}
+
 }
 
 
@@ -1006,12 +1013,22 @@ void FDynamicMeshEditor::AppendMesh(const FDynamicMesh3* AppendMesh,
 			FDynamicMeshUVOverlay* ToUVs = Mesh->Attributes()->GetUVLayer(UVLayerIndex);
 			if (FromUVs != nullptr && ToUVs != nullptr)
 			{
-				FIndexMapi& UVMap = IndexMapsOut.GetUVMap(0);
+				FIndexMapi& UVMap = IndexMapsOut.GetUVMap(UVLayerIndex);
 				UVMap.Reserve(FromUVs->ElementCount());
 				AppendUVs(AppendMesh, FromUVs, ToUVs,
 					VertexMap, TriangleMap, UVMap);
 			}
 		}
+
+		if (AppendMesh->Attributes()->HasMaterialID() && Mesh->Attributes()->HasMaterialID())
+		{
+			const FDynamicMeshMaterialAttribute* FromMaterialIDs = AppendMesh->Attributes()->GetMaterialID();
+			FDynamicMeshMaterialAttribute* ToMaterialIDs = Mesh->Attributes()->GetMaterialID();
+			for (int TriID : AppendMesh->TriangleIndicesItr())
+			{
+				ToMaterialIDs->SetValue(TriangleMap.GetTo(TriID), FromMaterialIDs->GetValue(TriID));
+			}
+		}
 	}
 }
 
@@ -1183,6 +1200,13 @@ static void AppendAttributes(const FDynamicMesh3* FromMesh, int FromTriangleID,
 		}
 		ToNormalOverlay->SetTriangle(ToTriangleID, ToElemTri);
 	}
+
+	if (FromMesh->Attributes()->HasMaterialID() && ToMesh->Attributes()->HasMaterialID())
+	{
+		const FDynamicMeshMaterialAttribute* FromMaterialIDs = FromMesh->Attributes()->GetMaterialID();
+		FDynamicMeshMaterialAttribute* ToMaterialIDs = ToMesh->Attributes()->GetMaterialID();
+		ToMaterialIDs->SetValue(ToTriangleID, FromMaterialIDs->GetValue(FromTriangleID));
+	}
 }
 
 
@@ -1239,3 +1263,91 @@ void FDynamicMeshEditor::AppendTriangles(const FDynamicMesh3* SourceMesh, const
 		//Mesh->CheckValidity(true);
 	}
 }
+
+
+bool FDynamicMeshEditor::SplitMesh(const FDynamicMesh3* SourceMesh, TArray<FDynamicMesh3>& SplitMeshes, TFunctionRef<int(int)> TriIDToMeshID, int DeleteMeshID)
+{
+	TMap<int, int> MeshIDToIndex;
+	int NumMeshes = 0;
+	bool bAlsoDelete = false;
+	for (int TID : SourceMesh->TriangleIndicesItr())
+	{
+		int MeshID = TriIDToMeshID(TID);
+		if (MeshID == DeleteMeshID)
+		{
+			bAlsoDelete = true;
+			continue;
+		}
+		if (!MeshIDToIndex.Contains(MeshID))
+		{
+			MeshIDToIndex.Add(MeshID, NumMeshes++);
+		}
+	}
+
+	if (!bAlsoDelete && NumMeshes < 2)
+	{
+		return false; // nothing to do, so don't bother filling the split meshes array
+	}
+
+	SplitMeshes.Reset();
+	SplitMeshes.SetNum(NumMeshes);
+	// enable matching attributes
+	for (FDynamicMesh3& M : SplitMeshes)
+	{
+		if (SourceMesh->HasAttributes())
+		{
+			M.EnableAttributes();
+			M.Attributes()->EnableMatchingAttributes(*SourceMesh->Attributes());
+		}
+	}
+
+	if (NumMeshes == 0) // full delete case, just leave the empty mesh
+	{
+		return true;
+	}
+
+	TArray<FMeshIndexMappings> Mappings; Mappings.Reserve(NumMeshes);
+	FDynamicMeshEditResult UnusedInvalidResultAccumulator; // only here because some functions require it
+	for (int Idx = 0; Idx < NumMeshes; Idx++)
+	{
+		FMeshIndexMappings& Map = Mappings.Emplace_GetRef();
+		Map.Initialize(&SplitMeshes[Idx]);
+	}
+
+	for (int SourceTID : SourceMesh->TriangleIndicesItr())
+	{
+		int MeshID = TriIDToMeshID(SourceTID);
+		if (MeshID == DeleteMeshID)
+		{
+			continue; // just skip triangles w/ the Delete Mesh ID
+		}
+		int MeshIndex = MeshIDToIndex[MeshID];
+		FDynamicMesh3& Mesh = SplitMeshes[MeshIndex];
+		FMeshIndexMappings& IndexMaps = Mappings[MeshIndex];
+
+		FIndex3i Tri = SourceMesh->GetTriangle(SourceTID);
+
+		// FindOrCreateDuplicateGroup
+		int SourceGID = SourceMesh->GetTriangleGroup(SourceTID);
+		int NewGID = IndexMaps.GetNewGroup(SourceGID);
+
+		FIndex3i NewTri;
+		for (int j = 0; j < 3; ++j)
+		{
+			int SourceVID = Tri[j];
+			int NewVID = IndexMaps.GetNewVertex(SourceVID);
+			if (NewVID == IndexMaps.InvalidID())
+			{
+				NewVID = Mesh.AppendVertex(*SourceMesh, SourceVID);
+				IndexMaps.SetVertex(SourceVID, NewVID);
+			}
+			NewTri[j] = NewVID;
+		}
+
+		int NewTID = Mesh.AppendTriangle(NewTri, NewGID);
+		IndexMaps.SetTriangle(SourceTID, NewTID);
+		AppendAttributes(SourceMesh, SourceTID, &Mesh, NewTID, IndexMaps, UnusedInvalidResultAccumulator);
+	}
+	
+	return true;
+}
\ No newline at end of file
diff --git a/Engine/Plugins/Experimental/GeometryProcessing/Source/DynamicMesh/Private/Operations/MeshPlaneCut.cpp b/Engine/Plugins/Experimental/GeometryProcessing/Source/DynamicMesh/Private/Operations/MeshPlaneCut.cpp
--- a/Engine/Plugins/Experimental/GeometryProcessing/Source/DynamicMesh/Private/Operations/MeshPlaneCut.cpp
+++ b/Engine/Plugins/Experimental/GeometryProcessing/Source/DynamicMesh/Private/Operations/MeshPlaneCut.cpp
@@ -1,6 +1,10 @@
 // Copyright 1998-2019 Epic Games, Inc. All Rights Reserved.
 
 #include "Operations/MeshPlaneCut.h"
+
+#include "DynamicMesh3.h"
+#include "DynamicMeshTriangleAttribute.h"
+
 #include "Operations/SimpleHoleFiller.h"
 #include "Operations/PlanarHoleFiller.h"
 #include "MeshNormals.h"
@@ -9,8 +13,11 @@
 
 #include "Async/ParallelFor.h"
 
-bool FMeshPlaneCut::Cut()
+
+void FMeshPlaneCut::SplitCrossingEdges(TArray<double>& Signs, TSet<int>& ZeroEdges, TSet<int>& OnCutEdges, bool bDeleteTrisOnPlane)
 {
+	Signs.Reset(); ZeroEdges.Reset(); OnCutEdges.Reset();
+
 	double InvalidDist = -FMathd::MaxReal;
 
 	// TODO: handle selections
@@ -23,7 +30,7 @@ bool FMeshPlaneCut::Cut()
 
 	// compute Signs
 	int MaxVID = Mesh->MaxVertexID();
-	TArray<double> Signs;
+	
 	Signs.SetNum(MaxVID);
 
 	bool bNoParallel = false;
@@ -38,9 +45,43 @@ bool FMeshPlaneCut::Cut()
 			Signs[VID] = InvalidDist;
 		}
 	},
-	bNoParallel);
+		bNoParallel);
 
-	TSet<int> ZeroEdges, ZeroVertices, OnCutEdges;
+	if (bDeleteTrisOnPlane)
+	{
+		for (int TID = 0; TID < Mesh->MaxTriangleID(); TID++)
+		{
+			if (!Mesh->IsTriangle(TID))
+			{
+				continue;
+			}
+
+			FIndex3i Tri = Mesh->GetTriangle(TID);
+			FIndex3i TriEdges = Mesh->GetTriEdges(TID);
+			if (   FMathd::Abs(Signs[Tri.A]) < PlaneTolerance
+				&& FMathd::Abs(Signs[Tri.B]) < PlaneTolerance
+				&& FMathd::Abs(Signs[Tri.C]) < PlaneTolerance)
+			{
+				EMeshResult Res = Mesh->RemoveTriangle(TID, true, false);
+				ensure(Res == EMeshResult::Ok);
+				for (int EIdx = 0; EIdx < 3; EIdx++)
+				{
+					int EID = TriEdges[EIdx];
+					if (Mesh->IsEdge(EID))
+					{
+						// any edge that still exists after removal is a possible cut edge
+						OnCutEdges.Add(EID);
+					}
+					else
+					{
+						// in case the now-gone edge was added to cut edges earlier, make sure it's removed
+						// (note: if it's not present this function will just do nothing, which is fine)
+						OnCutEdges.Remove(EID);
+					}
+				}
+			}
+		}
+	}
 
 	// have to skip processing of new edges. If edge id
 	// is > max at start, is new. Otherwise if in NewEdges list, also new.
@@ -71,18 +112,19 @@ bool FMeshPlaneCut::Cut()
 
 		// If both Signs are 0, this edge is on-contour
 		// If one sign is 0, that vertex is on-contour
-		int n0 = (FMathd::Abs(f0) < FMathd::Epsilon) ? 1 : 0;
-		int n1 = (FMathd::Abs(f1) < FMathd::Epsilon) ? 1 : 0;
+		int n0 = (FMathd::Abs(f0) < PlaneTolerance) ? 1 : 0;
+		int n1 = (FMathd::Abs(f1) < PlaneTolerance) ? 1 : 0;
 		if (n0 + n1 > 0)
 		{
 			if (n0 + n1 == 2)
 			{
 				ZeroEdges.Add(EID);
 			}
-			else
-			{
-				ZeroVertices.Add((n0 == 1) ? ev[0] : ev[1]);
-			}
+			// TODO: if we need to track ZeroVertices, can add this code back:
+			//else
+			//{
+			//	ZeroVertices.Add((n0 == 1) ? ev[0] : ev[1]);
+			//}
 			continue;
 		}
 
@@ -107,15 +149,193 @@ bool FMeshPlaneCut::Cut()
 			NewEdges.Add(splitInfo.NewEdges.C); OnCutEdges.Add(splitInfo.NewEdges.C);
 		}
 	}
+}
 
-	// remove one-rings of all positive-side vertices. 
+bool FMeshPlaneCut::CutWithoutDelete(bool bSplitVerticesAtPlane, float OffsetVertices, TDynamicMeshScalarTriangleAttribute<int>* TriLabels, int NewLabelStartID)
+{
+	TArray<double> Signs;
+	TSet<int> ZeroEdges, OnCutEdges;
+	SplitCrossingEdges(Signs, ZeroEdges, OnCutEdges, bSplitVerticesAtPlane /* only delete on-plane tris if we're also splitting vertices apart */);
+
+	if (!bSplitVerticesAtPlane)
+	{
+		ensure(OffsetVertices == 0.0); // it would be weird to not split vertices and still request any offset of the 'other side' vertices; please don't do that
+	}
+
+
+	if (!TriLabels)
+	{
+		ensure(!bSplitVerticesAtPlane); // need labels to split verts currently
+		return false;
+	}
+
+	// collapse degenerate edges if we got em
+	if (bCollapseDegenerateEdgesOnCut)
+	{
+		CollapseDegenerateEdges(OnCutEdges, ZeroEdges);
+	}
+
+	TMap<int, int> OldLabelToNew;
+	int AvailableID = NewLabelStartID;
+	FVector3d VertexOffsetVec = OffsetVertices * PlaneNormal;
+	for (int VID : Mesh->VertexIndicesItr())
+	{
+		if (VID < Signs.Num() && Signs[VID] > PlaneTolerance)
+		{
+			Mesh->SetVertex(VID, Mesh->GetVertex(VID) + VertexOffsetVec);
+			for (int TID : Mesh->VtxTrianglesItr(VID))
+			{
+				int LabelID = TriLabels->GetValue(TID);
+				if (LabelID >= NewLabelStartID)
+				{
+					continue;
+				}
+				if (!OldLabelToNew.Contains(LabelID))
+				{
+					OldLabelToNew.Add(LabelID, AvailableID++);
+				}
+				TriLabels->SetValue(TID, OldLabelToNew[LabelID]);
+			}
+		}
+	}
+
+	if (!bSplitVerticesAtPlane)
+	{
+		return true;
+	}
+
+	// split the mesh apart and add open boundary info
+	TMap<int, int> SplitVertices;
+	TSet<int> BoundaryVertices;
+	const TSet<int>* Sets[2] { &OnCutEdges, &ZeroEdges };
+	for (int SetIdx = 0; SetIdx < 2; SetIdx++)
+	{
+		const TSet<int>& Set = *(Sets[SetIdx]);
+		for (int EID : Set)
+		{
+			if (!Mesh->IsEdge(EID))
+			{
+				continue;
+			}
+			FIndex4i Edge = Mesh->GetEdge(EID);
+			BoundaryVertices.Add(Edge.A);
+			BoundaryVertices.Add(Edge.B);
+		}
+	}
+	TArray<int> Triangles;
+	DynamicMeshInfo::FVertexSplitInfo SplitInfo;
+	for (int VID : BoundaryVertices)
+	{
+		if (!ensure(Mesh->IsVertex(VID))) // should not have invalid vertices in BoundaryVertices
+		{
+			continue;
+		}
+		Triangles.Reset();
+		int NonSplitTriCount = 0;
+		for (int TID : Mesh->VtxTrianglesItr(VID))
+		{
+			if (TriLabels->GetValue(TID) >= NewLabelStartID)
+			{
+				Triangles.Add(TID);
+			}
+			else
+			{
+				NonSplitTriCount++;
+			}
+		}
+		if (NonSplitTriCount > 0) // connected to both old and new labels -- needs split
+		{
+			if (EMeshResult::Ok == Mesh->SplitVertex(VID, Triangles, SplitInfo))
+			{
+				SplitVertices.Add(VID, SplitInfo.NewVertex);
+				Mesh->SetVertex(SplitInfo.NewVertex, Mesh->GetVertex(SplitInfo.NewVertex) + VertexOffsetVec);
+			}
+		}
+		else if (Signs[VID] <= PlaneTolerance) // wasn't already offset and has no connections to 'old' labels -- needs offset
+		{
+			Mesh->SetVertex(VID, Mesh->GetVertex(VID) + VertexOffsetVec);
+		}
+	}
+	
+	// transfer ZeroEdges and OnCutEdges to newly split edges
+	TMap<int, TSet<int>> LabelToCutEdges;
+	for (int SetIdx = 0; SetIdx < 2; SetIdx++)
+	{
+		const TSet<int>& Set = *(Sets[SetIdx]);
+		for (int EID : Set)
+		{
+			if (!Mesh->IsEdge(EID))
+			{
+				continue;
+			}
+			FIndex4i Edge = Mesh->GetEdge(EID);
+			if (Edge.D >= 0) // only care about boundary edges
+			{
+				continue;
+			}
+			{
+				int LabelID = TriLabels->GetValue(Edge.C);
+				TSet<int>& LabelCutEdges = LabelToCutEdges.FindOrAdd(LabelID);
+				LabelCutEdges.Add(EID);
+			}
+
+			// try to find and add the corresponding edge
+			const int* SplitA = SplitVertices.Find(Edge.A);
+			const int* SplitB = SplitVertices.Find(Edge.B);
+			if (SplitA && SplitB)
+			{
+				int CorrEID = Mesh->FindEdge(*SplitA, *SplitB);
+				if (CorrEID >= 0) // corresponding edge exists
+				{
+					FIndex4i CorrEdge = Mesh->GetEdge(CorrEID);
+					if (CorrEdge.D < 0) // we only care if it's a boundary edge
+					{
+						int LabelID = TriLabels->GetValue(CorrEdge.C);
+						TSet<int>& LabelCutEdges = LabelToCutEdges.FindOrAdd(LabelID);
+						LabelCutEdges.Add(CorrEID);
+					}
+				}
+			}
+			BoundaryVertices.Add(Edge.A);
+			BoundaryVertices.Add(Edge.B);
+		}
+	}
+
+	bool AllExtractionsOk = true;
+	TSet<int> UnusedZeroEdgesSet;
+	for (TPair<int, TSet<int>>& LabelIDEdges : LabelToCutEdges)
+	{
+		int LabelID = LabelIDEdges.Key;
+		TSet<int>& Edges = LabelIDEdges.Value;
+		FMeshPlaneCut::FOpenBoundary& Boundary = OpenBoundaries.Emplace_GetRef();
+		Boundary.Label = LabelID;
+		if (LabelID >= NewLabelStartID)
+		{
+			Boundary.NormalSign = -1;
+		}
+
+		check(UnusedZeroEdgesSet.Num() == 0); // for simplicity, we only put stuff in the CutEdges set
+		bool ExtractOk = ExtractBoundaryLoops(Edges, UnusedZeroEdgesSet, Boundary);
+		AllExtractionsOk = ExtractOk && AllExtractionsOk;
+	}
+
+	return AllExtractionsOk;
+}
+
+bool FMeshPlaneCut::Cut()
+{
+	TArray<double> Signs;
+	TSet<int> ZeroEdges, OnCutEdges;
+	SplitCrossingEdges(Signs, ZeroEdges, OnCutEdges);
+	
 	// @todo handle selection logic
 	//IEnumerable<int> vertexSet = Interval1i.Range(MaxVID);
 	//if (CutVertexSet != null)
 	//	vertexSet = CutVertexSet;
+	// remove one-rings of all positive-side vertices. 
 	for (int VID : Mesh->VertexIndicesItr())
-	{ 
-		if (VID < Signs.Num() && Signs[VID] > FMathd::Epsilon)
+	{
+		if (VID < Signs.Num() && Signs[VID] > PlaneTolerance)
 		{
 			Mesh->RemoveVertex(VID, true, false);
 		}
@@ -127,10 +347,16 @@ bool FMeshPlaneCut::Cut()
 		CollapseDegenerateEdges(OnCutEdges, ZeroEdges);
 	}
 
+	FMeshPlaneCut::FOpenBoundary& Boundary = OpenBoundaries.Emplace_GetRef();
+	return ExtractBoundaryLoops(OnCutEdges, ZeroEdges, Boundary);
+
+}
 
+bool FMeshPlaneCut::ExtractBoundaryLoops(const TSet<int>& OnCutEdges, const TSet<int>& ZeroEdges, FMeshPlaneCut::FOpenBoundary& Boundary)
+{
 	// ok now we extract boundary loops, but restricted
 	// to either the zero-edges we found, or the edges we created! bang!!
-	
+
 	FMeshBoundaryLoops Loops(Mesh, false);
 	Loops.EdgeFilterFunc = [&OnCutEdges, &ZeroEdges](int EID)
 	{
@@ -140,25 +366,23 @@ bool FMeshPlaneCut::Cut()
 
 	if (bFoundLoops)
 	{
-		CutLoops = Loops.Loops;
-		CutSpans = Loops.Spans;
-		CutLoopsFailed = false;
-		FoundOpenSpans = CutSpans.Num() > 0;
+		Boundary.CutLoops = Loops.Loops;
+		Boundary.CutSpans = Loops.Spans;
+		Boundary.CutLoopsFailed = false;
+		Boundary.FoundOpenSpans = Boundary.CutSpans.Num() > 0;
 	}
 	else
 	{
-		CutLoops.Empty();
-		CutLoopsFailed = true;
+		Boundary.CutLoops.Empty();
+		Boundary.CutLoopsFailed = true;
 	}
 
-	return !CutLoopsFailed;
+	return !Boundary.CutLoopsFailed;
 }
 
-
-
 void FMeshPlaneCut::CollapseDegenerateEdges(const TSet<int>& OnCutEdges, const TSet<int>& ZeroEdges)
 {
-	TSet<int> Sets[2] { OnCutEdges, ZeroEdges };
+	const TSet<int>* Sets[2] { &OnCutEdges, &ZeroEdges };
 
 	double Tol2 = DegenerateEdgeTol * DegenerateEdgeTol;
 	FVector3d A, B;
@@ -168,7 +392,8 @@ void FMeshPlaneCut::CollapseDegenerateEdges(const TSet<int>& OnCutEdges, const T
 		Collapsed = 0;
 		for (int SetIdx = 0; SetIdx < 2; SetIdx++)
 		{
-			for (int EID : Sets[SetIdx])
+			const TSet<int>& Set = *(Sets[SetIdx]);
+			for (int EID : Set)
 			{
 				if (!Mesh->IsEdge(EID))
 				{
@@ -208,59 +433,85 @@ bool FMeshPlaneCut::SimpleHoleFill(int ConstantGroupID)
 	bool bAllOk = true;
 
 	HoleFillTriangles.Empty();
+	for (FOpenBoundary& Boundary : OpenBoundaries)
+	{
+		TArray<int> BoundaryFillTriangles = HoleFillTriangles.Emplace_GetRef();
+		FFrame3d ProjectionFrame(PlaneOrigin, PlaneNormal);
+
+		for (const FEdgeLoop& Loop : Boundary.CutLoops)
+		{
+			FSimpleHoleFiller Filler(Mesh, Loop);
+			int GID = ConstantGroupID >= 0 ? ConstantGroupID : Mesh->AllocateTriangleGroup();
+			bAllOk = Filler.Fill(GID) && bAllOk;
+
+			BoundaryFillTriangles.Append(Filler.NewTriangles);
+
+			if (Mesh->HasAttributes())
+			{
+				FDynamicMeshEditor Editor(Mesh);
+				Editor.SetTriangleNormals(Filler.NewTriangles, (FVector3f)PlaneNormal * Boundary.NormalSign);
+				Editor.SetTriangleUVsFromProjection(Filler.NewTriangles, ProjectionFrame, UVScaleFactor);
+			}
+		}
+	}
 
-	FFrame3d ProjectionFrame(PlaneOrigin, PlaneNormal);
+	return bAllOk;
+}
+
+
+bool FMeshPlaneCut::HoleFill(TFunction<TArray<FIndex3i>(const FGeneralPolygon2d&)> PlanarTriangulationFunc, bool bFillSpans, int ConstantGroupID)
+{
+	bool bAllOk = true;
 
-	for (const FEdgeLoop& Loop : CutLoops)
+	HoleFillTriangles.Empty();
+	for (FMeshPlaneCut::FOpenBoundary& Boundary : OpenBoundaries)
 	{
-		FSimpleHoleFiller Filler(Mesh, Loop);
+		TArray<TArray<int>> LoopVertices;
+		for (const FEdgeLoop& Loop : Boundary.CutLoops)
+		{
+			LoopVertices.Add(Loop.Vertices);
+		}
+		if (bFillSpans)
+		{
+			for (const FEdgeSpan& Span : Boundary.CutSpans)
+			{
+				LoopVertices.Add(Span.Vertices);
+			}
+		}
+		FPlanarHoleFiller Filler(Mesh, &LoopVertices, PlanarTriangulationFunc, PlaneOrigin, PlaneNormal*Boundary.NormalSign);
+
 		int GID = ConstantGroupID >= 0 ? ConstantGroupID : Mesh->AllocateTriangleGroup();
-		bAllOk = Filler.Fill(GID) && bAllOk;
-		
-		HoleFillTriangles.Add(Filler.NewTriangles);
+		bool bFullyFilledHole = Filler.Fill(GID);
 
+		HoleFillTriangles.Add(Filler.NewTriangles);
 		if (Mesh->HasAttributes())
 		{
 			FDynamicMeshEditor Editor(Mesh);
-			Editor.SetTriangleNormals(Filler.NewTriangles, (FVector3f)PlaneNormal);
+			Editor.SetTriangleNormals(Filler.NewTriangles, (FVector3f)(PlaneNormal*Boundary.NormalSign));
+
+			FFrame3d ProjectionFrame(PlaneOrigin, PlaneNormal*Boundary.NormalSign);
 			Editor.SetTriangleUVsFromProjection(Filler.NewTriangles, ProjectionFrame, UVScaleFactor);
 		}
+
+		bAllOk = bAllOk && bFullyFilledHole;
 	}
 
 	return bAllOk;
 }
 
-
-bool FMeshPlaneCut::HoleFill(TFunction<TArray<FIndex3i>(const FGeneralPolygon2d&)> PlanarTriangulationFunc, bool bFillSpans, int ConstantGroupID)
+void FMeshPlaneCut::TransferTriangleLabelsToHoleFillTriangles(TDynamicMeshScalarTriangleAttribute<int>* TriLabels)
 {
-	HoleFillTriangles.Empty();
-
-	TArray<TArray<int>> LoopVertices;
-	for (const FEdgeLoop& Loop : CutLoops)
+	if (!ensure(OpenBoundaries.Num() == HoleFillTriangles.Num()))
 	{
-		LoopVertices.Add(Loop.Vertices);
+		return;
 	}
-	if (bFillSpans)
+	for (int BoundaryIdx = 0; BoundaryIdx < OpenBoundaries.Num(); BoundaryIdx++)
 	{
-		for (const FEdgeSpan& Span : CutSpans)
+		const TArray<int>& Triangles = HoleFillTriangles[BoundaryIdx];
+		const FOpenBoundary& Boundary = OpenBoundaries[BoundaryIdx];
+		for (int TID : Triangles)
 		{
-			LoopVertices.Add(Span.Vertices);
+			TriLabels->SetValue(TID, Boundary.Label);
 		}
 	}
-	FPlanarHoleFiller Filler(Mesh, &LoopVertices, PlanarTriangulationFunc, PlaneOrigin, PlaneNormal);
-
-	int GID = ConstantGroupID >= 0 ? ConstantGroupID : Mesh->AllocateTriangleGroup();
-	bool bFullyFilledHole = Filler.Fill(GID);
-	
-	HoleFillTriangles.Add(Filler.NewTriangles);
-	if (Mesh->HasAttributes())
-	{
-		FDynamicMeshEditor Editor(Mesh);
-		Editor.SetTriangleNormals(Filler.NewTriangles, (FVector3f)PlaneNormal);
-
-		FFrame3d ProjectionFrame(PlaneOrigin, PlaneNormal);
-		Editor.SetTriangleUVsFromProjection(Filler.NewTriangles, ProjectionFrame, UVScaleFactor);
-	}
-
-	return bFullyFilledHole;
 }
diff --git a/Engine/Plugins/Experimental/GeometryProcessing/Source/DynamicMesh/Public/DynamicAttribute.h b/Engine/Plugins/Experimental/GeometryProcessing/Source/DynamicMesh/Public/DynamicAttribute.h
--- a/Engine/Plugins/Experimental/GeometryProcessing/Source/DynamicMesh/Public/DynamicAttribute.h
+++ b/Engine/Plugins/Experimental/GeometryProcessing/Source/DynamicMesh/Public/DynamicAttribute.h
@@ -61,6 +61,9 @@ public:
 
 	/** Allocate a new copy of the attribute layer, optionally with a different parent mesh */
 	virtual TDynamicAttributeBase* MakeCopy(ParentType* ParentIn) const = 0;
+	/** Allocate a new empty instance of the same type of attribute layer */
+	virtual TDynamicAttributeBase* MakeNew(ParentType* ParentIn) const = 0;
+	
 	virtual void Reparent(ParentType* NewParent) = 0;
 
 	virtual void OnNewVertex(int VertexID, bool bInserted)
diff --git a/Engine/Plugins/Experimental/GeometryProcessing/Source/DynamicMesh/Public/DynamicMeshAttributeSet.h b/Engine/Plugins/Experimental/GeometryProcessing/Source/DynamicMesh/Public/DynamicMeshAttributeSet.h
--- a/Engine/Plugins/Experimental/GeometryProcessing/Source/DynamicMesh/Public/DynamicMeshAttributeSet.h
+++ b/Engine/Plugins/Experimental/GeometryProcessing/Source/DynamicMesh/Public/DynamicMeshAttributeSet.h
@@ -54,6 +54,10 @@ public:
 			EnableMaterialID();
 			MaterialIDAttrib->Copy(*(Copy.MaterialIDAttrib));
 		}
+		else
+		{
+			DisableMaterialID();
+		}
 
 		GenericAttributes.Reset();
 		ResetRegisteredAttributes();
@@ -66,6 +70,36 @@ public:
 		// parent mesh is *not* copied!
 	}
 
+	/**
+	 * Enable the matching attributes and overlay layers as the reference Copy set, but do not copy any data across
+	 */
+	void EnableMatchingAttributes(const FDynamicMeshAttributeSet& ToMatch)
+	{
+		SetNumUVLayers(ToMatch.NumUVLayers());
+		for (int UVIdx = 0; UVIdx < NumUVLayers(); UVIdx++)
+		{
+			UVLayers[UVIdx].ClearElements();
+		}
+		Normals0.ClearElements();
+
+		if (ToMatch.MaterialIDAttrib)
+		{
+			EnableMaterialID();
+		}
+		else
+		{
+			DisableMaterialID();
+		}
+
+		GenericAttributes.Reset();
+		ResetRegisteredAttributes();
+		for (int Idx = 0; Idx < ToMatch.GenericAttributes.Num(); Idx++)
+		{
+			const FDynamicMeshAttributeBase* SourceAttrib = ToMatch.GenericAttributes[Idx].Get();
+			AttachAttribute(SourceAttrib->MakeNew(ParentMesh));
+		}
+	}
+
 	/** @return the parent mesh for this overlay */
 	const FDynamicMesh3* GetParentMesh() const { return ParentMesh; }
 	/** @return the parent mesh for this overlay */
@@ -208,6 +242,8 @@ public:
 
 	void EnableMaterialID();
 
+	void DisableMaterialID();
+
 	FDynamicMeshMaterialAttribute* GetMaterialID()
 	{
 		return MaterialIDAttrib.Get();
diff --git a/Engine/Plugins/Experimental/GeometryProcessing/Source/DynamicMesh/Public/DynamicMeshEditor.h b/Engine/Plugins/Experimental/GeometryProcessing/Source/DynamicMesh/Public/DynamicMeshEditor.h
--- a/Engine/Plugins/Experimental/GeometryProcessing/Source/DynamicMesh/Public/DynamicMeshEditor.h
+++ b/Engine/Plugins/Experimental/GeometryProcessing/Source/DynamicMesh/Public/DynamicMeshEditor.h
@@ -426,6 +426,16 @@ public:
 	 */
 	void AppendTriangles(const FDynamicMesh3* SourceMesh, const TArray<int>& SourceTriangles, FMeshIndexMappings& IndexMaps, FDynamicMeshEditResult& ResultOut);
 
+	/**
+	 * Create multiple meshes out of the source mesh by splitting triangles out.
+	 * Static because it creates multiple output meshes, so doesn't quite fit in the FDynamicMeshEditor model of operating on a single mesh
+	 *
+	 * @param SourceMesh
+	 * @param SplitMeshes
+	 * @param TriIDToMeshID
+	 * @return true if needed split, false if there were not multiple mesh ids so no split was needed
+	 */
+	static bool SplitMesh(const FDynamicMesh3* SourceMesh, TArray<FDynamicMesh3>& SplitMeshes, TFunctionRef<int(int)> TriIDToMeshID, int DeleteMeshID = -1);
 
 };
 
diff --git a/Engine/Plugins/Experimental/GeometryProcessing/Source/DynamicMesh/Public/DynamicMeshTriangleAttribute.h b/Engine/Plugins/Experimental/GeometryProcessing/Source/DynamicMesh/Public/DynamicMeshTriangleAttribute.h
--- a/Engine/Plugins/Experimental/GeometryProcessing/Source/DynamicMesh/Public/DynamicMeshTriangleAttribute.h
+++ b/Engine/Plugins/Experimental/GeometryProcessing/Source/DynamicMesh/Public/DynamicMeshTriangleAttribute.h
@@ -85,6 +85,13 @@ public:
 	/** @return the parent mesh for this overlay */
 	FDynamicMesh3* GetParentMesh() { return ParentMesh; }
 
+	virtual FDynamicMeshAttributeBase* MakeNew(FDynamicMesh3* ParentMeshIn) const override
+	{
+		TDynamicMeshTriangleAttribute<AttribValueType, AttribDimension>* Matching = new TDynamicMeshTriangleAttribute<AttribValueType, AttribDimension>(ParentMeshIn);
+		Matching->Initialize();
+		return Matching;
+	}
+
 	virtual FDynamicMeshAttributeBase* MakeCopy(FDynamicMesh3* ParentMeshIn) const override
 	{
 		TDynamicMeshTriangleAttribute<AttribValueType, AttribDimension>* ToFill = new TDynamicMeshTriangleAttribute<AttribValueType, AttribDimension>(ParentMeshIn);
@@ -247,6 +254,25 @@ public:
 		// nothing to do here because triangles did not change
 	}
 
+	virtual AttribValueType GetDefaultAttributeValue()
+	{
+		return (AttribValueType)0;
+	}
+
+	inline void ResizeAttribStoreIfNeeded(int TriangleID)
+	{
+		int k = TriangleID * AttribDimension;
+		if (k+AttribDimension > AttribValues.Num())
+		{
+			AttribValues.Resize(k + AttribDimension, GetDefaultAttributeValue());
+		}
+	}
+
+	virtual void OnNewTriangle(int TriangleID, bool bInserted) override
+	{
+		ResizeAttribStoreIfNeeded(TriangleID);
+	}
+
 
 };
 
diff --git a/Engine/Plugins/Experimental/GeometryProcessing/Source/DynamicMesh/Public/DynamicVertexAttribute.h b/Engine/Plugins/Experimental/GeometryProcessing/Source/DynamicMesh/Public/DynamicVertexAttribute.h
--- a/Engine/Plugins/Experimental/GeometryProcessing/Source/DynamicMesh/Public/DynamicVertexAttribute.h
+++ b/Engine/Plugins/Experimental/GeometryProcessing/Source/DynamicMesh/Public/DynamicVertexAttribute.h
@@ -77,6 +77,12 @@ public:
 private:
 	void Reparent(ParentType* NewParent) override { Parent = NewParent;  }
 public:
+	virtual TDynamicAttributeBase<ParentType>* MakeNew(ParentType* ParentIn) const override
+	{
+		TDynamicVertexAttribute<AttribValueType, AttribDimension, ParentType>* Matching = new TDynamicVertexAttribute<AttribValueType, AttribDimension, ParentType>(ParentIn);
+		Matching->Initialize();
+		return Matching;
+	}
 	virtual TDynamicAttributeBase<ParentType>* MakeCopy(ParentType* ParentIn) const override
 	{
 		TDynamicVertexAttribute<AttribValueType, AttribDimension, ParentType>* ToFill = new TDynamicVertexAttribute<AttribValueType, AttribDimension, ParentType>(ParentIn);
diff --git a/Engine/Plugins/Experimental/GeometryProcessing/Source/DynamicMesh/Public/Operations/MeshPlaneCut.h b/Engine/Plugins/Experimental/GeometryProcessing/Source/DynamicMesh/Public/Operations/MeshPlaneCut.h
--- a/Engine/Plugins/Experimental/GeometryProcessing/Source/DynamicMesh/Public/Operations/MeshPlaneCut.h
+++ b/Engine/Plugins/Experimental/GeometryProcessing/Source/DynamicMesh/Public/Operations/MeshPlaneCut.h
@@ -10,9 +10,9 @@
 #include "MeshBoundaryLoops.h"
 #include "Curve/GeneralPolygon2.h"
 
-
 class FDynamicMesh3;
-
+template<typename RealType>
+class TDynamicMeshScalarTriangleAttribute;
 
 /**
  * Cut the Mesh with the Plane. The *positive* side, ie (p-o).n > 0, is removed.
@@ -22,12 +22,14 @@ class FDynamicMesh3;
  * 
  * Algorithm is:
  *    1) find all edge crossings
- *    2) Do edge splits at crossings
- *    3) delete all vertices on positive side
- *    4) (optionally) collapse any degenerate boundary edges 
- *    5) find loops through valid boundary edges (ie connected to splits, or on-plane edges)
- * 
- * @todo Also discard any triangles with all vertex distances < epsilon.
+ *	  2) optionally discard any triangles with all vertex distances < epsilon.
+ *    3) Do edge splits at crossings
+ *    4 option a) (optionally) delete all vertices on positive side
+ *	  4 option b) (OR optionally) disconnect all triangles w/ vertices on positive side (if keeping both sides)
+ *	  4 option c) do nothing (if keeping both sides and not disconnecting them)
+ *    5) (optionally) collapse any degenerate boundary edges 
+ *	  6) (optionally) change an attribute tag for all triangles on positive side
+ *    7) find loops through valid boundary edges (ie connected to splits, or on-plane edges) (if second half was kept, do this separately for each separate mesh ID label)
  */
 class DYNAMICMESH_API FMeshPlaneCut
 {
@@ -45,22 +47,32 @@ public:
 	/** UVs on any hole fill surfaces are scaled by this amount */
 	float UVScaleFactor = 1.0f;
 
+	/** Tolerance distance for considering a vertex to be 'on plane' */
+	double PlaneTolerance = FMathf::ZeroTolerance * 10.0;
+
+
 	// TODO support optionally restricting plane cut to a mesh selection
 	//MeshFaceSelection CutFaceSet;
 
 	//
 	// Outputs
 	//
-	TArray<FEdgeLoop> CutLoops;
-	TArray<FEdgeSpan> CutSpans;
-	bool CutLoopsFailed = false;		// set to true if we could not compute cut loops/spans
-	bool FoundOpenSpans = false;     // set to true if we found open spans in cut
-
-	// Note: In practice this is 1:1 with CutLoops if all the hole fills succeed w/ simple hole filling, but this will not be true for other hole filling methods, or if any of the hole fills fails
-	// please do not write anything relying on a correspondence between the two
+	struct FOpenBoundary
+	{
+		int Label; // optional ID, used to transfer label to new hole-fill triangles
+		float NormalSign = 1; // -1 for the open boundary on the other side of the cut (for the CutWithoutDelete path)
+		TArray<FEdgeLoop> CutLoops;
+		TArray<FEdgeSpan> CutSpans;
+		bool CutLoopsFailed = false;		// set to true if we could not compute cut loops/spans
+		bool FoundOpenSpans = false;     // set to true if we found open spans in cut
+	};
+	// note: loops and spans within a single FOpenBoundary could be part of the same hole-fill triangulation
+	//	separate open boundary structs will be considered separately and will not share hole fill triangles
+	TArray<FOpenBoundary> OpenBoundaries;
+
+	// Triangle IDs of hole fill triangles.  Outer array is 1:1 with the OpenBoundaries array
 	TArray<TArray<int>> HoleFillTriangles;
 
-
 public:
 
 	/**
@@ -80,8 +92,13 @@ public:
 		return EOperationValidationResult::Ok;
 	}
 
+	// split plane-crossing edges and remove geometry on the positive side of the cutting plane
 	virtual bool Cut();
 
+	// alternate path for splitting the mesh without deleting the other half
+	virtual bool CutWithoutDelete(bool bSplitVerticesAtPlane, float OffsetSeparatedPortion = 0.0f,
+		TDynamicMeshScalarTriangleAttribute<int>* TriLabels = nullptr, int NewLabelStartID = 0);
+
 	/**
 	 *  Fill cut loops with FSimpleHoleFiller
 	 */
@@ -93,7 +110,11 @@ public:
 	virtual bool HoleFill(TFunction<TArray<FIndex3i>(const FGeneralPolygon2d&)> PlanarTriangulationFunc, bool bFillSpans, int ConstantGroupID = -1);
 
 	
+	virtual void TransferTriangleLabelsToHoleFillTriangles(TDynamicMeshScalarTriangleAttribute<int>* TriLabels);
+
 protected:
 
 	void CollapseDegenerateEdges(const TSet<int>& OnCutEdges, const TSet<int>& ZeroEdges);
+	void SplitCrossingEdges(TArray<double>& Signs, TSet<int>& ZeroEdges, TSet<int>& OnCutEdges, bool bDeleteTrisOnPlane = true);
+	bool ExtractBoundaryLoops(const TSet<int>& OnCutEdges, const TSet<int>& ZeroEdges, FMeshPlaneCut::FOpenBoundary& Boundary);
 };
diff --git a/Engine/Plugins/Experimental/GeometryProcessing/Source/GeometryAlgorithms/Private/ConstrainedDelaunay2.cpp b/Engine/Plugins/Experimental/GeometryProcessing/Source/GeometryAlgorithms/Private/ConstrainedDelaunay2.cpp
--- a/Engine/Plugins/Experimental/GeometryProcessing/Source/GeometryAlgorithms/Private/ConstrainedDelaunay2.cpp
+++ b/Engine/Plugins/Experimental/GeometryProcessing/Source/GeometryAlgorithms/Private/ConstrainedDelaunay2.cpp
@@ -195,6 +195,8 @@ bool TConstrainedDelaunay2<RealType>::Triangulate()
 		return false;
 	}
 
+	const std::vector<int>& Duplicates = Delaunay.GetDuplicates();
+
 	std::vector<int> OutEdges;
 	bool InsertConstraintFailure = false;
 	TMap<TPair<int, int>, bool> BoundaryMap, HoleMap; // tracks all the boundary edges as they are added, so we can later flood fill across them for inside/outside decisions
@@ -208,7 +210,9 @@ bool TConstrainedDelaunay2<RealType>::Triangulate()
 		TMap<TPair<int, int>, bool>& InputMap = *EdgeAndHoleMaps[EdgeOrHole];
 		for (const FIndex2i& Edge : Input)
 		{
-			if (!Delaunay.Insert({{Edge.A, Edge.B}}, OutEdges))
+			int A = Duplicates[Edge.A];
+			int B = Duplicates[Edge.B];
+			if (!Delaunay.Insert({{A, B}}, OutEdges))
 			{
 				// Note the failed edge; we will try to proceed anyway, just without this edge.  Hopefully the CDT is robust and this never happens!
 				ensureMsgf(false, TEXT("CDT edge insertion failed"));
diff --git a/Engine/Plugins/Experimental/MeshModelingToolset/Source/MeshModelingTools/Private/MeshSelectionTool.cpp b/Engine/Plugins/Experimental/MeshModelingToolset/Source/MeshModelingTools/Private/MeshSelectionTool.cpp
--- a/Engine/Plugins/Experimental/MeshModelingToolset/Source/MeshModelingTools/Private/MeshSelectionTool.cpp
+++ b/Engine/Plugins/Experimental/MeshModelingToolset/Source/MeshModelingTools/Private/MeshSelectionTool.cpp
@@ -1004,6 +1004,7 @@ void UMeshSelectionTool::SeparateSelectedTriangles()
 	// extract copy of triangles
 	FDynamicMesh3 SeparatedMesh;
 	SeparatedMesh.EnableAttributes();
+	SeparatedMesh.Attributes()->EnableMatchingAttributes(*SourceMesh->Attributes());
 	FDynamicMeshEditor Editor(&SeparatedMesh);
 	FMeshIndexMappings Mappings; FDynamicMeshEditResult EditResult;
 	Editor.AppendTriangles(SourceMesh, SelectedFaces, Mappings, EditResult);
diff --git a/Engine/Plugins/Experimental/MeshModelingToolset/Source/MeshModelingTools/Private/PlaneCutTool.cpp b/Engine/Plugins/Experimental/MeshModelingToolset/Source/MeshModelingTools/Private/PlaneCutTool.cpp
--- a/Engine/Plugins/Experimental/MeshModelingToolset/Source/MeshModelingTools/Private/PlaneCutTool.cpp
+++ b/Engine/Plugins/Experimental/MeshModelingToolset/Source/MeshModelingTools/Private/PlaneCutTool.cpp
@@ -7,6 +7,8 @@
 #include "ToolSetupUtil.h"
 
 #include "DynamicMesh3.h"
+#include "DynamicMeshTriangleAttribute.h"
+#include "DynamicMeshEditor.h"
 #include "BaseBehaviors/MultiClickSequenceInputBehavior.h"
 #include "Selection/SelectClickedAction.h"
 
@@ -21,6 +23,8 @@
 #include "Drawing/MeshDebugDrawing.h"
 #include "AssetGenerationUtil.h"
 
+#include "Changes/ToolCommandChangeSequence.h"
+
 #include "CuttingOps/PlaneCutOp.h"
 
 #define LOCTEXT_NAMESPACE "UPlaneCutTool"
@@ -68,7 +72,6 @@ UInteractiveTool* UPlaneCutToolBuilder::BuildTool(const FToolBuilderState& Scene
 
 UPlaneCutToolProperties::UPlaneCutToolProperties()
 {
-	bDiscardAttributes = false;
 	bKeepBothHalves = false;
 	bFillCutHole = true;
 	SpacingBetweenHalves = 1;
@@ -100,13 +103,30 @@ void UPlaneCutTool::Setup()
 	}
 
 	// Convert input mesh descriptions to dynamic mesh
-	for (TUniquePtr<FPrimitiveComponentTarget>& ComponentTarget : ComponentTargets)
+	for (int Idx = 0; Idx < ComponentTargets.Num(); Idx++)
 	{
-		TSharedPtr<FDynamicMesh3> OriginalDynamicMesh = MakeShared<FDynamicMesh3>();
+		TUniquePtr<FPrimitiveComponentTarget>& ComponentTarget = ComponentTargets[Idx];
+		FDynamicMesh3* OriginalDynamicMesh = new FDynamicMesh3;
 		FMeshDescriptionToDynamicMesh Converter;
 		Converter.bPrintDebugMessages = true;
+
+		
 		Converter.Convert(ComponentTarget->GetMesh(), *OriginalDynamicMesh);
-		OriginalDynamicMeshes.Add(OriginalDynamicMesh);
+		OriginalDynamicMesh->EnableAttributes();
+		TDynamicMeshScalarTriangleAttribute<int>* SubObjectIDs = new TDynamicMeshScalarTriangleAttribute<int>(OriginalDynamicMesh);
+		SubObjectIDs->Initialize(0);
+		int AttribIndex = OriginalDynamicMesh->Attributes()->AttachAttribute(SubObjectIDs);
+
+		/// fill in the MeshesToCut array
+		UDynamicMeshReplacementChangeTarget* Target = MeshesToCut.Add_GetRef(NewObject<UDynamicMeshReplacementChangeTarget>());
+		MeshSubObjectAttribIndices.Add(AttribIndex);
+		check(MeshSubObjectAttribIndices.Num() == MeshesToCut.Num());
+		// store a UV scale based on the original mesh bounds (we don't want to recompute this between cuts b/c we want consistent UV scale)
+		MeshUVScaleFactor.Add(1.0 / OriginalDynamicMesh->GetBounds().MaxDim());
+
+		// Set callbacks so previews are invalidated on undo/redo changing the meshes
+		Target->SetMesh(TSharedPtr<const FDynamicMesh3>(OriginalDynamicMesh));
+		Target->OnMeshChanged.AddLambda([this, Idx]() { Previews[Idx]->InvalidateResult(); });
 	}
 
 	// click to set plane behavior
@@ -133,11 +153,12 @@ void UPlaneCutTool::Setup()
 	PlaneTransformGizmo = GizmoManager->Create3AxisTransformGizmo(this);
 
 	// initialize our properties
+	ToolPropertyObjects.Add(this);
 	BasicProperties = NewObject<UPlaneCutToolProperties>(this, TEXT("Plane Cut Settings"));
 	AddToolPropertySource(BasicProperties);
 
 	// initialize the PreviewMesh+BackgroundCompute object
-	UpdateNumPreviews();
+	SetupPreviews();
 
 	// set initial cut plane (also attaches gizmo/proxy)
 	FBox CombinedBounds; CombinedBounds.Init();
@@ -158,47 +179,65 @@ void UPlaneCutTool::Setup()
 }
 
 
-void UPlaneCutTool::UpdateNumPreviews()
+void UPlaneCutTool::SetupPreviews()
 {
 	int32 CurrentNumPreview = Previews.Num();
-	int32 NumSourceMeshes = OriginalDynamicMeshes.Num();
-	int32 TargetNumPreview = (BasicProperties->bKeepBothHalves ? 2 : 1) * NumSourceMeshes;
-	if (TargetNumPreview < CurrentNumPreview)
+	int32 NumSourceMeshes = MeshesToCut.Num();
+	int32 TargetNumPreview = NumSourceMeshes;
+	for (int32 PreviewIdx = CurrentNumPreview; PreviewIdx < TargetNumPreview; PreviewIdx++)
 	{
-		for (int32 PreviewIdx = CurrentNumPreview - 1; PreviewIdx >= TargetNumPreview; PreviewIdx--)
-		{
-			Previews[PreviewIdx]->Cancel();
-		}
-		Previews.SetNum(TargetNumPreview);
+		UPlaneCutOperatorFactory *CutSide = NewObject<UPlaneCutOperatorFactory>();
+		CutSide->CutTool = this;
+		CutSide->ComponentIndex = PreviewIdx;
+		UMeshOpPreviewWithBackgroundCompute* Preview = Previews.Add_GetRef(NewObject<UMeshOpPreviewWithBackgroundCompute>(CutSide, "Preview"));
+		Preview->Setup(this->TargetWorld, CutSide);
+
+		FComponentMaterialSet MaterialSet;
+		ComponentTargets[PreviewIdx]->GetMaterialSet(MaterialSet);
+		Preview->ConfigureMaterials(MaterialSet.Materials,
+			ToolSetupUtil::GetDefaultWorkingMaterial(GetToolManager())
+		);
+
+		// set initial preview to un-processed mesh, so stuff doesn't just disappear if the first cut takes a while
+		Preview->PreviewMesh->UpdatePreview(MeshesToCut[PreviewIdx]->GetMesh().Get());
+		Preview->PreviewMesh->SetTransform(ComponentTargets[PreviewIdx]->GetWorldTransform());
+		Preview->SetVisibility(BasicProperties->bShowPreview);
 	}
-	else
+}
+
+
+
+void UPlaneCutTool::DoCut()
+{
+	if (!CanAccept())
 	{
-		for (int32 PreviewIdx = CurrentNumPreview; PreviewIdx < TargetNumPreview; PreviewIdx++)
-		{
-			UPlaneCutOperatorFactory *CutSide = NewObject<UPlaneCutOperatorFactory>();
-			CutSide->CutTool = this;
-			CutSide->bCutBackSide = PreviewIdx >= NumSourceMeshes;
-			int32 SrcIdx = PreviewIdx % NumSourceMeshes;
-			CutSide->ComponentIndex = SrcIdx;
-			UMeshOpPreviewWithBackgroundCompute* Preview = Previews.Add_GetRef(NewObject<UMeshOpPreviewWithBackgroundCompute>(CutSide, "Preview"));
-			Preview->Setup(this->TargetWorld, CutSide);
-
-			FComponentMaterialSet MaterialSet;
-			ComponentTargets[SrcIdx]->GetMaterialSet(MaterialSet);
-			Preview->ConfigureMaterials(MaterialSet.Materials,
-				ToolSetupUtil::GetDefaultWorkingMaterial(GetToolManager())
-			);
-
-			// set initial preview to un-processed mesh, so stuff doesn't just disappear if the first cut takes a while
-			Preview->PreviewMesh->UpdatePreview(OriginalDynamicMeshes[SrcIdx].Get());
-			Preview->PreviewMesh->SetTransform(ComponentTargets[SrcIdx]->GetWorldTransform());
-			Preview->SetVisibility(BasicProperties->bShowPreview);
-			
-		}
+		return;
+	}
+
+	
+
+	TUniquePtr<FToolCommandChangeSequence> ChangeSeq = MakeUnique<FToolCommandChangeSequence>();
+
+	TArray<FDynamicMeshOpResult> Results;
+	for (int Idx = 0, N = MeshesToCut.Num(); Idx < N; Idx++)
+	{
+		UMeshOpPreviewWithBackgroundCompute* Preview = Previews[Idx];
+		TUniquePtr<FDynamicMesh3> ResultMesh = Preview->PreviewMesh->ExtractPreviewMesh();
+		ChangeSeq->AppendChange(MeshesToCut[Idx], MeshesToCut[Idx]->ReplaceMesh(
+			TSharedPtr<const FDynamicMesh3>(ResultMesh.Release()))
+		);
+	}
+
+	// emit combined change sequence
+	GetToolManager()->EmitObjectChange(this, MoveTemp(ChangeSeq), LOCTEXT("MeshPlaneCut", "Cut Mesh with Plane"));
+	for (UMeshOpPreviewWithBackgroundCompute* Preview : Previews)
+	{
+		Preview->InvalidateResult();
 	}
 }
 
 
+
 void UPlaneCutTool::Shutdown(EToolShutdownType ShutdownType)
 {
 	// Restore (unhide) the source meshes
@@ -233,7 +272,6 @@ void UPlaneCutTool::SetAssetAPI(IToolsContextAssetAPI* AssetAPIIn)
 TUniquePtr<FDynamicMeshOperator> UPlaneCutOperatorFactory::MakeNewOperator()
 {
 	TUniquePtr<FPlaneCutOp> CutOp = MakeUnique<FPlaneCutOp>();
-	CutOp->bDiscardAttributes = CutTool->BasicProperties->bDiscardAttributes;
 	CutOp->bFillCutHole = CutTool->BasicProperties->bFillCutHole;
 	CutOp->bFillSpans = CutTool->BasicProperties->bFillSpans;
 
@@ -242,18 +280,13 @@ TUniquePtr<FDynamicMeshOperator> UPlaneCutOperatorFactory::MakeNewOperator()
 	FVector LocalOrigin = WorldToLocal.TransformPosition(CutTool->CutPlaneOrigin);
 	FVector WorldNormal = CutTool->CutPlaneOrientation.GetAxisZ();
 	FVector LocalNormal = WorldToLocal.TransformVectorNoScale(WorldNormal); // TODO: correct for nonuniform scaling?
-	if (bCutBackSide)
-	{
-		LocalNormal = -LocalNormal;
-	}
 	CutOp->LocalPlaneOrigin = LocalOrigin;
 	CutOp->LocalPlaneNormal = LocalNormal;
-	CutOp->OriginalMesh = CutTool->OriginalDynamicMeshes[ComponentIndex];
-	
-	if (bCutBackSide)
-	{
-		LocalToWorld *= FTransform(CutTool->BasicProperties->SpacingBetweenHalves * WorldNormal);
-	}
+	CutOp->OriginalMesh = CutTool->MeshesToCut[ComponentIndex]->GetMesh();
+	CutOp->bKeepBothHalves = CutTool->BasicProperties->bKeepBothHalves;
+	CutOp->OtherHalfOffsetDistance = CutTool->BasicProperties->SpacingBetweenHalves;
+	CutOp->UVScaleFactor = CutTool->MeshUVScaleFactor[ComponentIndex];
+	CutOp->SubObjectsAttribIndex = CutTool->MeshSubObjectAttribIndices[ComponentIndex];
 	CutOp->SetTransform(LocalToWorld);
 
 	return CutOp;
@@ -285,7 +318,6 @@ void UPlaneCutTool::Tick(float DeltaTime)
 #if WITH_EDITOR
 void UPlaneCutTool::PostEditChangeProperty(FPropertyChangedEvent& PropertyChangedEvent)
 {
-	UpdateNumPreviews();
 	for (UMeshOpPreviewWithBackgroundCompute* Preview : Previews)
 	{
 		Preview->InvalidateResult();
@@ -307,7 +339,6 @@ void UPlaneCutTool::OnPropertyModified(UObject* PropertySet, UProperty* Property
 		}
 	}
 
-	UpdateNumPreviews();
 	for (UMeshOpPreviewWithBackgroundCompute* Preview : Previews)
 	{
 		Preview->InvalidateResult();
@@ -317,7 +348,6 @@ void UPlaneCutTool::OnPropertyModified(UObject* PropertySet, UProperty* Property
 
 void UPlaneCutTool::TransformChanged(UTransformProxy* Proxy, FTransform Transform)
 {
-	// TODO: if multi-select is re-enabled, only invalidate the preview that actually needs it?
 	CutPlaneOrientation = Transform.GetRotation();
 	CutPlaneOrigin = (FVector)Transform.GetTranslation();
 	for (UMeshOpPreviewWithBackgroundCompute* Preview : Previews)
@@ -359,63 +389,117 @@ bool UPlaneCutTool::CanAccept() const
 
 void UPlaneCutTool::GenerateAsset(const TArray<FDynamicMeshOpResult>& Results)
 {
-	if (Results.Num() == 0 
-		|| Results[0].Mesh.IsValid() == false
-		|| Results[0].Mesh->TriangleCount() == 0 )
+	if (Results.Num() == 0)
 	{
 		return;
 	}
 
+	for (int ResIdx = 0; ResIdx < Results.Num(); ResIdx++)
+	{
+		if (   Results[ResIdx].Mesh.IsValid() == false
+			|| Results[ResIdx].Mesh->TriangleCount() == 0)
+		{
+			// TODO handle the "some meshes are fully deleted" case in some better way than this.
+			// (we cannot currently save them as empty assets, as that will break undo and possibly other downstream things that make bad assumptions)
+			return;
+		}
+	}
+
 	GetToolManager()->BeginUndoTransaction(LOCTEXT("PlaneCutToolTransactionName", "Plane Cut Tool"));
 	
 
 	// currently in-place replaces the first half, and adds a new actor for the second half (if it was generated)
 	// TODO: options to support other choices re what should be a new actor
+
 	ensure(Results.Num() > 0);
-	int32 NumSourceMeshes = OriginalDynamicMeshes.Num();
+	int32 NumSourceMeshes = MeshesToCut.Num();
+	TArray<TArray<FDynamicMesh3>> AllSplitMeshes;
+	
+
 	for (int OrigMeshIdx = 0; OrigMeshIdx < NumSourceMeshes; OrigMeshIdx++)
 	{
-		check(Results[OrigMeshIdx].Mesh.Get() != nullptr);
-		ComponentTargets[OrigMeshIdx]->CommitMesh([&Results, OrigMeshIdx](const FPrimitiveComponentTarget::FCommitParams& CommitParams)
+		FDynamicMesh3* UseMesh = Results[OrigMeshIdx].Mesh.Get();
+		check(UseMesh != nullptr);
+
+		TDynamicMeshScalarTriangleAttribute<int>* SubMeshIDs = 
+			static_cast<TDynamicMeshScalarTriangleAttribute<int>*>(UseMesh->Attributes()->GetAttachedAttribute(
+				MeshSubObjectAttribIndices[OrigMeshIdx]));
+		TArray<FDynamicMesh3>& SplitMeshes = AllSplitMeshes.Emplace_GetRef();
+		bool bWasSplit = FDynamicMeshEditor::SplitMesh(UseMesh, SplitMeshes, [SubMeshIDs](int TID)
+		{
+			return SubMeshIDs->GetValue(TID);
+		}
+		);
+		if (bWasSplit)
+		{
+			// split mesh did something but has no meshes in the output array??
+			if (!ensure(SplitMeshes.Num() > 0))
+			{
+				continue;
+			}
+			UseMesh = &SplitMeshes[0];
+		}
+
+		ComponentTargets[OrigMeshIdx]->CommitMesh([&UseMesh](const FPrimitiveComponentTarget::FCommitParams& CommitParams)
 		{
 			FDynamicMeshToMeshDescription Converter;
-			Converter.Convert(Results[OrigMeshIdx].Mesh.Get(), *CommitParams.MeshDescription);
+			Converter.Convert(UseMesh, *CommitParams.MeshDescription);
 		});
 	}
 
-	// The method for creating a new mesh (AssetGenerationUtil::GenerateStaticMeshActor) is editor-only; just creating the other half if not in editor
+	// The method for creating a new mesh (AssetGenerationUtil::GenerateStaticMeshActor) is editor-only
 #if WITH_EDITOR
-	
-	if (Results.Num() > NumSourceMeshes)
+	bool bNeedToAdd = false;
+	for (int MeshIdx = 0; MeshIdx < NumSourceMeshes; MeshIdx++)
 	{
-		ensure(Results.Num() == NumSourceMeshes * 2);
+		TArray<FDynamicMesh3>& SplitMeshes = AllSplitMeshes[MeshIdx];
+		if (SplitMeshes.Num() < 2)
+		{
+			continue;
+		}
 
+		bNeedToAdd = true;
+	}
+
+	if (bNeedToAdd)
+	{
 		FSelectedOjectsChangeList NewSelection;
 		NewSelection.ModificationType = ESelectedObjectsModificationType::Replace;
 		for (TUniquePtr<FPrimitiveComponentTarget>& ComponentTarget : ComponentTargets)
 		{
 			NewSelection.Actors.Add(ComponentTarget->GetOwnerActor());
 		}
-		
-		for (int32 AddedMeshIdx = NumSourceMeshes; AddedMeshIdx < Results.Num(); AddedMeshIdx++)
+
+		for (int OrigMeshIdx = 0; OrigMeshIdx < NumSourceMeshes; OrigMeshIdx++)
 		{
-			check(Results[AddedMeshIdx].Mesh.Get() != nullptr);
+			TArray<FDynamicMesh3>& SplitMeshes = AllSplitMeshes[OrigMeshIdx];
+			if (SplitMeshes.Num() < 2)
+			{
+				continue;
+			}
 
+			// build array of materials from the original
 			TArray<UMaterialInterface*> Materials;
-			TUniquePtr<FPrimitiveComponentTarget>& ComponentTarget = ComponentTargets[AddedMeshIdx - NumSourceMeshes];
+			TUniquePtr<FPrimitiveComponentTarget>& ComponentTarget = ComponentTargets[OrigMeshIdx];
 			for (int MaterialIdx = 0, NumMaterials = ComponentTarget->GetNumMaterials(); MaterialIdx < NumMaterials; MaterialIdx++)
 			{
 				Materials.Add(ComponentTarget->GetMaterial(MaterialIdx));
 			}
 
-			AActor* NewActor = AssetGenerationUtil::GenerateStaticMeshActor(
-				AssetAPI, TargetWorld,
-				Results[AddedMeshIdx].Mesh.Get(), Results[AddedMeshIdx].Transform, TEXT("PlaneCutOtherHalf"),
-				AssetGenerationUtil::GetDefaultAutoGeneratedAssetPath(), Materials);
-			NewSelection.Actors.Add(NewActor);
+			// add all the additional meshes
+			for (int AddMeshIdx = 1; AddMeshIdx < SplitMeshes.Num(); AddMeshIdx++)
+			{
+				AActor* NewActor = AssetGenerationUtil::GenerateStaticMeshActor(
+					AssetAPI, TargetWorld,
+					&SplitMeshes[AddMeshIdx], Results[OrigMeshIdx].Transform, TEXT("PlaneCutOtherPart"),
+					AssetGenerationUtil::GetDefaultAutoGeneratedAssetPath(), Materials);
+				NewSelection.Actors.Add(NewActor);
+			}
 		}
+
 		GetToolManager()->RequestSelectionChange(NewSelection);
 	}
+
 #endif
 
 	GetToolManager()->EndUndoTransaction();
diff --git a/Engine/Plugins/Experimental/MeshModelingToolset/Source/MeshModelingTools/Public/PlaneCutTool.h b/Engine/Plugins/Experimental/MeshModelingToolset/Source/MeshModelingTools/Public/PlaneCutTool.h
--- a/Engine/Plugins/Experimental/MeshModelingToolset/Source/MeshModelingTools/Public/PlaneCutTool.h
+++ b/Engine/Plugins/Experimental/MeshModelingToolset/Source/MeshModelingTools/Public/PlaneCutTool.h
@@ -9,6 +9,7 @@
 #include "InteractiveToolBuilder.h"
 #include "MeshOpPreviewHelpers.h"
 #include "DynamicMesh3.h"
+#include "Changes/DynamicMeshChangeTarget.h"
 #include "BaseTools/SingleClickTool.h"
 #include "PlaneCutTool.generated.h"
 
@@ -53,10 +54,6 @@ class MESHMODELINGTOOLS_API UPlaneCutToolProperties : public UInteractiveToolPro
 public:
 	UPlaneCutToolProperties();
 
-	/** If true, UVs and Normals are discarded  */
-	UPROPERTY(EditAnywhere, Category = Options)
-	bool bDiscardAttributes;
-
 	/** If true, both halves of the cut are computed */
 	UPROPERTY(EditAnywhere, Category = Options)
 	bool bKeepBothHalves;
@@ -91,9 +88,6 @@ public:
 	UPROPERTY()
 	UPlaneCutTool *CutTool;
 
-	UPROPERTY()
-	bool bCutBackSide = false;
-
 	int ComponentIndex;
 };
 
@@ -130,6 +124,7 @@ public:
 
 	virtual void OnPropertyModified(UObject* PropertySet, UProperty* Property) override;
 
+
 protected:
 
 	UPROPERTY()
@@ -146,9 +141,19 @@ protected:
 	UPROPERTY()
 	TArray<UMeshOpPreviewWithBackgroundCompute*> Previews;
 
+	/** Do the plane cut without exiting the tool, useful for doing a lot of cuts quickly */
+	UFUNCTION(CallInEditor, Category = Options, meta = (DisplayName = "Cut"))
+	void DoCut();
 
 protected:
-	TArray<TSharedPtr<FDynamicMesh3>> OriginalDynamicMeshes;
+
+	UPROPERTY()
+	TArray<UDynamicMeshReplacementChangeTarget*> MeshesToCut;
+
+	// for each mesh in MeshesToCut, the index of the attached generic triangle attribute tracking the object index
+	TArray<int> MeshSubObjectAttribIndices;
+	// UV Scale factor is cached based on the bounding box of the mesh before any cuts are performed, so you don't get inconsistent UVs if you multi-cut the object to smaller sizes
+	TArray<float> MeshUVScaleFactor;
 
 	UWorld* TargetWorld;
 	IToolsContextAssetAPI* AssetAPI;
@@ -162,8 +167,9 @@ protected:
 	UTransformProxy* PlaneTransformProxy;
 
 	void TransformChanged(UTransformProxy* Proxy, FTransform Transform);
+	void MeshChanged();
 
-	void UpdateNumPreviews();
+	void SetupPreviews();
 
 	IClickBehaviorTarget* SetPointInWorldConnector = nullptr;
 
diff --git a/Engine/Plugins/Experimental/MeshModelingToolset/Source/ModelingComponents/Private/Changes/DynamicMeshChangeTarget.cpp b/Engine/Plugins/Experimental/MeshModelingToolset/Source/ModelingComponents/Private/Changes/DynamicMeshChangeTarget.cpp
new file mode 100644
--- /dev/null
+++ b/Engine/Plugins/Experimental/MeshModelingToolset/Source/ModelingComponents/Private/Changes/DynamicMeshChangeTarget.cpp
@@ -0,0 +1,18 @@
+// Copyright 1998-2019 Epic Games, Inc. All Rights Reserved.
+
+#include "Changes/DynamicMeshChangeTarget.h"
+#include "DynamicMesh3.h"
+
+
+void UDynamicMeshReplacementChangeTarget::ApplyChange(const FMeshReplacementChange* Change, bool bRevert)
+{
+	Mesh = Change->GetMesh(bRevert);
+	OnMeshChanged.Broadcast();
+}
+
+TUniquePtr<FMeshReplacementChange> UDynamicMeshReplacementChangeTarget::ReplaceMesh(const TSharedPtr<const FDynamicMesh3>& UpdateMesh)
+{
+	TUniquePtr<FMeshReplacementChange> Change = MakeUnique<FMeshReplacementChange>(Mesh, UpdateMesh);
+	Mesh = UpdateMesh;
+	return Change;
+}
\ No newline at end of file
diff --git a/Engine/Plugins/Experimental/MeshModelingToolset/Source/ModelingComponents/Private/Changes/MeshReplacementChange.cpp b/Engine/Plugins/Experimental/MeshModelingToolset/Source/ModelingComponents/Private/Changes/MeshReplacementChange.cpp
new file mode 100644
--- /dev/null
+++ b/Engine/Plugins/Experimental/MeshModelingToolset/Source/ModelingComponents/Private/Changes/MeshReplacementChange.cpp
@@ -0,0 +1,43 @@
+// Copyright 1998-2019 Epic Games, Inc. All Rights Reserved.
+
+#include "Changes/MeshReplacementChange.h"
+#include "DynamicMesh3.h"
+
+FMeshReplacementChange::FMeshReplacementChange()
+{
+}
+
+FMeshReplacementChange::FMeshReplacementChange(TSharedPtr<const FDynamicMesh3> BeforeIn, TSharedPtr<const FDynamicMesh3> AfterIn)
+{
+	Before = BeforeIn;
+	After = AfterIn;
+}
+
+void FMeshReplacementChange::Apply(UObject* Object)
+{
+	IMeshReplacementCommandChangeTarget* ChangeTarget = CastChecked<IMeshReplacementCommandChangeTarget>(Object);
+	ChangeTarget->ApplyChange(this, false);
+
+	if (OnChangeAppliedFunc)
+	{
+		OnChangeAppliedFunc(this, Object, true);
+	}
+}
+
+void FMeshReplacementChange::Revert(UObject* Object)
+{
+	IMeshReplacementCommandChangeTarget* ChangeTarget = CastChecked<IMeshReplacementCommandChangeTarget>(Object);
+	ChangeTarget->ApplyChange(this, true);
+
+	if (OnChangeAppliedFunc)
+	{
+		OnChangeAppliedFunc(this, Object, false);
+	}
+}
+
+
+FString FMeshReplacementChange::ToString() const
+{
+	return FString(TEXT("Mesh Change"));
+}
+
diff --git a/Engine/Plugins/Experimental/MeshModelingToolset/Source/ModelingComponents/Private/OctreeDynamicMeshComponent.cpp b/Engine/Plugins/Experimental/MeshModelingToolset/Source/ModelingComponents/Private/OctreeDynamicMeshComponent.cpp
--- a/Engine/Plugins/Experimental/MeshModelingToolset/Source/ModelingComponents/Private/OctreeDynamicMeshComponent.cpp
+++ b/Engine/Plugins/Experimental/MeshModelingToolset/Source/ModelingComponents/Private/OctreeDynamicMeshComponent.cpp
@@ -408,4 +408,19 @@ void UOctreeDynamicMeshComponent::ApplyChange(const FMeshChange* Change, bool bR
 	//NotifyMeshUpdated();
 
 	OnMeshChanged.Broadcast();
-}
\ No newline at end of file
+}
+
+
+void UOctreeDynamicMeshComponent::ApplyChange(const FMeshReplacementChange* Change, bool bRevert)
+{
+	// full clear and reset
+	Mesh->Clear();
+	Mesh->Copy(*Change->GetMesh(bRevert));
+	Octree = MakeUnique<FDynamicMeshOctree3>();
+	Octree->Initialize(GetMesh());
+	OctreeCut = MakeUnique<FDynamicMeshOctree3::FTreeCutSet>();
+
+	//NotifyMeshUpdated();
+
+	OnMeshChanged.Broadcast();
+}
diff --git a/Engine/Plugins/Experimental/MeshModelingToolset/Source/ModelingComponents/Private/PreviewMesh.cpp b/Engine/Plugins/Experimental/MeshModelingToolset/Source/ModelingComponents/Private/PreviewMesh.cpp
--- a/Engine/Plugins/Experimental/MeshModelingToolset/Source/ModelingComponents/Private/PreviewMesh.cpp
+++ b/Engine/Plugins/Experimental/MeshModelingToolset/Source/ModelingComponents/Private/PreviewMesh.cpp
@@ -321,6 +321,15 @@ void UPreviewMesh::ApplyChange(const FMeshChange* Change, bool bRevert)
 		MeshAABBTree.SetMesh(DynamicMeshComponent->GetMesh(), true);
 	}
 }
+void UPreviewMesh::ApplyChange(const FMeshReplacementChange* Change, bool bRevert)
+{
+	check(DynamicMeshComponent != nullptr);
+	DynamicMeshComponent->ApplyChange(Change, bRevert);
+	if (bBuildSpatialDataStructure)
+	{
+		MeshAABBTree.SetMesh(DynamicMeshComponent->GetMesh(), true);
+	}
+}
 
 FSimpleMulticastDelegate& UPreviewMesh::GetOnMeshChanged()
 {
diff --git a/Engine/Plugins/Experimental/MeshModelingToolset/Source/ModelingComponents/Private/SimpleDynamicMeshComponent.cpp b/Engine/Plugins/Experimental/MeshModelingToolset/Source/ModelingComponents/Private/SimpleDynamicMeshComponent.cpp
--- a/Engine/Plugins/Experimental/MeshModelingToolset/Source/ModelingComponents/Private/SimpleDynamicMeshComponent.cpp
+++ b/Engine/Plugins/Experimental/MeshModelingToolset/Source/ModelingComponents/Private/SimpleDynamicMeshComponent.cpp
@@ -275,6 +275,15 @@ void USimpleDynamicMeshComponent::ApplyChange(const FMeshChange* Change, bool bR
 {
 	Change->DynamicMeshChange->Apply(Mesh.Get(), bRevert);
 
+	NotifyMeshUpdated();
+	OnMeshChanged.Broadcast();
+}
+
+
+void USimpleDynamicMeshComponent::ApplyChange(const FMeshReplacementChange* Change, bool bRevert)
+{
+	Mesh->Copy(*Change->GetMesh(bRevert));
+
 	NotifyMeshUpdated();
 	OnMeshChanged.Broadcast();
 }
\ No newline at end of file
diff --git a/Engine/Plugins/Experimental/MeshModelingToolset/Source/ModelingComponents/Public/BaseDynamicMeshComponent.h b/Engine/Plugins/Experimental/MeshModelingToolset/Source/ModelingComponents/Public/BaseDynamicMeshComponent.h
--- a/Engine/Plugins/Experimental/MeshModelingToolset/Source/ModelingComponents/Public/BaseDynamicMeshComponent.h
+++ b/Engine/Plugins/Experimental/MeshModelingToolset/Source/ModelingComponents/Public/BaseDynamicMeshComponent.h
@@ -8,6 +8,7 @@
 #include "InteractiveToolObjects.h"
 #include "Changes/MeshVertexChange.h"
 #include "Changes/MeshChange.h"
+#include "Changes/MeshReplacementChange.h"
 
 #include "DynamicMesh3.h"
 
@@ -41,7 +42,7 @@ enum class EDynamicMeshTangentCalcType : uint8
  * Currently no functionality lives here, only some interface functions are defined that various subclasses implement.
  */
 UCLASS(hidecategories = (LOD, Physics, Collision), editinlinenew, ClassGroup = Rendering)
-class MODELINGCOMPONENTS_API UBaseDynamicMeshComponent : public UMeshComponent, public IToolFrameworkComponent, public IMeshVertexCommandChangeTarget, public IMeshCommandChangeTarget
+class MODELINGCOMPONENTS_API UBaseDynamicMeshComponent : public UMeshComponent, public IToolFrameworkComponent, public IMeshVertexCommandChangeTarget, public IMeshCommandChangeTarget, public IMeshReplacementCommandChangeTarget
 {
 	GENERATED_UCLASS_BODY()
 
@@ -72,6 +73,14 @@ public:
 		unimplemented();
 	}
 
+	/**
+	* Apply a full mesh replacement change to the internal mesh  (implements IMeshReplacementCommandChangeTarget)
+	*/
+	virtual void ApplyChange(const FMeshReplacementChange* Change, bool bRevert) override
+	{
+		unimplemented();
+	}
+
 
 protected:
 	/**
diff --git a/Engine/Plugins/Experimental/MeshModelingToolset/Source/ModelingComponents/Public/Changes/DynamicMeshChangeTarget.h b/Engine/Plugins/Experimental/MeshModelingToolset/Source/ModelingComponents/Public/Changes/DynamicMeshChangeTarget.h
new file mode 100644
--- /dev/null
+++ b/Engine/Plugins/Experimental/MeshModelingToolset/Source/ModelingComponents/Public/Changes/DynamicMeshChangeTarget.h
@@ -0,0 +1,65 @@
+// Copyright 1998-2019 Epic Games, Inc. All Rights Reserved.
+
+#pragma once
+
+#include "CoreMinimal.h"
+#include "UObject/Interface.h"
+#include "InteractiveToolChange.h"
+
+#include "Changes/MeshChange.h"
+#include "Changes/MeshVertexChange.h"
+#include "Changes/MeshReplacementChange.h"
+
+#include "DynamicMeshChangeTarget.generated.h"
+
+
+class FDynamicMesh3;
+
+/**
+ * Bare bones wrapper of FDynamicMesh3 that supports MeshReplacementChange-based updates.  Shares the mesh ptr with the MeshReplacementChange objects, so it must not be changed directly
+ * TODO: also support other MeshChange objects (by making a copy of the mesh when applying these changes)
+ */
+UCLASS(Transient)
+class MODELINGCOMPONENTS_API UDynamicMeshReplacementChangeTarget : public UObject, public IMeshReplacementCommandChangeTarget
+{
+	GENERATED_BODY()
+
+protected:
+	TSharedPtr<const FDynamicMesh3> Mesh;
+
+public:
+	UDynamicMeshReplacementChangeTarget() = default;
+	UDynamicMeshReplacementChangeTarget(TSharedPtr<const FDynamicMesh3> Mesh) : Mesh(Mesh)
+	{}
+	virtual ~UDynamicMeshReplacementChangeTarget() = default;
+
+	/**
+	* Apply/Revert a general mesh change to the internal mesh   (implements IMeshReplacementCommandChangeTarget)
+	*/
+	virtual void ApplyChange(const FMeshReplacementChange* Change, bool bRevert) override;
+
+	virtual void SetMesh(TSharedPtr<const FDynamicMesh3> MeshIn)
+	{
+		Mesh = MeshIn;
+	}
+
+	const TSharedPtr<const FDynamicMesh3>& GetMesh() const
+	{
+		return Mesh;
+	}
+
+	/**
+	 * Update the mesh and return a corresponding FChange
+	 * Does not emit the change.
+	 * Does not fire OnMeshChanged.
+	 */
+	TUniquePtr<FMeshReplacementChange> ReplaceMesh(const TSharedPtr<const FDynamicMesh3>& UpdateMesh);
+
+	/**
+	* This delegate fires when a FCommandChange is applied to this component, so that
+	* parent objects know the mesh has changed.
+	*/
+	FSimpleMulticastDelegate OnMeshChanged;
+
+};
+
diff --git a/Engine/Plugins/Experimental/MeshModelingToolset/Source/ModelingComponents/Public/Changes/MeshChange.h b/Engine/Plugins/Experimental/MeshModelingToolset/Source/ModelingComponents/Public/Changes/MeshChange.h
--- a/Engine/Plugins/Experimental/MeshModelingToolset/Source/ModelingComponents/Public/Changes/MeshChange.h
+++ b/Engine/Plugins/Experimental/MeshModelingToolset/Source/ModelingComponents/Public/Changes/MeshChange.h
@@ -16,7 +16,7 @@
 
 /**
  * FMeshChange represents an undoable change to a FDynamicMesh3.
- * Currently only valid to call Apply/Revert when the Object is a USimpleDynamicMeshComponent
+ * Currently only valid to call Apply/Revert when the Object is a one of several components backed by FDynamicMesh: USimpleDynamicMeshComponent, UOctreeDynamicMeshComponent, UPreviewMesh
  */
 class MODELINGCOMPONENTS_API FMeshChange : public FToolCommandChange
 {
diff --git a/Engine/Plugins/Experimental/MeshModelingToolset/Source/ModelingComponents/Public/Changes/MeshReplacementChange.h b/Engine/Plugins/Experimental/MeshModelingToolset/Source/ModelingComponents/Public/Changes/MeshReplacementChange.h
new file mode 100644
--- /dev/null
+++ b/Engine/Plugins/Experimental/MeshModelingToolset/Source/ModelingComponents/Public/Changes/MeshReplacementChange.h
@@ -0,0 +1,65 @@
+// Copyright 1998-2019 Epic Games, Inc. All Rights Reserved.
+
+#pragma once
+
+#include "CoreMinimal.h"
+#include "UObject/Interface.h"
+#include "InteractiveToolChange.h"
+#include "MeshReplacementChange.generated.h"
+
+
+class FDynamicMesh3;
+
+
+
+
+/**
+ * FMeshReplacementChange represents an undoable *complete* change to a FDynamicMesh3.
+ * Currently only valid to call Apply/Revert when the Object is a USimpleDynamicMeshComponent
+ */
+class MODELINGCOMPONENTS_API FMeshReplacementChange : public FToolCommandChange
+{
+	TSharedPtr<const FDynamicMesh3> Before, After;
+
+public:
+	FMeshReplacementChange();
+	FMeshReplacementChange(TSharedPtr<const FDynamicMesh3> Before, TSharedPtr<const FDynamicMesh3> After);
+
+	const TSharedPtr<const FDynamicMesh3>& GetMesh(bool bRevert) const
+	{
+		return bRevert ? Before : After;
+	}
+
+	/** This function is called on Apply and Revert (last argument is true on Apply)*/
+	TFunction<void(FMeshReplacementChange*, UObject*, bool)> OnChangeAppliedFunc;
+
+	/** Makes the change to the object */
+	virtual void Apply(UObject* Object) override;
+
+	/** Reverts change to the object */
+	virtual void Revert(UObject* Object) override;
+
+	/** Describes this change (for debugging) */
+	virtual FString ToString() const override;
+};
+
+
+
+
+
+UINTERFACE()
+class MODELINGCOMPONENTS_API UMeshReplacementCommandChangeTarget : public UInterface
+{
+	GENERATED_BODY()
+};
+/**
+ * IMeshReplacementCommandChangeTarget is an interface which is used to apply a mesh replacement change
+ */
+class MODELINGCOMPONENTS_API IMeshReplacementCommandChangeTarget
+{
+	GENERATED_BODY()
+public:
+	virtual void ApplyChange(const FMeshReplacementChange* Change, bool bRevert) = 0;
+};
+
+
diff --git a/Engine/Plugins/Experimental/MeshModelingToolset/Source/ModelingComponents/Public/OctreeDynamicMeshComponent.h b/Engine/Plugins/Experimental/MeshModelingToolset/Source/ModelingComponents/Public/OctreeDynamicMeshComponent.h
--- a/Engine/Plugins/Experimental/MeshModelingToolset/Source/ModelingComponents/Public/OctreeDynamicMeshComponent.h
+++ b/Engine/Plugins/Experimental/MeshModelingToolset/Source/ModelingComponents/Public/OctreeDynamicMeshComponent.h
@@ -98,6 +98,11 @@ public:
 	 */
 	virtual void ApplyChange(const FMeshChange* Change, bool bRevert) override;
 
+	/**
+	* Apply a general mesh replacement change to the internal mesh
+	*/
+	virtual void ApplyChange(const FMeshReplacementChange* Change, bool bRevert) override;
+
 
 	/**
 	 * This delegate fires when a FCommandChange is applied to this component, so that
diff --git a/Engine/Plugins/Experimental/MeshModelingToolset/Source/ModelingComponents/Public/PreviewMesh.h b/Engine/Plugins/Experimental/MeshModelingToolset/Source/ModelingComponents/Public/PreviewMesh.h
--- a/Engine/Plugins/Experimental/MeshModelingToolset/Source/ModelingComponents/Public/PreviewMesh.h
+++ b/Engine/Plugins/Experimental/MeshModelingToolset/Source/ModelingComponents/Public/PreviewMesh.h
@@ -53,7 +53,7 @@ public:
  * 
  */
 UCLASS(Transient)
-class MODELINGCOMPONENTS_API UPreviewMesh : public UObject, public IMeshVertexCommandChangeTarget, public IMeshCommandChangeTarget
+class MODELINGCOMPONENTS_API UPreviewMesh : public UObject, public IMeshVertexCommandChangeTarget, public IMeshCommandChangeTarget, public IMeshReplacementCommandChangeTarget
 {
 	GENERATED_BODY()
 
@@ -275,6 +275,11 @@ public:
 	 */
 	virtual void ApplyChange(const FMeshChange* Change, bool bRevert) override;
 
+	/**
+	* Apply/Revert a general mesh change to the internal mesh   (implements IMeshReplacementCommandChangeTarget)
+	*/
+	virtual void ApplyChange(const FMeshReplacementChange* Change, bool bRevert) override;
+
 	/** @return delegate that is broadcast whenever the internal mesh component is changed */
 	FSimpleMulticastDelegate& GetOnMeshChanged();
 
diff --git a/Engine/Plugins/Experimental/MeshModelingToolset/Source/ModelingComponents/Public/SimpleDynamicMeshComponent.h b/Engine/Plugins/Experimental/MeshModelingToolset/Source/ModelingComponents/Public/SimpleDynamicMeshComponent.h
--- a/Engine/Plugins/Experimental/MeshModelingToolset/Source/ModelingComponents/Public/SimpleDynamicMeshComponent.h
+++ b/Engine/Plugins/Experimental/MeshModelingToolset/Source/ModelingComponents/Public/SimpleDynamicMeshComponent.h
@@ -113,6 +113,11 @@ public:
 	 */
 	virtual void ApplyChange(const FMeshChange* Change, bool bRevert) override;
 
+	/**
+	* Apply a general mesh replacement change to the internal mesh
+	*/
+	virtual void ApplyChange(const FMeshReplacementChange* Change, bool bRevert) override;
+
 
 	/**
 	 * This delegate fires when a FCommandChange is applied to this component, so that
diff --git a/Engine/Plugins/Experimental/MeshModelingToolset/Source/ModelingOperators/Private/CuttingOps/PlaneCutOp.cpp b/Engine/Plugins/Experimental/MeshModelingToolset/Source/ModelingOperators/Private/CuttingOps/PlaneCutOp.cpp
--- a/Engine/Plugins/Experimental/MeshModelingToolset/Source/ModelingOperators/Private/CuttingOps/PlaneCutOp.cpp
+++ b/Engine/Plugins/Experimental/MeshModelingToolset/Source/ModelingOperators/Private/CuttingOps/PlaneCutOp.cpp
@@ -21,22 +21,35 @@ void FPlaneCutOp::CalculateResult(FProgressCancel* Progress)
 	{
 		return;
 	}
-	ResultMesh->Copy(*OriginalMesh, true, true, true, !bDiscardAttributes);
-
+	ResultMesh->Copy(*OriginalMesh, true, true, true, true);
+	
 	FMeshPlaneCut Cut(ResultMesh.Get(), LocalPlaneOrigin, LocalPlaneNormal);
 
 	if (Progress->Cancelled())
 	{
 		return;
 	}
-	Cut.UVScaleFactor = 1.0 / ResultMesh->GetBounds().MaxDim();
+	Cut.UVScaleFactor = UVScaleFactor;
 
 
 	if (Progress->Cancelled())
 	{
 		return;
 	}
-	Cut.Cut();
+	if (bKeepBothHalves)
+	{
+		int MaxSubObjectID = -1;
+		TDynamicMeshScalarTriangleAttribute<int>* SubObjectAttrib = static_cast<TDynamicMeshScalarTriangleAttribute<int>*>(ResultMesh->Attributes()->GetAttachedAttribute(SubObjectsAttribIndex));
+		for (int TID : ResultMesh->TriangleIndicesItr())
+		{
+			MaxSubObjectID = FMath::Max(MaxSubObjectID, SubObjectAttrib->GetValue(TID));
+		}
+		Cut.CutWithoutDelete(true, OtherHalfOffsetDistance, SubObjectAttrib, MaxSubObjectID+1);
+	}
+	else
+	{
+		Cut.Cut();
+	}
 
 
 	if (Progress->Cancelled())
@@ -47,4 +60,14 @@ void FPlaneCutOp::CalculateResult(FProgressCancel* Progress)
 	{
 		Cut.HoleFill(ConstrainedDelaunayTriangulate<double>, bFillSpans);
 	}
+
+	if (Progress->Cancelled())
+	{
+		return;
+	}
+	if (bFillCutHole && bKeepBothHalves)
+	{
+		TDynamicMeshScalarTriangleAttribute<int>* SubObjectAttrib = static_cast<TDynamicMeshScalarTriangleAttribute<int>*>(ResultMesh->Attributes()->GetAttachedAttribute(SubObjectsAttribIndex));
+		Cut.TransferTriangleLabelsToHoleFillTriangles(SubObjectAttrib);
+	}
 }
\ No newline at end of file
diff --git a/Engine/Plugins/Experimental/MeshModelingToolset/Source/ModelingOperators/Public/CuttingOps/PlaneCutOp.h b/Engine/Plugins/Experimental/MeshModelingToolset/Source/ModelingOperators/Public/CuttingOps/PlaneCutOp.h
--- a/Engine/Plugins/Experimental/MeshModelingToolset/Source/ModelingOperators/Public/CuttingOps/PlaneCutOp.h
+++ b/Engine/Plugins/Experimental/MeshModelingToolset/Source/ModelingOperators/Public/CuttingOps/PlaneCutOp.h
@@ -17,8 +17,12 @@ public:
 
 	// inputs
 	FVector LocalPlaneOrigin, LocalPlaneNormal;
-	bool bFillCutHole, bDiscardAttributes, bFillSpans;
-	TSharedPtr<FDynamicMesh3> OriginalMesh;
+	bool bFillCutHole, bFillSpans;
+	bool bKeepBothHalves;
+	float OtherHalfOffsetDistance = 0;
+	float UVScaleFactor;
+	int SubObjectsAttribIndex;
+	TSharedPtr<const FDynamicMesh3> OriginalMesh;
 
 	void SetTransform(const FTransform& Transform);
 
