ModelingTools: rename enum parameter based on artists feedback
#rnx

diff --git a/Engine/Plugins/Experimental/MeshModelingToolset/Source/MeshModelingTools/Private/AddPrimitiveTool.cpp b/Engine/Plugins/Experimental/MeshModelingToolset/Source/MeshModelingTools/Private/AddPrimitiveTool.cpp
--- a/Engine/Plugins/Experimental/MeshModelingToolset/Source/MeshModelingTools/Private/AddPrimitiveTool.cpp
+++ b/Engine/Plugins/Experimental/MeshModelingToolset/Source/MeshModelingTools/Private/AddPrimitiveTool.cpp
@@ -229,7 +229,7 @@ void UAddPrimitiveTool::UpdatePreviewPosition(const FInputDeviceRay& DeviceClick
 	// hit position (temp)
 	bool bHit = false;
 
-	if (ShapeSettings->PlaceMode == EMakeMeshPlacementType::OnPlane)
+	if (ShapeSettings->PlaceMode == EMakeMeshPlacementType::GroundPlane)
 	{
 		FPlane DrawPlane(FVector::ZeroVector, FVector(0, 0, 1));
 		FVector DrawPlanePos = FMath::RayPlaneIntersection(ClickPosWorldRay.Origin, ClickPosWorldRay.Direction, DrawPlane);
diff --git a/Engine/Plugins/Experimental/MeshModelingToolset/Source/MeshModelingTools/Public/AddPrimitiveTool.h b/Engine/Plugins/Experimental/MeshModelingToolset/Source/MeshModelingTools/Public/AddPrimitiveTool.h
--- a/Engine/Plugins/Experimental/MeshModelingToolset/Source/MeshModelingTools/Public/AddPrimitiveTool.h
+++ b/Engine/Plugins/Experimental/MeshModelingToolset/Source/MeshModelingTools/Public/AddPrimitiveTool.h
@@ -58,7 +58,7 @@ ENUM_CLASS_FLAGS(EMakeMeshShapeType);
 UENUM()
 enum class EMakeMeshPlacementType : uint8
 {
-	OnPlane = 0,
+	GroundPlane = 0,
 	OnScene = 1
 };
 
