ModelingTools: update Simple and Octree DynamicMeshComponent's to use the SceneProxy pointer in base PrimitiveComponent, rather than their own copy. Simple Component's FastNotifyXYZ functions now use MarkRenderDynamicDataDirty() instead of MarkRenderStateDirty(), which was deleting the Proxy. Change MeshSelectionTool to use fast color update.
#rnx

diff --git a/Engine/Plugins/Experimental/MeshModelingToolset/Source/MeshModelingTools/Private/MeshSelectionTool.cpp b/Engine/Plugins/Experimental/MeshModelingToolset/Source/MeshModelingTools/Private/MeshSelectionTool.cpp
--- a/Engine/Plugins/Experimental/MeshModelingToolset/Source/MeshModelingTools/Private/MeshSelectionTool.cpp
+++ b/Engine/Plugins/Experimental/MeshModelingToolset/Source/MeshModelingTools/Private/MeshSelectionTool.cpp
@@ -525,7 +525,7 @@ void UMeshSelectionTool::UpdateVisualization()
 			{
 				return SelectedTriangles[TriangleID] ? LinearColors::VideoRed3b() : LinearColors::SelectFColor(Mesh->GetTriangleGroup(TriangleID));
 
-			}, UPreviewMesh::ERenderUpdateMode::FullUpdate);
+			}, UPreviewMesh::ERenderUpdateMode::FastUpdate);
 		}
 		else
 		{
@@ -533,12 +533,12 @@ void UMeshSelectionTool::UpdateVisualization()
 			{
 				return SelectedTriangles[TriangleID] ? LinearColors::VideoRed3b() : LinearColors::VideoWhite3b();
 
-			}, UPreviewMesh::ERenderUpdateMode::FullUpdate);
+			}, UPreviewMesh::ERenderUpdateMode::FastUpdate);
 		}
 	}
 	else
 	{
-		PreviewMesh->ClearTriangleColorFunction(UPreviewMesh::ERenderUpdateMode::FullUpdate);
+		PreviewMesh->ClearTriangleColorFunction(UPreviewMesh::ERenderUpdateMode::FastUpdate);
 	}
 }
 
diff --git a/Engine/Plugins/Experimental/MeshModelingToolset/Source/ModelingComponents/Private/BaseDynamicMeshComponent.cpp b/Engine/Plugins/Experimental/MeshModelingToolset/Source/ModelingComponents/Private/BaseDynamicMeshComponent.cpp
--- a/Engine/Plugins/Experimental/MeshModelingToolset/Source/ModelingComponents/Private/BaseDynamicMeshComponent.cpp
+++ b/Engine/Plugins/Experimental/MeshModelingToolset/Source/ModelingComponents/Private/BaseDynamicMeshComponent.cpp
@@ -13,14 +13,20 @@ UBaseDynamicMeshComponent::UBaseDynamicMeshComponent(const FObjectInitializer& O
 
 void UBaseDynamicMeshComponent::SetOverrideRenderMaterial(UMaterialInterface* Material)
 {
-	OverrideRenderMaterial = Material;
-	NotifyMaterialSetUpdated();
+	if (OverrideRenderMaterial != Material)
+	{
+		OverrideRenderMaterial = Material;
+		NotifyMaterialSetUpdated();
+	}
 }
 
 void UBaseDynamicMeshComponent::ClearOverrideRenderMaterial()
 {
-	OverrideRenderMaterial = nullptr;
-	NotifyMaterialSetUpdated();
+	if (OverrideRenderMaterial != nullptr)
+	{
+		OverrideRenderMaterial = nullptr;
+		NotifyMaterialSetUpdated();
+	}
 }
 
 
diff --git a/Engine/Plugins/Experimental/MeshModelingToolset/Source/ModelingComponents/Private/OctreeDynamicMeshComponent.cpp b/Engine/Plugins/Experimental/MeshModelingToolset/Source/ModelingComponents/Private/OctreeDynamicMeshComponent.cpp
--- a/Engine/Plugins/Experimental/MeshModelingToolset/Source/ModelingComponents/Private/OctreeDynamicMeshComponent.cpp
+++ b/Engine/Plugins/Experimental/MeshModelingToolset/Source/ModelingComponents/Private/OctreeDynamicMeshComponent.cpp
@@ -105,7 +105,7 @@ void UOctreeDynamicMeshComponent::ApplyTransform(const FTransform3d& Transform,
 		MeshTransforms::ApplyTransform(*GetMesh(), Transform);
 	}
 
-	if (CurrentProxy != nullptr)
+	if (GetCurrentSceneProxy() != nullptr)
 	{
 		Octree->ModifiedBounds = FAxisAlignedBox3d(
 			-TNumericLimits<float>::Max()*FVector3d::One(),
@@ -144,7 +144,7 @@ void UOctreeDynamicMeshComponent::Bake(FMeshDescription* MeshDescription, bool b
 
 void UOctreeDynamicMeshComponent::NotifyMeshUpdated()
 {
-	if (CurrentProxy != nullptr)
+	if (GetCurrentSceneProxy() != nullptr)
 	{
 		FAxisAlignedBox3d DirtyBox = Octree->ModifiedBounds;
 		Octree->ResetModifiedBounds();
@@ -227,14 +227,10 @@ void UOctreeDynamicMeshComponent::NotifyMeshUpdated()
 		//UE_LOG(LogTemp, Warning, TEXT("Updating %d of %d decomposition sets, %d tris total, spillcount %d"), SetsToUpdate.Num(), CutCellSetMap.Num(), TotalTriangleCount, SpillTriangleCount);
 		{
 			SCOPE_CYCLE_COUNTER(STAT_SculptToolOctree_UpdateFromDecomp);
-			CurrentProxy->UpdateFromDecomposition(TriangleDecomposition, SetsToUpdate);
+			GetCurrentSceneProxy()->UpdateFromDecomposition(TriangleDecomposition, SetsToUpdate);
 		}
 
 	}
-
-	//MarkRenderStateDirty();
-	//UpdateBounds();
-	//CurrentProxy = nullptr;
 }
 
 
@@ -267,15 +263,17 @@ static void InitializeOctreeCutSet(const FDynamicMesh3& Mesh, const FDynamicMesh
 
 FPrimitiveSceneProxy* UOctreeDynamicMeshComponent::CreateSceneProxy()
 {
-	CurrentProxy = nullptr;
+	check(GetCurrentSceneProxy() == nullptr);
+
+	FOctreeDynamicMeshSceneProxy* NewProxy = nullptr;
 	if (Mesh->TriangleCount() > 0)
 	{
-		CurrentProxy = new FOctreeDynamicMeshSceneProxy(this);
+		NewProxy = new FOctreeDynamicMeshSceneProxy(this);
 
 		if (TriangleColorFunc != nullptr)
 		{
-			CurrentProxy->bUsePerTriangleColor = true;
-			CurrentProxy->PerTriangleColorFunc = [this](const FDynamicMesh3* Mesh, int TriangleID) { return GetTriangleColor(TriangleID); };
+			NewProxy->bUsePerTriangleColor = true;
+			NewProxy->PerTriangleColorFunc = [this](const FDynamicMesh3* Mesh, int TriangleID) { return GetTriangleColor(TriangleID); };
 		}
 
 
@@ -312,18 +310,18 @@ FPrimitiveSceneProxy* UOctreeDynamicMeshComponent::CreateSceneProxy()
 			});
 		}
 
-		CurrentProxy->InitializeFromDecomposition(TriangleDecomposition);
+		NewProxy->InitializeFromDecomposition(TriangleDecomposition);
 	}
-	return CurrentProxy;
+	return NewProxy;
 }
 
 
 
 void UOctreeDynamicMeshComponent::NotifyMaterialSetUpdated()
 {
-	if (CurrentProxy != nullptr)
+	if (GetCurrentSceneProxy() != nullptr)
 	{
-		CurrentProxy->UpdatedReferencedMaterials();
+		GetCurrentSceneProxy()->UpdatedReferencedMaterials();
 	}
 }
 
diff --git a/Engine/Plugins/Experimental/MeshModelingToolset/Source/ModelingComponents/Private/SimpleDynamicMeshComponent.cpp b/Engine/Plugins/Experimental/MeshModelingToolset/Source/ModelingComponents/Private/SimpleDynamicMeshComponent.cpp
--- a/Engine/Plugins/Experimental/MeshModelingToolset/Source/ModelingComponents/Private/SimpleDynamicMeshComponent.cpp
+++ b/Engine/Plugins/Experimental/MeshModelingToolset/Source/ModelingComponents/Private/SimpleDynamicMeshComponent.cpp
@@ -152,7 +152,6 @@ void USimpleDynamicMeshComponent::NotifyMeshUpdated()
 	// Need to recreate scene proxy to send it over
 	MarkRenderStateDirty();
 	UpdateBounds();
-	CurrentProxy = nullptr;
 
 	if (TangentsType != EDynamicMeshTangentCalcType::ExternallyCalculated)
 	{
@@ -163,11 +162,22 @@ void USimpleDynamicMeshComponent::NotifyMeshUpdated()
 
 void USimpleDynamicMeshComponent::FastNotifyColorsUpdated()
 {
-	MarkRenderStateDirty();
-
-	if (CurrentProxy != nullptr)
+	FSimpleDynamicMeshSceneProxy* Proxy = GetCurrentSceneProxy();
+	if (Proxy != nullptr)
 	{
-		CurrentProxy->FastUpdateVertices(false, false, true);
+		if (TriangleColorFunc != nullptr &&  Proxy->bUsePerTriangleColor == false )
+		{
+			Proxy->bUsePerTriangleColor = true;
+			Proxy->PerTriangleColorFunc = [this](const FDynamicMesh3* MeshIn, int TriangleID) { return GetTriangleColor(MeshIn, TriangleID); };
+		} 
+		else if (TriangleColorFunc == nullptr && Proxy->bUsePerTriangleColor == true)
+		{
+			Proxy->bUsePerTriangleColor = false;
+			Proxy->PerTriangleColorFunc = nullptr;
+		}
+
+		Proxy->FastUpdateVertices(false, false, true);
+		MarkRenderDynamicDataDirty();
 	}
 	else
 	{
@@ -179,12 +189,11 @@ void USimpleDynamicMeshComponent::FastNotifyColorsUpdated()
 
 void USimpleDynamicMeshComponent::FastNotifyPositionsUpdated()
 {
-	MarkRenderStateDirty();
-	UpdateBounds();
-
-	if (CurrentProxy != nullptr)
+	if (GetCurrentSceneProxy() != nullptr)
 	{
-		CurrentProxy->FastUpdateVertices(true, false, false);
+		GetCurrentSceneProxy()->FastUpdateVertices(true, false, false);
+		MarkRenderDynamicDataDirty();
+		UpdateBounds();
 	}
 	else
 	{
@@ -198,28 +207,30 @@ void USimpleDynamicMeshComponent::FastNotifyPositionsUpdated()
 
 FPrimitiveSceneProxy* USimpleDynamicMeshComponent::CreateSceneProxy()
 {
-	CurrentProxy = nullptr;
+	check(GetCurrentSceneProxy() == nullptr);
+
+	FSimpleDynamicMeshSceneProxy* NewProxy = nullptr;
 	if (Mesh->TriangleCount() > 0)
 	{
-		CurrentProxy = new FSimpleDynamicMeshSceneProxy(this);
+		NewProxy = new FSimpleDynamicMeshSceneProxy(this);
 
 		if (TriangleColorFunc)
 		{
-			CurrentProxy->bUsePerTriangleColor = true;
-			CurrentProxy->PerTriangleColorFunc = [this](const FDynamicMesh3* MeshIn, int TriangleID) { return GetTriangleColor(MeshIn, TriangleID); };
+			NewProxy->bUsePerTriangleColor = true;
+			NewProxy->PerTriangleColorFunc = [this](const FDynamicMesh3* MeshIn, int TriangleID) { return GetTriangleColor(MeshIn, TriangleID); };
 		}
 
-		CurrentProxy->Initialize();
+		NewProxy->Initialize();
 	}
-	return CurrentProxy;
+	return NewProxy;
 }
 
 
 void USimpleDynamicMeshComponent::NotifyMaterialSetUpdated()
 {
-	if (CurrentProxy != nullptr)
+	if (GetCurrentSceneProxy() != nullptr)
 	{
-		CurrentProxy->UpdatedReferencedMaterials();
+		GetCurrentSceneProxy()->UpdatedReferencedMaterials();
 	}
 }
 
diff --git a/Engine/Plugins/Experimental/MeshModelingToolset/Source/ModelingComponents/Public/OctreeDynamicMeshComponent.h b/Engine/Plugins/Experimental/MeshModelingToolset/Source/ModelingComponents/Public/OctreeDynamicMeshComponent.h
--- a/Engine/Plugins/Experimental/MeshModelingToolset/Source/ModelingComponents/Public/OctreeDynamicMeshComponent.h
+++ b/Engine/Plugins/Experimental/MeshModelingToolset/Source/ModelingComponents/Public/OctreeDynamicMeshComponent.h
@@ -135,7 +135,7 @@ protected:
 
 private:
 
-	FOctreeDynamicMeshSceneProxy* CurrentProxy = nullptr;
+	FOctreeDynamicMeshSceneProxy* GetCurrentSceneProxy() { return (FOctreeDynamicMeshSceneProxy*)SceneProxy; }
 
 	//~ Begin UPrimitiveComponent Interface.
 	virtual FPrimitiveSceneProxy* CreateSceneProxy() override;
diff --git a/Engine/Plugins/Experimental/MeshModelingToolset/Source/ModelingComponents/Public/SimpleDynamicMeshComponent.h b/Engine/Plugins/Experimental/MeshModelingToolset/Source/ModelingComponents/Public/SimpleDynamicMeshComponent.h
--- a/Engine/Plugins/Experimental/MeshModelingToolset/Source/ModelingComponents/Public/SimpleDynamicMeshComponent.h
+++ b/Engine/Plugins/Experimental/MeshModelingToolset/Source/ModelingComponents/Public/SimpleDynamicMeshComponent.h
@@ -163,7 +163,7 @@ protected:
 
 private:
 
-	FSimpleDynamicMeshSceneProxy* CurrentProxy = nullptr;
+	FSimpleDynamicMeshSceneProxy* GetCurrentSceneProxy() { return (FSimpleDynamicMeshSceneProxy*)SceneProxy; }
 
 	//~ Begin UPrimitiveComponent Interface.
 	virtual FPrimitiveSceneProxy* CreateSceneProxy() override;
