Fix a compile error when using BulkData2, caused by ::LockReadOnly returning a pointer of the BulkDatas type rather than void*.

So we need to use reinterpret_cast instead of static_cast.

#rnx

diff --git a/Engine/Plugins/Experimental/MeshModelingToolset/Source/MeshModelingTools/Private/DisplaceMeshTool.cpp b/Engine/Plugins/Experimental/MeshModelingToolset/Source/MeshModelingTools/Private/DisplaceMeshTool.cpp
--- a/Engine/Plugins/Experimental/MeshModelingToolset/Source/MeshModelingTools/Private/DisplaceMeshTool.cpp
+++ b/Engine/Plugins/Experimental/MeshModelingToolset/Source/MeshModelingTools/Private/DisplaceMeshTool.cpp
@@ -123,7 +123,7 @@ namespace {
 #endif
 			DisplacementMap->UpdateResource();
 
-			FormattedImageData = static_cast<const FColor*>(DisplacementMap->PlatformData->Mips[0].BulkData.LockReadOnly());
+			FormattedImageData = reinterpret_cast<const FColor*>(DisplacementMap->PlatformData->Mips[0].BulkData.LockReadOnly());
 		}
 		FTextureAccess(const FTextureAccess&) = delete;
 		FTextureAccess(FTextureAccess&&) = delete;
