Missing includes fix for unity build.

#rnx

diff --git a/Engine/Plugins/Experimental/MeshModelingToolset/Source/MeshModelingToolsEditorOnly/Private/AttributeEditorTool.cpp b/Engine/Plugins/Experimental/MeshModelingToolset/Source/MeshModelingToolsEditorOnly/Private/AttributeEditorTool.cpp
--- a/Engine/Plugins/Experimental/MeshModelingToolset/Source/MeshModelingToolsEditorOnly/Private/AttributeEditorTool.cpp
+++ b/Engine/Plugins/Experimental/MeshModelingToolset/Source/MeshModelingToolsEditorOnly/Private/AttributeEditorTool.cpp
@@ -11,6 +11,8 @@
 #include "StaticMeshAttributes.h"
 #include "MeshDescriptionOperations.h"
 
+#include "Components/PrimitiveComponent.h"
+
 
 #define LOCTEXT_NAMESPACE "UAttributeEditorTool"
 
diff --git a/Engine/Source/Developer/AssetTools/Private/AssetFixUpRedirectors.cpp b/Engine/Source/Developer/AssetTools/Private/AssetFixUpRedirectors.cpp
--- a/Engine/Source/Developer/AssetTools/Private/AssetFixUpRedirectors.cpp
+++ b/Engine/Source/Developer/AssetTools/Private/AssetFixUpRedirectors.cpp
@@ -22,6 +22,7 @@
 #include "ObjectTools.h"
 #include "Logging/MessageLog.h"
 #include "AssetTools.h"
+#include "Engine/Blueprint.h"
 
 #define LOCTEXT_NAMESPACE "AssetFixUpRedirectors"
 
diff --git a/Engine/Source/Runtime/Engine/Private/Blueprint.cpp b/Engine/Source/Runtime/Engine/Private/Blueprint.cpp
--- a/Engine/Source/Runtime/Engine/Private/Blueprint.cpp
+++ b/Engine/Source/Runtime/Engine/Private/Blueprint.cpp
@@ -41,6 +41,7 @@
 #include "UObject/MetaData.h"
 #include "BlueprintAssetHandler.h"
 #include "Blueprint/BlueprintExtension.h"
+#include "UObject/TextProperty.h"
 #endif
 #include "Engine/InheritableComponentHandler.h"
 
