Fix for warning after upgrading to PVS-Studio 7.0.4:

V1047: Lifetime of the lambda is greater than lifetime of the local variable 'LocalEyePosition' captured by reference.

#rnx

diff --git a/Engine/Plugins/Experimental/MeshModelingToolset/Source/MeshModelingTools/Private/DynamicMeshSculptTool.cpp b/Engine/Plugins/Experimental/MeshModelingToolset/Source/MeshModelingTools/Private/DynamicMeshSculptTool.cpp
--- a/Engine/Plugins/Experimental/MeshModelingToolset/Source/MeshModelingTools/Private/DynamicMeshSculptTool.cpp
+++ b/Engine/Plugins/Experimental/MeshModelingToolset/Source/MeshModelingTools/Private/DynamicMeshSculptTool.cpp
@@ -1014,7 +1014,7 @@ int UDynamicMeshSculptTool::FindHitTargetMeshTriangle(const FRay3d& LocalRay)
 		GetToolManager()->GetContextQueriesAPI()->GetCurrentViewState(StateOut);
 		FVector3d LocalEyePosition(CurTargetTransform.InverseTransformPosition(StateOut.Position));
 
-		BrushTargetMeshSpatial.TriangleFilterF = [this, Mesh, &LocalEyePosition](int TriangleID) {
+		BrushTargetMeshSpatial.TriangleFilterF = [this, Mesh, LocalEyePosition](int TriangleID) {
 			FVector3d Normal, Centroid;
 			double Area;
 			Mesh->GetTriInfo(TriangleID, Normal, Area, Centroid);
