MeshModelingToolset: Change Simplify default group and material boundary constraints to "ignore"

The default boundary constraints for the Simplify tool were set too strictly which prevented our primitive sphere from being simplified. This change relaxes those constraints.

#rnx

diff --git a/Engine/Plugins/Experimental/MeshModelingToolset/Source/MeshModelingToolsEditorOnly/Private/SimplifyMeshTool.cpp b/Engine/Plugins/Experimental/MeshModelingToolset/Source/MeshModelingToolsEditorOnly/Private/SimplifyMeshTool.cpp
--- a/Engine/Plugins/Experimental/MeshModelingToolset/Source/MeshModelingToolsEditorOnly/Private/SimplifyMeshTool.cpp
+++ b/Engine/Plugins/Experimental/MeshModelingToolset/Source/MeshModelingToolsEditorOnly/Private/SimplifyMeshTool.cpp
@@ -2,6 +2,7 @@
 
 #include "SimplifyMeshTool.h"
 #include "InteractiveToolManager.h"
+#include "Properties/RemeshProperties.h"
 #include "ToolBuilderUtil.h"
 
 #include "ToolSetupUtil.h"
@@ -70,6 +71,8 @@ USimplifyMeshToolProperties::USimplifyMeshToolProperties()
 	bDiscardAttributes = false;
 	bShowWireframe = true;
 	bShowGroupColors = false;
+	GroupBoundaryConstraint = EGroupBoundaryConstraint::Ignore;
+	MaterialBoundaryConstraint = EMaterialBoundaryConstraint::Ignore;
 }
 
 void USimplifyMeshTool::SetWorld(UWorld* World)
