invert faces for invert scale in combine tool #jira UE-87427
#rnx

diff --git a/Engine/Plugins/Experimental/GeometryProcessing/Source/GeometricObjects/Public/TransformTypes.h b/Engine/Plugins/Experimental/GeometryProcessing/Source/GeometricObjects/Public/TransformTypes.h
--- a/Engine/Plugins/Experimental/GeometryProcessing/Source/GeometricObjects/Public/TransformTypes.h
+++ b/Engine/Plugins/Experimental/GeometryProcessing/Source/GeometricObjects/Public/TransformTypes.h
@@ -119,6 +119,11 @@ public:
 		Scale3D = ScaleIn;
 	}
 
+	RealType GetDeterminant()
+	{
+		return Scale3D.X * Scale3D.Y * Scale3D.Z;
+	}
+
 	/**
 	 * @return inverse of this transform
 	 */
diff --git a/Engine/Plugins/Experimental/MeshModelingToolset/Source/MeshModelingTools/Private/CombineMeshesTool.cpp b/Engine/Plugins/Experimental/MeshModelingToolset/Source/MeshModelingTools/Private/CombineMeshesTool.cpp
--- a/Engine/Plugins/Experimental/MeshModelingToolset/Source/MeshModelingTools/Private/CombineMeshesTool.cpp
+++ b/Engine/Plugins/Experimental/MeshModelingToolset/Source/MeshModelingTools/Private/CombineMeshesTool.cpp
@@ -193,6 +193,12 @@ void UCombineMeshesTool::UpdateAssets()
 			FDynamicMesh3 ComponentDMesh;
 			Converter.Convert(ComponentTarget->GetMesh(), ComponentDMesh);
 
+			FTransform3d XF = (FTransform3d)(ComponentTarget->GetWorldTransform() * ToAccum);
+			if (XF.GetDeterminant() < 0)
+			{
+				ComponentDMesh.ReverseOrientation(false);
+			}
+
 			// update material IDs to account for combined material set
 			FDynamicMeshMaterialAttribute* MatAttrib = ComponentDMesh.Attributes()->GetMaterialID();
 			for (int TID : ComponentDMesh.TriangleIndicesItr())
@@ -202,7 +208,6 @@ void UCombineMeshesTool::UpdateAssets()
 
 			FDynamicMeshEditor Editor(&AccumulateDMesh);
 			FMeshIndexMappings IndexMapping;
-			FTransform3d XF = (FTransform3d)(ComponentTarget->GetWorldTransform() * ToAccum);
 			Editor.AppendMesh(&ComponentDMesh, IndexMapping, 
 				[&XF](int Unused, const FVector3d P) { return XF.TransformPosition(P); },
 				[&XF](int Unused, const FVector3d N) { return XF.TransformNormal(N); });
