account for mesh scale in poly cut tool #jira UE-87442
#rnx

diff --git a/Engine/Plugins/Experimental/MeshModelingToolset/Source/MeshModelingToolsEditorOnly/Private/PolygonOnMeshTool.cpp b/Engine/Plugins/Experimental/MeshModelingToolset/Source/MeshModelingToolsEditorOnly/Private/PolygonOnMeshTool.cpp
--- a/Engine/Plugins/Experimental/MeshModelingToolset/Source/MeshModelingToolsEditorOnly/Private/PolygonOnMeshTool.cpp
+++ b/Engine/Plugins/Experimental/MeshModelingToolset/Source/MeshModelingToolsEditorOnly/Private/PolygonOnMeshTool.cpp
@@ -202,16 +202,20 @@ TUniquePtr<FDynamicMeshOperator> UPolygonOnMeshTool::MakeNewOperator()
 	TUniquePtr<FEmbedPolygonsOp> EmbedOp = MakeUnique<FEmbedPolygonsOp>();
 	EmbedOp->bDiscardAttributes = false;
 	EmbedOp->Operation = BasicProperties->Operation;
-	EmbedOp->EmbedPolygon = ActivePolygon;
-	double Scale = BasicProperties->PolygonScale;
-	EmbedOp->EmbedPolygon.Scale(FVector2d(Scale, Scale), ActivePolygon.Bounds().Center());
 
-	// EmbedOp->ExtrudeDistance = Tool->BasicProperties->ExtrudeDistance;
-	
 	FFrame3d LocalFrame = DrawPlaneWorld;
 	FTransform3d ToLocal = WorldTransform.Inverse();
 	LocalFrame.Transform(ToLocal);
 	EmbedOp->PolygonFrame = LocalFrame;
+	
+	FVector2d LocalFrameScale(ToLocal.TransformVector(LocalFrame.X()).Length(), ToLocal.TransformVector(LocalFrame.Y()).Length());
+	LocalFrameScale *= BasicProperties->PolygonScale;
+	EmbedOp->EmbedPolygon = ActivePolygon;
+	EmbedOp->EmbedPolygon.Scale(LocalFrameScale, FVector2d::Zero());
+
+	// TODO: scale any extrude by ToLocal.TransformVector(LocalFrame.Z()).Length() ??
+	// EmbedOp->ExtrudeDistance = Tool->BasicProperties->ExtrudeDistance;
+
 	EmbedOp->OriginalMesh = OriginalDynamicMesh;
 	EmbedOp->SetResultTransform(WorldTransform);
 
