fix side of hole cut + uv spacing for polycut tool

diff --git a/Engine/Plugins/Experimental/MeshModelingToolset/Source/ModelingOperatorsEditorOnly/Private/CuttingOps/EmbedPolygonsOp.cpp b/Engine/Plugins/Experimental/MeshModelingToolset/Source/ModelingOperatorsEditorOnly/Private/CuttingOps/EmbedPolygonsOp.cpp
--- a/Engine/Plugins/Experimental/MeshModelingToolset/Source/ModelingOperatorsEditorOnly/Private/CuttingOps/EmbedPolygonsOp.cpp
+++ b/Engine/Plugins/Experimental/MeshModelingToolset/Source/ModelingOperatorsEditorOnly/Private/CuttingOps/EmbedPolygonsOp.cpp
@@ -124,7 +124,7 @@ void FEmbedPolygonsOp::CalculateResult(FProgressCancel* Progress)
 	bool bCollapseDegenerateEdges = true; // TODO make this optional?
 	
 	FFrame3d Frame = PolygonFrame;
-	Frame.Origin = Frame.Origin - (2*MeshRadius*Frame.Z());
+	Frame.Origin = Frame.Origin + (2*MeshRadius*Frame.Z());
 
 	FPolygon2d Polygon = GetPolygon();
 	double Perimeter = Polygon.Perimeter();
@@ -132,7 +132,7 @@ void FEmbedPolygonsOp::CalculateResult(FProgressCancel* Progress)
 	TArray<int> PathVertIDs1, PathVertCorrespond1;
 	
 	TArray<TPair<float, int>> SortedHitTriangles;
-	TMeshQueries<FDynamicMesh3>::FindHitTriangles_LinearSearch(*ResultMesh, FRay3d(Frame.FromPlaneUV(Polygon[0]), Frame.Z()), SortedHitTriangles);
+	TMeshQueries<FDynamicMesh3>::FindHitTriangles_LinearSearch(*ResultMesh, FRay3d(Frame.FromPlaneUV(Polygon[0]), -Frame.Z()), SortedHitTriangles);
 
 	if (SortedHitTriangles.Num() < 1)
 	{
@@ -221,16 +221,17 @@ void FEmbedPolygonsOp::CalculateResult(FProgressCancel* Progress)
 		}
 		FDynamicMeshEditor MeshEditor(ResultMesh.Get());
 		FDynamicMeshEditResult ResultOut;
-		bool bStitched = MeshEditor.StitchSparselyCorrespondedVertexLoops(PathVertIDs2, PathVertCorrespond2, PathVertIDs1, PathVertCorrespond1, ResultOut);
+		bool bStitched = MeshEditor.StitchSparselyCorrespondedVertexLoops(PathVertIDs1, PathVertCorrespond1, PathVertIDs2, PathVertCorrespond2, ResultOut);
 		if (ResultMesh->HasAttributes())
 		{
-			MeshEditor.SetTubeNormals(ResultOut.NewTriangles, PathVertIDs2, PathVertCorrespond2, PathVertIDs1, PathVertCorrespond1);
+			MeshEditor.SetTubeNormals(ResultOut.NewTriangles, PathVertIDs1, PathVertCorrespond1, PathVertIDs2, PathVertCorrespond2);
 			TArray<float> UValues; UValues.SetNumUninitialized(PathVertCorrespond2.Num() + 1);
-			FVector3f ZVec = (FVector3f)Frame.Z();
+			FVector3f ZVec = -(FVector3f)Frame.Z();
+			float Along = 0;
 			for (int UIdx = 0; UIdx < UValues.Num(); UIdx++)
 			{
-				// todo stop assuming vertices are uniformly distributed on polygon
-				UValues[UIdx] = Perimeter * float(UIdx) / float(UValues.Num());
+				UValues[UIdx] = Along;
+				Along += Polygon[UIdx % Polygon.VertexCount()].Distance(Polygon[(UIdx + 1) % Polygon.VertexCount()]);
 			}
 
 			for (int UVIdx = 0, NumUVLayers = ResultMesh->Attributes()->NumUVLayers(); UVIdx < NumUVLayers; UVIdx++)
@@ -256,7 +257,7 @@ void FEmbedPolygonsOp::CalculateResult(FProgressCancel* Progress)
 		{
 			if (ResultMesh->HasAttributes())
 			{
-				MeshEditor.SetTriangleNormals(Filler.NewTriangles, (FVector3f)-Frame.Z()); // TODO: should we use the best fit plane instead of the projection plane for the normal? ... probably.
+				MeshEditor.SetTriangleNormals(Filler.NewTriangles, (FVector3f)Frame.Z()); // TODO: should we use the best fit plane instead of the projection plane for the normal? ... probably.
 				
 				for (int UVIdx = 0, NumUVLayers = ResultMesh->Attributes()->NumUVLayers(); UVIdx < NumUVLayers; UVIdx++)
 				{
