make separated pieces from selection tool not persist if the tool exited by canceling

#rnx

diff --git a/Engine/Plugins/Experimental/MeshModelingToolset/Source/MeshModelingTools/Private/MeshSelectionTool.cpp b/Engine/Plugins/Experimental/MeshModelingToolset/Source/MeshModelingTools/Private/MeshSelectionTool.cpp
--- a/Engine/Plugins/Experimental/MeshModelingToolset/Source/MeshModelingTools/Private/MeshSelectionTool.cpp
+++ b/Engine/Plugins/Experimental/MeshModelingToolset/Source/MeshModelingTools/Private/MeshSelectionTool.cpp
@@ -184,6 +184,13 @@ void UMeshSelectionTool::OnShutdown(EToolShutdownType ShutdownType)
 		});
 		GetToolManager()->EndUndoTransaction();
 	}
+	else if (ShutdownType == EToolShutdownType::Cancel)
+	{
+		for (AActor* Spawned : SpawnedActors)
+		{
+			Spawned->Destroy();
+		}
+	}
 }
 
 
@@ -1329,6 +1336,7 @@ void UMeshSelectionTool::SeparateSelectedTriangles()
 	{
 		AssignMaterial(NewActor, ComponentTarget);
 	}
+	SpawnedActors.Add(NewActor);
 	GetToolManager()->EndUndoTransaction();
 
 	// todo: undo won't remove this asset...
diff --git a/Engine/Plugins/Experimental/MeshModelingToolset/Source/MeshModelingTools/Public/MeshSelectionTool.h b/Engine/Plugins/Experimental/MeshModelingToolset/Source/MeshModelingTools/Public/MeshSelectionTool.h
--- a/Engine/Plugins/Experimental/MeshModelingToolset/Source/MeshModelingTools/Public/MeshSelectionTool.h
+++ b/Engine/Plugins/Experimental/MeshModelingToolset/Source/MeshModelingTools/Public/MeshSelectionTool.h
@@ -318,6 +318,9 @@ protected:
 	UPROPERTY()
 	UMeshSelectionSet* Selection;
 
+	UPROPERTY()
+	TArray<AActor*> SpawnedActors;
+
 	UWorld* TargetWorld;
 	IToolsContextAssetAPI* AssetAPI;
 
