Fix TriangleCount -> MaxTriangleID to support non-compact meshes
#rnx

diff --git a/Engine/Plugins/Experimental/GeometryProcessing/Source/GeometricObjects/Public/Util/MeshCaches.h b/Engine/Plugins/Experimental/GeometryProcessing/Source/GeometricObjects/Public/Util/MeshCaches.h
--- a/Engine/Plugins/Experimental/GeometryProcessing/Source/GeometricObjects/Public/Util/MeshCaches.h
+++ b/Engine/Plugins/Experimental/GeometryProcessing/Source/GeometricObjects/Public/Util/MeshCaches.h
@@ -26,14 +26,17 @@ struct FMeshTriInfoCache
 	static FMeshTriInfoCache BuildTriInfoCache(const TriangleMeshType& Mesh)
 	{
 		FMeshTriInfoCache Cache;
-		int NT = Mesh.TriangleCount();
+		int NT = Mesh.MaxTriangleID();
 		Cache.Centroids.Resize(NT);
 		Cache.Normals.Resize(NT);
 		Cache.Areas.Resize(NT);
 
-		ParallelFor(NT, [&](int TriIdx)
+		ParallelFor(NT, [&](int TID)
 		{
-			TMeshQueries<TriangleMeshType>::GetTriNormalAreaCentroid(Mesh, TriIdx, Cache.Normals[TriIdx], Cache.Areas[TriIdx], Cache.Centroids[TriIdx]);
+			if (Mesh.IsTriangle(TID))
+			{
+				TMeshQueries<TriangleMeshType>::GetTriNormalAreaCentroid(Mesh, TID, Cache.Normals[TID], Cache.Areas[TID], Cache.Centroids[TID]);
+			}
 		});
 
 		return Cache;
diff --git a/Engine/Plugins/Experimental/MeshModelingToolset/Source/ModelingOperators/Private/ParameterizationOps/UVProjectionOp.cpp b/Engine/Plugins/Experimental/MeshModelingToolset/Source/ModelingOperators/Private/ParameterizationOps/UVProjectionOp.cpp
--- a/Engine/Plugins/Experimental/MeshModelingToolset/Source/ModelingOperators/Private/ParameterizationOps/UVProjectionOp.cpp
+++ b/Engine/Plugins/Experimental/MeshModelingToolset/Source/ModelingOperators/Private/ParameterizationOps/UVProjectionOp.cpp
@@ -86,8 +86,8 @@ void FUVProjectionOp::CalculateResult(FProgressCancel* Progress)
 	{	
 		// All cases in this branch require normals 
 		// compute normals on the transformed vertices
-		TArray<FVector3d> TransformedNormals; TransformedNormals.SetNumUninitialized(ResultMesh->TriangleCount());
-		ParallelFor(ResultMesh->TriangleCount(), [this, &TransformedVertices, &TransformedNormals](int32 TriangleID)
+		TArray<FVector3d> TransformedNormals; TransformedNormals.SetNumUninitialized(ResultMesh->MaxTriangleID());
+		ParallelFor(ResultMesh->MaxTriangleID(), [this, &TransformedVertices, &TransformedNormals](int32 TriangleID)
 		{
 			if (ResultMesh->IsTriangle(TriangleID))
 			{
