fix naming for add primitive tool's two sphere options #jira UE-88931

#rnx

diff --git a/Engine/Plugins/Experimental/MeshModelingToolset/Source/MeshModelingTools/Public/AddPrimitiveTool.h b/Engine/Plugins/Experimental/MeshModelingToolset/Source/MeshModelingTools/Public/AddPrimitiveTool.h
--- a/Engine/Plugins/Experimental/MeshModelingToolset/Source/MeshModelingTools/Public/AddPrimitiveTool.h
+++ b/Engine/Plugins/Experimental/MeshModelingToolset/Source/MeshModelingTools/Public/AddPrimitiveTool.h
@@ -48,8 +48,8 @@ enum class EMakeMeshShapeType : uint32
 	Disc			 = 0x040 UMETA(DisplayName = "Disc"),
 	PuncturedDisc	 = 0x080 UMETA(DisplayName = "Punctured Disc"),
 	Torus			 = 0x100 UMETA(DisplayName = "Torus"),
-	Sphere			 = 0x200 UMETA(DisplayName = "Sphere"),
-	SphericalBox	 = 0x400 UMETA(DisplayName = "Spherical Box")
+	SphericalBox     = 0x200 UMETA(DisplayName = "Sphere Type 1"),
+	Sphere			 = 0x400 UMETA(DisplayName = "Sphere Type 2")
 	
 };
 ENUM_CLASS_FLAGS(EMakeMeshShapeType);
