Merge CL 11420156

DrawPolygonTool: fix for extrusion appearing in unexpected location.

#rnx

diff --git a/Engine/Plugins/Experimental/MeshModelingToolset/Source/MeshModelingTools/Private/DrawPolygonTool.cpp b/Engine/Plugins/Experimental/MeshModelingToolset/Source/MeshModelingTools/Private/DrawPolygonTool.cpp
--- a/Engine/Plugins/Experimental/MeshModelingToolset/Source/MeshModelingTools/Private/DrawPolygonTool.cpp
+++ b/Engine/Plugins/Experimental/MeshModelingToolset/Source/MeshModelingTools/Private/DrawPolygonTool.cpp
@@ -946,8 +946,19 @@ void UDrawPolygonTool::BeginInteractiveExtrude()
 	HeightMechanic->CurrentHeight = 1.0f;  // initialize to something non-zero...prob should be based on polygon bounds maybe?
 
 	FDynamicMesh3 HeightMesh;
-	FFrame3d HeightMeshFrame;
-	GeneratePolygonMesh(PolygonVertices, PolygonHolesVertices, &HeightMesh, HeightMeshFrame, false, 99999, true);
+	FFrame3d MeshFrame_Discard;
+	GeneratePolygonMesh(PolygonVertices, PolygonHolesVertices, &HeightMesh, MeshFrame_Discard, false, 99999, true);
+
+	// We don't actually want the same FFrame3d that was computed by GeneratePolygonMesh
+	int NumVerts = PolygonVertices.Num();
+	FVector3d Centroid(0, 0, 0);
+	for (int k = 0; k < NumVerts; ++k)
+	{
+		Centroid += PolygonVertices[k];
+	}
+	Centroid /= (double)NumVerts;
+	FFrame3d HeightMeshFrame(Centroid, TQuaternion<double>(DrawPlaneOrientation));
+
 	HeightMechanic->Initialize( MoveTemp(HeightMesh), HeightMeshFrame, false);
 
 	ShowExtrudeMessage();
@@ -1078,10 +1089,11 @@ bool UDrawPolygonTool::GeneratePolygonMesh(const TArray<FVector>& Polygon, const
 	FVector3d Centroid(0, 0, 0);
 	for (int k = 0; k < NumVerts; ++k)
 	{
-		Centroid += Polygon[k];
+		Centroid.X += Polygon[k].X;
+		Centroid.Y += Polygon[k].Y;
 	}
 	Centroid /= (double)NumVerts;
-	WorldFrameOut.Origin = Centroid;
+	WorldFrameOut.Origin += WorldFrameOut.Rotation * Centroid;
 
 	// Compute outer polygon & bounds
 	auto VertexArrayToPolygon = [&WorldFrameOut](const TArray<FVector>& Vertices)
