Merge CL 11423483

EditMeshPoygons: fix check firing when switching modes in TriEdit tool. Cancel any in-progress operation before starting a new one.

#rnx

diff --git a/Engine/Plugins/Experimental/MeshModelingToolset/Source/MeshModelingTools/Private/EditMeshPolygonsTool.cpp b/Engine/Plugins/Experimental/MeshModelingToolset/Source/MeshModelingTools/Private/EditMeshPolygonsTool.cpp
--- a/Engine/Plugins/Experimental/MeshModelingToolset/Source/MeshModelingTools/Private/EditMeshPolygonsTool.cpp
+++ b/Engine/Plugins/Experimental/MeshModelingToolset/Source/MeshModelingTools/Private/EditMeshPolygonsTool.cpp
@@ -639,6 +639,8 @@ void UEditMeshPolygonsTool::Tick(float DeltaTime)
 
 	if (PendingAction != EEditMeshPolygonsToolActions::NoAction)
 	{
+		CancelMeshEditChange();
+
 		if (PendingAction == EEditMeshPolygonsToolActions::Extrude || PendingAction == EEditMeshPolygonsToolActions::Offset)
 		{
 			GetToolManager()->EmitObjectChange(this, MakeUnique<FBeginInteractivePolyEditChange>(CurrentOperationTimestamp), LOCTEXT("PolyMeshEditBeginExtrude", "Extrude"));
