GeometryProcessing, MeshModelingToolset: Avoid a bug caused by copying a non-copyable FRemesher by marking that class non-movable and non-copyable and switching the usage to reference semantics

#rnx

diff --git a/Engine/Plugins/Experimental/GeometryProcessing/Source/DynamicMesh/Public/MeshRefinerBase.h b/Engine/Plugins/Experimental/GeometryProcessing/Source/DynamicMesh/Public/MeshRefinerBase.h
--- a/Engine/Plugins/Experimental/GeometryProcessing/Source/DynamicMesh/Public/MeshRefinerBase.h
+++ b/Engine/Plugins/Experimental/GeometryProcessing/Source/DynamicMesh/Public/MeshRefinerBase.h
@@ -90,6 +90,10 @@ public:
 
 
 public:
+	FMeshRefinerBase(const FMeshRefinerBase&) = delete;
+	FMeshRefinerBase(FMeshRefinerBase&&) = delete;
+	FMeshRefinerBase& operator=(const FMeshRefinerBase&) = delete;
+	FMeshRefinerBase& operator=(FMeshRefinerBase&&) = delete;
 	virtual ~FMeshRefinerBase() {}
 
 	/** Get the current mesh we are operating on */
diff --git a/Engine/Plugins/Experimental/MeshModelingToolset/Source/MeshModelingTools/Private/DynamicMeshSculptTool.cpp b/Engine/Plugins/Experimental/MeshModelingToolset/Source/MeshModelingTools/Private/DynamicMeshSculptTool.cpp
--- a/Engine/Plugins/Experimental/MeshModelingToolset/Source/MeshModelingTools/Private/DynamicMeshSculptTool.cpp
+++ b/Engine/Plugins/Experimental/MeshModelingToolset/Source/MeshModelingTools/Private/DynamicMeshSculptTool.cpp
@@ -1333,10 +1333,12 @@ Free	TRUE	TRUE		FALSE				FALSE
 Ignore	TRUE	TRUE		FALSE				TRUE
 */
 
-FSubRegionRemesher
-UDynamicMeshSculptTool::MakeRemesher(FDynamicMesh3* Mesh, FDynamicMeshOctree3* Octree)
+void
+UDynamicMeshSculptTool::ConfigureRemesher(FSubRegionRemesher& Remesher)
 {
-	FSubRegionRemesher Remesher(Mesh);
+	FDynamicMesh3* Mesh = DynamicMeshComponent->GetMesh();
+	FDynamicMeshOctree3* Octree = DynamicMeshComponent->GetOctree();
+
 	double TargetEdgeLength = RemeshProperties->RelativeSize * InitialEdgeLength;
 	Remesher.SetTargetEdgeLength(TargetEdgeLength);
 
@@ -1348,7 +1350,7 @@ UDynamicMeshSculptTool::MakeRemesher(FDynamicMesh3* Mesh, FDynamicMeshOctree3* O
 	{
 		Remesher.MinEdgeLength = TargetEdgeLength * 0.1;
 
-		Remesher.CustomSmoothSpeedF = [this, &UseSmoothing](const FDynamicMesh3& Mesh, int vID)
+		Remesher.CustomSmoothSpeedF = [this, UseSmoothing](const FDynamicMesh3& Mesh, int vID)
 		{
 			FVector3d Pos = Mesh.GetVertex(vID);
 			double Falloff = CalculateBrushFalloff(Pos.Distance((FVector3d)LastBrushPosLocal));
@@ -1399,14 +1401,14 @@ UDynamicMeshSculptTool::MakeRemesher(FDynamicMesh3* Mesh, FDynamicMeshOctree3* O
 			Remesher.SetExternalConstraints(MoveTemp(Constraints));
 		}
 	}
-	return Remesher;
 }
 
 void UDynamicMeshSculptTool::RemeshROIPass()
 {
 	FDynamicMesh3* Mesh = DynamicMeshComponent->GetMesh();
 	FDynamicMeshOctree3* Octree = DynamicMeshComponent->GetOctree();
-	FSubRegionRemesher Remesher = MakeRemesher(Mesh, Octree);
+	FSubRegionRemesher Remesher(Mesh);
+	ConfigureRemesher(Remesher);
 
 	bool bIsUniformSmooth = (Remesher.SmoothType == FRemesher::ESmoothTypes::Uniform);
 	for (int k = 0; k < 5; ++k)
diff --git a/Engine/Plugins/Experimental/MeshModelingToolset/Source/MeshModelingTools/Public/DynamicMeshSculptTool.h b/Engine/Plugins/Experimental/MeshModelingToolset/Source/MeshModelingTools/Public/DynamicMeshSculptTool.h
--- a/Engine/Plugins/Experimental/MeshModelingToolset/Source/MeshModelingTools/Public/DynamicMeshSculptTool.h
+++ b/Engine/Plugins/Experimental/MeshModelingToolset/Source/MeshModelingTools/Public/DynamicMeshSculptTool.h
@@ -285,7 +285,7 @@ private:
 	bool bEnableRemeshing;
 	double InitialEdgeLength;
 
-	FSubRegionRemesher MakeRemesher(FDynamicMesh3* Mesh, FDynamicMeshOctree3* Octree);
+	void ConfigureRemesher(FSubRegionRemesher&);
 
 	bool bInDrag;
 
