edigrate CL 11716928 to Release-4.25

Fixed memory stomp related to FConstraintData symbol name collision.

- Both AnimationCore and MeshModelingTools define an FConstraintData structure
- AnimNode_Constraint.cpp(106) calls TArray<FConstraintData>::Reset which calls ResizeTo
- ResizeTo calls AllocatorInstance.CalculateSlackReserve passing in 128 bytes per element which is the correct size of the FConstraintData in AnimationCore
- When we end up in DefaultCalculateSlackReserve, BytesPerElement is 40 which is incorrect. That's the size of FConstraintData in MeshModelingTools.
- This causes the array to under-allocate space for the number of elements requested
- Then AnimNode_Constraint.cpp(130) writes beyond the allocated buffer during FConstraintData copy construction in the TArray::Add call

By renaming the FConstraintData in MeshModelingTools, the compiler no longer conflates the two different types together and emits the correct size.

diff --git a/Engine/Plugins/Experimental/MeshModelingToolset/Source/MeshModelingTools/Private/DeformMeshPolygonsTool.cpp b/Engine/Plugins/Experimental/MeshModelingToolset/Source/MeshModelingTools/Private/DeformMeshPolygonsTool.cpp
--- a/Engine/Plugins/Experimental/MeshModelingToolset/Source/MeshModelingTools/Private/DeformMeshPolygonsTool.cpp
+++ b/Engine/Plugins/Experimental/MeshModelingToolset/Source/MeshModelingTools/Private/DeformMeshPolygonsTool.cpp
@@ -68,7 +68,7 @@ void UDeformMeshPolygonsTransformProperties::PostEditChangeProperty(FPropertyCha
  */
 
 
-void FConstrainedMeshDeformerTask::UpdateDeformer(const ELaplacianWeightScheme SelectedWeightScheme, const FDynamicMesh3& SrcMesh, const TArray<FConstraintData>& ConstraintArray, const TArray<int32>& SrcIDBufferSubset, bool bNewTransaction, const FRichCurve* Curve)
+void FConstrainedMeshDeformerTask::UpdateDeformer(const ELaplacianWeightScheme SelectedWeightScheme, const FDynamicMesh3& SrcMesh, const TArray<FDeformerVertexConstraintData>& ConstraintArray, const TArray<int32>& SrcIDBufferSubset, bool bNewTransaction, const FRichCurve* Curve)
 {
 	bIsNewTransaction = bNewTransaction;
 	SrcMeshMaxVertexID = SrcMesh.MaxVertexID();
@@ -132,7 +132,7 @@ void FConstrainedMeshDeformerTask::DoWork()
 		//Update our deformer's constraints before deforming using the copy of the constraint buffer
 		for (int32 SubsetVertexID = 0; SubsetVertexID < SubsetConstraintBuffer.Num(); ++SubsetVertexID)
 		{
-			FConstraintData& CData = SubsetConstraintBuffer[SubsetVertexID];
+			FDeformerVertexConstraintData& CData = SubsetConstraintBuffer[SubsetVertexID];
 			ConstrainedDeformer->AddConstraint(SubsetVertexID, CData.Weight, CData.Position, CData.bPostFix);
 		}
 
@@ -146,7 +146,7 @@ void FConstrainedMeshDeformerTask::DoWork()
 
 		for (int32 SubsetVertexID = 0; SubsetVertexID < SubsetConstraintBuffer.Num(); ++SubsetVertexID)
 		{
-			FConstraintData& CData = SubsetConstraintBuffer[SubsetVertexID];
+			FDeformerVertexConstraintData& CData = SubsetConstraintBuffer[SubsetVertexID];
 			ConstrainedDeformer->UpdateConstraintPosition(SubsetVertexID, CData.Position, CData.bPostFix);
 		}
 
@@ -276,7 +276,7 @@ void FConstrainedMeshDeformerTask::ApplyAttenuation()
 	for (int32 SubVertexID = 0; SubVertexID < SubsetConstraintBuffer.Num(); ++SubVertexID)
 	{
 
-		FConstraintData& CData = SubsetConstraintBuffer[SubVertexID];
+		FDeformerVertexConstraintData& CData = SubsetConstraintBuffer[SubVertexID];
 		// Update bounding box
 		InPlaceMinMaxElements(Min, Max, CData.Position);
 		
@@ -297,7 +297,7 @@ void FConstrainedMeshDeformerTask::ApplyAttenuation()
 	{
 		if (!Handles.Contains(SubVertexID))
 		{
-			FConstraintData& CData = SubsetConstraintBuffer[SubVertexID];
+			FDeformerVertexConstraintData& CData = SubsetConstraintBuffer[SubVertexID];
 			double T = CData.Position.Distance(FMath::ClosestPointOnSegment((FVector)CData.Position, (FVector)MinHandles, (FVector)MaxHandles)) / ExtentLength;
 			CData.Weight = WeightAttenuationCurve.Eval(T) * LeastWeight;
 		}
@@ -330,7 +330,7 @@ void FGroupTopologyLaplacianDeformer::InitializeConstraintBuffer()
 	
 	for (int32 VertexID : Mesh->VertexIndicesItr())
 	{
-		FConstraintData& CD = SrcMeshConstraintBuffer[VertexID];
+		FDeformerVertexConstraintData& CD = SrcMeshConstraintBuffer[VertexID];
 		CD.Position = Mesh->GetVertex(VertexID); 
 		CD.Weight   = 0.0;
 		CD.bPostFix = false;
@@ -542,7 +542,7 @@ void FGroupTopologyLaplacianDeformer::UpdateSolution(FDynamicMesh3 * TargetMesh,
 	for (int32 VertexID : ModifiedVertices)
 	{
 		//Get the vertex's data from the constraint buffer
-		FConstraintData& CData = SrcMeshConstraintBuffer[VertexID];
+		FDeformerVertexConstraintData& CData = SrcMeshConstraintBuffer[VertexID];
 
 		CData.Position = TargetMesh->GetVertex(VertexID);
 		CData.Weight   = 0.0; //A weight of zero is used to allow this point to move freely when moving the handles
@@ -555,7 +555,7 @@ void FGroupTopologyLaplacianDeformer::UpdateSolution(FDynamicMesh3 * TargetMesh,
 		const FVector3d DeformPos = HandleVertexDeformFunc(TargetMesh, VertexID);
 
 		//Get the vertex's data from the constraint buffer
-		FConstraintData& CData = SrcMeshConstraintBuffer[VertexID];
+		FDeformerVertexConstraintData& CData = SrcMeshConstraintBuffer[VertexID];
 
 		//Set the new vertex data
 		CData.Position = DeformPos;
diff --git a/Engine/Plugins/Experimental/MeshModelingToolset/Source/MeshModelingTools/Public/DeformMeshPolygonsTool.h b/Engine/Plugins/Experimental/MeshModelingToolset/Source/MeshModelingTools/Public/DeformMeshPolygonsTool.h
--- a/Engine/Plugins/Experimental/MeshModelingToolset/Source/MeshModelingTools/Public/DeformMeshPolygonsTool.h
+++ b/Engine/Plugins/Experimental/MeshModelingToolset/Source/MeshModelingTools/Public/DeformMeshPolygonsTool.h
@@ -23,7 +23,7 @@
 class FMeshVertexChangeBuilder;
 class FGroupTopologyLaplacianDeformer;
 class FDeformTask;
-struct FConstraintData;
+struct FDeformerVertexConstraintData;
 
 /**
  * ToolBuilder
@@ -149,9 +149,9 @@ public:
 
 //Stores per-vertex data needed by the laplacian deformer object
 //TODO: May be a candidate for a subclass of the FGroupTopologyLaplacianDeformer
-struct FConstraintData
+struct FDeformerVertexConstraintData
 {
-	FConstraintData& operator=(const FConstraintData& other)
+	FDeformerVertexConstraintData& operator=(const FDeformerVertexConstraintData& other)
 	{
 		Position = other.Position;
 		Weight   = other.Weight;
@@ -172,7 +172,7 @@ struct FConstraintData
 *
 *   
 *   // define constraints.  Need Constraints[VertID] to hold the constraints for the corresponding vertex.
-*   TArray<FConstraintData> Constraints;
+*   TArray<FDeformerVertexConstraintData> Constraints;
 *   ....
 *
 *   // populate with the VertexIDs of the vertices that are in the region you wish to deform.  
@@ -222,7 +222,7 @@ public:
 	
 	/** Called by the main thread in the tool, this copies the Constraint buffer right before the task begins on another thread.
 	  * Ensures the FConstrainedMeshDeformer is using correct mesh subset and the selected settings, then updates on change in properties, i.e. weight scheme */
-	void UpdateDeformer(const ELaplacianWeightScheme SelectedWeightScheme, const FDynamicMesh3& Mesh, const TArray<FConstraintData>& ConstraintArray, const TArray<int32>& SrcIDBufferSubset, bool bNewTransaction, const FRichCurve* Curve);
+	void UpdateDeformer(const ELaplacianWeightScheme SelectedWeightScheme, const FDynamicMesh3& Mesh, const TArray<FDeformerVertexConstraintData>& ConstraintArray, const TArray<int32>& SrcIDBufferSubset, bool bNewTransaction, const FRichCurve* Curve);
 
 	/** Required by the FAsyncTaskExecutor */
 	void SetAbortSource(bool* bAbort) { bAbortSource = bAbort; };
@@ -249,7 +249,7 @@ private:
 	TArray<FVector3d> SubsetPositionBuffer;
 
 	/** constraint data for each vertex in subset mesh - for use by the deformer*/
-	TArray<FConstraintData> SubsetConstraintBuffer;
+	TArray<FDeformerVertexConstraintData> SubsetConstraintBuffer;
 
 	FRichCurve WeightAttenuationCurve;
 
@@ -349,7 +349,7 @@ public:
 	const TArray<FROIFace>& GetROIFaces() const { return ROIFaces; }
 
 	/** Stores the position of the vertex constraints and corresponding weights for the entire mesh.  This is used as a form of scratch space.*/
-	TArray<FConstraintData> SrcMeshConstraintBuffer;
+	TArray<FDeformerVertexConstraintData> SrcMeshConstraintBuffer;
 
 	/** Array of vertex indices organized in groups of three - basically an index buffer - that defines the subset of the mesh that the deformation task will work on.*/
 	TArray<int32> SubsetIDBuffer;
