ModelingTools: remove copypasta that resulted in log spam
#rnx

diff --git a/Engine/Plugins/Experimental/MeshModelingToolset/Source/MeshModelingTools/Private/CombineMeshesTool.cpp b/Engine/Plugins/Experimental/MeshModelingToolset/Source/MeshModelingTools/Private/CombineMeshesTool.cpp
--- a/Engine/Plugins/Experimental/MeshModelingToolset/Source/MeshModelingTools/Private/CombineMeshesTool.cpp
+++ b/Engine/Plugins/Experimental/MeshModelingToolset/Source/MeshModelingTools/Private/CombineMeshesTool.cpp
@@ -189,7 +189,6 @@ void UCombineMeshesTool::UpdateAssets()
 			TUniquePtr<FPrimitiveComponentTarget>& ComponentTarget = ComponentTargets[ComponentIdx];
 
 			FMeshDescriptionToDynamicMesh Converter;
-			Converter.bPrintDebugMessages = true;
 			FDynamicMesh3 ComponentDMesh;
 			Converter.Convert(ComponentTarget->GetMesh(), ComponentDMesh);
 
diff --git a/Engine/Plugins/Experimental/MeshModelingToolset/Source/MeshModelingTools/Private/EditNormalsTool.cpp b/Engine/Plugins/Experimental/MeshModelingToolset/Source/MeshModelingTools/Private/EditNormalsTool.cpp
--- a/Engine/Plugins/Experimental/MeshModelingToolset/Source/MeshModelingTools/Private/EditNormalsTool.cpp
+++ b/Engine/Plugins/Experimental/MeshModelingToolset/Source/MeshModelingTools/Private/EditNormalsTool.cpp
@@ -139,7 +139,6 @@ void UEditNormalsTool::UpdateNumPreviews()
 			OpFactory->ComponentIndex = PreviewIdx;
 			OriginalDynamicMeshes[PreviewIdx] = MakeShared<FDynamicMesh3>();
 			FMeshDescriptionToDynamicMesh Converter;
-			Converter.bPrintDebugMessages = true;
 			Converter.Convert(ComponentTargets[PreviewIdx]->GetMesh(), *OriginalDynamicMeshes[PreviewIdx]);
 
 			UMeshOpPreviewWithBackgroundCompute* Preview = Previews.Add_GetRef(NewObject<UMeshOpPreviewWithBackgroundCompute>(OpFactory, "Preview"));
diff --git a/Engine/Plugins/Experimental/MeshModelingToolset/Source/MeshModelingTools/Private/MeshSpaceDeformerTool.cpp b/Engine/Plugins/Experimental/MeshModelingToolset/Source/MeshModelingTools/Private/MeshSpaceDeformerTool.cpp
--- a/Engine/Plugins/Experimental/MeshModelingToolset/Source/MeshModelingTools/Private/MeshSpaceDeformerTool.cpp
+++ b/Engine/Plugins/Experimental/MeshModelingToolset/Source/MeshModelingTools/Private/MeshSpaceDeformerTool.cpp
@@ -146,7 +146,6 @@ void UMeshSpaceDeformerTool::Setup()
 	{
 		OriginalDynamicMesh = MakeShared< FDynamicMesh3 >();
 		FMeshDescriptionToDynamicMesh Converter;
-		Converter.bPrintDebugMessages = false;
 		Converter.Convert(ComponentTarget->GetMesh(), *OriginalDynamicMesh);
 	}
 
diff --git a/Engine/Plugins/Experimental/MeshModelingToolset/Source/MeshModelingTools/Private/PlaneCutTool.cpp b/Engine/Plugins/Experimental/MeshModelingToolset/Source/MeshModelingTools/Private/PlaneCutTool.cpp
--- a/Engine/Plugins/Experimental/MeshModelingToolset/Source/MeshModelingTools/Private/PlaneCutTool.cpp
+++ b/Engine/Plugins/Experimental/MeshModelingToolset/Source/MeshModelingTools/Private/PlaneCutTool.cpp
@@ -171,9 +171,6 @@ void UPlaneCutTool::Setup()
 		TUniquePtr<FPrimitiveComponentTarget>& ComponentTarget = ComponentTargets[Idx];
 		FDynamicMesh3* OriginalDynamicMesh = new FDynamicMesh3;
 		FMeshDescriptionToDynamicMesh Converter;
-		Converter.bPrintDebugMessages = true;
-
-		
 		Converter.Convert(ComponentTarget->GetMesh(), *OriginalDynamicMesh);
 		OriginalDynamicMesh->EnableAttributes();
 		TDynamicMeshScalarTriangleAttribute<int>* SubObjectIDs = new TDynamicMeshScalarTriangleAttribute<int>(OriginalDynamicMesh);
diff --git a/Engine/Plugins/Experimental/MeshModelingToolset/Source/MeshModelingTools/Private/RemeshMeshTool.cpp b/Engine/Plugins/Experimental/MeshModelingToolset/Source/MeshModelingTools/Private/RemeshMeshTool.cpp
--- a/Engine/Plugins/Experimental/MeshModelingToolset/Source/MeshModelingTools/Private/RemeshMeshTool.cpp
+++ b/Engine/Plugins/Experimental/MeshModelingToolset/Source/MeshModelingTools/Private/RemeshMeshTool.cpp
@@ -111,7 +111,6 @@ void URemeshMeshTool::Setup()
 
 	OriginalMesh = MakeShared<FDynamicMesh3>();
 	FMeshDescriptionToDynamicMesh Converter;
-	Converter.bPrintDebugMessages = true;
 	Converter.Convert(ComponentTarget->GetMesh(), *OriginalMesh);
 
 	Preview->PreviewMesh->SetTransform(ComponentTarget->GetWorldTransform());
diff --git a/Engine/Plugins/Experimental/MeshModelingToolset/Source/MeshModelingTools/Private/RemoveOccludedTrianglesTool.cpp b/Engine/Plugins/Experimental/MeshModelingToolset/Source/MeshModelingTools/Private/RemoveOccludedTrianglesTool.cpp
--- a/Engine/Plugins/Experimental/MeshModelingToolset/Source/MeshModelingTools/Private/RemoveOccludedTrianglesTool.cpp
+++ b/Engine/Plugins/Experimental/MeshModelingToolset/Source/MeshModelingTools/Private/RemoveOccludedTrianglesTool.cpp
@@ -186,7 +186,6 @@ void URemoveOccludedTrianglesTool::SetupPreviews()
 			OpFactory->PreviewIdx = PreviewIdx;
 			OriginalDynamicMeshes[PreviewIdx] = MakeShared<FDynamicMesh3>();
 			FMeshDescriptionToDynamicMesh Converter;
-			Converter.bPrintDebugMessages = true;
 			Converter.Convert(ComponentTargets[TargetIdx]->GetMesh(), *OriginalDynamicMeshes[PreviewIdx]);
 
 			UMeshOpPreviewWithBackgroundCompute* Preview = Previews.Add_GetRef(NewObject<UMeshOpPreviewWithBackgroundCompute>(OpFactory, "Preview"));
diff --git a/Engine/Plugins/Experimental/MeshModelingToolset/Source/MeshModelingTools/Private/SmoothMeshTool.cpp b/Engine/Plugins/Experimental/MeshModelingToolset/Source/MeshModelingTools/Private/SmoothMeshTool.cpp
--- a/Engine/Plugins/Experimental/MeshModelingToolset/Source/MeshModelingTools/Private/SmoothMeshTool.cpp
+++ b/Engine/Plugins/Experimental/MeshModelingToolset/Source/MeshModelingTools/Private/SmoothMeshTool.cpp
@@ -76,7 +76,6 @@ void USmoothMeshTool::Setup()
 
 	// populate the SrcDynamicMesh with a conversion of the input mesh.
 	FMeshDescriptionToDynamicMesh Converter;
-	Converter.bPrintDebugMessages = false;
 	Converter.Convert(ComponentTarget->GetMesh(), SrcDynamicMesh);
 
 	// Initialize the preview mesh with a copy of the source mesh.
diff --git a/Engine/Plugins/Experimental/MeshModelingToolset/Source/MeshModelingTools/Private/UVProjectionTool.cpp b/Engine/Plugins/Experimental/MeshModelingToolset/Source/MeshModelingTools/Private/UVProjectionTool.cpp
--- a/Engine/Plugins/Experimental/MeshModelingToolset/Source/MeshModelingTools/Private/UVProjectionTool.cpp
+++ b/Engine/Plugins/Experimental/MeshModelingToolset/Source/MeshModelingTools/Private/UVProjectionTool.cpp
@@ -178,7 +178,6 @@ void UUVProjectionTool::UpdateNumPreviews()
 			OpFactory->ComponentIndex = PreviewIdx;
 			OriginalDynamicMeshes[PreviewIdx] = MakeShared<FDynamicMesh3>();
 			FMeshDescriptionToDynamicMesh Converter;
-			Converter.bPrintDebugMessages = true;
 			Converter.Convert(ComponentTargets[PreviewIdx]->GetMesh(), *OriginalDynamicMeshes[PreviewIdx]);
 
 			FVector Center, Extents;
diff --git a/Engine/Plugins/Experimental/MeshModelingToolset/Source/MeshModelingToolsEditorOnly/Private/ParameterizeMeshTool.cpp b/Engine/Plugins/Experimental/MeshModelingToolset/Source/MeshModelingToolsEditorOnly/Private/ParameterizeMeshTool.cpp
--- a/Engine/Plugins/Experimental/MeshModelingToolset/Source/MeshModelingToolsEditorOnly/Private/ParameterizeMeshTool.cpp
+++ b/Engine/Plugins/Experimental/MeshModelingToolset/Source/MeshModelingToolsEditorOnly/Private/ParameterizeMeshTool.cpp
@@ -122,7 +122,6 @@ void UParameterizeMeshTool::Setup()
 	{
 		FDynamicMesh3 Mesh;
 		FMeshDescriptionToDynamicMesh Converter;
-		Converter.bPrintDebugMessages = false;
 		Converter.Convert(InputMesh.Get(), Mesh);
 		bHasGroups = FaceGroupUtil::HasMultipleGroups(Mesh);
 
diff --git a/Engine/Plugins/Experimental/MeshModelingToolset/Source/MeshModelingToolsEditorOnly/Private/SimplifyMeshTool.cpp b/Engine/Plugins/Experimental/MeshModelingToolset/Source/MeshModelingToolsEditorOnly/Private/SimplifyMeshTool.cpp
--- a/Engine/Plugins/Experimental/MeshModelingToolset/Source/MeshModelingToolsEditorOnly/Private/SimplifyMeshTool.cpp
+++ b/Engine/Plugins/Experimental/MeshModelingToolset/Source/MeshModelingToolsEditorOnly/Private/SimplifyMeshTool.cpp
@@ -121,7 +121,6 @@ void USimplifyMeshTool::Setup()
 		EnterProgressFrame(1);
 		OriginalMesh = MakeShared<FDynamicMesh3>();
 		FMeshDescriptionToDynamicMesh Converter;
-		Converter.bPrintDebugMessages = true;
 		Converter.Convert(ComponentTarget->GetMesh(), *OriginalMesh);
 		EnterProgressFrame(2);
 		OriginalMeshSpatial = MakeShared<FDynamicMeshAABBTree3>(OriginalMesh.Get(), true);
diff --git a/Engine/Plugins/Experimental/MeshModelingToolset/Source/MeshModelingToolsEditorOnly/Private/UVLayoutTool.cpp b/Engine/Plugins/Experimental/MeshModelingToolset/Source/MeshModelingToolsEditorOnly/Private/UVLayoutTool.cpp
--- a/Engine/Plugins/Experimental/MeshModelingToolset/Source/MeshModelingToolsEditorOnly/Private/UVLayoutTool.cpp
+++ b/Engine/Plugins/Experimental/MeshModelingToolset/Source/MeshModelingToolsEditorOnly/Private/UVLayoutTool.cpp
@@ -153,7 +153,6 @@ void UUVLayoutTool::UpdateNumPreviews()
 			OpFactory->ComponentIndex = PreviewIdx;
 			OriginalDynamicMeshes[PreviewIdx] = MakeShared<FDynamicMesh3>();
 			FMeshDescriptionToDynamicMesh Converter;
-			Converter.bPrintDebugMessages = true;
 			Converter.Convert(ComponentTargets[PreviewIdx]->GetMesh(), *OriginalDynamicMeshes[PreviewIdx]);
 
 			UMeshOpPreviewWithBackgroundCompute* Preview = Previews.Add_GetRef(NewObject<UMeshOpPreviewWithBackgroundCompute>(OpFactory, "Preview"));
diff --git a/Engine/Plugins/Experimental/MeshModelingToolset/Source/ModelingComponents/Private/OctreeDynamicMeshComponent.cpp b/Engine/Plugins/Experimental/MeshModelingToolset/Source/ModelingComponents/Private/OctreeDynamicMeshComponent.cpp
--- a/Engine/Plugins/Experimental/MeshModelingToolset/Source/ModelingComponents/Private/OctreeDynamicMeshComponent.cpp
+++ b/Engine/Plugins/Experimental/MeshModelingToolset/Source/ModelingComponents/Private/OctreeDynamicMeshComponent.cpp
@@ -50,7 +50,6 @@ UOctreeDynamicMeshComponent::UOctreeDynamicMeshComponent(const FObjectInitialize
 void UOctreeDynamicMeshComponent::InitializeMesh(FMeshDescription* MeshDescription)
 {
 	FMeshDescriptionToDynamicMesh Converter;
-	Converter.bPrintDebugMessages = true;
 	Mesh->Clear();
 	Converter.Convert(MeshDescription, *Mesh);
 
diff --git a/Engine/Plugins/Experimental/MeshModelingToolset/Source/ModelingComponents/Private/SimpleDynamicMeshComponent.cpp b/Engine/Plugins/Experimental/MeshModelingToolset/Source/ModelingComponents/Private/SimpleDynamicMeshComponent.cpp
--- a/Engine/Plugins/Experimental/MeshModelingToolset/Source/ModelingComponents/Private/SimpleDynamicMeshComponent.cpp
+++ b/Engine/Plugins/Experimental/MeshModelingToolset/Source/ModelingComponents/Private/SimpleDynamicMeshComponent.cpp
@@ -50,7 +50,6 @@ USimpleDynamicMeshComponent::USimpleDynamicMeshComponent(const FObjectInitialize
 void USimpleDynamicMeshComponent::InitializeMesh(FMeshDescription* MeshDescription)
 {
 	FMeshDescriptionToDynamicMesh Converter;
-	Converter.bPrintDebugMessages = true;
 	Mesh->Clear();
 	Converter.Convert(MeshDescription, *Mesh);
 	if (TangentsType == EDynamicMeshTangentCalcType::ExternallyCalculated)
diff --git a/Engine/Plugins/Experimental/MeshModelingToolset/Source/ModelingOperatorsEditorOnly/Private/CompositionOps/VoxelBooleanMeshesOp.cpp b/Engine/Plugins/Experimental/MeshModelingToolset/Source/ModelingOperatorsEditorOnly/Private/CompositionOps/VoxelBooleanMeshesOp.cpp
--- a/Engine/Plugins/Experimental/MeshModelingToolset/Source/ModelingOperatorsEditorOnly/Private/CompositionOps/VoxelBooleanMeshesOp.cpp
+++ b/Engine/Plugins/Experimental/MeshModelingToolset/Source/ModelingOperatorsEditorOnly/Private/CompositionOps/VoxelBooleanMeshesOp.cpp
@@ -112,7 +112,6 @@ void FVoxelBooleanMeshesOp::CalculateResult(FProgressCancel* Progress)
 
 	// Convert to dynamic mesh
 	FMeshDescriptionToDynamicMesh Converter;
-	Converter.bPrintDebugMessages = true;
 	Converter.Convert(&ResultMeshDescription, *ResultMesh);
 
 	if (Progress->Cancelled())
diff --git a/Engine/Plugins/Experimental/MeshModelingToolset/Source/ModelingOperatorsEditorOnly/Private/CompositionOps/VoxelMergeMeshesOp.cpp b/Engine/Plugins/Experimental/MeshModelingToolset/Source/ModelingOperatorsEditorOnly/Private/CompositionOps/VoxelMergeMeshesOp.cpp
--- a/Engine/Plugins/Experimental/MeshModelingToolset/Source/ModelingOperatorsEditorOnly/Private/CompositionOps/VoxelMergeMeshesOp.cpp
+++ b/Engine/Plugins/Experimental/MeshModelingToolset/Source/ModelingOperatorsEditorOnly/Private/CompositionOps/VoxelMergeMeshesOp.cpp
@@ -72,7 +72,6 @@ void FVoxelMergeMeshesOp::CalculateResult(FProgressCancel* Progress)
 
 	// Convert to dynamic mesh
 	FMeshDescriptionToDynamicMesh Converter;
-	Converter.bPrintDebugMessages = true;
 	Converter.Convert(&MergedMeshesDescription, *ResultMesh);
 
 	if (Progress->Cancelled())
