ModelingMode: reduce thickness of rendered edges in triangle editing mode
#rnx

diff --git a/Engine/Plugins/Experimental/MeshModelingToolset/Source/MeshModelingTools/Private/EditMeshPolygonsTool.cpp b/Engine/Plugins/Experimental/MeshModelingToolset/Source/MeshModelingTools/Private/EditMeshPolygonsTool.cpp
--- a/Engine/Plugins/Experimental/MeshModelingToolset/Source/MeshModelingTools/Private/EditMeshPolygonsTool.cpp
+++ b/Engine/Plugins/Experimental/MeshModelingToolset/Source/MeshModelingTools/Private/EditMeshPolygonsTool.cpp
@@ -153,6 +153,10 @@ void UEditMeshPolygonsTool::Setup()
 	SelectionMechanic = NewObject<UPolygonSelectionMechanic>(this);
 	SelectionMechanic->Setup(this);
 	SelectionMechanic->OnSelectionChanged.AddUObject(this, &UEditMeshPolygonsTool::OnSelectionModifiedEvent);
+	if (bTriangleMode)
+	{
+		SelectionMechanic->PolyEdgesRenderer.LineThickness = 1.0;
+	}
 
 	// initialize AABBTree
 	MeshSpatial.SetMesh(DynamicMeshComponent->GetMesh());
diff --git a/Engine/Plugins/Experimental/MeshModelingToolset/Source/ModelingComponents/Public/Selection/PolygonSelectionMechanic.h b/Engine/Plugins/Experimental/MeshModelingToolset/Source/ModelingComponents/Public/Selection/PolygonSelectionMechanic.h
--- a/Engine/Plugins/Experimental/MeshModelingToolset/Source/ModelingComponents/Public/Selection/PolygonSelectionMechanic.h
+++ b/Engine/Plugins/Experimental/MeshModelingToolset/Source/ModelingComponents/Public/Selection/PolygonSelectionMechanic.h
@@ -160,6 +160,7 @@ protected:
 	TUniquePtr<FPolygonSelectionMechanicSelectionChange> ActiveChange;
 
 	FViewCameraState CameraState;
+public:
 	FToolDataVisualizer PolyEdgesRenderer;
 	FToolDataVisualizer HilightRenderer;
 	FToolDataVisualizer SelectionRenderer;
