ModelingMode: recalculate normals on preview path meshes in DrawPolyPathTool. Handle undefined-UVs case in TMeshTangents.
#rnx

diff --git a/Engine/Plugins/Experimental/GeometryProcessing/Source/DynamicMesh/Private/MeshTangents.cpp b/Engine/Plugins/Experimental/GeometryProcessing/Source/DynamicMesh/Private/MeshTangents.cpp
--- a/Engine/Plugins/Experimental/GeometryProcessing/Source/DynamicMesh/Private/MeshTangents.cpp
+++ b/Engine/Plugins/Experimental/GeometryProcessing/Source/DynamicMesh/Private/MeshTangents.cpp
@@ -35,7 +35,7 @@ void TMeshTangents<RealType>::Internal_ComputePerTriangleTangents(const FDynamic
 
 	ParallelFor(MaxTriangleID, [this, NormalOverlay, UVOverlay](int TriangleID)
 	{
-		if (Mesh->IsTriangle(TriangleID) == false)
+		if (Mesh->IsTriangle(TriangleID) == false || UVOverlay->IsSetTriangle(TriangleID) == false)
 		{
 			return;
 		}
diff --git a/Engine/Plugins/Experimental/MeshModelingToolset/Source/MeshModelingTools/Private/DrawPolyPathTool.cpp b/Engine/Plugins/Experimental/MeshModelingToolset/Source/MeshModelingTools/Private/DrawPolyPathTool.cpp
--- a/Engine/Plugins/Experimental/MeshModelingToolset/Source/MeshModelingTools/Private/DrawPolyPathTool.cpp
+++ b/Engine/Plugins/Experimental/MeshModelingToolset/Source/MeshModelingTools/Private/DrawPolyPathTool.cpp
@@ -18,6 +18,7 @@
 #include "MeshTransforms.h"
 #include "Selection/ToolSelectionUtil.h"
 #include "Operations/ExtrudeMesh.h"
+#include "MeshNormals.h"
 
 #define LOCTEXT_NAMESPACE "UDrawPolyPathTool"
 
@@ -511,6 +512,8 @@ void UDrawPolyPathTool::GeneratePathMesh(FDynamicMesh3& Mesh)
 			Mesh.SetVertexUV(NumPoints+k, FVector2f((float)ArcLengths[k], (float)k));
 			ShiftX += DeltaX;
 		}
+
+		FMeshNormals::QuickRecomputeOverlayNormals(Mesh);
 	}
 }
 
@@ -601,6 +604,8 @@ void UDrawPolyPathTool::GenerateExtrudeMesh(FDynamicMesh3& PathMesh)
 	Extruder.UVScaleFactor = 1.0 / Bounds.MaxDim();
 	Extruder.IsPositiveOffset = (CurHeight >= 0);
 	Extruder.Apply();
+
+	FMeshNormals::QuickRecomputeOverlayNormals(PathMesh);
 }
 
 void UDrawPolyPathTool::GenerateRampMesh(FDynamicMesh3& PathMesh)
@@ -623,6 +628,8 @@ void UDrawPolyPathTool::GenerateRampMesh(FDynamicMesh3& PathMesh)
 	Extruder.UVScaleFactor = 1.0 / Bounds.MaxDim();
 	Extruder.IsPositiveOffset = (CurHeight >= 0);
 	Extruder.Apply();
+
+	FMeshNormals::QuickRecomputeOverlayNormals(PathMesh);
 }
 
 
diff --git a/Engine/Plugins/Experimental/MeshModelingToolset/Source/MeshModelingTools/Public/DrawPolyPathTool.h b/Engine/Plugins/Experimental/MeshModelingToolset/Source/MeshModelingTools/Public/DrawPolyPathTool.h
--- a/Engine/Plugins/Experimental/MeshModelingToolset/Source/MeshModelingTools/Public/DrawPolyPathTool.h
+++ b/Engine/Plugins/Experimental/MeshModelingToolset/Source/MeshModelingTools/Public/DrawPolyPathTool.h
@@ -82,7 +82,7 @@ public:
 	float Height = 10.0f;
 
 	UPROPERTY(EditAnywhere, Category = Shape, meta = (UIMin = "0.01", UIMax = "1.0", ClampMin = "0", ClampMax = "100.0"))
-	float RampStartRatio = 0.25f;
+	float RampStartRatio = 0.05f;
 
 
 	UPROPERTY(EditAnywhere, Category = Snapping)
