ModelingTools: typo prevented normals on DynamicMeshComponent from being recomputed in certain cases
#rnx

diff --git a/Engine/Plugins/Experimental/MeshModelingToolset/Source/ModelingComponents/Private/BaseDynamicMeshSceneProxy.h b/Engine/Plugins/Experimental/MeshModelingToolset/Source/ModelingComponents/Private/BaseDynamicMeshSceneProxy.h
--- a/Engine/Plugins/Experimental/MeshModelingToolset/Source/ModelingComponents/Private/BaseDynamicMeshSceneProxy.h
+++ b/Engine/Plugins/Experimental/MeshModelingToolset/Source/ModelingComponents/Private/BaseDynamicMeshSceneProxy.h
@@ -576,7 +576,7 @@ public:
 		{
 			FIndex3i Tri = Mesh->GetTriangle(TriangleID);
 
-			FIndex3i TriNormal = (bUpdateColors && NormalOverlay != nullptr) ? NormalOverlay->GetTriangle(TriangleID) : FIndex3i::Zero();
+			FIndex3i TriNormal = (bUpdateNormals && NormalOverlay != nullptr) ? NormalOverlay->GetTriangle(TriangleID) : FIndex3i::Zero();
 
 			FColor TriColor = ConstantVertexColor;
 			if (bUpdateColors && bUsePerTriangleColor && PerTriangleColorFunc != nullptr)
