MeshModeling: Make checkerboard work with the Existing Mesh Material property save restore

THIS CHANGE IS AN EDIGRATE OF CL 11130397 IN DEV-EDITOR

#rnx

diff --git a/Engine/Plugins/Experimental/MeshModelingToolset/Source/MeshModelingTools/Private/MeshInspectorTool.cpp b/Engine/Plugins/Experimental/MeshModelingToolset/Source/MeshModelingTools/Private/MeshInspectorTool.cpp
--- a/Engine/Plugins/Experimental/MeshModelingToolset/Source/MeshModelingTools/Private/MeshInspectorTool.cpp
+++ b/Engine/Plugins/Experimental/MeshModelingToolset/Source/MeshModelingTools/Private/MeshInspectorTool.cpp
@@ -126,6 +126,8 @@ void UMeshInspectorTool::Setup()
 	UMeshAnalysisProperties* MeshAnalysis = NewObject<UMeshAnalysisProperties>(this);
 	MeshAnalysis->Update(*PreviewMesh->GetPreviewDynamicMesh(), ComponentTarget->GetWorldTransform());
 	AddToolPropertySource(MeshAnalysis);
+
+	UpdateVisualization();
 }
 
 
@@ -286,8 +288,13 @@ void UMeshInspectorTool::Render(IToolsContextRenderAPI* RenderAPI)
 void UMeshInspectorTool::OnPropertyModified(UObject* PropertySet, FProperty* Property)
 {
 	GetToolManager()->PostInvalidation();
+	UpdateVisualization();
+}
+
+void UMeshInspectorTool::UpdateVisualization()
+{
 	PreviewMesh->EnableWireframe(Settings->bWireframe);
-	
+
 	MaterialSettings->UpdateMaterials();
 	UMaterialInterface* OverrideMaterial = MaterialSettings->GetActiveOverrideMaterial();
 	if (OverrideMaterial == nullptr)
diff --git a/Engine/Plugins/Experimental/MeshModelingToolset/Source/MeshModelingTools/Private/Properties/MeshMaterialProperties.cpp b/Engine/Plugins/Experimental/MeshModelingToolset/Source/MeshModelingTools/Private/Properties/MeshMaterialProperties.cpp
--- a/Engine/Plugins/Experimental/MeshModelingToolset/Source/MeshModelingTools/Private/Properties/MeshMaterialProperties.cpp
+++ b/Engine/Plugins/Experimental/MeshModelingToolset/Source/MeshModelingTools/Private/Properties/MeshMaterialProperties.cpp
@@ -79,7 +79,6 @@ void UExistingMeshMaterialProperties::SaveProperties(UInteractiveTool* SaveFromT
 	PropertyCache->MaterialMode = this->MaterialMode;
 	PropertyCache->CheckerDensity = this->CheckerDensity;
 	PropertyCache->OverrideMaterial = this->OverrideMaterial;
-	PropertyCache->CheckerMaterial = this->CheckerMaterial;
 }
 
 void UExistingMeshMaterialProperties::RestoreProperties(UInteractiveTool* RestoreToTool)
@@ -88,7 +87,7 @@ void UExistingMeshMaterialProperties::RestoreProperties(UInteractiveTool* Restor
 	this->MaterialMode = PropertyCache->MaterialMode;
 	this->CheckerDensity = PropertyCache->CheckerDensity;
 	this->OverrideMaterial = PropertyCache->OverrideMaterial;
-	this->CheckerMaterial = PropertyCache->CheckerMaterial;
+	this->Setup();
 }
 
 void UMeshEditingViewProperties::SaveProperties(UInteractiveTool* SaveFromTool)
@@ -106,6 +105,4 @@ void UMeshEditingViewProperties::RestoreProperties(UInteractiveTool* RestoreToTo
 }
 
 
-
-
 #undef LOCTEXT_NAMESPACE
diff --git a/Engine/Plugins/Experimental/MeshModelingToolset/Source/MeshModelingTools/Private/UVProjectionTool.cpp b/Engine/Plugins/Experimental/MeshModelingToolset/Source/MeshModelingTools/Private/UVProjectionTool.cpp
--- a/Engine/Plugins/Experimental/MeshModelingToolset/Source/MeshModelingTools/Private/UVProjectionTool.cpp
+++ b/Engine/Plugins/Experimental/MeshModelingToolset/Source/MeshModelingTools/Private/UVProjectionTool.cpp
@@ -149,6 +149,7 @@ void UUVProjectionTool::Setup()
 	{
 		Preview->InvalidateResult();
 	}
+	UpdateVisualization();
 }
 
 
@@ -314,16 +315,15 @@ void UUVProjectionTool::PostEditChangeProperty(FPropertyChangedEvent& PropertyCh
 }
 #endif
 
-void UUVProjectionTool::OnPropertyModified(UObject* PropertySet, FProperty* Property)
+void UUVProjectionTool::UpdateVisualization()
 {
-	// if we don't know what changed, or we know checker density changed, update checker material
 	MaterialSettings->UpdateMaterials();
 	for (int PreviewIdx = 0; PreviewIdx < Previews.Num(); PreviewIdx++)
 	{
 		UMeshOpPreviewWithBackgroundCompute* Preview = Previews[PreviewIdx];
 		Preview->OverrideMaterial = MaterialSettings->GetActiveOverrideMaterial();
 	}
-	
+
 	UpdateNumPreviews();
 	for (UMeshOpPreviewWithBackgroundCompute* Preview : Previews)
 	{
@@ -331,6 +331,11 @@ void UUVProjectionTool::OnPropertyModified(UObject* PropertySet, FProperty* Prop
 	}
 }
 
+void UUVProjectionTool::OnPropertyModified(UObject* PropertySet, FProperty* Property)
+{
+	UpdateVisualization();
+}
+
 
 void UUVProjectionTool::TransformChanged(UTransformProxy* Proxy, FTransform Transform)
 {
diff --git a/Engine/Plugins/Experimental/MeshModelingToolset/Source/MeshModelingTools/Public/MeshInspectorTool.h b/Engine/Plugins/Experimental/MeshModelingToolset/Source/MeshModelingTools/Public/MeshInspectorTool.h
--- a/Engine/Plugins/Experimental/MeshModelingToolset/Source/MeshModelingTools/Public/MeshInspectorTool.h
+++ b/Engine/Plugins/Experimental/MeshModelingToolset/Source/MeshModelingTools/Public/MeshInspectorTool.h
@@ -136,5 +136,7 @@ protected:
 	TArray<int> UVSeamEdges;
 	TArray<int> NormalSeamEdges;
 	TArray<int> GroupBoundaryEdges;
+
+	void UpdateVisualization();
 	void Precompute();
 };
diff --git a/Engine/Plugins/Experimental/MeshModelingToolset/Source/MeshModelingTools/Public/UVProjectionTool.h b/Engine/Plugins/Experimental/MeshModelingToolset/Source/MeshModelingTools/Public/UVProjectionTool.h
--- a/Engine/Plugins/Experimental/MeshModelingToolset/Source/MeshModelingTools/Public/UVProjectionTool.h
+++ b/Engine/Plugins/Experimental/MeshModelingToolset/Source/MeshModelingTools/Public/UVProjectionTool.h
@@ -195,6 +195,7 @@ protected:
 	FToolDataVisualizer ProjectionShapeVisualizer;
 
 	void UpdateNumPreviews();
+	void UpdateVisualization();
 
 	void GenerateAsset(const TArray<FDynamicMeshOpResult>& Results);
 
diff --git a/Engine/Plugins/Experimental/MeshModelingToolset/Source/MeshModelingToolsEditorOnly/Private/UVLayoutTool.cpp b/Engine/Plugins/Experimental/MeshModelingToolset/Source/MeshModelingToolsEditorOnly/Private/UVLayoutTool.cpp
--- a/Engine/Plugins/Experimental/MeshModelingToolset/Source/MeshModelingToolsEditorOnly/Private/UVLayoutTool.cpp
+++ b/Engine/Plugins/Experimental/MeshModelingToolset/Source/MeshModelingToolsEditorOnly/Private/UVLayoutTool.cpp
@@ -120,13 +120,7 @@ void UUVLayoutTool::Setup()
 	MaterialSettings->RestoreProperties(this);
 	AddToolPropertySource(MaterialSettings);
 
-	// initialize the PreviewMesh+BackgroundCompute object
-	UpdateNumPreviews();
-	
-	for (UMeshOpPreviewWithBackgroundCompute* Preview : Previews)
-	{
-		Preview->InvalidateResult();
-	}
+	UpdateVisualization();
 }
 
 
@@ -243,14 +237,19 @@ void UUVLayoutTool::PostEditChangeProperty(FPropertyChangedEvent& PropertyChange
 void UUVLayoutTool::OnPropertyModified(UObject* PropertySet, FProperty* Property)
 {
 	// if we don't know what changed, or we know checker density changed, update checker material
+	UpdateVisualization();
+}
+
+void UUVLayoutTool::UpdateVisualization()
+{
 	MaterialSettings->UpdateMaterials();
+	UpdateNumPreviews();
 	for (int PreviewIdx = 0; PreviewIdx < Previews.Num(); PreviewIdx++)
 	{
 		UMeshOpPreviewWithBackgroundCompute* Preview = Previews[PreviewIdx];
 		Preview->OverrideMaterial = MaterialSettings->GetActiveOverrideMaterial();
 	}
-	
-	UpdateNumPreviews();
+
 	for (UMeshOpPreviewWithBackgroundCompute* Preview : Previews)
 	{
 		Preview->InvalidateResult();
diff --git a/Engine/Plugins/Experimental/MeshModelingToolset/Source/MeshModelingToolsEditorOnly/Public/UVLayoutTool.h b/Engine/Plugins/Experimental/MeshModelingToolset/Source/MeshModelingToolsEditorOnly/Public/UVLayoutTool.h
--- a/Engine/Plugins/Experimental/MeshModelingToolset/Source/MeshModelingToolsEditorOnly/Public/UVLayoutTool.h
+++ b/Engine/Plugins/Experimental/MeshModelingToolset/Source/MeshModelingToolsEditorOnly/Public/UVLayoutTool.h
@@ -164,6 +164,7 @@ protected:
 	FViewCameraState CameraState;
 
 	void UpdateNumPreviews();
+	void UpdateVisualization();
 
 	void GenerateAsset(const TArray<FDynamicMeshOpResult>& Results);
 };
