Add support for tracking changes to overlay normals to FMeshVertexChange/Builder. Add support for applying this to mesh in USimpleDynamicMeshComponent::ApplyChange and Octree component. FGroupTopologyDeformer now calculates list of modified overlay normals. EditMeshPolygonsTool and DeformMeshPolygonsTool now relcalculate normals on edit and propagate to change. Moved deformation calculation in EditMeshPolygonsTool from mouse handler to tick, significantly improves responsiveness.
#rnx

diff --git a/Engine/Plugins/Experimental/GeometryProcessing/Source/DynamicMesh/Private/Operations/GroupTopologyDeformer.cpp b/Engine/Plugins/Experimental/GeometryProcessing/Source/DynamicMesh/Private/Operations/GroupTopologyDeformer.cpp
--- a/Engine/Plugins/Experimental/GeometryProcessing/Source/DynamicMesh/Private/Operations/GroupTopologyDeformer.cpp
+++ b/Engine/Plugins/Experimental/GeometryProcessing/Source/DynamicMesh/Private/Operations/GroupTopologyDeformer.cpp
@@ -6,6 +6,7 @@
 #include "SegmentTypes.h"
 #include "Async/ParallelFor.h"
 #include "Containers/BitArray.h"
+#include "MeshNormals.h"
 
 
 void FGroupTopologyDeformer::Initialize(const FDynamicMesh3* MeshIn, const FGroupTopology* TopologyIn)
@@ -187,6 +188,21 @@ void FGroupTopologyDeformer::CalculateROI(const TArray<int>& HandleGroups, const
 			ModifiedVertices.Add(vid);
 		}
 	}
+
+	if (Mesh->HasAttributes() && Mesh->Attributes()->PrimaryNormals() != nullptr)
+	{
+		const FDynamicMeshNormalOverlay* Normals = Mesh->Attributes()->PrimaryNormals();
+		for (int32 vid : ModifiedVertices)
+		{
+			for (int32 tid : Mesh->VtxTrianglesItr(vid))
+			{
+				FIndex3i NormalTri = Normals->GetTriangle(tid);
+				ModifiedOverlayNormals.Add(NormalTri.A);
+				ModifiedOverlayNormals.Add(NormalTri.B);
+				ModifiedOverlayNormals.Add(NormalTri.C);
+			}
+		}
+	}
 }
 
 
@@ -314,6 +330,8 @@ void FGroupTopologyDeformer::UpdateSolution(FDynamicMesh3* TargetMesh, const TFu
 			TargetMesh->SetVertex(Face.InteriorVerts[vi], Sum);
 		}
 	});
+
+	FMeshNormals::QuickRecomputeOverlayNormals(*TargetMesh);
 }
 
 
diff --git a/Engine/Plugins/Experimental/GeometryProcessing/Source/DynamicMesh/Public/Operations/GroupTopologyDeformer.h b/Engine/Plugins/Experimental/GeometryProcessing/Source/DynamicMesh/Public/Operations/GroupTopologyDeformer.h
--- a/Engine/Plugins/Experimental/GeometryProcessing/Source/DynamicMesh/Public/Operations/GroupTopologyDeformer.h
+++ b/Engine/Plugins/Experimental/GeometryProcessing/Source/DynamicMesh/Public/Operations/GroupTopologyDeformer.h
@@ -110,7 +110,8 @@ public:
 	const TSet<int>& GetHandleVertices() const { return HandleVertices; }
 	/** @return the set of all vertices that will be modified by the deformation*/
 	const TSet<int>& GetModifiedVertices() const { return ModifiedVertices; }
-
+	/** @return the set of all overlay normals that will be modified by the deformation*/
+	const TSet<int>& GetModifiedOverlayNormals() const { return ModifiedOverlayNormals; }
 
 protected:
 	//
@@ -148,6 +149,7 @@ protected:
 	TSet<int> ROIEdgeVertices;
 	TSet<int> FaceVertsTemp;
 	TSet<int> FaceBoundaryVertsTemp;
+	TSet<int> ModifiedOverlayNormals;
 
 	struct FROIEdge
 	{
diff --git a/Engine/Plugins/Experimental/MeshModelingToolset/Source/MeshModelingTools/Private/DeformMeshPolygonsTool.cpp b/Engine/Plugins/Experimental/MeshModelingToolset/Source/MeshModelingTools/Private/DeformMeshPolygonsTool.cpp
--- a/Engine/Plugins/Experimental/MeshModelingToolset/Source/MeshModelingTools/Private/DeformMeshPolygonsTool.cpp
+++ b/Engine/Plugins/Experimental/MeshModelingToolset/Source/MeshModelingTools/Private/DeformMeshPolygonsTool.cpp
@@ -1024,9 +1024,9 @@ void UDeformMeshPolygonsTool::UpdateChangeFromROI(bool bFinal)
 
 	FDynamicMesh3* Mesh = DynamicMeshComponent->GetMesh();
 	const TSet<int>& ModifiedVertices = (bIsLaplacian) ? LaplacianDeformer.GetModifiedVertices() : LinearDeformer.GetModifiedVertices();
-	
 	ActiveVertexChange->SavePositions(Mesh, ModifiedVertices, !bFinal);
-
+	const TSet<int>& ModifiedNormals = (bIsLaplacian) ? LaplacianDeformer.GetModifiedOverlayNormals() : LinearDeformer.GetModifiedOverlayNormals();
+	ActiveVertexChange->SaveOverlayNormals(Mesh, ModifiedNormals, !bFinal);
 }
 
 
@@ -1211,7 +1211,7 @@ void UDeformMeshPolygonsTool::ComputeUpdate_Rotate()
 		if (DeformationStrategy == EGroupTopologyDeformationStrategy::Linear)
 		{
 
-			DynamicMeshComponent->FastNotifyPositionsUpdated();
+			DynamicMeshComponent->FastNotifyPositionsUpdated(true);
 			GetToolManager()->PostInvalidation();
 		}
 		bUpdatePending = false;
@@ -1249,7 +1249,7 @@ void UDeformMeshPolygonsTool::ComputeUpdate_Rotate()
 	//TODO: This is unseemly here, need to potentially defer this so that it's handled the same way as laplacian. Placeholder for now.
 	if (!bIsLaplacian)
 	{
-		DynamicMeshComponent->FastNotifyPositionsUpdated();
+		DynamicMeshComponent->FastNotifyPositionsUpdated(true);
 		GetToolManager()->PostInvalidation();
 	}
 	bUpdatePending = false;
@@ -1312,7 +1312,7 @@ void UDeformMeshPolygonsTool::ComputeUpdate_Translate()
 		//TODO: This is unseemly here, need to potentially defer this so that it's handled the same way as laplacian. Placeholder for now.
 		if (!bIsLaplacian)
 		{
-			DynamicMeshComponent->FastNotifyPositionsUpdated();
+			DynamicMeshComponent->FastNotifyPositionsUpdated(true);
 			GetToolManager()->PostInvalidation();
 		}
 	}
@@ -1448,7 +1448,7 @@ void UDeformMeshPolygonsTool::BeginChange()
 	{
 		if (ActiveVertexChange == nullptr)
 		{
-			ActiveVertexChange = new FMeshVertexChangeBuilder();
+			ActiveVertexChange = new FMeshVertexChangeBuilder(true);
 			UpdateChangeFromROI(false);
 		}
 	}
diff --git a/Engine/Plugins/Experimental/MeshModelingToolset/Source/MeshModelingTools/Private/EditMeshPolygonsTool.cpp b/Engine/Plugins/Experimental/MeshModelingToolset/Source/MeshModelingTools/Private/EditMeshPolygonsTool.cpp
--- a/Engine/Plugins/Experimental/MeshModelingToolset/Source/MeshModelingTools/Private/EditMeshPolygonsTool.cpp
+++ b/Engine/Plugins/Experimental/MeshModelingToolset/Source/MeshModelingTools/Private/EditMeshPolygonsTool.cpp
@@ -486,12 +486,13 @@ void UEditMeshPolygonsTool::OnMultiTransformerTransformUpdate()
 {
 	if (MultiTransformer->InGizmoEdit())
 	{
-		ComputeUpdate_Gizmo();
+		CacheUpdate_Gizmo();
 	}
 }
 
 void UEditMeshPolygonsTool::OnMultiTransformerTransformEnd()
 {
+	bGizmoUpdatePending = false;
 	bSpatialDirty = true;
 	SelectionMechanic->NotifyMeshChanged(false);
 
@@ -572,17 +573,24 @@ void UEditMeshPolygonsTool::UpdateDeformerFromSelection(const FGroupTopologySele
 
 
 
-
+void UEditMeshPolygonsTool::CacheUpdate_Gizmo()
+{
+	LastUpdateGizmoFrame = MultiTransformer->GetCurrentGizmoFrame();
+	LastUpdateGizmoScale = MultiTransformer->GetCurrentGizmoScale();
+	GetToolManager()->PostInvalidation();
+	bGizmoUpdatePending = true;
+}
 
 void UEditMeshPolygonsTool::ComputeUpdate_Gizmo()
 {
-	if (SelectionMechanic->HasSelection() == false)
+	if (SelectionMechanic->HasSelection() == false || bGizmoUpdatePending == false)
 	{
 		return;
 	}
+	bGizmoUpdatePending = false;
 
-	FFrame3d CurFrame = MultiTransformer->GetCurrentGizmoFrame();
-	FVector3d CurScale = MultiTransformer->GetCurrentGizmoScale();
+	FFrame3d CurFrame = LastUpdateGizmoFrame;
+	FVector3d CurScale = LastUpdateGizmoScale;
 	FVector3d TranslationDelta = CurFrame.Origin - InitialGizmoFrame.Origin;
 	FQuaterniond RotateDelta = CurFrame.Rotation - InitialGizmoFrame.Rotation;
 	FVector3d CurScaleDelta = CurScale - InitialGizmoScale;
@@ -607,7 +615,7 @@ void UEditMeshPolygonsTool::ComputeUpdate_Gizmo()
 		// Reset mesh to initial positions.
 		LinearDeformer.ClearSolution(Mesh);
 	}
-	DynamicMeshComponent->FastNotifyPositionsUpdated();
+	DynamicMeshComponent->FastNotifyPositionsUpdated(true);
 	GetToolManager()->PostInvalidation();
 }
 
@@ -624,6 +632,11 @@ void UEditMeshPolygonsTool::Tick(float DeltaTime)
 
 	MultiTransformer->Tick(DeltaTime);
 
+	if (bGizmoUpdatePending)
+	{
+		ComputeUpdate_Gizmo();
+	}
+
 	if (bSelectionStateDirty)
 	{
 		// update color highlights
@@ -825,8 +838,8 @@ void UEditMeshPolygonsTool::UpdateChangeFromROI(bool bFinal)
 	}
 
 	FDynamicMesh3* Mesh = DynamicMeshComponent->GetMesh();
-	const TSet<int>& ModifiedVertices = LinearDeformer.GetModifiedVertices();
-	ActiveVertexChange->SavePositions(Mesh, ModifiedVertices, !bFinal);
+	ActiveVertexChange->SavePositions(Mesh, LinearDeformer.GetModifiedVertices(), !bFinal);
+	ActiveVertexChange->SaveOverlayNormals(Mesh, LinearDeformer.GetModifiedOverlayNormals(), !bFinal);
 }
 
 
@@ -834,7 +847,7 @@ void UEditMeshPolygonsTool::BeginChange()
 {
 	if (ActiveVertexChange == nullptr)
 	{
-		ActiveVertexChange = new FMeshVertexChangeBuilder();
+		ActiveVertexChange = new FMeshVertexChangeBuilder(true);
 		UpdateChangeFromROI(false);
 	}
 }
diff --git a/Engine/Plugins/Experimental/MeshModelingToolset/Source/MeshModelingTools/Public/EditMeshPolygonsTool.h b/Engine/Plugins/Experimental/MeshModelingToolset/Source/MeshModelingTools/Public/EditMeshPolygonsTool.h
--- a/Engine/Plugins/Experimental/MeshModelingToolset/Source/MeshModelingTools/Public/EditMeshPolygonsTool.h
+++ b/Engine/Plugins/Experimental/MeshModelingToolset/Source/MeshModelingTools/Public/EditMeshPolygonsTool.h
@@ -479,6 +479,10 @@ protected:
 
 	FFrame3d InitialGizmoFrame;
 	FVector3d InitialGizmoScale;
+	void CacheUpdate_Gizmo();
+	bool bGizmoUpdatePending = false;
+	FFrame3d LastUpdateGizmoFrame;
+	FVector3d LastUpdateGizmoScale;
 	void ComputeUpdate_Gizmo();
 
 	TUniquePtr<FGroupTopology> Topology;
diff --git a/Engine/Plugins/Experimental/MeshModelingToolset/Source/ModelingComponents/Private/Changes/MeshVertexChange.cpp b/Engine/Plugins/Experimental/MeshModelingToolset/Source/ModelingComponents/Private/Changes/MeshVertexChange.cpp
--- a/Engine/Plugins/Experimental/MeshModelingToolset/Source/ModelingComponents/Private/Changes/MeshVertexChange.cpp
+++ b/Engine/Plugins/Experimental/MeshModelingToolset/Source/ModelingComponents/Private/Changes/MeshVertexChange.cpp
@@ -28,9 +28,15 @@ FString FMeshVertexChange::ToString() const
 
 
 
-FMeshVertexChangeBuilder::FMeshVertexChangeBuilder()
+FMeshVertexChangeBuilder::FMeshVertexChangeBuilder(bool bSaveOverlayNormalsIn)
 {
 	Change = MakeUnique<FMeshVertexChange>();
+
+	bSaveOverlayNormals = bSaveOverlayNormalsIn;
+	if (bSaveOverlayNormals)
+	{
+		Change->bHaveOverlayNormals = true;
+	}
 }
 
 void FMeshVertexChangeBuilder::UpdateVertex(int VertexID, const FVector3d& OldPosition, const FVector3d& NewPosition)
@@ -62,6 +68,9 @@ void FMeshVertexChangeBuilder::UpdateVertexFinal(int VertexID, const FVector3d&
 }
 
 
+
+
+
 void FMeshVertexChangeBuilder::SavePosition(const FDynamicMesh3* Mesh, int VertexID, bool bInitial)
 {
 	FVector3d Pos = Mesh->GetVertex(VertexID);
@@ -115,4 +124,93 @@ void FMeshVertexChangeBuilder::SavePositions(const FDynamicMesh3* Mesh, const TS
 			UpdateVertexFinal(VertexID, Pos);
 		}
 	}
+}
+
+
+
+
+
+
+
+void FMeshVertexChangeBuilder::UpdateOverlayNormal(int ElementID, const FVector3f& OldNormal, const FVector3f& NewNormal)
+{
+	const int* FoundIndex = SavedNormalElements.Find(ElementID);
+	if (FoundIndex == nullptr)
+	{
+		int NewIndex = Change->Normals.Num();
+		SavedNormalElements.Add(ElementID, NewIndex);
+		Change->Normals.Add(ElementID);
+		Change->OldNormals.Add(OldNormal);
+		Change->NewNormals.Add(NewNormal);
+	}
+	else
+	{
+		Change->NewNormals[*FoundIndex] = NewNormal;
+	}
+}
+
+void FMeshVertexChangeBuilder::UpdateOverlayNormalFinal(int ElementID, const FVector3f& NewNormal)
+{
+	check(SavedNormalElements.Contains(ElementID));
+	const int* Index = SavedNormalElements.Find(ElementID);
+	if (Index != nullptr)
+	{
+		Change->NewNormals[*Index] = NewNormal;
+	}
+}
+
+
+
+void FMeshVertexChangeBuilder::SaveOverlayNormals(const FDynamicMesh3* Mesh, const TArray<int>& ElementIDs, bool bInitial)
+{
+	if (Mesh->HasAttributes() == false || Mesh->Attributes()->PrimaryNormals() == nullptr)
+	{
+		return;
+	}
+	const FDynamicMeshNormalOverlay* Overlay = Mesh->Attributes()->PrimaryNormals();
+
+	int Num = ElementIDs.Num();
+	if (bInitial)
+	{
+		for (int k = 0; k < Num; ++k)
+		{
+			FVector3f Normal = Overlay->GetElement(ElementIDs[k]);
+			UpdateOverlayNormal(ElementIDs[k], Normal, Normal);
+		}
+	}
+	else
+	{
+		for (int k = 0; k < Num; ++k)
+		{
+			FVector3f Normal = Overlay->GetElement(ElementIDs[k]);
+			UpdateOverlayNormalFinal(ElementIDs[k], Normal);
+		}
+	}
+}
+
+
+void FMeshVertexChangeBuilder::SaveOverlayNormals(const FDynamicMesh3* Mesh, const TSet<int>& ElementIDs, bool bInitial)
+{
+	if (Mesh->HasAttributes() == false || Mesh->Attributes()->PrimaryNormals() == nullptr)
+	{
+		return;
+	}
+	const FDynamicMeshNormalOverlay* Overlay = Mesh->Attributes()->PrimaryNormals();
+
+	if (bInitial)
+	{
+		for (int ElementID : ElementIDs)
+		{
+			FVector3f Normal = Overlay->GetElement(ElementID);
+			UpdateOverlayNormal(ElementID, Normal, Normal);
+		}
+	}
+	else
+	{
+		for (int ElementID : ElementIDs)
+		{
+			FVector3f Normal = Overlay->GetElement(ElementID);
+			UpdateOverlayNormalFinal(ElementID, Normal);
+		}
+	}
 }
\ No newline at end of file
diff --git a/Engine/Plugins/Experimental/MeshModelingToolset/Source/ModelingComponents/Private/OctreeDynamicMeshComponent.cpp b/Engine/Plugins/Experimental/MeshModelingToolset/Source/ModelingComponents/Private/OctreeDynamicMeshComponent.cpp
--- a/Engine/Plugins/Experimental/MeshModelingToolset/Source/ModelingComponents/Private/OctreeDynamicMeshComponent.cpp
+++ b/Engine/Plugins/Experimental/MeshModelingToolset/Source/ModelingComponents/Private/OctreeDynamicMeshComponent.cpp
@@ -356,16 +356,10 @@ FBoxSphereBounds UOctreeDynamicMeshComponent::CalcBounds(const FTransform& Local
 
 void UOctreeDynamicMeshComponent::ApplyChange(const FMeshVertexChange* Change, bool bRevert)
 {
-	int NV = Change->Vertices.Num();
-	const TArray<FVector3d>& Positions = (bRevert) ? Change->OldPositions : Change->NewPositions;
-	
 	Octree->ResetModifiedBounds();
 	TSet<int> TrianglesToUpdate;
-
-	for (int k = 0; k < NV; ++k)
+	auto NotifyTriVerticesAffected = [&](int32 vid) 
 	{
-		int vid = Change->Vertices[k];
-
 		for (int tid : Mesh->VtxTrianglesItr(vid))
 		{
 			if (TrianglesToUpdate.Contains(tid) == false)
@@ -374,10 +368,34 @@ void UOctreeDynamicMeshComponent::ApplyChange(const FMeshVertexChange* Change, b
 				TrianglesToUpdate.Add(tid);
 			}
 		}
+	};
 
+	int NV = Change->Vertices.Num();
+	const TArray<FVector3d>& Positions = (bRevert) ? Change->OldPositions : Change->NewPositions;
+	for (int k = 0; k < NV; ++k)
+	{
+		int vid = Change->Vertices[k];
+		NotifyTriVerticesAffected(vid);
 		Mesh->SetVertex(vid, Positions[k]);
 	}
 
+	if (Change->bHaveOverlayNormals && Mesh->HasAttributes() && Mesh->Attributes()->PrimaryNormals())
+	{
+		FDynamicMeshNormalOverlay* Overlay = Mesh->Attributes()->PrimaryNormals();
+		int32 NumNormals = Change->Normals.Num();
+		const TArray<FVector3f>& UseNormals = (bRevert) ? Change->OldNormals : Change->NewNormals;
+		for (int32 k = 0; k < NumNormals; ++k)
+		{
+			int32 elemid = Change->Normals[k];
+			if (Overlay->IsElement(elemid))
+			{
+				Overlay->SetElement(elemid, UseNormals[k]);
+				int32 ParentVID = Overlay->GetParentVertex(elemid);
+				NotifyTriVerticesAffected(ParentVID);
+			}
+		}
+	}
+
 	Octree->ReinsertTriangles(TrianglesToUpdate);
 
 	//NotifyMeshUpdated();
diff --git a/Engine/Plugins/Experimental/MeshModelingToolset/Source/ModelingComponents/Private/SimpleDynamicMeshComponent.cpp b/Engine/Plugins/Experimental/MeshModelingToolset/Source/ModelingComponents/Private/SimpleDynamicMeshComponent.cpp
--- a/Engine/Plugins/Experimental/MeshModelingToolset/Source/ModelingComponents/Private/SimpleDynamicMeshComponent.cpp
+++ b/Engine/Plugins/Experimental/MeshModelingToolset/Source/ModelingComponents/Private/SimpleDynamicMeshComponent.cpp
@@ -325,13 +325,30 @@ FBoxSphereBounds USimpleDynamicMeshComponent::CalcBounds(const FTransform& Local
 
 void USimpleDynamicMeshComponent::ApplyChange(const FMeshVertexChange* Change, bool bRevert)
 {
-	int NV = Change->Vertices.Num();
+	int32 NV = Change->Vertices.Num();
 	const TArray<FVector3d>& Positions = (bRevert) ? Change->OldPositions : Change->NewPositions;
-	
-	for (int k = 0; k < NV; ++k)
+	for (int32 k = 0; k < NV; ++k)
+	{
+		int32 vid = Change->Vertices[k];
+		if (Mesh->IsVertex(vid))
+		{
+			Mesh->SetVertex(vid, Positions[k]);
+		}
+	}
+
+	if (Change->bHaveOverlayNormals && Mesh->HasAttributes() && Mesh->Attributes()->PrimaryNormals() )
 	{
-		int vid = Change->Vertices[k];
-		Mesh->SetVertex(vid, Positions[k]);
+		FDynamicMeshNormalOverlay* Overlay = Mesh->Attributes()->PrimaryNormals();
+		int32 NumNormals = Change->Normals.Num();
+		const TArray<FVector3f>& UseNormals = (bRevert) ? Change->OldNormals : Change->NewNormals;
+		for (int32 k = 0; k < NumNormals; ++k)
+		{
+			int32 elemid = Change->Normals[k];
+			if (Overlay->IsElement(elemid))
+			{
+				Overlay->SetElement(elemid, UseNormals[k]);
+			}
+		}
 	}
 
 	NotifyMeshUpdated();
diff --git a/Engine/Plugins/Experimental/MeshModelingToolset/Source/ModelingComponents/Public/Changes/MeshVertexChange.h b/Engine/Plugins/Experimental/MeshModelingToolset/Source/ModelingComponents/Public/Changes/MeshVertexChange.h
--- a/Engine/Plugins/Experimental/MeshModelingToolset/Source/ModelingComponents/Public/Changes/MeshVertexChange.h
+++ b/Engine/Plugins/Experimental/MeshModelingToolset/Source/ModelingComponents/Public/Changes/MeshVertexChange.h
@@ -24,6 +24,11 @@ public:
 	TArray<FVector3d> OldPositions;
 	TArray<FVector3d> NewPositions;
 
+	bool bHaveOverlayNormals = false;
+	TArray<int32> Normals;
+	TArray<FVector3f> OldNormals;
+	TArray<FVector3f> NewNormals;
+
 
 	/** Makes the change to the object */
 	virtual void Apply(UObject* Object) override;
@@ -47,7 +52,10 @@ public:
 	TUniquePtr<FMeshVertexChange> Change;
 	TMap<int, int> SavedVertices;
 
-	FMeshVertexChangeBuilder();
+	bool bSaveOverlayNormals;
+	TMap<int, int> SavedNormalElements;
+
+	FMeshVertexChangeBuilder(bool bSaveOverlayNormals = false);
 
 	void UpdateVertex(int VertexID, const FVector3d& OldPosition, const FVector3d& NewPosition);
 	void UpdateVertexFinal(int VertexID, const FVector3d& NewPosition);
@@ -55,6 +63,14 @@ public:
 	void SavePosition(const FDynamicMesh3* Mesh, int VertexID, bool bInitial);
 	void SavePositions(const FDynamicMesh3* Mesh, const TArray<int>& VertexIDs, bool bInitial);
 	void SavePositions(const FDynamicMesh3* Mesh, const TSet<int>& VertexIDs, bool bInitial);
+
+
+	void UpdateOverlayNormal(int ElementID, const FVector3f& OldNormal, const FVector3f& NewNormal);
+	void UpdateOverlayNormalFinal(int ElementID, const FVector3f& NewNormal);
+	void SaveOverlayNormals(const FDynamicMesh3* Mesh, const TArray<int>& ElementIDs, bool bInitial);
+	void SaveOverlayNormals(const FDynamicMesh3* Mesh, const TSet<int>& ElementIDs, bool bInitial);
+
+
 };
 
 
