Update DrawPolygonTool interaction to not add initial two points on mouse down/up, made it to easy to add tiny segments. Now first point is only dropped on mouse-up.
#rnx

diff --git a/Engine/Plugins/Experimental/MeshModelingToolset/Source/MeshModelingTools/Private/DrawPolygonTool.cpp b/Engine/Plugins/Experimental/MeshModelingToolset/Source/MeshModelingTools/Private/DrawPolygonTool.cpp
--- a/Engine/Plugins/Experimental/MeshModelingToolset/Source/MeshModelingTools/Private/DrawPolygonTool.cpp
+++ b/Engine/Plugins/Experimental/MeshModelingToolset/Source/MeshModelingTools/Private/DrawPolygonTool.cpp
@@ -639,19 +639,10 @@ void UDrawPolygonTool::OnBeginClickSequence(const FInputDeviceRay& ClickPos)
 		return;		// cannot start a poly an a point that is not visible, this is almost certainly an error due to draw plane
 	}
 
-	AppendVertex(HitPos);
 	UpdatePreviewVertex(HitPos);
 
 	bInFixedPolygonMode = (PolygonProperties->PolygonType != EDrawPolygonDrawMode::Freehand);
 	FixedPolygonClickPoints.Reset();
-	FixedPolygonClickPoints.Add(HitPos);
-
-	// if we are starting a freehand poly, add start point as snap target, but then ignore it until we get 3 verts
-	if (bInFixedPolygonMode == false)
-	{
-		SnapEngine.AddPointTarget(PolygonVertices[0], StartPointSnapID, 1);
-		SnapEngine.AddIgnoreTarget(StartPointSnapID);
-	}
 }
 
 void UDrawPolygonTool::OnNextSequencePreview(const FInputDeviceRay& ClickPos)
@@ -678,13 +669,12 @@ void UDrawPolygonTool::OnNextSequencePreview(const FInputDeviceRay& ClickPos)
 		return;
 	}
 
+	UpdatePreviewVertex(HitPos);
 	if (PolygonVertices.Num() > 2)
 	{
 		bPreviewUpdatePending = true;
+		UpdateSelfIntersection();
 	}
-
-	UpdatePreviewVertex(HitPos);
-	UpdateSelfIntersection();
 }
 
 bool UDrawPolygonTool::OnNextSequenceClick(const FInputDeviceRay& ClickPos)
@@ -706,7 +696,7 @@ bool UDrawPolygonTool::OnNextSequenceClick(const FInputDeviceRay& ClickPos)
 	if (bInFixedPolygonMode)
 	{
 		// ignore very close click points
-		if ( ToolSceneQueriesUtil::PointSnapQuery(this, FixedPolygonClickPoints[FixedPolygonClickPoints.Num()-1], HitPos) )
+		if (FixedPolygonClickPoints.Num() > 0 && ToolSceneQueriesUtil::PointSnapQuery(this, FixedPolygonClickPoints[FixedPolygonClickPoints.Num()-1], HitPos) )
 		{
 			return true;
 		}
@@ -722,7 +712,7 @@ bool UDrawPolygonTool::OnNextSequenceClick(const FInputDeviceRay& ClickPos)
 	else
 	{
 		// ignore very close click points
-		if (ToolSceneQueriesUtil::PointSnapQuery(this, PolygonVertices[PolygonVertices.Num()-1], HitPos))
+		if (PolygonVertices.Num() > 0 && ToolSceneQueriesUtil::PointSnapQuery(this, PolygonVertices[PolygonVertices.Num()-1], HitPos))
 		{
 			return true;
 		}
@@ -769,6 +759,13 @@ bool UDrawPolygonTool::OnNextSequenceClick(const FInputDeviceRay& ClickPos)
 	}
 
 	AppendVertex(HitPos);
+
+	// if we are starting a freehand poly, add start point as snap target, but then ignore it until we get 3 verts
+	if (bInFixedPolygonMode == false && PolygonVertices.Num() == 1)
+	{
+		SnapEngine.AddPointTarget(PolygonVertices[0], StartPointSnapID, 1);
+		SnapEngine.AddIgnoreTarget(StartPointSnapID);
+	}
 	if (PolygonVertices.Num() > 2)
 	{
 		SnapEngine.RemoveIgnoreTarget(StartPointSnapID);
