DrawPolygonTool undo behavior is now consistent with DrawPolyPathTool behavior. Removed backspace hotkey that pops last vertex, can now do that action with undo.
#rnx

diff --git a/Engine/Plugins/Experimental/MeshModelingToolset/Source/MeshModelingTools/Private/DrawPolygonTool.cpp b/Engine/Plugins/Experimental/MeshModelingToolset/Source/MeshModelingTools/Private/DrawPolygonTool.cpp
--- a/Engine/Plugins/Experimental/MeshModelingToolset/Source/MeshModelingTools/Private/DrawPolygonTool.cpp
+++ b/Engine/Plugins/Experimental/MeshModelingToolset/Source/MeshModelingTools/Private/DrawPolygonTool.cpp
@@ -237,21 +237,12 @@ void UDrawPolygonTool::Shutdown(EToolShutdownType ShutdownType)
 
 void UDrawPolygonTool::RegisterActions(FInteractiveToolActionSet& ActionSet)
 {
-	ActionSet.RegisterAction(this, (int32)EStandardToolActions::BaseClientDefinedActionID + 1,
-		TEXT("PopLastVertex"), 
-		LOCTEXT("PopLastVertex", "Pop Last Vertex"),
-		LOCTEXT("PopLastVertexTooltip", "Pop last vertex added to polygon"),
-		EModifierKey::None, EKeys::BackSpace,
-		[this]() { PopLastVertexAction(); });
-
-
 	ActionSet.RegisterAction(this, (int32)EStandardToolActions::BaseClientDefinedActionID + 2,
 		TEXT("ToggleGizmo"),
 		LOCTEXT("ToggleGizmo", "Toggle Gizmo"),
 		LOCTEXT("ToggleGizmoTooltip", "Toggle visibility of the transformation Gizmo"),
 		EModifierKey::None, EKeys::A,
 		[this]() { PolygonProperties->bShowGizmo = !PolygonProperties->bShowGizmo; });
-
 }
 
 
@@ -274,6 +265,7 @@ void UDrawPolygonTool::PopLastVertexAction()
 		{
 			PolygonVertices.RemoveAt(0);
 			bAbortActivePolygonDraw = true;
+			CurrentCurveTimestamp++;
 		}
 	}
 	else
@@ -287,6 +279,7 @@ void UDrawPolygonTool::PopLastVertexAction()
 		{
 			FixedPolygonClickPoints.RemoveAt(0);
 			bAbortActivePolygonDraw = true;
+			CurrentCurveTimestamp++;
 		}
 	}
 }
@@ -504,6 +497,7 @@ void UDrawPolygonTool::ResetPolygon()
 	SnapEngine.Reset();
 	bHaveSurfaceHit = false;
 	bInFixedPolygonMode = false;
+	CurrentCurveTimestamp++;
 }
 
 void UDrawPolygonTool::UpdatePreviewVertex(const FVector& PreviewVertexIn)
@@ -760,6 +754,9 @@ bool UDrawPolygonTool::OnNextSequenceClick(const FInputDeviceRay& ClickPos)
 
 	AppendVertex(HitPos);
 
+	// emit change event
+	GetToolManager()->EmitObjectChange(this, MakeUnique<FDrawPolygonStateChange>(CurrentCurveTimestamp), LOCTEXT("DrawPolyAddPoint", "Add Point"));
+
 	// if we are starting a freehand poly, add start point as snap target, but then ignore it until we get 3 verts
 	if (bInFixedPolygonMode == false && PolygonVertices.Num() == 1)
 	{
@@ -1287,4 +1284,38 @@ void UDrawPolygonTool::ShowExtrudeMessage()
 }
 
 
+
+
+void UDrawPolygonTool::UndoCurrentOperation()
+{
+	if (bInInteractiveExtrude)
+	{
+		PreviewMesh->ClearPreview();
+		PreviewMesh->SetVisible(false);
+		bInInteractiveExtrude = false;
+		PopLastVertexAction();
+	}
+	else
+	{
+		PopLastVertexAction();
+	}
+}
+
+
+void FDrawPolygonStateChange::Revert(UObject* Object)
+{
+	Cast<UDrawPolygonTool>(Object)->UndoCurrentOperation();
+	bHaveDoneUndo = true;
+}
+bool FDrawPolygonStateChange::HasExpired(UObject* Object) const
+{
+	return bHaveDoneUndo || (Cast<UDrawPolygonTool>(Object)->CheckInCurve(CurveTimestamp) == false);
+}
+FString FDrawPolygonStateChange::ToString() const
+{
+	return TEXT("FDrawPolygonStateChange");
+}
+
+
+
 #undef LOCTEXT_NAMESPACE
diff --git a/Engine/Plugins/Experimental/MeshModelingToolset/Source/MeshModelingTools/Public/DrawPolygonTool.h b/Engine/Plugins/Experimental/MeshModelingToolset/Source/MeshModelingTools/Public/DrawPolygonTool.h
--- a/Engine/Plugins/Experimental/MeshModelingToolset/Source/MeshModelingTools/Public/DrawPolygonTool.h
+++ b/Engine/Plugins/Experimental/MeshModelingToolset/Source/MeshModelingTools/Public/DrawPolygonTool.h
@@ -339,9 +339,29 @@ protected:
 	// user feedback messages
 	void ShowStartupMessage();
 	void ShowExtrudeMessage();
-};
 
 
+	friend class FDrawPolygonStateChange;
+	int32 CurrentCurveTimestamp = 1;
+	void UndoCurrentOperation();
+	bool CheckInCurve(int32 Timestamp) const { return CurrentCurveTimestamp == Timestamp; }
+};
 
 
 
+// Change event used by DrawPolygonTool to undo draw state.
+// Currently does not redo.
+class MESHMODELINGTOOLS_API FDrawPolygonStateChange : public FToolCommandChange
+{
+public:
+	bool bHaveDoneUndo = false;
+	int32 CurveTimestamp = 0;
+	FDrawPolygonStateChange(int32 CurveTimestampIn)
+	{
+		CurveTimestamp = CurveTimestampIn;
+	}
+	virtual void Apply(UObject* Object) override {}
+	virtual void Revert(UObject* Object) override;
+	virtual bool HasExpired(UObject* Object) const override;
+	virtual FString ToString() const override;
+};
