MeshModelingTools : Clean up update logic during progressive computation in DisplaceMeshTool

#rnx

diff --git a/Engine/Plugins/Experimental/MeshModelingToolset/Source/MeshModelingTools/Private/DisplaceMeshTool.cpp b/Engine/Plugins/Experimental/MeshModelingToolset/Source/MeshModelingTools/Private/DisplaceMeshTool.cpp
--- a/Engine/Plugins/Experimental/MeshModelingToolset/Source/MeshModelingTools/Private/DisplaceMeshTool.cpp
+++ b/Engine/Plugins/Experimental/MeshModelingToolset/Source/MeshModelingTools/Private/DisplaceMeshTool.cpp
@@ -651,9 +651,8 @@ void UDisplaceMeshTool::StartComputation()
 		if (SubdivideTask)
 		{
 			SubdivideTask->CancelAndDelete();
-			SubdividedMesh = nullptr;
 		}
-		auto SubdivideMeshOp = Subdivider->MakeNewOperator();
+		SubdividedMesh = nullptr;
 		SubdivideTask = new FAsyncTaskExecuterWithAbort<TModelingOpTask<FDynamicMeshOperator>>(Subdivider->MakeNewOperator());
 		SubdivideTask->StartBackgroundTask();
 		bNeedsSubdivided = false;
