MeshModelingTools: Optimize property change update for Remesh and Simplify tools to not rebuild when display only changes occur

EDIGRATE of CL11130378

#rnx

diff --git a/Engine/Plugins/Experimental/MeshModelingToolset/Source/MeshModelingTools/Private/RemeshMeshTool.cpp b/Engine/Plugins/Experimental/MeshModelingToolset/Source/MeshModelingTools/Private/RemeshMeshTool.cpp
--- a/Engine/Plugins/Experimental/MeshModelingToolset/Source/MeshModelingTools/Private/RemeshMeshTool.cpp
+++ b/Engine/Plugins/Experimental/MeshModelingToolset/Source/MeshModelingTools/Private/RemeshMeshTool.cpp
@@ -218,7 +218,18 @@ void URemeshMeshTool::Render(IToolsContextRenderAPI* RenderAPI)
 
 void URemeshMeshTool::OnPropertyModified(UObject* PropertySet, FProperty* Property)
 {
-	Preview->InvalidateResult();
+	if ( Property )
+	{
+		if ( ( Property->GetFName() == GET_MEMBER_NAME_CHECKED(URemeshMeshToolProperties, bShowWireframe) ) ||
+			 ( Property->GetFName() == GET_MEMBER_NAME_CHECKED(URemeshMeshToolProperties, bShowGroupColors) ) )
+		{
+			UpdateVisualization();
+		}
+		else
+		{
+			Preview->InvalidateResult();
+		}
+	}
 }
 
 void URemeshMeshTool::UpdateVisualization()
diff --git a/Engine/Plugins/Experimental/MeshModelingToolset/Source/MeshModelingToolsEditorOnly/Private/SimplifyMeshTool.cpp b/Engine/Plugins/Experimental/MeshModelingToolset/Source/MeshModelingToolsEditorOnly/Private/SimplifyMeshTool.cpp
--- a/Engine/Plugins/Experimental/MeshModelingToolset/Source/MeshModelingToolsEditorOnly/Private/SimplifyMeshTool.cpp
+++ b/Engine/Plugins/Experimental/MeshModelingToolset/Source/MeshModelingToolsEditorOnly/Private/SimplifyMeshTool.cpp
@@ -148,6 +148,7 @@ void USimplifyMeshTool::Setup()
 		MeshStatisticsProperties->Update(*Compute->PreviewMesh->GetPreviewDynamicMesh());
 	});
 
+	UpdateVisualization();
 	Preview->InvalidateResult();
 }
 
@@ -226,7 +227,18 @@ void USimplifyMeshTool::Render(IToolsContextRenderAPI* RenderAPI)
 
 void USimplifyMeshTool::OnPropertyModified(UObject* PropertySet, FProperty* Property)
 {
-	Preview->InvalidateResult();
+	if ( Property )
+	{
+		if ( ( Property->GetFName() == GET_MEMBER_NAME_CHECKED(USimplifyMeshToolProperties, bShowWireframe) ) ||
+			 ( Property->GetFName() == GET_MEMBER_NAME_CHECKED(USimplifyMeshToolProperties, bShowGroupColors) ) )
+		{
+			UpdateVisualization();
+		}
+		else
+		{
+			Preview->InvalidateResult();
+		}
+	}
 }
 
 void USimplifyMeshTool::UpdateVisualization()
