Enable autocalculated tangents for most mesh editing tools. This preserves appearance when editing meshes with normal maps. However these tangents are not written back to the MeshDescription, we current assume that the build settings will result in tangents being recomputed on Accept.
#rnx

diff --git a/Engine/Plugins/Experimental/MeshModelingToolset/Source/MeshModelingTools/Private/ConvertToPolygonsTool.cpp b/Engine/Plugins/Experimental/MeshModelingToolset/Source/MeshModelingTools/Private/ConvertToPolygonsTool.cpp
--- a/Engine/Plugins/Experimental/MeshModelingToolset/Source/MeshModelingTools/Private/ConvertToPolygonsTool.cpp
+++ b/Engine/Plugins/Experimental/MeshModelingToolset/Source/MeshModelingTools/Private/ConvertToPolygonsTool.cpp
@@ -86,6 +86,7 @@ void UConvertToPolygonsTool::Setup()
 	PreviewMesh->CreateInWorld(ComponentTarget->GetOwnerActor()->GetWorld(), FTransform::Identity);
 	PreviewMesh->SetVisible(false);
 	PreviewMesh->SetTransform(ComponentTarget->GetWorldTransform());
+	PreviewMesh->SetTangentsMode(EDynamicMeshTangentCalcType::AutoCalculated);
 
 	FComponentMaterialSet MaterialSet;
 	ComponentTarget->GetMaterialSet(MaterialSet);
diff --git a/Engine/Plugins/Experimental/MeshModelingToolset/Source/MeshModelingTools/Private/DisplaceMeshTool.cpp b/Engine/Plugins/Experimental/MeshModelingToolset/Source/MeshModelingTools/Private/DisplaceMeshTool.cpp
--- a/Engine/Plugins/Experimental/MeshModelingToolset/Source/MeshModelingTools/Private/DisplaceMeshTool.cpp
+++ b/Engine/Plugins/Experimental/MeshModelingToolset/Source/MeshModelingTools/Private/DisplaceMeshTool.cpp
@@ -526,6 +526,7 @@ void UDisplaceMeshTool::Setup()
 		DynamicMeshComponent->SetMaterial(k, MaterialSet.Materials[k]);
 	}
 
+	DynamicMeshComponent->TangentsType = EDynamicMeshTangentCalcType::AutoCalculated;
 	DynamicMeshComponent->InitializeMesh(ComponentTarget->GetMesh());
 	OriginalMesh.Copy(*DynamicMeshComponent->GetMesh());
 
diff --git a/Engine/Plugins/Experimental/MeshModelingToolset/Source/MeshModelingTools/Private/EditNormalsTool.cpp b/Engine/Plugins/Experimental/MeshModelingToolset/Source/MeshModelingTools/Private/EditNormalsTool.cpp
--- a/Engine/Plugins/Experimental/MeshModelingToolset/Source/MeshModelingTools/Private/EditNormalsTool.cpp
+++ b/Engine/Plugins/Experimental/MeshModelingToolset/Source/MeshModelingTools/Private/EditNormalsTool.cpp
@@ -143,6 +143,7 @@ void UEditNormalsTool::UpdateNumPreviews()
 
 			UMeshOpPreviewWithBackgroundCompute* Preview = Previews.Add_GetRef(NewObject<UMeshOpPreviewWithBackgroundCompute>(OpFactory, "Preview"));
 			Preview->Setup(this->TargetWorld, OpFactory);
+			Preview->PreviewMesh->SetTangentsMode(EDynamicMeshTangentCalcType::AutoCalculated);
 
 			FComponentMaterialSet MaterialSet;
 			ComponentTargets[PreviewIdx]->GetMaterialSet(MaterialSet);
diff --git a/Engine/Plugins/Experimental/MeshModelingToolset/Source/MeshModelingTools/Private/MeshSelectionTool.cpp b/Engine/Plugins/Experimental/MeshModelingToolset/Source/MeshModelingTools/Private/MeshSelectionTool.cpp
--- a/Engine/Plugins/Experimental/MeshModelingToolset/Source/MeshModelingTools/Private/MeshSelectionTool.cpp
+++ b/Engine/Plugins/Experimental/MeshModelingToolset/Source/MeshModelingTools/Private/MeshSelectionTool.cpp
@@ -108,6 +108,9 @@ void UMeshSelectionTool::Setup()
 	EditActions = CreateEditActions();
 	AddToolPropertySource(EditActions);
 
+	// set autocalculated tangents
+	PreviewMesh->SetTangentsMode(EDynamicMeshTangentCalcType::AutoCalculated);
+
 	// enable wireframe on component
 	PreviewMesh->EnableWireframe(SelectionProps->bShowWireframe);
 
diff --git a/Engine/Plugins/Experimental/MeshModelingToolset/Source/MeshModelingTools/Private/MeshSpaceDeformerTool.cpp b/Engine/Plugins/Experimental/MeshModelingToolset/Source/MeshModelingTools/Private/MeshSpaceDeformerTool.cpp
--- a/Engine/Plugins/Experimental/MeshModelingToolset/Source/MeshModelingTools/Private/MeshSpaceDeformerTool.cpp
+++ b/Engine/Plugins/Experimental/MeshModelingToolset/Source/MeshModelingTools/Private/MeshSpaceDeformerTool.cpp
@@ -257,6 +257,7 @@ void UMeshSpaceDeformerTool::Setup()
 
 		Preview = NewObject<UMeshOpPreviewWithBackgroundCompute>(DeformerOperatorFactory, "Preview");
 		Preview->Setup(this->TargetWorld, DeformerOperatorFactory);
+		Preview->PreviewMesh->SetTangentsMode(EDynamicMeshTangentCalcType::AutoCalculated);
 
 		// Give the preview something to display
 		Preview->PreviewMesh->UpdatePreview(OriginalDynamicMesh.Get());
diff --git a/Engine/Plugins/Experimental/MeshModelingToolset/Source/MeshModelingTools/Private/PlaneCutTool.cpp b/Engine/Plugins/Experimental/MeshModelingToolset/Source/MeshModelingTools/Private/PlaneCutTool.cpp
--- a/Engine/Plugins/Experimental/MeshModelingToolset/Source/MeshModelingTools/Private/PlaneCutTool.cpp
+++ b/Engine/Plugins/Experimental/MeshModelingToolset/Source/MeshModelingTools/Private/PlaneCutTool.cpp
@@ -277,6 +277,7 @@ void UPlaneCutTool::SetupPreviews()
 		CutSide->ComponentIndex = PreviewIdx;
 		UMeshOpPreviewWithBackgroundCompute* Preview = Previews.Add_GetRef(NewObject<UMeshOpPreviewWithBackgroundCompute>(CutSide, "Preview"));
 		Preview->Setup(this->TargetWorld, CutSide);
+		Preview->PreviewMesh->SetTangentsMode(EDynamicMeshTangentCalcType::AutoCalculated);
 
 		FComponentMaterialSet MaterialSet;
 		ComponentTargets[PreviewIdx]->GetMaterialSet(MaterialSet);
diff --git a/Engine/Plugins/Experimental/MeshModelingToolset/Source/MeshModelingTools/Private/RemeshMeshTool.cpp b/Engine/Plugins/Experimental/MeshModelingToolset/Source/MeshModelingTools/Private/RemeshMeshTool.cpp
--- a/Engine/Plugins/Experimental/MeshModelingToolset/Source/MeshModelingTools/Private/RemeshMeshTool.cpp
+++ b/Engine/Plugins/Experimental/MeshModelingToolset/Source/MeshModelingTools/Private/RemeshMeshTool.cpp
@@ -114,6 +114,7 @@ void URemeshMeshTool::Setup()
 	Converter.Convert(ComponentTarget->GetMesh(), *OriginalMesh);
 
 	Preview->PreviewMesh->SetTransform(ComponentTarget->GetWorldTransform());
+	Preview->PreviewMesh->SetTangentsMode(EDynamicMeshTangentCalcType::AutoCalculated);
 	Preview->PreviewMesh->UpdatePreview(OriginalMesh.Get());
 
 	OriginalMeshSpatial = MakeShared<FDynamicMeshAABBTree3>(OriginalMesh.Get(), true);
diff --git a/Engine/Plugins/Experimental/MeshModelingToolset/Source/MeshModelingTools/Private/RemoveOccludedTrianglesTool.cpp b/Engine/Plugins/Experimental/MeshModelingToolset/Source/MeshModelingTools/Private/RemoveOccludedTrianglesTool.cpp
--- a/Engine/Plugins/Experimental/MeshModelingToolset/Source/MeshModelingTools/Private/RemoveOccludedTrianglesTool.cpp
+++ b/Engine/Plugins/Experimental/MeshModelingToolset/Source/MeshModelingTools/Private/RemoveOccludedTrianglesTool.cpp
@@ -190,6 +190,7 @@ void URemoveOccludedTrianglesTool::SetupPreviews()
 
 			UMeshOpPreviewWithBackgroundCompute* Preview = Previews.Add_GetRef(NewObject<UMeshOpPreviewWithBackgroundCompute>(OpFactory, "Preview"));
 			Preview->Setup(this->TargetWorld, OpFactory);
+			Preview->PreviewMesh->SetTangentsMode(EDynamicMeshTangentCalcType::AutoCalculated);
 
 			FComponentMaterialSet MaterialSet;
 			ComponentTargets[TargetIdx]->GetMaterialSet(MaterialSet);
diff --git a/Engine/Plugins/Experimental/MeshModelingToolset/Source/MeshModelingTools/Private/SmoothMeshTool.cpp b/Engine/Plugins/Experimental/MeshModelingToolset/Source/MeshModelingTools/Private/SmoothMeshTool.cpp
--- a/Engine/Plugins/Experimental/MeshModelingToolset/Source/MeshModelingTools/Private/SmoothMeshTool.cpp
+++ b/Engine/Plugins/Experimental/MeshModelingToolset/Source/MeshModelingTools/Private/SmoothMeshTool.cpp
@@ -83,6 +83,7 @@ void USmoothMeshTool::Setup()
 		// Construct the preview object and set the material on it.
 		Preview = NewObject<UMeshOpPreviewWithBackgroundCompute>(this, "Preview");
 		Preview->Setup(this->TargetWorld, this); // Adds the actual functional tool in the Preview object
+		Preview->PreviewMesh->SetTangentsMode(EDynamicMeshTangentCalcType::AutoCalculated);
 
 		FComponentMaterialSet MaterialSet;
 		ComponentTarget->GetMaterialSet(MaterialSet);
diff --git a/Engine/Plugins/Experimental/MeshModelingToolset/Source/MeshModelingTools/Private/UVProjectionTool.cpp b/Engine/Plugins/Experimental/MeshModelingToolset/Source/MeshModelingTools/Private/UVProjectionTool.cpp
--- a/Engine/Plugins/Experimental/MeshModelingToolset/Source/MeshModelingTools/Private/UVProjectionTool.cpp
+++ b/Engine/Plugins/Experimental/MeshModelingToolset/Source/MeshModelingTools/Private/UVProjectionTool.cpp
@@ -189,6 +189,7 @@ void UUVProjectionTool::UpdateNumPreviews()
 
 			UMeshOpPreviewWithBackgroundCompute* Preview = Previews.Add_GetRef(NewObject<UMeshOpPreviewWithBackgroundCompute>(OpFactory, "Preview"));
 			Preview->Setup(this->TargetWorld, OpFactory);
+			Preview->PreviewMesh->SetTangentsMode(EDynamicMeshTangentCalcType::AutoCalculated);
 
 			FComponentMaterialSet MaterialSet;
 			ComponentTargets[PreviewIdx]->GetMaterialSet(MaterialSet);
diff --git a/Engine/Plugins/Experimental/MeshModelingToolset/Source/MeshModelingTools/Private/WeldMeshEdgesTool.cpp b/Engine/Plugins/Experimental/MeshModelingToolset/Source/MeshModelingTools/Private/WeldMeshEdgesTool.cpp
--- a/Engine/Plugins/Experimental/MeshModelingToolset/Source/MeshModelingTools/Private/WeldMeshEdgesTool.cpp
+++ b/Engine/Plugins/Experimental/MeshModelingToolset/Source/MeshModelingTools/Private/WeldMeshEdgesTool.cpp
@@ -69,6 +69,7 @@ void UWeldMeshEdgesTool::Setup()
 		DynamicMeshComponent->SetMaterial(k, MaterialSet.Materials[k]);
 	}
 
+	DynamicMeshComponent->TangentsType = EDynamicMeshTangentCalcType::AutoCalculated;
 	DynamicMeshComponent->InitializeMesh(ComponentTarget->GetMesh());
 	OriginalMesh.Copy(*DynamicMeshComponent->GetMesh());
 
diff --git a/Engine/Plugins/Experimental/MeshModelingToolset/Source/MeshModelingToolsEditorOnly/Private/ParameterizeMeshTool.cpp b/Engine/Plugins/Experimental/MeshModelingToolset/Source/MeshModelingToolsEditorOnly/Private/ParameterizeMeshTool.cpp
--- a/Engine/Plugins/Experimental/MeshModelingToolset/Source/MeshModelingToolsEditorOnly/Private/ParameterizeMeshTool.cpp
+++ b/Engine/Plugins/Experimental/MeshModelingToolset/Source/MeshModelingToolsEditorOnly/Private/ParameterizeMeshTool.cpp
@@ -116,6 +116,7 @@ void UParameterizeMeshTool::Setup()
 	// Construct the preview object and set the material on it
 	Preview = NewObject<UMeshOpPreviewWithBackgroundCompute>(this, "Preview");
 	Preview->Setup(this->TargetWorld, this);
+	Preview->PreviewMesh->SetTangentsMode(EDynamicMeshTangentCalcType::AutoCalculated);
 
 	// Initialize the preview mesh with a copy of the source mesh.
 	bool bHasGroups = false;
diff --git a/Engine/Plugins/Experimental/MeshModelingToolset/Source/MeshModelingToolsEditorOnly/Private/PolygonOnMeshTool.cpp b/Engine/Plugins/Experimental/MeshModelingToolset/Source/MeshModelingToolsEditorOnly/Private/PolygonOnMeshTool.cpp
--- a/Engine/Plugins/Experimental/MeshModelingToolset/Source/MeshModelingToolsEditorOnly/Private/PolygonOnMeshTool.cpp
+++ b/Engine/Plugins/Experimental/MeshModelingToolset/Source/MeshModelingToolsEditorOnly/Private/PolygonOnMeshTool.cpp
@@ -166,6 +166,7 @@ void UPolygonOnMeshTool::UpdateNumPreviews()
 {
 	Preview = NewObject<UMeshOpPreviewWithBackgroundCompute>(this);
 	Preview->Setup(this->TargetWorld, this);
+	Preview->PreviewMesh->SetTangentsMode(EDynamicMeshTangentCalcType::AutoCalculated);
 
 	FComponentMaterialSet MaterialSet;
 	ComponentTarget->GetMaterialSet(MaterialSet);
diff --git a/Engine/Plugins/Experimental/MeshModelingToolset/Source/MeshModelingToolsEditorOnly/Private/ShapeSprayTool.cpp b/Engine/Plugins/Experimental/MeshModelingToolset/Source/MeshModelingToolsEditorOnly/Private/ShapeSprayTool.cpp
--- a/Engine/Plugins/Experimental/MeshModelingToolset/Source/MeshModelingToolsEditorOnly/Private/ShapeSprayTool.cpp
+++ b/Engine/Plugins/Experimental/MeshModelingToolset/Source/MeshModelingToolsEditorOnly/Private/ShapeSprayTool.cpp
@@ -75,9 +75,12 @@ void UShapeSprayTool::Setup()
 
 	UMaterialInterface* VtxColorMaterial = GetToolManager()->GetContextQueriesAPI()->GetStandardMaterial(EStandardToolContextMaterials::VertexColorMaterial);
 	check(VtxColorMaterial);
-	AccumMeshComponent->SetMaterial(0, VtxColorMaterial);
-	AccumMeshComponent->GetMesh()->EnableVertexColors(FVector3f::One());
-	AccumMeshComponent->GetMesh()->EnableVertexNormals(FVector3f::UnitX());
+	if (VtxColorMaterial)
+	{
+		AccumMeshComponent->SetMaterial(0, VtxColorMaterial);
+		AccumMeshComponent->GetMesh()->EnableVertexColors(FVector3f::One());
+		AccumMeshComponent->GetMesh()->EnableVertexNormals(FVector3f::UnitX());
+	}
 }
 
 
diff --git a/Engine/Plugins/Experimental/MeshModelingToolset/Source/MeshModelingToolsEditorOnly/Private/SimplifyMeshTool.cpp b/Engine/Plugins/Experimental/MeshModelingToolset/Source/MeshModelingToolsEditorOnly/Private/SimplifyMeshTool.cpp
--- a/Engine/Plugins/Experimental/MeshModelingToolset/Source/MeshModelingToolsEditorOnly/Private/SimplifyMeshTool.cpp
+++ b/Engine/Plugins/Experimental/MeshModelingToolset/Source/MeshModelingToolsEditorOnly/Private/SimplifyMeshTool.cpp
@@ -127,6 +127,7 @@ void USimplifyMeshTool::Setup()
 	}
 
 	Preview->PreviewMesh->SetTransform(ComponentTarget->GetWorldTransform());
+	Preview->PreviewMesh->SetTangentsMode(EDynamicMeshTangentCalcType::AutoCalculated);
 	Preview->PreviewMesh->UpdatePreview(OriginalMesh.Get());
 
 	// initialize our properties
diff --git a/Engine/Plugins/Experimental/MeshModelingToolset/Source/ModelingComponents/Private/SimpleDynamicMeshSceneProxy.h b/Engine/Plugins/Experimental/MeshModelingToolset/Source/ModelingComponents/Private/SimpleDynamicMeshSceneProxy.h
--- a/Engine/Plugins/Experimental/MeshModelingToolset/Source/ModelingComponents/Private/SimpleDynamicMeshSceneProxy.h
+++ b/Engine/Plugins/Experimental/MeshModelingToolset/Source/ModelingComponents/Private/SimpleDynamicMeshSceneProxy.h
@@ -105,9 +105,22 @@ public:
 			NormalOverlay = Mesh->Attributes()->PrimaryNormals();
 		}
 
+		TFunction<void(int, int, int, FVector3f&, FVector3f&)> TangentsFunc = nullptr;
+		if (UVOverlay != nullptr)
+		{
+			FMeshTangentsf* Tangents = ParentComponent->GetTangents();
+			if (Tangents != nullptr)
+			{
+				TangentsFunc = [Tangents](int VertexID, int TriangleID, int TriVtxIdx, FVector3f& TangentX, FVector3f& TangentY)
+				{
+					return Tangents->GetPerTriangleTangent(TriangleID, TriVtxIdx, TangentX, TangentY);
+				};
+			}
+		}
+
 		InitializeBuffersFromOverlays(RenderBuffers, Mesh,
 			Mesh->TriangleCount(), Mesh->TriangleIndicesItr(),
-			UVOverlay, NormalOverlay);
+			UVOverlay, NormalOverlay, TangentsFunc);
 
 		ENQUEUE_RENDER_COMMAND(FOctreeDynamicMeshSceneProxyInitializeSingle)(
 			[RenderBuffers](FRHICommandListImmediate& RHICmdList)
@@ -189,7 +202,7 @@ public:
 
 				InitializeBuffersFromOverlays(RenderBuffers, Mesh,
 					Triangles.Num(), Triangles,
-					UVOverlay, NormalOverlay);
+					UVOverlay, NormalOverlay, TangentsFunc);
 
 				RenderBuffers->Triangles = Triangles;
 
