Improve alignment of UV islands in ParameterizeMeshOp
#rnx

diff --git a/Engine/Plugins/Experimental/MeshModelingToolset/Source/ModelingOperatorsEditorOnly/Private/ParameterizationOps/ParameterizeMeshOp.cpp b/Engine/Plugins/Experimental/MeshModelingToolset/Source/ModelingOperatorsEditorOnly/Private/ParameterizationOps/ParameterizeMeshOp.cpp
--- a/Engine/Plugins/Experimental/MeshModelingToolset/Source/ModelingOperatorsEditorOnly/Private/ParameterizationOps/ParameterizeMeshOp.cpp
+++ b/Engine/Plugins/Experimental/MeshModelingToolset/Source/ModelingOperatorsEditorOnly/Private/ParameterizationOps/ParameterizeMeshOp.cpp
@@ -272,6 +272,9 @@ bool FParameterizeMeshOp::ComputeUVs_ExpMap(FDynamicMesh3& Mesh, TFunction<bool(
 	FFrame3d SeedFrame = Mesh.GetTriFrame(SeedTriangleID);
 	FIndex3i SeedNbrs = Mesh.GetTriangle(SeedTriangleID);
 
+	// try to generate consistent frame alignment...
+	SeedFrame.ConstrainedAlignPerpAxes(0, 1, 2, FVector3d::UnitX(), FVector3d::UnitY(), 0.95 );
+
 	TMeshLocalParam<FDynamicMesh3> Param(&Mesh);
 	//Param.ParamMode = ELocalParamTypes::PlanarProjection;
 	//Param.ParamMode = ELocalParamTypes::ExponentialMap;
