Restore shift-click-to-ignore-hit-normal functionality when repositioning construction plane in various tools
#rnx

diff --git a/Engine/Plugins/Experimental/MeshModelingToolset/Source/ModelingComponents/Private/Mechanics/ConstructionPlaneMechanic.cpp b/Engine/Plugins/Experimental/MeshModelingToolset/Source/ModelingComponents/Private/Mechanics/ConstructionPlaneMechanic.cpp
--- a/Engine/Plugins/Experimental/MeshModelingToolset/Source/ModelingComponents/Private/Mechanics/ConstructionPlaneMechanic.cpp
+++ b/Engine/Plugins/Experimental/MeshModelingToolset/Source/ModelingComponents/Private/Mechanics/ConstructionPlaneMechanic.cpp
@@ -39,9 +39,9 @@ void UConstructionPlaneMechanic::Initialize(UWorld* TargetWorld, const FFrame3d&
 	// click to set plane behavior
 	FSelectClickedAction* SetPlaneAction = new FSelectClickedAction();
 	SetPlaneAction->World = TargetWorld;
-	SetPlaneAction->OnClickedPositionFunc = [this](const FHitResult& Hit)
+	SetPlaneAction->OnClickedPositionFunc = [this, SetPlaneAction](const FHitResult& Hit)
 	{
-		SetDrawPlaneFromWorldPos(FVector3d(Hit.ImpactPoint), FVector3d(Hit.ImpactNormal));
+		SetDrawPlaneFromWorldPos(FVector3d(Hit.ImpactPoint), FVector3d(Hit.ImpactNormal), SetPlaneAction->bShiftModifierToggle);
 	};
 	SetPlaneAction->ExternalCanClickPredicate = [this]() { return this->CanUpdatePlaneFunc(); };
 
@@ -49,6 +49,7 @@ void UConstructionPlaneMechanic::Initialize(UWorld* TargetWorld, const FFrame3d&
 
 	ClickToSetPlaneBehavior = NewObject<USingleClickInputBehavior>();
 	ClickToSetPlaneBehavior->ModifierCheckFunc = FInputDeviceState::IsCtrlKeyDown;
+	ClickToSetPlaneBehavior->Modifiers.RegisterModifier(FSelectClickedAction::ShiftModifier, FInputDeviceState::IsShiftKeyDown);
 	ClickToSetPlaneBehavior->Initialize(SetPointInWorldConnector.Get());
 
 	GetParentTool()->AddInputBehavior(ClickToSetPlaneBehavior);
@@ -79,10 +80,13 @@ void UConstructionPlaneMechanic::TransformChanged(UTransformProxy* Proxy, FTrans
 }
 
 
-void UConstructionPlaneMechanic::SetDrawPlaneFromWorldPos(const FVector3d& Position, const FVector3d& Normal)
+void UConstructionPlaneMechanic::SetDrawPlaneFromWorldPos(const FVector3d& Position, const FVector3d& Normal, bool bIgnoreNormal)
 {
 	Plane.Origin = Position;
-	Plane.AlignAxis(2, Normal);
+	if (bIgnoreNormal == false)
+	{
+		Plane.AlignAxis(2, Normal);
+	}
 	PlaneTransformGizmo->SetActiveTarget(PlaneTransformProxy);
 	PlaneTransformGizmo->SetNewGizmoTransform(Plane.ToFTransform());
 
diff --git a/Engine/Plugins/Experimental/MeshModelingToolset/Source/ModelingComponents/Public/Mechanics/ConstructionPlaneMechanic.h b/Engine/Plugins/Experimental/MeshModelingToolset/Source/ModelingComponents/Public/Mechanics/ConstructionPlaneMechanic.h
--- a/Engine/Plugins/Experimental/MeshModelingToolset/Source/ModelingComponents/Public/Mechanics/ConstructionPlaneMechanic.h
+++ b/Engine/Plugins/Experimental/MeshModelingToolset/Source/ModelingComponents/Public/Mechanics/ConstructionPlaneMechanic.h
@@ -45,7 +45,7 @@ public:
 
 	void Initialize(UWorld* TargetWorld, const FFrame3d& InitialPlane);
 
-	void SetDrawPlaneFromWorldPos(const FVector3d& Position, const FVector3d& Normal);
+	void SetDrawPlaneFromWorldPos(const FVector3d& Position, const FVector3d& Normal, bool bIgnoreNormal);
 
 	void SetEnableGridSnaping(bool bEnable);
 
diff --git a/Engine/Plugins/Experimental/MeshModelingToolset/Source/ModelingComponents/Public/Selection/SelectClickedAction.h b/Engine/Plugins/Experimental/MeshModelingToolset/Source/ModelingComponents/Public/Selection/SelectClickedAction.h
--- a/Engine/Plugins/Experimental/MeshModelingToolset/Source/ModelingComponents/Public/Selection/SelectClickedAction.h
+++ b/Engine/Plugins/Experimental/MeshModelingToolset/Source/ModelingComponents/Public/Selection/SelectClickedAction.h
@@ -31,6 +31,10 @@ public:
 	TFunction<void(const FHitResult&)> OnClickedPositionFunc = nullptr;
 	TUniqueFunction<bool()> ExternalCanClickPredicate = nullptr;
 
+	// can alternately track shift modifier, however client must register this modifier w/ behavior
+	static const int ShiftModifier = 1;
+	bool bShiftModifierToggle = false;
+
 	virtual FInputRayHit IsHitByClick(const FInputDeviceRay& ClickPos) override
 	{
 		if (ExternalCanClickPredicate && ExternalCanClickPredicate() == false)
@@ -44,4 +48,14 @@ public:
 	{
 		DoRayCast(ClickPos, true);
 	}
+
+
+	virtual void OnUpdateModifierState(int ModifierID, bool bIsOn)
+	{
+		if (ModifierID == ShiftModifier)
+		{
+			bShiftModifierToggle = bIsOn;
+		}
+	}
+
 };
