Draw Poly Path tool: allow clicking on the first point to create a closed path.

#rnx

diff --git a/Engine/Plugins/Experimental/MeshModelingToolset/Source/MeshModelingTools/Private/DrawPolyPathTool.cpp b/Engine/Plugins/Experimental/MeshModelingToolset/Source/MeshModelingTools/Private/DrawPolyPathTool.cpp
--- a/Engine/Plugins/Experimental/MeshModelingToolset/Source/MeshModelingTools/Private/DrawPolyPathTool.cpp
+++ b/Engine/Plugins/Experimental/MeshModelingToolset/Source/MeshModelingTools/Private/DrawPolyPathTool.cpp
@@ -23,6 +23,125 @@
 #define LOCTEXT_NAMESPACE "UDrawPolyPathTool"
 
 
+namespace
+{
+	// Simple mesh generator that generates a quad for each edge of an input polygon
+	class FPolygonEdgeMeshGenerator : public FMeshShapeGenerator
+	{
+
+	private:
+
+		// Polygon to triangulate. Assumed to be closed (i.e. last edge is (LastVertex, FirstVertex)). If Polygon has 
+		// self-intersections or degenerate edges, result is undefined.
+		TArray<FFrame3d> Polygon;
+
+		// For each polygon vertex, a scale factor for the patch width at that vertex. Helps keep the width constant
+		// going around acute corners. 
+		TArray<double> OffsetScaleFactors;
+
+		// Width of quads to generate.
+		double Width = 1.0;
+
+		// Normal vector of all vertices will be set to this value. Default is +Z axis.
+		FVector3d Normal = FVector3d::UnitZ();
+
+	public:
+
+		FPolygonEdgeMeshGenerator(const TArray<FFrame3d>& InPolygon, 
+								  const TArray<double>& InOffsetScaleFactors,
+								  double InWidth = 1.0,
+								  FVector3d InNormal = FVector3d::UnitZ()) :
+			Polygon(InPolygon),
+			OffsetScaleFactors(InOffsetScaleFactors),
+			Width(InWidth),
+			Normal(InNormal)
+		{
+			check(Polygon.Num() == OffsetScaleFactors.Num());
+		}
+
+
+		// Generate the triangulation
+		virtual FMeshShapeGenerator& Generate() final
+		{
+			const int NumInputVertices = Polygon.Num();
+			check(NumInputVertices >= 3);
+			if (NumInputVertices < 3)
+			{
+				return *this;
+			}
+			const int NumVertices = 2 * NumInputVertices;
+			const int NumTriangles = 2 * NumInputVertices;
+			SetBufferSizes(NumVertices, NumTriangles, NumVertices, NumVertices);
+
+			// Trace the input path, placing vertices on either side of each input vertex 
+
+			const FVector3d LeftVertex{ 0, -Width, 0 };
+			const FVector3d RightVertex{ 0, Width, 0 };
+			for (int CurrentInputVertex = 0; CurrentInputVertex < NumInputVertices; ++CurrentInputVertex)
+			{
+				const FFrame3d& CurrentFrame = Polygon[CurrentInputVertex];
+				const int NewVertexAIndex = 2 * CurrentInputVertex;
+				const int NewVertexBIndex = NewVertexAIndex + 1;
+				Vertices[NewVertexAIndex] = CurrentFrame.FromFramePoint(OffsetScaleFactors[CurrentInputVertex] * LeftVertex);
+				Vertices[NewVertexBIndex] = CurrentFrame.FromFramePoint(OffsetScaleFactors[CurrentInputVertex] * RightVertex);
+			}
+
+			// Triangulate the vertices we just placed
+
+			for (int CurrentVertex = 0; CurrentVertex < NumInputVertices; ++CurrentVertex)
+			{
+				const int NewVertexA = 2 * CurrentVertex;
+				const int NewVertexB = NewVertexA + 1;
+				const int NewVertexC = (NewVertexA + 2) % NumVertices;
+				const int NewVertexD = (NewVertexA + 3) % NumVertices;
+
+				FIndex3i NewTriA{ NewVertexA, NewVertexB, NewVertexC };
+				const int NewTriAIndex = 2 * CurrentVertex;
+
+				SetTriangle(NewTriAIndex, NewTriA);
+				SetTriangleUVs(NewTriAIndex, NewTriA);
+				SetTriangleNormals(NewTriAIndex, NewTriA);
+				SetTrianglePolygon(NewTriAIndex, 0);
+
+				const int NewTriBIndex = NewTriAIndex + 1;
+				FIndex3i NewTriB{ NewVertexC, NewVertexB, NewVertexD };
+
+				SetTriangle(NewTriBIndex, NewTriB);
+				SetTriangleUVs(NewTriBIndex, NewTriB);
+				SetTriangleNormals(NewTriBIndex, NewTriB);
+				SetTrianglePolygon(NewTriBIndex, 0);
+			}
+
+			// Set UVs, etc. for the vertices we put down previously
+
+			FAxisAlignedBox2d BoundingBox;
+			for (auto& CurrentFrame : Polygon)
+			{
+				BoundingBox.Contain(FVector2d{ CurrentFrame.Origin.X, CurrentFrame.Origin.Y });
+			}
+			BoundingBox.Max += {Width, Width};
+			BoundingBox.Min -= {Width, Width};
+
+			double BoxWidth = BoundingBox.Width(), BoxHeight = BoundingBox.Height();
+			double UVScale = FMath::Max(BoxWidth, BoxHeight);
+
+			for (int NewVertexIndex = 0; NewVertexIndex < NumVertices; ++NewVertexIndex)
+			{
+				FVector3d Pos = Vertices[NewVertexIndex];
+
+				UVs[NewVertexIndex] = FVector2f(
+					(float)((Pos.X - BoundingBox.Min.X) / UVScale),
+					(float)((Pos.Y - BoundingBox.Min.Y) / UVScale));
+				UVParentVertex[NewVertexIndex] = NewVertexIndex;
+				Normals[NewVertexIndex] = FVector3f(Normal);
+				NormalParentVertex[NewVertexIndex] = NewVertexIndex;
+			}
+
+			return *this;
+		}
+	};		// class FPolygonEdgeMeshGenerator
+
+}	// unnamed namespace
 
 /*
  * ToolBuilder
@@ -214,6 +333,7 @@ void UDrawPolyPathTool::OnClicked(const FInputDeviceRay& ClickPos)
 		{
 			if (SurfacePathMechanic->IsDone())
 			{
+				bPathIsClosed = SurfacePathMechanic->LoopWasClosed();
 				GetToolManager()->EmitObjectChange(this, MakeUnique<FDrawPolyPathStateChange>(CurrentCurveTimestamp), LOCTEXT("DrawPolyPathBeginOffset", "Set Offset"));
 				OnCompleteSurfacePath();
 			}
@@ -325,6 +445,7 @@ void UDrawPolyPathTool::InitializeNewSurfacePath()
 	{
 		return true && ToolSceneQueriesUtil::CalculateViewVisualAngleD(this->CameraState, Position1, Position2) < SnapTol;
 	};
+	SurfacePathMechanic->SetDoubleClickOrCloseLoopMode();
 	UpdateSurfacePathPlane();
 }
 
@@ -358,10 +479,17 @@ void UDrawPolyPathTool::OnCompleteSurfacePath()
 	OffsetScaleFactors[0] = OffsetScaleFactors[NumPoints-1] = 1.0;
 	ArcLengths.SetNum(NumPoints);
 	ArcLengths[0] = 0;
-	for (int j = 1; j < NumPoints - 1; ++j)
+
+	// Set local frames for path points. If the path is closed, we will adjust the first and last frames for continuity,
+	// otherwise we will leave them as set above.
+	int LastPointIndex = bPathIsClosed ? NumPoints : NumPoints - 1;
+	int FirstPointIndex = bPathIsClosed ? 0 : 1;
+	for (int j = FirstPointIndex; j < LastPointIndex; ++j)
 	{
-		FVector3d Prev(CurPathPoints[j - 1].Origin), Next(CurPathPoints[j + 1].Origin), Cur(CurPathPoints[j].Origin);
-		ArcLengths[j] = ArcLengths[j-1] + Cur.Distance(Prev);
+		int NextJ = (j + 1) % NumPoints;
+		int PrevJ = (j - 1 + NumPoints) % NumPoints;
+		FVector3d Prev(CurPathPoints[PrevJ].Origin), Next(CurPathPoints[NextJ].Origin), Cur(CurPathPoints[j].Origin);
+		ArcLengths[j] = ArcLengths[PrevJ] + Cur.Distance(Prev);
 		FLine3d Line1(FLine3d::FromPoints(Prev, Cur)), Line2(FLine3d::FromPoints(Cur, Next));
 		Line1.Origin += DistOffsetDelta * PlaneNormal.Cross(Line1.Direction);
 		Line2.Origin += DistOffsetDelta * PlaneNormal.Cross(Line2.Direction);
@@ -474,43 +602,52 @@ void UDrawPolyPathTool::GeneratePathMesh(FDynamicMesh3& Mesh)
 	int32 NumPoints = CurPathPoints.Num();
 	if (NumPoints > 1)
 	{
-		CurPathLength = 0;
-		for (int32 k = 1; k < NumPoints; ++k)
+		if (bPathIsClosed)
 		{
-			CurPathLength += CurPathPoints[k].Origin.Distance(CurPathPoints[k-1].Origin);
+			FPolygonEdgeMeshGenerator MeshGen(CurPathPoints, OffsetScaleFactors, CurOffsetDistance, FVector3d::UnitZ());
+			MeshGen.Generate();
+			Mesh.Copy(&MeshGen);
 		}
-
-		FRectangleMeshGenerator MeshGen;
-		MeshGen.Width = CurPathLength;
-		MeshGen.Height = 2 * CurOffsetDistance;
-		MeshGen.Normal = FVector3f::UnitZ();
-		MeshGen.Origin = FVector3d(CurPathLength / 2, 0, 0);
-		MeshGen.HeightVertexCount = 2;
-		MeshGen.WidthVertexCount = NumPoints;
-		MeshGen.Generate();
-		Mesh.Copy(&MeshGen);
-		Mesh.EnableVertexUVs(FVector2f::Zero());		// we will store arc length for each vtx in VertexUV
-
-		FAxisAlignedBox3d PathBounds = Mesh.GetBounds();
-
-		double ShiftX = 0;
-		double DeltaX = CurPathLength / (double)(NumPoints - 1);
-		for (int32 k = 0; k < NumPoints; ++k)
+		else
 		{
-			FFrame3d PathFrame = CurPathPoints[k];
-			FVector3d V0 = Mesh.GetVertex(k);
-			V0.X -= ShiftX;
-			V0.Y *= OffsetScaleFactors[k];
-			V0 = PathFrame.FromFramePoint(V0);
-			Mesh.SetVertex(k, V0);
-			Mesh.SetVertexUV(k, FVector2f((float)ArcLengths[k],(float)k));
-			FVector3d V1 = Mesh.GetVertex(NumPoints + k);
-			V1.X -= ShiftX;
-			V1.Y *= OffsetScaleFactors[k];
-			V1 = PathFrame.FromFramePoint(V1);
-			Mesh.SetVertex(NumPoints + k, V1);
-			Mesh.SetVertexUV(NumPoints+k, FVector2f((float)ArcLengths[k], (float)k));
-			ShiftX += DeltaX;
+			CurPathLength = 0;
+			for (int32 k = 1; k < NumPoints; ++k)
+			{
+				CurPathLength += CurPathPoints[k].Origin.Distance(CurPathPoints[k - 1].Origin);
+			}
+
+			FRectangleMeshGenerator MeshGen;
+			MeshGen.Width = CurPathLength;
+			MeshGen.Height = 2 * CurOffsetDistance;
+			MeshGen.Normal = FVector3f::UnitZ();
+			MeshGen.Origin = FVector3d(CurPathLength / 2, 0, 0);
+			MeshGen.HeightVertexCount = 2;
+			MeshGen.WidthVertexCount = NumPoints;
+			MeshGen.Generate();
+			Mesh.Copy(&MeshGen);
+			Mesh.EnableVertexUVs(FVector2f::Zero());		// we will store arc length for each vtx in VertexUV
+
+			FAxisAlignedBox3d PathBounds = Mesh.GetBounds();
+
+			double ShiftX = 0;
+			double DeltaX = CurPathLength / (double)(NumPoints - 1);
+			for (int32 k = 0; k < NumPoints; ++k)
+			{
+				FFrame3d PathFrame = CurPathPoints[k];
+				FVector3d V0 = Mesh.GetVertex(k);
+				V0.X -= ShiftX;
+				V0.Y *= OffsetScaleFactors[k];
+				V0 = PathFrame.FromFramePoint(V0);
+				Mesh.SetVertex(k, V0);
+				Mesh.SetVertexUV(k, FVector2f((float)ArcLengths[k], (float)k));
+				FVector3d V1 = Mesh.GetVertex(NumPoints + k);
+				V1.X -= ShiftX;
+				V1.Y *= OffsetScaleFactors[k];
+				V1 = PathFrame.FromFramePoint(V1);
+				Mesh.SetVertex(NumPoints + k, V1);
+				Mesh.SetVertexUV(NumPoints + k, FVector2f((float)ArcLengths[k], (float)k));
+				ShiftX += DeltaX;
+			}
 		}
 
 		FMeshNormals::QuickRecomputeOverlayNormals(Mesh);
diff --git a/Engine/Plugins/Experimental/MeshModelingToolset/Source/MeshModelingTools/Public/DrawPolyPathTool.h b/Engine/Plugins/Experimental/MeshModelingToolset/Source/MeshModelingTools/Public/DrawPolyPathTool.h
--- a/Engine/Plugins/Experimental/MeshModelingToolset/Source/MeshModelingTools/Public/DrawPolyPathTool.h
+++ b/Engine/Plugins/Experimental/MeshModelingToolset/Source/MeshModelingTools/Public/DrawPolyPathTool.h
@@ -195,6 +195,7 @@ protected:
 	double CurPathLength;
 	double CurOffsetDistance;
 	double CurHeight;
+	bool bPathIsClosed = false;		// If true, CurPathPoints are assumed to define a closed path
 
 	UPROPERTY()
 	UPolyEditPreviewMesh* EditPreview;
@@ -246,4 +247,4 @@ public:
 	virtual void Revert(UObject* Object) override;
 	virtual bool HasExpired(UObject* Object) const override;
 	virtual FString ToString() const override;
-};
\ No newline at end of file
+};
diff --git a/Engine/Plugins/Experimental/MeshModelingToolset/Source/ModelingComponents/Private/Mechanics/CollectSurfacePathMechanic.cpp b/Engine/Plugins/Experimental/MeshModelingToolset/Source/ModelingComponents/Private/Mechanics/CollectSurfacePathMechanic.cpp
--- a/Engine/Plugins/Experimental/MeshModelingToolset/Source/ModelingComponents/Private/Mechanics/CollectSurfacePathMechanic.cpp
+++ b/Engine/Plugins/Experimental/MeshModelingToolset/Source/ModelingComponents/Private/Mechanics/CollectSurfacePathMechanic.cpp
@@ -95,6 +95,10 @@ void UCollectSurfacePathMechanic::SetDrawClosedLoopMode()
 	DoneMode = ECollectSurfacePathDoneMode::SnapCloseLoop;
 }
 
+void UCollectSurfacePathMechanic::SetDoubleClickOrCloseLoopMode()
+{
+	DoneMode = ECollectSurfacePathDoneMode::SnapDoubleClickOrCloseLoop;
+}
 
 
 bool UCollectSurfacePathMechanic::IsHitByRay(const FRay3d& Ray, FFrame3d& HitPoint)
@@ -237,7 +241,7 @@ bool UCollectSurfacePathMechanic::IsDone() const
 	{
 		return IsDoneFunc();
 	}
-	else if (DoneMode == ECollectSurfacePathDoneMode::SnapCloseLoop || DoneMode == ECollectSurfacePathDoneMode::SnapDoubleClick)
+	else if (DoneMode == ECollectSurfacePathDoneMode::SnapCloseLoop || DoneMode == ECollectSurfacePathDoneMode::SnapDoubleClick || DoneMode == ECollectSurfacePathDoneMode::SnapDoubleClickOrCloseLoop)
 	{
 		return bGeometricCloseOcurred;
 	}
@@ -253,20 +257,28 @@ bool UCollectSurfacePathMechanic::CheckGeometricClosure(const FFrame3d& Point)
 		return false;
 	}
 
-	if (DoneMode == ECollectSurfacePathDoneMode::SnapCloseLoop)
+	if (DoneMode == ECollectSurfacePathDoneMode::SnapCloseLoop || DoneMode == ECollectSurfacePathDoneMode::SnapDoubleClickOrCloseLoop)
 	{
 		if (HitPath.Num() > 2)
 		{
 			const FFrame3d& FirstPoint = HitPath[0];
-			return SpatialSnapPointsFunc(Point.Origin, FirstPoint.Origin);
+			if (SpatialSnapPointsFunc(Point.Origin, FirstPoint.Origin))
+			{
+				bLoopWasClosed = true;		// We finished by clicking on the first point
+				return true;
+			}
 		}
 	}
-	else if (DoneMode == ECollectSurfacePathDoneMode::SnapDoubleClick)
+
+	if (DoneMode == ECollectSurfacePathDoneMode::SnapDoubleClick || DoneMode == ECollectSurfacePathDoneMode::SnapDoubleClickOrCloseLoop)
 	{
 		if (HitPath.Num() > 1)
 		{
 			const FFrame3d& LastPoint = HitPath[HitPath.Num() - 1];
-			return SpatialSnapPointsFunc(Point.Origin, LastPoint.Origin);
+			if (SpatialSnapPointsFunc(Point.Origin, LastPoint.Origin))
+			{
+				return true;
+			}
 		}
 	}
 
diff --git a/Engine/Plugins/Experimental/MeshModelingToolset/Source/ModelingComponents/Public/Mechanics/CollectSurfacePathMechanic.h b/Engine/Plugins/Experimental/MeshModelingToolset/Source/ModelingComponents/Public/Mechanics/CollectSurfacePathMechanic.h
--- a/Engine/Plugins/Experimental/MeshModelingToolset/Source/ModelingComponents/Public/Mechanics/CollectSurfacePathMechanic.h
+++ b/Engine/Plugins/Experimental/MeshModelingToolset/Source/ModelingComponents/Public/Mechanics/CollectSurfacePathMechanic.h
@@ -14,6 +14,7 @@ enum class ECollectSurfacePathDoneMode
 {
 	SnapCloseLoop,
 	SnapDoubleClick,
+	SnapDoubleClickOrCloseLoop,
 	ExternalLambda,
 	FixedNumPoints
 };
@@ -65,6 +66,8 @@ public:
 	virtual void SetFixedNumPointsMode(int32 NumPoints);
 	virtual void SetDrawClosedLoopMode();
 	virtual void SetCloseWithLambdaMode();
+	virtual void SetDoubleClickOrCloseLoopMode();
+
 
 	virtual bool IsHitByRay(const FRay3d& Ray, FFrame3d& HitPoint);
 	virtual bool UpdatePreviewPoint(const FRay3d& Ray);
@@ -74,6 +77,11 @@ public:
 
 	virtual bool IsDone() const;
 
+	/** Whether the path was finished by the user clicking on the first point */
+	bool LoopWasClosed() const
+	{
+		return bLoopWasClosed;
+	}
 
 protected:
 	FDynamicMesh3 TargetSurface;
@@ -89,4 +97,6 @@ protected:
 	bool bCurrentPreviewWillComplete = false;
 	bool bGeometricCloseOcurred = false;
 	bool CheckGeometricClosure(const FFrame3d& Point);
+
+	bool bLoopWasClosed = false;
 };
