ModelingTools: Improve tool info messages
#rnx

diff --git a/Engine/Plugins/Experimental/MeshModelingToolset/Source/MeshModelingTools/Private/AddPrimitiveTool.cpp b/Engine/Plugins/Experimental/MeshModelingToolset/Source/MeshModelingTools/Private/AddPrimitiveTool.cpp
--- a/Engine/Plugins/Experimental/MeshModelingToolset/Source/MeshModelingTools/Private/AddPrimitiveTool.cpp
+++ b/Engine/Plugins/Experimental/MeshModelingToolset/Source/MeshModelingTools/Private/AddPrimitiveTool.cpp
@@ -163,7 +163,7 @@ void UAddPrimitiveTool::Setup()
 	UpdatePreviewMesh();
 
 	GetToolManager()->DisplayMessage(
-		LOCTEXT("OnStartAddPrimitiveTool", "Position the Primitive by moving the mouse over the scene. Drop a new instance by Left-clicking."),
+		LOCTEXT("OnStartAddPrimitiveTool", "This Tool creates new Primitive objects. Position the Primitive by moving the mouse over the scene. Drop a new Asset or Instance by left-clicking (depending on Instance setting)."),
 		EToolMessageLevel::UserNotification);
 }
 
diff --git a/Engine/Plugins/Experimental/MeshModelingToolset/Source/MeshModelingTools/Private/AlignObjectsTool.cpp b/Engine/Plugins/Experimental/MeshModelingToolset/Source/MeshModelingTools/Private/AlignObjectsTool.cpp
--- a/Engine/Plugins/Experimental/MeshModelingToolset/Source/MeshModelingTools/Private/AlignObjectsTool.cpp
+++ b/Engine/Plugins/Experimental/MeshModelingToolset/Source/MeshModelingTools/Private/AlignObjectsTool.cpp
@@ -24,7 +24,7 @@
 
 bool UAlignObjectsToolBuilder::CanBuildTool(const FToolBuilderState& SceneState) const
 {
-	return ToolBuilderUtil::CountComponents(SceneState, CanMakeComponentTarget) >= 1;
+	return ToolBuilderUtil::CountComponents(SceneState, CanMakeComponentTarget) >= 2;
 }
 
 UInteractiveTool* UAlignObjectsToolBuilder::BuildTool(const FToolBuilderState& SceneState) const
@@ -104,6 +104,10 @@ void UAlignObjectsTool::Setup()
 	AddToolPropertySource(AlignProps);
 
 	Precompute();
+
+	GetToolManager()->DisplayMessage(
+		LOCTEXT("OnStartTool", "This Tool aligns the Pivots or Bounding Boxes of the input Objects."),
+		EToolMessageLevel::UserNotification);
 }
 
 
diff --git a/Engine/Plugins/Experimental/MeshModelingToolset/Source/MeshModelingTools/Private/BakeTransformTool.cpp b/Engine/Plugins/Experimental/MeshModelingToolset/Source/MeshModelingTools/Private/BakeTransformTool.cpp
--- a/Engine/Plugins/Experimental/MeshModelingToolset/Source/MeshModelingTools/Private/BakeTransformTool.cpp
+++ b/Engine/Plugins/Experimental/MeshModelingToolset/Source/MeshModelingTools/Private/BakeTransformTool.cpp
@@ -108,6 +108,10 @@ void UBakeTransformTool::Setup()
 	}
 	
 	GetToolManager()->DisplayMessage(AllTheWarnings, EToolMessageLevel::UserWarning);
+
+	GetToolManager()->DisplayMessage(
+		LOCTEXT("OnStartTool", "This Tool applies the current Rotation and/or Scaling of the object's Transform to the underlying mesh Asset."),
+		EToolMessageLevel::UserNotification);
 }
 
 
diff --git a/Engine/Plugins/Experimental/MeshModelingToolset/Source/MeshModelingTools/Private/CombineMeshesTool.cpp b/Engine/Plugins/Experimental/MeshModelingToolset/Source/MeshModelingTools/Private/CombineMeshesTool.cpp
--- a/Engine/Plugins/Experimental/MeshModelingToolset/Source/MeshModelingTools/Private/CombineMeshesTool.cpp
+++ b/Engine/Plugins/Experimental/MeshModelingToolset/Source/MeshModelingTools/Private/CombineMeshesTool.cpp
@@ -95,6 +95,20 @@ void UCombineMeshesTool::Setup()
 
 	BasicProperties = NewObject<UCombineMeshesToolProperties>(this);
 	AddToolPropertySource(BasicProperties);
+
+
+	if (bDuplicateMode)
+	{
+		GetToolManager()->DisplayMessage(
+			LOCTEXT("OnStartTool", "This Tool duplicates input Asset into a new Asset, and optionally replaces the input Actor with a new Actor containing the new Asset."),
+			EToolMessageLevel::UserNotification);
+	}
+	else
+	{
+		GetToolManager()->DisplayMessage(
+			LOCTEXT("OnStartTool", "This Tool appends the meshes from the input Assets into a new Asset, and optionally replaces the source Actors with a new Actor containing the new Asset."),
+			EToolMessageLevel::UserNotification);
+	}
 }
 
 
diff --git a/Engine/Plugins/Experimental/MeshModelingToolset/Source/MeshModelingTools/Private/DrawPolyPathTool.cpp b/Engine/Plugins/Experimental/MeshModelingToolset/Source/MeshModelingTools/Private/DrawPolyPathTool.cpp
--- a/Engine/Plugins/Experimental/MeshModelingToolset/Source/MeshModelingTools/Private/DrawPolyPathTool.cpp
+++ b/Engine/Plugins/Experimental/MeshModelingToolset/Source/MeshModelingTools/Private/DrawPolyPathTool.cpp
@@ -239,9 +239,7 @@ void UDrawPolyPathTool::Setup()
 	// begin path draw
 	InitializeNewSurfacePath();
 
-	GetToolManager()->DisplayMessage(
-		LOCTEXT("OnStartDrawPolyPathTool", "Click to begin drawing path. Doubleclick to finish path."),
-		EToolMessageLevel::UserNotification);
+	ShowStartupMessage();
 }
 
 
@@ -447,6 +445,8 @@ void UDrawPolyPathTool::InitializeNewSurfacePath()
 	};
 	SurfacePathMechanic->SetDoubleClickOrCloseLoopMode();
 	UpdateSurfacePathPlane();
+
+	ShowStartupMessage();
 }
 
 
@@ -537,6 +537,8 @@ void UDrawPolyPathTool::BeginInteractiveOffsetDistance()
 	CurveDistMechanic->InitializePolyCurve(CurPolyLine, FTransform3d::Identity());
 
 	InitializePreviewMesh();
+
+	ShowOffsetMessage();
 }
 
 
@@ -684,6 +686,8 @@ void UDrawPolyPathTool::BeginInteractiveExtrudeHeight()
 	FFrame3d UseFrame = DrawPlaneWorld; 
 	UseFrame.Origin = CurPathPoints.Last().Origin;
 	ExtrudeHeightMechanic->Initialize(MoveTemp(TmpMesh), UseFrame, true);
+
+	ShowExtrudeMessage();
 }
 
 
@@ -806,6 +810,29 @@ void UDrawPolyPathTool::EmitNewObject(EDrawPolyPathOutputMode OutputMode)
 
 
 
+void UDrawPolyPathTool::ShowStartupMessage()
+{
+	GetToolManager()->DisplayMessage(
+		LOCTEXT("OnStartDraw", "Use this Tool to draw a path on the Drawing Plane, and then Extrude. Left-click to place points, Double-click (or close) to complete path. Ctrl-click on the scene to reposition the Plane (Shift+Ctrl-click to ignore Normal). [A] toggles Gizmo. Hold Shift to ignore Snapping."),
+		EToolMessageLevel::UserNotification);
+}
+
+void UDrawPolyPathTool::ShowOffsetMessage()
+{
+	GetToolManager()->DisplayMessage(
+		LOCTEXT("OnStartOffset", "Set the Width of the Extrusion by clicking on the Drawing Plane."),
+		EToolMessageLevel::UserNotification);
+}
+
+void UDrawPolyPathTool::ShowExtrudeMessage()
+{
+	GetToolManager()->DisplayMessage(
+		LOCTEXT("OnStartExtrude", "Set the Height of the Extrusion by positioning the mouse over the extrusion volume, or over the scene to snap to relative heights."),
+		EToolMessageLevel::UserNotification);
+}
+
+
+
 void UDrawPolyPathTool::UndoCurrentOperation()
 {
 	if (SurfacePathMechanic != nullptr)
diff --git a/Engine/Plugins/Experimental/MeshModelingToolset/Source/MeshModelingTools/Private/DrawPolygonTool.cpp b/Engine/Plugins/Experimental/MeshModelingToolset/Source/MeshModelingTools/Private/DrawPolygonTool.cpp
--- a/Engine/Plugins/Experimental/MeshModelingToolset/Source/MeshModelingTools/Private/DrawPolygonTool.cpp
+++ b/Engine/Plugins/Experimental/MeshModelingToolset/Source/MeshModelingTools/Private/DrawPolygonTool.cpp
@@ -254,6 +254,8 @@ void UDrawPolygonTool::PopLastVertexAction()
 		return;
 	}
 
+	bHaveSelfIntersection = false;
+
 	if (bInFixedPolygonMode == false)
 	{
 		int NumVertices = PolygonVertices.Num();
@@ -497,6 +499,7 @@ void UDrawPolygonTool::ResetPolygon()
 	SnapEngine.Reset();
 	bHaveSurfaceHit = false;
 	bInFixedPolygonMode = false;
+	bHaveSelfIntersection = false;
 	CurrentCurveTimestamp++;
 }
 
@@ -664,10 +667,10 @@ void UDrawPolygonTool::OnNextSequencePreview(const FInputDeviceRay& ClickPos)
 	}
 
 	UpdatePreviewVertex(HitPos);
+	UpdateSelfIntersection();
 	if (PolygonVertices.Num() > 2)
 	{
 		bPreviewUpdatePending = true;
-		UpdateSelfIntersection();
 	}
 }
 
@@ -730,7 +733,7 @@ bool UDrawPolygonTool::OnNextSequenceClick(const FInputDeviceRay& ClickPos)
 	
 	if (bDonePolygon)
 	{
-		SnapEngine.Reset();
+		//SnapEngine.Reset();
 		bHaveSurfaceHit = false;
 		if (PolygonProperties->OutputMode == EDrawPolygonOutputMode::ExtrudedInteractive)
 		{
@@ -808,6 +811,10 @@ bool UDrawPolygonTool::UpdateSelfIntersection()
 	}
 
 	int NumVertices = PolygonVertices.Num();
+	if (NumVertices < 3)
+	{
+		return false;
+	}
 
 	FFrame3f DrawFrame(DrawPlaneOrigin, DrawPlaneOrientation);
 	FSegment2f PreviewSegment(DrawFrame.ToPlaneUV(PolygonVertices[NumVertices - 1],2), DrawFrame.ToPlaneUV(PreviewVertex,2));
@@ -1272,7 +1279,7 @@ bool UDrawPolygonTool::GeneratePolygonMesh(const TArray<FVector>& Polygon, const
 void UDrawPolygonTool::ShowStartupMessage()
 {
 	GetToolManager()->DisplayMessage(
-		LOCTEXT("OnStartDraw", "Left-click to place points on the Drawing Plane. Hold Shift to ignore Snapping. Ctrl-click on the scene to reposition the Plane (Shift+Ctrl-click to only Translate). Backspace to discard last vertex. A key toggles Gizmo."),
+		LOCTEXT("OnStartDraw", "Use this Tool to draw a polygon on the Drawing Plane, and Extrude it. Left-click to place points. Ctrl-click on the scene to reposition the Plane (Shift+Ctrl-click to ignore Normal). [A] toggles Gizmo. Hold Shift to ignore Snapping."),
 		EToolMessageLevel::UserNotification);
 }
 
diff --git a/Engine/Plugins/Experimental/MeshModelingToolset/Source/MeshModelingTools/Private/EditPivotTool.cpp b/Engine/Plugins/Experimental/MeshModelingToolset/Source/MeshModelingTools/Private/EditPivotTool.cpp
--- a/Engine/Plugins/Experimental/MeshModelingToolset/Source/MeshModelingTools/Private/EditPivotTool.cpp
+++ b/Engine/Plugins/Experimental/MeshModelingToolset/Source/MeshModelingTools/Private/EditPivotTool.cpp
@@ -115,6 +115,10 @@ void UEditPivotTool::Setup()
 	}
 
 	GetToolManager()->DisplayMessage(AllTheWarnings, EToolMessageLevel::UserWarning);
+
+	GetToolManager()->DisplayMessage(
+		LOCTEXT("OnStartTool", "This Tool edits the Pivot (Origin) of the input objects. Use the Gizmo or enable Snap Dragging and click-drag on the surface to reposition."),
+		EToolMessageLevel::UserNotification);
 }
 
 
diff --git a/Engine/Plugins/Experimental/MeshModelingToolset/Source/MeshModelingTools/Private/MeshSelectionTool.cpp b/Engine/Plugins/Experimental/MeshModelingToolset/Source/MeshModelingTools/Private/MeshSelectionTool.cpp
--- a/Engine/Plugins/Experimental/MeshModelingToolset/Source/MeshModelingTools/Private/MeshSelectionTool.cpp
+++ b/Engine/Plugins/Experimental/MeshModelingToolset/Source/MeshModelingTools/Private/MeshSelectionTool.cpp
@@ -158,6 +158,12 @@ void UMeshSelectionTool::Setup()
 
 	RecalculateBrushRadius();
 	UpdateVisualization(true);
+
+	GetToolManager()->DisplayMessage(
+		LOCTEXT("OnStartMeshSelectionTool", "This Tool allows you to modify the mesh based on a triangle selection. [Q] cyles through Selection Mode. [A] cycles through Face Color modes. [ and ] change brush size, < and > grow/shrink selection."),
+		EToolMessageLevel::UserNotification);
+
+
 }
 
 
@@ -254,14 +260,14 @@ void UMeshSelectionTool::RegisterActions(FInteractiveToolActionSet& ActionSet)
 		TEXT("ShrinkSelection"),
 		LOCTEXT("ShrinkSelection", "Shrink Selection"),
 		LOCTEXT("ShrinkSelectionTooltip", "Shrink selection"),
-		EModifierKey::None, EKeys::Comma,
+		EModifierKey::Shift, EKeys::Comma,
 		[this]() { GrowShrinkSelection(false); });
 
 	ActionSet.RegisterAction(this, (int32)EMeshSelectionToolActions::GrowSelection,
 		TEXT("GrowSelection"),
 		LOCTEXT("GrowSelection", "Grow Selection"),
 		LOCTEXT("GrowSelectionTooltip", "Grow selection"),
-		EModifierKey::None, EKeys::Period,
+		EModifierKey::Shift, EKeys::Period,
 		[this]() { GrowShrinkSelection(true); });
 
 	ActionSet.RegisterAction(this, (int32)EMeshSelectionToolActions::OptimizeSelection,
diff --git a/Engine/Plugins/Experimental/MeshModelingToolset/Source/MeshModelingTools/Private/TransformMeshesTool.cpp b/Engine/Plugins/Experimental/MeshModelingToolset/Source/MeshModelingTools/Private/TransformMeshesTool.cpp
--- a/Engine/Plugins/Experimental/MeshModelingToolset/Source/MeshModelingTools/Private/TransformMeshesTool.cpp
+++ b/Engine/Plugins/Experimental/MeshModelingToolset/Source/MeshModelingTools/Private/TransformMeshesTool.cpp
@@ -82,7 +82,7 @@ void UTransformMeshesTool::Setup()
 	UpdateTransformMode(TransformProps->TransformMode);
 
 	GetToolManager()->DisplayMessage(
-		LOCTEXT("OnStartTransformMeshesTool", "To transform Objects around points, reposition the Gizmo using Set Pivot mode (S). To quickly position Objects, enable Snap Drag mode (D). A cycles through Transform modes, W and E through SnapDrag Source and Rotation types."),
+		LOCTEXT("OnStartTransformMeshesTool", "This Tool allows you to transform multiple objects. [A] cycles through Transform modes. Reposition the Gizmo using Set Pivot Mode [S]. Toggle Snap Drag Mode [D] to move the object by click-draging. [W] and [E] cycle through SnapDrag Source and Rotation types."),
 		EToolMessageLevel::UserNotification);
 }
 
diff --git a/Engine/Plugins/Experimental/MeshModelingToolset/Source/MeshModelingTools/Public/DrawPolyPathTool.h b/Engine/Plugins/Experimental/MeshModelingToolset/Source/MeshModelingTools/Public/DrawPolyPathTool.h
--- a/Engine/Plugins/Experimental/MeshModelingToolset/Source/MeshModelingTools/Public/DrawPolyPathTool.h
+++ b/Engine/Plugins/Experimental/MeshModelingToolset/Source/MeshModelingTools/Public/DrawPolyPathTool.h
@@ -225,6 +225,11 @@ protected:
 	void GenerateRampMesh(FDynamicMesh3& PathMesh);
 	void EmitNewObject(EDrawPolyPathOutputMode OutputMode);
 
+	// user feedback messages
+	void ShowStartupMessage();
+	void ShowExtrudeMessage();
+	void ShowOffsetMessage();
+
 	friend class FDrawPolyPathStateChange;
 	int32 CurrentCurveTimestamp = 1;
 	void UndoCurrentOperation();
diff --git a/Engine/Plugins/Experimental/MeshModelingToolset/Source/MeshModelingToolsEditorOnly/Private/MergeMeshesTool.cpp b/Engine/Plugins/Experimental/MeshModelingToolset/Source/MeshModelingToolsEditorOnly/Private/MergeMeshesTool.cpp
--- a/Engine/Plugins/Experimental/MeshModelingToolset/Source/MeshModelingToolsEditorOnly/Private/MergeMeshesTool.cpp
+++ b/Engine/Plugins/Experimental/MeshModelingToolset/Source/MeshModelingToolsEditorOnly/Private/MergeMeshesTool.cpp
@@ -110,6 +110,10 @@ void UMergeMeshesTool::Setup()
 
 	Preview->InvalidateResult();    // start compute
 
+
+	GetToolManager()->DisplayMessage(
+		LOCTEXT("OnStartTool", "This Tool combines the input meshes into closed solids using voxelization techniques. UVs, sharp edges, and small/thin features will be lost. Increase Voxel Count to enhance accuracy."),
+		EToolMessageLevel::UserNotification);
 }
 
 
diff --git a/Engine/Plugins/Experimental/MeshModelingToolset/Source/MeshModelingToolsEditorOnly/Private/VoxelCSGMeshesTool.cpp b/Engine/Plugins/Experimental/MeshModelingToolset/Source/MeshModelingToolsEditorOnly/Private/VoxelCSGMeshesTool.cpp
--- a/Engine/Plugins/Experimental/MeshModelingToolset/Source/MeshModelingToolsEditorOnly/Private/VoxelCSGMeshesTool.cpp
+++ b/Engine/Plugins/Experimental/MeshModelingToolset/Source/MeshModelingToolsEditorOnly/Private/VoxelCSGMeshesTool.cpp
@@ -107,6 +107,10 @@ void UVoxelCSGMeshesTool::Setup()
 	);
 	
 	Preview->InvalidateResult();    // start compute
+
+	GetToolManager()->DisplayMessage(
+		LOCTEXT("OnStartTool", "This Tool computes a CSG Boolean of the input meshes using voxelization techniques. UVs, sharp edges, and small/thin features will be lost. Increase Voxel Count to enhance accuracy."),
+		EToolMessageLevel::UserNotification);
 }
 
 
