ModelingTools: make facegroup settings more consistent in Add Primitive tool, rename parameters
#rnx

diff --git a/Engine/Plugins/Experimental/GeometryProcessing/Source/GeometricObjects/Private/Generators/DiscMeshGenerator.cpp b/Engine/Plugins/Experimental/GeometryProcessing/Source/GeometricObjects/Private/Generators/DiscMeshGenerator.cpp
--- a/Engine/Plugins/Experimental/GeometryProcessing/Source/GeometricObjects/Private/Generators/DiscMeshGenerator.cpp
+++ b/Engine/Plugins/Experimental/GeometryProcessing/Source/GeometricObjects/Private/Generators/DiscMeshGenerator.cpp
@@ -66,14 +66,22 @@ FMeshShapeGenerator& FDiscMeshGenerator::Generate()
 	for (int AngleIdx = 0; AngleIdx+1 < AngleNV; AngleIdx++)
 	{
 		SetTriangleWithMatchedUVNormal(TriIdx, 0, AngleIdx + 2, AngleIdx + 1);
-		SetTrianglePolygon(TriIdx, PolyIdx++);
+		SetTrianglePolygon(TriIdx, PolyIdx);
+		if (!bSinglePolygroup)
+		{
+			PolyIdx++;
+		}
 		TriIdx++;
 	}
 	if (bFullDisc)
 	{
 		SetTriangleWithMatchedUVNormal(TriIdx, 0, 1, AngleNV);
-		SetTrianglePolygon(TriIdx, PolyIdx++);
+		SetTrianglePolygon(TriIdx, PolyIdx);
 		TriIdx++;
+		if (!bSinglePolygroup)
+		{
+			PolyIdx++;
+		}
 	}
 
 	for (int RadiusIdx = 0; RadiusIdx+1 < RadialNV; RadiusIdx++)
@@ -90,7 +98,10 @@ FMeshShapeGenerator& FDiscMeshGenerator::Generate()
 			SetTrianglePolygon(TriIdx, PolyIdx);
 			TriIdx++;
 
-			PolyIdx++;
+			if (!bSinglePolygroup)
+			{
+				PolyIdx++;
+			}
 		}
 		if (bFullDisc)
 		{
@@ -102,7 +113,10 @@ FMeshShapeGenerator& FDiscMeshGenerator::Generate()
 			SetTrianglePolygon(TriIdx, PolyIdx);
 			TriIdx++;
 
-			PolyIdx++;
+			if (!bSinglePolygroup)
+			{
+				PolyIdx++;
+			}
 		}
 	}
 
@@ -163,7 +177,10 @@ FMeshShapeGenerator& FPuncturedDiscMeshGenerator::Generate()
 			SetTrianglePolygon(TriIdx, PolyIdx);
 			TriIdx++;
 
-			PolyIdx++;
+			if (!bSinglePolygroup)
+			{
+				PolyIdx++;
+			}
 		}
 		if (bFullDisc)
 		{
@@ -175,7 +192,10 @@ FMeshShapeGenerator& FPuncturedDiscMeshGenerator::Generate()
 			SetTrianglePolygon(TriIdx, PolyIdx);
 			TriIdx++;
 
-			PolyIdx++;
+			if (!bSinglePolygroup)
+			{
+				PolyIdx++;
+			}
 		}
 	}
 
diff --git a/Engine/Plugins/Experimental/GeometryProcessing/Source/GeometricObjects/Private/Generators/RectangleMeshGenerator.cpp b/Engine/Plugins/Experimental/GeometryProcessing/Source/GeometricObjects/Private/Generators/RectangleMeshGenerator.cpp
--- a/Engine/Plugins/Experimental/GeometryProcessing/Source/GeometricObjects/Private/Generators/RectangleMeshGenerator.cpp
+++ b/Engine/Plugins/Experimental/GeometryProcessing/Source/GeometricObjects/Private/Generators/RectangleMeshGenerator.cpp
@@ -96,7 +96,10 @@ FMeshShapeGenerator& FRectangleMeshGenerator::Generate()
 			SetTriangleNormals(ti, i00, i10, i11);
 
 			ti++;
-			PolyIndex++;
+			if (!bSinglePolygroup)
+			{
+				PolyIndex++;
+			}
 		}
 	}
 
@@ -230,7 +233,10 @@ FMeshShapeGenerator& FRoundedRectangleMeshGenerator::Generate()
 			SetTriangleNormals(ti, i00, i10, i11);
 
 			ti++;
-			PolyIndex++;
+			if (!bSinglePolygroup)
+			{
+				PolyIndex++;
+			}
 		}
 	}
 
@@ -288,7 +294,10 @@ FMeshShapeGenerator& FRoundedRectangleMeshGenerator::Generate()
 				SetTriangleNormals(ti, VCenterIdx, LastUsedIdx, ActingSinIdx);
 				ti++;
 
-				PolyIndex++;
+				if (! bSinglePolygroup)
+				{
+					PolyIndex++;
+				}
 			}
 		}
 	}
diff --git a/Engine/Plugins/Experimental/GeometryProcessing/Source/GeometricObjects/Public/Generators/DiscMeshGenerator.h b/Engine/Plugins/Experimental/GeometryProcessing/Source/GeometricObjects/Public/Generators/DiscMeshGenerator.h
--- a/Engine/Plugins/Experimental/GeometryProcessing/Source/GeometricObjects/Public/Generators/DiscMeshGenerator.h
+++ b/Engine/Plugins/Experimental/GeometryProcessing/Source/GeometricObjects/Public/Generators/DiscMeshGenerator.h
@@ -28,6 +28,9 @@ public:
 	/** End of angle range spanned by disc, in degrees */
 	float EndAngle;
 
+	/** If true, output mesh has a single polygroup, otherwise each quad/tri gets a separate group */
+	bool bSinglePolygroup = false;
+
 	/** How to map 2D indices to 3D. Default is (0,1) = (x,y,0). */
 	FIndex2i IndicesMap;
 
diff --git a/Engine/Plugins/Experimental/GeometryProcessing/Source/GeometricObjects/Public/Generators/GridBoxMeshGenerator.h b/Engine/Plugins/Experimental/GeometryProcessing/Source/GeometricObjects/Public/Generators/GridBoxMeshGenerator.h
--- a/Engine/Plugins/Experimental/GeometryProcessing/Source/GeometricObjects/Public/Generators/GridBoxMeshGenerator.h
+++ b/Engine/Plugins/Experimental/GeometryProcessing/Source/GeometricObjects/Public/Generators/GridBoxMeshGenerator.h
@@ -21,6 +21,9 @@ public:
 	/** If true (default), UVs are scaled so that there is no stretching. If false, UVs are scaled to fill unit square */
 	bool bScaleUVByAspectRatio = true;
 
+	/** If true, each quad of box gets a separate polygroup */
+	bool bPolygroupPerQuad = false;
+
 public:
 
 	/** Generate the mesh */
@@ -162,6 +165,7 @@ public:
 		// create the face triangles and UVs+normals
 		int CurrentTriIdx = 0;
 		int CurrentUVIdx = 0;
+		int CurrentQuadIdx = 0;
 		for (int Dim = 0; Dim < 3; Dim++)
 		{
 			int FaceIdxBase = FaceDimOrder[Dim]*2;
@@ -233,7 +237,7 @@ public:
 							FaceUVStartInd + D1+Side + (D0+SideOpp) * N[D[1][Dim]],
 							FaceUVStartInd + D1+1 + (D0+1) * N[D[1][Dim]]
 						);
-						SetTrianglePolygon(CurrentTriIdx, MajorFaceInd);
+						SetTrianglePolygon(CurrentTriIdx, (bPolygroupPerQuad) ? CurrentQuadIdx : MajorFaceInd);
 						CurrentTriIdx++;
 
 						SetTriangle(CurrentTriIdx,
@@ -251,8 +255,9 @@ public:
 							FaceUVStartInd + D1+1 + (D0+1) * N[D[1][Dim]],
 							FaceUVStartInd + D1+SideOpp + (D0+Side) * N[D[1][Dim]]
 						);
-						SetTrianglePolygon(CurrentTriIdx, MajorFaceInd);
+						SetTrianglePolygon(CurrentTriIdx, (bPolygroupPerQuad) ? CurrentQuadIdx : MajorFaceInd);
 						CurrentTriIdx++;
+						CurrentQuadIdx++;
 					}
 				}
 			}
diff --git a/Engine/Plugins/Experimental/GeometryProcessing/Source/GeometricObjects/Public/Generators/RectangleMeshGenerator.h b/Engine/Plugins/Experimental/GeometryProcessing/Source/GeometricObjects/Public/Generators/RectangleMeshGenerator.h
--- a/Engine/Plugins/Experimental/GeometryProcessing/Source/GeometricObjects/Public/Generators/RectangleMeshGenerator.h
+++ b/Engine/Plugins/Experimental/GeometryProcessing/Source/GeometricObjects/Public/Generators/RectangleMeshGenerator.h
@@ -30,6 +30,9 @@ public:
 	/** If true (default), UVs are scaled so that there is no stretching. If false, UVs are scaled to fill unit square */
 	bool bScaleUVByAspectRatio = true;
 
+	/** If true, output mesh has a single polygroup, otherwise each quad gets a separate group */
+	bool bSinglePolygroup = false;
+
 	/** Specifies how 2D indices are mapped to 3D points. Default is (0,1) = (x,y,0). */ 
 	FIndex2i IndicesMap;
 
diff --git a/Engine/Plugins/Experimental/GeometryProcessing/Source/GeometricObjects/Public/Generators/SphereGenerator.h b/Engine/Plugins/Experimental/GeometryProcessing/Source/GeometricObjects/Public/Generators/SphereGenerator.h
--- a/Engine/Plugins/Experimental/GeometryProcessing/Source/GeometricObjects/Public/Generators/SphereGenerator.h
+++ b/Engine/Plugins/Experimental/GeometryProcessing/Source/GeometricObjects/Public/Generators/SphereGenerator.h
@@ -18,6 +18,9 @@ public:
 	int NumPhi = 16; // number of vertices along vertical extent from north pole to south pole
 	int NumTheta = 16; // number of vertices around circles
 
+	/** If true, each quad of sphere gets a separate polygroup */
+	bool bPolygroupPerQuad = false;
+
 public:
 
 	inline static FVector3d SphericalToCartesian(double r, double theta, double phi)
@@ -161,7 +164,10 @@ public:
 				SetTriangleUVs(TriIdx, UVCorners[2], UVCorners[0], UVCorners[3]);
 				SetTriangleNormals(TriIdx, Corners[2], Corners[0], Corners[3]);
 				TriIdx++;
-				PolyIdx++;
+				if (bPolygroupPerQuad)
+				{
+					PolyIdx++;
+				}
 			}
 		}
 
@@ -178,7 +184,10 @@ public:
 				SetTriangleUVs(TriIdx, UVCorners[0], UVCorners[2], UVCorners[1]);
 				SetTriangleNormals(TriIdx, Corners[0], Corners[2], Corners[1]);
 				TriIdx++;
-				PolyIdx++;
+				if (bPolygroupPerQuad)
+				{
+					PolyIdx++;
+				}
 			}
 		}
 
@@ -195,7 +204,10 @@ public:
 				SetTriangleUVs(TriIdx, UVCorners[0], UVCorners[2], UVCorners[1]);
 				SetTriangleNormals(TriIdx, Corners[0], Corners[2], Corners[1]);
 				TriIdx++;
-				PolyIdx++;
+				if (bPolygroupPerQuad)
+				{
+					PolyIdx++;
+				}
 			}
 		}
 
diff --git a/Engine/Plugins/Experimental/GeometryProcessing/Source/GeometricObjects/Public/Generators/SweepGenerator.h b/Engine/Plugins/Experimental/GeometryProcessing/Source/GeometricObjects/Public/Generators/SweepGenerator.h
--- a/Engine/Plugins/Experimental/GeometryProcessing/Source/GeometricObjects/Public/Generators/SweepGenerator.h
+++ b/Engine/Plugins/Experimental/GeometryProcessing/Source/GeometricObjects/Public/Generators/SweepGenerator.h
@@ -31,6 +31,9 @@ public:
 	{
 	}
 
+	/** If true, each quad gets a separate polygroup */
+	bool bPolygroupPerQuad = false;
+
 protected:
 	int32 CapVertStart[2], CapNormalStart[2], CapUVStart[2], CapTriangleStart[2], CapPolygonStart[2];
 
@@ -149,6 +152,7 @@ protected:
 
 		// fill in UVs and triangles along length
 		int MinValidCrossSections = bLoop ? 3 : 2;
+		int CurFaceGroupIndex = NumPolygons;
 		if (NumCrossSections >= MinValidCrossSections)
 		{
 			int32 UVSection = 0, UVSubIdx = 0;
@@ -240,8 +244,8 @@ protected:
 						int32 PIdx = XVerts * XIdx + VertSubIdx;
 						int32 NextXIdx = (XIdx + 1) % CrossSectionsMod;
 						int32 NextNXIdx = (NXIdx + 1) % NormalCrossSectionsMod;
-						SetTrianglePolygon(T0Idx, PIdx);
-						SetTrianglePolygon(T1Idx, PIdx);
+						SetTrianglePolygon(T0Idx, (bPolygroupPerQuad) ? PIdx : (CurFaceGroupIndex+XIdx) );
+						SetTrianglePolygon(T1Idx, (bPolygroupPerQuad) ? PIdx : (CurFaceGroupIndex+XIdx) );
 						SetTriangle(T0Idx,
 									XIdx * XVerts + VertSubIdx,
 									XIdx * XVerts + WrappedNextVertexSubIdx,
diff --git a/Engine/Plugins/Experimental/MeshModelingToolset/Source/MeshModelingTools/Private/AddPrimitiveTool.cpp b/Engine/Plugins/Experimental/MeshModelingToolset/Source/MeshModelingTools/Private/AddPrimitiveTool.cpp
--- a/Engine/Plugins/Experimental/MeshModelingToolset/Source/MeshModelingTools/Private/AddPrimitiveTool.cpp
+++ b/Engine/Plugins/Experimental/MeshModelingToolset/Source/MeshModelingTools/Private/AddPrimitiveTool.cpp
@@ -18,6 +18,7 @@
 #include "Generators/BoxSphereGenerator.h"
 #include "Generators/DiscMeshGenerator.h"
 #include "DynamicMeshAttributeSet.h"
+#include "FaceGroupUtil.h"
 
 #include "Drawing/MeshDebugDrawing.h"
 
@@ -55,6 +56,7 @@ UProceduralShapeToolProperties::UProceduralShapeToolProperties()
 	//EndAngle = 360;
 	Slices = 16;
 	Subdivisions = 0;
+	PolygroupMode = EMakeMeshPolygroupMode::PerFace;
 	PivotLocation = EMakeMeshPivotLocation::Base;
 	PlaceMode = EMakeMeshPlacementType::OnScene;
 	bAlignShapeToPlacementSurface = true;
@@ -69,9 +71,11 @@ void UProceduralShapeToolProperties::SaveProperties(UInteractiveTool* SaveFromTo
 	PropertyCache->FeatureRadius = this->FeatureRadius;
 	PropertyCache->Slices = this->Slices;
 	PropertyCache->Subdivisions = this->Subdivisions;
+	PropertyCache->PolygroupMode = this->PolygroupMode;
 	PropertyCache->PivotLocation = this->PivotLocation;
 	PropertyCache->PlaceMode = this->PlaceMode;
 	PropertyCache->bAlignShapeToPlacementSurface = this->bAlignShapeToPlacementSurface;
+	PropertyCache->bInstanceIfPossible = this->bInstanceIfPossible;
 }
 
 void UProceduralShapeToolProperties::RestoreProperties(UInteractiveTool* RestoreToTool)
@@ -83,9 +87,12 @@ void UProceduralShapeToolProperties::RestoreProperties(UInteractiveTool* Restore
 	this->FeatureRadius = PropertyCache->FeatureRadius;
 	this->Slices = PropertyCache->Slices;
 	this->Subdivisions = PropertyCache->Subdivisions;
+	this->PolygroupMode = PropertyCache->PolygroupMode;
 	this->PivotLocation = PropertyCache->PivotLocation;
 	this->PlaceMode = PropertyCache->PlaceMode;
 	this->bAlignShapeToPlacementSurface = PropertyCache->bAlignShapeToPlacementSurface;
+	this->bInstanceIfPossible = PropertyCache->bInstanceIfPossible;
+
 }
 
 namespace
@@ -100,10 +107,11 @@ namespace
 	  { TEXT("Height"),        EMakeMeshShapeType::Box | EMakeMeshShapeType::Cylinder | EMakeMeshShapeType::Cone | EMakeMeshShapeType::Arrow | EMakeMeshShapeType::Rectangle | EMakeMeshShapeType::RoundedRectangle },
 	  { TEXT("FeatureRadius"),		   EMakeMeshShapeType::Arrow | EMakeMeshShapeType::RoundedRectangle | EMakeMeshShapeType::PuncturedDisc | EMakeMeshShapeType::Torus },
 	  { TEXT("Rotation"),      EMakeMeshShapeType::All },
+	  { TEXT("PolygroupMode"), EMakeMeshShapeType::All },
 	  { TEXT("PlaceMode"),     EMakeMeshShapeType::All },
 	  { TEXT("PivotLocation"), EMakeMeshShapeType::All },
 	  { TEXT("bAlignShapeToPlacementSurface"), EMakeMeshShapeType::All },
-	  { TEXT("bInstanceLastCreatedAssetIfPossible"), EMakeMeshShapeType::All },
+	  { TEXT("bInstanceIfPossible"), EMakeMeshShapeType::All },
 	  { TEXT("Slices"),        EMakeMeshShapeType::Cylinder | EMakeMeshShapeType::Cone | EMakeMeshShapeType::Arrow | EMakeMeshShapeType::RoundedRectangle | EMakeMeshShapeType::Disc | EMakeMeshShapeType::PuncturedDisc | EMakeMeshShapeType::Sphere | EMakeMeshShapeType::Torus },
 	  { TEXT("Subdivisions"),  EMakeMeshShapeType::Box | EMakeMeshShapeType::Rectangle | EMakeMeshShapeType::RoundedRectangle | EMakeMeshShapeType::Disc | EMakeMeshShapeType::PuncturedDisc | EMakeMeshShapeType::Cylinder |
 							   EMakeMeshShapeType::Cone | EMakeMeshShapeType::Arrow | EMakeMeshShapeType::SphericalBox | EMakeMeshShapeType::Torus }
@@ -163,7 +171,7 @@ void UAddPrimitiveTool::Setup()
 	UpdatePreviewMesh();
 
 	GetToolManager()->DisplayMessage(
-		LOCTEXT("OnStartAddPrimitiveTool", "This Tool creates new Primitive objects. Position the Primitive by moving the mouse over the scene. Drop a new Asset or Instance by left-clicking (depending on Instance setting)."),
+		LOCTEXT("OnStartAddPrimitiveTool", "This Tool creates new Primitive mesh assets. Position the Primitive by moving the mouse over the scene. Drop a new Asset or Instance by left-clicking (depending on Asset settings)."),
 		EToolMessageLevel::UserNotification);
 }
 
@@ -322,6 +330,11 @@ void UAddPrimitiveTool::UpdatePreviewMesh()
 		break;
 	}
 
+	if (ShapeSettings->PolygroupMode == EMakeMeshPolygroupMode::Single)
+	{
+		FaceGroupUtil::SetGroupID(NewMesh, 0);
+	}
+
 	if (MaterialProperties->UVScale != 1.0 || MaterialProperties->bWorldSpaceUVScale)
 	{
 		FDynamicMeshEditor Editor(&NewMesh);
@@ -358,7 +371,7 @@ void UAddPrimitiveTool::OnClicked(const FInputDeviceRay& DeviceClickPos)
 	FString ShapeTypeName = MakeMeshShapeTypeEnum->GetNameStringByValue((int64)ShapeSettings->Shape);
 	UMaterialInterface* Material = PreviewMesh->GetMaterial();
 
-	if (ShapeSettings->bInstanceLastCreatedAssetIfPossible && IsEquivalentLastGeneratedAsset())
+	if (ShapeSettings->bInstanceIfPossible && IsEquivalentLastGeneratedAsset())
 	{
 		GetToolManager()->BeginUndoTransaction(LOCTEXT("AddPrimitiveToolTransactionName", "Add Primitive Mesh"));
 		FActorSpawnParameters SpawnParameters;
@@ -413,6 +426,10 @@ void UAddPrimitiveTool::GenerateBox(FDynamicMesh3* OutMesh)
 	BoxGen.Box = FOrientedBox3d(FVector3d::Zero(), 0.5*FVector3d(ShapeSettings->Width, ShapeSettings->Width, ShapeSettings->Height));
 	int EdgeNum = ShapeSettings->Subdivisions + 2;
 	BoxGen.EdgeVertices = FIndex3i(EdgeNum, EdgeNum, EdgeNum);
+	if (ShapeSettings->PolygroupMode == EMakeMeshPolygroupMode::PerQuad)
+	{
+		BoxGen.bPolygroupPerQuad = true;
+	}
 	BoxGen.Generate();
 	OutMesh->Copy(&BoxGen);
 }
@@ -424,6 +441,10 @@ void UAddPrimitiveTool::GenerateRectangle(FDynamicMesh3* OutMesh)
 	RectGen.Width = ShapeSettings->Width;
 	RectGen.Height = ShapeSettings->Height;
 	RectGen.WidthVertexCount = RectGen.HeightVertexCount = ShapeSettings->Subdivisions + 2;
+	if (ShapeSettings->PolygroupMode != EMakeMeshPolygroupMode::PerQuad)
+	{
+		RectGen.bSinglePolygroup = true;
+	}
 	RectGen.Generate();
 	OutMesh->Copy(&RectGen);
 }
@@ -434,6 +455,10 @@ void UAddPrimitiveTool::GenerateRoundedRectangle(FDynamicMesh3* OutMesh)
 	RectGen.Width = ShapeSettings->Width;
 	RectGen.Height = ShapeSettings->Height;
 	RectGen.WidthVertexCount = RectGen.HeightVertexCount = ShapeSettings->Subdivisions + 2;
+	if (ShapeSettings->PolygroupMode != EMakeMeshPolygroupMode::PerQuad)
+	{
+		RectGen.bSinglePolygroup = true;
+	}
 	RectGen.Radius = ShapeSettings->FeatureRadius;
 	RectGen.AngleSamples = ShapeSettings->Slices;
 	RectGen.Generate();
@@ -447,6 +472,10 @@ void UAddPrimitiveTool::GenerateDisc(FDynamicMesh3* OutMesh)
 	Gen.Radius = ShapeSettings->Width * 0.5f;
 	Gen.AngleSamples = ShapeSettings->Slices;
 	Gen.RadialSamples = ShapeSettings->Subdivisions;
+	if (ShapeSettings->PolygroupMode != EMakeMeshPolygroupMode::PerQuad)
+	{
+		Gen.bSinglePolygroup = true;
+	}
 	Gen.Generate();
 	OutMesh->Copy(&Gen);
 }
@@ -459,6 +488,10 @@ void UAddPrimitiveTool::GeneratePuncturedDisc(FDynamicMesh3* OutMesh)
 	Gen.HoleRadius = FMath::Min(ShapeSettings->FeatureRadius, Gen.Radius * .999f); // hole cannot be bigger than outer radius
 	Gen.AngleSamples = ShapeSettings->Slices;
 	Gen.RadialSamples = ShapeSettings->Subdivisions;
+	if (ShapeSettings->PolygroupMode != EMakeMeshPolygroupMode::PerQuad)
+	{
+		Gen.bSinglePolygroup = true;
+	}
 	Gen.Generate();
 	OutMesh->Copy(&Gen);
 }
@@ -475,6 +508,10 @@ void UAddPrimitiveTool::GenerateTorus(FDynamicMesh3* OutMesh)
 	}
 	Gen.bLoop = true;
 	Gen.bCapped = false;
+	if (ShapeSettings->PolygroupMode == EMakeMeshPolygroupMode::PerQuad)
+	{
+		Gen.bPolygroupPerQuad = true;
+	}
 	Gen.InitialFrame = FFrame3d(Gen.Path[0]);
 	Gen.Generate();
 	OutMesh->Copy(&Gen);
@@ -490,6 +527,10 @@ void UAddPrimitiveTool::GenerateCylinder(FDynamicMesh3* OutMesh)
 	CylGen.AngleSamples = ShapeSettings->Slices;
 	CylGen.LengthSamples = ShapeSettings->Subdivisions;
 	CylGen.bCapped = true;
+	if (ShapeSettings->PolygroupMode == EMakeMeshPolygroupMode::PerQuad)
+	{
+		CylGen.bPolygroupPerQuad = true;
+	}
 	CylGen.Generate();
 	OutMesh->Copy(&CylGen);
 }
@@ -505,6 +546,10 @@ void UAddPrimitiveTool::GenerateCone(FDynamicMesh3* OutMesh)
 	CylGen.AngleSamples = ShapeSettings->Slices;
 	CylGen.LengthSamples = ShapeSettings->Subdivisions;
 	CylGen.bCapped = true;
+	if (ShapeSettings->PolygroupMode == EMakeMeshPolygroupMode::PerQuad)
+	{
+		CylGen.bPolygroupPerQuad = true;
+	}
 	CylGen.Generate();
 	OutMesh->Copy(&CylGen);
 }
@@ -520,6 +565,10 @@ void UAddPrimitiveTool::GenerateArrow(FDynamicMesh3* OutMesh)
 	ArrowGen.HeadLength = ShapeSettings->Height * .75f;
 	ArrowGen.AngleSamples = ShapeSettings->Slices;
 	ArrowGen.bCapped = true;
+	if (ShapeSettings->PolygroupMode == EMakeMeshPolygroupMode::PerQuad)
+	{
+		ArrowGen.bPolygroupPerQuad = true;
+	}
 	ArrowGen.DistributeAdditionalLengthSamples(ShapeSettings->Subdivisions);
 	ArrowGen.Generate();
 	OutMesh->Copy(&ArrowGen);
@@ -532,6 +581,10 @@ void UAddPrimitiveTool::GenerateSphere(FDynamicMesh3* OutMesh)
 	SphereGen.Radius = ShapeSettings->Width * 0.5f;
 	SphereGen.NumTheta = ShapeSettings->Slices;
 	SphereGen.NumPhi = ShapeSettings->Slices;
+	if (ShapeSettings->PolygroupMode == EMakeMeshPolygroupMode::PerQuad)
+	{
+		SphereGen.bPolygroupPerQuad = true;
+	}
 	SphereGen.Generate();
 	OutMesh->Copy(&SphereGen);
 }
@@ -544,6 +597,10 @@ void UAddPrimitiveTool::GenerateSphericalBox(FDynamicMesh3* OutMesh)
 	SphereGen.Box = FOrientedBox3d(FVector3d::Zero(), 0.5*FVector3d(ShapeSettings->Width, ShapeSettings->Width, ShapeSettings->Width));
 	int EdgeNum = ShapeSettings->Subdivisions + 3;
 	SphereGen.EdgeVertices = FIndex3i(EdgeNum, EdgeNum, EdgeNum);
+	if (ShapeSettings->PolygroupMode == EMakeMeshPolygroupMode::PerQuad)
+	{
+		SphereGen.bPolygroupPerQuad = true;
+	}
 	SphereGen.Generate();
 	OutMesh->Copy(&SphereGen);
 }
diff --git a/Engine/Plugins/Experimental/MeshModelingToolset/Source/MeshModelingTools/Private/TransformMeshesTool.cpp b/Engine/Plugins/Experimental/MeshModelingToolset/Source/MeshModelingTools/Private/TransformMeshesTool.cpp
--- a/Engine/Plugins/Experimental/MeshModelingToolset/Source/MeshModelingTools/Private/TransformMeshesTool.cpp
+++ b/Engine/Plugins/Experimental/MeshModelingToolset/Source/MeshModelingTools/Private/TransformMeshesTool.cpp
@@ -82,7 +82,7 @@ void UTransformMeshesTool::Setup()
 	UpdateTransformMode(TransformProps->TransformMode);
 
 	GetToolManager()->DisplayMessage(
-		LOCTEXT("OnStartTransformMeshesTool", "This Tool allows you to transform multiple objects. [A] cycles through Transform modes. Reposition the Gizmo using Set Pivot Mode [S]. Toggle Snap Drag Mode [D] to move the object by click-draging. [W] and [E] cycle through SnapDrag Source and Rotation types."),
+		LOCTEXT("OnStartTransformMeshesTool", "Translate/Rotate/Scale the selected objects. [A] cycles through Transform modes. Reposition the Gizmo using Set Pivot Mode [S]. Toggle Snap Drag Mode [D] to move the object by click-draging. [W] and [E] cycle through SnapDrag Source and Rotation types."),
 		EToolMessageLevel::UserNotification);
 }
 
diff --git a/Engine/Plugins/Experimental/MeshModelingToolset/Source/MeshModelingTools/Public/AddPrimitiveTool.h b/Engine/Plugins/Experimental/MeshModelingToolset/Source/MeshModelingTools/Public/AddPrimitiveTool.h
--- a/Engine/Plugins/Experimental/MeshModelingToolset/Source/MeshModelingTools/Public/AddPrimitiveTool.h
+++ b/Engine/Plugins/Experimental/MeshModelingToolset/Source/MeshModelingTools/Public/AddPrimitiveTool.h
@@ -71,6 +71,20 @@ enum class EMakeMeshPivotLocation : uint8
 	Top
 };
 
+
+/** Polygroup mode for primitive */
+UENUM()
+enum class EMakeMeshPolygroupMode : uint8
+{
+	/** One polygroup for entire output mesh */
+	Single,
+	/** One polygroup per geometric face of primitive */
+	PerFace,
+	/** One polygroup per mesh quad/triangle */
+	PerQuad
+};
+
+
 UCLASS()
 class MESHMODELINGTOOLS_API UProceduralShapeToolProperties : public UInteractiveToolPropertySet
 {
@@ -85,8 +99,12 @@ public:
 #endif // WITH_EDITOR	
 	// End of UObject interface
 
+	/** If the shape settings haven't changed, create instances of the last created asset rather than creating a whole new asset.  If false, all created actors will have separate underlying mesh assets. */
+	UPROPERTY(EditAnywhere, Category = AssetSettings)
+	bool bInstanceIfPossible = true;
+
 	/** Type of shape to generate */
-	UPROPERTY(EditAnywhere, Category = ShapeSettings)
+	UPROPERTY(EditAnywhere, Category = ShapeSettings, meta = (DisplayName = "Primitive Type"))
 	EMakeMeshShapeType Shape;
 
 	/** Width of Shape */
@@ -98,28 +116,36 @@ public:
 	float Height;
 
 	/** Radius of additional circular features of the shape (not implicitly defined by the width of the shape) */
-	UPROPERTY(EditAnywhere, Category = ShapeSettings, meta = (DisplayName = "Feature Radius", UIMin = "1.0", UIMax = "1000.0", ClampMin = "0.0001", ClampMax = "1000000.0"))
+	UPROPERTY(EditAnywhere, Category = ShapeSettings, meta = (DisplayName = "Radius", UIMin = "1.0", UIMax = "1000.0", ClampMin = "0.0001", ClampMax = "1000000.0"))
 	float FeatureRadius;
 
-	/** Rotation around up axis */
-	UPROPERTY(EditAnywhere, Category = ShapeSettings, meta = (DisplayName = "Rotation", UIMin = "0.0", UIMax = "360.0"))
-	float Rotation;
+	/** Number of Slices */
+	UPROPERTY(EditAnywhere, Category = ShapeSettings, meta = (DisplayName = "Slices", UIMin = "3", UIMax = "128", ClampMin = "3", ClampMax = "999"))
+	int Slices;
 
-	/** Type of shape to generate */
+	/** Number of Subdivisions */
+	UPROPERTY(EditAnywhere, Category = ShapeSettings, meta = (DisplayName = "Subdivisions", UIMin = "0", UIMax = "100", ClampMin = "0", ClampMax = "4000"))
+	int Subdivisions;
+
+	/** How should Polygroups be assigned to triangles of Primitive */
 	UPROPERTY(EditAnywhere, Category = ShapeSettings)
+	EMakeMeshPolygroupMode PolygroupMode;
+
+	/** How to place Primitive in the Scene */
+	UPROPERTY(EditAnywhere, Category = Positioning, meta = (DisplayName = "Target Surface"))
 	EMakeMeshPlacementType PlaceMode;
 
-	/** Center shape at click point */
-	UPROPERTY(EditAnywhere, Category = ShapeSettings)
+	/** Location of Pivot within Primitive */
+	UPROPERTY(EditAnywhere, Category = Positioning)
     EMakeMeshPivotLocation PivotLocation;
 
-	/** Align shape to placement surface */
-	UPROPERTY(EditAnywhere, Category = ShapeSettings, meta = (EditCondition = "PlaceMode == EMakeMeshPlacementType::OnScene"))
-	bool bAlignShapeToPlacementSurface = true;
+	/** Rotation of Primitive around up axis */
+	UPROPERTY(EditAnywhere, Category = Positioning, meta = (DisplayName = "Rotation", UIMin = "0.0", UIMax = "360.0"))
+	float Rotation;
 
-	/** If the shape settings haven't changed, create instances of the last created asset rather than creating a whole new asset.  If false, all created actors will have separate underlying mesh assets. */
-	UPROPERTY(EditAnywhere, Category = ShapeSettings)
-	bool bInstanceLastCreatedAssetIfPossible = true;
+	/** Align shape to Placement Surface */
+	UPROPERTY(EditAnywhere, Category = Positioning, meta = (DisplayName = "Align to Normal", EditCondition = "PlaceMode == EMakeMeshPlacementType::OnScene"))
+	bool bAlignShapeToPlacementSurface = true;
 
 	///** Start Angle of Shape */
 	//UPROPERTY(EditAnywhere, Category = ShapeSettings, meta = (DisplayName = "Start Angle", UIMin = "0.0", UIMax = "360.0", ClampMin = "-10000", ClampMax = "10000.0"))
@@ -130,13 +156,6 @@ public:
 	//float EndAngle;
 
 
-	/** Number of Slices */
-	UPROPERTY(EditAnywhere, Category = ShapeSettings, meta = (DisplayName = "Slices", UIMin = "3", UIMax = "128", ClampMin = "3", ClampMax = "999"))
-	int Slices;
-
-	/** Subdivisions */
-	UPROPERTY(EditAnywhere, Category = ShapeSettings, meta = (DisplayName = "Subdivisions", UIMin = "0", UIMax = "100", ClampMin = "0", ClampMax = "4000"))
-	int Subdivisions;
 
 	virtual void SaveProperties(UInteractiveTool* SaveFromTool) override;
 	virtual void RestoreProperties(UInteractiveTool* RestoreToTool) override;
@@ -231,6 +250,7 @@ protected:
 			LastGenerated->ShapeSettings->Height == ShapeSettings->Height &&
 			LastGenerated->ShapeSettings->Width == ShapeSettings->Width &&
 			LastGenerated->ShapeSettings->Shape == ShapeSettings->Shape &&
+			LastGenerated->ShapeSettings->PolygroupMode == ShapeSettings->PolygroupMode &&
 			LastGenerated->MaterialProperties->UVScale == MaterialProperties->UVScale &&
 			LastGenerated->MaterialProperties->bWorldSpaceUVScale == MaterialProperties->bWorldSpaceUVScale
 			;
