MeshModelingTools: Simplify, Remesh, and DynaSculpt tools now save/restore their properties
#rnx

diff --git a/Engine/Plugins/Experimental/MeshModelingToolset/Source/MeshModelingTools/Private/DynamicMeshSculptTool.cpp b/Engine/Plugins/Experimental/MeshModelingToolset/Source/MeshModelingTools/Private/DynamicMeshSculptTool.cpp
--- a/Engine/Plugins/Experimental/MeshModelingToolset/Source/MeshModelingTools/Private/DynamicMeshSculptTool.cpp
+++ b/Engine/Plugins/Experimental/MeshModelingToolset/Source/MeshModelingTools/Private/DynamicMeshSculptTool.cpp
@@ -100,6 +100,28 @@ void UFixedPlaneBrushProperties::RestoreProperties(UInteractiveTool* RestoreToTo
 	this->Position = PropertyCache->Position;
 }
 
+void
+UBrushRemeshProperties::SaveRestoreProperties(UInteractiveTool* RestoreToTool, bool bSaving)
+{
+	UBrushRemeshProperties* PropertyCache = GetPropertyCache<UBrushRemeshProperties>();
+
+	// MeshConstraintProperties
+	SaveRestoreProperty(PropertyCache->bPreserveSharpEdges, this->bPreserveSharpEdges, bSaving);
+	SaveRestoreProperty(PropertyCache->MeshBoundaryConstraint, this->MeshBoundaryConstraint, bSaving);
+	SaveRestoreProperty(PropertyCache->GroupBoundaryConstraint, this->GroupBoundaryConstraint, bSaving);
+	SaveRestoreProperty(PropertyCache->MaterialBoundaryConstraint, this->MaterialBoundaryConstraint, bSaving);
+	SaveRestoreProperty(PropertyCache->bPreventNormalFlips, this->bPreventNormalFlips, bSaving);
+
+	//RemeshProperties
+	SaveRestoreProperty(PropertyCache->SmoothingStrength, this->SmoothingStrength, bSaving);
+	SaveRestoreProperty(PropertyCache->bFlips, this->bFlips, bSaving);
+	SaveRestoreProperty(PropertyCache->bSplits, this->bSplits, bSaving);
+	SaveRestoreProperty(PropertyCache->bCollapses, this->bCollapses, bSaving);
+
+	//BrushRemeshProperties
+	SaveRestoreProperty(PropertyCache->RelativeSize, this->RelativeSize, bSaving);
+}
+
 /*
  * Tool
  */
@@ -173,6 +195,8 @@ void UDynamicMeshSculptTool::Setup()
 	SculptProperties = NewObject<UBrushSculptProperties>(this, TEXT("Sculpting"));
 
 	RemeshProperties = NewObject<UBrushRemeshProperties>(this, TEXT("Remeshing"));
+	RemeshProperties->RestoreProperties(this);
+
 	InitialEdgeLength = EstimateIntialSafeTargetLength(*Mesh, 5000);
 
 	// hide input Component
@@ -284,6 +308,7 @@ void UDynamicMeshSculptTool::Shutdown(EToolShutdownType ShutdownType)
 	SculptProperties->SaveProperties(this);
 	ViewProperties->SaveProperties(this);
 	GizmoProperties->SaveProperties(this);
+	RemeshProperties->SaveProperties(this);
 }
 
 void UDynamicMeshSculptTool::OnDynamicMeshComponentChanged()
diff --git a/Engine/Plugins/Experimental/MeshModelingToolset/Source/MeshModelingTools/Private/RemeshMeshTool.cpp b/Engine/Plugins/Experimental/MeshModelingToolset/Source/MeshModelingTools/Private/RemeshMeshTool.cpp
--- a/Engine/Plugins/Experimental/MeshModelingToolset/Source/MeshModelingTools/Private/RemeshMeshTool.cpp
+++ b/Engine/Plugins/Experimental/MeshModelingToolset/Source/MeshModelingTools/Private/RemeshMeshTool.cpp
@@ -66,6 +66,36 @@ URemeshMeshToolProperties::URemeshMeshToolProperties()
 	bUseTargetEdgeLength = false;
 }
 
+void
+URemeshMeshToolProperties::SaveRestoreProperties(UInteractiveTool* RestoreToTool, bool bSaving)
+{
+	URemeshMeshToolProperties* PropertyCache = GetPropertyCache<URemeshMeshToolProperties>();
+
+	// MeshConstraintProperties
+	SaveRestoreProperty(PropertyCache->bPreserveSharpEdges, this->bPreserveSharpEdges, bSaving);
+	SaveRestoreProperty(PropertyCache->MeshBoundaryConstraint, this->MeshBoundaryConstraint, bSaving);
+	SaveRestoreProperty(PropertyCache->GroupBoundaryConstraint, this->GroupBoundaryConstraint, bSaving);
+	SaveRestoreProperty(PropertyCache->MaterialBoundaryConstraint, this->MaterialBoundaryConstraint, bSaving);
+	SaveRestoreProperty(PropertyCache->bPreventNormalFlips, this->bPreventNormalFlips, bSaving);
+
+	//RemeshProperties
+	SaveRestoreProperty(PropertyCache->SmoothingStrength, this->SmoothingStrength, bSaving);
+	SaveRestoreProperty(PropertyCache->bFlips, this->bFlips, bSaving);
+	SaveRestoreProperty(PropertyCache->bSplits, this->bSplits, bSaving);
+	SaveRestoreProperty(PropertyCache->bCollapses, this->bCollapses, bSaving);
+
+	//RemeshMeshToolProperties
+	SaveRestoreProperty(PropertyCache->TargetTriangleCount, this->TargetTriangleCount, bSaving);
+	SaveRestoreProperty(PropertyCache->SmoothingType, this->SmoothingType, bSaving);
+	SaveRestoreProperty(PropertyCache->RemeshIterations, this->RemeshIterations, bSaving);
+	SaveRestoreProperty(PropertyCache->bDiscardAttributes, this->bDiscardAttributes, bSaving);
+	SaveRestoreProperty(PropertyCache->bShowWireframe, this->bShowWireframe, bSaving);
+	SaveRestoreProperty(PropertyCache->bShowGroupColors, this->bShowGroupColors, bSaving);
+	SaveRestoreProperty(PropertyCache->bUseTargetEdgeLength, this->bUseTargetEdgeLength, bSaving);
+	SaveRestoreProperty(PropertyCache->TargetEdgeLength, this->TargetEdgeLength, bSaving);
+	SaveRestoreProperty(PropertyCache->bReproject, this->bReproject, bSaving);
+}
+
 URemeshMeshTool::URemeshMeshTool()
 {
 }
@@ -85,6 +115,7 @@ void URemeshMeshTool::Setup()
 	UInteractiveTool::Setup();
 
 	BasicProperties = NewObject<URemeshMeshToolProperties>(this);
+	BasicProperties->RestoreProperties(this);
 	MeshStatisticsProperties = NewObject<UMeshStatisticsProperties>(this);
 
 	// hide component and create + show preview
@@ -142,6 +173,7 @@ void URemeshMeshTool::Setup()
 
 void URemeshMeshTool::Shutdown(EToolShutdownType ShutdownType)
 {
+	BasicProperties->SaveProperties(this);
 	ComponentTarget->SetOwnerVisibility(true);
 	FDynamicMeshOpResult Result = Preview->Shutdown();
 	if (ShutdownType == EToolShutdownType::Accept)
diff --git a/Engine/Plugins/Experimental/MeshModelingToolset/Source/MeshModelingTools/Public/DynamicMeshSculptTool.h b/Engine/Plugins/Experimental/MeshModelingToolset/Source/MeshModelingTools/Public/DynamicMeshSculptTool.h
--- a/Engine/Plugins/Experimental/MeshModelingToolset/Source/MeshModelingTools/Public/DynamicMeshSculptTool.h
+++ b/Engine/Plugins/Experimental/MeshModelingToolset/Source/MeshModelingTools/Public/DynamicMeshSculptTool.h
@@ -179,6 +179,8 @@ class MESHMODELINGTOOLS_API UBrushRemeshProperties : public URemeshProperties
 	GENERATED_BODY()
 
 public:
+	void SaveRestoreProperties(UInteractiveTool* RestoreToTool, bool bSaving) override;
+
 	/** Target Relative Triangle Sizefor Dynamic Meshing */
 	UPROPERTY(EditAnywhere, Category = Remeshing, meta = (UIMin = "0.5", UIMax = "2.0", ClampMin = "0.1", ClampMax = "100.0"))
 	float RelativeSize = 1.0f;
diff --git a/Engine/Plugins/Experimental/MeshModelingToolset/Source/MeshModelingTools/Public/RemeshMeshTool.h b/Engine/Plugins/Experimental/MeshModelingToolset/Source/MeshModelingTools/Public/RemeshMeshTool.h
--- a/Engine/Plugins/Experimental/MeshModelingToolset/Source/MeshModelingTools/Public/RemeshMeshTool.h
+++ b/Engine/Plugins/Experimental/MeshModelingToolset/Source/MeshModelingTools/Public/RemeshMeshTool.h
@@ -41,6 +41,8 @@ class MESHMODELINGTOOLS_API URemeshMeshToolProperties : public URemeshProperties
 public:
 	URemeshMeshToolProperties();
 
+	void SaveRestoreProperties(UInteractiveTool* RestoreToTool, bool bSaving) override;
+
 	/** Target triangle count */
 	UPROPERTY(EditAnywhere, Category = Remeshing, meta = (EditCondition = "bUseTargetEdgeLength == false"))
 	int TargetTriangleCount;
diff --git a/Engine/Plugins/Experimental/MeshModelingToolset/Source/MeshModelingToolsEditorOnly/Private/SimplifyMeshTool.cpp b/Engine/Plugins/Experimental/MeshModelingToolset/Source/MeshModelingToolsEditorOnly/Private/SimplifyMeshTool.cpp
--- a/Engine/Plugins/Experimental/MeshModelingToolset/Source/MeshModelingToolsEditorOnly/Private/SimplifyMeshTool.cpp
+++ b/Engine/Plugins/Experimental/MeshModelingToolset/Source/MeshModelingToolsEditorOnly/Private/SimplifyMeshTool.cpp
@@ -75,6 +75,30 @@ USimplifyMeshToolProperties::USimplifyMeshToolProperties()
 	MaterialBoundaryConstraint = EMaterialBoundaryConstraint::Ignore;
 }
 
+void
+USimplifyMeshToolProperties::SaveRestoreProperties(UInteractiveTool* RestoreToTool, bool bSaving)
+{
+	USimplifyMeshToolProperties* PropertyCache = GetPropertyCache<USimplifyMeshToolProperties>();
+
+	// MeshConstraintProperties
+	SaveRestoreProperty(PropertyCache->bPreserveSharpEdges, this->bPreserveSharpEdges, bSaving);
+	SaveRestoreProperty(PropertyCache->MeshBoundaryConstraint, this->MeshBoundaryConstraint, bSaving);
+	SaveRestoreProperty(PropertyCache->GroupBoundaryConstraint, this->GroupBoundaryConstraint, bSaving);
+	SaveRestoreProperty(PropertyCache->MaterialBoundaryConstraint, this->MaterialBoundaryConstraint, bSaving);
+	SaveRestoreProperty(PropertyCache->bPreventNormalFlips, this->bPreventNormalFlips, bSaving);
+
+	// SimplifyMeshToolProperties
+	SaveRestoreProperty(PropertyCache->TargetMode, this->TargetMode, bSaving);
+	SaveRestoreProperty(PropertyCache->SimplifierType, this->SimplifierType, bSaving);
+	SaveRestoreProperty(PropertyCache->TargetPercentage, this->TargetPercentage, bSaving);
+	SaveRestoreProperty(PropertyCache->TargetEdgeLength, this->TargetEdgeLength, bSaving);
+	SaveRestoreProperty(PropertyCache->TargetCount, this->TargetCount, bSaving);
+	SaveRestoreProperty(PropertyCache->bDiscardAttributes, this->bDiscardAttributes, bSaving);
+	SaveRestoreProperty(PropertyCache->bShowWireframe, this->bShowWireframe, bSaving);
+	SaveRestoreProperty(PropertyCache->bShowGroupColors, this->bShowGroupColors, bSaving);
+	SaveRestoreProperty(PropertyCache->bReproject, this->bReproject, bSaving);
+}
+
 void USimplifyMeshTool::SetWorld(UWorld* World)
 {
 	this->TargetWorld = World;
@@ -132,6 +156,7 @@ void USimplifyMeshTool::Setup()
 
 	// initialize our properties
 	SimplifyProperties = NewObject<USimplifyMeshToolProperties>(this);
+	SimplifyProperties->RestoreProperties(this);
 	AddToolPropertySource(SimplifyProperties);
 
 	ShowGroupsWatcher.Initialize(
@@ -156,12 +181,13 @@ void USimplifyMeshTool::Setup()
 
 void USimplifyMeshTool::Shutdown(EToolShutdownType ShutdownType)
 {
-		ComponentTarget->SetOwnerVisibility(true);
+	SimplifyProperties->SaveProperties(this);
+	ComponentTarget->SetOwnerVisibility(true);
 	FDynamicMeshOpResult Result = Preview->Shutdown();
-		if (ShutdownType == EToolShutdownType::Accept)
-		{
+	if (ShutdownType == EToolShutdownType::Accept)
+	{
 		GenerateAsset(Result);
-		}
+	}
 }
 
 
diff --git a/Engine/Plugins/Experimental/MeshModelingToolset/Source/MeshModelingToolsEditorOnly/Public/SimplifyMeshTool.h b/Engine/Plugins/Experimental/MeshModelingToolset/Source/MeshModelingToolsEditorOnly/Public/SimplifyMeshTool.h
--- a/Engine/Plugins/Experimental/MeshModelingToolset/Source/MeshModelingToolsEditorOnly/Public/SimplifyMeshTool.h
+++ b/Engine/Plugins/Experimental/MeshModelingToolset/Source/MeshModelingToolsEditorOnly/Public/SimplifyMeshTool.h
@@ -46,6 +46,8 @@ class MESHMODELINGTOOLSEDITORONLY_API USimplifyMeshToolProperties : public UMesh
 public:
 	USimplifyMeshToolProperties();
 
+	void SaveRestoreProperties(UInteractiveTool* RestoreToTool, bool bSaving) override;
+
 	/** Simplification Target Type  */
 	UPROPERTY(EditAnywhere, Category = Options)
 	ESimplifyTargetType TargetMode;
