give localization strings unique names
#rnx

diff --git a/Engine/Plugins/Experimental/MeshModelingToolset/Source/MeshModelingTools/Private/CombineMeshesTool.cpp b/Engine/Plugins/Experimental/MeshModelingToolset/Source/MeshModelingTools/Private/CombineMeshesTool.cpp
--- a/Engine/Plugins/Experimental/MeshModelingToolset/Source/MeshModelingTools/Private/CombineMeshesTool.cpp
+++ b/Engine/Plugins/Experimental/MeshModelingToolset/Source/MeshModelingTools/Private/CombineMeshesTool.cpp
@@ -100,13 +100,13 @@ void UCombineMeshesTool::Setup()
 	if (bDuplicateMode)
 	{
 		GetToolManager()->DisplayMessage(
-			LOCTEXT("OnStartTool", "This Tool duplicates input Asset into a new Asset, and optionally replaces the input Actor with a new Actor containing the new Asset."),
+			LOCTEXT("OnStartToolDuplicate", "This Tool duplicates input Asset into a new Asset, and optionally replaces the input Actor with a new Actor containing the new Asset."),
 			EToolMessageLevel::UserNotification);
 	}
 	else
 	{
 		GetToolManager()->DisplayMessage(
-			LOCTEXT("OnStartTool", "This Tool appends the meshes from the input Assets into a new Asset, and optionally replaces the source Actors with a new Actor containing the new Asset."),
+			LOCTEXT("OnStartToolCombine", "This Tool appends the meshes from the input Assets into a new Asset, and optionally replaces the source Actors with a new Actor containing the new Asset."),
 			EToolMessageLevel::UserNotification);
 	}
 }
