ModelingTools: resolve many sculpting issues & limitations
- fix broken sculpt material that went to white in many cases, replace with diffuse shading and rim light in Emissive channel / unlit texture
- add image-based sculpt materials provided by techart, expose custom variant that allows user-defined shader image
- make flat shading optional in view settings, for relevant shaders
- reorganize properties and update tooltips to be more informative
- fix broken Falloff support and expose in property set, display in brush highlight. Make this separately exposable in BaseBrush property set since we don't need Strength.
- add viewpoint-based Sculpt and Plane brush modes
- fix brush preview location when in hover, was not respecting target surface setting
#rnx

diff --git a/Engine/Plugins/Experimental/MeshModelingToolset/Source/MeshModelingTools/Private/DynamicMeshSculptTool.cpp b/Engine/Plugins/Experimental/MeshModelingToolset/Source/MeshModelingTools/Private/DynamicMeshSculptTool.cpp
--- a/Engine/Plugins/Experimental/MeshModelingToolset/Source/MeshModelingTools/Private/DynamicMeshSculptTool.cpp
+++ b/Engine/Plugins/Experimental/MeshModelingToolset/Source/MeshModelingTools/Private/DynamicMeshSculptTool.cpp
@@ -30,6 +30,8 @@
 #include "BaseGizmos/GizmoComponents.h"
 #include "BaseGizmos/TransformGizmo.h"
 #include "UObject/ObjectMacros.h"
+#include "Materials/Material.h"
+#include "Materials/MaterialInstanceDynamic.h"
 
 #define LOCTEXT_NAMESPACE "UDynamicMeshSculptTool"
 
@@ -44,64 +46,74 @@ UMeshSurfacePointTool* UDynamicMeshSculptToolBuilder::CreateNewTool(const FToolB
 	return SculptTool;
 }
 
+USculptBrushProperties::USculptBrushProperties() : UBrushBaseProperties()
+{
+	this->BrushFalloffAmount = 0.5f;
+}
+
+
+
+void USculptBrushProperties::SaveRestoreProperties(UInteractiveTool* Tool, bool bSaving)
+{
+	// do not restore base properties, they have their own CDO
+	//UBrushBaseProperties::SaveRestoreProperties(Tool, bSaving);
+
+	USculptBrushProperties* PropertyCache = GetPropertyCache<USculptBrushProperties>();
+
+	SaveRestoreProperty(PropertyCache->BrushSize, this->BrushSize, bSaving);
+	SaveRestoreProperty(PropertyCache->bSpecifyRadius, this->bSpecifyRadius, bSaving);
+	SaveRestoreProperty(PropertyCache->BrushRadius, this->BrushRadius, bSaving);
+	SaveRestoreProperty(PropertyCache->BrushFalloffAmount, this->BrushFalloffAmount, bSaving);
+	SaveRestoreProperty(PropertyCache->BrushStrength, this->BrushStrength, bSaving);
+
+	SaveRestoreProperty(PropertyCache->Depth, this->Depth, bSaving);
+	SaveRestoreProperty(PropertyCache->bHitBackFaces, this->bHitBackFaces, bSaving);
+}
+
 /*
  * BrushSculptProperties
  */
-void UBrushSculptProperties::SaveProperties(UInteractiveTool* SaveFromTool)
+void UBrushSculptProperties::SaveRestoreProperties(UInteractiveTool* Tool, bool bSaving)
 {
 	UBrushSculptProperties* PropertyCache = GetPropertyCache<UBrushSculptProperties>();
-	PropertyCache->PrimaryBrushType = this->PrimaryBrushType;
-	PropertyCache->PrimaryBrushSpeed = this->PrimaryBrushSpeed;
-	PropertyCache->SmoothBrushSpeed = this->SmoothBrushSpeed;
-	PropertyCache->bPreserveUVFlow = this->bPreserveUVFlow;
-	PropertyCache->bSmoothBrushRemeshes = this->bSmoothBrushRemeshes;
-	PropertyCache->BrushDepth = this->BrushDepth;
-	PropertyCache->bFreezeTarget = this->bFreezeTarget;
-	PropertyCache->bHitBackFaces = this->bHitBackFaces;
+	SaveRestoreProperty(PropertyCache->PrimaryBrushType, this->PrimaryBrushType, bSaving);
+	SaveRestoreProperty(PropertyCache->PrimaryBrushSpeed, this->PrimaryBrushSpeed, bSaving);
+	SaveRestoreProperty(PropertyCache->SmoothBrushSpeed, this->SmoothBrushSpeed, bSaving);
+	SaveRestoreProperty(PropertyCache->bPreserveUVFlow, this->bPreserveUVFlow, bSaving);
+	SaveRestoreProperty(PropertyCache->bSmoothBrushRemeshes, this->bSmoothBrushRemeshes, bSaving);
+	SaveRestoreProperty(PropertyCache->bFreezeTarget, this->bFreezeTarget, bSaving);
 }
 
-void UBrushSculptProperties::RestoreProperties(UInteractiveTool* RestoreToTool)
+
+
+void UPlaneBrushProperties::SaveRestoreProperties(UInteractiveTool* Tool, bool bSaving)
 {
-	UBrushSculptProperties* PropertyCache = GetPropertyCache<UBrushSculptProperties>();
-	this->PrimaryBrushType = PropertyCache->PrimaryBrushType;
-	this->PrimaryBrushSpeed = PropertyCache->PrimaryBrushSpeed;
-	this->SmoothBrushSpeed = PropertyCache->SmoothBrushSpeed;
-	this->bPreserveUVFlow = PropertyCache->bPreserveUVFlow;
-	this->bSmoothBrushRemeshes = PropertyCache->bSmoothBrushRemeshes;
-	this->BrushDepth = PropertyCache->BrushDepth;
-	this->bFreezeTarget = PropertyCache->bFreezeTarget;
-	this->bHitBackFaces = PropertyCache->bHitBackFaces;
+	UPlaneBrushProperties* PropertyCache = GetPropertyCache<UPlaneBrushProperties>();
+	SaveRestoreProperty(PropertyCache->WhichSide, this->WhichSide, bSaving);
 }
 
-/*
- * FixedPlaneBrushProperties
- */
+
+
 UFixedPlaneBrushProperties::UFixedPlaneBrushProperties()
 {
 	bPropertySetEnabled = true;
 	bSnapToGrid = true;
 	bShowGizmo = true;
 	Position = FVector::ZeroVector;
+	Rotation = FQuat::Identity;
 }
 
-void UFixedPlaneBrushProperties::SaveProperties(UInteractiveTool* SaveFromTool)
+void UFixedPlaneBrushProperties::SaveRestoreProperties(UInteractiveTool* Tool, bool bSaving)
 {
 	UFixedPlaneBrushProperties* PropertyCache = GetPropertyCache<UFixedPlaneBrushProperties>();
-	PropertyCache->bShowGizmo = this->bShowGizmo;
-	PropertyCache->bSnapToGrid = this->bSnapToGrid;
-	PropertyCache->Position = this->Position;
+	SaveRestoreProperty(PropertyCache->bShowGizmo, this->bShowGizmo, bSaving);
+	SaveRestoreProperty(PropertyCache->bSnapToGrid, this->bSnapToGrid, bSaving);
+	SaveRestoreProperty(PropertyCache->Position, this->Position, bSaving);
+	SaveRestoreProperty(PropertyCache->Rotation, this->Rotation, bSaving);
 }
 
-void UFixedPlaneBrushProperties::RestoreProperties(UInteractiveTool* RestoreToTool)
-{
-	UFixedPlaneBrushProperties* PropertyCache = GetPropertyCache<UFixedPlaneBrushProperties>();
-	this->bShowGizmo = PropertyCache->bShowGizmo;
-	this->bSnapToGrid = PropertyCache->bSnapToGrid;
-	this->Position = PropertyCache->Position;
-}
 
-void
-UBrushRemeshProperties::SaveRestoreProperties(UInteractiveTool* RestoreToTool, bool bSaving)
+void UBrushRemeshProperties::SaveRestoreProperties(UInteractiveTool* RestoreToTool, bool bSaving)
 {
 	UBrushRemeshProperties* PropertyCache = GetPropertyCache<UBrushRemeshProperties>();
 
@@ -188,13 +200,14 @@ void UDynamicMeshSculptTool::Setup()
 	// initialize brush radius range interval, brush properties
 	double MaxDimension = DynamicMeshComponent->GetMesh()->GetCachedBounds().MaxDim();
 	BrushRelativeSizeRange = FInterval1d(MaxDimension*0.01, MaxDimension);
-	BrushProperties = NewObject<UBrushBaseProperties>(this, TEXT("Brush"));
+	BrushProperties = NewObject<USculptBrushProperties>(this);
+	BrushProperties->bShowStrength = false;
 	CalculateBrushRadius();
 
 	// initialize other properties
-	SculptProperties = NewObject<UBrushSculptProperties>(this, TEXT("Sculpting"));
+	SculptProperties = NewObject<UBrushSculptProperties>(this);
 
-	RemeshProperties = NewObject<UBrushRemeshProperties>(this, TEXT("Remeshing"));
+	RemeshProperties = NewObject<UBrushRemeshProperties>(this);
 	RemeshProperties->RestoreProperties(this);
 
 	InitialEdgeLength = EstimateIntialSafeTargetLength(*Mesh, 5000);
@@ -214,11 +227,20 @@ void UDynamicMeshSculptTool::Setup()
 	BrushIndicator = GetToolManager()->GetPairedGizmoManager()->CreateGizmo<UBrushStampIndicator>(BrushIndicatorGizmoType, FString(), this);
 	BrushIndicatorMesh = MakeDefaultSphereMesh(this, TargetWorld);
 	BrushIndicator->AttachedComponent = BrushIndicatorMesh->GetRootComponent();
-	BrushIndicator->bDrawIndicatorLines = false;
+	BrushIndicator->LineThickness = 1.0;
+	BrushIndicator->bDrawIndicatorLines = true;
+	BrushIndicator->bDrawRadiusCircle = false;
+	BrushIndicator->bDrawFalloffCircle = true;
+	BrushIndicator->LineColor = FLinearColor(0.9f, 0.4f, 0.4f);
 
 	// initialize our properties
-	AddToolPropertySource(SculptProperties);
 	AddToolPropertySource(BrushProperties);
+	AddToolPropertySource(SculptProperties);
+
+	PlaneBrushProperties = NewObject<UPlaneBrushProperties>(this);
+	PlaneBrushProperties->RestoreProperties(this);
+	AddToolPropertySource(PlaneBrushProperties);
+
 	if (this->bEnableRemeshing)
 	{
 		SculptProperties->bIsRemeshingEnabled = true;
@@ -233,25 +255,46 @@ void UDynamicMeshSculptTool::Setup()
 	GizmoProperties->RestoreProperties(this);
 	AddToolPropertySource(GizmoProperties);
 
+	// disable tool-specific properties
+	SetToolPropertySourceEnabled(PlaneBrushProperties, false);
+	SetToolPropertySourceEnabled(GizmoProperties, false);
+
 	ViewProperties = NewObject<UMeshEditingViewProperties>();
 	ViewProperties->RestoreProperties(this);
 	AddToolPropertySource(ViewProperties);
+
+	// register watchers
 	ShowWireframeWatcher.Initialize(
 		[this]() { return ViewProperties->bShowWireframe; },
-		[this](bool bNewValue) { DynamicMeshComponent->bExplicitShowWireframe = bNewValue; }, false);
+		[this](bool bNewValue) { DynamicMeshComponent->bExplicitShowWireframe = bNewValue; }, ViewProperties->bShowWireframe);
 	MaterialModeWatcher.Initialize(
 		[this]() { return ViewProperties->MaterialMode; },
 		[this](EMeshEditingMaterialModes NewMode) { UpdateMaterialMode(NewMode); }, EMeshEditingMaterialModes::ExistingMaterial);
+	FlatShadingWatcher.Initialize(
+		[this]() { return ViewProperties->bFlatShading; },
+		[this](bool bNewValue) { UpdateFlatShadingSetting(bNewValue); }, ViewProperties->bFlatShading);
+	ColorWatcher.Initialize(
+		[this]() { return ViewProperties->Color; },
+		[this](FLinearColor NewColor) { UpdateColorSetting(NewColor); }, ViewProperties->Color);
+	ImageWatcher.Initialize(
+		[this]() { return ViewProperties->Image; },
+		[this](UTexture2D* NewImage) { UpdateImageSetting(NewImage); }, ViewProperties->Image);
+	BrushTypeWatcher.Initialize(
+		[this]() { return SculptProperties->PrimaryBrushType; },
+		[this](EDynamicMeshSculptBrushType NewBrushType) { UpdateBrushType(NewBrushType); }, SculptProperties->PrimaryBrushType);
+	GizmoPositionWatcher.Initialize(
+		[this]() { return GizmoProperties->Position; },
+		[this](FVector NewPosition) { UpdateGizmoFromProperties(); }, GizmoProperties->Position);
+	GizmoRotationWatcher.Initialize(
+		[this]() { return GizmoProperties->Rotation; },
+		[this](FQuat NewRotation) { UpdateGizmoFromProperties(); }, GizmoProperties->Rotation);
+
 
 	// create proxy for plane gizmo, but not gizmo itself, as it only appears in FixedPlane brush mode
 	// listen for changes to the proxy and update the plane when that happens
 	PlaneTransformProxy = NewObject<UTransformProxy>(this);
 	PlaneTransformProxy->OnTransformChanged.AddUObject(this, &UDynamicMeshSculptTool::PlaneTransformChanged);
 
-	GetToolManager()->DisplayMessage(
-		LOCTEXT("OnStartSculptTool", "Hold Shift to Smooth, Ctrl to Invert (where applicable). Shift+Q/A keys cycle through Brush Types. Shift+S/D change Size (Ctrl+Shift to small-step), Shift+W/E change Speed."),
-		EToolMessageLevel::UserNotification);
-
 	if (bEnableRemeshing)
 	{
 		PrecomputeRemeshInfo();
@@ -268,6 +311,8 @@ void UDynamicMeshSculptTool::Setup()
 				EToolMessageLevel::UserWarning);
 		}
 	}
+
+	UpdateBrushType(SculptProperties->PrimaryBrushType);
 }
 
 void UDynamicMeshSculptTool::Shutdown(EToolShutdownType ShutdownType)
@@ -308,6 +353,7 @@ void UDynamicMeshSculptTool::Shutdown(EToolShutdownType ShutdownType)
 	SculptProperties->SaveProperties(this);
 	ViewProperties->SaveProperties(this);
 	GizmoProperties->SaveProperties(this);
+	PlaneBrushProperties->SaveProperties(this);
 	RemeshProperties->SaveProperties(this);
 }
 
@@ -354,7 +400,7 @@ void UDynamicMeshSculptTool::OnBeginDrag(const FRay& Ray)
 	FHitResult OutHit;
 	if (HitTest(Ray, OutHit))
 	{
-		BrushStartCenterWorld = Ray.PointAt(OutHit.Distance) + SculptProperties->BrushDepth*CurrentBrushRadius*Ray.Direction;
+		BrushStartCenterWorld = Ray.PointAt(OutHit.Distance) + BrushProperties->Depth*CurrentBrushRadius*Ray.Direction;
 
 		bInDrag = true;
 
@@ -372,7 +418,12 @@ void UDynamicMeshSculptTool::OnBeginDrag(const FRay& Ray)
 
 		if (SculptProperties->PrimaryBrushType == EDynamicMeshSculptBrushType::Plane)
 		{
-			ActiveFixedBrushPlane = ComputeROIBrushPlane(LastBrushPosLocal, false);
+			ActiveFixedBrushPlane = ComputeROIBrushPlane(LastBrushPosLocal, false, false);
+		}
+		else if (SculptProperties->PrimaryBrushType == EDynamicMeshSculptBrushType::PlaneViewAligned)
+		{
+			AlignBrushToView();
+			ActiveFixedBrushPlane = ComputeROIBrushPlane(LastBrushPosLocal, false, true);
 		}
 
 		// apply initial stamp
@@ -476,7 +527,10 @@ void UDynamicMeshSculptTool::ApplyStamp(const FRay& WorldRay)
 	switch (SculptProperties->PrimaryBrushType)
 	{
 		case EDynamicMeshSculptBrushType::Offset:
-			ApplyOffsetBrush(WorldRay);
+			ApplyOffsetBrush(WorldRay, false);
+			break;
+		case EDynamicMeshSculptBrushType::SculptView:
+			ApplyOffsetBrush(WorldRay, true);
 			break;
 		case EDynamicMeshSculptBrushType::SculptMax:
 			ApplySculptMaxBrush(WorldRay);
@@ -499,6 +553,9 @@ void UDynamicMeshSculptTool::ApplyStamp(const FRay& WorldRay)
 		case EDynamicMeshSculptBrushType::Plane:
 			ApplyPlaneBrush(WorldRay);
 			break;
+		case EDynamicMeshSculptBrushType::PlaneViewAligned:
+			ApplyPlaneBrush(WorldRay);
+			break;
 		case EDynamicMeshSculptBrushType::FixedPlane:
 			ApplyFixedPlaneBrush(WorldRay);
 			break;
@@ -509,12 +566,13 @@ void UDynamicMeshSculptTool::ApplyStamp(const FRay& WorldRay)
 
 double UDynamicMeshSculptTool::CalculateBrushFalloff(double Distance)
 {
+	double f = FMathd::Clamp(1.0 - BrushProperties->BrushFalloffAmount, 0.0, 1.0);
 	double d = Distance / CurrentBrushRadius;
 	double w = 1;
-	if (d > 0.5)
+	if (d > f)
 	{
-		d = VectorUtil::Clamp((d - 0.5) / 0.5, 0.0, 1.0);
-		w = (1 - d * d);
+		d = FMathd::Clamp((d - f) / (1.0 - f), 0.0, 1.0);
+		w = (1.0 - d * d);
 		w = w * w * w;
 	}
 	return w;
@@ -608,15 +666,19 @@ void UDynamicMeshSculptTool::ApplyMoveBrush(const FRay& WorldRay)
 	LastBrushPosLocal = NewBrushPosLocal;
 }
 
-void UDynamicMeshSculptTool::ApplyOffsetBrush(const FRay& WorldRay)
+void UDynamicMeshSculptTool::ApplyOffsetBrush(const FRay& WorldRay, bool bUseViewDirection)
 {
-	bool bHit = UpdateBrushPositionOnTargetMesh(WorldRay);
-	if (bHit == false)
+	if ( UpdateBrushPositionOnTargetMesh(WorldRay) == false )
 	{
 		return;
 	}
+	if (bUseViewDirection)
+	{
+		AlignBrushToView();
+	}
 
 	FVector3d NewBrushPosLocal = CurTargetTransform.InverseTransformPosition(LastBrushPosWorld);
+	FVector3d LocalNormal = CurTargetTransform.InverseTransformNormal(LastBrushPosNormalWorld);
 
 	double Direction = (bInvert) ? -1.0 : 1.0;
 	double UseSpeed = Direction * FMathd::Sqrt(CurrentBrushRadius) * (SculptProperties->PrimaryBrushSpeed) * ActivePressure;
@@ -625,7 +687,7 @@ void UDynamicMeshSculptTool::ApplyOffsetBrush(const FRay& WorldRay)
 	int NumV = VertexROI.Num();
 	ROIPositionBuffer.SetNum(NumV, false);
 
-	ParallelFor(NumV, [this, Mesh, NewBrushPosLocal, UseSpeed](int k)
+	ParallelFor(NumV, [&](int k)
 	{
 		int VertIdx = VertexROI[k];
 		FVector3d OrigPos = Mesh->GetVertex(VertIdx);
@@ -637,7 +699,7 @@ void UDynamicMeshSculptTool::ApplyOffsetBrush(const FRay& WorldRay)
 		}
 		else
 		{
-			FVector3d MoveVec = UseSpeed * BaseNormal;
+			FVector3d MoveVec = (bUseViewDirection) ?  (UseSpeed*LocalNormal) : (UseSpeed*BaseNormal);
 			double Falloff = CalculateBrushFalloff(OrigPos.Distance(NewBrushPosLocal));
 			FVector3d NewPos = OrigPos + Falloff * MoveVec;
 			ROIPositionBuffer[k] = NewPos;
@@ -727,7 +789,7 @@ void UDynamicMeshSculptTool::ApplyPinchBrush(const FRay& WorldRay)
 
 	FVector3d NewBrushPosLocal = CurTargetTransform.InverseTransformPosition(LastBrushPosWorld);
 	FVector3d BrushNormalLocal = CurTargetTransform.InverseTransformNormal(LastBrushPosNormalWorld);
-	FVector3d OffsetBrushPosLocal = NewBrushPosLocal - SculptProperties->BrushDepth * CurrentBrushRadius * BrushNormalLocal;
+	FVector3d OffsetBrushPosLocal = NewBrushPosLocal - BrushProperties->Depth * CurrentBrushRadius * BrushNormalLocal;
 
 	// hardcoded lazybrush...
 	FVector3d NewSmoothBrushPosLocal = (0.75f)*LastSmoothBrushPosLocal + (0.25f)*NewBrushPosLocal;
@@ -784,7 +846,7 @@ void UDynamicMeshSculptTool::ApplyPinchBrush(const FRay& WorldRay)
 	LastSmoothBrushPosLocal = NewSmoothBrushPosLocal;
 }
 
-FFrame3d UDynamicMeshSculptTool::ComputeROIBrushPlane(const FVector3d& BrushCenter, bool bIgnoreDepth)
+FFrame3d UDynamicMeshSculptTool::ComputeROIBrushPlane(const FVector3d& BrushCenter, bool bIgnoreDepth, bool bViewAligned)
 {
 	FDynamicMesh3* Mesh = DynamicMeshComponent->GetMesh();
 	FVector3d AverageNormal(0, 0, 0);
@@ -802,10 +864,15 @@ FFrame3d UDynamicMeshSculptTool::ComputeROIBrushPlane(const FVector3d& BrushCent
 	AverageNormal.Normalize();
 	AveragePos /= WeightSum;
 
+	if (bViewAligned)
+	{
+		AverageNormal = -CameraState.Forward();
+	}
+
 	FFrame3d Result = FFrame3d(AveragePos, AverageNormal);
 	if (bIgnoreDepth == false)
 	{
-		Result.Origin -= SculptProperties->BrushDepth * CurrentBrushRadius * Result.Z();
+		Result.Origin -= BrushProperties->Depth * CurrentBrushRadius * Result.Z();
 	}
 
 	return Result;
@@ -819,26 +886,31 @@ void UDynamicMeshSculptTool::ApplyPlaneBrush(const FRay& WorldRay)
 		return;
 	}
 
+	static const double PlaneSigns[3] = { 0, -1, 1 };
+	double PlaneSign = PlaneSigns[(int32)PlaneBrushProperties->WhichSide];
+
 	FVector3d NewBrushPosLocal = CurTargetTransform.InverseTransformPosition(LastBrushPosWorld);
 	FVector3d BrushNormalLocal = CurTargetTransform.InverseTransformNormal(LastBrushPosNormalWorld);
 	double UseSpeed = FMathd::Sqrt(CurrentBrushRadius) * FMathd::Sqrt(SculptProperties->PrimaryBrushSpeed) * 0.05 * ActivePressure;
 
-
 	FDynamicMesh3* Mesh = DynamicMeshComponent->GetMesh();
 	int NumV = VertexROI.Num();
 	ROIPositionBuffer.SetNum(NumV, false);
 
-	ParallelFor(NumV, [this, Mesh, NewBrushPosLocal, UseSpeed](int k)
+	ParallelFor(NumV, [&](int k)
 	{
 		int VertIdx = VertexROI[k];
 		FVector3d OrigPos = Mesh->GetVertex(VertIdx);
 		FVector3d PlanePos = ActiveFixedBrushPlane.ToPlane(OrigPos, 2);
 		FVector3d Delta = PlanePos - OrigPos;
-		FVector3d MoveVec = UseSpeed * Delta;
-
-		double Falloff = CalculateBrushFalloff(OrigPos.Distance(NewBrushPosLocal));
-
-		FVector3d NewPos = OrigPos + Falloff * MoveVec;
+		double Dot = Delta.Dot(ActiveFixedBrushPlane.Z());
+		FVector3d NewPos = OrigPos;
+		if (Dot * PlaneSign >= 0)
+		{
+			FVector3d MoveVec = UseSpeed * Delta;
+			double Falloff = CalculateBrushFalloff(OrigPos.Distance(NewBrushPosLocal));
+			NewPos = OrigPos + Falloff * MoveVec;
+		}
 		ROIPositionBuffer[k] = NewPos;
 	});
 
@@ -864,30 +936,37 @@ void UDynamicMeshSculptTool::ApplyFixedPlaneBrush(const FRay& WorldRay)
 		return;
 	}
 
+	static const double PlaneSigns[3] = { 0, -1, 1 };
+	double PlaneSign = PlaneSigns[(int32)PlaneBrushProperties->WhichSide];
+
 	FVector3d NewBrushPosLocal = CurTargetTransform.InverseTransformPosition(LastBrushPosWorld);
 	FVector3d BrushNormalLocal = CurTargetTransform.InverseTransformNormal(LastBrushPosNormalWorld);
 	double UseSpeed = CurrentBrushRadius * FMathd::Sqrt(SculptProperties->PrimaryBrushSpeed) * 0.1 * ActivePressure;
 
 	FFrame3d FixedPlaneLocal(
 		CurTargetTransform.InverseTransformPosition(GizmoProperties->Position),
-		CurTargetTransform.GetRotation().Inverse() * (FQuaterniond)DrawPlaneOrientation);
+		CurTargetTransform.GetRotation().Inverse() * (FQuaterniond)GizmoProperties->Rotation);
 
 	FDynamicMesh3* Mesh = DynamicMeshComponent->GetMesh();
 	int NumV = VertexROI.Num();
 	ROIPositionBuffer.SetNum(NumV, false);
 
-	ParallelFor(NumV, [this, Mesh, NewBrushPosLocal, UseSpeed, FixedPlaneLocal](int k)
+	ParallelFor(NumV, [&](int k)
 	{
 		int VertIdx = VertexROI[k];
 		FVector3d OrigPos = Mesh->GetVertex(VertIdx);
 		FVector3d PlanePos = FixedPlaneLocal.ToPlane(OrigPos, 2);
 		FVector3d Delta = PlanePos - OrigPos;
-		double MaxDist = Delta.Normalize();
-		double Falloff = CalculateBrushFalloff(OrigPos.Distance(NewBrushPosLocal));
-		FVector3d MoveVec = Falloff * UseSpeed * Delta;
-		FVector3d NewPos = (MoveVec.SquaredLength() > MaxDist*MaxDist) ?
-			PlanePos : OrigPos + Falloff * MoveVec;
-
+		double Dot = Delta.Dot(FixedPlaneLocal.Z());
+		FVector3d NewPos = OrigPos;
+		if (Dot * PlaneSign >= 0)
+		{
+			double MaxDist = Delta.Normalize();
+			double Falloff = CalculateBrushFalloff(OrigPos.Distance(NewBrushPosLocal));
+			FVector3d MoveVec = Falloff * UseSpeed * Delta;
+			NewPos = (MoveVec.SquaredLength() > MaxDist* MaxDist) ?
+				PlanePos : OrigPos + Falloff * MoveVec;
+		}
 		ROIPositionBuffer[k] = NewPos;
 	});
 
@@ -917,8 +996,8 @@ void UDynamicMeshSculptTool::ApplyFlattenBrush(const FRay& WorldRay)
 	FVector3d BrushNormalLocal = CurTargetTransform.InverseTransformNormal(LastBrushPosNormalWorld);
 
 	double UseSpeed = FMathd::Sqrt(CurrentBrushRadius) * FMathd::Sqrt(SculptProperties->PrimaryBrushSpeed) * 0.05 * ActivePressure;
-	FFrame3d StampFlattenPlane = ComputeROIBrushPlane(NewBrushPosLocal, true);
-	//StampFlattenPlane.Origin -= SculptProperties->BrushDepth * CurrentBrushRadius * StampFlattenPlane.Z();
+	FFrame3d StampFlattenPlane = ComputeROIBrushPlane(NewBrushPosLocal, true, false);
+	//StampFlattenPlane.Origin -= BrushProperties->Depth * CurrentBrushRadius * StampFlattenPlane.Z();
 
 	FDynamicMesh3* Mesh = DynamicMeshComponent->GetMesh();
 	int NumV = VertexROI.Num();
@@ -1007,7 +1086,7 @@ void UDynamicMeshSculptTool::ApplyInflateBrush(const FRay& WorldRay)
 
 int UDynamicMeshSculptTool::FindHitSculptMeshTriangle(const FRay3d& LocalRay)
 {
-	if (SculptProperties->bHitBackFaces)
+	if (BrushProperties->bHitBackFaces)
 	{
 		return DynamicMeshComponent->GetOctree()->FindNearestHitObject(LocalRay);
 	}
@@ -1031,7 +1110,7 @@ int UDynamicMeshSculptTool::FindHitSculptMeshTriangle(const FRay3d& LocalRay)
 
 int UDynamicMeshSculptTool::FindHitTargetMeshTriangle(const FRay3d& LocalRay)
 {
-	if (SculptProperties->bHitBackFaces == false)
+	if (BrushProperties->bHitBackFaces == false)
 	{
 		FDynamicMesh3* Mesh = DynamicMeshComponent->GetMesh();
 
@@ -1049,7 +1128,7 @@ int UDynamicMeshSculptTool::FindHitTargetMeshTriangle(const FRay3d& LocalRay)
 
 	int HitTID =  BrushTargetMeshSpatial.FindNearestHitTriangle(LocalRay);
 
-	if (SculptProperties->bHitBackFaces == false)
+	if (BrushProperties->bHitBackFaces == false)
 	{
 		BrushTargetMeshSpatial.TriangleFilterF = nullptr;
 	}
@@ -1112,6 +1191,60 @@ bool UDynamicMeshSculptTool::UpdateBrushPositionOnSculptMesh(const FRay& WorldRa
 	return false;
 }
 
+void UDynamicMeshSculptTool::AlignBrushToView()
+{
+	LastBrushPosNormalWorld = -CameraState.Forward();
+}
+
+
+bool UDynamicMeshSculptTool::UpdateBrushPosition(const FRay& WorldRay)
+{
+	// This is an unfortunate hack necessary because we haven't refactored brushes properly yet
+
+	if (bSmoothing)
+	{
+		return UpdateBrushPositionOnSculptMesh(WorldRay);
+	}
+
+	bool bHit = false;
+	switch (SculptProperties->PrimaryBrushType)
+	{
+	case EDynamicMeshSculptBrushType::Offset:
+	case EDynamicMeshSculptBrushType::SculptMax:
+	case EDynamicMeshSculptBrushType::Pinch:
+		bHit = UpdateBrushPositionOnTargetMesh(WorldRay);
+		break;
+
+	case EDynamicMeshSculptBrushType::SculptView:
+	case EDynamicMeshSculptBrushType::PlaneViewAligned:
+		bHit = UpdateBrushPositionOnTargetMesh(WorldRay);
+		AlignBrushToView();
+		break;
+
+	case EDynamicMeshSculptBrushType::Move:
+		//return UpdateBrushPositionOnActivePlane(WorldRay);
+		bHit = UpdateBrushPositionOnSculptMesh(WorldRay);
+		break;
+
+	case EDynamicMeshSculptBrushType::Smooth:
+	case EDynamicMeshSculptBrushType::Inflate:
+	case EDynamicMeshSculptBrushType::Flatten:
+	case EDynamicMeshSculptBrushType::Plane:
+	case EDynamicMeshSculptBrushType::FixedPlane:
+		bHit = UpdateBrushPositionOnSculptMesh(WorldRay);
+		break;
+
+	default:
+		UE_LOG(LogTemp, Warning, TEXT("UDynamicMeshSculptTool: unknown brush type in UpdateBrushPosition"));
+		bHit = UpdateBrushPositionOnSculptMesh(WorldRay);
+		break;
+	}
+
+	return bHit;
+}
+
+
+
 void UDynamicMeshSculptTool::OnEndDrag(const FRay& Ray)
 {
 	bInDrag = false;
@@ -1140,12 +1273,14 @@ bool UDynamicMeshSculptTool::OnUpdateHover(const FInputDeviceRay& DevicePos)
 	}
 	else
 	{
-		FHitResult OutHit;
-		if (HitTest(DevicePos.WorldRay, OutHit))
-		{
-			LastBrushPosWorld = DevicePos.WorldRay.PointAt(OutHit.Distance + SculptProperties->BrushDepth*CurrentBrushRadius);
-			LastBrushPosNormalWorld = OutHit.Normal;
-		}
+		UpdateBrushPosition(DevicePos.WorldRay);
+
+		//FHitResult OutHit;
+		//if (HitTest(DevicePos.WorldRay, OutHit))
+		//{
+		//	LastBrushPosWorld = DevicePos.WorldRay.PointAt(OutHit.Distance + BrushProperties->Depth*CurrentBrushRadius);
+		//	LastBrushPosNormalWorld = OutHit.Normal;
+		//}
 	}
 	return true;
 }
@@ -1153,8 +1288,9 @@ bool UDynamicMeshSculptTool::OnUpdateHover(const FInputDeviceRay& DevicePos)
 void UDynamicMeshSculptTool::Render(IToolsContextRenderAPI* RenderAPI)
 {
 	UMeshSurfacePointTool::Render(RenderAPI);
+	GetToolManager()->GetContextQueriesAPI()->GetCurrentViewState(CameraState);
 
-	BrushIndicator->Update( (float)this->CurrentBrushRadius, (FVector)this->LastBrushPosWorld, (FVector)this->LastBrushPosNormalWorld, 1.0f);
+	BrushIndicator->Update( (float)this->CurrentBrushRadius, (FVector)this->LastBrushPosWorld, (FVector)this->LastBrushPosNormalWorld, 1.0f-BrushProperties->BrushFalloffAmount);
 
 	if (SculptProperties->PrimaryBrushType == EDynamicMeshSculptBrushType::FixedPlane)
 	{
@@ -1163,7 +1299,7 @@ void UDynamicMeshSculptTool::Render(IToolsContextRenderAPI* RenderAPI)
 		float GridThickness = 0.5f;
 		float GridLineSpacing = 25.0f;   // @todo should be relative to view
 		int NumGridLines = 10;
-		FFrame3f DrawFrame(GizmoProperties->Position, DrawPlaneOrientation);
+		FFrame3f DrawFrame(GizmoProperties->Position, GizmoProperties->Rotation);
 		MeshDebugDraw::DrawSimpleGrid(DrawFrame, NumGridLines, GridLineSpacing, GridThickness, GridColor, false, PDI, FTransform::Identity);
 	}
 }
@@ -1186,16 +1322,22 @@ void UDynamicMeshSculptTool::Tick(float DeltaTime)
 
 	ShowWireframeWatcher.CheckAndUpdate();
 	MaterialModeWatcher.CheckAndUpdate();
+	FlatShadingWatcher.CheckAndUpdate();
+	ColorWatcher.CheckAndUpdate();
+	ImageWatcher.CheckAndUpdate();
+	BrushTypeWatcher.CheckAndUpdate();
+	GizmoPositionWatcher.CheckAndUpdate();
+	GizmoRotationWatcher.CheckAndUpdate();
 
 	bool bGizmoVisible = (SculptProperties->PrimaryBrushType == EDynamicMeshSculptBrushType::FixedPlane)
 		&& (GizmoProperties->bShowGizmo);
 	UpdateFixedPlaneGizmoVisibility(bGizmoVisible);
 	GizmoProperties->bPropertySetEnabled = (SculptProperties->PrimaryBrushType == EDynamicMeshSculptBrushType::FixedPlane);
 
-	if (bPendingSetFixedPlanePosition)
+	if (PendingWorkPlaneUpdate != EPendingWorkPlaneUpdate::NoUpdatePending)
 	{
-		SetFixedSculptPlaneFromWorldPos((FVector)LastBrushPosWorld);
-		bPendingSetFixedPlanePosition = false;
+		SetFixedSculptPlaneFromWorldPos((FVector)LastBrushPosWorld, (FVector)LastBrushPosNormalWorld, PendingWorkPlaneUpdate);
+		PendingWorkPlaneUpdate = EPendingWorkPlaneUpdate::NoUpdatePending;
 	}
 
 	// if user changed to not-frozen, we need to update the target
@@ -1691,6 +1833,26 @@ void UDynamicMeshSculptTool::DecreaseBrushSpeedAction()
 	SculptProperties->PrimaryBrushSpeed = FMath::Clamp(SculptProperties->PrimaryBrushSpeed - 0.05f, 0.0f, 1.0f);
 }
 
+void UDynamicMeshSculptTool::IncreaseBrushFalloffAction()
+{
+	const float ChangeAmount = 0.1f;
+	const float OldValue = BrushProperties->BrushFalloffAmount;
+
+	float NewValue = OldValue + ChangeAmount;
+	BrushProperties->BrushFalloffAmount = FMath::Clamp(NewValue, 0.f, 1.f);
+}
+
+void UDynamicMeshSculptTool::DecreaseBrushFalloffAction()
+{
+	const float ChangeAmount = 0.1f;
+	const float OldValue = BrushProperties->BrushFalloffAmount;
+
+	float NewValue = OldValue - ChangeAmount;
+	BrushProperties->BrushFalloffAmount = FMath::Clamp(NewValue, 0.f, 1.f);
+}
+
+
+
 void UDynamicMeshSculptTool::NextBrushModeAction()
 {
 	uint8 LastMode = (uint8)EDynamicMeshSculptBrushType::LastValue;
@@ -1748,18 +1910,33 @@ void UDynamicMeshSculptTool::RegisterActions(FInteractiveToolActionSet& ActionSe
 		EModifierKey::None, EKeys::LeftBracket,
 		[this]() { DecreaseBrushRadiusAction(); } );
 
+	ActionSet.RegisterAction(this, (int32)EStandardToolActions::BaseClientDefinedActionID + 12,
+		TEXT("BrushIncreaseFalloff"),
+		LOCTEXT("BrushIncreaseFalloff", "Increase Brush Falloff"),
+		LOCTEXT("BrushIncreaseFalloffTooltip", "Press this key to increase brush falloff by a fixed increment."),
+		EModifierKey::Shift | EModifierKey::Control, EKeys::RightBracket,
+		[this]() { IncreaseBrushFalloffAction(); });
+
+	ActionSet.RegisterAction(this, (int32)EStandardToolActions::BaseClientDefinedActionID + 13,
+		TEXT("BrushDecreaseFalloff"),
+		LOCTEXT("BrushDecreaseFalloff", "Decrease Brush Falloff"),
+		LOCTEXT("BrushDecreaseFalloffTooltip", "Press this key to decrease brush falloff by a fixed increment."),
+		EModifierKey::Shift | EModifierKey::Control, EKeys::LeftBracket,
+		[this]() { DecreaseBrushFalloffAction(); });
+
+
 	ActionSet.RegisterAction(this, (int32)EStandardToolActions::BaseClientDefinedActionID+1,
 		TEXT("NextBrushMode"),
 		LOCTEXT("SculptNextBrushMode", "Next Brush Type"),
 		LOCTEXT("SculptNextBrushModeTooltip", "Cycle to next Brush Type"),
-		EModifierKey::Shift, EKeys::A,
+		EModifierKey::None, EKeys::A,
 		[this]() { NextBrushModeAction(); });
 
 	ActionSet.RegisterAction(this, (int32)EStandardToolActions::BaseClientDefinedActionID+2,
 		TEXT("PreviousBrushMode"),
 		LOCTEXT("SculptPreviousBrushMode", "Previous Brush Type"),
 		LOCTEXT("SculptPreviousBrushModeTooltip", "Cycle to previous Brush Type"),
-		EModifierKey::Shift, EKeys::Q,
+		EModifierKey::None, EKeys::Q,
 		[this]() { PreviousBrushModeAction(); });
 
 	//ActionSet.RegisterAction(this, (int32)EStandardToolActions::BaseClientDefinedActionID + 10,
@@ -1780,28 +1957,28 @@ void UDynamicMeshSculptTool::RegisterActions(FInteractiveToolActionSet& ActionSe
 		TEXT("SculptIncreaseSize"),
 		LOCTEXT("SculptIncreaseSize", "Increase Size"),
 		LOCTEXT("SculptIncreaseSizeTooltip", "Increase Brush Size"),
-		EModifierKey::Shift, EKeys::D,
+		EModifierKey::None, EKeys::D,
 		[this]() { IncreaseBrushRadiusAction(); });
 
 	ActionSet.RegisterAction(this, (int32)EStandardToolActions::BaseClientDefinedActionID + 51,
 		TEXT("SculptDecreaseSize"),
 		LOCTEXT("SculptDecreaseSize", "Decrease Size"),
 		LOCTEXT("SculptDecreaseSizeTooltip", "Decrease Brush Size"),
-		EModifierKey::Shift, EKeys::S,
+		EModifierKey::None, EKeys::S,
 		[this]() { DecreaseBrushRadiusAction(); });
 
 	ActionSet.RegisterAction(this, (int32)EStandardToolActions::BaseClientDefinedActionID + 52,
 		TEXT("SculptIncreaseSizeSmallStep"),
 		LOCTEXT("SculptIncreaseSize", "Increase Size"),
 		LOCTEXT("SculptIncreaseSizeTooltip", "Increase Brush Size"),
-		EModifierKey::Shift | EModifierKey::Control, EKeys::D,
+		EModifierKey::Shift, EKeys::D,
 		[this]() { IncreaseBrushRadiusSmallStepAction(); });
 
 	ActionSet.RegisterAction(this, (int32)EStandardToolActions::BaseClientDefinedActionID + 53,
 		TEXT("SculptDecreaseSizeSmallStemp"),
 		LOCTEXT("SculptDecreaseSize", "Decrease Size"),
 		LOCTEXT("SculptDecreaseSizeTooltip", "Decrease Brush Size"),
-		EModifierKey::Shift | EModifierKey::Control, EKeys::S,
+		EModifierKey::Shift, EKeys::S,
 		[this]() { DecreaseBrushRadiusSmallStepAction(); });
 
 
@@ -1811,14 +1988,14 @@ void UDynamicMeshSculptTool::RegisterActions(FInteractiveToolActionSet& ActionSe
 		TEXT("SculptIncreaseSpeed"),
 		LOCTEXT("SculptIncreaseSpeed", "Increase Speed"),
 		LOCTEXT("SculptIncreaseSpeedTooltip", "Increase Brush Speed"),
-		EModifierKey::Shift, EKeys::E,
+		EModifierKey::None, EKeys::E,
 		[this]() { IncreaseBrushSpeedAction(); });
 
 	ActionSet.RegisterAction(this, (int32)EStandardToolActions::BaseClientDefinedActionID + 61,
 		TEXT("SculptDecreaseSpeed"),
 		LOCTEXT("SculptDecreaseSpeed", "Decrease Speed"),
 		LOCTEXT("SculptDecreaseSpeedTooltip", "Decrease Brush Speed"),
-		EModifierKey::Shift, EKeys::W,
+		EModifierKey::None, EKeys::W,
 		[this]() { DecreaseBrushSpeedAction(); });
 
 
@@ -1832,11 +2009,25 @@ void UDynamicMeshSculptTool::RegisterActions(FInteractiveToolActionSet& ActionSe
 
 
 	ActionSet.RegisterAction(this, (int32)EStandardToolActions::BaseClientDefinedActionID + 100,
-		TEXT("SetFixedSculptPlane"),
-		LOCTEXT("SetFixedSculptPlane", "Set Fixed Sculpt Plane"),
-		LOCTEXT("SetFixedSculptPlaneTooltip", "Set position of fixed sculpt plane"),
-		EModifierKey::None, EKeys::P,
-		[this]() { bPendingSetFixedPlanePosition = true; });
+		TEXT("SetSculptWorkSurfacePosNormal"),
+		LOCTEXT("SetSculptWorkSurfacePosNormal", "Reorient Work Surface"),
+		LOCTEXT("SetSculptWorkSurfacePosNormalTooltip", "Move the Sculpting Work Plane/Surface to Position and Normal of World hit point under cursor"),
+		EModifierKey::Shift, EKeys::T,
+		[this]() { PendingWorkPlaneUpdate = EPendingWorkPlaneUpdate::MoveToHitPositionNormal; });
+
+	ActionSet.RegisterAction(this, (int32)EStandardToolActions::BaseClientDefinedActionID + 101,
+		TEXT("SetSculptWorkSurfacePos"),
+		LOCTEXT("SetSculptWorkSurfacePos", "Reposition Work Surface"),
+		LOCTEXT("SetSculptWorkSurfacePosTooltip", "Move the Sculpting Work Plane/Surface to World hit point under cursor (keep current Orientation)"),
+		EModifierKey::None, EKeys::T,
+		[this]() { PendingWorkPlaneUpdate = EPendingWorkPlaneUpdate::MoveToHitPosition; });
+
+	ActionSet.RegisterAction(this, (int32)EStandardToolActions::BaseClientDefinedActionID + 102,
+		TEXT("SetSculptWorkSurfaceView"),
+		LOCTEXT("SetSculptWorkSurfaceView", "View-Align Work Surface"),
+		LOCTEXT("SetSculptWorkSurfaceViewTooltip", "Move the Sculpting Work Plane/Surface to World hit point under cursor and align to View"),
+		EModifierKey::Control | EModifierKey::Shift, EKeys::T,
+		[this]() { PendingWorkPlaneUpdate = EPendingWorkPlaneUpdate::MoveToHitPositionViewAligned; });
 
 }
 
@@ -1907,36 +2098,153 @@ void UDynamicMeshSculptTool::UpdateMaterialMode(EMeshEditingMaterialModes Materi
 	{
 		DynamicMeshComponent->ClearOverrideRenderMaterial();
 		DynamicMeshComponent->bCastDynamicShadow = ComponentTarget->GetOwnerComponent()->bCastDynamicShadow;
+		ActiveOverrideMaterial = nullptr;
 	}
-	else if (MaterialMode == EMeshEditingMaterialModes::MeshFocusMaterial)
+	else 
 	{
-		UMaterialInterface* SculptMaterial = ToolSetupUtil::GetSculptMaterial1(GetToolManager());
-		if (SculptMaterial != nullptr)
+		if (MaterialMode == EMeshEditingMaterialModes::Custom)
+		{
+			ActiveOverrideMaterial = ToolSetupUtil::GetCustomImageBasedSculptMaterial(GetToolManager(), ViewProperties->Image);
+			if (ViewProperties->Image != nullptr)
+			{
+				ActiveOverrideMaterial->SetTextureParameterValue(TEXT("ImageTexture"), ViewProperties->Image);
+			}
+		}
+		else
+		{
+			UMaterialInterface* SculptMaterial = nullptr;
+			switch (MaterialMode)
+			{
+			case EMeshEditingMaterialModes::Diffuse:
+				SculptMaterial = ToolSetupUtil::GetDefaultSculptMaterial(GetToolManager());
+				break;
+			case EMeshEditingMaterialModes::Grey:
+				SculptMaterial = ToolSetupUtil::GetImageBasedSculptMaterial(GetToolManager(), ToolSetupUtil::ImageMaterialType::DefaultBasic);
+				break;
+			case EMeshEditingMaterialModes::Soft:
+				SculptMaterial = ToolSetupUtil::GetImageBasedSculptMaterial(GetToolManager(), ToolSetupUtil::ImageMaterialType::DefaultSoft);
+				break;
+			case EMeshEditingMaterialModes::TangentNormal:
+				SculptMaterial = ToolSetupUtil::GetImageBasedSculptMaterial(GetToolManager(), ToolSetupUtil::ImageMaterialType::TangentNormalFromView);
+				break;
+			}
+			if (SculptMaterial != nullptr)
+			{
+				ActiveOverrideMaterial = UMaterialInstanceDynamic::Create(SculptMaterial, this);
+			}
+		}
+
+		if (ActiveOverrideMaterial != nullptr)
 		{
-			DynamicMeshComponent->SetOverrideRenderMaterial(SculptMaterial);
+			DynamicMeshComponent->SetOverrideRenderMaterial(ActiveOverrideMaterial);
+			ActiveOverrideMaterial->SetScalarParameterValue(TEXT("FlatShading"), (ViewProperties->bFlatShading) ? 1.0f : 0.0f);
 		}
+
 		DynamicMeshComponent->bCastDynamicShadow = false;
 	}
 }
 
-void UDynamicMeshSculptTool::SetFixedSculptPlaneFromWorldPos(const FVector& Position)
+
+void UDynamicMeshSculptTool::UpdateFlatShadingSetting(bool bNewValue)
 {
-	UpdateFixedSculptPlanePosition(Position);
+	if (ActiveOverrideMaterial != nullptr)
+	{
+		ActiveOverrideMaterial->SetScalarParameterValue(TEXT("FlatShading"), (bNewValue) ? 1.0f : 0.0f);
+	}
+}
+
+
+void UDynamicMeshSculptTool::UpdateColorSetting(FLinearColor NewColor)
+{
+	if (ActiveOverrideMaterial != nullptr)
+	{
+		ActiveOverrideMaterial->SetVectorParameterValue(TEXT("Color"), NewColor);
+	}
+}
+
+void UDynamicMeshSculptTool::UpdateImageSetting(UTexture2D* NewImage)
+{
+	if (ActiveOverrideMaterial != nullptr)
+	{
+		ActiveOverrideMaterial->SetTextureParameterValue(TEXT("ImageTexture"), NewImage);
+	}
+}
+
+void UDynamicMeshSculptTool::UpdateBrushType(EDynamicMeshSculptBrushType BrushType)
+{
+	static const FText BaseMessage = LOCTEXT("OnStartSculptTool", "Hold Shift to Smooth, Ctrl to Invert (where applicable). Q/A keys cycle Brush Type. S/D changes Size (+Shift to small-step), W/E changes Strength.");
+	FTextBuilder Builder;
+	Builder.AppendLine(BaseMessage);
+
+	SetToolPropertySourceEnabled(GizmoProperties, false);
+	SetToolPropertySourceEnabled(PlaneBrushProperties, false);
+
+	if (BrushType == EDynamicMeshSculptBrushType::FixedPlane)
+	{
+		Builder.AppendLine(LOCTEXT("FixedPlaneTip", "Use T to reposition Work Plane at cursor, Shift+T to align to Normal, Ctrl+Shift+T to align to View"));
+		SetToolPropertySourceEnabled(PlaneBrushProperties, true);
+		SetToolPropertySourceEnabled(GizmoProperties, true);
+	}
+	if (BrushType == EDynamicMeshSculptBrushType::Plane || BrushType == EDynamicMeshSculptBrushType::PlaneViewAligned)
+	{
+		SetToolPropertySourceEnabled(PlaneBrushProperties, true);
+	}
+
+	GetToolManager()->DisplayMessage(Builder.ToText(), EToolMessageLevel::UserNotification);
+}
+
+
+void UDynamicMeshSculptTool::SetFixedSculptPlaneFromWorldPos(const FVector& Position, const FVector& Normal, EPendingWorkPlaneUpdate UpdateType)
+{
+	if (UpdateType == EPendingWorkPlaneUpdate::MoveToHitPositionNormal)
+	{
+		UpdateFixedSculptPlanePosition(Position);
+		FFrame3d CurFrame(FVector::ZeroVector, GizmoProperties->Rotation);
+		CurFrame.AlignAxis(2, (FVector3d)Normal );
+		UpdateFixedSculptPlaneRotation( (FQuat)CurFrame.Rotation );
+	}
+	else if (UpdateType == EPendingWorkPlaneUpdate::MoveToHitPositionViewAligned)
+	{
+		UpdateFixedSculptPlanePosition(Position);
+		FFrame3d CurFrame(FVector::ZeroVector, GizmoProperties->Rotation);
+		CurFrame.AlignAxis(2, -(FVector3d)CameraState.Forward());
+		UpdateFixedSculptPlaneRotation((FQuat)CurFrame.Rotation);
+	}
+	else
+	{
+		UpdateFixedSculptPlanePosition(Position);
+	}
+
 	if (PlaneTransformGizmo != nullptr)
 	{
-		PlaneTransformGizmo->SetNewGizmoTransform(FTransform(DrawPlaneOrientation, GizmoProperties->Position));
+		PlaneTransformGizmo->SetNewGizmoTransform(FTransform(GizmoProperties->Rotation, GizmoProperties->Position));
 	}
 }
 
 void UDynamicMeshSculptTool::PlaneTransformChanged(UTransformProxy* Proxy, FTransform Transform)
 {
-	DrawPlaneOrientation = Transform.GetRotation();
+	UpdateFixedSculptPlaneRotation(Transform.GetRotation());
 	UpdateFixedSculptPlanePosition(Transform.GetLocation());
 }
 
 void UDynamicMeshSculptTool::UpdateFixedSculptPlanePosition(const FVector& Position)
 {
 	GizmoProperties->Position = Position;
+	GizmoPositionWatcher.SilentUpdate();
+}
+
+void UDynamicMeshSculptTool::UpdateFixedSculptPlaneRotation(const FQuat& Rotation)
+{
+	GizmoProperties->Rotation = Rotation;
+	GizmoRotationWatcher.SilentUpdate();
+}
+
+void UDynamicMeshSculptTool::UpdateGizmoFromProperties()
+{
+	if (PlaneTransformGizmo != nullptr)
+	{
+		PlaneTransformGizmo->SetNewGizmoTransform(FTransform(GizmoProperties->Rotation, GizmoProperties->Position));
+	}
 }
 
 void UDynamicMeshSculptTool::UpdateFixedPlaneGizmoVisibility(bool bVisible)
@@ -1955,8 +2263,10 @@ void UDynamicMeshSculptTool::UpdateFixedPlaneGizmoVisibility(bool bVisible)
 		{
 			PlaneTransformGizmo = GetToolManager()->GetPairedGizmoManager()->CreateCustomTransformGizmo(
 				ETransformGizmoSubElements::StandardTranslateRotate, this);
+			PlaneTransformGizmo->bUseContextCoordinateSystem = false;
+			PlaneTransformGizmo->CurrentCoordinateSystem = EToolContextCoordinateSystem::Local;
 			PlaneTransformGizmo->SetActiveTarget(PlaneTransformProxy, GetToolManager());
-			PlaneTransformGizmo->SetNewGizmoTransform(FTransform(DrawPlaneOrientation, GizmoProperties->Position));
+			PlaneTransformGizmo->SetNewGizmoTransform(FTransform(GizmoProperties->Rotation, GizmoProperties->Position));
 		}
 
 		PlaneTransformGizmo->bSnapToWorldGrid = GizmoProperties->bSnapToGrid;
diff --git a/Engine/Plugins/Experimental/MeshModelingToolset/Source/MeshModelingTools/Private/MeshSelectionTool.cpp b/Engine/Plugins/Experimental/MeshModelingToolset/Source/MeshModelingTools/Private/MeshSelectionTool.cpp
--- a/Engine/Plugins/Experimental/MeshModelingToolset/Source/MeshModelingTools/Private/MeshSelectionTool.cpp
+++ b/Engine/Plugins/Experimental/MeshModelingToolset/Source/MeshModelingTools/Private/MeshSelectionTool.cpp
@@ -92,7 +92,7 @@ void UMeshSelectionTool::Setup()
 	UDynamicMeshBrushTool::Setup();
 
 	// hide strength and falloff
-	BrushProperties->bShowFullSettings = false;
+	BrushProperties->bShowStrength = BrushProperties->bShowFalloff = false;
 	BrushProperties->RestoreProperties(this);
 
 	SelectionProps = NewObject<UMeshSelectionToolProperties>(this);
diff --git a/Engine/Plugins/Experimental/MeshModelingToolset/Source/MeshModelingTools/Private/Properties/MeshMaterialProperties.cpp b/Engine/Plugins/Experimental/MeshModelingToolset/Source/MeshModelingTools/Private/Properties/MeshMaterialProperties.cpp
--- a/Engine/Plugins/Experimental/MeshModelingToolset/Source/MeshModelingTools/Private/Properties/MeshMaterialProperties.cpp
+++ b/Engine/Plugins/Experimental/MeshModelingToolset/Source/MeshModelingTools/Private/Properties/MeshMaterialProperties.cpp
@@ -95,6 +95,9 @@ void UMeshEditingViewProperties::SaveProperties(UInteractiveTool* SaveFromTool)
 	UMeshEditingViewProperties* PropertyCache = GetPropertyCache<UMeshEditingViewProperties>();
 	PropertyCache->bShowWireframe = this->bShowWireframe;
 	PropertyCache->MaterialMode = this->MaterialMode;
+	PropertyCache->bFlatShading = this->bFlatShading;
+	PropertyCache->Color = this->Color;
+	PropertyCache->Image = this->Image;
 }
 
 void UMeshEditingViewProperties::RestoreProperties(UInteractiveTool* RestoreToTool)
@@ -102,6 +105,9 @@ void UMeshEditingViewProperties::RestoreProperties(UInteractiveTool* RestoreToTo
 	UMeshEditingViewProperties* PropertyCache = GetPropertyCache<UMeshEditingViewProperties>();
 	this->bShowWireframe = PropertyCache->bShowWireframe;
 	this->MaterialMode = PropertyCache->MaterialMode;
+	this->bFlatShading = PropertyCache->bFlatShading;
+	this->Color = PropertyCache->Color;
+	this->Image = PropertyCache->Image;
 }
 
 
diff --git a/Engine/Plugins/Experimental/MeshModelingToolset/Source/MeshModelingTools/Public/DynamicMeshSculptTool.h b/Engine/Plugins/Experimental/MeshModelingToolset/Source/MeshModelingTools/Public/DynamicMeshSculptTool.h
--- a/Engine/Plugins/Experimental/MeshModelingToolset/Source/MeshModelingTools/Public/DynamicMeshSculptTool.h
+++ b/Engine/Plugins/Experimental/MeshModelingToolset/Source/MeshModelingTools/Public/DynamicMeshSculptTool.h
@@ -22,6 +22,7 @@
 
 class UTransformGizmo;
 class UTransformProxy;
+class UMaterialInstanceDynamic;
 
 DECLARE_STATS_GROUP(TEXT("SculptTool"), STATGROUP_SculptTool, STATCAT_Advanced);
 DECLARE_CYCLE_STAT(TEXT("SculptTool_UpdateROI"), SculptTool_UpdateROI, STATGROUP_SculptTool);
@@ -51,31 +52,37 @@ class FSubRegionRemesher;
 UENUM()
 enum class EDynamicMeshSculptBrushType : uint8
 {
-	/** Move Brush moves vertices parallel to the view plane  */
+	/** Move vertices parallel to the view plane  */
 	Move UMETA(DisplayName = "Move"),
 
-	/** Smooth brush smooths mesh vertices  */
+	/** Smooth mesh vertices */
 	Smooth UMETA(DisplayName = "Smooth"),
 
-	/** Sculpt Brush displaces vertices from the surface */
-	Offset UMETA(DisplayName = "Sculpt"),
+	/** Displace vertices along the average surface normal (Ctrl to invert) */
+	Offset UMETA(DisplayName = "Sculpt (Normal)"),
 
-	/** SculptMax Brush displaces vertices to a maximum offset height */
+	/** Displace vertices towards the camera viewpoint (Ctrl to invert) */
+	SculptView UMETA(DisplayName = "Sculpt (Viewpoint)"),
+
+	/** Displaces vertices along the average surface normal to a maximum height based on the brush size (Ctrl to invert) */
 	SculptMax UMETA(DisplayName = "Sculpt Max"),
 
-	/** Inflate brush displaces vertices along their vertex normals */
+	/** Displace vertices along their vertex normals */
 	Inflate UMETA(DisplayName = "Inflate"),
 
-	/** Pinch Brush pulls vertices towards the center of the brush*/
+	/** Move vertices towards the center of the brush (Ctrl to push away) */
 	Pinch UMETA(DisplayName = "Pinch"),
 
-	/** Flatten Brush pulls vertices towards the average plane of the brush stamp */
+	/** Move vertices towards the average plane of the brush stamp region */
 	Flatten UMETA(DisplayName = "Flatten"),
 
-	/** Plane Brush pulls vertices to a plane defined by the initial brush position  */
-	Plane UMETA(DisplayName = "Plane"),
+	/** Move vertices towards a plane defined by the initial brush position  */
+	Plane UMETA(DisplayName = "Plane (Normal)"),
+
+	/** Move vertices towards a view-facing plane defined at the initial brush position */
+	PlaneViewAligned UMETA(DisplayName = "Plane (Viewpoint)"),
 
-	/** FixedPlane Brush pulls vertices to a fixed plane defined by the gizmo */
+	/** Move vertices towards a fixed plane in world space, positioned with a 3D gizmo */
 	FixedPlane UMETA(DisplayName = "FixedPlane"),
 
 	LastValue UMETA(Hidden)
@@ -101,52 +108,90 @@ public:
 	virtual UMeshSurfacePointTool* CreateNewTool(const FToolBuilderState& SceneState) const override;
 };
 
+
+
+UCLASS()
+class MESHMODELINGTOOLS_API USculptBrushProperties : public UBrushBaseProperties
+{
+	GENERATED_BODY()
+public:
+	USculptBrushProperties();
+
+	virtual void SaveRestoreProperties(UInteractiveTool* Tool, bool bSaving) override;
+
+	/** Depth of Brush into surface along view ray or surface normal, depending on the Active Brush Type */
+	UPROPERTY(EditAnywhere, Category = Brush, meta = (UIMin = "-0.5", UIMax = "0.5", ClampMin = "-1.0", ClampMax = "1.0", DisplayPriority = 5))
+	float Depth = 0;
+
+	/** Allow the Brush to hit the back-side of the mesh */
+	UPROPERTY(EditAnywhere, Category = Brush, meta = (DisplayPriority = 6))
+	bool bHitBackFaces = true;
+};
+
+
 UCLASS()
 class MESHMODELINGTOOLS_API UBrushSculptProperties : public UInteractiveToolPropertySet
 {
 	GENERATED_BODY()
 
 public:
-	/* This is a dupe of the bool in the tool class.  I needed it here so it could be checked as an EDITCONDITION */
+	virtual void SaveRestoreProperties(UInteractiveTool* Tool, bool bSaving) override;
+
+	/* This is a dupe of the bool in the tool class.  I needed it here so it could be checked as an EditCondition */
 	UPROPERTY()
 	bool bIsRemeshingEnabled = false;
 
 	/** Primary Brush Mode */
-	UPROPERTY(EditAnywhere, Category = Sculpting)
+	UPROPERTY(EditAnywhere, Category = Sculpting, meta = (DisplayName = "Brush Type"))
 	EDynamicMeshSculptBrushType PrimaryBrushType = EDynamicMeshSculptBrushType::Move;
 
-	/** Power/Speed of Brush */
-	UPROPERTY(EditAnywhere, Category = Sculpting, meta = (UIMin = "0.0", UIMax = "1.0", ClampMin = "0.0", ClampMax = "1.0", EditCondition = "PrimaryBrushType != EDynamicMeshSculptBrushType::Pull"))
+	/** Strength of the Primary Brush */
+	UPROPERTY(EditAnywhere, Category = Sculpting, meta = (DisplayName = "Strength", UIMin = "0.0", UIMax = "1.0", ClampMin = "0.0", ClampMax = "1.0", EditCondition = "PrimaryBrushType != EDynamicMeshSculptBrushType::Pull"))
 	float PrimaryBrushSpeed = 0.5;
 
-	/** Smoothing Speed of Smoothing brush */
-	UPROPERTY(EditAnywhere, Category = Sculpting, meta = (UIMin = "0.0", UIMax = "1.0", ClampMin = "0.0", ClampMax = "1.0"))
+	/** If true, try to preserve the shape of the UV/3D mapping. This will limit Smoothing and Remeshing in some cases. */
+	UPROPERTY(EditAnywhere, Category = Sculpting)
+	bool bPreserveUVFlow = false;
+
+	/** When Freeze Target is toggled on, the Brush Target Surface will be Frozen in its current state, until toggled off. Brush strokes will be applied relative to the Target Surface, for applicable Brushes */
+	UPROPERTY(EditAnywhere, Category = Sculpting, meta = (EditCondition = "PrimaryBrushType == EDynamicMeshSculptBrushType::Sculpt || PrimaryBrushType == EDynamicMeshSculptBrushType::SculptMax || PrimaryBrushType == EDynamicMeshSculptBrushType::SculptView || PrimaryBrushType == EDynamicMeshSculptBrushType::Pinch" ))
+	bool bFreezeTarget = false;
+
+
+	/** Strength of Shift-to-Smooth Brushing and Smoothing Brush */
+	UPROPERTY(EditAnywhere, Category = Smoothing, meta = (DisplayName = "Smoothing Strength", UIMin = "0.0", UIMax = "1.0", ClampMin = "0.0", ClampMax = "1.0"))
 	float SmoothBrushSpeed = 0.25;
 
-	/** If enabled, Full Remeshing is applied during smoothing, which will wipe out fine details */
-	UPROPERTY(EditAnywhere, Category = Sculpting, meta = (EditConditionHides, HideEditConditionToggle, EditCondition = "bIsRemeshingEnabled"))
+	/** If enabled, Remeshing is applied during Smoothing, which will wipe out fine details. Otherwise no remeshing is applied during Smoothing. */
+	UPROPERTY(EditAnywhere, Category = Smoothing, meta = (DisplayName = "Remesh in Smooth", EditConditionHides, HideEditConditionToggle, EditCondition = "bIsRemeshingEnabled"))
 	bool bSmoothBrushRemeshes = false;
+};
 
-	/** If true, try to preserve the shape of the UV/3D mapping. This will prevent smoothing in some cases */
-	UPROPERTY(EditAnywhere, Category = Sculpting)
-	bool bPreserveUVFlow = false;
 
-	/** Depth of Brush into surface along view ray or surface normal, depending on brush */
-	UPROPERTY(EditAnywhere, Category = Sculpting, meta = (UIMin = "-0.5", UIMax = "0.5", ClampMin = "-1.0", ClampMax = "1.0"))
-	float BrushDepth = 0;
 
-	/** Disable updating of the brush Target Surface after each stroke */
-	UPROPERTY(EditAnywhere, Category = Sculpting, meta = (EditCondition = "PrimaryBrushType == EDynamicMeshSculptBrushType::Sculpt || PrimaryBrushType == EDynamicMeshSculptBrushType::SculptMax || PrimaryBrushType == EDynamicMeshSculptBrushType::Pinch" ))
-	bool bFreezeTarget = false;
 
-	/** Hit back sides of triangles */
-	UPROPERTY(EditAnywhere, Category = Sculpting)
-	bool bHitBackFaces = true;
+UENUM()
+enum class EPlaneBrushSideMode : uint8
+{
+	BothSides = 0,
+	PushDown = 1,
+	PullTowards = 2
+};
+
+
+UCLASS()
+class MESHMODELINGTOOLS_API UPlaneBrushProperties : public UInteractiveToolPropertySet
+{
+	GENERATED_BODY()
+public:
+	virtual void SaveRestoreProperties(UInteractiveTool* Tool, bool bSaving) override;
 
-	virtual void SaveProperties(UInteractiveTool* SaveFromTool) override;
-	virtual void RestoreProperties(UInteractiveTool* RestoreToTool) override;
+	/** Control whether effect of brush should be limited to one side of the Plane  */
+	UPROPERTY(EditAnywhere, Category = PlaneBrush)
+	EPlaneBrushSideMode WhichSide = EPlaneBrushSideMode::BothSides;
 };
 
+
 UCLASS()
 class MESHMODELINGTOOLS_API UFixedPlaneBrushProperties : public UInteractiveToolPropertySet
 {
@@ -155,22 +200,24 @@ class MESHMODELINGTOOLS_API UFixedPlaneBrushProperties : public UInteractiveTool
 public:
 	UFixedPlaneBrushProperties();
 
+	virtual void SaveRestoreProperties(UInteractiveTool* Tool, bool bSaving) override;
+
 	UPROPERTY()
 	bool bPropertySetEnabled = true;
 
-	/** Toggle whether interactive plane gizmo is visible */
-	UPROPERTY(EditAnywhere, Category = TargetPlane, meta = (EditCondition = "bPropertySetEnabled == true"))
+	/** Toggle whether Work Plane Positioing Gizmo is visible */
+	UPROPERTY(EditAnywhere, Category = TargetPlane, meta = (HideEditConditionToggle, EditCondition = "bPropertySetEnabled == true"))
 	bool bShowGizmo;
 
-	/** Toggle whether fixed plane snaps to grid */
-	UPROPERTY(EditAnywhere, Category = TargetPlane, meta = (EditCondition = "bPropertySetEnabled == true"))
+	/** Toggle whether Work Plane snaps to grid when using Gizmo */
+	UPROPERTY(EditAnywhere, Category = TargetPlane, meta = (HideEditConditionToggle, EditCondition = "bPropertySetEnabled == true"))
 	bool bSnapToGrid;
 
-	UPROPERTY(EditAnywhere, Category = TargetPlane, AdvancedDisplay, meta = (EditCondition = "bPropertySetEnabled == true"))
+	UPROPERTY(EditAnywhere, Category = TargetPlane, AdvancedDisplay, meta = (HideEditConditionToggle, EditCondition = "bPropertySetEnabled == true"))
 	FVector Position;
 
-	virtual void SaveProperties(UInteractiveTool* SaveFromTool) override;
-	virtual void RestoreProperties(UInteractiveTool* RestoreToTool) override;
+	UPROPERTY(EditAnywhere, Category = TargetPlane, AdvancedDisplay, meta = (HideEditConditionToggle, EditCondition = "bPropertySetEnabled == true"))
+	FQuat Rotation;
 };
 
 UCLASS()
@@ -181,7 +228,7 @@ class MESHMODELINGTOOLS_API UBrushRemeshProperties : public URemeshProperties
 public:
 	void SaveRestoreProperties(UInteractiveTool* RestoreToTool, bool bSaving) override;
 
-	/** Target Relative Triangle Sizefor Dynamic Meshing */
+	/** Desired size of triangles after Remeshing, relative to average initial triangle size. Larger value results in larger triangles.  */
 	UPROPERTY(EditAnywhere, Category = Remeshing, meta = (UIMin = "0.5", UIMax = "2.0", ClampMin = "0.1", ClampMax = "100.0"))
 	float RelativeSize = 1.0f;
 };
@@ -224,29 +271,34 @@ public:
 	virtual void OnPropertyModified(UObject* PropertySet, FProperty* Property) override;
 
 public:
+	/** Properties that control brush size/etc*/
+	UPROPERTY()
+	USculptBrushProperties* BrushProperties;
+
 	/** Properties that control sculpting*/
 	UPROPERTY()
 	UBrushSculptProperties* SculptProperties;
 
-	/** Properties that control brush size/etc*/
 	UPROPERTY()
-	UBrushBaseProperties* BrushProperties;
+	UPlaneBrushProperties* PlaneBrushProperties;
 
 	/** Properties that control dynamic remeshing */
 	UPROPERTY()
 	UBrushRemeshProperties* RemeshProperties;
 
 	UPROPERTY()
-	UMeshEditingViewProperties* ViewProperties;
+	UFixedPlaneBrushProperties* GizmoProperties;
 
 	UPROPERTY()
-	UFixedPlaneBrushProperties* GizmoProperties;
+	UMeshEditingViewProperties* ViewProperties;
 
 public:
 	virtual void IncreaseBrushRadiusAction();
 	virtual void DecreaseBrushRadiusAction();
 	virtual void IncreaseBrushRadiusSmallStepAction();
 	virtual void DecreaseBrushRadiusSmallStepAction();
+	virtual void IncreaseBrushFalloffAction();
+	virtual void DecreaseBrushFalloffAction();
 
 	virtual void IncreaseBrushSpeedAction();
 	virtual void DecreaseBrushSpeedAction();
@@ -259,6 +311,7 @@ public:
 
 private:
 	UWorld* TargetWorld;		// required to spawn UPreviewMesh/etc
+	FViewCameraState CameraState;
 
 	UPROPERTY()
 	UBrushStampIndicator* BrushIndicator;
@@ -269,6 +322,9 @@ private:
 	UPROPERTY()
 	UOctreeDynamicMeshComponent* DynamicMeshComponent;
 
+	UPROPERTY()
+	UMaterialInstanceDynamic* ActiveOverrideMaterial;
+
 	FTransform3d InitialTargetTransform;
 	FTransform3d CurTargetTransform;
 
@@ -278,7 +334,18 @@ private:
 
 	TValueWatcher<bool> ShowWireframeWatcher;
 	TValueWatcher<EMeshEditingMaterialModes> MaterialModeWatcher;
+	TValueWatcher<bool> FlatShadingWatcher;
+	TValueWatcher<FLinearColor> ColorWatcher;
+	TValueWatcher<UTexture2D*> ImageWatcher;
+	TValueWatcher<EDynamicMeshSculptBrushType> BrushTypeWatcher;
+	TValueWatcher<FVector> GizmoPositionWatcher;
+	TValueWatcher<FQuat> GizmoRotationWatcher;
 	void UpdateMaterialMode(EMeshEditingMaterialModes NewMode);
+	void UpdateFlatShadingSetting(bool bNewValue);
+	void UpdateColorSetting(FLinearColor NewColor);
+	void UpdateImageSetting(UTexture2D* NewImage);
+	void UpdateBrushType(EDynamicMeshSculptBrushType BrushType);
+	void UpdateGizmoFromProperties();
 
 	FInterval1d BrushRelativeSizeRange;
 	double CurrentBrushRadius;
@@ -332,13 +399,15 @@ private:
 	int FindHitSculptMeshTriangle(const FRay3d& LocalRay);
 	int FindHitTargetMeshTriangle(const FRay3d& LocalRay);
 
+	bool UpdateBrushPosition(const FRay& WorldRay);
 	bool UpdateBrushPositionOnActivePlane(const FRay& WorldRay);
 	bool UpdateBrushPositionOnTargetMesh(const FRay& WorldRay);
 	bool UpdateBrushPositionOnSculptMesh(const FRay& WorldRay);
+	void AlignBrushToView();
 
 	void ApplySmoothBrush(const FRay& WorldRay);
 	void ApplyMoveBrush(const FRay& WorldRay);
-	void ApplyOffsetBrush(const FRay& WorldRay);
+	void ApplyOffsetBrush(const FRay& WorldRay, bool bUseViewDirection);
 	void ApplySculptMaxBrush(const FRay& WorldRay);
 	void ApplyPinchBrush(const FRay& WorldRay);
 	void ApplyInflateBrush(const FRay& WorldRay);
@@ -350,7 +419,7 @@ private:
 	TArray<FVector3d> ROIPositionBuffer;
 
 	FFrame3d ActiveFixedBrushPlane;
-	FFrame3d ComputeROIBrushPlane(const FVector3d& BrushCenter, bool bIgnoreDepth);
+	FFrame3d ComputeROIBrushPlane(const FVector3d& BrushCenter, bool bIgnoreDepth, bool bViewAligned);
 
 	TArray<int> TrianglesBuffer;
 	TArray<int> NormalsBuffer;
@@ -391,17 +460,19 @@ private:
 	UPROPERTY()
 	UTransformProxy* PlaneTransformProxy;
 
-	/** Orientation of plane we will draw polygon on */
-	UPROPERTY()
-	FQuat DrawPlaneOrientation;
-
 	void PlaneTransformChanged(UTransformProxy* Proxy, FTransform Transform);
 
-	bool bPendingSetFixedPlanePosition = false;
-	void SetFixedSculptPlaneFromWorldPos(const FVector& Position);
-
+	enum class EPendingWorkPlaneUpdate
+	{
+		NoUpdatePending,
+		MoveToHitPositionNormal,
+		MoveToHitPosition,
+		MoveToHitPositionViewAligned
+	};
+	EPendingWorkPlaneUpdate PendingWorkPlaneUpdate;
+	void SetFixedSculptPlaneFromWorldPos(const FVector& Position, const FVector& Normal, EPendingWorkPlaneUpdate UpdateType);
 	void UpdateFixedSculptPlanePosition(const FVector& Position);
-
+	void UpdateFixedSculptPlaneRotation(const FQuat& Rotation);
 	void UpdateFixedPlaneGizmoVisibility(bool bVisible);
 };
 
diff --git a/Engine/Plugins/Experimental/MeshModelingToolset/Source/MeshModelingTools/Public/Properties/MeshMaterialProperties.h b/Engine/Plugins/Experimental/MeshModelingToolset/Source/MeshModelingTools/Public/Properties/MeshMaterialProperties.h
--- a/Engine/Plugins/Experimental/MeshModelingToolset/Source/MeshModelingTools/Public/Properties/MeshMaterialProperties.h
+++ b/Engine/Plugins/Experimental/MeshModelingToolset/Source/MeshModelingTools/Public/Properties/MeshMaterialProperties.h
@@ -3,7 +3,7 @@
 #pragma once
 
 #include "CoreMinimal.h"
-
+#include "Engine/Classes/Engine/Texture2D.h"
 #include "InteractiveToolBuilder.h"
 
 #include "MeshMaterialProperties.generated.h"
@@ -107,7 +107,11 @@ UENUM()
 enum class EMeshEditingMaterialModes
 {
 	ExistingMaterial = 0,
-	MeshFocusMaterial = 1
+	Diffuse = 1,
+	Grey = 2,
+	Soft = 3,
+	TangentNormal = 4,
+	Custom = 5
 };
 
 
@@ -118,12 +122,24 @@ class MESHMODELINGTOOLS_API UMeshEditingViewProperties : public UInteractiveTool
 
 public:
 	/** Toggle drawing of wireframe overlay on/off [Alt+W] */
-	UPROPERTY(EditAnywhere, Category = ViewOptions)
+	UPROPERTY(EditAnywhere, Category = Rendering)
 	bool bShowWireframe = false;
 
 	/** Set which material to use on object */
-	UPROPERTY(EditAnywhere, Category = ViewOptions)
-	EMeshEditingMaterialModes MaterialMode = EMeshEditingMaterialModes::MeshFocusMaterial;
+	UPROPERTY(EditAnywhere, Category = Rendering)
+	EMeshEditingMaterialModes MaterialMode = EMeshEditingMaterialModes::Diffuse;
+
+	/** Toggle flat shading on/off */
+	UPROPERTY(EditAnywhere, Category = Rendering, meta = (EditConditionHides, EditCondition = "MaterialMode != EMeshEditingMaterialModes::ExistingMaterial") )
+	bool bFlatShading = true;
+
+	/** Main Color of Material */
+	UPROPERTY(EditAnywhere, Category = Rendering, meta = (EditConditionHides, EditCondition = "MaterialMode == EMeshEditingMaterialModes::Diffuse"))
+	FLinearColor Color = FLinearColor(0.4f, 0.4f, 0.4f);
+
+	/** Image used in Image-Based Material */
+	UPROPERTY(EditAnywhere, Category = Rendering, meta = (EditConditionHides, EditCondition = "MaterialMode == EMeshEditingMaterialModes::Custom") )
+	UTexture2D* Image;
 
 	virtual void SaveProperties(UInteractiveTool* SaveFromTool) override;
 	virtual void RestoreProperties(UInteractiveTool* RestoreToTool) override;
diff --git a/Engine/Plugins/Experimental/MeshModelingToolset/Source/MeshModelingTools/Public/Properties/RemeshProperties.h b/Engine/Plugins/Experimental/MeshModelingToolset/Source/MeshModelingTools/Public/Properties/RemeshProperties.h
--- a/Engine/Plugins/Experimental/MeshModelingToolset/Source/MeshModelingTools/Public/Properties/RemeshProperties.h
+++ b/Engine/Plugins/Experimental/MeshModelingToolset/Source/MeshModelingTools/Public/Properties/RemeshProperties.h
@@ -69,19 +69,19 @@ class MESHMODELINGTOOLS_API URemeshProperties : public UMeshConstraintProperties
 	GENERATED_BODY()
 
 public:
-	/** Smoothing speed for dynamic meshing */
-	UPROPERTY(EditAnywhere, Category = Remeshing, meta = (UIMin = "0.0", UIMax = "1.0", ClampMin = "0.0", ClampMax = "1.0"))
+	/** Amount of Vertex Smoothing applied within Remeshing */
+	UPROPERTY(EditAnywhere, Category = Remeshing, meta = (DisplayName = "Smoothing Rate", UIMin = "0.0", UIMax = "1.0", ClampMin = "0.0", ClampMax = "1.0"))
 	float SmoothingStrength = 0.1f;
 
-	/** Enable edge flips */
+	/** Enable edge flips in Remeshing (tends to lose edges/detail) */
 	UPROPERTY(EditAnywhere, Category = "Remeshing|Edge Constraints", meta = (DisplayName = "Allow Flips"))
 	bool bFlips = false;
 
-	/** Enable edge splits */
+	/** Enable edge splits in Remeshing */
 	UPROPERTY(EditAnywhere, Category = "Remeshing|Edge Constraints", meta = (DisplayName = "Allow Splits"))
 	bool bSplits = true;
 
-	/** Enable edge collapses */
+	/** Enable edge collapses in Remeshing */
 	UPROPERTY(EditAnywhere, Category = "Remeshing|Edge Constraints", meta = (DisplayName = "Allow Collapses"))
 	bool bCollapses = true;
 
diff --git a/Engine/Plugins/Experimental/MeshModelingToolset/Source/ModelingComponents/Private/ToolSetupUtil.cpp b/Engine/Plugins/Experimental/MeshModelingToolset/Source/ModelingComponents/Private/ToolSetupUtil.cpp
--- a/Engine/Plugins/Experimental/MeshModelingToolset/Source/ModelingComponents/Private/ToolSetupUtil.cpp
+++ b/Engine/Plugins/Experimental/MeshModelingToolset/Source/ModelingComponents/Private/ToolSetupUtil.cpp
@@ -42,7 +42,7 @@ UMaterialInterface* ToolSetupUtil::GetDefaultBrushVolumeMaterial(UInteractiveToo
 
 
 
-UMaterialInterface* ToolSetupUtil::GetSculptMaterial1(UInteractiveToolManager* ToolManager)
+UMaterialInterface* ToolSetupUtil::GetDefaultSculptMaterial(UInteractiveToolManager* ToolManager)
 {
 	UMaterialInterface* Material = LoadObject<UMaterial>(nullptr, TEXT("/MeshModelingToolset/Materials/SculptMaterial"));
 	if (Material == nullptr && ToolManager != nullptr)
@@ -53,6 +53,51 @@ UMaterialInterface* ToolSetupUtil::GetSculptMaterial1(UInteractiveToolManager* T
 }
 
 
+
+
+UMaterialInterface* ToolSetupUtil::GetImageBasedSculptMaterial(UInteractiveToolManager* ToolManager, ImageMaterialType Type)
+{
+	UMaterialInterface* Material = nullptr;
+	if (Type == ImageMaterialType::DefaultBasic)
+	{
+		Material = LoadObject<UMaterialInstance>(nullptr, TEXT("/MeshModelingToolset/Materials/SculptMaterial_Basic"));
+	}
+	else if (Type == ImageMaterialType::DefaultSoft)
+	{
+		Material = LoadObject<UMaterialInstance>(nullptr, TEXT("/MeshModelingToolset/Materials/SculptMaterial_Soft"));
+	}
+	else if (Type == ImageMaterialType::TangentNormalFromView)
+	{
+		Material = LoadObject<UMaterial>(nullptr, TEXT("/MeshModelingToolset/Materials/SculptMaterial_TangentNormalFromView"));
+	}
+
+	if (Material == nullptr && ToolManager != nullptr)
+	{
+		Material = GetDefaultSculptMaterial(ToolManager);
+	}
+
+	return Material;
+}
+
+
+
+UMaterialInstanceDynamic* ToolSetupUtil::GetCustomImageBasedSculptMaterial(UInteractiveToolManager* ToolManager, UTexture* SetImage)
+{
+	UMaterial* Material = LoadObject<UMaterial>(nullptr, TEXT("/MeshModelingToolset/Materials/ImageBasedMaterial_Master"));
+	if (Material != nullptr)
+	{
+		UMaterialInstanceDynamic* MatInstance = UMaterialInstanceDynamic::Create(Material, ToolManager);
+		if (SetImage != nullptr)
+		{
+			MatInstance->SetTextureParameterValue(TEXT("ImageTexture"), SetImage);
+		}
+		return MatInstance;
+	}
+	return nullptr;
+}
+
+
+
 UMaterialInterface* ToolSetupUtil::GetSelectionMaterial(UInteractiveToolManager* ToolManager)
 {
 	UMaterialInterface* Material = LoadObject<UMaterial>(nullptr, TEXT("/MeshModelingToolset/Materials/SelectionMaterial"));
diff --git a/Engine/Plugins/Experimental/MeshModelingToolset/Source/ModelingComponents/Public/ToolSetupUtil.h b/Engine/Plugins/Experimental/MeshModelingToolset/Source/ModelingComponents/Public/ToolSetupUtil.h
--- a/Engine/Plugins/Experimental/MeshModelingToolset/Source/ModelingComponents/Public/ToolSetupUtil.h
+++ b/Engine/Plugins/Experimental/MeshModelingToolset/Source/ModelingComponents/Public/ToolSetupUtil.h
@@ -5,6 +5,8 @@
 #include "CoreMinimal.h"
 
 class UMaterialInterface;
+class UMaterialInstanceDynamic;
+class UTexture;
 class UInteractiveToolManager;
 
 /**
@@ -33,7 +35,26 @@ namespace ToolSetupUtil
 	/**
 	 * @return Sculpt Material 1
 	 */
-	MODELINGCOMPONENTS_API UMaterialInterface* GetSculptMaterial1(UInteractiveToolManager* ToolManager);
+	MODELINGCOMPONENTS_API UMaterialInterface* GetDefaultSculptMaterial(UInteractiveToolManager* ToolManager);
+
+
+	/** Types of image-based material that we can create */
+	enum class ImageMaterialType
+	{
+		DefaultBasic,
+		DefaultSoft,
+		TangentNormalFromView
+	};
+
+	/**
+	 * @return Image-based sculpt material instance, based ImageMaterialType
+	 */
+	MODELINGCOMPONENTS_API UMaterialInterface* GetImageBasedSculptMaterial(UInteractiveToolManager* ToolManager, ImageMaterialType Type);
+
+	/**
+	 * @return Image-based sculpt material that supports changing the image
+	 */
+	MODELINGCOMPONENTS_API UMaterialInstanceDynamic* GetCustomImageBasedSculptMaterial(UInteractiveToolManager* ToolManager, UTexture* SetImage);
 
 
 	/**
diff --git a/Engine/Source/Runtime/Experimental/InteractiveToolsFramework/Private/BaseGizmos/BrushStampIndicator.cpp b/Engine/Source/Runtime/Experimental/InteractiveToolsFramework/Private/BaseGizmos/BrushStampIndicator.cpp
--- a/Engine/Source/Runtime/Experimental/InteractiveToolsFramework/Private/BaseGizmos/BrushStampIndicator.cpp
+++ b/Engine/Source/Runtime/Experimental/InteractiveToolsFramework/Private/BaseGizmos/BrushStampIndicator.cpp
@@ -29,7 +29,14 @@ void UBrushStampIndicator::Render(IToolsContextRenderAPI* RenderAPI)
 		FToolDataVisualizer Draw;
 		Draw.BeginFrame(RenderAPI);
 
-		Draw.DrawCircle(BrushPosition, BrushNormal, BrushRadius, SampleStepCount, LineColor, LineThickness, bDepthTested);
+		if (bDrawRadiusCircle)
+		{
+			Draw.DrawCircle(BrushPosition, BrushNormal, BrushRadius, SampleStepCount, LineColor, LineThickness, bDepthTested);
+		}
+		if (bDrawFalloffCircle)
+		{
+			Draw.DrawCircle(BrushPosition, BrushNormal, BrushFalloff, SampleStepCount, LineColor, LineThickness, bDepthTested);
+		}
 
 		if (bDrawSecondaryLines)
 		{
diff --git a/Engine/Source/Runtime/Experimental/InteractiveToolsFramework/Private/BaseTools/BaseBrushTool.cpp b/Engine/Source/Runtime/Experimental/InteractiveToolsFramework/Private/BaseTools/BaseBrushTool.cpp
--- a/Engine/Source/Runtime/Experimental/InteractiveToolsFramework/Private/BaseTools/BaseBrushTool.cpp
+++ b/Engine/Source/Runtime/Experimental/InteractiveToolsFramework/Private/BaseTools/BaseBrushTool.cpp
@@ -21,27 +21,16 @@ UBrushBaseProperties::UBrushBaseProperties()
 	BrushFalloffAmount = 1.0f;
 }
 
-void UBrushBaseProperties::SaveProperties(UInteractiveTool* SaveFromTool)
+void UBrushBaseProperties::SaveRestoreProperties(UInteractiveTool* Tool, bool bSaving)
 {
 	UBrushBaseProperties* PropertyCache = GetPropertyCache<UBrushBaseProperties>();
-	PropertyCache->BrushSize = this->BrushSize;
-	PropertyCache->bSpecifyRadius = this->bSpecifyRadius;
-	PropertyCache->BrushRadius = this->BrushRadius;
-	PropertyCache->BrushFalloffAmount = this->BrushFalloffAmount;
-	PropertyCache->BrushStrength = this->BrushStrength;
+	SaveRestoreProperty(PropertyCache->BrushSize, this->BrushSize, bSaving);
+	SaveRestoreProperty(PropertyCache->bSpecifyRadius, this->bSpecifyRadius, bSaving);
+	SaveRestoreProperty(PropertyCache->BrushRadius, this->BrushRadius, bSaving);
+	SaveRestoreProperty(PropertyCache->BrushFalloffAmount, this->BrushFalloffAmount, bSaving);
+	SaveRestoreProperty(PropertyCache->BrushStrength, this->BrushStrength, bSaving);
 }
 
-void UBrushBaseProperties::RestoreProperties(UInteractiveTool* RestoreToTool)
-{
-	UBrushBaseProperties* PropertyCache = GetPropertyCache<UBrushBaseProperties>();
-	this->BrushSize = PropertyCache->BrushSize;
-	this->bSpecifyRadius = PropertyCache->bSpecifyRadius;
-	this->BrushRadius = PropertyCache->BrushRadius;
-	this->BrushFalloffAmount = PropertyCache->BrushFalloffAmount;
-	this->BrushStrength = PropertyCache->BrushStrength;
-}
-
-
 
 
 UBaseBrushTool::UBaseBrushTool()
diff --git a/Engine/Source/Runtime/Experimental/InteractiveToolsFramework/Public/BaseGizmos/BrushStampIndicator.h b/Engine/Source/Runtime/Experimental/InteractiveToolsFramework/Public/BaseGizmos/BrushStampIndicator.h
--- a/Engine/Source/Runtime/Experimental/InteractiveToolsFramework/Public/BaseGizmos/BrushStampIndicator.h
+++ b/Engine/Source/Runtime/Experimental/InteractiveToolsFramework/Public/BaseGizmos/BrushStampIndicator.h
@@ -61,6 +61,11 @@ public:
 	UPROPERTY()
 	bool bDrawIndicatorLines = true;
 
+	UPROPERTY()
+	bool bDrawRadiusCircle = true;
+
+	UPROPERTY()
+	bool bDrawFalloffCircle = true;
 
 	UPROPERTY()
 	int SampleStepCount = 32;
diff --git a/Engine/Source/Runtime/Experimental/InteractiveToolsFramework/Public/BaseTools/BaseBrushTool.h b/Engine/Source/Runtime/Experimental/InteractiveToolsFramework/Public/BaseTools/BaseBrushTool.h
--- a/Engine/Source/Runtime/Experimental/InteractiveToolsFramework/Public/BaseTools/BaseBrushTool.h
+++ b/Engine/Source/Runtime/Experimental/InteractiveToolsFramework/Public/BaseTools/BaseBrushTool.h
@@ -20,7 +20,7 @@ public:
 	UBrushBaseProperties();
 
 	/** Relative size of brush */
-	UPROPERTY(EditAnywhere, Category = Brush, meta = (DisplayName = "Size", UIMin = "0.0", UIMax = "1.0", ClampMin = "0.0", ClampMax = "10.0"))
+	UPROPERTY(EditAnywhere, Category = Brush, meta = (DisplayName = "Size", UIMin = "0.0", UIMax = "1.0", ClampMin = "0.0", ClampMax = "10.0", DisplayPriority = 1))
 	float BrushSize;
 
 	/** If true, ignore relative Brush Size and use explicit World Radius */
@@ -32,22 +32,25 @@ public:
 	float BrushRadius;
 
 	/** Strength of the brush (0.0 - 1.0) */
-	UPROPERTY(EditAnywhere, Category = Brush, meta = (DisplayName = "Strength", UIMin = "0.0", UIMax = "1.0", ClampMin = "0.0", ClampMax = "1.0", HideEditConditionToggle, EditConditionHides, EditCondition = "bShowFullSettings"))
+	UPROPERTY(EditAnywhere, Category = Brush, meta = (DisplayName = "Strength", UIMin = "0.0", UIMax = "1.0", ClampMin = "0.0", ClampMax = "1.0", HideEditConditionToggle, EditConditionHides, EditCondition = "bShowStrength", DisplayPriority = 2))
 	float BrushStrength;
 
 	/** Amount of falloff to apply (0.0 - 1.0) */
-	UPROPERTY(EditAnywhere, Category = Brush, meta = (DisplayName = "Falloff", UIMin = "0.0", UIMax = "1.0", ClampMin = "0.0", ClampMax = "1.0", HideEditConditionToggle, EditConditionHides, EditCondition="bShowFullSettings"))
+	UPROPERTY(EditAnywhere, Category = Brush, meta = (DisplayName = "Falloff", UIMin = "0.0", UIMax = "1.0", ClampMin = "0.0", ClampMax = "1.0", HideEditConditionToggle, EditConditionHides, EditCondition="bShowFalloff", DisplayPriority = 3))
 	float BrushFalloffAmount;
 
-	/** If false, then BrushStrength and BrushFalloffAmount will not be shown in DetailsView panels (otherwise no effect)*/
+	/** If false, then BrushStrength will not be shown in DetailsView panels (otherwise no effect) */
 	UPROPERTY()
-	bool bShowFullSettings = true;
+	bool bShowStrength = true;
+
+	/** If false, then BrushFalloffAmount will not be shown in DetailsView panels (otherwise no effect) */
+	UPROPERTY()
+	bool bShowFalloff = true;
 
 	//
 	// save/restore support
 	//
-	virtual void SaveProperties(UInteractiveTool* SaveFromTool) override;
-	virtual void RestoreProperties(UInteractiveTool* RestoreToTool) override;
+	virtual void SaveRestoreProperties(UInteractiveTool* Tool, bool bSaving) override;
 };
 
 
