fix crashes on edge cases for mesh poly cut
#rnx

diff --git a/Engine/Plugins/Experimental/GeometryProcessing/Source/DynamicMesh/Private/Operations/EmbedSurfacePath.cpp b/Engine/Plugins/Experimental/GeometryProcessing/Source/DynamicMesh/Private/Operations/EmbedSurfacePath.cpp
--- a/Engine/Plugins/Experimental/GeometryProcessing/Source/DynamicMesh/Private/Operations/EmbedSurfacePath.cpp
+++ b/Engine/Plugins/Experimental/GeometryProcessing/Source/DynamicMesh/Private/Operations/EmbedSurfacePath.cpp
@@ -778,6 +778,10 @@ bool EmbedProjectedPath(FDynamicMesh3* Mesh, int StartTriID, FFrame3d Frame, con
 		{
 			WalkDir /= WalkLen;
 			FVector3d WalkNormal(-WalkDir.Y, WalkDir.X, 0);
+			if (!ensureMsgf(Mesh->IsTriangle(CurrentSeedTriID), TEXT("Invalid triangle somehow passed as seed for mesh path embedding: %d"), CurrentSeedTriID))
+			{
+				return false;
+			}
 			bool bWalkSuccess = SurfacePath.AddViaPlanarWalk(CurrentSeedTriID, StartPos, -1, LastVert, FVector3d(Path2D[IdxB].X, Path2D[IdxB].Y, 0), WalkNormal, ProjectToFrame, false, FMathf::ZeroTolerance, PtSnapVertexOrEdgeThresholdSq);
 			if (!bWalkSuccess)
 			{
diff --git a/Engine/Plugins/Experimental/MeshModelingToolset/Source/ModelingOperatorsEditorOnly/Private/CuttingOps/EmbedPolygonsOp.cpp b/Engine/Plugins/Experimental/MeshModelingToolset/Source/ModelingOperatorsEditorOnly/Private/CuttingOps/EmbedPolygonsOp.cpp
--- a/Engine/Plugins/Experimental/MeshModelingToolset/Source/ModelingOperatorsEditorOnly/Private/CuttingOps/EmbedPolygonsOp.cpp
+++ b/Engine/Plugins/Experimental/MeshModelingToolset/Source/ModelingOperatorsEditorOnly/Private/CuttingOps/EmbedPolygonsOp.cpp
@@ -156,6 +156,10 @@ void FEmbedPolygonsOp::CalculateResult(FProgressCancel* Progress)
 
 	auto CutHole = [](FDynamicMesh3& Mesh, const FFrame3d& F, int TriStart, const FPolygon2d& PolygonArg, TArray<int>& PathVertIDs, TArray<int>& PathVertCorrespond, bool bCollapseDegenerateEdges)
 	{
+		if (!Mesh.IsTriangle(TriStart))
+		{
+			return false;
+		}
 		FMeshFaceSelection Selection(&Mesh);
 		bool bDidEmbed = EmbedProjectedPath(&Mesh, TriStart, F, PolygonArg.GetVertices(), PathVertIDs, PathVertCorrespond, true, &Selection);
 		if (!bDidEmbed)
@@ -222,7 +226,7 @@ void FEmbedPolygonsOp::CalculateResult(FProgressCancel* Progress)
 		FDynamicMeshEditor MeshEditor(ResultMesh.Get());
 		FDynamicMeshEditResult ResultOut;
 		bool bStitched = MeshEditor.StitchSparselyCorrespondedVertexLoops(PathVertIDs1, PathVertCorrespond1, PathVertIDs2, PathVertCorrespond2, ResultOut);
-		if (ResultMesh->HasAttributes())
+		if (bStitched && ResultMesh->HasAttributes())
 		{
 			MeshEditor.SetTubeNormals(ResultOut.NewTriangles, PathVertIDs1, PathVertCorrespond1, PathVertIDs2, PathVertCorrespond2);
 			TArray<float> UValues; UValues.SetNumUninitialized(PathVertCorrespond2.Num() + 1);
