ModelingTools: Further renaming/improving sculpting parameters. Add remeshing on/off toggle. Triangle RelativeSize is now an integer range. Add PreserveDetail remeshing setting. Add Resample brush type. Add TMathUtil::Pow.
#rnx

diff --git a/Engine/Plugins/Experimental/GeometryProcessing/Source/GeometricObjects/Public/MathUtil.h b/Engine/Plugins/Experimental/GeometryProcessing/Source/GeometricObjects/Public/MathUtil.h
--- a/Engine/Plugins/Experimental/GeometryProcessing/Source/GeometricObjects/Public/MathUtil.h
+++ b/Engine/Plugins/Experimental/GeometryProcessing/Source/GeometricObjects/Public/MathUtil.h
@@ -134,6 +134,7 @@ public:
 	static inline RealType ACos(const RealType Value);
 	static inline RealType Floor(const RealType Value);
 	static inline RealType Ceil(const RealType Value);
+	static inline RealType Pow(const RealType Value, const RealType Power);
 	static inline RealType Lerp(const RealType A, const RealType B, RealType Alpha);
 
 
@@ -255,6 +256,12 @@ RealType TMathUtil<RealType>::Ceil(const RealType Value)
 	return ceil(Value);
 }
 
+template<typename RealType>
+RealType TMathUtil<RealType>::Pow(const RealType Value, const RealType Power)
+{
+	return pow(Value, Power);
+}
+
 template<typename RealType>
 RealType TMathUtil<RealType>::Lerp(const RealType A, const RealType B, RealType Alpha)
 {
diff --git a/Engine/Plugins/Experimental/MeshModelingToolset/Source/MeshModelingTools/Private/DynamicMeshSculptTool.cpp b/Engine/Plugins/Experimental/MeshModelingToolset/Source/MeshModelingTools/Private/DynamicMeshSculptTool.cpp
--- a/Engine/Plugins/Experimental/MeshModelingToolset/Source/MeshModelingTools/Private/DynamicMeshSculptTool.cpp
+++ b/Engine/Plugins/Experimental/MeshModelingToolset/Source/MeshModelingTools/Private/DynamicMeshSculptTool.cpp
@@ -80,7 +80,7 @@ void UBrushSculptProperties::SaveRestoreProperties(UInteractiveTool* Tool, bool
 	SaveRestoreProperty(PropertyCache->PrimaryBrushSpeed, this->PrimaryBrushSpeed, bSaving);
 	SaveRestoreProperty(PropertyCache->SmoothBrushSpeed, this->SmoothBrushSpeed, bSaving);
 	SaveRestoreProperty(PropertyCache->bPreserveUVFlow, this->bPreserveUVFlow, bSaving);
-	SaveRestoreProperty(PropertyCache->bSmoothBrushRemeshes, this->bSmoothBrushRemeshes, bSaving);
+	SaveRestoreProperty(PropertyCache->bDetailPreservingSmooth, this->bDetailPreservingSmooth, bSaving);
 	SaveRestoreProperty(PropertyCache->bFreezeTarget, this->bFreezeTarget, bSaving);
 }
 
@@ -131,7 +131,9 @@ void UBrushRemeshProperties::SaveRestoreProperties(UInteractiveTool* RestoreToTo
 	SaveRestoreProperty(PropertyCache->bCollapses, this->bCollapses, bSaving);
 
 	//BrushRemeshProperties
-	SaveRestoreProperty(PropertyCache->RelativeSize, this->RelativeSize, bSaving);
+	SaveRestoreProperty(PropertyCache->TriangleSize, this->TriangleSize, bSaving);
+	SaveRestoreProperty(PropertyCache->PreserveDetail, this->PreserveDetail, bSaving);
+	SaveRestoreProperty(PropertyCache->bEnableRemeshing, this->bEnableRemeshing, bSaving);
 }
 
 /*
@@ -559,6 +561,9 @@ void UDynamicMeshSculptTool::ApplyStamp(const FRay& WorldRay)
 		case EDynamicMeshSculptBrushType::FixedPlane:
 			ApplyFixedPlaneBrush(WorldRay);
 			break;
+		case EDynamicMeshSculptBrushType::Resample:
+			ApplyResampleBrush(WorldRay);
+			break;
 		case EDynamicMeshSculptBrushType::LastValue:
 			break;
 	}
@@ -617,7 +622,7 @@ void UDynamicMeshSculptTool::ApplySmoothBrush(const FRay& WorldRay)
 		UpdateSavedVertex(VertIdx, OrigPos, NewPos);
 	}
 
-	bRemeshPending = bEnableRemeshing;
+	ScheduleRemeshPass();
 	LastBrushPosLocal = NewBrushPosLocal;
 }
 
@@ -662,7 +667,7 @@ void UDynamicMeshSculptTool::ApplyMoveBrush(const FRay& WorldRay)
 		UpdateSavedVertex(VertIdx, OrigPos, NewPos);
 	}
 
-	bRemeshPending = bEnableRemeshing;
+	ScheduleRemeshPass();
 	LastBrushPosLocal = NewBrushPosLocal;
 }
 
@@ -715,8 +720,7 @@ void UDynamicMeshSculptTool::ApplyOffsetBrush(const FRay& WorldRay, bool bUseVie
 		UpdateSavedVertex(VertIdx, OrigPos, NewPos);
 	}
 
-	bRemeshPending = bEnableRemeshing;
-
+	ScheduleRemeshPass();
 	LastBrushPosLocal = NewBrushPosLocal;
 }
 
@@ -774,8 +778,7 @@ void UDynamicMeshSculptTool::ApplySculptMaxBrush(const FRay& WorldRay)
 		UpdateSavedVertex(VertIdx, OrigPos, NewPos);
 	}
 
-	bRemeshPending = bEnableRemeshing;
-
+	ScheduleRemeshPass();
 	LastBrushPosLocal = NewBrushPosLocal;
 }
 
@@ -840,8 +843,7 @@ void UDynamicMeshSculptTool::ApplyPinchBrush(const FRay& WorldRay)
 	}
 
 
-	bRemeshPending = bEnableRemeshing;
-
+	ScheduleRemeshPass();
 	LastBrushPosLocal = NewBrushPosLocal;
 	LastSmoothBrushPosLocal = NewSmoothBrushPosLocal;
 }
@@ -923,8 +925,7 @@ void UDynamicMeshSculptTool::ApplyPlaneBrush(const FRay& WorldRay)
 		UpdateSavedVertex(VertIdx, OrigPos, NewPos);
 	}
 
-	bRemeshPending = bEnableRemeshing;
-
+	ScheduleRemeshPass();
 	LastBrushPosLocal = NewBrushPosLocal;
 }
 
@@ -979,8 +980,7 @@ void UDynamicMeshSculptTool::ApplyFixedPlaneBrush(const FRay& WorldRay)
 		UpdateSavedVertex(VertIdx, OrigPos, NewPos);
 	}
 
-	bRemeshPending = bEnableRemeshing;
-
+	ScheduleRemeshPass();
 	LastBrushPosLocal = NewBrushPosLocal;
 }
 
@@ -1026,8 +1026,7 @@ void UDynamicMeshSculptTool::ApplyFlattenBrush(const FRay& WorldRay)
 		UpdateSavedVertex(VertIdx, OrigPos, NewPos);
 	}
 
-	bRemeshPending = bEnableRemeshing;
-
+	ScheduleRemeshPass();
 	LastBrushPosLocal = NewBrushPosLocal;
 }
 
@@ -1079,11 +1078,60 @@ void UDynamicMeshSculptTool::ApplyInflateBrush(const FRay& WorldRay)
 		UpdateSavedVertex(VertIdx, OrigPos, NewPos);
 	}
 
-	bRemeshPending = bEnableRemeshing;
+	ScheduleRemeshPass();
+	LastBrushPosLocal = NewBrushPosLocal;
+}
 
+
+
+void UDynamicMeshSculptTool::ApplyResampleBrush(const FRay& WorldRay)
+{
+	if (UpdateBrushPositionOnTargetMesh(WorldRay) == false)
+	{
+		return;
+	}
+
+	FVector3d NewBrushPosLocal = CurTargetTransform.InverseTransformPosition(LastBrushPosWorld);
+	FVector3d LocalNormal = CurTargetTransform.InverseTransformNormal(LastBrushPosNormalWorld);
+
+	double UseSpeed = FMathd::Sqrt(CurrentBrushRadius) * (SculptProperties->PrimaryBrushSpeed) * ActivePressure;
+
+	FDynamicMesh3* Mesh = DynamicMeshComponent->GetMesh();
+	int NumV = VertexROI.Num();
+	ROIPositionBuffer.SetNum(NumV, false);
+	ParallelFor(NumV, [&](int k)
+	{
+		int VertIdx = VertexROI[k];
+		FVector3d OrigPos = Mesh->GetVertex(VertIdx);
+
+		FVector3d BasePos, BaseNormal;
+		if (GetTargetMeshNearest(OrigPos, (double)(4 * CurrentBrushRadius), BasePos, BaseNormal) == false)
+		{
+			ROIPositionBuffer[k] = OrigPos;
+		}
+		else
+		{
+			double Falloff = CalculateBrushFalloff(OrigPos.Distance(NewBrushPosLocal));
+			FVector3d NewPos = BasePos;// FVector3d::Lerp(OrigPos, BasePos, Falloff);
+			ROIPositionBuffer[k] = NewPos;
+		}
+	});
+
+	for (int k = 0; k < NumV; ++k)
+	{
+		int VertIdx = VertexROI[k];
+		const FVector3d& NewPos = ROIPositionBuffer[k];
+		FVector3d OrigPos = Mesh->GetVertex(VertIdx);
+		Mesh->SetVertex(VertIdx, NewPos);
+		UpdateSavedVertex(VertIdx, OrigPos, NewPos);
+	}
+
+	ScheduleRemeshPass();
 	LastBrushPosLocal = NewBrushPosLocal;
 }
 
+
+
 int UDynamicMeshSculptTool::FindHitSculptMeshTriangle(const FRay3d& LocalRay)
 {
 	if (BrushProperties->bHitBackFaces)
@@ -1212,6 +1260,7 @@ bool UDynamicMeshSculptTool::UpdateBrushPosition(const FRay& WorldRay)
 	case EDynamicMeshSculptBrushType::Offset:
 	case EDynamicMeshSculptBrushType::SculptMax:
 	case EDynamicMeshSculptBrushType::Pinch:
+	case EDynamicMeshSculptBrushType::Resample:
 		bHit = UpdateBrushPositionOnTargetMesh(WorldRay);
 		break;
 
@@ -1492,6 +1541,14 @@ void UDynamicMeshSculptTool::PrecomputeRemeshInfo()
 	}
 }
 
+void UDynamicMeshSculptTool::ScheduleRemeshPass()
+{
+	if (bEnableRemeshing && RemeshProperties != nullptr && RemeshProperties->bEnableRemeshing)
+	{
+		bRemeshPending = true;
+	}
+}
+
 /*
         Split	Collapse	Vertices Pinned	Flip
 Fixed	FALSE	FALSE		TRUE				FALSE
@@ -1500,16 +1557,22 @@ Free	TRUE	TRUE		FALSE				FALSE
 Ignore	TRUE	TRUE		FALSE				TRUE
 */
 
-void
-UDynamicMeshSculptTool::ConfigureRemesher(FSubRegionRemesher& Remesher)
+void UDynamicMeshSculptTool::ConfigureRemesher(FSubRegionRemesher& Remesher)
 {
 	FDynamicMesh3* Mesh = DynamicMeshComponent->GetMesh();
 	FDynamicMeshOctree3* Octree = DynamicMeshComponent->GetOctree();
 
-	double TargetEdgeLength = RemeshProperties->RelativeSize * InitialEdgeLength;
+	const double SizeRange = 5;
+	double LengthMultiplier = (RemeshProperties->TriangleSize >= 0) ?
+		FMathd::Lerp(1.0, 5.0, FMathd::Pow( (double)RemeshProperties->TriangleSize / SizeRange, 2.0) )
+		: FMathd::Lerp(0.25, 1.0, 1.0 - FMathd::Pow(FMathd::Abs((double)RemeshProperties->TriangleSize) / SizeRange, 2.0) );
+	double TargetEdgeLength = LengthMultiplier * InitialEdgeLength;
+
 	Remesher.SetTargetEdgeLength(TargetEdgeLength);
 
+	double DetailT = (double)(RemeshProperties->PreserveDetail) / 5.0;
 	double UseSmoothing = RemeshProperties->SmoothingStrength * 0.25;
+	UseSmoothing *= FMathd::Lerp(1.0, 0.25, DetailT);
 	Remesher.SmoothSpeedT = UseSmoothing;
 
 	// this is a temporary tweak for Pinch brush. Remesh params should be per-brush!
@@ -1524,13 +1587,27 @@ UDynamicMeshSculptTool::ConfigureRemesher(FSubRegionRemesher& Remesher)
 			return (1.0f - Falloff) * UseSmoothing;
 		};
 	}
-
-	// tweak remesh params for Smooth brush
-	if (bSmoothing && SculptProperties->bSmoothBrushRemeshes == false)
+	else if (bSmoothing && SculptProperties->bDetailPreservingSmooth)
 	{
-		Remesher.MaxEdgeLength = 2*InitialEdgeLength;
+		// this is the case where we don't want remeshing in smoothing
+		Remesher.MaxEdgeLength = 2 * InitialEdgeLength;
 		Remesher.MinEdgeLength = InitialEdgeLength * 0.01;
 	}
+	else
+	{
+		if (RemeshProperties->PreserveDetail > 0)
+		{
+			Remesher.MinEdgeLength *= FMathd::Lerp(1.0, 0.1, DetailT);
+			Remesher.CustomSmoothSpeedF = [this, UseSmoothing, DetailT](const FDynamicMesh3& Mesh, int vID)
+			{
+				FVector3d Pos = Mesh.GetVertex(vID);
+				double FalloffT = 1.0 - CalculateBrushFalloff(Pos.Distance((FVector3d)LastBrushPosLocal));
+				FalloffT = FMathd::Lerp(1.0, FalloffT, DetailT);
+				return FalloffT * UseSmoothing;
+			};
+		}
+	}
+
 
 	Remesher.SmoothType = (SculptProperties->bPreserveUVFlow) ?
 		FRemesher::ESmoothTypes::MeanValue : FRemesher::ESmoothTypes::Uniform;
diff --git a/Engine/Plugins/Experimental/MeshModelingToolset/Source/MeshModelingTools/Public/DynamicMeshSculptTool.h b/Engine/Plugins/Experimental/MeshModelingToolset/Source/MeshModelingTools/Public/DynamicMeshSculptTool.h
--- a/Engine/Plugins/Experimental/MeshModelingToolset/Source/MeshModelingTools/Public/DynamicMeshSculptTool.h
+++ b/Engine/Plugins/Experimental/MeshModelingToolset/Source/MeshModelingTools/Public/DynamicMeshSculptTool.h
@@ -85,6 +85,9 @@ enum class EDynamicMeshSculptBrushType : uint8
 	/** Move vertices towards a fixed plane in world space, positioned with a 3D gizmo */
 	FixedPlane UMETA(DisplayName = "FixedPlane"),
 
+	/** Remesh the brushed region but do not otherwise deform it */
+	Resample UMETA(DisplayName = "Resample"),
+
 	LastValue UMETA(Hidden)
 
 };
@@ -154,17 +157,16 @@ public:
 	bool bPreserveUVFlow = false;
 
 	/** When Freeze Target is toggled on, the Brush Target Surface will be Frozen in its current state, until toggled off. Brush strokes will be applied relative to the Target Surface, for applicable Brushes */
-	UPROPERTY(EditAnywhere, Category = Sculpting, meta = (EditCondition = "PrimaryBrushType == EDynamicMeshSculptBrushType::Sculpt || PrimaryBrushType == EDynamicMeshSculptBrushType::SculptMax || PrimaryBrushType == EDynamicMeshSculptBrushType::SculptView || PrimaryBrushType == EDynamicMeshSculptBrushType::Pinch" ))
+	UPROPERTY(EditAnywhere, Category = Sculpting, meta = (EditCondition = "PrimaryBrushType == EDynamicMeshSculptBrushType::Sculpt || PrimaryBrushType == EDynamicMeshSculptBrushType::SculptMax || PrimaryBrushType == EDynamicMeshSculptBrushType::SculptView || PrimaryBrushType == EDynamicMeshSculptBrushType::Pinch || PrimaryBrushType == EDynamicMeshSculptBrushType::Resample" ))
 	bool bFreezeTarget = false;
 
-
 	/** Strength of Shift-to-Smooth Brushing and Smoothing Brush */
 	UPROPERTY(EditAnywhere, Category = Smoothing, meta = (DisplayName = "Smoothing Strength", UIMin = "0.0", UIMax = "1.0", ClampMin = "0.0", ClampMax = "1.0"))
 	float SmoothBrushSpeed = 0.25;
 
-	/** If enabled, Remeshing is applied during Smoothing, which will wipe out fine details. Otherwise no remeshing is applied during Smoothing. */
-	UPROPERTY(EditAnywhere, Category = Smoothing, meta = (DisplayName = "Remesh in Smooth", EditConditionHides, HideEditConditionToggle, EditCondition = "bIsRemeshingEnabled"))
-	bool bSmoothBrushRemeshes = false;
+	/** If enabled, Remeshing is limited during Smoothing to avoid wiping out higher-density triangle areas */
+	UPROPERTY(EditAnywhere, Category = Smoothing, meta = (DisplayName = "Preserve Tri Density", EditConditionHides, HideEditConditionToggle, EditCondition = "bIsRemeshingEnabled"))
+	bool bDetailPreservingSmooth = true;
 };
 
 
@@ -228,9 +230,18 @@ class MESHMODELINGTOOLS_API UBrushRemeshProperties : public URemeshProperties
 public:
 	void SaveRestoreProperties(UInteractiveTool* RestoreToTool, bool bSaving) override;
 
-	/** Desired size of triangles after Remeshing, relative to average initial triangle size. Larger value results in larger triangles.  */
-	UPROPERTY(EditAnywhere, Category = Remeshing, meta = (UIMin = "0.5", UIMax = "2.0", ClampMin = "0.1", ClampMax = "100.0"))
-	float RelativeSize = 1.0f;
+	/** Toggle remeshing on/off */
+	UPROPERTY(EditAnywhere, Category = Remeshing, meta = (DisplayPriority = 1))
+	bool bEnableRemeshing = true;
+
+	// Note that if you change range here, you must also update UDynamicMeshSculptTool::ConfigureRemesher!
+	/** Desired size of triangles after Remeshing, relative to average initial triangle size. Larger value results in larger triangles. */
+	UPROPERTY(EditAnywhere, Category = Remeshing, meta = (DisplayName = "Relative Tri Size", UIMin = "-5", UIMax = "5", ClampMin = "-5", ClampMax = "5", DisplayPriority = 2))
+	int TriangleSize = 0;
+
+	/** Control the amount of simplification during sculpting. Higher values will avoid wiping out fine details on the mesh. */
+	UPROPERTY(EditAnywhere, Category = Remeshing, meta = (UIMin = "0", UIMax = "5", ClampMin = "0", ClampMax = "5", DisplayPriority = 3))
+	int PreserveDetail = 0;
 };
 
 /**
@@ -353,7 +364,7 @@ private:
 
 	bool bEnableRemeshing;
 	double InitialEdgeLength;
-
+	void ScheduleRemeshPass();
 	void ConfigureRemesher(FSubRegionRemesher&);
 
 	bool bInDrag;
@@ -414,6 +425,7 @@ private:
 	void ApplyPlaneBrush(const FRay& WorldRay);
 	void ApplyFixedPlaneBrush(const FRay& WorldRay);
 	void ApplyFlattenBrush(const FRay& WorldRay);
+	void ApplyResampleBrush(const FRay& WorldRay);
 
 	double CalculateBrushFalloff(double Distance);
 	TArray<FVector3d> ROIPositionBuffer;
