ModelingTools: Sculpt brushes no longer get "stuck" at edge of target mesh, now flow onto view-aligned plane at last hit point (except Sculpt-normal brush which will then "pull" the mesh with it). Add control over Max brush height.
#rnx

diff --git a/Engine/Plugins/Experimental/GeometryProcessing/Source/GeometricObjects/Public/RayTypes.h b/Engine/Plugins/Experimental/GeometryProcessing/Source/GeometricObjects/Public/RayTypes.h
--- a/Engine/Plugins/Experimental/GeometryProcessing/Source/GeometricObjects/Public/RayTypes.h
+++ b/Engine/Plugins/Experimental/GeometryProcessing/Source/GeometricObjects/Public/RayTypes.h
@@ -62,7 +62,7 @@ public:
 	}
 
 	/**
-	 * @return smallest squared distance from line to QueryPoint
+	 * @return smallest squared distance from ray to QueryPoint
 	 */
 	inline RealType DistanceSquared(const FVector3<RealType>& QueryPoint) const
 	{
@@ -78,6 +78,14 @@ public:
 		}
 	}
 
+	/**
+	 * @return smallest squared distance from ray to QueryPoint
+	 */
+	inline RealType Distance(const FVector3<RealType>& QueryPoint) const
+	{
+		return TMathUtil<RealType>::Sqrt(DistanceSquared(QueryPoint));
+	}
+
 	/**
 	 * @return nearest point on line to QueryPoint
 	 */
diff --git a/Engine/Plugins/Experimental/MeshModelingToolset/Source/MeshModelingTools/Private/DynamicMeshSculptTool.cpp b/Engine/Plugins/Experimental/MeshModelingToolset/Source/MeshModelingTools/Private/DynamicMeshSculptTool.cpp
--- a/Engine/Plugins/Experimental/MeshModelingToolset/Source/MeshModelingTools/Private/DynamicMeshSculptTool.cpp
+++ b/Engine/Plugins/Experimental/MeshModelingToolset/Source/MeshModelingTools/Private/DynamicMeshSculptTool.cpp
@@ -113,6 +113,13 @@ void UFixedPlaneBrushProperties::SaveRestoreProperties(UInteractiveTool* Tool, b
 }
 
 
+void USculptMaxBrushProperties::SaveRestoreProperties(UInteractiveTool* Tool, bool bSaving)
+{
+	USculptMaxBrushProperties* PropertyCache = GetPropertyCache<USculptMaxBrushProperties>();
+	SaveRestoreProperty(PropertyCache->MaxHeight, this->MaxHeight, bSaving);
+}
+
+
 void UBrushRemeshProperties::SaveRestoreProperties(UInteractiveTool* RestoreToTool, bool bSaving)
 {
 	UBrushRemeshProperties* PropertyCache = GetPropertyCache<UBrushRemeshProperties>();
@@ -239,10 +246,19 @@ void UDynamicMeshSculptTool::Setup()
 	AddToolPropertySource(BrushProperties);
 	AddToolPropertySource(SculptProperties);
 
+	// add brush-specific properties 
 	PlaneBrushProperties = NewObject<UPlaneBrushProperties>(this);
 	PlaneBrushProperties->RestoreProperties(this);
 	AddToolPropertySource(PlaneBrushProperties);
 
+	SculptMaxBrushProperties = NewObject<USculptMaxBrushProperties>();
+	SculptMaxBrushProperties->RestoreProperties(this);
+	AddToolPropertySource(SculptMaxBrushProperties);
+
+	GizmoProperties = NewObject<UFixedPlaneBrushProperties>();
+	GizmoProperties->RestoreProperties(this);
+	AddToolPropertySource(GizmoProperties);
+
 	if (this->bEnableRemeshing)
 	{
 		SculptProperties->bIsRemeshingEnabled = true;
@@ -253,13 +269,10 @@ void UDynamicMeshSculptTool::Setup()
 	CalculateBrushRadius();
 	SculptProperties->RestoreProperties(this);
 
-	GizmoProperties = NewObject<UFixedPlaneBrushProperties>();
-	GizmoProperties->RestoreProperties(this);
-	AddToolPropertySource(GizmoProperties);
-
 	// disable tool-specific properties
 	SetToolPropertySourceEnabled(PlaneBrushProperties, false);
 	SetToolPropertySourceEnabled(GizmoProperties, false);
+	SetToolPropertySourceEnabled(SculptMaxBrushProperties, false);
 
 	ViewProperties = NewObject<UMeshEditingViewProperties>();
 	ViewProperties->RestoreProperties(this);
@@ -356,6 +369,7 @@ void UDynamicMeshSculptTool::Shutdown(EToolShutdownType ShutdownType)
 	ViewProperties->SaveProperties(this);
 	GizmoProperties->SaveProperties(this);
 	PlaneBrushProperties->SaveProperties(this);
+	SculptMaxBrushProperties->SaveProperties(this);
 	RemeshProperties->SaveProperties(this);
 }
 
@@ -585,7 +599,7 @@ double UDynamicMeshSculptTool::CalculateBrushFalloff(double Distance)
 
 void UDynamicMeshSculptTool::ApplySmoothBrush(const FRay& WorldRay)
 {
-	bool bHit = UpdateBrushPositionOnSculptMesh(WorldRay);
+	bool bHit = UpdateBrushPositionOnSculptMesh(WorldRay, true);
 	if (bHit == false)
 	{
 		return;
@@ -673,10 +687,7 @@ void UDynamicMeshSculptTool::ApplyMoveBrush(const FRay& WorldRay)
 
 void UDynamicMeshSculptTool::ApplyOffsetBrush(const FRay& WorldRay, bool bUseViewDirection)
 {
-	if ( UpdateBrushPositionOnTargetMesh(WorldRay) == false )
-	{
-		return;
-	}
+	UpdateBrushPositionOnTargetMesh(WorldRay, bUseViewDirection);
 	if (bUseViewDirection)
 	{
 		AlignBrushToView();
@@ -687,6 +698,7 @@ void UDynamicMeshSculptTool::ApplyOffsetBrush(const FRay& WorldRay, bool bUseVie
 
 	double Direction = (bInvert) ? -1.0 : 1.0;
 	double UseSpeed = Direction * FMathd::Sqrt(CurrentBrushRadius) * (SculptProperties->PrimaryBrushSpeed) * ActivePressure;
+	double MaxOffset = CurrentBrushRadius;
 
 	FDynamicMesh3* Mesh = DynamicMeshComponent->GetMesh();
 	int NumV = VertexROI.Num();
@@ -726,17 +738,18 @@ void UDynamicMeshSculptTool::ApplyOffsetBrush(const FRay& WorldRay, bool bUseVie
 
 void UDynamicMeshSculptTool::ApplySculptMaxBrush(const FRay& WorldRay)
 {
-	bool bHit = UpdateBrushPositionOnTargetMesh(WorldRay);
-	if (bHit == false)
-	{
-		return;
-	}
-
+	UpdateBrushPositionOnTargetMesh(WorldRay, true);
 	FVector3d NewBrushPosLocal = CurTargetTransform.InverseTransformPosition(LastBrushPosWorld);
 
 	double Direction = (bInvert) ? -1.0 : 1.0;
 	double UseSpeed = Direction * FMathd::Sqrt(CurrentBrushRadius) * (SculptProperties->PrimaryBrushSpeed) * ActivePressure;
-	double MaxOffset = CurrentBrushRadius * 0.5;
+
+	double MaxOffset = CurrentBrushRadius * SculptMaxBrushProperties->MaxHeight;
+	if (SculptMaxBrushProperties->bFreezeCurrentHeight && SculptMaxFixedHeight >= 0)
+	{
+		MaxOffset = SculptMaxFixedHeight;
+	}
+	SculptMaxFixedHeight = MaxOffset;
 
 	FDynamicMesh3* Mesh = DynamicMeshComponent->GetMesh();
 	int NumV = VertexROI.Num();
@@ -784,12 +797,7 @@ void UDynamicMeshSculptTool::ApplySculptMaxBrush(const FRay& WorldRay)
 
 void UDynamicMeshSculptTool::ApplyPinchBrush(const FRay& WorldRay)
 {
-	bool bHit = UpdateBrushPositionOnTargetMesh(WorldRay);
-	if (bHit == false)
-	{
-		return;
-	}
-
+	UpdateBrushPositionOnTargetMesh(WorldRay, true);
 	FVector3d NewBrushPosLocal = CurTargetTransform.InverseTransformPosition(LastBrushPosWorld);
 	FVector3d BrushNormalLocal = CurTargetTransform.InverseTransformNormal(LastBrushPosNormalWorld);
 	FVector3d OffsetBrushPosLocal = NewBrushPosLocal - BrushProperties->Depth * CurrentBrushRadius * BrushNormalLocal;
@@ -882,7 +890,7 @@ FFrame3d UDynamicMeshSculptTool::ComputeROIBrushPlane(const FVector3d& BrushCent
 
 void UDynamicMeshSculptTool::ApplyPlaneBrush(const FRay& WorldRay)
 {
-	bool bHit = UpdateBrushPositionOnSculptMesh(WorldRay);
+	bool bHit = UpdateBrushPositionOnSculptMesh(WorldRay, true);
 	if (bHit == false)
 	{
 		return;
@@ -931,7 +939,7 @@ void UDynamicMeshSculptTool::ApplyPlaneBrush(const FRay& WorldRay)
 
 void UDynamicMeshSculptTool::ApplyFixedPlaneBrush(const FRay& WorldRay)
 {
-	bool bHit = UpdateBrushPositionOnSculptMesh(WorldRay);
+	bool bHit = UpdateBrushPositionOnSculptMesh(WorldRay, true);
 	if (bHit == false)
 	{
 		return;
@@ -986,12 +994,15 @@ void UDynamicMeshSculptTool::ApplyFixedPlaneBrush(const FRay& WorldRay)
 
 void UDynamicMeshSculptTool::ApplyFlattenBrush(const FRay& WorldRay)
 {
-	bool bHit = UpdateBrushPositionOnSculptMesh(WorldRay);
+	bool bHit = UpdateBrushPositionOnSculptMesh(WorldRay, true);
 	if (bHit == false)
 	{
 		return;
 	}
 
+	static const double PlaneSigns[3] = { 0, -1, 1 };
+	double PlaneSign = PlaneSigns[(int32)PlaneBrushProperties->WhichSide];
+
 	FVector3d NewBrushPosLocal = CurTargetTransform.InverseTransformPosition(LastBrushPosWorld);
 	FVector3d BrushNormalLocal = CurTargetTransform.InverseTransformNormal(LastBrushPosNormalWorld);
 
@@ -1003,17 +1014,24 @@ void UDynamicMeshSculptTool::ApplyFlattenBrush(const FRay& WorldRay)
 	int NumV = VertexROI.Num();
 	ROIPositionBuffer.SetNum(NumV, false);
 
-	ParallelFor(NumV, [this, Mesh, NewBrushPosLocal, UseSpeed, StampFlattenPlane](int k)
+	ParallelFor(NumV, [&](int k)
 	{
 		int VertIdx = VertexROI[k];
 		FVector3d OrigPos = Mesh->GetVertex(VertIdx);
 		FVector3d PlanePos = StampFlattenPlane.ToPlane(OrigPos, 2);
 		FVector3d Delta = PlanePos - OrigPos;
-		FVector3d MoveVec = UseSpeed * Delta;
 
-		double Falloff = CalculateBrushFalloff(OrigPos.Distance(NewBrushPosLocal));
+		double Dot = Delta.Dot(StampFlattenPlane.Z());
+		FVector3d NewPos = OrigPos;
+		if (Dot * PlaneSign >= 0)
+		{
+			double MaxDist = Delta.Normalize();
+			double Falloff = CalculateBrushFalloff(OrigPos.Distance(NewBrushPosLocal));
+			FVector3d MoveVec = Falloff * UseSpeed * Delta;
+			NewPos = (MoveVec.SquaredLength() > MaxDist*MaxDist) ?
+				PlanePos : OrigPos + Falloff * MoveVec;
+		}
 
-		FVector3d NewPos = OrigPos + Falloff * MoveVec;
 		ROIPositionBuffer[k] = NewPos;
 	});
 
@@ -1032,7 +1050,7 @@ void UDynamicMeshSculptTool::ApplyFlattenBrush(const FRay& WorldRay)
 
 void UDynamicMeshSculptTool::ApplyInflateBrush(const FRay& WorldRay)
 {
-	bool bHit = UpdateBrushPositionOnSculptMesh(WorldRay);
+	bool bHit = UpdateBrushPositionOnSculptMesh(WorldRay, true);
 	if (bHit == false)
 	{
 		return;
@@ -1086,11 +1104,7 @@ void UDynamicMeshSculptTool::ApplyInflateBrush(const FRay& WorldRay)
 
 void UDynamicMeshSculptTool::ApplyResampleBrush(const FRay& WorldRay)
 {
-	if (UpdateBrushPositionOnTargetMesh(WorldRay) == false)
-	{
-		return;
-	}
-
+	UpdateBrushPositionOnTargetMesh(WorldRay, true);
 	FVector3d NewBrushPosLocal = CurTargetTransform.InverseTransformPosition(LastBrushPosWorld);
 	FVector3d LocalNormal = CurTargetTransform.InverseTransformNormal(LastBrushPosNormalWorld);
 
@@ -1193,17 +1207,17 @@ bool UDynamicMeshSculptTool::UpdateBrushPositionOnActivePlane(const FRay& WorldR
 	return true;
 }
 
-bool UDynamicMeshSculptTool::UpdateBrushPositionOnTargetMesh(const FRay& WorldRay)
+bool UDynamicMeshSculptTool::UpdateBrushPositionOnTargetMesh(const FRay& WorldRay, bool bFallbackToViewPlane)
 {
 	FRay3d LocalRay(CurTargetTransform.InverseTransformPosition(WorldRay.Origin),
 		CurTargetTransform.InverseTransformVector(WorldRay.Direction));
 	LocalRay.Direction.Normalize();
 
+	const FDynamicMesh3* TargetMesh = BrushTargetMeshSpatial.GetMesh();
+
 	int HitTID = FindHitTargetMeshTriangle(LocalRay);
 	if (HitTID != IndexConstants::InvalidID)
 	{
-		const FDynamicMesh3* TargetMesh = BrushTargetMeshSpatial.GetMesh();
-
 		FTriangle3d Triangle;
 		TargetMesh->GetTriVertices(HitTID, Triangle.V[0], Triangle.V[1], Triangle.V[2]);
 		FIntrRay3Triangle3d Query(LocalRay, Triangle);
@@ -1213,10 +1227,21 @@ bool UDynamicMeshSculptTool::UpdateBrushPositionOnTargetMesh(const FRay& WorldRa
 		LastBrushPosWorld = CurTargetTransform.TransformPosition(LocalRay.PointAt(Query.RayParameter));
 		return true;
 	}
+
+	if (bFallbackToViewPlane)
+	{
+		FFrame3d BrushPlane(LastBrushPosWorld, CameraState.Forward());
+		FVector3d NewHitPosWorld;
+		BrushPlane.RayPlaneIntersection(WorldRay.Origin, WorldRay.Direction, 2, NewHitPosWorld);
+		LastBrushPosWorld = NewHitPosWorld;
+		LastBrushPosNormalWorld = ActiveDragPlane.Z();
+		return true;
+	}
+
 	return false;
 }
 
-bool UDynamicMeshSculptTool::UpdateBrushPositionOnSculptMesh(const FRay& WorldRay)
+bool UDynamicMeshSculptTool::UpdateBrushPositionOnSculptMesh(const FRay& WorldRay, bool bFallbackToViewPlane)
 {
 	FRay3d LocalRay(CurTargetTransform.InverseTransformPosition(WorldRay.Origin),
 		CurTargetTransform.InverseTransformVector(WorldRay.Direction));
@@ -1236,6 +1261,17 @@ bool UDynamicMeshSculptTool::UpdateBrushPositionOnSculptMesh(const FRay& WorldRa
 		LastBrushPosWorld = CurTargetTransform.TransformPosition(LocalRay.PointAt(Query.RayParameter));
 		return true;
 	}
+
+	if (bFallbackToViewPlane)
+	{
+		FFrame3d BrushPlane(LastBrushPosWorld, CameraState.Forward());
+		FVector3d NewHitPosWorld;
+		BrushPlane.RayPlaneIntersection(WorldRay.Origin, WorldRay.Direction, 2, NewHitPosWorld);
+		LastBrushPosWorld = NewHitPosWorld;
+		LastBrushPosNormalWorld = ActiveDragPlane.Z();
+		return true;
+	}
+
 	return false;
 }
 
@@ -1251,7 +1287,7 @@ bool UDynamicMeshSculptTool::UpdateBrushPosition(const FRay& WorldRay)
 
 	if (bSmoothing)
 	{
-		return UpdateBrushPositionOnSculptMesh(WorldRay);
+		return UpdateBrushPositionOnSculptMesh(WorldRay, false);
 	}
 
 	bool bHit = false;
@@ -1261,18 +1297,18 @@ bool UDynamicMeshSculptTool::UpdateBrushPosition(const FRay& WorldRay)
 	case EDynamicMeshSculptBrushType::SculptMax:
 	case EDynamicMeshSculptBrushType::Pinch:
 	case EDynamicMeshSculptBrushType::Resample:
-		bHit = UpdateBrushPositionOnTargetMesh(WorldRay);
+		bHit = UpdateBrushPositionOnTargetMesh(WorldRay, false);
 		break;
 
 	case EDynamicMeshSculptBrushType::SculptView:
 	case EDynamicMeshSculptBrushType::PlaneViewAligned:
-		bHit = UpdateBrushPositionOnTargetMesh(WorldRay);
+		bHit = UpdateBrushPositionOnTargetMesh(WorldRay, false);
 		AlignBrushToView();
 		break;
 
 	case EDynamicMeshSculptBrushType::Move:
 		//return UpdateBrushPositionOnActivePlane(WorldRay);
-		bHit = UpdateBrushPositionOnSculptMesh(WorldRay);
+		bHit = UpdateBrushPositionOnSculptMesh(WorldRay, false);
 		break;
 
 	case EDynamicMeshSculptBrushType::Smooth:
@@ -1280,12 +1316,12 @@ bool UDynamicMeshSculptTool::UpdateBrushPosition(const FRay& WorldRay)
 	case EDynamicMeshSculptBrushType::Flatten:
 	case EDynamicMeshSculptBrushType::Plane:
 	case EDynamicMeshSculptBrushType::FixedPlane:
-		bHit = UpdateBrushPositionOnSculptMesh(WorldRay);
+		bHit = UpdateBrushPositionOnSculptMesh(WorldRay, false);
 		break;
 
 	default:
 		UE_LOG(LogTemp, Warning, TEXT("UDynamicMeshSculptTool: unknown brush type in UpdateBrushPosition"));
-		bHit = UpdateBrushPositionOnSculptMesh(WorldRay);
+		bHit = UpdateBrushPositionOnSculptMesh(WorldRay, false);
 		break;
 	}
 
@@ -1309,6 +1345,12 @@ void UDynamicMeshSculptTool::OnEndDrag(const FRay& Ray)
 	EndChange();
 }
 
+
+FInputRayHit UDynamicMeshSculptTool::BeginHoverSequenceHitTest(const FInputDeviceRay& PressPos)
+{
+	return UMeshSurfacePointTool::BeginHoverSequenceHitTest(PressPos);
+}
+
 bool UDynamicMeshSculptTool::OnUpdateHover(const FInputDeviceRay& DevicePos)
 {
 	PendingStampType = SculptProperties->PrimaryBrushType;
@@ -1590,8 +1632,8 @@ void UDynamicMeshSculptTool::ConfigureRemesher(FSubRegionRemesher& Remesher)
 	else if (bSmoothing && SculptProperties->bDetailPreservingSmooth)
 	{
 		// this is the case where we don't want remeshing in smoothing
-		Remesher.MaxEdgeLength = 2 * InitialEdgeLength;
-		Remesher.MinEdgeLength = InitialEdgeLength * 0.01;
+		Remesher.MaxEdgeLength = 3 * InitialEdgeLength;
+		Remesher.MinEdgeLength = InitialEdgeLength * 0.05;
 	}
 	else
 	{
@@ -2255,6 +2297,7 @@ void UDynamicMeshSculptTool::UpdateBrushType(EDynamicMeshSculptBrushType BrushTy
 
 	SetToolPropertySourceEnabled(GizmoProperties, false);
 	SetToolPropertySourceEnabled(PlaneBrushProperties, false);
+	SetToolPropertySourceEnabled(SculptMaxBrushProperties, false);
 
 	if (BrushType == EDynamicMeshSculptBrushType::FixedPlane)
 	{
@@ -2262,10 +2305,14 @@ void UDynamicMeshSculptTool::UpdateBrushType(EDynamicMeshSculptBrushType BrushTy
 		SetToolPropertySourceEnabled(PlaneBrushProperties, true);
 		SetToolPropertySourceEnabled(GizmoProperties, true);
 	}
-	if (BrushType == EDynamicMeshSculptBrushType::Plane || BrushType == EDynamicMeshSculptBrushType::PlaneViewAligned)
+	if (BrushType == EDynamicMeshSculptBrushType::Plane || BrushType == EDynamicMeshSculptBrushType::PlaneViewAligned || BrushType == EDynamicMeshSculptBrushType::Flatten)
 	{
 		SetToolPropertySourceEnabled(PlaneBrushProperties, true);
 	}
+	if (BrushType == EDynamicMeshSculptBrushType::SculptMax)
+	{
+		SetToolPropertySourceEnabled(SculptMaxBrushProperties, true);
+	}
 
 	GetToolManager()->DisplayMessage(Builder.ToText(), EToolMessageLevel::UserNotification);
 }
diff --git a/Engine/Plugins/Experimental/MeshModelingToolset/Source/MeshModelingTools/Public/DynamicMeshSculptTool.h b/Engine/Plugins/Experimental/MeshModelingToolset/Source/MeshModelingTools/Public/DynamicMeshSculptTool.h
--- a/Engine/Plugins/Experimental/MeshModelingToolset/Source/MeshModelingTools/Public/DynamicMeshSculptTool.h
+++ b/Engine/Plugins/Experimental/MeshModelingToolset/Source/MeshModelingTools/Public/DynamicMeshSculptTool.h
@@ -222,6 +222,24 @@ public:
 	FQuat Rotation;
 };
 
+
+UCLASS()
+class MESHMODELINGTOOLS_API USculptMaxBrushProperties : public UInteractiveToolPropertySet
+{
+	GENERATED_BODY()
+public:
+	virtual void SaveRestoreProperties(UInteractiveTool* Tool, bool bSaving) override;
+
+	/** Specify maximum displacement height (relative to brush size) */
+	UPROPERTY(EditAnywhere, Category = SculptMaxBrush, meta = (UIMin = "0.0", UIMax = "1.0", ClampMin = "0.0", ClampMax = "1.0"))
+	float MaxHeight = 0.5;
+
+	/** Use maximum height from last brush stroke, regardless of brush size. Note that spatial brush falloff still applies.  */
+	UPROPERTY(EditAnywhere, Category = SculptMaxBrush)
+	bool bFreezeCurrentHeight = false;
+};
+
+
 UCLASS()
 class MESHMODELINGTOOLS_API UBrushRemeshProperties : public URemeshProperties
 {
@@ -273,6 +291,8 @@ public:
 	virtual void OnBeginDrag(const FRay& Ray) override;
 	virtual void OnUpdateDrag(const FRay& Ray) override;
 	virtual void OnEndDrag(const FRay& Ray) override;
+
+	virtual FInputRayHit BeginHoverSequenceHitTest(const FInputDeviceRay& PressPos) override;
 	virtual bool OnUpdateHover(const FInputDeviceRay& DevicePos) override;
 
 	virtual void SetEnableRemeshing(bool bEnable) { bEnableRemeshing = bEnable; }
@@ -293,6 +313,10 @@ public:
 	UPROPERTY()
 	UPlaneBrushProperties* PlaneBrushProperties;
 
+	UPROPERTY()
+	USculptMaxBrushProperties* SculptMaxBrushProperties;
+	
+
 	/** Properties that control dynamic remeshing */
 	UPROPERTY()
 	UBrushRemeshProperties* RemeshProperties;
@@ -412,8 +436,8 @@ private:
 
 	bool UpdateBrushPosition(const FRay& WorldRay);
 	bool UpdateBrushPositionOnActivePlane(const FRay& WorldRay);
-	bool UpdateBrushPositionOnTargetMesh(const FRay& WorldRay);
-	bool UpdateBrushPositionOnSculptMesh(const FRay& WorldRay);
+	bool UpdateBrushPositionOnTargetMesh(const FRay& WorldRay, bool bFallbackToViewPlane);
+	bool UpdateBrushPositionOnSculptMesh(const FRay& WorldRay, bool bFallbackToViewPlane);
 	void AlignBrushToView();
 
 	void ApplySmoothBrush(const FRay& WorldRay);
@@ -427,6 +451,8 @@ private:
 	void ApplyFlattenBrush(const FRay& WorldRay);
 	void ApplyResampleBrush(const FRay& WorldRay);
 
+	double SculptMaxFixedHeight = -1.0;
+
 	double CalculateBrushFalloff(double Distance);
 	TArray<FVector3d> ROIPositionBuffer;
 
