ModelingTools: width of depthfade ring on sculpt brush sphere is now dependent on screen size, rather than a fixed constant dimension, which doesn't work for tiny or huge things
#rnx

diff --git a/Engine/Plugins/Experimental/MeshModelingToolset/Source/MeshModelingTools/Private/DynamicMeshSculptTool.cpp b/Engine/Plugins/Experimental/MeshModelingToolset/Source/MeshModelingTools/Private/DynamicMeshSculptTool.cpp
--- a/Engine/Plugins/Experimental/MeshModelingToolset/Source/MeshModelingTools/Private/DynamicMeshSculptTool.cpp
+++ b/Engine/Plugins/Experimental/MeshModelingToolset/Source/MeshModelingTools/Private/DynamicMeshSculptTool.cpp
@@ -1376,12 +1376,18 @@ bool UDynamicMeshSculptTool::OnUpdateHover(const FInputDeviceRay& DevicePos)
 	return true;
 }
 
+
 void UDynamicMeshSculptTool::Render(IToolsContextRenderAPI* RenderAPI)
 {
 	UMeshSurfacePointTool::Render(RenderAPI);
 	GetToolManager()->GetContextQueriesAPI()->GetCurrentViewState(CameraState);
 
 	BrushIndicator->Update( (float)this->CurrentBrushRadius, (FVector)this->LastBrushPosWorld, (FVector)this->LastBrushPosNormalWorld, 1.0f-BrushProperties->BrushFalloffAmount);
+	if (BrushIndicatorMaterial)
+	{
+		double FixedDimScale = ToolSceneQueriesUtil::CalculateDimensionFromVisualAngleD(CameraState, LastBrushPosWorld, 1.5f);
+		BrushIndicatorMaterial->SetScalarParameterValue(TEXT("FalloffWidth"), FixedDimScale);
+	}
 
 	if (SculptProperties->PrimaryBrushType == EDynamicMeshSculptBrushType::FixedPlane)
 	{
@@ -1914,7 +1920,13 @@ UPreviewMesh* UDynamicMeshSculptTool::MakeDefaultSphereMesh(UObject* Parent, UWo
 	SphereGen.Generate();
 	FDynamicMesh3 Mesh(&SphereGen);
 	SphereMesh->UpdatePreview(&Mesh);
-	SphereMesh->SetMaterial(ToolSetupUtil::GetDefaultBrushVolumeMaterial(nullptr));
+
+	BrushIndicatorMaterial = ToolSetupUtil::GetDefaultBrushVolumeMaterial(GetToolManager());
+	if (BrushIndicatorMaterial)
+	{
+		SphereMesh->SetMaterial(BrushIndicatorMaterial);
+	}
+
 	return SphereMesh;
 }
 
diff --git a/Engine/Plugins/Experimental/MeshModelingToolset/Source/MeshModelingTools/Public/DynamicMeshSculptTool.h b/Engine/Plugins/Experimental/MeshModelingToolset/Source/MeshModelingTools/Public/DynamicMeshSculptTool.h
--- a/Engine/Plugins/Experimental/MeshModelingToolset/Source/MeshModelingTools/Public/DynamicMeshSculptTool.h
+++ b/Engine/Plugins/Experimental/MeshModelingToolset/Source/MeshModelingTools/Public/DynamicMeshSculptTool.h
@@ -351,6 +351,9 @@ private:
 	UPROPERTY()
 	UBrushStampIndicator* BrushIndicator;
 
+	UPROPERTY()
+	UMaterialInstanceDynamic* BrushIndicatorMaterial;
+
 	UPROPERTY()
 	UPreviewMesh* BrushIndicatorMesh;
 
diff --git a/Engine/Plugins/Experimental/MeshModelingToolset/Source/ModelingComponents/Private/ToolSetupUtil.cpp b/Engine/Plugins/Experimental/MeshModelingToolset/Source/ModelingComponents/Private/ToolSetupUtil.cpp
--- a/Engine/Plugins/Experimental/MeshModelingToolset/Source/ModelingComponents/Private/ToolSetupUtil.cpp
+++ b/Engine/Plugins/Experimental/MeshModelingToolset/Source/ModelingComponents/Private/ToolSetupUtil.cpp
@@ -30,14 +30,15 @@ UMaterialInterface* ToolSetupUtil::GetDefaultWorkingMaterial(UInteractiveToolMan
 }
 
 
-UMaterialInterface* ToolSetupUtil::GetDefaultBrushVolumeMaterial(UInteractiveToolManager* ToolManager)
+UMaterialInstanceDynamic* ToolSetupUtil::GetDefaultBrushVolumeMaterial(UInteractiveToolManager* ToolManager)
 {
-	UMaterialInterface* Material = LoadObject<UMaterial>(nullptr, TEXT("/MeshModelingToolset/Materials/BrushIndicatorMaterial"));
-	if (Material == nullptr && ToolManager != nullptr)
+	UMaterial* Material = LoadObject<UMaterial>(nullptr, TEXT("/MeshModelingToolset/Materials/BrushIndicatorMaterial"));
+	if (Material != nullptr)
 	{
-		return ToolManager->GetContextQueriesAPI()->GetStandardMaterial(EStandardToolContextMaterials::VertexColorMaterial);
+		UMaterialInstanceDynamic* MatInstance = UMaterialInstanceDynamic::Create(Material, ToolManager);
+		return MatInstance;
 	}
-	return Material;
+	return nullptr;
 }
 
 
diff --git a/Engine/Plugins/Experimental/MeshModelingToolset/Source/ModelingComponents/Public/ToolSetupUtil.h b/Engine/Plugins/Experimental/MeshModelingToolset/Source/ModelingComponents/Public/ToolSetupUtil.h
--- a/Engine/Plugins/Experimental/MeshModelingToolset/Source/ModelingComponents/Public/ToolSetupUtil.h
+++ b/Engine/Plugins/Experimental/MeshModelingToolset/Source/ModelingComponents/Public/ToolSetupUtil.h
@@ -29,7 +29,7 @@ namespace ToolSetupUtil
 	/**
 	 * @return default material to use for brush volume indicators
 	 */
-	MODELINGCOMPONENTS_API UMaterialInterface* GetDefaultBrushVolumeMaterial(UInteractiveToolManager* ToolManager);
+	MODELINGCOMPONENTS_API UMaterialInstanceDynamic* GetDefaultBrushVolumeMaterial(UInteractiveToolManager* ToolManager);
 
 
 	/**
