ModelingTools: active change was left in a non-null state after change was ended if selection was not actually modified, which would then fail check on next BeginChange()
#rnx

diff --git a/Engine/Plugins/Experimental/MeshModelingToolset/Source/ModelingComponents/Private/Selection/PolygonSelectionMechanic.cpp b/Engine/Plugins/Experimental/MeshModelingToolset/Source/ModelingComponents/Private/Selection/PolygonSelectionMechanic.cpp
--- a/Engine/Plugins/Experimental/MeshModelingToolset/Source/ModelingComponents/Private/Selection/PolygonSelectionMechanic.cpp
+++ b/Engine/Plugins/Experimental/MeshModelingToolset/Source/ModelingComponents/Private/Selection/PolygonSelectionMechanic.cpp
@@ -280,6 +280,7 @@ TUniquePtr<FToolCommandChange> UPolygonSelectionMechanic::EndChange()
 	{
 		return MoveTemp(ActiveChange);
 	}
+	ActiveChange = TUniquePtr<FPolygonSelectionMechanicSelectionChange>();
 	return TUniquePtr<FToolCommandChange>();
 }
 
@@ -293,6 +294,7 @@ bool UPolygonSelectionMechanic::EndChangeAndEmitIfModified()
 			LOCTEXT("SelectionChange", "Selection Change"));
 		return true;
 	}
+	ActiveChange = TUniquePtr<FPolygonSelectionMechanicSelectionChange>();
 	return false;
 }
 
