use the multi-material generate mesh function when separating triangles in the mesh selection tool, to give separated mesh pieces the same materials as the source mesh.  + associated cleanup

#rnx

diff --git a/Engine/Plugins/Experimental/MeshModelingToolset/Source/MeshModelingTools/Private/MeshSelectionTool.cpp b/Engine/Plugins/Experimental/MeshModelingToolset/Source/MeshModelingTools/Private/MeshSelectionTool.cpp
--- a/Engine/Plugins/Experimental/MeshModelingToolset/Source/MeshModelingTools/Private/MeshSelectionTool.cpp
+++ b/Engine/Plugins/Experimental/MeshModelingToolset/Source/MeshModelingTools/Private/MeshSelectionTool.cpp
@@ -1240,34 +1240,6 @@ void UMeshSelectionTool::DeleteSelectedTriangles()
 }
 
 
-
-void AssignMaterial(AActor* ToActor, const TUniquePtr<FPrimitiveComponentTarget>& FromTarget)
-{
-	UMaterialInterface* Material = FromTarget->GetMaterial(0);
-	if (!Material)
-	{
-		return;
-	}
-
-	//if (Cast<AStaticMeshActor>(ToActor) != nullptr)
-	//{
-	//	UStaticMeshComponent* Component = Cast<AStaticMeshActor>(ToActor)->GetStaticMeshComponent();
-	//	if (Component)
-	//	{
-	//		Component->SetMaterial(0, Material);
-	//	}
-	//} 
-	//else
-	//{
-		USceneComponent* Component = ToActor->GetRootComponent();
-		if (Cast<UPrimitiveComponent>(Component) != nullptr)
-		{
-			Cast<UPrimitiveComponent>(Component)->SetMaterial(0, Material);
-		}
-	//}
-}
-
-
 void UMeshSelectionTool::DisconnectSelectedTriangles()
 {
 	check(SelectionType == EMeshSelectionElementType::Face);
@@ -1339,12 +1311,15 @@ void UMeshSelectionTool::SeparateSelectedTriangles()
 	// emit new asset
 	FTransform3d Transform(PreviewMesh->GetTransform());
 	GetToolManager()->BeginUndoTransaction(LOCTEXT("MeshSelectionToolSeparate", "Separate"));
-	AActor* NewActor = AssetGenerationUtil::GenerateStaticMeshActor(
-		AssetAPI, TargetWorld, &SeparatedMesh, Transform, TEXT("Submesh"));
-	if (NewActor != nullptr)
+
+	// build array of materials from the original
+	TArray<UMaterialInterface*> Materials;
+	for (int MaterialIdx = 0, NumMaterials = ComponentTarget->GetNumMaterials(); MaterialIdx < NumMaterials; MaterialIdx++)
 	{
-		AssignMaterial(NewActor, ComponentTarget);
+		Materials.Add(ComponentTarget->GetMaterial(MaterialIdx));
 	}
+	AActor* NewActor = AssetGenerationUtil::GenerateStaticMeshActor(
+		AssetAPI, TargetWorld, &SeparatedMesh, Transform, TEXT("Submesh"), Materials);
 	SpawnedActors.Add(NewActor);
 	GetToolManager()->EndUndoTransaction();
 
diff --git a/Engine/Plugins/Experimental/MeshModelingToolset/Source/MeshModelingTools/Public/MeshSelectionTool.h b/Engine/Plugins/Experimental/MeshModelingToolset/Source/MeshModelingTools/Public/MeshSelectionTool.h
--- a/Engine/Plugins/Experimental/MeshModelingToolset/Source/MeshModelingTools/Public/MeshSelectionTool.h
+++ b/Engine/Plugins/Experimental/MeshModelingToolset/Source/MeshModelingTools/Public/MeshSelectionTool.h
@@ -48,7 +48,6 @@ enum class EMeshSelectionToolActions
 	SeparateSelected,
 	FlipSelected,
 	CreateGroup,
-	AssignMaterial,
 
 	CycleSelectionMode,
 	CycleViewMode
