ModelingTools: fix incorrectly-disabled min-stretch parameter
#rnx

diff --git a/Engine/Plugins/Experimental/MeshModelingToolset/Source/MeshModelingToolsEditorOnly/Public/ParameterizeMeshTool.h b/Engine/Plugins/Experimental/MeshModelingToolset/Source/MeshModelingToolsEditorOnly/Public/ParameterizeMeshTool.h
--- a/Engine/Plugins/Experimental/MeshModelingToolset/Source/MeshModelingToolsEditorOnly/Public/ParameterizeMeshTool.h
+++ b/Engine/Plugins/Experimental/MeshModelingToolset/Source/MeshModelingToolsEditorOnly/Public/ParameterizeMeshTool.h
@@ -110,7 +110,7 @@ public:
 	EUVUnwrapType UnwrapType = EUVUnwrapType::ExpMap;
 
 	/** Maximum amount of stretch, from none to any.  If zero stretch is specified each triangle will likey be its own chart */
-	UPROPERTY(EditAnywhere, Category = Options, meta = (Default = "0.166", UIMin = "0", UIMax = "1", ClampMin = "0", ClampMax = "1", EditCondition = "UnwrapType == EUVUnwrapType::MinStretch"))
+	UPROPERTY(EditAnywhere, Category = Options, meta = (Default = "0.166", UIMin = "0", UIMax = "1", ClampMin = "0", ClampMax = "1", EditCondition = "bIsGlobalMode || UnwrapType == EUVUnwrapType::MinStretch"))
 	float ChartStretch = 0.11f;
 
 	/** Scaling applied to UV islands */
