fix duplicate tool changing pivot
#rnx

diff --git a/Engine/Plugins/Experimental/MeshModelingToolset/Source/MeshModelingTools/Private/CombineMeshesTool.cpp b/Engine/Plugins/Experimental/MeshModelingToolset/Source/MeshModelingTools/Private/CombineMeshesTool.cpp
--- a/Engine/Plugins/Experimental/MeshModelingToolset/Source/MeshModelingTools/Private/CombineMeshesTool.cpp
+++ b/Engine/Plugins/Experimental/MeshModelingToolset/Source/MeshModelingTools/Private/CombineMeshesTool.cpp
@@ -229,15 +229,29 @@ void UCombineMeshesTool::UpdateAssets()
 
 			FDynamicMeshEditor Editor(&AccumulateDMesh);
 			FMeshIndexMappings IndexMapping;
-			Editor.AppendMesh(&ComponentDMesh, IndexMapping, 
-				[&XF](int Unused, const FVector3d P) { return XF.TransformPosition(P); },
-				[&XF](int Unused, const FVector3d N) { return XF.TransformNormal(N); });
+			if (bDuplicateMode) // no transform if duplicating
+			{
+				Editor.AppendMesh(&ComponentDMesh, IndexMapping);
+			}
+			else
+			{
+				Editor.AppendMesh(&ComponentDMesh, IndexMapping,
+					[&XF](int Unused, const FVector3d P) { return XF.TransformPosition(P); },
+					[&XF](int Unused, const FVector3d N) { return XF.TransformNormal(N); });
+			}
+			
 
 			MatIndexBase += ComponentTarget->GetNumMaterials();
 		}
 
 		SlowTask.EnterProgressFrame(1);
 
+		if (bDuplicateMode)
+		{
+			// TODO: will need to refactor this when we support duplicating multiple
+			check(ComponentTargets.Num() == 1);
+			AccumToWorld = ComponentTargets[0]->GetWorldTransform();
+		}
 		AActor* NewActor = AssetGenerationUtil::GenerateStaticMeshActor(
 			AssetAPI, TargetWorld,
 			&AccumulateDMesh, (FTransform3d)AccumToWorld, TEXT("Combined Meshes"), AllMaterials);
