ModelingTools: DrawPolygonTool extrusion-preview hit-test was broken by previous fixes. Resolved this and also converted most math to double-precision.
#rnx

diff --git a/Engine/Plugins/Experimental/MeshModelingToolset/Source/MeshModelingTools/Private/DrawPolygonTool.cpp b/Engine/Plugins/Experimental/MeshModelingToolset/Source/MeshModelingTools/Private/DrawPolygonTool.cpp
--- a/Engine/Plugins/Experimental/MeshModelingToolset/Source/MeshModelingTools/Private/DrawPolygonTool.cpp
+++ b/Engine/Plugins/Experimental/MeshModelingToolset/Source/MeshModelingTools/Private/DrawPolygonTool.cpp
@@ -126,8 +126,8 @@ void UDrawPolygonToolSnapProperties::RestoreProperties(UInteractiveTool* Restore
  */
 UDrawPolygonTool::UDrawPolygonTool()
 {
-	DrawPlaneOrigin = FVector::ZeroVector;
-	DrawPlaneOrientation = FQuat::Identity;
+	DrawPlaneOrigin = FVector3d::Zero();
+	DrawPlaneOrientation = FQuaterniond::Identity();
 	bInInteractiveExtrude = false;
 }
 
@@ -181,7 +181,7 @@ void UDrawPolygonTool::Setup()
 	// Create a new TransformGizmo and associated TransformProxy. The TransformProxy will not be the
 	// parent of any Components in this case, we just use it's transform and change delegate.
 	PlaneTransformProxy = NewObject<UTransformProxy>(this);
-	PlaneTransformProxy->SetTransform(FTransform(DrawPlaneOrientation, DrawPlaneOrigin));
+	PlaneTransformProxy->SetTransform(FTransform((FQuat)DrawPlaneOrientation, (FVector)DrawPlaneOrigin));
 	PlaneTransformGizmo = GetToolManager()->GetPairedGizmoManager()->CreateCustomTransformGizmo(
 		ETransformGizmoSubElements::StandardTranslateRotate, this);
 	PlaneTransformGizmo->SetActiveTarget(PlaneTransformProxy, GetToolManager());
@@ -204,7 +204,7 @@ void UDrawPolygonTool::Setup()
 	SnapEngine.SnapMetricFunc = [this](const FVector3d& Position1, const FVector3d& Position2) {
 		return ToolSceneQueriesUtil::CalculateViewVisualAngleD(this->CameraState, Position1, Position2);
 	};
-	SnapEngine.Plane = FFrame3d((FVector3d)DrawPlaneOrigin, (FQuaterniond)DrawPlaneOrientation);
+	SnapEngine.Plane = FFrame3d(DrawPlaneOrigin, DrawPlaneOrientation);
 
 	SnapProperties = NewObject<UDrawPolygonToolSnapProperties>(this, TEXT("Snapping"));
 	SnapProperties->RestoreProperties(this);
@@ -301,19 +301,19 @@ void UDrawPolygonTool::Tick(float DeltaTime)
 
 
 void DrawEdgeTicks(FPrimitiveDrawInterface* PDI, 
-	const FSegment3f& Segment, float Height,
-	const FVector3f& PlaneNormal, 
+	const FSegment3d& Segment, float Height,
+	const FVector3d& PlaneNormal, 
 	const FLinearColor& Color, uint8 DepthPriorityGroup, float LineThickness, bool bIsScreenSpace)
 {
-	FVector3f Center = Segment.Center;
-	FVector3f X = Segment.Direction;
-	FVector3f Y = X.Cross(PlaneNormal);
+	FVector3d Center = Segment.Center;
+	FVector3d X = Segment.Direction;
+	FVector3d Y = X.Cross(PlaneNormal);
 	Y.Normalize();
-	FVector3f A = Center - Height * 0.25f*X - Height * Y;
-	FVector3f B = Center + Height * 0.25f*X + Height * Y;
+	FVector3d A = Center - Height * 0.25*X - Height * Y;
+	FVector3d B = Center + Height * 0.25*X + Height * Y;
 	PDI->DrawLine((FVector)A, (FVector)B, Color, DepthPriorityGroup, LineThickness, 0.0f, bIsScreenSpace);
-	A += Height * 0.5f*X;
-	B += Height * 0.5f*X;
+	A += Height * 0.5*X;
+	B += Height * 0.5*X;
 	PDI->DrawLine((FVector)A, (FVector)B, Color, DepthPriorityGroup, LineThickness, 0.0f, bIsScreenSpace);
 }
 
@@ -347,7 +347,7 @@ void UDrawPolygonTool::Render(IToolsContextRenderAPI* RenderAPI)
 
 	if (bInInteractiveExtrude == false)
 	{
-		FFrame3f DrawFrame(DrawPlaneOrigin, DrawPlaneOrientation);
+		FFrame3f DrawFrame((FVector3f)DrawPlaneOrigin, (FQuaternionf)DrawPlaneOrientation);
 		MeshDebugDraw::DrawSimpleGrid(DrawFrame, NumGridLines, GridLineSpacing, GridThickness, GridColor, false, PDI, FTransform::Identity);
 	}
 
@@ -398,13 +398,13 @@ void UDrawPolygonTool::Render(IToolsContextRenderAPI* RenderAPI)
 			if (SnapEngine.HaveActiveSnapDistance())
 			{
 				int iSegment = SnapEngine.GetActiveSnapDistanceID();
-				TArray<FVector>& HistoryPoints = (bInFixedPolygonMode) ? FixedPolygonClickPoints : PolygonVertices;
-				FVector UseNormal = DrawPlaneOrientation.GetAxisZ();
-				DrawEdgeTicks(PDI, FSegment3f(HistoryPoints[iSegment], HistoryPoints[iSegment+1]),
+				TArray<FVector3d>& HistoryPoints = (bInFixedPolygonMode) ? FixedPolygonClickPoints : PolygonVertices;
+				FVector3d UseNormal = DrawPlaneOrientation.AxisZ();
+				DrawEdgeTicks(PDI, FSegment3d(HistoryPoints[iSegment], HistoryPoints[iSegment+1]),
 					0.75f*ElementSize, UseNormal, SnapHighlightColor, SDPG_Foreground, 1.0f, true);
-				DrawEdgeTicks(PDI, FSegment3f(HistoryPoints[HistoryPoints.Num()-1], PreviewVertex),
+				DrawEdgeTicks(PDI, FSegment3d(HistoryPoints[HistoryPoints.Num()-1], PreviewVertex),
 					0.75f*ElementSize, UseNormal, SnapHighlightColor, SDPG_Foreground, 1.0f, true);
-				PDI->DrawLine(HistoryPoints[iSegment], HistoryPoints[iSegment + 1],
+				PDI->DrawLine((FVector)HistoryPoints[iSegment], (FVector)HistoryPoints[iSegment + 1],
 					SnapHighlightColor, SDPG_Foreground, 2.0f, 0.0f, true);
 			}
 		}
@@ -427,7 +427,7 @@ void UDrawPolygonTool::Render(IToolsContextRenderAPI* RenderAPI)
 	if (PolygonVertices.Num() > 0)
 	{
 		FColor UseColor = (bIsClosed) ? ClosedPolygonColor : OpenPolygonColor;
-		FVector UseLastVertex = (bIsClosed) ? PolygonVertices[0] : PreviewVertex;
+		FVector3d UseLastVertex = (bIsClosed) ? PolygonVertices[0] : PreviewVertex;
 		float UseThickness = LineThickness;
 		if (bHaveSelfIntersection)
 		{
@@ -435,11 +435,11 @@ void UDrawPolygonTool::Render(IToolsContextRenderAPI* RenderAPI)
 			UseThickness = SelfIntersectThickness;
 		}
 
-		auto DrawVertices = [&PDI, &UseColor](const TArray<FVector>& Vertices, ESceneDepthPriorityGroup Group, float Thickness)
+		auto DrawVertices = [&PDI, &UseColor](const TArray<FVector3d>& Vertices, ESceneDepthPriorityGroup Group, float Thickness)
 		{
 			for (int lasti = Vertices.Num() - 1, i = 0, NumVertices = Vertices.Num(); i < NumVertices; lasti = i++)
 			{
-				PDI->DrawLine(Vertices[lasti], Vertices[i], UseColor, Group, Thickness, 0.0f, true);
+				PDI->DrawLine((FVector)Vertices[lasti], (FVector)Vertices[i], UseColor, Group, Thickness, 0.0f, true);
 			}
 		};
 
@@ -447,10 +447,10 @@ void UDrawPolygonTool::Render(IToolsContextRenderAPI* RenderAPI)
 		//DrawVertices(PolygonVertices, SDPG_Foreground, HiddenLineThickness);
 		for (int i = 0; i < NumVerts - 1; ++i)
 		{
-			PDI->DrawLine(PolygonVertices[i], PolygonVertices[i + 1],
+			PDI->DrawLine((FVector)PolygonVertices[i], (FVector)PolygonVertices[i + 1],
 				UseColor, SDPG_Foreground, HiddenLineThickness, 0.0f, true);
 		}
-		PDI->DrawLine(PolygonVertices[NumVerts - 1], UseLastVertex,
+		PDI->DrawLine((FVector)PolygonVertices[NumVerts - 1], (FVector)UseLastVertex,
 			UseColor, SDPG_Foreground, HiddenLineThickness, 0.0f, true);
 		for (int HoleIdx = 0; HoleIdx < PolygonHolesVertices.Num(); HoleIdx++)
 		{
@@ -461,10 +461,10 @@ void UDrawPolygonTool::Render(IToolsContextRenderAPI* RenderAPI)
 		//DrawVertices(PolygonVertices, SDPG_World, LineThickness);
 		for (int i = 0; i < NumVerts - 1; ++i)
 		{
-			PDI->DrawLine(PolygonVertices[i], PolygonVertices[i + 1],
+			PDI->DrawLine((FVector)PolygonVertices[i], (FVector)PolygonVertices[i + 1],
 				UseColor, SDPG_World, LineThickness, 0.0f, true);
 		}
-		PDI->DrawLine(PolygonVertices[NumVerts - 1], UseLastVertex,
+		PDI->DrawLine((FVector)PolygonVertices[NumVerts - 1], (FVector)UseLastVertex,
 			UseColor, SDPG_World, LineThickness, 0.0f, true);
 		for (int HoleIdx = 0; HoleIdx < PolygonHolesVertices.Num(); HoleIdx++)
 		{
@@ -478,8 +478,7 @@ void UDrawPolygonTool::Render(IToolsContextRenderAPI* RenderAPI)
 	}
 
 	// draw preview vertex
-	PDI->DrawPoint(PreviewVertex, ClosedPolygonColor, 10, SDPG_Foreground);
-
+	PDI->DrawPoint((FVector)PreviewVertex, ClosedPolygonColor, 10, SDPG_Foreground);
 
 	// draw height preview stuff
 	if (bInInteractiveExtrude)
@@ -487,7 +486,6 @@ void UDrawPolygonTool::Render(IToolsContextRenderAPI* RenderAPI)
 		HeightMechanic->Render(RenderAPI);
 	}
 
-
 	ShowGizmoWatcher.CheckAndUpdate();
 }
 
@@ -503,24 +501,24 @@ void UDrawPolygonTool::ResetPolygon()
 	CurrentCurveTimestamp++;
 }
 
-void UDrawPolygonTool::UpdatePreviewVertex(const FVector& PreviewVertexIn)
+void UDrawPolygonTool::UpdatePreviewVertex(const FVector3d& PreviewVertexIn)
 {
 	PreviewVertex = PreviewVertexIn;
 
 	// update length and angle
 	if (PolygonVertices.Num() > 0)
 	{
-		FVector LastVertex = PolygonVertices[PolygonVertices.Num() - 1];
-		SnapProperties->SegmentLength = FVector::Distance(LastVertex, PreviewVertex);
+		FVector3d LastVertex = PolygonVertices[PolygonVertices.Num() - 1];
+		SnapProperties->SegmentLength = LastVertex.Distance(PreviewVertex);
 	}
 }
 
-void UDrawPolygonTool::AppendVertex(const FVector& Vertex)
+void UDrawPolygonTool::AppendVertex(const FVector3d& Vertex)
 {
 	PolygonVertices.Add(Vertex);
 }
 
-bool UDrawPolygonTool::FindDrawPlaneHitPoint(const FInputDeviceRay& ClickPos, FVector& HitPosOut)
+bool UDrawPolygonTool::FindDrawPlaneHitPoint(const FInputDeviceRay& ClickPos, FVector3d& HitPosOut)
 {
 	bHaveSurfaceHit = false;
 
@@ -561,7 +559,7 @@ bool UDrawPolygonTool::FindDrawPlaneHitPoint(const FInputDeviceRay& ClickPos, FV
 			}
 		}
 
-		TArray<FVector>& HistoryPoints = (bInFixedPolygonMode) ? FixedPolygonClickPoints : PolygonVertices;
+		TArray<FVector3d>& HistoryPoints = (bInFixedPolygonMode) ? FixedPolygonClickPoints : PolygonVertices;
 		SnapEngine.UpdatePointHistory(HistoryPoints);
 		if (SnapProperties->bSnapToAngles)
 		{
@@ -578,7 +576,7 @@ bool UDrawPolygonTool::FindDrawPlaneHitPoint(const FInputDeviceRay& ClickPos, FV
 
 	if (SnapEngine.HaveActiveSnap())
 	{
-		HitPosOut = (FVector)SnapEngine.GetActiveSnapToPoint();
+		HitPosOut = SnapEngine.GetActiveSnapToPoint();
 		return true;
 	}
 
@@ -593,20 +591,20 @@ bool UDrawPolygonTool::FindDrawPlaneHitPoint(const FInputDeviceRay& ClickPos, FV
 		{
 			bHaveSurfaceHit = true;
 			SurfaceHitPoint = Result.ImpactPoint;
-			FVector UseHitPos = Result.ImpactPoint + SnapProperties->HitNormalOffset*Result.Normal;
-			HitPos = Frame.ToPlane((FVector3d)UseHitPos, 2);
+			FVector3d UseHitPos = Result.ImpactPoint + SnapProperties->HitNormalOffset*Result.Normal;
+			HitPos = Frame.ToPlane(UseHitPos, 2);
 			SurfaceOffsetPoint = UseHitPos;
 		}
 	}
 
-	HitPosOut = (FVector)HitPos;
+	HitPosOut = HitPos;
 	return true;
 }
 
 void UDrawPolygonTool::OnBeginSequencePreview(const FInputDeviceRay& DevicePos)
 {
 	// just update snapped point preview
-	FVector HitPos;
+	FVector3d HitPos;
 	if (FindDrawPlaneHitPoint(DevicePos, HitPos))
 	{
 		PreviewVertex = HitPos;
@@ -623,7 +621,7 @@ void UDrawPolygonTool::OnBeginClickSequence(const FInputDeviceRay& ClickPos)
 {
 	ResetPolygon();
 	
-	FVector HitPos;
+	FVector3d HitPos;
 	bool bHit = FindDrawPlaneHitPoint(ClickPos, HitPos);
 	if (bHit == false)
 	{
@@ -652,7 +650,7 @@ void UDrawPolygonTool::OnNextSequencePreview(const FInputDeviceRay& ClickPos)
 		return;
 	}
 
-	FVector HitPos;
+	FVector3d HitPos;
 	bool bHit = FindDrawPlaneHitPoint(ClickPos, HitPos);
 	if (bHit == false)
 	{
@@ -682,7 +680,7 @@ bool UDrawPolygonTool::OnNextSequenceClick(const FInputDeviceRay& ClickPos)
 		return false;
 	}
 
-	FVector HitPos;
+	FVector3d HitPos;
 	bool bHit = FindDrawPlaneHitPoint(ClickPos, HitPos);
 	if (bHit == false)
 	{
@@ -816,14 +814,14 @@ bool UDrawPolygonTool::UpdateSelfIntersection()
 		return false;
 	}
 
-	FFrame3f DrawFrame(DrawPlaneOrigin, DrawPlaneOrientation);
-	FSegment2f PreviewSegment(DrawFrame.ToPlaneUV(PolygonVertices[NumVertices - 1],2), DrawFrame.ToPlaneUV(PreviewVertex,2));
+	FFrame3d DrawFrame(DrawPlaneOrigin, DrawPlaneOrientation);
+	FSegment2d PreviewSegment(DrawFrame.ToPlaneUV(PolygonVertices[NumVertices - 1],2), DrawFrame.ToPlaneUV(PreviewVertex,2));
 
-	float BestIntersectionParameter = FMathf::MaxReal;
+	double BestIntersectionParameter = FMathd::MaxReal;
 	for (int k = 0; k < NumVertices - 2; ++k) 
 	{
-		FSegment2f Segment(DrawFrame.ToPlaneUV(PolygonVertices[k],2), DrawFrame.ToPlaneUV(PolygonVertices[k + 1],2));
-		FIntrSegment2Segment2f Intersection(PreviewSegment, Segment);
+		FSegment2d Segment(DrawFrame.ToPlaneUV(PolygonVertices[k],2), DrawFrame.ToPlaneUV(PolygonVertices[k + 1],2));
+		FIntrSegment2Segment2d Intersection(PreviewSegment, Segment);
 		if (Intersection.Find()) 
 		{
 			bHaveSelfIntersection = true;
@@ -838,75 +836,75 @@ bool UDrawPolygonTool::UpdateSelfIntersection()
 	return bHaveSelfIntersection;
 }
 
-void UDrawPolygonTool::GetPolygonParametersFromFixedPoints(const TArray<FVector>& FixedPoints, FVector2f& FirstReferencePt, FVector2f& BoxSize, float& YSign, float& AngleRad)
+void UDrawPolygonTool::GetPolygonParametersFromFixedPoints(const TArray<FVector3d>& FixedPoints, FVector2d& FirstReferencePt, FVector2d& BoxSize, double& YSign, double& AngleRad)
 {
 	if (FixedPoints.Num() < 2)
 	{
 		return;
 	}
 
-	FFrame3f DrawFrame(DrawPlaneOrigin, DrawPlaneOrientation);
+	FFrame3d DrawFrame(DrawPlaneOrigin, DrawPlaneOrientation);
 	FirstReferencePt = DrawFrame.ToPlaneUV(FixedPoints[0], 2);
 
-	FVector2f EdgePt = DrawFrame.ToPlaneUV(FixedPoints[1], 2);
-	FVector2f Delta = EdgePt - FirstReferencePt;
-	AngleRad = FMath::Atan2(Delta.Y, Delta.X);
+	FVector2d EdgePt = DrawFrame.ToPlaneUV(FixedPoints[1], 2);
+	FVector2d Delta = EdgePt - FirstReferencePt;
+	AngleRad = FMathd::Atan2(Delta.Y, Delta.X);
 
-	float Radius = Delta.Length();
-	FVector2f AxisX = Delta / Radius;
-	FVector2f AxisY = -AxisX.Perp();
-	FVector2f HeightPt = DrawFrame.ToPlaneUV((FixedPoints.Num() == 3) ? FixedPoints[2] : FixedPoints[1], 2);
-	FVector2f HeightDelta = HeightPt - FirstReferencePt;
-	YSign = FMath::Sign(HeightDelta.Dot(AxisY));
+	double Radius = Delta.Length();
+	FVector2d AxisX = Delta / Radius;
+	FVector2d AxisY = -AxisX.Perp();
+	FVector2d HeightPt = DrawFrame.ToPlaneUV((FixedPoints.Num() == 3) ? FixedPoints[2] : FixedPoints[1], 2);
+	FVector2d HeightDelta = HeightPt - FirstReferencePt;
+	YSign = FMathd::Sign(HeightDelta.Dot(AxisY));
 	BoxSize.X = Radius;
-	BoxSize.Y = FMath::Abs(HeightDelta.Dot(AxisY));
+	BoxSize.Y = FMathd::Abs(HeightDelta.Dot(AxisY));
 }
 
-void UDrawPolygonTool::GenerateFixedPolygon(const TArray<FVector>& FixedPoints, TArray<FVector>& VerticesOut, TArray<TArray<FVector>>& HolesVerticesOut)
+void UDrawPolygonTool::GenerateFixedPolygon(const TArray<FVector3d>& FixedPoints, TArray<FVector3d>& VerticesOut, TArray<TArray<FVector3d>>& HolesVerticesOut)
 {
-	FVector2f FirstReferencePt, BoxSize;
-	float YSign, AngleRad;
+	FVector2d FirstReferencePt, BoxSize;
+	double YSign, AngleRad;
 	GetPolygonParametersFromFixedPoints(FixedPoints, FirstReferencePt, BoxSize, YSign, AngleRad);
-	float Width = BoxSize.X, Height = BoxSize.Y;
-	FMatrix2f RotationMat = FMatrix2f::RotationRad(AngleRad);
+	double Width = BoxSize.X, Height = BoxSize.Y;
+	FMatrix2d RotationMat = FMatrix2d::RotationRad(AngleRad);
 
-	FPolygon2f Polygon;
-	TArray<FPolygon2f> PolygonHoles;
+	FPolygon2d Polygon;
+	TArray<FPolygon2d> PolygonHoles;
 	if (PolygonProperties->PolygonType == EDrawPolygonDrawMode::Square)
 	{
-		Polygon = FPolygon2f::MakeRectangle(FVector2f::Zero(), 2*Width, 2*Width);
+		Polygon = FPolygon2d::MakeRectangle(FVector2d::Zero(), 2*Width, 2*Width);
 	}
 	else if (PolygonProperties->PolygonType == EDrawPolygonDrawMode::Rectangle || PolygonProperties->PolygonType == EDrawPolygonDrawMode::RoundedRectangle)
 	{
 		if (PolygonProperties->PolygonType == EDrawPolygonDrawMode::Rectangle)
 		{
-			Polygon = FPolygon2f::MakeRectangle(FVector2f(Width / 2, YSign*Height / 2), Width, Height);
+			Polygon = FPolygon2d::MakeRectangle(FVector2d(Width / 2, YSign*Height / 2), Width, Height);
 		}
 		else // PolygonProperties->PolygonType == EDrawPolygonDrawMode::RoundedRectangle
 		{
-			Polygon = FPolygon2f::MakeRoundedRectangle(FVector2f(Width / 2, YSign*Height / 2), Width, Height, FMath::Min(Width,Height) * FMathf::Clamp(PolygonProperties->FeatureSizeRatio, .01, .99) * .5, PolygonProperties->Steps);
+			Polygon = FPolygon2d::MakeRoundedRectangle(FVector2d(Width / 2, YSign*Height / 2), Width, Height, FMathd::Min(Width,Height) * FMathd::Clamp(PolygonProperties->FeatureSizeRatio, .01, .99) * .5, PolygonProperties->Steps);
 		}
 	}
 	else // Circle or HoleyCircle
 	{
-		Polygon = FPolygon2f::MakeCircle(Width, PolygonProperties->Steps, 0);
+		Polygon = FPolygon2d::MakeCircle(Width, PolygonProperties->Steps, 0);
 		if (PolygonProperties->PolygonType == EDrawPolygonDrawMode::HoleyCircle)
 		{
-			PolygonHoles.Add(FPolygon2f::MakeCircle(Width * FMathd::Clamp(PolygonProperties->FeatureSizeRatio, .01, .99), PolygonProperties->Steps, 0));
+			PolygonHoles.Add(FPolygon2d::MakeCircle(Width * FMathd::Clamp(PolygonProperties->FeatureSizeRatio, .01, .99), PolygonProperties->Steps, 0));
 		}
 	}
-	Polygon.Transform([RotationMat](const FVector2f& Pt) { return RotationMat * Pt; });
-	for (FPolygon2f& Hole : PolygonHoles)
+	Polygon.Transform([RotationMat](const FVector2d& Pt) { return RotationMat * Pt; });
+	for (FPolygon2d& Hole : PolygonHoles)
 	{
-		Hole.Transform([RotationMat](const FVector2f& Pt) { return RotationMat * Pt; });
+		Hole.Transform([RotationMat](const FVector2d& Pt) { return RotationMat * Pt; });
 	}
 
-	FFrame3f DrawFrame(DrawPlaneOrigin, DrawPlaneOrientation);
+	FFrame3d DrawFrame(DrawPlaneOrigin, DrawPlaneOrientation);
 	VerticesOut.SetNum(Polygon.VertexCount());
 	for (int k = 0; k < Polygon.VertexCount(); ++k)
 	{
-		FVector2f NewPt = FirstReferencePt + Polygon[k];
-		VerticesOut[k] = (FVector)DrawFrame.FromPlaneUV(NewPt, 2);
+		FVector2d NewPt = FirstReferencePt + Polygon[k];
+		VerticesOut[k] = DrawFrame.FromPlaneUV(NewPt, 2);
 	}
 
 	HolesVerticesOut.SetNum(PolygonHoles.Num());
@@ -916,8 +914,8 @@ void UDrawPolygonTool::GenerateFixedPolygon(const TArray<FVector>& FixedPoints,
 		HolesVerticesOut[HoleIdx].SetNum(NumHoleVerts);
 		for (int k = 0; k < NumHoleVerts; ++k)
 		{
-			FVector2f NewPt = FirstReferencePt + PolygonHoles[HoleIdx][k];
-			HolesVerticesOut[HoleIdx][k] = (FVector)DrawFrame.FromPlaneUV(NewPt, 2);
+			FVector2d NewPt = FirstReferencePt + PolygonHoles[HoleIdx][k];
+			HolesVerticesOut[HoleIdx][k] = DrawFrame.FromPlaneUV(NewPt, 2);
 		}
 	}
 }
@@ -947,20 +945,9 @@ void UDrawPolygonTool::BeginInteractiveExtrude()
 	HeightMechanic->CurrentHeight = 1.0f;  // initialize to something non-zero...prob should be based on polygon bounds maybe?
 
 	FDynamicMesh3 HeightMesh;
-	FFrame3d MeshFrame_Discard;
-	GeneratePolygonMesh(PolygonVertices, PolygonHolesVertices, &HeightMesh, MeshFrame_Discard, false, 99999, true);
-
-	// We don't actually want the same FFrame3d that was computed by GeneratePolygonMesh
-	int NumVerts = PolygonVertices.Num();
-	FVector3d Centroid(0, 0, 0);
-	for (int k = 0; k < NumVerts; ++k)
-	{
-		Centroid += PolygonVertices[k];
-	}
-	Centroid /= (double)NumVerts;
-	FFrame3d HeightMeshFrame(Centroid, TQuaternion<double>(DrawPlaneOrientation));
-
-	HeightMechanic->Initialize( MoveTemp(HeightMesh), HeightMeshFrame, false);
+	FFrame3d WorldMeshFrame;
+	GeneratePolygonMesh(PolygonVertices, PolygonHolesVertices, &HeightMesh, WorldMeshFrame, false, 99999, true);
+	HeightMechanic->Initialize( MoveTemp(HeightMesh), WorldMeshFrame, false);
 
 	ShowExtrudeMessage();
 }
@@ -980,32 +967,32 @@ void UDrawPolygonTool::EndInteractiveExtrude()
 
 
 
-void UDrawPolygonTool::SetDrawPlaneFromWorldPos(const FVector& Position, const FVector& Normal)
+void UDrawPolygonTool::SetDrawPlaneFromWorldPos(const FVector3d& Position, const FVector3d& Normal)
 {
 	DrawPlaneOrigin = Position;
 
-	FFrame3f DrawPlane(Position, DrawPlaneOrientation);
+	FFrame3d DrawPlane(Position, DrawPlaneOrientation);
 	if (bIgnoreSnappingToggle == false)
 	{
 		DrawPlane.AlignAxis(2, Normal);
 		DrawPlane.ConstrainedAlignPerpAxes();
-		DrawPlaneOrientation = (FQuat)DrawPlane.Rotation;
+		DrawPlaneOrientation = DrawPlane.Rotation;
 	}
 
-	SnapEngine.Plane = FFrame3d((FVector3d)DrawPlane.Origin, (FQuaterniond)DrawPlane.Rotation);
+	SnapEngine.Plane = FFrame3d(DrawPlane.Origin, DrawPlane.Rotation);
 
 	if (PlaneTransformGizmo != nullptr)
 	{
-		PlaneTransformGizmo->SetNewGizmoTransform(FTransform(DrawPlaneOrientation, DrawPlaneOrigin));
+		PlaneTransformGizmo->SetNewGizmoTransform(FTransform((FQuat)DrawPlaneOrientation, (FVector)DrawPlaneOrigin));
 	}
 }
 
 
 void UDrawPolygonTool::PlaneTransformChanged(UTransformProxy* Proxy, FTransform Transform)
 {
-	DrawPlaneOrientation = Transform.GetRotation();
-	DrawPlaneOrigin = Transform.GetLocation();
-	SnapEngine.Plane = FFrame3d((FVector3d)DrawPlaneOrigin, (FQuaterniond)DrawPlaneOrientation);
+	DrawPlaneOrientation = (FQuaterniond)Transform.GetRotation();
+	DrawPlaneOrigin = (FVector3d)Transform.GetLocation();
+	SnapEngine.Plane = FFrame3d(DrawPlaneOrigin, DrawPlaneOrientation);
 }
 
 void UDrawPolygonTool::UpdateShowGizmoState(bool bNewVisibility)
@@ -1020,7 +1007,7 @@ void UDrawPolygonTool::UpdateShowGizmoState(bool bNewVisibility)
 		PlaneTransformGizmo = GetToolManager()->GetPairedGizmoManager()->CreateCustomTransformGizmo(
 			ETransformGizmoSubElements::StandardTranslateRotate, this);
 		PlaneTransformGizmo->SetActiveTarget(PlaneTransformProxy, GetToolManager());
-		PlaneTransformGizmo->SetNewGizmoTransform(FTransform(DrawPlaneOrientation, DrawPlaneOrigin));
+		PlaneTransformGizmo->SetNewGizmoTransform(FTransform((FQuat)DrawPlaneOrientation, (FVector)DrawPlaneOrigin));
 	}
 }
 
@@ -1081,23 +1068,23 @@ void UDrawPolygonTool::UpdateLivePreview()
 	}
 }
 
-bool UDrawPolygonTool::GeneratePolygonMesh(const TArray<FVector>& Polygon, const TArray<TArray<FVector>>& PolygonHoles, FDynamicMesh3* ResultMeshOut, FFrame3d& WorldFrameOut, bool bIncludePreviewVtx, double ExtrudeDistance, bool bExtrudeSymmetric)
+bool UDrawPolygonTool::GeneratePolygonMesh(const TArray<FVector3d>& Polygon, const TArray<TArray<FVector3d>>& PolygonHoles, FDynamicMesh3* ResultMeshOut, FFrame3d& WorldFrameOut, bool bIncludePreviewVtx, double ExtrudeDistance, bool bExtrudeSymmetric)
 {
 	// construct centered frame for polygon
 	WorldFrameOut = FFrame3d(DrawPlaneOrigin, DrawPlaneOrientation);
 
 	int NumVerts = Polygon.Num();
-	FVector3d Centroid(0, 0, 0);
+	FVector3d Centroid3d(0, 0, 0);
 	for (int k = 0; k < NumVerts; ++k)
 	{
-		Centroid.X += Polygon[k].X;
-		Centroid.Y += Polygon[k].Y;
+		Centroid3d += Polygon[k];
 	}
-	Centroid /= (double)NumVerts;
-	WorldFrameOut.Origin += WorldFrameOut.Rotation * Centroid;
+	Centroid3d /= (double)NumVerts;
+	FVector2d CentroidInDrawPlane = WorldFrameOut.ToPlaneUV(Centroid3d);
+	WorldFrameOut.Origin = Centroid3d;
 
 	// Compute outer polygon & bounds
-	auto VertexArrayToPolygon = [&WorldFrameOut](const TArray<FVector>& Vertices)
+	auto VertexArrayToPolygon = [&WorldFrameOut](const TArray<FVector3d>& Vertices)
 	{
 		FPolygon2d OutPolygon;
 		for (int k = 0, N = Vertices.Num(); k < N; ++k)
@@ -1110,7 +1097,7 @@ bool UDrawPolygonTool::GeneratePolygonMesh(const TArray<FVector>& Polygon, const
 	// add preview vertex
 	if (bIncludePreviewVtx)
 	{
-		if (FVector::Dist(PreviewVertex, Polygon[NumVerts - 1]) > 0.1)
+		if (PreviewVertex.Distance(Polygon[NumVerts - 1]) > 0.1)
 		{
 			OuterPolygon.AppendVertex(WorldFrameOut.ToPlaneUV(PreviewVertex, 2));
 		}
@@ -1121,11 +1108,11 @@ bool UDrawPolygonTool::GeneratePolygonMesh(const TArray<FVector>& Polygon, const
 	if (PolygonProperties->PolygonType == EDrawPolygonDrawMode::HoleyCircle || PolygonProperties->PolygonType == EDrawPolygonDrawMode::Circle || PolygonProperties->PolygonType == EDrawPolygonDrawMode::RoundedRectangle)
 	{
 		// get polygon parameters
-		FVector2f FirstReferencePt, BoxSize;
-		float YSign, AngleRad;
+		FVector2d FirstReferencePt, BoxSize;
+		double YSign, AngleRad;
 		GetPolygonParametersFromFixedPoints(FixedPolygonClickPoints, FirstReferencePt, BoxSize, YSign, AngleRad);
-		FirstReferencePt -= FVector2f(Centroid.X, Centroid.Y);
-		FMatrix2f RotationMat = FMatrix2f::RotationRad(AngleRad);
+		FirstReferencePt -= CentroidInDrawPlane;
+		FMatrix2d RotationMat = FMatrix2d::RotationRad(AngleRad);
 
 		// translate general polygon parameters to specific mesh generator parameters, and generate mesh
 		if (PolygonProperties->PolygonType == EDrawPolygonDrawMode::HoleyCircle)
@@ -1134,7 +1121,7 @@ bool UDrawPolygonTool::GeneratePolygonMesh(const TArray<FVector>& Polygon, const
 			HCGen.AngleSamples = PolygonProperties->Steps;
 			HCGen.RadialSamples = 1;
 			HCGen.Radius = BoxSize.X;
-			HCGen.HoleRadius = BoxSize.X * FMath::Clamp(PolygonProperties->FeatureSizeRatio, .01f, .99f);
+			HCGen.HoleRadius = BoxSize.X * FMathd::Clamp(PolygonProperties->FeatureSizeRatio, .01f, .99f);
 			ResultMeshOut->Copy(&HCGen.Generate());
 		}
 		else if (PolygonProperties->PolygonType == EDrawPolygonDrawMode::Circle)
@@ -1148,9 +1135,9 @@ bool UDrawPolygonTool::GeneratePolygonMesh(const TArray<FVector>& Polygon, const
 		else if (PolygonProperties->PolygonType == EDrawPolygonDrawMode::RoundedRectangle)
 		{
 			FRoundedRectangleMeshGenerator RRGen;
-			FirstReferencePt += RotationMat * (FVector2f(BoxSize.X, BoxSize.Y * YSign)*.5f);
+			FirstReferencePt += RotationMat * (FVector2d(BoxSize.X, BoxSize.Y * YSign)*.5f);
 			RRGen.AngleSamples = PolygonProperties->Steps;
-			RRGen.Radius = .5 * FMath::Min(BoxSize.X, BoxSize.Y) * FMath::Clamp(PolygonProperties->FeatureSizeRatio, .01f, .99f);
+			RRGen.Radius = .5 * FMathd::Min(BoxSize.X, BoxSize.Y) * FMathd::Clamp(PolygonProperties->FeatureSizeRatio, .01f, .99f);
 			RRGen.Height = BoxSize.Y - RRGen.Radius * 2.;
 			RRGen.Width = BoxSize.X - RRGen.Radius * 2.;
 			RRGen.WidthVertexCount = 1;
@@ -1162,7 +1149,7 @@ bool UDrawPolygonTool::GeneratePolygonMesh(const TArray<FVector>& Polygon, const
 		for (int VertIdx : ResultMeshOut->VertexIndicesItr())
 		{
 			FVector3d V = ResultMeshOut->GetVertex(VertIdx);
-			FVector2f VTransformed = RotationMat * FVector2f(V.X, V.Y) + FirstReferencePt;
+			FVector2d VTransformed = RotationMat * FVector2d(V.X, V.Y) + FirstReferencePt;
 			ResultMeshOut->SetVertex(VertIdx, FVector3d(VTransformed.X, VTransformed.Y, 0));
 		}
 	}
diff --git a/Engine/Plugins/Experimental/MeshModelingToolset/Source/MeshModelingTools/Public/DrawPolygonTool.h b/Engine/Plugins/Experimental/MeshModelingToolset/Source/MeshModelingTools/Public/DrawPolygonTool.h
--- a/Engine/Plugins/Experimental/MeshModelingToolset/Source/MeshModelingTools/Public/DrawPolygonTool.h
+++ b/Engine/Plugins/Experimental/MeshModelingToolset/Source/MeshModelingTools/Public/DrawPolygonTool.h
@@ -211,9 +211,9 @@ public:
 
 	// polygon drawing functions
 	virtual void ResetPolygon();
-	virtual void UpdatePreviewVertex(const FVector& PreviewVertex);
-	virtual void AppendVertex(const FVector& Vertex);
-	virtual bool FindDrawPlaneHitPoint(const FInputDeviceRay& ClickPos, FVector& HitPosOut);
+	virtual void UpdatePreviewVertex(const FVector3d& PreviewVertex);
+	virtual void AppendVertex(const FVector3d& Vertex);
+	virtual bool FindDrawPlaneHitPoint(const FInputDeviceRay& ClickPos, FVector3d& HitPosOut);
 	virtual void EmitCurrentPolygon();
 
 	virtual void BeginInteractiveExtrude();
@@ -243,23 +243,19 @@ protected:
 	
 
 	/** Origin of plane we will draw polygon on */
-	UPROPERTY()
-	FVector DrawPlaneOrigin;
+	FVector3d DrawPlaneOrigin;
 
 	/** Orientation of plane we will draw polygon on */
-	UPROPERTY()
-	FQuat DrawPlaneOrientation;
+	FQuaterniond DrawPlaneOrientation;
 	
 	/** Vertices of current preview polygon */
-	UPROPERTY()
-	TArray<FVector> PolygonVertices;
+	TArray<FVector3d> PolygonVertices;
 
 	/** Vertices of holes in current preview polygon */
-	TArray<TArray<FVector>> PolygonHolesVertices;
+	TArray<TArray<FVector3d>> PolygonHolesVertices;
 
 	/** last vertex of polygon that is actively being updated as input device is moved */
-	UPROPERTY()
-	FVector PreviewVertex;
+	FVector3d PreviewVertex;
 
 protected:
 	UWorld* TargetWorld;
@@ -286,7 +282,7 @@ protected:
 	IClickBehaviorTarget* SetPointInWorldConnector = nullptr;
 
 	// updates plane and gizmo position
-	virtual void SetDrawPlaneFromWorldPos(const FVector& Position, const FVector& Normal);
+	virtual void SetDrawPlaneFromWorldPos(const FVector3d& Position, const FVector3d& Normal);
 
 	TValueWatcher<bool> ShowGizmoWatcher;
 	void UpdateShowGizmoState(bool bNewVisibility);
@@ -297,26 +293,26 @@ protected:
 	bool bAbortActivePolygonDraw;
 
 	bool bInFixedPolygonMode = false;
-	TArray<FVector> FixedPolygonClickPoints;
+	TArray<FVector3d> FixedPolygonClickPoints;
 
 	// can close poly if current segment intersects existing segment
 	bool UpdateSelfIntersection();
 	bool bHaveSelfIntersection;
 	int SelfIntersectSegmentIdx;
-	FVector3f SelfIntersectionPoint;
+	FVector3d SelfIntersectionPoint;
 
 	// only used when SnapSettings.bHitSceneObjects = true
 	bool bHaveSurfaceHit;
-	FVector3f SurfaceHitPoint;
-	FVector3f SurfaceOffsetPoint;
+	FVector3d SurfaceHitPoint;
+	FVector3d SurfaceOffsetPoint;
 
 	bool bIgnoreSnappingToggle = false;		// toggled by hotkey (shift)
 	FPointPlanarSnapSolver SnapEngine;
 	ToolSceneQueriesUtil::FSnapGeometry LastSnapGeometry;
 	FVector3d LastGridSnapPoint;
 
-	void GetPolygonParametersFromFixedPoints(const TArray<FVector>& FixedPoints, FVector2f& FirstReferencePt, FVector2f& BoxSize, float& YSign, float& AngleRad);
-	void GenerateFixedPolygon(const TArray<FVector>& FixedPoints, TArray<FVector>& VerticesOut, TArray<TArray<FVector>>& HolesVerticesOut);
+	void GetPolygonParametersFromFixedPoints(const TArray<FVector3d>& FixedPoints, FVector2d& FirstReferencePt, FVector2d& BoxSize, double& YSign, double& AngleRad);
+	void GenerateFixedPolygon(const TArray<FVector3d>& FixedPoints, TArray<FVector3d>& VerticesOut, TArray<TArray<FVector3d>>& HolesVerticesOut);
 
 
 	// extrusion control
@@ -333,7 +329,7 @@ protected:
 	FFrame3f HitPosFrameWorld;
 
 	/** Generate extruded meshes.  Returns true on success. */
-	bool GeneratePolygonMesh(const TArray<FVector>& Polygon, const TArray<TArray<FVector>>& PolygonHoles, FDynamicMesh3* ResultMeshOut, FFrame3d& WorldFrameOut, bool bIncludePreviewVtx, double ExtrudeDistance, bool bExtrudeSymmetric);
+	bool GeneratePolygonMesh(const TArray<FVector3d>& Polygon, const TArray<TArray<FVector3d>>& PolygonHoles, FDynamicMesh3* ResultMeshOut, FFrame3d& WorldFrameOut, bool bIncludePreviewVtx, double ExtrudeDistance, bool bExtrudeSymmetric);
 
 
 	// user feedback messages
