Edigrate Dev-Editor updates to InteractiveToolsFramework, GeometryProcessing, MeshModelingToolset, and ModelingTools Editor Mode Plugins. Update a few files in new MeshPaint Plugin for UInteractiveToo

diff --git a/Engine/Plugins/Experimental/MeshModelingToolset/Source/MeshModelingTools/Private/AddPrimitiveTool.cpp b/Engine/Plugins/Experimental/MeshModelingToolset/Source/MeshModelingTools/Private/AddPrimitiveTool.cpp
--- a/Engine/Plugins/Experimental/MeshModelingToolset/Source/MeshModelingTools/Private/AddPrimitiveTool.cpp
+++ b/Engine/Plugins/Experimental/MeshModelingToolset/Source/MeshModelingTools/Private/AddPrimitiveTool.cpp
@@ -23,6 +23,8 @@
 #include "Drawing/MeshDebugDrawing.h"
 
 #include "DynamicMeshEditor.h"
+#include "UObject/PropertyIterator.h"
+#include "UObject/UnrealType.h"
 
 #define LOCTEXT_NAMESPACE "UAddPrimitiveTool"
 
@@ -37,103 +39,72 @@ bool UAddPrimitiveToolBuilder::CanBuildTool(const FToolBuilderState& SceneState)
 
 UInteractiveTool* UAddPrimitiveToolBuilder::BuildTool(const FToolBuilderState& SceneState) const
 {
-	UAddPrimitiveTool* NewTool = NewObject<UAddPrimitiveTool>(SceneState.ToolManager);
+	UAddPrimitiveTool* NewTool = nullptr;
+	switch (ShapeType)
+	{
+	case EMakeMeshShapeType::Box:
+		NewTool = NewObject<UAddBoxPrimitiveTool>(SceneState.ToolManager);
+		break;
+	case EMakeMeshShapeType::Cylinder:
+		NewTool = NewObject<UAddCylinderPrimitiveTool>(SceneState.ToolManager);
+		break;
+	case EMakeMeshShapeType::Cone:
+		NewTool = NewObject<UAddConePrimitiveTool>(SceneState.ToolManager);
+		break;
+	case EMakeMeshShapeType::Arrow:
+		NewTool = NewObject<UAddArrowPrimitiveTool>(SceneState.ToolManager);
+		break;
+	case EMakeMeshShapeType::Rectangle:
+		NewTool = NewObject<UAddRectanglePrimitiveTool>(SceneState.ToolManager);
+		break;
+	case EMakeMeshShapeType::RoundedRectangle:
+		NewTool = NewObject<UAddRoundedRectanglePrimitiveTool>(SceneState.ToolManager);
+		break;
+	case EMakeMeshShapeType::Disc:
+		NewTool = NewObject<UAddDiscPrimitiveTool>(SceneState.ToolManager);
+		break;
+	case EMakeMeshShapeType::PuncturedDisc:
+		NewTool = NewObject<UAddPuncturedDiscPrimitiveTool>(SceneState.ToolManager);
+		break;
+	case EMakeMeshShapeType::Torus:
+		NewTool = NewObject<UAddTorusPrimitiveTool>(SceneState.ToolManager);
+		break;
+	case EMakeMeshShapeType::SphericalBox:
+		NewTool = NewObject<UAddSphericalBoxPrimitiveTool>(SceneState.ToolManager);
+		break;
+	case EMakeMeshShapeType::Sphere:
+		NewTool = NewObject<UAddSpherePrimitiveTool>(SceneState.ToolManager);
+		break;
+	default:
+		break;
+	}
 	NewTool->SetWorld(SceneState.World);
 	NewTool->SetAssetAPI(AssetAPI);
 	return NewTool;
 }
 
-/*
- * Tool
- */
-UProceduralShapeToolProperties::UProceduralShapeToolProperties()
-{
-	Shape = EMakeMeshShapeType::Box;
-	Width = 100;
-	Height = 200;
-	FeatureRadius = 25;
-	//StartAngle = 0;
-	//EndAngle = 360;
-	Slices = 16;
-	Subdivisions = 0;
-	PolygroupMode = EMakeMeshPolygroupMode::PerFace;
-	PivotLocation = EMakeMeshPivotLocation::Base;
-	PlaceMode = EMakeMeshPlacementType::OnScene;
-	bAlignShapeToPlacementSurface = true;
-}
-
-void UProceduralShapeToolProperties::SaveProperties(UInteractiveTool* SaveFromTool)
-{
-	UProceduralShapeToolProperties* PropertyCache = GetPropertyCache<UProceduralShapeToolProperties>();
-	PropertyCache->Shape = this->Shape;
-	PropertyCache->Width = this->Width;
-	PropertyCache->Height = this->Height;
-	PropertyCache->FeatureRadius = this->FeatureRadius;
-	PropertyCache->Slices = this->Slices;
-	PropertyCache->Subdivisions = this->Subdivisions;
-	PropertyCache->PolygroupMode = this->PolygroupMode;
-	PropertyCache->PivotLocation = this->PivotLocation;
-	PropertyCache->PlaceMode = this->PlaceMode;
-	PropertyCache->bAlignShapeToPlacementSurface = this->bAlignShapeToPlacementSurface;
-	PropertyCache->bInstanceIfPossible = this->bInstanceIfPossible;
-}
-
-void UProceduralShapeToolProperties::RestoreProperties(UInteractiveTool* RestoreToTool)
-{
-	UProceduralShapeToolProperties* PropertyCache = GetPropertyCache<UProceduralShapeToolProperties>();
-	this->Shape = PropertyCache->Shape;
-	this->Width = PropertyCache->Width;
-	this->Height = PropertyCache->Height;
-	this->FeatureRadius = PropertyCache->FeatureRadius;
-	this->Slices = PropertyCache->Slices;
-	this->Subdivisions = PropertyCache->Subdivisions;
-	this->PolygroupMode = PropertyCache->PolygroupMode;
-	this->PivotLocation = PropertyCache->PivotLocation;
-	this->PlaceMode = PropertyCache->PlaceMode;
-	this->bAlignShapeToPlacementSurface = PropertyCache->bAlignShapeToPlacementSurface;
-	this->bInstanceIfPossible = PropertyCache->bInstanceIfPossible;
-
-}
-
-namespace
-{
-	struct {
-		FName Name;
-		EMakeMeshShapeType EnabledShapes;
-	} EnabledShapesMap[] =
-	{ 
-	  { TEXT("Shape"),         EMakeMeshShapeType::All },
-	  { TEXT("Width"),         EMakeMeshShapeType::All },
-	  { TEXT("Height"),        EMakeMeshShapeType::Box | EMakeMeshShapeType::Cylinder | EMakeMeshShapeType::Cone | EMakeMeshShapeType::Arrow | EMakeMeshShapeType::Rectangle | EMakeMeshShapeType::RoundedRectangle },
-	  { TEXT("FeatureRadius"),		   EMakeMeshShapeType::Arrow | EMakeMeshShapeType::RoundedRectangle | EMakeMeshShapeType::PuncturedDisc | EMakeMeshShapeType::Torus },
-	  { TEXT("Rotation"),      EMakeMeshShapeType::All },
-	  { TEXT("PolygroupMode"), EMakeMeshShapeType::All },
-	  { TEXT("PlaceMode"),     EMakeMeshShapeType::All },
-	  { TEXT("PivotLocation"), EMakeMeshShapeType::All },
-	  { TEXT("bAlignShapeToPlacementSurface"), EMakeMeshShapeType::All },
-	  { TEXT("bInstanceIfPossible"), EMakeMeshShapeType::All },
-	  { TEXT("Slices"),        EMakeMeshShapeType::Cylinder | EMakeMeshShapeType::Cone | EMakeMeshShapeType::Arrow | EMakeMeshShapeType::RoundedRectangle | EMakeMeshShapeType::Disc | EMakeMeshShapeType::PuncturedDisc | EMakeMeshShapeType::Sphere | EMakeMeshShapeType::Torus },
-	  { TEXT("Subdivisions"),  EMakeMeshShapeType::Box | EMakeMeshShapeType::Rectangle | EMakeMeshShapeType::RoundedRectangle | EMakeMeshShapeType::Disc | EMakeMeshShapeType::PuncturedDisc | EMakeMeshShapeType::Cylinder |
-							   EMakeMeshShapeType::Cone | EMakeMeshShapeType::Arrow | EMakeMeshShapeType::SphericalBox | EMakeMeshShapeType::Torus }
-	};
-};
-
-// UObject interface
-#if WITH_EDITOR
 bool
-UProceduralShapeToolProperties::CanEditChange( const FProperty* InProperty) const
+UProceduralShapeToolProperties::IsEquivalent( const UProceduralShapeToolProperties* Other ) const
 {
-	auto* Elem = Algo::FindByPredicate(EnabledShapesMap, [InProperty](const auto& Elem) { return Elem.Name == InProperty->GetFName(); });
-	if (Elem != nullptr)
+#if WITH_EDITOR
+	UClass* Class = GetClass();
+	if ( Other->GetClass() != Class )
 	{
-		return (Shape & Elem->EnabledShapes) != EMakeMeshShapeType::None;
+		return false;
 	}
-	checkNoEntry();
+	for ( FProperty* Prop : TFieldRange<FProperty>(Class) )
+	{
+		if (Prop->HasMetaData(TEXT("ProceduralShapeSetting")) &&
+			(!Prop->Identical_InContainer(this, Other)))
+		{
+			return false;
+		}
+	}
+	return true;
+#else
 	return false;
+#endif
 }
-#endif // WITH_EDITOR	
-// End of UObject interface
-
 
 void UAddPrimitiveTool::SetWorld(UWorld* World)
 {
@@ -145,6 +116,10 @@ void UAddPrimitiveTool::SetAssetAPI(IToolsContextAssetAPI* AssetAPIIn)
 	this->AssetAPI = AssetAPIIn;
 }
 
+UAddPrimitiveTool::UAddPrimitiveTool(const FObjectInitializer&)
+{
+	ShapeSettings = CreateDefaultSubobject<UProceduralShapeToolProperties>(TEXT("ShapeSettings"));
+}
 
 void UAddPrimitiveTool::Setup()
 {
@@ -154,7 +129,6 @@ void UAddPrimitiveTool::Setup()
 	HoverBehavior->Initialize(this);
 	AddInputBehavior(HoverBehavior);
 
-	ShapeSettings = NewObject<UProceduralShapeToolProperties>(this);
 	AddToolPropertySource(ShapeSettings);
 	ShapeSettings->RestoreProperties(this);
 
@@ -167,6 +141,7 @@ void UAddPrimitiveTool::Setup()
 	PreviewMesh->CreateInWorld(TargetWorld, FTransform::Identity);
 	PreviewMesh->SetVisible(false);
 	PreviewMesh->SetMaterial(MaterialProperties->Material);
+	PreviewMesh->EnableWireframe(MaterialProperties->bWireframe);
 
 	UpdatePreviewMesh();
 
@@ -278,57 +253,10 @@ void UAddPrimitiveTool::UpdatePreviewPosition(const FInputDeviceRay& DeviceClick
 	}
 }
 
-
 void UAddPrimitiveTool::UpdatePreviewMesh()
 {
 	FDynamicMesh3 NewMesh;
-	switch (ShapeSettings->Shape)
-	{
-	case EMakeMeshShapeType::Rectangle:
-		GenerateRectangle(&NewMesh);
-		break;
-
-	case EMakeMeshShapeType::RoundedRectangle:
-		GenerateRoundedRectangle(&NewMesh);
-		break;
-
-	case EMakeMeshShapeType::Disc:
-		GenerateDisc(&NewMesh);
-		break;
-
-	case EMakeMeshShapeType::PuncturedDisc:
-		GeneratePuncturedDisc(&NewMesh);
-		break;
-
-	case EMakeMeshShapeType::Cylinder:
-		GenerateCylinder(&NewMesh);
-		break;
-
-	case EMakeMeshShapeType::Cone:
-		GenerateCone(&NewMesh);
-		break;
-
-	case EMakeMeshShapeType::Arrow:
-		GenerateArrow(&NewMesh);
-		break;
-
-	case EMakeMeshShapeType::Torus:
-		GenerateTorus(&NewMesh);
-		break;
-
-	case EMakeMeshShapeType::Sphere:
-		GenerateSphere(&NewMesh);
-		break;
-
-	case EMakeMeshShapeType::SphericalBox:
-		GenerateSphericalBox(&NewMesh);
-		break;
-
-	case EMakeMeshShapeType::Box:
-	default:
-		GenerateBox(&NewMesh);
-		break;
-	}
+	GenerateMesh( &NewMesh );
 
 	if (ShapeSettings->PolygroupMode == EMakeMeshPolygroupMode::Single)
 	{
@@ -367,8 +295,6 @@ void UAddPrimitiveTool::UpdatePreviewMesh()
 void UAddPrimitiveTool::OnClicked(const FInputDeviceRay& DeviceClickPos)
 {
 #if WITH_EDITOR
-	const UEnum* const MakeMeshShapeTypeEnum = StaticEnum<EMakeMeshShapeType>();		
-	FString ShapeTypeName = MakeMeshShapeTypeEnum->GetNameStringByValue((int64)ShapeSettings->Shape);
 	UMaterialInterface* Material = PreviewMesh->GetMaterial();
 
 	if (ShapeSettings->bInstanceIfPossible && IsEquivalentLastGeneratedAsset())
@@ -393,12 +319,12 @@ void UAddPrimitiveTool::OnClicked(const FInputDeviceRay& DeviceClickPos)
 
 	const FDynamicMesh3* CurMesh = PreviewMesh->GetPreviewDynamicMesh();
 	FTransform3d CurTransform(PreviewMesh->GetTransform());
-	
+
 	GetToolManager()->BeginUndoTransaction(LOCTEXT("AddPrimitiveToolTransactionName", "Add Primitive Mesh"));
 
 	AActor* NewActor = AssetGenerationUtil::GenerateStaticMeshActor(
 		AssetAPI, TargetWorld,
-		CurMesh, CurTransform, ShapeTypeName, Material);
+		CurMesh, CurTransform, AssetName, Material);
 	if (NewActor != nullptr)
 	{
 		LastGenerated = NewObject<ULastActorInfo>(this);
@@ -418,14 +344,14 @@ void UAddPrimitiveTool::OnClicked(const FInputDeviceRay& DeviceClickPos)
 #endif
 }
 
-
-
-void UAddPrimitiveTool::GenerateBox(FDynamicMesh3* OutMesh)
+void UAddBoxPrimitiveTool::GenerateMesh(FDynamicMesh3* OutMesh) const
 {
 	FGridBoxMeshGenerator BoxGen;
-	BoxGen.Box = FOrientedBox3d(FVector3d::Zero(), 0.5*FVector3d(ShapeSettings->Width, ShapeSettings->Width, ShapeSettings->Height));
-	int EdgeNum = ShapeSettings->Subdivisions + 2;
-	BoxGen.EdgeVertices = FIndex3i(EdgeNum, EdgeNum, EdgeNum);
+	auto* BoxSettings = Cast<UProceduralBoxToolProperties>(ShapeSettings);
+	BoxGen.Box = FOrientedBox3d(FVector3d::Zero(), 0.5*FVector3d(BoxSettings->Width, BoxSettings->Height, BoxSettings->Depth));
+	BoxGen.EdgeVertices = FIndex3i(BoxSettings->WidthSubdivisions + 1,
+								   BoxSettings->DepthSubdivisions + 1,
+								   BoxSettings->HeightSubdivisions + 1);
 	if (ShapeSettings->PolygroupMode == EMakeMeshPolygroupMode::PerQuad)
 	{
 		BoxGen.bPolygroupPerQuad = true;
@@ -434,13 +360,14 @@ void UAddPrimitiveTool::GenerateBox(FDynamicMesh3* OutMesh)
 	OutMesh->Copy(&BoxGen);
 }
 
-
-void UAddPrimitiveTool::GenerateRectangle(FDynamicMesh3* OutMesh)
+void UAddRectanglePrimitiveTool::GenerateMesh(FDynamicMesh3* OutMesh) const
 {
 	FRectangleMeshGenerator RectGen;
-	RectGen.Width = ShapeSettings->Width;
-	RectGen.Height = ShapeSettings->Height;
-	RectGen.WidthVertexCount = RectGen.HeightVertexCount = ShapeSettings->Subdivisions + 2;
+	auto* RectangleSettings = Cast<UProceduralRectangleToolProperties>(ShapeSettings);
+	RectGen.Width = RectangleSettings->Width;
+	RectGen.Height = RectangleSettings->Depth;
+	RectGen.WidthVertexCount = RectangleSettings->WidthSubdivisions + 1;
+	RectGen.HeightVertexCount = RectangleSettings->DepthSubdivisions + 1;
 	if (ShapeSettings->PolygroupMode != EMakeMeshPolygroupMode::PerQuad)
 	{
 		RectGen.bSinglePolygroup = true;
@@ -449,29 +376,31 @@ void UAddPrimitiveTool::GenerateRectangle(FDynamicMesh3* OutMesh)
 	OutMesh->Copy(&RectGen);
 }
 
-void UAddPrimitiveTool::GenerateRoundedRectangle(FDynamicMesh3* OutMesh)
+void UAddRoundedRectanglePrimitiveTool::GenerateMesh(FDynamicMesh3* OutMesh) const
 {
 	FRoundedRectangleMeshGenerator RectGen;
-	RectGen.Width = ShapeSettings->Width;
-	RectGen.Height = ShapeSettings->Height;
-	RectGen.WidthVertexCount = RectGen.HeightVertexCount = ShapeSettings->Subdivisions + 2;
+	auto* RectangleSettings = Cast<UProceduralRoundedRectangleToolProperties>(ShapeSettings);
+	RectGen.Width = RectangleSettings->Width;
+	RectGen.Height = RectangleSettings->Depth;
+	RectGen.WidthVertexCount = RectangleSettings->WidthSubdivisions + 1;
+	RectGen.HeightVertexCount = RectangleSettings->DepthSubdivisions + 1;
 	if (ShapeSettings->PolygroupMode != EMakeMeshPolygroupMode::PerQuad)
 	{
 		RectGen.bSinglePolygroup = true;
 	}
-	RectGen.Radius = ShapeSettings->FeatureRadius;
-	RectGen.AngleSamples = ShapeSettings->Slices;
+	RectGen.Radius = RectangleSettings->CornerRadius;
+	RectGen.AngleSamples = RectangleSettings->CornerSlices - 1;
 	RectGen.Generate();
 	OutMesh->Copy(&RectGen);
 }
 
-
-void UAddPrimitiveTool::GenerateDisc(FDynamicMesh3* OutMesh)
+void UAddDiscPrimitiveTool::GenerateMesh(FDynamicMesh3* OutMesh) const
 {
 	FDiscMeshGenerator Gen;
-	Gen.Radius = ShapeSettings->Width * 0.5f;
-	Gen.AngleSamples = ShapeSettings->Slices;
-	Gen.RadialSamples = ShapeSettings->Subdivisions;
+	auto* DiscSettings = Cast<UProceduralDiscToolProperties>(ShapeSettings);
+	Gen.Radius = DiscSettings->Radius;
+	Gen.AngleSamples = DiscSettings->RadialSlices;
+	Gen.RadialSamples = DiscSettings->RadialSubdivisions;
 	if (ShapeSettings->PolygroupMode != EMakeMeshPolygroupMode::PerQuad)
 	{
 		Gen.bSinglePolygroup = true;
@@ -480,14 +409,14 @@ void UAddPrimitiveTool::GenerateDisc(FDynamicMesh3* OutMesh)
 	OutMesh->Copy(&Gen);
 }
 
-
-void UAddPrimitiveTool::GeneratePuncturedDisc(FDynamicMesh3* OutMesh)
+void UAddPuncturedDiscPrimitiveTool::GenerateMesh(FDynamicMesh3* OutMesh) const
 {
 	FPuncturedDiscMeshGenerator Gen;
-	Gen.Radius = ShapeSettings->Width * 0.5f;
-	Gen.HoleRadius = FMath::Min(ShapeSettings->FeatureRadius, Gen.Radius * .999f); // hole cannot be bigger than outer radius
-	Gen.AngleSamples = ShapeSettings->Slices;
-	Gen.RadialSamples = ShapeSettings->Subdivisions;
+	auto* DiscSettings = Cast<UProceduralPuncturedDiscToolProperties>(ShapeSettings);
+	Gen.Radius = DiscSettings->Radius;
+	Gen.HoleRadius = FMath::Min(DiscSettings->HoleRadius, Gen.Radius * .999f); // hole cannot be bigger than outer radius
+	Gen.AngleSamples = DiscSettings->RadialSlices;
+	Gen.RadialSamples = DiscSettings->RadialSubdivisions;
 	if (ShapeSettings->PolygroupMode != EMakeMeshPolygroupMode::PerQuad)
 	{
 		Gen.bSinglePolygroup = true;
@@ -496,12 +425,12 @@ void UAddPrimitiveTool::GeneratePuncturedDisc(FDynamicMesh3* OutMesh)
 	OutMesh->Copy(&Gen);
 }
 
-
-void UAddPrimitiveTool::GenerateTorus(FDynamicMesh3* OutMesh)
+void UAddTorusPrimitiveTool::GenerateMesh(FDynamicMesh3* OutMesh) const
 {
 	FGeneralizedCylinderGenerator Gen;
-	Gen.CrossSection = FPolygon2d::MakeCircle(ShapeSettings->FeatureRadius, ShapeSettings->Slices);
-	FPolygon2d PathCircle = FPolygon2d::MakeCircle(ShapeSettings->Width*.5, ShapeSettings->Subdivisions+4);
+	auto* TorusSettings = Cast<UProceduralTorusToolProperties>(ShapeSettings);
+	Gen.CrossSection = FPolygon2d::MakeCircle(TorusSettings->MinorRadius, TorusSettings->CrossSectionSlices);
+	FPolygon2d PathCircle = FPolygon2d::MakeCircle(TorusSettings->MajorRadius, TorusSettings->TubeSlices);
 	for (int Idx = 0; Idx < PathCircle.VertexCount(); Idx++)
 	{
 		Gen.Path.Add(PathCircle[Idx]);
@@ -517,15 +446,14 @@ void UAddPrimitiveTool::GenerateTorus(FDynamicMesh3* OutMesh)
 	OutMesh->Copy(&Gen);
 }
 
-
-void UAddPrimitiveTool::GenerateCylinder(FDynamicMesh3* OutMesh)
+void UAddCylinderPrimitiveTool::GenerateMesh(FDynamicMesh3* OutMesh) const
 {
 	FCylinderGenerator CylGen;
-	CylGen.Radius[0] = ShapeSettings->Width * 0.5f;
-	CylGen.Radius[1] = CylGen.Radius[0];
-	CylGen.Height = ShapeSettings->Height;
-	CylGen.AngleSamples = ShapeSettings->Slices;
-	CylGen.LengthSamples = ShapeSettings->Subdivisions;
+	auto* CylinderSettings = Cast<UProceduralCylinderToolProperties>(ShapeSettings);
+	CylGen.Radius[1] = CylGen.Radius[0] = CylinderSettings->Radius;
+	CylGen.Height = CylinderSettings->Height;
+	CylGen.AngleSamples = CylinderSettings->RadialSlices;
+	CylGen.LengthSamples = CylinderSettings->HeightSubdivisions - 1;
 	CylGen.bCapped = true;
 	if (ShapeSettings->PolygroupMode == EMakeMeshPolygroupMode::PerQuad)
 	{
@@ -535,16 +463,16 @@ void UAddPrimitiveTool::GenerateCylinder(FDynamicMesh3* OutMesh)
 	OutMesh->Copy(&CylGen);
 }
 
-
-void UAddPrimitiveTool::GenerateCone(FDynamicMesh3* OutMesh)
+void UAddConePrimitiveTool::GenerateMesh(FDynamicMesh3* OutMesh) const
 {
 	// Unreal's standard cone is just a cylinder with a very small top
 	FCylinderGenerator CylGen;
-	CylGen.Radius[0] = ShapeSettings->Width * 0.5f;
+	auto* ConeSettings = Cast<UProceduralConeToolProperties>(ShapeSettings);
+	CylGen.Radius[0] = ConeSettings->Radius;
 	CylGen.Radius[1] = .01;
-	CylGen.Height = ShapeSettings->Height;
-	CylGen.AngleSamples = ShapeSettings->Slices;
-	CylGen.LengthSamples = ShapeSettings->Subdivisions;
+	CylGen.Height = ConeSettings->Height;
+	CylGen.AngleSamples = ConeSettings->RadialSlices;
+	CylGen.LengthSamples = ConeSettings->HeightSubdivisions - 1;
 	CylGen.bCapped = true;
 	if (ShapeSettings->PolygroupMode == EMakeMeshPolygroupMode::PerQuad)
 	{
@@ -554,33 +482,33 @@ void UAddPrimitiveTool::GenerateCone(FDynamicMesh3* OutMesh)
 	OutMesh->Copy(&CylGen);
 }
 
-
-void UAddPrimitiveTool::GenerateArrow(FDynamicMesh3* OutMesh)
+void UAddArrowPrimitiveTool::GenerateMesh(FDynamicMesh3* OutMesh) const
 {
 	FArrowGenerator ArrowGen;
-	ArrowGen.StickRadius = ShapeSettings->FeatureRadius;
-	ArrowGen.StickLength = ShapeSettings->Height * .25f;
-	ArrowGen.HeadBaseRadius = ShapeSettings->Width * .5f;
-	ArrowGen.TipRadius = .01f;
-	ArrowGen.HeadLength = ShapeSettings->Height * .75f;
-	ArrowGen.AngleSamples = ShapeSettings->Slices;
+	auto* ArrowSettings = Cast<UProceduralArrowToolProperties>(ShapeSettings);
+	ArrowGen.StickRadius = ArrowSettings->ShaftRadius;
+	ArrowGen.StickLength = ArrowSettings->ShaftHeight;
+	ArrowGen.HeadBaseRadius = ArrowSettings->HeadRadius;
+	ArrowGen.HeadTipRadius = .01f;
+	ArrowGen.HeadLength = ArrowSettings->HeadHeight;
+	ArrowGen.AngleSamples = ArrowSettings->RadialSlices;
 	ArrowGen.bCapped = true;
 	if (ShapeSettings->PolygroupMode == EMakeMeshPolygroupMode::PerQuad)
 	{
 		ArrowGen.bPolygroupPerQuad = true;
 	}
-	ArrowGen.DistributeAdditionalLengthSamples(ShapeSettings->Subdivisions);
+	ArrowGen.DistributeAdditionalLengthSamples(ArrowSettings->TotalSubdivisions);
 	ArrowGen.Generate();
 	OutMesh->Copy(&ArrowGen);
 }
 
-
-void UAddPrimitiveTool::GenerateSphere(FDynamicMesh3* OutMesh)
+void UAddSpherePrimitiveTool::GenerateMesh(FDynamicMesh3* OutMesh) const
 {
 	FSphereGenerator SphereGen;
-	SphereGen.Radius = ShapeSettings->Width * 0.5f;
-	SphereGen.NumTheta = ShapeSettings->Slices;
-	SphereGen.NumPhi = ShapeSettings->Slices;
+	auto* SphereSettings = Cast<UProceduralSphereToolProperties>(ShapeSettings);
+	SphereGen.Radius = SphereSettings->Radius;
+	SphereGen.NumTheta = SphereSettings->LongitudeSlices + 1;
+	SphereGen.NumPhi = SphereSettings->LatitudeSlices + 1;
 	if (ShapeSettings->PolygroupMode == EMakeMeshPolygroupMode::PerQuad)
 	{
 		SphereGen.bPolygroupPerQuad = true;
@@ -589,13 +517,16 @@ void UAddPrimitiveTool::GenerateSphere(FDynamicMesh3* OutMesh)
 	OutMesh->Copy(&SphereGen);
 }
 
-
-void UAddPrimitiveTool::GenerateSphericalBox(FDynamicMesh3* OutMesh)
+void UAddSphericalBoxPrimitiveTool::GenerateMesh(FDynamicMesh3* OutMesh) const
 {
 	FBoxSphereGenerator SphereGen;
-	SphereGen.Radius = ShapeSettings->Width * 0.5f;
-	SphereGen.Box = FOrientedBox3d(FVector3d::Zero(), 0.5*FVector3d(ShapeSettings->Width, ShapeSettings->Width, ShapeSettings->Width));
-	int EdgeNum = ShapeSettings->Subdivisions + 3;
+	auto* SphereSettings = Cast<UProceduralSphericalBoxToolProperties>(ShapeSettings);
+	SphereGen.Radius = SphereSettings->Radius;
+	SphereGen.Box = FOrientedBox3d(FVector3d::Zero(),
+								   0.5*FVector3d(SphereSettings->Subdivisions + 1,
+												 SphereSettings->Subdivisions + 1,
+												 SphereSettings->Subdivisions + 1));
+	int EdgeNum = SphereSettings->Subdivisions + 1;
 	SphereGen.EdgeVertices = FIndex3i(EdgeNum, EdgeNum, EdgeNum);
 	if (ShapeSettings->PolygroupMode == EMakeMeshPolygroupMode::PerQuad)
 	{
@@ -605,5 +536,4 @@ void UAddPrimitiveTool::GenerateSphericalBox(FDynamicMesh3* OutMesh)
 	OutMesh->Copy(&SphereGen);
 }
 
-
 #undef LOCTEXT_NAMESPACE
diff --git a/Engine/Plugins/Experimental/MeshModelingToolset/Source/MeshModelingTools/Private/AlignObjectsTool.cpp b/Engine/Plugins/Experimental/MeshModelingToolset/Source/MeshModelingTools/Private/AlignObjectsTool.cpp
--- a/Engine/Plugins/Experimental/MeshModelingToolset/Source/MeshModelingTools/Private/AlignObjectsTool.cpp
+++ b/Engine/Plugins/Experimental/MeshModelingToolset/Source/MeshModelingTools/Private/AlignObjectsTool.cpp
@@ -50,32 +50,6 @@ UInteractiveTool* UAlignObjectsToolBuilder::BuildTool(const FToolBuilderState& S
 	return NewTool;
 }
 
-
-
-
-void UAlignObjectsToolProperties::SaveProperties(UInteractiveTool* SaveFromTool)
-{
-	UAlignObjectsToolProperties* PropertyCache = GetPropertyCache<UAlignObjectsToolProperties>();
-	PropertyCache->AlignType = this->AlignType;
-	PropertyCache->AlignTo = this->AlignTo;
-	PropertyCache->BoxPosition = this->BoxPosition;
-	PropertyCache->bAlignX = this->bAlignX;
-	PropertyCache->bAlignY = this->bAlignY;
-	PropertyCache->bAlignZ = this->bAlignZ;
-}
-
-void UAlignObjectsToolProperties::RestoreProperties(UInteractiveTool* RestoreToTool)
-{
-	UAlignObjectsToolProperties* PropertyCache = GetPropertyCache<UAlignObjectsToolProperties>();
-	this->AlignType = PropertyCache->AlignType;
-	this->AlignTo = PropertyCache->AlignTo;
-	this->BoxPosition = PropertyCache->BoxPosition;
-	this->bAlignX = PropertyCache->bAlignX;
-	this->bAlignY = PropertyCache->bAlignY;
-	this->bAlignZ = PropertyCache->bAlignZ;
-}
-
-
 /*
  * Tool
  */
@@ -140,7 +114,7 @@ void UAlignObjectsTool::Shutdown(EToolShutdownType ShutdownType)
 
 
 
-void UAlignObjectsTool::Tick(float DeltaTime)
+void UAlignObjectsTool::OnTick(float DeltaTime)
 {
 	if (bAlignDirty)
 	{
diff --git a/Engine/Plugins/Experimental/MeshModelingToolset/Source/MeshModelingTools/Private/BakeTransformTool.cpp b/Engine/Plugins/Experimental/MeshModelingToolset/Source/MeshModelingTools/Private/BakeTransformTool.cpp
--- a/Engine/Plugins/Experimental/MeshModelingToolset/Source/MeshModelingTools/Private/BakeTransformTool.cpp
+++ b/Engine/Plugins/Experimental/MeshModelingToolset/Source/MeshModelingTools/Private/BakeTransformTool.cpp
@@ -131,16 +131,6 @@ void UBakeTransformTool::SetAssetAPI(IToolsContextAssetAPI* AssetAPIIn)
 
 
 
-void UBakeTransformTool::Render(IToolsContextRenderAPI* RenderAPI)
-{
-}
-
-void UBakeTransformTool::Tick(float DeltaTime)
-{
-}
-
-
-
 bool UBakeTransformTool::HasAccept() const
 {
 	return true;
diff --git a/Engine/Plugins/Experimental/MeshModelingToolset/Source/MeshModelingTools/Private/CSGMeshesTool.cpp b/Engine/Plugins/Experimental/MeshModelingToolset/Source/MeshModelingTools/Private/CSGMeshesTool.cpp
new file mode 100644
--- /dev/null
+++ b/Engine/Plugins/Experimental/MeshModelingToolset/Source/MeshModelingTools/Private/CSGMeshesTool.cpp
@@ -0,0 +1,383 @@
+// Copyright Epic Games, Inc. All Rights Reserved.
+
+#include "CSGMeshesTool.h"
+#include "CompositionOps/BooleanMeshesOp.h"
+#include "InteractiveToolManager.h"
+#include "ToolBuilderUtil.h"
+
+#include "ToolSetupUtil.h"
+
+#include "Selection/ToolSelectionUtil.h"
+
+#include "DynamicMesh3.h"
+#include "MeshTransforms.h"
+
+#include "MeshDescriptionToDynamicMesh.h"
+#include "DynamicMeshToMeshDescription.h"
+
+#include "InteractiveGizmoManager.h"
+
+#include "BaseGizmos/GizmoComponents.h"
+#include "BaseGizmos/TransformGizmo.h"
+
+#include "AssetGenerationUtil.h"
+
+
+#define LOCTEXT_NAMESPACE "UCSGMeshesTool"
+
+
+/*
+ * ToolBuilder
+ */
+
+
+bool UCSGMeshesToolBuilder::CanBuildTool(const FToolBuilderState& SceneState) const
+{
+	return AssetAPI != nullptr && ToolBuilderUtil::CountComponents(SceneState, CanMakeComponentTarget) == 2;
+}
+
+UInteractiveTool* UCSGMeshesToolBuilder::BuildTool(const FToolBuilderState& SceneState) const
+{
+	UCSGMeshesTool* NewTool = NewObject<UCSGMeshesTool>(SceneState.ToolManager);
+
+	TArray<UActorComponent*> Components = ToolBuilderUtil::FindAllComponents(SceneState, CanMakeComponentTarget);
+	check(Components.Num() > 0);
+
+	TArray<TUniquePtr<FPrimitiveComponentTarget>> ComponentTargets;
+	for (UActorComponent* ActorComponent : Components)
+	{
+		auto* MeshComponent = Cast<UPrimitiveComponent>(ActorComponent);
+		if ( MeshComponent )
+		{
+			ComponentTargets.Add(MakeComponentTarget(MeshComponent));
+		}
+	}
+
+	NewTool->SetSelection(MoveTemp(ComponentTargets));
+	NewTool->SetWorld(SceneState.World);
+	NewTool->SetAssetAPI(AssetAPI);
+
+	return NewTool;
+}
+
+
+
+/*
+ * Tool
+ */
+UCSGMeshesTool::UCSGMeshesTool()
+{
+}
+
+void UCSGMeshesTool::SetWorld(UWorld* World)
+{
+	this->TargetWorld = World;
+}
+
+void UCSGMeshesTool::Setup()
+{
+	UInteractiveTool::Setup();
+
+	// hide input StaticMeshComponents
+	for (TUniquePtr<FPrimitiveComponentTarget>& ComponentTarget : ComponentTargets)
+	{
+		ComponentTarget->SetOwnerVisibility(false);
+	}
+
+
+	// initialize our properties
+	CSGProperties = NewObject<UCSGMeshesToolProperties>(this);
+	CSGProperties->RestoreProperties(this);
+	AddToolPropertySource(CSGProperties);
+
+
+	// initialize the PreviewMesh+BackgroundCompute object
+	SetupPreview();
+
+	Preview->InvalidateResult();
+}
+
+
+
+
+void UCSGMeshesTool::SetupPreview()
+{
+	FComponentMaterialSet AllMaterialSet;
+	TMap<UMaterialInterface*, int> KnownMaterials;
+	TArray<TArray<int>> MaterialRemap; MaterialRemap.SetNum(ComponentTargets.Num());
+	for (int ComponentIdx = 0; ComponentIdx < ComponentTargets.Num(); ComponentIdx++)
+	{
+		FComponentMaterialSet ComponentMaterialSet;
+		ComponentTargets[ComponentIdx]->GetMaterialSet(ComponentMaterialSet);
+		for (UMaterialInterface* Mat : ComponentMaterialSet.Materials)
+		{
+			int* FoundMatIdx = KnownMaterials.Find(Mat);
+			int MatIdx;
+			if (FoundMatIdx)
+			{
+				MatIdx = *FoundMatIdx;
+			}
+			else
+			{
+				MatIdx = AllMaterialSet.Materials.Add(Mat);
+				KnownMaterials.Add(Mat, MatIdx);
+			}
+			MaterialRemap[ComponentIdx].Add(MatIdx);
+		}
+	}
+
+	OriginalDynamicMeshes.SetNum(ComponentTargets.Num());
+
+	for (int ComponentIdx = 0; ComponentIdx < ComponentTargets.Num(); ComponentIdx++)
+	{
+		OriginalDynamicMeshes[ComponentIdx] = MakeShared<FDynamicMesh3>();
+		FMeshDescriptionToDynamicMesh Converter;
+		Converter.Convert(ComponentTargets[ComponentIdx]->GetMesh(), *OriginalDynamicMeshes[ComponentIdx]);
+
+		// ensure materials and attributes are always enabled
+		OriginalDynamicMeshes[ComponentIdx]->EnableAttributes();
+		OriginalDynamicMeshes[ComponentIdx]->Attributes()->EnableMaterialID();
+		FDynamicMeshMaterialAttribute* MaterialIDs = OriginalDynamicMeshes[ComponentIdx]->Attributes()->GetMaterialID();
+		for (int TID : OriginalDynamicMeshes[ComponentIdx]->TriangleIndicesItr())
+		{
+			MaterialIDs->SetValue(TID, MaterialRemap[ComponentIdx][MaterialIDs->GetValue(TID)]);
+		}
+	}
+
+	Preview = NewObject<UMeshOpPreviewWithBackgroundCompute>(this, "Preview");
+	Preview->Setup(this->TargetWorld, this);
+	Preview->ConfigureMaterials(AllMaterialSet.Materials, ToolSetupUtil::GetDefaultWorkingMaterial(GetToolManager()));
+
+	
+	DrawnLineSet = NewObject<ULineSetComponent>(Preview->PreviewMesh->GetRootComponent());
+	DrawnLineSet->SetupAttachment(Preview->PreviewMesh->GetRootComponent());
+	DrawnLineSet->SetLineMaterial(ToolSetupUtil::GetDefaultLineComponentMaterial(GetToolManager()));
+	DrawnLineSet->RegisterComponent();
+
+	Preview->OnOpCompleted.AddLambda(
+		[this](const FDynamicMeshOperator* Op)
+		{
+			const FBooleanMeshesOp* BooleanOp = (const FBooleanMeshesOp*)(Op);
+			CreatedBoundaryEdges = BooleanOp->GetCreatedBoundaryEdges();
+		}
+	);
+	Preview->OnMeshUpdated.AddLambda(
+		[this](const UMeshOpPreviewWithBackgroundCompute*)
+		{
+			GetToolManager()->PostInvalidation();
+			UpdateVisualization();
+		}
+	);
+	SetTransformGizmos();
+}
+
+void UCSGMeshesTool::UpdateVisualization()
+{
+	FColor BoundaryEdgeColor(240, 15, 15);
+	float BoundaryEdgeThickness = 2.0;
+	float BoundaryEdgeDepthBias = 2.0f;
+
+	const FDynamicMesh3* TargetMesh = Preview->PreviewMesh->GetPreviewDynamicMesh();
+	FVector3d A, B;
+
+	DrawnLineSet->Clear();
+	if (CSGProperties->bShowNewBoundaryEdges)
+	{
+		for (int EID : CreatedBoundaryEdges)
+		{
+			TargetMesh->GetEdgeV(EID, A, B);
+			DrawnLineSet->AddLine((FVector)A, (FVector)B, BoundaryEdgeColor, BoundaryEdgeThickness, BoundaryEdgeDepthBias);
+		}
+	}
+}
+
+void UCSGMeshesTool::UpdateGizmoVisibility()
+{
+	for (UTransformGizmo* Gizmo : TransformGizmos)
+	{
+		Gizmo->SetVisibility(CSGProperties->bShowTransformUI);
+	}
+}
+
+void UCSGMeshesTool::SetTransformGizmos()
+{
+	UInteractiveGizmoManager* GizmoManager = GetToolManager()->GetPairedGizmoManager();
+
+	for (int ComponentIdx = 0; ComponentIdx < ComponentTargets.Num(); ComponentIdx++)
+	{
+		TUniquePtr<FPrimitiveComponentTarget>& Target = ComponentTargets[ComponentIdx];
+		UTransformProxy* Proxy = TransformProxies.Add_GetRef(NewObject<UTransformProxy>(this));
+		UTransformGizmo* Gizmo = TransformGizmos.Add_GetRef(GizmoManager->Create3AxisTransformGizmo(this));
+		Gizmo->SetActiveTarget(Proxy);
+		FTransform InitialTransform = Target->GetWorldTransform();
+		Gizmo->SetNewGizmoTransform(InitialTransform);
+		Proxy->OnTransformChanged.AddUObject(this, &UCSGMeshesTool::TransformChanged);
+	}
+	UpdateGizmoVisibility();
+}
+
+void UCSGMeshesTool::TransformChanged(UTransformProxy* Proxy, FTransform Transform)
+{
+	Preview->InvalidateResult();
+}
+
+void UCSGMeshesTool::Shutdown(EToolShutdownType ShutdownType)
+{
+	CSGProperties->SaveProperties(this);
+
+	FDynamicMeshOpResult Result = Preview->Shutdown();
+	if (ShutdownType == EToolShutdownType::Accept)
+	{
+		// Generate the result
+		{
+			GetToolManager()->BeginUndoTransaction(LOCTEXT("BooleanMeshes", "Boolean Meshes"));
+
+			GenerateAsset(Result);
+
+			GetToolManager()->EndUndoTransaction();
+		}
+
+		// Hide or destroy the sources
+		{
+
+			if (CSGProperties->bDeleteInputActors)
+			{
+				GetToolManager()->BeginUndoTransaction(LOCTEXT("RemoveSources", "Remove Sources"));
+			}
+
+			for (auto& ComponentTarget : ComponentTargets)
+			{
+				ComponentTarget->SetOwnerVisibility(true);
+				AActor* Actor = ComponentTarget->GetOwnerActor();
+				if (CSGProperties->bDeleteInputActors)
+				{
+					Actor->Destroy();
+				}
+				else
+				{
+#if WITH_EDITOR
+					Actor->SetIsTemporarilyHiddenInEditor(true);
+#endif
+				}
+			}
+
+			if (CSGProperties->bDeleteInputActors)
+			{
+				GetToolManager()->EndUndoTransaction();
+			}
+
+		}
+
+	}
+	else
+	{
+		// Restore (unhide) the source meshes
+		for (auto& ComponentTarget : ComponentTargets)
+		{
+			ComponentTarget->SetOwnerVisibility(true);
+		}
+	}
+
+	UInteractiveGizmoManager* GizmoManager = GetToolManager()->GetPairedGizmoManager();
+	GizmoManager->DestroyAllGizmosByOwner(this);
+}
+
+void UCSGMeshesTool::SetAssetAPI(IToolsContextAssetAPI* AssetAPIIn)
+{
+	this->AssetAPI = AssetAPIIn;
+}
+
+TUniquePtr<FDynamicMeshOperator> UCSGMeshesTool::MakeNewOperator()
+{
+	TUniquePtr<FBooleanMeshesOp> BooleanOp = MakeUnique<FBooleanMeshesOp>();
+	
+	BooleanOp->Operation = CSGProperties->Operation;
+	BooleanOp->bAttemptFixHoles = CSGProperties->bAttemptFixHoles;
+
+	check(OriginalDynamicMeshes.Num() == 2);
+	check(ComponentTargets.Num() == 2);
+	BooleanOp->Transforms.SetNum(2);
+	BooleanOp->Meshes.SetNum(2);
+	for (int Idx = 0; Idx < 2; Idx++)
+	{
+		BooleanOp->Meshes[Idx] = OriginalDynamicMeshes[Idx];
+		BooleanOp->Transforms[Idx] = TransformProxies[Idx]->GetTransform();
+	}
+
+	return BooleanOp;
+}
+
+
+
+void UCSGMeshesTool::Render(IToolsContextRenderAPI* RenderAPI)
+{
+}
+
+void UCSGMeshesTool::OnTick(float DeltaTime)
+{
+	Preview->Tick(DeltaTime);
+}
+
+
+#if WITH_EDITOR
+void UCSGMeshesTool::PostEditChangeProperty(FPropertyChangedEvent& PropertyChangedEvent)
+{
+	Preview->InvalidateResult();
+	UpdateGizmoVisibility();
+}
+#endif
+
+void UCSGMeshesTool::OnPropertyModified(UObject* PropertySet, FProperty* Property)
+{
+	if (Property && (Property->GetFName() == GET_MEMBER_NAME_CHECKED(UCSGMeshesToolProperties, bShowTransformUI)))
+	{
+		UpdateGizmoVisibility();
+	}
+	else if (Property && (Property->GetFName() == GET_MEMBER_NAME_CHECKED(UCSGMeshesToolProperties, bDeleteInputActors)))
+	{
+		// nothing
+	}
+	else if (Property && (Property->GetFName() == GET_MEMBER_NAME_CHECKED(UCSGMeshesToolProperties, bShowNewBoundaryEdges)))
+	{
+		GetToolManager()->PostInvalidation();
+		UpdateVisualization();
+	}
+	else
+	{
+		Preview->InvalidateResult();
+	}
+}
+
+
+bool UCSGMeshesTool::HasAccept() const
+{
+	return true;
+}
+
+bool UCSGMeshesTool::CanAccept() const
+{
+	return Preview->HaveValidResult();
+}
+
+
+void UCSGMeshesTool::GenerateAsset(const FDynamicMeshOpResult& Result)
+{
+	check(Result.Mesh.Get() != nullptr);
+
+	FVector3d Center = Result.Mesh->GetCachedBounds().Center();
+	MeshTransforms::Translate(*Result.Mesh, -Center);
+	FTransform3d CenteredTransform = Result.Transform;
+	CenteredTransform.SetTranslation(CenteredTransform.GetTranslation() + Result.Transform.TransformVector(Center));
+	AActor* NewActor = AssetGenerationUtil::GenerateStaticMeshActor(
+		AssetAPI, TargetWorld,
+		Result.Mesh.Get(), CenteredTransform, TEXT("CSGMesh"), Preview->StandardMaterials);
+	if (NewActor != nullptr)
+	{
+		ToolSelectionUtil::SetNewActorSelection(GetToolManager(), NewActor);
+	}
+}
+
+
+
+
+#undef LOCTEXT_NAMESPACE
diff --git a/Engine/Plugins/Experimental/MeshModelingToolset/Source/MeshModelingTools/Private/CombineMeshesTool.cpp b/Engine/Plugins/Experimental/MeshModelingToolset/Source/MeshModelingTools/Private/CombineMeshesTool.cpp
--- a/Engine/Plugins/Experimental/MeshModelingToolset/Source/MeshModelingTools/Private/CombineMeshesTool.cpp
+++ b/Engine/Plugins/Experimental/MeshModelingToolset/Source/MeshModelingTools/Private/CombineMeshesTool.cpp
@@ -126,18 +126,6 @@ void UCombineMeshesTool::SetAssetAPI(IToolsContextAssetAPI* AssetAPIIn)
 	this->AssetAPI = AssetAPIIn;
 }
 
-
-
-void UCombineMeshesTool::Render(IToolsContextRenderAPI* RenderAPI)
-{
-}
-
-void UCombineMeshesTool::Tick(float DeltaTime)
-{
-}
-
-
-
 bool UCombineMeshesTool::HasAccept() const
 {
 	return true;
diff --git a/Engine/Plugins/Experimental/MeshModelingToolset/Source/MeshModelingTools/Private/ConvertToPolygonsTool.cpp b/Engine/Plugins/Experimental/MeshModelingToolset/Source/MeshModelingTools/Private/ConvertToPolygonsTool.cpp
--- a/Engine/Plugins/Experimental/MeshModelingToolset/Source/MeshModelingTools/Private/ConvertToPolygonsTool.cpp
+++ b/Engine/Plugins/Experimental/MeshModelingToolset/Source/MeshModelingTools/Private/ConvertToPolygonsTool.cpp
@@ -40,23 +40,9 @@ UInteractiveTool* UConvertToPolygonsToolBuilder::BuildTool(const FToolBuilderSta
 	return NewTool;
 }
 
-
-
-void UConvertToPolygonsToolProperties::SaveRestoreProperties(UInteractiveTool* RestoreToTool, bool bSaving)
-{
-	UConvertToPolygonsToolProperties* PropertyCache = GetPropertyCache<UConvertToPolygonsToolProperties>();
-	SaveRestoreProperty(PropertyCache->ConversionMode, this->ConversionMode, bSaving);
-	SaveRestoreProperty(PropertyCache->AngleTolerance, this->AngleTolerance, bSaving);
-	SaveRestoreProperty(PropertyCache->bCalculateNormals, this->bCalculateNormals, bSaving);
-	SaveRestoreProperty(PropertyCache->bShowGroupColors, this->bShowGroupColors, bSaving);
-}
-
-
-
 /*
  * Tool
  */
-
 UConvertToPolygonsTool::UConvertToPolygonsTool()
 {
 	SetToolDisplayName(LOCTEXT("ConvertToPolygonsToolName", "Find PolyGroups"));
@@ -92,13 +78,11 @@ void UConvertToPolygonsTool::Setup()
 	ComponentTarget->GetMaterialSet(MaterialSet);
 	PreviewMesh->SetMaterials(MaterialSet.Materials);
 
-	ConvertModeWatcher.Initialize(
-		[this]() { return Settings->ConversionMode; },
-		[this](EConvertToPolygonsMode NewMode) { bPolygonsValid = false; }, Settings->ConversionMode);
-
-	ShowGroupsWatcher.Initialize(
-		[this]() { return Settings->bShowGroupColors; },
-		[this](bool bNewValue) { UpdateVisualization(); }, Settings->bShowGroupColors );
+	Settings->WatchProperty(Settings->ConversionMode,
+							[this](EConvertToPolygonsMode NewMode)
+							{ bPolygonsValid = false; });
+	Settings->WatchProperty(Settings->bShowGroupColors,
+							[this](bool bNewValue) { UpdateVisualization(); });
 	if (Settings->bShowGroupColors)
 	{
 		UpdateVisualization();
@@ -137,11 +121,8 @@ void UConvertToPolygonsTool::OnPropertyModified(UObject* PropertySet, FProperty*
 }
 
 
-void UConvertToPolygonsTool::Tick(float DeltaTime)
+void UConvertToPolygonsTool::OnTick(float DeltaTime)
 {
-	ConvertModeWatcher.CheckAndUpdate();
-	ShowGroupsWatcher.CheckAndUpdate();
-
 	if (bPolygonsValid == false)
 	{
 		UpdatePolygons();
diff --git a/Engine/Plugins/Experimental/MeshModelingToolset/Source/MeshModelingTools/Private/DeformMeshPolygonsTool.cpp b/Engine/Plugins/Experimental/MeshModelingToolset/Source/MeshModelingTools/Private/DeformMeshPolygonsTool.cpp
--- a/Engine/Plugins/Experimental/MeshModelingToolset/Source/MeshModelingTools/Private/DeformMeshPolygonsTool.cpp
+++ b/Engine/Plugins/Experimental/MeshModelingToolset/Source/MeshModelingTools/Private/DeformMeshPolygonsTool.cpp
@@ -56,13 +56,6 @@ UDeformMeshPolygonsTransformProperties::UDeformMeshPolygonsTransformProperties()
 	bSnapToWorldGrid = false;
 }
 
-#if WITH_EDITOR
-void UDeformMeshPolygonsTransformProperties::PostEditChangeProperty(FPropertyChangedEvent& PropertyChangedEvent)
-{
-	Super::PostEditChangeProperty(PropertyChangedEvent);
-}
-#endif
-
 /*
  * Asynchronous Task
  */
@@ -1024,7 +1017,7 @@ void UDeformMeshPolygonsTool::UpdateChangeFromROI(bool bFinal)
 
 	FDynamicMesh3* Mesh = DynamicMeshComponent->GetMesh();
 	const TSet<int>& ModifiedVertices = (bIsLaplacian) ? LaplacianDeformer.GetModifiedVertices() : LinearDeformer.GetModifiedVertices();
-	ActiveVertexChange->SavePositions(Mesh, ModifiedVertices, !bFinal);
+	ActiveVertexChange->SaveVertices(Mesh, ModifiedVertices, !bFinal);
 	const TSet<int>& ModifiedNormals = (bIsLaplacian) ? LaplacianDeformer.GetModifiedOverlayNormals() : LinearDeformer.GetModifiedOverlayNormals();
 	ActiveVertexChange->SaveOverlayNormals(Mesh, ModifiedNormals, !bFinal);
 }
@@ -1325,10 +1318,8 @@ void UDeformMeshPolygonsTool::ComputeUpdate_Translate()
 
 
 
-void UDeformMeshPolygonsTool::Tick(float DeltaTime)
+void UDeformMeshPolygonsTool::OnTick(float DeltaTime)
 {
-	UMeshSurfacePointTool::Tick(DeltaTime);
-	
 	LaplacianDeformer.HandleWeights   = TransformProps->HandleWeight;
 	LaplacianDeformer.bPostfixHandles = TransformProps->bPostFixHandles;
 }
@@ -1448,7 +1439,7 @@ void UDeformMeshPolygonsTool::BeginChange()
 	{
 		if (ActiveVertexChange == nullptr)
 		{
-			ActiveVertexChange = new FMeshVertexChangeBuilder(true);
+			ActiveVertexChange = new FMeshVertexChangeBuilder(EMeshVertexChangeComponents::VertexPositions | EMeshVertexChangeComponents::OverlayNormals);
 			UpdateChangeFromROI(false);
 		}
 	}
diff --git a/Engine/Plugins/Experimental/MeshModelingToolset/Source/MeshModelingTools/Private/DisplaceMeshTool.cpp b/Engine/Plugins/Experimental/MeshModelingToolset/Source/MeshModelingTools/Private/DisplaceMeshTool.cpp
--- a/Engine/Plugins/Experimental/MeshModelingToolset/Source/MeshModelingTools/Private/DisplaceMeshTool.cpp
+++ b/Engine/Plugins/Experimental/MeshModelingToolset/Source/MeshModelingTools/Private/DisplaceMeshTool.cpp
@@ -8,12 +8,13 @@
 #include "SimpleDynamicMeshComponent.h"
 #include "MeshNormals.h"
 #include "ModelingOperators.h"
-#include "SceneManagement.h" // for FPrimitiveDrawInterface
 #include "Async/ParallelFor.h"
 #include "ProfilingDebugging/ScopedTimers.h"
+#include "MeshDescription.h"
 #define LOCTEXT_NAMESPACE "UDisplaceMeshTool"
 
 namespace {
+
 	void SubdivideMesh(FDynamicMesh3& Mesh)
 	{
 		TArray<int> EdgesToProcess;
@@ -60,26 +61,80 @@ namespace {
 
 	namespace ComputeDisplacement 
 	{
-		void Constant(const FDynamicMesh3& Mesh,
-				const TArray<FVector3d>& Positions, const FMeshNormals& Normals, double Intensity,
-				TArray<FVector3d>& DisplacedPositions)
+		/// Directional Filter: Scale displacement for a given vertex based on how well 
+		/// the vertex normal agrees with the specified direction.
+		struct FDirectionalFilter
 		{
-			DisplacedPositions.SetNumUninitialized(Positions.Num());
-			for (int vid : Mesh.VertexIndicesItr())
+			bool bEnableFilter = false;
+			FVector3d FilterDirection = {1,0,0};
+			double FilterWidth = 0.1;
+			const double RampSlope = 5.0;
+
+			double FilterValue(const FVector3d& EvalNormal) const
 			{
-				DisplacedPositions[vid] = Positions[vid] + (Intensity * Normals[vid]);
+				if (!bEnableFilter)	{ return 1.0;}
+
+				double DotWithFilterDirection = EvalNormal.Dot(FilterDirection);
+				double Offset = 1.0 / RampSlope;
+				double MinX = 1.0 - (2.0 + Offset) * FilterWidth;			// Start increasing here
+				double MaxX = FMathd::Min(1.0, MinX + Offset);				// Stop increasing here
+				
+				if (FMathd::Abs(MaxX - MinX) < FMathd::ZeroTolerance) { return 0.0; }
+				
+				double Y = (DotWithFilterDirection - MinX) / (MaxX - MinX); // Clamped linear interpolation for the ramp region
+				return FMathd::Clamp(Y, 0.0, 1.0);
 			}
+		};
+
+		template<typename DisplaceFunc>
+		void ParallelDisplace(const FDynamicMesh3& Mesh,
+			const TArray<FVector3d>& Positions,
+			const FMeshNormals& Normals,
+			TArray<FVector3d>& DisplacedPositions,
+			DisplaceFunc Displace)
+		{
+			ensure(Positions.Num() == Normals.GetNormals().Num());
+			ensure(Positions.Num() == DisplacedPositions.Num());
+			ensure(Mesh.VertexCount() == Positions.Num());
+
+			int32 NumVertices = Mesh.MaxVertexID();
+			ParallelFor(NumVertices, [&](int32 vid)
+			{
+				if (Mesh.IsVertex(vid))
+				{
+					DisplacedPositions[vid] = Displace(vid, Positions[vid], Normals[vid]);
+				}
+			});
 		}
 
+
+		void Constant(const FDynamicMesh3& Mesh,
+			const TArray<FVector3d>& Positions, 
+			const FMeshNormals& Normals, 
+			TFunctionRef<float(const FVector3d&, const FVector3d&)> IntensityFunc,
+			TArray<FVector3d>& DisplacedPositions)
+		{
+			ParallelDisplace(Mesh, Positions, Normals, DisplacedPositions,
+				[&](int32 vid, const FVector3d& Position, const FVector3d& Normal)
+			{
+				double Intensity = IntensityFunc(Position, Normal);
+				return Position + (Intensity * Normal);
+			});
+		}
+
+
 		void RandomNoise(const FDynamicMesh3& Mesh,
-				const TArray<FVector3d>& Positions, const FMeshNormals& Normals,
-				double Intensity, int RandomSeed, 
-				TArray<FVector3d>& DisplacedPositions)
+			const TArray<FVector3d>& Positions, 
+			const FMeshNormals& Normals,
+			TFunctionRef<float(const FVector3d&, const FVector3d&)> IntensityFunc,
+			int RandomSeed, 
+			TArray<FVector3d>& DisplacedPositions)
 		{
 			FMath::SRandInit(RandomSeed);
 			for (int vid : Mesh.VertexIndicesItr())
 			{
 				double RandVal = 2.0 * (FMath::SRand() - 0.5);
+				double Intensity = IntensityFunc(Positions[vid], Normals[vid]);
 				DisplacedPositions[vid] = Positions[vid] + (Normals[vid] * RandVal * Intensity);
 			}
 		}
@@ -87,24 +142,35 @@ namespace {
 		void PerlinNoise(const FDynamicMesh3& Mesh,
 			const TArray<FVector3d>& Positions,
 			const FMeshNormals& Normals,
-			double Intensity,
+			TFunctionRef<float(const FVector3d&, const FVector3d&)> IntensityFunc,
+			const TArray<FPerlinLayerProperties>& PerlinLayerProperties,
 			int RandomSeed,
-			double Frequency,
 			TArray<FVector3d>& DisplacedPositions)
 		{
 			FMath::SRandInit(RandomSeed);
-			float RandomOffset = 10000.0f * FMath::SRand();
-			for (int vid : Mesh.VertexIndicesItr())
+			const float RandomOffset = 10000.0f * FMath::SRand();
+
+			ParallelDisplace(Mesh, Positions, Normals, DisplacedPositions,
+				[&](int32 vid, const FVector3d& Position, const FVector3d& Normal)
 			{
-				float NoiseValue = FMath::PerlinNoise3D(Frequency * FVector(Positions[vid] + RandomOffset));
-				DisplacedPositions[vid] = Positions[vid] + (NoiseValue * Intensity * Normals[vid]);
-			}
+				// Compute the sum of Perlin noise evaluations for this point
+				FVector EvalLocation(Position + RandomOffset);
+				double TotalNoiseValue = 0.0;
+				for (int32 Layer = 0; Layer < PerlinLayerProperties.Num(); ++Layer)
+				{
+					TotalNoiseValue += PerlinLayerProperties[Layer].Intensity * FMath::PerlinNoise3D(PerlinLayerProperties[Layer].Frequency * EvalLocation);
+				}
+				double Intensity = IntensityFunc(Position, Normal);
+				return Position + (TotalNoiseValue * Intensity * Normal);
+			});
 		}
 
 		void Map(const FDynamicMesh3& Mesh,
-				const TArray<FVector3d>& Positions, const FMeshNormals& Normals,
-				double Intensity, const FSampledScalarField2f& DisplaceField,
-				TArray<FVector3d>& DisplacedPositions)
+			const TArray<FVector3d>& Positions, 
+			const FMeshNormals& Normals,
+			TFunctionRef<float(const FVector3d&, const FVector3d&)> IntensityFunc,
+			const FSampledScalarField2f& DisplaceField,
+			TArray<FVector3d>& DisplacedPositions)
 		{
 			const FDynamicMeshUVOverlay* UVOverlay = Mesh.Attributes()->GetUVLayer(0);
 			for (int tid : Mesh.TriangleIndicesItr())
@@ -116,6 +182,7 @@ namespace {
 					int vid = Tri[j];
 					FVector2f UV = UVOverlay->GetElement(UVTri[j]);
 					double Offset = DisplaceField.BilinearSampleClamped(UV);
+					double Intensity = IntensityFunc(Positions[vid], Normals[vid]);
 					DisplacedPositions[vid] = Positions[vid] + (Offset * Intensity * Normals[vid]);
 				}
 			}
@@ -123,21 +190,29 @@ namespace {
 		
 		void Sine(const FDynamicMesh3& Mesh,
 			const TArray<FVector3d>& Positions,
-			double Amplitude,
+			const FMeshNormals& Normals,
+			TFunctionRef<float(const FVector3d&, const FVector3d&)> IntensityFunc,
 			double Frequency,
 			double PhaseShift,
+			const FVector3d& Direction,
 			TArray<FVector3d>& DisplacedPositions)
 		{
-			for (int vid : Mesh.VertexIndicesItr())
+			FQuaterniond RotateToDirection(Direction, { 0.0, 0.0, 1.0 });
+
+			ParallelDisplace(Mesh, Positions, Normals, DisplacedPositions,
+				[&](int32 vid, const FVector3d& Position, const FVector3d& Normal)
 			{
-				DisplacedPositions[vid].X = Positions[vid].X;
-				DisplacedPositions[vid].Y = Positions[vid].Y;
-				double DistXY = FMath::Sqrt(Positions[vid].X*Positions[vid].X + Positions[vid].Y*Positions[vid].Y);
-				DisplacedPositions[vid].Z = Positions[vid].Z + Amplitude * FMath::Sin(Frequency*DistXY + PhaseShift);
-			}
-		}
+				FVector3d RotatedPosition = RotateToDirection * Position;
+				double DistXY = FMath::Sqrt(RotatedPosition.X * RotatedPosition.X + RotatedPosition.Y * RotatedPosition.Y);
+				double Intensity = IntensityFunc(Position, Normal);
+				FVector3d Offset = Intensity * FMath::Sin(Frequency * DistXY + PhaseShift) * Direction;
+				return Position + Offset;
 
+			});
 		}
+
+	}
+
 	class FTextureAccess
 	{
 	public:
@@ -249,34 +324,49 @@ namespace {
 		return SubdivisionsCount;
 	}
 
+
+	// A collection of parameters to avoid having excess function parameters
+	struct DisplaceMeshParameters
+	{
+		float DisplaceIntensity = 0.0f;
+		int RandomSeed = 0;
+		UTexture2D* DisplacementMap = nullptr;
+		float SineWaveFrequency = 0.0f;
+		float SineWavePhaseShift = 0.0f;
+		FVector SineWaveDirection = { 0.0f, 0.0f, 0.0f };
+		bool bEnableFilter = false;
+		FVector FilterDirection = { 0.0f, 0.0f, 0.0f };
+		float FilterWidth = 0.0f;
+		FSampledScalarField2f DisplaceField;
+		TArray<FPerlinLayerProperties> PerlinLayerProperties;
+
+		TSharedPtr<FIndexedWeightMap> WeightMap;
+		TFunction<float(const FVector3d&, const FIndexedWeightMap)> WeightMapQueryFunc;
+	};
+
 	class FDisplaceMeshOp : public FDynamicMeshOperator
 	{
 	public:
-		FDisplaceMeshOp(TSharedPtr<FDynamicMesh3> SourceMeshIn, const FSampledScalarField2f& DisplaceFieldIn,
-			float DisplaceIntensityIn, int RandomSeedIn, float DisplaceFrequencyIn,  float DisplacePhaseShiftIn,
-			EDisplaceMeshToolDisplaceType DisplacementTypeIn);
+		FDisplaceMeshOp(TSharedPtr<FDynamicMesh3> SourceMeshIn, 
+						const DisplaceMeshParameters& DisplaceParametersIn,
+						EDisplaceMeshToolDisplaceType DisplacementTypeIn);
 		void CalculateResult(FProgressCancel* Progress) final;
 
 	private:
 		TSharedPtr<FDynamicMesh3> SourceMesh;
-		float DisplaceIntensity;
-		int RandomSeed;
+		DisplaceMeshParameters Parameters;
 		EDisplaceMeshToolDisplaceType DisplacementType;
-		const FSampledScalarField2f& DisplaceField;
-		float DisplaceFrequency;
-		float DisplacePhaseShift;
-
 		TArray<FVector3d> SourcePositions;
 		FMeshNormals SourceNormals;
 		TArray<FVector3d> DisplacedPositions;
 	};
 
-	FDisplaceMeshOp::FDisplaceMeshOp(TSharedPtr<FDynamicMesh3> SourceMeshIn, const FSampledScalarField2f& DisplaceFieldIn,
-		float DisplaceIntensityIn, int RandomSeedIn, float DisplaceFrequencyIn, float DisplacePhaseShiftIn, 
-		EDisplaceMeshToolDisplaceType DisplacementTypeIn)
-		: SourceMesh(MoveTemp(SourceMeshIn)), DisplaceIntensity(DisplaceIntensityIn), RandomSeed(RandomSeedIn),
-		DisplacementType(DisplacementTypeIn), DisplaceField(DisplaceFieldIn), DisplaceFrequency(DisplaceFrequencyIn),
-		DisplacePhaseShift(DisplacePhaseShiftIn)
+	FDisplaceMeshOp::FDisplaceMeshOp(TSharedPtr<FDynamicMesh3> SourceMeshIn, 
+									 const DisplaceMeshParameters& DisplaceParametersIn,
+									 EDisplaceMeshToolDisplaceType DisplacementTypeIn)
+		: SourceMesh(MoveTemp(SourceMeshIn)), 
+		  Parameters(DisplaceParametersIn), 
+		  DisplacementType(DisplacementTypeIn)
 	{
 	}
 
@@ -301,19 +391,41 @@ namespace {
 		DisplacedPositions.SetNum(SourceMesh->MaxVertexID());
 
 		if (Progress->Cancelled()) return;
+
+		ComputeDisplacement::FDirectionalFilter DirectionalFilter{ Parameters.bEnableFilter,
+			FVector3d(Parameters.FilterDirection),
+			Parameters.FilterWidth };
+		double Intensity = Parameters.DisplaceIntensity;
+
+		TUniqueFunction<float(const FVector3d&)> WeightMapQueryFunc = [&](const FVector3d&) { return 1.0f; };
+		if (Parameters.WeightMap.IsValid())
+		{
+			WeightMapQueryFunc = [&](const FVector3d& Pos) { return Parameters.WeightMapQueryFunc(Pos, *Parameters.WeightMap); };
+		}
+		auto IntensityFunc = [&](const FVector3d& Position, const FVector3d& Normal) 
+		{
+			return Intensity * DirectionalFilter.FilterValue(Normal) * WeightMapQueryFunc(Position);
+		};
+
+
 		// compute Displaced positions in PositionBuffer
 		switch (DisplacementType)
 		{
 		default:
 		case EDisplaceMeshToolDisplaceType::Constant:
-			ComputeDisplacement::Constant(*SourceMesh, SourcePositions, SourceNormals,
-				DisplaceIntensity,
+			ComputeDisplacement::Constant(*SourceMesh, 
+				SourcePositions, 
+				SourceNormals,
+				IntensityFunc,
 				DisplacedPositions);
 			break;
 
 		case EDisplaceMeshToolDisplaceType::RandomNoise:
-			ComputeDisplacement::RandomNoise(*SourceMesh, SourcePositions, SourceNormals,
-				DisplaceIntensity, RandomSeed,
+			ComputeDisplacement::RandomNoise(*SourceMesh, 
+				SourcePositions, 
+				SourceNormals,
+				IntensityFunc,
+				Parameters.RandomSeed,
 				DisplacedPositions);
 			break;
 			
@@ -321,24 +433,29 @@ namespace {
 			ComputeDisplacement::PerlinNoise(*SourceMesh,
 				SourcePositions,
 				SourceNormals,
-				DisplaceIntensity,
-				RandomSeed,
-				DisplaceFrequency,
+				IntensityFunc,
+				Parameters.PerlinLayerProperties,	
+				Parameters.RandomSeed,
 				DisplacedPositions);
 			break;
 
 		case EDisplaceMeshToolDisplaceType::DisplacementMap:
-			ComputeDisplacement::Map(*SourceMesh, SourcePositions, SourceNormals,
-				DisplaceIntensity, DisplaceField,
+			ComputeDisplacement::Map(*SourceMesh, 
+				SourcePositions, 
+				SourceNormals,
+				IntensityFunc,
+				Parameters.DisplaceField, 
 				DisplacedPositions);
 			break;
 
 		case EDisplaceMeshToolDisplaceType::SineWave:
-			double Amplitude = DisplaceIntensity;
-			double Frequency = DisplaceFrequency;
-			double Phase = DisplacePhaseShift;
-			ComputeDisplacement::Sine(*SourceMesh, SourcePositions, 
-				Amplitude, Frequency, Phase,
+			ComputeDisplacement::Sine(*SourceMesh, 
+				SourcePositions, 
+				SourceNormals,
+				IntensityFunc,
+				Parameters.SineWaveFrequency,
+				Parameters.SineWavePhaseShift,
+				Parameters.SineWaveDirection,
 				DisplacedPositions);
 			break;
 		}
@@ -367,70 +484,83 @@ namespace {
 	{
 	public:
 		FDisplaceMeshOpFactory(TSharedPtr<FDynamicMesh3>& SourceMeshIn,
-			float DisplaceIntensityIn, int RandomSeedIn, UTexture2D* DisplacementMapIn,
-			float DisplaceFrequencyIn, float DisplacePhaseShiftIn,
+			const DisplaceMeshParameters& DisplaceParametersIn,
 			EDisplaceMeshToolDisplaceType DisplacementTypeIn )
 			: SourceMesh(SourceMeshIn)
 		{
-			SetIntensity(DisplaceIntensityIn);
-			SetRandomSeed(RandomSeed);
-			SetDisplacementMap(DisplacementMapIn);
+			SetIntensity(DisplaceParametersIn.DisplaceIntensity);
+			SetRandomSeed(DisplaceParametersIn.RandomSeed);
+			SetDisplacementMap(DisplaceParametersIn.DisplacementMap);
+			SetFrequency(DisplaceParametersIn.SineWaveFrequency);
+			SetPhaseShift(DisplaceParametersIn.SineWavePhaseShift);
+			SetSineWaveDirection(DisplaceParametersIn.SineWaveDirection);
+			SetEnableDirectionalFilter(DisplaceParametersIn.bEnableFilter);
+			SetFilterDirection(DisplaceParametersIn.FilterDirection);
+			SetFilterFalloffWidth(DisplaceParametersIn.FilterWidth);
+			SetPerlinNoiseLayerProperties(DisplaceParametersIn.PerlinLayerProperties);
 			SetDisplacementType(DisplacementTypeIn);
-			SetFrequency(DisplaceFrequencyIn);
-			SetPhaseShift(DisplacePhaseShiftIn);
+
+			Parameters.WeightMap = DisplaceParametersIn.WeightMap;
+			Parameters.WeightMapQueryFunc = DisplaceParametersIn.WeightMapQueryFunc;
 		}
 		void SetIntensity(float IntensityIn);
 		void SetRandomSeed(int RandomSeedIn);
 		void SetDisplacementMap(UTexture2D* DisplacementMapIn);
 		void SetFrequency(float FrequencyIn);
 		void SetPhaseShift(float PhaseShiftIn);
+		void SetSineWaveDirection(const FVector& Direction);
 		void SetDisplacementType(EDisplaceMeshToolDisplaceType TypeIn);
+		void SetEnableDirectionalFilter(bool EnableDirectionalFilter);
+		void SetFilterDirection(const FVector& Direction);
+		void SetFilterFalloffWidth(float FalloffWidth);
+		void SetPerlinNoiseLayerProperties(const TArray<FPerlinLayerProperties>& PerlinLayerProperties);
+		void SetWeightMap(TSharedPtr<FIndexedWeightMap> WeightMap);
 
 		TUniquePtr<FDynamicMeshOperator> MakeNewOperator() final
 		{
-			return MakeUnique<FDisplaceMeshOp>(SourceMesh, DisplaceField, DisplaceIntensity, RandomSeed, DisplaceFrequency, DisplacePhaseShift, DisplacementType);
+			return MakeUnique<FDisplaceMeshOp>(SourceMesh, Parameters, DisplacementType);
 		}
 	private:
 		void UpdateMap();
 
-		float DisplaceIntensity;
-		int RandomSeed;
-		UTexture2D* DisplacementMap;
-		float DisplaceFrequency;
-		float DisplacePhaseShift;
+		DisplaceMeshParameters Parameters;
 		EDisplaceMeshToolDisplaceType DisplacementType;
 
 		TSharedPtr<FDynamicMesh3>& SourceMesh;
-		FSampledScalarField2f DisplaceField;
 	};
 
 	void FDisplaceMeshOpFactory::SetIntensity(float IntensityIn)
 	{
-		DisplaceIntensity = IntensityIn;
+		Parameters.DisplaceIntensity = IntensityIn;
 	}
 
 	void FDisplaceMeshOpFactory::SetRandomSeed(int RandomSeedIn)
 	{
-		RandomSeed = RandomSeedIn;
+		Parameters.RandomSeed = RandomSeedIn;
 	}
 
 	void FDisplaceMeshOpFactory::SetDisplacementMap(UTexture2D* DisplacementMapIn)
 	{
-		if (DisplacementMap != DisplacementMapIn)
+		if (Parameters.DisplacementMap != DisplacementMapIn)
 		{
-			DisplacementMap = DisplacementMapIn;
+			Parameters.DisplacementMap = DisplacementMapIn;
 			UpdateMap();
 		}
 	}
 
 	void FDisplaceMeshOpFactory::SetFrequency(float FrequencyIn)
 	{
-		DisplaceFrequency = FrequencyIn;
+		Parameters.SineWaveFrequency = FrequencyIn;
 	}
 
 	void FDisplaceMeshOpFactory::SetPhaseShift(float PhaseShiftIn)
 	{
-		DisplacePhaseShift = PhaseShiftIn;
+		Parameters.SineWavePhaseShift = PhaseShiftIn;
+	}
+
+	void FDisplaceMeshOpFactory::SetSineWaveDirection(const FVector& Direction)
+	{
+		Parameters.SineWaveDirection = Direction;
 	}
 
 	void FDisplaceMeshOpFactory::SetDisplacementType(EDisplaceMeshToolDisplaceType TypeIn)
@@ -440,25 +570,25 @@ namespace {
 
 	void FDisplaceMeshOpFactory::UpdateMap()
 	{
-		if (DisplacementMap == nullptr ||
-			DisplacementMap->PlatformData == nullptr ||
-			DisplacementMap->PlatformData->Mips.Num() < 1)
+		if (Parameters.DisplacementMap == nullptr ||
+			Parameters.DisplacementMap->PlatformData == nullptr ||
+			Parameters.DisplacementMap->PlatformData->Mips.Num() < 1)
 		{
-			DisplaceField = FSampledScalarField2f();
+			Parameters.DisplaceField = FSampledScalarField2f();
 			return;
 		}
 
-		FTextureAccess TextureAccess(DisplacementMap);
+		FTextureAccess TextureAccess(Parameters.DisplacementMap);
 		if (!TextureAccess.HasData())
 		{
-			DisplaceField = FSampledScalarField2f();
+			Parameters.DisplaceField = FSampledScalarField2f();
 		}
 		else
 		{
-			int TextureWidth = DisplacementMap->GetSizeX();
-			int TextureHeight = DisplacementMap->GetSizeY();
-			DisplaceField.Resize(TextureWidth, TextureHeight, 0.0f);
-			DisplaceField.SetCellSize(1.0f / (float)TextureWidth);
+			int TextureWidth = Parameters.DisplacementMap->GetSizeX();
+			int TextureHeight = Parameters.DisplacementMap->GetSizeY();
+			Parameters.DisplaceField.Resize(TextureWidth, TextureHeight, 0.0f);
+			Parameters.DisplaceField.SetCellSize(1.0f / (float)TextureWidth);
 
 			const FColor* FormattedData = TextureAccess.GetData();
 			for (int y = 0; y < TextureHeight; ++y)
@@ -467,11 +597,37 @@ namespace {
 				{
 					FColor PixelColor = FormattedData[y * TextureWidth + x];
 					float Value = PixelColor.R / 255.0;
-					DisplaceField.GridValues[y * TextureWidth + x] = Value;
+					Parameters.DisplaceField.GridValues[y * TextureWidth + x] = Value;
 				}
 			}
 		}
 	}
+
+	void FDisplaceMeshOpFactory::SetEnableDirectionalFilter(bool EnableDirectionalFilter)
+	{
+		Parameters.bEnableFilter = EnableDirectionalFilter;
+	}
+
+	void FDisplaceMeshOpFactory::SetFilterDirection(const FVector& Direction)
+	{
+		Parameters.FilterDirection = Direction;
+	}
+
+	void FDisplaceMeshOpFactory::SetFilterFalloffWidth(float FalloffWidth)
+	{
+		Parameters.FilterWidth = FalloffWidth;
+	}
+
+	void FDisplaceMeshOpFactory::SetPerlinNoiseLayerProperties(const TArray<FPerlinLayerProperties>& LayerProperties )
+	{
+		Parameters.PerlinLayerProperties = LayerProperties;
+	}
+
+	void FDisplaceMeshOpFactory::SetWeightMap(TSharedPtr<FIndexedWeightMap> WeightMap)
+	{
+		Parameters.WeightMap = WeightMap;
+	}
+
 } // namespace
 
 /*
@@ -498,20 +654,47 @@ UInteractiveTool* UDisplaceMeshToolBuilder::BuildTool(const FToolBuilderState& S
 /*
  * Tool
  */
-UDisplaceMeshTool::UDisplaceMeshTool()
+TArray<FString> UDisplaceMeshCommonProperties::GetWeightMapsFunc()
 {
-	DisplacementType = EDisplaceMeshToolDisplaceType::Constant;
-	DisplaceIntensity = 10.0f;
-	DisplaceFrequency = 0.1f;
-	DisplacePhaseShift = 0.0f;
-	RandomSeed = 31337;
-	Subdivisions = 4;
+	return WeightMapsList;
 }
 
 void UDisplaceMeshTool::Setup()
 {
 	UInteractiveTool::Setup();
 
+	// UInteractiveToolPropertySets
+	NoiseProperties = NewObject<UDisplaceMeshPerlinNoiseProperties>();
+	NoiseProperties->RestoreProperties(this);
+	CommonProperties = NewObject<UDisplaceMeshCommonProperties>();
+	CommonProperties->RestoreProperties(this);
+	DirectionalFilterProperties = NewObject<UDisplaceMeshDirectionalFilterProperties>();
+	DirectionalFilterProperties->RestoreProperties(this);
+	TextureMapProperties = NewObject<UDisplaceMeshTextureMapProperties>();
+	TextureMapProperties->RestoreProperties(this);
+	SineWaveProperties = NewObject<UDisplaceMeshSineWaveProperties>();
+	SineWaveProperties->RestoreProperties(this);
+
+	if (TextureMapProperties->DisplacementMap != nullptr && TextureMapProperties->DisplacementMap->IsValidLowLevel() == false)
+	{
+		TextureMapProperties->DisplacementMap = nullptr;
+	}
+
+	// populate weight maps list
+	TArray<FName> WeightMaps;
+	UE::WeightMaps::FindVertexWeightMaps(ComponentTarget->GetMesh(), WeightMaps);
+	CommonProperties->WeightMapsList.Add(TEXT("None"));
+	for (FName Name : WeightMaps)
+	{
+		CommonProperties->WeightMapsList.Add(Name.ToString());
+	}
+	if (WeightMaps.Contains(CommonProperties->WeightMap) == false)		// discard restored value if it doesn't apply
+	{
+		CommonProperties->WeightMap = FName(CommonProperties->WeightMapsList[0]);
+	}
+	UpdateActiveWeightMap();
+
+
 	// create dynamic mesh component to use for live preview
 	DynamicMeshComponent = NewObject<USimpleDynamicMeshComponent>(ComponentTarget->GetOwnerActor(), "DynamicMesh");
 	DynamicMeshComponent->SetupAttachment(ComponentTarget->GetOwnerActor()->GetRootComponent());
@@ -529,21 +712,72 @@ void UDisplaceMeshTool::Setup()
 	DynamicMeshComponent->TangentsType = EDynamicMeshTangentCalcType::AutoCalculated;
 	DynamicMeshComponent->InitializeMesh(ComponentTarget->GetMesh());
 	OriginalMesh.Copy(*DynamicMeshComponent->GetMesh());
-
-	Subdivider = MakeUnique<FSubdivideMeshOpFactory>(OriginalMesh, Subdivisions);
-	Displacer = MakeUnique<FDisplaceMeshOpFactory>(SubdividedMesh, DisplaceIntensity, RandomSeed, DisplacementMap, DisplaceFrequency, DisplacePhaseShift, DisplacementType);
+	OriginalMeshSpatial.SetMesh(&OriginalMesh, true);
+
+	Subdivider = MakeUnique<FSubdivideMeshOpFactory>(OriginalMesh, CommonProperties->Subdivisions);
+	
+	DisplaceMeshParameters Parameters;
+	Parameters.DisplaceIntensity = CommonProperties->DisplaceIntensity;
+	Parameters.RandomSeed = CommonProperties->RandomSeed;
+	Parameters.DisplacementMap = TextureMapProperties->DisplacementMap;
+	Parameters.SineWaveFrequency = SineWaveProperties->SineWaveFrequency;
+	Parameters.SineWavePhaseShift = SineWaveProperties->SineWavePhaseShift;
+	Parameters.SineWaveDirection = SineWaveProperties->SineWaveDirection;
+	Parameters.bEnableFilter = DirectionalFilterProperties->bEnableFilter;
+	Parameters.FilterDirection = DirectionalFilterProperties->FilterDirection;
+	Parameters.FilterWidth = DirectionalFilterProperties->FilterWidth;
+	Parameters.PerlinLayerProperties = NoiseProperties->PerlinLayerProperties;
+	Parameters.WeightMap = ActiveWeightMap;
+	Parameters.WeightMapQueryFunc = [this](const FVector3d& Position, const FIndexedWeightMap& WeightMap) { return WeightMapQuery(Position, WeightMap);	};
+
+	Displacer = MakeUnique<FDisplaceMeshOpFactory>(SubdividedMesh, Parameters, CommonProperties->DisplacementType);
 		
 	// hide input StaticMeshComponent
 	ComponentTarget->SetOwnerVisibility(false);
 
 	// initialize our properties
 	ToolPropertyObjects.Add(this);
+
+	AddToolPropertySource(CommonProperties);
+	SetToolPropertySourceEnabled(CommonProperties, true);
+
+	AddToolPropertySource(DirectionalFilterProperties);
+	SetToolPropertySourceEnabled(DirectionalFilterProperties, true);
+
+	AddToolPropertySource(TextureMapProperties);
+	SetToolPropertySourceEnabled(TextureMapProperties, CommonProperties->DisplacementType == EDisplaceMeshToolDisplaceType::DisplacementMap);
+
+	AddToolPropertySource(SineWaveProperties);
+	SetToolPropertySourceEnabled(SineWaveProperties, CommonProperties->DisplacementType == EDisplaceMeshToolDisplaceType::SineWave);
+
+	AddToolPropertySource(NoiseProperties);
+	SetToolPropertySourceEnabled(NoiseProperties, CommonProperties->DisplacementType == EDisplaceMeshToolDisplaceType::PerlinNoise);
+
+	// Set up a callback for when the type of displacement changes
+	CommonProperties->WatchProperty(CommonProperties->DisplacementType,
+									[this](EDisplaceMeshToolDisplaceType NewType)
+									{
+										SetToolPropertySourceEnabled(NoiseProperties, (NewType == EDisplaceMeshToolDisplaceType::PerlinNoise));
+										SetToolPropertySourceEnabled(SineWaveProperties, (NewType == EDisplaceMeshToolDisplaceType::SineWave));
+										SetToolPropertySourceEnabled(TextureMapProperties, (NewType == EDisplaceMeshToolDisplaceType::DisplacementMap));
+									} );
+
 	ValidateSubdivisions();
 	StartComputation();
+
+	GetToolManager()->DisplayMessage(
+		LOCTEXT("OnStartDisplaceMesh", "Subdivide and Displace the input mesh using different noise functions and maps"),
+		EToolMessageLevel::UserNotification);
 }
 
 void UDisplaceMeshTool::Shutdown(EToolShutdownType ShutdownType)
 {
+	CommonProperties->SaveProperties(this);
+	NoiseProperties->SaveProperties(this);
+	DirectionalFilterProperties->SaveProperties(this);
+	SineWaveProperties->SaveProperties(this);
+	TextureMapProperties->SaveProperties(this);
+
 	if (DynamicMeshComponent != nullptr)
 	{
 		ComponentTarget->SetOwnerVisibility(true);
@@ -554,7 +788,7 @@ void UDisplaceMeshTool::Shutdown(EToolShutdownType ShutdownType)
 			GetToolManager()->BeginUndoTransaction(LOCTEXT("DisplaceMeshToolTransactionName", "Displace Mesh"));
 			ComponentTarget->CommitMesh([=](const FPrimitiveComponentTarget::FCommitParams& CommitParams)
 			{
-				DynamicMeshComponent->Bake(CommitParams.MeshDescription, Subdivisions > 0);
+				DynamicMeshComponent->Bake(CommitParams.MeshDescription, CommonProperties->Subdivisions > 0);
 			});
 			GetToolManager()->EndUndoTransaction();
 		}
@@ -567,42 +801,54 @@ void UDisplaceMeshTool::Shutdown(EToolShutdownType ShutdownType)
 
 void UDisplaceMeshTool::ValidateSubdivisions()
 {
+	if (CommonProperties->bDisableSizeWarning)
+	{
+		GetToolManager()->DisplayMessage(FText::GetEmpty(), EToolMessageLevel::UserWarning);
+		return;
+	}
+
 	constexpr int MaxTriangles = 3000000;
 	double NumTriangles = OriginalMesh.MaxTriangleID();
 	int MaxSubdivisions = (int)floor(log2(MaxTriangles / NumTriangles) / 2.0);
-	if (Subdivisions > MaxSubdivisions)
+	if (CommonProperties->Subdivisions > MaxSubdivisions)
 	{
-		FText WarningText = FText::Format(LOCTEXT("SubdivisionsTooHigh", "Desired number of Subdivisions ({0}) exceeds maximum number of {1}"), FText::AsNumber(Subdivisions), FText::AsNumber(MaxSubdivisions));
+		FText WarningText = FText::Format(LOCTEXT("SubdivisionsTooHigh", "Desired number of Subdivisions ({0}) exceeds maximum number of {1}"), 
+			FText::AsNumber(CommonProperties->Subdivisions), 
+			FText::AsNumber(MaxSubdivisions));
 		GetToolManager()->DisplayMessage(WarningText, EToolMessageLevel::UserWarning);
-		Subdivisions = MaxSubdivisions;
+		CommonProperties->Subdivisions = MaxSubdivisions;
 	}
 	else
 	{
 		FText ClearWarningText;
 		GetToolManager()->DisplayMessage(ClearWarningText, EToolMessageLevel::UserWarning);
 	}
-	if (Subdivisions < 0)
+	if (CommonProperties->Subdivisions < 0)
 	{
-		Subdivisions = 0;
+		CommonProperties->Subdivisions = 0;
 	}
 }
 
 #if WITH_EDITOR
-void UDisplaceMeshTool::PostEditChangeProperty(FPropertyChangedEvent& PropertyChangedEvent)
+
+void UDisplaceMeshTool::OnPropertyModified(UObject* PropertySet, FProperty* Property)
 {
-	FProperty* PropertyThatChanged = PropertyChangedEvent.Property;
-	if (PropertyThatChanged)
+	if (PropertySet && Property)
 	{
-		FSubdivideMeshOpFactory*  SubdividerDownCast = static_cast<FSubdivideMeshOpFactory*>(Subdivider.Get());
-		FDisplaceMeshOpFactory*  DisplacerDownCast = static_cast<FDisplaceMeshOpFactory*>(Displacer.Get());
-		const FName PropName = PropertyThatChanged->GetFName();
+		FDisplaceMeshOpFactory* DisplacerDownCast = static_cast<FDisplaceMeshOpFactory*>(Displacer.Get());
+		FSubdivideMeshOpFactory* SubdividerDownCast = static_cast<FSubdivideMeshOpFactory*>(Subdivider.Get());
+
+		const FString PropertySetName = PropertySet->GetFName().GetPlainNameString();
+		const FName PropName = Property->GetFName();
+	
 		bNeedsDisplaced = true;
-		if (PropName == GET_MEMBER_NAME_CHECKED(UDisplaceMeshTool, Subdivisions))
+
+		if (PropName == GET_MEMBER_NAME_CHECKED(UDisplaceMeshCommonProperties, Subdivisions))
 		{
 			ValidateSubdivisions();
-			if (Subdivisions != SubdividerDownCast->GetSubdivisionsCount())
+			if (CommonProperties->Subdivisions != SubdividerDownCast->GetSubdivisionsCount())
 			{
-				SubdividerDownCast->SetSubdivisionsCount(Subdivisions);
+				SubdividerDownCast->SetSubdivisionsCount(CommonProperties->Subdivisions);
 				bNeedsSubdivided = true;
 			}
 			else
@@ -610,38 +856,70 @@ void UDisplaceMeshTool::PostEditChangeProperty(FPropertyChangedEvent& PropertyCh
 				return;
 			}
 		}
-		else if (PropName == GET_MEMBER_NAME_CHECKED(UDisplaceMeshTool, RandomSeed))
+		else if (PropName == GET_MEMBER_NAME_CHECKED(UDisplaceMeshCommonProperties, RandomSeed))
 		{
-			DisplacerDownCast->SetRandomSeed(RandomSeed);
+			DisplacerDownCast->SetRandomSeed(CommonProperties->RandomSeed);
 		}
-		else if (PropName == GET_MEMBER_NAME_CHECKED(UDisplaceMeshTool, DisplacementType))
+		else if (PropName == GET_MEMBER_NAME_CHECKED(UDisplaceMeshCommonProperties, DisplacementType))
 		{
-			DisplacerDownCast->SetDisplacementType(DisplacementType);
+			DisplacerDownCast->SetDisplacementType(CommonProperties->DisplacementType);
 		}
-		else if (PropName == GET_MEMBER_NAME_CHECKED(UDisplaceMeshTool, DisplaceIntensity))
+		else if (PropName == GET_MEMBER_NAME_CHECKED(UDisplaceMeshCommonProperties, DisplaceIntensity))
 		{
-			DisplacerDownCast->SetIntensity(DisplaceIntensity);
+			DisplacerDownCast->SetIntensity(CommonProperties->DisplaceIntensity);
 		}
-		else if (PropName == GET_MEMBER_NAME_CHECKED(UDisplaceMeshTool, DisplaceFrequency))
+		else if (PropName == GET_MEMBER_NAME_CHECKED(UDisplaceMeshSineWaveProperties, SineWaveFrequency))
 		{
-			DisplacerDownCast->SetFrequency(DisplaceFrequency);
+			DisplacerDownCast->SetFrequency(SineWaveProperties->SineWaveFrequency);
 		}
-		else if (PropName == GET_MEMBER_NAME_CHECKED(UDisplaceMeshTool, DisplacePhaseShift))
+		else if (PropName == GET_MEMBER_NAME_CHECKED(UDisplaceMeshSineWaveProperties, SineWavePhaseShift))
 		{
-			DisplacerDownCast->SetPhaseShift(DisplacePhaseShift);
+			DisplacerDownCast->SetPhaseShift(SineWaveProperties->SineWavePhaseShift);
 		}
-		else if (PropName == GET_MEMBER_NAME_CHECKED(UDisplaceMeshTool, DisplacementMap))
+		else if (PropName == GET_MEMBER_NAME_CHECKED(UDisplaceMeshTextureMapProperties, DisplacementMap))
 		{
-			DisplacerDownCast->SetDisplacementMap(DisplacementMap);
+			DisplacerDownCast->SetDisplacementMap(TextureMapProperties->DisplacementMap);
 		}
+		else if (PropName == GET_MEMBER_NAME_CHECKED(UDisplaceMeshCommonProperties, WeightMap) 
+				 || PropName == GET_MEMBER_NAME_CHECKED(UDisplaceMeshCommonProperties, bInvertWeightMap))
+		{
+			UpdateActiveWeightMap();
+			DisplacerDownCast->SetWeightMap(ActiveWeightMap);
+		}
+		else if (PropName == GET_MEMBER_NAME_CHECKED(UDisplaceMeshDirectionalFilterProperties, bEnableFilter))
+		{
+			DisplacerDownCast->SetEnableDirectionalFilter(DirectionalFilterProperties->bEnableFilter);
+		}
+		else if (PropName == GET_MEMBER_NAME_CHECKED(UDisplaceMeshDirectionalFilterProperties, FilterWidth))
+		{
+			DisplacerDownCast->SetFilterFalloffWidth(DirectionalFilterProperties->FilterWidth);
+		}
+		else if ((PropName == GET_MEMBER_NAME_CHECKED(FPerlinLayerProperties, Frequency)) || (PropName == GET_MEMBER_NAME_CHECKED(FPerlinLayerProperties, Intensity)))
+		{
+			DisplacerDownCast->SetPerlinNoiseLayerProperties(NoiseProperties->PerlinLayerProperties);
+		}
+		else if (PropName == "X" || PropName == "Y" || PropName == "Z")
+		{
+			if (PropertySetName == "DisplaceMeshDirectionalFilterProperties")
+			{
+				DirectionalFilterProperties->FilterDirection.Normalize();
+				DisplacerDownCast->SetFilterDirection(DirectionalFilterProperties->FilterDirection);
+			}
+			else if (PropertySetName == "DisplaceMeshSineWaveProperties")
+			{
+				SineWaveProperties->SineWaveDirection.Normalize();
+				DisplacerDownCast->SetSineWaveDirection(SineWaveProperties->SineWaveDirection);
+			}
+		}
+
+
 		StartComputation();
 	}
 }
 #endif
 
-void UDisplaceMeshTool::Tick(float DeltaTime)
+void UDisplaceMeshTool::OnTick(float DeltaTime)
 {
-	Super::Tick(DeltaTime); 
 	AdvanceComputation();
 }
 
@@ -694,4 +972,35 @@ void UDisplaceMeshTool::AdvanceComputation()
 	}
 }
 
+
+
+void UDisplaceMeshTool::UpdateActiveWeightMap()
+{
+	TSharedPtr<FIndexedWeightMap> NewWeightMap = MakeShared<FIndexedWeightMap>();
+	UE::WeightMaps::GetVertexWeightMap(ComponentTarget->GetMesh(), CommonProperties->WeightMap, *NewWeightMap, 1.0f);
+	if (CommonProperties->bInvertWeightMap)
+	{
+		NewWeightMap->InvertWeightMap();
+	}
+	ActiveWeightMap = NewWeightMap;
+}
+
+
+float UDisplaceMeshTool::WeightMapQuery(const FVector3d& Position, const FIndexedWeightMap& WeightMap) const
+{
+	double NearDistSqr;
+	int32 NearTID = OriginalMeshSpatial.FindNearestTriangle(Position, NearDistSqr);
+	if (NearTID < 0)
+	{
+		return 1.0f;
+	}
+	FDistPoint3Triangle3d Distance = TMeshQueries<FDynamicMesh3>::TriangleDistance(OriginalMesh, NearTID, Position);
+	FIndex3i Tri = OriginalMesh.GetTriangle(NearTID);
+	return WeightMap.GetInterpValue(Tri, Distance.TriangleBaryCoords);
+}
+
+
+#include "Tests/DisplaceMeshTool_Tests.inl"
+
+
 #undef LOCTEXT_NAMESPACE
diff --git a/Engine/Plugins/Experimental/MeshModelingToolset/Source/MeshModelingTools/Private/DrawPolyPathTool.cpp b/Engine/Plugins/Experimental/MeshModelingToolset/Source/MeshModelingTools/Private/DrawPolyPathTool.cpp
--- a/Engine/Plugins/Experimental/MeshModelingToolset/Source/MeshModelingTools/Private/DrawPolyPathTool.cpp
+++ b/Engine/Plugins/Experimental/MeshModelingToolset/Source/MeshModelingTools/Private/DrawPolyPathTool.cpp
@@ -159,31 +159,9 @@ UInteractiveTool* UDrawPolyPathToolBuilder::BuildTool(const FToolBuilderState& S
 	return NewTool;
 }
 
-
-void UDrawPolyPathProperties::SaveRestoreProperties(UInteractiveTool* RestoreToTool, bool bSaving)
-{
-	UDrawPolyPathProperties* PropertyCache = GetPropertyCache<UDrawPolyPathProperties>();
-	SaveRestoreProperty(PropertyCache->OutputType, this->OutputType, bSaving);
-	SaveRestoreProperty(PropertyCache->WidthMode, this->WidthMode, bSaving);
-	SaveRestoreProperty(PropertyCache->Width, this->Width, bSaving);
-	SaveRestoreProperty(PropertyCache->HeightMode, this->HeightMode, bSaving);
-	SaveRestoreProperty(PropertyCache->Height, this->Height, bSaving);
-	SaveRestoreProperty(PropertyCache->RampStartRatio, this->RampStartRatio, bSaving);
-	SaveRestoreProperty(PropertyCache->bSnapToWorldGrid, this->bSnapToWorldGrid, bSaving);
-}
-
-
-
-
-
 /*
 * Tool methods
 */
-
-UDrawPolyPathTool::UDrawPolyPathTool()
-{
-}
-
 void UDrawPolyPathTool::SetWorld(UWorld* World)
 {
 	this->TargetWorld = World;
@@ -396,10 +374,8 @@ bool UDrawPolyPathTool::OnUpdateHover(const FInputDeviceRay& DevicePos)
 
 
 
-void UDrawPolyPathTool::Tick(float DeltaTime)
+void UDrawPolyPathTool::OnTick(float DeltaTime)
 {
-	UInteractiveTool::Tick(DeltaTime);
-
 	if (PlaneMechanic != nullptr)
 	{
 		PlaneMechanic->SetEnableGridSnaping(TransformProps->bSnapToWorldGrid);
diff --git a/Engine/Plugins/Experimental/MeshModelingToolset/Source/MeshModelingTools/Private/DrawPolygonTool.cpp b/Engine/Plugins/Experimental/MeshModelingToolset/Source/MeshModelingTools/Private/DrawPolygonTool.cpp
--- a/Engine/Plugins/Experimental/MeshModelingToolset/Source/MeshModelingTools/Private/DrawPolygonTool.cpp
+++ b/Engine/Plugins/Experimental/MeshModelingToolset/Source/MeshModelingTools/Private/DrawPolygonTool.cpp
@@ -65,62 +65,6 @@ UDrawPolygonToolStandardProperties::UDrawPolygonToolStandardProperties()
 {
 }
 
-
-void UDrawPolygonToolStandardProperties::SaveProperties(UInteractiveTool* SaveFromTool)
-{
-	UDrawPolygonToolStandardProperties* PropertyCache = GetPropertyCache<UDrawPolygonToolStandardProperties>();
-	PropertyCache->PolygonType = this->PolygonType;
-	PropertyCache->OutputMode = this->OutputMode;
-	PropertyCache->ExtrudeHeight = this->ExtrudeHeight;
-	PropertyCache->Steps = this->Steps;
-	PropertyCache->bAllowSelfIntersections = this->bAllowSelfIntersections;
-	PropertyCache->bShowGizmo = this->bShowGizmo;
-}
-
-void UDrawPolygonToolStandardProperties::RestoreProperties(UInteractiveTool* RestoreToTool)
-{
-	UDrawPolygonToolStandardProperties* PropertyCache = GetPropertyCache<UDrawPolygonToolStandardProperties>();
-	this->PolygonType = PropertyCache->PolygonType;
-	this->OutputMode = PropertyCache->OutputMode;
-	this->ExtrudeHeight = PropertyCache->ExtrudeHeight;
-	this->Steps = PropertyCache->Steps;
-	this->bAllowSelfIntersections = PropertyCache->bAllowSelfIntersections;
-	this->bShowGizmo = PropertyCache->bShowGizmo;
-}
-
-
-
-void UDrawPolygonToolSnapProperties::SaveProperties(UInteractiveTool* SaveFromTool)
-{
-	UDrawPolygonToolSnapProperties* PropertyCache = GetPropertyCache<UDrawPolygonToolSnapProperties>();
-	PropertyCache->bEnableSnapping = this->bEnableSnapping;
-	PropertyCache->bSnapToWorldGrid = this->bSnapToWorldGrid;
-	PropertyCache->bSnapToVertices = this->bSnapToVertices;
-	PropertyCache->bSnapToEdges = this->bSnapToEdges;
-	PropertyCache->bSnapToAngles = this->bSnapToAngles;
-	PropertyCache->bSnapToLengths = this->bSnapToLengths;
-	PropertyCache->bHitSceneObjects = this->bHitSceneObjects;
-	//PropertyCache->SegmentLength = this->Length;		// this is purely a feedback property
-	PropertyCache->HitNormalOffset = this->HitNormalOffset;
-}
-
-void UDrawPolygonToolSnapProperties::RestoreProperties(UInteractiveTool* RestoreToTool)
-{
-	UDrawPolygonToolSnapProperties* PropertyCache = GetPropertyCache<UDrawPolygonToolSnapProperties>();
-	this->bEnableSnapping = PropertyCache->bEnableSnapping;
-	this->bSnapToWorldGrid = PropertyCache->bSnapToWorldGrid;
-	this->bSnapToVertices = PropertyCache->bSnapToVertices;
-	this->bSnapToEdges = PropertyCache->bSnapToEdges;
-	this->bSnapToAngles = PropertyCache->bSnapToAngles;
-	this->bSnapToLengths = PropertyCache->bSnapToLengths;
-	this->bHitSceneObjects = PropertyCache->bHitSceneObjects;
-	//this->SegmentLength = PropertyCache->Length;
-	this->HitNormalOffset = PropertyCache->HitNormalOffset;
-}
-
-
-
-
 /*
  * Tool
  */
@@ -173,10 +117,8 @@ void UDrawPolygonTool::Setup()
 
 	PolygonProperties = NewObject<UDrawPolygonToolStandardProperties>(this, TEXT("Polygon Settings"));
 	PolygonProperties->RestoreProperties(this);
-	ShowGizmoWatcher.Initialize(
-		[this]() { return this->PolygonProperties->bShowGizmo; }, 
-		[this](bool bNewValue) { this->UpdateShowGizmoState(bNewValue); },
-		true);
+	PolygonProperties->WatchProperty(PolygonProperties->bShowGizmo,
+									 [this](bool bNewValue) { this->UpdateShowGizmoState(bNewValue); });
 
 	// Create a new TransformGizmo and associated TransformProxy. The TransformProxy will not be the
 	// parent of any Components in this case, we just use it's transform and change delegate.
@@ -289,7 +231,7 @@ void UDrawPolygonTool::PopLastVertexAction()
 
 
 
-void UDrawPolygonTool::Tick(float DeltaTime)
+void UDrawPolygonTool::OnTick(float DeltaTime)
 {
 	if (PlaneTransformGizmo != nullptr)
 	{
@@ -438,6 +380,7 @@ void UDrawPolygonTool::Render(IToolsContextRenderAPI* RenderAPI)
 		auto DrawVertices = [&PDI, &UseColor](const TArray<FVector3d>& Vertices, ESceneDepthPriorityGroup Group, float Thickness)
 		{
 			for (int lasti = Vertices.Num() - 1, i = 0, NumVertices = Vertices.Num(); i < NumVertices; lasti = i++)
+
 			{
 				PDI->DrawLine((FVector)Vertices[lasti], (FVector)Vertices[i], UseColor, Group, Thickness, 0.0f, true);
 			}
@@ -485,8 +428,6 @@ void UDrawPolygonTool::Render(IToolsContextRenderAPI* RenderAPI)
 	{
 		HeightMechanic->Render(RenderAPI);
 	}
-
-	ShowGizmoWatcher.CheckAndUpdate();
 }
 
 
@@ -1004,8 +945,10 @@ void UDrawPolygonTool::UpdateShowGizmoState(bool bNewVisibility)
 	}
 	else
 	{
-		PlaneTransformGizmo = GetToolManager()->GetPairedGizmoManager()->CreateCustomTransformGizmo(
-			ETransformGizmoSubElements::StandardTranslateRotate, this);
+		if (!PlaneTransformGizmo) {
+			PlaneTransformGizmo = GetToolManager()->GetPairedGizmoManager()->CreateCustomTransformGizmo(
+				ETransformGizmoSubElements::StandardTranslateRotate, this);
+		}
 		PlaneTransformGizmo->SetActiveTarget(PlaneTransformProxy, GetToolManager());
 		PlaneTransformGizmo->SetNewGizmoTransform(FTransform((FQuat)DrawPlaneOrientation, (FVector)DrawPlaneOrigin));
 	}
