Edigrate Dev-Editor updates to InteractiveToolsFramework, GeometryProcessing, MeshModelingToolset, and ModelingTools Editor Mode Plugins. Update a few files in new MeshPaint Plugin for UInteractiveToo

diff --git a/Engine/Plugins/Experimental/MeshModelingToolset/Source/MeshModelingTools/Private/DynamicMeshSculptTool.cpp b/Engine/Plugins/Experimental/MeshModelingToolset/Source/MeshModelingTools/Private/DynamicMeshSculptTool.cpp
--- a/Engine/Plugins/Experimental/MeshModelingToolset/Source/MeshModelingTools/Private/DynamicMeshSculptTool.cpp
+++ b/Engine/Plugins/Experimental/MeshModelingToolset/Source/MeshModelingTools/Private/DynamicMeshSculptTool.cpp
@@ -26,6 +26,8 @@
 #include "Components/PrimitiveComponent.h"
 #include "Generators/SphereGenerator.h"
 
+#include "Sculpting/KelvinletBrushOp.h"
+
 #include "InteractiveGizmoManager.h"
 #include "BaseGizmos/GizmoComponents.h"
 #include "BaseGizmos/TransformGizmo.h"
@@ -35,6 +37,8 @@
 
 #define LOCTEXT_NAMESPACE "UDynamicMeshSculptTool"
 
+
+
 /*
  * ToolBuilder
  */
@@ -46,102 +50,6 @@ UMeshSurfacePointTool* UDynamicMeshSculptToolBuilder::CreateNewTool(const FToolB
 	return SculptTool;
 }
 
-USculptBrushProperties::USculptBrushProperties() : UBrushBaseProperties()
-{
-	this->BrushFalloffAmount = 0.5f;
-}
-
-
-
-void USculptBrushProperties::SaveRestoreProperties(UInteractiveTool* Tool, bool bSaving)
-{
-	// do not restore base properties, they have their own CDO
-	//UBrushBaseProperties::SaveRestoreProperties(Tool, bSaving);
-
-	USculptBrushProperties* PropertyCache = GetPropertyCache<USculptBrushProperties>();
-
-	SaveRestoreProperty(PropertyCache->BrushSize, this->BrushSize, bSaving);
-	SaveRestoreProperty(PropertyCache->bSpecifyRadius, this->bSpecifyRadius, bSaving);
-	SaveRestoreProperty(PropertyCache->BrushRadius, this->BrushRadius, bSaving);
-	SaveRestoreProperty(PropertyCache->BrushFalloffAmount, this->BrushFalloffAmount, bSaving);
-	SaveRestoreProperty(PropertyCache->BrushStrength, this->BrushStrength, bSaving);
-
-	SaveRestoreProperty(PropertyCache->Depth, this->Depth, bSaving);
-	SaveRestoreProperty(PropertyCache->bHitBackFaces, this->bHitBackFaces, bSaving);
-}
-
-/*
- * BrushSculptProperties
- */
-void UBrushSculptProperties::SaveRestoreProperties(UInteractiveTool* Tool, bool bSaving)
-{
-	UBrushSculptProperties* PropertyCache = GetPropertyCache<UBrushSculptProperties>();
-	SaveRestoreProperty(PropertyCache->PrimaryBrushType, this->PrimaryBrushType, bSaving);
-	SaveRestoreProperty(PropertyCache->PrimaryBrushSpeed, this->PrimaryBrushSpeed, bSaving);
-	SaveRestoreProperty(PropertyCache->SmoothBrushSpeed, this->SmoothBrushSpeed, bSaving);
-	SaveRestoreProperty(PropertyCache->bPreserveUVFlow, this->bPreserveUVFlow, bSaving);
-	SaveRestoreProperty(PropertyCache->bDetailPreservingSmooth, this->bDetailPreservingSmooth, bSaving);
-	SaveRestoreProperty(PropertyCache->bFreezeTarget, this->bFreezeTarget, bSaving);
-}
-
-
-
-void UPlaneBrushProperties::SaveRestoreProperties(UInteractiveTool* Tool, bool bSaving)
-{
-	UPlaneBrushProperties* PropertyCache = GetPropertyCache<UPlaneBrushProperties>();
-	SaveRestoreProperty(PropertyCache->WhichSide, this->WhichSide, bSaving);
-}
-
-
-
-UFixedPlaneBrushProperties::UFixedPlaneBrushProperties()
-{
-	bPropertySetEnabled = true;
-	bSnapToGrid = true;
-	bShowGizmo = true;
-	Position = FVector::ZeroVector;
-	Rotation = FQuat::Identity;
-}
-
-void UFixedPlaneBrushProperties::SaveRestoreProperties(UInteractiveTool* Tool, bool bSaving)
-{
-	UFixedPlaneBrushProperties* PropertyCache = GetPropertyCache<UFixedPlaneBrushProperties>();
-	SaveRestoreProperty(PropertyCache->bShowGizmo, this->bShowGizmo, bSaving);
-	SaveRestoreProperty(PropertyCache->bSnapToGrid, this->bSnapToGrid, bSaving);
-	SaveRestoreProperty(PropertyCache->Position, this->Position, bSaving);
-	SaveRestoreProperty(PropertyCache->Rotation, this->Rotation, bSaving);
-}
-
-
-void USculptMaxBrushProperties::SaveRestoreProperties(UInteractiveTool* Tool, bool bSaving)
-{
-	USculptMaxBrushProperties* PropertyCache = GetPropertyCache<USculptMaxBrushProperties>();
-	SaveRestoreProperty(PropertyCache->MaxHeight, this->MaxHeight, bSaving);
-}
-
-
-void UBrushRemeshProperties::SaveRestoreProperties(UInteractiveTool* RestoreToTool, bool bSaving)
-{
-	UBrushRemeshProperties* PropertyCache = GetPropertyCache<UBrushRemeshProperties>();
-
-	// MeshConstraintProperties
-	SaveRestoreProperty(PropertyCache->bPreserveSharpEdges, this->bPreserveSharpEdges, bSaving);
-	SaveRestoreProperty(PropertyCache->MeshBoundaryConstraint, this->MeshBoundaryConstraint, bSaving);
-	SaveRestoreProperty(PropertyCache->GroupBoundaryConstraint, this->GroupBoundaryConstraint, bSaving);
-	SaveRestoreProperty(PropertyCache->MaterialBoundaryConstraint, this->MaterialBoundaryConstraint, bSaving);
-	SaveRestoreProperty(PropertyCache->bPreventNormalFlips, this->bPreventNormalFlips, bSaving);
-
-	//RemeshProperties
-	SaveRestoreProperty(PropertyCache->SmoothingStrength, this->SmoothingStrength, bSaving);
-	SaveRestoreProperty(PropertyCache->bFlips, this->bFlips, bSaving);
-	SaveRestoreProperty(PropertyCache->bSplits, this->bSplits, bSaving);
-	SaveRestoreProperty(PropertyCache->bCollapses, this->bCollapses, bSaving);
-
-	//BrushRemeshProperties
-	SaveRestoreProperty(PropertyCache->TriangleSize, this->TriangleSize, bSaving);
-	SaveRestoreProperty(PropertyCache->PreserveDetail, this->PreserveDetail, bSaving);
-	SaveRestoreProperty(PropertyCache->bEnableRemeshing, this->bEnableRemeshing, bSaving);
-}
 
 /*
  * Tool
@@ -175,6 +83,7 @@ void UDynamicMeshSculptTool::Setup()
 	DynamicMeshComponent->InitializeMesh(ComponentTarget->GetMesh());
 
 	// transform mesh to world space because handling scaling inside brush is a mess
+	// Note: this transform does not include translation ( so only the 3x3 transform)
 	InitialTargetTransform = FTransform3d(ComponentTarget->GetWorldTransform());
 	// clamp scaling because if we allow zero-scale we cannot invert this transform on Accept
 	InitialTargetTransform.ClampMinimumScale(0.01);
@@ -215,6 +124,7 @@ void UDynamicMeshSculptTool::Setup()
 
 	// initialize other properties
 	SculptProperties = NewObject<UBrushSculptProperties>(this);
+	KelvinBrushProperties = NewObject<UKelvinBrushProperties>(this);
 
 	RemeshProperties = NewObject<UBrushRemeshProperties>(this);
 	RemeshProperties->RestoreProperties(this);
@@ -255,6 +165,9 @@ void UDynamicMeshSculptTool::Setup()
 	SculptMaxBrushProperties->RestoreProperties(this);
 	AddToolPropertySource(SculptMaxBrushProperties);
 
+	AddToolPropertySource(KelvinBrushProperties);
+	KelvinBrushProperties->RestoreProperties(this);
+
 	GizmoProperties = NewObject<UFixedPlaneBrushProperties>();
 	GizmoProperties->RestoreProperties(this);
 	AddToolPropertySource(GizmoProperties);
@@ -273,6 +186,7 @@ void UDynamicMeshSculptTool::Setup()
 	SetToolPropertySourceEnabled(PlaneBrushProperties, false);
 	SetToolPropertySourceEnabled(GizmoProperties, false);
 	SetToolPropertySourceEnabled(SculptMaxBrushProperties, false);
+	SetToolPropertySourceEnabled(KelvinBrushProperties, false);
 
 	ViewProperties = NewObject<UMeshEditingViewProperties>();
 	ViewProperties->RestoreProperties(this);
@@ -366,6 +280,7 @@ void UDynamicMeshSculptTool::Shutdown(EToolShutdownType ShutdownType)
 
 	BrushProperties->SaveProperties(this);
 	SculptProperties->SaveProperties(this);
+	KelvinBrushProperties->SaveProperties(this);
 	ViewProperties->SaveProperties(this);
 	GizmoProperties->SaveProperties(this);
 	PlaneBrushProperties->SaveProperties(this);
@@ -554,15 +469,27 @@ void UDynamicMeshSculptTool::ApplyStamp(const FRay& WorldRay)
 		case EDynamicMeshSculptBrushType::Move:
 			ApplyMoveBrush(WorldRay);
 			break;
+		case EDynamicMeshSculptBrushType::PullKelvin:
+			ApplyPullKelvinBrush(WorldRay);
+			break;
+		case EDynamicMeshSculptBrushType::PullSharpKelvin:
+			ApplyPullSharpKelvinBrush(WorldRay);
+			break;
 		case EDynamicMeshSculptBrushType::Smooth:
 			ApplySmoothBrush(WorldRay);
 			break;
 		case EDynamicMeshSculptBrushType::Pinch:
 			ApplyPinchBrush(WorldRay);
 			break;
+		case EDynamicMeshSculptBrushType::TwistKelvin:
+			ApplyTwistKelvinBrush(WorldRay);
+			break;
 		case EDynamicMeshSculptBrushType::Inflate:
 			ApplyInflateBrush(WorldRay);
 			break;
+		case EDynamicMeshSculptBrushType::ScaleKelvin:
+			ApplyScaleKelvinBrush(WorldRay);
+			break;
 		case EDynamicMeshSculptBrushType::Flatten:
 			ApplyFlattenBrush(WorldRay);
 			break;
@@ -597,6 +524,23 @@ double UDynamicMeshSculptTool::CalculateBrushFalloff(double Distance)
 	return w;
 }
 
+
+
+void UDynamicMeshSculptTool::SyncMeshWithPositionBuffer(FDynamicMesh3* Mesh)
+{
+	const int NumV = ROIPositionBuffer.Num();
+	checkSlow(VertexROI.Num() <= NumV);
+
+	for (int k = 0; k < NumV; ++k)
+	{
+		int VertIdx = VertexROI[k];
+		const FVector3d& NewPos = ROIPositionBuffer[k];
+		FVector3d OrigPos = Mesh->GetVertex(VertIdx);
+		Mesh->SetVertex(VertIdx, NewPos);
+		UpdateSavedVertex(VertIdx, OrigPos, NewPos);
+	}
+}
+
 void UDynamicMeshSculptTool::ApplySmoothBrush(const FRay& WorldRay)
 {
 	bool bHit = UpdateBrushPositionOnSculptMesh(WorldRay, true);
@@ -627,14 +571,7 @@ void UDynamicMeshSculptTool::ApplySmoothBrush(const FRay& WorldRay)
 		ROIPositionBuffer[k] = NewPos;
 	});
 
-	for (int k = 0; k < NumV; ++k)
-	{
-		int VertIdx = VertexROI[k];
-		const FVector3d& NewPos = ROIPositionBuffer[k];
-		FVector3d OrigPos = Mesh->GetVertex(VertIdx);
-		Mesh->SetVertex(VertIdx, NewPos);
-		UpdateSavedVertex(VertIdx, OrigPos, NewPos);
-	}
+	SyncMeshWithPositionBuffer(Mesh);
 
 	ScheduleRemeshPass();
 	LastBrushPosLocal = NewBrushPosLocal;
@@ -642,7 +579,7 @@ void UDynamicMeshSculptTool::ApplySmoothBrush(const FRay& WorldRay)
 
 void UDynamicMeshSculptTool::ApplyMoveBrush(const FRay& WorldRay)
 {
-	bool bHit = UpdateBrushPositionOnActivePlane(WorldRay);
+	UpdateBrushPositionOnActivePlane(WorldRay);
 
 	FVector3d NewBrushPosLocal = CurTargetTransform.InverseTransformPosition(LastBrushPosWorld);
 	FVector3d MoveVec = NewBrushPosLocal - LastBrushPosLocal;
@@ -672,14 +609,8 @@ void UDynamicMeshSculptTool::ApplyMoveBrush(const FRay& WorldRay)
 		ROIPositionBuffer[k] = NewPos;
 	});
 
-	for (int k = 0; k < NumV; ++k)
-	{
-		int VertIdx = VertexROI[k];
-		const FVector3d& NewPos = ROIPositionBuffer[k];
-		FVector3d OrigPos = Mesh->GetVertex(VertIdx);
-		Mesh->SetVertex(VertIdx, NewPos);
-		UpdateSavedVertex(VertIdx, OrigPos, NewPos);
-	}
+	// Update the mesh positions to match those in the position buffer
+	SyncMeshWithPositionBuffer(Mesh);
 
 	ScheduleRemeshPass();
 	LastBrushPosLocal = NewBrushPosLocal;
@@ -723,14 +654,8 @@ void UDynamicMeshSculptTool::ApplyOffsetBrush(const FRay& WorldRay, bool bUseVie
 		}
 	});
 
-	for (int k = 0; k < NumV; ++k)
-	{
-		int VertIdx = VertexROI[k];
-		const FVector3d& NewPos = ROIPositionBuffer[k];
-		FVector3d OrigPos = Mesh->GetVertex(VertIdx);
-		Mesh->SetVertex(VertIdx, NewPos);
-		UpdateSavedVertex(VertIdx, OrigPos, NewPos);
-	}
+	// Update the mesh positions to match those in the position buffer
+	SyncMeshWithPositionBuffer(Mesh);
 
 	ScheduleRemeshPass();
 	LastBrushPosLocal = NewBrushPosLocal;
@@ -782,14 +707,8 @@ void UDynamicMeshSculptTool::ApplySculptMaxBrush(const FRay& WorldRay)
 		}
 	});
 
-	for (int k = 0; k < NumV; ++k)
-	{
-		int VertIdx = VertexROI[k];
-		const FVector3d& NewPos = ROIPositionBuffer[k];
-		FVector3d OrigPos = Mesh->GetVertex(VertIdx);
-		Mesh->SetVertex(VertIdx, NewPos);
-		UpdateSavedVertex(VertIdx, OrigPos, NewPos);
-	}
+	// Update the mesh positions to match those in the position buffer
+	SyncMeshWithPositionBuffer(Mesh);
 
 	ScheduleRemeshPass();
 	LastBrushPosLocal = NewBrushPosLocal;
@@ -841,14 +760,8 @@ void UDynamicMeshSculptTool::ApplyPinchBrush(const FRay& WorldRay)
 		ROIPositionBuffer[k] = NewPos;
 	});
 
-	for (int k = 0; k < NumV; ++k)
-	{
-		int VertIdx = VertexROI[k];
-		const FVector3d& NewPos = ROIPositionBuffer[k];
-		FVector3d OrigPos = Mesh->GetVertex(VertIdx);
-		Mesh->SetVertex(VertIdx, NewPos);
-		UpdateSavedVertex(VertIdx, OrigPos, NewPos);
-	}
+	// Update the mesh positions to match those in the position buffer
+	SyncMeshWithPositionBuffer(Mesh);
 
 
 	ScheduleRemeshPass();
@@ -924,14 +837,8 @@ void UDynamicMeshSculptTool::ApplyPlaneBrush(const FRay& WorldRay)
 		ROIPositionBuffer[k] = NewPos;
 	});
 
-	for (int k = 0; k < NumV; ++k)
-	{
-		int VertIdx = VertexROI[k];
-		const FVector3d& NewPos = ROIPositionBuffer[k];
-		FVector3d OrigPos = Mesh->GetVertex(VertIdx);
-		Mesh->SetVertex(VertIdx, NewPos);
-		UpdateSavedVertex(VertIdx, OrigPos, NewPos);
-	}
+	// Update the mesh positions to match those in the position buffer
+	SyncMeshWithPositionBuffer(Mesh);
 
 	ScheduleRemeshPass();
 	LastBrushPosLocal = NewBrushPosLocal;
@@ -979,14 +886,8 @@ void UDynamicMeshSculptTool::ApplyFixedPlaneBrush(const FRay& WorldRay)
 		ROIPositionBuffer[k] = NewPos;
 	});
 
-	for (int k = 0; k < NumV; ++k)
-	{
-		int VertIdx = VertexROI[k];
-		const FVector3d& NewPos = ROIPositionBuffer[k];
-		FVector3d OrigPos = Mesh->GetVertex(VertIdx);
-		Mesh->SetVertex(VertIdx, NewPos);
-		UpdateSavedVertex(VertIdx, OrigPos, NewPos);
-	}
+	// Update the mesh positions to match those in the position buffer
+	SyncMeshWithPositionBuffer(Mesh);
 
 	ScheduleRemeshPass();
 	LastBrushPosLocal = NewBrushPosLocal;
@@ -1035,14 +936,8 @@ void UDynamicMeshSculptTool::ApplyFlattenBrush(const FRay& WorldRay)
 		ROIPositionBuffer[k] = NewPos;
 	});
 
-	for (int k = 0; k < NumV; ++k)
-	{
-		int VertIdx = VertexROI[k];
-		const FVector3d& NewPos = ROIPositionBuffer[k];
-		FVector3d OrigPos = Mesh->GetVertex(VertIdx);
-		Mesh->SetVertex(VertIdx, NewPos);
-		UpdateSavedVertex(VertIdx, OrigPos, NewPos);
-	}
+	// Update the mesh positions to match those in the position buffer
+	SyncMeshWithPositionBuffer(Mesh);
 
 	ScheduleRemeshPass();
 	LastBrushPosLocal = NewBrushPosLocal;
@@ -1087,21 +982,13 @@ void UDynamicMeshSculptTool::ApplyInflateBrush(const FRay& WorldRay)
 		ROIPositionBuffer[k] = NewPos;
 	});
 
-	for (int k = 0; k < NumV; ++k)
-	{
-		int VertIdx = VertexROI[k];
-		const FVector3d& NewPos = ROIPositionBuffer[k];
-		FVector3d OrigPos = Mesh->GetVertex(VertIdx);
-		Mesh->SetVertex(VertIdx, NewPos);
-		UpdateSavedVertex(VertIdx, OrigPos, NewPos);
-	}
-
+	// Update the mesh positions to match those in the position buffer
+	SyncMeshWithPositionBuffer(Mesh);
 	ScheduleRemeshPass();
 	LastBrushPosLocal = NewBrushPosLocal;
 }
 
 
-
 void UDynamicMeshSculptTool::ApplyResampleBrush(const FRay& WorldRay)
 {
 	UpdateBrushPositionOnTargetMesh(WorldRay, true);
@@ -1131,21 +1018,139 @@ void UDynamicMeshSculptTool::ApplyResampleBrush(const FRay& WorldRay)
 		}
 	});
 
-	for (int k = 0; k < NumV; ++k)
+	SyncMeshWithPositionBuffer(Mesh);
+	ScheduleRemeshPass();
+	LastBrushPosLocal = NewBrushPosLocal;
+}
+
+
+void UDynamicMeshSculptTool::ApplyPullKelvinBrush(const FRay& WorldRay)
+{
+	UpdateBrushPositionOnActivePlane(WorldRay);
+	FVector3d BrushNormalLocal = CurTargetTransform.InverseTransformNormal(LastBrushPosNormalWorld);
+	FVector3d NewBrushPosLocal = CurTargetTransform.InverseTransformPosition(LastBrushPosWorld);
+	FVector3d MoveVec = NewBrushPosLocal - LastBrushPosLocal;
+
+	if (MoveVec.SquaredLength() <= 0)
 	{
-		int VertIdx = VertexROI[k];
-		const FVector3d& NewPos = ROIPositionBuffer[k];
-		FVector3d OrigPos = Mesh->GetVertex(VertIdx);
-		Mesh->SetVertex(VertIdx, NewPos);
-		UpdateSavedVertex(VertIdx, OrigPos, NewPos);
+		LastBrushPosLocal = NewBrushPosLocal;
+		return;
+	}
+	
+	FDynamicMesh3* Mesh = DynamicMeshComponent->GetMesh();
+	FKelvinletBrushOp KelvinBrushOp(*Mesh);
+
+	const EKelvinletBrushMode KelvinMode = EKelvinletBrushMode::PullKelvinlet;
+
+	FKelvinletBrushOp::FKelvinletBrushOpProperties  KelvinletBrushOpProperties(KelvinMode, *KelvinBrushProperties, *BrushProperties);
+	KelvinletBrushOpProperties.Direction = FVector(MoveVec.X, MoveVec.Y, MoveVec.Z);  //FVector(BrushNormalLocal.X, BrushNormalLocal.Y, BrushNormalLocal.Z);
+	KelvinletBrushOpProperties.Size *= 0.6;
+	
+
+	FMatrix ToBrush; ToBrush.SetIdentity();  ToBrush.SetOrigin(-FVector(NewBrushPosLocal.X, NewBrushPosLocal.Y, NewBrushPosLocal.Z)); //  ToBrush.
+
+	KelvinBrushOp.ApplyBrush(KelvinletBrushOpProperties, ToBrush, VertexROI, ROIPositionBuffer);
+
+	// Update the mesh positions to match those in the position buffer
+	SyncMeshWithPositionBuffer(Mesh);
+	ScheduleRemeshPass();
+	LastBrushPosLocal = NewBrushPosLocal;
+}
+
+
+void UDynamicMeshSculptTool::ApplyPullSharpKelvinBrush(const FRay& WorldRay)
+{
+	UpdateBrushPositionOnActivePlane(WorldRay);
+	FVector3d BrushNormalLocal = CurTargetTransform.InverseTransformNormal(LastBrushPosNormalWorld);
+	FVector3d NewBrushPosLocal = CurTargetTransform.InverseTransformPosition(LastBrushPosWorld);
+	FVector3d MoveVec = NewBrushPosLocal - LastBrushPosLocal;
+
+	if (MoveVec.SquaredLength() <= 0)
+	{
+		LastBrushPosLocal = NewBrushPosLocal;
+		return;
 	}
 
+	FDynamicMesh3* Mesh = DynamicMeshComponent->GetMesh();
+	FKelvinletBrushOp KelvinBrushOp(*Mesh);
+
+	const EKelvinletBrushMode KelvinMode = EKelvinletBrushMode::SharpPullKelvinlet;
+
+	FKelvinletBrushOp::FKelvinletBrushOpProperties  KelvinletBrushOpProperties(KelvinMode, *KelvinBrushProperties, *BrushProperties);
+	KelvinletBrushOpProperties.Direction = FVector(MoveVec.X, MoveVec.Y, MoveVec.Z);  //FVector(BrushNormalLocal.X, BrushNormalLocal.Y, BrushNormalLocal.Z);
+	KelvinletBrushOpProperties.Size *= 0.6;
+
+
+	FMatrix ToBrush; ToBrush.SetIdentity();  ToBrush.SetOrigin(-FVector(NewBrushPosLocal.X, NewBrushPosLocal.Y, NewBrushPosLocal.Z)); //  ToBrush.
+
+	KelvinBrushOp.ApplyBrush(KelvinletBrushOpProperties, ToBrush, VertexROI, ROIPositionBuffer);
+
+	// Update the mesh positions to match those in the position buffer
+	SyncMeshWithPositionBuffer(Mesh);
 	ScheduleRemeshPass();
 	LastBrushPosLocal = NewBrushPosLocal;
 }
 
+void UDynamicMeshSculptTool::ApplyTwistKelvinBrush(const FRay& WorldRay)
+{
+	UpdateBrushPositionOnTargetMesh(WorldRay, true);
+	
+	FVector3d BrushNormalLocal = CurTargetTransform.InverseTransformNormal(LastBrushPosNormalWorld);
+	FVector3d NewBrushPosLocal = CurTargetTransform.InverseTransformPosition(LastBrushPosWorld);
+	
+
+
+	double Direction = (bInvert) ? -1.0 : 1.0;
+	double UseSpeed  = Direction * FMathd::Sqrt(CurrentBrushRadius) * (SculptProperties->PrimaryBrushSpeed) * ActivePressure;
+
+	FDynamicMesh3* Mesh = DynamicMeshComponent->GetMesh();
+	FKelvinletBrushOp KelvinBrushOp(*Mesh);
 
 
+	FKelvinletBrushOp::FKelvinletBrushOpProperties  KelvinletBrushOpProperties(EKelvinletBrushMode::TwistKelvinlet, *KelvinBrushProperties, *BrushProperties);
+	KelvinletBrushOpProperties.Direction = UseSpeed * FVector(BrushNormalLocal.X, BrushNormalLocal.Y, BrushNormalLocal.Z); // twist about local normal
+	KelvinletBrushOpProperties.Size *= 0.35; // reduce the core size of this brush.
+	
+	FMatrix ToBrush; ToBrush.SetIdentity();  ToBrush.SetOrigin(-FVector(NewBrushPosLocal.X, NewBrushPosLocal.Y, NewBrushPosLocal.Z)); //  ToBrush.
+
+	KelvinBrushOp.ApplyBrush(KelvinletBrushOpProperties, ToBrush, VertexROI, ROIPositionBuffer);
+
+	// Update the mesh positions to match those in the position buffer
+	SyncMeshWithPositionBuffer(Mesh);
+	ScheduleRemeshPass();
+	LastBrushPosLocal = NewBrushPosLocal;
+}
+
+void UDynamicMeshSculptTool::ApplyScaleKelvinBrush(const FRay& WorldRay)
+{
+	UpdateBrushPositionOnSculptMesh(WorldRay, true);
+
+	FVector3d NewBrushPosLocal = CurTargetTransform.InverseTransformPosition(LastBrushPosWorld);
+	FVector3d BrushNormalLocal = CurTargetTransform.InverseTransformNormal(LastBrushPosNormalWorld);
+	FVector3d OffsetBrushPosLocal = NewBrushPosLocal - BrushProperties->Depth * CurrentBrushRadius * BrushNormalLocal;
+	
+
+	double Direction = (bInvert) ? -1.0 : 1.0;
+	double UseSpeed = Direction * FMath::Sqrt(CurrentBrushRadius) * SculptProperties->PrimaryBrushSpeed * 0.025 * ActivePressure; ; 
+
+	FDynamicMesh3* Mesh = DynamicMeshComponent->GetMesh();
+	FKelvinletBrushOp KelvinBrushOp(*Mesh);
+
+
+	FKelvinletBrushOp::FKelvinletBrushOpProperties  KelvinletBrushOpProperties(EKelvinletBrushMode::ScaleKelvinlet, *KelvinBrushProperties, *BrushProperties);
+	KelvinletBrushOpProperties.Direction = FVector(UseSpeed, 0., 0.); // it is a bit iffy, but we only use the first component for the scale
+	KelvinletBrushOpProperties.Size *= 0.35;
+
+	FMatrix ToBrush; ToBrush.SetIdentity();  ToBrush.SetOrigin(-FVector(OffsetBrushPosLocal.X, OffsetBrushPosLocal.Y, OffsetBrushPosLocal.Z)); //  ToBrush.
+
+	KelvinBrushOp.ApplyBrush(KelvinletBrushOpProperties, ToBrush, VertexROI, ROIPositionBuffer);
+
+	// Update the mesh positions to match those in the position buffer
+	SyncMeshWithPositionBuffer(Mesh);
+	ScheduleRemeshPass();
+	LastBrushPosLocal = NewBrushPosLocal;
+}
+
 int UDynamicMeshSculptTool::FindHitSculptMeshTriangle(const FRay3d& LocalRay)
 {
 	if (BrushProperties->bHitBackFaces)
@@ -1172,6 +1177,8 @@ int UDynamicMeshSculptTool::FindHitSculptMeshTriangle(const FRay3d& LocalRay)
 
 int UDynamicMeshSculptTool::FindHitTargetMeshTriangle(const FRay3d& LocalRay)
 {
+	IMeshSpatial::FQueryOptions RaycastOptions;
+
 	if (BrushProperties->bHitBackFaces == false)
 	{
 		FDynamicMesh3* Mesh = DynamicMeshComponent->GetMesh();
@@ -1180,7 +1187,7 @@ int UDynamicMeshSculptTool::FindHitTargetMeshTriangle(const FRay3d& LocalRay)
 		GetToolManager()->GetContextQueriesAPI()->GetCurrentViewState(StateOut);
 		FVector3d LocalEyePosition(CurTargetTransform.InverseTransformPosition(StateOut.Position));
 
-		BrushTargetMeshSpatial.TriangleFilterF = [this, Mesh, LocalEyePosition](int TriangleID) {
+		RaycastOptions.TriangleFilterF = [this, Mesh, LocalEyePosition](int TriangleID) {
 			FVector3d Normal, Centroid;
 			double Area;
 			Mesh->GetTriInfo(TriangleID, Normal, Area, Centroid);
@@ -1188,12 +1195,7 @@ int UDynamicMeshSculptTool::FindHitTargetMeshTriangle(const FRay3d& LocalRay)
 		};
 	}
 
-	int HitTID =  BrushTargetMeshSpatial.FindNearestHitTriangle(LocalRay);
-
-	if (BrushProperties->bHitBackFaces == false)
-	{
-		BrushTargetMeshSpatial.TriangleFilterF = nullptr;
-	}
+	int HitTID =  BrushTargetMeshSpatial.FindNearestHitTriangle(LocalRay, RaycastOptions);
 
 	return HitTID;
 }
@@ -1401,12 +1403,10 @@ void UDynamicMeshSculptTool::Render(IToolsContextRenderAPI* RenderAPI)
 	}
 }
 
-void UDynamicMeshSculptTool::Tick(float DeltaTime)
+void UDynamicMeshSculptTool::OnTick(float DeltaTime)
 {
 	SCOPE_CYCLE_COUNTER(STAT_SculptToolTick);
 
-	UMeshSurfacePointTool::Tick(DeltaTime);
-
 	ActivePressure = GetCurrentDevicePressure();
 
 	// Allow a tick to pass between application of brush stamps. Bizarrely this
@@ -1599,10 +1599,10 @@ void UDynamicMeshSculptTool::ScheduleRemeshPass()
 
 /*
         Split	Collapse	Vertices Pinned	Flip
-Fixed	FALSE	FALSE		TRUE				FALSE
-Refine	TRUE	FALSE		TRUE				FALSE
-Free	TRUE	TRUE		FALSE				FALSE
-Ignore	TRUE	TRUE		FALSE				TRUE
+Fixed	FALSE	FALSE	TRUE	FALSE
+Refine	TRUE	FALSE	TRUE	FALSE
+Free	TRUE	TRUE	FALSE	FALSE
+Ignore	TRUE	TRUE	FALSE	TRUE
 */
 
 void UDynamicMeshSculptTool::ConfigureRemesher(FSubRegionRemesher& Remesher)
diff --git a/Engine/Plugins/Experimental/MeshModelingToolset/Source/MeshModelingTools/Private/EditMeshPolygonsTool.cpp b/Engine/Plugins/Experimental/MeshModelingToolset/Source/MeshModelingTools/Private/EditMeshPolygonsTool.cpp
--- a/Engine/Plugins/Experimental/MeshModelingToolset/Source/MeshModelingTools/Private/EditMeshPolygonsTool.cpp
+++ b/Engine/Plugins/Experimental/MeshModelingToolset/Source/MeshModelingTools/Private/EditMeshPolygonsTool.cpp
@@ -59,27 +59,6 @@ void UEditMeshPolygonsToolActionPropertySet::PostAction(EEditMeshPolygonsToolAct
 	}
 }
 
-
-void UPolyEditExtrudeProperties::SaveRestoreProperties(UInteractiveTool* RestoreToTool, bool bSaving)
-{
-	UPolyEditExtrudeProperties* PropertyCache = GetPropertyCache<UPolyEditExtrudeProperties>();
-	SaveRestoreProperty(PropertyCache->Direction, this->Direction, bSaving);
-}
-void UPolyEditCutProperties::SaveRestoreProperties(UInteractiveTool* RestoreToTool, bool bSaving)
-{
-	UPolyEditCutProperties* PropertyCache = GetPropertyCache<UPolyEditCutProperties>();
-	SaveRestoreProperty(PropertyCache->Orientation, this->Orientation, bSaving);
-	SaveRestoreProperty(PropertyCache->bSnapToVertices, this->bSnapToVertices, bSaving);
-}
-void UPolyEditSetUVProperties::SaveRestoreProperties(UInteractiveTool* RestoreToTool, bool bSaving)
-{
-	UPolyEditSetUVProperties* PropertyCache = GetPropertyCache<UPolyEditSetUVProperties>();
-	SaveRestoreProperty(PropertyCache->bShowMaterial, this->bShowMaterial, bSaving);
-}
-
-
-
-
 /*
 * Tool methods
 */
@@ -146,8 +125,13 @@ void UEditMeshPolygonsTool::Setup()
 	// add properties
 	TransformProps = NewObject<UPolyEditTransformProperties>(this);
 	AddToolPropertySource(TransformProps);
-	LocalFrameModeWatcher.Initialize([&]() { return TransformProps->LocalFrameMode; }, [this](ELocalFrameMode) { UpdateMultiTransformerFrame(); }, TransformProps->LocalFrameMode);
-	LockFrameWatcher.Initialize([&]() { return TransformProps->bLockRotation; }, [this](bool) { LockedTransfomerFrame = LastTransformerFrame; UpdateMultiTransformerFrame(); }, TransformProps->bLockRotation);
+	TransformProps->WatchProperty(TransformProps->LocalFrameMode,
+								  [this](ELocalFrameMode) { UpdateMultiTransformerFrame(); });
+	TransformProps->WatchProperty(TransformProps->bLockRotation,
+								  [this](bool)
+								  {
+									  LockedTransfomerFrame = LastTransformerFrame; UpdateMultiTransformerFrame();
+								  });
 
 	// set up SelectionMechanic
 	SelectionMechanic = NewObject<UPolygonSelectionMechanic>(this);
@@ -215,7 +199,8 @@ void UEditMeshPolygonsTool::Setup()
 	ExtrudeProperties->RestoreProperties(this);
 	AddToolPropertySource(ExtrudeProperties);
 	SetToolPropertySourceEnabled(ExtrudeProperties, false);
-	ExtrudeDirectionWatcher.Initialize([this]() { return ExtrudeProperties->Direction; }, [this](EPolyEditExtrudeDirection) { RestartExtrude(); }, ExtrudeProperties->Direction);
+	ExtrudeProperties->WatchProperty(ExtrudeProperties->Direction,
+									 [this](EPolyEditExtrudeDirection){ RestartExtrude(); });
 
 	CutProperties = NewObject<UPolyEditCutProperties>();
 	CutProperties->RestoreProperties(this);
@@ -621,15 +606,8 @@ void UEditMeshPolygonsTool::ComputeUpdate_Gizmo()
 
 
 
-void UEditMeshPolygonsTool::Tick(float DeltaTime)
+void UEditMeshPolygonsTool::OnTick(float DeltaTime)
 {
-	UMeshSurfacePointTool::Tick(DeltaTime);
-
-	// update value watchers
-	ExtrudeDirectionWatcher.CheckAndUpdate();
-	LocalFrameModeWatcher.CheckAndUpdate();
-	LockFrameWatcher.CheckAndUpdate();
-
 	MultiTransformer->Tick(DeltaTime);
 
 	if (bGizmoUpdatePending)
@@ -838,7 +816,7 @@ void UEditMeshPolygonsTool::UpdateChangeFromROI(bool bFinal)
 	}
 
 	FDynamicMesh3* Mesh = DynamicMeshComponent->GetMesh();
-	ActiveVertexChange->SavePositions(Mesh, LinearDeformer.GetModifiedVertices(), !bFinal);
+	ActiveVertexChange->SaveVertices(Mesh, LinearDeformer.GetModifiedVertices(), !bFinal);
 	ActiveVertexChange->SaveOverlayNormals(Mesh, LinearDeformer.GetModifiedOverlayNormals(), !bFinal);
 }
 
@@ -847,7 +825,7 @@ void UEditMeshPolygonsTool::BeginChange()
 {
 	if (ActiveVertexChange == nullptr)
 	{
-		ActiveVertexChange = new FMeshVertexChangeBuilder(true);
+		ActiveVertexChange = new FMeshVertexChangeBuilder(EMeshVertexChangeComponents::VertexPositions | EMeshVertexChangeComponents::OverlayNormals);
 		UpdateChangeFromROI(false);
 	}
 }
diff --git a/Engine/Plugins/Experimental/MeshModelingToolset/Source/MeshModelingTools/Private/EditNormalsTool.cpp b/Engine/Plugins/Experimental/MeshModelingToolset/Source/MeshModelingTools/Private/EditNormalsTool.cpp
--- a/Engine/Plugins/Experimental/MeshModelingToolset/Source/MeshModelingTools/Private/EditNormalsTool.cpp
+++ b/Engine/Plugins/Experimental/MeshModelingToolset/Source/MeshModelingTools/Private/EditNormalsTool.cpp
@@ -206,7 +206,7 @@ void UEditNormalsTool::Render(IToolsContextRenderAPI* RenderAPI)
 {
 }
 
-void UEditNormalsTool::Tick(float DeltaTime)
+void UEditNormalsTool::OnTick(float DeltaTime)
 {
 	for (UMeshOpPreviewWithBackgroundCompute* Preview : Previews)
 	{
diff --git a/Engine/Plugins/Experimental/MeshModelingToolset/Source/MeshModelingTools/Private/EditPivotTool.cpp b/Engine/Plugins/Experimental/MeshModelingToolset/Source/MeshModelingTools/Private/EditPivotTool.cpp
--- a/Engine/Plugins/Experimental/MeshModelingToolset/Source/MeshModelingTools/Private/EditPivotTool.cpp
+++ b/Engine/Plugins/Experimental/MeshModelingToolset/Source/MeshModelingTools/Private/EditPivotTool.cpp
@@ -193,7 +193,7 @@ void UEditPivotTool::RequestAction(EEditPivotToolActions ActionType)
 }
 
 
-void UEditPivotTool::Tick(float DeltaTime)
+void UEditPivotTool::OnTick(float DeltaTime)
 {
 	if (PendingAction != EEditPivotToolActions::NoAction)
 	{
diff --git a/Engine/Plugins/Experimental/MeshModelingToolset/Source/MeshModelingTools/Private/EditUVIslandsTool.cpp b/Engine/Plugins/Experimental/MeshModelingToolset/Source/MeshModelingTools/Private/EditUVIslandsTool.cpp
--- a/Engine/Plugins/Experimental/MeshModelingToolset/Source/MeshModelingTools/Private/EditUVIslandsTool.cpp
+++ b/Engine/Plugins/Experimental/MeshModelingToolset/Source/MeshModelingTools/Private/EditUVIslandsTool.cpp
@@ -411,9 +411,8 @@ void UEditUVIslandsTool::ComputeUpdate_Gizmo()
 
 
 
-void UEditUVIslandsTool::Tick(float DeltaTime)
+void UEditUVIslandsTool::OnTick(float DeltaTime)
 {
-	UMeshSurfacePointTool::Tick(DeltaTime);
 	MultiTransformer->Tick(DeltaTime);
 
 	if (bSelectionStateDirty)
@@ -481,9 +480,8 @@ void FUVGroupTopology::CalculateIslandGroups()
 FFrame3d FUVGroupTopology::GetIslandFrame(int32 GroupID, FDynamicMeshAABBTree3& AABBTree)
 {
 	FFrame3d Frame = GetGroupFrame(GroupID);
-	AABBTree.TriangleFilterF = [&](int32 TriangleID) { return GetGroupID(TriangleID) == GroupID; };
-	Frame.Origin = AABBTree.FindNearestPoint(Frame.Origin);
-	AABBTree.TriangleFilterF = nullptr;
+	IMeshSpatial::FQueryOptions QueryOptions([&](int32 TriangleID) { return GetGroupID(TriangleID) == GroupID; });
+	Frame.Origin = AABBTree.FindNearestPoint(Frame.Origin, QueryOptions);
 
 	const TArray<int32>& Triangles = GetGroupTriangles(GroupID);
 
diff --git a/Engine/Plugins/Experimental/MeshModelingToolset/Source/MeshModelingTools/Private/HoleFillTool.cpp b/Engine/Plugins/Experimental/MeshModelingToolset/Source/MeshModelingTools/Private/HoleFillTool.cpp
new file mode 100644
--- /dev/null
+++ b/Engine/Plugins/Experimental/MeshModelingToolset/Source/MeshModelingTools/Private/HoleFillTool.cpp
@@ -0,0 +1,384 @@
+// Copyright Epic Games, Inc. All Rights Reserved.
+
+#include "HoleFillTool.h"
+#include "ToolBuilderUtil.h"
+#include "InteractiveToolManager.h"
+#include "MeshDescriptionToDynamicMesh.h"
+#include "ToolSetupUtil.h"
+#include "MeshNormals.h"
+#include "Changes/DynamicMeshChangeTarget.h"
+#include "DynamicMeshToMeshDescription.h"
+#include "BaseBehaviors/SingleClickBehavior.h"
+#include "BaseBehaviors/MouseHoverBehavior.h"
+#include "MeshBoundaryLoops.h"
+#include "MeshOpPreviewHelpers.h"
+#include "Selection/PolygonSelectionMechanic.h"
+
+#define LOCTEXT_NAMESPACE "UHoleFillTool"
+
+/*
+ * ToolBuilder
+ */
+
+bool UHoleFillToolBuilder::CanBuildTool(const FToolBuilderState& SceneState) const
+{
+	return ToolBuilderUtil::CountComponents(SceneState, CanMakeComponentTarget) == 1;
+}
+
+UInteractiveTool* UHoleFillToolBuilder::BuildTool(const FToolBuilderState& SceneState) const
+{
+	UActorComponent* ActorComponent = ToolBuilderUtil::FindFirstComponent(SceneState, CanMakeComponentTarget);
+	auto* MeshComponent = Cast<UPrimitiveComponent>(ActorComponent);
+	check(MeshComponent != nullptr);
+
+	UHoleFillTool* NewTool = NewObject<UHoleFillTool>(SceneState.ToolManager);
+	NewTool->SetSelection(MakeComponentTarget(MeshComponent));
+	NewTool->SetWorld(SceneState.World);
+
+	return NewTool;
+}
+
+/*
+ * Tool properties
+ */
+void UHoleFillToolActions::PostAction(EHoleFillToolActions Action)
+{
+	if (ParentTool.IsValid())
+	{
+		ParentTool->RequestAction(Action);
+	}
+}
+
+
+/*
+* Op Factory
+*/
+
+TUniquePtr<FDynamicMeshOperator> UHoleFillOperatorFactory::MakeNewOperator()
+{
+	TUniquePtr<FHoleFillOp> FillOp = MakeUnique<FHoleFillOp>();
+
+	FTransform LocalToWorld = FillTool->ComponentTarget->GetWorldTransform();
+	FillOp->SetResultTransform((FTransform3d)LocalToWorld);
+	FillOp->OriginalMesh = &(FillTool->OriginalMesh);
+	FillOp->FillType = FillTool->Properties->FillType;
+	FillOp->MeshUVScaleFactor = FillTool->MeshUVScaleFactor;
+
+	FillTool->GetLoopsToFill(FillOp->Loops);
+
+	return FillOp;
+}
+
+
+/*
+ * Tool
+ */
+
+void UHoleFillTool::Setup()
+{
+	USingleSelectionTool::Setup();
+
+	// initialize properties
+	Properties = NewObject<UHoleFillToolProperties>(this, TEXT("Hole Fill Settings"));
+	Properties->RestoreProperties(this);
+	AddToolPropertySource(Properties);
+
+	Actions = NewObject<UHoleFillToolActions>(this, TEXT("Hole Fill Actions"));
+	Actions->Initialize(this);
+	AddToolPropertySource(Actions);
+
+	ToolPropertyObjects.Add(this);
+
+	if (!ComponentTarget)
+	{
+		return;
+	}
+
+	// click behavior
+	USingleClickInputBehavior* ClickBehavior = NewObject<USingleClickInputBehavior>();
+	ClickBehavior->Initialize(this);
+	AddInputBehavior(ClickBehavior);
+
+	// hover behavior
+	UMouseHoverBehavior* HoverBehavior = NewObject<UMouseHoverBehavior>();
+	HoverBehavior->Initialize(this);
+	AddInputBehavior(HoverBehavior);
+
+	// create mesh to operate on
+	FMeshDescriptionToDynamicMesh Converter;
+	Converter.Convert(ComponentTarget->GetMesh(), const_cast<FDynamicMesh3&>(OriginalMesh));
+
+	// initialize hit query
+	MeshSpatial.SetMesh(&OriginalMesh);
+
+	// initialize topology
+	Topology = MakeUnique<FBasicTopology>(&OriginalMesh, false);
+	Topology->RebuildTopology();
+
+	// Set up selection mechanic to find and select edges
+	SelectionMechanic = NewObject<UPolygonSelectionMechanic>(this);
+	SelectionMechanic->bAddSelectionFilterPropertiesToParentTool = false;
+	SelectionMechanic->Setup(this);
+	SelectionMechanic->Properties->bSelectEdges = true;
+	SelectionMechanic->Properties->bSelectFaces = false;
+	SelectionMechanic->Properties->bSelectVertices = false;
+	SelectionMechanic->Initialize(&OriginalMesh,
+		FTransform3d(ComponentTarget->GetWorldTransform()),
+		Topology.Get(),
+		[this]() { return &MeshSpatial; },
+		[]() { return true; }	// allow adding to selection without modifier key
+	);
+	
+	// Store a UV scale based on the original mesh bounds
+	MeshUVScaleFactor = (1.0 / OriginalMesh.GetBounds().MaxDim());
+
+	// initialize the PreviewMesh+BackgroundCompute object
+	SetupPreview();
+	Preview->InvalidateResult();
+
+	// Hide all meshes except the Preview
+	ComponentTarget->SetOwnerVisibility(false);
+}
+
+
+void UHoleFillTool::OnTick(float DeltaTime)
+{
+	Preview->Tick(DeltaTime);
+
+	if (bHavePendingAction)
+	{
+		ApplyAction(PendingAction);
+		bHavePendingAction = false;
+		PendingAction = EHoleFillToolActions::NoAction;
+	}
+}
+
+void UHoleFillTool::OnPropertyModified(UObject* PropertySet, FProperty* Property)
+{
+	Preview->InvalidateResult();
+}
+
+bool UHoleFillTool::CanAccept() const
+{
+	return Preview->HaveValidResult();
+}
+
+void UHoleFillTool::Shutdown(EToolShutdownType ShutdownType)
+{
+	Properties->SaveProperties(this);
+
+	ComponentTarget->SetOwnerVisibility(true);
+
+	FDynamicMeshOpResult Result = Preview->Shutdown();
+	if (ShutdownType == EToolShutdownType::Accept)
+	{
+		GetToolManager()->BeginUndoTransaction(LOCTEXT("HoleFillToolTransactionName", "Hole Fill Tool"));
+
+		check(Result.Mesh.Get() != nullptr);
+		ComponentTarget->CommitMesh([&Result](const FPrimitiveComponentTarget::FCommitParams& CommitParams)
+		{
+			FDynamicMeshToMeshDescription Converter;
+
+			// full conversion if normal topology changed or faces were inverted
+			Converter.Convert(Result.Mesh.Get(), *CommitParams.MeshDescription);
+		});
+
+		GetToolManager()->EndUndoTransaction();
+	}
+}
+
+FInputRayHit UHoleFillTool::IsHitByClick(const FInputDeviceRay& ClickPos)
+{
+	FHitResult OutHit;
+	if (SelectionMechanic->TopologyHitTest(ClickPos.WorldRay, OutHit))
+	{
+		return FInputRayHit(OutHit.Distance);
+	}
+
+	return FInputRayHit(TNumericLimits<float>::Max());
+}
+
+void UHoleFillTool::OnClicked(const FInputDeviceRay& ClickPos)
+{
+	// update selection
+	GetToolManager()->BeginUndoTransaction(LOCTEXT("PolyMeshSelectionChange", "Selection"));
+	SelectionMechanic->BeginChange();
+
+	FVector3d LocalHitPosition, LocalHitNormal;
+	bool bSelectionModified = SelectionMechanic->UpdateSelection(ClickPos.WorldRay, LocalHitPosition, LocalHitNormal);
+
+	if (bSelectionModified)
+	{
+		UpdateActiveBoundaryLoopSelection();
+		Preview->InvalidateResult();
+	}
+
+	SelectionMechanic->EndChangeAndEmitIfModified();
+	GetToolManager()->EndUndoTransaction();
+}
+
+
+FInputRayHit UHoleFillTool::BeginHoverSequenceHitTest(const FInputDeviceRay& PressPos)
+{
+	FHitResult OutHit;
+	if (SelectionMechanic->TopologyHitTest(PressPos.WorldRay, OutHit))
+	{
+		return FInputRayHit(OutHit.Distance);
+	}
+	return FInputRayHit();
+}
+
+bool UHoleFillTool::OnUpdateHover(const FInputDeviceRay& DevicePos) 
+{
+	SelectionMechanic->UpdateHighlight(DevicePos.WorldRay);
+	return true;
+}
+
+void UHoleFillTool::OnEndHover()
+{
+	SelectionMechanic->ClearHighlight();
+}
+
+void UHoleFillTool::RequestAction(EHoleFillToolActions ActionType)
+{
+	if (bHavePendingAction)
+	{
+		return;
+	}
+
+	PendingAction = ActionType;
+	bHavePendingAction = true;
+}
+
+void UHoleFillTool::SetWorld(UWorld* World)
+{
+	this->TargetWorld = World;
+}
+
+void UHoleFillTool::SetupPreview()
+{
+	UHoleFillOperatorFactory* OpFactory = NewObject<UHoleFillOperatorFactory>();
+	OpFactory->FillTool = this;
+
+	Preview = NewObject<UMeshOpPreviewWithBackgroundCompute>(OpFactory, "Preview");
+	Preview->Setup(this->TargetWorld, OpFactory);
+
+	FComponentMaterialSet MaterialSet;
+	ComponentTarget->GetMaterialSet(MaterialSet);
+	Preview->ConfigureMaterials(MaterialSet.Materials,
+		ToolSetupUtil::GetDefaultWorkingMaterial(GetToolManager())
+	);
+
+	// configure secondary render material
+	UMaterialInterface* SelectionMaterial = ToolSetupUtil::GetSelectionMaterial(FLinearColor(0.8f, 0.75f, 0.0f), GetToolManager());
+	if (SelectionMaterial != nullptr)
+	{
+		Preview->PreviewMesh->SetSecondaryRenderMaterial(SelectionMaterial);
+	}
+
+	// enable secondary triangle buffers
+	Preview->OnOpCompleted.AddLambda(
+		[this](const FDynamicMeshOperator* Op)
+	{
+		const FHoleFillOp* HoleFilleOp = (const FHoleFillOp*)(Op);
+		NewTriangleIDs = TSet<int32>(HoleFilleOp->NewTriangles);
+	});
+
+	Preview->PreviewMesh->EnableSecondaryTriangleBuffers(
+		[this](const FDynamicMesh3* Mesh, int32 TriangleID)
+	{
+		return NewTriangleIDs.Contains(TriangleID);
+	});
+
+	// set initial preview to un-processed mesh
+	Preview->PreviewMesh->SetTransform(ComponentTarget->GetWorldTransform());
+	Preview->PreviewMesh->UpdatePreview(&OriginalMesh);
+
+	Preview->SetVisibility(true);
+}
+
+
+void UHoleFillTool::ApplyAction(EHoleFillToolActions ActionType)
+{
+	switch (ActionType)
+	{
+	case EHoleFillToolActions::SelectAll:
+		SelectAll();
+		break;
+	case EHoleFillToolActions::ClearSelection:
+		ClearSelection();
+		break;
+	}
+}
+
+void UHoleFillTool::SelectAll()
+{
+	FGroupTopologySelection NewSelection;
+	for (int32 i = 0; i < Topology->Edges.Num(); ++i)
+	{
+		NewSelection.SelectedEdgeIDs.Add(i);
+	}
+
+	SelectionMechanic->SetSelection(NewSelection);
+	UpdateActiveBoundaryLoopSelection();
+	Preview->InvalidateResult();
+}
+
+
+void UHoleFillTool::ClearSelection()
+{
+	SelectionMechanic->ClearSelection();
+	UpdateActiveBoundaryLoopSelection();
+	Preview->InvalidateResult();
+}
+
+void UHoleFillTool::UpdateActiveBoundaryLoopSelection()
+{
+	ActiveBoundaryLoopSelection.Reset();
+
+	const FGroupTopologySelection& ActiveSelection = SelectionMechanic->GetActiveSelection();
+	int NumEdges = ActiveSelection.SelectedEdgeIDs.Num();
+	if (NumEdges == 0)
+	{
+		return;
+	}
+
+	ActiveBoundaryLoopSelection.Reserve(NumEdges);
+	for (int32 k = 0; k < NumEdges; ++k)
+	{
+		int32 EdgeID = ActiveSelection.SelectedEdgeIDs[k];
+		if (Topology->IsBoundaryEdge(EdgeID))
+		{
+			FSelectedBoundaryLoop& Loop = ActiveBoundaryLoopSelection.Emplace_GetRef();
+			Loop.EdgeTopoID = EdgeID;
+			Loop.EdgeIDs = Topology->GetGroupEdgeEdges(EdgeID);
+		}
+	}
+}
+
+
+void UHoleFillTool::Render(IToolsContextRenderAPI* RenderAPI)
+{
+	SelectionMechanic->Render(RenderAPI);
+}
+
+
+void UHoleFillTool::GetLoopsToFill(TArray<FEdgeLoop>& OutLoops) const
+{
+	OutLoops.Reset();
+
+	for (const FSelectedBoundaryLoop& FillEdge : ActiveBoundaryLoopSelection)
+	{
+		if (OriginalMesh.IsBoundaryEdge(FillEdge.EdgeIDs[0]))		// may no longer be boundary due to previous fill
+		{
+			FMeshBoundaryLoops BoundaryLoops(&OriginalMesh);
+			int32 LoopID = BoundaryLoops.FindLoopContainingEdge(FillEdge.EdgeIDs[0]);
+			if (LoopID >= 0)
+			{
+				OutLoops.Add(BoundaryLoops.Loops[LoopID]);
+			}
+		}
+	}
+}
+
+#undef LOCTEXT_NAMESPACE
diff --git a/Engine/Plugins/Experimental/MeshModelingToolset/Source/MeshModelingTools/Private/MeshInspectorTool.cpp b/Engine/Plugins/Experimental/MeshModelingToolset/Source/MeshModelingTools/Private/MeshInspectorTool.cpp
--- a/Engine/Plugins/Experimental/MeshModelingToolset/Source/MeshModelingTools/Private/MeshInspectorTool.cpp
+++ b/Engine/Plugins/Experimental/MeshModelingToolset/Source/MeshModelingTools/Private/MeshInspectorTool.cpp
@@ -14,8 +14,12 @@
 #include "Properties/MeshStatisticsProperties.h"
 #include "Properties/MeshAnalysisProperties.h"
 
+#include "Drawing/LineSetComponent.h"
+
 #include "SceneManagement.h" // for FPrimitiveDrawInterface
 
+#include "ToolSetupUtil.h"
+
 #define LOCTEXT_NAMESPACE "UMeshInspectorTool"
 
 
@@ -41,47 +45,9 @@ UInteractiveTool* UMeshInspectorToolBuilder::BuildTool(const FToolBuilderState&
 	return NewTool;
 }
 
-
-
-/*
- * Properties
- */
-
-void UMeshInspectorProperties::SaveProperties(UInteractiveTool* SaveFromTool)
-{
-	UMeshInspectorProperties* PropertyCache = GetPropertyCache<UMeshInspectorProperties>();
-	PropertyCache->bWireframe = this->bWireframe;
-	PropertyCache->bBoundaryEdges = this->bBoundaryEdges;
-	PropertyCache->bPolygonBorders = this->bPolygonBorders;
-	PropertyCache->bUVSeams = this->bUVSeams;
-	PropertyCache->bNormalSeams = this->bNormalSeams;
-	PropertyCache->bNormalVectors = this->bNormalVectors;
-	PropertyCache->bTangentVectors = this->bTangentVectors;
-	PropertyCache->NormalLength = this->NormalLength;
-	PropertyCache->TangentLength = this->TangentLength;
-}
-
-void UMeshInspectorProperties::RestoreProperties(UInteractiveTool* RestoreToTool)
-{
-	UMeshInspectorProperties* PropertyCache = GetPropertyCache<UMeshInspectorProperties>();
-	this->bWireframe = PropertyCache->bWireframe;
-	this->bBoundaryEdges = PropertyCache->bBoundaryEdges;
-	this->bPolygonBorders = PropertyCache->bPolygonBorders;
-	this->bUVSeams = PropertyCache->bUVSeams;
-	this->bNormalSeams = PropertyCache->bNormalSeams;
-	this->bNormalVectors = PropertyCache->bNormalVectors;
-	this->bTangentVectors = PropertyCache->bTangentVectors;
-	this->NormalLength = PropertyCache->NormalLength;
-	this->TangentLength = PropertyCache->TangentLength;
-}
-
-
-
-
 /*
  * Tool
  */
-
 UMeshInspectorTool::UMeshInspectorTool()
 {
 }
@@ -103,6 +69,16 @@ void UMeshInspectorTool::Setup()
 	PreviewMesh->SetTangentsMode(EDynamicMeshTangentCalcType::ExternallyCalculated);
 	PreviewMesh->InitializeMesh(ComponentTarget->GetMesh());
 
+	DrawnLineSet = NewObject<ULineSetComponent>(PreviewMesh->GetRootComponent(), "MeshInspectorToolLineSet");
+
+	DrawnLineSet->SetupAttachment(PreviewMesh->GetRootComponent());
+
+	// Note: at the moment, the material we use for the line set supports the use of a depth bias. If we ever need 
+	// extra speedup from the line set, we can disable depth bias support and set blend mode within the material to
+	// translucent, which gives a variable amount of speedup under heavy load.
+	DrawnLineSet->SetLineMaterial(ToolSetupUtil::GetDefaultLineComponentMaterial(GetToolManager()));
+	DrawnLineSet->RegisterComponent();
+
 	Precompute();
 
 	// hide input StaticMeshComponent
@@ -183,33 +159,66 @@ void UMeshInspectorTool::Precompute()
 
 void UMeshInspectorTool::Render(IToolsContextRenderAPI* RenderAPI)
 {
-	FPrimitiveDrawInterface* PDI = RenderAPI->GetPrimitiveDrawInterface();
-	FTransform3d Transform(ComponentTarget->GetWorldTransform());
+	
+}
+
+
+void UMeshInspectorTool::OnPropertyModified(UObject* PropertySet, FProperty* Property)
+{
+	GetToolManager()->PostInvalidation();
+	UpdateVisualization();
+}
+
+void UMeshInspectorTool::UpdateVisualization()
+{
+	PreviewMesh->EnableWireframe(Settings->bWireframe);
+
+	MaterialSettings->UpdateMaterials();
+	UMaterialInterface* OverrideMaterial = MaterialSettings->GetActiveOverrideMaterial();
+	if (OverrideMaterial == nullptr)
+	{
+		PreviewMesh->ClearOverrideRenderMaterial();
+	}
+	else
+	{
+		PreviewMesh->SetOverrideRenderMaterial(OverrideMaterial);
+	}
 
 	FColor BoundaryEdgeColor(240, 15, 15);
-	float BoundaryEdgeThickness = LineWidthMultiplier * 4.0;
+	float BoundaryEdgeThickness = LineWidthMultiplier * 2.0;
 	FColor UVSeamColor(15, 240, 15);
-	float UVSeamThickness = LineWidthMultiplier * 2.0;
+	float UVSeamThickness = LineWidthMultiplier * 1.0;
 	FColor NormalSeamColor(15, 240, 240);
-	float NormalSeamThickness = LineWidthMultiplier * 2.0;
+	float NormalSeamThickness = LineWidthMultiplier * 1.0;
 	FColor PolygonBorderColor(240, 15, 240);
-	float PolygonBorderThickness = LineWidthMultiplier * 2.0;
+	float PolygonBorderThickness = LineWidthMultiplier * 1.0;
 	FColor NormalColor(15, 15, 240);
-	float NormalThickness = LineWidthMultiplier * 2.0f;
+	float NormalThickness = LineWidthMultiplier * 1.0f;
 	FColor TangentColor(240, 15, 15);
 	FColor BinormalColor(15, 240, 15);
-	float TangentThickness = LineWidthMultiplier * 2.0f;
+	float TangentThickness = LineWidthMultiplier * 1.0f;
+
+	float BoundaryEdgeDepthBias = 2.0f;
+	float UVSeamDepthBias = 3.0f;
+	float NormalSeamDepthBias = 3.0f;
+	float PolygonBorderDepthBias = 2.0f;
+	float NormalDepthBias = 0.0f;
+	float TangentDepthBias = 3.5f;
+
+	// Used to scale normals and tangents appropriately
+	FVector ComponentScale = DrawnLineSet->GetComponentTransform().GetScale3D();
+	FVector3f InverseScale(1 / ComponentScale.X, 1 / ComponentScale.Y, 1 / ComponentScale.Z);
 
 	const FDynamicMesh3* TargetMesh = PreviewMesh->GetPreviewDynamicMesh();
 	FVector3d A, B;
 
+	DrawnLineSet->Clear();
 	if (Settings->bBoundaryEdges)
 	{
 		for (int eid : BoundaryEdges)
 		{
 			TargetMesh->GetEdgeV(eid, A, B);
-			PDI->DrawLine((FVector)Transform.TransformPosition(A), (FVector)Transform.TransformPosition(B),
-				BoundaryEdgeColor, 0, BoundaryEdgeThickness, 2.0f, true);
+			DrawnLineSet->AddLine((FVector)A, (FVector)B, BoundaryEdgeColor, BoundaryEdgeThickness, BoundaryEdgeDepthBias);
 		}
 	}
 
@@ -218,8 +227,7 @@ void UMeshInspectorTool::Render(IToolsContextRenderAPI* RenderAPI)
 		for (int eid : UVSeamEdges)
 		{
 			TargetMesh->GetEdgeV(eid, A, B);
-			PDI->DrawLine((FVector)Transform.TransformPosition(A), (FVector)Transform.TransformPosition(B),
-				UVSeamColor, 0, UVSeamThickness, 3.0f, true);
+			DrawnLineSet->AddLine((FVector)A, (FVector)B, UVSeamColor, UVSeamThickness, UVSeamDepthBias);
 		}
 	}
 
@@ -228,8 +236,7 @@ void UMeshInspectorTool::Render(IToolsContextRenderAPI* RenderAPI)
 		for (int eid : NormalSeamEdges)
 		{
 			TargetMesh->GetEdgeV(eid, A, B);
-			PDI->DrawLine((FVector)Transform.TransformPosition(A), (FVector)Transform.TransformPosition(B),
-				NormalSeamColor, 0, NormalSeamThickness, 3.0f, true);
+			DrawnLineSet->AddLine((FVector)A, (FVector)B, NormalSeamColor, NormalSeamThickness, NormalSeamDepthBias);
 		}
 	}
 
@@ -238,13 +245,18 @@ void UMeshInspectorTool::Render(IToolsContextRenderAPI* RenderAPI)
 		for (int eid : GroupBoundaryEdges)
 		{
 			TargetMesh->GetEdgeV(eid, A, B);
-			PDI->DrawLine((FVector)Transform.TransformPosition(A), (FVector)Transform.TransformPosition(B),
-				PolygonBorderColor, 0, PolygonBorderThickness, 2.0f, true);
+			DrawnLineSet->AddLine((FVector)A, (FVector)B, PolygonBorderColor, PolygonBorderThickness, PolygonBorderDepthBias);
 		}
 	}
 
 	if (Settings->bNormalVectors && TargetMesh->HasAttributes() && TargetMesh->Attributes()->PrimaryNormals() != nullptr)
 	{
+		// Note that for normals and tangent vectors, we want to allow the origin of the vector to undergo the full
+		// transform of the component, but the direction we travel along the vector has to be inversely scaled
+		// to end up with the original length in the end (and to reflect the fact that Unreal does not consider
+		// the scale transform when adjusting normals for an object).
+		FVector3f NormalScaling = Settings->NormalLength * InverseScale;
+
 		const FDynamicMeshNormalOverlay* NormalOverlay = TargetMesh->Attributes()->PrimaryNormals();
 		FVector3d TriV[3];
 		FVector3f TriN[3];
@@ -254,15 +266,17 @@ void UMeshInspectorTool::Render(IToolsContextRenderAPI* RenderAPI)
 			NormalOverlay->GetTriElements(tid, TriN[0], TriN[1], TriN[2]);
 			for (int j = 0; j < 3; ++j)
 			{
-				TriV[j] = Transform.TransformPosition(TriV[j]);
-				PDI->DrawLine((FVector)TriV[j], (FVector)((FVector3f)TriV[j] + Settings->NormalLength * Transform.TransformVectorNoScale(TriN[j])),
-					NormalColor, SDPG_World, NormalThickness, 0.0f, true);
+				DrawnLineSet->AddLine((FVector)TriV[j], (FVector)((FVector3f)TriV[j] + NormalScaling * TriN[j]),
+					NormalColor, NormalThickness, NormalDepthBias);
 			}
 		}
 	}
 
 	if (Settings->bTangentVectors && PreviewMesh->GetTangents() != nullptr)
 	{
+		// See note in normals about scaling
+		FVector3f TangentScaling = Settings->TangentLength * InverseScale;
+
 		const FMeshTangentsf* Tangents = PreviewMesh->GetTangents();
 		for (int TID : TargetMesh->TriangleIndicesItr())
 		{
@@ -270,41 +284,16 @@ void UMeshInspectorTool::Render(IToolsContextRenderAPI* RenderAPI)
 			TargetMesh->GetTriVertices(TID, TriV[0], TriV[1], TriV[2]);
 			for (int SubIdx = 0; SubIdx < 3; SubIdx++)
 			{
-				FVector3f Vert(Transform.TransformPosition(TriV[SubIdx]));
+				FVector3f Vert(TriV[SubIdx]);
 				FVector3f Tangent, Bitangent;
 				Tangents->GetPerTriangleTangent(TID, SubIdx, Tangent, Bitangent);
-				PDI->DrawLine((FVector)Vert, (FVector)(Vert + Settings->TangentLength * Transform.TransformVectorNoScale(Tangent)),
-					TangentColor, SDPG_World, TangentThickness, 3.5f, true);
-				PDI->DrawLine((FVector)Vert, (FVector)(Vert + Settings->TangentLength * Transform.TransformVectorNoScale(Bitangent)),
-					BinormalColor, SDPG_World, TangentThickness, 3.5f, true);
+				DrawnLineSet->AddLine((FVector)Vert, (FVector)(Vert + TangentScaling * Tangent),
+					TangentColor, TangentThickness, TangentDepthBias);
+				DrawnLineSet->AddLine((FVector)Vert, (FVector)(Vert + TangentScaling * Bitangent),
+					BinormalColor, TangentThickness, TangentDepthBias);
 			}
 		}
 	}
-
-
-}
-
-
-void UMeshInspectorTool::OnPropertyModified(UObject* PropertySet, FProperty* Property)
-{
-	GetToolManager()->PostInvalidation();
-	UpdateVisualization();
-}
-
-void UMeshInspectorTool::UpdateVisualization()
-{
-	PreviewMesh->EnableWireframe(Settings->bWireframe);
-
-	MaterialSettings->UpdateMaterials();
-	UMaterialInterface* OverrideMaterial = MaterialSettings->GetActiveOverrideMaterial();
-	if (OverrideMaterial == nullptr)
-	{
-		PreviewMesh->ClearOverrideRenderMaterial();
-	}
-	else
-	{
-		PreviewMesh->SetOverrideRenderMaterial(OverrideMaterial);
-	}
 }
 
 
diff --git a/Engine/Plugins/Experimental/MeshModelingToolset/Source/MeshModelingTools/Private/MeshSelectionTool.cpp b/Engine/Plugins/Experimental/MeshModelingToolset/Source/MeshModelingTools/Private/MeshSelectionTool.cpp
--- a/Engine/Plugins/Experimental/MeshModelingToolset/Source/MeshModelingTools/Private/MeshSelectionTool.cpp
+++ b/Engine/Plugins/Experimental/MeshModelingToolset/Source/MeshModelingTools/Private/MeshSelectionTool.cpp
@@ -24,7 +24,6 @@
 /*
  * ToolBuilder
  */
-
 UMeshSurfacePointTool* UMeshSelectionToolBuilder::CreateNewTool(const FToolBuilderState& SceneState) const
 {
 	UMeshSelectionTool* SelectionTool = NewObject<UMeshSelectionTool>(SceneState.ToolManager);
@@ -33,9 +32,9 @@ UMeshSurfacePointTool* UMeshSelectionToolBuilder::CreateNewTool(const FToolBuild
 	return SelectionTool;
 }
 
-
-
-
+/*
+ * Properties
+ */
 void UMeshSelectionToolActionPropertySet::PostAction(EMeshSelectionToolActions Action)
 {
 	if (ParentTool.IsValid())
@@ -44,33 +43,9 @@ void UMeshSelectionToolActionPropertySet::PostAction(EMeshSelectionToolActions A
 	}
 }
 
-
-void UMeshSelectionToolProperties::SaveProperties(UInteractiveTool* SaveFromTool)
-{
-	UMeshSelectionToolProperties* PropertyCache = GetPropertyCache<UMeshSelectionToolProperties>();
-	PropertyCache->SelectionMode = this->SelectionMode;
-	PropertyCache->AngleTolerance = this->AngleTolerance;
-	PropertyCache->bHitBackFaces = this->bHitBackFaces;
-	PropertyCache->bShowWireframe = this->bShowWireframe;
-	PropertyCache->FaceColorMode = this->FaceColorMode;
-}
-
-void UMeshSelectionToolProperties::RestoreProperties(UInteractiveTool* RestoreToTool)
-{
-	UMeshSelectionToolProperties* PropertyCache = GetPropertyCache<UMeshSelectionToolProperties>();
-	this->SelectionMode = PropertyCache->SelectionMode;
-	this->AngleTolerance = PropertyCache->AngleTolerance;
-	this->bHitBackFaces = PropertyCache->bHitBackFaces;
-	this->bShowWireframe = PropertyCache->bShowWireframe;
-	this->FaceColorMode = PropertyCache->FaceColorMode;
-}
-
-
-
 /*
  * Tool
  */
-
 UMeshSelectionTool::UMeshSelectionTool()
 {
 }
@@ -147,13 +122,13 @@ void UMeshSelectionTool::Setup()
 	// rebuild octree if mesh changes
 	PreviewMesh->GetOnMeshChanged().AddLambda([this]() { bOctreeValid = false; bFullMeshInvalidationPending = true; });
 
-	ShowWireframeWatcher.Initialize(
-		[this]() { return SelectionProps->bShowWireframe; },
-		[this](bool bNewValue) { PreviewMesh->EnableWireframe(bNewValue); }, SelectionProps->bShowWireframe);
-
-	ColorModeWatcher.Initialize(
-		[this]() { return SelectionProps->FaceColorMode; },
-		[this](EMeshFacesColorMode NewValue) { bColorsUpdatePending = true; UpdateVisualization(false); }, SelectionProps->FaceColorMode );
+	SelectionProps->WatchProperty(SelectionProps->bShowWireframe,
+								  [this](bool bNewValue) { PreviewMesh->EnableWireframe(bNewValue); });
+	SelectionProps->WatchProperty(SelectionProps->FaceColorMode,
+								  [this](EMeshFacesColorMode NewValue)
+								  {
+									  bColorsUpdatePending = true; UpdateVisualization(false);
+								  });
 	bColorsUpdatePending = (SelectionProps->FaceColorMode != EMeshFacesColorMode::None);
 
 	RecalculateBrushRadius();
@@ -451,8 +426,6 @@ static void UpdateList(TArray<int>& List, int Value, bool bAdd)
 
 void UMeshSelectionTool::ApplyStamp(const FBrushStampData& Stamp)
 {
-	//const FDynamicMesh3* Mesh = PreviewMesh->GetPreviewDynamicMesh();
-
 	IndexBuf.Reset();
 
 	bool bDesiredValue = bInRemoveStroke ? false : true;
@@ -589,7 +562,7 @@ void UMeshSelectionTool::OnEndDrag(const FRay& Ray)
 	bStampPending = false;
 
 	// close change record
-	TUniquePtr<FMeshSelectionChange> Change = EndChange();
+	TUniquePtr<FToolCommandChange> Change = EndChange();
 	GetToolManager()->EmitObjectChange(Selection, MoveTemp(Change), LOCTEXT("MeshSelectionChange", "Mesh Selection"));
 }
 
@@ -635,7 +608,7 @@ void UMeshSelectionTool::UpdateVisualization(bool bSelectionModified)
 	// force an update of renderbuffers
 	if (bSelectionModified)
 	{
-		PreviewMesh->NotifyDeferredEditCompleted(UPreviewMesh::ERenderUpdateMode::FullUpdate, true);
+		PreviewMesh->NotifyDeferredEditCompleted(UPreviewMesh::ERenderUpdateMode::FullUpdate, EMeshRenderAttributeFlags::All, true);
 	}
 
 	if (bColorsUpdatePending)
@@ -730,14 +703,8 @@ void UMeshSelectionTool::Render(IToolsContextRenderAPI* RenderAPI)
 }
 
 
-void UMeshSelectionTool::Tick(float DeltaTime)
-
+void UMeshSelectionTool::OnTick(float DeltaTime)
 {
-	UDynamicMeshBrushTool::Tick(DeltaTime);
-
-	ShowWireframeWatcher.CheckAndUpdate();
-	ColorModeWatcher.CheckAndUpdate();
-
 	if (bStampPending)
 	{
 		ApplyStamp(LastStamp);
@@ -769,7 +736,7 @@ void UMeshSelectionTool::CancelChange()
 	}
 }
 
-TUniquePtr<FMeshSelectionChange> UMeshSelectionTool::EndChange()
+TUniquePtr<FToolCommandChange> UMeshSelectionTool::EndChange()
 {
 	check(ActiveSelectionChange);
 	if (ActiveSelectionChange != nullptr)
@@ -882,7 +849,7 @@ void UMeshSelectionTool::SelectAll()
 	ActiveSelectionChange->Add(AddFaces);
 	Selection->AddIndices(EMeshSelectionElementType::Face, AddFaces);
 	
-	TUniquePtr<FMeshSelectionChange> SelectionChange = EndChange();
+	TUniquePtr<FToolCommandChange> SelectionChange = EndChange();
 
 	GetToolManager()->EmitObjectChange(Selection, MoveTemp(SelectionChange), LOCTEXT("SelectAll", "Select All"));
 
@@ -903,7 +870,7 @@ void UMeshSelectionTool::ClearSelection()
 	ActiveSelectionChange->Add(SelectedFaces);
 	Selection->RemoveIndices(EMeshSelectionElementType::Face, SelectedFaces);
 
-	TUniquePtr<FMeshSelectionChange> SelectionChange = EndChange();
+	TUniquePtr<FToolCommandChange> SelectionChange = EndChange();
 
 	GetToolManager()->EmitObjectChange(Selection, MoveTemp(SelectionChange), LOCTEXT("ClearSelection", "Clear Selection"));
 
@@ -938,7 +905,7 @@ void UMeshSelectionTool::InvertSelection()
 	BeginChange(false);
 	ActiveSelectionChange->Add(SelectedFaces);
 	Selection->RemoveIndices(EMeshSelectionElementType::Face, SelectedFaces);
-	TUniquePtr<FMeshSelectionChange> ClearChange = EndChange();
+	TUniquePtr<FToolCommandChange> ClearChange = EndChange();
 
 	GetToolManager()->EmitObjectChange(Selection, MoveTemp(ClearChange), LOCTEXT("InvertSelection", "Invert Selection"));
 
@@ -946,7 +913,7 @@ void UMeshSelectionTool::InvertSelection()
 	BeginChange(true);
 	ActiveSelectionChange->Add(InvertedFaces);
 	Selection->AddIndices(EMeshSelectionElementType::Face, InvertedFaces);
-	TUniquePtr<FMeshSelectionChange> AddChange = EndChange();
+	TUniquePtr<FToolCommandChange> AddChange = EndChange();
 
 	GetToolManager()->EmitObjectChange(Selection, MoveTemp(AddChange), LOCTEXT("InvertSelection", "Invert Selection"));
 
@@ -1022,13 +989,13 @@ void UMeshSelectionTool::GrowShrinkSelection(bool bGrow)
 	if (bGrow)
 	{
 		Selection->AddIndices(EMeshSelectionElementType::Face, ChangeFaces);
-		TUniquePtr<FMeshSelectionChange> SelectionChange = EndChange();
+		TUniquePtr<FToolCommandChange> SelectionChange = EndChange();
 		GetToolManager()->EmitObjectChange(Selection, MoveTemp(SelectionChange), LOCTEXT("GrowSelection", "Grow Selection"));
 	}
 	else
 	{
 		Selection->RemoveIndices(EMeshSelectionElementType::Face, ChangeFaces);
-		TUniquePtr<FMeshSelectionChange> SelectionChange = EndChange();
+		TUniquePtr<FToolCommandChange> SelectionChange = EndChange();
 		GetToolManager()->EmitObjectChange(Selection, MoveTemp(SelectionChange), LOCTEXT("ShrinkSelection", "Shrink Selection"));
 	}
 	OnExternalSelectionChange();
@@ -1075,7 +1042,7 @@ void UMeshSelectionTool::ExpandToConnected()
 	BeginChange(true);
 	ActiveSelectionChange->Add(AddFaces);
 	Selection->AddIndices(EMeshSelectionElementType::Face, AddFaces);
-	TUniquePtr<FMeshSelectionChange> SelectionChange = EndChange();
+	TUniquePtr<FToolCommandChange> SelectionChange = EndChange();
 	GetToolManager()->EmitObjectChange(Selection, MoveTemp(SelectionChange), LOCTEXT("ExpandToConnected", "Expand Selection"));
 	OnExternalSelectionChange();
 }
@@ -1131,7 +1098,7 @@ void UMeshSelectionTool::SelectLargestComponent(bool bWeightByArea)
 		}
 	}
 	
-	TUniquePtr<FMeshSelectionChange> SelectionChange = EndChange();
+	TUniquePtr<FToolCommandChange> SelectionChange = EndChange();
 
 	GetToolManager()->EmitObjectChange(Selection, MoveTemp(SelectionChange), LOCTEXT("SelectLargestComponentByArea", "Select Largest Component By Area"));
 
@@ -1170,7 +1137,7 @@ void UMeshSelectionTool::OptimizeSelection()
 	}
 	Selection->NotifySelectionSetModified();
 
-	TUniquePtr<FMeshSelectionChange> DeselectChange = EndChange();
+	TUniquePtr<FToolCommandChange> DeselectChange = EndChange();
 
 	GetToolManager()->EmitObjectChange(Selection, MoveTemp(DeselectChange), LOCTEXT("OptimizeSelection", "Optimize Selection"));
 
@@ -1190,7 +1157,7 @@ void UMeshSelectionTool::OptimizeSelection()
 
 	check(Selection->Faces.Num() == FaceSelection.Num());
 	
-	TUniquePtr<FMeshSelectionChange> AddChange = EndChange();
+	TUniquePtr<FToolCommandChange> AddChange = EndChange();
 
 	GetToolManager()->EmitObjectChange(Selection, MoveTemp(AddChange), LOCTEXT("OptimizeSelection", "Optimize Selection"));
 
@@ -1218,7 +1185,7 @@ void UMeshSelectionTool::DeleteSelectedTriangles()
 		ActiveSelectionChange->Add(tid);
 	}
 	Selection->RemoveIndices(EMeshSelectionElementType::Face, SelectedFaces);
-	TUniquePtr<FMeshSelectionChange> SelectionChange = EndChange();
+	TUniquePtr<FToolCommandChange> SelectionChange = EndChange();
 	ChangeSeq->AppendChange(Selection, MoveTemp(SelectionChange));
 
 	// delete triangles and emit delete triangles change
@@ -1342,7 +1309,7 @@ void UMeshSelectionTool::FlipSelectedTriangles()
 		ActiveSelectionChange->Add(tid);
 	}
 	Selection->RemoveIndices(EMeshSelectionElementType::Face, SelectedFaces);
-	TUniquePtr<FMeshSelectionChange> SelectionChange = EndChange();
+	TUniquePtr<FToolCommandChange> SelectionChange = EndChange();
 	ChangeSeq->AppendChange(Selection, MoveTemp(SelectionChange));
 
 	// flip normals
@@ -1383,7 +1350,7 @@ void UMeshSelectionTool::AssignNewGroupToSelectedTriangles()
 		ActiveSelectionChange->Add(tid);
 	}
 	Selection->RemoveIndices(EMeshSelectionElementType::Face, SelectedFaces);
-	TUniquePtr<FMeshSelectionChange> SelectionChange = EndChange();
+	TUniquePtr<FToolCommandChange> SelectionChange = EndChange();
 	ChangeSeq->AppendChange(Selection, MoveTemp(SelectionChange));
 
 	// assign new groups to triangles
diff --git a/Engine/Plugins/Experimental/MeshModelingToolset/Source/MeshModelingTools/Private/MeshVertexSculptTool.cpp b/Engine/Plugins/Experimental/MeshModelingToolset/Source/MeshModelingTools/Private/MeshVertexSculptTool.cpp
new file mode 100644
--- /dev/null
+++ b/Engine/Plugins/Experimental/MeshModelingToolset/Source/MeshModelingTools/Private/MeshVertexSculptTool.cpp
@@ -0,0 +1,1115 @@
+// Copyright Epic Games, Inc. All Rights Reserved.
+
+#include "MeshVertexSculptTool.h"
+#include "Containers/Map.h"
+#include "InteractiveToolManager.h"
+#include "InteractiveGizmoManager.h"
+#include "ToolBuilderUtil.h"
+#include "Async/ParallelFor.h"
+#include "Async/Async.h"
+
+#include "ProjectionTargets.h"
+#include "MeshWeights.h"
+#include "MeshNormals.h"
+#include "MeshIndexUtil.h"
+#include "PreviewMesh.h"
+#include "ToolSetupUtil.h"
+#include "Util/BufferUtil.h"
+
+#include "Changes/MeshVertexChange.h"
+
+#include "ToolDataVisualizer.h"
+#include "Components/PrimitiveComponent.h"
+
+//#include "Sculpting/KelvinletBrushOp.h"
+#include "Sculpting/MeshSmoothingBrushOps.h"
+#include "Sculpting/MeshInflateBrushOps.h"
+#include "Sculpting/MeshMoveBrushOps.h"
+#include "Sculpting/MeshPlaneBrushOps.h"
+#include "Sculpting/MeshPinchBrushOps.h"
+#include "Sculpting/MeshSculptBrushOps.h"
+#include "Sculpting/MeshSculptUtil.h"
+
+#include "UObject/ObjectMacros.h"
+#include "Materials/Material.h"
+#include "Materials/MaterialInstanceDynamic.h"
+
+#define LOCTEXT_NAMESPACE "UMeshVertexSculptTool"
+
+namespace
+{
+	// probably should be something defined for the whole tool framework...
+#if WITH_EDITOR
+	static EAsyncExecution VertexSculptToolAsyncExecTarget = EAsyncExecution::LargeThreadPool;
+#else
+	static EAsyncExecution VertexSculptToolAsyncExecTarget = EAsyncExecution::ThreadPool;
+#endif
+}
+
+
+/*
+ * ToolBuilder
+ */
+UMeshSurfacePointTool* UMeshVertexSculptToolBuilder::CreateNewTool(const FToolBuilderState& SceneState) const
+{
+	UMeshVertexSculptTool* SculptTool = NewObject<UMeshVertexSculptTool>(SceneState.ToolManager);
+	SculptTool->SetWorld(SceneState.World);
+	return SculptTool;
+}
+
+/*
+ * Tool
+ */
+
+void UMeshVertexSculptTool::Setup()
+{
+	UMeshSculptToolBase::Setup();
+
+	// create dynamic mesh component to use for live preview
+	DynamicMeshComponent = NewObject<USimpleDynamicMeshComponent>(ComponentTarget->GetOwnerActor());
+	DynamicMeshComponent->SetupAttachment(ComponentTarget->GetOwnerActor()->GetRootComponent());
+	DynamicMeshComponent->RegisterComponent();
+
+	// initialize from LOD-0 MeshDescription
+	DynamicMeshComponent->InitializeMesh(ComponentTarget->GetMesh());
+	double MaxDimension = DynamicMeshComponent->GetMesh()->GetCachedBounds().MaxDim();
+
+	// transform mesh to world space because handling scaling inside brush is a mess
+	// Note: this transform does not include translation ( so only the 3x3 transform)
+	InitialTargetTransform = FTransform3d(ComponentTarget->GetWorldTransform());
+	// clamp scaling because if we allow zero-scale we cannot invert this transform on Accept
+	InitialTargetTransform.ClampMinimumScale(0.01);
+	FVector3d Translation = InitialTargetTransform.GetTranslation();
+	InitialTargetTransform.SetTranslation(FVector3d::Zero());
+	DynamicMeshComponent->ApplyTransform(InitialTargetTransform, false);
+	// since we moved to World coords there is not a current transform anymore.
+	CurTargetTransform = FTransform3d(Translation);
+	DynamicMeshComponent->SetWorldTransform((FTransform)CurTargetTransform);
+
+	// assign materials
+	FComponentMaterialSet MaterialSet;
+	ComponentTarget->GetMaterialSet(MaterialSet);
+	for (int k = 0; k < MaterialSet.Materials.Num(); ++k)
+	{
+		DynamicMeshComponent->SetMaterial(k, MaterialSet.Materials[k]);
+	}
+
+	DynamicMeshComponent->bInvalidateProxyOnChange = false;
+	OnDynamicMeshComponentChangedHandle = DynamicMeshComponent->OnMeshVerticesChanged.AddUObject(this, &UMeshVertexSculptTool::OnDynamicMeshComponentChanged);
+
+	FDynamicMesh3* Mesh = DynamicMeshComponent->GetMesh();
+	// do we always want to keep vertex normals updated? Perhaps should discard vertex normals before baking?
+	//FMeshNormals::QuickComputeVertexNormals(*Mesh);
+
+	// initialize dynamic octree
+	Octree.RootDimension = MaxDimension / 10.0;
+	Octree.SetMaxTreeDepth(3);
+	Octree.Initialize(Mesh);
+	FDynamicMeshOctree3::FStatistics Stats;
+	Octree.ComputeStatistics(Stats);
+	//UE_LOG(LogTemp, Warning, TEXT("Octree Stats: %s"), *Stats.ToString());
+
+	// initialize decomposition
+	TUniquePtr<FMeshRenderDecomposition> Decomp = MakeUnique<FMeshRenderDecomposition>();
+	FMeshRenderDecomposition::BuildChunkedDecomposition(Mesh, &MaterialSet, *Decomp);
+	Decomp->BuildAssociations(Mesh);
+	//UE_LOG(LogTemp, Warning, TEXT("Decomposition has %d groups"), Decomp->Num());
+	DynamicMeshComponent->SetExternalDecomposition(MoveTemp(Decomp));
+
+	// initialize target mesh
+	UpdateBaseMesh(nullptr);
+	bTargetDirty = false;
+
+	// initialize brush radius range interval, brush properties
+	UMeshSculptToolBase::InitializeBrushSizeRange(DynamicMeshComponent->GetMesh()->GetCachedBounds());
+
+	// initialize other properties
+	SculptProperties = NewObject<UVertexBrushSculptProperties>(this);
+	//KelvinBrushProperties = NewObject<UKelvinBrushProperties>(this);
+
+	// hide input Component
+	ComponentTarget->SetOwnerVisibility(false);
+
+	// init state flags flags
+	bInDrag = false;
+	bStampPending = false;
+	ActiveVertexChange = nullptr;
+
+	InitializeIndicator();
+
+	// initialize our properties
+	AddToolPropertySource(UMeshSculptToolBase::BrushProperties);
+	AddToolPropertySource(SculptProperties);
+
+	// add brush-specific properties 
+	PlaneBrushProperties = NewObject<UPlaneBrushProperties>(this);
+	PlaneBrushProperties->RestoreProperties(this);
+	AddToolPropertySource(PlaneBrushProperties);
+	SetToolPropertySourceEnabled(PlaneBrushProperties, false);
+
+	SculptMaxBrushProperties = NewObject<USculptMaxBrushProperties>();
+	SculptMaxBrushProperties->RestoreProperties(this);
+	AddToolPropertySource(SculptMaxBrushProperties);
+
+	//AddToolPropertySource(KelvinBrushProperties);
+	//KelvinBrushProperties->RestoreProperties(this);
+
+	AddToolPropertySource(UMeshSculptToolBase::GizmoProperties);
+	SetToolPropertySourceEnabled(UMeshSculptToolBase::GizmoProperties, false);
+
+	CalculateBrushRadius();
+	SculptProperties->RestoreProperties(this);
+
+	// disable tool-specific properties
+	SetToolPropertySourceEnabled(SculptMaxBrushProperties, false);
+	//SetToolPropertySourceEnabled(KelvinBrushProperties, false);
+
+	AddToolPropertySource(UMeshSculptToolBase::ViewProperties);
+
+	// register watchers
+	BrushTypeWatcher.Initialize(
+		[this]() { return SculptProperties->PrimaryBrushType; },
+		[this](EMeshVertexSculptBrushType NewBrushType) { UpdateBrushType(NewBrushType); }, SculptProperties->PrimaryBrushType);
+
+	ActiveFalloff = MakeShared<FMeshSculptFallofFunc>();
+	ActiveFalloff->FalloffFunc = [&](const FSculptBrushStamp& StampInfo, const FVector3d& Position)
+	{
+		return CalculateBrushFalloff(Position.Distance(StampInfo.LocalFrame.Origin));
+	};
+
+	SecondaryBrushOp = MakeUnique<FSmoothBrushOp>();
+	SecondaryBrushOp->Falloff = ActiveFalloff;
+	PrimaryBrushOp = MakeUnique<FSmoothBrushOp>();  // set default
+	PrimaryBrushOp->Falloff = ActiveFalloff;
+
+	UpdateBrushType(SculptProperties->PrimaryBrushType);
+}
+
+void UMeshVertexSculptTool::Shutdown(EToolShutdownType ShutdownType)
+{
+	UMeshSculptToolBase::Shutdown(ShutdownType);
+
+	if (DynamicMeshComponent != nullptr)
+	{
+		DynamicMeshComponent->OnMeshChanged.Remove(OnDynamicMeshComponentChangedHandle);
+
+		ComponentTarget->SetOwnerVisibility(true);
+
+		if (ShutdownType == EToolShutdownType::Accept)
+		{
+			// safe to do this here because we are about to destroy componeont
+			DynamicMeshComponent->ApplyTransform(InitialTargetTransform, true);
+
+			// this block bakes the modified DynamicMeshComponent back into the StaticMeshComponent inside an undo transaction
+			GetToolManager()->BeginUndoTransaction(LOCTEXT("SculptMeshToolTransactionName", "Sculpt Mesh"));
+			ComponentTarget->CommitMesh([=](const FPrimitiveComponentTarget::FCommitParams& CommitParams)
+			{
+				DynamicMeshComponent->Bake(CommitParams.MeshDescription, false);
+			});
+			GetToolManager()->EndUndoTransaction();
+		}
+
+		DynamicMeshComponent->UnregisterComponent();
+		DynamicMeshComponent->DestroyComponent();
+		DynamicMeshComponent = nullptr;
+	}
+
+	SculptProperties->SaveProperties(this);
+	//KelvinBrushProperties->SaveProperties(this);
+	PlaneBrushProperties->SaveProperties(this);
+	SculptMaxBrushProperties->SaveProperties(this);
+}
+
+
+void UMeshVertexSculptTool::OnPropertyModified(UObject* PropertySet, FProperty* Property)
+{
+	CalculateBrushRadius();
+}
+
+bool UMeshVertexSculptTool::HitTest(const FRay& Ray, FHitResult& OutHit)
+{
+	FRay3d LocalRay(CurTargetTransform.InverseTransformPosition(Ray.Origin),
+		CurTargetTransform.InverseTransformVector(Ray.Direction));
+	LocalRay.Direction.Normalize();
+	FDynamicMesh3* Mesh = DynamicMeshComponent->GetMesh();
+
+	int HitTID = FindHitSculptMeshTriangle(LocalRay);
+	if (HitTID != IndexConstants::InvalidID)
+	{
+		FTriangle3d Triangle;
+		Mesh->GetTriVertices(HitTID, Triangle.V[0], Triangle.V[1], Triangle.V[2]);
+		FIntrRay3Triangle3d Query(LocalRay, Triangle);
+		Query.Find();
+
+		OutHit.FaceIndex = HitTID;
+		OutHit.Distance = Query.RayParameter;
+		OutHit.Normal = (FVector)CurTargetTransform.TransformNormal(Mesh->GetTriNormal(HitTID));
+		OutHit.ImpactPoint = (FVector)CurTargetTransform.TransformPosition(LocalRay.PointAt(Query.RayParameter));
+		return true;
+	}
+	return false;
+}
+
+void UMeshVertexSculptTool::OnBeginDrag(const FRay& Ray)
+{
+	bSmoothing = GetShiftToggle();
+	bInvert = GetCtrlToggle();
+
+	FHitResult OutHit;
+	if (HitTest(Ray, OutHit))
+	{
+		FVector3d BrushStartCenterWorld = Ray.PointAt(OutHit.Distance) + BrushProperties->Depth*CurrentBrushRadius*Ray.Direction;
+
+		bInDrag = true;
+
+		ActiveDragPlane = FFrame3d(BrushStartCenterWorld, -Ray.Direction);
+		ActiveDragPlane.RayPlaneIntersection(Ray.Origin, Ray.Direction, 2, LastHitPosWorld);
+
+		LastBrushPosWorld = LastHitPosWorld;
+		LastBrushPosNormalWorld = ActiveDragPlane.Z();
+		LastBrushPosLocal = CurTargetTransform.InverseTransformPosition(LastHitPosWorld);
+
+		BeginChange();
+
+		UpdateROI(LastBrushPosLocal);
+
+		if (SculptProperties->PrimaryBrushType == EMeshVertexSculptBrushType::Plane)
+		{
+			ActiveFixedBrushPlane = ComputeROIBrushPlane(LastBrushPosLocal, false, false);
+		}
+		else if (SculptProperties->PrimaryBrushType == EMeshVertexSculptBrushType::PlaneViewAligned)
+		{
+			AlignBrushToView();
+			ActiveFixedBrushPlane = ComputeROIBrushPlane(LastBrushPosLocal, false, true);
+		}
+		else if (SculptProperties->PrimaryBrushType == EMeshVertexSculptBrushType::FixedPlane)
+		{
+			ActiveFixedBrushPlane = FFrame3d(
+				CurTargetTransform.InverseTransformPosition(GizmoProperties->Position),
+				CurTargetTransform.GetRotation().Inverse() * (FQuaterniond)GizmoProperties->Rotation);
+		}
+
+		// init last stamp
+		LastStamp.WorldFrame = FFrame3d(LastBrushPosWorld, LastBrushPosNormalWorld);
+		LastStamp.LocalFrame = LastStamp.WorldFrame;
+		LastStamp.LocalFrame.Transform(CurTargetTransform.Inverse());
+		LastStamp.Radius = CurrentBrushRadius;
+		LastStamp.Direction = (bInvert) ? -1.0 : 1.0;
+		LastStamp.Depth = BrushProperties->Depth;
+		LastStamp.Power = GetActivePressure() * ((bSmoothing) ? SculptProperties->SmoothBrushSpeed : SculptProperties->PrimaryBrushSpeed);
+		TUniquePtr<FMeshSculptBrushOp>& UseBrushOp = (bSmoothing) ? SecondaryBrushOp : PrimaryBrushOp;
+
+		FSculptBrushOptions SculptOptions;
+		SculptOptions.bPreserveUVFlow = SculptProperties->bPreserveUVFlow;
+		SculptOptions.MaxHeight = CurrentBrushRadius * SculptMaxBrushProperties->MaxHeight;
+		if (SculptMaxBrushProperties->bFreezeCurrentHeight && SculptMaxFixedHeight >= 0)
+		{
+			SculptOptions.MaxHeight = SculptMaxFixedHeight;
+		}
+		SculptMaxFixedHeight = SculptOptions.MaxHeight;		// yech
+		SculptOptions.WhichPlaneSideIndex = (int32)PlaneBrushProperties->WhichSide;
+		SculptOptions.ConstantReferencePlane = ActiveFixedBrushPlane;
+		UseBrushOp->ConfigureOptions(SculptOptions);
+
+		UseBrushOp->BeginStroke(DynamicMeshComponent->GetMesh(), LastStamp, VertexROI);
+
+		AccumulatedTriangleROI.Reset();
+
+		// apply initial stamp
+		PendingStampRay = Ray;
+		bStampPending = true;
+	}
+}
+
+
+
+void UMeshVertexSculptTool::UpdateROI(const FVector3d& BrushPos)
+{
+	SCOPE_CYCLE_COUNTER(VtxSculptTool_UpdateROI);
+
+	// TODO: need dynamic vertex hash table?
+
+	float RadiusSqr = CurrentBrushRadius * CurrentBrushRadius;
+
+	FAxisAlignedBox3d BrushBox(
+		BrushPos - CurrentBrushRadius * FVector3d::One(),
+		BrushPos + CurrentBrushRadius * FVector3d::One());
+
+	VertexSetBuffer.Reset();
+	TriangleROI.Reset();
+	FDynamicMesh3* Mesh = DynamicMeshComponent->GetMesh();
+	Octree.RangeQuery(BrushBox,
+		[&](int TriIdx) {
+		FIndex3i TriV = Mesh->GetTriangle(TriIdx);
+		int Count = 0;
+		for (int j = 0; j < 3; ++j)
+		{
+			if (VertexSetBuffer.Contains(TriV[j]))
+			{
+				Count++;
+			}
+			else if (BrushPos.DistanceSquared(Mesh->GetVertex(TriV[j])) < RadiusSqr)
+			{
+				VertexSetBuffer.Add(TriV[j]);
+				Count++;
+			}
+		}
+		if (Count > 0)
+		{
+			TriangleROI.Add(TriIdx);
+		}
+	});
+
+	VertexROI.SetNum(0, false);
+	BufferUtil::AppendElements(VertexROI, VertexSetBuffer);
+
+	int NumV = VertexROI.Num();
+	ROIPositionBuffer.SetNum(NumV, false);
+}
+
+void UMeshVertexSculptTool::OnUpdateDrag(const FRay& WorldRay)
+{
+	if (bInDrag)
+	{
+		PendingStampRay = WorldRay;
+		bStampPending = true;
+	}
+}
+
+
+bool UMeshVertexSculptTool::UpdateStampPosition(const FRay& WorldRay)
+{
+	StampTimestamp++;
+
+	CalculateBrushRadius();
+
+	TUniquePtr<FMeshSculptBrushOp>& UseBrushOp = (bSmoothing) ? SecondaryBrushOp : PrimaryBrushOp;
+	UseBrushOp->Falloff = ActiveFalloff;
+
+	ESculptBrushOpTargetType TargetType = UseBrushOp->GetBrushTargetType();
+	switch (TargetType)
+	{
+	case ESculptBrushOpTargetType::SculptMesh:
+		UpdateBrushPositionOnSculptMesh(WorldRay, true);
+		break;
+	case ESculptBrushOpTargetType::TargetMesh:
+		UpdateBrushPositionOnTargetMesh(WorldRay, true);
+		break;
+	case ESculptBrushOpTargetType::ActivePlane:
+		UpdateBrushPositionOnActivePlane(WorldRay);
+		break;
+	}
+
+	if (UseBrushOp->GetAlignStampToView())
+	{
+		AlignBrushToView();
+	}
+
+	CurrentStamp = LastStamp;
+	CurrentStamp.WorldFrame = FFrame3d(LastBrushPosWorld, LastBrushPosNormalWorld);
+	CurrentStamp.LocalFrame = CurrentStamp.WorldFrame;
+	CurrentStamp.LocalFrame.Transform(CurTargetTransform.Inverse());
+	CurrentStamp.Power = GetActivePressure() * ((bSmoothing) ? SculptProperties->SmoothBrushSpeed : SculptProperties->PrimaryBrushSpeed);
+
+	CurrentStamp.PrevLocalFrame = LastStamp.LocalFrame;
+	CurrentStamp.PrevWorldFrame = LastStamp.WorldFrame;
+
+	FVector3d MoveDelta = CurrentStamp.LocalFrame.Origin - CurrentStamp.PrevLocalFrame.Origin;
+	if (UseBrushOp->IgnoreZeroMovements() && MoveDelta.SquaredLength() < FMathd::ZeroTolerance)
+	{
+		return false;
+	}
+
+	LastBrushPosLocal = CurrentStamp.LocalFrame.Origin;
+
+	return true;
+}
+
+
+void UMeshVertexSculptTool::ApplyStamp()
+{
+	SCOPE_CYCLE_COUNTER(VtxSculptToolApplyStamp);
+
+	TUniquePtr<FMeshSculptBrushOp>& UseBrushOp = (bSmoothing) ? SecondaryBrushOp : PrimaryBrushOp;
+	if (UseBrushOp->WantsStampRegionPlane())
+	{
+		CurrentStamp.RegionPlane = ComputeROIBrushPlane(CurrentStamp.LocalFrame.Origin, true, false);
+	}
+
+	FDynamicMesh3* Mesh = DynamicMeshComponent->GetMesh();
+	UseBrushOp->ApplyStamp(Mesh, CurrentStamp, VertexROI, ROIPositionBuffer);
+
+	SyncMeshWithPositionBuffer(Mesh);
+
+	LastStamp = CurrentStamp;
+}
+
+
+
+double UMeshVertexSculptTool::CalculateBrushFalloff(double Distance)
+{
+	double f = FMathd::Clamp(1.0 - BrushProperties->BrushFalloffAmount, 0.0, 1.0);
+	double d = Distance / CurrentBrushRadius;
+	double w = 1;
+	if (d > f)
+	{
+		d = FMathd::Clamp((d - f) / (1.0 - f), 0.0, 1.0);
+		w = (1.0 - d * d);
+		w = w * w * w;
+	}
+	return w;
+}
+
+
+
+void UMeshVertexSculptTool::SyncMeshWithPositionBuffer(FDynamicMesh3* Mesh)
+{
+	const int NumV = ROIPositionBuffer.Num();
+	checkSlow(VertexROI.Num() <= NumV);
+
+	// actually doing this in parallel is slower here due to locking (?)
+	ParallelFor(NumV, [&](int32 k)
+	{
+		int VertIdx = VertexROI[k];
+		const FVector3d& NewPos = ROIPositionBuffer[k];
+		FVector3d OrigPos = Mesh->GetVertex(VertIdx);
+		Mesh->SetVertex(VertIdx, NewPos);
+
+		//UpdateSavedVertexLock.Lock();
+		ActiveVertexChange->UpdateVertex(VertIdx, OrigPos, NewPos);
+		//UpdateSavedVertexLock.Unlock();
+	}, true);
+}
+
+
+
+
+
+
+
+
+FFrame3d UMeshVertexSculptTool::ComputeROIBrushPlane(const FVector3d& BrushCenter, bool bIgnoreDepth, bool bViewAligned)
+{
+	FDynamicMesh3* Mesh = DynamicMeshComponent->GetMesh();
+	FVector3d AverageNormal(0, 0, 0);
+	FVector3d AveragePos(0, 0, 0);
+	double WeightSum = 0;
+	for (int TriID : TriangleROI)
+	{
+		FVector3d Centroid = Mesh->GetTriCentroid(TriID);
+		double Weight = CalculateBrushFalloff(BrushCenter.Distance(Centroid));
+
+		AverageNormal += Weight * Mesh->GetTriNormal(TriID);
+		AveragePos += Weight * Centroid;
+		WeightSum += Weight;
+	}
+	AverageNormal.Normalize();
+	AveragePos /= WeightSum;
+
+	if (bViewAligned)
+	{
+		AverageNormal = -CameraState.Forward();
+	}
+
+	FFrame3d Result = FFrame3d(AveragePos, AverageNormal);
+	if (bIgnoreDepth == false)
+	{
+		Result.Origin -= BrushProperties->Depth * CurrentBrushRadius * Result.Z();
+	}
+
+	return Result;
+}
+
+
+
+
+int UMeshVertexSculptTool::FindHitSculptMeshTriangle(const FRay3d& LocalRay)
+{
+	if (BrushProperties->bHitBackFaces)
+	{
+		return Octree.FindNearestHitObject(LocalRay);
+	}
+	else
+	{
+		FDynamicMesh3* Mesh = DynamicMeshComponent->GetMesh();
+
+		FViewCameraState StateOut;
+		GetToolManager()->GetContextQueriesAPI()->GetCurrentViewState(StateOut);
+		FVector3d LocalEyePosition(CurTargetTransform.InverseTransformPosition(StateOut.Position));
+		int HitTID = Octree.FindNearestHitObject(LocalRay,
+			[this, Mesh, &LocalEyePosition](int TriangleID) {
+			FVector3d Normal, Centroid;
+			double Area;
+			Mesh->GetTriInfo(TriangleID, Normal, Area, Centroid);
+			return Normal.Dot((Centroid - LocalEyePosition)) < 0;
+		});
+		return HitTID;
+	}
+}
+
+int UMeshVertexSculptTool::FindHitTargetMeshTriangle(const FRay3d& LocalRay)
+{
+	IMeshSpatial::FQueryOptions RaycastOptions;
+
+	if (BrushProperties->bHitBackFaces == false)
+	{
+		FDynamicMesh3* Mesh = DynamicMeshComponent->GetMesh();
+
+		FViewCameraState StateOut;
+		GetToolManager()->GetContextQueriesAPI()->GetCurrentViewState(StateOut);
+		FVector3d LocalEyePosition(CurTargetTransform.InverseTransformPosition(StateOut.Position));
+
+		RaycastOptions.TriangleFilterF = [this, Mesh, LocalEyePosition](int TriangleID) {
+			FVector3d Normal, Centroid;
+			double Area;
+			Mesh->GetTriInfo(TriangleID, Normal, Area, Centroid);
+			return Normal.Dot((Centroid - LocalEyePosition)) < 0;
+		};
+
+		return BaseMeshSpatial.FindNearestHitObject(LocalRay, RaycastOptions.TriangleFilterF);
+	}
+	else
+	{
+		return BaseMeshSpatial.FindNearestHitObject(LocalRay);
+	}
+}
+
+bool UMeshVertexSculptTool::UpdateBrushPositionOnActivePlane(const FRay& WorldRay)
+{
+	FVector3d NewHitPosWorld;
+	ActiveDragPlane.RayPlaneIntersection(WorldRay.Origin, WorldRay.Direction, 2, NewHitPosWorld);
+	LastBrushPosWorld = NewHitPosWorld;
+	LastBrushPosNormalWorld = ActiveDragPlane.Z();
+	return true;
+}
+
+bool UMeshVertexSculptTool::UpdateBrushPositionOnTargetMesh(const FRay& WorldRay, bool bFallbackToViewPlane)
+{
+	FRay3d LocalRay(CurTargetTransform.InverseTransformPosition(WorldRay.Origin),
+		CurTargetTransform.InverseTransformVector(WorldRay.Direction));
+	LocalRay.Direction.Normalize();
+
+	int32 HitTID = FindHitTargetMeshTriangle(LocalRay);
+	if (HitTID != IndexConstants::InvalidID)
+	{
+		FIntrRay3Triangle3d Query = TMeshQueries<FDynamicMesh3>::TriangleIntersection(BaseMesh, HitTID, LocalRay);
+		LastBrushPosNormalWorld = CurTargetTransform.TransformNormal(BaseMesh.GetTriNormal(HitTID));
+		LastBrushPosWorld = CurTargetTransform.TransformPosition(LocalRay.PointAt(Query.RayParameter));
+		return true;
+	}
+
+	if (bFallbackToViewPlane)
+	{
+		FFrame3d BrushPlane(LastBrushPosWorld, CameraState.Forward());
+		FVector3d NewHitPosWorld;
+		BrushPlane.RayPlaneIntersection(WorldRay.Origin, WorldRay.Direction, 2, NewHitPosWorld);
+		LastBrushPosWorld = NewHitPosWorld;
+		LastBrushPosNormalWorld = ActiveDragPlane.Z();
+		return true;
+	}
+
+	return false;
+}
+
+bool UMeshVertexSculptTool::UpdateBrushPositionOnSculptMesh(const FRay& WorldRay, bool bFallbackToViewPlane)
+{
+	FRay3d LocalRay(CurTargetTransform.InverseTransformPosition(WorldRay.Origin),
+		CurTargetTransform.InverseTransformVector(WorldRay.Direction));
+	LocalRay.Direction.Normalize();
+
+	int32 HitTID = FindHitSculptMeshTriangle(LocalRay);
+	if (HitTID != IndexConstants::InvalidID)
+	{
+		const FDynamicMesh3* SculptMesh = DynamicMeshComponent->GetMesh();
+		FIntrRay3Triangle3d Query = TMeshQueries<FDynamicMesh3>::TriangleIntersection(*SculptMesh, HitTID, LocalRay);
+		LastBrushPosNormalWorld = CurTargetTransform.TransformNormal(SculptMesh->GetTriNormal(HitTID));
+		LastBrushPosWorld = CurTargetTransform.TransformPosition(LocalRay.PointAt(Query.RayParameter));
+		return true;
+	}
+
+	if (bFallbackToViewPlane)
+	{
+		FFrame3d BrushPlane(LastBrushPosWorld, CameraState.Forward());
+		FVector3d NewHitPosWorld;
+		BrushPlane.RayPlaneIntersection(WorldRay.Origin, WorldRay.Direction, 2, NewHitPosWorld);
+		LastBrushPosWorld = NewHitPosWorld;
+		LastBrushPosNormalWorld = ActiveDragPlane.Z();
+		return true;
+	}
+
+	return false;
+}
+
+void UMeshVertexSculptTool::AlignBrushToView()
+{
+	LastBrushPosNormalWorld = -CameraState.Forward();
+}
+
+
+bool UMeshVertexSculptTool::UpdateBrushPosition(const FRay& WorldRay)
+{
+	// This is an unfortunate hack necessary because we haven't refactored brushes properly yet
+
+	if (bSmoothing)
+	{
+		return UpdateBrushPositionOnSculptMesh(WorldRay, false);
+	}
+
+	bool bHit = false;
+	switch (SculptProperties->PrimaryBrushType)
+	{
+	case EMeshVertexSculptBrushType::Offset:
+	case EMeshVertexSculptBrushType::SculptMax:
+	case EMeshVertexSculptBrushType::Pinch:
+		bHit = UpdateBrushPositionOnTargetMesh(WorldRay, false);
+		break;
+
+	case EMeshVertexSculptBrushType::SculptView:
+	case EMeshVertexSculptBrushType::PlaneViewAligned:
+		bHit = UpdateBrushPositionOnTargetMesh(WorldRay, false);
+		AlignBrushToView();
+		break;
+
+	case EMeshVertexSculptBrushType::Move:
+		//return UpdateBrushPositionOnActivePlane(WorldRay);
+		bHit = UpdateBrushPositionOnSculptMesh(WorldRay, false);
+		break;
+
+	case EMeshVertexSculptBrushType::Smooth:
+	case EMeshVertexSculptBrushType::Inflate:
+	case EMeshVertexSculptBrushType::Flatten:
+	case EMeshVertexSculptBrushType::Plane:
+	case EMeshVertexSculptBrushType::FixedPlane:
+		bHit = UpdateBrushPositionOnSculptMesh(WorldRay, false);
+		break;
+
+	default:
+		UE_LOG(LogTemp, Warning, TEXT("UMeshVertexSculptTool: unknown brush type in UpdateBrushPosition"));
+		bHit = UpdateBrushPositionOnSculptMesh(WorldRay, false);
+		break;
+	}
+
+	return bHit;
+}
+
+
+
+void UMeshVertexSculptTool::OnEndDrag(const FRay& Ray)
+{
+	bInDrag = false;
+
+	// cancel these! otherwise change record could become invalid
+	bStampPending = false;
+
+	// update spatial
+	bTargetDirty = true;
+
+	TUniquePtr<FMeshSculptBrushOp>& UseBrushOp = (bSmoothing) ? SecondaryBrushOp : PrimaryBrushOp;
+	UseBrushOp->EndStroke(DynamicMeshComponent->GetMesh(), LastStamp, VertexROI);
+
+	// close change record
+	EndChange();
+}
+
+
+FInputRayHit UMeshVertexSculptTool::BeginHoverSequenceHitTest(const FInputDeviceRay& PressPos)
+{
+	return UMeshSurfacePointTool::BeginHoverSequenceHitTest(PressPos);
+}
+
+bool UMeshVertexSculptTool::OnUpdateHover(const FInputDeviceRay& DevicePos)
+{
+	PendingStampType = SculptProperties->PrimaryBrushType;
+
+	if (bInDrag)
+	{
+		FVector3d NewHitPosWorld;
+		ActiveDragPlane.RayPlaneIntersection(DevicePos.WorldRay.Origin, DevicePos.WorldRay.Direction, 2, NewHitPosWorld);
+		LastBrushPosWorld = NewHitPosWorld;
+		LastBrushPosNormalWorld = ActiveDragPlane.Z();
+	}
+	else
+	{
+		UpdateBrushPosition(DevicePos.WorldRay);
+	}
+	return true;
+}
+
+
+
+void UMeshVertexSculptTool::OnTick(float DeltaTime)
+{
+	UMeshSculptToolBase::OnTick(DeltaTime);
+
+	SCOPE_CYCLE_COUNTER(VtxSculptToolTick);
+
+	UpdateHoverStamp(LastBrushPosWorld, LastBrushPosNormalWorld);
+
+	BrushTypeWatcher.CheckAndUpdate();
+
+	UpdateWorkPlane();
+
+	// process the undo update
+	if (bUndoUpdatePending)
+	{
+		// wait for updates
+		WaitForPendingUndoRedo();
+
+		// post rendering update
+		DynamicMeshComponent->FastNotifyTriangleVerticesUpdated(AccumulatedTriangleROI,
+			EMeshRenderAttributeFlags::Positions | EMeshRenderAttributeFlags::VertexNormals);
+		GetToolManager()->PostInvalidation();
+
+		// ignore stamp and wait for next tick to do anything else
+		bUndoUpdatePending = false;
+		return;
+	}
+
+	// if user changed to not-frozen, we need to reinitialize the target
+	if (bCachedFreezeTarget != SculptProperties->bFreezeTarget)
+	{
+		UpdateBaseMesh(nullptr);
+		bTargetDirty = false;
+	}
+
+
+	if (bStampPending)
+	{
+		//UE_LOG(LogTemp, Warning, TEXT("dt is %.3f, tick fps %.2f - roi size %d/%d"), DeltaTime, 1.0 / DeltaTime, VertexROI.Num(), TriangleROI.Num());
+		SCOPE_CYCLE_COUNTER(VtxSculptTool_Tick_ApplyStampBlock);
+		// hardcoded brush flow (todo: make optional)
+		bStampPending = bInDrag;
+
+		// update brush position
+		if (UpdateStampPosition(PendingStampRay) == false)
+		{
+			return;
+		}
+
+		// update sculpt ROI
+		UpdateROI(CurrentStamp.LocalFrame.Origin);
+
+		// append updated ROI to modified region (async)
+		TFuture<void> AccumulateROI = Async(VertexSculptToolAsyncExecTarget, [&]()
+		{
+			AccumulatedTriangleROI.Append(TriangleROI);
+		});
+
+		// apply the stamp
+		ApplyStamp();
+
+		// begin octree rebuild calculation
+		TFuture<void> OctreeRebuild = Async(VertexSculptToolAsyncExecTarget, [&]()
+		{
+			SCOPE_CYCLE_COUNTER(VtxSculptTool_Tick_ApplyStamp_Insert);
+			Octree.ReinsertTriangles(TriangleROI);
+		});
+
+		// recalculate normals. This has to complete before we can update component
+		// (in fact we could do it per-chunk...)
+		FDynamicMesh3* Mesh = DynamicMeshComponent->GetMesh();
+		{
+			SCOPE_CYCLE_COUNTER(VtxSculptTool_Tick_NormalsBlock);
+			UE::SculptUtil::RecalculateROINormals(Mesh, TriangleROI, VertexSetBuffer, NormalsBuffer);
+		}
+
+		{
+			SCOPE_CYCLE_COUNTER(VtxSculptTool_Tick_UpdateMeshBlock);
+			DynamicMeshComponent->FastNotifyTriangleVerticesUpdated(TriangleROI,
+				EMeshRenderAttributeFlags::Positions | EMeshRenderAttributeFlags::VertexNormals);
+			GetToolManager()->PostInvalidation();
+		}
+
+		// we don't really need to wait for these to happen to end Tick()...
+		AccumulateROI.Wait();
+		OctreeRebuild.Wait();
+	} 
+	else if (bTargetDirty)
+	{
+		SCOPE_CYCLE_COUNTER(VtxSculptTool_Tick_UpdateTargetBlock);
+		check(bInDrag == false);
+
+		// this spawns futures that we could allow to run while other things happen...
+		UpdateBaseMesh(&AccumulatedTriangleROI);
+		AccumulatedTriangleROI.Reset();
+
+		bTargetDirty = false;
+	}
+
+}
+
+
+
+
+void UMeshVertexSculptTool::UpdateBaseMesh(const TSet<int32>* TriangleSet)
+{
+	if (SculptProperties != nullptr)
+	{
+		bCachedFreezeTarget = SculptProperties->bFreezeTarget;
+		if (SculptProperties->bFreezeTarget)
+		{
+			return;   // do not update frozen target
+		}
+	}
+
+	const FDynamicMesh3* SculptMesh = DynamicMeshComponent->GetMesh();
+	if ( ! TriangleSet )
+	{
+		BaseMesh.Copy(*SculptMesh, false, false, false, false);
+		BaseMesh.EnableVertexNormals(FVector3f::UnitZ());
+		FMeshNormals::QuickComputeVertexNormals(BaseMesh);
+		BaseMeshSpatial.SetMaxTreeDepth(8);
+		BaseMeshSpatial = FDynamicMeshOctree3();   // need to clear...
+		BaseMeshSpatial.Initialize(&BaseMesh);
+	}
+	else
+	{
+		BaseMeshIndexBuffer.Reset();
+		for ( int32 tid : *TriangleSet)
+		{ 
+			FIndex3i Tri = BaseMesh.GetTriangle(tid);
+			BaseMesh.SetVertex(Tri.A, SculptMesh->GetVertex(Tri.A));
+			BaseMesh.SetVertex(Tri.B, SculptMesh->GetVertex(Tri.B));
+			BaseMesh.SetVertex(Tri.C, SculptMesh->GetVertex(Tri.C));
+			BaseMeshIndexBuffer.Add(tid);
+		}
+		auto UpdateBaseNormals = Async(EAsyncExecution::Thread, [&]()
+		{
+			FMeshNormals::QuickComputeVertexNormalsForTriangles(BaseMesh, BaseMeshIndexBuffer);
+		});
+		auto ReinsertTriangles = Async(EAsyncExecution::Thread, [&]()
+		{
+			BaseMeshSpatial.ReinsertTriangles(*TriangleSet);
+		});
+		UpdateBaseNormals.Wait();
+		ReinsertTriangles.Wait();
+	}
+}
+
+
+bool UMeshVertexSculptTool::GetBaseMeshNearest(int32 VertexID, const FVector3d& Position, double SearchRadius, FVector3d& TargetPosOut, FVector3d& TargetNormalOut)
+{
+	TargetPosOut = BaseMesh.GetVertex(VertexID);
+	TargetNormalOut = (FVector3d)BaseMesh.GetVertexNormal(VertexID);
+	return true;
+
+	// Octree FindNearestObject is currently quite expensive...
+	//int32 NearestTID = BaseMeshSpatial.FindNearestObject(Position, SearchRadius);
+	//if (NearestTID <= 0)
+	//{
+	//	return false;
+	//}
+	//FDistPoint3Triangle3d DistanceQuery = TMeshQueries<FDynamicMesh3>::TriangleDistance(BaseMesh, NearestTID, Position);
+	//TargetPosOut = DistanceQuery.ClosestTrianglePoint;
+	//TargetNormalOut = BaseMesh.GetTriBaryNormal(NearestTID, DistanceQuery.TriangleBaryCoords.X, DistanceQuery.TriangleBaryCoords.Y, DistanceQuery.TriangleBaryCoords.Z);
+	//return true;
+}
+
+
+
+
+
+
+void UMeshVertexSculptTool::IncreaseBrushSpeedAction()
+{
+	SculptProperties->PrimaryBrushSpeed = FMath::Clamp(SculptProperties->PrimaryBrushSpeed + 0.05f, 0.0f, 1.0f);
+}
+
+void UMeshVertexSculptTool::DecreaseBrushSpeedAction()
+{
+	SculptProperties->PrimaryBrushSpeed = FMath::Clamp(SculptProperties->PrimaryBrushSpeed - 0.05f, 0.0f, 1.0f);
+}
+
+
+
+void UMeshVertexSculptTool::NextBrushModeAction()
+{
+	uint8 LastMode = (uint8)EMeshVertexSculptBrushType::LastValue;
+	SculptProperties->PrimaryBrushType = (EMeshVertexSculptBrushType)(((uint8)SculptProperties->PrimaryBrushType + 1) % LastMode);
+}
+
+void UMeshVertexSculptTool::PreviousBrushModeAction()
+{
+	uint8 LastMode = (uint8)EMeshVertexSculptBrushType::LastValue;
+	uint8 CurMode = (uint8)SculptProperties->PrimaryBrushType;
+	if (CurMode == 0)
+	{
+		SculptProperties->PrimaryBrushType = (EMeshVertexSculptBrushType)((uint8)LastMode - 1);
+	}
+	else
+	{
+		SculptProperties->PrimaryBrushType = (EMeshVertexSculptBrushType)((uint8)CurMode - 1);
+	}
+}
+
+
+
+//
+// Change Tracking
+//
+void UMeshVertexSculptTool::BeginChange()
+{
+	check(ActiveVertexChange == nullptr);
+	ActiveVertexChange = new FMeshVertexChangeBuilder();
+}
+
+void UMeshVertexSculptTool::EndChange()
+{
+	check(ActiveVertexChange);
+
+	TUniquePtr<TWrappedToolCommandChange<FMeshVertexChange>> NewChange = MakeUnique<TWrappedToolCommandChange<FMeshVertexChange>>();
+	NewChange->WrappedChange = MoveTemp(ActiveVertexChange->Change);
+	NewChange->BeforeModify = [this](bool bRevert)
+	{
+		this->WaitForPendingUndoRedo();
+	};
+
+	GetToolManager()->EmitObjectChange(DynamicMeshComponent, MoveTemp(NewChange), LOCTEXT("VertexSculptChange", "Brush Stroke"));
+
+	delete ActiveVertexChange;
+	ActiveVertexChange = nullptr;
+}
+
+
+void UMeshVertexSculptTool::WaitForPendingUndoRedo()
+{
+	if (bUndoUpdatePending)
+	{
+		UndoNormalsFuture.Wait();
+		UndoUpdateOctreeFuture.Wait();
+		UndoUpdateBaseMeshFuture.Wait();
+		bUndoUpdatePending = false;
+	}
+}
+
+void UMeshVertexSculptTool::OnDynamicMeshComponentChanged(USimpleDynamicMeshComponent* Component, const FMeshVertexChange* Change, bool bRevert)
+{
+	// update octree
+	FDynamicMesh3* Mesh = DynamicMeshComponent->GetMesh();
+
+	// make sure any previous async computations are done, and update the undo ROI
+	if (bUndoUpdatePending)
+	{
+		// we should never hit this anymore, because of pre-change calling WaitForPendingUndoRedo()
+		WaitForPendingUndoRedo();
+
+		// this is not right because now we are going to do extra recomputation, but it's very messy otherwise...
+		MeshIndexUtil::VertexToTriangleOneRing(Mesh, Change->Vertices, AccumulatedTriangleROI);
+	}
+	else
+	{
+		AccumulatedTriangleROI.Reset();
+		MeshIndexUtil::VertexToTriangleOneRing(Mesh, Change->Vertices, AccumulatedTriangleROI);
+	}
+
+	// start the normal recomputation
+	UndoNormalsFuture = Async(EAsyncExecution::Thread, [&]()
+	{
+		UE::SculptUtil::RecalculateROINormals(Mesh, AccumulatedTriangleROI, VertexSetBuffer, NormalsBuffer);
+		return true;
+	});
+
+	// start the octree update
+	UndoUpdateOctreeFuture = Async(EAsyncExecution::Thread, [&]()
+	{
+		Octree.ReinsertTriangles(AccumulatedTriangleROI);
+		return true;
+	});
+
+	// start the base mesh update
+	UndoUpdateBaseMeshFuture = Async(EAsyncExecution::Thread, [&]()
+	{
+		UpdateBaseMesh(&AccumulatedTriangleROI);
+		return true;
+	});
+
+	// note that we have a pending update
+	bUndoUpdatePending = true;
+}
+
+
+
+
+
+
+
+
+
+
+void UMeshVertexSculptTool::UpdateBrushType(EMeshVertexSculptBrushType BrushType)
+{
+	static const FText BaseMessage = LOCTEXT("OnStartSculptTool", "Hold Shift to Smooth, Ctrl to Invert (where applicable). Q/A keys cycle Brush Type. S/D changes Size (+Shift to small-step), W/E changes Strength.");
+	FTextBuilder Builder;
+	Builder.AppendLine(BaseMessage);
+
+	TUniqueFunction<bool(int32, const FVector3d&, double MaxDist, FVector3d&, FVector3d&)> BaseMeshQueryFunc =
+		[&](int32 VertexID, const FVector3d& Position, double MaxDist, FVector3d& PosOut, FVector3d& NormalOut)
+		{
+			//return GetTargetMeshNearest(VertexID, Position, MaxDist, PosOut, NormalOut);
+			return GetBaseMeshNearest(VertexID, Position, MaxDist, PosOut, NormalOut);
+		};
+
+	switch (BrushType)
+	{
+	case EMeshVertexSculptBrushType::Smooth:
+		PrimaryBrushOp = MakeUnique<FSmoothBrushOp>();
+		break;
+	case EMeshVertexSculptBrushType::Offset:
+		PrimaryBrushOp = MakeUnique<FSurfaceSculptBrushOp>(MoveTemp(BaseMeshQueryFunc));
+		break;
+	case EMeshVertexSculptBrushType::SculptView:
+		PrimaryBrushOp = MakeUnique<FViewAlignedSculptBrushOp>(MoveTemp(BaseMeshQueryFunc));
+		break;
+	case EMeshVertexSculptBrushType::SculptMax:
+		PrimaryBrushOp = MakeUnique<FSurfaceMaxSculptBrushOp>(MoveTemp(BaseMeshQueryFunc));
+		break;
+	case EMeshVertexSculptBrushType::Inflate:
+		PrimaryBrushOp = MakeUnique<FInflateBrushOp>();
+		break;
+	case EMeshVertexSculptBrushType::Move:
+		PrimaryBrushOp = MakeUnique<FMoveBrushOp>();
+		break;
+	case EMeshVertexSculptBrushType::Flatten:
+		PrimaryBrushOp = MakeUnique<FFlattenBrushOp>();
+		break;
+	case EMeshVertexSculptBrushType::Plane:
+		PrimaryBrushOp = MakeUnique<FPlaneBrushOp>();
+		break;
+	case EMeshVertexSculptBrushType::PlaneViewAligned:
+		PrimaryBrushOp = MakeUnique<FPlaneBrushOp>();
+		break;
+	case EMeshVertexSculptBrushType::FixedPlane:
+		PrimaryBrushOp = MakeUnique<FPlaneBrushOp>();
+		break;
+	case EMeshVertexSculptBrushType::Pinch:
+		PrimaryBrushOp = MakeUnique<FPinchBrushOp>();
+		break;
+	}
+	PrimaryBrushOp->Falloff = ActiveFalloff;
+
+
+	SetToolPropertySourceEnabled(GizmoProperties, false);
+	SetToolPropertySourceEnabled(PlaneBrushProperties, false);
+	SetToolPropertySourceEnabled(SculptMaxBrushProperties, false);
+
+	if (BrushType == EMeshVertexSculptBrushType::FixedPlane)
+	{
+		Builder.AppendLine(LOCTEXT("FixedPlaneTip", "Use T to reposition Work Plane at cursor, Shift+T to align to Normal, Ctrl+Shift+T to align to View"));
+		SetToolPropertySourceEnabled(PlaneBrushProperties, true);
+		SetToolPropertySourceEnabled(GizmoProperties, true);
+	}
+	if (BrushType == EMeshVertexSculptBrushType::Plane || BrushType == EMeshVertexSculptBrushType::PlaneViewAligned || BrushType == EMeshVertexSculptBrushType::Flatten)
+	{
+		SetToolPropertySourceEnabled(PlaneBrushProperties, true);
+	}
+	if (BrushType == EMeshVertexSculptBrushType::SculptMax)
+	{
+		SetToolPropertySourceEnabled(SculptMaxBrushProperties, true);
+	}
+
+	GetToolManager()->DisplayMessage(Builder.ToText(), EToolMessageLevel::UserNotification);
+}
+
+
+
+#undef LOCTEXT_NAMESPACE
