Edigrate Dev-Editor updates to InteractiveToolsFramework, GeometryProcessing, MeshModelingToolset, and ModelingTools Editor Mode Plugins. Update a few files in new MeshPaint Plugin for UInteractiveToo

diff --git a/Engine/Plugins/Experimental/MeshModelingToolset/Source/MeshModelingTools/Private/PlaneCutTool.cpp b/Engine/Plugins/Experimental/MeshModelingToolset/Source/MeshModelingTools/Private/PlaneCutTool.cpp
--- a/Engine/Plugins/Experimental/MeshModelingToolset/Source/MeshModelingTools/Private/PlaneCutTool.cpp
+++ b/Engine/Plugins/Experimental/MeshModelingToolset/Source/MeshModelingTools/Private/PlaneCutTool.cpp
@@ -82,41 +82,6 @@ UPlaneCutToolProperties::UPlaneCutToolProperties() :
 {
 }
 
-void UPlaneCutToolProperties::SaveProperties(UInteractiveTool* SaveFromTool)
-{
-	UPlaneCutToolProperties* PropertyCache = GetPropertyCache<UPlaneCutToolProperties>();
-	PropertyCache->bKeepBothHalves = this->bKeepBothHalves;
-	PropertyCache->bFillCutHole = this->bFillCutHole;
-	PropertyCache->SpacingBetweenHalves = this->SpacingBetweenHalves;
-	PropertyCache->bSnapToWorldGrid = this->bSnapToWorldGrid;
-	PropertyCache->bFillSpans = this->bFillSpans;
-	PropertyCache->bShowPreview = this->bShowPreview;
-}
-
-void UPlaneCutToolProperties::RestoreProperties(UInteractiveTool* RestoreToTool)
-{
-	UPlaneCutToolProperties* PropertyCache = GetPropertyCache<UPlaneCutToolProperties>();
-	this->bKeepBothHalves = PropertyCache->bKeepBothHalves;
-	this->bFillCutHole = PropertyCache->bFillCutHole;
-	this->SpacingBetweenHalves = PropertyCache->SpacingBetweenHalves;
-	this->bSnapToWorldGrid = PropertyCache->bSnapToWorldGrid;
-	this->bFillSpans = PropertyCache->bFillSpans;
-	this->bShowPreview = PropertyCache->bShowPreview;
-}
-
-void UAcceptOutputProperties::SaveProperties(UInteractiveTool* SaveFromTool)
-{
-	UAcceptOutputProperties* PropertyCache = GetPropertyCache<UAcceptOutputProperties>();
-	PropertyCache->bExportSeparatedPiecesAsNewMeshAssets = this->bExportSeparatedPiecesAsNewMeshAssets;
-}
-
-void UAcceptOutputProperties::RestoreProperties(UInteractiveTool* RestoreToTool)
-{
-	UAcceptOutputProperties* PropertyCache = GetPropertyCache<UAcceptOutputProperties>();
-	this->bExportSeparatedPiecesAsNewMeshAssets = PropertyCache->bExportSeparatedPiecesAsNewMeshAssets;
-}
-
-
 UPlaneCutTool::UPlaneCutTool()
 {
 	CutPlaneOrigin = FVector::ZeroVector;
@@ -416,7 +381,7 @@ void UPlaneCutTool::Render(IToolsContextRenderAPI* RenderAPI)
 	MeshDebugDraw::DrawSimpleGrid(DrawFrame, NumGridLines, GridLineSpacing, GridThickness, GridColor, false, PDI, FTransform::Identity);
 }
 
-void UPlaneCutTool::Tick(float DeltaTime)
+void UPlaneCutTool::OnTick(float DeltaTime)
 {
 	if (PlaneTransformGizmo != nullptr)
 	{
diff --git a/Engine/Plugins/Experimental/MeshModelingToolset/Source/MeshModelingTools/Private/Properties/MeshMaterialProperties.cpp b/Engine/Plugins/Experimental/MeshModelingToolset/Source/MeshModelingTools/Private/Properties/MeshMaterialProperties.cpp
--- a/Engine/Plugins/Experimental/MeshModelingToolset/Source/MeshModelingTools/Private/Properties/MeshMaterialProperties.cpp
+++ b/Engine/Plugins/Experimental/MeshModelingToolset/Source/MeshModelingTools/Private/Properties/MeshMaterialProperties.cpp
@@ -17,27 +17,6 @@ UNewMeshMaterialProperties::UNewMeshMaterialProperties()
 	Material = CreateDefaultSubobject<UMaterialInterface>(TEXT("MATERIAL"));
 }
 
-
-void UNewMeshMaterialProperties::SaveProperties(UInteractiveTool* SaveFromTool)
-{
-	UNewMeshMaterialProperties* PropertyCache = GetPropertyCache<UNewMeshMaterialProperties>();
-	PropertyCache->Material = this->Material;
-	PropertyCache->UVScale = this->UVScale;
-	PropertyCache->bWorldSpaceUVScale = this->bWorldSpaceUVScale;
-	// not bWireframe
-}
-
-void UNewMeshMaterialProperties::RestoreProperties(UInteractiveTool* RestoreToTool)
-{
-	UNewMeshMaterialProperties* PropertyCache = GetPropertyCache<UNewMeshMaterialProperties>();
-	this->Material = PropertyCache->Material;
-	this->UVScale = PropertyCache->UVScale;
-	this->bWorldSpaceUVScale = PropertyCache->bWorldSpaceUVScale;
-}
-
-
-
-
 void UExistingMeshMaterialProperties::Setup()
 {
 	UMaterial* CheckerMaterialBase = LoadObject<UMaterial>(nullptr, TEXT("/MeshModelingToolset/Materials/CheckerMaterial"));
@@ -73,42 +52,4 @@ UMaterialInterface* UExistingMeshMaterialProperties::GetActiveOverrideMaterial()
 	return nullptr;
 }
 
-void UExistingMeshMaterialProperties::SaveProperties(UInteractiveTool* SaveFromTool)
-{
-	UExistingMeshMaterialProperties* PropertyCache = GetPropertyCache<UExistingMeshMaterialProperties>();
-	PropertyCache->MaterialMode = this->MaterialMode;
-	PropertyCache->CheckerDensity = this->CheckerDensity;
-	PropertyCache->OverrideMaterial = this->OverrideMaterial;
-}
-
-void UExistingMeshMaterialProperties::RestoreProperties(UInteractiveTool* RestoreToTool)
-{
-	UExistingMeshMaterialProperties* PropertyCache = GetPropertyCache<UExistingMeshMaterialProperties>();
-	this->MaterialMode = PropertyCache->MaterialMode;
-	this->CheckerDensity = PropertyCache->CheckerDensity;
-	this->OverrideMaterial = PropertyCache->OverrideMaterial;
-	this->Setup();
-}
-
-void UMeshEditingViewProperties::SaveProperties(UInteractiveTool* SaveFromTool)
-{
-	UMeshEditingViewProperties* PropertyCache = GetPropertyCache<UMeshEditingViewProperties>();
-	PropertyCache->bShowWireframe = this->bShowWireframe;
-	PropertyCache->MaterialMode = this->MaterialMode;
-	PropertyCache->bFlatShading = this->bFlatShading;
-	PropertyCache->Color = this->Color;
-	PropertyCache->Image = this->Image;
-}
-
-void UMeshEditingViewProperties::RestoreProperties(UInteractiveTool* RestoreToTool)
-{
-	UMeshEditingViewProperties* PropertyCache = GetPropertyCache<UMeshEditingViewProperties>();
-	this->bShowWireframe = PropertyCache->bShowWireframe;
-	this->MaterialMode = PropertyCache->MaterialMode;
-	this->bFlatShading = PropertyCache->bFlatShading;
-	this->Color = PropertyCache->Color;
-	this->Image = PropertyCache->Image;
-}
-
-
 #undef LOCTEXT_NAMESPACE
diff --git a/Engine/Plugins/Experimental/MeshModelingToolset/Source/MeshModelingTools/Private/RemeshMeshTool.cpp b/Engine/Plugins/Experimental/MeshModelingToolset/Source/MeshModelingTools/Private/RemeshMeshTool.cpp
--- a/Engine/Plugins/Experimental/MeshModelingToolset/Source/MeshModelingTools/Private/RemeshMeshTool.cpp
+++ b/Engine/Plugins/Experimental/MeshModelingToolset/Source/MeshModelingTools/Private/RemeshMeshTool.cpp
@@ -52,6 +52,7 @@ URemeshMeshToolProperties::URemeshMeshToolProperties()
 	SmoothingStrength = 0.25;
 	RemeshIterations = 20;
 	bDiscardAttributes = false;
+	RemeshType = ERemeshType::Standard;
 	SmoothingType = ERemeshSmoothingType::MeanValue;
 	bPreserveSharpEdges = true;
 	bShowWireframe = true;
@@ -66,36 +67,6 @@ URemeshMeshToolProperties::URemeshMeshToolProperties()
 	bUseTargetEdgeLength = false;
 }
 
-void
-URemeshMeshToolProperties::SaveRestoreProperties(UInteractiveTool* RestoreToTool, bool bSaving)
-{
-	URemeshMeshToolProperties* PropertyCache = GetPropertyCache<URemeshMeshToolProperties>();
-
-	// MeshConstraintProperties
-	SaveRestoreProperty(PropertyCache->bPreserveSharpEdges, this->bPreserveSharpEdges, bSaving);
-	SaveRestoreProperty(PropertyCache->MeshBoundaryConstraint, this->MeshBoundaryConstraint, bSaving);
-	SaveRestoreProperty(PropertyCache->GroupBoundaryConstraint, this->GroupBoundaryConstraint, bSaving);
-	SaveRestoreProperty(PropertyCache->MaterialBoundaryConstraint, this->MaterialBoundaryConstraint, bSaving);
-	SaveRestoreProperty(PropertyCache->bPreventNormalFlips, this->bPreventNormalFlips, bSaving);
-
-	//RemeshProperties
-	SaveRestoreProperty(PropertyCache->SmoothingStrength, this->SmoothingStrength, bSaving);
-	SaveRestoreProperty(PropertyCache->bFlips, this->bFlips, bSaving);
-	SaveRestoreProperty(PropertyCache->bSplits, this->bSplits, bSaving);
-	SaveRestoreProperty(PropertyCache->bCollapses, this->bCollapses, bSaving);
-
-	//RemeshMeshToolProperties
-	SaveRestoreProperty(PropertyCache->TargetTriangleCount, this->TargetTriangleCount, bSaving);
-	SaveRestoreProperty(PropertyCache->SmoothingType, this->SmoothingType, bSaving);
-	SaveRestoreProperty(PropertyCache->RemeshIterations, this->RemeshIterations, bSaving);
-	SaveRestoreProperty(PropertyCache->bDiscardAttributes, this->bDiscardAttributes, bSaving);
-	SaveRestoreProperty(PropertyCache->bShowWireframe, this->bShowWireframe, bSaving);
-	SaveRestoreProperty(PropertyCache->bShowGroupColors, this->bShowGroupColors, bSaving);
-	SaveRestoreProperty(PropertyCache->bUseTargetEdgeLength, this->bUseTargetEdgeLength, bSaving);
-	SaveRestoreProperty(PropertyCache->TargetEdgeLength, this->TargetEdgeLength, bSaving);
-	SaveRestoreProperty(PropertyCache->bReproject, this->bReproject, bSaving);
-}
-
 URemeshMeshTool::URemeshMeshTool()
 {
 }
@@ -133,12 +104,10 @@ void URemeshMeshTool::Setup()
 	);
 	Preview->PreviewMesh->EnableWireframe(BasicProperties->bShowWireframe);
 
-	ShowGroupsWatcher.Initialize(
-		[this]() { return BasicProperties->bShowGroupColors; },
-		[this](bool bNewValue) { UpdateVisualization(); }, BasicProperties->bShowGroupColors );
-	ShowWireFrameWatcher.Initialize(
-		[this]() { return BasicProperties->bShowWireframe; },
-		[this](bool bNewValue) { UpdateVisualization(); }, BasicProperties->bShowWireframe );
+	BasicProperties->WatchProperty(BasicProperties->bShowGroupColors,
+								   [this](bool bNewValue) { UpdateVisualization();});
+	BasicProperties->WatchProperty(BasicProperties->bShowWireframe,
+								   [this](bool bNewValue) { UpdateVisualization();});
 
 	OriginalMesh = MakeShared<FDynamicMesh3>();
 	FMeshDescriptionToDynamicMesh Converter;
@@ -182,11 +151,8 @@ void URemeshMeshTool::Shutdown(EToolShutdownType ShutdownType)
 	}
 }
 
-void URemeshMeshTool::Tick(float DeltaTime)
+void URemeshMeshTool::OnTick(float DeltaTime)
 {
-	ShowWireFrameWatcher.CheckAndUpdate();
-	ShowGroupsWatcher.CheckAndUpdate();
-
 	Preview->Tick(DeltaTime);
 }
 
@@ -194,6 +160,8 @@ TUniquePtr<FDynamicMeshOperator> URemeshMeshTool::MakeNewOperator()
 {
 	TUniquePtr<FRemeshMeshOp> Op = MakeUnique<FRemeshMeshOp>();
 
+	Op->RemeshType = BasicProperties->RemeshType;
+
 	if (!BasicProperties->bUseTargetEdgeLength)
 	{
 		Op->TargetEdgeLength = CalculateTargetEdgeLength(BasicProperties->TargetTriangleCount);
@@ -256,7 +224,7 @@ void URemeshMeshTool::OnPropertyModified(UObject* PropertySet, FProperty* Proper
 		if ( ( Property->GetFName() == GET_MEMBER_NAME_CHECKED(URemeshMeshToolProperties, bShowWireframe) ) ||
 			 ( Property->GetFName() == GET_MEMBER_NAME_CHECKED(URemeshMeshToolProperties, bShowGroupColors) ) )
 		{
-			UpdateVisualization();
+			//UpdateVisualization();
 		}
 		else
 		{
diff --git a/Engine/Plugins/Experimental/MeshModelingToolset/Source/MeshModelingTools/Private/RemoveOccludedTrianglesTool.cpp b/Engine/Plugins/Experimental/MeshModelingToolset/Source/MeshModelingTools/Private/RemoveOccludedTrianglesTool.cpp
--- a/Engine/Plugins/Experimental/MeshModelingToolset/Source/MeshModelingTools/Private/RemoveOccludedTrianglesTool.cpp
+++ b/Engine/Plugins/Experimental/MeshModelingToolset/Source/MeshModelingTools/Private/RemoveOccludedTrianglesTool.cpp
@@ -325,13 +325,7 @@ TUniquePtr<FDynamicMeshOperator> URemoveOccludedTrianglesOperatorFactory::MakeNe
 	return Op;
 }
 
-
-
-void URemoveOccludedTrianglesTool::Render(IToolsContextRenderAPI* RenderAPI)
-{
-}
-
-void URemoveOccludedTrianglesTool::Tick(float DeltaTime)
+void URemoveOccludedTrianglesTool::OnTick(float DeltaTime)
 {
 	for (UMeshOpPreviewWithBackgroundCompute* Preview : Previews)
 	{
diff --git a/Engine/Plugins/Experimental/MeshModelingToolset/Source/MeshModelingTools/Private/Sculpting/MeshSculptToolBase.cpp b/Engine/Plugins/Experimental/MeshModelingToolset/Source/MeshModelingTools/Private/Sculpting/MeshSculptToolBase.cpp
new file mode 100644
--- /dev/null
+++ b/Engine/Plugins/Experimental/MeshModelingToolset/Source/MeshModelingTools/Private/Sculpting/MeshSculptToolBase.cpp
@@ -0,0 +1,578 @@
+// Copyright Epic Games, Inc. All Rights Reserved. 
+
+#include "Sculpting/MeshSculptToolBase.h"
+#include "InteractiveToolManager.h"
+#include "InteractiveGizmoManager.h"
+#include "BaseGizmos/BrushStampIndicator.h"
+#include "ToolSetupUtil.h"
+#include "ToolSceneQueriesUtil.h"
+#include "PreviewMesh.h"
+#include "BaseGizmos/GizmoComponents.h"
+#include "BaseGizmos/TransformGizmo.h"
+#include "Drawing/MeshDebugDrawing.h"
+
+
+
+#include "Generators/SphereGenerator.h"
+
+
+
+#define LOCTEXT_NAMESPACE "UMeshSculptToolBase"
+
+
+
+namespace
+{
+	const FString VertexSculptIndicatorGizmoType = TEXT("BrushIndicatorGizmoType");
+}
+
+
+void UMeshSculptToolBase::SetWorld(UWorld* World)
+{
+	this->TargetWorld = World;
+}
+
+
+void UMeshSculptToolBase::Setup()
+{
+	UMeshSurfacePointTool::Setup();
+
+	BrushProperties = NewObject<USculptBrushProperties>(this);
+	BrushProperties->RestoreProperties(this);
+	BrushProperties->bShowStrength = false;
+
+
+	// work plane
+
+	GizmoProperties = NewObject<UWorkPlaneProperties>();
+	GizmoProperties->RestoreProperties(this);
+
+	// create proxy for plane gizmo, but not gizmo itself, as it only appears in FixedPlane brush mode
+	// listen for changes to the proxy and update the plane when that happens
+	PlaneTransformProxy = NewObject<UTransformProxy>(this);
+	PlaneTransformProxy->OnTransformChanged.AddUObject(this, &UMeshSculptToolBase::PlaneTransformChanged);
+
+	//GizmoProperties->WatchProperty(GizmoProperties->Position,
+	//	[this](FVector NewPosition) { UpdateGizmoFromProperties(); });
+	//GizmoProperties->WatchProperty(GizmoProperties->Rotation,
+	//	[this](FQuat NewRotation) { UpdateGizmoFromProperties(); });
+	GizmoPositionWatcher.Initialize(
+		[this]() { return GizmoProperties->Position; },
+		[this](FVector NewPosition) { UpdateGizmoFromProperties(); }, GizmoProperties->Position);
+	GizmoRotationWatcher.Initialize(
+		[this]() { return GizmoProperties->Rotation; },
+		[this](FQuat NewRotation) { UpdateGizmoFromProperties(); }, GizmoProperties->Rotation);
+
+
+
+	// display
+	ViewProperties = NewObject<UMeshEditingViewProperties>();
+	ViewProperties->RestoreProperties(this);
+
+	ViewProperties->WatchProperty(ViewProperties->bShowWireframe,
+		[this](bool bNewValue) { UpdateWireframeVisibility(bNewValue); });
+	ViewProperties->WatchProperty(ViewProperties->MaterialMode,
+		[this](EMeshEditingMaterialModes NewMode) { UpdateMaterialMode(NewMode); });
+	ViewProperties->WatchProperty(ViewProperties->bFlatShading,
+		[this](bool bNewValue) { UpdateFlatShadingSetting(bNewValue); });
+	ViewProperties->WatchProperty(ViewProperties->Color,
+		[this](FLinearColor NewColor) { UpdateColorSetting(NewColor); });
+	ViewProperties->WatchProperty(ViewProperties->Image,
+		[this](UTexture2D* NewImage) { UpdateImageSetting(NewImage); });
+}
+
+
+void UMeshSculptToolBase::Shutdown(EToolShutdownType ShutdownType)
+{
+	UMeshSurfacePointTool::Shutdown(ShutdownType);
+
+	BrushIndicatorMesh->Disconnect();
+	BrushIndicatorMesh = nullptr;
+
+	GetToolManager()->GetPairedGizmoManager()->DestroyAllGizmosByOwner(this);
+	BrushIndicator = nullptr;
+	GetToolManager()->GetPairedGizmoManager()->DeregisterGizmoType(VertexSculptIndicatorGizmoType);
+
+	BrushProperties->SaveProperties(this);
+	if (GizmoProperties)
+	{
+		GizmoProperties->SaveProperties(this);
+	}
+
+	ViewProperties->SaveProperties(this);
+}
+
+
+
+void UMeshSculptToolBase::OnTick(float DeltaTime)
+{
+	GizmoPositionWatcher.CheckAndUpdate();
+	GizmoRotationWatcher.CheckAndUpdate();
+
+	ActivePressure = GetCurrentDevicePressure();
+}
+
+
+void UMeshSculptToolBase::Render(IToolsContextRenderAPI* RenderAPI)
+{
+	UMeshSurfacePointTool::Render(RenderAPI);
+	GetToolManager()->GetContextQueriesAPI()->GetCurrentViewState(CameraState);
+
+	BrushIndicator->Update((float)GetCurrentBrushRadius(),
+		(FVector)HoverStamp.WorldFrame.Origin, (FVector)HoverStamp.WorldFrame.Z(),
+		1.0f - BrushProperties->BrushFalloffAmount);
+	if (BrushIndicatorMaterial)
+	{
+		double FixedDimScale = ToolSceneQueriesUtil::CalculateDimensionFromVisualAngleD(CameraState, HoverStamp.WorldFrame.Origin, 1.5f);
+		BrushIndicatorMaterial->SetScalarParameterValue(TEXT("FalloffWidth"), FixedDimScale);
+	}
+
+	if (ShowWorkPlane())
+	{
+		FPrimitiveDrawInterface* PDI = RenderAPI->GetPrimitiveDrawInterface();
+		FColor GridColor(128, 128, 128, 32);
+		float GridThickness = 0.5f;
+		float GridLineSpacing = 25.0f;   // @todo should be relative to view
+		int NumGridLines = 10;
+		FFrame3f DrawFrame(GizmoProperties->Position, GizmoProperties->Rotation);
+		MeshDebugDraw::DrawSimpleGrid(DrawFrame, NumGridLines, GridLineSpacing, GridThickness, GridColor, false, PDI, FTransform::Identity);
+	}
+}
+
+
+
+
+
+
+
+
+void UMeshSculptToolBase::UpdateHoverStamp(const FVector3d& WorldPos, const FVector3d& WorldNormal)
+{
+	HoverStamp.WorldFrame = FFrame3d(WorldPos, WorldNormal);
+}
+
+
+
+void UMeshSculptToolBase::InitializeBrushSizeRange(const FAxisAlignedBox3d& TargetBounds)
+{
+	double MaxDimension = TargetBounds.MaxDim();
+	BrushRelativeSizeRange = FInterval1d(MaxDimension * 0.01, MaxDimension);
+	CalculateBrushRadius();
+}
+
+
+void UMeshSculptToolBase::CalculateBrushRadius()
+{
+	CurrentBrushRadius = 0.5 * BrushRelativeSizeRange.Interpolate(BrushProperties->BrushSize);
+	if (BrushProperties->bSpecifyRadius)
+	{
+		CurrentBrushRadius = BrushProperties->BrushRadius;
+	}
+	else
+	{
+		BrushProperties->BrushRadius = CurrentBrushRadius;
+	}
+}
+
+
+
+void UMeshSculptToolBase::IncreaseBrushRadiusAction()
+{
+	BrushProperties->BrushSize = FMath::Clamp(BrushProperties->BrushSize + 0.025f, 0.0f, 1.0f);
+	CalculateBrushRadius();
+}
+
+void UMeshSculptToolBase::DecreaseBrushRadiusAction()
+{
+	BrushProperties->BrushSize = FMath::Clamp(BrushProperties->BrushSize - 0.025f, 0.0f, 1.0f);
+	CalculateBrushRadius();
+}
+
+void UMeshSculptToolBase::IncreaseBrushRadiusSmallStepAction()
+{
+	BrushProperties->BrushSize = FMath::Clamp(BrushProperties->BrushSize + 0.005f, 0.0f, 1.0f);
+	CalculateBrushRadius();
+}
+
+void UMeshSculptToolBase::DecreaseBrushRadiusSmallStepAction()
+{
+	BrushProperties->BrushSize = FMath::Clamp(BrushProperties->BrushSize - 0.005f, 0.0f, 1.0f);
+	CalculateBrushRadius();
+}
+
+void UMeshSculptToolBase::IncreaseBrushFalloffAction()
+{
+	const float ChangeAmount = 0.1f;
+	const float OldValue = BrushProperties->BrushFalloffAmount;
+	float NewValue = OldValue + ChangeAmount;
+	BrushProperties->BrushFalloffAmount = FMath::Clamp(NewValue, 0.f, 1.f);
+}
+
+void UMeshSculptToolBase::DecreaseBrushFalloffAction()
+{
+	const float ChangeAmount = 0.1f;
+	const float OldValue = BrushProperties->BrushFalloffAmount;
+	float NewValue = OldValue - ChangeAmount;
+	BrushProperties->BrushFalloffAmount = FMath::Clamp(NewValue, 0.f, 1.f);
+}
+
+
+
+
+
+
+
+void UMeshSculptToolBase::UpdateWireframeVisibility(bool bNewValue)
+{
+	GetSculptMeshComponent()->SetEnableWireframeRenderPass(bNewValue);
+}
+
+void UMeshSculptToolBase::UpdateFlatShadingSetting(bool bNewValue)
+{
+	if (ActiveOverrideMaterial != nullptr)
+	{
+		ActiveOverrideMaterial->SetScalarParameterValue(TEXT("FlatShading"), (bNewValue) ? 1.0f : 0.0f);
+	}
+}
+
+void UMeshSculptToolBase::UpdateColorSetting(FLinearColor NewColor)
+{
+	if (ActiveOverrideMaterial != nullptr)
+	{
+		ActiveOverrideMaterial->SetVectorParameterValue(TEXT("Color"), NewColor);
+	}
+}
+
+void UMeshSculptToolBase::UpdateImageSetting(UTexture2D* NewImage)
+{
+	if (ActiveOverrideMaterial != nullptr)
+	{
+		ActiveOverrideMaterial->SetTextureParameterValue(TEXT("ImageTexture"), NewImage);
+	}
+}
+
+
+
+void UMeshSculptToolBase::UpdateMaterialMode(EMeshEditingMaterialModes MaterialMode)
+{
+	if (MaterialMode == EMeshEditingMaterialModes::ExistingMaterial)
+	{
+		GetSculptMeshComponent()->ClearOverrideRenderMaterial();
+		GetSculptMeshComponent()->bCastDynamicShadow = ComponentTarget->GetOwnerComponent()->bCastDynamicShadow;
+		ActiveOverrideMaterial = nullptr;
+	}
+	else
+	{
+		if (MaterialMode == EMeshEditingMaterialModes::Custom)
+		{
+			ActiveOverrideMaterial = ToolSetupUtil::GetCustomImageBasedSculptMaterial(GetToolManager(), ViewProperties->Image);
+			if (ViewProperties->Image != nullptr)
+			{
+				ActiveOverrideMaterial->SetTextureParameterValue(TEXT("ImageTexture"), ViewProperties->Image);
+			}
+		}
+		else
+		{
+			UMaterialInterface* SculptMaterial = nullptr;
+			switch (MaterialMode)
+			{
+			case EMeshEditingMaterialModes::Diffuse:
+				SculptMaterial = ToolSetupUtil::GetDefaultSculptMaterial(GetToolManager());
+				break;
+			case EMeshEditingMaterialModes::Grey:
+				SculptMaterial = ToolSetupUtil::GetImageBasedSculptMaterial(GetToolManager(), ToolSetupUtil::ImageMaterialType::DefaultBasic);
+				break;
+			case EMeshEditingMaterialModes::Soft:
+				SculptMaterial = ToolSetupUtil::GetImageBasedSculptMaterial(GetToolManager(), ToolSetupUtil::ImageMaterialType::DefaultSoft);
+				break;
+			case EMeshEditingMaterialModes::TangentNormal:
+				SculptMaterial = ToolSetupUtil::GetImageBasedSculptMaterial(GetToolManager(), ToolSetupUtil::ImageMaterialType::TangentNormalFromView);
+				break;
+			}
+			if (SculptMaterial != nullptr)
+			{
+				ActiveOverrideMaterial = UMaterialInstanceDynamic::Create(SculptMaterial, this);
+			}
+		}
+
+		if (ActiveOverrideMaterial != nullptr)
+		{
+			GetSculptMeshComponent()->SetOverrideRenderMaterial(ActiveOverrideMaterial);
+			ActiveOverrideMaterial->SetScalarParameterValue(TEXT("FlatShading"), (ViewProperties->bFlatShading) ? 1.0f : 0.0f);
+		}
+
+		GetSculptMeshComponent()->bCastDynamicShadow = false;
+	}
+}
+
+
+
+
+
+
+
+
+
+void UMeshSculptToolBase::InitializeIndicator()
+{
+	// register and spawn brush indicator gizmo
+	GetToolManager()->GetPairedGizmoManager()->RegisterGizmoType(VertexSculptIndicatorGizmoType, NewObject<UBrushStampIndicatorBuilder>());
+	BrushIndicator = GetToolManager()->GetPairedGizmoManager()->CreateGizmo<UBrushStampIndicator>(VertexSculptIndicatorGizmoType, FString(), this);
+	BrushIndicatorMesh = MakeDefaultIndicatorSphereMesh(this, TargetWorld);
+	BrushIndicator->AttachedComponent = BrushIndicatorMesh->GetRootComponent();
+	BrushIndicator->LineThickness = 1.0;
+	BrushIndicator->bDrawIndicatorLines = true;
+	BrushIndicator->bDrawRadiusCircle = false;
+	BrushIndicator->bDrawFalloffCircle = true;
+	BrushIndicator->LineColor = FLinearColor(0.9f, 0.4f, 0.4f);
+}
+
+
+UPreviewMesh* UMeshSculptToolBase::MakeDefaultIndicatorSphereMesh(UObject* Parent, UWorld* World, int Resolution /*= 32*/)
+{
+	UPreviewMesh* SphereMesh = NewObject<UPreviewMesh>(Parent);
+	SphereMesh->CreateInWorld(World, FTransform::Identity);
+	FSphereGenerator SphereGen;
+	SphereGen.NumPhi = SphereGen.NumTheta = Resolution;
+	SphereGen.Generate();
+	FDynamicMesh3 Mesh(&SphereGen);
+	SphereMesh->UpdatePreview(&Mesh);
+
+	BrushIndicatorMaterial = ToolSetupUtil::GetDefaultBrushVolumeMaterial(GetToolManager());
+	if (BrushIndicatorMaterial)
+	{
+		SphereMesh->SetMaterial(BrushIndicatorMaterial);
+	}
+
+	return SphereMesh;
+}
+
+
+
+
+void UMeshSculptToolBase::UpdateWorkPlane()
+{
+	bool bGizmoVisible = ShowWorkPlane() && (GizmoProperties->bShowGizmo);
+	UpdateFixedPlaneGizmoVisibility(bGizmoVisible);
+	GizmoProperties->bPropertySetEnabled = ShowWorkPlane();
+
+	if (PendingWorkPlaneUpdate != EPendingWorkPlaneUpdate::NoUpdatePending)
+	{
+		SetFixedSculptPlaneFromWorldPos((FVector)HoverStamp.WorldFrame.Origin, (FVector)HoverStamp.WorldFrame.Z(), PendingWorkPlaneUpdate);
+		PendingWorkPlaneUpdate = EPendingWorkPlaneUpdate::NoUpdatePending;
+	}
+}
+
+
+void UMeshSculptToolBase::SetFixedSculptPlaneFromWorldPos(const FVector& Position, const FVector& Normal, EPendingWorkPlaneUpdate UpdateType)
+{
+	if (UpdateType == EPendingWorkPlaneUpdate::MoveToHitPositionNormal)
+	{
+		UpdateFixedSculptPlanePosition(Position);
+		FFrame3d CurFrame(FVector::ZeroVector, GizmoProperties->Rotation);
+		CurFrame.AlignAxis(2, (FVector3d)Normal);
+		UpdateFixedSculptPlaneRotation((FQuat)CurFrame.Rotation);
+	}
+	else if (UpdateType == EPendingWorkPlaneUpdate::MoveToHitPositionViewAligned)
+	{
+		UpdateFixedSculptPlanePosition(Position);
+		FFrame3d CurFrame(FVector::ZeroVector, GizmoProperties->Rotation);
+		CurFrame.AlignAxis(2, -(FVector3d)CameraState.Forward());
+		UpdateFixedSculptPlaneRotation((FQuat)CurFrame.Rotation);
+	}
+	else
+	{
+		UpdateFixedSculptPlanePosition(Position);
+	}
+
+	if (PlaneTransformGizmo != nullptr)
+	{
+		PlaneTransformGizmo->SetNewGizmoTransform(FTransform(GizmoProperties->Rotation, GizmoProperties->Position));
+	}
+}
+
+void UMeshSculptToolBase::PlaneTransformChanged(UTransformProxy* Proxy, FTransform Transform)
+{
+	UpdateFixedSculptPlaneRotation(Transform.GetRotation());
+	UpdateFixedSculptPlanePosition(Transform.GetLocation());
+}
+
+void UMeshSculptToolBase::UpdateFixedSculptPlanePosition(const FVector& Position)
+{
+	GizmoProperties->Position = Position;
+	GizmoPositionWatcher.SilentUpdate();
+}
+
+void UMeshSculptToolBase::UpdateFixedSculptPlaneRotation(const FQuat& Rotation)
+{
+	GizmoProperties->Rotation = Rotation;
+	GizmoPositionWatcher.SilentUpdate();
+}
+
+void UMeshSculptToolBase::UpdateGizmoFromProperties()
+{
+	if (PlaneTransformGizmo != nullptr)
+	{
+		PlaneTransformGizmo->SetNewGizmoTransform(FTransform(GizmoProperties->Rotation, GizmoProperties->Position));
+	}
+}
+
+void UMeshSculptToolBase::UpdateFixedPlaneGizmoVisibility(bool bVisible)
+{
+	if (bVisible == false)
+	{
+		if (PlaneTransformGizmo != nullptr)
+		{
+			GetToolManager()->GetPairedGizmoManager()->DestroyGizmo(PlaneTransformGizmo);
+			PlaneTransformGizmo = nullptr;
+		}
+	}
+	else
+	{
+		if (PlaneTransformGizmo == nullptr)
+		{
+			PlaneTransformGizmo = GetToolManager()->GetPairedGizmoManager()->CreateCustomTransformGizmo(
+				ETransformGizmoSubElements::StandardTranslateRotate, this);
+			PlaneTransformGizmo->bUseContextCoordinateSystem = false;
+			PlaneTransformGizmo->CurrentCoordinateSystem = EToolContextCoordinateSystem::Local;
+			PlaneTransformGizmo->SetActiveTarget(PlaneTransformProxy, GetToolManager());
+			PlaneTransformGizmo->SetNewGizmoTransform(FTransform(GizmoProperties->Rotation, GizmoProperties->Position));
+		}
+
+		PlaneTransformGizmo->bSnapToWorldGrid = GizmoProperties->bSnapToGrid;
+	}
+}
+
+
+
+
+
+
+
+void UMeshSculptToolBase::RegisterActions(FInteractiveToolActionSet& ActionSet)
+{
+	ActionSet.RegisterAction(this, (int32)EStandardToolActions::IncreaseBrushSize,
+		TEXT("SculptIncreaseRadius"),
+		LOCTEXT("SculptIncreaseRadius", "Increase Sculpt Radius"),
+		LOCTEXT("SculptIncreaseRadiusTooltip", "Increase radius of sculpting brush"),
+		EModifierKey::None, EKeys::RightBracket,
+		[this]() { IncreaseBrushRadiusAction(); });
+
+	ActionSet.RegisterAction(this, (int32)EStandardToolActions::DecreaseBrushSize,
+		TEXT("SculptDecreaseRadius"),
+		LOCTEXT("SculptDecreaseRadius", "Decrease Sculpt Radius"),
+		LOCTEXT("SculptDecreaseRadiusTooltip", "Decrease radius of sculpting brush"),
+		EModifierKey::None, EKeys::LeftBracket,
+		[this]() { DecreaseBrushRadiusAction(); });
+
+	ActionSet.RegisterAction(this, (int32)EStandardToolActions::BaseClientDefinedActionID + 12,
+		TEXT("BrushIncreaseFalloff"),
+		LOCTEXT("BrushIncreaseFalloff", "Increase Brush Falloff"),
+		LOCTEXT("BrushIncreaseFalloffTooltip", "Press this key to increase brush falloff by a fixed increment."),
+		EModifierKey::Shift | EModifierKey::Control, EKeys::RightBracket,
+		[this]() { IncreaseBrushFalloffAction(); });
+
+	ActionSet.RegisterAction(this, (int32)EStandardToolActions::BaseClientDefinedActionID + 13,
+		TEXT("BrushDecreaseFalloff"),
+		LOCTEXT("BrushDecreaseFalloff", "Decrease Brush Falloff"),
+		LOCTEXT("BrushDecreaseFalloffTooltip", "Press this key to decrease brush falloff by a fixed increment."),
+		EModifierKey::Shift | EModifierKey::Control, EKeys::LeftBracket,
+		[this]() { DecreaseBrushFalloffAction(); });
+
+
+
+	ActionSet.RegisterAction(this, (int32)EStandardToolActions::BaseClientDefinedActionID + 1,
+		TEXT("NextBrushMode"),
+		LOCTEXT("SculptNextBrushMode", "Next Brush Type"),
+		LOCTEXT("SculptNextBrushModeTooltip", "Cycle to next Brush Type"),
+		EModifierKey::None, EKeys::A,
+		[this]() { NextBrushModeAction(); });
+
+	ActionSet.RegisterAction(this, (int32)EStandardToolActions::BaseClientDefinedActionID + 2,
+		TEXT("PreviousBrushMode"),
+		LOCTEXT("SculptPreviousBrushMode", "Previous Brush Type"),
+		LOCTEXT("SculptPreviousBrushModeTooltip", "Cycle to previous Brush Type"),
+		EModifierKey::None, EKeys::Q,
+		[this]() { PreviousBrushModeAction(); });
+
+
+	ActionSet.RegisterAction(this, (int32)EStandardToolActions::BaseClientDefinedActionID + 60,
+		TEXT("SculptIncreaseSpeed"),
+		LOCTEXT("SculptIncreaseSpeed", "Increase Speed"),
+		LOCTEXT("SculptIncreaseSpeedTooltip", "Increase Brush Speed"),
+		EModifierKey::None, EKeys::E,
+		[this]() { IncreaseBrushSpeedAction(); });
+
+	ActionSet.RegisterAction(this, (int32)EStandardToolActions::BaseClientDefinedActionID + 61,
+		TEXT("SculptDecreaseSpeed"),
+		LOCTEXT("SculptDecreaseSpeed", "Decrease Speed"),
+		LOCTEXT("SculptDecreaseSpeedTooltip", "Decrease Brush Speed"),
+		EModifierKey::None, EKeys::W,
+		[this]() { DecreaseBrushSpeedAction(); });
+
+
+
+	ActionSet.RegisterAction(this, (int32)EStandardToolActions::BaseClientDefinedActionID + 50,
+		TEXT("SculptIncreaseSize"),
+		LOCTEXT("SculptIncreaseSize", "Increase Size"),
+		LOCTEXT("SculptIncreaseSizeTooltip", "Increase Brush Size"),
+		EModifierKey::None, EKeys::D,
+		[this]() { IncreaseBrushRadiusAction(); });
+
+	ActionSet.RegisterAction(this, (int32)EStandardToolActions::BaseClientDefinedActionID + 51,
+		TEXT("SculptDecreaseSize"),
+		LOCTEXT("SculptDecreaseSize", "Decrease Size"),
+		LOCTEXT("SculptDecreaseSizeTooltip", "Decrease Brush Size"),
+		EModifierKey::None, EKeys::S,
+		[this]() { DecreaseBrushRadiusAction(); });
+
+	ActionSet.RegisterAction(this, (int32)EStandardToolActions::BaseClientDefinedActionID + 52,
+		TEXT("SculptIncreaseSizeSmallStep"),
+		LOCTEXT("SculptIncreaseSize", "Increase Size"),
+		LOCTEXT("SculptIncreaseSizeTooltip", "Increase Brush Size"),
+		EModifierKey::Shift, EKeys::D,
+		[this]() { IncreaseBrushRadiusSmallStepAction(); });
+
+	ActionSet.RegisterAction(this, (int32)EStandardToolActions::BaseClientDefinedActionID + 53,
+		TEXT("SculptDecreaseSizeSmallStemp"),
+		LOCTEXT("SculptDecreaseSize", "Decrease Size"),
+		LOCTEXT("SculptDecreaseSizeTooltip", "Decrease Brush Size"),
+		EModifierKey::Shift, EKeys::S,
+		[this]() { DecreaseBrushRadiusSmallStepAction(); });
+
+
+
+	ActionSet.RegisterAction(this, (int32)EStandardToolActions::ToggleWireframe,
+		TEXT("ToggleWireframe"),
+		LOCTEXT("ToggleWireframe", "Toggle Wireframe"),
+		LOCTEXT("ToggleWireframeTooltip", "Toggle visibility of wireframe overlay"),
+		EModifierKey::Alt, EKeys::W,
+		[this]() { ViewProperties->bShowWireframe = !ViewProperties->bShowWireframe; });
+
+
+
+	ActionSet.RegisterAction(this, (int32)EStandardToolActions::BaseClientDefinedActionID + 100,
+		TEXT("SetSculptWorkSurfacePosNormal"),
+		LOCTEXT("SetSculptWorkSurfacePosNormal", "Reorient Work Surface"),
+		LOCTEXT("SetSculptWorkSurfacePosNormalTooltip", "Move the Sculpting Work Plane/Surface to Position and Normal of World hit point under cursor"),
+		EModifierKey::Shift, EKeys::T,
+		[this]() { PendingWorkPlaneUpdate = EPendingWorkPlaneUpdate::MoveToHitPositionNormal; });
+
+	ActionSet.RegisterAction(this, (int32)EStandardToolActions::BaseClientDefinedActionID + 101,
+		TEXT("SetSculptWorkSurfacePos"),
+		LOCTEXT("SetSculptWorkSurfacePos", "Reposition Work Surface"),
+		LOCTEXT("SetSculptWorkSurfacePosTooltip", "Move the Sculpting Work Plane/Surface to World hit point under cursor (keep current Orientation)"),
+		EModifierKey::None, EKeys::T,
+		[this]() { PendingWorkPlaneUpdate = EPendingWorkPlaneUpdate::MoveToHitPosition; });
+
+	ActionSet.RegisterAction(this, (int32)EStandardToolActions::BaseClientDefinedActionID + 102,
+		TEXT("SetSculptWorkSurfaceView"),
+		LOCTEXT("SetSculptWorkSurfaceView", "View-Align Work Surface"),
+		LOCTEXT("SetSculptWorkSurfaceViewTooltip", "Move the Sculpting Work Plane/Surface to World hit point under cursor and align to View"),
+		EModifierKey::Control | EModifierKey::Shift, EKeys::T,
+		[this]() { PendingWorkPlaneUpdate = EPendingWorkPlaneUpdate::MoveToHitPositionViewAligned; });
+}
+
+
+#undef LOCTEXT_NAMESPACE
+
diff --git a/Engine/Plugins/Experimental/MeshModelingToolset/Source/MeshModelingTools/Private/Sculpting/MeshSculptUtil.h b/Engine/Plugins/Experimental/MeshModelingToolset/Source/MeshModelingTools/Private/Sculpting/MeshSculptUtil.h
new file mode 100644
--- /dev/null
+++ b/Engine/Plugins/Experimental/MeshModelingToolset/Source/MeshModelingTools/Private/Sculpting/MeshSculptUtil.h
@@ -0,0 +1,88 @@
+// Copyright Epic Games, Inc. All Rights Reserved.
+
+#pragma once
+
+#include "CoreMinimal.h"
+#include "DynamicMeshAABBTree3.h"
+#include "DynamicMeshOctree3.h"
+#include "MeshNormals.h"
+
+
+namespace UE { namespace SculptUtil {
+
+
+
+	void RecalculateNormals_Overlay(FDynamicMesh3* Mesh, const TSet<int32>& ModifiedTris, TSet<int32>& VertexSetBuffer, TArray<int32>& NormalsBuffer)
+	{
+		FDynamicMeshNormalOverlay* Normals = Mesh->HasAttributes() ? Mesh->Attributes()->PrimaryNormals() : nullptr;
+		check(Normals != nullptr);
+
+		NormalsBuffer.Reset();
+		VertexSetBuffer.Reset();
+		for (int32 TriangleID : ModifiedTris)
+		{
+			FIndex3i TriElems = Normals->GetTriangle(TriangleID);
+			for (int32 j = 0; j < 3; ++j)
+			{
+				int32 elemid = TriElems[j];
+				if (VertexSetBuffer.Contains(elemid) == false)
+				{
+					VertexSetBuffer.Add(elemid);
+					NormalsBuffer.Add(elemid);
+				}
+			}
+		}
+
+		ParallelFor(NormalsBuffer.Num(), [&](int32 k) {
+			int32 elemid = NormalsBuffer[k];
+			FVector3d NewNormal = FMeshNormals::ComputeOverlayNormal(*Mesh, Normals, elemid);
+			Normals->SetElement(elemid, (FVector3f)NewNormal);
+		});
+	}
+
+
+
+
+	void RecalculateNormals_PerVertex(FDynamicMesh3* Mesh, const TSet<int32>& ModifiedTris, TSet<int32>& VertexSetBuffer, TArray<int32>& NormalsBuffer)
+	{
+		NormalsBuffer.Reset();
+		VertexSetBuffer.Reset();
+		for (int32 TriangleID : ModifiedTris)
+		{
+			FIndex3i TriV = Mesh->GetTriangle(TriangleID);
+			for (int32 j = 0; j < 3; ++j)
+			{
+				int32 vid = TriV[j];
+				if (VertexSetBuffer.Contains(vid) == false)
+				{
+					VertexSetBuffer.Add(vid);
+					NormalsBuffer.Add(vid);
+				}
+			}
+		}
+
+		ParallelFor(NormalsBuffer.Num(), [&](int32 k) {
+			int32 vid = NormalsBuffer[k];
+			FVector3d NewNormal = FMeshNormals::ComputeVertexNormal(*Mesh, vid);
+			Mesh->SetVertexNormal(vid, (FVector3f)NewNormal);
+		});
+	}
+
+
+
+
+	void RecalculateROINormals(FDynamicMesh3* Mesh, const TSet<int32>& TriangleROI, TSet<int32>& VertexSetBuffer, TArray<int32>& NormalsBuffer, bool bForceVertex = false)
+	{
+		if (Mesh->HasAttributes() && Mesh->Attributes()->PrimaryNormals() && bForceVertex == false)
+		{
+			RecalculateNormals_Overlay(Mesh, TriangleROI, VertexSetBuffer, NormalsBuffer);
+		}
+		else
+		{
+			RecalculateNormals_PerVertex(Mesh, TriangleROI, VertexSetBuffer, NormalsBuffer);
+		}
+	}
+
+
+
+/* end namespace UE::SculptUtil */  } }
\ No newline at end of file
diff --git a/Engine/Plugins/Experimental/MeshModelingToolset/Source/MeshModelingTools/Private/SelfUnionMeshesTool.cpp b/Engine/Plugins/Experimental/MeshModelingToolset/Source/MeshModelingTools/Private/SelfUnionMeshesTool.cpp
new file mode 100644
--- /dev/null
+++ b/Engine/Plugins/Experimental/MeshModelingToolset/Source/MeshModelingTools/Private/SelfUnionMeshesTool.cpp
@@ -0,0 +1,351 @@
+// Copyright Epic Games, Inc. All Rights Reserved.
+
+#include "SelfUnionMeshesTool.h"
+#include "CompositionOps/SelfUnionMeshesOp.h"
+#include "InteractiveToolManager.h"
+#include "ToolBuilderUtil.h"
+
+#include "ToolSetupUtil.h"
+
+#include "Selection/ToolSelectionUtil.h"
+
+#include "DynamicMesh3.h"
+#include "DynamicMeshEditor.h"
+#include "MeshTransforms.h"
+
+#include "MeshDescriptionToDynamicMesh.h"
+#include "DynamicMeshToMeshDescription.h"
+
+#include "AssetGenerationUtil.h"
+
+
+#define LOCTEXT_NAMESPACE "USelfUnionMeshesTool"
+
+
+/*
+ * ToolBuilder
+ */
+
+
+bool USelfUnionMeshesToolBuilder::CanBuildTool(const FToolBuilderState& SceneState) const
+{
+	return AssetAPI != nullptr && ToolBuilderUtil::CountComponents(SceneState, CanMakeComponentTarget) > 0;
+}
+
+UInteractiveTool* USelfUnionMeshesToolBuilder::BuildTool(const FToolBuilderState& SceneState) const
+{
+	USelfUnionMeshesTool* NewTool = NewObject<USelfUnionMeshesTool>(SceneState.ToolManager);
+
+	TArray<UActorComponent*> Components = ToolBuilderUtil::FindAllComponents(SceneState, CanMakeComponentTarget);
+	check(Components.Num() > 0);
+
+	TArray<TUniquePtr<FPrimitiveComponentTarget>> ComponentTargets;
+	for (UActorComponent* ActorComponent : Components)
+	{
+		auto* MeshComponent = Cast<UPrimitiveComponent>(ActorComponent);
+		if ( MeshComponent )
+		{
+			ComponentTargets.Add(MakeComponentTarget(MeshComponent));
+		}
+	}
+
+	NewTool->SetSelection(MoveTemp(ComponentTargets));
+	NewTool->SetWorld(SceneState.World);
+	NewTool->SetAssetAPI(AssetAPI);
+
+	return NewTool;
+}
+
+
+
+/*
+ * Tool
+ */
+USelfUnionMeshesTool::USelfUnionMeshesTool()
+{
+}
+
+void USelfUnionMeshesTool::SetWorld(UWorld* World)
+{
+	this->TargetWorld = World;
+}
+
+void USelfUnionMeshesTool::Setup()
+{
+	UInteractiveTool::Setup();
+
+	// hide input StaticMeshComponents
+	for (TUniquePtr<FPrimitiveComponentTarget>& ComponentTarget : ComponentTargets)
+	{
+		ComponentTarget->SetOwnerVisibility(false);
+	}
+
+
+	// initialize our properties
+	Properties = NewObject<USelfUnionMeshesToolProperties>(this);
+	Properties->RestoreProperties(this);
+	AddToolPropertySource(Properties);
+
+
+	// initialize the PreviewMesh+BackgroundCompute object
+	SetupPreview();
+
+	Preview->InvalidateResult();
+}
+
+
+
+
+void USelfUnionMeshesTool::SetupPreview()
+{
+	FComponentMaterialSet AllMaterialSet;
+	TMap<UMaterialInterface*, int> KnownMaterials;
+	TArray<TArray<int>> MaterialRemap; MaterialRemap.SetNum(ComponentTargets.Num());
+	for (int ComponentIdx = 0; ComponentIdx < ComponentTargets.Num(); ComponentIdx++)
+	{
+		FComponentMaterialSet ComponentMaterialSet;
+		ComponentTargets[ComponentIdx]->GetMaterialSet(ComponentMaterialSet);
+		for (UMaterialInterface* Mat : ComponentMaterialSet.Materials)
+		{
+			int* FoundMatIdx = KnownMaterials.Find(Mat);
+			int MatIdx;
+			if (FoundMatIdx)
+			{
+				MatIdx = *FoundMatIdx;
+			}
+			else
+			{
+				MatIdx = AllMaterialSet.Materials.Add(Mat);
+				KnownMaterials.Add(Mat, MatIdx);
+			}
+			MaterialRemap[ComponentIdx].Add(MatIdx);
+		}
+	}
+
+	CombinedSourceMeshes = MakeShared<FDynamicMesh3>();
+	CombinedSourceMeshes->EnableAttributes();
+	CombinedSourceMeshes->Attributes()->EnableMaterialID();
+	FDynamicMeshEditor AppendEditor(CombinedSourceMeshes.Get());
+
+	for (int ComponentIdx = 0; ComponentIdx < ComponentTargets.Num(); ComponentIdx++)
+	{
+		FDynamicMesh3 ComponentMesh;
+		FMeshDescriptionToDynamicMesh Converter;
+		Converter.Convert(ComponentTargets[ComponentIdx]->GetMesh(), ComponentMesh);
+
+		// ensure materials and attributes are always enabled
+		ComponentMesh.EnableAttributes();
+		ComponentMesh.Attributes()->EnableMaterialID();
+		FDynamicMeshMaterialAttribute* MaterialIDs = ComponentMesh.Attributes()->GetMaterialID();
+		for (int TID : ComponentMesh.TriangleIndicesItr())
+		{
+			MaterialIDs->SetValue(TID, MaterialRemap[ComponentIdx][MaterialIDs->GetValue(TID)]);
+		}
+		// TODO: center the meshes
+		FTransform3d WorldTransform = (FTransform3d)ComponentTargets[ComponentIdx]->GetWorldTransform();
+		FMeshIndexMappings IndexMaps;
+		AppendEditor.AppendMesh(&ComponentMesh, IndexMaps,
+			[WorldTransform](int VID, const FVector3d& Pos)
+			{
+				return WorldTransform.TransformPosition(Pos);
+			},
+			[WorldTransform](int VID, const FVector3d& Normal)
+			{
+				return WorldTransform.TransformNormal(Normal);
+			}
+		);
+	}
+
+	Preview = NewObject<UMeshOpPreviewWithBackgroundCompute>(this, "Preview");
+	Preview->Setup(this->TargetWorld, this);
+	Preview->PreviewMesh->UpdatePreview(CombinedSourceMeshes.Get());
+	Preview->ConfigureMaterials(AllMaterialSet.Materials, ToolSetupUtil::GetDefaultWorkingMaterial(GetToolManager()));
+
+
+	DrawnLineSet = NewObject<ULineSetComponent>(Preview->PreviewMesh->GetRootComponent());
+	DrawnLineSet->SetupAttachment(Preview->PreviewMesh->GetRootComponent());
+	DrawnLineSet->SetLineMaterial(ToolSetupUtil::GetDefaultLineComponentMaterial(GetToolManager()));
+	DrawnLineSet->RegisterComponent();
+
+	Preview->OnOpCompleted.AddLambda(
+		[this](const FDynamicMeshOperator* Op)
+		{
+			const FSelfUnionMeshesOp* UnionOp = (const FSelfUnionMeshesOp*)(Op);
+			CreatedBoundaryEdges = UnionOp->GetCreatedBoundaryEdges();
+		}
+	);
+	Preview->OnMeshUpdated.AddLambda(
+		[this](const UMeshOpPreviewWithBackgroundCompute*)
+		{
+			GetToolManager()->PostInvalidation();
+			UpdateVisualization();
+		}
+	);
+}
+
+
+void USelfUnionMeshesTool::UpdateVisualization()
+{
+	FColor BoundaryEdgeColor(240, 15, 15);
+	float BoundaryEdgeThickness = 2.0;
+	float BoundaryEdgeDepthBias = 2.0f;
+
+	const FDynamicMesh3* TargetMesh = Preview->PreviewMesh->GetPreviewDynamicMesh();
+	FVector3d A, B;
+
+	DrawnLineSet->Clear();
+	if (Properties->bShowNewBoundaryEdges)
+	{
+		for (int EID : CreatedBoundaryEdges)
+		{
+			TargetMesh->GetEdgeV(EID, A, B);
+			DrawnLineSet->AddLine((FVector)A, (FVector)B, BoundaryEdgeColor, BoundaryEdgeThickness, BoundaryEdgeDepthBias);
+		}
+	}
+}
+
+
+void USelfUnionMeshesTool::Shutdown(EToolShutdownType ShutdownType)
+{
+	Properties->SaveProperties(this);
+
+	FDynamicMeshOpResult Result = Preview->Shutdown();
+	if (ShutdownType == EToolShutdownType::Accept)
+	{
+		// Generate the result
+		{
+			GetToolManager()->BeginUndoTransaction(LOCTEXT("SelfUnionMeshes", "Boolean Meshes"));
+
+			GenerateAsset(Result);
+
+			GetToolManager()->EndUndoTransaction();
+		}
+
+		// Hide or destroy the sources
+		{
+
+			if (Properties->bDeleteInputActors)
+			{
+				GetToolManager()->BeginUndoTransaction(LOCTEXT("RemoveSources", "Remove Sources"));
+			}
+
+			for (auto& ComponentTarget : ComponentTargets)
+			{
+				ComponentTarget->SetOwnerVisibility(true);
+				AActor* Actor = ComponentTarget->GetOwnerActor();
+				if (Properties->bDeleteInputActors)
+				{
+					Actor->Destroy();
+				}
+				else
+				{
+#if WITH_EDITOR
+					Actor->SetIsTemporarilyHiddenInEditor(true);
+#endif
+				}
+			}
+
+			if (Properties->bDeleteInputActors)
+			{
+				GetToolManager()->EndUndoTransaction();
+			}
+
+		}
+
+	}
+	else
+	{
+		// Restore (unhide) the source meshes
+		for (auto& ComponentTarget : ComponentTargets)
+		{
+			ComponentTarget->SetOwnerVisibility(true);
+		}
+	}
+}
+
+void USelfUnionMeshesTool::SetAssetAPI(IToolsContextAssetAPI* AssetAPIIn)
+{
+	this->AssetAPI = AssetAPIIn;
+}
+
+TUniquePtr<FDynamicMeshOperator> USelfUnionMeshesTool::MakeNewOperator()
+{
+	TUniquePtr<FSelfUnionMeshesOp> Op = MakeUnique<FSelfUnionMeshesOp>();
+	
+	Op->bAttemptFixHoles = Properties->bAttemptFixHoles;
+
+	Op->SetResultTransform(FTransform3d::Identity()); // TODO Center the combined meshes (when building them) and change this transform accordingly
+	Op->CombinedMesh = CombinedSourceMeshes;
+
+	return Op;
+}
+
+
+
+void USelfUnionMeshesTool::Render(IToolsContextRenderAPI* RenderAPI)
+{
+}
+
+void USelfUnionMeshesTool::OnTick(float DeltaTime)
+{
+	Preview->Tick(DeltaTime);
+}
+
+
+#if WITH_EDITOR
+void USelfUnionMeshesTool::PostEditChangeProperty(FPropertyChangedEvent& PropertyChangedEvent)
+{
+	Preview->InvalidateResult();
+}
+#endif
+
+void USelfUnionMeshesTool::OnPropertyModified(UObject* PropertySet, FProperty* Property)
+{
+	if (Property && (Property->GetFName() == GET_MEMBER_NAME_CHECKED(USelfUnionMeshesToolProperties, bDeleteInputActors)))
+	{
+		// nothing
+	}
+	else if (Property && (Property->GetFName() == GET_MEMBER_NAME_CHECKED(USelfUnionMeshesToolProperties, bShowNewBoundaryEdges)))
+	{
+		GetToolManager()->PostInvalidation();
+		UpdateVisualization();
+	}
+	else
+	{
+		Preview->InvalidateResult();
+	}
+}
+
+
+bool USelfUnionMeshesTool::HasAccept() const
+{
+	return true;
+}
+
+bool USelfUnionMeshesTool::CanAccept() const
+{
+	return Preview->HaveValidResult();
+}
+
+
+void USelfUnionMeshesTool::GenerateAsset(const FDynamicMeshOpResult& Result)
+{
+	check(Result.Mesh.Get() != nullptr);
+
+	FVector3d Center = Result.Mesh->GetCachedBounds().Center();
+	MeshTransforms::Translate(*Result.Mesh, -Center);
+	FTransform3d CenteredTransform = Result.Transform;
+	CenteredTransform.SetTranslation(CenteredTransform.GetTranslation() + Result.Transform.TransformVector(Center));
+	AActor* NewActor = AssetGenerationUtil::GenerateStaticMeshActor(
+		AssetAPI, TargetWorld,
+		Result.Mesh.Get(), CenteredTransform, TEXT("Merged Mesh"), Preview->StandardMaterials);
+	if (NewActor != nullptr)
+	{
+		ToolSelectionUtil::SetNewActorSelection(GetToolManager(), NewActor);
+	}
+}
+
+
+
+
+#undef LOCTEXT_NAMESPACE
diff --git a/Engine/Plugins/Experimental/MeshModelingToolset/Source/MeshModelingTools/Private/SmoothMeshTool.cpp b/Engine/Plugins/Experimental/MeshModelingToolset/Source/MeshModelingTools/Private/SmoothMeshTool.cpp
--- a/Engine/Plugins/Experimental/MeshModelingToolset/Source/MeshModelingTools/Private/SmoothMeshTool.cpp
+++ b/Engine/Plugins/Experimental/MeshModelingToolset/Source/MeshModelingTools/Private/SmoothMeshTool.cpp
@@ -155,7 +155,7 @@ void USmoothMeshTool::OnPropertyModified(UObject* PropertySet, FProperty* Proper
 }
 
 
-void USmoothMeshTool::Tick(float DeltaTime)
+void USmoothMeshTool::OnTick(float DeltaTime)
 {
 	Preview->Tick(DeltaTime);
 }
diff --git a/Engine/Plugins/Experimental/MeshModelingToolset/Source/MeshModelingTools/Private/TransformMeshesTool.cpp b/Engine/Plugins/Experimental/MeshModelingToolset/Source/MeshModelingTools/Private/TransformMeshesTool.cpp
--- a/Engine/Plugins/Experimental/MeshModelingToolset/Source/MeshModelingTools/Private/TransformMeshesTool.cpp
+++ b/Engine/Plugins/Experimental/MeshModelingToolset/Source/MeshModelingTools/Private/TransformMeshesTool.cpp
@@ -99,7 +99,7 @@ void UTransformMeshesTool::Shutdown(EToolShutdownType ShutdownType)
 
 
 
-void UTransformMeshesTool::Tick(float DeltaTime)
+void UTransformMeshesTool::OnTick(float DeltaTime)
 {
 	// make sure state is up-to-date
 
diff --git a/Engine/Plugins/Experimental/MeshModelingToolset/Source/MeshModelingTools/Private/UVProjectionTool.cpp b/Engine/Plugins/Experimental/MeshModelingToolset/Source/MeshModelingTools/Private/UVProjectionTool.cpp
--- a/Engine/Plugins/Experimental/MeshModelingToolset/Source/MeshModelingTools/Private/UVProjectionTool.cpp
+++ b/Engine/Plugins/Experimental/MeshModelingToolset/Source/MeshModelingTools/Private/UVProjectionTool.cpp
@@ -76,29 +76,6 @@ UUVProjectionToolProperties::UUVProjectionToolProperties()
 	UVOffset = FVector2D::ZeroVector;
 }
 
-void UUVProjectionToolProperties::SaveProperties(UInteractiveTool* SaveFromTool)
-{
-	UUVProjectionToolProperties* PropertyCache = GetPropertyCache<UUVProjectionToolProperties>();
-	PropertyCache->UVProjectionMethod = this->UVProjectionMethod;
-	PropertyCache->ProjectionPrimitiveScale = this->ProjectionPrimitiveScale;
-	PropertyCache->UVScale = this->UVScale;
-	PropertyCache->UVOffset = this->UVOffset;
-	PropertyCache->bWorldSpaceUVScale = this->bWorldSpaceUVScale;
-	PropertyCache->CylinderProjectToTopOrBottomAngleThreshold = this->CylinderProjectToTopOrBottomAngleThreshold;
-}
-
-void UUVProjectionToolProperties::RestoreProperties(UInteractiveTool* RestoreToTool)
-{
-	UUVProjectionToolProperties* PropertyCache = GetPropertyCache<UUVProjectionToolProperties>();
-	this->UVProjectionMethod = PropertyCache->UVProjectionMethod;
-	this->ProjectionPrimitiveScale = PropertyCache->ProjectionPrimitiveScale;
-	this->UVScale = PropertyCache->UVScale;
-	this->UVOffset = PropertyCache->UVOffset;
-	this->bWorldSpaceUVScale = PropertyCache->bWorldSpaceUVScale;
-	this->CylinderProjectToTopOrBottomAngleThreshold = PropertyCache->CylinderProjectToTopOrBottomAngleThreshold;
-}
-
-
 UUVProjectionAdvancedProperties::UUVProjectionAdvancedProperties()
 {
 }
@@ -295,7 +272,7 @@ void UUVProjectionTool::Render(IToolsContextRenderAPI* RenderAPI)
 	ProjectionShapeVisualizer.EndFrame();
 }
 
-void UUVProjectionTool::Tick(float DeltaTime)
+void UUVProjectionTool::OnTick(float DeltaTime)
 {
 	for (UMeshOpPreviewWithBackgroundCompute* Preview : Previews)
 	{
diff --git a/Engine/Plugins/Experimental/MeshModelingToolset/Source/MeshModelingTools/Public/AddPatchTool.h b/Engine/Plugins/Experimental/MeshModelingToolset/Source/MeshModelingTools/Public/AddPatchTool.h
--- a/Engine/Plugins/Experimental/MeshModelingToolset/Source/MeshModelingTools/Public/AddPatchTool.h
+++ b/Engine/Plugins/Experimental/MeshModelingToolset/Source/MeshModelingTools/Public/AddPatchTool.h
@@ -58,9 +58,6 @@ public:
 	/** Rotation around up axis */
 	UPROPERTY(EditAnywhere, Category = PatchSettings, meta = (DisplayName = "Shift", UIMin = "-1000", UIMax = "1000"))
 	float Shift;
-
-	virtual void SaveProperties(UInteractiveTool* SaveFromTool) override;
-	virtual void RestoreProperties(UInteractiveTool* RestoreToTool) override;
 };
 
 
@@ -85,7 +82,7 @@ public:
 	virtual void Shutdown(EToolShutdownType ShutdownType) override;
 
 	virtual void Render(IToolsContextRenderAPI* RenderAPI) override;
-	virtual void Tick(float DeltaTime) override;
+	virtual void OnTick(float DeltaTime) override;
 
 	virtual bool HasCancel() const override { return false; }
 	virtual bool HasAccept() const override { return false; }
diff --git a/Engine/Plugins/Experimental/MeshModelingToolset/Source/MeshModelingTools/Public/AddPrimitiveTool.h b/Engine/Plugins/Experimental/MeshModelingToolset/Source/MeshModelingTools/Public/AddPrimitiveTool.h
--- a/Engine/Plugins/Experimental/MeshModelingToolset/Source/MeshModelingTools/Public/AddPrimitiveTool.h
+++ b/Engine/Plugins/Experimental/MeshModelingToolset/Source/MeshModelingTools/Public/AddPrimitiveTool.h
@@ -2,19 +2,21 @@
 
 #pragma once
 
+#include "BaseTools/SingleClickTool.h"
 #include "CoreMinimal.h"
-#include "MaterialShared.h"
-#include "UObject/NoExportTypes.h"
 #include "InteractiveToolBuilder.h"
-#include "BaseTools/SingleClickTool.h"
+#include "MaterialShared.h"
 #include "PreviewMesh.h"
 #include "Properties/MeshMaterialProperties.h"
+#include "UObject/NoExportTypes.h"
+#include "UObject/UObjectGlobals.h"
+
 #include "AddPrimitiveTool.generated.h"
 
 class FDynamicMesh3;
 
 /**
- *
+ * Builder
  */
 UCLASS()
 class MESHMODELINGTOOLS_API UAddPrimitiveToolBuilder : public UInteractiveToolBuilder
@@ -22,44 +24,39 @@ class MESHMODELINGTOOLS_API UAddPrimitiveToolBuilder : public UInteractiveToolBu
 	GENERATED_BODY()
 
 public:
+	enum class EMakeMeshShapeType : uint32
+	{
+		Box,
+		Cylinder,
+		Cone,
+		Arrow,
+		Rectangle,
+		RoundedRectangle,
+		Disc,
+		PuncturedDisc,
+		Torus,
+		SphericalBox,
+		Sphere
+	};
+
 	IToolsContextAssetAPI* AssetAPI;
 
-	UAddPrimitiveToolBuilder() 
+	UAddPrimitiveToolBuilder()
 	{
 		AssetAPI = nullptr;
 	}
-
 	virtual bool CanBuildTool(const FToolBuilderState& SceneState) const override;
 	virtual UInteractiveTool* BuildTool(const FToolBuilderState& SceneState) const override;
-};
 
-/** Shape Types */
-UENUM()
-enum class EMakeMeshShapeType : uint32
-{
-	None			 = 0x000 UMETA(DisplayName = "None", Hidden),
-	All				 = 0xfff UMETA(DisplayName = "All", Hidden), 
-	Box				 = 0x001 UMETA(DisplayName = "Box"),
-	Cylinder		 = 0x002 UMETA(DisplayName = "Cylinder"),
-	Cone			 = 0x004 UMETA(DisplayName = "Cone"),
-	Arrow			 = 0x008 UMETA(DisplayName = "Arrow"),
-	Rectangle		 = 0x010 UMETA(DisplayName = "Rectangle"),
-	RoundedRectangle = 0x020 UMETA(DisplayName = "Rounded Rectangle"),
-	Disc			 = 0x040 UMETA(DisplayName = "Disc"),
-	PuncturedDisc	 = 0x080 UMETA(DisplayName = "Punctured Disc"),
-	Torus			 = 0x100 UMETA(DisplayName = "Torus"),
-	SphericalBox     = 0x200 UMETA(DisplayName = "Sphere Type 1"),
-	Sphere			 = 0x400 UMETA(DisplayName = "Sphere Type 2")
-	
+	EMakeMeshShapeType ShapeType{EMakeMeshShapeType::Box};
 };
-ENUM_CLASS_FLAGS(EMakeMeshShapeType);
 
 /** Placement Target Types */
 UENUM()
 enum class EMakeMeshPlacementType : uint8
 {
 	GroundPlane = 0,
-	OnScene = 1
+	OnScene     = 1
 };
 
 /** Placement Pivot Location */
@@ -71,7 +68,6 @@ enum class EMakeMeshPivotLocation : uint8
 	Top
 };
 
-
 /** Polygroup mode for primitive */
 UENUM()
 enum class EMakeMeshPolygroupMode : uint8
@@ -84,64 +80,33 @@ enum class EMakeMeshPolygroupMode : uint8
 	PerQuad
 };
 
-
 UCLASS()
 class MESHMODELINGTOOLS_API UProceduralShapeToolProperties : public UInteractiveToolPropertySet
 {
 	GENERATED_BODY()
 
 public:
-	UProceduralShapeToolProperties();
-
-	// UObject interface
-#if WITH_EDITOR
-	virtual bool CanEditChange( const FProperty* InProperty ) const override;
-#endif // WITH_EDITOR	
-	// End of UObject interface
+	UProceduralShapeToolProperties() = default;
 
 	/** If the shape settings haven't changed, create instances of the last created asset rather than creating a whole new asset.  If false, all created actors will have separate underlying mesh assets. */
 	UPROPERTY(EditAnywhere, Category = AssetSettings)
 	bool bInstanceIfPossible = true;
 
-	/** Type of shape to generate */
-	UPROPERTY(EditAnywhere, Category = ShapeSettings, meta = (DisplayName = "Primitive Type"))
-	EMakeMeshShapeType Shape;
-
-	/** Width of Shape */
-	UPROPERTY(EditAnywhere, Category = ShapeSettings, meta = (DisplayName = "Width", UIMin = "1.0", UIMax = "1000.0", ClampMin = "0.0001", ClampMax = "1000000.0"))
-	float Width;
-
-	/** Height of Shape */
-	UPROPERTY(EditAnywhere, Category = ShapeSettings, meta = (DisplayName = "Height", UIMin = "1.0", UIMax = "1000.0", ClampMin = "0.0001", ClampMax = "1000000.0"))
-	float Height;
-
-	/** Radius of additional circular features of the shape (not implicitly defined by the width of the shape) */
-	UPROPERTY(EditAnywhere, Category = ShapeSettings, meta = (DisplayName = "Radius", UIMin = "1.0", UIMax = "1000.0", ClampMin = "0.0001", ClampMax = "1000000.0"))
-	float FeatureRadius;
-
-	/** Number of Slices */
-	UPROPERTY(EditAnywhere, Category = ShapeSettings, meta = (DisplayName = "Slices", UIMin = "3", UIMax = "128", ClampMin = "3", ClampMax = "999"))
-	int Slices;
-
-	/** Number of Subdivisions */
-	UPROPERTY(EditAnywhere, Category = ShapeSettings, meta = (DisplayName = "Subdivisions", UIMin = "0", UIMax = "100", ClampMin = "0", ClampMax = "4000"))
-	int Subdivisions;
-
 	/** How should Polygroups be assigned to triangles of Primitive */
-	UPROPERTY(EditAnywhere, Category = ShapeSettings)
-	EMakeMeshPolygroupMode PolygroupMode;
+	UPROPERTY(EditAnywhere, Category = ShapeSettings, meta = (ProceduralShapeSetting))
+	EMakeMeshPolygroupMode PolygroupMode = EMakeMeshPolygroupMode::PerFace;
 
 	/** How to place Primitive in the Scene */
 	UPROPERTY(EditAnywhere, Category = Positioning, meta = (DisplayName = "Target Surface"))
-	EMakeMeshPlacementType PlaceMode;
+	EMakeMeshPlacementType PlaceMode = EMakeMeshPlacementType::OnScene;
 
 	/** Location of Pivot within Primitive */
-	UPROPERTY(EditAnywhere, Category = Positioning)
-    EMakeMeshPivotLocation PivotLocation;
+	UPROPERTY(EditAnywhere, Category = Positioning, meta = (ProceduralShapeSetting))
+	EMakeMeshPivotLocation PivotLocation = EMakeMeshPivotLocation::Base;
 
 	/** Rotation of Primitive around up axis */
 	UPROPERTY(EditAnywhere, Category = Positioning, meta = (DisplayName = "Rotation", UIMin = "0.0", UIMax = "360.0"))
-	float Rotation;
+	float Rotation = 0.0;
 
 	/** Align shape to Placement Surface */
 	UPROPERTY(EditAnywhere, Category = Positioning, meta = (DisplayName = "Align to Normal", EditCondition = "PlaceMode == EMakeMeshPlacementType::OnScene"))
@@ -149,19 +114,225 @@ public:
 
 	///** Start Angle of Shape */
 	//UPROPERTY(EditAnywhere, Category = ShapeSettings, meta = (DisplayName = "Start Angle", UIMin = "0.0", UIMax = "360.0", ClampMin = "-10000", ClampMax = "10000.0"))
-	//float StartAngle;
+	//float StartAngle = 0.f;
 
 	///** End Angle of Shape */
 	//UPROPERTY(EditAnywhere, Category = ShapeSettings, meta = (DisplayName = "End Angle", UIMin = "0.0", UIMax = "360.0", ClampMin = "-10000", ClampMax = "10000.0"))
-	//float EndAngle;
+	//float EndAngle = 360.f;
+
+	bool IsEquivalent( const UProceduralShapeToolProperties* ) const;
+};
+
+UCLASS()
+class MESHMODELINGTOOLS_API UProceduralRectangleToolProperties : public UProceduralShapeToolProperties
+{
+	GENERATED_BODY()
+public:
+	UProceduralRectangleToolProperties() = default;
+
+	/** Width of Shape */
+	UPROPERTY(EditAnywhere, Category = "ShapeSettings|Dimensions", meta = (DisplayName = "Width", UIMin = "1.0", UIMax = "1000.0", ClampMin = "0.0001", ClampMax = "1000000.0", ProceduralShapeSetting))
+	float Width = 100.f;
+
+	/** Depth of Shape */
+	UPROPERTY(EditAnywhere, Category = "ShapeSettings|Dimensions", meta = (DisplayName = "Depth", UIMin = "1.0", UIMax = "1000.0", ClampMin = "0.0001", ClampMax = "1000000.0", ProceduralShapeSetting))
+	float Depth = 100.f;
+
+	/** Number of Subdivisions Along the Width */
+	UPROPERTY(EditAnywhere, Category = "ShapeSettings|Subdivisions", meta = (DisplayName = "Width", UIMin = "1", UIMax = "100", ClampMin = "1", ClampMax = "4000", ProceduralShapeSetting))
+	int WidthSubdivisions = 1;
+
+	/** Number of Subdivisions Along the Depth*/
+	UPROPERTY(EditAnywhere, Category = "ShapeSettings|Subdivisions", meta = (DisplayName = "Depth", UIMin = "1", UIMax = "100", ClampMin = "1", ClampMax = "4000", ProceduralShapeSetting))
+	int DepthSubdivisions = 1;
+};
+
+UCLASS()
+class MESHMODELINGTOOLS_API UProceduralBoxToolProperties : public UProceduralRectangleToolProperties
+{
+	GENERATED_BODY()
+public:
+	UProceduralBoxToolProperties() = default;
+
+	/** Height of Shape */
+	UPROPERTY(EditAnywhere, Category = "ShapeSettings|Dimensions", meta = (DisplayName = "Height", UIMin = "1.0", UIMax = "1000.0", ClampMin = "0.0001", ClampMax = "1000000.0", ProceduralShapeSetting))
+	float Height = 100.f;
+
+	/** Number of Subdivisions Along the Height */
+	UPROPERTY(EditAnywhere, Category = "ShapeSettings|Subdivisions", meta = (DisplayName = "Height", UIMin = "0", UIMax = "100", ClampMin = "0", ClampMax = "4000", ProceduralShapeSetting))
+	int HeightSubdivisions = 1;
+};
+
+UCLASS()
+class MESHMODELINGTOOLS_API UProceduralRoundedRectangleToolProperties : public UProceduralRectangleToolProperties
+{
+	GENERATED_BODY()
+public:
+	/** Radius of Rounded Corners */
+	UPROPERTY(EditAnywhere, Category = "ShapeSettings|Dimensions", meta = (DisplayName = "Corner Radius", UIMin = "1.0", UIMax = "1000.0", ClampMin = "0.0001", ClampMax = "1000000.0", ProceduralShapeSetting))
+	float CornerRadius = 25.f;
+
+	/** Number of Angular Slices in Each Corner*/
+	UPROPERTY(EditAnywhere, Category = "ShapeSettings|Slices", meta = (DisplayName = "Corner", UIMin = "3", UIMax = "128", ClampMin = "3", ClampMax = "999", ProceduralShapeSetting))
+	int CornerSlices = 16;
+};
+
+UCLASS()
+class MESHMODELINGTOOLS_API UProceduralDiscToolProperties : public UProceduralShapeToolProperties
+{
+	GENERATED_BODY()
+public:
+	/** Radius of Disc */
+	UPROPERTY(EditAnywhere, Category = "ShapeSettings|Dimensions", meta = (DisplayName = "Radius", UIMin = "1.0", UIMax = "1000.0", ClampMin = "0.0001", ClampMax = "1000000.0", ProceduralShapeSetting))
+	float Radius = 50.f;
+
+	/** Number of Angular Slices around the Disc */
+	UPROPERTY(EditAnywhere, Category = "ShapeSettings|Slices", meta = (DisplayName = "Radial", UIMin = "3", UIMax = "128", ClampMin = "3", ClampMax = "999", ProceduralShapeSetting))
+	int RadialSlices = 16;
+
+	/** Number of Radial Subdivisions in the Disc */
+	UPROPERTY(EditAnywhere, Category = "ShapeSettings|Subdivisions", meta = (DisplayName = "Radial", UIMin = "1", UIMax = "100", ClampMin = "1", ClampMax = "4000", ProceduralShapeSetting))
+	int RadialSubdivisions = 1;
+};
+
+UCLASS()
+class MESHMODELINGTOOLS_API UProceduralPuncturedDiscToolProperties : public UProceduralDiscToolProperties
+{
+	GENERATED_BODY()
+public:
+	/** Radius of the Disc's Hole */
+	UPROPERTY(EditAnywhere, Category = "ShapeSettings|Dimensions", meta = (DisplayName = "Hole Radius", UIMin = "1.0", UIMax = "1000.0", ClampMin = "0.0001", ClampMax = "1000000.0", ProceduralShapeSetting))
+	float HoleRadius = 25.f;
+};
+
+UCLASS()
+class MESHMODELINGTOOLS_API UProceduralTorusToolProperties : public UProceduralShapeToolProperties
+{
+	GENERATED_BODY()
+public:
+	/** Radius from the Torus Center to the Center of the Torus Tube */
+	UPROPERTY(EditAnywhere, Category = "ShapeSettings|Dimensions", meta = (DisplayName = "Major Radius", UIMin = "1.0", UIMax = "1000.0", ClampMin = "0.0001", ClampMax = "1000000.0", ProceduralShapeSetting))
+	float MajorRadius = 50.f;
+
+	/** Radius of the Torus Tube */
+	UPROPERTY(EditAnywhere, Category = "ShapeSettings|Dimensions", meta = (DisplayName = "Minor Radius", UIMin = "1.0", UIMax = "1000.0", ClampMin = "0.0001", ClampMax = "1000000.0", ProceduralShapeSetting))
+	float MinorRadius = 25.f;
+
+	/** Number of Angular Slices Along the Torus Tube */
+	UPROPERTY(EditAnywhere, Category = "ShapeSettings|Slices", meta = (DisplayName = "Major", UIMin = "3", UIMax = "128", ClampMin = "3", ClampMax = "999", ProceduralShapeSetting))
+	int TubeSlices = 16;
+
+	/** Number of Angular Slices Around the Tube of the Torus */
+	UPROPERTY(EditAnywhere, Category = "ShapeSettings|Slices", meta = (DisplayName = "Minor", UIMin = "3", UIMax = "128", ClampMin = "3", ClampMax = "999", ProceduralShapeSetting))
+	int CrossSectionSlices = 16;
+};
+
+UCLASS()
+class MESHMODELINGTOOLS_API UProceduralCylinderToolProperties : public UProceduralShapeToolProperties
+{
+	GENERATED_BODY()
+public:
+
+	/** Radius of The Cylinder */
+	UPROPERTY(EditAnywhere, Category = "ShapeSettings|Dimensions", meta = (DisplayName = "Radius", UIMin = "1.0", UIMax = "1000.0", ClampMin = "0.0001", ClampMax = "1000000.0", ProceduralShapeSetting))
+	float Radius = 50.f;
+
+	/** Height of Cylinder */
+	UPROPERTY(EditAnywhere, Category = "ShapeSettings|Dimensions", meta = (DisplayName = "Height", UIMin = "1.0", UIMax = "1000.0", ClampMin = "0.0001", ClampMax = "1000000.0", ProceduralShapeSetting))
+	float Height = 200.f;
 
+	/** Number of Slices on the Cylinder Caps */
+	UPROPERTY(EditAnywhere, Category = "ShapeSettings|Slices", meta = (DisplayName = "Radial", UIMin = "3", UIMax = "128", ClampMin = "3", ClampMax = "999", ProceduralShapeSetting))
+	int RadialSlices = 16;
+
+	/** Number of Vertical Subdivisions Along the Height of the Cylidner */
+	UPROPERTY(EditAnywhere, Category = "ShapeSettings|Subdivisions", meta = (DisplayName = "Height", UIMin = "1", UIMax = "100", ClampMin = "1", ClampMax = "4000", ProceduralShapeSetting))
+	int HeightSubdivisions = 1;
+};
+
+UCLASS()
+class MESHMODELINGTOOLS_API UProceduralConeToolProperties : public UProceduralShapeToolProperties
+{
+	GENERATED_BODY()
+public:
+	/** Radius of the Cone */
+	UPROPERTY(EditAnywhere, Category = "ShapeSettings|Dimensions", meta = (DisplayName = "Radius", UIMin = "1.0", UIMax = "1000.0", ClampMin = "0.0001", ClampMax = "1000000.0", ProceduralShapeSetting))
+	float Radius = 50.f;
+
+	/** Height of Cone */
+	UPROPERTY(EditAnywhere, Category = "ShapeSettings|Dimensions", meta = (DisplayName = "Height", UIMin = "1.0", UIMax = "1000.0", ClampMin = "0.0001", ClampMax = "1000000.0", ProceduralShapeSetting))
+	float Height = 200.f;
+
+	/** Number of Slices on the Cone Base */
+	UPROPERTY(EditAnywhere, Category = "ShapeSettings|Slices", meta = (DisplayName = "Radial", UIMin = "3", UIMax = "128", ClampMin = "3", ClampMax = "999", ProceduralShapeSetting))
+	int RadialSlices = 16;
+
+	/** Number of Vertical Subdivisions Along the Hight of the Cone */
+	UPROPERTY(EditAnywhere, Category = "ShapeSettings|Subdivisions", meta = (DisplayName = "Height", UIMin = "1", UIMax = "100", ClampMin = "1", ClampMax = "4000", ProceduralShapeSetting))
+	int HeightSubdivisions = 1;
+};
+
+UCLASS()
+class MESHMODELINGTOOLS_API UProceduralArrowToolProperties : public UProceduralShapeToolProperties
+{
+	GENERATED_BODY()
+public:
+	/** Radius of the Arrow Shaft */
+	UPROPERTY(EditAnywhere, Category = "ShapeSettings|Dimensions", meta = (DisplayName = "Shaft Radius", UIMin = "1.0", UIMax = "1000.0", ClampMin = "0.0001", ClampMax = "1000000.0", ProceduralShapeSetting))
+	float ShaftRadius = 20.f;
+
+	/** Height of Arrow Shaft */
+	UPROPERTY(EditAnywhere, Category = "ShapeSettings|Dimensions", meta = (DisplayName = "Shaft Height", UIMin = "1.0", UIMax = "1000.0", ClampMin = "0.0001", ClampMax = "1000000.0", ProceduralShapeSetting))
+	float ShaftHeight = 200.f;
+
+	/** Radius of the Arrow Head */
+	UPROPERTY(EditAnywhere, Category = "ShapeSettings|Dimensions", meta = (DisplayName = "Head Radius", UIMin = "1.0", UIMax = "1000.0", ClampMin = "0.0001", ClampMax = "1000000.0", ProceduralShapeSetting))
+	float HeadRadius = 60.f;
+
+	/** Height of Arrow's Head */
+	UPROPERTY(EditAnywhere, Category = "ShapeSettings|Dimensions", meta = (DisplayName = "Head Height", UIMin = "1.0", UIMax = "1000.0", ClampMin = "0.0001", ClampMax = "1000000.0", ProceduralShapeSetting))
+	float HeadHeight = 120.f;
+
+	/** Number of Angular Slices Around the Arrow */
+	UPROPERTY(EditAnywhere, Category = "ShapeSettings|Slices", meta = (DisplayName = "Radial Slices", UIMin = "3", UIMax = "100", ClampMin = "3", ClampMax = "4000", ProceduralShapeSetting))
+	int RadialSlices = 16;
+
+	/** Number of Vertical Subdivisions Along in the Arrow */
+	UPROPERTY(EditAnywhere, Category = "ShapeSettings|Subdivisions", meta = (DisplayName = "Total", UIMin = "1", UIMax = "100", ClampMin = "1", ClampMax = "4000", ProceduralShapeSetting))
+	int TotalSubdivisions = 1;
+};
+
+UCLASS()
+class MESHMODELINGTOOLS_API UProceduralSphereToolProperties : public UProceduralShapeToolProperties
+{
+	GENERATED_BODY()
+public:
+	/** Radius of the Sphere */
+	UPROPERTY(EditAnywhere, Category = "ShapeSettings|Dimensions", meta = (DisplayName = "Radius", UIMin = "1.0", UIMax = "1000.0", ClampMin = "0.0001", ClampMax = "1000000.0", ProceduralShapeSetting))
+	float Radius = 50.f;
 
+	/** Number of Latitudinal Slices of the Sphere */
+	UPROPERTY(EditAnywhere, Category = "ShapeSettings|Slices", meta = (DisplayName = "Latitude Slices", UIMin = "3", UIMax = "100", ClampMin = "4", ClampMax = "4000", ProceduralShapeSetting))
+	int LatitudeSlices = 16;
 
-	virtual void SaveProperties(UInteractiveTool* SaveFromTool) override;
-	virtual void RestoreProperties(UInteractiveTool* RestoreToTool) override;
+	/** Number of Longitudinal Slices around the Sphere */
+	UPROPERTY(EditAnywhere, Category = "ShapeSettings|Slices", meta = (DisplayName = "Longitude Slices", UIMin = "3", UIMax = "100", ClampMin = "4", ClampMax = "4000", ProceduralShapeSetting))
+	int LongitudeSlices = 16;
 };
 
+UCLASS()
+class MESHMODELINGTOOLS_API UProceduralSphericalBoxToolProperties : public UProceduralShapeToolProperties
+{
+	GENERATED_BODY()
+public:
+	/** Radius of the Sphere */
+	UPROPERTY(EditAnywhere, Category = "ShapeSettings|Dimensions", meta = (DisplayName = "Radius", UIMin = "1.0", UIMax = "1000.0", ClampMin = "0.0001", ClampMax = "1000000.0", ProceduralShapeSetting))
+	float Radius = 50.f;
 
+	/** Number of Subdivisions of each Side of the Sphere */
+	UPROPERTY(EditAnywhere, Category = "ShapeSettings|Subdivisions", meta = (DisplayName = "Side", UIMin = "1", UIMax = "100", ClampMin = "1", ClampMax = "4000", ProceduralShapeSetting))
+	int Subdivisions = 16;
+};
 
 UCLASS(Transient)
 class MESHMODELINGTOOLS_API ULastActorInfo : public UObject
@@ -173,7 +344,7 @@ public:
 
 	UPROPERTY()
 	AActor* Actor = nullptr;
-	
+
 	UPROPERTY()
 	UStaticMesh* StaticMesh = nullptr;
 
@@ -189,10 +360,8 @@ public:
 	}
 };
 
-
-
 /**
- *
+ * Base tool to create primitives
  */
 UCLASS()
 class MESHMODELINGTOOLS_API UAddPrimitiveTool : public USingleClickTool, public IHoverBehaviorTarget
@@ -200,6 +369,8 @@ class MESHMODELINGTOOLS_API UAddPrimitiveTool : public USingleClickTool, public
 	GENERATED_BODY()
 
 public:
+	UAddPrimitiveTool(const FObjectInitializer&);
+
 	virtual void SetWorld(UWorld* World);
 	virtual void SetAssetAPI(IToolsContextAssetAPI* AssetAPI);
 
@@ -225,6 +396,9 @@ public:
 
 
 protected:
+	virtual void GenerateMesh(FDynamicMesh3* OutMesh) const {}
+	virtual UProceduralShapeToolProperties* CreateShapeSettings(){return nullptr;}
+
 	UPROPERTY()
 	UProceduralShapeToolProperties* ShapeSettings;
 
@@ -241,19 +415,9 @@ protected:
 		{
 			return false;
 		}
-		// manual diff because not all shape setting changes result in a different asset (e.g. some just affect the transform)
-		return
-			LastGenerated->ShapeSettings->Subdivisions == ShapeSettings->Subdivisions &&
-			LastGenerated->ShapeSettings->Slices == ShapeSettings->Slices &&
-			LastGenerated->ShapeSettings->PivotLocation == ShapeSettings->PivotLocation &&
-			LastGenerated->ShapeSettings->FeatureRadius == ShapeSettings->FeatureRadius &&
-			LastGenerated->ShapeSettings->Height == ShapeSettings->Height &&
-			LastGenerated->ShapeSettings->Width == ShapeSettings->Width &&
-			LastGenerated->ShapeSettings->Shape == ShapeSettings->Shape &&
-			LastGenerated->ShapeSettings->PolygroupMode == ShapeSettings->PolygroupMode &&
-			LastGenerated->MaterialProperties->UVScale == MaterialProperties->UVScale &&
-			LastGenerated->MaterialProperties->bWorldSpaceUVScale == MaterialProperties->bWorldSpaceUVScale
-			;
+		return (LastGenerated->MaterialProperties->UVScale == MaterialProperties->UVScale) &&
+			(LastGenerated->MaterialProperties->bWorldSpaceUVScale == MaterialProperties->bWorldSpaceUVScale) &&
+			ShapeSettings->IsEquivalent(LastGenerated->ShapeSettings);
 	}
 
 
@@ -263,7 +427,8 @@ protected:
 	UPROPERTY()
 	ULastActorInfo* LastGenerated;
 
-
+	UPROPERTY()
+	FString AssetName = TEXT("GeneratedAsset");
 
 protected:
 	UWorld* TargetWorld;
@@ -273,16 +438,158 @@ protected:
 	FFrame3f ShapeFrame;
 
 	void UpdatePreviewMesh();
+};
+
+UCLASS()
+class UAddBoxPrimitiveTool : public UAddPrimitiveTool
+{
+	GENERATED_BODY()
+public:
+	UAddBoxPrimitiveTool(const FObjectInitializer& ObjectInitializer)
+	: Super(ObjectInitializer.SetDefaultSubobjectClass<UProceduralBoxToolProperties>(TEXT("ShapeSettings")))
+	{
+		AssetName = TEXT("Box");
+	}
+protected:
+	void GenerateMesh(FDynamicMesh3* OutMesh) const override;
+};
+
+UCLASS()
+class UAddCylinderPrimitiveTool : public UAddPrimitiveTool
+{
+	GENERATED_BODY()
+public:
+	UAddCylinderPrimitiveTool(const FObjectInitializer& ObjectInitializer)
+	: Super(ObjectInitializer.SetDefaultSubobjectClass<UProceduralCylinderToolProperties>(TEXT("ShapeSettings")))
+	{
+		AssetName = TEXT("Cylinder");
+	}
+protected:
+	void GenerateMesh(FDynamicMesh3* OutMesh) const override;
+};
+
+UCLASS()
+class UAddConePrimitiveTool : public UAddPrimitiveTool
+{
+	GENERATED_BODY()
+public:
+	UAddConePrimitiveTool(const FObjectInitializer& ObjectInitializer)
+	: Super(ObjectInitializer.SetDefaultSubobjectClass<UProceduralConeToolProperties>(TEXT("ShapeSettings")))
+	{
+		AssetName = TEXT("Cone");
+	}
+protected:
+	void GenerateMesh(FDynamicMesh3* OutMesh) const override;
+};
+
+UCLASS()
+class UAddRectanglePrimitiveTool : public UAddPrimitiveTool
+{
+	GENERATED_BODY()
+public:
+	UAddRectanglePrimitiveTool(const FObjectInitializer& ObjectInitializer)
+	: Super(ObjectInitializer.SetDefaultSubobjectClass<UProceduralRectangleToolProperties>(TEXT("ShapeSettings")))
+	{
+		AssetName = TEXT("Rectangle");
+	}
+protected:
+	void GenerateMesh(FDynamicMesh3* OutMesh) const override;
+};
+
+UCLASS()
+class UAddRoundedRectanglePrimitiveTool : public UAddPrimitiveTool
+{
+	GENERATED_BODY()
+public:
+	UAddRoundedRectanglePrimitiveTool(const FObjectInitializer& ObjectInitializer)
+	: Super(ObjectInitializer.SetDefaultSubobjectClass<UProceduralRoundedRectangleToolProperties>(TEXT("ShapeSettings")))
+	{
+		AssetName = TEXT("RoundedRectangle");
+	}
+protected:
+	void GenerateMesh(FDynamicMesh3* OutMesh) const override;
+};
+
+UCLASS()
+class UAddDiscPrimitiveTool : public UAddPrimitiveTool
+{
+	GENERATED_BODY()
+public:
+	UAddDiscPrimitiveTool(const FObjectInitializer& ObjectInitializer)
+	: Super(ObjectInitializer.SetDefaultSubobjectClass<UProceduralDiscToolProperties>(TEXT("ShapeSettings")))
+	{
+		AssetName = TEXT("Disc");
+	}
+protected:
+	void GenerateMesh(FDynamicMesh3* OutMesh) const override;
+};
+
+UCLASS()
+class UAddPuncturedDiscPrimitiveTool : public UAddPrimitiveTool
+{
+	GENERATED_BODY()
+public:
+	UAddPuncturedDiscPrimitiveTool(const FObjectInitializer& ObjectInitializer)
+	: Super(ObjectInitializer.SetDefaultSubobjectClass<UProceduralPuncturedDiscToolProperties>(TEXT("ShapeSettings")))
+	{
+		AssetName = TEXT("PuncturedDisc");
+	}
+protected:
+	void GenerateMesh(FDynamicMesh3* OutMesh) const override;
+};
 
-	void GenerateCylinder(FDynamicMesh3* OutMesh);
-	void GenerateCone(FDynamicMesh3* OutMesh);
-	void GenerateBox(FDynamicMesh3* OutMesh);
-	void GenerateRectangle(FDynamicMesh3* OutMesh);
-	void GenerateRoundedRectangle(FDynamicMesh3* OutMesh);
-	void GenerateDisc(FDynamicMesh3* OutMesh);
-	void GeneratePuncturedDisc(FDynamicMesh3* OutMesh);
-	void GenerateTorus(FDynamicMesh3* OutMesh);
-	void GenerateSphere(FDynamicMesh3* OutMesh);
-	void GenerateSphericalBox(FDynamicMesh3* OutMesh);
-	void GenerateArrow(FDynamicMesh3* OutMesh);
+UCLASS()
+class UAddTorusPrimitiveTool : public UAddPrimitiveTool
+{
+	GENERATED_BODY()
+public:
+	UAddTorusPrimitiveTool(const FObjectInitializer& ObjectInitializer)
+	: Super(ObjectInitializer.SetDefaultSubobjectClass<UProceduralTorusToolProperties>(TEXT("ShapeSettings")))
+	{
+		AssetName = TEXT("Torus");
+	}
+protected:
+	void GenerateMesh(FDynamicMesh3* OutMesh) const override;
+};
+
+UCLASS()
+class UAddArrowPrimitiveTool : public UAddPrimitiveTool
+{
+	GENERATED_BODY()
+public:
+	UAddArrowPrimitiveTool(const FObjectInitializer& ObjectInitializer)
+	: Super(ObjectInitializer.SetDefaultSubobjectClass<UProceduralArrowToolProperties>(TEXT("ShapeSettings")))
+	{
+		AssetName = TEXT("Arrow");
+	}
+protected:
+	void GenerateMesh(FDynamicMesh3* OutMesh) const override;
+};
+
+UCLASS()
+class UAddSpherePrimitiveTool : public UAddPrimitiveTool
+{
+	GENERATED_BODY()
+public:
+	UAddSpherePrimitiveTool(const FObjectInitializer& ObjectInitializer)
+	: Super(ObjectInitializer.SetDefaultSubobjectClass<UProceduralSphereToolProperties>(TEXT("ShapeSettings")))
+	{
+		AssetName = TEXT("Sphere Type 2");
+	}
+protected:
+	void GenerateMesh(FDynamicMesh3* OutMesh) const override;
+};
+
+UCLASS()
+class UAddSphericalBoxPrimitiveTool : public UAddPrimitiveTool
+{
+	GENERATED_BODY()
+public:
+	UAddSphericalBoxPrimitiveTool(const FObjectInitializer& ObjectInitializer)
+	: Super(ObjectInitializer.SetDefaultSubobjectClass<UProceduralSphericalBoxToolProperties>(TEXT("ShapeSettings")))
+	{
+		AssetName = TEXT("Spherical Type 1");
+	}
+protected:
+	void GenerateMesh(FDynamicMesh3* OutMesh) const override;
 };
diff --git a/Engine/Plugins/Experimental/MeshModelingToolset/Source/MeshModelingTools/Public/AlignObjectsTool.h b/Engine/Plugins/Experimental/MeshModelingToolset/Source/MeshModelingTools/Public/AlignObjectsTool.h
--- a/Engine/Plugins/Experimental/MeshModelingToolset/Source/MeshModelingTools/Public/AlignObjectsTool.h
+++ b/Engine/Plugins/Experimental/MeshModelingToolset/Source/MeshModelingTools/Public/AlignObjectsTool.h
@@ -84,9 +84,6 @@ public:
 
 	UPROPERTY(EditAnywhere, Category = Axes)
 	bool bAlignZ = true;
-
-	virtual void SaveProperties(UInteractiveTool* SaveFromTool) override;
-	virtual void RestoreProperties(UInteractiveTool* RestoreToTool) override;
 };
 
 
@@ -113,7 +110,7 @@ public:
 	virtual void Setup() override;
 	virtual void Shutdown(EToolShutdownType ShutdownType) override;
 
-	virtual void Tick(float DeltaTime) override;
+	virtual void OnTick(float DeltaTime) override;
 	virtual void Render(IToolsContextRenderAPI* RenderAPI) override;
 
 	virtual bool HasCancel() const override { return true; }
diff --git a/Engine/Plugins/Experimental/MeshModelingToolset/Source/MeshModelingTools/Public/BakeTransformTool.h b/Engine/Plugins/Experimental/MeshModelingToolset/Source/MeshModelingTools/Public/BakeTransformTool.h
--- a/Engine/Plugins/Experimental/MeshModelingToolset/Source/MeshModelingTools/Public/BakeTransformTool.h
+++ b/Engine/Plugins/Experimental/MeshModelingToolset/Source/MeshModelingTools/Public/BakeTransformTool.h
@@ -88,9 +88,6 @@ public:
 	virtual void SetWorld(UWorld* World);
 	virtual void SetAssetAPI(IToolsContextAssetAPI* AssetAPI);
 
-	virtual void Tick(float DeltaTime) override;
-	virtual void Render(IToolsContextRenderAPI* RenderAPI) override;
-
 	virtual bool HasCancel() const override { return true; }
 	virtual bool HasAccept() const override;
 	virtual bool CanAccept() const override;
diff --git a/Engine/Plugins/Experimental/MeshModelingToolset/Source/MeshModelingTools/Public/CSGMeshesTool.h b/Engine/Plugins/Experimental/MeshModelingToolset/Source/MeshModelingTools/Public/CSGMeshesTool.h
new file mode 100644
--- /dev/null
+++ b/Engine/Plugins/Experimental/MeshModelingToolset/Source/MeshModelingTools/Public/CSGMeshesTool.h
@@ -0,0 +1,136 @@
+// Copyright Epic Games, Inc. All Rights Reserved.
+
+#pragma once
+
+#include "CoreMinimal.h"
+#include "UObject/NoExportTypes.h"
+#include "MultiSelectionTool.h"
+#include "InteractiveToolBuilder.h"
+#include "Drawing/LineSetComponent.h"
+#include "MeshOpPreviewHelpers.h"
+#include "BaseTools/SingleClickTool.h"
+
+#include "CompositionOps/BooleanMeshesOp.h"
+
+#include "CSGMeshesTool.generated.h"
+
+// predeclarations
+class FDynamicMesh3;
+class UTransformGizmo;
+class UTransformProxy;
+
+
+UCLASS()
+class MESHMODELINGTOOLS_API UCSGMeshesToolBuilder : public UInteractiveToolBuilder
+{
+	GENERATED_BODY()
+
+public:
+	IToolsContextAssetAPI* AssetAPI = nullptr;
+
+	virtual bool CanBuildTool(const FToolBuilderState& SceneState) const override;
+	virtual UInteractiveTool* BuildTool(const FToolBuilderState& SceneState) const override;
+};
+
+
+
+/**
+ * Standard properties of the CSG operation
+ */
+UCLASS()
+class MESHMODELINGTOOLS_API UCSGMeshesToolProperties : public UInteractiveToolPropertySet
+{
+	GENERATED_BODY()
+public:
+	/** The type of operation */
+	UPROPERTY(EditAnywhere, Category = Options)
+	ECSGOperation Operation = ECSGOperation::Union;
+
+	/** Remove the source Actors/Components when accepting results of tool.*/
+	UPROPERTY(EditAnywhere, Category = Options)
+	bool bDeleteInputActors = true;
+
+	/** Show UI to allow changing translation, rotation and scale of input meshes */
+	UPROPERTY(EditAnywhere, Category = Options)
+	bool bShowTransformUI = true;
+
+	/** Show boundary edges created by the CSG operation -- often due to numerical error */
+	UPROPERTY(EditAnywhere, Category = Options)
+	bool bShowNewBoundaryEdges = true;
+
+	/** Automatically attempt to fill any holes left by CSG (e.g. due to numerical errors) */
+	UPROPERTY(EditAnywhere, Category = Options)
+	bool bAttemptFixHoles = false;
+};
+
+
+
+/**
+ * Simple Mesh Plane Cutting Tool
+ */
+UCLASS()
+class MESHMODELINGTOOLS_API UCSGMeshesTool : public UMultiSelectionTool, public IDynamicMeshOperatorFactory
+{
+	GENERATED_BODY()
+
+public:
+
+	UCSGMeshesTool();
+
+	virtual void Setup() override;
+	virtual void Shutdown(EToolShutdownType ShutdownType) override;
+
+	virtual void SetWorld(UWorld* World);
+	virtual void SetAssetAPI(IToolsContextAssetAPI* AssetAPI);
+
+	virtual void OnTick(float DeltaTime) override;
+	virtual void Render(IToolsContextRenderAPI* RenderAPI) override;
+
+	virtual bool HasCancel() const override { return true; }
+	virtual bool HasAccept() const override;
+	virtual bool CanAccept() const override;
+
+#if WITH_EDITOR
+	virtual void PostEditChangeProperty(FPropertyChangedEvent &PropertyChangedEvent) override;
+#endif
+
+	virtual void OnPropertyModified(UObject* PropertySet, FProperty* Property) override;
+
+	// IDynamicMeshOperatorFactory API
+	virtual TUniquePtr<FDynamicMeshOperator> MakeNewOperator() override;
+
+protected:
+
+	UPROPERTY()
+	UMeshOpPreviewWithBackgroundCompute* Preview;
+
+	UPROPERTY()
+	UCSGMeshesToolProperties* CSGProperties;
+
+	UPROPERTY()
+	TArray<UTransformProxy*> TransformProxies;
+
+	UPROPERTY()
+	TArray<UTransformGizmo*> TransformGizmos;
+
+	TArray<TSharedPtr<FDynamicMesh3>> OriginalDynamicMeshes;
+
+	UPROPERTY()
+	ULineSetComponent* DrawnLineSet;
+
+	void TransformChanged(UTransformProxy* Proxy, FTransform Transform);
+
+	UWorld* TargetWorld;
+	IToolsContextAssetAPI* AssetAPI;
+
+	void SetupPreview();
+	void SetTransformGizmos();
+	void UpdateGizmoVisibility();
+
+	void GenerateAsset(const FDynamicMeshOpResult& Result);
+
+	void UpdateVisualization();
+
+	// for visualization of any errors in the currently-previewed CSG operation
+	TArray<int> CreatedBoundaryEdges;
+};
diff --git a/Engine/Plugins/Experimental/MeshModelingToolset/Source/MeshModelingTools/Public/CombineMeshesTool.h b/Engine/Plugins/Experimental/MeshModelingToolset/Source/MeshModelingTools/Public/CombineMeshesTool.h
--- a/Engine/Plugins/Experimental/MeshModelingToolset/Source/MeshModelingTools/Public/CombineMeshesTool.h
+++ b/Engine/Plugins/Experimental/MeshModelingToolset/Source/MeshModelingTools/Public/CombineMeshesTool.h
@@ -76,9 +76,6 @@ public:
 	virtual void SetWorld(UWorld* World);
 	virtual void SetAssetAPI(IToolsContextAssetAPI* AssetAPI);
 
-	virtual void Tick(float DeltaTime) override;
-	virtual void Render(IToolsContextRenderAPI* RenderAPI) override;
-
 	virtual bool HasCancel() const override { return true; }
 	virtual bool HasAccept() const override;
 	virtual bool CanAccept() const override;
diff --git a/Engine/Plugins/Experimental/MeshModelingToolset/Source/MeshModelingTools/Public/ConvertToPolygonsTool.h b/Engine/Plugins/Experimental/MeshModelingToolset/Source/MeshModelingTools/Public/ConvertToPolygonsTool.h
--- a/Engine/Plugins/Experimental/MeshModelingToolset/Source/MeshModelingTools/Public/ConvertToPolygonsTool.h
+++ b/Engine/Plugins/Experimental/MeshModelingToolset/Source/MeshModelingTools/Public/ConvertToPolygonsTool.h
@@ -9,7 +9,6 @@
 #include "DynamicMesh3.h"
 #include "FindPolygonsAlgorithm.h"
 #include "PreviewMesh.h"
-#include "Changes/ValueWatcher.h"
 #include "ConvertToPolygonsTool.generated.h"
 
 // predeclaration
@@ -60,8 +59,6 @@ public:
 
 	UPROPERTY(EditAnywhere, Category = Display)
 	bool bShowGroupColors = true;
-
-	virtual void SaveRestoreProperties(UInteractiveTool* RestoreToTool, bool bSaving) override;
 };
 
 /**
@@ -79,7 +76,7 @@ public:
 	virtual void Shutdown(EToolShutdownType ShutdownType) override;
 
 	virtual void Render(IToolsContextRenderAPI* RenderAPI) override;
-	virtual void Tick(float DeltaTime) override;
+	virtual void OnTick(float DeltaTime) override;
 
 	virtual bool HasCancel() const override { return true; }
 	virtual bool HasAccept() const override;
@@ -98,9 +95,6 @@ protected:
 	FDynamicMesh3 SearchMesh;
 	FDynamicMeshNormalOverlay InitialNormals;
 
-	TValueWatcher<EConvertToPolygonsMode> ConvertModeWatcher;
-	TValueWatcher<bool> ShowGroupsWatcher;
-
 	FFindPolygonsAlgorithm Polygons;
 	bool bPolygonsValid = false;
 	void UpdatePolygons();
diff --git a/Engine/Plugins/Experimental/MeshModelingToolset/Source/MeshModelingTools/Public/DeformMeshPolygonsTool.h b/Engine/Plugins/Experimental/MeshModelingToolset/Source/MeshModelingTools/Public/DeformMeshPolygonsTool.h
--- a/Engine/Plugins/Experimental/MeshModelingToolset/Source/MeshModelingTools/Public/DeformMeshPolygonsTool.h
+++ b/Engine/Plugins/Experimental/MeshModelingToolset/Source/MeshModelingTools/Public/DeformMeshPolygonsTool.h
@@ -93,12 +93,6 @@ class MESHMODELINGTOOLS_API UDeformMeshPolygonsTransformProperties : public UInt
 public:
 	UDeformMeshPolygonsTransformProperties();
 
-#if WITH_EDITOR
-	// UObject interface
-	virtual void PostEditChangeProperty(FPropertyChangedEvent& PropertyChangedEvent) override;
-	// End of UObject interface
-#endif
-
 	//Options
 
 	UPROPERTY(EditAnywhere, Category = Options, meta = (DisplayName = "Deformation Type", ToolTip = "Select the type of deformation you wish to employ on a polygroup."))
@@ -399,7 +393,7 @@ public:
 	virtual void Setup() override;
 	virtual void Shutdown(EToolShutdownType ShutdownType) override;
 
-	virtual void Tick(float DeltaTime) override;
+	virtual void OnTick(float DeltaTime) override;
 	virtual void Render(IToolsContextRenderAPI* RenderAPI) override;
 
 	virtual bool HasCancel() const override { return true; }
diff --git a/Engine/Plugins/Experimental/MeshModelingToolset/Source/MeshModelingTools/Public/DisplaceMeshTool.h b/Engine/Plugins/Experimental/MeshModelingToolset/Source/MeshModelingTools/Public/DisplaceMeshTool.h
--- a/Engine/Plugins/Experimental/MeshModelingToolset/Source/MeshModelingTools/Public/DisplaceMeshTool.h
+++ b/Engine/Plugins/Experimental/MeshModelingToolset/Source/MeshModelingTools/Public/DisplaceMeshTool.h
@@ -11,6 +11,7 @@
 #include "MeshNormals.h"
 #include "MeshOpPreviewHelpers.h"
 #include "Spatial/SampledScalarField2.h"
+#include "WeightMapUtil.h"
 #include "DisplaceMeshTool.generated.h"
 
 struct FMeshDescription;
@@ -34,12 +35,157 @@ enum class EDisplaceMeshToolDisplaceType : uint8
 	PerlinNoise UMETA(DisplayName = "Perlin Noise"),
 
 	/** Displace with N iterations */
-	DisplacementMap UMETA(DisplayName = "Displacement Map"),
+	DisplacementMap UMETA(DisplayName = "Texture2D Map"),
 
 	/** Move vertices in spatial sine wave pattern */
 	SineWave UMETA(DisplayName = "Sine Wave"),
 };
 
+/** The basic set of properties shared by (more or less) all DisplacementTypes. */
+UCLASS()
+class MESHMODELINGTOOLS_API UDisplaceMeshCommonProperties : public UInteractiveToolPropertySet
+{
+	GENERATED_BODY()
+
+public:
+	/** Displacement type */
+	UPROPERTY(EditAnywhere, Category = Options)
+	EDisplaceMeshToolDisplaceType DisplacementType = EDisplaceMeshToolDisplaceType::PerlinNoise;
+
+	/** Displacement intensity */
+	UPROPERTY(EditAnywhere, Category = Options,
+		meta = (UIMin = "-100.0", UIMax = "100.0", ClampMin = "-10000.0", ClampMax = "100000.0"))
+	float DisplaceIntensity = 10.0f;
+
+	/** Seed for randomization */
+	UPROPERTY(EditAnywhere, Category = Options,
+		meta = (EditCondition = "DisplacementType == EDisplaceMeshToolDisplaceType::RandomNoise || DisplacementType == EDisplaceMeshToolDisplaceType::PerlinNoise"))
+	int RandomSeed = 31337;
+
+	/** Subdivision iterations for mesh */
+	UPROPERTY(EditAnywhere, Category = Options,
+		meta = (UIMin = "0", UIMax = "10", ClampMin = "0", ClampMax = "100"))
+	int Subdivisions = 4;
+
+	/** Select vertex weight map. If configured, the weight map value will be sampled to modulate displacement intensity. */
+	UPROPERTY(EditAnywhere, Category = Options, meta = (GetOptions = GetWeightMapsFunc))
+	FName WeightMap;
+
+	UFUNCTION()
+	TArray<FString> GetWeightMapsFunc();
+
+	UPROPERTY()
+	TArray<FString> WeightMapsList;
+
+	UPROPERTY(EditAnywhere, Category = Options)
+	bool bInvertWeightMap = false;
+
+	UPROPERTY(EditAnywhere, Category = Options, AdvancedDisplay, meta = (TransientToolProperty) )
+	bool bDisableSizeWarning = false;
+};
+
+
+
+/** PropertySet for properties affecting the Image Map displacement type. */
+UCLASS()
+class MESHMODELINGTOOLS_API UDisplaceMeshTextureMapProperties : public UInteractiveToolPropertySet
+{
+	GENERATED_BODY()
+			
+public:
+	/** Displacement map */
+	UPROPERTY(EditAnywhere, Category = Options)
+	UTexture2D* DisplacementMap = nullptr;
+};
+
+
+
+/** Properties for a directional filter. Allows for displacement to be applied only to vertices whose normals point in a given direction */
+UCLASS()
+class MESHMODELINGTOOLS_API UDisplaceMeshDirectionalFilterProperties : public UInteractiveToolPropertySet
+{
+	GENERATED_BODY()
+
+public:
+	/** Whether the directional filter is active. */
+	UPROPERTY(EditAnywhere, Category = DirectionalFilterOptions)
+	bool bEnableFilter = false;
+
+	/** Unit vector representing the direction to filter along. */
+	UPROPERTY(EditAnywhere, Category = DirectionalFilterOptions, 
+		meta = (EditCondition = "bEnableFilter == true"))
+	FVector FilterDirection = {0.0f, 0.0f, 1.0f};
+
+	/** Scalar value determining how close to the filter direction the vertex normals must be in order to be displaced.
+		0: Only normals pointing exactly in the filter direction are displaced.
+		0.5: Normals forming angle up to 90 from the filter direction are displaced.
+		1.0: All vertices are displaced.
+	*/
+	UPROPERTY(EditAnywhere, Category = DirectionalFilterOptions,
+		meta = (EditCondition = "bEnableFilter == true", UIMin = "0.0", UIMax = "1.0", ClampMin = "0.0", ClampMax = "1.0"))
+	float FilterWidth = 0.25f;
+
+};
+
+/** Per-layer properties for Perlin noise. Each layer has independent Frequency and Intensity. */
+USTRUCT()
+struct FPerlinLayerProperties
+{
+	GENERATED_BODY()
+
+	FPerlinLayerProperties() = default;
+
+	FPerlinLayerProperties(float FrequencyIn, float IntensityIn) : 
+		Frequency(FrequencyIn), 
+		Intensity(IntensityIn) 
+    {}
+
+	/** Frequency of Perlin noise layer */
+	UPROPERTY(EditAnywhere, Category = PerlinNoiseOptions,
+			meta = (UIMin = "0.0", UIMax = "1.0", ClampMin = "0.0", ClampMax = "10.0"))
+	float Frequency = 0.1f;
+
+	/** Intensity/amplitude of Perlin noise layer */
+	UPROPERTY(EditAnywhere, Category = PerlinNoiseOptions,
+		meta = (UIMin = "-10.0", UIMax = "10.0", ClampMin = "-100.0", ClampMax = "100.0"))
+	float Intensity = 1.0f;
+};
+
+/** PropertySet for properties affecting the Perlin Noise displacement type. */
+UCLASS()
+class MESHMODELINGTOOLS_API UDisplaceMeshPerlinNoiseProperties : public UInteractiveToolPropertySet
+{
+	GENERATED_BODY()
+
+public:
+	UPROPERTY(EditAnywhere, Category = PerlinNoiseOptions)
+	TArray<FPerlinLayerProperties> PerlinLayerProperties = 
+		{ /*{Frequency, Intensity}*/ {0.05f, 1.0f}, {0.25f, 0.5f}, {0.5f, 0.2f}, {1.0f, 0.1f} };
+};
+
+
+/** PropertySet for Sine wave displacement */
+UCLASS()
+class MESHMODELINGTOOLS_API UDisplaceMeshSineWaveProperties : public UInteractiveToolPropertySet
+{
+	GENERATED_BODY()
+
+public:
+	/** Sine wave displacement frequency */
+	UPROPERTY(EditAnywhere, Category = SineWaveOptions,
+		meta = (UIMin = "0.0", UIMax = "1.0", ClampMin = "0.0", ClampMax = "10.0"))
+	float SineWaveFrequency = 0.1f;
+
+	/** Sine wave phase shift */
+	UPROPERTY(EditAnywhere, Category = SineWaveOptions,
+		meta = (UIMin = "0.0", UIMax = "6.28318531", ClampMin = "0.0", ClampMax = "6.28318531"))
+	float SineWavePhaseShift = 0.0f;
+
+	/** Unit vector representing the direction of wave displacement. */
+	UPROPERTY(EditAnywhere, Category = SineWaveOptions)
+	FVector SineWaveDirection = { 0.0f, 0.0f, 1.0f };
+};
+
 /**
  * Builder for Simple Mesh Displacement Tool
  */
@@ -61,54 +207,36 @@ class MESHMODELINGTOOLS_API UDisplaceMeshTool : public USingleSelectionTool
 	GENERATED_BODY()
 
 public:
-	UDisplaceMeshTool();
-
 	void Setup() final;
 	void Shutdown(EToolShutdownType ShutdownType) final;
-	void Tick(float DeltaTime) final;
+	void OnTick(float DeltaTime) final;
 
 	bool HasCancel() const final { return true; }
 	bool HasAccept() const final { return true; }
 	bool CanAccept() const final { return true; }
 #if WITH_EDITOR
-	void PostEditChangeProperty(FPropertyChangedEvent& PropertyChangedEvent) final;
+	void OnPropertyModified(UObject* PropertySet, FProperty* Property) final;
 #endif
 
-	/** primary brush mode */
-	UPROPERTY(EditAnywhere, Category = Options)
-	EDisplaceMeshToolDisplaceType DisplacementType;
+	/** Shared properties for all displacement modes. */
+	UPROPERTY()
+	UDisplaceMeshCommonProperties* CommonProperties;
 
-	/** Displacement intensity */
-	UPROPERTY(EditAnywhere, Category = Options,
-		meta = (UIMin = "-100.0", UIMax = "100.0", ClampMin = "-10000.0", ClampMax = "100000.0"))
-	float DisplaceIntensity;
+	/** Properties defining the directional filter. */
+	UPROPERTY()
+	UDisplaceMeshDirectionalFilterProperties* DirectionalFilterProperties;
 
-	/** Seed for randomization */
-	UPROPERTY(EditAnywhere, Category = Options,
-		meta = (EditCondition = "DisplacementType == EDisplaceMeshToolDisplaceType::RandomNoise || DisplacementType == EDisplaceMeshToolDisplaceType::PerlinNoise"))
-	int RandomSeed;
+	/** Properties defining the texture map */
+	UPROPERTY()
+	UDisplaceMeshTextureMapProperties* TextureMapProperties;
 
-	/** Subdivision iterations for mesh */
-	UPROPERTY(EditAnywhere, Category = Options,
-		meta = (UIMin = "0", UIMax = "10", ClampMin = "0", ClampMax = "100"))
-	int Subdivisions;
+	/** Multi-layer Perlin noise frequencies and intensities */
+	UPROPERTY()
+	UDisplaceMeshPerlinNoiseProperties* NoiseProperties;
 
-	/** Displacement map */
-	UPROPERTY(EditAnywhere, Category = Options,
-		meta = (EditCondition = "DisplacementType == EDisplaceMeshToolDisplaceType::DisplacementMap"))
-	UTexture2D* DisplacementMap;
-
-	/** Displacement frequency */
-	UPROPERTY(EditAnywhere, Category = Options,
-		meta = (EditCondition = "DisplacementType == EDisplaceMeshToolDisplaceType::SineWave || DisplacementType == EDisplaceMeshToolDisplaceType::PerlinNoise", 
-				UIMin = "0.0", UIMax = "1.0", ClampMin = "0.0", ClampMax = "10.0"))
-	float DisplaceFrequency;
-
-	/** Displacement phase shift */
-	UPROPERTY(EditAnywhere, Category = Options,
-		meta = (EditCondition = "DisplacementType == EDisplaceMeshToolDisplaceType::SineWave", 
-				UIMin = "0.0", UIMax = "6.28318531", ClampMin = "0.0", ClampMax = "6.28318531"))
-	float DisplacePhaseShift;
+	/** Sine wave parameters and direction of displacement */
+	UPROPERTY()
+	UDisplaceMeshSineWaveProperties* SineWaveProperties;
 
 private:
 	void StartComputation();
@@ -121,6 +249,11 @@ private:
 	bool bNeedsDisplaced = true;
 
 	FDynamicMesh3 OriginalMesh;
+	FDynamicMeshAABBTree3 OriginalMeshSpatial;
+	TSharedPtr<FIndexedWeightMap> ActiveWeightMap;
+	void UpdateActiveWeightMap();
+	float WeightMapQuery(const FVector3d& Position, const FIndexedWeightMap& WeightMap) const;
+
 	TSharedPtr<FDynamicMesh3> SubdividedMesh = nullptr;
 	USimpleDynamicMeshComponent* DynamicMeshComponent = nullptr;
 
diff --git a/Engine/Plugins/Experimental/MeshModelingToolset/Source/MeshModelingTools/Public/DrawPolyPathTool.h b/Engine/Plugins/Experimental/MeshModelingToolset/Source/MeshModelingTools/Public/DrawPolyPathTool.h
--- a/Engine/Plugins/Experimental/MeshModelingToolset/Source/MeshModelingTools/Public/DrawPolyPathTool.h
+++ b/Engine/Plugins/Experimental/MeshModelingToolset/Source/MeshModelingTools/Public/DrawPolyPathTool.h
@@ -5,7 +5,6 @@
 #include "CoreMinimal.h"
 #include "UObject/NoExportTypes.h"
 #include "BaseTools/MeshSurfacePointTool.h"
-#include "Changes/ValueWatcher.h"
 #include "Mechanics/PlaneDistanceFromHitMechanic.h"
 #include "Mechanics/SpatialCurveDistanceMechanic.h"
 #include "Mechanics/CollectSurfacePathMechanic.h"
@@ -87,8 +86,6 @@ public:
 
 	UPROPERTY(EditAnywhere, Category = Snapping)
 	bool bSnapToWorldGrid = true;
-
-	virtual void SaveRestoreProperties(UInteractiveTool* RestoreToTool, bool bSaving) override;
 };
 
 
@@ -129,8 +126,6 @@ class MESHMODELINGTOOLS_API UDrawPolyPathTool : public UInteractiveTool, public
 	GENERATED_BODY()
 
 public:
-	UDrawPolyPathTool();
-
 	virtual void RegisterActions(FInteractiveToolActionSet& ActionSet) override;
 
 	virtual void SetWorld(UWorld* World);
@@ -139,7 +134,7 @@ public:
 	virtual void Setup() override;
 	virtual void Shutdown(EToolShutdownType ShutdownType) override;
 
-	virtual void Tick(float DeltaTime) override;
+	virtual void OnTick(float DeltaTime) override;
 	virtual void Render(IToolsContextRenderAPI* RenderAPI) override;
 
 	virtual bool HasCancel() const override { return false; }
diff --git a/Engine/Plugins/Experimental/MeshModelingToolset/Source/MeshModelingTools/Public/DrawPolygonTool.h b/Engine/Plugins/Experimental/MeshModelingToolset/Source/MeshModelingTools/Public/DrawPolygonTool.h
--- a/Engine/Plugins/Experimental/MeshModelingToolset/Source/MeshModelingTools/Public/DrawPolygonTool.h
+++ b/Engine/Plugins/Experimental/MeshModelingToolset/Source/MeshModelingTools/Public/DrawPolygonTool.h
@@ -12,7 +12,6 @@
 #include "Snapping/PointPlanarSnapSolver.h"
 #include "ToolSceneQueriesUtil.h"
 #include "Properties/MeshMaterialProperties.h"
-#include "Changes/ValueWatcher.h"
 #include "Mechanics/PlaneDistanceFromHitMechanic.h"
 #include "DrawPolygonTool.generated.h"
 
@@ -117,17 +116,8 @@ public:
 
 	UPROPERTY(EditAnywhere, NonTransactional, Category = Polygon)
 	bool bShowGizmo = true;
-
-	//
-	// save/restore support
-	//
-	virtual void SaveProperties(UInteractiveTool* SaveFromTool) override;
-	virtual void RestoreProperties(UInteractiveTool* RestoreToTool) override;
 };
 
-
-
-
 UCLASS()
 class MESHMODELINGTOOLS_API UDrawPolygonToolSnapProperties : public UInteractiveToolPropertySet
 {
@@ -152,7 +142,7 @@ public:
 	UPROPERTY(EditAnywhere, NonTransactional, Category = Snapping, Meta = (EditCondition = "bEnableSnapping"))
 	bool bSnapToLengths = true;
 
-	UPROPERTY(VisibleAnywhere, NonTransactional, Category = Snapping)
+	UPROPERTY(VisibleAnywhere, NonTransactional, Category = Snapping, meta = (TransientToolProperty))
 	float SegmentLength = 0.0f;
 
 	UPROPERTY(EditAnywhere, NonTransactional, Category = Snapping)
@@ -160,12 +150,6 @@ public:
 
 	UPROPERTY(EditAnywhere, NonTransactional, Category = Snapping, Meta = (EditCondition = "bHitSceneObjects"))
 	float HitNormalOffset = 0.0f;
-
-	//
-	// save/restore support
-	//
-	virtual void SaveProperties(UInteractiveTool* SaveFromTool) override;
-	virtual void RestoreProperties(UInteractiveTool* RestoreToTool) override;
 };
 
 
@@ -191,7 +175,7 @@ public:
 	virtual void Setup() override;
 	virtual void Shutdown(EToolShutdownType ShutdownType) override;
 
-	virtual void Tick(float DeltaTime) override;
+	virtual void OnTick(float DeltaTime) override;
 	virtual void Render(IToolsContextRenderAPI* RenderAPI) override;
 
 	virtual bool HasCancel() const override { return false; }
@@ -284,7 +268,6 @@ protected:
 	// updates plane and gizmo position
 	virtual void SetDrawPlaneFromWorldPos(const FVector3d& Position, const FVector3d& Normal);
 
-	TValueWatcher<bool> ShowGizmoWatcher;
 	void UpdateShowGizmoState(bool bNewVisibility);
 
 
