Edigrate Dev-Editor updates to InteractiveToolsFramework, GeometryProcessing, MeshModelingToolset, and ModelingTools Editor Mode Plugins. Update a few files in new MeshPaint Plugin for UInteractiveToo

diff --git a/Engine/Plugins/Experimental/MeshModelingToolset/Source/MeshModelingTools/Public/DynamicMeshSculptTool.h b/Engine/Plugins/Experimental/MeshModelingToolset/Source/MeshModelingTools/Public/DynamicMeshSculptTool.h
--- a/Engine/Plugins/Experimental/MeshModelingToolset/Source/MeshModelingTools/Public/DynamicMeshSculptTool.h
+++ b/Engine/Plugins/Experimental/MeshModelingToolset/Source/MeshModelingTools/Public/DynamicMeshSculptTool.h
@@ -18,6 +18,7 @@
 #include "Properties/MeshMaterialProperties.h"
 #include "Properties/RemeshProperties.h"
 #include "TransformTypes.h"
+#include "Sculpting/MeshSculptToolBase.h"
 #include "DynamicMeshSculptTool.generated.h"
 
 class UTransformGizmo;
@@ -55,7 +56,13 @@ enum class EDynamicMeshSculptBrushType : uint8
 	/** Move vertices parallel to the view plane  */
 	Move UMETA(DisplayName = "Move"),
 
-	/** Smooth mesh vertices */
+	/** Grab Brush, fall-off alters the influence of the grab */
+	PullKelvin UMETA(DisplayName = "Kelvin Grab"),
+
+	/** Grab Brush that may generate cusps, fall-off alters the influence of the grab */
+	PullSharpKelvin UMETA(DisplayName = "Sharp Kelvin Grab"),
+
+	/** Smooth mesh vertices  */
 	Smooth UMETA(DisplayName = "Smooth"),
 
 	/** Displace vertices along the average surface normal (Ctrl to invert) */
@@ -70,9 +77,15 @@ enum class EDynamicMeshSculptBrushType : uint8
 	/** Displace vertices along their vertex normals */
 	Inflate UMETA(DisplayName = "Inflate"),
 
-	/** Move vertices towards the center of the brush (Ctrl to push away) */
+	/** Scale Brush will inflate or pinch radially from the center of the brush */
+	ScaleKelvin UMETA(DisplayName = "Kelvin Scale"),
+
+	/** Move vertices towards the center of the brush (Ctrl to push away)*/
 	Pinch UMETA(DisplayName = "Pinch"),
 
+	/** Twist Brush moves vertices in the plane perpendicular to the local mesh normal */
+	TwistKelvin UMETA(DisplayName = "Kelvin Twist"),
+
 	/** Move vertices towards the average plane of the brush stamp region */
 	Flatten UMETA(DisplayName = "Flatten"),
 
@@ -113,24 +126,6 @@ public:
 
 
 
-UCLASS()
-class MESHMODELINGTOOLS_API USculptBrushProperties : public UBrushBaseProperties
-{
-	GENERATED_BODY()
-public:
-	USculptBrushProperties();
-
-	virtual void SaveRestoreProperties(UInteractiveTool* Tool, bool bSaving) override;
-
-	/** Depth of Brush into surface along view ray or surface normal, depending on the Active Brush Type */
-	UPROPERTY(EditAnywhere, Category = Brush, meta = (UIMin = "-0.5", UIMax = "0.5", ClampMin = "-1.0", ClampMax = "1.0", DisplayPriority = 5))
-	float Depth = 0;
-
-	/** Allow the Brush to hit the back-side of the mesh */
-	UPROPERTY(EditAnywhere, Category = Brush, meta = (DisplayPriority = 6))
-	bool bHitBackFaces = true;
-};
-
 
 UCLASS()
 class MESHMODELINGTOOLS_API UBrushSculptProperties : public UInteractiveToolPropertySet
@@ -138,10 +133,8 @@ class MESHMODELINGTOOLS_API UBrushSculptProperties : public UInteractiveToolProp
 	GENERATED_BODY()
 
 public:
-	virtual void SaveRestoreProperties(UInteractiveTool* Tool, bool bSaving) override;
-
 	/* This is a dupe of the bool in the tool class.  I needed it here so it could be checked as an EditCondition */
-	UPROPERTY()
+	UPROPERTY(meta = (TransientToolProperty))
 	bool bIsRemeshingEnabled = false;
 
 	/** Primary Brush Mode */
@@ -171,96 +164,54 @@ public:
 
 
 
-
-UENUM()
-enum class EPlaneBrushSideMode : uint8
-{
-	BothSides = 0,
-	PushDown = 1,
-	PullTowards = 2
-};
-
-
 UCLASS()
-class MESHMODELINGTOOLS_API UPlaneBrushProperties : public UInteractiveToolPropertySet
+class MESHMODELINGTOOLS_API UBrushRemeshProperties : public URemeshProperties
 {
 	GENERATED_BODY()
+
 public:
-	virtual void SaveRestoreProperties(UInteractiveTool* Tool, bool bSaving) override;
+	/** Toggle remeshing on/off */
+	UPROPERTY(EditAnywhere, Category = Remeshing, meta = (DisplayPriority = 1))
+	bool bEnableRemeshing = true;
 
-	/** Control whether effect of brush should be limited to one side of the Plane  */
-	UPROPERTY(EditAnywhere, Category = PlaneBrush)
-	EPlaneBrushSideMode WhichSide = EPlaneBrushSideMode::BothSides;
+	// Note that if you change range here, you must also update UDynamicMeshSculptTool::ConfigureRemesher!
+	/** Desired size of triangles after Remeshing, relative to average initial triangle size. Larger value results in larger triangles. */
+	UPROPERTY(EditAnywhere, Category = Remeshing, meta = (DisplayName = "Relative Tri Size", UIMin = "-5", UIMax = "5", ClampMin = "-5", ClampMax = "5", DisplayPriority = 2))
+	int TriangleSize = 0;
+
+	/** Control the amount of simplification during sculpting. Higher values will avoid wiping out fine details on the mesh. */
+	UPROPERTY(EditAnywhere, Category = Remeshing, meta = (UIMin = "0", UIMax = "5", ClampMin = "0", ClampMax = "5", DisplayPriority = 3))
+	int PreserveDetail = 0;
 };
 
 
+
 UCLASS()
 class MESHMODELINGTOOLS_API UFixedPlaneBrushProperties : public UInteractiveToolPropertySet
 {
 	GENERATED_BODY()
 
 public:
-	UFixedPlaneBrushProperties();
-
-	virtual void SaveRestoreProperties(UInteractiveTool* Tool, bool bSaving) override;
-
-	UPROPERTY()
+	UPROPERTY( meta = (TransientToolProperty) )
 	bool bPropertySetEnabled = true;
 
 	/** Toggle whether Work Plane Positioing Gizmo is visible */
 	UPROPERTY(EditAnywhere, Category = TargetPlane, meta = (HideEditConditionToggle, EditCondition = "bPropertySetEnabled == true"))
-	bool bShowGizmo;
+	bool bShowGizmo = true;
 
 	/** Toggle whether Work Plane snaps to grid when using Gizmo */
 	UPROPERTY(EditAnywhere, Category = TargetPlane, meta = (HideEditConditionToggle, EditCondition = "bPropertySetEnabled == true"))
-	bool bSnapToGrid;
+	bool bSnapToGrid = true;
 
 	UPROPERTY(EditAnywhere, Category = TargetPlane, AdvancedDisplay, meta = (HideEditConditionToggle, EditCondition = "bPropertySetEnabled == true"))
-	FVector Position;
+	FVector Position = FVector::ZeroVector;
 
 	UPROPERTY(EditAnywhere, Category = TargetPlane, AdvancedDisplay, meta = (HideEditConditionToggle, EditCondition = "bPropertySetEnabled == true"))
-	FQuat Rotation;
+	FQuat Rotation = FQuat::Identity;
 };
 
 
-UCLASS()
-class MESHMODELINGTOOLS_API USculptMaxBrushProperties : public UInteractiveToolPropertySet
-{
-	GENERATED_BODY()
-public:
-	virtual void SaveRestoreProperties(UInteractiveTool* Tool, bool bSaving) override;
 
-	/** Specify maximum displacement height (relative to brush size) */
-	UPROPERTY(EditAnywhere, Category = SculptMaxBrush, meta = (UIMin = "0.0", UIMax = "1.0", ClampMin = "0.0", ClampMax = "1.0"))
-	float MaxHeight = 0.5;
-
-	/** Use maximum height from last brush stroke, regardless of brush size. Note that spatial brush falloff still applies.  */
-	UPROPERTY(EditAnywhere, Category = SculptMaxBrush)
-	bool bFreezeCurrentHeight = false;
-};
-
-
-UCLASS()
-class MESHMODELINGTOOLS_API UBrushRemeshProperties : public URemeshProperties
-{
-	GENERATED_BODY()
-
-public:
-	void SaveRestoreProperties(UInteractiveTool* RestoreToTool, bool bSaving) override;
-
-	/** Toggle remeshing on/off */
-	UPROPERTY(EditAnywhere, Category = Remeshing, meta = (DisplayPriority = 1))
-	bool bEnableRemeshing = true;
-
-	// Note that if you change range here, you must also update UDynamicMeshSculptTool::ConfigureRemesher!
-	/** Desired size of triangles after Remeshing, relative to average initial triangle size. Larger value results in larger triangles. */
-	UPROPERTY(EditAnywhere, Category = Remeshing, meta = (DisplayName = "Relative Tri Size", UIMin = "-5", UIMax = "5", ClampMin = "-5", ClampMax = "5", DisplayPriority = 2))
-	int TriangleSize = 0;
-
-	/** Control the amount of simplification during sculpting. Higher values will avoid wiping out fine details on the mesh. */
-	UPROPERTY(EditAnywhere, Category = Remeshing, meta = (UIMin = "0", UIMax = "5", ClampMin = "0", ClampMax = "5", DisplayPriority = 3))
-	int PreserveDetail = 0;
-};
 
 /**
  * Dynamic Mesh Sculpt Tool Class
@@ -279,7 +230,7 @@ public:
 	virtual void Setup() override;
 	virtual void Shutdown(EToolShutdownType ShutdownType) override;
 
-	virtual void Tick(float DeltaTime) override;
+	virtual void OnTick(float DeltaTime) override;
 	virtual void Render(IToolsContextRenderAPI* RenderAPI) override;
 
 	virtual bool HasCancel() const override { return true; }
@@ -316,6 +267,8 @@ public:
 	UPROPERTY()
 	USculptMaxBrushProperties* SculptMaxBrushProperties;
 	
+	UPROPERTY()
+	UKelvinBrushProperties* KelvinBrushProperties;
 
 	/** Properties that control dynamic remeshing */
 	UPROPERTY()
@@ -453,11 +406,16 @@ private:
 	void ApplyFixedPlaneBrush(const FRay& WorldRay);
 	void ApplyFlattenBrush(const FRay& WorldRay);
 	void ApplyResampleBrush(const FRay& WorldRay);
+	void ApplyPullKelvinBrush(const FRay& WorldRay);
+	void ApplyPullSharpKelvinBrush(const FRay& WorldRay);
+	void ApplyTwistKelvinBrush(const FRay& WorldRay);
+	void ApplyScaleKelvinBrush(const FRay& WorldRay);
 
 	double SculptMaxFixedHeight = -1.0;
 
 	double CalculateBrushFalloff(double Distance);
 	TArray<FVector3d> ROIPositionBuffer;
+	void SyncMeshWithPositionBuffer(FDynamicMesh3* Mesh);
 
 	FFrame3d ActiveFixedBrushPlane;
 	FFrame3d ComputeROIBrushPlane(const FVector3d& BrushCenter, bool bIgnoreDepth, bool bViewAligned);
diff --git a/Engine/Plugins/Experimental/MeshModelingToolset/Source/MeshModelingTools/Public/EditMeshPolygonsTool.h b/Engine/Plugins/Experimental/MeshModelingToolset/Source/MeshModelingTools/Public/EditMeshPolygonsTool.h
--- a/Engine/Plugins/Experimental/MeshModelingToolset/Source/MeshModelingTools/Public/EditMeshPolygonsTool.h
+++ b/Engine/Plugins/Experimental/MeshModelingToolset/Source/MeshModelingTools/Public/EditMeshPolygonsTool.h
@@ -17,7 +17,6 @@
 #include "Operations/GroupTopologyDeformer.h"
 #include "ModelingOperators/Public/ModelingTaskTypes.h"
 #include "Transforms/MultiTransformer.h"
-#include "Changes/ValueWatcher.h"
 #include "Selection/PolygonSelectionMechanic.h"
 #include "Mechanics/PlaneDistanceFromHitMechanic.h"
 #include "Mechanics/SpatialCurveDistanceMechanic.h"
@@ -316,8 +315,6 @@ class MESHMODELINGTOOLS_API UPolyEditExtrudeProperties : public UInteractiveTool
 public:
 	UPROPERTY(EditAnywhere, Category = Extrude)
 	EPolyEditExtrudeDirection Direction = EPolyEditExtrudeDirection::SelectionNormal;
-
-	virtual void SaveRestoreProperties(UInteractiveTool* RestoreToTool, bool bSaving) override;
 };
 
 
@@ -343,8 +340,6 @@ public:
 
 	UPROPERTY(EditAnywhere, Category = Cut)
 	bool bSnapToVertices = true;
-
-	virtual void SaveRestoreProperties(UInteractiveTool* RestoreToTool, bool bSaving) override;
 };
 
 
@@ -358,8 +353,6 @@ class MESHMODELINGTOOLS_API UPolyEditSetUVProperties : public UInteractiveToolPr
 public:
 	UPROPERTY(EditAnywhere, Category = PlanarProjectUV)
 	bool bShowMaterial = false;
-
-	virtual void SaveRestoreProperties(UInteractiveTool* RestoreToTool, bool bSaving) override;
 };
 
 
@@ -382,7 +375,7 @@ public:
 	virtual void Setup() override;
 	virtual void Shutdown(EToolShutdownType ShutdownType) override;
 
-	virtual void Tick(float DeltaTime) override;
+	virtual void OnTick(float DeltaTime) override;
 	virtual void Render(IToolsContextRenderAPI* RenderAPI) override;
 
 	virtual bool HasCancel() const override { return true; }
@@ -420,8 +413,6 @@ protected:
 
 	UPROPERTY()
 	UPolyEditTransformProperties* TransformProps;
-	TValueWatcher<ELocalFrameMode> LocalFrameModeWatcher;
-	TValueWatcher<bool> LockFrameWatcher;
 
 	UPROPERTY()
 	UEditMeshPolygonsToolActions* EditActions;
@@ -438,7 +429,6 @@ protected:
 
 	UPROPERTY()
 	UPolyEditExtrudeProperties* ExtrudeProperties;
-	TValueWatcher<EPolyEditExtrudeDirection> ExtrudeDirectionWatcher;
 
 	UPROPERTY()
 	UPolyEditCutProperties* CutProperties;
@@ -646,4 +636,4 @@ public:
 	virtual void Revert(UObject* Object) override;
 	virtual bool HasExpired(UObject* Object) const override;
 	virtual FString ToString() const override;
-};
\ No newline at end of file
+};
diff --git a/Engine/Plugins/Experimental/MeshModelingToolset/Source/MeshModelingTools/Public/EditNormalsTool.h b/Engine/Plugins/Experimental/MeshModelingToolset/Source/MeshModelingTools/Public/EditNormalsTool.h
--- a/Engine/Plugins/Experimental/MeshModelingToolset/Source/MeshModelingTools/Public/EditNormalsTool.h
+++ b/Engine/Plugins/Experimental/MeshModelingToolset/Source/MeshModelingTools/Public/EditNormalsTool.h
@@ -141,7 +141,7 @@ public:
 	virtual void SetWorld(UWorld* World);
 	virtual void SetAssetAPI(IToolsContextAssetAPI* AssetAPI);
 
-	virtual void Tick(float DeltaTime) override;
+	virtual void OnTick(float DeltaTime) override;
 	virtual void Render(IToolsContextRenderAPI* RenderAPI) override;
 
 	virtual bool HasCancel() const override { return true; }
diff --git a/Engine/Plugins/Experimental/MeshModelingToolset/Source/MeshModelingTools/Public/EditPivotTool.h b/Engine/Plugins/Experimental/MeshModelingToolset/Source/MeshModelingTools/Public/EditPivotTool.h
--- a/Engine/Plugins/Experimental/MeshModelingToolset/Source/MeshModelingTools/Public/EditPivotTool.h
+++ b/Engine/Plugins/Experimental/MeshModelingToolset/Source/MeshModelingTools/Public/EditPivotTool.h
@@ -157,7 +157,7 @@ public:
 	virtual void Setup() override;
 	virtual void Shutdown(EToolShutdownType ShutdownType) override;
 
-	virtual void Tick(float DeltaTime) override;
+	virtual void OnTick(float DeltaTime) override;
 	virtual void Render(IToolsContextRenderAPI* RenderAPI) override;
 
 	virtual bool HasCancel() const override { return true; }
diff --git a/Engine/Plugins/Experimental/MeshModelingToolset/Source/MeshModelingTools/Public/EditUVIslandsTool.h b/Engine/Plugins/Experimental/MeshModelingToolset/Source/MeshModelingTools/Public/EditUVIslandsTool.h
--- a/Engine/Plugins/Experimental/MeshModelingToolset/Source/MeshModelingTools/Public/EditUVIslandsTool.h
+++ b/Engine/Plugins/Experimental/MeshModelingToolset/Source/MeshModelingTools/Public/EditUVIslandsTool.h
@@ -15,7 +15,6 @@
 #include "Selection/GroupTopologySelector.h"
 #include "ModelingOperators/Public/ModelingTaskTypes.h"
 #include "Transforms/MultiTransformer.h"
-#include "Changes/ValueWatcher.h"
 #include "Selection/PolygonSelectionMechanic.h"
 #include "EditUVIslandsTool.generated.h"
 
@@ -72,7 +71,7 @@ public:
 	virtual void Setup() override;
 	virtual void Shutdown(EToolShutdownType ShutdownType) override;
 
-	virtual void Tick(float DeltaTime) override;
+	virtual void OnTick(float DeltaTime) override;
 	virtual void Render(IToolsContextRenderAPI* RenderAPI) override;
 
 	virtual bool HasCancel() const override { return true; }
diff --git a/Engine/Plugins/Experimental/MeshModelingToolset/Source/MeshModelingTools/Public/HoleFillTool.h b/Engine/Plugins/Experimental/MeshModelingToolset/Source/MeshModelingTools/Public/HoleFillTool.h
new file mode 100644
--- /dev/null
+++ b/Engine/Plugins/Experimental/MeshModelingToolset/Source/MeshModelingTools/Public/HoleFillTool.h
@@ -0,0 +1,214 @@
+// Copyright Epic Games, Inc. All Rights Reserved.
+
+// HoleFillTool: Fill one or more boundary loops on a selected mesh. Several hole-filling methods are available.
+
+#pragma once
+
+#include "CoreMinimal.h"
+#include "InteractiveToolBuilder.h"
+#include "SingleSelectionTool.h"
+#include "CleaningOps/HoleFillOp.h"
+#include "BaseTools/MeshSurfacePointTool.h"
+#include "BaseBehaviors/BehaviorTargetInterfaces.h"
+#include "DynamicMeshAABBTree3.h"
+#include "GroupTopology.h"
+#include "HoleFillTool.generated.h"
+
+class UPolygonSelectionMechanic;
+class UDynamicMeshReplacementChangeTarget;
+class UMeshOpPreviewWithBackgroundCompute;
+class USimpleDynamicMeshComponent;
+struct FDynamicMeshOpResult;
+class UHoleFillTool;
+
+/*
+ * Tool builder
+ */
+
+UCLASS()
+class MESHMODELINGTOOLS_API UHoleFillToolBuilder : public UInteractiveToolBuilder
+{
+	GENERATED_BODY()
+
+public:
+	bool CanBuildTool(const FToolBuilderState& SceneState) const override;
+	UInteractiveTool* BuildTool(const FToolBuilderState& SceneState) const override;
+};
+
+/*
+ * Properties
+ */
+
+UCLASS()
+class MESHMODELINGTOOLS_API UHoleFillToolProperties : public UInteractiveToolPropertySet
+{
+	GENERATED_BODY()
+
+public:
+	UPROPERTY(EditAnywhere, Category = Options)
+	EHoleFillOpFillType FillType = EHoleFillOpFillType::PolygonEarClipping;
+};
+
+
+UENUM()
+enum class EHoleFillToolActions
+{
+	NoAction,
+	SelectAll,
+	ClearSelection
+};
+
+UCLASS()
+class MESHMODELINGTOOLS_API UHoleFillToolActions : public UInteractiveToolPropertySet
+{
+GENERATED_BODY()
+
+	TWeakObjectPtr<UHoleFillTool> ParentTool;
+
+public:
+
+	void Initialize(UHoleFillTool* ParentToolIn)
+	{
+		ParentTool = ParentToolIn;
+	}
+
+	void PostAction(EHoleFillToolActions Action);
+
+	UFUNCTION(CallInEditor, Category = SelectionEdits, meta = (DisplayName = "Select All", DisplayPriority = 1))
+		void SelectAll()
+	{
+		PostAction(EHoleFillToolActions::SelectAll);
+	}
+
+	UFUNCTION(CallInEditor, Category = SelectionEdits, meta = (DisplayName = "Clear", DisplayPriority = 1))
+		void Clear()
+	{
+		PostAction(EHoleFillToolActions::ClearSelection);
+	}
+};
+
+/*
+ * Operator factory
+ */
+
+UCLASS()
+class MESHMODELINGTOOLS_API UHoleFillOperatorFactory : public UObject, public IDynamicMeshOperatorFactory
+{
+	GENERATED_BODY()
+
+public:
+
+	TUniquePtr<FDynamicMeshOperator> MakeNewOperator() override;
+
+	UPROPERTY()
+	UHoleFillTool* FillTool;
+};
+
+/*
+ * Tool
+ * Inherit from IClickBehaviorTarget so we can click on boundary loops.
+ */
+
+UCLASS()
+class MESHMODELINGTOOLS_API UHoleFillTool : public USingleSelectionTool, public IClickBehaviorTarget, public IHoverBehaviorTarget
+{
+	GENERATED_BODY()
+
+public:
+
+	// UMeshSurfacePointTool
+	void Setup() override;
+	void OnTick(float DeltaTime) override;
+	void OnPropertyModified(UObject* PropertySet, FProperty* Property) override;
+	bool HasCancel() const override { return true; }
+	bool HasAccept() const override { return true; }
+	bool CanAccept() const override;
+	void Shutdown(EToolShutdownType ShutdownType) override;
+
+	// IClickBehaviorTarget
+	FInputRayHit IsHitByClick(const FInputDeviceRay& ClickPos) override;
+	void OnClicked(const FInputDeviceRay& ClickPos) override;
+
+	// IHoverBehaviorTarget
+	FInputRayHit BeginHoverSequenceHitTest(const FInputDeviceRay& PressPos) override;
+	void OnBeginHover(const FInputDeviceRay& DevicePos) override {}
+	bool OnUpdateHover(const FInputDeviceRay& DevicePos) override;
+	void OnEndHover() override;
+
+	virtual void RequestAction(EHoleFillToolActions Action);
+
+protected:
+
+	friend UHoleFillOperatorFactory;
+	friend UHoleFillToolBuilder;
+
+	UPROPERTY()
+	UHoleFillToolProperties* Properties = nullptr;
+
+	UPROPERTY()
+	UHoleFillToolActions* Actions = nullptr;
+
+	UPROPERTY()
+	UMeshOpPreviewWithBackgroundCompute* Preview = nullptr;
+
+	UPROPERTY()
+	UPolygonSelectionMechanic* SelectionMechanic = nullptr;
+
+	// Input mesh
+	const FDynamicMesh3 OriginalMesh;
+
+	// UV Scale factor is cached based on the bounding box of the mesh before any fills are performed
+	float MeshUVScaleFactor = 0.0f;
+
+	// World in which to create the Preview mesh actor
+	UWorld* TargetWorld = nullptr;
+
+	// Used for hit querying
+	FDynamicMeshAABBTree3 MeshSpatial;
+
+	TSet<int32> NewTriangleIDs;
+
+	void SetWorld(UWorld* World);
+
+	// Create the Preview object
+	void SetupPreview();
+
+	bool bHavePendingAction = false;
+	EHoleFillToolActions PendingAction;
+	virtual void ApplyAction(EHoleFillToolActions ActionType);
+	void SelectAll();
+	void ClearSelection();
+
+	// Object used to get boundary loop information. All triangles return the same GroupID, so the only boundaries that 
+	// are returned are the the actual mesh boundaries.
+	// TODO: It seems like overkill to use a FGroupTopology subclass when we don't actually care about groups.
+	class FBasicTopology : public FGroupTopology
+	{
+	public:
+		FBasicTopology(const FDynamicMesh3* Mesh, bool bAutoBuild) :
+			FGroupTopology(Mesh, bAutoBuild)
+		{}
+
+		int GetGroupID(int TriangleID) const override
+		{
+			return Mesh->IsTriangle(TriangleID) ? 1 : 0;
+		}
+	};
+	TUniquePtr<FBasicTopology> Topology;
+
+	struct FSelectedBoundaryLoop
+	{
+		int32 EdgeTopoID;
+		TArray<int32> EdgeIDs;
+	};
+	TArray<FSelectedBoundaryLoop> ActiveBoundaryLoopSelection;
+	void UpdateActiveBoundaryLoopSelection();
+
+	// Just call the SelectionMechanism's Render function
+	void Render(IToolsContextRenderAPI* RenderAPI) override;
+
+	// Populate an array of Edge Loops to be processed by an FHoleFillOp. Returns the edge loops currently selected
+	// by this tool.
+	void GetLoopsToFill(TArray<FEdgeLoop>& OutLoops) const;
+
+};
diff --git a/Engine/Plugins/Experimental/MeshModelingToolset/Source/MeshModelingTools/Public/MeshAttributePaintTool.h b/Engine/Plugins/Experimental/MeshModelingToolset/Source/MeshModelingTools/Public/MeshAttributePaintTool.h
new file mode 100644
--- /dev/null
+++ b/Engine/Plugins/Experimental/MeshModelingToolset/Source/MeshModelingTools/Public/MeshAttributePaintTool.h
@@ -0,0 +1,258 @@
+// Copyright Epic Games, Inc. All Rights Reserved.
+
+#pragma once
+
+#include "DynamicMeshBrushTool.h"
+#include "Changes\ValueWatcher.h"
+#include "Changes\IndexedAttributeChange.h"
+#include "DynamicVerticesOctree3.h"
+#include "MeshDescription.h"
+#include "MeshAttributePaintTool.generated.h"
+
+
+struct FMeshDescription;
+
+
+/**
+ * Maps float values to linear color ramp.
+ */
+class FFloatAttributeColorMapper
+{
+public:
+	virtual ~FFloatAttributeColorMapper() {}
+
+	FLinearColor LowColor = FLinearColor(0.9f, 0.9f, 0.9f, 1.0f);
+	FLinearColor HighColor = FLinearColor(1.0f, 0.0f, 0.0f, 1.0f);
+
+	virtual FLinearColor ToColor(float Value)
+	{
+		float t = FMath::Clamp(Value, 0.0f, 1.0f);
+		return FLinearColor(
+			FMathf::Lerp(LowColor.R, HighColor.R, t),
+			FMathf::Lerp(LowColor.G, HighColor.G, t),
+			FMathf::Lerp(LowColor.B, HighColor.B, t),
+			1.0f);
+	}
+};
+
+
+/**
+ * Abstract interface to a single-channel indexed floating-point attribute
+ */
+class IMeshVertexAttributeAdapter
+{
+public:
+	virtual ~IMeshVertexAttributeAdapter() {}
+
+	virtual int ElementNum() const = 0;
+	virtual float GetValue(int32 Index) const = 0;
+	virtual void SetValue(int32 Index, float Value) = 0;
+	virtual FInterval1f GetValueRange() = 0;
+};
+
+
+
+/**
+ * Abstract interface to a set of single-channel indexed floating-point attributes
+ */
+class IMeshVertexAttributeSource
+{
+public:
+	virtual ~IMeshVertexAttributeSource() {}
+
+	virtual TArray<FName> GetAttributeList() = 0;
+	virtual TUniquePtr<IMeshVertexAttributeAdapter> GetAttribute(FName AttributeName) = 0;
+	/** @return number of indices in each attribute */
+	virtual int32 GetAttributeElementNum() = 0;
+};
+
+
+
+
+
+
+/**
+ * Tool Builder for Attribute Paint Tool
+ */
+UCLASS()
+class MESHMODELINGTOOLS_API UMeshAttributePaintToolBuilder : public UMeshSurfacePointToolBuilder
+{
+	GENERATED_BODY()
+
+public:
+	IToolsContextAssetAPI* AssetAPI = nullptr;
+
+	/** Optional color map customization */
+	TUniqueFunction<TUniquePtr<FFloatAttributeColorMapper>()> ColorMapFactory;
+
+	virtual UMeshSurfacePointTool* CreateNewTool(const FToolBuilderState& SceneState) const override;
+};
+
+
+
+
+/**
+ * Selected-Attribute settings Attribute Paint Tool
+ */
+UCLASS()
+class MESHMODELINGTOOLS_API UMeshAttributePaintToolProperties : public UInteractiveToolPropertySet
+{
+	GENERATED_BODY()
+public:
+	UPROPERTY(VisibleAnywhere, Category = Attribute)
+	TArray<FName> Attributes;
+
+	UPROPERTY(EditAnywhere, Category = Attribute, meta = (ArrayClamp = "Attributes"))
+	int SelectedAttribute = 0;
+
+	UPROPERTY(VisibleAnywhere, Category = Attribute)
+	FString AttributeName;
+};
+
+
+
+
+
+
+UENUM()
+enum class EMeshAttributePaintToolActions
+{
+	NoAction
+};
+
+
+
+UCLASS()
+class MESHMODELINGTOOLS_API UMeshAttributePaintEditActions : public UInteractiveToolPropertySet
+{
+	GENERATED_BODY()
+
+public:
+	TWeakObjectPtr<UMeshAttributePaintTool> ParentTool;
+
+	void Initialize(UMeshAttributePaintTool* ParentToolIn) { ParentTool = ParentToolIn; }
+
+	void PostAction(EMeshAttributePaintToolActions Action);
+};
+
+
+
+/**
+ * FCommandChange for color map changes
+ */
+class MESHMODELINGTOOLS_API FMeshAttributePaintChange : public TCustomIndexedValuesChange<float, int32>
+{
+public:
+	virtual FString ToString() const override
+	{
+		return FString(TEXT("Paint Attribute"));
+	}
+};
+
+
+
+/**
+ * UMeshAttributePaintTool paints single-channel float attributes on a MeshDescription.
+ * 
+ */
+UCLASS()
+class MESHMODELINGTOOLS_API UMeshAttributePaintTool : public UDynamicMeshBrushTool
+{
+	GENERATED_BODY()
+
+public:
+	virtual void SetWorld(UWorld* World);
+	virtual void SetAssetAPI(IToolsContextAssetAPI* AssetAPI);
+
+	virtual void RegisterActions(FInteractiveToolActionSet& ActionSet) override;
+
+	virtual void Setup() override;
+	virtual void OnTick(float DeltaTime) override;
+
+	virtual bool HasCancel() const override { return true; }
+	virtual bool HasAccept() const override { return true; }
+	virtual bool CanAccept() const override { return true; }
+
+	// UBaseBrushTool overrides
+	virtual void OnBeginDrag(const FRay& Ray) override;
+	virtual void OnUpdateDrag(const FRay& Ray) override;
+	virtual void OnEndDrag(const FRay& Ray) override;
+	virtual bool OnUpdateHover(const FInputDeviceRay& DevicePos) override;
+
+	virtual void RequestAction(EMeshAttributePaintToolActions ActionType);
+
+	virtual void SetColorMap(TUniquePtr<FFloatAttributeColorMapper> ColorMap);
+
+protected:
+	virtual void ApplyStamp(const FBrushStampData& Stamp);
+
+	virtual void OnShutdown(EToolShutdownType ShutdownType) override;
+	
+
+protected:
+	UPROPERTY()
+	UMeshAttributePaintToolProperties* AttribProps;
+
+	TValueWatcher<int32> SelectedAttributeWatcher;
+
+	//UPROPERTY()
+	//UMeshAttributePaintEditActions* AttributeEditActions;
+
+protected:
+	UWorld* TargetWorld;
+	IToolsContextAssetAPI* AssetAPI;
+
+	bool bInRemoveStroke = false;
+	bool bInSmoothStroke = false;
+	FBrushStampData StartStamp;
+	FBrushStampData LastStamp;
+	bool bStampPending;
+
+	TUniquePtr<FMeshDescription> EditedMesh;
+
+	double CalculateBrushFalloff(double Distance);
+	TDynamicVerticesOctree3<FDynamicMesh3> VerticesOctree;
+	TArray<int> PreviewBrushROI;
+	void CalculateVertexROI(const FBrushStampData& Stamp, TArray<int>& VertexROI);
+
+	TUniquePtr<FFloatAttributeColorMapper> ColorMapper;
+	TUniquePtr<IMeshVertexAttributeSource> AttributeSource;
+
+	struct FAttributeData
+	{
+		FName Name;
+		TUniquePtr<IMeshVertexAttributeAdapter> Attribute;
+		TArray<float> CurrentValues;
+		TArray<float> InitialValues;
+	};
+	TArray<FAttributeData> Attributes;
+	int32 AttributeBufferCount;
+	int32 CurrentAttributeIndex;
+	FInterval1f CurrentValueRange;
+
+	// actions
+
+	bool bHavePendingAction = false;
+	EMeshAttributePaintToolActions PendingAction;
+	virtual void ApplyAction(EMeshAttributePaintToolActions ActionType);
+
+	bool bVisibleAttributeValid = false;
+	int32 PendingNewSelectedIndex = -1;
+	
+	void InitializeAttributes();
+	void StoreCurrentAttribute();
+	void UpdateVisibleAttribute();
+	void UpdateSelectedAttribute(int32 NewSelectedIndex);
+
+	TUniquePtr<TIndexedValuesChangeBuilder<float, FMeshAttributePaintChange>> ActiveChangeBuilder;
+	void BeginChange();
+	TUniquePtr<FMeshAttributePaintChange> EndChange();
+	void ExternalUpdateValues(int32 AttribIndex, const TArray<int32>& VertexIndices, const TArray<float>& NewValues);
+};
+
+
+
+
+
+
+
diff --git a/Engine/Plugins/Experimental/MeshModelingToolset/Source/MeshModelingTools/Public/MeshInspectorTool.h b/Engine/Plugins/Experimental/MeshModelingToolset/Source/MeshModelingTools/Public/MeshInspectorTool.h
--- a/Engine/Plugins/Experimental/MeshModelingToolset/Source/MeshModelingTools/Public/MeshInspectorTool.h
+++ b/Engine/Plugins/Experimental/MeshModelingToolset/Source/MeshModelingTools/Public/MeshInspectorTool.h
@@ -6,6 +6,7 @@
 #include "UObject/NoExportTypes.h"
 #include "SingleSelectionTool.h"
 #include "InteractiveToolBuilder.h"
+#include "Drawing/LineSetComponent.h"
 #include "DynamicMesh3.h"
 #include "DynamicMeshAABBTree3.h"
 #include "Properties/MeshMaterialProperties.h"
@@ -73,18 +74,8 @@ public:
 	/** Length of line segments representing tangent vectors */
 	UPROPERTY(EditAnywhere, Category = Options, meta = (EditCondition = "bTangentVectors"))
 	float TangentLength = 5.0f;
-
-
-	//
-	// save/restore support
-	//
-	virtual void SaveProperties(UInteractiveTool* SaveFromTool) override;
-	virtual void RestoreProperties(UInteractiveTool* RestoreToTool) override;
 };
 
-
-
-
 /**
  * Mesh Inspector Tool for visualizing mesh information
  */
@@ -129,6 +120,9 @@ protected:
 	UPROPERTY()
 	UPreviewMesh* PreviewMesh;
 
+	UPROPERTY()
+	ULineSetComponent* DrawnLineSet;
+
 	UPROPERTY()
 	UMaterialInterface* DefaultMaterial = nullptr;
 
diff --git a/Engine/Plugins/Experimental/MeshModelingToolset/Source/MeshModelingTools/Public/MeshSelectionTool.h b/Engine/Plugins/Experimental/MeshModelingToolset/Source/MeshModelingTools/Public/MeshSelectionTool.h
--- a/Engine/Plugins/Experimental/MeshModelingToolset/Source/MeshModelingTools/Public/MeshSelectionTool.h
+++ b/Engine/Plugins/Experimental/MeshModelingToolset/Source/MeshModelingTools/Public/MeshSelectionTool.h
@@ -6,7 +6,6 @@
 #include "DynamicMeshBrushTool.h"
 #include "SelectionSet.h"
 #include "Changes/MeshSelectionChange.h"
-#include "Changes/ValueWatcher.h"
 #include "DynamicMeshOctree3.h"
 #include "MeshSelectionTool.generated.h"
 
@@ -250,9 +249,6 @@ public:
 
 	UPROPERTY(EditAnywhere, Category = Selection)
 	EMeshFacesColorMode FaceColorMode = EMeshFacesColorMode::None;
-
-	virtual void SaveProperties(UInteractiveTool* SaveFromTool) override;
-	virtual void RestoreProperties(UInteractiveTool* RestoreToTool) override;
 };
 
 
@@ -275,7 +271,7 @@ public:
 
 	virtual void Setup() override;
 
-	virtual void Tick(float DeltaTime) override;
+	virtual void OnTick(float DeltaTime) override;
 	virtual void Render(IToolsContextRenderAPI* RenderAPI) override;
 
 	virtual bool HasCancel() const override { return true; }
@@ -330,10 +326,6 @@ protected:
 
 	EMeshSelectionElementType SelectionType = EMeshSelectionElementType::Face;
 
-	TValueWatcher<bool> ShowWireframeWatcher;
-	TValueWatcher<EMeshFacesColorMode> ColorModeWatcher;
-
-
 	bool bInRemoveStroke = false;
 	FBrushStampData StartStamp;
 	FBrushStampData LastStamp;
@@ -364,7 +356,7 @@ protected:
 	// selection change
 	FMeshSelectionChangeBuilder* ActiveSelectionChange = nullptr;
 	void BeginChange(bool bAdding);
-	TUniquePtr<FMeshSelectionChange> EndChange();
+	TUniquePtr<FToolCommandChange> EndChange();
 	void CancelChange();
 
 
diff --git a/Engine/Plugins/Experimental/MeshModelingToolset/Source/MeshModelingTools/Public/MeshVertexSculptTool.h b/Engine/Plugins/Experimental/MeshModelingToolset/Source/MeshModelingTools/Public/MeshVertexSculptTool.h
new file mode 100644
--- /dev/null
+++ b/Engine/Plugins/Experimental/MeshModelingToolset/Source/MeshModelingTools/Public/MeshVertexSculptTool.h
@@ -0,0 +1,282 @@
+// Copyright Epic Games, Inc. All Rights Reserved.
+
+#pragma once
+
+#include "CoreMinimal.h"
+#include "HAL/Platform.h"
+#include "UObject/NoExportTypes.h"
+#include "SimpleDynamicMeshComponent.h"
+#include "DynamicMeshAABBTree3.h"
+#include "DynamicMeshOctree3.h"
+#include "MeshNormals.h"
+#include "BaseTools/BaseBrushTool.h"
+#include "Changes/ValueWatcher.h"
+#include "TransformTypes.h"
+#include "Sculpting/MeshSculptToolBase.h"
+#include "Sculpting/MeshBrushOpBase.h"
+#include "MeshVertexSculptTool.generated.h"
+
+class UTransformGizmo;
+class UTransformProxy;
+class UMaterialInstanceDynamic;
+
+DECLARE_STATS_GROUP(TEXT("VertexSculptTool"), STATGROUP_VtxSculptTool, STATCAT_Advanced);
+DECLARE_CYCLE_STAT(TEXT("VertexSculptTool_UpdateROI"), VtxSculptTool_UpdateROI, STATGROUP_VtxSculptTool);
+DECLARE_CYCLE_STAT(TEXT("VertexSculptTool_ApplyStamp"), VtxSculptToolApplyStamp, STATGROUP_VtxSculptTool );
+DECLARE_CYCLE_STAT(TEXT("VertexSculptTool_Tick"), VtxSculptToolTick, STATGROUP_VtxSculptTool);
+DECLARE_CYCLE_STAT(TEXT("VertexSculptTool_Tick_ApplyStampBlock"), VtxSculptTool_Tick_ApplyStampBlock, STATGROUP_VtxSculptTool);
+DECLARE_CYCLE_STAT(TEXT("VertexSculptTool_Tick_ApplyStamp_Remove"), VtxSculptTool_Tick_ApplyStamp_Remove, STATGROUP_VtxSculptTool);
+DECLARE_CYCLE_STAT(TEXT("VertexSculptTool_Tick_ApplyStamp_Insert"), VtxSculptTool_Tick_ApplyStamp_Insert, STATGROUP_VtxSculptTool);
+DECLARE_CYCLE_STAT(TEXT("VertexSculptTool_Tick_NormalsBlock"), VtxSculptTool_Tick_NormalsBlock, STATGROUP_VtxSculptTool);
+DECLARE_CYCLE_STAT(TEXT("VertexSculptTool_Tick_UpdateMeshBlock"), VtxSculptTool_Tick_UpdateMeshBlock, STATGROUP_VtxSculptTool);
+DECLARE_CYCLE_STAT(TEXT("VertexSculptTool_Tick_UpdateTargetBlock"), VtxSculptTool_Tick_UpdateTargetBlock, STATGROUP_VtxSculptTool);
+DECLARE_CYCLE_STAT(TEXT("VertexSculptTool_Normals_Collect"), VtxSculptTool_Normals_Collect, STATGROUP_VtxSculptTool);
+DECLARE_CYCLE_STAT(TEXT("VertexSculptTool_Normals_Compute"), VtxSculptTool_Normals_Compute, STATGROUP_VtxSculptTool);
+
+class FMeshVertexChangeBuilder;
+class UPreviewMesh;
+
+
+
+
+/**
+ * Tool Builder
+ */
+UCLASS()
+class MESHMODELINGTOOLS_API UMeshVertexSculptToolBuilder : public UMeshSurfacePointToolBuilder
+{
+	GENERATED_BODY()
+
+public:
+	virtual UMeshSurfacePointTool* CreateNewTool(const FToolBuilderState& SceneState) const override;
+};
+
+
+
+
+
+/** Mesh Sculpting Brush Types */
+UENUM()
+enum class EMeshVertexSculptBrushType : uint8
+{
+	/** Move vertices parallel to the view plane  */
+	Move UMETA(DisplayName = "Move"),
+
+	/** Smooth mesh vertices  */
+	Smooth UMETA(DisplayName = "Smooth"),
+
+	/** Displace vertices along the average surface normal (Ctrl to invert) */
+	Offset UMETA(DisplayName = "Sculpt (Normal)"),
+
+	/** Displace vertices towards the camera viewpoint (Ctrl to invert) */
+	SculptView UMETA(DisplayName = "Sculpt (Viewpoint)"),
+
+	/** Displaces vertices along the average surface normal to a maximum height based on the brush size (Ctrl to invert) */
+	SculptMax UMETA(DisplayName = "Sculpt Max"),
+
+	/** Displace vertices along their vertex normals */
+	Inflate UMETA(DisplayName = "Inflate"),
+
+	/** Move vertices towards the center of the brush (Ctrl to push away)*/
+	Pinch UMETA(DisplayName = "Pinch"),
+
+	/** Move vertices towards the average plane of the brush stamp region */
+	Flatten UMETA(DisplayName = "Flatten"),
+
+	/** Move vertices towards a plane defined by the initial brush position  */
+	Plane UMETA(DisplayName = "Plane (Normal)"),
+
+	/** Move vertices towards a view-facing plane defined at the initial brush position */
+	PlaneViewAligned UMETA(DisplayName = "Plane (Viewpoint)"),
+
+	/** Move vertices towards a fixed plane in world space, positioned with a 3D gizmo */
+	FixedPlane UMETA(DisplayName = "FixedPlane"),
+
+	LastValue UMETA(Hidden)
+
+};
+
+
+
+
+
+
+UCLASS()
+class MESHMODELINGTOOLS_API UVertexBrushSculptProperties : public UInteractiveToolPropertySet
+{
+	GENERATED_BODY()
+
+public:
+	/** Primary Brush Mode */
+	UPROPERTY(EditAnywhere, Category = Sculpting, meta = (DisplayName = "Brush Type"))
+	EMeshVertexSculptBrushType PrimaryBrushType = EMeshVertexSculptBrushType::Move;
+
+	/** Strength of the Primary Brush */
+	UPROPERTY(EditAnywhere, Category = Sculpting, meta = (DisplayName = "Strength", UIMin = "0.0", UIMax = "1.0", ClampMin = "0.0", ClampMax = "1.0", EditCondition = "PrimaryBrushType != EMeshVertexSculptBrushType::Pull"))
+	float PrimaryBrushSpeed = 0.5;
+
+	/** If true, try to preserve the shape of the UV/3D mapping. This will limit Smoothing and Remeshing in some cases. */
+	UPROPERTY(EditAnywhere, Category = Sculpting)
+	bool bPreserveUVFlow = false;
+
+	/** When Freeze Target is toggled on, the Brush Target Surface will be Frozen in its current state, until toggled off. Brush strokes will be applied relative to the Target Surface, for applicable Brushes */
+	UPROPERTY(EditAnywhere, Category = Sculpting, meta = (EditCondition = "PrimaryBrushType == EMeshVertexSculptBrushType::Sculpt || PrimaryBrushType == EMeshVertexSculptBrushType::SculptMax || PrimaryBrushType == EMeshVertexSculptBrushType::SculptView || PrimaryBrushType == EMeshVertexSculptBrushType::Pinch || PrimaryBrushType == EMeshVertexSculptBrushType::Resample" ))
+	bool bFreezeTarget = false;
+
+	/** Strength of Shift-to-Smooth Brushing and Smoothing Brush */
+	UPROPERTY(EditAnywhere, Category = Smoothing, meta = (DisplayName = "Smoothing Strength", UIMin = "0.0", UIMax = "1.0", ClampMin = "0.0", ClampMax = "1.0"))
+	float SmoothBrushSpeed = 0.25;
+};
+
+
+
+
+
+/**
+ * Mesh Vertex Sculpt Tool Class
+ */
+UCLASS()
+class MESHMODELINGTOOLS_API UMeshVertexSculptTool : public UMeshSculptToolBase
+{
+	GENERATED_BODY()
+
+public:
+	virtual void Setup() override;
+	virtual void Shutdown(EToolShutdownType ShutdownType) override;
+
+	virtual void OnTick(float DeltaTime) override;
+
+	virtual bool HasCancel() const override { return true; }
+	virtual bool HasAccept() const override { return true; }
+	virtual bool CanAccept() const override { return true; }
+
+	virtual bool HitTest(const FRay& Ray, FHitResult& OutHit) override;
+
+	virtual void OnBeginDrag(const FRay& Ray) override;
+	virtual void OnUpdateDrag(const FRay& Ray) override;
+	virtual void OnEndDrag(const FRay& Ray) override;
+
+	virtual FInputRayHit BeginHoverSequenceHitTest(const FInputDeviceRay& PressPos) override;
+	virtual bool OnUpdateHover(const FInputDeviceRay& DevicePos) override;
+
+	virtual void OnPropertyModified(UObject* PropertySet, FProperty* Property) override;
+
+public:
+
+	/** Properties that control sculpting*/
+	UPROPERTY()
+	UVertexBrushSculptProperties* SculptProperties;
+
+	UPROPERTY()
+	UPlaneBrushProperties* PlaneBrushProperties;
+
+	UPROPERTY()
+	USculptMaxBrushProperties* SculptMaxBrushProperties;
+	
+	//UPROPERTY()
+	//UKelvinBrushProperties* KelvinBrushProperties;
+
+
+public:
+	virtual void IncreaseBrushSpeedAction() override;
+	virtual void DecreaseBrushSpeedAction() override;
+	virtual void NextBrushModeAction() override;
+	virtual void PreviousBrushModeAction() override;
+
+protected:
+	UPROPERTY()
+	USimpleDynamicMeshComponent* DynamicMeshComponent;
+
+	virtual UBaseDynamicMeshComponent* GetSculptMeshComponent() { return DynamicMeshComponent; }
+
+	FTransform3d InitialTargetTransform;
+	FTransform3d CurTargetTransform;
+
+	// realtime visualization
+	void OnDynamicMeshComponentChanged(USimpleDynamicMeshComponent* Component, const FMeshVertexChange* Change, bool bRevert);
+	FDelegateHandle OnDynamicMeshComponentChangedHandle;
+
+	TValueWatcher<EMeshVertexSculptBrushType> BrushTypeWatcher;
+	void UpdateBrushType(EMeshVertexSculptBrushType BrushType);
+
+
+	bool bInDrag;
+
+	FFrame3d ActiveDragPlane;
+	FVector3d LastHitPosWorld;
+	FVector3d LastBrushPosLocal;
+	FVector3d LastBrushPosWorld;
+	FVector3d LastBrushPosNormalWorld;
+
+	TSet<int32> AccumulatedTriangleROI;
+	bool bUndoUpdatePending = false;
+	TFuture<bool> UndoNormalsFuture;
+	TFuture<bool> UndoUpdateOctreeFuture;
+	TFuture<bool> UndoUpdateBaseMeshFuture;
+	TArray<int> NormalsBuffer;
+	void WaitForPendingUndoRedo();
+
+	TArray<int> VertexROI;
+	TSet<int> VertexSetBuffer;
+	TSet<int> TriangleROI;
+	void UpdateROI(const FVector3d& BrushPos);
+
+	bool bTargetDirty;
+
+	bool bSmoothing;
+	bool bInvert;
+
+	bool bStampPending;
+	FRay PendingStampRay;
+	int StampTimestamp = 0;
+	EMeshVertexSculptBrushType PendingStampType = EMeshVertexSculptBrushType::Smooth;
+
+	TUniquePtr<FMeshSculptBrushOp> PrimaryBrushOp;
+	TUniquePtr<FMeshSculptBrushOp> SecondaryBrushOp;
+	TSharedPtr<FMeshSculptFallofFunc> ActiveFalloff;
+
+	bool UpdateStampPosition(const FRay& WorldRay);
+	void ApplyStamp();
+
+
+	FDynamicMesh3 BaseMesh;
+	FDynamicMeshOctree3 BaseMeshSpatial;
+	TArray<int32> BaseMeshIndexBuffer;
+	bool bCachedFreezeTarget = false;
+	void UpdateBaseMesh(const TSet<int32>* TriangleROI = nullptr);
+	bool GetBaseMeshNearest(int32 VertexID, const FVector3d& Position, double SearchRadius, FVector3d& TargetPosOut, FVector3d& TargetNormalOut);
+
+	FDynamicMeshOctree3 Octree;
+
+	int FindHitSculptMeshTriangle(const FRay3d& LocalRay);
+	int FindHitTargetMeshTriangle(const FRay3d& LocalRay);
+
+	bool UpdateBrushPosition(const FRay& WorldRay);
+	bool UpdateBrushPositionOnActivePlane(const FRay& WorldRay);
+	bool UpdateBrushPositionOnTargetMesh(const FRay& WorldRay, bool bFallbackToViewPlane);
+	bool UpdateBrushPositionOnSculptMesh(const FRay& WorldRay, bool bFallbackToViewPlane);
+	void AlignBrushToView();
+
+	double SculptMaxFixedHeight = -1.0;
+
+	double CalculateBrushFalloff(double Distance);
+	TArray<FVector3d> ROIPositionBuffer;
+	void SyncMeshWithPositionBuffer(FDynamicMesh3* Mesh);
+
+	FFrame3d ActiveFixedBrushPlane;
+	FFrame3d ComputeROIBrushPlane(const FVector3d& BrushCenter, bool bIgnoreDepth, bool bViewAligned);
+
+	FMeshVertexChangeBuilder* ActiveVertexChange = nullptr;
+	void BeginChange();
+	void EndChange();
+
+	FCriticalSection UpdateSavedVertexLock;
+
+
+
+protected:
+	virtual bool ShowWorkPlane() const override { return SculptProperties->PrimaryBrushType == EMeshVertexSculptBrushType::FixedPlane; }
+};
+
+
+
diff --git a/Engine/Plugins/Experimental/MeshModelingToolset/Source/MeshModelingTools/Public/PlaneCutTool.h b/Engine/Plugins/Experimental/MeshModelingToolset/Source/MeshModelingTools/Public/PlaneCutTool.h
--- a/Engine/Plugins/Experimental/MeshModelingToolset/Source/MeshModelingTools/Public/PlaneCutTool.h
+++ b/Engine/Plugins/Experimental/MeshModelingToolset/Source/MeshModelingTools/Public/PlaneCutTool.h
@@ -15,8 +15,6 @@
 
 
 // predeclarations
-struct FMeshDescription;
-class USimpleDynamicMeshComponent;
 class UTransformGizmo;
 class UTransformProxy;
 
@@ -54,12 +52,6 @@ public:
 	/** If true, meshes cut into multiple pieces will be saved as separate assets on 'accept'. */
 	UPROPERTY(EditAnywhere, Category = ToolOutputOptions)
 	bool bExportSeparatedPiecesAsNewMeshAssets = true;
-
-	//
-	// save/restore support
-	//
-	virtual void SaveProperties(UInteractiveTool* SaveFromTool) override;
-	virtual void RestoreProperties(UInteractiveTool* RestoreToTool) override;
 };
 
 
@@ -99,12 +91,6 @@ public:
 
 	UPROPERTY(EditAnywhere, Category = Options, AdvancedDisplay)
 	bool bFillSpans;
-
-	//
-	// save/restore support
-	//
-	virtual void SaveProperties(UInteractiveTool* SaveFromTool) override;
-	virtual void RestoreProperties(UInteractiveTool* RestoreToTool) override;
 };
 
 
@@ -147,7 +133,7 @@ public:
 
 	virtual void RegisterActions(FInteractiveToolActionSet& ActionSet) override;
 
-	virtual void Tick(float DeltaTime) override;
+	virtual void OnTick(float DeltaTime) override;
 	virtual void Render(IToolsContextRenderAPI* RenderAPI) override;
 
 	virtual bool HasCancel() const override { return true; }
diff --git a/Engine/Plugins/Experimental/MeshModelingToolset/Source/MeshModelingTools/Public/Properties/MeshMaterialProperties.h b/Engine/Plugins/Experimental/MeshModelingToolset/Source/MeshModelingTools/Public/Properties/MeshMaterialProperties.h
--- a/Engine/Plugins/Experimental/MeshModelingToolset/Source/MeshModelingTools/Public/Properties/MeshMaterialProperties.h
+++ b/Engine/Plugins/Experimental/MeshModelingToolset/Source/MeshModelingTools/Public/Properties/MeshMaterialProperties.h
@@ -24,7 +24,7 @@ public:
 	UNewMeshMaterialProperties();
 
 	/** Material for new mesh*/
-	UPROPERTY(EditAnywhere, NonTransactional, Category = Material, meta = (DisplayName = "Material"))
+	UPROPERTY(EditAnywhere, NonTransactional, Category = Material, meta = (DisplayName = "Material", TransientToolProperty))
 	UMaterialInterface* Material;
 
 	/** Scale factor for generated UVs */
@@ -42,12 +42,6 @@ public:
 	// controls visibility of UV/etc properties
 	UPROPERTY()
 	bool bShowExtendedOptions = true;
-
-	//
-	// save/restore support
-	//
-	virtual void SaveProperties(UInteractiveTool* SaveFromTool) override;
-	virtual void RestoreProperties(UInteractiveTool* RestoreToTool) override;
 };
 
 
@@ -83,26 +77,15 @@ public:
 	UPROPERTY(EditAnywhere, Category = MaterialPreview, meta = (EditCondition = "MaterialMode == ESetMeshMaterialMode::Override"))
 	UMaterialInterface* OverrideMaterial = nullptr;
 
-	UPROPERTY()
+	UPROPERTY(meta = (TransientToolProperty))
 	UMaterialInstanceDynamic* CheckerMaterial = nullptr;
 
 	void Setup();
 
 	void UpdateMaterials();
 	UMaterialInterface* GetActiveOverrideMaterial() const;
-
-	//
-	// save/restore support
-	//
-	virtual void SaveProperties(UInteractiveTool* SaveFromTool) override;
-	virtual void RestoreProperties(UInteractiveTool* RestoreToTool) override;
 };
 
-
-
-
-
-
 UENUM()
 enum class EMeshEditingMaterialModes
 {
@@ -138,9 +121,6 @@ public:
 	FLinearColor Color = FLinearColor(0.4f, 0.4f, 0.4f);
 
 	/** Image used in Image-Based Material */
-	UPROPERTY(EditAnywhere, Category = Rendering, meta = (EditConditionHides, EditCondition = "MaterialMode == EMeshEditingMaterialModes::Custom") )
+	UPROPERTY(EditAnywhere, Category = Rendering, meta = (EditConditionHides, EditCondition = "MaterialMode == EMeshEditingMaterialModes::Custom", TransientToolProperty) )
 	UTexture2D* Image;
-
-	virtual void SaveProperties(UInteractiveTool* SaveFromTool) override;
-	virtual void RestoreProperties(UInteractiveTool* RestoreToTool) override;
 };
diff --git a/Engine/Plugins/Experimental/MeshModelingToolset/Source/MeshModelingTools/Public/RemeshMeshTool.h b/Engine/Plugins/Experimental/MeshModelingToolset/Source/MeshModelingTools/Public/RemeshMeshTool.h
--- a/Engine/Plugins/Experimental/MeshModelingToolset/Source/MeshModelingTools/Public/RemeshMeshTool.h
+++ b/Engine/Plugins/Experimental/MeshModelingToolset/Source/MeshModelingTools/Public/RemeshMeshTool.h
@@ -6,7 +6,6 @@
 #include "UObject/NoExportTypes.h"
 #include "SingleSelectionTool.h"
 #include "InteractiveToolBuilder.h"
-#include "Changes/ValueWatcher.h"
 #include "DynamicMesh3.h"
 #include "DynamicMeshAABBTree3.h"
 #include "MeshOpPreviewHelpers.h"
@@ -41,21 +40,14 @@ class MESHMODELINGTOOLS_API URemeshMeshToolProperties : public URemeshProperties
 public:
 	URemeshMeshToolProperties();
 
-	void SaveRestoreProperties(UInteractiveTool* RestoreToTool, bool bSaving) override;
-
 	/** Target triangle count */
 	UPROPERTY(EditAnywhere, Category = Remeshing, meta = (EditCondition = "bUseTargetEdgeLength == false"))
 	int TargetTriangleCount;
 
-
 	/** Smoothing type */
 	UPROPERTY(EditAnywhere, Category = Remeshing)
 	ERemeshSmoothingType SmoothingType;
 
-	/** Number of Remeshing passes */
-	UPROPERTY(EditAnywhere, Category = Remeshing, meta = (UIMin = "0", UIMax = "50", ClampMin = "0", ClampMax = "1000"))
-	int RemeshIterations;
-
 	/** If true, UVs and Normals are discarded  */
 	UPROPERTY(EditAnywhere, Category = Remeshing)
 	bool bDiscardAttributes;
@@ -68,6 +60,14 @@ public:
 	UPROPERTY(EditAnywhere, Category = Display)
 	bool bShowGroupColors = false;
 
+	/** Remeshing type */
+	UPROPERTY(EditAnywhere, Category = Remeshing, AdvancedDisplay)
+	ERemeshType RemeshType;
+
+	/** Number of Remeshing passes */
+	UPROPERTY(EditAnywhere, Category = Remeshing, AdvancedDisplay, meta = (EditCondition = "RemeshType == ERemeshType::FullPass", UIMin = "0", UIMax = "50", ClampMin = "0", ClampMax = "1000"))
+	int RemeshIterations;
+
 	/** If true, the target count is ignored and the target edge length is used directly */
 	UPROPERTY(EditAnywhere, Category = Remeshing, AdvancedDisplay)
 	bool bUseTargetEdgeLength;
@@ -99,7 +99,7 @@ public:
 	virtual void Setup() override;
 	virtual void Shutdown(EToolShutdownType ShutdownType) override;
 
-	virtual void Tick(float DeltaTime) override;
+	virtual void OnTick(float DeltaTime) override;
 	virtual void Render(IToolsContextRenderAPI* RenderAPI) override;
 
 	virtual bool HasCancel() const override { return true; }
@@ -128,9 +128,6 @@ private:
 	TSharedPtr<FDynamicMeshAABBTree3> OriginalMeshSpatial;
 	double InitialMeshArea;
 
-	TValueWatcher<bool> ShowWireFrameWatcher;
-	TValueWatcher<bool> ShowGroupsWatcher;
-
 	double CalculateTargetEdgeLength(int TargetTriCount);
 	void GenerateAsset(const FDynamicMeshOpResult& Result);
 	void UpdateVisualization();
diff --git a/Engine/Plugins/Experimental/MeshModelingToolset/Source/MeshModelingTools/Public/RemoveOccludedTrianglesTool.h b/Engine/Plugins/Experimental/MeshModelingToolset/Source/MeshModelingTools/Public/RemoveOccludedTrianglesTool.h
--- a/Engine/Plugins/Experimental/MeshModelingToolset/Source/MeshModelingTools/Public/RemoveOccludedTrianglesTool.h
+++ b/Engine/Plugins/Experimental/MeshModelingToolset/Source/MeshModelingTools/Public/RemoveOccludedTrianglesTool.h
@@ -14,14 +14,6 @@
 #include "RemoveOccludedTrianglesTool.generated.h"
 
 
-// predeclarations
-struct FMeshDescription;
-class USimpleDynamicMeshComponent;
-
-
-
-
-
 /**
  *
  */
@@ -154,8 +146,7 @@ public:
 	virtual void SetWorld(UWorld* World);
 	virtual void SetAssetAPI(IToolsContextAssetAPI* AssetAPI);
 
-	virtual void Tick(float DeltaTime) override;
-	virtual void Render(IToolsContextRenderAPI* RenderAPI) override;
+	virtual void OnTick(float DeltaTime) override;
 
 	virtual bool HasCancel() const override { return true; }
 	virtual bool HasAccept() const override;
diff --git a/Engine/Plugins/Experimental/MeshModelingToolset/Source/MeshModelingTools/Public/Sculpting/MeshMoveBrushOps.h b/Engine/Plugins/Experimental/MeshModelingToolset/Source/MeshModelingTools/Public/Sculpting/MeshMoveBrushOps.h
new file mode 100644
--- /dev/null
+++ b/Engine/Plugins/Experimental/MeshModelingToolset/Source/MeshModelingTools/Public/Sculpting/MeshMoveBrushOps.h
@@ -0,0 +1,88 @@
+// Copyright Epic Games, Inc. All Rights Reserved.
+
+#pragma once
+
+#include "Sculpting/MeshBrushOpBase.h"
+#include "DynamicMesh3.h"
+#include "MeshNormals.h"
+#include "MeshWeights.h"
+#include "Async/ParallelFor.h"
+
+
+class FMoveBrushOp : public FMeshSculptBrushOp
+{
+
+public:
+	
+
+
+	virtual void ApplyStamp(const FDynamicMesh3* Mesh, const FSculptBrushStamp& Stamp, const TArray<int32>& Vertices, TArray<FVector3d>& NewPositionsOut) override
+	{
+		const FVector3d& StampPos = Stamp.LocalFrame.Origin;
+
+		double UsePower = Stamp.Power;
+		FVector3d MoveVec = Stamp.LocalFrame.Origin - Stamp.PrevLocalFrame.Origin;
+
+		ParallelFor(Vertices.Num(), [&](int32 k)
+		{
+			int32 VertIdx = Vertices[k];
+			FVector3d OrigPos = Mesh->GetVertex(VertIdx);
+
+			double Falloff = GetFalloff().Evaluate(Stamp, OrigPos);
+
+			FVector3d NewPos = OrigPos + Falloff * MoveVec;
+			NewPositionsOut[k] = NewPos;
+		});
+	}
+
+
+	virtual ESculptBrushOpTargetType GetBrushTargetType() const override
+	{
+		return ESculptBrushOpTargetType::ActivePlane;
+	}
+
+	virtual bool IgnoreZeroMovements() const override
+	{
+		return true;
+	}
+};
+
+
+
+
+
+//void UMeshVertexSculptTool::ApplyMoveBrush(const FRay& WorldRay)
+//{
+//	FVector3d NewBrushPosLocal = CurTargetTransform.InverseTransformPosition(LastBrushPosWorld);
+//	FVector3d MoveVec = NewBrushPosLocal - LastBrushPosLocal;
+//
+//	if (MoveVec.SquaredLength() <= 0)
+//	{
+//		LastBrushPosLocal = NewBrushPosLocal;
+//		return;
+//	}
+//
+//	FDynamicMesh3* Mesh = DynamicMeshComponent->GetMesh();
+//	int NumV = VertexROI.Num();
+//	ROIPositionBuffer.SetNum(NumV, false);
+//
+//	ParallelFor(NumV, [this, Mesh, NewBrushPosLocal, MoveVec](int k)
+//	{
+//		int VertIdx = VertexROI[k];
+//		FVector3d OrigPos = Mesh->GetVertex(VertIdx);
+//
+//		double PrevDist = (OrigPos - LastBrushPosLocal).Length();
+//		double NewDist = (OrigPos - NewBrushPosLocal).Length();
+//		double UseDist = FMath::Min(PrevDist, NewDist);
+//
+//		double Falloff = CalculateBrushFalloff(UseDist) * ActivePressure;
+//
+//		FVector3d NewPos = OrigPos + Falloff * MoveVec;
+//		ROIPositionBuffer[k] = NewPos;
+//	});
+//
+//	// Update the mesh positions to match those in the position buffer
+//	SyncMeshWithPositionBuffer(Mesh);
+//
+//	LastBrushPosLocal = NewBrushPosLocal;
+//}
\ No newline at end of file
diff --git a/Engine/Plugins/Experimental/MeshModelingToolset/Source/MeshModelingTools/Public/Sculpting/MeshPinchBrushOps.h b/Engine/Plugins/Experimental/MeshModelingToolset/Source/MeshModelingTools/Public/Sculpting/MeshPinchBrushOps.h
new file mode 100644
--- /dev/null
+++ b/Engine/Plugins/Experimental/MeshModelingToolset/Source/MeshModelingTools/Public/Sculpting/MeshPinchBrushOps.h
@@ -0,0 +1,78 @@
+// Copyright Epic Games, Inc. All Rights Reserved.
+
+#pragma once
+
+#include "Sculpting/MeshBrushOpBase.h"
+#include "DynamicMesh3.h"
+#include "Async/ParallelFor.h"
+#include "LineTypes.h"
+
+
+class FPinchBrushOp : public FMeshSculptBrushOp
+{
+
+public:
+	double BrushSpeedTuning = 0.05;
+
+	FVector3d LastSmoothBrushPosLocal;
+	FVector3d LastSmoothBrushNormalLocal;;
+
+	virtual void BeginStroke(const FDynamicMesh3* Mesh, const FSculptBrushStamp& Stamp, const TArray<int32>& InitialVertices) 
+	{
+		LastSmoothBrushPosLocal = Stamp.LocalFrame.Origin;
+		LastSmoothBrushNormalLocal = Stamp.LocalFrame.Z();
+	}
+
+
+	virtual void ApplyStamp(const FDynamicMesh3* Mesh, const FSculptBrushStamp& Stamp, const TArray<int32>& Vertices, TArray<FVector3d>& NewPositionsOut) override
+	{
+		// hardcoded LazyBrush
+		FVector3d NewSmoothBrushPosLocal = FVector3d::Lerp(LastSmoothBrushPosLocal, Stamp.LocalFrame.Origin, 0.75f);
+		FVector3d NewSmoothBrushNormal = FVector3d::Lerp(LastSmoothBrushNormalLocal, Stamp.LocalFrame.Z(), 0.75f);
+		NewSmoothBrushNormal.Normalize();
+
+		FVector3d MotionVec = NewSmoothBrushPosLocal - LastSmoothBrushPosLocal;
+		bool bHaveMotion = (MotionVec.Length() > FMathd::ZeroTolerance);
+		MotionVec.Normalize();
+		FLine3d MoveLine(LastSmoothBrushPosLocal, MotionVec);
+
+		FVector3d DepthPosLocal = NewSmoothBrushPosLocal - (Stamp.Depth * Stamp.Radius * NewSmoothBrushNormal);
+		double UseSpeed = Stamp.Direction * Stamp.Radius * Stamp.Power * BrushSpeedTuning;
+
+		LastSmoothBrushPosLocal = NewSmoothBrushPosLocal;
+
+		ParallelFor(Vertices.Num(), [&](int32 k)
+		{
+			int32 VertIdx = Vertices[k];
+			FVector3d OrigPos = Mesh->GetVertex(VertIdx);
+
+			FVector3d Delta = DepthPosLocal - OrigPos;
+			FVector3d MoveVec = UseSpeed * Delta;
+
+			//double Falloff = GetFalloff().Evaluate(Stamp, OrigPos);
+
+			// pinch uses 1/x falloff
+			double Distance = OrigPos.Distance(NewSmoothBrushPosLocal);
+			double NormalizedDistance = (Distance / Stamp.Radius) + 0.0001;
+			double Falloff = FMathd::Clamp(1.0 - NormalizedDistance, 0.0, 1.0);
+			Falloff *= Falloff; Falloff *= Falloff;
+
+			if (bHaveMotion && Falloff < 0.8f)
+			{
+				double AnglePower = 1.0 - FMathd::Abs(MoveVec.Normalized().Dot(MotionVec));
+				Falloff *= AnglePower;
+			}
+
+			FVector3d NewPos = OrigPos + Falloff * MoveVec;
+			NewPositionsOut[k] = NewPos;
+		});
+	}
+
+
+	virtual ESculptBrushOpTargetType GetBrushTargetType() const
+	{
+		return ESculptBrushOpTargetType::TargetMesh;
+	}
+};
+
+
diff --git a/Engine/Plugins/Experimental/MeshModelingToolset/Source/MeshModelingTools/Public/Sculpting/MeshSculptToolBase.h b/Engine/Plugins/Experimental/MeshModelingToolset/Source/MeshModelingTools/Public/Sculpting/MeshSculptToolBase.h
new file mode 100644
--- /dev/null
+++ b/Engine/Plugins/Experimental/MeshModelingToolset/Source/MeshModelingTools/Public/Sculpting/MeshSculptToolBase.h
@@ -0,0 +1,292 @@
+// Copyright Epic Games, Inc. All Rights Reserved.
+
+#pragma once
+
+#include "CoreMinimal.h"
+#include "BaseTools/MeshSurfacePointTool.h"
+#include "BaseTools/BaseBrushTool.h"
+#include "Sculpting/MeshBrushOpBase.h"
+#include "BoxTypes.h"
+#include "Properties/MeshMaterialProperties.h"
+#include "Changes/ValueWatcher.h"
+#include "MeshSculptToolBase.generated.h"
+
+
+class UMaterialInstanceDynamic;
+class UTransformGizmo;
+class UTransformProxy;
+class UPreviewMesh;
+class UBaseDynamicMeshComponent;
+
+
+UCLASS()
+class MESHMODELINGTOOLS_API USculptBrushProperties : public UBrushBaseProperties
+{
+	GENERATED_BODY()
+public:
+	USculptBrushProperties()
+	{
+		this->BrushFalloffAmount = 0.5f;
+	}
+
+	/** Depth of Brush into surface along view ray or surface normal, depending on the Active Brush Type */
+	UPROPERTY(EditAnywhere, Category = Brush, meta = (UIMin = "-0.5", UIMax = "0.5", ClampMin = "-1.0", ClampMax = "1.0", DisplayPriority = 5))
+	float Depth = 0;
+
+	/** Allow the Brush to hit the back-side of the mesh */
+	UPROPERTY(EditAnywhere, Category = Brush, meta = (DisplayPriority = 6))
+	bool bHitBackFaces = true;
+};
+
+
+
+
+UCLASS()
+class MESHMODELINGTOOLS_API UKelvinBrushProperties : public UInteractiveToolPropertySet
+{
+	GENERATED_BODY()
+
+public:
+	/** Brush Fall off as fraction of brush size*/
+	UPROPERTY(EditAnywhere, Category = Kelvin, meta = (UIMin = "0.0", ClampMin = "0.0"))
+	float FallOffDistance = 1.f;
+
+	/** How much the mesh resists shear */
+	UPROPERTY(EditAnywhere, Category = Kelvin, meta = (UIMin = "0.0", ClampMin = "0.0"))
+	float Stiffness = 1.f;
+
+	/** How compressible the spatial region is: 1 - 2 x Poisson ratio */
+	UPROPERTY(EditAnywhere, Category = Kelvin, meta = (UIMin = "0.0", UIMax = "0.1", ClampMin = "0.0", ClampMax = "1.0"))
+	float Incompressiblity = 1.f;
+
+	/** Integration steps*/
+	UPROPERTY(EditAnywhere, Category = Kelvin, meta = (UIMin = "0.0", UIMax = "100", ClampMin = "0.0", ClampMax = "100"))
+	int BrushSteps = 3;
+};
+
+
+
+UENUM()
+enum class EPlaneBrushSideMode : uint8
+{
+	BothSides = 0,
+	PushDown = 1,
+	PullTowards = 2
+};
+
+
+
+UCLASS()
+class MESHMODELINGTOOLS_API UPlaneBrushProperties : public UInteractiveToolPropertySet
+{
+	GENERATED_BODY()
+public:
+	/** Control whether effect of brush should be limited to one side of the Plane  */
+	UPROPERTY(EditAnywhere, Category = PlaneBrush)
+	EPlaneBrushSideMode WhichSide = EPlaneBrushSideMode::BothSides;
+};
+
+
+
+UCLASS()
+class MESHMODELINGTOOLS_API UWorkPlaneProperties : public UInteractiveToolPropertySet
+{
+	GENERATED_BODY()
+
+public:
+	UPROPERTY( meta = (TransientToolProperty) )
+	bool bPropertySetEnabled = true;
+
+	/** Toggle whether Work Plane Positioing Gizmo is visible */
+	UPROPERTY(EditAnywhere, Category = TargetPlane, meta = (HideEditConditionToggle, EditCondition = "bPropertySetEnabled == true"))
+	bool bShowGizmo = true;
+
+	/** Toggle whether Work Plane snaps to grid when using Gizmo */
+	UPROPERTY(EditAnywhere, Category = TargetPlane, meta = (HideEditConditionToggle, EditCondition = "bPropertySetEnabled == true"))
+	bool bSnapToGrid = true;
+
+	UPROPERTY(EditAnywhere, Category = TargetPlane, AdvancedDisplay, meta = (HideEditConditionToggle, EditCondition = "bPropertySetEnabled == true"))
+	FVector Position = FVector::ZeroVector;
+
+	UPROPERTY(EditAnywhere, Category = TargetPlane, AdvancedDisplay, meta = (HideEditConditionToggle, EditCondition = "bPropertySetEnabled == true"))
+	FQuat Rotation = FQuat::Identity;
+};
+
+
+
+
+UCLASS()
+class MESHMODELINGTOOLS_API USculptMaxBrushProperties : public UInteractiveToolPropertySet
+{
+	GENERATED_BODY()
+public:
+	/** Specify maximum displacement height (relative to brush size) */
+	UPROPERTY(EditAnywhere, Category = SculptMaxBrush, meta = (UIMin = "0.0", UIMax = "1.0", ClampMin = "0.0", ClampMax = "1.0"))
+	float MaxHeight = 0.5;
+
+	/** Use maximum height from last brush stroke, regardless of brush size. Note that spatial brush falloff still applies.  */
+	UPROPERTY(EditAnywhere, Category = SculptMaxBrush)
+	bool bFreezeCurrentHeight = false;
+};
+
+
+
+
+
+/**
+ * Base Tool for mesh sculpting tools, provides some shared functionality
+ */
+UCLASS()
+class MESHMODELINGTOOLS_API UMeshSculptToolBase : public UMeshSurfacePointTool
+{
+	GENERATED_BODY()
+public:
+	virtual void RegisterActions(FInteractiveToolActionSet& ActionSet) override;
+
+	virtual void SetWorld(UWorld* World);
+
+	virtual void Setup() override;
+	virtual void Shutdown(EToolShutdownType ShutdownType) override;
+	virtual void Render(IToolsContextRenderAPI* RenderAPI) override;
+
+protected:
+	virtual void OnTick(float DeltaTime) override;
+
+protected:
+	UWorld* TargetWorld;		// required to spawn UPreviewMesh/etc
+	FViewCameraState CameraState;
+
+	/**
+	 * Subclass must implement this and return relevant rendering component
+	 */
+	virtual UBaseDynamicMeshComponent* GetSculptMeshComponent() { check(false); return nullptr; }
+
+
+public:
+	/** Properties that control brush size/etc*/
+	UPROPERTY()
+	USculptBrushProperties* BrushProperties;
+
+	UPROPERTY()
+	UWorkPlaneProperties* GizmoProperties;
+
+
+	//
+	// Brush Size
+	//
+protected:
+	FInterval1d BrushRelativeSizeRange;
+	double CurrentBrushRadius;
+	virtual void InitializeBrushSizeRange(const FAxisAlignedBox3d& TargetBounds);
+	virtual void CalculateBrushRadius();
+	virtual double GetCurrentBrushRadius() const { return CurrentBrushRadius; }
+
+	double ActivePressure = 1.0;
+	virtual double GetActivePressure() const { return ActivePressure; }
+
+public:
+	virtual void IncreaseBrushRadiusAction();
+	virtual void DecreaseBrushRadiusAction();
+	virtual void IncreaseBrushRadiusSmallStepAction();
+	virtual void DecreaseBrushRadiusSmallStepAction();
+	virtual void IncreaseBrushFalloffAction();
+	virtual void DecreaseBrushFalloffAction();
+
+
+	// client currently needs to implement these...
+	virtual void IncreaseBrushSpeedAction() {}
+	virtual void DecreaseBrushSpeedAction() {}
+	virtual void NextBrushModeAction() {}
+	virtual void PreviousBrushModeAction() {}
+
+
+	//
+	// Stamps
+	//
+protected:
+	FSculptBrushStamp HoverStamp;
+	FSculptBrushStamp CurrentStamp;
+	FSculptBrushStamp LastStamp;
+	void UpdateHoverStamp(const FVector3d& WorldPos, const FVector3d& WorldNormal);
+
+
+
+
+	//
+	// Display / Material
+	//
+public:
+	UPROPERTY()
+	UMeshEditingViewProperties* ViewProperties;
+
+	UPROPERTY()
+	UMaterialInstanceDynamic* ActiveOverrideMaterial;
+
+protected:
+	virtual void UpdateWireframeVisibility(bool bNewValue);
+	virtual void UpdateMaterialMode(EMeshEditingMaterialModes NewMode);
+	virtual void UpdateFlatShadingSetting(bool bNewValue);
+	virtual void UpdateColorSetting(FLinearColor NewColor);
+	virtual void UpdateImageSetting(UTexture2D* NewImage);
+
+
+
+	//
+	// brush indicator
+	//
+protected:
+	// subclasses should call this to create indicator in their ::Setup()
+	virtual void InitializeIndicator();
+
+protected:
+	UPROPERTY()
+	UBrushStampIndicator* BrushIndicator;
+
+	UPROPERTY()
+	UMaterialInstanceDynamic* BrushIndicatorMaterial;
+
+	UPROPERTY()
+	UPreviewMesh* BrushIndicatorMesh;
+
+	// creates default sphere indicator
+	UPreviewMesh* MakeDefaultIndicatorSphereMesh(UObject* Parent, UWorld* World, int Resolution = 32);
+
+
+
+
+	//
+	// Work Plane
+	//
+public:
+	// plane gizmo
+	UPROPERTY()
+	UTransformGizmo* PlaneTransformGizmo;
+
+	UPROPERTY()
+	UTransformProxy* PlaneTransformProxy;
+
+protected:
+	virtual void UpdateWorkPlane();
+	virtual bool ShowWorkPlane() const { return false; };
+
+private:
+	TValueWatcher<FVector> GizmoPositionWatcher;
+	TValueWatcher<FQuat> GizmoRotationWatcher;
+
+	void UpdateGizmoFromProperties();
+	void PlaneTransformChanged(UTransformProxy* Proxy, FTransform Transform);
+
+	enum class EPendingWorkPlaneUpdate
+	{
+		NoUpdatePending,
+		MoveToHitPositionNormal,
+		MoveToHitPosition,
+		MoveToHitPositionViewAligned
+	};
+	EPendingWorkPlaneUpdate PendingWorkPlaneUpdate;
+	void SetFixedSculptPlaneFromWorldPos(const FVector& Position, const FVector& Normal, EPendingWorkPlaneUpdate UpdateType);
+	void UpdateFixedSculptPlanePosition(const FVector& Position);
+	void UpdateFixedSculptPlaneRotation(const FQuat& Rotation);
+	void UpdateFixedPlaneGizmoVisibility(bool bVisible);
+
+};
\ No newline at end of file
diff --git a/Engine/Plugins/Experimental/MeshModelingToolset/Source/MeshModelingTools/Public/Sculpting/MeshSmoothingBrushOps.h b/Engine/Plugins/Experimental/MeshModelingToolset/Source/MeshModelingTools/Public/Sculpting/MeshSmoothingBrushOps.h
new file mode 100644
--- /dev/null
+++ b/Engine/Plugins/Experimental/MeshModelingToolset/Source/MeshModelingTools/Public/Sculpting/MeshSmoothingBrushOps.h
@@ -0,0 +1,92 @@
+// Copyright Epic Games, Inc. All Rights Reserved.
+
+#pragma once
+
+#include "Sculpting/MeshBrushOpBase.h"
+#include "DynamicMesh3.h"
+#include "MeshWeights.h"
+#include "Async/ParallelFor.h"
+
+
+class FSmoothBrushOp : public FMeshSculptBrushOp
+{
+
+public:
+
+	virtual void ApplyStamp(const FDynamicMesh3* Mesh, const FSculptBrushStamp& Stamp, const TArray<int32>& Vertices, TArray<FVector3d>& NewPositionsOut) override
+	{
+		const FVector3d& StampPos = Stamp.LocalFrame.Origin;
+
+		ParallelFor(Vertices.Num(), [&](int32 k)
+		{
+			int32 VertIdx = Vertices[k];
+
+			FVector3d OrigPos = Mesh->GetVertex(VertIdx);
+
+			double Falloff = GetFalloff().Evaluate(Stamp, OrigPos);
+
+			FVector3d SmoothedPos = (CurrentOptions.bPreserveUVFlow) ?
+				FMeshWeights::MeanValueCentroid(*Mesh, VertIdx) : FMeshWeights::UniformCentroid(*Mesh, VertIdx);
+
+			FVector3d NewPos = FVector3d::Lerp(OrigPos, SmoothedPos, Falloff * Stamp.Power);
+
+			NewPositionsOut[k] = NewPos;
+		});
+	}
+};
+
+
+
+
+
+
+
+
+
+
+class FFlattenBrushOp : public FMeshSculptBrushOp
+{
+
+public:
+	double BrushSpeedTuning = 0.05;
+
+	virtual void ApplyStamp(const FDynamicMesh3* Mesh, const FSculptBrushStamp& Stamp, const TArray<int32>& Vertices, TArray<FVector3d>& NewPositionsOut) override
+	{
+		static const double PlaneSigns[3] = { 0, -1, 1 };
+		double PlaneSign = PlaneSigns[CurrentOptions.WhichPlaneSideIndex];
+
+		const FVector3d& StampPos = Stamp.LocalFrame.Origin;
+
+		double UseSpeed = Stamp.Power * Stamp.Radius * BrushSpeedTuning;
+		const FFrame3d& FlattenPlane = Stamp.RegionPlane;
+
+		ParallelFor(Vertices.Num(), [&](int32 k)
+		{
+			int32 VertIdx = Vertices[k];
+			FVector3d OrigPos = Mesh->GetVertex(VertIdx);
+			FVector3d PlanePos = FlattenPlane.ToPlane(OrigPos, 2);
+			FVector3d Delta = PlanePos - OrigPos;
+
+			double Dot = Delta.Dot(FlattenPlane.Z());
+			FVector3d NewPos = OrigPos;
+			if (Dot * PlaneSign >= 0)
+			{
+				double Falloff = GetFalloff().Evaluate(Stamp, OrigPos);
+				FVector3d MoveVec = Falloff * UseSpeed * Delta;
+				double MaxDist = Delta.Normalize();
+				NewPos = (MoveVec.SquaredLength() > MaxDist * MaxDist) ?
+					PlanePos : (OrigPos + Falloff * MoveVec);
+			}
+
+			NewPositionsOut[k] = NewPos;
+		});
+	}
+
+
+	virtual bool WantsStampRegionPlane() const
+	{
+		return true;
+	}
+};
+
+
diff --git a/Engine/Plugins/Experimental/MeshModelingToolset/Source/MeshModelingTools/Public/SelfUnionMeshesTool.h b/Engine/Plugins/Experimental/MeshModelingToolset/Source/MeshModelingTools/Public/SelfUnionMeshesTool.h
new file mode 100644
--- /dev/null
+++ b/Engine/Plugins/Experimental/MeshModelingToolset/Source/MeshModelingTools/Public/SelfUnionMeshesTool.h
@@ -0,0 +1,117 @@
+// Copyright Epic Games, Inc. All Rights Reserved.
+
+#pragma once
+
+#include "CoreMinimal.h"
+#include "UObject/NoExportTypes.h"
+#include "MultiSelectionTool.h"
+#include "InteractiveToolBuilder.h"
+#include "Drawing/LineSetComponent.h"
+#include "MeshOpPreviewHelpers.h"
+#include "BaseTools/SingleClickTool.h"
+
+#include "CompositionOps/SelfUnionMeshesOp.h"
+
+#include "SelfUnionMeshesTool.generated.h"
+
+// predeclarations
+class FDynamicMesh3;
+
+
+UCLASS()
+class MESHMODELINGTOOLS_API USelfUnionMeshesToolBuilder : public UInteractiveToolBuilder
+{
+	GENERATED_BODY()
+
+public:
+	IToolsContextAssetAPI* AssetAPI = nullptr;
+
+	virtual bool CanBuildTool(const FToolBuilderState& SceneState) const override;
+	virtual UInteractiveTool* BuildTool(const FToolBuilderState& SceneState) const override;
+};
+
+
+
+/**
+ * Standard properties of the self-union operation
+ */
+UCLASS()
+class MESHMODELINGTOOLS_API USelfUnionMeshesToolProperties : public UInteractiveToolPropertySet
+{
+	GENERATED_BODY()
+public:
+
+	/** Remove the source Actors/Components when accepting results of tool.*/
+	UPROPERTY(EditAnywhere, Category = Options)
+	bool bDeleteInputActors = true;
+
+	/** Automatically attempt to fill any holes left by merging (e.g. due to numerical errors) */
+	UPROPERTY(EditAnywhere, Category = Options)
+	bool bAttemptFixHoles = false;
+
+	/** Show boundary edges created by the union operation -- often due to numerical error */
+	UPROPERTY(EditAnywhere, Category = Options)
+	bool bShowNewBoundaryEdges = true;
+};
+
+
+
+/**
+ * Union of meshes, resolving self intersections
+ */
+UCLASS()
+class MESHMODELINGTOOLS_API USelfUnionMeshesTool : public UMultiSelectionTool, public IDynamicMeshOperatorFactory
+{
+	GENERATED_BODY()
+
+public:
+
+	USelfUnionMeshesTool();
+
+	virtual void Setup() override;
+	virtual void Shutdown(EToolShutdownType ShutdownType) override;
+
+	virtual void SetWorld(UWorld* World);
+	virtual void SetAssetAPI(IToolsContextAssetAPI* AssetAPI);
+
+	virtual void OnTick(float DeltaTime) override;
+	virtual void Render(IToolsContextRenderAPI* RenderAPI) override;
+
+	virtual bool HasCancel() const override { return true; }
+	virtual bool HasAccept() const override;
+	virtual bool CanAccept() const override;
+
+#if WITH_EDITOR
+	virtual void PostEditChangeProperty(FPropertyChangedEvent &PropertyChangedEvent) override;
+#endif
+
+	virtual void OnPropertyModified(UObject* PropertySet, FProperty* Property) override;
+
+	// IDynamicMeshOperatorFactory API
+	virtual TUniquePtr<FDynamicMeshOperator> MakeNewOperator() override;
+
+protected:
+
+	UPROPERTY()
+	UMeshOpPreviewWithBackgroundCompute* Preview;
+
+	UPROPERTY()
+	USelfUnionMeshesToolProperties* Properties;
+
+	UPROPERTY()
+	ULineSetComponent* DrawnLineSet;
+
+	TSharedPtr<FDynamicMesh3> CombinedSourceMeshes;
+
+	UWorld* TargetWorld;
+	IToolsContextAssetAPI* AssetAPI;
+
+	void SetupPreview();
+
+	void GenerateAsset(const FDynamicMeshOpResult& Result);
+
+	void UpdateVisualization();
+
+	// for visualization of any errors in the currently-previewed merge operation
+	TArray<int> CreatedBoundaryEdges;
+};
diff --git a/Engine/Plugins/Experimental/MeshModelingToolset/Source/MeshModelingTools/Public/UVProjectionTool.h b/Engine/Plugins/Experimental/MeshModelingToolset/Source/MeshModelingTools/Public/UVProjectionTool.h
--- a/Engine/Plugins/Experimental/MeshModelingToolset/Source/MeshModelingTools/Public/UVProjectionTool.h
+++ b/Engine/Plugins/Experimental/MeshModelingToolset/Source/MeshModelingTools/Public/UVProjectionTool.h
@@ -78,13 +78,6 @@ public:
 	/** If set, UV scales will be relative to world space so different objects created with the same UV scale should have the same average texel size */
 	UPROPERTY(EditAnywhere, Category = ProjectionSettings, meta = (DisplayName = "UV Scale Relative to World Space"))
 	bool bWorldSpaceUVScale = false;
-
-	//
-	// save/restore support
-	//
-	virtual void SaveProperties(UInteractiveTool* SaveFromTool) override;
-	virtual void RestoreProperties(UInteractiveTool* RestoreToTool) override;
-
 };
 
 
@@ -143,7 +136,7 @@ public:
 	virtual void SetWorld(UWorld* World, UInteractiveGizmoManager* GizmoManagerIn);
 	virtual void SetAssetAPI(IToolsContextAssetAPI* AssetAPI);
 
-	virtual void Tick(float DeltaTime) override;
+	virtual void OnTick(float DeltaTime) override;
 	virtual void Render(IToolsContextRenderAPI* RenderAPI) override;
 
 	virtual bool HasCancel() const override { return true; }
diff --git a/Engine/Plugins/Experimental/MeshModelingToolset/Source/MeshModelingToolsEditorOnly/Private/AttributeEditorTool.cpp b/Engine/Plugins/Experimental/MeshModelingToolset/Source/MeshModelingToolsEditorOnly/Private/AttributeEditorTool.cpp
--- a/Engine/Plugins/Experimental/MeshModelingToolset/Source/MeshModelingToolsEditorOnly/Private/AttributeEditorTool.cpp
+++ b/Engine/Plugins/Experimental/MeshModelingToolset/Source/MeshModelingToolsEditorOnly/Private/AttributeEditorTool.cpp
@@ -46,7 +46,6 @@ UInteractiveTool* UAttributeEditorToolBuilder::BuildTool(const FToolBuilderState
 
 	NewTool->SetSelection(MoveTemp(ComponentTargets));
 	NewTool->SetWorld(SceneState.World);
-	NewTool->SetAssetAPI(AssetAPI);
 
 	return NewTool;
 }
@@ -54,39 +53,30 @@ UInteractiveTool* UAttributeEditorToolBuilder::BuildTool(const FToolBuilderState
 
 
 
+void UAttributeEditorActionPropertySet::PostAction(EAttributeEditorToolActions Action)
+{
+	if (ParentTool.IsValid() && Cast<UAttributeEditorTool>(ParentTool))
+	{
+		Cast<UAttributeEditorTool>(ParentTool)->RequestAction(Action);
+	}
+}
+
+
+
+
 /*
  * Tool
  */
 
- // UObject interface
-#if WITH_EDITOR
-bool
-UAttributeEditorToolProperties::CanEditChange( const FProperty* InProperty) const
-{
-	static TArray<FString> UVLayers
-	{
-		TEXT("bClearUVLayer0"),
-		TEXT("bClearUVLayer1"),
-		TEXT("bClearUVLayer2"),
-		TEXT("bClearUVLayer3"),
-		TEXT("bClearUVLayer4"),
-		TEXT("bClearUVLayer5"),
-		TEXT("bClearUVLayer6"),
-		TEXT("bClearUVLayer7")
-	};
 
-	int32 Index;
-	if (UVLayers.Find(InProperty->GetFName().ToString(), Index))
-	{
-		return Index < NumUVLayers;
-	}
-	else
-	{
-		return true;
-	}
+
+
+TArray<FString> UAttributeEditorModifyAttributeActions::GetAttributeNamesFunc()
+{
+	return AttributeNamesList;
 }
-#endif // WITH_EDITOR	
-// End of UObject interface
+
+
 
 
 UAttributeEditorTool::UAttributeEditorTool()
@@ -111,60 +101,335 @@ void UAttributeEditorTool::Setup()
 		NumUVLayers = FMath::Max(NumUVLayers, InstanceUVs.GetNumIndices());
 	}
 
-	RemovalProperties = NewObject<UAttributeEditorToolProperties>(this);
-	RemovalProperties->NumUVLayers = NumUVLayers;
-	AddToolPropertySource(RemovalProperties);
+	NormalsActions = NewObject<UAttributeEditorNormalsActions>(this);
+	NormalsActions->Initialize(this);
+	AddToolPropertySource(NormalsActions);
+
+	UVActions = NewObject<UAttributeEditorUVActions>(this);
+	UVActions->Initialize(this);
+	AddToolPropertySource(UVActions);
+	UVActions->NumUVLayers = NumUVLayers;
+
+	if (ComponentTargets.Num() == 1)
+	{
+		NewAttributeProps = NewObject<UAttributeEditorNewAttributeActions>(this);
+		NewAttributeProps->Initialize(this);
+		AddToolPropertySource(NewAttributeProps);
+
+		ModifyAttributeProps = NewObject<UAttributeEditorModifyAttributeActions>(this);
+		ModifyAttributeProps->Initialize(this);
+		AddToolPropertySource(ModifyAttributeProps);
+		//SetToolPropertySourceEnabled(ModifyAttributeProps, false);
+
+		CopyAttributeProps = NewObject<UAttributeEditorCopyAttributeActions>(this);
+		CopyAttributeProps->Initialize(this);
+		AddToolPropertySource(CopyAttributeProps);
+		SetToolPropertySourceEnabled(CopyAttributeProps, false);
+
+		AttributeProps = NewObject<UAttributeEditorAttribProperties>(this);
+		AddToolPropertySource(AttributeProps);
+
+		InitializeAttributeLists();
+	}
+
+	GetToolManager()->DisplayMessage(
+		LOCTEXT("OnStartAttribEditor", "Inspect and Modify Attributes of a StaticMesh Asset"),
+		EToolMessageLevel::UserNotification);
 }
 
 
+template<typename AttribSetType>
+void ExtractAttribList(FMeshDescription* Mesh, AttribSetType& AttribSet, EAttributeEditorElementType ElemType, TArray<FAttributeEditorAttribInfo>& AttribList, TArray<FString>& StringList)
+{
+	AttribList.Reset();
+	StringList.Reset();
+
+	static FString EnumStrings[] = {
+		"Int32", "Boolean", "Float", "Vector2", "Vector3", "Vector4", "String", "Unknown"
+	};
+
+	AttribSet.ForEach([&](const FName AttributeName, auto AttributesRef)
+	{
+		FAttributeEditorAttribInfo AttribInfo;
+		AttribInfo.Name = AttributeName;
+		AttribInfo.ElementType = ElemType;
+		AttribInfo.DataType = EAttributeEditorAttribType::Unknown;
+		if (AttribSet.template HasAttributeOfType<int32>(AttributeName))
+		{
+			AttribInfo.DataType = EAttributeEditorAttribType::Int32;
+		}
+		else if (AttribSet.template HasAttributeOfType<float>(AttributeName))
+		{
+			AttribInfo.DataType = EAttributeEditorAttribType::Float;
+		}
+		else if (AttribSet.template HasAttributeOfType<bool>(AttributeName))
+		{
+			AttribInfo.DataType = EAttributeEditorAttribType::Boolean;
+		}
+		else if (AttribSet.template HasAttributeOfType<FVector2D>(AttributeName))
+		{
+			AttribInfo.DataType = EAttributeEditorAttribType::Vector2;
+		}
+		else if (AttribSet.template HasAttributeOfType<FVector>(AttributeName))
+		{
+			AttribInfo.DataType = EAttributeEditorAttribType::Vector3;
+		}
+		else if (AttribSet.template HasAttributeOfType<FVector4>(AttributeName))
+		{
+			AttribInfo.DataType = EAttributeEditorAttribType::Vector4;
+		}
+		else if (AttribSet.template HasAttributeOfType<FName>(AttributeName))
+		{
+			AttribInfo.DataType = EAttributeEditorAttribType::String;
+		}
+		AttribList.Add(AttribInfo);
+
+		bool bIsAutoGen = (AttributesRef.GetFlags() & EMeshAttributeFlags::AutoGenerated) != EMeshAttributeFlags::None;
+		
+		FString UIString = (bIsAutoGen) ?
+			FString::Printf(TEXT("%s - %s (autogen)"), *(AttribInfo.Name.ToString()), *EnumStrings[(int32)AttribInfo.DataType])
+			: FString::Printf(TEXT("%s - %s"), *(AttribInfo.Name.ToString()), *EnumStrings[(int32)AttribInfo.DataType]);
+		StringList.Add(UIString);
+	});
+}
+
+
+
+static FAttributesSetBase* GetAttributeSetByType(FMeshDescription* Mesh, EAttributeEditorElementType ElemType)
+{
+	switch (ElemType)
+	{
+	case EAttributeEditorElementType::Vertex:
+		return &Mesh->VertexAttributes();
+	case EAttributeEditorElementType::VertexInstance:
+		return &Mesh->VertexInstanceAttributes();
+	case EAttributeEditorElementType::Triangle:
+		return &Mesh->TriangleAttributes();
+	case EAttributeEditorElementType::Polygon:
+		return &Mesh->PolygonAttributes();
+	case EAttributeEditorElementType::Edge:
+		return &Mesh->EdgeAttributes();
+	case EAttributeEditorElementType::PolygonGroup:
+		return &Mesh->PolygonGroupAttributes();
+	}
+	check(false);
+	return nullptr;
+}
+
+static bool HasAttribute(FMeshDescription* Mesh, EAttributeEditorElementType ElemType, FName AttributeName)
+{
+	FAttributesSetBase* AttribSetBase = GetAttributeSetByType(Mesh, ElemType);
+	return (AttribSetBase) ? AttribSetBase->HasAttribute(AttributeName) : false;
+}
+
+static bool AddAttribute(FMeshDescription* Mesh, EAttributeEditorElementType ElemType, EAttributeEditorAttribType AttribType, FName AttributeName)
+{
+	FAttributesSetBase* AttribSetBase = GetAttributeSetByType(Mesh, ElemType);
+	if (AttribSetBase != nullptr)
+	{
+		switch (AttribType)
+		{
+		case EAttributeEditorAttribType::Int32:
+			AttribSetBase->RegisterAttribute<int32>(AttributeName, 1, 0, EMeshAttributeFlags::None);
+			return true;
+		case EAttributeEditorAttribType::Boolean:
+			AttribSetBase->RegisterAttribute<bool>(AttributeName, 1, false, EMeshAttributeFlags::None);
+			return true;
+		case EAttributeEditorAttribType::Float:
+			AttribSetBase->RegisterAttribute<float>(AttributeName, 1, 0.0f, EMeshAttributeFlags::Lerpable);
+			return true;
+		case EAttributeEditorAttribType::Vector2:
+			AttribSetBase->RegisterAttribute<FVector2D>(AttributeName, 1, FVector2D::ZeroVector, EMeshAttributeFlags::Lerpable);
+			return true;
+		case EAttributeEditorAttribType::Vector3:
+			AttribSetBase->RegisterAttribute<FVector>(AttributeName, 1, FVector::ZeroVector, EMeshAttributeFlags::Lerpable);
+			return true;
+		case EAttributeEditorAttribType::Vector4:
+			AttribSetBase->RegisterAttribute<FVector4>(AttributeName, 1, FVector4(0,0,0,1), EMeshAttributeFlags::Lerpable);
+			return true;
+		}
+	}
+	return false;
+}
+
+
+
+static bool RemoveAttribute(FMeshDescription* Mesh, EAttributeEditorElementType ElemType, FName AttributeName)
+{
+	FAttributesSetBase* AttribSetBase = GetAttributeSetByType(Mesh, ElemType);
+	if (AttribSetBase != nullptr)
+	{
+		AttribSetBase->UnregisterAttribute(AttributeName);
+		return true;
+	}
+	return false;
+}
+
+
+static bool IsReservedName(FName AttributeName)
+{
+	return (AttributeName == MeshAttribute::Vertex::Position)
+		|| (AttributeName == MeshAttribute::Vertex::CornerSharpness)
+		|| (AttributeName == MeshAttribute::VertexInstance::TextureCoordinate)
+		|| (AttributeName == MeshAttribute::VertexInstance::Normal)
+		|| (AttributeName == MeshAttribute::VertexInstance::Tangent)
+		|| (AttributeName == MeshAttribute::VertexInstance::BinormalSign)
+		|| (AttributeName == MeshAttribute::VertexInstance::Color)
+		|| (AttributeName == MeshAttribute::Edge::IsHard)
+		|| (AttributeName == MeshAttribute::Edge::IsUVSeam)
+		|| (AttributeName == MeshAttribute::Edge::CreaseSharpness)
+		|| (AttributeName == MeshAttribute::Polygon::Normal)
+		|| (AttributeName == MeshAttribute::Polygon::Tangent)
+		|| (AttributeName == MeshAttribute::Polygon::Binormal)
+		|| (AttributeName == MeshAttribute::Polygon::Center)
+		|| (AttributeName == MeshAttribute::PolygonGroup::ImportedMaterialSlotName)
+		|| (AttributeName == MeshAttribute::PolygonGroup::EnableCollision)
+		|| (AttributeName == MeshAttribute::PolygonGroup::CastShadow);
+}
+
+
+void UAttributeEditorTool::InitializeAttributeLists()
+{
+	FMeshDescription* Mesh = ComponentTargets[0]->GetMesh();
+
+	ExtractAttribList(Mesh, Mesh->VertexAttributes(), EAttributeEditorElementType::Vertex, this->VertexAttributes, AttributeProps->VertexAttributes);
+	ExtractAttribList(Mesh, Mesh->VertexInstanceAttributes(), EAttributeEditorElementType::VertexInstance, this->InstanceAttributes, AttributeProps->InstanceAttributes);
+	ExtractAttribList(Mesh, Mesh->TriangleAttributes(), EAttributeEditorElementType::Triangle, this->TriangleAttributes, AttributeProps->TriangleAttributes);
+	ExtractAttribList(Mesh, Mesh->PolygonAttributes(), EAttributeEditorElementType::Polygon, this->PolygonAttributes, AttributeProps->PolygonAttributes);
+	ExtractAttribList(Mesh, Mesh->EdgeAttributes(), EAttributeEditorElementType::Edge, this->EdgeAttributes, AttributeProps->EdgeAttributes);
+	ExtractAttribList(Mesh, Mesh->PolygonGroupAttributes(), EAttributeEditorElementType::PolygonGroup, this->GroupAttributes, AttributeProps->GroupAttributes);
+
+	ModifyAttributeProps->AttributeNamesList.Reset();
+	CopyAttributeProps->FromAttribute.Reset();
+	CopyAttributeProps->ToAttribute.Reset();
+
+	//TArray<TArray<FAttributeEditorAttribInfo>*> AttribInfos = {
+	//	&this->VertexAttributes, &this->InstanceAttributes,
+	//	&this->TriangleAttributes, & this->PolygonAttributes,
+	//	& this->EdgeAttributes, & this->GroupAttributes };
+	TArray<TArray<FAttributeEditorAttribInfo>*> AttribInfos = {
+		&this->VertexAttributes, &this->PolygonAttributes };
+
+	for (TArray<FAttributeEditorAttribInfo>* AttribInfoList : AttribInfos)
+	{
+		for (FAttributeEditorAttribInfo& AttribInfo : *AttribInfoList)
+		{
+			//if (IsReservedName(AttribInfo.Name) == false)
+			//{
+				ModifyAttributeProps->AttributeNamesList.Add(AttribInfo.Name.ToString());
+			//}
+			//CopyAttributeProps->FromAttribute.Add(AttribInfo.Name);
+			//CopyAttributeProps->ToAttribute.Add(AttribInfo.Name);
+		}
+	}
+
+	bAttributeListsValid = true;
+}
+
 
 
 void UAttributeEditorTool::Shutdown(EToolShutdownType ShutdownType)
 {
-	if (ShutdownType == EToolShutdownType::Accept)
+}
+
+
+void UAttributeEditorTool::RequestAction(EAttributeEditorToolActions ActionType)
+{
+	if (PendingAction == EAttributeEditorToolActions::NoAction)
 	{
-		GenerateAssets();
+		PendingAction = ActionType;
 	}
 }
 
-void UAttributeEditorTool::GenerateAssets()
+
+void UAttributeEditorTool::OnTick(float DeltaTime)
 {
-	GetToolManager()->BeginUndoTransaction(LOCTEXT("AttributeEditorToolTransactionName", "Attribute Editor Tool Transaction"));
+	switch (PendingAction)
+	{
+	case EAttributeEditorToolActions::ClearNormals:
+		ClearNormals();
+		break;
+	case EAttributeEditorToolActions::ClearSelectedUVs:
+		ClearUVs(true);
+		break;
+	case EAttributeEditorToolActions::ClearAllUVs:
+		ClearUVs(false);
+		break;
+	case EAttributeEditorToolActions::AddAttribute:
+		AddNewAttribute();
+		break;
+	case EAttributeEditorToolActions::AddWeightMapLayer:
+		AddNewWeightMap();
+		break;
+	case EAttributeEditorToolActions::AddPolyGroupLayer:
+		AddNewGroupsLayer();
+		break;
+	case EAttributeEditorToolActions::DeleteAttribute:
+		DeleteAttribute();
+		break;
+	}
+	PendingAction = EAttributeEditorToolActions::NoAction;
 
-	for (TUniquePtr<FPrimitiveComponentTarget>& ComponentTarget : ComponentTargets)
+	if (bAttributeListsValid == false && ComponentTargets.Num() == 1)
+	{
+		InitializeAttributeLists();
+	}
+}
+
+
+
+void UAttributeEditorTool::ClearNormals()
+{
+	GetToolManager()->BeginUndoTransaction(LOCTEXT("ClearNormalsTransactionMessage", "Clear Normals"));
+
+	for (TUniquePtr<FPrimitiveComponentTarget>& Target : ComponentTargets)
 	{
-		ComponentTarget->CommitMesh([this](const FPrimitiveComponentTarget::FCommitParams& CommitParams)
+		Target->CommitMesh([&](const FPrimitiveComponentTarget::FCommitParams& CommitParams)
 		{
-			if (RemovalProperties->bClearNormals)
+			TEdgeAttributesRef<bool> EdgeHardnesses = CommitParams.MeshDescription->EdgeAttributes().GetAttributesRef<bool>(MeshAttribute::Edge::IsHard);
+			if (EdgeHardnesses.IsValid())
 			{
-				TEdgeAttributesRef<bool> EdgeHardnesses = CommitParams.MeshDescription->EdgeAttributes().GetAttributesRef<bool>(MeshAttribute::Edge::IsHard);
-				if (EdgeHardnesses.IsValid())
+				for (FEdgeID ElID : CommitParams.MeshDescription->Edges().GetElementIDs())
 				{
-					for (FEdgeID ElID : CommitParams.MeshDescription->Edges().GetElementIDs())
-					{
-						EdgeHardnesses[ElID] = false;
-					}
+					EdgeHardnesses[ElID] = false;
 				}
-				FStaticMeshOperations::ComputePolygonTangentsAndNormals(*CommitParams.MeshDescription, FMathf::Epsilon);
-				FStaticMeshOperations::RecomputeNormalsAndTangentsIfNeeded(*CommitParams.MeshDescription, EComputeNTBsFlags::WeightedNTBs | EComputeNTBsFlags::Normals);
 			}
+			FStaticMeshOperations::ComputePolygonTangentsAndNormals(*CommitParams.MeshDescription, FMathf::Epsilon);
+			FStaticMeshOperations::RecomputeNormalsAndTangentsIfNeeded(*CommitParams.MeshDescription, EComputeNTBsFlags::WeightedNTBs | EComputeNTBsFlags::Normals);
+		});
+	}
+	GetToolManager()->EndUndoTransaction();
+}
+
+
+
+
+void UAttributeEditorTool::ClearUVs(bool bSelectedOnly)
+{
+	GetToolManager()->BeginUndoTransaction(LOCTEXT("ClearUVsTransactionMessage", "Clear Selected UVs"));
+	for (TUniquePtr<FPrimitiveComponentTarget>& Target : ComponentTargets)
+	{
+		Target->CommitMesh([&](const FPrimitiveComponentTarget::FCommitParams& CommitParams)
+		{
 			bool RemoveLayers[8] =
 			{
-				RemovalProperties->bClearUVLayer0,
-				RemovalProperties->bClearUVLayer1,
-				RemovalProperties->bClearUVLayer2,
-				RemovalProperties->bClearUVLayer3,
-				RemovalProperties->bClearUVLayer4,
-				RemovalProperties->bClearUVLayer5,
-				RemovalProperties->bClearUVLayer6,
-				RemovalProperties->bClearUVLayer7
+				UVActions->bClearUVLayer0,
+				UVActions->bClearUVLayer1,
+				UVActions->bClearUVLayer2,
+				UVActions->bClearUVLayer3,
+				UVActions->bClearUVLayer4,
+				UVActions->bClearUVLayer5,
+				UVActions->bClearUVLayer6,
+				UVActions->bClearUVLayer7
 			};
+
 			TVertexInstanceAttributesRef<FVector2D> InstanceUVs =
 				CommitParams.MeshDescription->VertexInstanceAttributes().GetAttributesRef<FVector2D>(MeshAttribute::VertexInstance::TextureCoordinate);
 			const FVertexInstanceArray& Instances = CommitParams.MeshDescription->VertexInstances();
 			for (int LayerIndex = 7; LayerIndex >= 0; LayerIndex--)
 			{
-				if (RemoveLayers[LayerIndex] && LayerIndex < InstanceUVs.GetNumIndices())
+				bool bRemove = (bSelectedOnly == false) || RemoveLayers[LayerIndex];
+				if (bRemove && LayerIndex < InstanceUVs.GetNumIndices())
 				{
 					if (!FStaticMeshOperations::RemoveUVChannel(*CommitParams.MeshDescription, LayerIndex))
 					{
@@ -177,45 +442,114 @@ void UAttributeEditorTool::GenerateAssets()
 			}
 		});
 	}
-
 	GetToolManager()->EndUndoTransaction();
 }
 
 
 
-void UAttributeEditorTool::SetAssetAPI(IToolsContextAssetAPI* AssetAPIIn)
+
+
+void UAttributeEditorTool::AddNewAttribute(EAttributeEditorElementType ElemType, EAttributeEditorAttribType DataType, FName AttributeName)
 {
-	this->AssetAPI = AssetAPIIn;
+	if (AttributeName.IsNone())
+	{
+		GetToolManager()->DisplayMessage(LOCTEXT("InvalidAttributeName", "Invalid attribute name"), EToolMessageLevel::UserWarning);
+		return;
+	}
+
+	FMeshDescription* CurMesh = ComponentTargets[0]->GetMesh();
+	if (HasAttribute(CurMesh, ElemType, AttributeName))
+	{
+		GetToolManager()->DisplayMessage(LOCTEXT("ErrorAddingDuplicateNameMessage", "Attribute with this name already exists"), EToolMessageLevel::UserWarning);
+		return;
+	}
+
+	FMeshDescription NewMesh = *CurMesh;
+	if (AddAttribute(&NewMesh, ElemType, DataType, AttributeName) == false)
+	{
+		GetToolManager()->DisplayMessage(LOCTEXT("FailedAddingNewMessage", "Unknown error adding new Attribute"), EToolMessageLevel::UserWarning);
+		return;
+	}
+
+	GetToolManager()->BeginUndoTransaction(LOCTEXT("NewAttributeTransactionMessage", "Add Attribute"));
+	ComponentTargets[0]->CommitMesh([&NewMesh](const FPrimitiveComponentTarget::FCommitParams& CommitParams)
+	{
+		*CommitParams.MeshDescription = NewMesh;
+	});
+	GetToolManager()->EndUndoTransaction();
+
+	bAttributeListsValid = false;
 }
 
 
-void UAttributeEditorTool::Render(IToolsContextRenderAPI* RenderAPI)
+void UAttributeEditorTool::AddNewAttribute()
 {
+	if (NewAttributeProps->DataType == EAttributeEditorAttribType::Unknown)
+	{
+		GetToolManager()->DisplayMessage(LOCTEXT("ErrorAddingTypeMessage", "Currently cannot add this attribute type"), EToolMessageLevel::UserWarning);
+		return;
+	}
+
+	AddNewAttribute(NewAttributeProps->ElementType, NewAttributeProps->DataType, FName(NewAttributeProps->NewName));
 }
 
-void UAttributeEditorTool::Tick(float DeltaTime)
+void UAttributeEditorTool::AddNewWeightMap()
 {
+	AddNewAttribute(EAttributeEditorElementType::Vertex, EAttributeEditorAttribType::Float, FName(NewAttributeProps->NewName));
 }
 
-
-void UAttributeEditorTool::OnPropertyModified(UObject* PropertySet, FProperty* Property)
+void UAttributeEditorTool::AddNewGroupsLayer()
 {
+	AddNewAttribute(EAttributeEditorElementType::Polygon, EAttributeEditorAttribType::Int32, FName(NewAttributeProps->NewName));
 }
 
 
-
-bool UAttributeEditorTool::HasAccept() const
+void UAttributeEditorTool::ClearAttribute()
 {
-	return true;
 }
 
-bool UAttributeEditorTool::CanAccept() const
+void UAttributeEditorTool::DeleteAttribute()
 {
-	return true;
+	FMeshDescription* CurMesh = ComponentTargets[0]->GetMesh();
+	FName SelectedName(ModifyAttributeProps->Attribute);
+
+	if (IsReservedName(SelectedName))
+	{
+		GetToolManager()->DisplayMessage(LOCTEXT("CannotDeleteReservedNameError", "Cannot delete required mesh Attributes"), EToolMessageLevel::UserWarning);
+		return;
+	}
+
+	bool bIsVertex = HasAttribute(CurMesh, EAttributeEditorElementType::Vertex, SelectedName);
+	bool bIsPoly = HasAttribute(CurMesh, EAttributeEditorElementType::PolygonGroup, SelectedName);
+	if (bIsVertex == false && bIsPoly == false)
+	{
+		GetToolManager()->DisplayMessage(LOCTEXT("CannotDeleteAttribError", "Cannot delete the selected attribute"), EToolMessageLevel::UserWarning);
+		return;
+	}
+
+	EAttributeEditorElementType ElemType = (bIsVertex) ? EAttributeEditorElementType::Vertex : EAttributeEditorElementType::Polygon;
+
+	FMeshDescription NewMesh = *CurMesh;
+	if (RemoveAttribute(&NewMesh, ElemType, SelectedName) == false)
+	{
+		GetToolManager()->DisplayMessage(LOCTEXT("FailedRemovingNewMessage", "Unknown error removing Attribute"), EToolMessageLevel::UserWarning);
+		return;
+	}
+
+	GetToolManager()->BeginUndoTransaction(LOCTEXT("RemoveAttributeTransactionMessage", "Remove Attribute"));
+	ComponentTargets[0]->CommitMesh([&NewMesh](const FPrimitiveComponentTarget::FCommitParams& CommitParams)
+	{
+		*CommitParams.MeshDescription = NewMesh;
+	});
+	GetToolManager()->EndUndoTransaction();
+
+	bAttributeListsValid = false;
 }
 
 
 
 
 
+
+
 #undef LOCTEXT_NAMESPACE
diff --git a/Engine/Plugins/Experimental/MeshModelingToolset/Source/MeshModelingToolsEditorOnly/Private/EditMeshMaterialsTool.cpp b/Engine/Plugins/Experimental/MeshModelingToolset/Source/MeshModelingToolsEditorOnly/Private/EditMeshMaterialsTool.cpp
--- a/Engine/Plugins/Experimental/MeshModelingToolset/Source/MeshModelingToolsEditorOnly/Private/EditMeshMaterialsTool.cpp
+++ b/Engine/Plugins/Experimental/MeshModelingToolset/Source/MeshModelingToolsEditorOnly/Private/EditMeshMaterialsTool.cpp
@@ -55,9 +55,9 @@ void UEditMeshMaterialsTool::Setup()
 	CurrentMaterials = MaterialProps->Materials;
 	InitialMaterialKey = GetMaterialKey();
 
-	MaterialSetWatcher.Initialize(
-		[this]() { return GetMaterialKey(); },
-		[this](FMaterialSetKey NewKey) { OnMaterialSetChanged(); }, GetMaterialKey());
+	MaterialProps->WatchProperty<FMaterialSetKey>(
+		[this](){ return GetMaterialKey(); },
+		[this](FMaterialSetKey NewKey) { OnMaterialSetChanged(); });
 
 	FComponentMaterialSet ComponentMaterials;
 	ComponentTarget->GetMaterialSet(ComponentMaterials, false);
@@ -101,11 +101,9 @@ void UEditMeshMaterialsTool::RequestMaterialAction(EEditMeshMaterialsToolActions
 
 
 
-void UEditMeshMaterialsTool::Tick(float DeltaTime)
+void UEditMeshMaterialsTool::OnTick(float DeltaTime)
 {
-	UMeshSelectionTool::Tick(DeltaTime);
-
-	MaterialSetWatcher.CheckAndUpdate();
+	UMeshSelectionTool::OnTick(DeltaTime);
 
 	if (bHavePendingSubAction)
 	{
@@ -147,7 +145,7 @@ void UEditMeshMaterialsTool::AssignMaterialToSelectedTriangles()
 		ActiveSelectionChange->Add(tid);
 	}
 	Selection->RemoveIndices(EMeshSelectionElementType::Face, SelectedFaces);
-	TUniquePtr<FMeshSelectionChange> SelectionChange = EndChange();
+	TUniquePtr<FToolCommandChange> SelectionChange = EndChange();
 	ChangeSeq->AppendChange(Selection, MoveTemp(SelectionChange));
 
 	int32 SetMaterialID = MaterialProps->SelectedMaterial;
@@ -198,8 +196,10 @@ void UEditMeshMaterialsTool::OnMaterialSetChanged()
 
 void UEditMeshMaterialsTool::ExternalUpdateMaterialSet(const TArray<UMaterialInterface*>& NewMaterialSet)
 {
+	// Disable props so they don't update
+	SetToolPropertySourceEnabled(MaterialProps, false);
 	MaterialProps->Materials = NewMaterialSet;
-	MaterialSetWatcher.SilentUpdate();
+	SetToolPropertySourceEnabled(MaterialProps, true);
 	PreviewMesh->SetMaterials(MaterialProps->Materials);
 	CurrentMaterials = MaterialProps->Materials;
 }
@@ -280,4 +280,4 @@ FString FEditMeshMaterials_MaterialSetChange::ToString() const
 	return FString(TEXT("MaterialSet Change"));
 }
 
-#undef LOCTEXT_NAMESPACE
\ No newline at end of file
+#undef LOCTEXT_NAMESPACE
diff --git a/Engine/Plugins/Experimental/MeshModelingToolset/Source/MeshModelingToolsEditorOnly/Private/MergeMeshesTool.cpp b/Engine/Plugins/Experimental/MeshModelingToolset/Source/MeshModelingToolsEditorOnly/Private/MergeMeshesTool.cpp
--- a/Engine/Plugins/Experimental/MeshModelingToolset/Source/MeshModelingToolsEditorOnly/Private/MergeMeshesTool.cpp
+++ b/Engine/Plugins/Experimental/MeshModelingToolset/Source/MeshModelingToolsEditorOnly/Private/MergeMeshesTool.cpp
@@ -168,17 +168,11 @@ void UMergeMeshesTool::Shutdown(EToolShutdownType ShutdownType)
 
 
 
-void UMergeMeshesTool::Tick(float DeltaTime)
+void UMergeMeshesTool::OnTick(float DeltaTime)
 {
 	Preview->Tick(DeltaTime);
 }
 
-
-void UMergeMeshesTool::Render(IToolsContextRenderAPI* RenderAPI)
-{
-}
-
-
 bool UMergeMeshesTool::HasAccept() const
 {
 	return true;
diff --git a/Engine/Plugins/Experimental/MeshModelingToolset/Source/MeshModelingToolsEditorOnly/Private/ParameterizeMeshTool.cpp b/Engine/Plugins/Experimental/MeshModelingToolset/Source/MeshModelingToolsEditorOnly/Private/ParameterizeMeshTool.cpp
--- a/Engine/Plugins/Experimental/MeshModelingToolset/Source/MeshModelingToolsEditorOnly/Private/ParameterizeMeshTool.cpp
+++ b/Engine/Plugins/Experimental/MeshModelingToolset/Source/MeshModelingToolsEditorOnly/Private/ParameterizeMeshTool.cpp
@@ -45,37 +45,9 @@ UInteractiveTool* UParameterizeMeshToolBuilder::BuildTool(const FToolBuilderStat
 	return NewTool;
 }
 
-
-
-
-
-void UParameterizeMeshToolProperties::SaveProperties(UInteractiveTool* SaveFromTool)
-{
-	UParameterizeMeshToolProperties* PropertyCache = GetPropertyCache<UParameterizeMeshToolProperties>();
-	PropertyCache->ChartStretch = this->ChartStretch;
-	//PropertyCache->IslandMode = this->IslandMode;
-	PropertyCache->UnwrapType = this->UnwrapType;
-	PropertyCache->UVScaleMode = this->UVScaleMode;
-	PropertyCache->UVScale = this->UVScale;
-}
-
-void UParameterizeMeshToolProperties::RestoreProperties(UInteractiveTool* RestoreToTool)
-{
-	UParameterizeMeshToolProperties* PropertyCache = GetPropertyCache<UParameterizeMeshToolProperties>();
-	this->ChartStretch = PropertyCache->ChartStretch;
-	//this->IslandMode = PropertyCache->IslandMode;
-	this->UnwrapType = PropertyCache->UnwrapType;
-	this->UVScaleMode = PropertyCache->UVScaleMode;
-	this->UVScale = PropertyCache->UVScale;
-}
-
-
-
-
 /*
  * Tool
  */
-
 UParameterizeMeshTool::UParameterizeMeshTool()
 {
 }
@@ -202,15 +174,11 @@ void UParameterizeMeshTool::Shutdown(EToolShutdownType ShutdownType)
 
 }
 
-void UParameterizeMeshTool::Tick(float DeltaTime)
+void UParameterizeMeshTool::OnTick(float DeltaTime)
 {
 	Preview->Tick(DeltaTime);
 }
 
-void UParameterizeMeshTool::Render(IToolsContextRenderAPI* RenderAPI)
-{
-}
-
 bool UParameterizeMeshTool::HasAccept() const
 {
 	return true;
diff --git a/Engine/Plugins/Experimental/MeshModelingToolset/Source/MeshModelingToolsEditorOnly/Private/PolygonOnMeshTool.cpp b/Engine/Plugins/Experimental/MeshModelingToolset/Source/MeshModelingToolsEditorOnly/Private/PolygonOnMeshTool.cpp
--- a/Engine/Plugins/Experimental/MeshModelingToolset/Source/MeshModelingToolsEditorOnly/Private/PolygonOnMeshTool.cpp
+++ b/Engine/Plugins/Experimental/MeshModelingToolset/Source/MeshModelingToolsEditorOnly/Private/PolygonOnMeshTool.cpp
@@ -251,7 +251,7 @@ void UPolygonOnMeshTool::Render(IToolsContextRenderAPI* RenderAPI)
 	}
 }
 
-void UPolygonOnMeshTool::Tick(float DeltaTime)
+void UPolygonOnMeshTool::OnTick(float DeltaTime)
 {
 	PlaneMechanic->Tick(DeltaTime);
 	Preview->Tick(DeltaTime);
