Edigrate Dev-Editor updates to InteractiveToolsFramework, GeometryProcessing, MeshModelingToolset, and ModelingTools Editor Mode Plugins. Update a few files in new MeshPaint Plugin for UInteractiveToo

diff --git a/Engine/Plugins/Experimental/MeshModelingToolset/Source/MeshModelingToolsEditorOnly/Private/SimplifyMeshTool.cpp b/Engine/Plugins/Experimental/MeshModelingToolset/Source/MeshModelingToolsEditorOnly/Private/SimplifyMeshTool.cpp
--- a/Engine/Plugins/Experimental/MeshModelingToolset/Source/MeshModelingToolsEditorOnly/Private/SimplifyMeshTool.cpp
+++ b/Engine/Plugins/Experimental/MeshModelingToolset/Source/MeshModelingToolsEditorOnly/Private/SimplifyMeshTool.cpp
@@ -75,30 +75,6 @@ USimplifyMeshToolProperties::USimplifyMeshToolProperties()
 	MaterialBoundaryConstraint = EMaterialBoundaryConstraint::Ignore;
 }
 
-void
-USimplifyMeshToolProperties::SaveRestoreProperties(UInteractiveTool* RestoreToTool, bool bSaving)
-{
-	USimplifyMeshToolProperties* PropertyCache = GetPropertyCache<USimplifyMeshToolProperties>();
-
-	// MeshConstraintProperties
-	SaveRestoreProperty(PropertyCache->bPreserveSharpEdges, this->bPreserveSharpEdges, bSaving);
-	SaveRestoreProperty(PropertyCache->MeshBoundaryConstraint, this->MeshBoundaryConstraint, bSaving);
-	SaveRestoreProperty(PropertyCache->GroupBoundaryConstraint, this->GroupBoundaryConstraint, bSaving);
-	SaveRestoreProperty(PropertyCache->MaterialBoundaryConstraint, this->MaterialBoundaryConstraint, bSaving);
-	SaveRestoreProperty(PropertyCache->bPreventNormalFlips, this->bPreventNormalFlips, bSaving);
-
-	// SimplifyMeshToolProperties
-	SaveRestoreProperty(PropertyCache->TargetMode, this->TargetMode, bSaving);
-	SaveRestoreProperty(PropertyCache->SimplifierType, this->SimplifierType, bSaving);
-	SaveRestoreProperty(PropertyCache->TargetPercentage, this->TargetPercentage, bSaving);
-	SaveRestoreProperty(PropertyCache->TargetEdgeLength, this->TargetEdgeLength, bSaving);
-	SaveRestoreProperty(PropertyCache->TargetCount, this->TargetCount, bSaving);
-	SaveRestoreProperty(PropertyCache->bDiscardAttributes, this->bDiscardAttributes, bSaving);
-	SaveRestoreProperty(PropertyCache->bShowWireframe, this->bShowWireframe, bSaving);
-	SaveRestoreProperty(PropertyCache->bShowGroupColors, this->bShowGroupColors, bSaving);
-	SaveRestoreProperty(PropertyCache->bReproject, this->bReproject, bSaving);
-}
-
 void USimplifyMeshTool::SetWorld(UWorld* World)
 {
 	this->TargetWorld = World;
@@ -159,12 +135,10 @@ void USimplifyMeshTool::Setup()
 	SimplifyProperties->RestoreProperties(this);
 	AddToolPropertySource(SimplifyProperties);
 
-	ShowGroupsWatcher.Initialize(
-		[this]() { return SimplifyProperties->bShowGroupColors; },
-		[this](bool bNewValue) { UpdateVisualization(); }, SimplifyProperties->bShowGroupColors );
-	ShowWireFrameWatcher.Initialize(
-		[this]() { return SimplifyProperties->bShowWireframe; },
-		[this](bool bNewValue) { UpdateVisualization(); }, SimplifyProperties->bShowWireframe );
+	SimplifyProperties->WatchProperty(SimplifyProperties->bShowGroupColors,
+									  [this](bool bNewValue) { UpdateVisualization(); });
+	SimplifyProperties->WatchProperty(SimplifyProperties->bShowWireframe,
+									  [this](bool bNewValue) { UpdateVisualization(); });
 
 	MeshStatisticsProperties = NewObject<UMeshStatisticsProperties>(this);
 	AddToolPropertySource(MeshStatisticsProperties);
@@ -191,11 +165,8 @@ void USimplifyMeshTool::Shutdown(EToolShutdownType ShutdownType)
 }
 
 
-void USimplifyMeshTool::Tick(float DeltaTime)
+void USimplifyMeshTool::OnTick(float DeltaTime)
 {
-	ShowWireFrameWatcher.CheckAndUpdate();
-	ShowGroupsWatcher.CheckAndUpdate();
-
 	Preview->Tick(DeltaTime);
 }
 
diff --git a/Engine/Plugins/Experimental/MeshModelingToolset/Source/MeshModelingToolsEditorOnly/Private/UVLayoutTool.cpp b/Engine/Plugins/Experimental/MeshModelingToolset/Source/MeshModelingToolsEditorOnly/Private/UVLayoutTool.cpp
--- a/Engine/Plugins/Experimental/MeshModelingToolset/Source/MeshModelingToolsEditorOnly/Private/UVLayoutTool.cpp
+++ b/Engine/Plugins/Experimental/MeshModelingToolset/Source/MeshModelingToolsEditorOnly/Private/UVLayoutTool.cpp
@@ -66,24 +66,6 @@ UUVLayoutToolProperties::UUVLayoutToolProperties()
 
 }
 
-
-void UUVLayoutToolProperties::SaveProperties(UInteractiveTool* SaveFromTool)
-{
-	UUVLayoutToolProperties* PropertyCache = GetPropertyCache<UUVLayoutToolProperties>();
-	PropertyCache->bSeparateUVIslands = this->bSeparateUVIslands;
-	PropertyCache->TextureResolution = this->TextureResolution;
-	PropertyCache->UVScaleFactor = this->UVScaleFactor;
-}
-
-void UUVLayoutToolProperties::RestoreProperties(UInteractiveTool* RestoreToTool)
-{
-	UUVLayoutToolProperties* PropertyCache = GetPropertyCache<UUVLayoutToolProperties>();
-	this->bSeparateUVIslands = PropertyCache->bSeparateUVIslands;
-	this->TextureResolution = PropertyCache->TextureResolution;
-	this->UVScaleFactor = PropertyCache->UVScaleFactor;
-}
-
-
 UUVLayoutAdvancedProperties::UUVLayoutAdvancedProperties()
 {
 }
@@ -213,7 +195,7 @@ void UUVLayoutTool::Render(IToolsContextRenderAPI* RenderAPI)
 {
 }
 
-void UUVLayoutTool::Tick(float DeltaTime)
+void UUVLayoutTool::OnTick(float DeltaTime)
 {
 	for (UMeshOpPreviewWithBackgroundCompute* Preview : Previews)
 	{
diff --git a/Engine/Plugins/Experimental/MeshModelingToolset/Source/MeshModelingToolsEditorOnly/Private/VoxelCSGMeshesTool.cpp b/Engine/Plugins/Experimental/MeshModelingToolset/Source/MeshModelingToolsEditorOnly/Private/VoxelCSGMeshesTool.cpp
--- a/Engine/Plugins/Experimental/MeshModelingToolset/Source/MeshModelingToolsEditorOnly/Private/VoxelCSGMeshesTool.cpp
+++ b/Engine/Plugins/Experimental/MeshModelingToolset/Source/MeshModelingToolsEditorOnly/Private/VoxelCSGMeshesTool.cpp
@@ -165,7 +165,7 @@ void UVoxelCSGMeshesTool::Shutdown(EToolShutdownType ShutdownType)
 
 
 
-void UVoxelCSGMeshesTool::Tick(float DeltaTime)
+void UVoxelCSGMeshesTool::OnTick(float DeltaTime)
 {
 	Preview->Tick(DeltaTime);
 }
diff --git a/Engine/Plugins/Experimental/MeshModelingToolset/Source/MeshModelingToolsEditorOnly/Public/AttributeEditorTool.h b/Engine/Plugins/Experimental/MeshModelingToolset/Source/MeshModelingToolsEditorOnly/Public/AttributeEditorTool.h
--- a/Engine/Plugins/Experimental/MeshModelingToolset/Source/MeshModelingToolsEditorOnly/Public/AttributeEditorTool.h
+++ b/Engine/Plugins/Experimental/MeshModelingToolset/Source/MeshModelingToolsEditorOnly/Public/AttributeEditorTool.h
@@ -28,61 +28,259 @@ public:
 
 
 
-/**
- * Standard properties
- */
+
+UENUM()
+enum class EAttributeEditorElementType : uint8
+{
+	Vertex = 0,
+	VertexInstance = 1,
+	Triangle = 2,
+	Polygon = 3,
+	Edge = 4,
+	PolygonGroup = 5
+};
+
+
+UENUM()
+enum class EAttributeEditorAttribType : uint8
+{
+	Int32 = 0,
+	Boolean = 1,
+	Float = 2,
+	Vector2 = 3,
+	Vector3 = 4,
+	Vector4 = 5,
+	String = 6,
+	Unknown = 7
+};
+
+
+struct FAttributeEditorAttribInfo
+{
+	FName Name;
+	EAttributeEditorElementType ElementType;
+	EAttributeEditorAttribType DataType;
+};
+
+
+UCLASS()
+class MESHMODELINGTOOLSEDITORONLY_API UAttributeEditorAttribProperties : public UInteractiveToolPropertySet
+{
+	GENERATED_BODY()
+public:
+	UPROPERTY(VisibleAnywhere, Category = AttributesInspector)
+	TArray<FString> VertexAttributes;
+
+	UPROPERTY(VisibleAnywhere, Category = AttributesInspector)
+	TArray<FString> InstanceAttributes;
+
+	UPROPERTY(VisibleAnywhere, Category = AttributesInspector)
+	TArray<FString> TriangleAttributes;
+
+	UPROPERTY(VisibleAnywhere, Category = AttributesInspector)
+	TArray<FString> PolygonAttributes;
+
+	UPROPERTY(VisibleAnywhere, Category = AttributesInspector)
+	TArray<FString> EdgeAttributes;
+
+	UPROPERTY(VisibleAnywhere, Category = AttributesInspector)
+	TArray<FString> GroupAttributes;
+};
+
+
+
+
+UENUM()
+enum class EAttributeEditorToolActions
+{
+	NoAction,
+	ClearNormals,
+	ClearSelectedUVs,
+	ClearAllUVs,
+	AddAttribute,
+	AddWeightMapLayer,
+	AddPolyGroupLayer,
+	DeleteAttribute,
+	ClearAttribute,
+	CopyAttributeFromTo
+};
+
+
+
 UCLASS()
-class MESHMODELINGTOOLSEDITORONLY_API UAttributeEditorToolProperties : public UInteractiveToolPropertySet
+class MESHMODELINGTOOLSEDITORONLY_API UAttributeEditorActionPropertySet : public UInteractiveToolPropertySet
 {
 	GENERATED_BODY()
 
 public:
-	UAttributeEditorToolProperties()
-	{}
-	
-	int NumUVLayers = 8;
+	TWeakObjectPtr<UAttributeEditorTool> ParentTool;
+
+	void Initialize(UAttributeEditorTool* ParentToolIn) { ParentTool = ParentToolIn; }
+	void PostAction(EAttributeEditorToolActions Action);
+
+};
+
+
+UCLASS()
+class MESHMODELINGTOOLSEDITORONLY_API UAttributeEditorNormalsActions : public UAttributeEditorActionPropertySet
+{
+	GENERATED_BODY()
+public:
+	/** Remove any hard edges / split normals, setting all normals to a single vertex normal */
+	UFUNCTION(CallInEditor, Category = Normals, meta = (DisplayPriority = 1))
+	void ResetHardNormals()
+	{
+		PostAction(EAttributeEditorToolActions::ClearNormals);
+	}
+};
 
-	// UObject interface
-#if WITH_EDITOR
-	virtual bool CanEditChange( const FProperty* InProperty ) const override;
-#endif // WITH_EDITOR	
-	// End of UObject interface
 
-	/** Reset all normals to per-vertex smooth normals, removing all hard edges */
-	UPROPERTY(EditAnywhere, Category = Options)
-	bool bClearNormals;
 
-	// below -- just manually made 8 bools for the 8 possible uv layers in the dumbest possible way ~~
+UCLASS()
+class MESHMODELINGTOOLSEDITORONLY_API UAttributeEditorUVActions : public UAttributeEditorActionPropertySet
+{
+	GENERATED_BODY()
+public:
 
 	/** Clear data from UV layer */
-	UPROPERTY(EditAnywhere, Category = Options)
+	UPROPERTY(EditAnywhere, Category = UVs, meta = (DisplayName = "Layer 0", HideEditConditionToggle, EditConditionHides, EditCondition = "NumUVLayers > 0"))
 	bool bClearUVLayer0;
 	/** Clear data from UV layer */
-	UPROPERTY(EditAnywhere, Category = Options)
+	UPROPERTY(EditAnywhere, Category = UVs, meta = (DisplayName = "Layer 1", HideEditConditionToggle, EditConditionHides, EditCondition = "NumUVLayers > 1"))
 	bool bClearUVLayer1;
 	/** Clear data from UV layer */
-	UPROPERTY(EditAnywhere, Category = Options)
+	UPROPERTY(EditAnywhere, Category = UVs, meta = (DisplayName = "Layer 2", HideEditConditionToggle, EditConditionHides, EditCondition = "NumUVLayers > 2"))
 	bool bClearUVLayer2;
 	/** Clear data from UV layer */
-	UPROPERTY(EditAnywhere, Category = Options)
+	UPROPERTY(EditAnywhere, Category = UVs, meta = (DisplayName = "Layer 3", HideEditConditionToggle, EditConditionHides, EditCondition = "NumUVLayers > 3"))
 	bool bClearUVLayer3;
 	/** Clear data from UV layer */
-	UPROPERTY(EditAnywhere, Category = Options)
+	UPROPERTY(EditAnywhere, Category = UVs, meta = (DisplayName = "Layer 4", HideEditConditionToggle, EditConditionHides, EditCondition = "NumUVLayers > 4"))
 	bool bClearUVLayer4;
 	/** Clear data from UV layer */
-	UPROPERTY(EditAnywhere, Category = Options)
+	UPROPERTY(EditAnywhere, Category = UVs, meta = (DisplayName = "Layer 5", HideEditConditionToggle, EditConditionHides, EditCondition = "NumUVLayers > 5"))
 	bool bClearUVLayer5;
 	/** Clear data from UV layer */
-	UPROPERTY(EditAnywhere, Category = Options)
+	UPROPERTY(EditAnywhere, Category = UVs, meta = (DisplayName = "Layer 6", HideEditConditionToggle, EditConditionHides, EditCondition = "NumUVLayers > 6"))
 	bool bClearUVLayer6;
 	/** Clear data from UV layer */
-	UPROPERTY(EditAnywhere, Category = Options)
+	UPROPERTY(EditAnywhere, Category = UVs, meta = (DisplayName = "Layer 7", HideEditConditionToggle, EditConditionHides, EditCondition = "NumUVLayers > 7"))
 	bool bClearUVLayer7;
 
+	UPROPERTY()
+	int NumUVLayers = 0;
+
+	/** Clear the selected UV layers, setting all UV values to (0,0) */
+	UFUNCTION(CallInEditor, Category = UVs, meta = (DisplayPriority = 1))
+	void ClearSelectedUVSets()
+	{
+		PostAction(EAttributeEditorToolActions::ClearSelectedUVs);
+	}
+
+	/** Clear all UV layers, setting all UV values to (0,0) */
+	UFUNCTION(CallInEditor, Category = UVs, meta = (DisplayPriority = 2))
+	void ClearAllUVSets()
+	{
+		PostAction(EAttributeEditorToolActions::ClearAllUVs);
+	}
 };
 
+
+
+UCLASS()
+class MESHMODELINGTOOLSEDITORONLY_API UAttributeEditorNewAttributeActions : public UAttributeEditorActionPropertySet
+{
+	GENERATED_BODY()
+public:
+	UPROPERTY(EditAnywhere, Category = NewAttribute, meta = (DisplayName = "New Attribute Name") )
+	FString NewName;
+
+	//UPROPERTY(EditAnywhere, Category = NewAttribute)
+	UPROPERTY()
+	EAttributeEditorElementType ElementType;
+
+	//UPROPERTY(EditAnywhere, Category = NewAttribute)
+	UPROPERTY()
+	EAttributeEditorAttribType DataType;
+
+	//UFUNCTION(CallInEditor, Category = NewAttribute, meta = (DisplayPriority = 1))
+	//void AddNew()
+	//{
+	//	PostAction(EAttributeEditorToolActions::AddAttribute);
+	//}
+
+	/** Add a new Per-Vertex Weight Map layer with the given Name */
+	UFUNCTION(CallInEditor, Category = NewAttribute, meta = (DisplayPriority = 2))
+	void AddWeightMapLayer()
+	{
+		PostAction(EAttributeEditorToolActions::AddWeightMapLayer);
+	}
+
+	/** Add a new PolyGroup layer with the given Name */
+	UFUNCTION(CallInEditor, Category = NewAttribute, meta = (DisplayPriority = 3))
+	void AddPolyGroupLayer()
+	{
+		PostAction(EAttributeEditorToolActions::AddPolyGroupLayer);
+	}
+
+};
+
+
+UCLASS()
+class MESHMODELINGTOOLSEDITORONLY_API UAttributeEditorModifyAttributeActions : public UAttributeEditorActionPropertySet
+{
+	GENERATED_BODY()
+public:
+	UPROPERTY(EditAnywhere, Category = ModifyAttribute, meta = (GetOptions = GetAttributeNamesFunc))
+	FString Attribute;
+
+	UFUNCTION()
+	TArray<FString> GetAttributeNamesFunc();
+
+	UPROPERTY()
+	TArray<FString> AttributeNamesList;
+
+	/** Remove the selected Attribute Name from the mesh */
+	UFUNCTION(CallInEditor, Category = ModifyAttribute, meta = (DisplayPriority = 1))
+	void DeleteSelected()
+	{
+		PostAction(EAttributeEditorToolActions::DeleteAttribute);
+	}
+
+	//UFUNCTION(CallInEditor, Category = ModifyAttribute, meta = (DisplayPriority = 2))
+	//void Clear()
+	//{
+	//	PostAction(EAttributeEditorToolActions::ClearAttribute);
+	//}
+
+};
+
+
+UCLASS()
+class MESHMODELINGTOOLSEDITORONLY_API UAttributeEditorCopyAttributeActions : public UAttributeEditorActionPropertySet
+{
+	GENERATED_BODY()
+public:
+	UPROPERTY(EditAnywhere, Category = CopyAttribute)
+	TArray<FString> FromAttribute;
+
+	UPROPERTY(EditAnywhere, Category = CopyAttribute)
+	TArray<FString> ToAttribute;
+
+	UFUNCTION(CallInEditor, Category = CopyAttribute, meta = (DisplayPriority = 1))
+	void CopyFromTo()
+	{
+		PostAction(EAttributeEditorToolActions::CopyAttributeFromTo);
+	}
+};
+
+
+
+
+
+
 /**
- * Simple Mesh Normal Updating Tool
+ * Mesh Attribute Editor Tool
  */
 UCLASS()
 class MESHMODELINGTOOLSEDITORONLY_API UAttributeEditorTool : public UMultiSelectionTool
@@ -95,29 +293,61 @@ public:
 
 	virtual void Setup() override;
 	virtual void Shutdown(EToolShutdownType ShutdownType) override;
+	virtual void OnTick(float DeltaTime) override;
 
 	virtual void SetWorld(UWorld* World);
-	virtual void SetAssetAPI(IToolsContextAssetAPI* AssetAPI);
-
-	virtual void Tick(float DeltaTime) override;
-	virtual void Render(IToolsContextRenderAPI* RenderAPI) override;
 
-	virtual bool HasCancel() const override { return true; }
-	virtual bool HasAccept() const override;
-	virtual bool CanAccept() const override;
+	virtual bool HasCancel() const override { return false; }
+	virtual bool HasAccept() const override { return false; }
+	virtual bool CanAccept() const override { return false; }
 
-	virtual void OnPropertyModified(UObject* PropertySet, FProperty* Property) override;
+	virtual void RequestAction(EAttributeEditorToolActions ActionType);
 
 protected:
 
 	UPROPERTY()
-	UAttributeEditorToolProperties* RemovalProperties;
+	UAttributeEditorNormalsActions* NormalsActions;
+
+	UPROPERTY()
+	UAttributeEditorUVActions* UVActions;
+
+	UPROPERTY()
+	UAttributeEditorAttribProperties* AttributeProps;
+
+	UPROPERTY()
+	UAttributeEditorNewAttributeActions* NewAttributeProps;
+
+	UPROPERTY()
+	UAttributeEditorModifyAttributeActions* ModifyAttributeProps;
+
+	UPROPERTY()
+	UAttributeEditorCopyAttributeActions* CopyAttributeProps;
+
 
 protected:
 	UWorld* TargetWorld;
-	IToolsContextAssetAPI* AssetAPI;
 
-	FViewCameraState CameraState;
+	TArray<FAttributeEditorAttribInfo> VertexAttributes;
+	TArray<FAttributeEditorAttribInfo> InstanceAttributes;
+	TArray<FAttributeEditorAttribInfo> TriangleAttributes;
+	TArray<FAttributeEditorAttribInfo> PolygonAttributes;
+	TArray<FAttributeEditorAttribInfo> EdgeAttributes;
+	TArray<FAttributeEditorAttribInfo> GroupAttributes;
+
+
+	bool bAttributeListsValid = false;
+	void InitializeAttributeLists();
+
+
+
+	EAttributeEditorToolActions PendingAction = EAttributeEditorToolActions::NoAction;
+	void ClearNormals();
+	void ClearUVs(bool bSelectedOnly);
+	void AddNewAttribute();
+	void AddNewWeightMap();
+	void AddNewGroupsLayer();
+	void DeleteAttribute();
+	void ClearAttribute();
 
-	void GenerateAssets();
+	void AddNewAttribute(EAttributeEditorElementType ElemType, EAttributeEditorAttribType AttribType, FName AttributeName);
 };
diff --git a/Engine/Plugins/Experimental/MeshModelingToolset/Source/MeshModelingToolsEditorOnly/Public/EditMeshMaterialsTool.h b/Engine/Plugins/Experimental/MeshModelingToolset/Source/MeshModelingToolsEditorOnly/Public/EditMeshMaterialsTool.h
--- a/Engine/Plugins/Experimental/MeshModelingToolset/Source/MeshModelingToolsEditorOnly/Public/EditMeshMaterialsTool.h
+++ b/Engine/Plugins/Experimental/MeshModelingToolset/Source/MeshModelingToolsEditorOnly/Public/EditMeshMaterialsTool.h
@@ -79,7 +79,7 @@ class MESHMODELINGTOOLSEDITORONLY_API UEditMeshMaterialsTool : public UMeshSelec
 
 public:
 	virtual void Setup() override;
-	virtual void Tick(float DeltaTime) override;
+	virtual void OnTick(float DeltaTime) override;
 
 	virtual bool CanAccept() const override { return UMeshSelectionTool::CanAccept() || bHaveModifiedMaterials; }
 
@@ -111,7 +111,6 @@ protected:
 	FMaterialSetKey GetMaterialKey();
 
 	FMaterialSetKey InitialMaterialKey;
-	TValueWatcher<FMaterialSetKey> MaterialSetWatcher;
 	bool bHaveModifiedMaterials = false;
 
 	void ExternalUpdateMaterialSet(const TArray<UMaterialInterface*>& NewMaterialSet);
diff --git a/Engine/Plugins/Experimental/MeshModelingToolset/Source/MeshModelingToolsEditorOnly/Public/MergeMeshesTool.h b/Engine/Plugins/Experimental/MeshModelingToolset/Source/MeshModelingToolsEditorOnly/Public/MergeMeshesTool.h
--- a/Engine/Plugins/Experimental/MeshModelingToolset/Source/MeshModelingToolsEditorOnly/Public/MergeMeshesTool.h
+++ b/Engine/Plugins/Experimental/MeshModelingToolset/Source/MeshModelingToolsEditorOnly/Public/MergeMeshesTool.h
@@ -80,8 +80,7 @@ public:
 	virtual void Setup() override;
 	virtual void Shutdown(EToolShutdownType ShutdownType) override;
 
-	virtual void Tick(float DeltaTime) override;
-	virtual void Render(IToolsContextRenderAPI* RenderAPI) override;
+	virtual void OnTick(float DeltaTime) override;
 
 	virtual bool HasCancel() const override { return true; }
 	virtual bool HasAccept() const override;
diff --git a/Engine/Plugins/Experimental/MeshModelingToolset/Source/MeshModelingToolsEditorOnly/Public/ParameterizeMeshTool.h b/Engine/Plugins/Experimental/MeshModelingToolset/Source/MeshModelingToolsEditorOnly/Public/ParameterizeMeshTool.h
--- a/Engine/Plugins/Experimental/MeshModelingToolset/Source/MeshModelingToolsEditorOnly/Public/ParameterizeMeshTool.h
+++ b/Engine/Plugins/Experimental/MeshModelingToolset/Source/MeshModelingToolsEditorOnly/Public/ParameterizeMeshTool.h
@@ -103,7 +103,7 @@ class MESHMODELINGTOOLSEDITORONLY_API UParameterizeMeshToolProperties : public U
 public:
 
 	//UPROPERTY(EditAnywhere, Category = Options)
-	UPROPERTY()
+	UPROPERTY(meta = (TransientToolProperty))
 	EUVIslandMode IslandMode = EUVIslandMode::PolyGroups;
 
 	UPROPERTY(EditAnywhere, Category = Options, meta = (EditConditionHides, HideEditConditionToggle, EditCondition = "bIsGlobalMode == false"))
@@ -121,11 +121,8 @@ public:
 	UPROPERTY(EditAnywhere, Category = Options, meta = (EditCondition = "UVScaleMode!=EParameterizeMeshToolUVScaleMode::NoScaling", UIMin = "0.001", UIMax = "10", ClampMin = "0.00001", ClampMax = "1000000.0") )
 	float UVScale = 1.0;
 
-	UPROPERTY()
+	UPROPERTY(meta = (TransientToolProperty))
 	bool bIsGlobalMode = false;
-
-	virtual void SaveProperties(UInteractiveTool* SaveFromTool) override;
-	virtual void RestoreProperties(UInteractiveTool* RestoreToTool) override;
 };
 
 
@@ -150,8 +147,7 @@ public:
 	virtual void Setup() override;
 	virtual void Shutdown(EToolShutdownType ShutdownType) override;
 
-	virtual void Tick(float DeltaTime) override;
-	virtual void Render(IToolsContextRenderAPI* RenderAPI) override;
+	virtual void OnTick(float DeltaTime) override;
 
 	virtual bool HasCancel() const override { return true; }
 	virtual bool HasAccept() const override;
diff --git a/Engine/Plugins/Experimental/MeshModelingToolset/Source/MeshModelingToolsEditorOnly/Public/PolygonOnMeshTool.h b/Engine/Plugins/Experimental/MeshModelingToolset/Source/MeshModelingToolsEditorOnly/Public/PolygonOnMeshTool.h
--- a/Engine/Plugins/Experimental/MeshModelingToolset/Source/MeshModelingToolsEditorOnly/Public/PolygonOnMeshTool.h
+++ b/Engine/Plugins/Experimental/MeshModelingToolset/Source/MeshModelingToolsEditorOnly/Public/PolygonOnMeshTool.h
@@ -14,7 +14,6 @@
 #include "Mechanics/ConstructionPlaneMechanic.h"
 #include "Mechanics/CollectSurfacePathMechanic.h"
 #include "Polygon2.h"
-#include "Changes/ValueWatcher.h"
 #include "PolygonOnMeshTool.generated.h"
 
 
@@ -149,7 +148,7 @@ public:
 
 	virtual void SetWorld(UWorld* World);
 
-	virtual void Tick(float DeltaTime) override;
+	virtual void OnTick(float DeltaTime) override;
 	virtual void Render(IToolsContextRenderAPI* RenderAPI) override;
 
 	virtual bool HasCancel() const override { return true; }
diff --git a/Engine/Plugins/Experimental/MeshModelingToolset/Source/MeshModelingToolsEditorOnly/Public/SimplifyMeshTool.h b/Engine/Plugins/Experimental/MeshModelingToolset/Source/MeshModelingToolsEditorOnly/Public/SimplifyMeshTool.h
--- a/Engine/Plugins/Experimental/MeshModelingToolset/Source/MeshModelingToolsEditorOnly/Public/SimplifyMeshTool.h
+++ b/Engine/Plugins/Experimental/MeshModelingToolset/Source/MeshModelingToolsEditorOnly/Public/SimplifyMeshTool.h
@@ -9,7 +9,6 @@
 #include "DynamicMesh3.h"
 #include "DynamicMeshAABBTree3.h"
 #include "MeshOpPreviewHelpers.h"
-#include "Changes/ValueWatcher.h"
 #include "CleaningOps/SimplifyMeshOp.h"
 #include "Properties/MeshStatisticsProperties.h"
 #include "Properties/RemeshProperties.h"
@@ -46,8 +45,6 @@ class MESHMODELINGTOOLSEDITORONLY_API USimplifyMeshToolProperties : public UMesh
 public:
 	USimplifyMeshToolProperties();
 
-	void SaveRestoreProperties(UInteractiveTool* RestoreToTool, bool bSaving) override;
-
 	/** Simplification Target Type  */
 	UPROPERTY(EditAnywhere, Category = Options)
 	ESimplifyTargetType TargetMode;
@@ -103,7 +100,7 @@ public:
 	virtual void Setup() override;
 	virtual void Shutdown(EToolShutdownType ShutdownType) override;
 
-	virtual void Tick(float DeltaTime) override;
+	virtual void OnTick(float DeltaTime) override;
 	virtual void Render(IToolsContextRenderAPI* RenderAPI) override;
 
 	virtual bool HasCancel() const override { return true; }
@@ -128,9 +125,6 @@ private:
 	UWorld* TargetWorld;
 	IToolsContextAssetAPI* AssetAPI;
 
-	TValueWatcher<bool> ShowWireFrameWatcher;
-	TValueWatcher<bool> ShowGroupsWatcher;
-
 	TSharedPtr<FMeshDescription> OriginalMeshDescription;
 	// Dynamic Mesh versions precomputed in Setup (rather than recomputed for every simplify op)
 	TSharedPtr<FDynamicMesh3> OriginalMesh;
diff --git a/Engine/Plugins/Experimental/MeshModelingToolset/Source/MeshModelingToolsEditorOnly/Public/UVLayoutTool.h b/Engine/Plugins/Experimental/MeshModelingToolset/Source/MeshModelingToolsEditorOnly/Public/UVLayoutTool.h
--- a/Engine/Plugins/Experimental/MeshModelingToolset/Source/MeshModelingToolsEditorOnly/Public/UVLayoutTool.h
+++ b/Engine/Plugins/Experimental/MeshModelingToolset/Source/MeshModelingToolsEditorOnly/Public/UVLayoutTool.h
@@ -60,12 +60,6 @@ public:
 
 	UPROPERTY(EditAnywhere, Category = UVLayout)
 	float UVScaleFactor = 1;
-
-	//
-	// save/restore support
-	//
-	virtual void SaveProperties(UInteractiveTool* SaveFromTool) override;
-	virtual void RestoreProperties(UInteractiveTool* RestoreToTool) override;
 };
 
 
@@ -124,7 +118,7 @@ public:
 	virtual void SetWorld(UWorld* World, UInteractiveGizmoManager* GizmoManagerIn);
 	virtual void SetAssetAPI(IToolsContextAssetAPI* AssetAPI);
 
-	virtual void Tick(float DeltaTime) override;
+	virtual void OnTick(float DeltaTime) override;
 	virtual void Render(IToolsContextRenderAPI* RenderAPI) override;
 
 	virtual bool HasCancel() const override { return true; }
diff --git a/Engine/Plugins/Experimental/MeshModelingToolset/Source/ModelingComponents/Private/BaseDynamicMeshSceneProxy.h b/Engine/Plugins/Experimental/MeshModelingToolset/Source/ModelingComponents/Private/BaseDynamicMeshSceneProxy.h
--- a/Engine/Plugins/Experimental/MeshModelingToolset/Source/ModelingComponents/Private/BaseDynamicMeshSceneProxy.h
+++ b/Engine/Plugins/Experimental/MeshModelingToolset/Source/ModelingComponents/Private/BaseDynamicMeshSceneProxy.h
@@ -180,6 +180,47 @@ public:
 	}
 
 
+	void TransferVertexUpdateToGPU(bool bPositions, bool bNormals, bool bTexCoords, bool bColors)
+	{
+		check(IsInRenderingThread());
+		if (TriangleCount == 0)
+		{
+			return;
+		}
+
+		if (bPositions)
+		{
+			auto& VertexBuffer = this->PositionVertexBuffer;
+			void* VertexBufferData = RHILockVertexBuffer(VertexBuffer.VertexBufferRHI, 0, VertexBuffer.GetNumVertices() * VertexBuffer.GetStride(), RLM_WriteOnly);
+			FMemory::Memcpy(VertexBufferData, VertexBuffer.GetVertexData(), VertexBuffer.GetNumVertices() * VertexBuffer.GetStride());
+			RHIUnlockVertexBuffer(VertexBuffer.VertexBufferRHI);
+		}
+		if (bNormals)
+		{
+			auto& VertexBuffer = this->StaticMeshVertexBuffer;
+			void* VertexBufferData = RHILockVertexBuffer(VertexBuffer.TangentsVertexBuffer.VertexBufferRHI, 0, VertexBuffer.GetTangentSize(), RLM_WriteOnly);
+			FMemory::Memcpy(VertexBufferData, VertexBuffer.GetTangentData(), VertexBuffer.GetTangentSize());
+			RHIUnlockVertexBuffer(VertexBuffer.TangentsVertexBuffer.VertexBufferRHI);
+		}
+		if (bColors)
+		{
+			auto& VertexBuffer = this->ColorVertexBuffer;
+			void* VertexBufferData = RHILockVertexBuffer(VertexBuffer.VertexBufferRHI, 0, VertexBuffer.GetNumVertices() * VertexBuffer.GetStride(), RLM_WriteOnly);
+			FMemory::Memcpy(VertexBufferData, VertexBuffer.GetVertexData(), VertexBuffer.GetNumVertices() * VertexBuffer.GetStride());
+			RHIUnlockVertexBuffer(VertexBuffer.VertexBufferRHI);
+		}
+		if (bTexCoords)
+		{
+			auto& VertexBuffer = this->StaticMeshVertexBuffer;
+			void* VertexBufferData = RHILockVertexBuffer(VertexBuffer.TexCoordVertexBuffer.VertexBufferRHI, 0, VertexBuffer.GetTexCoordSize(), RLM_WriteOnly);
+			FMemory::Memcpy(VertexBufferData, VertexBuffer.GetTexCoordData(), VertexBuffer.GetTexCoordSize());
+			RHIUnlockVertexBuffer(VertexBuffer.TexCoordVertexBuffer.VertexBufferRHI);
+		}
+	}
+
+
+
+
 
 
 	/**
diff --git a/Engine/Plugins/Experimental/MeshModelingToolset/Source/ModelingComponents/Private/Changes/MeshVertexChange.cpp b/Engine/Plugins/Experimental/MeshModelingToolset/Source/ModelingComponents/Private/Changes/MeshVertexChange.cpp
--- a/Engine/Plugins/Experimental/MeshModelingToolset/Source/ModelingComponents/Private/Changes/MeshVertexChange.cpp
+++ b/Engine/Plugins/Experimental/MeshModelingToolset/Source/ModelingComponents/Private/Changes/MeshVertexChange.cpp
@@ -25,18 +25,24 @@ FString FMeshVertexChange::ToString() const
 }
 
 
+FMeshVertexChangeBuilder::FMeshVertexChangeBuilder()
+{
+	Change = MakeUnique<FMeshVertexChange>();
+}
 
 
-
-FMeshVertexChangeBuilder::FMeshVertexChangeBuilder(bool bSaveOverlayNormalsIn)
+FMeshVertexChangeBuilder::FMeshVertexChangeBuilder(EMeshVertexChangeComponents Components)
 {
 	Change = MakeUnique<FMeshVertexChange>();
 
-	bSaveOverlayNormals = bSaveOverlayNormalsIn;
-	if (bSaveOverlayNormals)
-	{
-		Change->bHaveOverlayNormals = true;
-	}
+	bSavePositions = ((Components & EMeshVertexChangeComponents::VertexPositions) != EMeshVertexChangeComponents::None);
+	Change->bHaveVertexPositions = bSavePositions;
+
+	bSaveColors = ((Components & EMeshVertexChangeComponents::VertexColors) != EMeshVertexChangeComponents::None);
+	Change->bHaveVertexColors = bSaveColors;
+	
+	bSaveOverlayNormals = ((Components & EMeshVertexChangeComponents::OverlayNormals) != EMeshVertexChangeComponents::None);
+	Change->bHaveOverlayNormals = bSaveOverlayNormals;
 }
 
 void FMeshVertexChangeBuilder::UpdateVertex(int VertexID, const FVector3d& OldPosition, const FVector3d& NewPosition)
@@ -71,57 +77,59 @@ void FMeshVertexChangeBuilder::UpdateVertexFinal(int VertexID, const FVector3d&
 
 
 
-void FMeshVertexChangeBuilder::SavePosition(const FDynamicMesh3* Mesh, int VertexID, bool bInitial)
+void FMeshVertexChangeBuilder::SaveVertexInitial(const FDynamicMesh3* Mesh, int32 VertexID)
 {
-	FVector3d Pos = Mesh->GetVertex(VertexID);
-	if (bInitial)
-	{
-		UpdateVertex(VertexID, Pos, Pos);
-	}
-	else 
-	{
-		UpdateVertexFinal(VertexID, Pos);
-	}
-}
-
-void FMeshVertexChangeBuilder::SavePositions(const FDynamicMesh3* Mesh, const TArray<int>& VertexIDs, bool bInitial)
-{
-	int Num = VertexIDs.Num();
-	if (bInitial)
+	const int32* FoundIndex = SavedVertices.Find(VertexID);
+	if (FoundIndex == nullptr)
 	{
-		for (int k = 0; k < Num; ++k)
+		int32 Index = Change->Vertices.Num();
+		SavedVertices.Add(VertexID, Index);
+		Change->Vertices.Add(VertexID);
+		if (bSavePositions)
 		{
-			FVector3d Pos = Mesh->GetVertex(VertexIDs[k]);
-			UpdateVertex(VertexIDs[k], Pos, Pos);
+			FVector3d Pos = Mesh->GetVertex(VertexID);
+			Change->OldPositions.Add(Pos);
+			Change->NewPositions.Add(Pos);
+		}
+		if (bSaveColors)
+		{
+			FVector3f Color = Mesh->GetVertexColor(VertexID);
+			Change->OldColors.Add(Color);
+			Change->NewColors.Add(Color);
+		}
+		if (OnNewVertexSaved)
+		{
+			OnNewVertexSaved(VertexID, Index);
 		}
 	}
 	else
 	{
-		for (int k = 0; k < Num; ++k)
+		int32 Index = *FoundIndex;
+		if (bSavePositions)
+		{
+			Change->NewPositions[Index] = Mesh->GetVertex(VertexID);
+		}
+		if (bSaveColors)
 		{
-			FVector3d Pos = Mesh->GetVertex(VertexIDs[k]);
-			UpdateVertexFinal(VertexIDs[k], Pos);
+			Change->NewColors[Index] = Mesh->GetVertexColor(VertexID);
 		}
 	}
 }
 
 
-void FMeshVertexChangeBuilder::SavePositions(const FDynamicMesh3* Mesh, const TSet<int>& VertexIDs, bool bInitial)
+void FMeshVertexChangeBuilder::SaveVertexFinal(const FDynamicMesh3* Mesh, int32 VertexID)
 {
-	if (bInitial)
+	const int32* FoundIndex = SavedVertices.Find(VertexID);
+	if (FoundIndex != nullptr)
 	{
-		for (int VertexID : VertexIDs)
+		int32 Index = *FoundIndex;
+		if (bSavePositions)
 		{
-			FVector3d Pos = Mesh->GetVertex(VertexID);
-			UpdateVertex(VertexID, Pos, Pos);
+			Change->NewPositions[Index] = Mesh->GetVertex(VertexID);
 		}
-	}
-	else
-	{
-		for (int VertexID : VertexIDs)
+		if (bSaveColors)
 		{
-			FVector3d Pos = Mesh->GetVertex(VertexID);
-			UpdateVertexFinal(VertexID, Pos);
+			Change->NewColors[Index] = Mesh->GetVertexColor(VertexID);
 		}
 	}
 }
@@ -132,6 +140,11 @@ void FMeshVertexChangeBuilder::SavePositions(const FDynamicMesh3* Mesh, const TS
 
 
 
+
+
+
+
+
 void FMeshVertexChangeBuilder::UpdateOverlayNormal(int ElementID, const FVector3f& OldNormal, const FVector3f& NewNormal)
 {
 	const int* FoundIndex = SavedNormalElements.Find(ElementID);
diff --git a/Engine/Plugins/Experimental/MeshModelingToolset/Source/ModelingComponents/Private/MeshOpPreviewHelpers.cpp b/Engine/Plugins/Experimental/MeshModelingToolset/Source/ModelingComponents/Private/MeshOpPreviewHelpers.cpp
--- a/Engine/Plugins/Experimental/MeshModelingToolset/Source/ModelingComponents/Private/MeshOpPreviewHelpers.cpp
+++ b/Engine/Plugins/Experimental/MeshModelingToolset/Source/ModelingComponents/Private/MeshOpPreviewHelpers.cpp
@@ -71,7 +71,9 @@ void UMeshOpPreviewWithBackgroundCompute::UpdateResults()
 	EBackgroundComputeTaskStatus Status = BackgroundCompute->CheckStatus();
 	if (Status == EBackgroundComputeTaskStatus::NewResultAvailable)
 	{
-		TUniquePtr<FDynamicMeshOperator> MeshOp  = BackgroundCompute->ExtractResult();
+		TUniquePtr<FDynamicMeshOperator> MeshOp = BackgroundCompute->ExtractResult();
+		OnOpCompleted.Broadcast(MeshOp.Get());
+
 		TUniquePtr<FDynamicMesh3> ResultMesh = MeshOp->ExtractResult();
 		PreviewMesh->SetTransform((FTransform)MeshOp->GetResultTransform());
 		PreviewMesh->UpdatePreview(ResultMesh.Get());  // copies the mesh @todo we could just give ownership to the Preview!
@@ -121,4 +123,4 @@ bool UMeshOpPreviewWithBackgroundCompute::IsUsingWorkingMaterial()
 {
 	bool bIsLongDelay = BackgroundCompute && BackgroundCompute->GetElapsedComputeTime() > SecondsBeforeWorkingMaterial;
 	return !(bResultValid || WorkingMaterial == nullptr || bIsLongDelay == false);
-}
\ No newline at end of file
+}
diff --git a/Engine/Plugins/Experimental/MeshModelingToolset/Source/ModelingComponents/Private/PreviewMesh.cpp b/Engine/Plugins/Experimental/MeshModelingToolset/Source/ModelingComponents/Private/PreviewMesh.cpp
--- a/Engine/Plugins/Experimental/MeshModelingToolset/Source/ModelingComponents/Private/PreviewMesh.cpp
+++ b/Engine/Plugins/Experimental/MeshModelingToolset/Source/ModelingComponents/Private/PreviewMesh.cpp
@@ -371,7 +371,7 @@ void UPreviewMesh::ForceRebuildSpatial()
 }
 
 
-void UPreviewMesh::NotifyDeferredEditCompleted(ERenderUpdateMode UpdateMode, bool bRebuildSpatial)
+void UPreviewMesh::NotifyDeferredEditCompleted(ERenderUpdateMode UpdateMode, EMeshRenderAttributeFlags ModifiedAttribs, bool bRebuildSpatial)
 {
 	if (bBuildSpatialDataStructure && bRebuildSpatial)
 	{
@@ -382,9 +382,20 @@ void UPreviewMesh::NotifyDeferredEditCompleted(ERenderUpdateMode UpdateMode, boo
 	{
 		DynamicMeshComponent->NotifyMeshUpdated();
 	}
-	else
+	else if (UpdateMode == ERenderUpdateMode::FastUpdate)
 	{
-		DynamicMeshComponent->FastNotifyPositionsUpdated();
+		bool bPositions = (ModifiedAttribs & EMeshRenderAttributeFlags::Positions) != EMeshRenderAttributeFlags::None;
+		bool bNormals = (ModifiedAttribs & EMeshRenderAttributeFlags::VertexNormals) != EMeshRenderAttributeFlags::None;
+		bool bColors = (ModifiedAttribs & EMeshRenderAttributeFlags::VertexColors) != EMeshRenderAttributeFlags::None;
+		bool bUVs = (ModifiedAttribs & EMeshRenderAttributeFlags::VertexUVs) != EMeshRenderAttributeFlags::None;
+		if (bPositions)
+		{
+			DynamicMeshComponent->FastNotifyPositionsUpdated(bNormals, bColors, bUVs);
+		}
+		else
+		{
+			DynamicMeshComponent->FastNotifyVertexAttributesUpdated(bNormals, bColors, bUVs);
+		}
 	}
 }
 
diff --git a/Engine/Plugins/Experimental/MeshModelingToolset/Source/ModelingComponents/Private/Selection/PolygonSelectionMechanic.cpp b/Engine/Plugins/Experimental/MeshModelingToolset/Source/ModelingComponents/Private/Selection/PolygonSelectionMechanic.cpp
--- a/Engine/Plugins/Experimental/MeshModelingToolset/Source/ModelingComponents/Private/Selection/PolygonSelectionMechanic.cpp
+++ b/Engine/Plugins/Experimental/MeshModelingToolset/Source/ModelingComponents/Private/Selection/PolygonSelectionMechanic.cpp
@@ -27,17 +27,18 @@ void UPolygonSelectionMechanic::Setup(UInteractiveTool* ParentToolIn)
 }
 
 
-void UPolygonSelectionMechanic::Initialize(const USimpleDynamicMeshComponent* MeshComponentIn, const FGroupTopology* TopologyIn,
-	TFunction<FDynamicMeshAABBTree3*()> GetSpatialSourceFunc,
-	TFunction<bool()> GetAddToSelectionModifierStateFuncIn)
+void UPolygonSelectionMechanic::Initialize(const FDynamicMesh3* MeshIn,
+	FTransform3d TargetTransformIn,
+	const FGroupTopology* TopologyIn,
+	TFunction<FDynamicMeshAABBTree3 * ()> GetSpatialSourceFuncIn,
+	TFunction<bool(void)> GetAddToSelectionModifierStateFuncIn)
 {
-	this->MeshComponent = MeshComponentIn;
+	this->Mesh = MeshIn;
 	this->Topology = TopologyIn;
+	this->TargetTransform = TargetTransformIn;
 
-	TargetTransform = FTransform3d(MeshComponent->GetComponentTransform());
-
-	TopoSelector.Initialize(MeshComponent->GetMesh(), Topology);
-	this->GetSpatialFunc = GetSpatialSourceFunc;
+	TopoSelector.Initialize(Mesh, Topology);
+	this->GetSpatialFunc = GetSpatialSourceFuncIn;
 	TopoSelector.SetSpatialSource(GetSpatialFunc);
 	TopoSelector.PointsWithinToleranceTest = [this](const FVector3d& Position1, const FVector3d& Position2) {
 		return ToolSceneQueriesUtil::PointSnapQuery(this->CameraState,
@@ -47,13 +48,23 @@ void UPolygonSelectionMechanic::Initialize(const USimpleDynamicMeshComponent* Me
 	GetAddToSelectionModifierStateFunc = GetAddToSelectionModifierStateFuncIn;
 }
 
-
+void UPolygonSelectionMechanic::Initialize(const USimpleDynamicMeshComponent* MeshComponentIn,
+	const FGroupTopology* TopologyIn,
+	TFunction<FDynamicMeshAABBTree3 * ()> GetSpatialSourceFuncIn,
+	TFunction<bool()> GetAddToSelectionModifierStateFuncIn)
+{
+	Initialize(MeshComponentIn->GetMesh(),
+		FTransform3d(MeshComponentIn->GetComponentTransform()),
+		TopologyIn,
+		GetSpatialSourceFuncIn,
+		GetAddToSelectionModifierStateFuncIn);
+}
 
 void UPolygonSelectionMechanic::Render(IToolsContextRenderAPI* RenderAPI)
 {
 	GetParentTool()->GetToolManager()->GetContextQueriesAPI()->GetCurrentViewState(CameraState);
 
-	const FDynamicMesh3* TargetMesh = MeshComponent->GetMesh();
+	const FDynamicMesh3* TargetMesh = this->Mesh;
 	FTransform Transform = (FTransform)TargetTransform;
 
 	PolyEdgesRenderer.BeginFrame(RenderAPI, CameraState);
diff --git a/Engine/Plugins/Experimental/MeshModelingToolset/Source/ModelingComponents/Private/SimpleDynamicMeshComponent.cpp b/Engine/Plugins/Experimental/MeshModelingToolset/Source/ModelingComponents/Private/SimpleDynamicMeshComponent.cpp
--- a/Engine/Plugins/Experimental/MeshModelingToolset/Source/ModelingComponents/Private/SimpleDynamicMeshComponent.cpp
+++ b/Engine/Plugins/Experimental/MeshModelingToolset/Source/ModelingComponents/Private/SimpleDynamicMeshComponent.cpp
@@ -15,8 +15,8 @@
 #include "EngineGlobals.h"
 #include "Engine/Engine.h"
 #include "StaticMeshResources.h"
-
 #include "StaticMeshAttributes.h"
+#include "Async/Async.h"
 
 #include "DynamicMeshAttributeSet.h"
 #include "MeshNormals.h"
@@ -26,6 +26,7 @@
 #include "Changes/MeshVertexChange.h"
 #include "Changes/MeshChange.h"
 #include "DynamicMeshChangeTracker.h"
+#include "MeshTransforms.h"
 
 
 
@@ -83,6 +84,21 @@ void USimpleDynamicMeshComponent::InitializeNewMesh()
 }
 
 
+void USimpleDynamicMeshComponent::ApplyTransform(const FTransform3d& Transform, bool bInvert)
+{
+	if (bInvert)
+	{
+		MeshTransforms::ApplyTransformInverse(*GetMesh(), Transform);
+	}
+	else
+	{
+		MeshTransforms::ApplyTransform(*GetMesh(), Transform);
+	}
+
+	NotifyMeshUpdated();
+}
+
+
 void USimpleDynamicMeshComponent::Bake(FMeshDescription* MeshDescription, bool bHaveModifiedTopology, const FConversionToMeshDescriptionOptions& ConversionOptions)
 {
 	FDynamicMeshToMeshDescription Converter(ConversionOptions);
@@ -210,15 +226,14 @@ void USimpleDynamicMeshComponent::FastNotifyPositionsUpdated(bool bNormals, bool
 }
 
 
-void USimpleDynamicMeshComponent::FastNotifyUVsUpdated()
+void USimpleDynamicMeshComponent::FastNotifyVertexAttributesUpdated(bool bNormals, bool bColors, bool bUVs)
 {
+	check(bNormals || bColors || bUVs);
 	if (GetCurrentSceneProxy() != nullptr)
 	{
-		GetCurrentSceneProxy()->FastUpdateVertices(false, false, false, true);
+		GetCurrentSceneProxy()->FastUpdateVertices(false, bNormals, bColors, bUVs);
 		//MarkRenderDynamicDataDirty();
-		MarkRenderTransformDirty();
-		LocalBounds = Mesh->GetCachedBounds();
-		UpdateBounds();
+		//MarkRenderTransformDirty();
 	}
 	else
 	{
@@ -227,6 +242,37 @@ void USimpleDynamicMeshComponent::FastNotifyUVsUpdated()
 }
 
 
+void USimpleDynamicMeshComponent::FastNotifyVertexAttributesUpdated(EMeshRenderAttributeFlags UpdatedAttributes)
+{
+	check(UpdatedAttributes != EMeshRenderAttributeFlags::None);
+	if (GetCurrentSceneProxy() != nullptr)
+	{
+		bool bPositions = (UpdatedAttributes & EMeshRenderAttributeFlags::Positions) != EMeshRenderAttributeFlags::None;
+		GetCurrentSceneProxy()->FastUpdateVertices(bPositions,
+			(UpdatedAttributes & EMeshRenderAttributeFlags::VertexNormals) != EMeshRenderAttributeFlags::None,
+			(UpdatedAttributes & EMeshRenderAttributeFlags::VertexColors) != EMeshRenderAttributeFlags::None,
+			(UpdatedAttributes & EMeshRenderAttributeFlags::VertexUVs) != EMeshRenderAttributeFlags::None);
+
+		if (bPositions)
+		{
+			MarkRenderTransformDirty();
+			LocalBounds = Mesh->GetCachedBounds();
+			UpdateBounds();
+		}
+	}
+	else
+	{
+		NotifyMeshUpdated();
+	}
+}
+
+
+void USimpleDynamicMeshComponent::FastNotifyUVsUpdated()
+{
+	FastNotifyVertexAttributesUpdated(EMeshRenderAttributeFlags::Positions);
+}
+
+
 
 void USimpleDynamicMeshComponent::FastNotifySecondaryTrianglesChanged()
 {
@@ -241,6 +287,98 @@ void USimpleDynamicMeshComponent::FastNotifySecondaryTrianglesChanged()
 }
 
 
+void USimpleDynamicMeshComponent::FastNotifyTriangleVerticesUpdated(const TArray<int32>& Triangles, EMeshRenderAttributeFlags UpdatedAttributes)
+{
+	if (GetCurrentSceneProxy() == nullptr)
+	{
+		NotifyMeshUpdated();
+	}
+	else if ( ! Decomposition )
+	{
+		FastNotifyVertexAttributesUpdated(UpdatedAttributes);
+	}
+	else
+	{
+		TArray<int32> UpdatedSets;
+		for (int32 tid : Triangles)
+		{
+			int32 SetID = Decomposition->GetGroupForTriangle(tid);
+			UpdatedSets.AddUnique(SetID);
+		}
+
+		bool bPositions = (UpdatedAttributes & EMeshRenderAttributeFlags::Positions) != EMeshRenderAttributeFlags::None;
+		GetCurrentSceneProxy()->FastUpdateVertices(UpdatedSets, bPositions,
+			(UpdatedAttributes & EMeshRenderAttributeFlags::VertexNormals) != EMeshRenderAttributeFlags::None,
+			(UpdatedAttributes & EMeshRenderAttributeFlags::VertexColors) != EMeshRenderAttributeFlags::None,
+			(UpdatedAttributes & EMeshRenderAttributeFlags::VertexUVs) != EMeshRenderAttributeFlags::None);
+
+		if (bPositions)
+		{
+			MarkRenderTransformDirty();
+			LocalBounds = Mesh->GetCachedBounds();
+			UpdateBounds();
+		}
+	}
+}
+
+
+
+void USimpleDynamicMeshComponent::FastNotifyTriangleVerticesUpdated(const TSet<int32>& Triangles, EMeshRenderAttributeFlags UpdatedAttributes)
+{
+	if (GetCurrentSceneProxy() == nullptr)
+	{
+		NotifyMeshUpdated();
+	}
+	else if (!Decomposition)
+	{
+		FastNotifyVertexAttributesUpdated(UpdatedAttributes);
+	}
+	else
+	{
+		TArray<int32> UpdatedSets;
+		for (int32 tid : Triangles)
+		{
+			int32 SetID = Decomposition->GetGroupForTriangle(tid);
+			UpdatedSets.AddUnique(SetID);
+		}
+
+		int32 TotalTris = 0;
+		for (int32 SetID : UpdatedSets)
+		{
+			TotalTris += Decomposition->GetGroup(SetID).Triangles.Num();
+		}
+
+		UE_LOG(LogTemp, Warning, TEXT("Updating %d groups with %d tris"), UpdatedSets.Num(), TotalTris);
+
+		bool bPositions = (UpdatedAttributes & EMeshRenderAttributeFlags::Positions) != EMeshRenderAttributeFlags::None;
+
+		// calculate bounds while we are updating vertices
+		TFuture<void> UpdateBoundsCalc;
+		if (bPositions)
+		{
+			UpdateBoundsCalc = Async(EAsyncExecution::ThreadPool, [&]()
+			{
+				LocalBounds = Mesh->GetCachedBounds();
+			});
+		}
+
+		GetCurrentSceneProxy()->FastUpdateVertices(UpdatedSets, bPositions,
+			(UpdatedAttributes & EMeshRenderAttributeFlags::VertexNormals) != EMeshRenderAttributeFlags::None,
+			(UpdatedAttributes & EMeshRenderAttributeFlags::VertexColors) != EMeshRenderAttributeFlags::None,
+			(UpdatedAttributes & EMeshRenderAttributeFlags::VertexUVs) != EMeshRenderAttributeFlags::None);
+
+		if (bPositions)
+		{
+			MarkRenderTransformDirty();
+			UpdateBoundsCalc.Wait();
+			//LocalBounds = Mesh->GetCachedBounds();
+			UpdateBounds();
+		}
+	}
+}
+
+
+
 
 FPrimitiveSceneProxy* USimpleDynamicMeshComponent::CreateSceneProxy()
 {
@@ -266,7 +404,14 @@ FPrimitiveSceneProxy* USimpleDynamicMeshComponent::CreateSceneProxy()
 			};
 		}
 
-		NewProxy->Initialize();
+		if (Decomposition)
+		{
+			NewProxy->InitializeFromDecomposition(Decomposition);
+		}
+		else
+		{
+			NewProxy->Initialize();
+		}
 	}
 	return NewProxy;
 }
@@ -296,6 +441,13 @@ void USimpleDynamicMeshComponent::DisableSecondaryTriangleBuffers()
 }
 
 
+void USimpleDynamicMeshComponent::SetExternalDecomposition(TUniquePtr<FMeshRenderDecomposition> DecompositionIn)
+{
+	Decomposition = MoveTemp(DecompositionIn);
+	NotifyMeshUpdated();
+}
+
+
 
 FColor USimpleDynamicMeshComponent::GetTriangleColor(const FDynamicMesh3* MeshIn, int TriangleID)
 {
@@ -324,14 +476,25 @@ FBoxSphereBounds USimpleDynamicMeshComponent::CalcBounds(const FTransform& Local
 
 void USimpleDynamicMeshComponent::ApplyChange(const FMeshVertexChange* Change, bool bRevert)
 {
+	bool bHavePositions = Change->bHaveVertexPositions;
+	bool bHaveColors = Change->bHaveVertexColors && Mesh->HasVertexColors();
+
 	int32 NV = Change->Vertices.Num();
 	const TArray<FVector3d>& Positions = (bRevert) ? Change->OldPositions : Change->NewPositions;
+	const TArray<FVector3f>& Colors = (bRevert) ? Change->OldColors : Change->NewColors;
 	for (int32 k = 0; k < NV; ++k)
 	{
 		int32 vid = Change->Vertices[k];
 		if (Mesh->IsVertex(vid))
 		{
-			Mesh->SetVertex(vid, Positions[k]);
+			if (bHavePositions)
+			{
+				Mesh->SetVertex(vid, Positions[k]);
+			}
+			if (bHaveColors)
+			{
+				Mesh->SetVertexColor(vid, Colors[k]);
+			}
 		}
 	}
 
@@ -350,8 +513,12 @@ void USimpleDynamicMeshComponent::ApplyChange(const FMeshVertexChange* Change, b
 		}
 	}
 
-	NotifyMeshUpdated();
+	if (bInvalidateProxyOnChange)
+	{
+		NotifyMeshUpdated();
+	}
 	OnMeshChanged.Broadcast();
+	OnMeshVerticesChanged.Broadcast(this, Change, bRevert);
 }
 
 
@@ -361,7 +528,10 @@ void USimpleDynamicMeshComponent::ApplyChange(const FMeshChange* Change, bool bR
 {
 	Change->DynamicMeshChange->Apply(Mesh.Get(), bRevert);
 
-	NotifyMeshUpdated();
+	if (bInvalidateProxyOnChange)
+	{
+		NotifyMeshUpdated();
+	}
 	OnMeshChanged.Broadcast();
 }
 
@@ -370,7 +540,10 @@ void USimpleDynamicMeshComponent::ApplyChange(const FMeshReplacementChange* Chan
 {
 	Mesh->Copy(*Change->GetMesh(bRevert));
 
-	NotifyMeshUpdated();
+	if (bInvalidateProxyOnChange)
+	{
+		NotifyMeshUpdated();
+	}
 	OnMeshChanged.Broadcast();
 }
 
diff --git a/Engine/Plugins/Experimental/MeshModelingToolset/Source/ModelingComponents/Private/SimpleDynamicMeshSceneProxy.h b/Engine/Plugins/Experimental/MeshModelingToolset/Source/ModelingComponents/Private/SimpleDynamicMeshSceneProxy.h
--- a/Engine/Plugins/Experimental/MeshModelingToolset/Source/ModelingComponents/Private/SimpleDynamicMeshSceneProxy.h
+++ b/Engine/Plugins/Experimental/MeshModelingToolset/Source/ModelingComponents/Private/SimpleDynamicMeshSceneProxy.h
@@ -83,6 +83,60 @@ public:
 
 
 
+	/**
+	 * Initialize multiple buffers for the mesh based on the given Decomposition
+	 */
+	virtual void InitializeFromDecomposition(TUniquePtr<FMeshRenderDecomposition>& Decomposition)
+	{
+		check(RenderBufferSets.Num() == 0);
+		int32 NumSets = Decomposition->Num();
+		RenderBufferSets.SetNum(NumSets);
+		for (int32 k = 0; k < NumSets; ++k)
+		{
+			RenderBufferSets[k] = AllocateNewRenderBufferSet();
+			RenderBufferSets[k]->Material = Decomposition->GetGroup(k).Material;
+		}
+
+		bIsSingleBuffer = false;
+
+		FDynamicMesh3* Mesh = ParentComponent->GetMesh();
+		// find suitable overlays
+		FDynamicMeshUVOverlay* UVOverlay = Mesh->Attributes()->PrimaryUV();
+		FDynamicMeshNormalOverlay* NormalOverlay = Mesh->Attributes()->PrimaryNormals();
+
+		TFunction<void(int, int, int, FVector3f&, FVector3f&)> TangentsFunc = nullptr;
+		FMeshTangentsf* Tangents = ParentComponent->GetTangents();
+		if (Tangents != nullptr)
+		{
+			TangentsFunc = [Tangents](int VertexID, int TriangleID, int TriVtxIdx, FVector3f& TangentX, FVector3f& TangentY)
+			{
+				return Tangents->GetPerTriangleTangent(TriangleID, TriVtxIdx, TangentX, TangentY);
+			};
+		}
+
+		// init renderbuffers for each set
+		ParallelFor(NumSets, [&](int32 SetIndex)
+		{
+			const FMeshRenderDecomposition::FGroup& Group = Decomposition->GetGroup(SetIndex);
+			const TArray<int32>& Triangles = Group.Triangles;
+			if (Triangles.Num() > 0)
+			{
+				FMeshRenderBufferSet* RenderBuffers = RenderBufferSets[SetIndex];
+				RenderBuffers->Triangles = Triangles;
+				InitializeBuffersFromOverlays(RenderBuffers, Mesh,
+					Triangles.Num(), Triangles,
+					UVOverlay, NormalOverlay, TangentsFunc);
+
+				ENQUEUE_RENDER_COMMAND(FSimpleDynamicMeshSceneProxyInitializeFromDecomposition)(
+					[RenderBuffers](FRHICommandListImmediate& RHICmdList)
+				{
+					RenderBuffers->Upload();
+				});
+			}
+		});
+	}
+
+
 
 
 
@@ -122,7 +176,7 @@ public:
 			Mesh->TriangleCount(), Mesh->TriangleIndicesItr(),
 			UVOverlay, NormalOverlay, TangentsFunc);
 
-		ENQUEUE_RENDER_COMMAND(FOctreeDynamicMeshSceneProxyInitializeSingle)(
+		ENQUEUE_RENDER_COMMAND(FSimpleDynamicMeshSceneProxyInitializeSingle)(
 			[RenderBuffers](FRHICommandListImmediate& RHICmdList)
 		{
 			RenderBuffers->Upload();
@@ -206,7 +260,7 @@ public:
 
 				RenderBuffers->Triangles = Triangles;
 
-				ENQUEUE_RENDER_COMMAND(FOctreeDynamicMeshSceneProxyInitializeSingle)(
+				ENQUEUE_RENDER_COMMAND(FSimpleDynamicMeshSceneProxyInitializeByMaterial)(
 					[RenderBuffers](FRHICommandListImmediate& RHICmdList)
 				{
 					RenderBuffers->Upload();
@@ -298,10 +352,67 @@ public:
 				});
 			});
 		}
+	}
+
 
+
+	/**
+	 * Update the vertex position/normal/color buffers
+	 */
+	virtual void FastUpdateVertices(const TArray<int32>& WhichBuffers, bool bPositions, bool bNormals, bool bColors, bool bUVs)
+	{
+		// This needs to be rewritten for split-by-material buffers.
+		// Could store triangle set with each buffer, and then rebuild vtx buffer(s) as needed?
+
+		FDynamicMesh3* Mesh = ParentComponent->GetMesh();
+
+		// find suitable overlays
+		FDynamicMeshNormalOverlay* NormalOverlay = nullptr;
+		if (bNormals)
+		{
+			check(Mesh->HasAttributes());
+			NormalOverlay = Mesh->Attributes()->PrimaryNormals();
+		}
+		FDynamicMeshUVOverlay* UVOVerlay = nullptr;
+		if (bUVs)
+		{
+			check(Mesh->HasAttributes());
+			UVOVerlay = Mesh->Attributes()->PrimaryUV();
+		}
+
+		ParallelFor(WhichBuffers.Num(), [&](int idx)
+		{
+			int32 BufferIndex = WhichBuffers[idx];
+			if ( RenderBufferSets.IsValidIndex(BufferIndex) == false || RenderBufferSets[BufferIndex]->TriangleCount == 0)
+			{
+				return;
+			}
+			FMeshRenderBufferSet* Buffers = RenderBufferSets[BufferIndex];
+			check(Buffers->Triangles.IsSet());
+			if (bPositions || bNormals || bColors)
+			{
+				UpdateVertexBuffersFromOverlays(Buffers, Mesh,
+					Buffers->Triangles->Num(), Buffers->Triangles.GetValue(),
+					NormalOverlay,
+					bPositions, bNormals, bColors);
+			}
+			if (bUVs)
+			{
+				UpdateVertexUVBufferFromOverlays(Buffers, Mesh,
+					Buffers->Triangles->Num(), Buffers->Triangles.GetValue(), UVOVerlay, 0);
+			}
+
+			ENQUEUE_RENDER_COMMAND(FSimpleDynamicMeshSceneProxyFastUpdateVerticesBufferList)(
+				[Buffers, bPositions, bNormals, bColors, bUVs](FRHICommandListImmediate& RHICmdList)
+			{
+				Buffers->TransferVertexUpdateToGPU(bPositions, bNormals, bUVs, bColors);
+			});
+		});
 	}
 
 
+
+
 	/**
 	 * Update index buffers inside each RenderBuffer set
 	 */
diff --git a/Engine/Plugins/Experimental/MeshModelingToolset/Source/ModelingComponents/Private/ToolSetupUtil.cpp b/Engine/Plugins/Experimental/MeshModelingToolset/Source/ModelingComponents/Private/ToolSetupUtil.cpp
--- a/Engine/Plugins/Experimental/MeshModelingToolset/Source/ModelingComponents/Private/ToolSetupUtil.cpp
+++ b/Engine/Plugins/Experimental/MeshModelingToolset/Source/ModelingComponents/Private/ToolSetupUtil.cpp
@@ -18,6 +18,17 @@ UMaterialInterface* ToolSetupUtil::GetDefaultMaterial(UInteractiveToolManager* T
 }
 
 
+UMaterialInstanceDynamic* ToolSetupUtil::GetVertexColorMaterial(UInteractiveToolManager* ToolManager)
+{
+	UMaterial* Material = LoadObject<UMaterial>(nullptr, TEXT("/MeshModelingToolset/Materials/MeshVertexColorMaterial"));
+	if (Material != nullptr)
+	{
+		UMaterialInstanceDynamic* MatInstance = UMaterialInstanceDynamic::Create(Material, ToolManager);
+		return MatInstance;
+	}
+	return nullptr;
+}
+
 
 UMaterialInterface* ToolSetupUtil::GetDefaultWorkingMaterial(UInteractiveToolManager* ToolManager)
 {
@@ -140,4 +151,15 @@ UMaterialInterface* ToolSetupUtil::GetDefaultPointComponentMaterial(bool bRoundP
 		return ToolManager->GetContextQueriesAPI()->GetStandardMaterial(EStandardToolContextMaterials::VertexColorMaterial);
 	}
 	return Material;
+}
+
+UMaterialInterface* ToolSetupUtil::GetDefaultLineComponentMaterial(UInteractiveToolManager* ToolManager)
+{
+	UMaterialInterface* Material = LoadObject<UMaterial>(nullptr, TEXT("/MeshModelingToolset/Materials/LineSetComponentMaterial"));
+	if (Material == nullptr && ToolManager != nullptr)
+	{
+		// We don't seem to have a default line material to use here.
+		return ToolManager->GetContextQueriesAPI()->GetStandardMaterial(EStandardToolContextMaterials::VertexColorMaterial);
+	}
+	return Material;
 }
\ No newline at end of file
diff --git a/Engine/Plugins/Experimental/MeshModelingToolset/Source/ModelingComponents/Public/BaseDynamicMeshComponent.h b/Engine/Plugins/Experimental/MeshModelingToolset/Source/ModelingComponents/Public/BaseDynamicMeshComponent.h
--- a/Engine/Plugins/Experimental/MeshModelingToolset/Source/ModelingComponents/Public/BaseDynamicMeshComponent.h
+++ b/Engine/Plugins/Experimental/MeshModelingToolset/Source/ModelingComponents/Public/BaseDynamicMeshComponent.h
@@ -21,6 +21,23 @@ class FMeshChange;
 
 
 
+/**
+ * EMeshRenderAttributeFlags is used to identify different mesh rendering attributes, for things
+ * like fast-update functions
+ */
+enum class EMeshRenderAttributeFlags : uint8
+{
+	None = 0,
+	Positions = 0x1,
+	VertexColors = 0x2,
+	VertexNormals = 0x4,
+	VertexUVs = 0x8,
+
+	All = 0xFF
+};
+ENUM_CLASS_FLAGS(EMeshRenderAttributeFlags);
+
+
 
 /**
  * Tangent calculation modes
@@ -95,6 +112,11 @@ protected:
 
 public:
 
+	/**
+	 * Configure whether wireframe rendering is enabled or not
+	 */
+	virtual void SetEnableWireframeRenderPass(bool bEnable) { check(false); }
+
 	/**
 	 * @return true if wireframe rendering pass is enabled (default false)
 	 */
diff --git a/Engine/Plugins/Experimental/MeshModelingToolset/Source/ModelingComponents/Public/Changes/MeshVertexChange.h b/Engine/Plugins/Experimental/MeshModelingToolset/Source/ModelingComponents/Public/Changes/MeshVertexChange.h
--- a/Engine/Plugins/Experimental/MeshModelingToolset/Source/ModelingComponents/Public/Changes/MeshVertexChange.h
+++ b/Engine/Plugins/Experimental/MeshModelingToolset/Source/ModelingComponents/Public/Changes/MeshVertexChange.h
@@ -20,16 +20,19 @@ class FDynamicMesh3;
 class MODELINGCOMPONENTS_API FMeshVertexChange : public FToolCommandChange
 {
 public:
-	TArray<int> Vertices;
+	bool bHaveVertexPositions = true;
+	bool bHaveVertexColors = false;
+	TArray<int32> Vertices;
 	TArray<FVector3d> OldPositions;
 	TArray<FVector3d> NewPositions;
+	TArray<FVector3f> OldColors;
+	TArray<FVector3f> NewColors;
 
 	bool bHaveOverlayNormals = false;
 	TArray<int32> Normals;
 	TArray<FVector3f> OldNormals;
 	TArray<FVector3f> NewNormals;
 
-
 	/** Makes the change to the object */
 	virtual void Apply(UObject* Object) override;
 
@@ -41,36 +44,55 @@ public:
 };
 
 
+enum class EMeshVertexChangeComponents : uint8
+{
+	None = 0,
+	VertexPositions = 1,
+	VertexColors = 2,
+	OverlayNormals = 16
+};
+ENUM_CLASS_FLAGS(EMeshVertexChangeComponents);
+
 
 /**
  * FMeshVertexChangeBuilder can be used to construct a FMeshVertexChange.
- * Usage is to call UpdateVertex() each time a vertex moves, with the old and new positions.
  */
 class MODELINGCOMPONENTS_API FMeshVertexChangeBuilder
 {
 public:
 	TUniquePtr<FMeshVertexChange> Change;
-	TMap<int, int> SavedVertices;
+	TMap<int32, int32> SavedVertices;
+	bool bSavePositions = true;
+	bool bSaveColors = false;
+
+	bool bSaveOverlayNormals = false;
+	TMap<int32, int32> SavedNormalElements;
+
+	/** If set, this function is called whenever a newly-seen VertexID is saved, parameters are (VertexID, Index) into saved-vertices array */
+	TUniqueFunction<void(int32, int32)> OnNewVertexSaved = nullptr;
 
-	bool bSaveOverlayNormals;
-	TMap<int, int> SavedNormalElements;
+	FMeshVertexChangeBuilder();
+	explicit FMeshVertexChangeBuilder(EMeshVertexChangeComponents Components);
 
-	FMeshVertexChangeBuilder(bool bSaveOverlayNormals = false);
+	void SaveVertexInitial(const FDynamicMesh3* Mesh, int32 VertexID);
+	void SaveVertexFinal(const FDynamicMesh3* Mesh, int32 VertexID);
 
-	void UpdateVertex(int VertexID, const FVector3d& OldPosition, const FVector3d& NewPosition);
-	void UpdateVertexFinal(int VertexID, const FVector3d& NewPosition);
+	template<typename Enumerable>
+	void SaveVertices(const FDynamicMesh3* Mesh, Enumerable Enum, bool bInitial);
 
-	void SavePosition(const FDynamicMesh3* Mesh, int VertexID, bool bInitial);
-	void SavePositions(const FDynamicMesh3* Mesh, const TArray<int>& VertexIDs, bool bInitial);
-	void SavePositions(const FDynamicMesh3* Mesh, const TSet<int>& VertexIDs, bool bInitial);
+	void SaveOverlayNormals(const FDynamicMesh3* Mesh, const TArray<int32>& ElementIDs, bool bInitial);
+	void SaveOverlayNormals(const FDynamicMesh3* Mesh, const TSet<int32>& ElementIDs, bool bInitial);
 
 
-	void UpdateOverlayNormal(int ElementID, const FVector3f& OldNormal, const FVector3f& NewNormal);
-	void UpdateOverlayNormalFinal(int ElementID, const FVector3f& NewNormal);
-	void SaveOverlayNormals(const FDynamicMesh3* Mesh, const TArray<int>& ElementIDs, bool bInitial);
-	void SaveOverlayNormals(const FDynamicMesh3* Mesh, const TSet<int>& ElementIDs, bool bInitial);
+public:
+	// currently only used in vertex sculpt tool. cannot be used if bSaveColors = true
+	void UpdateVertex(int32 VertexID, const FVector3d& OldPosition, const FVector3d& NewPosition);
 
+protected:
+	void UpdateVertexFinal(int32 VertexID, const FVector3d& NewPosition);
 
+	void UpdateOverlayNormal(int32 ElementID, const FVector3f& OldNormal, const FVector3f& NewNormal);
+	void UpdateOverlayNormalFinal(int32 ElementID, const FVector3f& NewNormal);
 };
 
 
@@ -90,3 +112,25 @@ public:
 	virtual void ApplyChange(const FMeshVertexChange* Change, bool bRevert) = 0;
 };
 
+
+
+
+
+template<typename Enumerable>
+void FMeshVertexChangeBuilder::SaveVertices(const FDynamicMesh3* Mesh, Enumerable Enum, bool bInitial)
+{
+	if (bInitial)
+	{
+		for (int32 k : Enum)
+		{
+			SaveVertexInitial(Mesh, k);
+		}
+	}
+	else
+	{
+		for (int32 k : Enum)
+		{
+			SaveVertexFinal(Mesh, k);
+		}
+	}
+}
diff --git a/Engine/Plugins/Experimental/MeshModelingToolset/Source/ModelingComponents/Public/MeshOpPreviewHelpers.h b/Engine/Plugins/Experimental/MeshModelingToolset/Source/ModelingComponents/Public/MeshOpPreviewHelpers.h
--- a/Engine/Plugins/Experimental/MeshModelingToolset/Source/ModelingComponents/Public/MeshOpPreviewHelpers.h
+++ b/Engine/Plugins/Experimental/MeshModelingToolset/Source/ModelingComponents/Public/MeshOpPreviewHelpers.h
@@ -3,6 +3,7 @@
 #pragma once
 
 #include "CoreMinimal.h"
+#include "Delegates/DelegateCombinations.h"
 #include "PreviewMesh.h"
 #include "Util/ProgressCancel.h"
 #include "ModelingOperators.h"
@@ -124,12 +125,12 @@ public:
 	//
 	// Change notification
 	//
-
 	DECLARE_MULTICAST_DELEGATE_OneParam(FOnMeshUpdated, UMeshOpPreviewWithBackgroundCompute*);
 	/** This delegate is broadcast whenever the embedded preview mesh is updated */
 	FOnMeshUpdated OnMeshUpdated;
 
-
+	DECLARE_MULTICAST_DELEGATE_OneParam(FOnOpCompleted, const FDynamicMeshOperator*);
+	FOnOpCompleted OnOpCompleted;
 
 public:
 	// preview of MeshOperator result
diff --git a/Engine/Plugins/Experimental/MeshModelingToolset/Source/ModelingComponents/Public/OctreeDynamicMeshComponent.h b/Engine/Plugins/Experimental/MeshModelingToolset/Source/ModelingComponents/Public/OctreeDynamicMeshComponent.h
--- a/Engine/Plugins/Experimental/MeshModelingToolset/Source/ModelingComponents/Public/OctreeDynamicMeshComponent.h
+++ b/Engine/Plugins/Experimental/MeshModelingToolset/Source/ModelingComponents/Public/OctreeDynamicMeshComponent.h
@@ -73,7 +73,7 @@ public:
 	}
 
 	/**
-	 * Apply transform to internal mesh. Updates Octree and RenderProxy if available.
+	 * Apply transform to internal mesh. Invalidates RenderProxy.
 	 * @param bInvert if true, inverse tranform is applied instead of forward transform
 	 */
 	void ApplyTransform(const FTransform3d& Transform, bool bInvert);
@@ -116,6 +116,11 @@ public:
 	UPROPERTY()
 	bool bExplicitShowWireframe = false;
 
+	/**
+	 * Configure whether wireframe rendering is enabled or not
+	 */
+	virtual void SetEnableWireframeRenderPass(bool bEnable) override { bExplicitShowWireframe = bEnable; }
+
 	/**
 	 * @return true if wireframe rendering pass is enabled
 	 */
diff --git a/Engine/Plugins/Experimental/MeshModelingToolset/Source/ModelingComponents/Public/PreviewMesh.h b/Engine/Plugins/Experimental/MeshModelingToolset/Source/ModelingComponents/Public/PreviewMesh.h
--- a/Engine/Plugins/Experimental/MeshModelingToolset/Source/ModelingComponents/Public/PreviewMesh.h
+++ b/Engine/Plugins/Experimental/MeshModelingToolset/Source/ModelingComponents/Public/PreviewMesh.h
@@ -314,9 +314,10 @@ public:
 	/**
 	 * Notify that a DeferredEditMesh sequence is complete and cause update of rendering data structures.
 	 * @param UpdateMode type of rendering update required for the applied mesh edits
+	 * @param ModifiedAttribs which mesh attributes have been modified and need to be updated
 	 * @param bRebuildSpatial if true, and internal spatial data structure is enabled, rebuild it for updated mesh
 	 */
-	void NotifyDeferredEditCompleted(ERenderUpdateMode UpdateMode, bool bRebuildSpatial);
+	void NotifyDeferredEditCompleted(ERenderUpdateMode UpdateMode, EMeshRenderAttributeFlags ModifiedAttribs, bool bRebuildSpatial);
 
 
 	/**
diff --git a/Engine/Plugins/Experimental/MeshModelingToolset/Source/ModelingComponents/Public/Selection/PolygonSelectionMechanic.h b/Engine/Plugins/Experimental/MeshModelingToolset/Source/ModelingComponents/Public/Selection/PolygonSelectionMechanic.h
--- a/Engine/Plugins/Experimental/MeshModelingToolset/Source/ModelingComponents/Public/Selection/PolygonSelectionMechanic.h
+++ b/Engine/Plugins/Experimental/MeshModelingToolset/Source/ModelingComponents/Public/Selection/PolygonSelectionMechanic.h
@@ -48,11 +48,18 @@ public:
 	virtual void Render(IToolsContextRenderAPI* RenderAPI) override;
 
 
-	void Initialize(const USimpleDynamicMeshComponent* MeshComponent, const FGroupTopology* Topology,
+	void Initialize(const FDynamicMesh3* Mesh, 
+		FTransform3d TargetTransform,
+		const FGroupTopology* Topology,
 		TFunction<FDynamicMeshAABBTree3*()> GetSpatialSourceFunc,
 		TFunction<bool(void)> GetAddToSelectionModifierStateFunc = []() {return false; }
 		);
 
+	void Initialize(const USimpleDynamicMeshComponent* MeshComponent, const FGroupTopology* Topology,
+		TFunction<FDynamicMeshAABBTree3 * ()> GetSpatialSourceFunc,
+		TFunction<bool(void)> GetAddToSelectionModifierStateFunc = []() {return false; }
+	);
+
 	/**
 	 * Notify internal data structures that the associated MeshComponent has been modified.
 	 * @param bTopologyModified if true, the underlying mesh topology has been changed. This clears the current selection.
@@ -143,7 +150,7 @@ public:
 	UPolygonSelectionMechanicProperties* Properties;
 
 protected:
-	const USimpleDynamicMeshComponent* MeshComponent;
+	const FDynamicMesh3* Mesh;
 	const FGroupTopology* Topology;
 	TFunction<FDynamicMeshAABBTree3*()> GetSpatialFunc;
 
diff --git a/Engine/Plugins/Experimental/MeshModelingToolset/Source/ModelingComponents/Public/SimpleDynamicMeshComponent.h b/Engine/Plugins/Experimental/MeshModelingToolset/Source/ModelingComponents/Public/SimpleDynamicMeshComponent.h
--- a/Engine/Plugins/Experimental/MeshModelingToolset/Source/ModelingComponents/Public/SimpleDynamicMeshComponent.h
+++ b/Engine/Plugins/Experimental/MeshModelingToolset/Source/ModelingComponents/Public/SimpleDynamicMeshComponent.h
@@ -5,8 +5,9 @@
 #include "CoreMinimal.h"
 #include "BaseDynamicMeshComponent.h"
 #include "MeshConversionOptions.h"
-
+#include "Drawing/MeshRenderDecomposition.h"
 #include "MeshTangents.h"
+#include "TransformTypes.h"
 
 #include "SimpleDynamicMeshComponent.generated.h"
 
@@ -85,6 +86,12 @@ public:
 		Bake(MeshDescription, bHaveModifiedTopology, ConversionOptions);
 	}
 
+	/**
+	 * Apply transform to internal mesh. Updates Octree and RenderProxy if available.
+	 * @param bInvert if true, inverse tranform is applied instead of forward transform
+	 */
+	void ApplyTransform(const FTransform3d& Transform, bool bInvert);
+
 
 	//
 	// change tracking/etc
@@ -103,11 +110,23 @@ public:
 	void FastNotifyColorsUpdated();
 
 	/**
-	 * Call this instead of NotifyMeshUpdated() if you have only updated the vertex positions.
+	 * Call this instead of NotifyMeshUpdated() if you have only updated the vertex positions (and possibly some attributes).
 	 * This function will update the existing RenderProxy buffers if possible
 	 */
 	void FastNotifyPositionsUpdated(bool bNormals = false, bool bColors = false, bool bUVs = false);
 
+	/**
+	 * Call this instead of NotifyMeshUpdated() if you have only updated the vertex attributes (but not positions).
+	 * This function will update the existing RenderProxy buffers if possible, rather than create new ones.
+	 */
+	void FastNotifyVertexAttributesUpdated(bool bNormals, bool bColors, bool bUVs);
+
+	/**
+	 * Call this instead of NotifyMeshUpdated() if you have only updated the vertex positions/attributes
+	 * This function will update the existing RenderProxy buffers if possible, rather than create new ones.
+	 */
+	void FastNotifyVertexAttributesUpdated(EMeshRenderAttributeFlags UpdatedAttributes);
+
 	/**
 	 * Call this instead of NotifyMeshUpdated() if you have only updated the vertex uvs.
 	 * This function will update the existing RenderProxy buffers if possible
@@ -120,6 +139,22 @@ public:
 	 */
 	void FastNotifySecondaryTrianglesChanged();
 
+	/**
+	 * This function updates vertex positions/attributes of existing SceneProxy render buffers if possible, for the given triangles.
+	 * If a FMeshRenderDecomposition has not been explicitly set, call is forwarded to FastNotifyVertexAttributesUpdated()
+	 */
+	void FastNotifyTriangleVerticesUpdated(const TArray<int32>& Triangles, EMeshRenderAttributeFlags UpdatedAttributes);
+
+	/**
+	 * This function updates vertex positions/attributes of existing SceneProxy render buffers if possible, for the given triangles.
+	 * If a FMeshRenderDecomposition has not been explicitly set, call is forwarded to FastNotifyVertexAttributesUpdated()
+	 */
+	void FastNotifyTriangleVerticesUpdated(const TSet<int32>& Triangles, EMeshRenderAttributeFlags UpdatedAttributes);
+
+
+	/** If false, we don't completely invalidate the RenderProxy when ApplyChange() is called (assumption is it will be handled elsewhere) */
+	UPROPERTY()
+	bool bInvalidateProxyOnChange = true;
 
 	/**
 	 * Apply a vertex deformation change to the internal mesh
@@ -144,12 +179,24 @@ public:
 	FSimpleMulticastDelegate OnMeshChanged;
 
 
+	/**
+	 * This delegate fires when ApplyChange(FMeshVertexChange) executes
+	 */
+	DECLARE_MULTICAST_DELEGATE_ThreeParams(FMeshVerticesModified, USimpleDynamicMeshComponent*, const FMeshVertexChange*, bool);
+	FMeshVerticesModified OnMeshVerticesChanged;
+
+
 	/**
 	 * if true, we always show the wireframe on top of the shaded mesh, even when not in wireframe mode
 	 */
 	UPROPERTY()
 	bool bExplicitShowWireframe = false;
 
+	/**
+	 * Configure whether wireframe rendering is enabled or not
+	 */
+	virtual void SetEnableWireframeRenderPass(bool bEnable) override { bExplicitShowWireframe = bEnable; }
+
 	/**
 	 * @return true if wireframe rendering pass is enabled
 	 */
@@ -174,6 +221,12 @@ public:
 	 */
 	virtual void DisableSecondaryTriangleBuffers();
 
+	/**
+	 * Configure a decomposition of the mesh, which will result in separate render buffers for each decomposition triangle group.
+	 * Invalidates existing SceneProxy.
+	 */
+	virtual void SetExternalDecomposition(TUniquePtr<FMeshRenderDecomposition> Decomposition);
+
 
 public:
 
@@ -217,6 +270,5 @@ private:
 
 	TUniqueFunction<bool(const FDynamicMesh3*, int32)> SecondaryTriFilterFunc = nullptr;
 
-
-	//friend class FCustomMeshSceneProxy;
+	TUniquePtr<FMeshRenderDecomposition> Decomposition;
 };
diff --git a/Engine/Plugins/Experimental/MeshModelingToolset/Source/ModelingComponents/Public/ToolSetupUtil.h b/Engine/Plugins/Experimental/MeshModelingToolset/Source/ModelingComponents/Public/ToolSetupUtil.h
--- a/Engine/Plugins/Experimental/MeshModelingToolset/Source/ModelingComponents/Public/ToolSetupUtil.h
+++ b/Engine/Plugins/Experimental/MeshModelingToolset/Source/ModelingComponents/Public/ToolSetupUtil.h
@@ -3,9 +3,9 @@
 #pragma once
 
 #include "CoreMinimal.h"
+#include "Materials/Material.h"
+#include "Materials/MaterialInstanceDynamic.h"
 
-class UMaterialInterface;
-class UMaterialInstanceDynamic;
 class UTexture;
 class UInteractiveToolManager;
 
@@ -21,6 +21,12 @@ namespace ToolSetupUtil
 	 */
 	MODELINGCOMPONENTS_API UMaterialInterface* GetDefaultMaterial(UInteractiveToolManager* ToolManager, UMaterialInterface* SourceMaterial = nullptr);
 
+	/**
+	 * @return configurable vertex color material
+	 */
+	MODELINGCOMPONENTS_API UMaterialInstanceDynamic* GetVertexColorMaterial(UInteractiveToolManager* ToolManager);
+
+
 	/**
 	 * @return default material to use for "Working"/In-Progress animations
 	 */
@@ -75,4 +81,8 @@ namespace ToolSetupUtil
 	 */
 	MODELINGCOMPONENTS_API UMaterialInterface* GetDefaultPointComponentMaterial(bool bRoundPoints, UInteractiveToolManager* ToolManager);
 
+	/**
+	 * @return custom material suitable for use with ULineSetComponent
+	 */
+	MODELINGCOMPONENTS_API UMaterialInterface* GetDefaultLineComponentMaterial(UInteractiveToolManager* ToolManager);
 }
\ No newline at end of file
diff --git a/Engine/Plugins/Experimental/MeshModelingToolset/Source/ModelingComponents/Public/WeightMapUtil.h b/Engine/Plugins/Experimental/MeshModelingToolset/Source/ModelingComponents/Public/WeightMapUtil.h
new file mode 100644
--- /dev/null
+++ b/Engine/Plugins/Experimental/MeshModelingToolset/Source/ModelingComponents/Public/WeightMapUtil.h
@@ -0,0 +1,63 @@
+// Copyright Epic Games, Inc. All Rights Reserved.
+
+#pragma once
+
+#include "CoreMinimal.h"
+#include "VectorTypes.h"
+#include "IndexTypes.h"
+#include "BoxTypes.h"
+
+struct FMeshDescription;
+
+/**
+ * FIndexedWeightMap stores an array of values, the intention is that these are "weights" on indices,
+ * for example per-vertex weights. 
+ */
+class FIndexedWeightMap
+{
+public:
+	float DefaultValue;
+	TArray<float> Values;
+
+	float GetValue(int32 Index) const
+	{
+		return Values[Index];
+	}
+
+	float GetInterpValue(const FIndex3i& Indices, const FVector3d& BaryCoords) const
+	{
+		return (float)((double)Values[Indices.A]*BaryCoords.X
+			+ (double)Values[Indices.B]*BaryCoords.Y
+			+ (double)Values[Indices.C]*BaryCoords.Z);
+	}
+
+	void InvertWeightMap(FInterval1f Range = FInterval1f(0.0f,1.0f))
+	{
+		int32 Num = Values.Num();
+		for (int32 k = 0; k < Num; ++k)
+		{
+			Values[k] = Range.Clamp( (Range.Max - (Values[k] - Range.Min)) );
+		}
+	}
+};
+
+
+namespace UE
+{
+	namespace WeightMaps
+	{
+
+		/**
+		 * Find the set of per-vertex weight map attributes on a MeshDescription
+		 */
+		MODELINGCOMPONENTS_API void FindVertexWeightMaps(const FMeshDescription* Mesh, TArray<FName>& PropertyNamesOut);
+
+		/**
+		 * Extract a per-vertex weight map from a MeshDescription
+		 * If the attribute with the given name is not found, a WeightMap initialized with the default value is returned
+		 * @return false if weight map was not found
+		 */
+		MODELINGCOMPONENTS_API bool GetVertexWeightMap(const FMeshDescription* Mesh, FName AttributeName, FIndexedWeightMap& WeightMap, float DefaultValue = 1.0f);
+
+	}
+}
\ No newline at end of file
diff --git a/Engine/Plugins/Experimental/MeshModelingToolset/Source/ModelingOperators/Private/CleaningOps/RemeshMeshOp.cpp b/Engine/Plugins/Experimental/MeshModelingToolset/Source/ModelingOperators/Private/CleaningOps/RemeshMeshOp.cpp
--- a/Engine/Plugins/Experimental/MeshModelingToolset/Source/ModelingOperators/Private/CleaningOps/RemeshMeshOp.cpp
+++ b/Engine/Plugins/Experimental/MeshModelingToolset/Source/ModelingOperators/Private/CleaningOps/RemeshMeshOp.cpp
@@ -6,9 +6,31 @@
 #include "DynamicMeshAABBTree3.h"
 #include "DynamicMeshAttributeSet.h"
 #include "Remesher.h"
+#include "QueueRemesher.h"
 #include "MeshConstraintsUtil.h"
 #include "ProjectionTargets.h"
 #include "MeshNormals.h"
+#include "NormalFlowRemesher.h"
+
+namespace
+{
+	TUniquePtr<FRemesher> RemesherFactory(ERemeshType Type, FDynamicMesh3* TargetMesh)
+	{
+		switch(Type)
+		{
+		case ERemeshType::Standard:
+			return MakeUnique<FQueueRemesher>(TargetMesh);
+		case ERemeshType::FullPass:
+			return MakeUnique<FRemesher>(TargetMesh);
+		case ERemeshType::NormalFlow:
+			return MakeUnique<FNormalFlowRemesher>(TargetMesh);
+		default:
+			check(!"Encountered unexpected Remesh Type");
+			return nullptr;
+		}
+	}
+
+}
 
 void FRemeshMeshOp::SetTransform(const FTransform& Transform)
 {
@@ -32,42 +54,43 @@ void FRemeshMeshOp::CalculateResult(FProgressCancel* Progress)
 
 	FDynamicMesh3* TargetMesh = ResultMesh.Get();
 
-	FRemesher Remesher(TargetMesh);
-	Remesher.bEnableSplits = bSplits;
-	Remesher.bEnableFlips = bFlips;
-	Remesher.bEnableCollapses = bCollapses;
+	TUniquePtr<FRemesher> Remesher = RemesherFactory(RemeshType, TargetMesh);
+
+	Remesher->bEnableSplits = bSplits;
+	Remesher->bEnableFlips = bFlips;
+	Remesher->bEnableCollapses = bCollapses;
 
-	Remesher.SetTargetEdgeLength(TargetEdgeLength);
+	Remesher->SetTargetEdgeLength(TargetEdgeLength);
 
-	Remesher.ProjectionMode = (bReproject) ? 
+	Remesher->ProjectionMode = (bReproject) ? 
 		FRemesher::ETargetProjectionMode::AfterRefinement : FRemesher::ETargetProjectionMode::NoProjection;
 
-	Remesher.bEnableSmoothing = (SmoothingStrength > 0);
-	Remesher.SmoothSpeedT = SmoothingStrength;
+	Remesher->bEnableSmoothing = (SmoothingStrength > 0);
+	Remesher->SmoothSpeedT = SmoothingStrength;
 	// convert smooth type from UI enum to (currently 1:1) FRemesher enum
-	Remesher.SmoothType = FRemesher::ESmoothTypes::Uniform;
+	Remesher->SmoothType = FRemesher::ESmoothTypes::Uniform;
 	if (!bDiscardAttributes)
 	{
 		switch (SmoothingType)
 		{
 		case ERemeshSmoothingType::Uniform:
-			Remesher.SmoothType = FRemesher::ESmoothTypes::Uniform;
+			Remesher->SmoothType = FRemesher::ESmoothTypes::Uniform;
 			break;
 		case ERemeshSmoothingType::Cotangent:
-			Remesher.SmoothType = FRemesher::ESmoothTypes::Cotan;
+			Remesher->SmoothType = FRemesher::ESmoothTypes::Cotan;
 			break;
 		case ERemeshSmoothingType::MeanValue:
-			Remesher.SmoothType = FRemesher::ESmoothTypes::MeanValue;
+			Remesher->SmoothType = FRemesher::ESmoothTypes::MeanValue;
 			break;
 		default:
 			ensure(false);
 		}
 	}
-	bool bIsUniformSmooth = (Remesher.SmoothType == FRemesher::ESmoothTypes::Uniform);
+	bool bIsUniformSmooth = (Remesher->SmoothType == FRemesher::ESmoothTypes::Uniform);
 
-	Remesher.bPreventNormalFlips = bPreventNormalFlips;
+	Remesher->bPreventNormalFlips = bPreventNormalFlips;
 
-	Remesher.DEBUG_CHECK_LEVEL = 0;
+	Remesher->DEBUG_CHECK_LEVEL = 0;
 
 	FMeshConstraints constraints;
 	FMeshConstraintsUtil::ConstrainAllBoundariesAndSeams(constraints, *TargetMesh,
@@ -76,33 +99,45 @@ void FRemeshMeshOp::CalculateResult(FProgressCancel* Progress)
 														 MaterialBoundaryConstraint,
 														 true, !bPreserveSharpEdges);
 
-	Remesher.SetExternalConstraints(MoveTemp(constraints));
+	Remesher->SetExternalConstraints(MoveTemp(constraints));
 
 	FMeshProjectionTarget ProjTarget(OriginalMesh.Get(), OriginalMeshSpatial.Get());
-	Remesher.SetProjectionTarget(&ProjTarget);
+	Remesher->SetProjectionTarget(&ProjTarget);
 
-	Remesher.Progress = Progress;
+	Remesher->Progress = Progress;
 
 	if (bDiscardAttributes && !bDiscardAttributesImmediately)
 	{
 		TargetMesh->DiscardAttributes();
 	}
 
-	// run the remesh iterations
-	for (int k = 0; k < RemeshIterations; ++k)
+	if (RemeshType == ERemeshType::FullPass)
 	{
-		// If we are not uniform smoothing, then flips seem to often make things worse.
-		// Possibly this is because without the tangential flow, we won't get to the nice tris.
-		// In this case we are better off basically not flipping, and just letting collapses resolve things
-		// regular-valence polygons - things stay "stuck". 
-		// @todo try implementing edge-length flip criteria instead of valence-flip
-		if (bIsUniformSmooth == false)
+		// Run a fixed number of iterations
+		for (int k = 0; k < RemeshIterations; ++k)
 		{
-			bool bUseFlipsThisPass = (k % 2 == 0 && k < RemeshIterations/2);
-			Remesher.bEnableFlips = bUseFlipsThisPass && bFlips;
+			// If we are not uniform smoothing, then flips seem to often make things worse.
+			// Possibly this is because without the tangential flow, we won't get to the nice tris.
+			// In this case we are better off basically not flipping, and just letting collapses resolve things
+			// regular-valence polygons - things stay "stuck". 
+			// @todo try implementing edge-length flip criteria instead of valence-flip
+			if (bIsUniformSmooth == false)
+			{
+				bool bUseFlipsThisPass = (k % 2 == 0 && k < RemeshIterations / 2);
+				Remesher->bEnableFlips = bUseFlipsThisPass && bFlips;
+			}
+
+			Remesher->BasicRemeshPass();
 		}
-
-		Remesher.BasicRemeshPass();
+	}
+	else if (RemeshType == ERemeshType::Standard || RemeshType == ERemeshType::NormalFlow)
+	{
+		// Run to convergence
+		Remesher->BasicRemeshPass();
+	}
+	else 
+	{
+		check(!"Encountered unexpected Remesh Type");
 	}
 
 	if (!TargetMesh->HasAttributes())
diff --git a/Engine/Plugins/Experimental/MeshModelingToolset/Source/ModelingOperators/Private/CompositionOps/BooleanMeshesOp.cpp b/Engine/Plugins/Experimental/MeshModelingToolset/Source/ModelingOperators/Private/CompositionOps/BooleanMeshesOp.cpp
new file mode 100644
--- /dev/null
+++ b/Engine/Plugins/Experimental/MeshModelingToolset/Source/ModelingOperators/Private/CompositionOps/BooleanMeshesOp.cpp
@@ -0,0 +1,95 @@
+// Copyright Epic Games, Inc. All Rights Reserved.
+
+#include "CompositionOps/BooleanMeshesOp.h"
+
+#include "Operations/MeshBoolean.h"
+
+#include "MeshBoundaryLoops.h"
+#include "Operations/MinimalHoleFiller.h"
+
+void FBooleanMeshesOp::SetTransform(const FTransform& Transform) {
+	ResultTransform = (FTransform3d)Transform;
+}
+
+void FBooleanMeshesOp::CalculateResult(FProgressCancel* Progress)
+{
+	if (Progress->Cancelled())
+	{
+		return;
+	}
+	check(Meshes.Num() == 2 && Transforms.Num() == 2);
+	
+	int FirstIdx = 0;
+	if (Operation == ECSGOperation::DifferenceBA)
+	{
+		FirstIdx = 1;
+	}
+	int OtherIdx = 1 - FirstIdx;
+
+	FMeshBoolean::EBooleanOp Op;
+	// convert UI enum to algorithm enum
+	switch (Operation)
+	{
+	case ECSGOperation::DifferenceAB:
+	case ECSGOperation::DifferenceBA:
+		Op = FMeshBoolean::EBooleanOp::Difference;
+		break;
+	case ECSGOperation::Union:
+		Op = FMeshBoolean::EBooleanOp::Union;
+		break;
+	case ECSGOperation::Intersect:
+		Op = FMeshBoolean::EBooleanOp::Intersect;
+		break;
+	default:
+		check(false); // all conversion cases should be implemented
+		Op = FMeshBoolean::EBooleanOp::Union;
+	}
+
+	FMeshBoolean MeshBoolean(Meshes[FirstIdx].Get(), (FTransform3d)Transforms[FirstIdx], Meshes[OtherIdx].Get(), (FTransform3d)Transforms[OtherIdx], ResultMesh.Get(), Op);
+	if (Progress->Cancelled())
+	{
+		return;
+	}
+
+	MeshBoolean.Progress = Progress;
+	bool bSuccess = MeshBoolean.Compute();
+	ResultTransform = MeshBoolean.ResultTransform;
+
+	if (Progress->Cancelled())
+	{
+		return;
+	}
+
+	if (!bSuccess && bAttemptFixHoles)
+	{
+		FMeshBoundaryLoops OpenBoundary(MeshBoolean.Result, false);
+		TSet<int> ConsiderEdges(MeshBoolean.CreatedBoundaryEdges);
+		OpenBoundary.EdgeFilterFunc = [&ConsiderEdges](int EID)
+		{
+			return ConsiderEdges.Contains(EID);
+		};
+		OpenBoundary.Compute();
+
+		if (Progress->Cancelled())
+		{
+			return;
+		}
+
+		for (FEdgeLoop& Loop : OpenBoundary.Loops)
+		{
+			FMinimalHoleFiller Filler(MeshBoolean.Result, Loop);
+			Filler.Fill();
+		}
+		for (int EID : MeshBoolean.CreatedBoundaryEdges)
+		{
+			if (MeshBoolean.Result->IsEdge(EID) && MeshBoolean.Result->IsBoundaryEdge(EID))
+			{
+				CreatedBoundaryEdges.Add(EID);
+			}
+		}
+	}
+	else
+	{
+		CreatedBoundaryEdges = MeshBoolean.CreatedBoundaryEdges;
+	}
+}
\ No newline at end of file
diff --git a/Engine/Plugins/Experimental/MeshModelingToolset/Source/ModelingOperators/Private/ParameterizationOps/UVProjectionOp.cpp b/Engine/Plugins/Experimental/MeshModelingToolset/Source/ModelingOperators/Private/ParameterizationOps/UVProjectionOp.cpp
--- a/Engine/Plugins/Experimental/MeshModelingToolset/Source/ModelingOperators/Private/ParameterizationOps/UVProjectionOp.cpp
+++ b/Engine/Plugins/Experimental/MeshModelingToolset/Source/ModelingOperators/Private/ParameterizationOps/UVProjectionOp.cpp
@@ -70,7 +70,7 @@ void FUVProjectionOp::CalculateResult(FProgressCancel* Progress)
 			if (ResultMesh->IsVertex(VID))
 			{
 				FVector2f UV = ProjAxis(TransformedVertices[VID], 0, 1, 1, 1);
-				UVLayer->InsertElement(VID, (float*)UV, VID, true);
+				UVLayer->InsertElement(VID, (float*)UV, true);
 			}
 		}
 		UVLayer->EndUnsafeElementsInsert();
@@ -129,7 +129,7 @@ void FUVProjectionOp::CalculateResult(FProgressCancel* Progress)
 					if (!ElementIdx)
 					{
 						FVector3d LocalV = TransformedVertices[SourceTri[SubIdx]];
-						int NewElementIdx = UVLayer->AppendElement(ProjAxis(LocalV, Minor1, Minor2, MajorAxisSign*Minor1Flip[MajorAxis], Minor2Flip[MajorAxis]), ElementKey.A);
+						int NewElementIdx = UVLayer->AppendElement(ProjAxis(LocalV, Minor1, Minor2, MajorAxisSign*Minor1Flip[MajorAxis], Minor2Flip[MajorAxis]));
 						VertIdxAndBucketIDToElementID.Add(ElementKey, NewElementIdx);
 						OverlayTri[SubIdx] = NewElementIdx;
 					}
@@ -170,7 +170,7 @@ void FUVProjectionOp::CalculateResult(FProgressCancel* Progress)
 						if (!ElementIdx)
 						{
 							FVector3d LocalV = TransformedVertices[SourceTri[SubIdx]];
-							int NewElementIdx = UVLayer->AppendElement(ProjAxis(LocalV, Minor1, Minor2, MajorAxisSign, 1), ElementKey.A);
+							int NewElementIdx = UVLayer->AppendElement(ProjAxis(LocalV, Minor1, Minor2, MajorAxisSign, 1));
 							VertIdxAndBucketIDToElementID.Add(ElementKey, NewElementIdx);
 							OverlayTri[SubIdx] = NewElementIdx;
 						}
@@ -208,7 +208,7 @@ void FUVProjectionOp::CalculateResult(FProgressCancel* Progress)
 						if (!ElementIdx)
 						{
 							FVector2f UV(-(float(VAngle)*FMathf::InvPi - 1.0f)*Scale.X + Offset.X, -float(LocalV.Z)*Scale.Y + Offset.Y);
-							int NewElementIdx = UVLayer->AppendElement(UV, ElementKey.A);
+							int NewElementIdx = UVLayer->AppendElement(UV);
 							VertIdxAndBucketIDToElementID.Add(ElementKey, NewElementIdx);
 							OverlayTri[SubIdx] = NewElementIdx;
 						}
diff --git a/Engine/Plugins/Experimental/MeshModelingToolset/Source/ModelingOperators/Public/CleaningOps/HoleFillOp.h b/Engine/Plugins/Experimental/MeshModelingToolset/Source/ModelingOperators/Public/CleaningOps/HoleFillOp.h
new file mode 100644
--- /dev/null
+++ b/Engine/Plugins/Experimental/MeshModelingToolset/Source/ModelingOperators/Public/CleaningOps/HoleFillOp.h
@@ -0,0 +1,42 @@
+// Copyright Epic Games, Inc. All Rights Reserved.
+
+#pragma once
+
+#include "ModelingOperators.h"
+#include "EdgeLoop.h"
+#include "HoleFillOp.generated.h"
+
+UENUM()
+enum class EHoleFillOpFillType : uint8
+{
+	/** Fill with a fan of triangles connected to a new central vertex. */
+	TriangleFan UMETA(DisplayName = "TriangleFan"),
+
+	/** Incrementally triangulate the hole boundary without introducing new interior vertices. */
+	PolygonEarClipping UMETA(DisplayName = "PolygonEarClipping"),
+
+	/** Choose a best-fit plane, project the boundary vertices to the plane, and use 2D Delaunay triangulation. */
+	Planar UMETA(DisplayName = "Planar"),
+
+	/** Fill with a triangulation which attempts to minimize Gaussian curvature and not introcude new interior vertices. */
+	Minimal UMETA(DisplayName = "Minimal")
+};
+
+
+class MODELINGOPERATORS_API FHoleFillOp : public FDynamicMeshOperator
+{
+public:
+
+	// inputs
+	const FDynamicMesh3* OriginalMesh = nullptr;
+	EHoleFillOpFillType FillType = EHoleFillOpFillType::TriangleFan;
+	double MeshUVScaleFactor = 1.0;
+	TArray<FEdgeLoop> Loops;
+
+	// output
+	TArray<int32> NewTriangles;
+
+	// FDynamicMeshOperator implementation
+	void CalculateResult(FProgressCancel* Progress) override;
+};
+
diff --git a/Engine/Plugins/Experimental/MeshModelingToolset/Source/ModelingOperators/Public/CleaningOps/RemeshMeshOp.h b/Engine/Plugins/Experimental/MeshModelingToolset/Source/ModelingOperators/Public/CleaningOps/RemeshMeshOp.h
--- a/Engine/Plugins/Experimental/MeshModelingToolset/Source/ModelingOperators/Public/CleaningOps/RemeshMeshOp.h
+++ b/Engine/Plugins/Experimental/MeshModelingToolset/Source/ModelingOperators/Public/CleaningOps/RemeshMeshOp.h
@@ -17,6 +17,21 @@ typedef TMeshAABBTree3<FDynamicMesh3> FDynamicMeshAABBTree3;
 
 /** Remeshing modes */
 UENUM()
+enum class ERemeshType : uint8
+{
+	/** One pass over the entire mesh, then remesh only changed edges */
+	Standard = 0 UMETA(DisplayName = "Standard"),
+
+	/** Multiple full passes over the entire mesh */
+	FullPass = 1 UMETA(DisplayName = "Full Pass"),
+
+	/** One pass over the entire mesh, then remesh only changed edges. Use Normal flow to align triangles with input.*/
+	NormalFlow = 2 UMETA(DisplayName = "Normal Flow"),
+
+};
+
+/** Smoothing modes */
+UENUM()
 enum class ERemeshSmoothingType : uint8
 {
 	/** Uniform Smoothing */
@@ -38,6 +53,8 @@ public:
 	TSharedPtr<FDynamicMesh3> OriginalMesh;
 	TSharedPtr<FDynamicMeshAABBTree3> OriginalMeshSpatial;
 
+	ERemeshType RemeshType = ERemeshType::Standard;
+
 	int RemeshIterations;
 	float SmoothingStrength, TargetEdgeLength;
 	ERemeshSmoothingType SmoothingType;
diff --git a/Engine/Plugins/Experimental/MeshModelingToolset/Source/ModelingOperators/Public/CompositionOps/BooleanMeshesOp.h b/Engine/Plugins/Experimental/MeshModelingToolset/Source/ModelingOperators/Public/CompositionOps/BooleanMeshesOp.h
new file mode 100644
--- /dev/null
+++ b/Engine/Plugins/Experimental/MeshModelingToolset/Source/ModelingOperators/Public/CompositionOps/BooleanMeshesOp.h
@@ -0,0 +1,61 @@
+// Copyright Epic Games, Inc. All Rights Reserved.
+
+#pragma once
+
+#include "CoreMinimal.h"
+#include "Util/ProgressCancel.h"
+#include "ModelingOperators.h"
+
+
+#include "BooleanMeshesOp.generated.h"
+
+
+/** CSG operation types */
+UENUM()
+enum class ECSGOperation : uint8
+{
+	/** Subtracts the first object from the second */
+	DifferenceAB = 0 UMETA(DisplayName = "A - B"),
+
+	/** Subtracts the second object from the first */
+	DifferenceBA = 1 UMETA(DisplayName = "B - A"),
+
+	/** intersection of two objects */
+	Intersect = 2 UMETA(DisplayName = "Intersect"),
+
+	/** union of two objects */
+	Union = 3 UMETA(DisplayName = "Union"),
+};
+
+
+
+class MODELINGOPERATORS_API FBooleanMeshesOp : public FDynamicMeshOperator
+{
+public:
+	virtual ~FBooleanMeshesOp() {}
+
+	// inputs
+	ECSGOperation Operation;
+	TArray<TSharedPtr<const FDynamicMesh3>> Meshes;
+	TArray<FTransform> Transforms; // 1:1 with Meshes
+	bool bAttemptFixHoles;
+
+	void SetTransform(const FTransform& Transform);
+
+	//
+	// FDynamicMeshOperator implementation
+	// 
+
+	virtual void CalculateResult(FProgressCancel* Progress) override;
+
+	// IDs of any newly-created boundary edges in the result mesh
+	TArray<int> GetCreatedBoundaryEdges() const
+	{
+		return CreatedBoundaryEdges;
+	}
+
+private:
+	TArray<int> CreatedBoundaryEdges;
+};
+
+
diff --git a/Engine/Plugins/Experimental/MeshModelingToolset/Source/ModelingOperators/Public/CompositionOps/SelfUnionMeshesOp.h b/Engine/Plugins/Experimental/MeshModelingToolset/Source/ModelingOperators/Public/CompositionOps/SelfUnionMeshesOp.h
new file mode 100644
--- /dev/null
+++ b/Engine/Plugins/Experimental/MeshModelingToolset/Source/ModelingOperators/Public/CompositionOps/SelfUnionMeshesOp.h
@@ -0,0 +1,37 @@
+// Copyright Epic Games, Inc. All Rights Reserved.
+
+#pragma once
+
+#include "CoreMinimal.h"
+#include "Util/ProgressCancel.h"
+#include "ModelingOperators.h"
+
+
+class MODELINGOPERATORS_API FSelfUnionMeshesOp : public FDynamicMeshOperator
+{
+public:
+	virtual ~FSelfUnionMeshesOp() {}
+
+	// inputs
+	TSharedPtr<const FDynamicMesh3> CombinedMesh;
+	bool bAttemptFixHoles;
+
+	void SetTransform(const FTransform& Transform);
+
+	//
+	// FDynamicMeshOperator implementation
+	// 
+
+	virtual void CalculateResult(FProgressCancel* Progress) override;
+
+	// IDs of any newly-created boundary edges in the result mesh
+	TArray<int> GetCreatedBoundaryEdges() const
+	{
+		return CreatedBoundaryEdges;
+	}
+
+private:
+	TArray<int> CreatedBoundaryEdges;
+};
+
+
diff --git a/Engine/Plugins/Experimental/MeshModelingToolset/Source/ModelingOperatorsEditorOnly/Private/ParameterizationOps/ParameterizeMeshOp.cpp b/Engine/Plugins/Experimental/MeshModelingToolset/Source/ModelingOperatorsEditorOnly/Private/ParameterizationOps/ParameterizeMeshOp.cpp
--- a/Engine/Plugins/Experimental/MeshModelingToolset/Source/ModelingOperatorsEditorOnly/Private/ParameterizationOps/ParameterizeMeshOp.cpp
+++ b/Engine/Plugins/Experimental/MeshModelingToolset/Source/ModelingOperatorsEditorOnly/Private/ParameterizationOps/ParameterizeMeshOp.cpp
@@ -208,10 +208,9 @@ bool FParameterizeMeshOp::ComputeUVs(FDynamicMesh3& Mesh,  TFunction<bool(float)
 
 			// The associated VertID in the dynamic mesh
 			const int32 VertOffset = VertexRemapArray[i];
-			const int32 VertID = LinearMesh.VertToID[VertOffset];
 
 			// add the UV to the mesh overlay
-			const int32 NewID = UVOverlay->AppendElement(FVector2f(UV), VertID);
+			const int32 NewID = UVOverlay->AppendElement(FVector2f(UV));
 			UVOffsetToElID.Add(NewID);
 		}
 
@@ -294,7 +293,7 @@ bool FParameterizeMeshOp::ComputeUVs_ExpMap(FDynamicMesh3& Mesh, TFunction<bool(
 	{
 		check(Param.HasUV(VertexID));
 		FVector2d UV = Param.GetUV(VertexID) - Center;
-		int32 NewElemID = UVOverlay->AppendElement(FVector2f(UV), VertexID);
+		int32 NewElemID = UVOverlay->AppendElement(FVector2f(UV));
 		check(NewElemID == VertexID);
 	}
 
@@ -476,8 +475,7 @@ void FParameterizeMeshOp::CalculateResult(FProgressCancel* Progress)
 			for (int32 ElemID : ComponentUVOverlay->ElementIndicesItr())
 			{
 				int ParentVertexID = ComponentUVOverlay->GetParentVertex(ElemID);
-				int32 BaseVertexID = ComponentSubmesh.MapVertexToBaseMesh(ParentVertexID);
-				int32 BaseElemID = UVOverlay->AppendElement(ComponentUVOverlay->GetElement(ElemID), BaseVertexID);
+				int32 BaseElemID = UVOverlay->AppendElement(ComponentUVOverlay->GetElement(ElemID));
 				ElemIDMap[ElemID] = BaseElemID;
 			}
 
diff --git a/Engine/Plugins/Experimental/ModelingToolsEditorMode/Source/ModelingToolsEditorMode/Private/ModelingToolsActions.cpp b/Engine/Plugins/Experimental/ModelingToolsEditorMode/Source/ModelingToolsEditorMode/Private/ModelingToolsActions.cpp
--- a/Engine/Plugins/Experimental/ModelingToolsEditorMode/Source/ModelingToolsEditorMode/Private/ModelingToolsActions.cpp
+++ b/Engine/Plugins/Experimental/ModelingToolsEditorMode/Source/ModelingToolsEditorMode/Private/ModelingToolsActions.cpp
@@ -3,6 +3,7 @@
 #include "ModelingToolsActions.h"
 #include "EditorStyleSet.h"
 #include "DynamicMeshSculptTool.h"
+#include "MeshVertexSculptTool.h"
 #include "MeshInspectorTool.h"
 #include "DrawPolygonTool.h"
 #include "EditMeshPolygonsTool.h"
@@ -73,6 +74,7 @@ void FModelingToolActionCommands::RegisterAllToolActions()
 {
 	FModelingToolActionCommands::Register();
 	FSculptToolActionCommands::Register();
+	FVertexSculptToolActionCommands::Register();
 	FDrawPolygonToolActionCommands::Register();
 	FTransformToolActionCommands::Register();
 	FMeshSelectionToolActionCommands::Register();
@@ -84,6 +86,7 @@ void FModelingToolActionCommands::UnregisterAllToolActions()
 {
 	FModelingToolActionCommands::Unregister();
 	FSculptToolActionCommands::Unregister();
+	FVertexSculptToolActionCommands::Unregister();
 	FDrawPolygonToolActionCommands::Unregister();
 	FTransformToolActionCommands::Unregister();
 	FMeshSelectionToolActionCommands::Unregister();
@@ -106,6 +109,10 @@ void FModelingToolActionCommands::UpdateToolCommandBinding(UInteractiveTool* Too
 	{
 		UPDATE_BINDING(FSculptToolActionCommands);
 	}
+	else if (Cast<UMeshVertexSculptTool>(Tool) != nullptr)
+	{
+		UPDATE_BINDING(FVertexSculptToolActionCommands);
+	}
 	else if (Cast<UDrawPolygonTool>(Tool) != nullptr)
 	{
 		UPDATE_BINDING(FDrawPolygonToolActionCommands);
@@ -132,104 +139,24 @@ void FModelingToolActionCommands::UpdateToolCommandBinding(UInteractiveTool* Too
 
 
 
-FSculptToolActionCommands::FSculptToolActionCommands() :
-	TInteractiveToolCommands<FSculptToolActionCommands>(
-		"ModelingToolsSculptTool", // Context name for fast lookup
-		NSLOCTEXT("Contexts", "ModelingToolsSculptTool", "Modeling Tools - Sculpt Tool"), // Localized context name for displaying
-		NAME_None, // Parent
-		FEditorStyle::GetStyleSetName() // Icon Style Set
-		)
-{
-}
-
-void FSculptToolActionCommands::GetToolDefaultObjectList(TArray<UInteractiveTool*>& ToolCDOs)
-{
-	ToolCDOs.Add(GetMutableDefault<UDynamicMeshSculptTool>());
-}
-
-
-
-FTransformToolActionCommands::FTransformToolActionCommands() :
-	TInteractiveToolCommands<FTransformToolActionCommands>(
-		"ModelingToolsTransformTool", // Context name for fast lookup
-		NSLOCTEXT("Contexts", "ModelingToolsTransformTool", "Modeling Tools - Transform Tool"), // Localized context name for displaying
-		NAME_None, // Parent
-		FEditorStyle::GetStyleSetName() // Icon Style Set
-		)
-{
-}
-
-void FTransformToolActionCommands::GetToolDefaultObjectList(TArray<UInteractiveTool*>& ToolCDOs)
-{
-	ToolCDOs.Add(GetMutableDefault<UTransformMeshesTool>());
-}
-
-
-
-FDrawPolygonToolActionCommands::FDrawPolygonToolActionCommands() :
-	TInteractiveToolCommands<FDrawPolygonToolActionCommands>(
-		"ModelingToolsDrawPolygonTool", // Context name for fast lookup
-		NSLOCTEXT("Contexts", "ModelingToolsDrawPolygonTool", "Modeling Tools - Draw Polygon Tool"), // Localized context name for displaying
-		NAME_None, // Parent
-		FEditorStyle::GetStyleSetName() // Icon Style Set
-		)
-{
+#define DEFINE_TOOL_ACTION_COMMANDS(CommandsClassName, ContextNameString, SettingsDialogString, ToolClassName ) \
+CommandsClassName::CommandsClassName() : TInteractiveToolCommands<CommandsClassName>( \
+ContextNameString, NSLOCTEXT("Contexts", ContextNameString, SettingsDialogString), NAME_None, FEditorStyle::GetStyleSetName()) {} \
+void CommandsClassName::GetToolDefaultObjectList(TArray<UInteractiveTool*>& ToolCDOs) \
+{\
+	ToolCDOs.Add(GetMutableDefault<ToolClassName>()); \
 }
 
-void FDrawPolygonToolActionCommands::GetToolDefaultObjectList(TArray<UInteractiveTool*>& ToolCDOs)
-{
-	ToolCDOs.Add(GetMutableDefault<UDrawPolygonTool>());
-}
 
 
-FMeshSelectionToolActionCommands::FMeshSelectionToolActionCommands() :
-	TInteractiveToolCommands<FMeshSelectionToolActionCommands>(
-		"ModelingToolsMeshSelectionTool", // Context name for fast lookup
-		NSLOCTEXT("Contexts", "ModelingToolsMeshSelectionTool", "Modeling Tools - Mesh Selection Tool"), // Localized context name for displaying
-		NAME_None, // Parent
-		FEditorStyle::GetStyleSetName() // Icon Style Set
-		)
-{
-}
-
-void FMeshSelectionToolActionCommands::GetToolDefaultObjectList(TArray<UInteractiveTool*>& ToolCDOs)
-{
-	ToolCDOs.Add(GetMutableDefault<UMeshSelectionTool>());
-}
-
-
-
-FMeshPlaneCutToolActionCommands::FMeshPlaneCutToolActionCommands() :
-	TInteractiveToolCommands<FMeshPlaneCutToolActionCommands>(
-		"ModelingToolsMeshPlaneCutTool", // Context name for fast lookup
-		NSLOCTEXT("Contexts", "ModelingToolsMeshPlaneCutTool", "Modeling Tools - Mesh Plane Cut Tool"), // Localized context name for displaying
-		NAME_None, // Parent
-		FEditorStyle::GetStyleSetName() // Icon Style Set
-		)
-{
-}
-
-void FMeshPlaneCutToolActionCommands::GetToolDefaultObjectList(TArray<UInteractiveTool*>& ToolCDOs)
-{
-	ToolCDOs.Add(GetMutableDefault<UPlaneCutTool>());
-}
-
-
-
-FEditMeshPolygonsToolActionCommands::FEditMeshPolygonsToolActionCommands() :
-	TInteractiveToolCommands<FEditMeshPolygonsToolActionCommands>(
-		"ModelingToolsEditMeshPolygonsTool", // Context name for fast lookup
-		NSLOCTEXT("Contexts", "ModelingToolsEditMeshPolygonsTool", "Modeling Tools - Edit Mesh Polygons Tool"), // Localized context name for displaying
-		NAME_None, // Parent
-		FEditorStyle::GetStyleSetName() // Icon Style Set
-		)
-{
-}
-void FEditMeshPolygonsToolActionCommands::GetToolDefaultObjectList(TArray<UInteractiveTool*>& ToolCDOs)
-{
-	ToolCDOs.Add(GetMutableDefault<UEditMeshPolygonsTool>());
-}
 
+DEFINE_TOOL_ACTION_COMMANDS(FSculptToolActionCommands, "ModelingToolsSculptTool", "Modeling Tools - Sculpt Tool", UDynamicMeshSculptTool);
+DEFINE_TOOL_ACTION_COMMANDS(FVertexSculptToolActionCommands, "ModelingToolsVertexSculptTool", "Modeling Tools - Vertex Sculpt Tool", UMeshVertexSculptTool);
+DEFINE_TOOL_ACTION_COMMANDS(FTransformToolActionCommands, "ModelingToolsTransformTool", "Modeling Tools - Transform Tool", UTransformMeshesTool);
+DEFINE_TOOL_ACTION_COMMANDS(FDrawPolygonToolActionCommands, "ModelingToolsDrawPolygonTool", "Modeling Tools - Draw Polygon Tool", UDrawPolygonTool);
+DEFINE_TOOL_ACTION_COMMANDS(FMeshSelectionToolActionCommands, "ModelingToolsMeshSelectionTool", "Modeling Tools - Mesh Selection Tool", UMeshSelectionTool);
+DEFINE_TOOL_ACTION_COMMANDS(FMeshPlaneCutToolActionCommands, "ModelingToolsMeshPlaneCutTool", "Modeling Tools - Mesh Plane Cut Tool", UPlaneCutTool);
+DEFINE_TOOL_ACTION_COMMANDS(FEditMeshPolygonsToolActionCommands, "ModelingToolsEditMeshPolygonsTool", "Modeling Tools - Edit Mesh Polygons Tool", UEditMeshPolygonsTool);
 
 
 
