ModelingComponents: new USimpleDynamicMeshComponent::FastNotifyUVsUpdated() implementation was not actually updating UVs.
#rnx

diff --git a/Engine/Plugins/Experimental/MeshModelingToolset/Source/ModelingComponents/Private/SimpleDynamicMeshComponent.cpp b/Engine/Plugins/Experimental/MeshModelingToolset/Source/ModelingComponents/Private/SimpleDynamicMeshComponent.cpp
--- a/Engine/Plugins/Experimental/MeshModelingToolset/Source/ModelingComponents/Private/SimpleDynamicMeshComponent.cpp
+++ b/Engine/Plugins/Experimental/MeshModelingToolset/Source/ModelingComponents/Private/SimpleDynamicMeshComponent.cpp
@@ -269,7 +269,7 @@ void USimpleDynamicMeshComponent::FastNotifyVertexAttributesUpdated(EMeshRenderA
 
 void USimpleDynamicMeshComponent::FastNotifyUVsUpdated()
 {
-	FastNotifyVertexAttributesUpdated(EMeshRenderAttributeFlags::Positions);
+	FastNotifyVertexAttributesUpdated(EMeshRenderAttributeFlags::VertexUVs);
 }
 
 
