optional grid snapping for mesh CSG tool widgets
#rnx

diff --git a/Engine/Plugins/Experimental/MeshModelingToolset/Source/MeshModelingTools/Private/CSGMeshesTool.cpp b/Engine/Plugins/Experimental/MeshModelingToolset/Source/MeshModelingTools/Private/CSGMeshesTool.cpp
--- a/Engine/Plugins/Experimental/MeshModelingToolset/Source/MeshModelingTools/Private/CSGMeshesTool.cpp
+++ b/Engine/Plugins/Experimental/MeshModelingToolset/Source/MeshModelingTools/Private/CSGMeshesTool.cpp
@@ -315,6 +315,11 @@ void UCSGMeshesTool::Render(IToolsContextRenderAPI* RenderAPI)
 
 void UCSGMeshesTool::OnTick(float DeltaTime)
 {
+	for (UTransformGizmo* Gizmo : TransformGizmos)
+	{
+		Gizmo->bSnapToWorldGrid = CSGProperties->bSnapToWorldGrid;
+	}
+
 	Preview->Tick(DeltaTime);
 }
 
@@ -333,7 +338,10 @@ void UCSGMeshesTool::OnPropertyModified(UObject* PropertySet, FProperty* Propert
 	{
 		UpdateGizmoVisibility();
 	}
-	else if (Property && (Property->GetFName() == GET_MEMBER_NAME_CHECKED(UCSGMeshesToolProperties, bDeleteInputActors)))
+	else if (Property && 
+		(  Property->GetFName() == GET_MEMBER_NAME_CHECKED(UCSGMeshesToolProperties, bDeleteInputActors)
+		|| Property->GetFName() == GET_MEMBER_NAME_CHECKED(UCSGMeshesToolProperties, bSnapToWorldGrid)
+		))
 	{
 		// nothing
 	}
diff --git a/Engine/Plugins/Experimental/MeshModelingToolset/Source/MeshModelingTools/Public/CSGMeshesTool.h b/Engine/Plugins/Experimental/MeshModelingToolset/Source/MeshModelingTools/Public/CSGMeshesTool.h
--- a/Engine/Plugins/Experimental/MeshModelingToolset/Source/MeshModelingTools/Public/CSGMeshesTool.h
+++ b/Engine/Plugins/Experimental/MeshModelingToolset/Source/MeshModelingTools/Public/CSGMeshesTool.h
@@ -54,6 +54,10 @@ public:
 	UPROPERTY(EditAnywhere, Category = Options)
 	bool bShowTransformUI = true;
 
+	/** Snap the cut plane to the world grid */
+	UPROPERTY(EditAnywhere, Category = Snapping, meta = (EditCondition = "bShowTransformUI == true"))
+	bool bSnapToWorldGrid = false;
+
 	/** Show boundary edges created by the CSG operation -- often due to numerical error */
 	UPROPERTY(EditAnywhere, Category = Options)
 	bool bShowNewBoundaryEdges = true;
