fix undo of input hiding for csg + merge tools, + add option to keep (not hide or delete) sources for mesh csg+merge

#rnx

diff --git a/Engine/Plugins/Experimental/MeshModelingToolset/Source/MeshModelingTools/Private/CSGMeshesTool.cpp b/Engine/Plugins/Experimental/MeshModelingToolset/Source/MeshModelingTools/Private/CSGMeshesTool.cpp
--- a/Engine/Plugins/Experimental/MeshModelingToolset/Source/MeshModelingTools/Private/CSGMeshesTool.cpp
+++ b/Engine/Plugins/Experimental/MeshModelingToolset/Source/MeshModelingTools/Private/CSGMeshesTool.cpp
@@ -237,35 +237,55 @@ void UCSGMeshesTool::Shutdown(EToolShutdownType ShutdownType)
 			GetToolManager()->EndUndoTransaction();
 		}
 
-		// Hide or destroy the sources
+		// Restore (unhide) the source meshes
+		for (auto& ComponentTarget : ComponentTargets)
 		{
+			ComponentTarget->SetOwnerVisibility(true);
+		}
 
-			if (CSGProperties->bDeleteInputActors)
+		// Hide or destroy the sources
+		if (CSGProperties->OnToolAccept != ECSGAcceptBehavior::LeaveOriginalsUnchanged)
+		{
+			bool bDelete = CSGProperties->OnToolAccept == ECSGAcceptBehavior::DeleteOriginals;
+			if (bDelete)
 			{
 				GetToolManager()->BeginUndoTransaction(LOCTEXT("RemoveSources", "Remove Sources"));
 			}
+			else
+			{
+#if WITH_EDITOR
+				GetToolManager()->BeginUndoTransaction(LOCTEXT("HideSources", "Hide Sources"));
+#endif
+			}
 
 			for (auto& ComponentTarget : ComponentTargets)
 			{
-				ComponentTarget->SetOwnerVisibility(true);
 				AActor* Actor = ComponentTarget->GetOwnerActor();
-				if (CSGProperties->bDeleteInputActors)
+				if (bDelete)
 				{
 					Actor->Destroy();
 				}
 				else
 				{
 #if WITH_EDITOR
+					// Save the actor to the transaction buffer to support undo/redo, but do
+					// not call Modify, as we do not want to dirty the actor's package and
+					// we're only editing temporary, transient values
+					SaveToTransactionBuffer(Actor, false);
 					Actor->SetIsTemporarilyHiddenInEditor(true);
 #endif
 				}
 			}
-
-			if (CSGProperties->bDeleteInputActors)
+			if (bDelete)
 			{
 				GetToolManager()->EndUndoTransaction();
 			}
-
+			else
+			{
+#if WITH_EDITOR
+				GetToolManager()->EndUndoTransaction();
+#endif
+			}
 		}
 
 	}
@@ -339,7 +359,7 @@ void UCSGMeshesTool::OnPropertyModified(UObject* PropertySet, FProperty* Propert
 		UpdateGizmoVisibility();
 	}
 	else if (Property && 
-		(  Property->GetFName() == GET_MEMBER_NAME_CHECKED(UCSGMeshesToolProperties, bDeleteInputActors)
+		(  Property->GetFName() == GET_MEMBER_NAME_CHECKED(UCSGMeshesToolProperties, OnToolAccept)
 		|| Property->GetFName() == GET_MEMBER_NAME_CHECKED(UCSGMeshesToolProperties, bSnapToWorldGrid)
 		))
 	{
diff --git a/Engine/Plugins/Experimental/MeshModelingToolset/Source/MeshModelingTools/Private/SelfUnionMeshesTool.cpp b/Engine/Plugins/Experimental/MeshModelingToolset/Source/MeshModelingTools/Private/SelfUnionMeshesTool.cpp
--- a/Engine/Plugins/Experimental/MeshModelingToolset/Source/MeshModelingTools/Private/SelfUnionMeshesTool.cpp
+++ b/Engine/Plugins/Experimental/MeshModelingToolset/Source/MeshModelingTools/Private/SelfUnionMeshesTool.cpp
@@ -221,35 +221,55 @@ void USelfUnionMeshesTool::Shutdown(EToolShutdownType ShutdownType)
 			GetToolManager()->EndUndoTransaction();
 		}
 
-		// Hide or destroy the sources
+		// Restore (unhide) the source meshes
+		for (auto& ComponentTarget : ComponentTargets)
 		{
+			ComponentTarget->SetOwnerVisibility(true);
+		}
 
-			if (Properties->bDeleteInputActors)
+		// Hide or destroy the sources
+		if (Properties->OnToolAccept != ESelfUnionAcceptBehavior::LeaveOriginalsUnchanged)
+		{
+			bool bDelete = Properties->OnToolAccept == ESelfUnionAcceptBehavior::DeleteOriginals;
+			if (bDelete)
 			{
 				GetToolManager()->BeginUndoTransaction(LOCTEXT("RemoveSources", "Remove Sources"));
 			}
+			else
+			{
+#if WITH_EDITOR
+				GetToolManager()->BeginUndoTransaction(LOCTEXT("HideSources", "Hide Sources"));
+#endif
+			}
 
 			for (auto& ComponentTarget : ComponentTargets)
 			{
-				ComponentTarget->SetOwnerVisibility(true);
 				AActor* Actor = ComponentTarget->GetOwnerActor();
-				if (Properties->bDeleteInputActors)
+				if (bDelete)
 				{
 					Actor->Destroy();
 				}
 				else
 				{
 #if WITH_EDITOR
+					// Save the actor to the transaction buffer to support undo/redo, but do
+					// not call Modify, as we do not want to dirty the actor's package and
+					// we're only editing temporary, transient values
+					SaveToTransactionBuffer(Actor, false);
 					Actor->SetIsTemporarilyHiddenInEditor(true);
 #endif
 				}
 			}
-
-			if (Properties->bDeleteInputActors)
+			if (bDelete)
 			{
 				GetToolManager()->EndUndoTransaction();
 			}
-
+			else
+			{
+#if WITH_EDITOR
+				GetToolManager()->EndUndoTransaction();
+#endif
+			}
 		}
 
 	}
@@ -301,7 +321,7 @@ void USelfUnionMeshesTool::PostEditChangeProperty(FPropertyChangedEvent& Propert
 
 void USelfUnionMeshesTool::OnPropertyModified(UObject* PropertySet, FProperty* Property)
 {
-	if (Property && (Property->GetFName() == GET_MEMBER_NAME_CHECKED(USelfUnionMeshesToolProperties, bDeleteInputActors)))
+	if (Property && (Property->GetFName() == GET_MEMBER_NAME_CHECKED(USelfUnionMeshesToolProperties, OnToolAccept)))
 	{
 		// nothing
 	}
diff --git a/Engine/Plugins/Experimental/MeshModelingToolset/Source/MeshModelingTools/Public/CSGMeshesTool.h b/Engine/Plugins/Experimental/MeshModelingToolset/Source/MeshModelingTools/Public/CSGMeshesTool.h
--- a/Engine/Plugins/Experimental/MeshModelingToolset/Source/MeshModelingTools/Public/CSGMeshesTool.h
+++ b/Engine/Plugins/Experimental/MeshModelingToolset/Source/MeshModelingTools/Public/CSGMeshesTool.h
@@ -33,6 +33,18 @@ public:
 };
 
 
+/** Choice of behavior on accepting CSG */
+UENUM()
+enum class ECSGAcceptBehavior : uint8
+{
+	DeleteOriginals = 0			UMETA(DisplayName = "Delete Inputs"),
+
+	HideOriginals = 1			UMETA(DisplayName = "Hide Inputs"),
+
+	LeaveOriginalsUnchanged = 2	UMETA(DisplayName = "Keep Inputs")
+};
+
+
 
 /**
  * Standard properties of the CSG operation
@@ -46,10 +58,6 @@ public:
 	UPROPERTY(EditAnywhere, Category = Options)
 	ECSGOperation Operation = ECSGOperation::Union;
 
-	/** Remove the source Actors/Components when accepting results of tool.*/
-	UPROPERTY(EditAnywhere, Category = Options)
-	bool bDeleteInputActors = true;
-
 	/** Show UI to allow changing translation, rotation and scale of input meshes */
 	UPROPERTY(EditAnywhere, Category = Options)
 	bool bShowTransformUI = true;
@@ -65,6 +73,10 @@ public:
 	/** Automatically attempt to fill any holes left by CSG (e.g. due to numerical errors) */
 	UPROPERTY(EditAnywhere, Category = Options)
 	bool bAttemptFixHoles = false;
+
+	/** What to do with the source Actors/Components when accepting results of tool.*/
+	UPROPERTY(EditAnywhere, Category = ToolOutputOptions)
+	ECSGAcceptBehavior OnToolAccept;
 };
 
 
diff --git a/Engine/Plugins/Experimental/MeshModelingToolset/Source/MeshModelingTools/Public/SelfUnionMeshesTool.h b/Engine/Plugins/Experimental/MeshModelingToolset/Source/MeshModelingTools/Public/SelfUnionMeshesTool.h
--- a/Engine/Plugins/Experimental/MeshModelingToolset/Source/MeshModelingTools/Public/SelfUnionMeshesTool.h
+++ b/Engine/Plugins/Experimental/MeshModelingToolset/Source/MeshModelingTools/Public/SelfUnionMeshesTool.h
@@ -30,6 +30,17 @@ public:
 	virtual UInteractiveTool* BuildTool(const FToolBuilderState& SceneState) const override;
 };
 
+/** Choice of behavior on accepting CSG */
+UENUM()
+enum class ESelfUnionAcceptBehavior : uint8
+{
+	DeleteOriginals = 0			UMETA(DisplayName = "Delete Inputs"),
+
+	HideOriginals = 1			UMETA(DisplayName = "Hide Inputs"),
+
+	LeaveOriginalsUnchanged = 2	UMETA(DisplayName = "Keep Inputs")
+};
+
 
 
 /**
@@ -40,11 +51,6 @@ class MESHMODELINGTOOLS_API USelfUnionMeshesToolProperties : public UInteractive
 {
 	GENERATED_BODY()
 public:
-
-	/** Remove the source Actors/Components when accepting results of tool.*/
-	UPROPERTY(EditAnywhere, Category = Options)
-	bool bDeleteInputActors = true;
-
 	/** Automatically attempt to fill any holes left by merging (e.g. due to numerical errors) */
 	UPROPERTY(EditAnywhere, Category = Options)
 	bool bAttemptFixHoles = false;
@@ -52,6 +58,10 @@ public:
 	/** Show boundary edges created by the union operation -- often due to numerical error */
 	UPROPERTY(EditAnywhere, Category = Options)
 	bool bShowNewBoundaryEdges = true;
+
+	/** What to do with the source Actors/Components when accepting results of tool.*/
+	UPROPERTY(EditAnywhere, Category = ToolOutputOptions)
+	ESelfUnionAcceptBehavior OnToolAccept;
 };
 
 
diff --git a/Engine/Plugins/Experimental/MeshModelingToolset/Source/MeshModelingToolsEditorOnly/Private/MergeMeshesTool.cpp b/Engine/Plugins/Experimental/MeshModelingToolset/Source/MeshModelingToolsEditorOnly/Private/MergeMeshesTool.cpp
--- a/Engine/Plugins/Experimental/MeshModelingToolset/Source/MeshModelingToolsEditorOnly/Private/MergeMeshesTool.cpp
+++ b/Engine/Plugins/Experimental/MeshModelingToolset/Source/MeshModelingToolsEditorOnly/Private/MergeMeshesTool.cpp
@@ -132,8 +132,14 @@ void UMergeMeshesTool::Shutdown(EToolShutdownType ShutdownType)
 
 		// Hide or destroy the sources
 		{
-			if (MergeProps->bDeleteInputActors) 
+			if (MergeProps->bDeleteInputActors)
+			{
 				GetToolManager()->BeginUndoTransaction(LOCTEXT("RemoveSources", "Remove Sources"));
+			}
+			else
+			{
+				GetToolManager()->BeginUndoTransaction(LOCTEXT("HideSource", "Hide Sources"));
+			}
 			
 			for (auto& ComponentTarget : ComponentTargets)
 			{
@@ -146,12 +152,15 @@ void UMergeMeshesTool::Shutdown(EToolShutdownType ShutdownType)
 				else
 				{
 					// just hide the result.
+					// Save the actor to the transaction buffer to support undo/redo, but do
+					// not call Modify, as we do not want to dirty the actor's package and
+					// we're only editing temporary, transient values
+					SaveToTransactionBuffer(Actor, false);
 					Actor->SetIsTemporarilyHiddenInEditor(true);
 				}
 			}
 
-			if (MergeProps->bDeleteInputActors) 
-				GetToolManager()->EndUndoTransaction();
+			GetToolManager()->EndUndoTransaction();
 		}
 		
 
diff --git a/Engine/Plugins/Experimental/MeshModelingToolset/Source/MeshModelingToolsEditorOnly/Private/VoxelCSGMeshesTool.cpp b/Engine/Plugins/Experimental/MeshModelingToolset/Source/MeshModelingToolsEditorOnly/Private/VoxelCSGMeshesTool.cpp
--- a/Engine/Plugins/Experimental/MeshModelingToolset/Source/MeshModelingToolsEditorOnly/Private/VoxelCSGMeshesTool.cpp
+++ b/Engine/Plugins/Experimental/MeshModelingToolset/Source/MeshModelingToolsEditorOnly/Private/VoxelCSGMeshesTool.cpp
@@ -130,8 +130,14 @@ void UVoxelCSGMeshesTool::Shutdown(EToolShutdownType ShutdownType)
 		// Hide or destroy the sources
 		{
 			
-			if (CSGProps->bDeleteInputActors) 
+			if (CSGProps->bDeleteInputActors)
+			{
 				GetToolManager()->BeginUndoTransaction(LOCTEXT("RemoveSources", "Remove Sources"));
+			}
+			else
+			{
+				GetToolManager()->BeginUndoTransaction(LOCTEXT("HideSource", "Hide Sources"));
+			}
 			
 			for (auto& ComponentTarget : ComponentTargets)
 			{
@@ -143,12 +149,15 @@ void UVoxelCSGMeshesTool::Shutdown(EToolShutdownType ShutdownType)
 				}
 				else
 				{
+					// Save the actor to the transaction buffer to support undo/redo, but do
+					// not call Modify, as we do not want to dirty the actor's package and
+					// we're only editing temporary, transient values
+					SaveToTransactionBuffer(Actor, false);
 					Actor->SetIsTemporarilyHiddenInEditor(true);
 				}
 			}
 
-			if (CSGProps->bDeleteInputActors) 
-				GetToolManager()->EndUndoTransaction();
+			GetToolManager()->EndUndoTransaction();
 			
 		}
 		
