MeshModelingToolset: Fix checkerboard material for tools which affect existing meshes

#rnx

diff --git a/Engine/Plugins/Experimental/MeshModelingToolset/Source/MeshModelingTools/Private/Properties/MeshMaterialProperties.cpp b/Engine/Plugins/Experimental/MeshModelingToolset/Source/MeshModelingTools/Private/Properties/MeshMaterialProperties.cpp
--- a/Engine/Plugins/Experimental/MeshModelingToolset/Source/MeshModelingTools/Private/Properties/MeshMaterialProperties.cpp
+++ b/Engine/Plugins/Experimental/MeshModelingToolset/Source/MeshModelingTools/Private/Properties/MeshMaterialProperties.cpp
@@ -17,6 +17,12 @@ UNewMeshMaterialProperties::UNewMeshMaterialProperties()
 	Material = CreateDefaultSubobject<UMaterialInterface>(TEXT("MATERIAL"));
 }
 
+void UExistingMeshMaterialProperties::RestoreProperties(UInteractiveTool* RestoreToTool)
+{
+	Super::RestoreProperties(RestoreToTool);
+	Setup();
+}
+
 void UExistingMeshMaterialProperties::Setup()
 {
 	UMaterial* CheckerMaterialBase = LoadObject<UMaterial>(nullptr, TEXT("/MeshModelingToolset/Materials/CheckerMaterial"));
diff --git a/Engine/Plugins/Experimental/MeshModelingToolset/Source/MeshModelingTools/Public/Properties/MeshMaterialProperties.h b/Engine/Plugins/Experimental/MeshModelingToolset/Source/MeshModelingTools/Public/Properties/MeshMaterialProperties.h
--- a/Engine/Plugins/Experimental/MeshModelingToolset/Source/MeshModelingTools/Public/Properties/MeshMaterialProperties.h
+++ b/Engine/Plugins/Experimental/MeshModelingToolset/Source/MeshModelingTools/Public/Properties/MeshMaterialProperties.h
@@ -80,6 +80,9 @@ public:
 	UPROPERTY(meta = (TransientToolProperty))
 	UMaterialInstanceDynamic* CheckerMaterial = nullptr;
 
+	// Needs custom restore in order to call setup
+	virtual void RestoreProperties(UInteractiveTool* RestoreToTool) override;
+
 	void Setup();
 
 	void UpdateMaterials();
