Add preliminary visual feedback for failure cases w/ poly cut tool

#rnx

diff --git a/Engine/Plugins/Experimental/MeshModelingToolset/Source/MeshModelingToolsEditorOnly/Private/PolygonOnMeshTool.cpp b/Engine/Plugins/Experimental/MeshModelingToolset/Source/MeshModelingToolsEditorOnly/Private/PolygonOnMeshTool.cpp
--- a/Engine/Plugins/Experimental/MeshModelingToolset/Source/MeshModelingToolsEditorOnly/Private/PolygonOnMeshTool.cpp
+++ b/Engine/Plugins/Experimental/MeshModelingToolset/Source/MeshModelingToolsEditorOnly/Private/PolygonOnMeshTool.cpp
@@ -10,6 +10,7 @@
 #include "BaseBehaviors/MouseHoverBehavior.h"
 #include "ToolDataVisualizer.h"
 #include "Util/ColorConstants.h"
+#include "Drawing/LineSetComponent.h"
 
 #include "MeshDescriptionToDynamicMesh.h"
 #include "DynamicMeshToMeshDescription.h"
@@ -96,7 +97,28 @@ void UPolygonOnMeshTool::Setup()
 	AddToolPropertySource(ActionProperties);
 
 	// initialize the PreviewMesh+BackgroundCompute object
-	UpdateNumPreviews();
+	SetupPreview();
+
+	DrawnLineSet = NewObject<ULineSetComponent>(Preview->PreviewMesh->GetRootComponent());
+	DrawnLineSet->SetupAttachment(Preview->PreviewMesh->GetRootComponent());
+	DrawnLineSet->SetLineMaterial(ToolSetupUtil::GetDefaultLineComponentMaterial(GetToolManager()));
+	DrawnLineSet->RegisterComponent();
+
+	Preview->OnOpCompleted.AddLambda(
+		[this](const FDynamicMeshOperator* Op)
+		{
+			const FEmbedPolygonsOp* PolygonsOp = (const FEmbedPolygonsOp*)(Op);
+			EmbeddedEdges = PolygonsOp->EmbeddedEdges;
+			bEmbedSucceeded = PolygonsOp->bEmbedSucceeded;
+		}
+	);
+	Preview->OnMeshUpdated.AddLambda(
+		[this](const UMeshOpPreviewWithBackgroundCompute*)
+		{
+			GetToolManager()->PostInvalidation();
+			UpdateVisualization();
+		}
+	);
 
 	DrawPlaneWorld = FFrame3d(WorldTransform.GetTranslation());
 
@@ -121,6 +143,25 @@ void UPolygonOnMeshTool::Setup()
 }
 
 
+void UPolygonOnMeshTool::UpdateVisualization()
+{
+	FColor PartialPathEdgeColor(240, 15, 15);
+	float PartialPathEdgeThickness = 2.0;
+	float PartialPathEdgeDepthBias = 2.0f;
+
+	const FDynamicMesh3* TargetMesh = Preview->PreviewMesh->GetPreviewDynamicMesh();
+	FVector3d A, B;
+
+	DrawnLineSet->Clear();
+	if (!bEmbedSucceeded)
+	{
+		for (int EID : EmbeddedEdges)
+		{
+			TargetMesh->GetEdgeV(EID, A, B);
+			DrawnLineSet->AddLine((FVector)A, (FVector)B, PartialPathEdgeColor, PartialPathEdgeThickness, PartialPathEdgeDepthBias);
+		}
+	}
+}
 
 
 void UPolygonOnMeshTool::UpdatePolygonType()
@@ -162,7 +203,7 @@ void UPolygonOnMeshTool::UpdateDrawPlane()
 }
 
 
-void UPolygonOnMeshTool::UpdateNumPreviews()
+void UPolygonOnMeshTool::SetupPreview()
 {
 	Preview = NewObject<UMeshOpPreviewWithBackgroundCompute>(this);
 	Preview->Setup(this->TargetWorld, this);
diff --git a/Engine/Plugins/Experimental/MeshModelingToolset/Source/MeshModelingToolsEditorOnly/Public/PolygonOnMeshTool.h b/Engine/Plugins/Experimental/MeshModelingToolset/Source/MeshModelingToolsEditorOnly/Public/PolygonOnMeshTool.h
--- a/Engine/Plugins/Experimental/MeshModelingToolset/Source/MeshModelingToolsEditorOnly/Public/PolygonOnMeshTool.h
+++ b/Engine/Plugins/Experimental/MeshModelingToolset/Source/MeshModelingToolsEditorOnly/Public/PolygonOnMeshTool.h
@@ -22,6 +22,7 @@ struct FMeshDescription;
 class USimpleDynamicMeshComponent;
 class UTransformGizmo;
 class UTransformProxy;
+class ULineSetComponent;
 
 
 
@@ -193,6 +194,11 @@ protected:
 	UPROPERTY()
 	UMeshOpPreviewWithBackgroundCompute* Preview;
 
+	UPROPERTY()
+	ULineSetComponent* DrawnLineSet;
+
+	TArray<int> EmbeddedEdges;
+	bool bEmbedSucceeded;
 
 protected:
 	UWorld* TargetWorld;
@@ -216,11 +222,13 @@ protected:
 	FPolygon2d ActivePolygon;
 	void UpdatePolygonType();
 
-	void UpdateNumPreviews();
+	void SetupPreview();
 	void UpdateDrawPlane();
 
 	void BeginDrawPolygon();
 	void CompleteDrawPolygon();
 
+	void UpdateVisualization();
+
 	void GenerateAsset(const TArray<FDynamicMeshOpResult>& Results);
 };
diff --git a/Engine/Plugins/Experimental/MeshModelingToolset/Source/ModelingOperatorsEditorOnly/Private/CuttingOps/EmbedPolygonsOp.cpp b/Engine/Plugins/Experimental/MeshModelingToolset/Source/ModelingOperatorsEditorOnly/Private/CuttingOps/EmbedPolygonsOp.cpp
--- a/Engine/Plugins/Experimental/MeshModelingToolset/Source/ModelingOperatorsEditorOnly/Private/CuttingOps/EmbedPolygonsOp.cpp
+++ b/Engine/Plugins/Experimental/MeshModelingToolset/Source/ModelingOperatorsEditorOnly/Private/CuttingOps/EmbedPolygonsOp.cpp
@@ -210,6 +210,7 @@ void FEmbedPolygonsOp::CalculateResult(FProgressCancel* Progress)
 	};
 
 	bool bCutSide1 = CutHole(*ResultMesh, Frame, SortedHitTriangles[0].Value, Polygon, PathVertIDs1, PathVertCorrespond1, bCollapseDegenerateEdges);
+	RecordEmbeddedEdges(PathVertIDs1);
 	if (!bCutSide1 || PathVertIDs1.Num() < 2)
 	{
 		return;
@@ -219,6 +220,7 @@ void FEmbedPolygonsOp::CalculateResult(FProgressCancel* Progress)
 	{
 		TArray<int> PathVertIDs2, PathVertCorrespond2;
 		bool bCutSide2 = CutHole(*ResultMesh, Frame, SortedHitTriangles[SecondHit].Value, Polygon, PathVertIDs2, PathVertCorrespond2, bCollapseDegenerateEdges);
+		RecordEmbeddedEdges(PathVertIDs1);
 		if (!bCutSide2 || PathVertIDs2.Num() < 2)
 		{
 			return;
@@ -270,4 +272,18 @@ void FEmbedPolygonsOp::CalculateResult(FProgressCancel* Progress)
 			}
 		}
 	}
+
+	bEmbedSucceeded = true;
+}
+
+void FEmbedPolygonsOp::RecordEmbeddedEdges(TArray<int>& PathVertIDs)
+{
+	for (int Idx = 0; Idx + 1 < PathVertIDs.Num(); Idx++)
+	{
+		int EID = ResultMesh->FindEdge(PathVertIDs[Idx], PathVertIDs[Idx + 1]);
+		if (ensure(ResultMesh->IsEdge(EID)))
+		{
+			EmbeddedEdges.Add(EID);
+		}
+	}
 }
\ No newline at end of file
diff --git a/Engine/Plugins/Experimental/MeshModelingToolset/Source/ModelingOperatorsEditorOnly/Public/CuttingOps/EmbedPolygonsOp.h b/Engine/Plugins/Experimental/MeshModelingToolset/Source/ModelingOperatorsEditorOnly/Public/CuttingOps/EmbedPolygonsOp.h
--- a/Engine/Plugins/Experimental/MeshModelingToolset/Source/ModelingOperatorsEditorOnly/Public/CuttingOps/EmbedPolygonsOp.h
+++ b/Engine/Plugins/Experimental/MeshModelingToolset/Source/ModelingOperatorsEditorOnly/Public/CuttingOps/EmbedPolygonsOp.h
@@ -31,7 +31,7 @@ public:
 	FFrame3d PolygonFrame;
 	FPolygon2d EmbedPolygon;
 
-	// TODO: stop hardcoding the polygon shape, switch to FGeneralPolygon2d
+	// TODO: switch to FGeneralPolygon2d?
 	FPolygon2d GetPolygon()
 	{
 		return EmbedPolygon;
@@ -50,6 +50,14 @@ public:
 	// 
 
 	virtual void CalculateResult(FProgressCancel* Progress) override;
+
+
+	// Outputs
+	TArray<int> EmbeddedEdges;
+	bool bEmbedSucceeded = false;
+
+private:
+	void RecordEmbeddedEdges(TArray<int>& PathVertIDs);
 };
 
 
