Pull delete/hide/keep original actor logic out of the separate mesh/vox union/boolean tools to a shared Properties file, & clean up surrounding logic

#rnx

diff --git a/Engine/Plugins/Experimental/MeshModelingToolset/Source/MeshModelingTools/Private/CSGMeshesTool.cpp b/Engine/Plugins/Experimental/MeshModelingToolset/Source/MeshModelingTools/Private/CSGMeshesTool.cpp
--- a/Engine/Plugins/Experimental/MeshModelingToolset/Source/MeshModelingTools/Private/CSGMeshesTool.cpp
+++ b/Engine/Plugins/Experimental/MeshModelingToolset/Source/MeshModelingTools/Private/CSGMeshesTool.cpp
@@ -89,7 +89,9 @@ void UCSGMeshesTool::Setup()
 	CSGProperties = NewObject<UCSGMeshesToolProperties>(this);
 	CSGProperties->RestoreProperties(this);
 	AddToolPropertySource(CSGProperties);
-
+	HandleSourcesProperties = NewObject<UOnAcceptHandleSourcesProperties>(this);
+	HandleSourcesProperties->RestoreProperties(this);
+	AddToolPropertySource(HandleSourcesProperties);
 
 	// initialize the PreviewMesh+BackgroundCompute object
 	SetupPreview();
@@ -224,8 +226,14 @@ void UCSGMeshesTool::TransformChanged(UTransformProxy* Proxy, FTransform Transfo
 void UCSGMeshesTool::Shutdown(EToolShutdownType ShutdownType)
 {
 	CSGProperties->SaveProperties(this);
+	HandleSourcesProperties->SaveProperties(this);
 
 	FDynamicMeshOpResult Result = Preview->Shutdown();
+	// Restore (unhide) the source meshes
+	for (auto& ComponentTarget : ComponentTargets)
+	{
+		ComponentTarget->SetOwnerVisibility(true);
+	}
 	if (ShutdownType == EToolShutdownType::Accept)
 	{
 		// Generate the result
@@ -237,65 +245,12 @@ void UCSGMeshesTool::Shutdown(EToolShutdownType ShutdownType)
 			GetToolManager()->EndUndoTransaction();
 		}
 
-		// Restore (unhide) the source meshes
-		for (auto& ComponentTarget : ComponentTargets)
-		{
-			ComponentTarget->SetOwnerVisibility(true);
-		}
-
-		// Hide or destroy the sources
-		if (CSGProperties->OnToolAccept != ECSGAcceptBehavior::LeaveOriginalsUnchanged)
-		{
-			bool bDelete = CSGProperties->OnToolAccept == ECSGAcceptBehavior::DeleteOriginals;
-			if (bDelete)
-			{
-				GetToolManager()->BeginUndoTransaction(LOCTEXT("RemoveSources", "Remove Sources"));
-			}
-			else
-			{
-#if WITH_EDITOR
-				GetToolManager()->BeginUndoTransaction(LOCTEXT("HideSources", "Hide Sources"));
-#endif
-			}
-
-			for (auto& ComponentTarget : ComponentTargets)
-			{
-				AActor* Actor = ComponentTarget->GetOwnerActor();
-				if (bDelete)
-				{
-					Actor->Destroy();
-				}
-				else
-				{
-#if WITH_EDITOR
-					// Save the actor to the transaction buffer to support undo/redo, but do
-					// not call Modify, as we do not want to dirty the actor's package and
-					// we're only editing temporary, transient values
-					SaveToTransactionBuffer(Actor, false);
-					Actor->SetIsTemporarilyHiddenInEditor(true);
-#endif
-				}
-			}
-			if (bDelete)
-			{
-				GetToolManager()->EndUndoTransaction();
-			}
-			else
-			{
-#if WITH_EDITOR
-				GetToolManager()->EndUndoTransaction();
-#endif
-			}
-		}
-
-	}
-	else
-	{
-		// Restore (unhide) the source meshes
+		TArray<AActor*> Actors;
 		for (auto& ComponentTarget : ComponentTargets)
 		{
-			ComponentTarget->SetOwnerVisibility(true);
+			Actors.Add(ComponentTarget->GetOwnerActor());
 		}
+		HandleSourcesProperties->ApplyMethod(Actors, GetToolManager());
 	}
 
 	UInteractiveGizmoManager* GizmoManager = GetToolManager()->GetPairedGizmoManager();
@@ -359,7 +314,7 @@ void UCSGMeshesTool::OnPropertyModified(UObject* PropertySet, FProperty* Propert
 		UpdateGizmoVisibility();
 	}
 	else if (Property && 
-		(  Property->GetFName() == GET_MEMBER_NAME_CHECKED(UCSGMeshesToolProperties, OnToolAccept)
+		(  PropertySet == HandleSourcesProperties
 		|| Property->GetFName() == GET_MEMBER_NAME_CHECKED(UCSGMeshesToolProperties, bSnapToWorldGrid)
 		))
 	{
diff --git a/Engine/Plugins/Experimental/MeshModelingToolset/Source/MeshModelingTools/Private/Properties/OnAcceptProperties.cpp b/Engine/Plugins/Experimental/MeshModelingToolset/Source/MeshModelingTools/Private/Properties/OnAcceptProperties.cpp
new file mode 100644
--- /dev/null
+++ b/Engine/Plugins/Experimental/MeshModelingToolset/Source/MeshModelingTools/Private/Properties/OnAcceptProperties.cpp
@@ -0,0 +1,57 @@
+// Copyright Epic Games, Inc. All Rights Reserved.
+
+#include "Properties/OnAcceptProperties.h"
+
+#include "InteractiveToolManager.h"
+
+#define LOCTEXT_NAMESPACE "UOnAcceptHandleSourcesProperties"
+
+
+void UOnAcceptHandleSourcesProperties::ApplyMethod(const TArray<AActor*>& Actors, UInteractiveToolManager* ToolManager)
+{
+	// Hide or destroy the sources
+	if (OnToolAccept != EHandleSourcesMethod::KeepSources)
+	{
+		bool bDelete = OnToolAccept == EHandleSourcesMethod::DeleteSources;
+		if (bDelete)
+		{
+			ToolManager->BeginUndoTransaction(LOCTEXT("RemoveSources", "Remove Sources"));
+		}
+		else
+		{
+#if WITH_EDITOR
+			ToolManager->BeginUndoTransaction(LOCTEXT("HideSources", "Hide Sources"));
+#endif
+		}
+
+		for (AActor* Actor : Actors)
+		{
+			if (bDelete)
+			{
+				Actor->Destroy();
+			}
+			else
+			{
+#if WITH_EDITOR
+				// Save the actor to the transaction buffer to support undo/redo, but do
+				// not call Modify, as we do not want to dirty the actor's package and
+				// we're only editing temporary, transient values
+				SaveToTransactionBuffer(Actor, false);
+				Actor->SetIsTemporarilyHiddenInEditor(true);
+#endif
+			}
+		}
+		if (bDelete)
+		{
+			ToolManager->EndUndoTransaction();
+		}
+		else
+		{
+#if WITH_EDITOR
+			ToolManager->EndUndoTransaction();
+#endif
+		}
+	}
+}
+
+#undef LOCTEXT_NAMESPACE
diff --git a/Engine/Plugins/Experimental/MeshModelingToolset/Source/MeshModelingTools/Private/SelfUnionMeshesTool.cpp b/Engine/Plugins/Experimental/MeshModelingToolset/Source/MeshModelingTools/Private/SelfUnionMeshesTool.cpp
--- a/Engine/Plugins/Experimental/MeshModelingToolset/Source/MeshModelingTools/Private/SelfUnionMeshesTool.cpp
+++ b/Engine/Plugins/Experimental/MeshModelingToolset/Source/MeshModelingTools/Private/SelfUnionMeshesTool.cpp
@@ -85,6 +85,9 @@ void USelfUnionMeshesTool::Setup()
 	Properties = NewObject<USelfUnionMeshesToolProperties>(this);
 	Properties->RestoreProperties(this);
 	AddToolPropertySource(Properties);
+	HandleSourcesProperties = NewObject<UOnAcceptHandleSourcesProperties>(this);
+	HandleSourcesProperties->RestoreProperties(this);
+	AddToolPropertySource(HandleSourcesProperties);
 
 
 	// initialize the PreviewMesh+BackgroundCompute object
@@ -208,8 +211,14 @@ void USelfUnionMeshesTool::UpdateVisualization()
 void USelfUnionMeshesTool::Shutdown(EToolShutdownType ShutdownType)
 {
 	Properties->SaveProperties(this);
+	HandleSourcesProperties->SaveProperties(this);
 
 	FDynamicMeshOpResult Result = Preview->Shutdown();
+	// Restore (unhide) the source meshes
+	for (auto& ComponentTarget : ComponentTargets)
+	{
+		ComponentTarget->SetOwnerVisibility(true);
+	}
 	if (ShutdownType == EToolShutdownType::Accept)
 	{
 		// Generate the result
@@ -221,65 +230,12 @@ void USelfUnionMeshesTool::Shutdown(EToolShutdownType ShutdownType)
 			GetToolManager()->EndUndoTransaction();
 		}
 
-		// Restore (unhide) the source meshes
-		for (auto& ComponentTarget : ComponentTargets)
-		{
-			ComponentTarget->SetOwnerVisibility(true);
-		}
-
-		// Hide or destroy the sources
-		if (Properties->OnToolAccept != ESelfUnionAcceptBehavior::LeaveOriginalsUnchanged)
-		{
-			bool bDelete = Properties->OnToolAccept == ESelfUnionAcceptBehavior::DeleteOriginals;
-			if (bDelete)
-			{
-				GetToolManager()->BeginUndoTransaction(LOCTEXT("RemoveSources", "Remove Sources"));
-			}
-			else
-			{
-#if WITH_EDITOR
-				GetToolManager()->BeginUndoTransaction(LOCTEXT("HideSources", "Hide Sources"));
-#endif
-			}
-
-			for (auto& ComponentTarget : ComponentTargets)
-			{
-				AActor* Actor = ComponentTarget->GetOwnerActor();
-				if (bDelete)
-				{
-					Actor->Destroy();
-				}
-				else
-				{
-#if WITH_EDITOR
-					// Save the actor to the transaction buffer to support undo/redo, but do
-					// not call Modify, as we do not want to dirty the actor's package and
-					// we're only editing temporary, transient values
-					SaveToTransactionBuffer(Actor, false);
-					Actor->SetIsTemporarilyHiddenInEditor(true);
-#endif
-				}
-			}
-			if (bDelete)
-			{
-				GetToolManager()->EndUndoTransaction();
-			}
-			else
-			{
-#if WITH_EDITOR
-				GetToolManager()->EndUndoTransaction();
-#endif
-			}
-		}
-
-	}
-	else
-	{
-		// Restore (unhide) the source meshes
+		TArray<AActor*> Actors;
 		for (auto& ComponentTarget : ComponentTargets)
 		{
-			ComponentTarget->SetOwnerVisibility(true);
+			Actors.Add(ComponentTarget->GetOwnerActor());
 		}
+		HandleSourcesProperties->ApplyMethod(Actors, GetToolManager());
 	}
 }
 
@@ -321,7 +277,7 @@ void USelfUnionMeshesTool::PostEditChangeProperty(FPropertyChangedEvent& Propert
 
 void USelfUnionMeshesTool::OnPropertyModified(UObject* PropertySet, FProperty* Property)
 {
-	if (Property && (Property->GetFName() == GET_MEMBER_NAME_CHECKED(USelfUnionMeshesToolProperties, OnToolAccept)))
+	if (PropertySet == HandleSourcesProperties)
 	{
 		// nothing
 	}
diff --git a/Engine/Plugins/Experimental/MeshModelingToolset/Source/MeshModelingTools/Public/CSGMeshesTool.h b/Engine/Plugins/Experimental/MeshModelingToolset/Source/MeshModelingTools/Public/CSGMeshesTool.h
--- a/Engine/Plugins/Experimental/MeshModelingToolset/Source/MeshModelingTools/Public/CSGMeshesTool.h
+++ b/Engine/Plugins/Experimental/MeshModelingToolset/Source/MeshModelingTools/Public/CSGMeshesTool.h
@@ -9,6 +9,7 @@
 #include "Drawing/LineSetComponent.h"
 #include "MeshOpPreviewHelpers.h"
 #include "BaseTools/SingleClickTool.h"
+#include "Properties/OnAcceptProperties.h"
 
 #include "CompositionOps/BooleanMeshesOp.h"
 
@@ -33,18 +34,6 @@ public:
 };
 
 
-/** Choice of behavior on accepting CSG */
-UENUM()
-enum class ECSGAcceptBehavior : uint8
-{
-	DeleteOriginals = 0			UMETA(DisplayName = "Delete Inputs"),
-
-	HideOriginals = 1			UMETA(DisplayName = "Hide Inputs"),
-
-	LeaveOriginalsUnchanged = 2	UMETA(DisplayName = "Keep Inputs")
-};
-
-
 
 /**
  * Standard properties of the CSG operation
@@ -73,10 +62,6 @@ public:
 	/** Automatically attempt to fill any holes left by CSG (e.g. due to numerical errors) */
 	UPROPERTY(EditAnywhere, Category = Options)
 	bool bAttemptFixHoles = false;
-
-	/** What to do with the source Actors/Components when accepting results of tool.*/
-	UPROPERTY(EditAnywhere, Category = ToolOutputOptions)
-	ECSGAcceptBehavior OnToolAccept;
 };
 
 
@@ -123,6 +108,9 @@ protected:
 	UPROPERTY()
 	UCSGMeshesToolProperties* CSGProperties;
 
+	UPROPERTY()
+	UOnAcceptHandleSourcesProperties* HandleSourcesProperties;
+
 	UPROPERTY()
 	TArray<UTransformProxy*> TransformProxies;
 
diff --git a/Engine/Plugins/Experimental/MeshModelingToolset/Source/MeshModelingTools/Public/Properties/OnAcceptProperties.h b/Engine/Plugins/Experimental/MeshModelingToolset/Source/MeshModelingTools/Public/Properties/OnAcceptProperties.h
new file mode 100644
--- /dev/null
+++ b/Engine/Plugins/Experimental/MeshModelingToolset/Source/MeshModelingTools/Public/Properties/OnAcceptProperties.h
@@ -0,0 +1,40 @@
+// Copyright Epic Games, Inc. All Rights Reserved.
+
+#pragma once
+
+#include "CoreMinimal.h"
+#include "InteractiveToolBuilder.h"
+
+#include "OnAcceptProperties.generated.h"
+
+
+class UInteractiveToolManager;
+
+
+/** Options to handle source meshes */
+UENUM()
+enum class EHandleSourcesMethod : uint8
+{
+	DeleteSources = 0	UMETA(DisplayName = "Delete Sources"),
+
+	HideSources = 1		UMETA(DisplayName = "Hide Sources"),
+
+	KeepSources = 2		UMETA(DisplayName = "Keep Sources")
+};
+
+
+
+
+// Standard property settings for tools that create a new actor and need to decide what to do with the input (source) actor(s)
+UCLASS()
+class MESHMODELINGTOOLS_API UOnAcceptHandleSourcesProperties : public UInteractiveToolPropertySet
+{
+	GENERATED_BODY()
+public:
+
+	/** What to do with the source Actors/Components when accepting results of tool.*/
+	UPROPERTY(EditAnywhere, Category = ToolOutputOptions)
+	EHandleSourcesMethod OnToolAccept;
+
+	void ApplyMethod(const TArray<AActor*>& Actors, UInteractiveToolManager* ToolManager);
+};
diff --git a/Engine/Plugins/Experimental/MeshModelingToolset/Source/MeshModelingTools/Public/SelfUnionMeshesTool.h b/Engine/Plugins/Experimental/MeshModelingToolset/Source/MeshModelingTools/Public/SelfUnionMeshesTool.h
--- a/Engine/Plugins/Experimental/MeshModelingToolset/Source/MeshModelingTools/Public/SelfUnionMeshesTool.h
+++ b/Engine/Plugins/Experimental/MeshModelingToolset/Source/MeshModelingTools/Public/SelfUnionMeshesTool.h
@@ -9,6 +9,7 @@
 #include "Drawing/LineSetComponent.h"
 #include "MeshOpPreviewHelpers.h"
 #include "BaseTools/SingleClickTool.h"
+#include "Properties/OnAcceptProperties.h"
 
 #include "CompositionOps/SelfUnionMeshesOp.h"
 
@@ -30,17 +31,6 @@ public:
 	virtual UInteractiveTool* BuildTool(const FToolBuilderState& SceneState) const override;
 };
 
-/** Choice of behavior on accepting CSG */
-UENUM()
-enum class ESelfUnionAcceptBehavior : uint8
-{
-	DeleteOriginals = 0			UMETA(DisplayName = "Delete Inputs"),
-
-	HideOriginals = 1			UMETA(DisplayName = "Hide Inputs"),
-
-	LeaveOriginalsUnchanged = 2	UMETA(DisplayName = "Keep Inputs")
-};
-
 
 
 /**
@@ -58,10 +48,6 @@ public:
 	/** Show boundary edges created by the union operation -- often due to numerical error */
 	UPROPERTY(EditAnywhere, Category = Options)
 	bool bShowNewBoundaryEdges = true;
-
-	/** What to do with the source Actors/Components when accepting results of tool.*/
-	UPROPERTY(EditAnywhere, Category = ToolOutputOptions)
-	ESelfUnionAcceptBehavior OnToolAccept;
 };
 
 
@@ -108,6 +94,9 @@ protected:
 	UPROPERTY()
 	USelfUnionMeshesToolProperties* Properties;
 
+	UPROPERTY()
+	UOnAcceptHandleSourcesProperties* HandleSourcesProperties;
+
 	UPROPERTY()
 	ULineSetComponent* DrawnLineSet;
 
diff --git a/Engine/Plugins/Experimental/MeshModelingToolset/Source/MeshModelingToolsEditorOnly/Private/MergeMeshesTool.cpp b/Engine/Plugins/Experimental/MeshModelingToolset/Source/MeshModelingToolsEditorOnly/Private/MergeMeshesTool.cpp
--- a/Engine/Plugins/Experimental/MeshModelingToolset/Source/MeshModelingToolsEditorOnly/Private/MergeMeshesTool.cpp
+++ b/Engine/Plugins/Experimental/MeshModelingToolset/Source/MeshModelingToolsEditorOnly/Private/MergeMeshesTool.cpp
@@ -79,11 +79,16 @@ void UMergeMeshesTool::Setup()
 	UInteractiveTool::Setup();
 
 	MergeProps = NewObject<UMergeMeshesToolProperties>();
+	MergeProps->RestoreProperties(this);
 	AddToolPropertySource(MergeProps);
 
 	MeshStatisticsProperties = NewObject<UMeshStatisticsProperties>(this);
 	AddToolPropertySource(MeshStatisticsProperties);
 
+	HandleSourcesProperties = NewObject<UOnAcceptHandleSourcesProperties>(this);
+	HandleSourcesProperties->RestoreProperties(this);
+	AddToolPropertySource(HandleSourcesProperties);
+
 	// Hide the source meshes
 	for (auto& ComponentTarget : ComponentTargets)
 	{
@@ -118,7 +123,15 @@ void UMergeMeshesTool::Setup()
 
 void UMergeMeshesTool::Shutdown(EToolShutdownType ShutdownType)
 {
+	MergeProps->SaveProperties(this);
+	HandleSourcesProperties->SaveProperties(this);
+
 	FDynamicMeshOpResult Result = Preview->Shutdown();
+	// Restore (unhide) the source meshes
+	for (auto& ComponentTarget : ComponentTargets)
+	{
+		ComponentTarget->SetOwnerVisibility(true);
+	}
 	if (ShutdownType == EToolShutdownType::Accept)
 	{
 		// Generate the result
@@ -130,48 +143,12 @@ void UMergeMeshesTool::Shutdown(EToolShutdownType ShutdownType)
 			GetToolManager()->EndUndoTransaction();
 		}
 
-		// Hide or destroy the sources
-		{
-			if (MergeProps->bDeleteInputActors)
-			{
-				GetToolManager()->BeginUndoTransaction(LOCTEXT("RemoveSources", "Remove Sources"));
-			}
-			else
-			{
-				GetToolManager()->BeginUndoTransaction(LOCTEXT("HideSource", "Hide Sources"));
-			}
-			
-			for (auto& ComponentTarget : ComponentTargets)
-			{
-				ComponentTarget->SetOwnerVisibility(true);
-				AActor* Actor = ComponentTarget->GetOwnerActor();
-				if (MergeProps->bDeleteInputActors)
-				{
-					Actor->Destroy();
-				}
-				else
-				{
-					// just hide the result.
-					// Save the actor to the transaction buffer to support undo/redo, but do
-					// not call Modify, as we do not want to dirty the actor's package and
-					// we're only editing temporary, transient values
-					SaveToTransactionBuffer(Actor, false);
-					Actor->SetIsTemporarilyHiddenInEditor(true);
-				}
-			}
-
-			GetToolManager()->EndUndoTransaction();
-		}
-		
-
-	}
-	else
-	{
-		// Restore (unhide) the source meshes
+		TArray<AActor*> Actors;
 		for (auto& ComponentTarget : ComponentTargets)
 		{
-			ComponentTarget->SetOwnerVisibility(true);
+			Actors.Add(ComponentTarget->GetOwnerActor());
 		}
+		HandleSourcesProperties->ApplyMethod(Actors, GetToolManager());
 	}
 }
 
diff --git a/Engine/Plugins/Experimental/MeshModelingToolset/Source/MeshModelingToolsEditorOnly/Private/VoxelCSGMeshesTool.cpp b/Engine/Plugins/Experimental/MeshModelingToolset/Source/MeshModelingToolsEditorOnly/Private/VoxelCSGMeshesTool.cpp
--- a/Engine/Plugins/Experimental/MeshModelingToolset/Source/MeshModelingToolsEditorOnly/Private/VoxelCSGMeshesTool.cpp
+++ b/Engine/Plugins/Experimental/MeshModelingToolset/Source/MeshModelingToolsEditorOnly/Private/VoxelCSGMeshesTool.cpp
@@ -74,14 +74,17 @@ void UVoxelCSGMeshesTool::Setup()
 	UInteractiveTool::Setup();
 
 	CSGProps = NewObject<UVoxelCSGMeshesToolProperties>();
-	CSGProps->VoxelCount = 128;
-	CSGProps->MeshAdaptivity = 0.01f;
-	CSGProps->OffsetDistance = 0.0f;
+	CSGProps->RestoreProperties(this);
 	AddToolPropertySource(CSGProps);
 
 	MeshStatisticsProperties = NewObject<UMeshStatisticsProperties>(this);
 	AddToolPropertySource(MeshStatisticsProperties);
 
+	HandleSourcesProperties = NewObject<UOnAcceptHandleSourcesProperties>(this);
+	HandleSourcesProperties->RestoreProperties(this);
+	AddToolPropertySource(HandleSourcesProperties);
+
+
 	// Hide the source meshes
 	for (auto& ComponentTarget : ComponentTargets)
 	{
@@ -115,7 +118,15 @@ void UVoxelCSGMeshesTool::Setup()
 
 void UVoxelCSGMeshesTool::Shutdown(EToolShutdownType ShutdownType)
 {
+	CSGProps->SaveProperties(this);
+	HandleSourcesProperties->SaveProperties(this);
+
 	FDynamicMeshOpResult Result = Preview->Shutdown();
+	// Restore (unhide) the source meshes
+	for (auto& ComponentTarget : ComponentTargets)
+	{
+		ComponentTarget->SetOwnerVisibility(true);
+	}
 	if (ShutdownType == EToolShutdownType::Accept)
 	{
 		// Generate the result
@@ -127,48 +138,12 @@ void UVoxelCSGMeshesTool::Shutdown(EToolShutdownType ShutdownType)
 			GetToolManager()->EndUndoTransaction();
 		}
 
-		// Hide or destroy the sources
-		{
-			
-			if (CSGProps->bDeleteInputActors)
-			{
-				GetToolManager()->BeginUndoTransaction(LOCTEXT("RemoveSources", "Remove Sources"));
-			}
-			else
-			{
-				GetToolManager()->BeginUndoTransaction(LOCTEXT("HideSource", "Hide Sources"));
-			}
-			
-			for (auto& ComponentTarget : ComponentTargets)
-			{
-				ComponentTarget->SetOwnerVisibility(true);
-				AActor* Actor = ComponentTarget->GetOwnerActor();
-				if (CSGProps->bDeleteInputActors)
-				{
-					Actor->Destroy();
-				}
-				else
-				{
-					// Save the actor to the transaction buffer to support undo/redo, but do
-					// not call Modify, as we do not want to dirty the actor's package and
-					// we're only editing temporary, transient values
-					SaveToTransactionBuffer(Actor, false);
-					Actor->SetIsTemporarilyHiddenInEditor(true);
-				}
-			}
-
-			GetToolManager()->EndUndoTransaction();
-			
-		}
-		
-	}
-	else
-	{
-		// Restore (unhide) the source meshes
+		TArray<AActor*> Actors;
 		for (auto& ComponentTarget : ComponentTargets)
 		{
-			ComponentTarget->SetOwnerVisibility(true);
+			Actors.Add(ComponentTarget->GetOwnerActor());
 		}
+		HandleSourcesProperties->ApplyMethod(Actors, GetToolManager());
 	}
 }
 
diff --git a/Engine/Plugins/Experimental/MeshModelingToolset/Source/MeshModelingToolsEditorOnly/Public/MergeMeshesTool.h b/Engine/Plugins/Experimental/MeshModelingToolset/Source/MeshModelingToolsEditorOnly/Public/MergeMeshesTool.h
--- a/Engine/Plugins/Experimental/MeshModelingToolset/Source/MeshModelingToolsEditorOnly/Public/MergeMeshesTool.h
+++ b/Engine/Plugins/Experimental/MeshModelingToolset/Source/MeshModelingToolsEditorOnly/Public/MergeMeshesTool.h
@@ -8,6 +8,7 @@
 #include "MeshOpPreviewHelpers.h"
 #include "ProxyLODVolume.h"
 #include "Properties/MeshStatisticsProperties.h"
+#include "Properties/OnAcceptProperties.h"
 #include "MergeMeshesTool.generated.h"
 
 
@@ -57,10 +58,6 @@ public:
 	/** Automatically simplify the result of voxel-based merge.*/
 	UPROPERTY(EditAnywhere, Category = Options)
 	bool bAutoSimplify = false;
-
-	/** Delete the source Actors/Components when accepting results of tool.*/
-	UPROPERTY(EditAnywhere, Category = Options)
-	bool bDeleteInputActors = true;
 };
 
 /**
@@ -98,6 +95,9 @@ protected:
 	UPROPERTY()
 	UMeshStatisticsProperties* MeshStatisticsProperties;
 
+	UPROPERTY()
+	UOnAcceptHandleSourcesProperties* HandleSourcesProperties;
+
 
 	UPROPERTY()
 	UMeshOpPreviewWithBackgroundCompute* Preview;
diff --git a/Engine/Plugins/Experimental/MeshModelingToolset/Source/MeshModelingToolsEditorOnly/Public/VoxelCSGMeshesTool.h b/Engine/Plugins/Experimental/MeshModelingToolset/Source/MeshModelingToolsEditorOnly/Public/VoxelCSGMeshesTool.h
--- a/Engine/Plugins/Experimental/MeshModelingToolset/Source/MeshModelingToolsEditorOnly/Public/VoxelCSGMeshesTool.h
+++ b/Engine/Plugins/Experimental/MeshModelingToolset/Source/MeshModelingToolsEditorOnly/Public/VoxelCSGMeshesTool.h
@@ -9,6 +9,8 @@
 #include "MultiSelectionTool.h"
 #include "ProxyLODVolume.h"
 #include "Properties/MeshStatisticsProperties.h"
+#include "Properties/OnAcceptProperties.h"
+
 #include "VoxelCSGMeshesTool.generated.h"
 
 
@@ -70,19 +72,15 @@ public:
 
 	/** Remeshing adaptivity, prior to optional simplification */
 	UPROPERTY(EditAnywhere, Category = Options, meta = (UIMin = "0", UIMax = "1", ClampMin = "0", ClampMax = "1"))
-	float MeshAdaptivity = 0.001f;
+	float MeshAdaptivity = 0.01f;
 
 	/** Offset when remeshing, note large offsets with high voxels counts will be slow */
 	UPROPERTY(EditAnywhere, Category = Options, meta = (UIMin = "-10", UIMax = "10", ClampMin = "-10", ClampMax = "10"))
-	float OffsetDistance = 0;
+	float OffsetDistance = 0.0f;
 
 	/** Automatically simplify the result of voxel-based merge.*/
 	UPROPERTY(EditAnywhere, Category = Options)
 	bool bAutoSimplify = false;
-
-	/** Remove the source Actors/Components when accepting results of tool.*/
-	UPROPERTY(EditAnywhere, Category = Options)
-	bool bDeleteInputActors = true;
 };
 
 
@@ -125,6 +123,9 @@ protected:
 	UPROPERTY()
 	UMeshStatisticsProperties* MeshStatisticsProperties;
 
+	UPROPERTY()
+	UOnAcceptHandleSourcesProperties* HandleSourcesProperties;
+
 	UPROPERTY()
 	UMeshOpPreviewWithBackgroundCompute* Preview;
 
