add includes to onacceptproperties to fix non-unity build

#rnx

diff --git a/Engine/Plugins/Experimental/MeshModelingToolset/Source/MeshModelingTools/Private/Properties/OnAcceptProperties.cpp b/Engine/Plugins/Experimental/MeshModelingToolset/Source/MeshModelingTools/Private/Properties/OnAcceptProperties.cpp
--- a/Engine/Plugins/Experimental/MeshModelingToolset/Source/MeshModelingTools/Private/Properties/OnAcceptProperties.cpp
+++ b/Engine/Plugins/Experimental/MeshModelingToolset/Source/MeshModelingTools/Private/Properties/OnAcceptProperties.cpp
@@ -2,6 +2,8 @@
 
 #include "Properties/OnAcceptProperties.h"
 
+#include "GameFramework/Actor.h"
+#include "UObject/UObjectGlobals.h"
 #include "InteractiveToolManager.h"
 
 #define LOCTEXT_NAMESPACE "UOnAcceptHandleSourcesProperties"
diff --git a/Engine/Plugins/Experimental/MeshModelingToolset/Source/MeshModelingTools/Public/Properties/OnAcceptProperties.h b/Engine/Plugins/Experimental/MeshModelingToolset/Source/MeshModelingTools/Public/Properties/OnAcceptProperties.h
--- a/Engine/Plugins/Experimental/MeshModelingToolset/Source/MeshModelingTools/Public/Properties/OnAcceptProperties.h
+++ b/Engine/Plugins/Experimental/MeshModelingToolset/Source/MeshModelingTools/Public/Properties/OnAcceptProperties.h
@@ -9,6 +9,7 @@
 
 
 class UInteractiveToolManager;
+class AActor;
 
 
 /** Options to handle source meshes */
