add option for mesh csg tools to only use material of first mesh (useful for sculpting case)

#rnx

diff --git a/Engine/Plugins/Experimental/MeshModelingToolset/Source/MeshModelingTools/Private/CSGMeshesTool.cpp b/Engine/Plugins/Experimental/MeshModelingToolset/Source/MeshModelingTools/Private/CSGMeshesTool.cpp
--- a/Engine/Plugins/Experimental/MeshModelingToolset/Source/MeshModelingTools/Private/CSGMeshesTool.cpp
+++ b/Engine/Plugins/Experimental/MeshModelingToolset/Source/MeshModelingTools/Private/CSGMeshesTool.cpp
@@ -101,34 +101,48 @@ void UCSGMeshesTool::Setup()
 
 
 
-
-void UCSGMeshesTool::SetupPreview()
+void UCSGMeshesTool::ConfigurePreviewMaterials()
 {
+	OriginalDynamicMeshes.SetNum(ComponentTargets.Num());
 	FComponentMaterialSet AllMaterialSet;
 	TMap<UMaterialInterface*, int> KnownMaterials;
 	TArray<TArray<int>> MaterialRemap; MaterialRemap.SetNum(ComponentTargets.Num());
-	for (int ComponentIdx = 0; ComponentIdx < ComponentTargets.Num(); ComponentIdx++)
+
+	if (!CSGProperties->bOnlyUseFirstMeshMaterials)
 	{
-		FComponentMaterialSet ComponentMaterialSet;
-		ComponentTargets[ComponentIdx]->GetMaterialSet(ComponentMaterialSet);
-		for (UMaterialInterface* Mat : ComponentMaterialSet.Materials)
+		for (int ComponentIdx = 0; ComponentIdx < ComponentTargets.Num(); ComponentIdx++)
 		{
-			int* FoundMatIdx = KnownMaterials.Find(Mat);
-			int MatIdx;
-			if (FoundMatIdx)
-			{
-				MatIdx = *FoundMatIdx;
-			}
-			else
+			FComponentMaterialSet ComponentMaterialSet;
+			ComponentTargets[ComponentIdx]->GetMaterialSet(ComponentMaterialSet);
+			for (UMaterialInterface* Mat : ComponentMaterialSet.Materials)
 			{
-				MatIdx = AllMaterialSet.Materials.Add(Mat);
-				KnownMaterials.Add(Mat, MatIdx);
+				int* FoundMatIdx = KnownMaterials.Find(Mat);
+				int MatIdx;
+				if (FoundMatIdx)
+				{
+					MatIdx = *FoundMatIdx;
+				}
+				else
+				{
+					MatIdx = AllMaterialSet.Materials.Add(Mat);
+					KnownMaterials.Add(Mat, MatIdx);
+				}
+				MaterialRemap[ComponentIdx].Add(MatIdx);
 			}
-			MaterialRemap[ComponentIdx].Add(MatIdx);
 		}
 	}
-
-	OriginalDynamicMeshes.SetNum(ComponentTargets.Num());
+	else
+	{
+		ComponentTargets[0]->GetMaterialSet(AllMaterialSet);
+		for (int MatIdx = 0; MatIdx < AllMaterialSet.Materials.Num(); MatIdx++)
+		{
+			MaterialRemap[0].Add(MatIdx);
+		}
+		for (int ComponentIdx = 1; ComponentIdx < ComponentTargets.Num(); ComponentIdx++)
+		{
+			MaterialRemap[ComponentIdx].Init(0, ComponentTargets[ComponentIdx]->GetNumMaterials());
+		}
+	}
 
 	for (int ComponentIdx = 0; ComponentIdx < ComponentTargets.Num(); ComponentIdx++)
 	{
@@ -145,11 +159,16 @@ void UCSGMeshesTool::SetupPreview()
 			MaterialIDs->SetValue(TID, MaterialRemap[ComponentIdx][MaterialIDs->GetValue(TID)]);
 		}
 	}
+	Preview->ConfigureMaterials(AllMaterialSet.Materials, ToolSetupUtil::GetDefaultWorkingMaterial(GetToolManager()));
+}
+
 
+void UCSGMeshesTool::SetupPreview()
+{
 	Preview = NewObject<UMeshOpPreviewWithBackgroundCompute>(this, "Preview");
 	Preview->Setup(this->TargetWorld, this);
-	Preview->ConfigureMaterials(AllMaterialSet.Materials, ToolSetupUtil::GetDefaultWorkingMaterial(GetToolManager()));
 
+	ConfigurePreviewMaterials();
 	
 	DrawnLineSet = NewObject<ULineSetComponent>(Preview->PreviewMesh->GetRootComponent());
 	DrawnLineSet->SetupAttachment(Preview->PreviewMesh->GetRootComponent());
@@ -309,7 +328,12 @@ void UCSGMeshesTool::PostEditChangeProperty(FPropertyChangedEvent& PropertyChang
 
 void UCSGMeshesTool::OnPropertyModified(UObject* PropertySet, FProperty* Property)
 {
-	if (Property && (Property->GetFName() == GET_MEMBER_NAME_CHECKED(UCSGMeshesToolProperties, bShowTransformUI)))
+	if (Property && (Property->GetFName() == GET_MEMBER_NAME_CHECKED(UCSGMeshesToolProperties, bOnlyUseFirstMeshMaterials)))
+	{
+		ConfigurePreviewMaterials();
+		Preview->InvalidateResult();
+	}
+	else if (Property && (Property->GetFName() == GET_MEMBER_NAME_CHECKED(UCSGMeshesToolProperties, bShowTransformUI)))
 	{
 		UpdateGizmoVisibility();
 	}
diff --git a/Engine/Plugins/Experimental/MeshModelingToolset/Source/MeshModelingTools/Private/SelfUnionMeshesTool.cpp b/Engine/Plugins/Experimental/MeshModelingToolset/Source/MeshModelingTools/Private/SelfUnionMeshesTool.cpp
--- a/Engine/Plugins/Experimental/MeshModelingToolset/Source/MeshModelingTools/Private/SelfUnionMeshesTool.cpp
+++ b/Engine/Plugins/Experimental/MeshModelingToolset/Source/MeshModelingTools/Private/SelfUnionMeshesTool.cpp
@@ -98,30 +98,45 @@ void USelfUnionMeshesTool::Setup()
 
 
 
-
-void USelfUnionMeshesTool::SetupPreview()
+void USelfUnionMeshesTool::ConfigurePreviewMaterials()
 {
 	FComponentMaterialSet AllMaterialSet;
 	TMap<UMaterialInterface*, int> KnownMaterials;
 	TArray<TArray<int>> MaterialRemap; MaterialRemap.SetNum(ComponentTargets.Num());
-	for (int ComponentIdx = 0; ComponentIdx < ComponentTargets.Num(); ComponentIdx++)
+
+	if (!Properties->bOnlyUseFirstMeshMaterials)
 	{
-		FComponentMaterialSet ComponentMaterialSet;
-		ComponentTargets[ComponentIdx]->GetMaterialSet(ComponentMaterialSet);
-		for (UMaterialInterface* Mat : ComponentMaterialSet.Materials)
+		for (int ComponentIdx = 0; ComponentIdx < ComponentTargets.Num(); ComponentIdx++)
 		{
-			int* FoundMatIdx = KnownMaterials.Find(Mat);
-			int MatIdx;
-			if (FoundMatIdx)
+			FComponentMaterialSet ComponentMaterialSet;
+			ComponentTargets[ComponentIdx]->GetMaterialSet(ComponentMaterialSet);
+			for (UMaterialInterface* Mat : ComponentMaterialSet.Materials)
 			{
-				MatIdx = *FoundMatIdx;
+				int* FoundMatIdx = KnownMaterials.Find(Mat);
+				int MatIdx;
+				if (FoundMatIdx)
+				{
+					MatIdx = *FoundMatIdx;
+				}
+				else
+				{
+					MatIdx = AllMaterialSet.Materials.Add(Mat);
+					KnownMaterials.Add(Mat, MatIdx);
+				}
+				MaterialRemap[ComponentIdx].Add(MatIdx);
 			}
-			else
-			{
-				MatIdx = AllMaterialSet.Materials.Add(Mat);
-				KnownMaterials.Add(Mat, MatIdx);
-			}
-			MaterialRemap[ComponentIdx].Add(MatIdx);
+		}
+	}
+	else
+	{
+		ComponentTargets[0]->GetMaterialSet(AllMaterialSet);
+		for (int MatIdx = 0; MatIdx < AllMaterialSet.Materials.Num(); MatIdx++)
+		{
+			MaterialRemap[0].Add(MatIdx);
+		}
+		for (int ComponentIdx = 1; ComponentIdx < ComponentTargets.Num(); ComponentIdx++)
+		{
+			MaterialRemap[ComponentIdx].Init(0, ComponentTargets[ComponentIdx]->GetNumMaterials());
 		}
 	}
 
@@ -156,15 +171,21 @@ void USelfUnionMeshesTool::SetupPreview()
 			{
 				return WorldTransform.TransformNormal(Normal);
 			}
-		);
+			);
 	}
 
-	Preview = NewObject<UMeshOpPreviewWithBackgroundCompute>(this, "Preview");
-	Preview->Setup(this->TargetWorld, this);
-	Preview->PreviewMesh->UpdatePreview(CombinedSourceMeshes.Get());
 	Preview->ConfigureMaterials(AllMaterialSet.Materials, ToolSetupUtil::GetDefaultWorkingMaterial(GetToolManager()));
+}
 
 
+void USelfUnionMeshesTool::SetupPreview()
+{
+	Preview = NewObject<UMeshOpPreviewWithBackgroundCompute>(this, "Preview");
+	Preview->Setup(this->TargetWorld, this);
+	ConfigurePreviewMaterials();
+
+	Preview->PreviewMesh->UpdatePreview(CombinedSourceMeshes.Get());
+	
 	DrawnLineSet = NewObject<ULineSetComponent>(Preview->PreviewMesh->GetRootComponent());
 	DrawnLineSet->SetupAttachment(Preview->PreviewMesh->GetRootComponent());
 	DrawnLineSet->SetLineMaterial(ToolSetupUtil::GetDefaultLineComponentMaterial(GetToolManager()));
@@ -277,7 +298,12 @@ void USelfUnionMeshesTool::PostEditChangeProperty(FPropertyChangedEvent& Propert
 
 void USelfUnionMeshesTool::OnPropertyModified(UObject* PropertySet, FProperty* Property)
 {
-	if (PropertySet == HandleSourcesProperties)
+	if (Property && (Property->GetFName() == GET_MEMBER_NAME_CHECKED(USelfUnionMeshesToolProperties, bOnlyUseFirstMeshMaterials)))
+	{
+		ConfigurePreviewMaterials();
+		Preview->InvalidateResult();
+	}
+	else if (PropertySet == HandleSourcesProperties)
 	{
 		// nothing
 	}
diff --git a/Engine/Plugins/Experimental/MeshModelingToolset/Source/MeshModelingTools/Public/CSGMeshesTool.h b/Engine/Plugins/Experimental/MeshModelingToolset/Source/MeshModelingTools/Public/CSGMeshesTool.h
--- a/Engine/Plugins/Experimental/MeshModelingToolset/Source/MeshModelingTools/Public/CSGMeshesTool.h
+++ b/Engine/Plugins/Experimental/MeshModelingToolset/Source/MeshModelingTools/Public/CSGMeshesTool.h
@@ -62,6 +62,10 @@ public:
 	/** Automatically attempt to fill any holes left by CSG (e.g. due to numerical errors) */
 	UPROPERTY(EditAnywhere, Category = Options)
 	bool bAttemptFixHoles = false;
+
+	/** If true, only the first mesh will keep its materials assignments; all other triangles will be assigned material 0 */
+	UPROPERTY(EditAnywhere, Category = Materials)
+	bool bOnlyUseFirstMeshMaterials = false;
 };
 
 
@@ -128,6 +132,7 @@ protected:
 	IToolsContextAssetAPI* AssetAPI;
 
 	void SetupPreview();
+	void ConfigurePreviewMaterials();
 	void SetTransformGizmos();
 	void UpdateGizmoVisibility();
 
diff --git a/Engine/Plugins/Experimental/MeshModelingToolset/Source/MeshModelingTools/Public/SelfUnionMeshesTool.h b/Engine/Plugins/Experimental/MeshModelingToolset/Source/MeshModelingTools/Public/SelfUnionMeshesTool.h
--- a/Engine/Plugins/Experimental/MeshModelingToolset/Source/MeshModelingTools/Public/SelfUnionMeshesTool.h
+++ b/Engine/Plugins/Experimental/MeshModelingToolset/Source/MeshModelingTools/Public/SelfUnionMeshesTool.h
@@ -48,6 +48,10 @@ public:
 	/** Show boundary edges created by the union operation -- often due to numerical error */
 	UPROPERTY(EditAnywhere, Category = Options)
 	bool bShowNewBoundaryEdges = true;
+
+	/** If true, only the first mesh will keep its materials assignments; all other triangles will be assigned material 0 */
+	UPROPERTY(EditAnywhere, Category = Materials)
+	bool bOnlyUseFirstMeshMaterials = false;
 };
 
 
@@ -106,6 +110,7 @@ protected:
 	IToolsContextAssetAPI* AssetAPI;
 
 	void SetupPreview();
+	void ConfigurePreviewMaterials();
 
 	void GenerateAsset(const FDynamicMeshOpResult& Result);
 
