ModelingComponents: refactor some DynamicMeshComponent functionality into base class
#rnx

diff --git a/Engine/Plugins/Experimental/MeshModelingToolset/Source/ModelingComponents/Public/BaseDynamicMeshComponent.h b/Engine/Plugins/Experimental/MeshModelingToolset/Source/ModelingComponents/Public/BaseDynamicMeshComponent.h
--- a/Engine/Plugins/Experimental/MeshModelingToolset/Source/ModelingComponents/Public/BaseDynamicMeshComponent.h
+++ b/Engine/Plugins/Experimental/MeshModelingToolset/Source/ModelingComponents/Public/BaseDynamicMeshComponent.h
@@ -9,13 +9,13 @@
 #include "Changes/MeshVertexChange.h"
 #include "Changes/MeshChange.h"
 #include "Changes/MeshReplacementChange.h"
-
+#include "MeshConversionOptions.h"
 #include "DynamicMesh3.h"
 
 #include "BaseDynamicMeshComponent.generated.h"
 
-
 // predecl
+struct FMeshDescription;
 class FMeshVertexChange;
 class FMeshChange;
 
@@ -65,6 +65,32 @@ class MODELINGCOMPONENTS_API UBaseDynamicMeshComponent : public UMeshComponent,
 
 public:
 
+	/**
+	 * initialize the internal mesh from a MeshDescription
+	 */
+	virtual void InitializeMesh(FMeshDescription* MeshDescription)
+	{
+		unimplemented();
+	}
+
+	/**
+	 * @return pointer to internal mesh
+	 */
+	virtual FDynamicMesh3* GetMesh()
+	{
+		unimplemented();
+		return nullptr;
+	}
+
+	/**
+	 * @return pointer to internal mesh
+	 */
+	virtual const FDynamicMesh3* GetMesh() const
+	{
+		unimplemented();
+		return nullptr;
+	}
+
 	/**
 	 * Call this if you update the mesh via GetMesh()
 	 * @todo should provide a function that calls a lambda to modify the mesh, and only return const mesh pointer
@@ -99,6 +125,28 @@ public:
 	}
 
 
+	//
+	// Modification support
+	//
+public:
+
+	virtual void ApplyTransform(const FTransform3d& Transform, bool bInvert)
+	{
+		unimplemented();
+	}
+
+	/**
+	 * Write the internal mesh to a MeshDescription
+	 * @param bHaveModifiedTopology if false, we only update the vertex positions in the MeshDescription, otherwise it is Empty()'d and regenerated entirely
+	 * @param ConversionOptions struct of additional options for the conversion
+	 */
+	virtual void Bake(FMeshDescription* MeshDescription, bool bHaveModifiedTopology, const FConversionToMeshDescriptionOptions& ConversionOptions)
+	{
+		unimplemented();
+	}
+
+
+
 protected:
 	/**
 	 * Subclass must implement this to notify allocated proxies of updated materials
diff --git a/Engine/Plugins/Experimental/MeshModelingToolset/Source/ModelingComponents/Public/OctreeDynamicMeshComponent.h b/Engine/Plugins/Experimental/MeshModelingToolset/Source/ModelingComponents/Public/OctreeDynamicMeshComponent.h
--- a/Engine/Plugins/Experimental/MeshModelingToolset/Source/ModelingComponents/Public/OctreeDynamicMeshComponent.h
+++ b/Engine/Plugins/Experimental/MeshModelingToolset/Source/ModelingComponents/Public/OctreeDynamicMeshComponent.h
@@ -41,12 +41,18 @@ public:
 	/**
 	 * initialize the internal mesh from a MeshDescription
 	 */
-	void InitializeMesh(FMeshDescription* MeshDescription);
+	virtual void InitializeMesh(FMeshDescription* MeshDescription) override;
 
 	/**
 	 * @return pointer to internal mesh
 	 */
-	FDynamicMesh3* GetMesh() { return Mesh.Get(); }
+	virtual FDynamicMesh3* GetMesh() override { return Mesh.Get(); }
+
+	/**
+	 * @return pointer to internal mesh
+	 */
+	virtual const FDynamicMesh3* GetMesh() const override { return Mesh.Get(); }
+
 
 	FDynamicMeshOctree3* GetOctree() { return Octree.Get(); }
 
@@ -60,7 +66,7 @@ public:
 	 * @param bHaveModifiedTopology if false, we only update the vertex positions in the MeshDescription, otherwise it is Empty()'d and regenerated entirely
 	 * @param ConversionOptions struct of additional options for the conversion
 	 */
-	void Bake(FMeshDescription* MeshDescription, bool bHaveModifiedTopology, const FConversionToMeshDescriptionOptions& ConversionOptions);
+	virtual void Bake(FMeshDescription* MeshDescription, bool bHaveModifiedTopology, const FConversionToMeshDescriptionOptions& ConversionOptions) override;
 
 	/**
 	* Write the internal mesh to a MeshDescription with default conversion options
@@ -76,7 +82,7 @@ public:
 	 * Apply transform to internal mesh. Invalidates RenderProxy.
 	 * @param bInvert if true, inverse tranform is applied instead of forward transform
 	 */
-	void ApplyTransform(const FTransform3d& Transform, bool bInvert);
+	virtual void ApplyTransform(const FTransform3d& Transform, bool bInvert) override;
 
 	//
 	// change tracking/etc
diff --git a/Engine/Plugins/Experimental/MeshModelingToolset/Source/ModelingComponents/Public/SimpleDynamicMeshComponent.h b/Engine/Plugins/Experimental/MeshModelingToolset/Source/ModelingComponents/Public/SimpleDynamicMeshComponent.h
--- a/Engine/Plugins/Experimental/MeshModelingToolset/Source/ModelingComponents/Public/SimpleDynamicMeshComponent.h
+++ b/Engine/Plugins/Experimental/MeshModelingToolset/Source/ModelingComponents/Public/SimpleDynamicMeshComponent.h
@@ -43,17 +43,17 @@ public:
 	/**
 	 * initialize the internal mesh from a MeshDescription
 	 */
-	void InitializeMesh(FMeshDescription* MeshDescription);
+	virtual void InitializeMesh(FMeshDescription* MeshDescription) override;
 
 	/**
 	 * @return pointer to internal mesh
 	 */
-	FDynamicMesh3* GetMesh() { return Mesh.Get(); }
+	virtual FDynamicMesh3* GetMesh() override { return Mesh.Get(); }
 
 	/**
 	 * @return pointer to internal mesh
 	 */
-	const FDynamicMesh3* GetMesh() const { return Mesh.Get(); }
+	virtual const FDynamicMesh3* GetMesh() const override { return Mesh.Get(); }
 
 	/**
 	 * @return the current internal mesh, which is replaced with an empty mesh
@@ -74,7 +74,7 @@ public:
 	 * @param bHaveModifiedTopology if false, we only update the vertex positions in the MeshDescription, otherwise it is Empty()'d and regenerated entirely
 	 * @param ConversionOptions struct of additional options for the conversion
 	 */
-	void Bake(FMeshDescription* MeshDescription, bool bHaveModifiedTopology, const FConversionToMeshDescriptionOptions& ConversionOptions);
+	virtual void Bake(FMeshDescription* MeshDescription, bool bHaveModifiedTopology, const FConversionToMeshDescriptionOptions& ConversionOptions) override;
 
 	/**
 	* Write the internal mesh to a MeshDescription with default conversion options
@@ -90,7 +90,7 @@ public:
 	 * Apply transform to internal mesh. Updates Octree and RenderProxy if available.
 	 * @param bInvert if true, inverse tranform is applied instead of forward transform
 	 */
-	void ApplyTransform(const FTransform3d& Transform, bool bInvert);
+	virtual void ApplyTransform(const FTransform3d& Transform, bool bInvert) override;
 
 
 	//
