ModelingTools: Ongoing Sculpting improvements:
- new UMeshSculptBrushOpProps baseclass for per-BrushOp Property Sets, added such a PropertySet for each BrushOp
    - separate Strength/Depth/Falloff for each BrushOp so that they can have different settings. Currently cannot be in a shared base PropSet class because of different UProp Category strings
- UMeshSculptToolBase now supports registering BrushOp factory paired w/ ToolProps, rewrote VertexSculpt to use this to handle switching active BrushOp and visible PropSet
- removed UPlaneBrushProperties, vertex sculpt no longer uses USculptMaxBrushProperties  (still used by dynamic sculpt)
- moved UMeshSurfacePointTool implementation up to UMeshSculptToolBase
- UMeshSculptToolBase now has OnBeginStroke() / OnEndStroke()API for subclasses to use to handle stroke (instead of overriding drag functions)
- Moved some setup to UMeshSculptToolBase. It can now initializes the target Component, can apply to Simple or Octree variants,
- moved Primary/Secondary BrushOp and active Falloff to UMeshSculptToolBase
- Support for various falloff types implemented in StampFalloffs.h, new EMeshSculptFalloffType exposed in VertexSculpt properties
- moved target component Transforms to UMeshSculptToolBase
- moved Brush Frames to UMeshSculptToolBase. Now storing Frames instead of Point+Normal, moved functions to update position for various stamp types
- UMeshSculptToolBase now has API for hit-testing against Sculpt and Base mesh, VertexSculpt implements
- UMeshSculptToolBase now has APIs for accessing transforms, current DMesh and BaseMesh, active stroke Depth/Strength/Falloff properties,
- moved Brush Target Plane (eg for Move), Stamp Plane (eg for Flatten), and StrokePlane (eg for Plane brushes) to UMeshSculptToolBase, and relevant calcuation APIs
- base USculptBrushProperties now has Lazyness parameter, implemented lazy-brush in base sculpt tool
#rnx

diff --git a/Engine/Plugins/Experimental/MeshModelingToolset/Source/MeshModelingTools/Private/DynamicMeshSculptTool.cpp b/Engine/Plugins/Experimental/MeshModelingToolset/Source/MeshModelingTools/Private/DynamicMeshSculptTool.cpp
--- a/Engine/Plugins/Experimental/MeshModelingToolset/Source/MeshModelingTools/Private/DynamicMeshSculptTool.cpp
+++ b/Engine/Plugins/Experimental/MeshModelingToolset/Source/MeshModelingTools/Private/DynamicMeshSculptTool.cpp
@@ -157,10 +157,6 @@ void UDynamicMeshSculptTool::Setup()
 	AddToolPropertySource(SculptProperties);
 
 	// add brush-specific properties 
-	PlaneBrushProperties = NewObject<UPlaneBrushProperties>(this);
-	PlaneBrushProperties->RestoreProperties(this);
-	AddToolPropertySource(PlaneBrushProperties);
-
 	SculptMaxBrushProperties = NewObject<USculptMaxBrushProperties>();
 	SculptMaxBrushProperties->RestoreProperties(this);
 	AddToolPropertySource(SculptMaxBrushProperties);
@@ -183,7 +179,6 @@ void UDynamicMeshSculptTool::Setup()
 	SculptProperties->RestoreProperties(this);
 
 	// disable tool-specific properties
-	SetToolPropertySourceEnabled(PlaneBrushProperties, false);
 	SetToolPropertySourceEnabled(GizmoProperties, false);
 	SetToolPropertySourceEnabled(SculptMaxBrushProperties, false);
 	SetToolPropertySourceEnabled(KelvinBrushProperties, false);
@@ -283,7 +278,6 @@ void UDynamicMeshSculptTool::Shutdown(EToolShutdownType ShutdownType)
 	KelvinBrushProperties->SaveProperties(this);
 	ViewProperties->SaveProperties(this);
 	GizmoProperties->SaveProperties(this);
-	PlaneBrushProperties->SaveProperties(this);
 	SculptMaxBrushProperties->SaveProperties(this);
 	RemeshProperties->SaveProperties(this);
 }
@@ -810,7 +804,7 @@ void UDynamicMeshSculptTool::ApplyPlaneBrush(const FRay& WorldRay)
 	}
 
 	static const double PlaneSigns[3] = { 0, -1, 1 };
-	double PlaneSign = PlaneSigns[(int32)PlaneBrushProperties->WhichSide];
+	double PlaneSign = PlaneSigns[0];
 
 	FVector3d NewBrushPosLocal = CurTargetTransform.InverseTransformPosition(LastBrushPosWorld);
 	FVector3d BrushNormalLocal = CurTargetTransform.InverseTransformNormal(LastBrushPosNormalWorld);
@@ -853,7 +847,7 @@ void UDynamicMeshSculptTool::ApplyFixedPlaneBrush(const FRay& WorldRay)
 	}
 
 	static const double PlaneSigns[3] = { 0, -1, 1 };
-	double PlaneSign = PlaneSigns[(int32)PlaneBrushProperties->WhichSide];
+	double PlaneSign = PlaneSigns[0];
 
 	FVector3d NewBrushPosLocal = CurTargetTransform.InverseTransformPosition(LastBrushPosWorld);
 	FVector3d BrushNormalLocal = CurTargetTransform.InverseTransformNormal(LastBrushPosNormalWorld);
@@ -902,7 +896,7 @@ void UDynamicMeshSculptTool::ApplyFlattenBrush(const FRay& WorldRay)
 	}
 
 	static const double PlaneSigns[3] = { 0, -1, 1 };
-	double PlaneSign = PlaneSigns[(int32)PlaneBrushProperties->WhichSide];
+	double PlaneSign = PlaneSigns[0];
 
 	FVector3d NewBrushPosLocal = CurTargetTransform.InverseTransformPosition(LastBrushPosWorld);
 	FVector3d BrushNormalLocal = CurTargetTransform.InverseTransformNormal(LastBrushPosNormalWorld);
@@ -2314,19 +2308,13 @@ void UDynamicMeshSculptTool::UpdateBrushType(EDynamicMeshSculptBrushType BrushTy
 	Builder.AppendLine(BaseMessage);
 
 	SetToolPropertySourceEnabled(GizmoProperties, false);
-	SetToolPropertySourceEnabled(PlaneBrushProperties, false);
 	SetToolPropertySourceEnabled(SculptMaxBrushProperties, false);
 
 	if (BrushType == EDynamicMeshSculptBrushType::FixedPlane)
 	{
 		Builder.AppendLine(LOCTEXT("FixedPlaneTip", "Use T to reposition Work Plane at cursor, Shift+T to align to Normal, Ctrl+Shift+T to align to View"));
-		SetToolPropertySourceEnabled(PlaneBrushProperties, true);
 		SetToolPropertySourceEnabled(GizmoProperties, true);
 	}
-	if (BrushType == EDynamicMeshSculptBrushType::Plane || BrushType == EDynamicMeshSculptBrushType::PlaneViewAligned || BrushType == EDynamicMeshSculptBrushType::Flatten)
-	{
-		SetToolPropertySourceEnabled(PlaneBrushProperties, true);
-	}
 	if (BrushType == EDynamicMeshSculptBrushType::SculptMax)
 	{
 		SetToolPropertySourceEnabled(SculptMaxBrushProperties, true);
diff --git a/Engine/Plugins/Experimental/MeshModelingToolset/Source/MeshModelingTools/Private/MeshVertexSculptTool.cpp b/Engine/Plugins/Experimental/MeshModelingToolset/Source/MeshModelingTools/Private/MeshVertexSculptTool.cpp
--- a/Engine/Plugins/Experimental/MeshModelingToolset/Source/MeshModelingTools/Private/MeshVertexSculptTool.cpp
+++ b/Engine/Plugins/Experimental/MeshModelingToolset/Source/MeshModelingTools/Private/MeshVertexSculptTool.cpp
@@ -1,26 +1,18 @@
 // Copyright Epic Games, Inc. All Rights Reserved.
 
 #include "MeshVertexSculptTool.h"
-#include "Containers/Map.h"
 #include "InteractiveToolManager.h"
 #include "InteractiveGizmoManager.h"
-#include "ToolBuilderUtil.h"
 #include "Async/ParallelFor.h"
 #include "Async/Async.h"
 
-#include "ProjectionTargets.h"
 #include "MeshWeights.h"
 #include "MeshNormals.h"
 #include "MeshIndexUtil.h"
-#include "PreviewMesh.h"
-#include "ToolSetupUtil.h"
 #include "Util/BufferUtil.h"
 
 #include "Changes/MeshVertexChange.h"
 
-#include "ToolDataVisualizer.h"
-#include "Components/PrimitiveComponent.h"
-
 //#include "Sculpting/KelvinletBrushOp.h"
 #include "Sculpting/MeshSmoothingBrushOps.h"
 #include "Sculpting/MeshInflateBrushOps.h"
@@ -28,12 +20,9 @@
 #include "Sculpting/MeshPlaneBrushOps.h"
 #include "Sculpting/MeshPinchBrushOps.h"
 #include "Sculpting/MeshSculptBrushOps.h"
+#include "Sculpting/StampFalloffs.h"
 #include "Sculpting/MeshSculptUtil.h"
 
-#include "UObject/ObjectMacros.h"
-#include "Materials/Material.h"
-#include "Materials/MaterialInstanceDynamic.h"
-
 #define LOCTEXT_NAMESPACE "UMeshVertexSculptTool"
 
 namespace
@@ -67,24 +56,7 @@ void UMeshVertexSculptTool::Setup()
 
 	// create dynamic mesh component to use for live preview
 	DynamicMeshComponent = NewObject<USimpleDynamicMeshComponent>(ComponentTarget->GetOwnerActor());
-	DynamicMeshComponent->SetupAttachment(ComponentTarget->GetOwnerActor()->GetRootComponent());
-	DynamicMeshComponent->RegisterComponent();
-
-	// initialize from LOD-0 MeshDescription
-	DynamicMeshComponent->InitializeMesh(ComponentTarget->GetMesh());
-	double MaxDimension = DynamicMeshComponent->GetMesh()->GetCachedBounds().MaxDim();
-
-	// transform mesh to world space because handling scaling inside brush is a mess
-	// Note: this transform does not include translation ( so only the 3x3 transform)
-	InitialTargetTransform = FTransform3d(ComponentTarget->GetWorldTransform());
-	// clamp scaling because if we allow zero-scale we cannot invert this transform on Accept
-	InitialTargetTransform.ClampMinimumScale(0.01);
-	FVector3d Translation = InitialTargetTransform.GetTranslation();
-	InitialTargetTransform.SetTranslation(FVector3d::Zero());
-	DynamicMeshComponent->ApplyTransform(InitialTargetTransform, false);
-	// since we moved to World coords there is not a current transform anymore.
-	CurTargetTransform = FTransform3d(Translation);
-	DynamicMeshComponent->SetWorldTransform((FTransform)CurTargetTransform);
+	InitializeSculptMeshComponent(DynamicMeshComponent);
 
 	// assign materials
 	FComponentMaterialSet MaterialSet;
@@ -97,19 +69,26 @@ void UMeshVertexSculptTool::Setup()
 	DynamicMeshComponent->bInvalidateProxyOnChange = false;
 	OnDynamicMeshComponentChangedHandle = DynamicMeshComponent->OnMeshVerticesChanged.AddUObject(this, &UMeshVertexSculptTool::OnDynamicMeshComponentChanged);
 
-	FDynamicMesh3* Mesh = DynamicMeshComponent->GetMesh();
-	// do we always want to keep vertex normals updated? Perhaps should discard vertex normals before baking?
-	//FMeshNormals::QuickComputeVertexNormals(*Mesh);
-
 	// initialize dynamic octree
-	Octree.RootDimension = MaxDimension / 10.0;
-	Octree.SetMaxTreeDepth(3);
+	FDynamicMesh3* Mesh = GetSculptMesh();
+	FAxisAlignedBox3d Bounds = Mesh->GetCachedBounds();
+	if (Mesh->TriangleCount() > 100000)
+	{
+		Octree.RootDimension = Bounds.MaxDim() / 10.0;
+		Octree.SetMaxTreeDepth(4);
+	} 
+	else
+	{
+		Octree.RootDimension = Bounds.MaxDim();
+		Octree.SetMaxTreeDepth(8);
+	}
 	Octree.Initialize(Mesh);
-	FDynamicMeshOctree3::FStatistics Stats;
-	Octree.ComputeStatistics(Stats);
+	//Octree.CheckValidity(EValidityCheckFailMode::Check, true, true);
+	//FDynamicMeshOctree3::FStatistics Stats;
+	//Octree.ComputeStatistics(Stats);
 	//UE_LOG(LogTemp, Warning, TEXT("Octree Stats: %s"), *Stats.ToString());
 
-	// initialize decomposition
+	// initialize render decomposition
 	TUniquePtr<FMeshRenderDecomposition> Decomp = MakeUnique<FMeshRenderDecomposition>();
 	FMeshRenderDecomposition::BuildChunkedDecomposition(Mesh, &MaterialSet, *Decomp);
 	Decomp->BuildAssociations(Mesh);
@@ -121,103 +100,114 @@ void UMeshVertexSculptTool::Setup()
 	bTargetDirty = false;
 
 	// initialize brush radius range interval, brush properties
-	UMeshSculptToolBase::InitializeBrushSizeRange(DynamicMeshComponent->GetMesh()->GetCachedBounds());
+	UMeshSculptToolBase::InitializeBrushSizeRange(Bounds);
 
 	// initialize other properties
 	SculptProperties = NewObject<UVertexBrushSculptProperties>(this);
-	//KelvinBrushProperties = NewObject<UKelvinBrushProperties>(this);
-
-	// hide input Component
-	ComponentTarget->SetOwnerVisibility(false);
 
 	// init state flags flags
-	bInDrag = false;
-	bStampPending = false;
 	ActiveVertexChange = nullptr;
 
 	InitializeIndicator();
 
 	// initialize our properties
 	AddToolPropertySource(UMeshSculptToolBase::BrushProperties);
+	UMeshSculptToolBase::BrushProperties->bShowPerBrushProps = false;
+	UMeshSculptToolBase::BrushProperties->bShowFalloff = false;
 	AddToolPropertySource(SculptProperties);
+	CalculateBrushRadius();
+	SculptProperties->RestoreProperties(this);
+
+	this->BaseMeshQueryFunc = [&](int32 VertexID, const FVector3d& Position, double MaxDist, FVector3d& PosOut, FVector3d& NormalOut)
+	{
+		return GetBaseMeshNearest(VertexID, Position, MaxDist, PosOut, NormalOut);
+	};
 
-	// add brush-specific properties 
-	PlaneBrushProperties = NewObject<UPlaneBrushProperties>(this);
-	PlaneBrushProperties->RestoreProperties(this);
-	AddToolPropertySource(PlaneBrushProperties);
-	SetToolPropertySourceEnabled(PlaneBrushProperties, false);
+	RegisterBrushType((int32)EMeshVertexSculptBrushType::Smooth,
+		MakeUnique<TBasicMeshSculptBrushOpFactory<FSmoothBrushOp>>(),
+		NewObject<USmoothBrushOpProps>(this));
 
-	SculptMaxBrushProperties = NewObject<USculptMaxBrushProperties>();
-	SculptMaxBrushProperties->RestoreProperties(this);
-	AddToolPropertySource(SculptMaxBrushProperties);
+	RegisterBrushType((int32)EMeshVertexSculptBrushType::SmoothFill,
+		MakeUnique<TBasicMeshSculptBrushOpFactory<FSmoothFillBrushOp>>(),
+		NewObject<USmoothFillBrushOpProps>(this));
 
-	//AddToolPropertySource(KelvinBrushProperties);
-	//KelvinBrushProperties->RestoreProperties(this);
+	RegisterBrushType((int32)EMeshVertexSculptBrushType::Move,
+		MakeUnique<TBasicMeshSculptBrushOpFactory<FMoveBrushOp>>(),
+		NewObject<UMoveBrushOpProps>(this));
 
-	AddToolPropertySource(UMeshSculptToolBase::GizmoProperties);
-	SetToolPropertySourceEnabled(UMeshSculptToolBase::GizmoProperties, false);
+	RegisterBrushType((int32)EMeshVertexSculptBrushType::Offset,
+		MakeUnique<FLambdaMeshSculptBrushOpFactory>([this]() { return MakeUnique<FSurfaceSculptBrushOp>(BaseMeshQueryFunc); }),
+		NewObject<UStandardSculptBrushOpProps>(this));
+
+	RegisterBrushType((int32)EMeshVertexSculptBrushType::SculptView,
+		MakeUnique<FLambdaMeshSculptBrushOpFactory>( [this]() { return MakeUnique<FViewAlignedSculptBrushOp>(BaseMeshQueryFunc); } ),
+		NewObject<UViewAlignedSculptBrushOpProps>(this));
+
+	RegisterBrushType((int32)EMeshVertexSculptBrushType::SculptMax,
+		MakeUnique<FLambdaMeshSculptBrushOpFactory>([this]() { return MakeUnique<FSurfaceMaxSculptBrushOp>(BaseMeshQueryFunc); }),
+		NewObject<USculptMaxBrushOpProps>(this));
+
+	RegisterBrushType((int32)EMeshVertexSculptBrushType::Inflate,
+		MakeUnique<TBasicMeshSculptBrushOpFactory<FInflateBrushOp>>(),
+		NewObject<UInflateBrushOpProps>(this));
+
+	RegisterBrushType((int32)EMeshVertexSculptBrushType::Pinch,
+		MakeUnique<TBasicMeshSculptBrushOpFactory<FPinchBrushOp>>(),
+		NewObject<UPinchBrushOpProps>(this));
+
+	RegisterBrushType((int32)EMeshVertexSculptBrushType::Flatten,
+		MakeUnique<TBasicMeshSculptBrushOpFactory<FFlattenBrushOp>>(),
+		NewObject<UFlattenBrushOpProps>(this));
+
+	RegisterBrushType((int32)EMeshVertexSculptBrushType::Plane,
+		MakeUnique<TBasicMeshSculptBrushOpFactory<FPlaneBrushOp>>(),
+		NewObject<UPlaneBrushOpProps>(this));
+
+	RegisterBrushType((int32)EMeshVertexSculptBrushType::PlaneViewAligned,
+		MakeUnique<TBasicMeshSculptBrushOpFactory<FPlaneBrushOp>>(),
+		NewObject<UViewAlignedPlaneBrushOpProps>(this));
+
+	RegisterBrushType((int32)EMeshVertexSculptBrushType::FixedPlane,
+		MakeUnique<TBasicMeshSculptBrushOpFactory<FPlaneBrushOp>>(),
+		NewObject<UFixedPlaneBrushOpProps>(this));
 
-	CalculateBrushRadius();
-	SculptProperties->RestoreProperties(this);
 
-	// disable tool-specific properties
-	SetToolPropertySourceEnabled(SculptMaxBrushProperties, false);
-	//SetToolPropertySourceEnabled(KelvinBrushProperties, false);
+	// secondary brushes
+	RegisterSecondaryBrushType((int32)EMeshVertexSculptBrushType::Smooth,
+		MakeUnique<TBasicMeshSculptBrushOpFactory<FSmoothBrushOp>>(),
+		NewObject<USecondarySmoothBrushOpProps>(this));
+
+
+	AddToolPropertySource(UMeshSculptToolBase::GizmoProperties);
+	SetToolPropertySourceEnabled(UMeshSculptToolBase::GizmoProperties, false);
 
 	AddToolPropertySource(UMeshSculptToolBase::ViewProperties);
 
 	// register watchers
-	BrushTypeWatcher.Initialize(
-		[this]() { return SculptProperties->PrimaryBrushType; },
-		[this](EMeshVertexSculptBrushType NewBrushType) { UpdateBrushType(NewBrushType); }, SculptProperties->PrimaryBrushType);
+	SculptProperties->WatchProperty( SculptProperties->PrimaryBrushType,
+		[this](EMeshVertexSculptBrushType NewType) { UpdateBrushType(NewType); });
 
-	ActiveFalloff = MakeShared<FMeshSculptFallofFunc>();
-	ActiveFalloff->FalloffFunc = [&](const FSculptBrushStamp& StampInfo, const FVector3d& Position)
-	{
-		return CalculateBrushFalloff(Position.Distance(StampInfo.LocalFrame.Origin));
-	};
+	SculptProperties->WatchProperty( SculptProperties->PrimaryFalloffType,
+		[this](EMeshSculptFalloffType NewType) { SetPrimaryFalloffType(NewType); });
 
-	SecondaryBrushOp = MakeUnique<FSmoothBrushOp>();
-	SecondaryBrushOp->Falloff = ActiveFalloff;
-	PrimaryBrushOp = MakeUnique<FSmoothBrushOp>();  // set default
-	PrimaryBrushOp->Falloff = ActiveFalloff;
+	// must call before updating brush type so that we register all brush properties?
+	UMeshSculptToolBase::OnCompleteSetup();
 
 	UpdateBrushType(SculptProperties->PrimaryBrushType);
+	SetActiveSecondaryBrushType((int32)EMeshVertexSculptBrushType::Smooth);
 }
 
 void UMeshVertexSculptTool::Shutdown(EToolShutdownType ShutdownType)
 {
-	UMeshSculptToolBase::Shutdown(ShutdownType);
-
 	if (DynamicMeshComponent != nullptr)
 	{
 		DynamicMeshComponent->OnMeshChanged.Remove(OnDynamicMeshComponentChangedHandle);
-
-		ComponentTarget->SetOwnerVisibility(true);
-
-		if (ShutdownType == EToolShutdownType::Accept)
-		{
-			// safe to do this here because we are about to destroy componeont
-			DynamicMeshComponent->ApplyTransform(InitialTargetTransform, true);
-
-			// this block bakes the modified DynamicMeshComponent back into the StaticMeshComponent inside an undo transaction
-			GetToolManager()->BeginUndoTransaction(LOCTEXT("SculptMeshToolTransactionName", "Sculpt Mesh"));
-			ComponentTarget->CommitMesh([=](const FPrimitiveComponentTarget::FCommitParams& CommitParams)
-			{
-				DynamicMeshComponent->Bake(CommitParams.MeshDescription, false);
-			});
-			GetToolManager()->EndUndoTransaction();
-		}
-
-		DynamicMeshComponent->UnregisterComponent();
-		DynamicMeshComponent->DestroyComponent();
-		DynamicMeshComponent = nullptr;
 	}
 
 	SculptProperties->SaveProperties(this);
-	//KelvinBrushProperties->SaveProperties(this);
-	PlaneBrushProperties->SaveProperties(this);
-	SculptMaxBrushProperties->SaveProperties(this);
+
+	// this call will commit result, unregister and destroy DynamicMeshComponent
+	UMeshSculptToolBase::Shutdown(ShutdownType);
 }
 
 
@@ -226,99 +216,57 @@ void UMeshVertexSculptTool::OnPropertyModified(UObject* PropertySet, FProperty*
 	CalculateBrushRadius();
 }
 
-bool UMeshVertexSculptTool::HitTest(const FRay& Ray, FHitResult& OutHit)
-{
-	FRay3d LocalRay(CurTargetTransform.InverseTransformPosition(Ray.Origin),
-		CurTargetTransform.InverseTransformVector(Ray.Direction));
-	LocalRay.Direction.Normalize();
-	FDynamicMesh3* Mesh = DynamicMeshComponent->GetMesh();
 
-	int HitTID = FindHitSculptMeshTriangle(LocalRay);
-	if (HitTID != IndexConstants::InvalidID)
-	{
-		FTriangle3d Triangle;
-		Mesh->GetTriVertices(HitTID, Triangle.V[0], Triangle.V[1], Triangle.V[2]);
-		FIntrRay3Triangle3d Query(LocalRay, Triangle);
-		Query.Find();
-
-		OutHit.FaceIndex = HitTID;
-		OutHit.Distance = Query.RayParameter;
-		OutHit.Normal = (FVector)CurTargetTransform.TransformNormal(Mesh->GetTriNormal(HitTID));
-		OutHit.ImpactPoint = (FVector)CurTargetTransform.TransformPosition(LocalRay.PointAt(Query.RayParameter));
-		return true;
-	}
-	return false;
-}
 
-void UMeshVertexSculptTool::OnBeginDrag(const FRay& Ray)
+void UMeshVertexSculptTool::OnBeginStroke(const FRay& WorldRay)
 {
-	bSmoothing = GetShiftToggle();
-	bInvert = GetCtrlToggle();
+	UpdateBrushPosition(WorldRay);
 
-	FHitResult OutHit;
-	if (HitTest(Ray, OutHit))
+	if (SculptProperties->PrimaryBrushType == EMeshVertexSculptBrushType::Plane ||
+		SculptProperties->PrimaryBrushType == EMeshVertexSculptBrushType::PlaneViewAligned)
 	{
-		FVector3d BrushStartCenterWorld = Ray.PointAt(OutHit.Distance) + BrushProperties->Depth*CurrentBrushRadius*Ray.Direction;
-
-		bInDrag = true;
-
-		ActiveDragPlane = FFrame3d(BrushStartCenterWorld, -Ray.Direction);
-		ActiveDragPlane.RayPlaneIntersection(Ray.Origin, Ray.Direction, 2, LastHitPosWorld);
+		UpdateROI(GetBrushFrameLocal().Origin);
+		UpdateStrokeReferencePlaneForROI(GetBrushFrameLocal(), TriangleROI,
+			(SculptProperties->PrimaryBrushType == EMeshVertexSculptBrushType::PlaneViewAligned));
+	}
+	else if (SculptProperties->PrimaryBrushType == EMeshVertexSculptBrushType::FixedPlane)
+	{
+		UpdateStrokeReferencePlaneFromWorkPlane();
+	}
 
-		LastBrushPosWorld = LastHitPosWorld;
-		LastBrushPosNormalWorld = ActiveDragPlane.Z();
-		LastBrushPosLocal = CurTargetTransform.InverseTransformPosition(LastHitPosWorld);
+	// initialize first "Last Stamp", so that we can assume all stamps in stroke have a valid previous stamp
+	LastStamp.WorldFrame = GetBrushFrameWorld();
+	LastStamp.LocalFrame = GetBrushFrameLocal();
+	LastStamp.Radius = GetCurrentBrushRadius();
+	LastStamp.Falloff = GetCurrentBrushFalloff();
+	LastStamp.Direction = GetInInvertStroke() ? -1.0 : 1.0;
+	LastStamp.Depth = GetCurrentBrushDepth();
+	LastStamp.Power = GetActivePressure() * GetCurrentBrushStrength();
+	LastStamp.TimeStamp = FDateTime::Now();
 
-		BeginChange();
+	FSculptBrushOptions SculptOptions;
+	//SculptOptions.bPreserveUVFlow = false; // SculptProperties->bPreserveUVFlow;
+	SculptOptions.ConstantReferencePlane = GetCurrentStrokeReferencePlane();
 
-		UpdateROI(LastBrushPosLocal);
+	TUniquePtr<FMeshSculptBrushOp>& UseBrushOp = GetActiveBrushOp();
+	UseBrushOp->ConfigureOptions(SculptOptions);
+	UseBrushOp->BeginStroke(GetSculptMesh(), LastStamp, VertexROI);
 
-		if (SculptProperties->PrimaryBrushType == EMeshVertexSculptBrushType::Plane)
-		{
-			ActiveFixedBrushPlane = ComputeROIBrushPlane(LastBrushPosLocal, false, false);
-		}
-		else if (SculptProperties->PrimaryBrushType == EMeshVertexSculptBrushType::PlaneViewAligned)
-		{
-			AlignBrushToView();
-			ActiveFixedBrushPlane = ComputeROIBrushPlane(LastBrushPosLocal, false, true);
-		}
-		else if (SculptProperties->PrimaryBrushType == EMeshVertexSculptBrushType::FixedPlane)
-		{
-			ActiveFixedBrushPlane = FFrame3d(
-				CurTargetTransform.InverseTransformPosition(GizmoProperties->Position),
-				CurTargetTransform.GetRotation().Inverse() * (FQuaterniond)GizmoProperties->Rotation);
-		}
+	AccumulatedTriangleROI.Reset();
 
-		// init last stamp
-		LastStamp.WorldFrame = FFrame3d(LastBrushPosWorld, LastBrushPosNormalWorld);
-		LastStamp.LocalFrame = LastStamp.WorldFrame;
-		LastStamp.LocalFrame.Transform(CurTargetTransform.Inverse());
-		LastStamp.Radius = CurrentBrushRadius;
-		LastStamp.Direction = (bInvert) ? -1.0 : 1.0;
-		LastStamp.Depth = BrushProperties->Depth;
-		LastStamp.Power = GetActivePressure() * ((bSmoothing) ? SculptProperties->SmoothBrushSpeed : SculptProperties->PrimaryBrushSpeed);
-		TUniquePtr<FMeshSculptBrushOp>& UseBrushOp = (bSmoothing) ? SecondaryBrushOp : PrimaryBrushOp;
-
-		FSculptBrushOptions SculptOptions;
-		SculptOptions.bPreserveUVFlow = SculptProperties->bPreserveUVFlow;
-		SculptOptions.MaxHeight = CurrentBrushRadius * SculptMaxBrushProperties->MaxHeight;
-		if (SculptMaxBrushProperties->bFreezeCurrentHeight && SculptMaxFixedHeight >= 0)
-		{
-			SculptOptions.MaxHeight = SculptMaxFixedHeight;
-		}
-		SculptMaxFixedHeight = SculptOptions.MaxHeight;		// yech
-		SculptOptions.WhichPlaneSideIndex = (int32)PlaneBrushProperties->WhichSide;
-		SculptOptions.ConstantReferencePlane = ActiveFixedBrushPlane;
-		UseBrushOp->ConfigureOptions(SculptOptions);
+	// begin change here? or wait for first stamp?
+	BeginChange();
+}
 
-		UseBrushOp->BeginStroke(DynamicMeshComponent->GetMesh(), LastStamp, VertexROI);
+void UMeshVertexSculptTool::OnEndStroke()
+{
+	// update spatial
+	bTargetDirty = true;
 
-		AccumulatedTriangleROI.Reset();
+	GetActiveBrushOp()->EndStroke(GetSculptMesh(), LastStamp, VertexROI);
 
-		// apply initial stamp
-		PendingStampRay = Ray;
-		bStampPending = true;
-	}
+	// close change record
+	EndChange();
 }
 
 
@@ -327,17 +275,14 @@ void UMeshVertexSculptTool::UpdateROI(const FVector3d& BrushPos)
 {
 	SCOPE_CYCLE_COUNTER(VtxSculptTool_UpdateROI);
 
-	// TODO: need dynamic vertex hash table?
-
-	float RadiusSqr = CurrentBrushRadius * CurrentBrushRadius;
-
+	float RadiusSqr = GetCurrentBrushRadius() * GetCurrentBrushRadius();
 	FAxisAlignedBox3d BrushBox(
-		BrushPos - CurrentBrushRadius * FVector3d::One(),
-		BrushPos + CurrentBrushRadius * FVector3d::One());
+		BrushPos - GetCurrentBrushRadius() * FVector3d::One(),
+		BrushPos + GetCurrentBrushRadius() * FVector3d::One());
 
 	VertexSetBuffer.Reset();
 	TriangleROI.Reset();
-	FDynamicMesh3* Mesh = DynamicMeshComponent->GetMesh();
+	FDynamicMesh3* Mesh = GetSculptMesh();
 	Octree.RangeQuery(BrushBox,
 		[&](int TriIdx) {
 		FIndex3i TriV = Mesh->GetTriangle(TriIdx);
@@ -363,28 +308,14 @@ void UMeshVertexSculptTool::UpdateROI(const FVector3d& BrushPos)
 	VertexROI.SetNum(0, false);
 	BufferUtil::AppendElements(VertexROI, VertexSetBuffer);
 
-	int NumV = VertexROI.Num();
-	ROIPositionBuffer.SetNum(NumV, false);
+	ROIPositionBuffer.SetNum(VertexROI.Num(), false);
 }
 
-void UMeshVertexSculptTool::OnUpdateDrag(const FRay& WorldRay)
-{
-	if (bInDrag)
-	{
-		PendingStampRay = WorldRay;
-		bStampPending = true;
-	}
-}
-
-
 bool UMeshVertexSculptTool::UpdateStampPosition(const FRay& WorldRay)
 {
-	StampTimestamp++;
-
 	CalculateBrushRadius();
 
-	TUniquePtr<FMeshSculptBrushOp>& UseBrushOp = (bSmoothing) ? SecondaryBrushOp : PrimaryBrushOp;
-	UseBrushOp->Falloff = ActiveFalloff;
+	TUniquePtr<FMeshSculptBrushOp>& UseBrushOp = GetActiveBrushOp();
 
 	ESculptBrushOpTargetType TargetType = UseBrushOp->GetBrushTargetType();
 	switch (TargetType)
@@ -406,10 +337,10 @@ bool UMeshVertexSculptTool::UpdateStampPosition(const FRay& WorldRay)
 	}
 
 	CurrentStamp = LastStamp;
-	CurrentStamp.WorldFrame = FFrame3d(LastBrushPosWorld, LastBrushPosNormalWorld);
-	CurrentStamp.LocalFrame = CurrentStamp.WorldFrame;
-	CurrentStamp.LocalFrame.Transform(CurTargetTransform.Inverse());
-	CurrentStamp.Power = GetActivePressure() * ((bSmoothing) ? SculptProperties->SmoothBrushSpeed : SculptProperties->PrimaryBrushSpeed);
+	CurrentStamp.DeltaTime = FMathd::Min((FDateTime::Now() - LastStamp.TimeStamp).GetTotalSeconds(), 1.0);
+	CurrentStamp.WorldFrame = GetBrushFrameWorld();
+	CurrentStamp.LocalFrame = GetBrushFrameLocal();
+	CurrentStamp.Power = GetActivePressure() * GetCurrentBrushStrength();
 
 	CurrentStamp.PrevLocalFrame = LastStamp.LocalFrame;
 	CurrentStamp.PrevWorldFrame = LastStamp.WorldFrame;
@@ -420,8 +351,6 @@ bool UMeshVertexSculptTool::UpdateStampPosition(const FRay& WorldRay)
 		return false;
 	}
 
-	LastBrushPosLocal = CurrentStamp.LocalFrame.Origin;
-
 	return true;
 }
 
@@ -430,55 +359,39 @@ void UMeshVertexSculptTool::ApplyStamp()
 {
 	SCOPE_CYCLE_COUNTER(VtxSculptToolApplyStamp);
 
-	TUniquePtr<FMeshSculptBrushOp>& UseBrushOp = (bSmoothing) ? SecondaryBrushOp : PrimaryBrushOp;
+	TUniquePtr<FMeshSculptBrushOp>& UseBrushOp = GetActiveBrushOp();
 	if (UseBrushOp->WantsStampRegionPlane())
 	{
-		CurrentStamp.RegionPlane = ComputeROIBrushPlane(CurrentStamp.LocalFrame.Origin, true, false);
+		CurrentStamp.RegionPlane = ComputeStampRegionPlane(CurrentStamp.LocalFrame, TriangleROI, true, false, false);
 	}
 
-	FDynamicMesh3* Mesh = DynamicMeshComponent->GetMesh();
+	FDynamicMesh3* Mesh = GetSculptMesh();
 	UseBrushOp->ApplyStamp(Mesh, CurrentStamp, VertexROI, ROIPositionBuffer);
 
 	SyncMeshWithPositionBuffer(Mesh);
 
 	LastStamp = CurrentStamp;
+	LastStamp.TimeStamp = FDateTime::Now();
 }
 
 
 
-double UMeshVertexSculptTool::CalculateBrushFalloff(double Distance)
-{
-	double f = FMathd::Clamp(1.0 - BrushProperties->BrushFalloffAmount, 0.0, 1.0);
-	double d = Distance / CurrentBrushRadius;
-	double w = 1;
-	if (d > f)
-	{
-		d = FMathd::Clamp((d - f) / (1.0 - f), 0.0, 1.0);
-		w = (1.0 - d * d);
-		w = w * w * w;
-	}
-	return w;
-}
-
-
 
 void UMeshVertexSculptTool::SyncMeshWithPositionBuffer(FDynamicMesh3* Mesh)
 {
-	const int NumV = ROIPositionBuffer.Num();
+	const int32 NumV = ROIPositionBuffer.Num();
 	checkSlow(VertexROI.Num() <= NumV);
 
-	// actually doing this in parallel is slower here due to locking (?)
-	ParallelFor(NumV, [&](int32 k)
+	// change update could be async here if we collected array of <idx,orig,new> and dispatched independenlty
+	for ( int32 k = 0; k < NumV; ++k)
 	{
 		int VertIdx = VertexROI[k];
 		const FVector3d& NewPos = ROIPositionBuffer[k];
 		FVector3d OrigPos = Mesh->GetVertex(VertIdx);
 		Mesh->SetVertex(VertIdx, NewPos);
 
-		//UpdateSavedVertexLock.Lock();
 		ActiveVertexChange->UpdateVertex(VertIdx, OrigPos, NewPos);
-		//UpdateSavedVertexLock.Unlock();
-	}, true);
+	}
 }
 
 
@@ -486,52 +399,15 @@ void UMeshVertexSculptTool::SyncMeshWithPositionBuffer(FDynamicMesh3* Mesh)
 
 
 
-
-
-FFrame3d UMeshVertexSculptTool::ComputeROIBrushPlane(const FVector3d& BrushCenter, bool bIgnoreDepth, bool bViewAligned)
+int32 UMeshVertexSculptTool::FindHitSculptMeshTriangle(const FRay3d& LocalRay)
 {
-	FDynamicMesh3* Mesh = DynamicMeshComponent->GetMesh();
-	FVector3d AverageNormal(0, 0, 0);
-	FVector3d AveragePos(0, 0, 0);
-	double WeightSum = 0;
-	for (int TriID : TriangleROI)
-	{
-		FVector3d Centroid = Mesh->GetTriCentroid(TriID);
-		double Weight = CalculateBrushFalloff(BrushCenter.Distance(Centroid));
-
-		AverageNormal += Weight * Mesh->GetTriNormal(TriID);
-		AveragePos += Weight * Centroid;
-		WeightSum += Weight;
-	}
-	AverageNormal.Normalize();
-	AveragePos /= WeightSum;
-
-	if (bViewAligned)
-	{
-		AverageNormal = -CameraState.Forward();
-	}
-
-	FFrame3d Result = FFrame3d(AveragePos, AverageNormal);
-	if (bIgnoreDepth == false)
-	{
-		Result.Origin -= BrushProperties->Depth * CurrentBrushRadius * Result.Z();
-	}
-
-	return Result;
-}
-
-
-
-
-int UMeshVertexSculptTool::FindHitSculptMeshTriangle(const FRay3d& LocalRay)
-{
-	if (BrushProperties->bHitBackFaces)
+	if (GetBrushCanHitBackFaces())
 	{
 		return Octree.FindNearestHitObject(LocalRay);
 	}
 	else
 	{
-		FDynamicMesh3* Mesh = DynamicMeshComponent->GetMesh();
+		FDynamicMesh3* Mesh = GetSculptMesh();
 
 		FViewCameraState StateOut;
 		GetToolManager()->GetContextQueriesAPI()->GetCurrentViewState(StateOut);
@@ -547,13 +423,13 @@ int UMeshVertexSculptTool::FindHitSculptMeshTriangle(const FRay3d& LocalRay)
 	}
 }
 
-int UMeshVertexSculptTool::FindHitTargetMeshTriangle(const FRay3d& LocalRay)
+int32 UMeshVertexSculptTool::FindHitTargetMeshTriangle(const FRay3d& LocalRay)
 {
 	IMeshSpatial::FQueryOptions RaycastOptions;
 
-	if (BrushProperties->bHitBackFaces == false)
+	if (GetBrushCanHitBackFaces())
 	{
-		FDynamicMesh3* Mesh = DynamicMeshComponent->GetMesh();
+		FDynamicMesh3* Mesh = GetSculptMesh();
 
 		FViewCameraState StateOut;
 		GetToolManager()->GetContextQueriesAPI()->GetCurrentViewState(StateOut);
@@ -574,119 +450,31 @@ int UMeshVertexSculptTool::FindHitTargetMeshTriangle(const FRay3d& LocalRay)
 	}
 }
 
-bool UMeshVertexSculptTool::UpdateBrushPositionOnActivePlane(const FRay& WorldRay)
-{
-	FVector3d NewHitPosWorld;
-	ActiveDragPlane.RayPlaneIntersection(WorldRay.Origin, WorldRay.Direction, 2, NewHitPosWorld);
-	LastBrushPosWorld = NewHitPosWorld;
-	LastBrushPosNormalWorld = ActiveDragPlane.Z();
-	return true;
-}
-
-bool UMeshVertexSculptTool::UpdateBrushPositionOnTargetMesh(const FRay& WorldRay, bool bFallbackToViewPlane)
-{
-	FRay3d LocalRay(CurTargetTransform.InverseTransformPosition(WorldRay.Origin),
-		CurTargetTransform.InverseTransformVector(WorldRay.Direction));
-	LocalRay.Direction.Normalize();
-
-	int32 HitTID = FindHitTargetMeshTriangle(LocalRay);
-	if (HitTID != IndexConstants::InvalidID)
-	{
-		FIntrRay3Triangle3d Query = TMeshQueries<FDynamicMesh3>::TriangleIntersection(BaseMesh, HitTID, LocalRay);
-		LastBrushPosNormalWorld = CurTargetTransform.TransformNormal(BaseMesh.GetTriNormal(HitTID));
-		LastBrushPosWorld = CurTargetTransform.TransformPosition(LocalRay.PointAt(Query.RayParameter));
-		return true;
-	}
-
-	if (bFallbackToViewPlane)
-	{
-		FFrame3d BrushPlane(LastBrushPosWorld, CameraState.Forward());
-		FVector3d NewHitPosWorld;
-		BrushPlane.RayPlaneIntersection(WorldRay.Origin, WorldRay.Direction, 2, NewHitPosWorld);
-		LastBrushPosWorld = NewHitPosWorld;
-		LastBrushPosNormalWorld = ActiveDragPlane.Z();
-		return true;
-	}
-
-	return false;
-}
-
-bool UMeshVertexSculptTool::UpdateBrushPositionOnSculptMesh(const FRay& WorldRay, bool bFallbackToViewPlane)
-{
-	FRay3d LocalRay(CurTargetTransform.InverseTransformPosition(WorldRay.Origin),
-		CurTargetTransform.InverseTransformVector(WorldRay.Direction));
-	LocalRay.Direction.Normalize();
-
-	int32 HitTID = FindHitSculptMeshTriangle(LocalRay);
-	if (HitTID != IndexConstants::InvalidID)
-	{
-		const FDynamicMesh3* SculptMesh = DynamicMeshComponent->GetMesh();
-		FIntrRay3Triangle3d Query = TMeshQueries<FDynamicMesh3>::TriangleIntersection(*SculptMesh, HitTID, LocalRay);
-		LastBrushPosNormalWorld = CurTargetTransform.TransformNormal(SculptMesh->GetTriNormal(HitTID));
-		LastBrushPosWorld = CurTargetTransform.TransformPosition(LocalRay.PointAt(Query.RayParameter));
-		return true;
-	}
-
-	if (bFallbackToViewPlane)
-	{
-		FFrame3d BrushPlane(LastBrushPosWorld, CameraState.Forward());
-		FVector3d NewHitPosWorld;
-		BrushPlane.RayPlaneIntersection(WorldRay.Origin, WorldRay.Direction, 2, NewHitPosWorld);
-		LastBrushPosWorld = NewHitPosWorld;
-		LastBrushPosNormalWorld = ActiveDragPlane.Z();
-		return true;
-	}
-
-	return false;
-}
-
-void UMeshVertexSculptTool::AlignBrushToView()
-{
-	LastBrushPosNormalWorld = -CameraState.Forward();
-}
 
 
 bool UMeshVertexSculptTool::UpdateBrushPosition(const FRay& WorldRay)
 {
-	// This is an unfortunate hack necessary because we haven't refactored brushes properly yet
+	TUniquePtr<FMeshSculptBrushOp>& UseBrushOp = GetActiveBrushOp();
 
-	if (bSmoothing)
-	{
-		return UpdateBrushPositionOnSculptMesh(WorldRay, false);
-	}
-
-	bool bHit = false;
-	switch (SculptProperties->PrimaryBrushType)
+	bool bHit = false; 
+	ESculptBrushOpTargetType TargetType = UseBrushOp->GetBrushTargetType();
+	switch (TargetType)
 	{
-	case EMeshVertexSculptBrushType::Offset:
-	case EMeshVertexSculptBrushType::SculptMax:
-	case EMeshVertexSculptBrushType::Pinch:
-		bHit = UpdateBrushPositionOnTargetMesh(WorldRay, false);
+	case ESculptBrushOpTargetType::SculptMesh:
+		bHit = UpdateBrushPositionOnSculptMesh(WorldRay, false);
 		break;
-
-	case EMeshVertexSculptBrushType::SculptView:
-	case EMeshVertexSculptBrushType::PlaneViewAligned:
+	case ESculptBrushOpTargetType::TargetMesh:
 		bHit = UpdateBrushPositionOnTargetMesh(WorldRay, false);
-		AlignBrushToView();
-		break;
-
-	case EMeshVertexSculptBrushType::Move:
-		//return UpdateBrushPositionOnActivePlane(WorldRay);
-		bHit = UpdateBrushPositionOnSculptMesh(WorldRay, false);
 		break;
-
-	case EMeshVertexSculptBrushType::Smooth:
-	case EMeshVertexSculptBrushType::Inflate:
-	case EMeshVertexSculptBrushType::Flatten:
-	case EMeshVertexSculptBrushType::Plane:
-	case EMeshVertexSculptBrushType::FixedPlane:
+	case ESculptBrushOpTargetType::ActivePlane:
+		//UpdateBrushPositionOnActivePlane(WorldRay);
 		bHit = UpdateBrushPositionOnSculptMesh(WorldRay, false);
 		break;
+	}
 
-	default:
-		UE_LOG(LogTemp, Warning, TEXT("UMeshVertexSculptTool: unknown brush type in UpdateBrushPosition"));
-		bHit = UpdateBrushPositionOnSculptMesh(WorldRay, false);
-		break;
+	if (bHit && UseBrushOp->GetAlignStampToView())
+	{
+		AlignBrushToView();
 	}
 
 	return bHit;
@@ -694,41 +482,12 @@ bool UMeshVertexSculptTool::UpdateBrushPosition(const FRay& WorldRay)
 
 
 
-void UMeshVertexSculptTool::OnEndDrag(const FRay& Ray)
-{
-	bInDrag = false;
-
-	// cancel these! otherwise change record could become invalid
-	bStampPending = false;
-
-	// update spatial
-	bTargetDirty = true;
-
-	TUniquePtr<FMeshSculptBrushOp>& UseBrushOp = (bSmoothing) ? SecondaryBrushOp : PrimaryBrushOp;
-	UseBrushOp->EndStroke(DynamicMeshComponent->GetMesh(), LastStamp, VertexROI);
-
-	// close change record
-	EndChange();
-}
-
-
-FInputRayHit UMeshVertexSculptTool::BeginHoverSequenceHitTest(const FInputDeviceRay& PressPos)
-{
-	return UMeshSurfacePointTool::BeginHoverSequenceHitTest(PressPos);
-}
 
 bool UMeshVertexSculptTool::OnUpdateHover(const FInputDeviceRay& DevicePos)
 {
 	PendingStampType = SculptProperties->PrimaryBrushType;
 
-	if (bInDrag)
-	{
-		FVector3d NewHitPosWorld;
-		ActiveDragPlane.RayPlaneIntersection(DevicePos.WorldRay.Origin, DevicePos.WorldRay.Direction, 2, NewHitPosWorld);
-		LastBrushPosWorld = NewHitPosWorld;
-		LastBrushPosNormalWorld = ActiveDragPlane.Z();
-	}
-	else
+	if(ensure(InStroke() == false))
 	{
 		UpdateBrushPosition(DevicePos.WorldRay);
 	}
@@ -743,12 +502,6 @@ void UMeshVertexSculptTool::OnTick(float DeltaTime)
 
 	SCOPE_CYCLE_COUNTER(VtxSculptToolTick);
 
-	UpdateHoverStamp(LastBrushPosWorld, LastBrushPosNormalWorld);
-
-	BrushTypeWatcher.CheckAndUpdate();
-
-	UpdateWorkPlane();
-
 	// process the undo update
 	if (bUndoUpdatePending)
 	{
@@ -773,15 +526,15 @@ void UMeshVertexSculptTool::OnTick(float DeltaTime)
 	}
 
 
-	if (bStampPending)
+	if (IsStampPending())
 	{
 		//UE_LOG(LogTemp, Warning, TEXT("dt is %.3f, tick fps %.2f - roi size %d/%d"), DeltaTime, 1.0 / DeltaTime, VertexROI.Num(), TriangleROI.Num());
 		SCOPE_CYCLE_COUNTER(VtxSculptTool_Tick_ApplyStampBlock);
-		// hardcoded brush flow (todo: make optional)
-		bStampPending = bInDrag;
+
+		ApplyStrokeFlowInTick();
 
 		// update brush position
-		if (UpdateStampPosition(PendingStampRay) == false)
+		if (UpdateStampPosition(GetPendingStampRayWorld()) == false)
 		{
 			return;
 		}
@@ -807,7 +560,7 @@ void UMeshVertexSculptTool::OnTick(float DeltaTime)
 
 		// recalculate normals. This has to complete before we can update component
 		// (in fact we could do it per-chunk...)
-		FDynamicMesh3* Mesh = DynamicMeshComponent->GetMesh();
+		FDynamicMesh3* Mesh = GetSculptMesh();
 		{
 			SCOPE_CYCLE_COUNTER(VtxSculptTool_Tick_NormalsBlock);
 			UE::SculptUtil::RecalculateROINormals(Mesh, TriangleROI, VertexSetBuffer, NormalsBuffer);
@@ -827,7 +580,7 @@ void UMeshVertexSculptTool::OnTick(float DeltaTime)
 	else if (bTargetDirty)
 	{
 		SCOPE_CYCLE_COUNTER(VtxSculptTool_Tick_UpdateTargetBlock);
-		check(bInDrag == false);
+		check(InStroke() == false);
 
 		// this spawns futures that we could allow to run while other things happen...
 		UpdateBaseMesh(&AccumulatedTriangleROI);
@@ -852,7 +605,7 @@ void UMeshVertexSculptTool::UpdateBaseMesh(const TSet<int32>* TriangleSet)
 		}
 	}
 
-	const FDynamicMesh3* SculptMesh = DynamicMeshComponent->GetMesh();
+	const FDynamicMesh3* SculptMesh = GetSculptMesh();
 	if ( ! TriangleSet )
 	{
 		BaseMesh.Copy(*SculptMesh, false, false, false, false);
@@ -892,17 +645,6 @@ bool UMeshVertexSculptTool::GetBaseMeshNearest(int32 VertexID, const FVector3d&
 	TargetPosOut = BaseMesh.GetVertex(VertexID);
 	TargetNormalOut = (FVector3d)BaseMesh.GetVertexNormal(VertexID);
 	return true;
-
-	// Octree FindNearestObject is currently quite expensive...
-	//int32 NearestTID = BaseMeshSpatial.FindNearestObject(Position, SearchRadius);
-	//if (NearestTID <= 0)
-	//{
-	//	return false;
-	//}
-	//FDistPoint3Triangle3d DistanceQuery = TMeshQueries<FDynamicMesh3>::TriangleDistance(BaseMesh, NearestTID, Position);
-	//TargetPosOut = DistanceQuery.ClosestTrianglePoint;
-	//TargetNormalOut = BaseMesh.GetTriBaryNormal(NearestTID, DistanceQuery.TriangleBaryCoords.X, DistanceQuery.TriangleBaryCoords.Y, DistanceQuery.TriangleBaryCoords.Z);
-	//return true;
 }
 
 
@@ -912,12 +654,12 @@ bool UMeshVertexSculptTool::GetBaseMeshNearest(int32 VertexID, const FVector3d&
 
 void UMeshVertexSculptTool::IncreaseBrushSpeedAction()
 {
-	SculptProperties->PrimaryBrushSpeed = FMath::Clamp(SculptProperties->PrimaryBrushSpeed + 0.05f, 0.0f, 1.0f);
+	//SculptProperties->PrimaryBrushSpeed = FMath::Clamp(SculptProperties->PrimaryBrushSpeed + 0.05f, 0.0f, 1.0f);
 }
 
 void UMeshVertexSculptTool::DecreaseBrushSpeedAction()
 {
-	SculptProperties->PrimaryBrushSpeed = FMath::Clamp(SculptProperties->PrimaryBrushSpeed - 0.05f, 0.0f, 1.0f);
+	//SculptProperties->PrimaryBrushSpeed = FMath::Clamp(SculptProperties->PrimaryBrushSpeed - 0.05f, 0.0f, 1.0f);
 }
 
 
@@ -985,7 +727,7 @@ void UMeshVertexSculptTool::WaitForPendingUndoRedo()
 void UMeshVertexSculptTool::OnDynamicMeshComponentChanged(USimpleDynamicMeshComponent* Component, const FMeshVertexChange* Change, bool bRevert)
 {
 	// update octree
-	FDynamicMesh3* Mesh = DynamicMeshComponent->GetMesh();
+	FDynamicMesh3* Mesh = GetSculptMesh();
 
 	// make sure any previous async computations are done, and update the undo ROI
 	if (bUndoUpdatePending)
@@ -1032,80 +774,20 @@ void UMeshVertexSculptTool::OnDynamicMeshComponentChanged(USimpleDynamicMeshComp
 
 
 
-
-
-
-
 void UMeshVertexSculptTool::UpdateBrushType(EMeshVertexSculptBrushType BrushType)
 {
 	static const FText BaseMessage = LOCTEXT("OnStartSculptTool", "Hold Shift to Smooth, Ctrl to Invert (where applicable). Q/A keys cycle Brush Type. S/D changes Size (+Shift to small-step), W/E changes Strength.");
 	FTextBuilder Builder;
 	Builder.AppendLine(BaseMessage);
 
-	TUniqueFunction<bool(int32, const FVector3d&, double MaxDist, FVector3d&, FVector3d&)> BaseMeshQueryFunc =
-		[&](int32 VertexID, const FVector3d& Position, double MaxDist, FVector3d& PosOut, FVector3d& NormalOut)
-		{
-			//return GetTargetMeshNearest(VertexID, Position, MaxDist, PosOut, NormalOut);
-			return GetBaseMeshNearest(VertexID, Position, MaxDist, PosOut, NormalOut);
-		};
-
-	switch (BrushType)
-	{
-	case EMeshVertexSculptBrushType::Smooth:
-		PrimaryBrushOp = MakeUnique<FSmoothBrushOp>();
-		break;
-	case EMeshVertexSculptBrushType::Offset:
-		PrimaryBrushOp = MakeUnique<FSurfaceSculptBrushOp>(MoveTemp(BaseMeshQueryFunc));
-		break;
-	case EMeshVertexSculptBrushType::SculptView:
-		PrimaryBrushOp = MakeUnique<FViewAlignedSculptBrushOp>(MoveTemp(BaseMeshQueryFunc));
-		break;
-	case EMeshVertexSculptBrushType::SculptMax:
-		PrimaryBrushOp = MakeUnique<FSurfaceMaxSculptBrushOp>(MoveTemp(BaseMeshQueryFunc));
-		break;
-	case EMeshVertexSculptBrushType::Inflate:
-		PrimaryBrushOp = MakeUnique<FInflateBrushOp>();
-		break;
-	case EMeshVertexSculptBrushType::Move:
-		PrimaryBrushOp = MakeUnique<FMoveBrushOp>();
-		break;
-	case EMeshVertexSculptBrushType::Flatten:
-		PrimaryBrushOp = MakeUnique<FFlattenBrushOp>();
-		break;
-	case EMeshVertexSculptBrushType::Plane:
-		PrimaryBrushOp = MakeUnique<FPlaneBrushOp>();
-		break;
-	case EMeshVertexSculptBrushType::PlaneViewAligned:
-		PrimaryBrushOp = MakeUnique<FPlaneBrushOp>();
-		break;
-	case EMeshVertexSculptBrushType::FixedPlane:
-		PrimaryBrushOp = MakeUnique<FPlaneBrushOp>();
-		break;
-	case EMeshVertexSculptBrushType::Pinch:
-		PrimaryBrushOp = MakeUnique<FPinchBrushOp>();
-		break;
-	}
-	PrimaryBrushOp->Falloff = ActiveFalloff;
-
+	SetActivePrimaryBrushType((int32)BrushType);
 
 	SetToolPropertySourceEnabled(GizmoProperties, false);
-	SetToolPropertySourceEnabled(PlaneBrushProperties, false);
-	SetToolPropertySourceEnabled(SculptMaxBrushProperties, false);
-
 	if (BrushType == EMeshVertexSculptBrushType::FixedPlane)
 	{
 		Builder.AppendLine(LOCTEXT("FixedPlaneTip", "Use T to reposition Work Plane at cursor, Shift+T to align to Normal, Ctrl+Shift+T to align to View"));
-		SetToolPropertySourceEnabled(PlaneBrushProperties, true);
 		SetToolPropertySourceEnabled(GizmoProperties, true);
 	}
-	if (BrushType == EMeshVertexSculptBrushType::Plane || BrushType == EMeshVertexSculptBrushType::PlaneViewAligned || BrushType == EMeshVertexSculptBrushType::Flatten)
-	{
-		SetToolPropertySourceEnabled(PlaneBrushProperties, true);
-	}
-	if (BrushType == EMeshVertexSculptBrushType::SculptMax)
-	{
-		SetToolPropertySourceEnabled(SculptMaxBrushProperties, true);
-	}
 
 	GetToolManager()->DisplayMessage(Builder.ToText(), EToolMessageLevel::UserNotification);
 }
diff --git a/Engine/Plugins/Experimental/MeshModelingToolset/Source/MeshModelingTools/Private/Sculpting/MeshSculptToolBase.cpp b/Engine/Plugins/Experimental/MeshModelingToolset/Source/MeshModelingTools/Private/Sculpting/MeshSculptToolBase.cpp
--- a/Engine/Plugins/Experimental/MeshModelingToolset/Source/MeshModelingTools/Private/Sculpting/MeshSculptToolBase.cpp
+++ b/Engine/Plugins/Experimental/MeshModelingToolset/Source/MeshModelingTools/Private/Sculpting/MeshSculptToolBase.cpp
@@ -11,7 +11,7 @@
 #include "BaseGizmos/TransformGizmo.h"
 #include "Drawing/MeshDebugDrawing.h"
 
-
+#include "Sculpting/StampFalloffs.h"
 
 #include "Generators/SphereGenerator.h"
 
@@ -23,7 +23,7 @@
 
 namespace
 {
-	const FString VertexSculptIndicatorGizmoType = TEXT("BrushIndicatorGizmoType");
+	const FString VertexSculptIndicatorGizmoType = TEXT("VertexSculptIndicatorGizmoType");
 }
 
 
@@ -82,6 +82,17 @@ void UMeshSculptToolBase::Setup()
 }
 
 
+void UMeshSculptToolBase::OnCompleteSetup()
+{
+	RestoreAllBrushTypeProperties(this);
+
+	for (auto Pair : BrushOpPropSets)
+	{
+		SetToolPropertySourceEnabled(Pair.Value, false);
+	}
+}
+
+
 void UMeshSculptToolBase::Shutdown(EToolShutdownType ShutdownType)
 {
 	UMeshSurfacePointTool::Shutdown(ShutdownType);
@@ -100,6 +111,36 @@ void UMeshSculptToolBase::Shutdown(EToolShutdownType ShutdownType)
 	}
 
 	ViewProperties->SaveProperties(this);
+
+	SaveAllBrushTypeProperties(this);
+
+
+	// bake result
+	UBaseDynamicMeshComponent* DynamicMeshComponent = GetSculptMeshComponent();
+	if (DynamicMeshComponent != nullptr)
+	{
+		ComponentTarget->SetOwnerVisibility(true);
+
+		if (ShutdownType == EToolShutdownType::Accept)
+		{
+			// safe to do this here because we are about to destroy componeont
+			DynamicMeshComponent->ApplyTransform(InitialTargetTransform, true);
+
+			// this block bakes the modified DynamicMeshComponent back into the StaticMeshComponent inside an undo transaction
+			GetToolManager()->BeginUndoTransaction(LOCTEXT("SculptMeshToolTransactionName", "Sculpt Mesh"));
+			ComponentTarget->CommitMesh([=](const FPrimitiveComponentTarget::FCommitParams& CommitParams)
+			{
+				FConversionToMeshDescriptionOptions ConversionOptions;
+				DynamicMeshComponent->Bake(CommitParams.MeshDescription, false, ConversionOptions);
+			});
+			GetToolManager()->EndUndoTransaction();
+		}
+
+		DynamicMeshComponent->UnregisterComponent();
+		DynamicMeshComponent->DestroyComponent();
+		DynamicMeshComponent = nullptr;
+	}
+
 }
 
 
@@ -110,6 +151,22 @@ void UMeshSculptToolBase::OnTick(float DeltaTime)
 	GizmoRotationWatcher.CheckAndUpdate();
 
 	ActivePressure = GetCurrentDevicePressure();
+
+	if (InStroke() == false)
+	{
+		SaveActiveStrokeModifiers();
+	}
+
+	// update cached falloff
+	CurrentBrushFalloff = 0.5;
+	if (GetActiveBrushOp()->PropertySet.IsValid())
+	{
+		CurrentBrushFalloff = FMathd::Clamp(GetActiveBrushOp()->PropertySet->GetFalloff(), 0.0, 1.0);
+	}
+
+	UpdateHoverStamp(GetBrushFrameWorld());
+
+	UpdateWorkPlane();
 }
 
 
@@ -120,7 +177,7 @@ void UMeshSculptToolBase::Render(IToolsContextRenderAPI* RenderAPI)
 
 	BrushIndicator->Update((float)GetCurrentBrushRadius(),
 		(FVector)HoverStamp.WorldFrame.Origin, (FVector)HoverStamp.WorldFrame.Z(),
-		1.0f - BrushProperties->BrushFalloffAmount);
+		1.0f - (float)GetCurrentBrushFalloff());
 	if (BrushIndicatorMaterial)
 	{
 		double FixedDimScale = ToolSceneQueriesUtil::CalculateDimensionFromVisualAngleD(CameraState, HoverStamp.WorldFrame.Origin, 1.5f);
@@ -143,12 +200,463 @@ void UMeshSculptToolBase::Render(IToolsContextRenderAPI* RenderAPI)
 
 
 
+void UMeshSculptToolBase::InitializeSculptMeshComponent(UBaseDynamicMeshComponent* Component)
+{
+	Component->SetupAttachment(ComponentTarget->GetOwnerActor()->GetRootComponent());
+	Component->RegisterComponent();
+
+	// initialize from LOD-0 MeshDescription
+	Component->InitializeMesh(ComponentTarget->GetMesh());
+	double MaxDimension = Component->GetMesh()->GetCachedBounds().MaxDim();
+
+	// bake rotation and scaling into mesh because handling these inside sculpting is a mess
+	// Note: this transform does not include translation ( so only the 3x3 transform)
+	InitialTargetTransform = FTransform3d(ComponentTarget->GetWorldTransform());
+	// clamp scaling because if we allow zero-scale we cannot invert this transform on Accept
+	InitialTargetTransform.ClampMinimumScale(0.01);
+	FVector3d Translation = InitialTargetTransform.GetTranslation();
+	InitialTargetTransform.SetTranslation(FVector3d::Zero());
+	Component->ApplyTransform(InitialTargetTransform, false);
+	CurTargetTransform = FTransform3d(Translation);
+	Component->SetWorldTransform((FTransform)CurTargetTransform);
+
+	// hide input Component
+	ComponentTarget->SetOwnerVisibility(false);
+}
+
+
+
+
+
+
+void UMeshSculptToolBase::RegisterBrushType(int32 Identifier, TUniquePtr<FMeshSculptBrushOpFactory> Factory, UMeshSculptBrushOpProps* PropSet)
+{
+	check(BrushOpPropSets.Contains(Identifier) == false && BrushOpFactories.Contains(Identifier) == false);
+	BrushOpPropSets.Add(Identifier, PropSet);
+	BrushOpFactories.Add(Identifier, MoveTemp(Factory));
+
+	AddToolPropertySource(PropSet);
+	SetToolPropertySourceEnabled(PropSet, false);
+}
+
+void UMeshSculptToolBase::RegisterSecondaryBrushType(int32 Identifier, TUniquePtr<FMeshSculptBrushOpFactory> Factory, UMeshSculptBrushOpProps* PropSet)
+{
+	check(SecondaryBrushOpPropSets.Contains(Identifier) == false && SecondaryBrushOpFactories.Contains(Identifier) == false);
+	SecondaryBrushOpPropSets.Add(Identifier, PropSet);
+	SecondaryBrushOpFactories.Add(Identifier, MoveTemp(Factory));
+
+	AddToolPropertySource(PropSet);
+	SetToolPropertySourceEnabled(PropSet, false);
+}
+
+
+void UMeshSculptToolBase::SaveAllBrushTypeProperties(UInteractiveTool* SaveFromTool)
+{
+	for (auto Pair : BrushOpPropSets)
+	{
+		Pair.Value->SaveProperties(SaveFromTool);
+	}
+	for (auto Pair : SecondaryBrushOpPropSets)
+	{
+		Pair.Value->SaveProperties(SaveFromTool);
+	}
+}
+void UMeshSculptToolBase::RestoreAllBrushTypeProperties(UInteractiveTool* RestoreToTool)
+{
+	for (auto Pair : BrushOpPropSets)
+	{
+		Pair.Value->RestoreProperties(RestoreToTool);
+	}
+	for (auto Pair : SecondaryBrushOpPropSets)
+	{
+		Pair.Value->RestoreProperties(RestoreToTool);
+	}
+}
+
+
+void UMeshSculptToolBase::SetActivePrimaryBrushType(int32 Identifier)
+{
+	TUniquePtr<FMeshSculptBrushOpFactory>* Factory = BrushOpFactories.Find(Identifier);
+	if (Factory == nullptr)
+	{
+		check(false);
+		return;
+	}
+
+	if (PrimaryVisiblePropSet != nullptr)
+	{
+		SetToolPropertySourceEnabled(PrimaryVisiblePropSet, false);
+		PrimaryVisiblePropSet = nullptr;
+	}
+
+	PrimaryBrushOp = (*Factory)->Build();
+	PrimaryBrushOp->Falloff = PrimaryFalloff;
+
+	UMeshSculptBrushOpProps** FoundProps = BrushOpPropSets.Find(Identifier);
+	if (FoundProps != nullptr)
+	{
+		SetToolPropertySourceEnabled(*FoundProps, true);
+		PrimaryVisiblePropSet = *FoundProps;
+
+		PrimaryBrushOp->PropertySet = PrimaryVisiblePropSet;
+	}
+}
+
+
+
+void UMeshSculptToolBase::SetActiveSecondaryBrushType(int32 Identifier)
+{
+	TUniquePtr<FMeshSculptBrushOpFactory>* Factory = SecondaryBrushOpFactories.Find(Identifier);
+	if (Factory == nullptr)
+	{
+		check(false);
+		return;
+	}
+
+	if (SecondaryVisiblePropSet != nullptr)
+	{
+		SetToolPropertySourceEnabled(SecondaryVisiblePropSet, false);
+		SecondaryVisiblePropSet = nullptr;
+	}
+
+	SecondaryBrushOp = (*Factory)->Build();
+	TSharedPtr<FMeshSculptFallofFunc> SecondaryFalloff = MakeShared<FMeshSculptFallofFunc>();;
+	SecondaryFalloff->FalloffFunc = UE::SculptFalloffs::MakeStandardSmoothFalloff();
+	SecondaryBrushOp->Falloff = SecondaryFalloff;
+
+	UMeshSculptBrushOpProps** FoundProps = SecondaryBrushOpPropSets.Find(Identifier);
+	if (FoundProps != nullptr)
+	{
+		SetToolPropertySourceEnabled(*FoundProps, true);
+		SecondaryVisiblePropSet = *FoundProps;
+
+		SecondaryBrushOp->PropertySet = SecondaryVisiblePropSet;
+	}
+}
+
+
+
+TUniquePtr<FMeshSculptBrushOp>& UMeshSculptToolBase::GetActiveBrushOp()
+{
+	if (GetInSmoothingStroke())
+	{
+		return SecondaryBrushOp;
+	}
+	else
+	{
+		return PrimaryBrushOp;
+	}
+}
+
+
+
+
+void UMeshSculptToolBase::SetPrimaryFalloffType(EMeshSculptFalloffType FalloffType)
+{
+	PrimaryFalloff = MakeShared<FMeshSculptFallofFunc>();
+	switch (FalloffType)
+	{
+	default:
+	case EMeshSculptFalloffType::Smooth:
+		PrimaryFalloff->FalloffFunc = UE::SculptFalloffs::MakeStandardSmoothFalloff();
+		break;
+	case EMeshSculptFalloffType::Linear:
+		PrimaryFalloff->FalloffFunc = UE::SculptFalloffs::MakeLinearFalloff();
+		break;
+	case EMeshSculptFalloffType::Inverse:
+		PrimaryFalloff->FalloffFunc = UE::SculptFalloffs::MakeInverseFalloff();
+		break;
+	case EMeshSculptFalloffType::Round:
+		PrimaryFalloff->FalloffFunc = UE::SculptFalloffs::MakeRoundFalloff();
+		break;
+	case EMeshSculptFalloffType::BoxSmooth:
+		PrimaryFalloff->FalloffFunc = UE::SculptFalloffs::MakeSmoothBoxFalloff();
+		break;
+	case EMeshSculptFalloffType::BoxLinear:
+		PrimaryFalloff->FalloffFunc = UE::SculptFalloffs::MakeLinearBoxFalloff();
+		break;
+	case EMeshSculptFalloffType::BoxInverse:
+		PrimaryFalloff->FalloffFunc = UE::SculptFalloffs::MakeInverseBoxFalloff();
+		break;
+	case EMeshSculptFalloffType::BoxRound:
+		PrimaryFalloff->FalloffFunc = UE::SculptFalloffs::MakeRoundBoxFalloff();
+		break;
+	}
+}
+
 
 
 
-void UMeshSculptToolBase::UpdateHoverStamp(const FVector3d& WorldPos, const FVector3d& WorldNormal)
+bool UMeshSculptToolBase::HitTest(const FRay& Ray, FHitResult& OutHit)
 {
-	HoverStamp.WorldFrame = FFrame3d(WorldPos, WorldNormal);
+	FRay3d LocalRay = GetLocalRay(Ray);
+
+	int HitTID = FindHitSculptMeshTriangle(LocalRay);
+	if (HitTID != IndexConstants::InvalidID)
+	{
+		FTriangle3d Triangle;
+		FDynamicMesh3* Mesh = GetSculptMesh();
+		Mesh->GetTriVertices(HitTID, Triangle.V[0], Triangle.V[1], Triangle.V[2]);
+		FIntrRay3Triangle3d Query(LocalRay, Triangle);
+		Query.Find();
+
+		OutHit.FaceIndex = HitTID;
+		OutHit.Distance = Query.RayParameter;
+		OutHit.Normal = (FVector)CurTargetTransform.TransformNormal(Mesh->GetTriNormal(HitTID));
+		OutHit.ImpactPoint = (FVector)CurTargetTransform.TransformPosition(LocalRay.PointAt(Query.RayParameter));
+		return true;
+	}
+	return false;
+}
+
+
+void UMeshSculptToolBase::OnBeginDrag(const FRay& WorldRay)
+{
+	SaveActiveStrokeModifiers();
+
+	FHitResult OutHit;
+	if (HitTest(WorldRay, OutHit))
+	{
+		bInStroke = true;
+
+		UpdateBrushTargetPlaneFromHit(WorldRay, OutHit);
+
+		// initialize first stamp
+		PendingStampRay = WorldRay;
+		bIsStampPending = true;
+
+		// set falloff
+		PrimaryBrushOp->Falloff = PrimaryFalloff;
+
+		OnBeginStroke(WorldRay);
+	}
+}
+
+void UMeshSculptToolBase::OnUpdateDrag(const FRay& WorldRay)
+{
+	if (InStroke())
+	{
+		PendingStampRay = WorldRay;
+		bIsStampPending = true;
+	}
+}
+
+void UMeshSculptToolBase::OnEndDrag(const FRay& Ray)
+{
+	bInStroke = false;
+
+	// cancel these! otherwise change record could become invalid
+	bIsStampPending = false;
+
+	OnEndStroke();
+
+}
+
+
+
+FRay3d UMeshSculptToolBase::GetLocalRay(const FRay& WorldRay) const
+{
+	FRay3d LocalRay(CurTargetTransform.InverseTransformPosition(WorldRay.Origin),
+		CurTargetTransform.InverseTransformVector(WorldRay.Direction));
+	LocalRay.Direction.Normalize();
+	return LocalRay;
+}
+
+
+
+void UMeshSculptToolBase::UpdateBrushFrameWorld(const FVector3d& NewPosition, const FVector3d& NewNormal)
+{
+	FFrame3d NewFrame = LastBrushFrameWorld;
+	NewFrame.Origin = NewPosition;
+	NewFrame.AlignAxis(2, NewNormal);
+	NewFrame.ConstrainedAlignPerpAxes();
+
+	if (InStroke() && BrushProperties->Lazyness > 0)
+	{
+		double t = FMathd::Lerp(1.0, 0.1, BrushProperties->Lazyness);
+		LastBrushFrameWorld.Origin = FVector3d::Lerp(LastBrushFrameWorld.Origin, NewFrame.Origin, t);
+		LastBrushFrameWorld.Rotation = FQuaterniond(LastBrushFrameWorld.Rotation, NewFrame.Rotation, t);
+	}
+	else
+	{
+		LastBrushFrameWorld = NewFrame;
+
+	}
+
+	LastBrushFrameLocal = LastBrushFrameWorld;
+	LastBrushFrameLocal.Transform(CurTargetTransform.Inverse());
+}
+
+void UMeshSculptToolBase::AlignBrushToView()
+{
+	UpdateBrushFrameWorld(GetBrushFrameWorld().Origin, -CameraState.Forward());
+}
+
+
+
+void UMeshSculptToolBase::UpdateBrushTargetPlaneFromHit(const FRay& WorldRay, const FHitResult& Hit)
+{
+	FVector3d WorldPosWithBrushDepth = WorldRay.PointAt(Hit.Distance) + GetCurrentBrushDepth() * GetCurrentBrushRadius() * WorldRay.Direction;
+	ActiveBrushTargetPlaneWorld = FFrame3d(WorldPosWithBrushDepth, -WorldRay.Direction);
+}
+
+bool UMeshSculptToolBase::UpdateBrushPositionOnActivePlane(const FRay& WorldRay)
+{
+	FVector3d NewHitPosWorld;
+	ActiveBrushTargetPlaneWorld.RayPlaneIntersection(WorldRay.Origin, WorldRay.Direction, 2, NewHitPosWorld);
+	UpdateBrushFrameWorld(NewHitPosWorld, ActiveBrushTargetPlaneWorld.Z());
+	return true;
+}
+
+bool UMeshSculptToolBase::UpdateBrushPositionOnTargetMesh(const FRay& WorldRay, bool bFallbackToViewPlane)
+{
+	FRay3d LocalRay = GetLocalRay(WorldRay);
+	int32 HitTID = FindHitTargetMeshTriangle(LocalRay);
+	if (HitTID != IndexConstants::InvalidID)
+	{
+		const FDynamicMesh3* BaseMesh = GetBaseMesh();
+		FIntrRay3Triangle3d Query = TMeshQueries<FDynamicMesh3>::TriangleIntersection(*BaseMesh, HitTID, LocalRay);
+		FVector3d WorldNormal = CurTargetTransform.TransformNormal(BaseMesh->GetTriNormal(HitTID));
+		FVector3d WorldPos = CurTargetTransform.TransformPosition(LocalRay.PointAt(Query.RayParameter));
+		UpdateBrushFrameWorld(WorldPos, WorldNormal);
+		return true;
+	}
+
+	if (bFallbackToViewPlane)
+	{
+		FFrame3d BrushPlane(GetBrushFrameWorld().Origin, CameraState.Forward());
+		FVector3d NewHitPosWorld;
+		BrushPlane.RayPlaneIntersection(WorldRay.Origin, WorldRay.Direction, 2, NewHitPosWorld);
+		UpdateBrushFrameWorld(NewHitPosWorld, ActiveBrushTargetPlaneWorld.Z());
+		return true;
+	}
+
+	return false;
+}
+
+bool UMeshSculptToolBase::UpdateBrushPositionOnSculptMesh(const FRay& WorldRay, bool bFallbackToViewPlane)
+{
+	FRay3d LocalRay = GetLocalRay(WorldRay);
+	int32 HitTID = FindHitSculptMeshTriangle(LocalRay);
+	if (HitTID != IndexConstants::InvalidID)
+	{
+		const FDynamicMesh3* SculptMesh = GetSculptMesh();
+		FIntrRay3Triangle3d Query = TMeshQueries<FDynamicMesh3>::TriangleIntersection(*SculptMesh, HitTID, LocalRay);
+		FVector3d WorldNormal = CurTargetTransform.TransformNormal(SculptMesh->GetTriNormal(HitTID));
+		FVector3d WorldPos = CurTargetTransform.TransformPosition(LocalRay.PointAt(Query.RayParameter));
+		UpdateBrushFrameWorld(WorldPos, WorldNormal);
+		return true;
+	}
+
+	if (bFallbackToViewPlane)
+	{
+		FFrame3d BrushPlane(GetBrushFrameWorld().Origin, CameraState.Forward());
+		FVector3d NewHitPosWorld;
+		BrushPlane.RayPlaneIntersection(WorldRay.Origin, WorldRay.Direction, 2, NewHitPosWorld);
+		UpdateBrushFrameWorld(NewHitPosWorld, ActiveBrushTargetPlaneWorld.Z());
+		return true;
+	}
+
+	return false;
+}
+
+
+
+
+
+void UMeshSculptToolBase::SaveActiveStrokeModifiers()
+{
+	bSmoothing = GetShiftToggle();
+	bInvert = GetCtrlToggle();
+}
+
+
+void UMeshSculptToolBase::UpdateHoverStamp(const FFrame3d& StampFrame)
+{
+	HoverStamp.WorldFrame = StampFrame;
+}
+
+void UMeshSculptToolBase::ApplyStrokeFlowInTick()
+{
+	bIsStampPending = InStroke();
+}
+
+
+
+FFrame3d UMeshSculptToolBase::ComputeStampRegionPlane(const FFrame3d& StampFrame, const TSet<int32>& StampTriangles, bool bIgnoreDepth, bool bViewAligned, bool bInvDistFalloff)
+{
+	const FDynamicMesh3* Mesh = GetSculptMesh();
+	double FalloffRadius = GetCurrentBrushRadius();
+	if (bInvDistFalloff)
+	{
+		FalloffRadius *= 0.5;
+	}
+	FVector3d StampNormal = StampFrame.Z();
+
+	FVector3d AverageNormal(0, 0, 0);
+	FVector3d AveragePos(0, 0, 0);
+	double WeightSum = 0;
+	for (int TriID : StampTriangles)
+	{
+		FVector3d Normal, Centroid; double Area;
+		Mesh->GetTriInfo(TriID, Normal, Area, Centroid);
+		if (Normal.Dot(StampNormal) < -0.2)		// ignore back-facing (heuristic to avoid "other side")
+		{
+			continue;
+		}
+
+		double Distance = StampFrame.Origin.Distance(Centroid);
+		double NormalizedDistance = (Distance / FalloffRadius) + 0.0001;
+
+		double Weight = Area;
+		if (bInvDistFalloff)
+		{
+			double RampT = FMathd::Clamp(1.0 - NormalizedDistance, 0.0, 1.0);
+			Weight *= FMathd::Clamp(RampT * RampT * RampT, 0.0, 1.0);
+		}
+		else
+		{
+			if (NormalizedDistance > 0.5)
+			{
+				double d = FMathd::Clamp((NormalizedDistance - 0.5) / (1.0 - 0.5), 0.0, 1.0);
+				double t = (1.0 - d * d);
+				Weight *= (t * t * t);
+			}
+		}
+
+		AverageNormal += Weight * Mesh->GetTriNormal(TriID);
+		AveragePos += Weight * Centroid;
+		WeightSum += Weight;
+	}
+	AverageNormal.Normalize();
+	AveragePos /= WeightSum;
+
+	if (bViewAligned)
+	{
+		AverageNormal = -CameraState.Forward();
+	}
+
+	FFrame3d Result = FFrame3d(AveragePos, AverageNormal);
+	if (bIgnoreDepth == false)
+	{
+		Result.Origin -= GetCurrentBrushDepth() * GetCurrentBrushRadius() * Result.Z();
+	}
+
+	return Result;
+}
+
+
+
+void UMeshSculptToolBase::UpdateStrokeReferencePlaneForROI(const FFrame3d& StampFrame, const TSet<int32>& TriangleROI, bool bViewAligned)
+{
+	StrokePlane = ComputeStampRegionPlane(GetBrushFrameLocal(), TriangleROI, false, bViewAligned);
+}
+
+void UMeshSculptToolBase::UpdateStrokeReferencePlaneFromWorkPlane()
+{
+	StrokePlane = FFrame3d(
+		CurTargetTransform.InverseTransformPosition(GizmoProperties->Position),
+		CurTargetTransform.GetRotation().Inverse() * (FQuaterniond)GizmoProperties->Rotation);
 }
 
 
@@ -175,6 +683,27 @@ void UMeshSculptToolBase::CalculateBrushRadius()
 }
 
 
+double UMeshSculptToolBase::GetCurrentBrushStrength()
+{
+	TUniquePtr<FMeshSculptBrushOp>& BrushOp = GetActiveBrushOp();
+	if (BrushOp->PropertySet.IsValid())
+	{
+		return FMathd::Clamp(BrushOp->PropertySet->GetStrength(), 0.0, 1.0);
+	}
+	return 1.0;
+}
+
+double UMeshSculptToolBase::GetCurrentBrushDepth()
+{
+	TUniquePtr<FMeshSculptBrushOp>& BrushOp = GetActiveBrushOp();
+	if (BrushOp->PropertySet.IsValid())
+	{
+		return FMathd::Clamp(BrushOp->PropertySet->GetDepth(), -1.0, 1.0);
+	}
+	return 0.0;
+}
+
+
 
 void UMeshSculptToolBase::IncreaseBrushRadiusAction()
 {
@@ -200,22 +729,6 @@ void UMeshSculptToolBase::DecreaseBrushRadiusSmallStepAction()
 	CalculateBrushRadius();
 }
 
-void UMeshSculptToolBase::IncreaseBrushFalloffAction()
-{
-	const float ChangeAmount = 0.1f;
-	const float OldValue = BrushProperties->BrushFalloffAmount;
-	float NewValue = OldValue + ChangeAmount;
-	BrushProperties->BrushFalloffAmount = FMath::Clamp(NewValue, 0.f, 1.f);
-}
-
-void UMeshSculptToolBase::DecreaseBrushFalloffAction()
-{
-	const float ChangeAmount = 0.1f;
-	const float OldValue = BrushProperties->BrushFalloffAmount;
-	float NewValue = OldValue - ChangeAmount;
-	BrushProperties->BrushFalloffAmount = FMath::Clamp(NewValue, 0.f, 1.f);
-}
-
 
 
 
@@ -406,7 +919,7 @@ void UMeshSculptToolBase::UpdateFixedSculptPlanePosition(const FVector& Position
 void UMeshSculptToolBase::UpdateFixedSculptPlaneRotation(const FQuat& Rotation)
 {
 	GizmoProperties->Rotation = Rotation;
-	GizmoPositionWatcher.SilentUpdate();
+	GizmoRotationWatcher.SilentUpdate();
 }
 
 void UMeshSculptToolBase::UpdateGizmoFromProperties()
@@ -465,20 +978,6 @@ void UMeshSculptToolBase::RegisterActions(FInteractiveToolActionSet& ActionSet)
 		EModifierKey::None, EKeys::LeftBracket,
 		[this]() { DecreaseBrushRadiusAction(); });
 
-	ActionSet.RegisterAction(this, (int32)EStandardToolActions::BaseClientDefinedActionID + 12,
-		TEXT("BrushIncreaseFalloff"),
-		LOCTEXT("BrushIncreaseFalloff", "Increase Brush Falloff"),
-		LOCTEXT("BrushIncreaseFalloffTooltip", "Press this key to increase brush falloff by a fixed increment."),
-		EModifierKey::Shift | EModifierKey::Control, EKeys::RightBracket,
-		[this]() { IncreaseBrushFalloffAction(); });
-
-	ActionSet.RegisterAction(this, (int32)EStandardToolActions::BaseClientDefinedActionID + 13,
-		TEXT("BrushDecreaseFalloff"),
-		LOCTEXT("BrushDecreaseFalloff", "Decrease Brush Falloff"),
-		LOCTEXT("BrushDecreaseFalloffTooltip", "Press this key to decrease brush falloff by a fixed increment."),
-		EModifierKey::Shift | EModifierKey::Control, EKeys::LeftBracket,
-		[this]() { DecreaseBrushFalloffAction(); });
-
 
 
 	ActionSet.RegisterAction(this, (int32)EStandardToolActions::BaseClientDefinedActionID + 1,
diff --git a/Engine/Plugins/Experimental/MeshModelingToolset/Source/MeshModelingTools/Public/DynamicMeshSculptTool.h b/Engine/Plugins/Experimental/MeshModelingToolset/Source/MeshModelingTools/Public/DynamicMeshSculptTool.h
--- a/Engine/Plugins/Experimental/MeshModelingToolset/Source/MeshModelingTools/Public/DynamicMeshSculptTool.h
+++ b/Engine/Plugins/Experimental/MeshModelingToolset/Source/MeshModelingTools/Public/DynamicMeshSculptTool.h
@@ -261,9 +261,6 @@ public:
 	UPROPERTY()
 	UBrushSculptProperties* SculptProperties;
 
-	UPROPERTY()
-	UPlaneBrushProperties* PlaneBrushProperties;
-
 	UPROPERTY()
 	USculptMaxBrushProperties* SculptMaxBrushProperties;
 	
diff --git a/Engine/Plugins/Experimental/MeshModelingToolset/Source/MeshModelingTools/Public/MeshVertexSculptTool.h b/Engine/Plugins/Experimental/MeshModelingToolset/Source/MeshModelingTools/Public/MeshVertexSculptTool.h
--- a/Engine/Plugins/Experimental/MeshModelingToolset/Source/MeshModelingTools/Public/MeshVertexSculptTool.h
+++ b/Engine/Plugins/Experimental/MeshModelingToolset/Source/MeshModelingTools/Public/MeshVertexSculptTool.h
@@ -10,7 +10,6 @@
 #include "DynamicMeshOctree3.h"
 #include "MeshNormals.h"
 #include "BaseTools/BaseBrushTool.h"
-#include "Changes/ValueWatcher.h"
 #include "TransformTypes.h"
 #include "Sculpting/MeshSculptToolBase.h"
 #include "Sculpting/MeshBrushOpBase.h"
@@ -65,6 +64,9 @@ enum class EMeshVertexSculptBrushType : uint8
 	/** Smooth mesh vertices  */
 	Smooth UMETA(DisplayName = "Smooth"),
 
+	/** Smooth mesh vertices but only in direction of normal (Ctrl to invert) */
+	SmoothFill UMETA(DisplayName = "SmoothFill"),
+
 	/** Displace vertices along the average surface normal (Ctrl to invert) */
 	Offset UMETA(DisplayName = "Sculpt (Normal)"),
 
@@ -109,23 +111,20 @@ class MESHMODELINGTOOLS_API UVertexBrushSculptProperties : public UInteractiveTo
 public:
 	/** Primary Brush Mode */
 	UPROPERTY(EditAnywhere, Category = Sculpting, meta = (DisplayName = "Brush Type"))
-	EMeshVertexSculptBrushType PrimaryBrushType = EMeshVertexSculptBrushType::Move;
-
-	/** Strength of the Primary Brush */
-	UPROPERTY(EditAnywhere, Category = Sculpting, meta = (DisplayName = "Strength", UIMin = "0.0", UIMax = "1.0", ClampMin = "0.0", ClampMax = "1.0", EditCondition = "PrimaryBrushType != EMeshVertexSculptBrushType::Pull"))
-	float PrimaryBrushSpeed = 0.5;
+	EMeshVertexSculptBrushType PrimaryBrushType = EMeshVertexSculptBrushType::Offset;
 
-	/** If true, try to preserve the shape of the UV/3D mapping. This will limit Smoothing and Remeshing in some cases. */
-	UPROPERTY(EditAnywhere, Category = Sculpting)
-	bool bPreserveUVFlow = false;
+	/** Primary Brush Falloff */
+	UPROPERTY(EditAnywhere, Category = Sculpting, meta = (DisplayName = "Falloff Shape"))
+	EMeshSculptFalloffType PrimaryFalloffType = EMeshSculptFalloffType::Smooth;
 
 	/** When Freeze Target is toggled on, the Brush Target Surface will be Frozen in its current state, until toggled off. Brush strokes will be applied relative to the Target Surface, for applicable Brushes */
-	UPROPERTY(EditAnywhere, Category = Sculpting, meta = (EditCondition = "PrimaryBrushType == EMeshVertexSculptBrushType::Sculpt || PrimaryBrushType == EMeshVertexSculptBrushType::SculptMax || PrimaryBrushType == EMeshVertexSculptBrushType::SculptView || PrimaryBrushType == EMeshVertexSculptBrushType::Pinch || PrimaryBrushType == EMeshVertexSculptBrushType::Resample" ))
+	UPROPERTY(EditAnywhere, Category = Sculpting, meta = (EditCondition = "PrimaryBrushType == EMeshVertexSculptBrushType::Offset || PrimaryBrushType == EMeshVertexSculptBrushType::SculptMax || PrimaryBrushType == EMeshVertexSculptBrushType::SculptView || PrimaryBrushType == EMeshVertexSculptBrushType::Pinch || PrimaryBrushType == EMeshVertexSculptBrushType::Resample" ))
 	bool bFreezeTarget = false;
 
-	/** Strength of Shift-to-Smooth Brushing and Smoothing Brush */
-	UPROPERTY(EditAnywhere, Category = Smoothing, meta = (DisplayName = "Smoothing Strength", UIMin = "0.0", UIMax = "1.0", ClampMin = "0.0", ClampMax = "1.0"))
-	float SmoothBrushSpeed = 0.25;
+	/** When enabled, instead of Mesh Smoothing, the Shift-Smooth modifier will "erase" the displacement relative to the Brush Target Surface */
+	//UPROPERTY(EditAnywhere, Category = Sculpting, meta = (DisplayName = "Shift-Smooth Erases", EditCondition = "bFreezeTarget == true"))
+	UPROPERTY()
+	bool bSmoothErases = false;
 };
 
 
@@ -150,13 +149,6 @@ public:
 	virtual bool HasAccept() const override { return true; }
 	virtual bool CanAccept() const override { return true; }
 
-	virtual bool HitTest(const FRay& Ray, FHitResult& OutHit) override;
-
-	virtual void OnBeginDrag(const FRay& Ray) override;
-	virtual void OnUpdateDrag(const FRay& Ray) override;
-	virtual void OnEndDrag(const FRay& Ray) override;
-
-	virtual FInputRayHit BeginHoverSequenceHitTest(const FInputDeviceRay& PressPos) override;
 	virtual bool OnUpdateHover(const FInputDeviceRay& DevicePos) override;
 
 	virtual void OnPropertyModified(UObject* PropertySet, FProperty* Property) override;
@@ -167,16 +159,6 @@ public:
 	UPROPERTY()
 	UVertexBrushSculptProperties* SculptProperties;
 
-	UPROPERTY()
-	UPlaneBrushProperties* PlaneBrushProperties;
-
-	UPROPERTY()
-	USculptMaxBrushProperties* SculptMaxBrushProperties;
-	
-	//UPROPERTY()
-	//UKelvinBrushProperties* KelvinBrushProperties;
-
-
 public:
 	virtual void IncreaseBrushSpeedAction() override;
 	virtual void DecreaseBrushSpeedAction() override;
@@ -184,30 +166,28 @@ public:
 	virtual void PreviousBrushModeAction() override;
 
 protected:
-	UPROPERTY()
-	USimpleDynamicMeshComponent* DynamicMeshComponent;
-
+	// UMeshSculptToolBase API
 	virtual UBaseDynamicMeshComponent* GetSculptMeshComponent() { return DynamicMeshComponent; }
+	virtual FDynamicMesh3* GetBaseMesh() { return &BaseMesh; }
+	virtual const FDynamicMesh3* GetBaseMesh() const { return &BaseMesh; }
+
+	virtual int32 FindHitSculptMeshTriangle(const FRay3d& LocalRay) override;
+	virtual int32 FindHitTargetMeshTriangle(const FRay3d& LocalRay) override;
 
-	FTransform3d InitialTargetTransform;
-	FTransform3d CurTargetTransform;
+	virtual void OnBeginStroke(const FRay& WorldRay) override;
+	virtual void OnEndStroke() override;
+	// end UMeshSculptToolBase API
+
+protected:
+	UPROPERTY()
+	USimpleDynamicMeshComponent* DynamicMeshComponent;
 
 	// realtime visualization
 	void OnDynamicMeshComponentChanged(USimpleDynamicMeshComponent* Component, const FMeshVertexChange* Change, bool bRevert);
 	FDelegateHandle OnDynamicMeshComponentChangedHandle;
 
-	TValueWatcher<EMeshVertexSculptBrushType> BrushTypeWatcher;
 	void UpdateBrushType(EMeshVertexSculptBrushType BrushType);
 
-
-	bool bInDrag;
-
-	FFrame3d ActiveDragPlane;
-	FVector3d LastHitPosWorld;
-	FVector3d LastBrushPosLocal;
-	FVector3d LastBrushPosWorld;
-	FVector3d LastBrushPosNormalWorld;
-
 	TSet<int32> AccumulatedTriangleROI;
 	bool bUndoUpdatePending = false;
 	TFuture<bool> UndoNormalsFuture;
@@ -223,56 +203,32 @@ protected:
 
 	bool bTargetDirty;
 
-	bool bSmoothing;
-	bool bInvert;
-
-	bool bStampPending;
-	FRay PendingStampRay;
-	int StampTimestamp = 0;
 	EMeshVertexSculptBrushType PendingStampType = EMeshVertexSculptBrushType::Smooth;
 
-	TUniquePtr<FMeshSculptBrushOp> PrimaryBrushOp;
-	TUniquePtr<FMeshSculptBrushOp> SecondaryBrushOp;
-	TSharedPtr<FMeshSculptFallofFunc> ActiveFalloff;
-
 	bool UpdateStampPosition(const FRay& WorldRay);
 	void ApplyStamp();
 
-
 	FDynamicMesh3 BaseMesh;
 	FDynamicMeshOctree3 BaseMeshSpatial;
 	TArray<int32> BaseMeshIndexBuffer;
 	bool bCachedFreezeTarget = false;
 	void UpdateBaseMesh(const TSet<int32>* TriangleROI = nullptr);
 	bool GetBaseMeshNearest(int32 VertexID, const FVector3d& Position, double SearchRadius, FVector3d& TargetPosOut, FVector3d& TargetNormalOut);
+	TFunction<bool(int32, const FVector3d&, double MaxDist, FVector3d&, FVector3d&)> BaseMeshQueryFunc;
 
 	FDynamicMeshOctree3 Octree;
 
-	int FindHitSculptMeshTriangle(const FRay3d& LocalRay);
-	int FindHitTargetMeshTriangle(const FRay3d& LocalRay);
-
 	bool UpdateBrushPosition(const FRay& WorldRay);
-	bool UpdateBrushPositionOnActivePlane(const FRay& WorldRay);
-	bool UpdateBrushPositionOnTargetMesh(const FRay& WorldRay, bool bFallbackToViewPlane);
-	bool UpdateBrushPositionOnSculptMesh(const FRay& WorldRay, bool bFallbackToViewPlane);
-	void AlignBrushToView();
 
 	double SculptMaxFixedHeight = -1.0;
 
-	double CalculateBrushFalloff(double Distance);
 	TArray<FVector3d> ROIPositionBuffer;
 	void SyncMeshWithPositionBuffer(FDynamicMesh3* Mesh);
 
-	FFrame3d ActiveFixedBrushPlane;
-	FFrame3d ComputeROIBrushPlane(const FVector3d& BrushCenter, bool bIgnoreDepth, bool bViewAligned);
-
 	FMeshVertexChangeBuilder* ActiveVertexChange = nullptr;
 	void BeginChange();
 	void EndChange();
 
-	FCriticalSection UpdateSavedVertexLock;
-
-
 
 protected:
 	virtual bool ShowWorkPlane() const override { return SculptProperties->PrimaryBrushType == EMeshVertexSculptBrushType::FixedPlane; }
diff --git a/Engine/Plugins/Experimental/MeshModelingToolset/Source/MeshModelingTools/Public/Sculpting/MeshBrushOpBase.h b/Engine/Plugins/Experimental/MeshModelingToolset/Source/MeshModelingTools/Public/Sculpting/MeshBrushOpBase.h
--- a/Engine/Plugins/Experimental/MeshModelingToolset/Source/MeshModelingTools/Public/Sculpting/MeshBrushOpBase.h
+++ b/Engine/Plugins/Experimental/MeshModelingToolset/Source/MeshModelingTools/Public/Sculpting/MeshBrushOpBase.h
@@ -5,6 +5,7 @@
 #include "CoreMinimal.h"
 #include "VectorTypes.h"
 #include "FrameTypes.h"
+#include "MeshBrushOpBase.generated.h"
 
 class FDynamicMesh3;
 
@@ -16,36 +17,50 @@ enum class ESculptBrushOpTargetType : uint8
 };
 
 
-struct FSculptBrushStamp
+UENUM()
+enum class EPlaneBrushSideMode : uint8
+{
+	BothSides = 0,
+	PushDown = 1,
+	PullTowards = 2
+};
+
+
+struct MESHMODELINGTOOLS_API FSculptBrushStamp
 {
 	FFrame3d WorldFrame;
 	FFrame3d LocalFrame;
 	double Radius;
+	double Falloff;
 	double Power;
 	double Direction;
 	double Depth;
+	double DeltaTime;
 
 	FFrame3d PrevWorldFrame;
 	FFrame3d PrevLocalFrame;
 
+	FDateTime TimeStamp;
+
 	// only initialized if current op requires it
 	FFrame3d RegionPlane;
+
+	FSculptBrushStamp()
+	{
+		TimeStamp = FDateTime::Now();
+	}
 };
 
 
-struct FSculptBrushOptions
+struct MESHMODELINGTOOLS_API FSculptBrushOptions
 {
-	bool bPreserveUVFlow = false;
-
-	double MaxHeight = 0.5f;
+	//bool bPreserveUVFlow = false;
 
 	FFrame3d ConstantReferencePlane;
-
-	int32 WhichPlaneSideIndex = 0;		// 	BothSides = 0,	PushDown = 1, PullTowards = 2
 };
 
 
-class FMeshSculptFallofFunc
+class MESHMODELINGTOOLS_API FMeshSculptFallofFunc
 {
 public:
 	TUniqueFunction<double(const FSculptBrushStamp& StampInfo, const FVector3d& Position)> FalloffFunc;
@@ -57,11 +72,35 @@ public:
 };
 
 
-class FMeshSculptBrushOp
+
+UCLASS()
+class MESHMODELINGTOOLS_API UMeshSculptBrushOpProps : public UInteractiveToolPropertySet
+{
+	GENERATED_BODY()
+public:
+
+	virtual float GetStrength() { return 1.0f; }
+	virtual float GetDepth() { return 0.0f; }
+	virtual float GetFalloff() { return 0.5f; }
+};
+
+
+
+class MESHMODELINGTOOLS_API FMeshSculptBrushOp
 {
 public:
 	virtual ~FMeshSculptBrushOp() {}
 
+	TWeakObjectPtr<UMeshSculptBrushOpProps> PropertySet;
+
+	template<typename PropType> 
+	PropType* GetPropertySetAs()
+	{
+		check(PropertySet.IsValid());
+		return CastChecked<PropType>(PropertySet.Get());
+	}
+
+
 	TSharedPtr<FMeshSculptFallofFunc> Falloff;
 	FSculptBrushOptions CurrentOptions;
 
@@ -107,3 +146,40 @@ public:
 
 
 
+class MESHMODELINGTOOLS_API FMeshSculptBrushOpFactory
+{
+public:
+	virtual ~FMeshSculptBrushOpFactory() {}
+	virtual TUniquePtr<FMeshSculptBrushOp> Build() = 0;
+};
+
+template<typename OpType>
+class TBasicMeshSculptBrushOpFactory : public FMeshSculptBrushOpFactory
+{
+public:
+	virtual TUniquePtr<FMeshSculptBrushOp> Build() override
+	{
+		return MakeUnique<OpType>();
+	}
+};
+
+
+class FLambdaMeshSculptBrushOpFactory : public FMeshSculptBrushOpFactory
+{
+public:
+	TUniqueFunction<TUniquePtr<FMeshSculptBrushOp>(void)> BuildFunc;
+
+	FLambdaMeshSculptBrushOpFactory()
+	{
+	}
+
+	FLambdaMeshSculptBrushOpFactory(TUniqueFunction<TUniquePtr<FMeshSculptBrushOp>(void)> BuildFuncIn)
+	{
+		BuildFunc = MoveTemp(BuildFuncIn);
+	}
+
+	virtual TUniquePtr<FMeshSculptBrushOp> Build() override
+	{
+		return BuildFunc();
+	}
+};
\ No newline at end of file
diff --git a/Engine/Plugins/Experimental/MeshModelingToolset/Source/MeshModelingTools/Public/Sculpting/MeshInflateBrushOps.h b/Engine/Plugins/Experimental/MeshModelingToolset/Source/MeshModelingTools/Public/Sculpting/MeshInflateBrushOps.h
--- a/Engine/Plugins/Experimental/MeshModelingToolset/Source/MeshModelingTools/Public/Sculpting/MeshInflateBrushOps.h
+++ b/Engine/Plugins/Experimental/MeshModelingToolset/Source/MeshModelingTools/Public/Sculpting/MeshInflateBrushOps.h
@@ -7,13 +7,31 @@
 #include "MeshNormals.h"
 #include "MeshWeights.h"
 #include "Async/ParallelFor.h"
+#include "MeshInflateBrushOps.generated.h"
 
 
+UCLASS()
+class MESHMODELINGTOOLS_API UInflateBrushOpProps : public UMeshSculptBrushOpProps
+{
+	GENERATED_BODY()
+public:
+	/** Strength of the Brush */
+	UPROPERTY(EditAnywhere, Category = InflateBrush, meta = (DisplayName = "Strength", UIMin = "0.0", UIMax = "1.0", ClampMin = "0.0", ClampMax = "1.0"))
+	float Strength = 0.5;
+
+	/** Amount of falloff to apply */
+	UPROPERTY(EditAnywhere, Category = InflateBrush, meta = (DisplayName = "Falloff", UIMin = "0.0", UIMax = "1.0", ClampMin = "0.0", ClampMax = "1.0"))
+	float Falloff = 0.5;
+
+	virtual float GetStrength() override { return Strength; }
+	virtual float GetFalloff() override { return Falloff; }
+};
+
 class FInflateBrushOp : public FMeshSculptBrushOp
 {
 
 public:
-	double BrushSpeedTuning = 0.05;
+	double BrushSpeedTuning = 2.0;
 
 	virtual ESculptBrushOpTargetType GetBrushTargetType() const override
 	{
@@ -24,7 +42,7 @@ public:
 	{
 		const FVector3d& StampPos = Stamp.LocalFrame.Origin;
 
-		double UsePower = Stamp.Direction * Stamp.Power * Stamp.Radius * BrushSpeedTuning;
+		double UsePower = Stamp.Direction * Stamp.Power * Stamp.Radius * Stamp.DeltaTime * BrushSpeedTuning;
 
 		ParallelFor(Vertices.Num(), [&](int32 k)
 		{
diff --git a/Engine/Plugins/Experimental/MeshModelingToolset/Source/MeshModelingTools/Public/Sculpting/MeshMoveBrushOps.h b/Engine/Plugins/Experimental/MeshModelingToolset/Source/MeshModelingTools/Public/Sculpting/MeshMoveBrushOps.h
--- a/Engine/Plugins/Experimental/MeshModelingToolset/Source/MeshModelingTools/Public/Sculpting/MeshMoveBrushOps.h
+++ b/Engine/Plugins/Experimental/MeshModelingToolset/Source/MeshModelingTools/Public/Sculpting/MeshMoveBrushOps.h
@@ -7,15 +7,36 @@
 #include "MeshNormals.h"
 #include "MeshWeights.h"
 #include "Async/ParallelFor.h"
+#include "MeshMoveBrushOps.generated.h"
 
 
-class FMoveBrushOp : public FMeshSculptBrushOp
+UCLASS()
+class MESHMODELINGTOOLS_API UMoveBrushOpProps : public UMeshSculptBrushOpProps
 {
-
+	GENERATED_BODY()
 public:
-	
+	/** Strength of the Brush */
+	UPROPERTY(EditAnywhere, Category = MoveBrush, meta = (DisplayName = "Strength", UIMin = "0.0", UIMax = "1.0", ClampMin = "0.0", ClampMax = "1.0"))
+	float Strength = 1.0;
+
+	/** Amount of falloff to apply */
+	UPROPERTY(EditAnywhere, Category = MoveBrush, meta = (DisplayName = "Falloff", UIMin = "0.0", UIMax = "1.0", ClampMin = "0.0", ClampMax = "1.0"))
+	float Falloff = 0.5;
+
+	/** Depth of Brush into surface along view ray */
+	UPROPERTY(EditAnywhere, Category = MoveBrush, meta = (UIMin = "-0.5", UIMax = "0.5", ClampMin = "-1.0", ClampMax = "1.0"))
+	float Depth = 0;
+
+	virtual float GetStrength() override { return Strength; }
+	virtual float GetFalloff() override { return Falloff; }
+	virtual float GetDepth() override { return Depth; }
+};
 
 
+class FMoveBrushOp : public FMeshSculptBrushOp
+{
+public:
+
 	virtual void ApplyStamp(const FDynamicMesh3* Mesh, const FSculptBrushStamp& Stamp, const TArray<int32>& Vertices, TArray<FVector3d>& NewPositionsOut) override
 	{
 		const FVector3d& StampPos = Stamp.LocalFrame.Origin;
@@ -46,43 +67,3 @@ public:
 		return true;
 	}
 };
-
-
-
-
-
-//void UMeshVertexSculptTool::ApplyMoveBrush(const FRay& WorldRay)
-//{
-//	FVector3d NewBrushPosLocal = CurTargetTransform.InverseTransformPosition(LastBrushPosWorld);
-//	FVector3d MoveVec = NewBrushPosLocal - LastBrushPosLocal;
-//
-//	if (MoveVec.SquaredLength() <= 0)
-//	{
-//		LastBrushPosLocal = NewBrushPosLocal;
-//		return;
-//	}
-//
-//	FDynamicMesh3* Mesh = DynamicMeshComponent->GetMesh();
-//	int NumV = VertexROI.Num();
-//	ROIPositionBuffer.SetNum(NumV, false);
-//
-//	ParallelFor(NumV, [this, Mesh, NewBrushPosLocal, MoveVec](int k)
-//	{
-//		int VertIdx = VertexROI[k];
-//		FVector3d OrigPos = Mesh->GetVertex(VertIdx);
-//
-//		double PrevDist = (OrigPos - LastBrushPosLocal).Length();
-//		double NewDist = (OrigPos - NewBrushPosLocal).Length();
-//		double UseDist = FMath::Min(PrevDist, NewDist);
-//
-//		double Falloff = CalculateBrushFalloff(UseDist) * ActivePressure;
-//
-//		FVector3d NewPos = OrigPos + Falloff * MoveVec;
-//		ROIPositionBuffer[k] = NewPos;
-//	});
-//
-//	// Update the mesh positions to match those in the position buffer
-//	SyncMeshWithPositionBuffer(Mesh);
-//
-//	LastBrushPosLocal = NewBrushPosLocal;
-//}
\ No newline at end of file
diff --git a/Engine/Plugins/Experimental/MeshModelingToolset/Source/MeshModelingTools/Public/Sculpting/MeshPinchBrushOps.h b/Engine/Plugins/Experimental/MeshModelingToolset/Source/MeshModelingTools/Public/Sculpting/MeshPinchBrushOps.h
--- a/Engine/Plugins/Experimental/MeshModelingToolset/Source/MeshModelingTools/Public/Sculpting/MeshPinchBrushOps.h
+++ b/Engine/Plugins/Experimental/MeshModelingToolset/Source/MeshModelingTools/Public/Sculpting/MeshPinchBrushOps.h
@@ -6,13 +6,37 @@
 #include "DynamicMesh3.h"
 #include "Async/ParallelFor.h"
 #include "LineTypes.h"
+#include "MeshPinchBrushOps.generated.h"
+
+
+UCLASS()
+class MESHMODELINGTOOLS_API UPinchBrushOpProps : public UMeshSculptBrushOpProps
+{
+	GENERATED_BODY()
+public:
+	/** Strength of the Brush */
+	UPROPERTY(EditAnywhere, Category = PinchBrush, meta = (DisplayName = "Strength", UIMin = "0.0", UIMax = "1.0", ClampMin = "0.0", ClampMax = "1.0"))
+	float Strength = 0.5;
+
+	/** Amount of falloff to apply */
+	UPROPERTY(EditAnywhere, Category = PinchBrush, meta = (DisplayName = "Falloff", UIMin = "0.0", UIMax = "1.0", ClampMin = "0.0", ClampMax = "1.0"))
+	float Falloff = 1.0;
+
+	/** Depth of Brush into surface along surface normal */
+	UPROPERTY(EditAnywhere, Category = PinchBrush, meta = (UIMin = "-0.5", UIMax = "0.5", ClampMin = "-1.0", ClampMax = "1.0"))
+	float Depth = 0;
+
+	virtual float GetStrength() override { return Strength; }
+	virtual float GetFalloff() override { return Falloff; }
+	virtual float GetDepth() override { return Depth; }
+};
 
 
 class FPinchBrushOp : public FMeshSculptBrushOp
 {
 
 public:
-	double BrushSpeedTuning = 0.05;
+	double BrushSpeedTuning = 6.0;
 
 	FVector3d LastSmoothBrushPosLocal;
 	FVector3d LastSmoothBrushNormalLocal;;
@@ -37,7 +61,7 @@ public:
 		FLine3d MoveLine(LastSmoothBrushPosLocal, MotionVec);
 
 		FVector3d DepthPosLocal = NewSmoothBrushPosLocal - (Stamp.Depth * Stamp.Radius * NewSmoothBrushNormal);
-		double UseSpeed = Stamp.Direction * Stamp.Radius * Stamp.Power * BrushSpeedTuning;
+		double UseSpeed = Stamp.Direction * Stamp.Radius * Stamp.Power * Stamp.DeltaTime * BrushSpeedTuning;
 
 		LastSmoothBrushPosLocal = NewSmoothBrushPosLocal;
 
diff --git a/Engine/Plugins/Experimental/MeshModelingToolset/Source/MeshModelingTools/Public/Sculpting/MeshPlaneBrushOps.h b/Engine/Plugins/Experimental/MeshModelingToolset/Source/MeshModelingTools/Public/Sculpting/MeshPlaneBrushOps.h
--- a/Engine/Plugins/Experimental/MeshModelingToolset/Source/MeshModelingTools/Public/Sculpting/MeshPlaneBrushOps.h
+++ b/Engine/Plugins/Experimental/MeshModelingToolset/Source/MeshModelingTools/Public/Sculpting/MeshPlaneBrushOps.h
@@ -5,13 +5,110 @@
 #include "Sculpting/MeshBrushOpBase.h"
 #include "DynamicMesh3.h"
 #include "Async/ParallelFor.h"
+#include "MeshPlaneBrushOps.generated.h"
+
+
+
+
+
+UCLASS()
+class MESHMODELINGTOOLS_API UBasePlaneBrushOpProps : public UMeshSculptBrushOpProps
+{
+	GENERATED_BODY()
+public:
+	virtual EPlaneBrushSideMode GetWhichSide() { return EPlaneBrushSideMode::BothSides; }
+};
+
+UCLASS()
+class MESHMODELINGTOOLS_API UPlaneBrushOpProps : public UBasePlaneBrushOpProps
+{
+	GENERATED_BODY()
+public:
+	/** Strength of the Brush */
+	UPROPERTY(EditAnywhere, Category = PlaneBrush, meta = (DisplayName = "Strength", UIMin = "0.0", UIMax = "1.0", ClampMin = "0.0", ClampMax = "1.0"))
+	float Strength = 0.5;
+
+	/** Amount of falloff to apply */
+	UPROPERTY(EditAnywhere, Category = PlaneBrush, meta = (DisplayName = "Falloff", UIMin = "0.0", UIMax = "1.0", ClampMin = "0.0", ClampMax = "1.0"))
+	float Falloff = 0.5;
+
+	/** Depth of Brush into surface along surface normal */
+	UPROPERTY(EditAnywhere, Category = PlaneBrush, meta = (UIMin = "-0.5", UIMax = "0.5", ClampMin = "-1.0", ClampMax = "1.0"))
+	float Depth = 0;
+
+	/** Control whether effect of brush should be limited to one side of the Plane  */
+	UPROPERTY(EditAnywhere, Category = PlaneBrush)
+	EPlaneBrushSideMode WhichSide = EPlaneBrushSideMode::BothSides;
+
+	virtual float GetStrength() override { return Strength; }
+	virtual float GetFalloff() override { return Falloff; }
+	virtual float GetDepth() override { return Depth; }
+	virtual EPlaneBrushSideMode GetWhichSide() { return WhichSide; }
+};
+
+
+UCLASS()
+class MESHMODELINGTOOLS_API UViewAlignedPlaneBrushOpProps : public UBasePlaneBrushOpProps
+{
+	GENERATED_BODY()
+public:
+	/** Strength of the Brush */
+	UPROPERTY(EditAnywhere, Category = ViewPlaneBrush, meta = (DisplayName = "Strength", UIMin = "0.0", UIMax = "1.0", ClampMin = "0.0", ClampMax = "1.0"))
+	float Strength = 0.5;
+
+	/** Amount of falloff to apply */
+	UPROPERTY(EditAnywhere, Category = ViewPlaneBrush, meta = (DisplayName = "Falloff", UIMin = "0.0", UIMax = "1.0", ClampMin = "0.0", ClampMax = "1.0"))
+	float Falloff = 0.5;
+
+	/** Depth of Brush into surface along view ray */
+	UPROPERTY(EditAnywhere, Category = ViewPlaneBrush, meta = (UIMin = "-0.5", UIMax = "0.5", ClampMin = "-1.0", ClampMax = "1.0"))
+	float Depth = 0;
+
+	/** Control whether effect of brush should be limited to one side of the Plane  */
+	UPROPERTY(EditAnywhere, Category = ViewPlaneBrush)
+	EPlaneBrushSideMode WhichSide = EPlaneBrushSideMode::BothSides;
+
+	virtual float GetStrength() override { return Strength; }
+	virtual float GetFalloff() override { return Falloff; }
+	virtual float GetDepth() override { return Depth; }
+	virtual EPlaneBrushSideMode GetWhichSide() { return WhichSide; }
+};
+
+
+UCLASS()
+class MESHMODELINGTOOLS_API UFixedPlaneBrushOpProps : public UBasePlaneBrushOpProps
+{
+	GENERATED_BODY()
+public:
+	/** Strength of the Brush */
+	UPROPERTY(EditAnywhere, Category = FixedPlaneBrush, meta = (DisplayName = "Strength", UIMin = "0.0", UIMax = "1.0", ClampMin = "0.0", ClampMax = "1.0"))
+	float Strength = 0.5;
+
+	/** Amount of falloff to apply */
+	UPROPERTY(EditAnywhere, Category = FixedPlaneBrush, meta = (DisplayName = "Falloff", UIMin = "0.0", UIMax = "1.0", ClampMin = "0.0", ClampMax = "1.0"))
+	float Falloff = 0.5;
+
+	/** Depth of Brush into surface relative to plane */
+	UPROPERTY(EditAnywhere, Category = FixedPlaneBrush, meta = (UIMin = "-0.5", UIMax = "0.5", ClampMin = "-1.0", ClampMax = "1.0"))
+	float Depth = 0;
+
+	/** Control whether effect of brush should be limited to one side of the Plane  */
+	UPROPERTY(EditAnywhere, Category = FixedPlaneBrush)
+	EPlaneBrushSideMode WhichSide = EPlaneBrushSideMode::BothSides;
+
+	virtual float GetStrength() override { return Strength; }
+	virtual float GetFalloff() override { return Falloff; }
+	virtual float GetDepth() override { return Depth; }
+	virtual EPlaneBrushSideMode GetWhichSide() { return WhichSide; }
+};
+
 
 
 class FPlaneBrushOp : public FMeshSculptBrushOp
 {
 
 public:
-	double BrushSpeedTuning = 0.05;
+	double BrushSpeedTuning = 6.0;
 
 	FFrame3d StrokePlane;
 
@@ -23,11 +120,12 @@ public:
 	virtual void ApplyStamp(const FDynamicMesh3* Mesh, const FSculptBrushStamp& Stamp, const TArray<int32>& Vertices, TArray<FVector3d>& NewPositionsOut) override
 	{
 		static const double PlaneSigns[3] = { 0, -1, 1 };
-		double PlaneSign = PlaneSigns[CurrentOptions.WhichPlaneSideIndex];
+		UBasePlaneBrushOpProps* Props = GetPropertySetAs<UBasePlaneBrushOpProps>();
+		double PlaneSign = PlaneSigns[(int32)Props->GetWhichSide()];
 
 		const FVector3d& StampPos = Stamp.LocalFrame.Origin;
 
-		double UseSpeed = Stamp.Power * Stamp.Radius * BrushSpeedTuning;
+		double UseSpeed = Stamp.Power * Stamp.Radius * Stamp.DeltaTime * BrushSpeedTuning;
 
 		ParallelFor(Vertices.Num(), [&](int32 k)
 		{
diff --git a/Engine/Plugins/Experimental/MeshModelingToolset/Source/MeshModelingTools/Public/Sculpting/MeshSculptBrushOps.h b/Engine/Plugins/Experimental/MeshModelingToolset/Source/MeshModelingTools/Public/Sculpting/MeshSculptBrushOps.h
--- a/Engine/Plugins/Experimental/MeshModelingToolset/Source/MeshModelingTools/Public/Sculpting/MeshSculptBrushOps.h
+++ b/Engine/Plugins/Experimental/MeshModelingToolset/Source/MeshModelingTools/Public/Sculpting/MeshSculptBrushOps.h
@@ -7,12 +7,30 @@
 #include "MeshNormals.h"
 #include "MeshWeights.h"
 #include "Async/ParallelFor.h"
+#include "MeshSculptBrushOps.generated.h"
 
 
+UCLASS()
+class MESHMODELINGTOOLS_API UStandardSculptBrushOpProps : public UMeshSculptBrushOpProps
+{
+	GENERATED_BODY()
+public:
+	/** Strength of the Brush */
+	UPROPERTY(EditAnywhere, Category = SculptBrush, meta = (DisplayName = "Strength", UIMin = "0.0", UIMax = "1.0", ClampMin = "0.0", ClampMax = "1.0"))
+	float Strength = 0.5;
+
+	/** Amount of falloff to apply */
+	UPROPERTY(EditAnywhere, Category = SculptBrush, meta = (DisplayName = "Falloff", UIMin = "0.0", UIMax = "1.0", ClampMin = "0.0", ClampMax = "1.0"))
+	float Falloff = 1.0;
+
+	virtual float GetStrength() override { return Strength; }
+	virtual float GetFalloff() override { return Falloff; }
+};
+
 class FSurfaceSculptBrushOp : public FMeshSculptBrushOp
 {
 public:
-	double BrushSpeedTuning = 0.05;
+	double BrushSpeedTuning = 6.0;
 
 	typedef TUniqueFunction<bool(int32, const FVector3d&, double, FVector3d&, FVector3d&)> NearestQueryFuncType;
 
@@ -30,7 +48,7 @@ public:
 
 	virtual void ApplyStamp(const FDynamicMesh3* Mesh, const FSculptBrushStamp& Stamp, const TArray<int32>& Vertices, TArray<FVector3d>& NewPositionsOut) override
 	{
-		double UsePower = Stamp.Direction * Stamp.Power * Stamp.Radius * BrushSpeedTuning;
+		double UsePower = Stamp.Direction * Stamp.Power * Stamp.Radius * Stamp.DeltaTime * BrushSpeedTuning;
 		double MaxOffset = Stamp.Radius;
 
 		ParallelFor(Vertices.Num(), [&](int32 k)
@@ -47,7 +65,7 @@ public:
 			else
 			{
 				FVector3d MoveVec = UsePower * BaseNormal;
-				double Falloff = GetFalloff().Evaluate(Stamp, BasePos);
+				double Falloff = GetFalloff().Evaluate(Stamp, OrigPos);
 				FVector3d NewPos = OrigPos + Falloff * MoveVec;
 				NewPositionsOut[k] = NewPos;
 			}
@@ -60,13 +78,28 @@ public:
 
 
 
+UCLASS()
+class MESHMODELINGTOOLS_API UViewAlignedSculptBrushOpProps : public UMeshSculptBrushOpProps
+{
+	GENERATED_BODY()
+public:
+	/** Strength of the Brush */
+	UPROPERTY(EditAnywhere, Category = SculptToViewBrush, meta = (DisplayName = "Strength", UIMin = "0.0", UIMax = "1.0", ClampMin = "0.0", ClampMax = "1.0"))
+	float Strength = 0.5;
 
+	/** Amount of falloff to apply */
+	UPROPERTY(EditAnywhere, Category = SculptToViewBrush, meta = (DisplayName = "Falloff", UIMin = "0.0", UIMax = "1.0", ClampMin = "0.0", ClampMax = "1.0"))
+	float Falloff = 1.0;
+
+	virtual float GetStrength() override { return Strength; }
+	virtual float GetFalloff() override { return Falloff; }
+};
 
 
 class FViewAlignedSculptBrushOp : public FMeshSculptBrushOp
 {
 public:
-	double BrushSpeedTuning = 0.05;
+	double BrushSpeedTuning = 6.0;
 
 	typedef TUniqueFunction<bool(int32, const FVector3d&, double, FVector3d&, FVector3d&)> NearestQueryFuncType;
 
@@ -91,7 +124,7 @@ public:
 	{
 		FVector3d StampNormal = Stamp.LocalFrame.Z();
 
-		double UsePower = Stamp.Direction * Stamp.Power * Stamp.Radius * BrushSpeedTuning;
+		double UsePower = Stamp.Direction * Stamp.Power * Stamp.Radius * Stamp.DeltaTime * BrushSpeedTuning;
 		double MaxOffset = Stamp.Radius;
 
 		ParallelFor(Vertices.Num(), [&](int32 k)
@@ -108,7 +141,7 @@ public:
 			else
 			{
 				FVector3d MoveVec = UsePower * StampNormal;
-				double Falloff = GetFalloff().Evaluate(Stamp, BasePos);
+				double Falloff = GetFalloff().Evaluate(Stamp, OrigPos);
 				FVector3d NewPos = OrigPos + Falloff * MoveVec;
 				NewPositionsOut[k] = NewPos;
 			}
@@ -121,12 +154,41 @@ public:
 
 
 
+UCLASS()
+class MESHMODELINGTOOLS_API USculptMaxBrushOpProps : public UMeshSculptBrushOpProps
+{
+	GENERATED_BODY()
+public:
+	/** Strength of the Brush */
+	UPROPERTY(EditAnywhere, Category = SculptMaxBrush, meta = (DisplayName = "Strength", UIMin = "0.0", UIMax = "1.0", ClampMin = "0.0", ClampMax = "1.0"))
+	float Strength = 0.5;
+
+	/** Amount of falloff to apply */
+	UPROPERTY(EditAnywhere, Category = SculptMaxBrush, meta = (DisplayName = "Falloff", UIMin = "0.0", UIMax = "1.0", ClampMin = "0.0", ClampMax = "1.0"))
+	float Falloff = 0.5;
+
+	/** Maximum height as fraction of brush size */
+	UPROPERTY(EditAnywhere, Category = SculptMaxBrush, meta = (UIMin = "0.0", UIMax = "1.0", ClampMin = "0.0", ClampMax = "1.0"))
+	float MaxHeight = 0.5;
+
+	/** If true, maximum height is defined using the FixedHeight constant instead of brush-relative size */
+	UPROPERTY(EditAnywhere, AdvancedDisplay, Category = SculptMaxBrush)
+	bool bUseFixedHeight = false;
+
+	/** Maximum height in world-space dimension */
+	UPROPERTY(EditAnywhere, AdvancedDisplay, Category = SculptMaxBrush)
+	float FixedHeight = 0.0;
+
+
+	virtual float GetStrength() override { return Strength; }
+	virtual float GetFalloff() override { return Falloff; }
+};
 
 
 class FSurfaceMaxSculptBrushOp : public FMeshSculptBrushOp
 {
 public:
-	double BrushSpeedTuning = 0.05;
+	double BrushSpeedTuning = 6.0;
 
 	typedef TUniqueFunction<bool(int32, const FVector3d&, double, FVector3d&, FVector3d&)> NearestQueryFuncType;
 
@@ -144,8 +206,10 @@ public:
 
 	virtual void ApplyStamp(const FDynamicMesh3* Mesh, const FSculptBrushStamp& Stamp, const TArray<int32>& Vertices, TArray<FVector3d>& NewPositionsOut) override
 	{
-		double UsePower = Stamp.Direction * Stamp.Power * Stamp.Radius * BrushSpeedTuning;
-		double MaxOffset = CurrentOptions.MaxHeight;
+		double UsePower = Stamp.Direction * Stamp.Power * Stamp.Radius * Stamp.DeltaTime * BrushSpeedTuning;
+
+		USculptMaxBrushOpProps* Props = GetPropertySetAs<USculptMaxBrushOpProps>();
+		double MaxOffset = (Props->bUseFixedHeight) ? Props->FixedHeight : (Props->MaxHeight * Stamp.Radius);
 
 		ParallelFor(Vertices.Num(), [&](int32 k)
 		{
@@ -161,7 +225,7 @@ public:
 			else
 			{
 				FVector3d MoveVec = UsePower * BaseNormal;
-				double Falloff = GetFalloff().Evaluate(Stamp, BasePos);
+				double Falloff = GetFalloff().Evaluate(Stamp, OrigPos);
 				FVector3d NewPos = OrigPos + Falloff * MoveVec;
 
 				FVector3d DeltaPos = NewPos - BasePos;
diff --git a/Engine/Plugins/Experimental/MeshModelingToolset/Source/MeshModelingTools/Public/Sculpting/MeshSculptToolBase.h b/Engine/Plugins/Experimental/MeshModelingToolset/Source/MeshModelingTools/Public/Sculpting/MeshSculptToolBase.h
--- a/Engine/Plugins/Experimental/MeshModelingToolset/Source/MeshModelingTools/Public/Sculpting/MeshSculptToolBase.h
+++ b/Engine/Plugins/Experimental/MeshModelingToolset/Source/MeshModelingTools/Public/Sculpting/MeshSculptToolBase.h
@@ -9,6 +9,7 @@
 #include "BoxTypes.h"
 #include "Properties/MeshMaterialProperties.h"
 #include "Changes/ValueWatcher.h"
+#include "BaseDynamicMeshComponent.h"
 #include "MeshSculptToolBase.generated.h"
 
 
@@ -16,7 +17,26 @@ class UMaterialInstanceDynamic;
 class UTransformGizmo;
 class UTransformProxy;
 class UPreviewMesh;
-class UBaseDynamicMeshComponent;
+
+
+/** Mesh Sculpting Brush Types */
+UENUM()
+enum class EMeshSculptFalloffType : uint8
+{
+	Smooth,
+	Linear,
+	Inverse,
+	Round,
+	BoxSmooth,
+	BoxLinear,
+	BoxInverse,
+	BoxRound,
+
+	LastValue UMETA(Hidden)
+};
+
+
+
 
 
 UCLASS()
@@ -30,12 +50,20 @@ public:
 	}
 
 	/** Depth of Brush into surface along view ray or surface normal, depending on the Active Brush Type */
-	UPROPERTY(EditAnywhere, Category = Brush, meta = (UIMin = "-0.5", UIMax = "0.5", ClampMin = "-1.0", ClampMax = "1.0", DisplayPriority = 5))
+	UPROPERTY(EditAnywhere, Category = Brush, meta = (UIMin = "-0.5", UIMax = "0.5", ClampMin = "-1.0", ClampMax = "1.0", DisplayPriority = 5, HideEditConditionToggle, EditConditionHides, EditCondition = "bShowPerBrushProps"))
 	float Depth = 0;
 
 	/** Allow the Brush to hit the back-side of the mesh */
 	UPROPERTY(EditAnywhere, Category = Brush, meta = (DisplayPriority = 6))
 	bool bHitBackFaces = true;
+
+
+	UPROPERTY(EditAnywhere, Category = Brush, meta = (UIMin = "0.0", UIMax = "1.0", ClampMin = "0.0", ClampMax = "1.0", DisplayPriority = 7))
+	float Lazyness = 0;
+
+	/**  */
+	UPROPERTY( meta = (TransientToolProperty))
+	bool bShowPerBrushProps = true;
 };
 
 
@@ -66,27 +94,6 @@ public:
 
 
 
-UENUM()
-enum class EPlaneBrushSideMode : uint8
-{
-	BothSides = 0,
-	PushDown = 1,
-	PullTowards = 2
-};
-
-
-
-UCLASS()
-class MESHMODELINGTOOLS_API UPlaneBrushProperties : public UInteractiveToolPropertySet
-{
-	GENERATED_BODY()
-public:
-	/** Control whether effect of brush should be limited to one side of the Plane  */
-	UPROPERTY(EditAnywhere, Category = PlaneBrush)
-	EPlaneBrushSideMode WhichSide = EPlaneBrushSideMode::BothSides;
-};
-
-
 
 UCLASS()
 class MESHMODELINGTOOLS_API UWorkPlaneProperties : public UInteractiveToolPropertySet
@@ -149,26 +156,103 @@ public:
 	virtual void Shutdown(EToolShutdownType ShutdownType) override;
 	virtual void Render(IToolsContextRenderAPI* RenderAPI) override;
 
+	// UMeshSurfacePointTool API
+	virtual bool HitTest(const FRay& Ray, FHitResult& OutHit) override;
+	virtual void OnBeginDrag(const FRay& Ray) override;
+	virtual void OnUpdateDrag(const FRay& Ray) override;
+	virtual void OnEndDrag(const FRay& Ray) override;
+	// end UMeshSurfacePointTool API
+
 protected:
 	virtual void OnTick(float DeltaTime) override;
 
+	
+	virtual void OnCompleteSetup();
+	virtual void OnBeginStroke(const FRay& WorldRay) { check(false); }
+	virtual void OnEndStroke() { check(false); }
+
+public:
+	/** Properties that control brush size/etc */
+	UPROPERTY()
+	USculptBrushProperties* BrushProperties;
+
+	/** Properties for 3D workplane / gizmo */
+	UPROPERTY()
+	UWorkPlaneProperties* GizmoProperties;
+
+
 protected:
 	UWorld* TargetWorld;		// required to spawn UPreviewMesh/etc
 	FViewCameraState CameraState;
 
+	/** Initial transformation on target mesh */
+	FTransform3d InitialTargetTransform;
+	/** Active transformation on target mesh, includes baked scale */
+	FTransform3d CurTargetTransform;
+
+	FRay3d GetLocalRay(const FRay& WorldRay) const;
+
+
 	/**
 	 * Subclass must implement this and return relevant rendering component
 	 */
 	virtual UBaseDynamicMeshComponent* GetSculptMeshComponent() { check(false); return nullptr; }
 
+	virtual FDynamicMesh3* GetSculptMesh() { return GetSculptMeshComponent()->GetMesh(); }
+	virtual const FDynamicMesh3* GetSculptMesh() const { return const_cast<UMeshSculptToolBase*>(this)->GetSculptMeshComponent()->GetMesh(); }
 
-public:
-	/** Properties that control brush size/etc*/
+	virtual FDynamicMesh3* GetBaseMesh() { check(false); return nullptr; }
+	virtual const FDynamicMesh3* GetBaseMesh() const { check(false); return nullptr; }
+
+
+	/**
+	 * Subclass calls this to set up editing component
+	 */
+	void InitializeSculptMeshComponent(UBaseDynamicMeshComponent* Component);
+
+
+
+	//
+	// Brush Types
+	//
+protected:
 	UPROPERTY()
-	USculptBrushProperties* BrushProperties;
+	TMap<int32, UMeshSculptBrushOpProps*> BrushOpPropSets;
+
+	TMap<int32, TUniquePtr<FMeshSculptBrushOpFactory>> BrushOpFactories;
 
 	UPROPERTY()
-	UWorkPlaneProperties* GizmoProperties;
+	TMap<int32, UMeshSculptBrushOpProps*> SecondaryBrushOpPropSets;
+
+	TMap<int32, TUniquePtr<FMeshSculptBrushOpFactory>> SecondaryBrushOpFactories;
+
+	void RegisterBrushType(int32 Identifier, TUniquePtr<FMeshSculptBrushOpFactory> Factory, UMeshSculptBrushOpProps* PropSet);
+	void RegisterSecondaryBrushType(int32 Identifier, TUniquePtr<FMeshSculptBrushOpFactory> Factory, UMeshSculptBrushOpProps* PropSet);
+
+	virtual void SaveAllBrushTypeProperties(UInteractiveTool* SaveFromTool);
+	virtual void RestoreAllBrushTypeProperties(UInteractiveTool* RestoreToTool);
+
+private:
+	TUniquePtr<FMeshSculptBrushOp> PrimaryBrushOp;
+	UMeshSculptBrushOpProps* PrimaryVisiblePropSet = nullptr;		// BrushOpPropSets prevents GC of this
+
+	TUniquePtr<FMeshSculptBrushOp> SecondaryBrushOp;
+	UMeshSculptBrushOpProps* SecondaryVisiblePropSet = nullptr;
+protected:
+	void SetActivePrimaryBrushType(int32 Identifier);
+	void SetActiveSecondaryBrushType(int32 Identifier);
+	TUniquePtr<FMeshSculptBrushOp>& GetActiveBrushOp();
+
+
+	//
+	// Falloff types
+	//
+private:
+	TSharedPtr<FMeshSculptFallofFunc> PrimaryFalloff;
+protected:
+	void SetPrimaryFalloffType(EMeshSculptFalloffType Falloff);
+
+
 
 
 	//
@@ -176,21 +260,25 @@ public:
 	//
 protected:
 	FInterval1d BrushRelativeSizeRange;
-	double CurrentBrushRadius;
+	double CurrentBrushRadius = 1.0;
 	virtual void InitializeBrushSizeRange(const FAxisAlignedBox3d& TargetBounds);
 	virtual void CalculateBrushRadius();
 	virtual double GetCurrentBrushRadius() const { return CurrentBrushRadius; }
 
+	double CurrentBrushFalloff = 0.5;
+	virtual double GetCurrentBrushFalloff() const { return CurrentBrushFalloff; }
+
 	double ActivePressure = 1.0;
 	virtual double GetActivePressure() const { return ActivePressure; }
 
+	virtual double GetCurrentBrushStrength();
+	virtual double GetCurrentBrushDepth();
+
 public:
 	virtual void IncreaseBrushRadiusAction();
 	virtual void DecreaseBrushRadiusAction();
 	virtual void IncreaseBrushRadiusSmallStepAction();
 	virtual void DecreaseBrushRadiusSmallStepAction();
-	virtual void IncreaseBrushFalloffAction();
-	virtual void DecreaseBrushFalloffAction();
 
 
 	// client currently needs to implement these...
@@ -200,16 +288,89 @@ public:
 	virtual void PreviousBrushModeAction() {}
 
 
+
+	//
+	// Brush/Stroke stuff
+	//
+private:
+	FFrame3d LastBrushFrameWorld;
+	FFrame3d LastBrushFrameLocal;
+protected:
+	const FFrame3d& GetBrushFrameWorld() const { return LastBrushFrameWorld; }
+	const FFrame3d& GetBrushFrameLocal() const { return LastBrushFrameLocal; }
+	void UpdateBrushFrameWorld(const FVector3d& NewPosition, const FVector3d& NewNormal);
+	void AlignBrushToView();
+
+	bool GetBrushCanHitBackFaces() const { return BrushProperties->bHitBackFaces; }
+
+	/** @return hit triangle at ray position - subclass must implement this */
+	virtual int32 FindHitSculptMeshTriangle(const FRay3d& LocalRay) { check(false); return IndexConstants::InvalidID; }
+	/** @return hit triangle at ray position - subclass should implement this for most brushes */
+	virtual int32 FindHitTargetMeshTriangle(const FRay3d& LocalRay) { check(false); return IndexConstants::InvalidID;	}
+
+	virtual bool UpdateBrushPositionOnActivePlane(const FRay& WorldRay);
+	virtual bool UpdateBrushPositionOnTargetMesh(const FRay& WorldRay, bool bFallbackToViewPlane);
+	virtual bool UpdateBrushPositionOnSculptMesh(const FRay& WorldRay, bool bFallbackToViewPlane);
+
+
+
+	//
+	// Brush Target Plane is plane that some brushes move on
+	//
+private:
+	FFrame3d ActiveBrushTargetPlaneWorld;
+protected:
+	void UpdateBrushTargetPlaneFromHit(const FRay& WorldRay, const FHitResult& Hit);
+
+
+	//
+	// Stroke Modifiers
+	//
+private:
+	bool bInStroke = false;
+	bool bSmoothing = false;
+	bool bInvert = false;
+protected:
+	void SaveActiveStrokeModifiers();
+	bool InStroke() const { return bInStroke; }
+	bool GetInSmoothingStroke() const { return bSmoothing; }
+	bool GetInInvertStroke() const { return bInvert; }
+
+	void ApplyStrokeFlowInTick();
+
 	//
 	// Stamps
 	//
+private:
+	bool bIsStampPending = false;
+	FRay PendingStampRay;
 protected:
 	FSculptBrushStamp HoverStamp;
 	FSculptBrushStamp CurrentStamp;
 	FSculptBrushStamp LastStamp;
-	void UpdateHoverStamp(const FVector3d& WorldPos, const FVector3d& WorldNormal);
+	void UpdateHoverStamp(const FFrame3d& StampFrame);
+	bool IsStampPending() const { return bIsStampPending; }
+	const FRay& GetPendingStampRayWorld() const { return PendingStampRay;  }
+
+
+	//
+	// Stamp ROI Plane is a plane used by some brush ops
+	//
+private:
+	FFrame3d StampRegionPlane;
+protected:
+	virtual FFrame3d ComputeStampRegionPlane(const FFrame3d& StampFrame, const TSet<int32>& StampTriangles, bool bIgnoreDepth, bool bViewAligned, bool bInvDistFalloff = true);
+
 
+	// Stroke plane is a plane used by some brush ops
+	//
+private:
+	FFrame3d StrokePlane;
+protected:
+	virtual const FFrame3d& GetCurrentStrokeReferencePlane() const { return StrokePlane; }
 
+	virtual void UpdateStrokeReferencePlaneForROI(const FFrame3d& StampFrame, const TSet<int32>& TriangleROI, bool bViewAligned);
+	virtual void UpdateStrokeReferencePlaneFromWorkPlane();
 
 
 	//
diff --git a/Engine/Plugins/Experimental/MeshModelingToolset/Source/MeshModelingTools/Public/Sculpting/MeshSmoothingBrushOps.h b/Engine/Plugins/Experimental/MeshModelingToolset/Source/MeshModelingTools/Public/Sculpting/MeshSmoothingBrushOps.h
--- a/Engine/Plugins/Experimental/MeshModelingToolset/Source/MeshModelingTools/Public/Sculpting/MeshSmoothingBrushOps.h
+++ b/Engine/Plugins/Experimental/MeshModelingToolset/Source/MeshModelingTools/Public/Sculpting/MeshSmoothingBrushOps.h
@@ -5,7 +5,66 @@
 #include "Sculpting/MeshBrushOpBase.h"
 #include "DynamicMesh3.h"
 #include "MeshWeights.h"
+#include "MeshNormals.h"
 #include "Async/ParallelFor.h"
+#include "MeshSmoothingBrushOps.generated.h"
+
+
+UCLASS()
+class MESHMODELINGTOOLS_API UBaseSmoothBrushOpProps : public UMeshSculptBrushOpProps
+{
+	GENERATED_BODY()
+public:
+	virtual bool GetPreserveUVFlow() { return false; }
+};
+
+
+UCLASS()
+class MESHMODELINGTOOLS_API USmoothBrushOpProps : public UBaseSmoothBrushOpProps
+{
+	GENERATED_BODY()
+public:
+	/** Strength of the Brush */
+	UPROPERTY(EditAnywhere, Category = SmoothBrush, meta = (DisplayName = "Strength", UIMin = "0.0", UIMax = "1.0", ClampMin = "0.0", ClampMax = "1.0"))
+	float Strength = 0.5;
+
+	/** Amount of falloff to apply (0.0 - 1.0) */
+	UPROPERTY(EditAnywhere, Category = SmoothBrush, meta = (DisplayName = "Falloff", UIMin = "0.0", UIMax = "1.0", ClampMin = "0.0", ClampMax = "1.0"))
+	float Falloff = 0.5;
+
+	/** If true, try to preserve the shape of the UV/3D mapping. This will limit Smoothing and Remeshing in some cases. */
+	UPROPERTY(EditAnywhere, Category = SmoothBrush)
+	bool bPreserveUVFlow = false;
+
+	virtual float GetStrength() override { return Strength; }
+	virtual float GetFalloff() override { return Falloff; }
+	virtual bool GetPreserveUVFlow() override { return bPreserveUVFlow; }
+};
+
+
+
+UCLASS()
+class MESHMODELINGTOOLS_API USecondarySmoothBrushOpProps : public UBaseSmoothBrushOpProps
+{
+	GENERATED_BODY()
+public:
+	/** Strength of the Brush */
+	UPROPERTY(EditAnywhere, Category = ShiftSmoothBrush, meta = (DisplayName = "Strength", UIMin = "0.0", UIMax = "1.0", ClampMin = "0.0", ClampMax = "1.0"))
+	float Strength = 0.5;
+
+	/** Amount of falloff to apply */
+	UPROPERTY(EditAnywhere, Category = ShiftSmoothBrush, meta = (DisplayName = "Falloff", UIMin = "0.0", UIMax = "1.0", ClampMin = "0.0", ClampMax = "1.0"))
+	float Falloff = 0.5;
+
+	/** If true, try to preserve the shape of the UV/3D mapping. This will limit Smoothing and Remeshing in some cases. */
+	UPROPERTY(EditAnywhere, Category = ShiftSmoothBrush)
+	bool bPreserveUVFlow = false;
+
+	virtual float GetStrength() override { return Strength; }
+	virtual float GetFalloff() override { return Falloff; }
+	virtual bool GetPreserveUVFlow() override { return bPreserveUVFlow; }
+};
+
 
 
 class FSmoothBrushOp : public FMeshSculptBrushOp
@@ -16,6 +75,7 @@ public:
 	virtual void ApplyStamp(const FDynamicMesh3* Mesh, const FSculptBrushStamp& Stamp, const TArray<int32>& Vertices, TArray<FVector3d>& NewPositionsOut) override
 	{
 		const FVector3d& StampPos = Stamp.LocalFrame.Origin;
+		bool bPreserveUVFlow = GetPropertySetAs<UBaseSmoothBrushOpProps>()->GetPreserveUVFlow();
 
 		ParallelFor(Vertices.Num(), [&](int32 k)
 		{
@@ -25,7 +85,7 @@ public:
 
 			double Falloff = GetFalloff().Evaluate(Stamp, OrigPos);
 
-			FVector3d SmoothedPos = (CurrentOptions.bPreserveUVFlow) ?
+			FVector3d SmoothedPos = (bPreserveUVFlow) ?
 				FMeshWeights::MeanValueCentroid(*Mesh, VertIdx) : FMeshWeights::UniformCentroid(*Mesh, VertIdx);
 
 			FVector3d NewPos = FVector3d::Lerp(OrigPos, SmoothedPos, Falloff * Stamp.Power);
@@ -39,41 +99,129 @@ public:
 
 
 
+UCLASS()
+class MESHMODELINGTOOLS_API USmoothFillBrushOpProps : public UBaseSmoothBrushOpProps
+{
+	GENERATED_BODY()
+public:
+	/** Strength of the Brush */
+	UPROPERTY(EditAnywhere, Category = SmoothFillBrush, meta = (DisplayName = "Strength", UIMin = "0.0", UIMax = "1.0", ClampMin = "0.0", ClampMax = "1.0"))
+	float Strength = 0.5;
+
+	/** Amount of falloff to apply (0.0 - 1.0) */
+	UPROPERTY(EditAnywhere, Category = SmoothFillBrush, meta = (DisplayName = "Falloff", UIMin = "0.0", UIMax = "1.0", ClampMin = "0.0", ClampMax = "1.0"))
+	float Falloff = 0.5;
+
+	/** If true, try to preserve the shape of the UV/3D mapping. This will limit Smoothing and Remeshing in some cases. */
+	UPROPERTY(EditAnywhere, Category = SmoothFillBrush)
+	bool bPreserveUVFlow = false;
+
+	virtual float GetStrength() override { return Strength; }
+	virtual float GetFalloff() override { return Falloff; }
+	virtual bool GetPreserveUVFlow() override { return bPreserveUVFlow; }
+};
+
+
+class FSmoothFillBrushOp : public FMeshSculptBrushOp
+{
+
+public:
+
+	virtual void ApplyStamp(const FDynamicMesh3* Mesh, const FSculptBrushStamp& Stamp, const TArray<int32>& Vertices, TArray<FVector3d>& NewPositionsOut) override
+	{
+		const FVector3d& StampPos = Stamp.LocalFrame.Origin;
+		double Direction = Stamp.Direction;
+		bool bPreserveUVFlow = GetPropertySetAs<UBaseSmoothBrushOpProps>()->GetPreserveUVFlow();
+
+		ParallelFor(Vertices.Num(), [&](int32 k)
+		{
+			int32 VertIdx = Vertices[k];
+
+			FVector3d OrigPos = Mesh->GetVertex(VertIdx);
+			FVector3d Normal = FMeshNormals::ComputeVertexNormal(*Mesh, VertIdx);
+
+			double Falloff = GetFalloff().Evaluate(Stamp, OrigPos);
+
+			FVector3d SmoothedPos = (bPreserveUVFlow) ?
+				FMeshWeights::MeanValueCentroid(*Mesh, VertIdx) : FMeshWeights::UniformCentroid(*Mesh, VertIdx);
+
+			FVector3d NewPos = FVector3d::Lerp(OrigPos, SmoothedPos, Falloff * Stamp.Power);
+
+			if ((NewPos - OrigPos).Dot(Normal) * Direction < 0)
+			{
+				NewPositionsOut[k] = OrigPos;
+			}
+			else
+			{
+				NewPositionsOut[k] = NewPos;
+			}
+		});
+	}
+};
+
+
+
+
+
+UCLASS()
+class MESHMODELINGTOOLS_API UFlattenBrushOpProps : public UMeshSculptBrushOpProps
+{
+	GENERATED_BODY()
+public:
+	/** Strength of the Brush */
+	UPROPERTY(EditAnywhere, Category = FlattenBrush, meta = (DisplayName = "Strength", UIMin = "0.0", UIMax = "1.0", ClampMin = "0.0", ClampMax = "1.0"))
+	float Strength = 0.5;
+
+	/** Amount of falloff to apply (0.0 - 1.0) */
+	UPROPERTY(EditAnywhere, Category = FlattenBrush, meta = (DisplayName = "Falloff", UIMin = "0.0", UIMax = "1.0", ClampMin = "0.0", ClampMax = "1.0"))
+	float Falloff = 0.5;
 
+	/** Depth of Brush into surface along surface normal */
+	UPROPERTY(EditAnywhere, Category = FlattenBrush, meta = (UIMin = "-0.5", UIMax = "0.5", ClampMin = "-1.0", ClampMax = "1.0"))
+	float Depth = 0;
 
+	/** Control whether effect of brush should be limited to one side of the Plane  */
+	UPROPERTY(EditAnywhere, Category = FlattenBrush)
+	EPlaneBrushSideMode WhichSide = EPlaneBrushSideMode::BothSides;
 
+	virtual float GetStrength() override { return Strength; }
+	virtual float GetFalloff() override { return Falloff; }
+	virtual float GetDepth() override { return Depth; }
+};
 
 
 class FFlattenBrushOp : public FMeshSculptBrushOp
 {
 
 public:
-	double BrushSpeedTuning = 0.05;
+	double BrushSpeedTuning = 1.0;
 
 	virtual void ApplyStamp(const FDynamicMesh3* Mesh, const FSculptBrushStamp& Stamp, const TArray<int32>& Vertices, TArray<FVector3d>& NewPositionsOut) override
 	{
 		static const double PlaneSigns[3] = { 0, -1, 1 };
-		double PlaneSign = PlaneSigns[CurrentOptions.WhichPlaneSideIndex];
+		double PlaneSign = PlaneSigns[(int32)GetPropertySetAs<UFlattenBrushOpProps>()->WhichSide];
 
 		const FVector3d& StampPos = Stamp.LocalFrame.Origin;
 
-		double UseSpeed = Stamp.Power * Stamp.Radius * BrushSpeedTuning;
+		double UseSpeed = Stamp.Power * Stamp.Radius * Stamp.DeltaTime * BrushSpeedTuning;
 		const FFrame3d& FlattenPlane = Stamp.RegionPlane;
+		FVector3d PlaneZ = FlattenPlane.Z();
 
 		ParallelFor(Vertices.Num(), [&](int32 k)
 		{
 			int32 VertIdx = Vertices[k];
 			FVector3d OrigPos = Mesh->GetVertex(VertIdx);
+			FVector3d Normal = FMeshNormals::ComputeVertexNormal(*Mesh, VertIdx);
 			FVector3d PlanePos = FlattenPlane.ToPlane(OrigPos, 2);
-			FVector3d Delta = PlanePos - OrigPos;
+			FVector3d MoveDelta = PlanePos - OrigPos;
 
-			double Dot = Delta.Dot(FlattenPlane.Z());
+			double PlaneDot = MoveDelta.Dot(FlattenPlane.Z());
 			FVector3d NewPos = OrigPos;
-			if (Dot * PlaneSign >= 0)
+			if (PlaneDot * PlaneSign >= 0)
 			{
 				double Falloff = GetFalloff().Evaluate(Stamp, OrigPos);
-				FVector3d MoveVec = Falloff * UseSpeed * Delta;
-				double MaxDist = Delta.Normalize();
+				FVector3d MoveVec = Falloff * UseSpeed * MoveDelta;
+				double MaxDist = MoveDelta.Normalize();
 				NewPos = (MoveVec.SquaredLength() > MaxDist * MaxDist) ?
 					PlanePos : (OrigPos + Falloff * MoveVec);
 			}
@@ -90,3 +238,79 @@ public:
 };
 
 
+
+
+
+
+UCLASS()
+class MESHMODELINGTOOLS_API UEraseBrushOpProps : public UMeshSculptBrushOpProps
+{
+	GENERATED_BODY()
+public:
+	/** Strength of the Brush */
+	UPROPERTY(EditAnywhere, Category = EraseBrush, meta = (DisplayName = "Strength", UIMin = "0.0", UIMax = "1.0", ClampMin = "0.0", ClampMax = "1.0"))
+	float Strength = 0.5;
+
+	/** Amount of falloff to apply (0.0 - 1.0) */
+	UPROPERTY(EditAnywhere, Category = EraseBrush, meta = (DisplayName = "Falloff", UIMin = "0.0", UIMax = "1.0", ClampMin = "0.0", ClampMax = "1.0"))
+	float Falloff = 0.5;
+
+	virtual float GetStrength() override { return Strength; }
+	virtual float GetFalloff() override { return Falloff; }
+};
+
+
+class FEraseToBaseMeshBrushOp : public FMeshSculptBrushOp
+{
+public:
+	double BrushSpeedTuning = 3.0;
+
+	typedef TUniqueFunction<bool(int32, const FVector3d&, double, FVector3d&, FVector3d&)> NearestQueryFuncType;
+
+	NearestQueryFuncType BaseMeshNearestQueryFunc;
+
+	FEraseToBaseMeshBrushOp(NearestQueryFuncType QueryFunc)
+	{
+		BaseMeshNearestQueryFunc = MoveTemp(QueryFunc);
+	}
+
+	virtual ESculptBrushOpTargetType GetBrushTargetType() const override
+	{
+		return ESculptBrushOpTargetType::TargetMesh;
+	}
+
+	virtual void ApplyStamp(const FDynamicMesh3* Mesh, const FSculptBrushStamp& Stamp, const TArray<int32>& Vertices, TArray<FVector3d>& NewPositionsOut) override
+	{
+		double UsePower = Stamp.Power * Stamp.Radius * Stamp.DeltaTime * BrushSpeedTuning;
+		double MaxOffset = Stamp.Radius;
+
+		ParallelFor(Vertices.Num(), [&](int32 k)
+		{
+			int32 VertIdx = Vertices[k];
+			FVector3d OrigPos = Mesh->GetVertex(VertIdx);
+
+			FVector3d BasePos, BaseNormal;
+			bool bFoundBasePos = BaseMeshNearestQueryFunc(VertIdx, OrigPos, 4.0 * Stamp.Radius, BasePos, BaseNormal);
+			if (bFoundBasePos == false)
+			{
+				NewPositionsOut[k] = OrigPos;
+			}
+			else
+			{
+				FVector3d MoveVec = (BasePos - OrigPos); 
+				double Falloff = GetFalloff().Evaluate(Stamp, OrigPos);
+				double MoveDist = Falloff * UsePower;
+				if (MoveVec.SquaredLength() < MoveDist * MoveDist)
+				{
+					NewPositionsOut[k] = BasePos;
+				}
+				else
+				{
+					MoveVec.Normalize();
+					NewPositionsOut[k] = OrigPos + MoveDist * MoveVec;
+				}
+			}
+		});
+	}
+
+};
diff --git a/Engine/Plugins/Experimental/MeshModelingToolset/Source/MeshModelingTools/Public/Sculpting/StampFalloffs.h b/Engine/Plugins/Experimental/MeshModelingToolset/Source/MeshModelingTools/Public/Sculpting/StampFalloffs.h
new file mode 100644
--- /dev/null
+++ b/Engine/Plugins/Experimental/MeshModelingToolset/Source/MeshModelingTools/Public/Sculpting/StampFalloffs.h
@@ -0,0 +1,215 @@
+// Copyright Epic Games, Inc. All Rights Reserved.
+
+#pragma once
+
+#include "Sculpting/MeshBrushOpBase.h"
+
+
+namespace UE
+{
+	namespace SculptFalloffs
+	{
+
+		static TUniqueFunction<double(const FSculptBrushStamp&, const FVector3d&)> MakeStandardSmoothFalloff()
+		{
+			return [](const FSculptBrushStamp& Stamp, const FVector3d& Position)
+			{
+				double Distance = Position.Distance(Stamp.LocalFrame.Origin);
+				double FalloffT = FMathd::Clamp(1.0 - Stamp.Falloff, 0.0, 1.0);
+				double UnitDistance = Distance / Stamp.Radius;
+				if (UnitDistance > FalloffT)
+				{
+					UnitDistance = FMathd::Clamp((UnitDistance - FalloffT) / (1.0 - FalloffT), 0.0, 1.0);
+					double Weight = (1.0 - UnitDistance * UnitDistance);
+					return Weight * Weight * Weight;
+				}
+				return 1.0;
+			};
+		}
+
+
+
+		static TUniqueFunction<double(const FSculptBrushStamp&, const FVector3d&)> MakeLinearFalloff()
+		{
+			return [](const FSculptBrushStamp& Stamp, const FVector3d& Position)
+			{
+				double Distance = Position.Distance(Stamp.LocalFrame.Origin);
+				double FalloffT = FMathd::Clamp(1.0 - Stamp.Falloff, 0.0, 1.0);
+				double UnitDistance = Distance / Stamp.Radius;
+				if (UnitDistance > FalloffT)
+				{
+					UnitDistance = FMathd::Clamp((UnitDistance - FalloffT) / (1.0 - FalloffT), 0.0, 1.0);
+					double Weight = FMathd::Clamp(1.0 - UnitDistance, 0.0, 1.0);
+					return Weight;
+				}
+				return 1.0;
+			};
+		}
+
+
+
+
+		static TUniqueFunction<double(const FSculptBrushStamp&, const FVector3d&)> MakeInverseFalloff()
+		{
+			return [](const FSculptBrushStamp& Stamp, const FVector3d& Position)
+			{
+				double Distance = Position.Distance(Stamp.LocalFrame.Origin);
+				double FalloffT = FMathd::Clamp(1.0 - Stamp.Falloff, 0.0, 1.0);
+				double UnitDistance = Distance / Stamp.Radius;
+				if (UnitDistance > FalloffT)
+				{
+					UnitDistance = FMathd::Clamp((UnitDistance - FalloffT) / (1.0 - FalloffT), 0.0, 1.0);
+					double Weight = FMathd::Clamp(1.0 - UnitDistance, 0.0, 1.0);
+					return Weight * Weight * Weight;
+				}
+				return 1.0;
+			};
+		}
+
+
+
+		static TUniqueFunction<double(const FSculptBrushStamp&, const FVector3d&)> MakeRoundFalloff()
+		{
+			return [](const FSculptBrushStamp& Stamp, const FVector3d& Position)
+			{
+				double Distance = Position.Distance(Stamp.LocalFrame.Origin);
+				double FalloffT = FMathd::Clamp(1.0 - Stamp.Falloff, 0.0, 1.0);
+				double UnitDistance = Distance / Stamp.Radius;
+				if (UnitDistance > FalloffT)
+				{
+					UnitDistance = FMathd::Clamp((UnitDistance - FalloffT) / (1.0 - FalloffT), 0.0, 1.0);
+					return FMathd::Clamp(1.0 - UnitDistance * UnitDistance, 0.0, 1.0);
+				}
+				return 1.0;
+			};
+		}
+
+
+
+		static TUniqueFunction<double(const FSculptBrushStamp&, const FVector3d&)> MakeSmoothBoxFalloff()
+		{
+			return [](const FSculptBrushStamp& Stamp, const FVector3d& Position)
+			{
+				double f = FMathd::Clamp(1.0 - Stamp.Falloff, 0.0, 1.0);
+				double InnerRadius = f * Stamp.Radius;
+				double FalloffWidth = Stamp.Radius - InnerRadius;
+				double BoxHalfWidth = FMathd::InvSqrt2 * InnerRadius;
+
+				FVector3d LocalPos = Stamp.LocalFrame.ToFramePoint(Position);
+				FAxisAlignedBox3d Box(FVector3d::Zero(), BoxHalfWidth);
+				if (Box.Contains(LocalPos))
+				{
+					return 1.0;
+				}
+				else
+				{
+					double DistSqr = Box.DistanceSquared(LocalPos);
+					if (DistSqr < FalloffWidth * FalloffWidth)
+					{
+						double t = FMathd::Clamp(FMathd::Sqrt(DistSqr) / FalloffWidth, 0.0, 1.0);
+						double w = (1.0 - t * t);
+						return w * w * w;
+					}
+					return 0.0;
+				}
+			};
+		}
+
+
+
+		static TUniqueFunction<double(const FSculptBrushStamp&, const FVector3d&)> MakeLinearBoxFalloff()
+		{
+			return [](const FSculptBrushStamp& Stamp, const FVector3d& Position)
+			{
+				double f = FMathd::Clamp(1.0 - Stamp.Falloff, 0.0, 1.0);
+				double InnerRadius = f * Stamp.Radius;
+				double FalloffWidth = Stamp.Radius - InnerRadius;
+				double BoxHalfWidth = FMathd::InvSqrt2 * InnerRadius;
+
+				FVector3d LocalPos = Stamp.LocalFrame.ToFramePoint(Position);
+				FAxisAlignedBox3d Box(FVector3d::Zero(), BoxHalfWidth);
+				if (Box.Contains(LocalPos))
+				{
+					return 1.0;
+				}
+				else
+				{
+					double DistSqr = Box.DistanceSquared(LocalPos);
+					if (DistSqr < FalloffWidth * FalloffWidth)
+					{
+						double t = FMathd::Clamp(FMathd::Sqrt(DistSqr) / FalloffWidth, 0.0, 1.0);
+						return (1.0 - t);
+					}
+					return 0.0;
+				}
+			};
+		}
+
+
+
+		static TUniqueFunction<double(const FSculptBrushStamp&, const FVector3d&)> MakeRoundBoxFalloff()
+		{
+			return [](const FSculptBrushStamp& Stamp, const FVector3d& Position)
+			{
+				double f = FMathd::Clamp(1.0 - Stamp.Falloff, 0.0, 1.0);
+				double InnerRadius = f * Stamp.Radius;
+				double FalloffWidth = Stamp.Radius - InnerRadius;
+				double BoxHalfWidth = FMathd::InvSqrt2 * InnerRadius;
+
+				FVector3d LocalPos = Stamp.LocalFrame.ToFramePoint(Position);
+				FAxisAlignedBox3d Box(FVector3d::Zero(), BoxHalfWidth);
+				if (Box.Contains(LocalPos))
+				{
+					return 1.0;
+				}
+				else
+				{
+					double DistSqr = Box.DistanceSquared(LocalPos);
+					if (DistSqr < FalloffWidth * FalloffWidth)
+					{
+						double t = FMathd::Clamp(FMathd::Sqrt(DistSqr) / FalloffWidth, 0.0, 1.0);
+						return (1.0 - t*t);
+					}
+					return 0.0;
+				}
+			};
+		}
+
+
+
+		static TUniqueFunction<double(const FSculptBrushStamp&, const FVector3d&)> MakeInverseBoxFalloff()
+		{
+			return [](const FSculptBrushStamp& Stamp, const FVector3d& Position)
+			{
+				double f = FMathd::Clamp(1.0 - Stamp.Falloff, 0.0, 1.0);
+				double InnerRadius = f * Stamp.Radius;
+				double FalloffWidth = Stamp.Radius - InnerRadius;
+				double BoxHalfWidth = FMathd::InvSqrt2 * InnerRadius;
+
+				FVector3d LocalPos = Stamp.LocalFrame.ToFramePoint(Position);
+				FAxisAlignedBox3d Box(FVector3d::Zero(), BoxHalfWidth);
+				if (Box.Contains(LocalPos))
+				{
+					return 1.0;
+				}
+				else
+				{
+					double DistSqr = Box.DistanceSquared(LocalPos);
+					if (DistSqr < FalloffWidth * FalloffWidth)
+					{
+						double t = FMathd::Clamp(FMathd::Sqrt(DistSqr) / FalloffWidth, 0.0, 1.0);
+						double w = (1.0 - t);
+						return w * w * w;
+					}
+					return 0.0;
+				}
+			};
+		}
+
+
+
+
+	}
+}
+
+
