ModelingTools: add missing header
#rnx

diff --git a/Engine/Plugins/Experimental/MeshModelingToolset/Source/MeshModelingTools/Public/Sculpting/MeshBrushOpBase.h b/Engine/Plugins/Experimental/MeshModelingToolset/Source/MeshModelingTools/Public/Sculpting/MeshBrushOpBase.h
--- a/Engine/Plugins/Experimental/MeshModelingToolset/Source/MeshModelingTools/Public/Sculpting/MeshBrushOpBase.h
+++ b/Engine/Plugins/Experimental/MeshModelingToolset/Source/MeshModelingTools/Public/Sculpting/MeshBrushOpBase.h
@@ -5,6 +5,7 @@
 #include "CoreMinimal.h"
 #include "VectorTypes.h"
 #include "FrameTypes.h"
+#include "InteractiveTool.h"
 #include "MeshBrushOpBase.generated.h"
 
 class FDynamicMesh3;
