ModelingComponents: get rid of log spam
#rnx

diff --git a/Engine/Plugins/Experimental/MeshModelingToolset/Source/ModelingComponents/Private/SimpleDynamicMeshComponent.cpp b/Engine/Plugins/Experimental/MeshModelingToolset/Source/ModelingComponents/Private/SimpleDynamicMeshComponent.cpp
--- a/Engine/Plugins/Experimental/MeshModelingToolset/Source/ModelingComponents/Private/SimpleDynamicMeshComponent.cpp
+++ b/Engine/Plugins/Experimental/MeshModelingToolset/Source/ModelingComponents/Private/SimpleDynamicMeshComponent.cpp
@@ -348,7 +348,7 @@ void USimpleDynamicMeshComponent::FastNotifyTriangleVerticesUpdated(const TSet<i
 			TotalTris += Decomposition->GetGroup(SetID).Triangles.Num();
 		}
 
-		UE_LOG(LogTemp, Warning, TEXT("Updating %d groups with %d tris"), UpdatedSets.Num(), TotalTris);
+		//UE_LOG(LogTemp, Warning, TEXT("Updating %d groups with %d tris"), UpdatedSets.Num(), TotalTris);
 
 		bool bPositions = (UpdatedAttributes & EMeshRenderAttributeFlags::Positions) != EMeshRenderAttributeFlags::None;
 
