ModelingTools: remove scale dependency in Pinch brush, expose perp-drag control as an option, use Falloff parameter to interpolate between linear and inverse-cubic falloff, tune speed
#rnx

diff --git a/Engine/Plugins/Experimental/MeshModelingToolset/Source/MeshModelingTools/Public/Sculpting/MeshPinchBrushOps.h b/Engine/Plugins/Experimental/MeshModelingToolset/Source/MeshModelingTools/Public/Sculpting/MeshPinchBrushOps.h
--- a/Engine/Plugins/Experimental/MeshModelingToolset/Source/MeshModelingTools/Public/Sculpting/MeshPinchBrushOps.h
+++ b/Engine/Plugins/Experimental/MeshModelingToolset/Source/MeshModelingTools/Public/Sculpting/MeshPinchBrushOps.h
@@ -20,12 +20,16 @@ public:
 
 	/** Amount of falloff to apply */
 	UPROPERTY(EditAnywhere, Category = PinchBrush, meta = (DisplayName = "Falloff", UIMin = "0.0", UIMax = "1.0", ClampMin = "0.0", ClampMax = "1.0"))
-	float Falloff = 1.0;
+	float Falloff = 0.75;
 
 	/** Depth of Brush into surface along surface normal */
 	UPROPERTY(EditAnywhere, Category = PinchBrush, meta = (UIMin = "-0.5", UIMax = "0.5", ClampMin = "-1.0", ClampMax = "1.0"))
 	float Depth = 0;
 
+	/** When enabled, brush will damp motion of vertices that would move perpendicular to brush stroke direction */
+	UPROPERTY(EditAnywhere, Category = PinchBrush)
+	bool bPerpDamping = true;
+
 	virtual float GetStrength() override { return Strength; }
 	virtual float GetFalloff() override { return Falloff; }
 	virtual float GetDepth() override { return Depth; }
@@ -36,7 +40,7 @@ class FPinchBrushOp : public FMeshSculptBrushOp
 {
 
 public:
-	double BrushSpeedTuning = 6.0;
+	double BrushSpeedTuning = 3.0;
 
 	FVector3d LastSmoothBrushPosLocal;
 	FVector3d LastSmoothBrushNormalLocal;;
@@ -65,29 +69,31 @@ public:
 
 		LastSmoothBrushPosLocal = NewSmoothBrushPosLocal;
 
+		bool bLimitDrag = GetPropertySetAs<UPinchBrushOpProps>()->bPerpDamping;
+
 		ParallelFor(Vertices.Num(), [&](int32 k)
 		{
 			int32 VertIdx = Vertices[k];
 			FVector3d OrigPos = Mesh->GetVertex(VertIdx);
 
 			FVector3d Delta = DepthPosLocal - OrigPos;
-			FVector3d MoveVec = UseSpeed * Delta;
-
-			//double Falloff = GetFalloff().Evaluate(Stamp, OrigPos);
+			Delta.Normalize();
+			FVector3d MoveVec = Delta;
 
 			// pinch uses 1/x falloff
 			double Distance = OrigPos.Distance(NewSmoothBrushPosLocal);
 			double NormalizedDistance = (Distance / Stamp.Radius) + 0.0001;
-			double Falloff = FMathd::Clamp(1.0 - NormalizedDistance, 0.0, 1.0);
-			Falloff *= Falloff; Falloff *= Falloff;
+			double LinearFalloff = FMathd::Clamp(1.0 - NormalizedDistance, 0.0, 1.0);
+			double InvFalloff = LinearFalloff * LinearFalloff * LinearFalloff;
+			double UseFalloff = FMathd::Lerp(LinearFalloff, InvFalloff, Stamp.Falloff);
 
-			if (bHaveMotion && Falloff < 0.8f)
+			if (bLimitDrag && bHaveMotion && UseFalloff < 0.7f)
 			{
 				double AnglePower = 1.0 - FMathd::Abs(MoveVec.Normalized().Dot(MotionVec));
-				Falloff *= AnglePower;
+				UseFalloff *= AnglePower;
 			}
 
-			FVector3d NewPos = OrigPos + Falloff * MoveVec;
+			FVector3d NewPos = OrigPos + UseFalloff * UseSpeed * MoveVec;
 			NewPositionsOut[k] = NewPos;
 		});
 	}
