ModelingComponents: don't allow null material in decomposition to propagate
#rnx

diff --git a/Engine/Plugins/Experimental/MeshModelingToolset/Source/ModelingComponents/Private/SimpleDynamicMeshSceneProxy.h b/Engine/Plugins/Experimental/MeshModelingToolset/Source/ModelingComponents/Private/SimpleDynamicMeshSceneProxy.h
--- a/Engine/Plugins/Experimental/MeshModelingToolset/Source/ModelingComponents/Private/SimpleDynamicMeshSceneProxy.h
+++ b/Engine/Plugins/Experimental/MeshModelingToolset/Source/ModelingComponents/Private/SimpleDynamicMeshSceneProxy.h
@@ -95,6 +95,10 @@ public:
 		{
 			RenderBufferSets[k] = AllocateNewRenderBufferSet();
 			RenderBufferSets[k]->Material = Decomposition->GetGroup(k).Material;
+			if (RenderBufferSets[k]->Material == nullptr)
+			{
+				RenderBufferSets[k]->Material = UMaterial::GetDefaultMaterial(MD_Surface);
+			}
 		}
 
 		bIsSingleBuffer = false;
