ModelingTools: enable weighted implicit smoothing in CotanSmoothingOp, expose in SmoothingTool. Update SmoothingTool to use Tool Property Sets. Remove unused MeanValueSmoothingOp, move functionality to IterativeSmoohtingOp along with old implicit code (which was still doing iterations).
#rnx

diff --git a/Engine/Plugins/Experimental/MeshModelingToolset/Source/MeshModelingTools/Private/SmoothMeshTool.cpp b/Engine/Plugins/Experimental/MeshModelingToolset/Source/MeshModelingTools/Private/SmoothMeshTool.cpp
--- a/Engine/Plugins/Experimental/MeshModelingToolset/Source/MeshModelingTools/Private/SmoothMeshTool.cpp
+++ b/Engine/Plugins/Experimental/MeshModelingToolset/Source/MeshModelingTools/Private/SmoothMeshTool.cpp
@@ -16,7 +16,6 @@
 // Smoothing operators
 #include "SmoothingOps/IterativeSmoothingOp.h"
 #include "SmoothingOps/CotanSmoothingOp.h"
-//#include "SmoothingOps/MeanValueSmoothingOp.h"
 
 
 #define LOCTEXT_NAMESPACE "USmoothMeshTool"
@@ -90,7 +89,7 @@ void USmoothMeshTool::Setup()
 		Preview->ConfigureMaterials(MaterialSet.Materials,
 			ToolSetupUtil::GetDefaultWorkingMaterial(GetToolManager())
 		);
-
+		Preview->SetWorkingMaterialDelay(0.75);
 		Preview->PreviewMesh->SetTransform(ComponentTarget->GetWorldTransform());
 		Preview->PreviewMesh->UpdatePreview(&SrcDynamicMesh);
 	}
@@ -98,18 +97,57 @@ void USmoothMeshTool::Setup()
 	// show the preview mesh
 	Preview->SetVisibility(true);
 
-	// start the compute
-	Preview->InvalidateResult();    // start compute
+	SmoothProperties = NewObject<USmoothMeshToolProperties>(this);
+	AddToolPropertySource(SmoothProperties);
+	SmoothProperties->RestoreProperties(this);
+	SmoothProperties->WatchProperty(SmoothProperties->SmoothingType,
+		[&](ESmoothMeshToolSmoothType) { UpdateVisiblePropertySets(); InvalidateResult();  });
+	//SmoothProperties->WatchProperty(SmoothProperties->bPreserveUVs,
+	//	[&](bool) { InvalidateResult(); });
+
+	IterativeProperties = NewObject<UIterativeSmoothProperties>(this);
+	AddToolPropertySource(IterativeProperties);
+	IterativeProperties->RestoreProperties(this);
+	SetToolPropertySourceEnabled(IterativeProperties, false);
+	IterativeProperties->GetOnModified().AddLambda([this](UObject*, FProperty*) { InvalidateResult(); });
+	//IterativeProperties->WatchProperty(IterativeProperties->SmoothingPerStep,[&](float) { InvalidateResult(); } );
+	//IterativeProperties->WatchProperty(IterativeProperties->Steps, [&](int) { InvalidateResult(); });
+	//IterativeProperties->WatchProperty(IterativeProperties->bSmoothBoundary, [&](int) { InvalidateResult(); });
+
+	DiffusionProperties = NewObject<UDiffusionSmoothProperties>(this);
+	AddToolPropertySource(DiffusionProperties);
+	DiffusionProperties->RestoreProperties(this);
+	SetToolPropertySourceEnabled(DiffusionProperties, false);
+	DiffusionProperties->GetOnModified().AddLambda([this](UObject*, FProperty*) { InvalidateResult(); });
+	//DiffusionProperties->WatchProperty(DiffusionProperties->SmoothingPerStep, [&](float) { InvalidateResult(); });
+	//DiffusionProperties->WatchProperty(DiffusionProperties->Steps, [&](int) { InvalidateResult(); });
+
+	ImplicitProperties = NewObject<UImplicitSmoothProperties>(this);
+	AddToolPropertySource(ImplicitProperties);
+	ImplicitProperties->RestoreProperties(this);
+	SetToolPropertySourceEnabled(ImplicitProperties, false);
+	ImplicitProperties->GetOnModified().AddLambda([this](UObject*, FProperty*) { InvalidateResult(); });
+	//ImplicitProperties->WatchProperty(ImplicitProperties->SmoothSpeed, [&](float) { InvalidateResult(); });
+	//ImplicitProperties->WatchProperty(ImplicitProperties->Smoothness, [&](float) { InvalidateResult(); });
 
-	bResultValid = false;
+	// start the compute
+	InvalidateResult();
 
+	GetToolManager()->DisplayMessage(
+		LOCTEXT("StartSmoothToolMessage", "Smooths the mesh vertex positions using various smoothing methods."),
+		EToolMessageLevel::UserNotification);
 }
 
 
 void USmoothMeshTool::Shutdown(EToolShutdownType ShutdownType)
 {
+	SmoothProperties->SaveProperties(this);
+	IterativeProperties->SaveProperties(this);
+	DiffusionProperties->SaveProperties(this);
+	ImplicitProperties->SaveProperties(this);
+
 	// Restore (unhide) the source meshes
-		ComponentTarget->SetOwnerVisibility(true);
+	ComponentTarget->SetOwnerVisibility(true);
 
 	if (Preview != nullptr)
 	{
@@ -139,33 +177,25 @@ void USmoothMeshTool::Render(IToolsContextRenderAPI* RenderAPI)
 	UpdateResult();
 }
 
-#if WITH_EDITOR
-void USmoothMeshTool::PostEditChangeProperty(FPropertyChangedEvent& PropertyChangedEvent)
+
+
+void USmoothMeshTool::OnTick(float DeltaTime)
 {
-	// one of the parameters changed.  Dirty the result for force a recompute.
-	Preview->InvalidateResult();
-	bResultValid = false;
+	Preview->Tick(DeltaTime);
 }
-#endif
 
-void USmoothMeshTool::OnPropertyModified(UObject* PropertySet, FProperty* Property)
+void USmoothMeshTool::InvalidateResult()
 {
 	Preview->InvalidateResult();
 	bResultValid = false;
 }
 
-
-void USmoothMeshTool::OnTick(float DeltaTime)
-{
-	Preview->Tick(DeltaTime);
-}
-
 void USmoothMeshTool::UpdateResult()
 {
 	if (bResultValid) 
 	{
-				return;
-			}
+		return;
+	}
 
 	bResultValid = Preview->HaveValidResult();
 }
@@ -180,24 +210,61 @@ bool USmoothMeshTool::CanAccept() const
 	return bResultValid;
 }
 
+
+void USmoothMeshTool::UpdateVisiblePropertySets()
+{
+	SetToolPropertySourceEnabled(IterativeProperties, false);
+	SetToolPropertySourceEnabled(DiffusionProperties, false);
+	SetToolPropertySourceEnabled(ImplicitProperties, false);
+
+	switch (SmoothProperties->SmoothingType)
+	{
+	case ESmoothMeshToolSmoothType::Iterative:
+		SetToolPropertySourceEnabled(IterativeProperties, true);
+		break;
+	case ESmoothMeshToolSmoothType::Diffusion:
+		SetToolPropertySourceEnabled(DiffusionProperties, true);
+		break;
+	case ESmoothMeshToolSmoothType::Implicit:
+		SetToolPropertySourceEnabled(ImplicitProperties, true);
+		break;
+	}
+}
+
+
 TUniquePtr<FDynamicMeshOperator> USmoothMeshTool::MakeNewOperator()
 {
 	TUniquePtr<FSmoothingOpBase> MeshOp;
-	switch (SmoothType)
+	
+	FSmoothingOpBase::FOptions Options;
+
+	switch (SmoothProperties->SmoothingType)
 	{
 	default:
 	case ESmoothMeshToolSmoothType::Iterative:
-		MeshOp = MakeUnique<FIterativeSmoothingOp>(&SrcDynamicMesh, SmoothSpeed, SmoothIterations);
+		Options.SmoothAlpha = IterativeProperties->SmoothingPerStep;
+		Options.Iterations = IterativeProperties->Steps;
+		Options.bSmoothBoundary = IterativeProperties->bSmoothBoundary;
+		Options.bUniform = true;
+		Options.bUseImplicit = false;
+		MeshOp = MakeUnique<FIterativeSmoothingOp>(&SrcDynamicMesh, Options);
 		break;
 
-	case ESmoothMeshToolSmoothType::BiHarmonic_Cotan:
-		MeshOp = MakeUnique<FCotanSmoothingOp>(&SrcDynamicMesh, SmoothSpeed, SmoothIterations);
+	case ESmoothMeshToolSmoothType::Diffusion:
+		Options.SmoothAlpha = DiffusionProperties->SmoothingPerStep;
+		Options.Iterations = DiffusionProperties->Steps;
+		Options.bUniform = DiffusionProperties->bPreserveUVs == false;
+		Options.bUseImplicit = true;
+		MeshOp = MakeUnique<FIterativeSmoothingOp>(&SrcDynamicMesh, Options);
 		break;
-/**
-	case ESmoothMeshToolSmoothType::BiHarmonic_MVW:
-		MeshOp = MakeShared<FMeanValueSmoothingOp>(&SrcDynamicMesh, SmoothSpeed, SmoothIterations);
+
+	case ESmoothMeshToolSmoothType::Implicit:
+		Options.SmoothAlpha = ImplicitProperties->SmoothSpeed;
+		Options.SmoothPower = (ImplicitProperties->Smoothness >= 100.0) ? FMathf::MaxReal : ImplicitProperties->Smoothness;
+		Options.bUniform = ImplicitProperties->bPreserveUVs == false;
+		Options.bUseImplicit = true;
+		MeshOp = MakeUnique<FCotanSmoothingOp>(&SrcDynamicMesh, Options);
 		break;
-*/
 	}
 
 	const FTransform XForm = ComponentTarget->GetWorldTransform();
diff --git a/Engine/Plugins/Experimental/MeshModelingToolset/Source/MeshModelingTools/Public/SmoothMeshTool.h b/Engine/Plugins/Experimental/MeshModelingToolset/Source/MeshModelingTools/Public/SmoothMeshTool.h
--- a/Engine/Plugins/Experimental/MeshModelingToolset/Source/MeshModelingTools/Public/SmoothMeshTool.h
+++ b/Engine/Plugins/Experimental/MeshModelingToolset/Source/MeshModelingTools/Public/SmoothMeshTool.h
@@ -10,24 +10,10 @@
 #include "MeshOpPreviewHelpers.h"
 #include "SmoothMeshTool.generated.h"
 
-
 // predeclarations
 struct FMeshDescription;
 class USimpleDynamicMeshComponent;
 
-
-UENUM()
-enum class ESmoothMeshToolSmoothType : uint8
-{
-	/** Fast smooth with N iterations */
-	Iterative UMETA(DisplayName = "Fast Smoothing"),
-
-	/** Implicit smoothing, usually better with UV preservation */
-	BiHarmonic_Cotan UMETA(DisplayName = "Implicit Smoothing"),
-
-};
-
-
 /**
  *
  */
@@ -53,6 +39,105 @@ public:
 
 
 
+UENUM()
+enum class ESmoothMeshToolSmoothType : uint8
+{
+	/** Iterative smoothing with N iterations */
+	Iterative UMETA(DisplayName = "Fast Iterative"),
+
+	/** Implicit smoothing, produces smoother output and does a better job at preserving UVs, but can be very slow on large meshes */
+	Implicit UMETA(DisplayName = "Implicit Single"),
+
+	/** Iterative implicit-diffusion smoothing with N iterations */
+	Diffusion UMETA(DisplayName = "Implicit Diffusion")
+};
+
+
+
+/** PropertySet for properties affecting the Smoother. */
+UCLASS()
+class MESHMODELINGTOOLS_API USmoothMeshToolProperties : public UInteractiveToolPropertySet
+{
+	GENERATED_BODY()
+
+public:
+	/** Type of smoothing to apply */
+	UPROPERTY(EditAnywhere, Category = Options)
+	ESmoothMeshToolSmoothType SmoothingType = ESmoothMeshToolSmoothType::Iterative;
+};
+
+
+
+/** Properties for Iterative Smoothing */
+UCLASS()
+class MESHMODELINGTOOLS_API UIterativeSmoothProperties : public UInteractiveToolPropertySet
+{
+	GENERATED_BODY()
+public:
+	/** Amount of smoothing allowed per step. Smaller steps will avoid things like collapse of small/thin features. */
+	UPROPERTY(EditAnywhere, Category = IterativeSmoothing, meta = (UIMin = "0.0", UIMax = "1.0", ClampMin = "0.0", ClampMax = "1.0"))
+	float SmoothingPerStep = 1.0f;
+
+	/** Number of Smoothing iterations */
+	UPROPERTY(EditAnywhere, Category = IterativeSmoothing, meta = (UIMin = "0", UIMax = "100", ClampMin = "0", ClampMax = "1000"))
+	int32 Steps = 1;
+
+	/** If this is false, the smoother will try to reshape the triangles to be more regular, which will distort UVs */
+	UPROPERTY(EditAnywhere, Category = IterativeSmoothing)
+	bool bSmoothBoundary = true;
+};
+
+
+
+/** Properties for Diffusion Smoothing */
+UCLASS()
+class MESHMODELINGTOOLS_API UDiffusionSmoothProperties : public UInteractiveToolPropertySet
+{
+	GENERATED_BODY()
+
+public:
+	/** Amount of smoothing allowed per step. Smaller steps will avoid things like collapse of small/thin features. */
+	UPROPERTY(EditAnywhere, Category = DiffusionSmoothing, meta = (UIMin = "0.0", UIMax = "1.0", ClampMin = "0.0", ClampMax = "1.0"))
+	float SmoothingPerStep = 1.0f;
+
+	/** Number of Smoothing iterations */
+	UPROPERTY(EditAnywhere, Category = DiffusionSmoothing, meta = (UIMin = "0", UIMax = "100", ClampMin = "0", ClampMax = "1000"))
+	int32 Steps = 1;
+
+	/** If this is false, the smoother will try to reshape the triangles to be more regular, which will distort UVs */
+	UPROPERTY(EditAnywhere, Category = DiffusionSmoothing)
+	bool bPreserveUVs = true;
+};
+
+
+
+
+/** Properties for Implicit smoothing */
+UCLASS()
+class MESHMODELINGTOOLS_API UImplicitSmoothProperties : public UInteractiveToolPropertySet
+{
+	GENERATED_BODY()
+
+public:
+	/** Smoothing speed */
+	//UPROPERTY(EditAnywhere, Category = ImplicitSmoothing, meta = (UIMin = "0.0", UIMax = "1.0", ClampMin = "0.0", ClampMax = "1.0"))
+	UPROPERTY()
+	float SmoothSpeed = 0.1f;
+
+	/** Desired Smoothness. This is not a linear quantity, but larger numbers produce smoother results */
+	UPROPERTY(EditAnywhere, Category = ImplicitSmoothing, meta = (UIMin = "0.0", UIMax = "100.0", ClampMin = "0.0", ClampMax = "10000.0"))
+	float Smoothness = 1.0f;
+
+	/** If this is false, the smoother will try to reshape the triangles to be more regular, which will distort UVs */
+	UPROPERTY(EditAnywhere, Category = ImplicitSmoothing)
+	bool bPreserveUVs = true;
+};
+
+
+
+
+
+
 /**
  * Simple Mesh Smoothing Tool
  */
@@ -77,43 +162,33 @@ public:
 	virtual bool HasAccept() const override;
 	virtual bool CanAccept() const override;
 
-#if WITH_EDITOR
-	virtual void PostEditChangeProperty(FPropertyChangedEvent& PropertyChangedEvent) override;
-#endif
-
-	virtual void OnPropertyModified(UObject* PropertySet, FProperty* Property) override;
-
 	virtual TUniquePtr<FDynamicMeshOperator> MakeNewOperator() override;
-
-
 	
 protected:
-	// need to update bResultValid if these are modified, so we don't publicly expose them. 
-	// @todo setters/getters for these
+	UPROPERTY()
+	USmoothMeshToolProperties* SmoothProperties = nullptr;
 
-	/** primary brush mode */
-	UPROPERTY(EditAnywhere, Category = Options)
-	ESmoothMeshToolSmoothType SmoothType = ESmoothMeshToolSmoothType::Iterative;
+	UPROPERTY()
+	UIterativeSmoothProperties* IterativeProperties = nullptr;
 
-	/** Smoothing speed */
-	UPROPERTY(EditAnywhere, Category = Options, meta = (UIMin = "0.0", UIMax = "1.0", ClampMin = "0.0", ClampMax = "1.0"))
-	float SmoothSpeed = 0.1f;
+	UPROPERTY()
+	UDiffusionSmoothProperties* DiffusionProperties = nullptr;
 
-	/** Smoothing speed */
-	UPROPERTY(EditAnywhere, Category = Options, meta = (UIMin = "0", UIMax = "100", ClampMin = "0", ClampMax = "1000"))
-	int32 SmoothIterations = 1;
+	UPROPERTY()
+	UImplicitSmoothProperties* ImplicitProperties = nullptr;
 
+	void UpdateVisiblePropertySets();
 
 protected:
 	void UpdateResult();
 
 	bool bResultValid = false;
+	void InvalidateResult();
 
 	UPROPERTY()
 	UMeshOpPreviewWithBackgroundCompute* Preview = nullptr;
 	FDynamicMesh3 SrcDynamicMesh;
 
-
 	UWorld* TargetWorld = nullptr;
 	IToolsContextAssetAPI* AssetAPI = nullptr;
 };
diff --git a/Engine/Plugins/Experimental/MeshModelingToolset/Source/ModelingComponents/Public/MeshOpPreviewHelpers.h b/Engine/Plugins/Experimental/MeshModelingToolset/Source/ModelingComponents/Public/MeshOpPreviewHelpers.h
--- a/Engine/Plugins/Experimental/MeshModelingToolset/Source/ModelingComponents/Public/MeshOpPreviewHelpers.h
+++ b/Engine/Plugins/Experimental/MeshModelingToolset/Source/ModelingComponents/Public/MeshOpPreviewHelpers.h
@@ -115,6 +115,10 @@ public:
 	 */
 	void SetVisibility(bool bVisible);
 
+	/**
+	 * Set time that Preview will wait before showing working material
+	 */
+	void SetWorkingMaterialDelay(float TimeInSeconds) { SecondsBeforeWorkingMaterial = TimeInSeconds; }
 
 	/**
 	 * @return true if currently using the 'in progress' working material
@@ -155,7 +159,7 @@ protected:
 
 	bool bVisible = true;
 
-	const float SecondsBeforeWorkingMaterial = 2.0;
+	float SecondsBeforeWorkingMaterial = 2.0;
 
 	// this object manages the background computes
 	TUniquePtr<FBackgroundDynamicMeshComputeSource> BackgroundCompute;
diff --git a/Engine/Plugins/Experimental/MeshModelingToolset/Source/ModelingOperators/Private/SmoothingOps/CotanSmoothingOp.cpp b/Engine/Plugins/Experimental/MeshModelingToolset/Source/ModelingOperators/Private/SmoothingOps/CotanSmoothingOp.cpp
--- a/Engine/Plugins/Experimental/MeshModelingToolset/Source/ModelingOperators/Private/SmoothingOps/CotanSmoothingOp.cpp
+++ b/Engine/Plugins/Experimental/MeshModelingToolset/Source/ModelingOperators/Private/SmoothingOps/CotanSmoothingOp.cpp
@@ -2,10 +2,13 @@
 
 #include "SmoothingOps/CotanSmoothingOp.h"
 #include "Solvers/MeshSmoothing.h"
+#include "Solvers/ConstrainedMeshSmoother.h"
+#include "MeshCurvature.h"
+#include "MeshWeights.h"
 
 
-FCotanSmoothingOp::FCotanSmoothingOp(const FDynamicMesh3* Mesh, float Speed, int32 Iterations) :
-	FSmoothingOpBase(Mesh, Speed, Iterations)
+FCotanSmoothingOp::FCotanSmoothingOp(const FDynamicMesh3* Mesh, const FSmoothingOpBase::FOptions& OptionsIn) :
+	FSmoothingOpBase(Mesh, OptionsIn)
 {
 }
 
@@ -22,6 +25,35 @@ void FCotanSmoothingOp::CalculateResult(FProgressCancel* Progress)
 
 void FCotanSmoothingOp::Smooth()
 {
-	double Intensity = 1.;
-	UE::MeshDeformation::ComputeSmoothing_BiHarmonic(ELaplacianWeightScheme::ClampedCotangent, *ResultMesh, SmoothSpeed, Intensity, SmoothIterations, PositionBuffer);
+	ELaplacianWeightScheme UseScheme = ELaplacianWeightScheme::ClampedCotangent;
+	if (SmoothOptions.bUniform)
+	{
+		UseScheme = ELaplacianWeightScheme::Uniform;
+	}
+
+	TUniquePtr<IConstrainedMeshSolver> Smoother = UE::MeshDeformation::ConstructConstrainedMeshSmoother(
+		UseScheme, *ResultMesh);
+
+	double Power = SmoothOptions.SmoothPower;
+	if (Power < 0.0001)
+	{
+		for (int32 vid : ResultMesh->VertexIndicesItr())
+		{
+			PositionBuffer[vid] = ResultMesh->GetVertex(vid);
+		}
+	}
+	else if ( Power > 10000 )
+	{
+		Smoother->Deform(PositionBuffer);
+	}
+	else
+	{
+		double Weight = 1.0 / Power;
+		for (int32 vid : ResultMesh->VertexIndicesItr())
+		{
+			Smoother->AddConstraint(vid, Weight, ResultMesh->GetVertex(vid), false);
+		}
+		Smoother->Deform(PositionBuffer);
+	}
+
 }
diff --git a/Engine/Plugins/Experimental/MeshModelingToolset/Source/ModelingOperators/Private/SmoothingOps/IterativeSmoothingOp.cpp b/Engine/Plugins/Experimental/MeshModelingToolset/Source/ModelingOperators/Private/SmoothingOps/IterativeSmoothingOp.cpp
--- a/Engine/Plugins/Experimental/MeshModelingToolset/Source/ModelingOperators/Private/SmoothingOps/IterativeSmoothingOp.cpp
+++ b/Engine/Plugins/Experimental/MeshModelingToolset/Source/ModelingOperators/Private/SmoothingOps/IterativeSmoothingOp.cpp
@@ -1,10 +1,13 @@
 // Copyright Epic Games, Inc. All Rights Reserved.
 
 #include "SmoothingOps/IterativeSmoothingOp.h"
+#include "Solvers/MeshSmoothing.h"
+#include "MeshWeights.h"
+#include "MeshCurvature.h"
 #include "Async/ParallelFor.h"
 
-FIterativeSmoothingOp::FIterativeSmoothingOp(const FDynamicMesh3* Mesh, float Speed, int32 Iterations) :
-	FSmoothingOpBase(Mesh, Speed, Iterations)
+FIterativeSmoothingOp::FIterativeSmoothingOp(const FDynamicMesh3* Mesh, const FSmoothingOpBase::FOptions& OptionsIn) :
+	FSmoothingOpBase(Mesh, OptionsIn)
 {
 	SmoothedBuffer = PositionBuffer;
 }
@@ -12,55 +15,132 @@ FIterativeSmoothingOp::FIterativeSmoothingOp(const FDynamicMesh3* Mesh, float Sp
 void FIterativeSmoothingOp::CalculateResult(FProgressCancel* Progress)
 {
 	// Update the values in the position buffer with smoothed positions.
-	Smooth();
+	if (SmoothOptions.bUseImplicit)
+	{
+		if (SmoothOptions.bUniform)
+		{
+			Smooth_MeanValue();
+		}
+		else
+		{
+			Smooth_Implicit_Cotan();
+		}
+	}
+	else
+	{
+		Smooth_Forward(SmoothOptions.bUniform);
+	}
 
 	// Copy the results back into the result mesh and update normals
 	UpdateResultMesh();
 
 }
 
-void FIterativeSmoothingOp::Smooth()
+void FIterativeSmoothingOp::Smooth_Implicit_Cotan()
+{
+	double Intensity = 1.;
+	UE::MeshDeformation::ComputeSmoothing_BiHarmonic(ELaplacianWeightScheme::ClampedCotangent, 
+		*ResultMesh, SmoothOptions.SmoothAlpha, Intensity, SmoothOptions.Iterations, PositionBuffer);
+}
+
+
+void FIterativeSmoothingOp::Smooth_MeanValue()
+{
+	double Intensity = 1.;
+	UE::MeshDeformation::ComputeSmoothing_BiHarmonic(ELaplacianWeightScheme::MeanValue,
+		*ResultMesh, SmoothOptions.SmoothAlpha, Intensity, SmoothOptions.Iterations, PositionBuffer);
+}
+
+
+// Each iteration is equivalent to time advancing the diffusion system 
+// 
+//  d p / dt = k*k L[ p ]
+// 
+// with forward Euler
+//  
+// p^{n+1} = p^{n} + dt k * k L[ p^{n} ]
+// 
+// Where { L[p] }_i = Sum[ w_{ij} p_j ]  
+//        
+//        and w_{ii} = - Sum[ w_{ij}, j != i }
+//
+// Here SmoothSpeed = -dt * k * k * w_{ii}
+void FIterativeSmoothingOp::Smooth_Forward(bool bUniform)
 {
+	int32 NV = ResultMesh->MaxVertexID();
 
-	// Each iteration is equivalent to time advancing the diffusion system 
-	// 
-	//  d p / dt = k*k L[ p ]
-	// 
-	// with forward Euler
-	//  
-	// p^{n+1} = p^{n} + dt k * k L[ p^{n} ]
-	// 
-	// Where { L[p] }_i = Sum[ w_{ij} p_j ]  
-	//        
-	//        and w_{ii} = - Sum[ w_{ij}, j != i }
-	//
-	// Here SmoothSpeed = -dt * k * k * w_{ii}
-	// 
-	// NB: Here -w_{ii} is the valence of the ith vertex.  
-	//     and 'k' may be spatially dependent. 
-
-	int NV = ResultMesh->MaxVertexID();
-
-	for (int k = 0; k < SmoothIterations; ++k)
+	// cache boundary verts info
+	TArray<bool> bIsBoundary;
+	TArray<int32> BoundaryVerts;
+	bIsBoundary.SetNum(NV);
+	ParallelFor(NV, [&](int32 vid)
 	{
-		ParallelFor(NV, [=](int vid) {
-			if (ResultMesh->IsReferencedVertex(vid) == false)
+		bIsBoundary[vid] = ResultMesh->IsBoundaryVertex(vid) && ResultMesh->IsReferencedVertex(vid);
+	});
+	for (int32 vid = 0; vid < NV; ++vid)
+	{
+		if (bIsBoundary[vid])
+		{
+			BoundaryVerts.Add(vid);
+		}
+	}
+	// have to limit boundary alpha to avoid oscillations because there are usually only 2 nbrs
+	double ClampedBoundaryAlpha = FMathd::Lerp(0.0, 0.9, SmoothOptions.SmoothAlpha);
+
+
+	for (int32 k = 0; k < SmoothOptions.Iterations; ++k)
+	{
+		// calculate smoothed positions of interior vertices
+		ParallelFor(NV, [&](int32 vid) 
+		{
+			if (ResultMesh->IsReferencedVertex(vid) == false || bIsBoundary[vid])
 			{
+				SmoothedBuffer[vid] = PositionBuffer[vid];
 				return;
 			}
 
-			FVector3d Centroid(0, 0, 0);
-			int NumNbrs = 0;
-			for (int nbrvid : ResultMesh->VtxVerticesItr(vid))
+			FVector3d Centroid;
+			if (bUniform)
 			{
-				Centroid += PositionBuffer[nbrvid];
-				NumNbrs++;
+				Centroid = FMeshWeights::UniformCentroid(*ResultMesh, vid, [&](int32 nbrvid) { return PositionBuffer[nbrvid]; });
 			}
-			Centroid /= (double)NumNbrs;
-			SmoothedBuffer[vid] = FVector3d::Lerp(PositionBuffer[vid], Centroid, SmoothSpeed);
-		}, false);
+			else
+			{
+				FVector3d Uniform = FMeshWeights::UniformCentroid(*ResultMesh, vid, [&](int32 nbrvid) { return PositionBuffer[nbrvid]; });
+				FVector3d MeanCurvNorm = UE::MeshCurvature::MeanCurvatureNormal(*ResultMesh, vid, [&](int32 nbrvid) { return PositionBuffer[nbrvid]; });
+				Centroid = PositionBuffer[vid] - 0.5*MeanCurvNorm;
+				if (Centroid.DistanceSquared(PositionBuffer[vid]) > Uniform.DistanceSquared(PositionBuffer[vid]))
+				{
+					Centroid = Uniform;
+				}
+			}
+
+			SmoothedBuffer[vid] = FVector3d::Lerp(PositionBuffer[vid], Centroid, SmoothOptions.SmoothAlpha);
+
+		} /*, EParallelForFlags::ForceSingleThread*/ );
+
+		// calculate boundary vertices
+		if (SmoothOptions.bSmoothBoundary)
+		{
+			ParallelFor(BoundaryVerts.Num(), [=](int32 idx)
+			{
+				int32 vid = BoundaryVerts[idx];
+				FVector3d Centroid(0, 0, 0);
+				int32 NumNbrs = 0;
+				for (int32 nbrvid : ResultMesh->VtxVerticesItr(vid))
+				{
+					if (bIsBoundary[nbrvid])
+					{
+						Centroid += PositionBuffer[nbrvid];
+						NumNbrs++;
+					}
+				}
+				Centroid = (NumNbrs > 0) ? (Centroid / (double)NumNbrs) : PositionBuffer[vid];
+				SmoothedBuffer[vid] = FVector3d::Lerp(PositionBuffer[vid], Centroid, ClampedBoundaryAlpha);
+			});
+		}
 
-		for (int vid = 0; vid < NV; ++vid)
+		for (int32 vid = 0; vid < NV; ++vid)
 		{
 			PositionBuffer[vid] = SmoothedBuffer[vid];
 		}
diff --git a/Engine/Plugins/Experimental/MeshModelingToolset/Source/ModelingOperators/Private/SmoothingOps/MeanValueSmoothingOp.cpp b/Engine/Plugins/Experimental/MeshModelingToolset/Source/ModelingOperators/Private/SmoothingOps/MeanValueSmoothingOp.cpp
deleted file mode 100644
diff --git a/Engine/Plugins/Experimental/MeshModelingToolset/Source/ModelingOperators/Private/SmoothingOps/SmoothingOpBase.cpp b/Engine/Plugins/Experimental/MeshModelingToolset/Source/ModelingOperators/Private/SmoothingOps/SmoothingOpBase.cpp
--- a/Engine/Plugins/Experimental/MeshModelingToolset/Source/ModelingOperators/Private/SmoothingOps/SmoothingOpBase.cpp
+++ b/Engine/Plugins/Experimental/MeshModelingToolset/Source/ModelingOperators/Private/SmoothingOps/SmoothingOpBase.cpp
@@ -6,10 +6,9 @@
 #include "MeshNormals.h"
 
 
-FSmoothingOpBase::FSmoothingOpBase(const FDynamicMesh3* Mesh, float Speed, int32 Iterations) :
-		FDynamicMeshOperator(),
-		SmoothSpeed(Speed),
-		SmoothIterations(Iterations)
+FSmoothingOpBase::FSmoothingOpBase(const FDynamicMesh3* Mesh, const FOptions& OptionsIn) :
+	FDynamicMeshOperator(),
+	SmoothOptions(OptionsIn)
 {
 	// deep copy the src mesh into the result mesh.  This ResultMesh will be directly updated by the smoothing.
 	ResultMesh->Copy(*Mesh);
diff --git a/Engine/Plugins/Experimental/MeshModelingToolset/Source/ModelingOperators/Public/SmoothingOps/CotanSmoothingOp.h b/Engine/Plugins/Experimental/MeshModelingToolset/Source/ModelingOperators/Public/SmoothingOps/CotanSmoothingOp.h
--- a/Engine/Plugins/Experimental/MeshModelingToolset/Source/ModelingOperators/Public/SmoothingOps/CotanSmoothingOp.h
+++ b/Engine/Plugins/Experimental/MeshModelingToolset/Source/ModelingOperators/Public/SmoothingOps/CotanSmoothingOp.h
@@ -9,7 +9,7 @@
 class MODELINGOPERATORS_API  FCotanSmoothingOp : public FSmoothingOpBase
 {
 public:
-	FCotanSmoothingOp(const FDynamicMesh3* Mesh, float Speed, int32 Iterations);
+	FCotanSmoothingOp(const FDynamicMesh3* Mesh, const FSmoothingOpBase::FOptions& OptionsIn);
 
 	~FCotanSmoothingOp() override {};
 
diff --git a/Engine/Plugins/Experimental/MeshModelingToolset/Source/ModelingOperators/Public/SmoothingOps/IterativeSmoothingOp.h b/Engine/Plugins/Experimental/MeshModelingToolset/Source/ModelingOperators/Public/SmoothingOps/IterativeSmoothingOp.h
--- a/Engine/Plugins/Experimental/MeshModelingToolset/Source/ModelingOperators/Public/SmoothingOps/IterativeSmoothingOp.h
+++ b/Engine/Plugins/Experimental/MeshModelingToolset/Source/ModelingOperators/Public/SmoothingOps/IterativeSmoothingOp.h
@@ -9,7 +9,7 @@
 class MODELINGOPERATORS_API  FIterativeSmoothingOp : public FSmoothingOpBase
 {
 public:
-	FIterativeSmoothingOp(const FDynamicMesh3* Mesh, float Speed, int32 Iterations);
+	FIterativeSmoothingOp(const FDynamicMesh3* Mesh, const FSmoothingOpBase::FOptions& OptionsIn);
 
 	~FIterativeSmoothingOp() override {};
 
@@ -18,9 +18,15 @@ public:
 
 private:
 
-	// Does the actual interative smoothign. 
-	void Smooth();
-	
+	// uniform iterative smoothing
+	void Smooth_Forward(bool bUniform);
+
+	// cotan smoothing iterations
+	void Smooth_Implicit_Cotan();
+
+	// mean value smoothing iterations
+	void Smooth_MeanValue();
+
 private:
 	TArray<FVector3d> SmoothedBuffer;
 };
\ No newline at end of file
diff --git a/Engine/Plugins/Experimental/MeshModelingToolset/Source/ModelingOperators/Public/SmoothingOps/MeanValueSmoothingOp.h b/Engine/Plugins/Experimental/MeshModelingToolset/Source/ModelingOperators/Public/SmoothingOps/MeanValueSmoothingOp.h
deleted file mode 100644
diff --git a/Engine/Plugins/Experimental/MeshModelingToolset/Source/ModelingOperators/Public/SmoothingOps/SmoothingOpBase.h b/Engine/Plugins/Experimental/MeshModelingToolset/Source/ModelingOperators/Public/SmoothingOps/SmoothingOpBase.h
--- a/Engine/Plugins/Experimental/MeshModelingToolset/Source/ModelingOperators/Public/SmoothingOps/SmoothingOpBase.h
+++ b/Engine/Plugins/Experimental/MeshModelingToolset/Source/ModelingOperators/Public/SmoothingOps/SmoothingOpBase.h
@@ -7,13 +7,38 @@
 #include "ModelingOperators.h"
 #include "VectorTypes.h"
 
-
 class FDynamicMesh3;
 
 class MODELINGOPERATORS_API FSmoothingOpBase : public FDynamicMeshOperator
 {
 public:
-	FSmoothingOpBase(const FDynamicMesh3* Mesh, float Speed, int32 Iterations);
+	struct FOptions
+	{
+		// value in range [0,1] where 0 is no smoothing and 1 is full smoothing
+		float SmoothAlpha = 1.0;
+
+		// number of iterations for iterative smoothing 
+		int32 Iterations = 1;
+
+		// Unconstrained value in range [0,FMathf::MaxValue] with 0=NoSmoothing and MaxValue=FullySmoothed
+		// Used by weighted implicit smoothing where weight is somewhat arbitrary...
+		float SmoothPower = 1.0;
+
+		// if true use implicit smoothing (where that is possible - depends on smoother?)
+		bool bUseImplicit = false;
+
+		// if true smooth the boundary, otherwise keep it fixed
+		bool bSmoothBoundary = true;
+
+		// if true use uniform weights, otherwise use something better
+		bool bUniform = false;
+
+		// use this value to clamp weights (eg for clamped mean value)
+		double WeightClamp = FMathf::MaxReal;
+	};
+
+
+	FSmoothingOpBase(const FDynamicMesh3* Mesh, const FOptions& OptionsIn);
 
 	virtual ~FSmoothingOpBase() override {}
 
@@ -27,8 +52,7 @@ public:
 	void UpdateResultMesh();
 
 protected:
-	float SmoothSpeed;
-	int32 SmoothIterations;
+	FOptions SmoothOptions;
 
 	TArray<FVector3d> PositionBuffer;
 
