ModelingTools: attempt to normalize for mesh scale and density in SmoothMeshTool. Mesh that is sent into smoother is scaled to have a constant surface area, and translated to origin. Smoothing power is scaled using empirically-tuned constants and vertex count. Net result is that smoothing power is much more consistent on (eg) model with same shape but different triangle count.

Also added parameter to Implicit smoothing that allows user to try to manually preserve volume (based on normal offset of constraint positions). Does not work very well, but it's kind of fun.

#rnx

diff --git a/Engine/Plugins/Experimental/MeshModelingToolset/Source/MeshModelingTools/Private/SmoothMeshTool.cpp b/Engine/Plugins/Experimental/MeshModelingToolset/Source/MeshModelingTools/Private/SmoothMeshTool.cpp
--- a/Engine/Plugins/Experimental/MeshModelingToolset/Source/MeshModelingTools/Private/SmoothMeshTool.cpp
+++ b/Engine/Plugins/Experimental/MeshModelingToolset/Source/MeshModelingTools/Private/SmoothMeshTool.cpp
@@ -5,6 +5,8 @@
 #include "ToolBuilderUtil.h"
 
 #include "DynamicMesh3.h"
+#include "MeshNormals.h"
+#include "MeshTransforms.h"
 #include "DynamicMeshToMeshDescription.h"
 #include "MeshDescriptionToDynamicMesh.h"
 
@@ -77,6 +79,27 @@ void USmoothMeshTool::Setup()
 	FMeshDescriptionToDynamicMesh Converter;
 	Converter.Convert(ComponentTarget->GetMesh(), SrcDynamicMesh);
 
+	// compute area of the input mesh and compute normalization scaling factor
+	FVector2d VolArea = TMeshQueries<FDynamicMesh3>::GetVolumeArea(SrcDynamicMesh);
+	double AreaScale = FMathd::Max(0.01, 6.0 / FMathd::Sqrt(VolArea.Y) );	// 6.0 is a bit arbitrary here...surface area of unit box
+
+	// translate to origin and then apply inverse of scale
+	FAxisAlignedBox3d Bounds = SrcDynamicMesh.GetCachedBounds();
+	SrcTranslate = Bounds.Center();
+	MeshTransforms::Translate(SrcDynamicMesh, -SrcTranslate);
+	SrcScale = AreaScale;
+	MeshTransforms::Scale(SrcDynamicMesh, (1.0/SrcScale)*FVector3d::One(), FVector3d::Zero() );
+
+	// apply that transform to target transform so that visible mesh stays in the same spot
+	OverrideTransform = ComponentTarget->GetWorldTransform();
+	FVector TranslateDelta = OverrideTransform.TransformVector((FVector)SrcTranslate);
+	FVector CurScale = OverrideTransform.GetScale3D();
+	OverrideTransform.AddToTranslation(TranslateDelta);
+	CurScale.X *= (float)SrcScale;
+	CurScale.Y *= (float)SrcScale;
+	CurScale.Z *= (float)SrcScale;
+	OverrideTransform.SetScale3D(CurScale);
+
 	// Initialize the preview mesh with a copy of the source mesh.
 	{
 		// Construct the preview object and set the material on it.
@@ -90,10 +113,14 @@ void USmoothMeshTool::Setup()
 			ToolSetupUtil::GetDefaultWorkingMaterial(GetToolManager())
 		);
 		Preview->SetWorkingMaterialDelay(0.75);
-		Preview->PreviewMesh->SetTransform(ComponentTarget->GetWorldTransform());
+		Preview->PreviewMesh->SetTransform(OverrideTransform);
 		Preview->PreviewMesh->UpdatePreview(&SrcDynamicMesh);
 	}
 
+	// calculate normals
+	BaseNormals = MakeShared<FMeshNormals>(&SrcDynamicMesh);
+	BaseNormals->ComputeVertexNormals();
+
 	// show the preview mesh
 	Preview->SetVisibility(true);
 
@@ -102,33 +129,24 @@ void USmoothMeshTool::Setup()
 	SmoothProperties->RestoreProperties(this);
 	SmoothProperties->WatchProperty(SmoothProperties->SmoothingType,
 		[&](ESmoothMeshToolSmoothType) { UpdateVisiblePropertySets(); InvalidateResult();  });
-	//SmoothProperties->WatchProperty(SmoothProperties->bPreserveUVs,
-	//	[&](bool) { InvalidateResult(); });
 
 	IterativeProperties = NewObject<UIterativeSmoothProperties>(this);
 	AddToolPropertySource(IterativeProperties);
 	IterativeProperties->RestoreProperties(this);
 	SetToolPropertySourceEnabled(IterativeProperties, false);
 	IterativeProperties->GetOnModified().AddLambda([this](UObject*, FProperty*) { InvalidateResult(); });
-	//IterativeProperties->WatchProperty(IterativeProperties->SmoothingPerStep,[&](float) { InvalidateResult(); } );
-	//IterativeProperties->WatchProperty(IterativeProperties->Steps, [&](int) { InvalidateResult(); });
-	//IterativeProperties->WatchProperty(IterativeProperties->bSmoothBoundary, [&](int) { InvalidateResult(); });
 
 	DiffusionProperties = NewObject<UDiffusionSmoothProperties>(this);
 	AddToolPropertySource(DiffusionProperties);
 	DiffusionProperties->RestoreProperties(this);
 	SetToolPropertySourceEnabled(DiffusionProperties, false);
 	DiffusionProperties->GetOnModified().AddLambda([this](UObject*, FProperty*) { InvalidateResult(); });
-	//DiffusionProperties->WatchProperty(DiffusionProperties->SmoothingPerStep, [&](float) { InvalidateResult(); });
-	//DiffusionProperties->WatchProperty(DiffusionProperties->Steps, [&](int) { InvalidateResult(); });
 
 	ImplicitProperties = NewObject<UImplicitSmoothProperties>(this);
 	AddToolPropertySource(ImplicitProperties);
 	ImplicitProperties->RestoreProperties(this);
 	SetToolPropertySourceEnabled(ImplicitProperties, false);
 	ImplicitProperties->GetOnModified().AddLambda([this](UObject*, FProperty*) { InvalidateResult(); });
-	//ImplicitProperties->WatchProperty(ImplicitProperties->SmoothSpeed, [&](float) { InvalidateResult(); });
-	//ImplicitProperties->WatchProperty(ImplicitProperties->Smoothness, [&](float) { InvalidateResult(); });
 
 	// start the compute
 	InvalidateResult();
@@ -160,6 +178,9 @@ void USmoothMeshTool::Shutdown(EToolShutdownType ShutdownType)
 			FDynamicMesh3* DynamicMeshResult = Result.Mesh.Get();
 			check(DynamicMeshResult != nullptr);
 
+			MeshTransforms::Scale(*DynamicMeshResult, FVector3d(SrcScale, SrcScale, SrcScale), FVector3d::Zero());
+			MeshTransforms::Translate(*DynamicMeshResult, SrcTranslate);
+
 			ComponentTarget->CommitMesh([DynamicMeshResult](const FPrimitiveComponentTarget::FCommitParams& CommitParams)
 			{
 				FDynamicMeshToMeshDescription Converter;
@@ -237,6 +258,7 @@ TUniquePtr<FDynamicMeshOperator> USmoothMeshTool::MakeNewOperator()
 	TUniquePtr<FSmoothingOpBase> MeshOp;
 	
 	FSmoothingOpBase::FOptions Options;
+	Options.BaseNormals = this->BaseNormals;
 
 	switch (SmoothProperties->SmoothingType)
 	{
@@ -259,16 +281,21 @@ TUniquePtr<FDynamicMeshOperator> USmoothMeshTool::MakeNewOperator()
 		break;
 
 	case ESmoothMeshToolSmoothType::Implicit:
+		{	
 		Options.SmoothAlpha = ImplicitProperties->SmoothSpeed;
-		Options.SmoothPower = (ImplicitProperties->Smoothness >= 100.0) ? FMathf::MaxReal : ImplicitProperties->Smoothness;
+		double NonlinearT = FMathd::Pow(ImplicitProperties->Smoothness, 2.0);
+		// this is an empirically-determined hack that seems to work OK to normalize the smoothing result for variable vertex count...
+		double ScaledPower = (NonlinearT/50.0) * SrcDynamicMesh.VertexCount();
+		Options.SmoothPower = ScaledPower;
 		Options.bUniform = ImplicitProperties->bPreserveUVs == false;
 		Options.bUseImplicit = true;
+		Options.NormalOffset = ImplicitProperties->VolumeCorrection;
 		MeshOp = MakeUnique<FCotanSmoothingOp>(&SrcDynamicMesh, Options);
+		}
 		break;
 	}
 
-	const FTransform XForm = ComponentTarget->GetWorldTransform();
-	FTransform3d XForm3d(XForm);
+	FTransform3d XForm3d(OverrideTransform);
 	MeshOp->SetTransform(XForm3d);
 
 	return MeshOp;
diff --git a/Engine/Plugins/Experimental/MeshModelingToolset/Source/MeshModelingTools/Public/SmoothMeshTool.h b/Engine/Plugins/Experimental/MeshModelingToolset/Source/MeshModelingTools/Public/SmoothMeshTool.h
--- a/Engine/Plugins/Experimental/MeshModelingToolset/Source/MeshModelingTools/Public/SmoothMeshTool.h
+++ b/Engine/Plugins/Experimental/MeshModelingToolset/Source/MeshModelingTools/Public/SmoothMeshTool.h
@@ -12,6 +12,7 @@
 
 // predeclarations
 struct FMeshDescription;
+class FMeshNormals;
 class USimpleDynamicMeshComponent;
 
 /**
@@ -46,10 +47,10 @@ enum class ESmoothMeshToolSmoothType : uint8
 	Iterative UMETA(DisplayName = "Fast Iterative"),
 
 	/** Implicit smoothing, produces smoother output and does a better job at preserving UVs, but can be very slow on large meshes */
-	Implicit UMETA(DisplayName = "Implicit Single"),
+	Implicit UMETA(DisplayName = "Fast Implicit"),
 
 	/** Iterative implicit-diffusion smoothing with N iterations */
-	Diffusion UMETA(DisplayName = "Implicit Diffusion")
+	Diffusion UMETA(DisplayName = "Iterative Diffusion")
 };
 
 
@@ -76,11 +77,11 @@ class MESHMODELINGTOOLS_API UIterativeSmoothProperties : public UInteractiveTool
 public:
 	/** Amount of smoothing allowed per step. Smaller steps will avoid things like collapse of small/thin features. */
 	UPROPERTY(EditAnywhere, Category = IterativeSmoothing, meta = (UIMin = "0.0", UIMax = "1.0", ClampMin = "0.0", ClampMax = "1.0"))
-	float SmoothingPerStep = 1.0f;
+	float SmoothingPerStep = 0.8f;
 
 	/** Number of Smoothing iterations */
 	UPROPERTY(EditAnywhere, Category = IterativeSmoothing, meta = (UIMin = "0", UIMax = "100", ClampMin = "0", ClampMax = "1000"))
-	int32 Steps = 1;
+	int32 Steps = 10;
 
 	/** If this is false, the smoother will try to reshape the triangles to be more regular, which will distort UVs */
 	UPROPERTY(EditAnywhere, Category = IterativeSmoothing)
@@ -98,7 +99,7 @@ class MESHMODELINGTOOLS_API UDiffusionSmoothProperties : public UInteractiveTool
 public:
 	/** Amount of smoothing allowed per step. Smaller steps will avoid things like collapse of small/thin features. */
 	UPROPERTY(EditAnywhere, Category = DiffusionSmoothing, meta = (UIMin = "0.0", UIMax = "1.0", ClampMin = "0.0", ClampMax = "1.0"))
-	float SmoothingPerStep = 1.0f;
+	float SmoothingPerStep = 0.8f;
 
 	/** Number of Smoothing iterations */
 	UPROPERTY(EditAnywhere, Category = DiffusionSmoothing, meta = (UIMin = "0", UIMax = "100", ClampMin = "0", ClampMax = "1000"))
@@ -125,12 +126,16 @@ public:
 	float SmoothSpeed = 0.1f;
 
 	/** Desired Smoothness. This is not a linear quantity, but larger numbers produce smoother results */
-	UPROPERTY(EditAnywhere, Category = ImplicitSmoothing, meta = (UIMin = "0.0", UIMax = "100.0", ClampMin = "0.0", ClampMax = "10000.0"))
-	float Smoothness = 1.0f;
+	UPROPERTY(EditAnywhere, Category = ImplicitSmoothing, meta = (UIMin = "0.0", UIMax = "1.0", ClampMin = "0.0", ClampMax = "100.0"))
+	float Smoothness = 0.2f;
 
 	/** If this is false, the smoother will try to reshape the triangles to be more regular, which will distort UVs */
 	UPROPERTY(EditAnywhere, Category = ImplicitSmoothing)
 	bool bPreserveUVs = true;
+
+	/** Magic number that allows you to try to correct for shrinking caused by smoothing */
+	UPROPERTY(EditAnywhere, Category = ImplicitSmoothing, meta = (UIMin = "0.0", UIMax = "1000.0", ClampMin = "0.0", ClampMax = "1000000.0"))
+	float VolumeCorrection = 0.0f;
 };
 
 
@@ -189,6 +194,14 @@ protected:
 	UMeshOpPreviewWithBackgroundCompute* Preview = nullptr;
 	FDynamicMesh3 SrcDynamicMesh;
 
+	// scale/translate applied to input mesh to regularize it
+	FVector3d SrcTranslate;
+	double SrcScale;
+	// transform that does the opposite of scale/translate so that mesh stays in the right spot on screen
+	FTransform OverrideTransform;
+
+	TSharedPtr<FMeshNormals> BaseNormals;
+
 	UWorld* TargetWorld = nullptr;
 	IToolsContextAssetAPI* AssetAPI = nullptr;
 };
diff --git a/Engine/Plugins/Experimental/MeshModelingToolset/Source/ModelingOperators/Private/SmoothingOps/CotanSmoothingOp.cpp b/Engine/Plugins/Experimental/MeshModelingToolset/Source/ModelingOperators/Private/SmoothingOps/CotanSmoothingOp.cpp
--- a/Engine/Plugins/Experimental/MeshModelingToolset/Source/ModelingOperators/Private/SmoothingOps/CotanSmoothingOp.cpp
+++ b/Engine/Plugins/Experimental/MeshModelingToolset/Source/ModelingOperators/Private/SmoothingOps/CotanSmoothingOp.cpp
@@ -5,6 +5,7 @@
 #include "Solvers/ConstrainedMeshSmoother.h"
 #include "MeshCurvature.h"
 #include "MeshWeights.h"
+#include "MeshNormals.h"
 
 
 FCotanSmoothingOp::FCotanSmoothingOp(const FDynamicMesh3* Mesh, const FSmoothingOpBase::FOptions& OptionsIn) :
@@ -51,7 +52,15 @@ void FCotanSmoothingOp::Smooth()
 		double Weight = 1.0 / Power;
 		for (int32 vid : ResultMesh->VertexIndicesItr())
 		{
-			Smoother->AddConstraint(vid, Weight, ResultMesh->GetVertex(vid), false);
+			FVector3d Position = ResultMesh->GetVertex(vid);
+
+			if (SmoothOptions.NormalOffset > 0)
+			{
+				check(SmoothOptions.BaseNormals.IsValid());
+				Position += SmoothOptions.NormalOffset * SmoothOptions.BaseNormals->GetNormals()[vid];
+			}
+
+			Smoother->AddConstraint(vid, Weight, Position, false);
 		}
 		Smoother->Deform(PositionBuffer);
 	}
diff --git a/Engine/Plugins/Experimental/MeshModelingToolset/Source/ModelingOperators/Public/SmoothingOps/SmoothingOpBase.h b/Engine/Plugins/Experimental/MeshModelingToolset/Source/ModelingOperators/Public/SmoothingOps/SmoothingOpBase.h
--- a/Engine/Plugins/Experimental/MeshModelingToolset/Source/ModelingOperators/Public/SmoothingOps/SmoothingOpBase.h
+++ b/Engine/Plugins/Experimental/MeshModelingToolset/Source/ModelingOperators/Public/SmoothingOps/SmoothingOpBase.h
@@ -8,6 +8,7 @@
 #include "VectorTypes.h"
 
 class FDynamicMesh3;
+class FMeshNormals;
 
 class MODELINGOPERATORS_API FSmoothingOpBase : public FDynamicMeshOperator
 {
@@ -35,6 +36,12 @@ public:
 
 		// use this value to clamp weights (eg for clamped mean value)
 		double WeightClamp = FMathf::MaxReal;
+
+		// mesh normals calculated for input mesh
+		TSharedPtr<FMeshNormals> BaseNormals;
+
+		// offset used by some smoothers
+		double NormalOffset = 0.0;
 	};
 
 
