ModelingTools: remove Q/A keybindings in Sculpt tool, problematic interaction w/ WASD navigation
#rnx

diff --git a/Engine/Plugins/Experimental/MeshModelingToolset/Source/MeshModelingTools/Private/DynamicMeshSculptTool.cpp b/Engine/Plugins/Experimental/MeshModelingToolset/Source/MeshModelingTools/Private/DynamicMeshSculptTool.cpp
--- a/Engine/Plugins/Experimental/MeshModelingToolset/Source/MeshModelingTools/Private/DynamicMeshSculptTool.cpp
+++ b/Engine/Plugins/Experimental/MeshModelingToolset/Source/MeshModelingTools/Private/DynamicMeshSculptTool.cpp
@@ -1986,25 +1986,6 @@ void UDynamicMeshSculptTool::DecreaseBrushFalloffAction()
 
 
 
-void UDynamicMeshSculptTool::NextBrushModeAction()
-{
-	uint8 LastMode = (uint8)EDynamicMeshSculptBrushType::LastValue;
-	SculptProperties->PrimaryBrushType = (EDynamicMeshSculptBrushType)(((uint8)SculptProperties->PrimaryBrushType + 1) % LastMode);
-}
-
-void UDynamicMeshSculptTool::PreviousBrushModeAction()
-{
-	uint8 LastMode = (uint8)EDynamicMeshSculptBrushType::LastValue;
-	uint8 CurMode = (uint8)SculptProperties->PrimaryBrushType;
-	if (CurMode == 0)
-	{
-		SculptProperties->PrimaryBrushType = (EDynamicMeshSculptBrushType)((uint8)LastMode - 1);
-	}
-	else
-	{
-		SculptProperties->PrimaryBrushType = (EDynamicMeshSculptBrushType)((uint8)CurMode - 1);
-	}
-}
 
 void UDynamicMeshSculptTool::NextHistoryBrushModeAction()
 {
@@ -2058,20 +2039,6 @@ void UDynamicMeshSculptTool::RegisterActions(FInteractiveToolActionSet& ActionSe
 		[this]() { DecreaseBrushFalloffAction(); });
 
 
-	ActionSet.RegisterAction(this, (int32)EStandardToolActions::BaseClientDefinedActionID+1,
-		TEXT("NextBrushMode"),
-		LOCTEXT("SculptNextBrushMode", "Next Brush Type"),
-		LOCTEXT("SculptNextBrushModeTooltip", "Cycle to next Brush Type"),
-		EModifierKey::None, EKeys::A,
-		[this]() { NextBrushModeAction(); });
-
-	ActionSet.RegisterAction(this, (int32)EStandardToolActions::BaseClientDefinedActionID+2,
-		TEXT("PreviousBrushMode"),
-		LOCTEXT("SculptPreviousBrushMode", "Previous Brush Type"),
-		LOCTEXT("SculptPreviousBrushModeTooltip", "Cycle to previous Brush Type"),
-		EModifierKey::None, EKeys::Q,
-		[this]() { PreviousBrushModeAction(); });
-
 	//ActionSet.RegisterAction(this, (int32)EStandardToolActions::BaseClientDefinedActionID + 10,
 	//	TEXT("NextBrushHistoryState"),
 	//	LOCTEXT("SculptNextBrushHistoryState", "Next Brush History State"),
@@ -2303,7 +2270,7 @@ void UDynamicMeshSculptTool::UpdateImageSetting(UTexture2D* NewImage)
 
 void UDynamicMeshSculptTool::UpdateBrushType(EDynamicMeshSculptBrushType BrushType)
 {
-	static const FText BaseMessage = LOCTEXT("OnStartSculptTool", "Hold Shift to Smooth, Ctrl to Invert (where applicable). Q/A keys cycle Brush Type. S/D changes Size (+Shift to small-step), W/E changes Strength.");
+	static const FText BaseMessage = LOCTEXT("OnStartSculptTool", "Hold Shift to Smooth, Ctrl to Invert (where applicable). [/] and S/D change Size (+Shift to small-step), W/E changes Strength.");
 	FTextBuilder Builder;
 	Builder.AppendLine(BaseMessage);
 
diff --git a/Engine/Plugins/Experimental/MeshModelingToolset/Source/MeshModelingTools/Private/MeshVertexSculptTool.cpp b/Engine/Plugins/Experimental/MeshModelingToolset/Source/MeshModelingTools/Private/MeshVertexSculptTool.cpp
--- a/Engine/Plugins/Experimental/MeshModelingToolset/Source/MeshModelingTools/Private/MeshVertexSculptTool.cpp
+++ b/Engine/Plugins/Experimental/MeshModelingToolset/Source/MeshModelingTools/Private/MeshVertexSculptTool.cpp
@@ -664,25 +664,6 @@ void UMeshVertexSculptTool::DecreaseBrushSpeedAction()
 
 
 
-void UMeshVertexSculptTool::NextBrushModeAction()
-{
-	uint8 LastMode = (uint8)EMeshVertexSculptBrushType::LastValue;
-	SculptProperties->PrimaryBrushType = (EMeshVertexSculptBrushType)(((uint8)SculptProperties->PrimaryBrushType + 1) % LastMode);
-}
-
-void UMeshVertexSculptTool::PreviousBrushModeAction()
-{
-	uint8 LastMode = (uint8)EMeshVertexSculptBrushType::LastValue;
-	uint8 CurMode = (uint8)SculptProperties->PrimaryBrushType;
-	if (CurMode == 0)
-	{
-		SculptProperties->PrimaryBrushType = (EMeshVertexSculptBrushType)((uint8)LastMode - 1);
-	}
-	else
-	{
-		SculptProperties->PrimaryBrushType = (EMeshVertexSculptBrushType)((uint8)CurMode - 1);
-	}
-}
 
 
 
@@ -776,7 +757,7 @@ void UMeshVertexSculptTool::OnDynamicMeshComponentChanged(USimpleDynamicMeshComp
 
 void UMeshVertexSculptTool::UpdateBrushType(EMeshVertexSculptBrushType BrushType)
 {
-	static const FText BaseMessage = LOCTEXT("OnStartSculptTool", "Hold Shift to Smooth, Ctrl to Invert (where applicable). Q/A keys cycle Brush Type. S/D changes Size (+Shift to small-step), W/E changes Strength.");
+	static const FText BaseMessage = LOCTEXT("OnStartSculptTool", "Hold Shift to Smooth, Ctrl to Invert (where applicable). [/] and S/D change Size (+Shift to small-step), W/E changes Strength.");
 	FTextBuilder Builder;
 	Builder.AppendLine(BaseMessage);
 
diff --git a/Engine/Plugins/Experimental/MeshModelingToolset/Source/MeshModelingTools/Public/DynamicMeshSculptTool.h b/Engine/Plugins/Experimental/MeshModelingToolset/Source/MeshModelingTools/Public/DynamicMeshSculptTool.h
--- a/Engine/Plugins/Experimental/MeshModelingToolset/Source/MeshModelingTools/Public/DynamicMeshSculptTool.h
+++ b/Engine/Plugins/Experimental/MeshModelingToolset/Source/MeshModelingTools/Public/DynamicMeshSculptTool.h
@@ -288,9 +288,6 @@ public:
 	virtual void IncreaseBrushSpeedAction();
 	virtual void DecreaseBrushSpeedAction();
 
-	virtual void NextBrushModeAction();
-	virtual void PreviousBrushModeAction();
-
 	virtual void NextHistoryBrushModeAction();
 	virtual void PreviousHistoryBrushModeAction();
 
diff --git a/Engine/Plugins/Experimental/MeshModelingToolset/Source/MeshModelingTools/Public/MeshVertexSculptTool.h b/Engine/Plugins/Experimental/MeshModelingToolset/Source/MeshModelingTools/Public/MeshVertexSculptTool.h
--- a/Engine/Plugins/Experimental/MeshModelingToolset/Source/MeshModelingTools/Public/MeshVertexSculptTool.h
+++ b/Engine/Plugins/Experimental/MeshModelingToolset/Source/MeshModelingTools/Public/MeshVertexSculptTool.h
@@ -162,8 +162,6 @@ public:
 public:
 	virtual void IncreaseBrushSpeedAction() override;
 	virtual void DecreaseBrushSpeedAction() override;
-	virtual void NextBrushModeAction() override;
-	virtual void PreviousBrushModeAction() override;
 
 protected:
 	// UMeshSculptToolBase API
