ModelingTools: Smooth Mesh Tool was scaling incorrectly, to very large size instead of unit size. Offset now has much larger effect, so change UI range to unit-dimension quantity.
#rnx

diff --git a/Engine/Plugins/Experimental/MeshModelingToolset/Source/MeshModelingTools/Private/SmoothMeshTool.cpp b/Engine/Plugins/Experimental/MeshModelingToolset/Source/MeshModelingTools/Private/SmoothMeshTool.cpp
--- a/Engine/Plugins/Experimental/MeshModelingToolset/Source/MeshModelingTools/Private/SmoothMeshTool.cpp
+++ b/Engine/Plugins/Experimental/MeshModelingToolset/Source/MeshModelingTools/Private/SmoothMeshTool.cpp
@@ -81,13 +81,13 @@ void USmoothMeshTool::Setup()
 
 	// compute area of the input mesh and compute normalization scaling factor
 	FVector2d VolArea = TMeshQueries<FDynamicMesh3>::GetVolumeArea(SrcDynamicMesh);
-	double AreaScale = FMathd::Max(0.01, 6.0 / FMathd::Sqrt(VolArea.Y) );	// 6.0 is a bit arbitrary here...surface area of unit box
+	double UnitScalingMeasure = FMathd::Max(0.01, FMathd::Sqrt(VolArea.Y / 6.0));  // 6.0 is a bit arbitrary here...surface area of unit box
 
 	// translate to origin and then apply inverse of scale
 	FAxisAlignedBox3d Bounds = SrcDynamicMesh.GetCachedBounds();
 	SrcTranslate = Bounds.Center();
 	MeshTransforms::Translate(SrcDynamicMesh, -SrcTranslate);
-	SrcScale = AreaScale;
+	SrcScale = UnitScalingMeasure;
 	MeshTransforms::Scale(SrcDynamicMesh, (1.0/SrcScale)*FVector3d::One(), FVector3d::Zero() );
 
 	// apply that transform to target transform so that visible mesh stays in the same spot
diff --git a/Engine/Plugins/Experimental/MeshModelingToolset/Source/MeshModelingTools/Public/SmoothMeshTool.h b/Engine/Plugins/Experimental/MeshModelingToolset/Source/MeshModelingTools/Public/SmoothMeshTool.h
--- a/Engine/Plugins/Experimental/MeshModelingToolset/Source/MeshModelingTools/Public/SmoothMeshTool.h
+++ b/Engine/Plugins/Experimental/MeshModelingToolset/Source/MeshModelingTools/Public/SmoothMeshTool.h
@@ -134,7 +134,7 @@ public:
 	bool bPreserveUVs = true;
 
 	/** Magic number that allows you to try to correct for shrinking caused by smoothing */
-	UPROPERTY(EditAnywhere, Category = ImplicitSmoothing, meta = (UIMin = "0.0", UIMax = "1000.0", ClampMin = "0.0", ClampMax = "1000000.0"))
+	UPROPERTY(EditAnywhere, Category = ImplicitSmoothing, meta = (UIMin = "0.0", UIMax = "1.0", ClampMin = "0.0", ClampMax = "10.0"))
 	float VolumeCorrection = 0.0f;
 };
 
