Added modeling tool to convert BSP brushes to static meshes.

#rnx

diff --git a/Engine/Plugins/Experimental/MeshModelingToolset/Source/MeshModelingToolsEditorOnly/Private/BspConversionTool.cpp b/Engine/Plugins/Experimental/MeshModelingToolset/Source/MeshModelingToolsEditorOnly/Private/BspConversionTool.cpp
new file mode 100644
--- /dev/null
+++ b/Engine/Plugins/Experimental/MeshModelingToolset/Source/MeshModelingToolsEditorOnly/Private/BspConversionTool.cpp
@@ -0,0 +1,864 @@
+// Copyright Epic Games, Inc. All Rights Reserved.
+
+#include "BspConversionTool.h"
+
+#include "AssetGenerationUtil.h"
+#include "AssetSelection.h"
+#include "ComponentSourceInterfaces.h"
+#include "DynamicMesh3.h"
+#include "Editor.h"
+#include "Editor/EditorEngine.h"
+#include "Engine/StaticMeshActor.h"
+#include "InteractiveToolManager.h"
+#include "MeshDescriptionToDynamicMesh.h"
+#include "MeshTransforms.h"
+#include "Model.h"
+#include "Operations/MeshBoolean.h"
+#include "StaticMeshAttributes.h"
+#include "ToolBuilderUtil.h"
+#include "Tools/EditorComponentSourceFactory.h"
+#include "ToolSetupUtil.h"
+#include "Engine/Selection.h"
+
+#include "Logging/MessageLog.h"
+#include "Misc/MessageDialog.h"
+#include "Logging/TokenizedMessage.h"
+
+#define LOCTEXT_NAMESPACE "UBspConversionTool"
+
+// Forward declarations of local functions
+void ConvertBrushesToDynamicMesh(TArray<ABrush*>& BrushesToConvert, FDynamicMesh3& OutputMesh, 
+	TArray<UMaterialInterface*>& OutputMaterials);
+void ApplyStaticMeshBooleanOperation(
+	FDynamicMesh3& MeshA, const FTransform3d& TransformA, const TArray<UMaterialInterface*>& MaterialsA,
+	FDynamicMesh3& MeshB, const FTransform3d& TransformB, const TArray<UMaterialInterface*>& MaterialsB,
+	FDynamicMesh3& OutputMesh, FTransform3d& OutputTransform, TArray<UMaterialInterface*>& OutputMaterials,
+	FMeshBoolean::EBooleanOp Operation);
+FText GetBrushGeometryErrorMessage(ABrush* Brush);
+
+// Element stored in CachedBrushes: the resulting dynamic mesh and the materials array
+typedef TPair<TSharedPtr<const FDynamicMesh3>, TSharedPtr<const TArray<UMaterialInterface*>>> FCachedResult;
+
+#pragma region Tool builder functions
+bool UBspConversionToolBuilder::CanBuildTool(const FToolBuilderState& SceneState) const
+{
+	// We allow the tool to be built even if nothing is selected because the tool has a "select all" option.
+	// We just need to be able to get the place to save assets.
+	return AssetAPI != nullptr;
+}
+
+UInteractiveTool* UBspConversionToolBuilder::BuildTool(const FToolBuilderState& SceneState) const
+{
+	UBspConversionTool* NewTool = NewObject<UBspConversionTool>(SceneState.ToolManager);
+	NewTool->SetWorld(SceneState.World);
+	NewTool->SetAssetAPI(AssetAPI);
+
+	return NewTool;
+}
+#pragma endregion
+
+#pragma region Tool property functions
+void UBspConversionToolActionPropertySet::PostAction(EBspConversionToolAction Action)
+{
+	if (ParentTool.IsValid())
+	{
+		ParentTool->RequestAction(Action);
+	}
+}
+#pragma endregion
+
+#pragma region Unused tool scaffolding
+UBspConversionTool::UBspConversionTool()
+{
+}
+
+void UBspConversionTool::RegisterActions(FInteractiveToolActionSet& ActionSet)
+{
+}
+
+void UBspConversionTool::Render(IToolsContextRenderAPI* RenderAPI)
+{
+}
+#pragma endregion
+
+#pragma region Tool scaffolding
+bool UBspConversionTool::CanAccept() const
+{
+	// We precompute this value and update it at various editor events
+	return bCanAccept;
+}
+
+void UBspConversionTool::Setup()
+{
+	UInteractiveTool::Setup();
+
+	// Link in the tool properties and actions that will be displayed in the side panel
+	Settings = NewObject<UBspConversionToolProperties>();
+	Settings->RestoreProperties(this);
+	AddToolPropertySource(Settings);
+
+	ToolActions = NewObject<UBspConversionToolActionPropertySet>(this);
+	ToolActions->Initialize(this);
+	AddToolPropertySource(ToolActions);
+
+	// Give a description to put in the side panel
+	GetToolManager()->DisplayMessage(
+		LOCTEXT("OnStartTool", "Convert geometry brushes (also known as BSP brushes) into a single static mesh. \"Convert then Combine\" "
+			"first converts the individual brushes and performs boolean operations on the resulting meshes (this requires individual "
+			"brushes to have manifold geometry. \"Combine then Convert\" is the old functionality, combining brushes, then converting them."),
+		EToolMessageLevel::UserNotification);
+
+	// We write out an empty warning message to make the sidebar look unchanged if we write out a warning message and then
+	// change it to be an empty string to clear it.
+	GetToolManager()->DisplayMessage(FText(), EToolMessageLevel::UserWarning);
+
+	// See if we have valid targets selected
+	bCanAccept = AtLeastOneValidConversionTarget();
+
+	// Set up the preview mesh
+	PreviewMesh = NewObject<UPreviewMesh>(this, "Preview Mesh");
+	PreviewMesh->CreateInWorld(TargetWorld, FTransform::Identity);
+
+	// The material is set to have opacity 0, so that only the wireframe shows.
+	PreviewMesh->SetOverrideRenderMaterial(ToolSetupUtil::GetSimpleCustomMaterial(GetToolManager(), FLinearColor(), 0.0f));
+	PreviewMesh->EnableWireframe(Settings->bShowPreview);
+
+	// Register with any editor events that will require us to recompute the preview or update bCanAccept
+	USelection::SelectionChangedEvent.AddUObject(this, &UBspConversionTool::OnEditorSelectionChanged);
+	GEngine->OnLevelActorListChanged().AddUObject(this, &UBspConversionTool::OnEditorLevelActorListChanged);
+	GEngine->OnActorMoved().AddUObject(this, &UBspConversionTool::OnEditorActorMoved);
+
+	// Do some precomputing if we show a preview.
+	if (bCanAccept && Settings->bShowPreview)
+	{
+		CompareAndUpdateConversionTargets();
+
+		FText ErrorMessage;
+		// When we do the actual computation, we may find new errors that prevent us from accepting
+		bCanAccept = ComputeAndUpdatePreviewMesh(&ErrorMessage);
+
+		if (!ErrorMessage.IsEmpty())
+		{
+			GetToolManager()->DisplayMessage(ErrorMessage, EToolMessageLevel::UserWarning);
+		}
+	}
+}
+
+/**
+ * Returns true if there is at least one valid conversion actor selected (an explicitly selected brush, or a volume if
+ * those are included).
+ *
+ * Notes:
+ * - An explictly selected subtractive brush is valid, as users may want to delete subtractive brushes through the tool.
+ * - Even with an additive brush, a resulting mesh could still end up empty, for instance if a brush was fully inside a subtractive one.
+ * - Does not check if targets have valid geometry for the "convert then combine" path (for mesh boolean operations),
+ *   as that currently requires doing an actual conversion.
+ */
+bool UBspConversionTool::AtLeastOneValidConversionTarget() const
+{
+	for (FSelectionIterator Iter(*GEditor->GetSelectedActors()); Iter; ++Iter)
+	{
+		ABrush* Brush = Cast<ABrush>(*Iter);
+		if (IsValidConversionTarget(Brush))
+		{
+			return true;
+		}
+	}
+
+	return false;
+}
+
+/**
+ * Determines whether the passed in pointer is a valid convertible brush. Safe to call with nullptr.
+ * Considers Settings->bIncludeVolumes.
+ */
+bool UBspConversionTool::IsValidConversionTarget(const ABrush* Brush) const
+{
+	return Brush && !Brush->IsPendingKill()
+		&& (!Brush->IsVolumeBrush() || Settings->bIncludeVolumes)
+		&& (Brush->BrushType == EBrushType::Brush_Add || Brush->BrushType == EBrushType::Brush_Subtract);
+}
+
+void UBspConversionTool::Shutdown(EToolShutdownType ShutdownType)
+{
+	if (ShutdownType == EToolShutdownType::Accept)
+	{
+		GetToolManager()->BeginUndoTransaction(LOCTEXT("BspConversionToolTransactionName", "BSP Conversion"));
+
+		// If settings had been set to show a preview, then we would have already generated a preview mesh
+		// that we could copy the result from. Thus, we only need to generate here if the settings were set
+		// to not use a preview.
+		if (!Settings->bShowPreview)
+		{
+			CompareAndUpdateConversionTargets();
+			FText ErrorMessage;
+			if (!ComputeAndUpdatePreviewMesh(&ErrorMessage))
+			{
+				// We're closing the tool, so it's too late to just display a side panel warning. Give a pop up.
+				FMessageDialog::Open(EAppMsgType::Ok, ErrorMessage);
+				
+				// Don't delete any brushes if we failed
+				BrushesToDelete.Empty();
+			}
+		}
+
+		// We only need to output something if the result wasn't empty.
+		if (PreviewMesh->GetMesh()->VertexCount() > 0)
+		{
+			TArray<UMaterialInterface*> Materials;
+			PreviewMesh->GetMaterials(Materials);
+
+			// Create the final actor
+			AActor* NewActor = AssetGenerationUtil::GenerateStaticMeshActor(
+				AssetAPI, TargetWorld,
+				PreviewMesh->GetMesh(), (FTransform3d)PreviewMesh->GetTransform(),
+				TEXT("BspConvertedToStaticMesh"), Materials);
+
+			if (NewActor)
+			{
+				// Select the final result
+				FSelectedOjectsChangeList NewSelection;
+				NewSelection.ModificationType = ESelectedObjectsModificationType::Replace;
+				NewSelection.Actors.Add(NewActor);
+				GetToolManager()->RequestSelectionChange(NewSelection);
+			}
+		}
+
+		// Delete brushes that we marked for deletion. This may need to happen even if the resulting mesh
+		// was empty, for instance if everything was inside a subtractive brush.
+		for (ABrush* Brush : BrushesToDelete)
+		{
+			TargetWorld->EditorDestroyActor(Brush, true);
+		}
+		GEditor->RebuildAlteredBSP();
+
+		GetToolManager()->EndUndoTransaction();
+	}
+
+	// Remove the preview mesh
+	PreviewMesh->SetVisible(false);
+	PreviewMesh->Disconnect();
+	PreviewMesh = nullptr;
+
+	// Deregister all the callbacks we registered
+	USelection::SelectionChangedEvent.RemoveAll(this);
+	GEngine->OnLevelActorListChanged().RemoveAll(this);
+	GEngine->OnActorMoved().RemoveAll(this);
+
+	// Empty cached data
+	CachedBrushes.Empty();
+
+	Settings->SaveProperties(this);
+}
+#pragma endregion
+
+#pragma region Conversion functions
+
+/**
+ * Updates the targets that the conversion functions operate on (BrushesToConvert, BrushForPivot),
+*  based on level composition, current selection, and settings. Also checks if this is a change 
+*  from the previous value of BrushesToConvert.
+ *
+ * @return true if BrushesToConvert changed.
+ */
+bool UBspConversionTool::CompareAndUpdateConversionTargets()
+{
+	TArray<ABrush*> PreviousBrushesToConvert = BrushesToConvert;
+	BrushesToConvert.Empty();
+	BrushesToDelete.Empty();
+
+	// One thing we need to check is whether we have just a single additive brush,
+	// in which case we want to keep the pivot the same (and we will set BrushForPivot)
+	int NumAdditiveBrushes = 0;
+
+	// The order of brush composition is determined by their order in the ULevel, and may be different from 
+	// their order of selection. We make BrushesToConvert contain selected brushes in the proper order.
+	TArray<AActor*>& LevelActorsInCompositionOrder = TargetWorld->GetCurrentLevel()->Actors;
+	for (int i = 0; i < LevelActorsInCompositionOrder.Num(); ++i)
+	{
+		ABrush* Brush = Cast<ABrush>(LevelActorsInCompositionOrder[i]);
+		if (IsValidConversionTarget(Brush)
+			&& (Brush->IsSelected() || (!Settings->bExplicitSubtractiveBrushSelection && Brush->BrushType == EBrushType::Brush_Subtract)))
+		{
+			BrushesToConvert.Add(Brush);
+
+			if (Brush->BrushType == EBrushType::Brush_Add)
+			{
+				++NumAdditiveBrushes;
+				BrushForPivot = Brush;
+
+				// Additive volume brushes may not be deleted, depending on settings
+				if (!Brush->IsVolumeBrush() || Settings->bRemoveConvertedVolumes)
+				{
+					BrushesToDelete.Add(Brush);
+				}
+			}
+			// Subtractive brushes only get deleted if they were explicitly selected and the settings allow it
+			else if (Settings->bRemoveConvertedSubtractiveBrushes && Brush->IsSelected())
+			{
+				BrushesToDelete.Add(Brush);
+			}
+		}
+	}
+
+	if (NumAdditiveBrushes != 1)
+	{
+		BrushForPivot = nullptr;
+	}
+
+	return BrushesToConvert != PreviousBrushesToConvert;
+}
+
+/**
+ * Performs conversion of BrushesToConvert and stores the result in the preview mesh.
+ * 
+ * @param OutErrorMessage Error message to fill if there is a problem.
+ * @return true if successful.
+ */
+bool UBspConversionTool::ComputeAndUpdatePreviewMesh(FText* OutErrorMessage)
+{
+	if (Settings->ConversionMode == EBspConversionMode::CombineFirst)
+	{
+		return CombineThenConvert(OutErrorMessage);
+	}
+	else
+	{
+		return ConvertThenCombine(OutErrorMessage);
+	}
+}
+
+/**
+ * The easy conversion path, where we just use the existing bsp conversion function to
+ * convert everything that was selected. Operates on non-manifold geometry. Doesn't write
+ * out an error message. Returns true unless BrushesToConvert was empty.
+ *
+ * BrushesToConvert must be initialized.
+ */
+bool UBspConversionTool::CombineThenConvert(FText*)
+{
+	if (BrushesToConvert.Num() == 0)
+	{
+		PreviewMesh->ClearPreview();
+		return false;
+	}
+
+	FDynamicMesh3 OutputMesh;
+	TArray<UMaterialInterface*> OutputMaterials;
+	ConvertBrushesToDynamicMesh(BrushesToConvert, OutputMesh, OutputMaterials);
+
+	// The created mesh is built with its pivot at the origin, so we need to reset it. If there was
+	// just one additive brush, we keep the pivot the same, otherwise we center it.
+	FVector3d NewPivot;
+	if (BrushForPivot)
+	{
+		NewPivot = BrushForPivot->GetTransform().GetLocation();
+	}
+	else
+	{
+		NewPivot = OutputMesh.GetCachedBounds().Center();
+	}
+
+	MeshTransforms::Translate(OutputMesh, -NewPivot);
+	FTransform3d Transform = FTransform3d::Identity();
+	Transform.SetTranslation(NewPivot);
+
+	PreviewMesh->UpdatePreview(&OutputMesh);
+	PreviewMesh->SetTransform((FTransform)Transform);
+	PreviewMesh->SetMaterials(OutputMaterials);
+
+	return true;
+}
+
+/**
+ * The more complicated conversion path, where we convert brushes individually and use static mesh boolean
+ * operations to combine them. Boolean operations fail when geometry is non-manifold, which happens in
+ * the case of stair brushes, mainly, due to them not being properly closed.
+ * The preview gets cleared in the case of an error.
+ * 
+ * @param ErrorMessage Place to write out error message if a conversion results in invalid geometry.
+ * @return false if BrushesToConvert is empty or a conversion gives invalid geometry, true otherwise.
+ */
+bool UBspConversionTool::ConvertThenCombine(FText *ErrorMessage)
+{
+	if (BrushesToConvert.Num() == 0)
+	{
+		PreviewMesh->ClearPreview();
+		return false;
+	}
+
+	// We'll need temporary space to store meshes as we build them up. The input and output meshes,
+	// transforms, and material arrays will swap after we merge in each new mesh.
+	FDynamicMesh3 MeshStorage[2];
+	FTransform3d TransformStorage[2];
+	TArray<UMaterialInterface*> MaterialArraysStorage[2];
+
+	FDynamicMesh3* InputMesh = &MeshStorage[0];
+	FDynamicMesh3* OutputMesh = &MeshStorage[1];
+	
+	FTransform3d* InputTransform = &TransformStorage[0];
+	FTransform3d* OutputTransform = &TransformStorage[1];
+
+	TArray<UMaterialInterface*>* InputMaterials = &MaterialArraysStorage[0];
+	TArray<UMaterialInterface*>* OutputMaterials = &MaterialArraysStorage[1];
+
+	// The individual mesh to merge in is dependent only on that brush, and therefore can be
+	// cached. Unfortunately, we need to be able to modify the mesh while doing boolean operations
+	// to update the material references of the triangles to a common material set, so we need to
+	// create copies of the cached mesh.
+	FDynamicMesh3 NextMesh;
+
+	// The transform when a brush is first converted is always the same
+	FTransform3d NextTransform = FTransform3d::Identity();
+
+	// We could actually point to cached materials since we don't change these, but for consistency,
+	// we'll have a local copy.
+	TArray<UMaterialInterface*> NextMaterials;
+
+	// This will point to the next bsp brush that we're considering
+	ABrush* NextBrush = BrushesToConvert[0];
+
+	// When converting individual brushes, we pass them to ConvertBrushesToDynamicMesh as a
+	// one-element array.
+	TArray<ABrush*> WrappedBrush = { NextBrush };
+
+	// Convert the first brush unless it is subtractive, in which case the mesh will stay empty.
+	// This isn't just an optimization- we also need to avoid caching this brush, or trying to use a cached
+	// additive version.
+	if (NextBrush->BrushType != EBrushType::Brush_Subtract)
+	{
+		if (Settings->bCacheBrushes && CachedBrushes.Contains(NextBrush))
+		{
+			OutputMesh->Copy(*CachedBrushes[NextBrush].Key);
+			*OutputMaterials = *CachedBrushes[NextBrush].Value;
+		}
+		else
+		{
+			ConvertBrushesToDynamicMesh(WrappedBrush, *OutputMesh, *OutputMaterials);
+
+			if (Settings->bCacheBrushes)
+			{
+				CachedBrushes.Add(NextBrush, FCachedResult(MakeShared<const FDynamicMesh3>(*OutputMesh),
+					MakeShared<const TArray<UMaterialInterface*>>(*OutputMaterials)));
+			}
+		}
+		*OutputTransform = FTransform3d::Identity(); // Always identity when first converted
+
+		// Check validity for the purposes of boolean operations. We'll do this after each brush conversion.
+		if (!OutputMesh->CheckValidity(FDynamicMesh3::FValidityOptions(), EValidityCheckFailMode::ReturnOnly))
+		{
+			PreviewMesh->ClearPreview();
+			if (ErrorMessage)
+			{
+				*ErrorMessage = GetBrushGeometryErrorMessage(NextBrush);
+			}
+			return false;
+		}
+	}
+
+	// Now convert the other meshes one at a time and apply them using boolean mesh operations
+	for (int i = 1; i < BrushesToConvert.Num(); ++i)
+	{
+		NextBrush = BrushesToConvert[i];
+
+		Swap(OutputMesh, InputMesh);
+		Swap(OutputTransform, InputTransform);
+		Swap(OutputMaterials, InputMaterials);
+		OutputMesh->Clear();
+
+		FMeshBoolean::EBooleanOp Operation = FMeshBoolean::EBooleanOp::Union;
+		bool IsSubtractiveBrush = NextBrush->BrushType == EBrushType::Brush_Subtract;
+		if (IsSubtractiveBrush)
+		{
+			Operation = FMeshBoolean::EBooleanOp::Difference;
+
+			// We need a solid mesh to subtract, so the conversion code needs to operate on an additive brush.
+			// However, we're going to undo this once we're done because the brush may persist after the conversion,
+			// depending on the settings.
+			// Despite the undoing, we need to call Modify now because it gets called in the conversion code, which
+			// would cause the changed property to be saved.
+			NextBrush->Modify();
+			NextBrush->BrushType = EBrushType::Brush_Add;
+		}
+
+		// Convert the next brush
+		if (Settings->bCacheBrushes && CachedBrushes.Contains(NextBrush))
+		{
+			NextMesh.Copy(*CachedBrushes[NextBrush].Key);
+			NextMaterials = *CachedBrushes[NextBrush].Value;
+		}
+		else
+		{
+			WrappedBrush[0] = NextBrush;
+
+			// Get the result
+			ConvertBrushesToDynamicMesh(WrappedBrush, NextMesh, NextMaterials);
+
+			if (Settings->bCacheBrushes)
+			{
+				CachedBrushes.Add(NextBrush, FCachedResult(MakeShared<const FDynamicMesh3>(NextMesh),
+					MakeShared<const TArray<UMaterialInterface*>>(NextMaterials)));
+			}
+		}
+
+		// See if converted geometry is valid
+		if (!NextMesh.CheckValidity(FDynamicMesh3::FValidityOptions(), EValidityCheckFailMode::ReturnOnly)) 
+		{
+			PreviewMesh->ClearPreview();
+			if (ErrorMessage)
+			{
+				*ErrorMessage = GetBrushGeometryErrorMessage(NextBrush);
+			}
+			return false;
+		}
+
+		// Apply the boolean operation
+		ApplyStaticMeshBooleanOperation(
+			*InputMesh, *InputTransform, *InputMaterials,
+			NextMesh, NextTransform, NextMaterials,
+			*OutputMesh, *OutputTransform, *OutputMaterials, Operation);
+		
+		if (IsSubtractiveBrush)
+		{
+			// Undo the change in brush type that we did before.
+			NextBrush->BrushType = EBrushType::Brush_Subtract;
+		}
+	}//end converting other brushes
+
+	// If there was only a single additive brush, we should keep its pivot in the original location.
+	if (BrushForPivot)
+	{
+		FVector3d BrushTranslation = BrushForPivot->GetTransform().GetLocation();
+		MeshTransforms::Translate(*OutputMesh, OutputTransform->GetTranslation() - BrushTranslation);
+		OutputTransform->SetTranslation(BrushTranslation);
+	}
+	// Otherwise, the pivot set by boolean operations is appropriate.
+
+	PreviewMesh->UpdatePreview(OutputMesh);
+	PreviewMesh->SetTransform((FTransform)(*OutputTransform));
+	PreviewMesh->SetMaterials(*OutputMaterials);
+
+	return true;
+}
+
+FText GetBrushGeometryErrorMessage(ABrush* Brush)
+{
+	return FText::Format(LOCTEXT("ConvertThenCombineInvalidGeometryError", "The brush \"{0}\" has invalid geometry "
+		"to be used with mesh boolean operations, so the \"Convert, then Combine\" mode cannot be used. Try using \"Combine, then Convert\" "
+		"to convert, then use MeshInspector to see problematic areas."),
+		FText::FromString(Brush->GetName()));
+}
+
+/** Uses our existing conversion functions to convert brushes to a single DynamicMesh. */
+void ConvertBrushesToDynamicMesh(TArray<ABrush*>& BrushesToConvert, FDynamicMesh3& OutputMesh, TArray<UMaterialInterface*>& OutputMaterials)
+{
+	// Have the editor rebuild a model composed only of selected brushes into a temporary UModel object, and make
+	// sure its polygons get built.
+	// Even though it's temporary, we're not allowed to make a UModel on the stack
+	UModel* TempModel = NewObject<UModel>();
+	TempModel->Initialize(nullptr);
+	GEditor->RebuildModelFromBrushes(BrushesToConvert, TempModel);
+	GEditor->bspBuildFPolys(TempModel, true, 0); // SurfLinks parameter doesn't matter
+
+	// Prep some output variables
+	FMeshDescription MeshDescription;
+	FStaticMeshAttributes StaticMeshAttributes(MeshDescription);
+	StaticMeshAttributes.Register();
+
+	TArray<FStaticMaterial> Materials;
+
+	// Do the actual conversion using our old conversion function
+	GetBrushMesh(nullptr, TempModel, MeshDescription, Materials);
+
+	// Get a list of material interfaces out of the list of materials that we got
+	OutputMaterials.Empty();
+	for (FStaticMaterial Material : Materials)
+	{
+		OutputMaterials.Add(Material.MaterialInterface);
+	}
+
+	// Turn the mesh description into a DynamicMesh
+	OutputMesh.Clear();
+	FMeshDescriptionToDynamicMesh Converter;
+	Converter.Convert(&MeshDescription, OutputMesh);
+}
+
+void ApplyStaticMeshBooleanOperation(
+	FDynamicMesh3& MeshA, const FTransform3d& TransformA, const TArray<UMaterialInterface*>& MaterialsA,
+	FDynamicMesh3& MeshB, const FTransform3d& TransformB, const TArray<UMaterialInterface*>& MaterialsB,
+	FDynamicMesh3& OutputMesh, FTransform3d& OutputTransform, TArray<UMaterialInterface*>& OutputMaterials,
+	FMeshBoolean::EBooleanOp Operation)
+{
+	// These need to be enabled on both meshes to deal with materials properly. This is relevant, for
+	// instance, if the first brush in the composition list was empty.
+	MeshA.EnableAttributes();
+	MeshA.Attributes()->EnableMaterialID();
+	MeshB.EnableAttributes();
+	MeshB.Attributes()->EnableMaterialID();
+
+
+	// We'll need to combine all the materials into the output, but not duplicate them across the two
+	// meshes. This will keep track of the materials that we have seen.
+	TMap<UMaterialInterface*, int> SeenMaterials;
+
+	// Start with the materials from the first mesh.
+	OutputMaterials = MaterialsA;
+	for (int i = 0; i < OutputMaterials.Num(); ++i)
+	{
+		SeenMaterials.Add(OutputMaterials[i], i);
+	}
+
+	// Add any new materials from mesh B, and remap any materials that we already have in that mesh. MaterialBRemap
+	// maps old material indices to new ones.
+	TArray<int> MaterialBRemap;
+	for (UMaterialInterface* Mat : MaterialsB)
+	{
+		int NewMatIndex;
+		int* FoundMatIdx = SeenMaterials.Find(Mat);
+		if (FoundMatIdx)
+		{
+			NewMatIndex = *FoundMatIdx;
+		}
+		else
+		{
+			NewMatIndex = OutputMaterials.Add(Mat);
+			SeenMaterials.Add(Mat, NewMatIndex);
+		}
+		MaterialBRemap.Add(NewMatIndex);
+	}
+
+	// Apply the remapping of material indices to mesh B.
+	FDynamicMeshMaterialAttribute* MaterialIDs = MeshB.Attributes()->GetMaterialID();
+	for (int TID : MeshB.TriangleIndicesItr())
+	{
+		MaterialIDs->SetValue(TID, MaterialBRemap[MaterialIDs->GetValue(TID)]);
+	}
+
+	// Perform the actual boolean operation.
+	FMeshBoolean BooleanOperation(&MeshA, TransformA, &MeshB, TransformB, &OutputMesh, Operation);
+	BooleanOperation.Compute();
+	OutputTransform = BooleanOperation.ResultTransform;
+}
+#pragma endregion
+
+#pragma region Action buttons
+void UBspConversionTool::RequestAction(EBspConversionToolAction ActionType)
+{
+	if (PendingAction == EBspConversionToolAction::NoAction)
+	{
+		PendingAction = ActionType;
+	}
+}
+
+void UBspConversionTool::OnTick(float DeltaTime)
+{
+	if (PendingAction != EBspConversionToolAction::NoAction)
+	{
+		ApplyAction(PendingAction);
+		PendingAction = EBspConversionToolAction::NoAction;
+	}
+}
+
+void UBspConversionTool::ApplyAction(EBspConversionToolAction ActionType)
+{
+	switch (ActionType)
+	{
+	case EBspConversionToolAction::SelectAllValidBrushes:
+	{
+		FSelectedOjectsChangeList NewSelection;
+		NewSelection.ModificationType = ESelectedObjectsModificationType::Replace;
+
+		// Accumulate all actors in the current level that are valid brushes
+		for (auto Iter(TargetWorld->GetCurrentLevel()->Actors.CreateConstIterator()); Iter; ++Iter)
+		{
+			ABrush* Brush = Cast<ABrush>(*Iter);
+			if (IsValidConversionTarget(Brush))
+			{
+				NewSelection.Actors.Add(Brush);
+			}
+		}
+
+		GetToolManager()->RequestSelectionChange(NewSelection);
+	}
+	break;
+	case EBspConversionToolAction::DeselectVolumes:
+	{
+		FSelectedOjectsChangeList NewSelection;
+		NewSelection.ModificationType = ESelectedObjectsModificationType::Replace;
+
+		// Out of the current selection, keep anything except volume brushes
+		for (FSelectionIterator Iter(*GEditor->GetSelectedActors()); Iter; ++Iter)
+		{
+			AActor* Actor = Cast<AActor>(*Iter);
+			if (!(Cast<ABrush>(Actor) && Cast<ABrush>(Actor)->IsVolumeBrush()))
+			{
+				NewSelection.Actors.Add(Actor);
+			}
+		}
+
+		GetToolManager()->RequestSelectionChange(NewSelection);
+	}
+	break;
+	case EBspConversionToolAction::DeselectNonValid:
+	{
+		// Normally, "GEditor->SelectNone(true, false, false)" would deselect all brushes, but
+		// it does not deselect volumes, which we want to do depending on the settings.
+
+		FSelectedOjectsChangeList NewSelection;
+		NewSelection.ModificationType = ESelectedObjectsModificationType::Replace;
+
+		// From current selection, select everything that is valid given current settings
+		for (FSelectionIterator Iter(*GEditor->GetSelectedActors()); Iter; ++Iter)
+		{
+			ABrush* Brush = Cast<ABrush>(*Iter);
+			if (IsValidConversionTarget(Brush))
+			{
+				NewSelection.Actors.Add(Brush);
+			}
+		}
+
+		GetToolManager()->RequestSelectionChange(NewSelection);
+	}
+	break;
+	}
+}
+#pragma endregion
+
+#pragma region Callback functions
+
+// We need to preserve the following:
+// - BrushesToConvert refle
+// - BrushesToConvert contains relevant brushes given settings and selection
+// - PreviewMesh is empty if bCanAccept is false, otherwise it reflects BrushesToConvert
+// - CachedBrushes does not contain any outdated conversions.
+
+/**
+ * This is the primary event we'll be responding to. It affects whether we can accept
+ * or not, and the preview, if it is being shown. Does not affect cached brushes.
+ */
+void UBspConversionTool::OnEditorSelectionChanged(UObject* NewSelection)
+{
+	USelection* Selection = Cast<USelection>(NewSelection);
+	if (!Selection)
+	{
+		// Clear things
+		bCanAccept = false;
+		PreviewMesh->ClearPreview();
+		BrushesToConvert.Empty();
+		GetToolManager()->DisplayMessage(FText(), EToolMessageLevel::UserWarning); // This "clears" it
+		return;
+	}
+
+	// Update whether the tool can accept safely
+	bCanAccept = AtLeastOneValidConversionTarget();
+
+	// Update preview. This may change a true bCanAccept to false if there are errors
+	if (!bCanAccept)
+	{
+		// If there isn't anything selected yet, clear things
+		PreviewMesh->ClearPreview();
+		BrushesToConvert.Empty();
+		GetToolManager()->DisplayMessage(FText(), EToolMessageLevel::UserWarning);
+	}
+	else if (Settings->bShowPreview && CompareAndUpdateConversionTargets())
+	{
+		FText ErrorMessage;
+		bCanAccept = ComputeAndUpdatePreviewMesh(&ErrorMessage);
+
+		// We do this regardless of success because it resets the error message if ErrorMessage is empty.
+		GetToolManager()->DisplayMessage(ErrorMessage, EToolMessageLevel::UserWarning);
+	}
+}
+
+/**
+ * Changes to the actor list may change the order of brush composition, or add or remove implicitly included
+ * subtractive brushes. Should not affect cached brushes.
+ */
+void UBspConversionTool::OnEditorLevelActorListChanged()
+{
+	// Safest thing to do is to do the same things as OnEditorSelectionChanged, even though bCanAccept should
+	// only change in very weird cases (such as removing an implicit subtractive stair brush).
+
+	// Update whether the tool can accept safely
+	bCanAccept = AtLeastOneValidConversionTarget();
+
+	// Update preview. This may change a true bCanAccept to false if there are errors
+	if (!bCanAccept)
+	{
+		// If there isn't anything selected yet, clear things
+		PreviewMesh->ClearPreview();
+		BrushesToConvert.Empty();
+		GetToolManager()->DisplayMessage(FText(), EToolMessageLevel::UserWarning);
+	}
+	else if (Settings->bShowPreview && CompareAndUpdateConversionTargets())
+	{
+		FText ErrorMessage;
+		bCanAccept = ComputeAndUpdatePreviewMesh(&ErrorMessage);
+
+		// We do this regardless of success because it resets the error message if ErrorMessage is empty.
+		GetToolManager()->DisplayMessage(ErrorMessage, EToolMessageLevel::UserWarning);
+	}
+}
+
+/**
+ * Actor movements include any transform changes, and change both the preview and the
+ * cached brushes. It shouldn't affect BrushesToConvert or bCanAccept.
+ */
+void UBspConversionTool::OnEditorActorMoved(AActor* Actor)
+{
+	// Make sure it was a brush that got moved
+	ABrush* Brush = Cast<ABrush>(Actor);
+	if (!Brush)
+	{
+		return;
+	}
+
+	// Remove cached version of brush
+	if (CachedBrushes.Contains(Brush))
+	{
+		CachedBrushes.Remove(Brush);
+	}
+
+	// Update preview if this is a relevant brush
+	if (bCanAccept && Settings->bShowPreview && BrushesToConvert.Contains(Brush))
+	{
+		FText ErrorMessage;
+		bCanAccept = ComputeAndUpdatePreviewMesh();
+		GetToolManager()->DisplayMessage(ErrorMessage, EToolMessageLevel::UserWarning);
+	}
+}
+
+void UBspConversionTool::OnPropertyModified(UObject* PropertySet, FProperty* Property)
+{
+	// There are a few properties that are only relevant on accept, and don't change anything now
+	if (Property && (Property->GetFName() == GET_MEMBER_NAME_CHECKED(UBspConversionToolProperties, bRemoveConvertedSubtractiveBrushes)
+		|| Property->GetFName() == GET_MEMBER_NAME_CHECKED(UBspConversionToolProperties, bRemoveConvertedVolumes)))
+	{
+		return;
+	}
+
+	// For the most part, it ends up being too messy to consider all properties individually. It
+	// is safer to reset the state (though only clear cached brushes if we have to).
+
+	if (!Settings->bCacheBrushes)
+	{
+		CachedBrushes.Empty();
+	}
+
+	bCanAccept = AtLeastOneValidConversionTarget();
+
+	// Clear preview
+	PreviewMesh->ClearPreview();
+	BrushesToConvert.Empty();
+	GetToolManager()->DisplayMessage(FText(), EToolMessageLevel::UserWarning); // This "clears" it
+
+	if (Settings->bShowPreview && bCanAccept)
+	{
+		CompareAndUpdateConversionTargets();
+		FText ErrorMessage;
+		bCanAccept = ComputeAndUpdatePreviewMesh(&ErrorMessage);
+		GetToolManager()->DisplayMessage(ErrorMessage, EToolMessageLevel::UserWarning);
+	}
+}
+#pragma endregion
+
+#undef LOCTEXT_NAMESPACE
diff --git a/Engine/Plugins/Experimental/MeshModelingToolset/Source/MeshModelingToolsEditorOnly/Public/BspConversionTool.h b/Engine/Plugins/Experimental/MeshModelingToolset/Source/MeshModelingToolsEditorOnly/Public/BspConversionTool.h
new file mode 100644
--- /dev/null
+++ b/Engine/Plugins/Experimental/MeshModelingToolset/Source/MeshModelingToolsEditorOnly/Public/BspConversionTool.h
@@ -0,0 +1,191 @@
+// Copyright Epic Games, Inc. All Rights Reserved.
+
+#pragma once
+
+#include "CoreMinimal.h"
+#include "DynamicMesh3.h"
+#include "UObject/ObjectMacros.h"
+#include "InteractiveTool.h"
+#include "InteractiveToolBuilder.h"
+#include "ToolContextInterfaces.h"
+#include "MultiSelectionTool.h"
+#include "PreviewMesh.h"
+#include "Properties/MeshMaterialProperties.h"
+
+#include "BspConversionTool.generated.h"
+
+
+/**
+ * Builder for UBspConversionTool.
+ */
+UCLASS()
+class MESHMODELINGTOOLSEDITORONLY_API UBspConversionToolBuilder : public UInteractiveToolBuilder
+{
+	GENERATED_BODY()
+
+public:
+	IToolsContextAssetAPI* AssetAPI = nullptr;
+
+	virtual bool CanBuildTool(const FToolBuilderState& SceneState) const override;
+	virtual UInteractiveTool* BuildTool(const FToolBuilderState& SceneState) const override;
+};
+
+
+UENUM()
+enum class EBspConversionMode : uint8
+{
+	/** First converts the brushes to static meshes, then performs static mesh boolean operations. */
+	ConvertFirst = 0 UMETA(DisplayName = "Convert, then Combine"),
+
+	/** First combines brushes, then converts result to static mesh. */
+	CombineFirst = 1 UMETA(DisplayName = "Combine, then Convert"),
+};
+
+/**
+ * 
+ */
+UCLASS()
+class MESHMODELINGTOOLSEDITORONLY_API UBspConversionToolProperties : public UInteractiveToolPropertySet
+{
+	GENERATED_BODY()
+
+public:
+
+	UPROPERTY(EditAnywhere, Category = Options)
+	EBspConversionMode ConversionMode = EBspConversionMode::ConvertFirst;
+
+	/** Whether to consider BSP volumes to be valid conversion targets. */
+	UPROPERTY(EditAnywhere, NonTransactional, Category = Options)
+	bool bIncludeVolumes = false;
+	
+	/** Whether to remove any selected BSP volumes after using them to create a static mesh. */
+	UPROPERTY(EditAnywhere, NonTransactional, Category = Options, meta = (EditCondition = "bIncludeVolumes"))
+	bool bRemoveConvertedVolumes = false;
+
+	/** Whether subtractive brushes have to be explicitly selected to be part of the conversion. If false, all
+	 subtractive brushes in the level will be used. */
+	UPROPERTY(EditAnywhere, NonTransactional, Category = Options)
+	bool bExplicitSubtractiveBrushSelection = true;
+
+	/** Whether subtractive brushes used in a conversion should be removed. Only acts on explicitly selected
+	 subtractive brushes. */
+	UPROPERTY(EditAnywhere, NonTransactional, Category = Options)
+	bool bRemoveConvertedSubtractiveBrushes = true;
+
+	/** Caches individual brush conversions in "convert then combine" mode during a single invocation of 
+	 the tool. Only useful if changing selections or properties after starting the tool. Cleared on tool shutdown. */
+	UPROPERTY(EditAnywhere, NonTransactional, Category = Options, meta = (EditCondition = "ConversionMode == EBspConversionMode::ConvertFirst"))
+		bool bCacheBrushes = true;
+
+	/** Determines whether a dynamic preview is shown. Note that this introduces non-background computations 
+	at each event that changes the result, rather than only performing a computation on Accept. */
+	UPROPERTY(EditAnywhere, NonTransactional, Category = PreviewOptions)
+	bool bShowPreview = true;
+};
+
+UENUM()
+enum class EBspConversionToolAction
+{
+	NoAction,
+
+	SelectAllValidBrushes,
+	DeselectVolumes,
+	DeselectNonValid
+};
+
+UCLASS()
+class MESHMODELINGTOOLSEDITORONLY_API UBspConversionToolActionPropertySet : public UInteractiveToolPropertySet
+{
+	GENERATED_BODY()
+
+public:
+	TWeakObjectPtr<UBspConversionTool> ParentTool;
+
+	void Initialize(UBspConversionTool* ParentToolIn) { ParentTool = ParentToolIn; }
+	void PostAction(EBspConversionToolAction Action);
+
+	/** Select all brushes that satisfy the current settings. */
+	UFUNCTION(CallInEditor, Category = SelectionOperations, Meta = (DisplayName = "Select All Valid Brushes", DisplayPriority = 1))
+	void SelectAllValidBrushes() { PostAction(EBspConversionToolAction::SelectAllValidBrushes); }
+
+	/** Deselect any currently selected volume brushes. */
+	UFUNCTION(CallInEditor, Category = SelectionOperations, Meta = (DisplayName = "Deselect Volumes", DisplayPriority = 2))
+	void DeselectVolumes() { PostAction(EBspConversionToolAction::DeselectVolumes); }
+
+	/** Deselect any currently selected brushes that would not be converted given current settings. */
+	UFUNCTION(CallInEditor, Category = SelectionOperations, Meta = (DisplayName = "Deselect Non-Valid", DisplayPriority = 3))
+	void DeselectNonValid() { PostAction(EBspConversionToolAction::DeselectNonValid); }
+};
+
+/**
+ * Converts BSP brushes to static meshes.
+ *
+ * Known limitations:
+ * - Preview does not respond to property changes in the brush detail panel while the tool is running. User would need
+ *   to create some event that does change the preview (such as tool property change, or selection change).
+ * - BSP brushes with non-manifold geometry (specifically, the stair brushes) cannot be used with the "Convert, then combine"
+ *   path because boolean operations do not allow them. The user gets properly notified of this if it comes up.
+ */
+UCLASS()
+class MESHMODELINGTOOLSEDITORONLY_API UBspConversionTool : public UInteractiveTool
+{
+	GENERATED_BODY()
+
+public:
+	UBspConversionTool();
+
+	virtual void RegisterActions(FInteractiveToolActionSet& ActionSet) override;
+
+	virtual void Setup() override;
+	virtual void Shutdown(EToolShutdownType ShutdownType) override;
+
+	virtual void SetWorld(UWorld* World) { this->TargetWorld = World; }
+	virtual void SetAssetAPI(IToolsContextAssetAPI* InAssetAPI) { this->AssetAPI = InAssetAPI; }
+
+	virtual void OnTick(float DeltaTime) override;
+	virtual void Render(IToolsContextRenderAPI* RenderAPI) override;
+
+	virtual bool HasCancel() const override { return true; }
+	virtual bool HasAccept() const override { return true; }
+	virtual bool CanAccept() const override;
+
+	virtual void OnPropertyModified(UObject* PropertySet, FProperty* Property) override;
+
+	virtual void RequestAction(EBspConversionToolAction ActionType);
+
+	UPROPERTY()
+	UBspConversionToolProperties* Settings = nullptr;
+
+	UPROPERTY()
+	UBspConversionToolActionPropertySet* ToolActions = nullptr;
+
+protected:
+	bool bCanAccept = false;
+
+	UWorld* TargetWorld;
+	IToolsContextAssetAPI* AssetAPI;
+	EBspConversionToolAction PendingAction;
+
+	UPROPERTY()
+	UPreviewMesh *PreviewMesh;
+
+	TArray<ABrush*> BrushesToConvert;
+	TArray<ABrush*> BrushesToDelete;
+	ABrush* BrushForPivot = nullptr;
+
+	TMap<ABrush*, TPair<TSharedPtr<const FDynamicMesh3>, 
+		TSharedPtr<const TArray<UMaterialInterface*>>>> CachedBrushes;
+
+	bool IsValidConversionTarget(const ABrush* Brush) const;
+	bool AtLeastOneValidConversionTarget() const;
+	bool CompareAndUpdateConversionTargets();
+
+	void OnEditorSelectionChanged(UObject* NewSelection);
+	void OnEditorLevelActorListChanged();
+	void OnEditorActorMoved(AActor* InActor);
+
+    void ApplyAction(EBspConversionToolAction ActionType);
+	bool ComputeAndUpdatePreviewMesh(FText* ErrorMessage = nullptr);
+	bool ConvertThenCombine(FText* ErrorMessage = nullptr);
+	bool CombineThenConvert(FText* ErrorMessage = nullptr);
+};
\ No newline at end of file
diff --git a/Engine/Plugins/Experimental/MeshModelingToolset/Source/ModelingComponents/Private/PreviewMesh.cpp b/Engine/Plugins/Experimental/MeshModelingToolset/Source/ModelingComponents/Private/PreviewMesh.cpp
--- a/Engine/Plugins/Experimental/MeshModelingToolset/Source/ModelingComponents/Private/PreviewMesh.cpp
+++ b/Engine/Plugins/Experimental/MeshModelingToolset/Source/ModelingComponents/Private/PreviewMesh.cpp
@@ -87,12 +87,26 @@ void UPreviewMesh::SetMaterials(const TArray<UMaterialInterface*>& Materials)
 	DynamicMeshComponent->NotifyMeshUpdated();
 }
 
+int32 UPreviewMesh::GetNumMaterials() const
+{
+	check(DynamicMeshComponent);
+	return DynamicMeshComponent->GetNumMaterials();
+}
+
 UMaterialInterface* UPreviewMesh::GetMaterial(int MaterialIndex) const
 {
 	check(DynamicMeshComponent);
 	return DynamicMeshComponent->GetMaterial(MaterialIndex);
 }
 
+void UPreviewMesh::GetMaterials(TArray<UMaterialInterface*>& OutMaterials) const
+{
+	check(DynamicMeshComponent);
+	for (int32 i = 0; i < DynamicMeshComponent->GetNumMaterials(); ++i)
+	{
+		OutMaterials.Add(DynamicMeshComponent->GetMaterial(i));
+	}
+}
 
 void UPreviewMesh::SetOverrideRenderMaterial(UMaterialInterface* Material)
 {
diff --git a/Engine/Plugins/Experimental/MeshModelingToolset/Source/ModelingComponents/Private/ToolSetupUtil.cpp b/Engine/Plugins/Experimental/MeshModelingToolset/Source/ModelingComponents/Private/ToolSetupUtil.cpp
--- a/Engine/Plugins/Experimental/MeshModelingToolset/Source/ModelingComponents/Private/ToolSetupUtil.cpp
+++ b/Engine/Plugins/Experimental/MeshModelingToolset/Source/ModelingComponents/Private/ToolSetupUtil.cpp
@@ -139,6 +139,18 @@ UMaterialInterface* ToolSetupUtil::GetSelectionMaterial(const FLinearColor& UseC
 }
 
 
+UMaterialInstanceDynamic* ToolSetupUtil::GetSimpleCustomMaterial(UInteractiveToolManager* ToolManager, const FLinearColor& Color, float Opacity)
+{
+	UMaterialInterface* Material = LoadObject<UMaterial>(nullptr, TEXT("/MeshModelingToolset/Materials/SimpleBaseMaterial"));
+	if (Material != nullptr)
+	{
+		UMaterialInstanceDynamic* MatInstance = UMaterialInstanceDynamic::Create(Material, ToolManager);
+		MatInstance->SetVectorParameterValue(TEXT("Color"), Color);
+		MatInstance->SetScalarParameterValue(TEXT("Opacity"), Opacity);
+		return MatInstance;
+	}
+	return nullptr;
+}
 
 
 UMaterialInterface* ToolSetupUtil::GetDefaultPointComponentMaterial(bool bRoundPoints, UInteractiveToolManager* ToolManager)
diff --git a/Engine/Plugins/Experimental/MeshModelingToolset/Source/ModelingComponents/Public/PreviewMesh.h b/Engine/Plugins/Experimental/MeshModelingToolset/Source/ModelingComponents/Public/PreviewMesh.h
--- a/Engine/Plugins/Experimental/MeshModelingToolset/Source/ModelingComponents/Public/PreviewMesh.h
+++ b/Engine/Plugins/Experimental/MeshModelingToolset/Source/ModelingComponents/Public/PreviewMesh.h
@@ -106,11 +106,21 @@ public:
 	 */
 	void SetMaterials(const TArray<UMaterialInterface*>& Materials);
 
+	/**
+	* Get number of materials in the preview mesh (base materials, i.e., not including override material)
+	*/
+	int32 GetNumMaterials() const;
+
 	/**
 	* Get material from the preview mesh
 	*/
 	UMaterialInterface* GetMaterial(int MaterialIndex = 0) const;
 
+	/**
+	* Get the entire materials array from the preview mesh. Appends to OutMaterials without clearing it.
+	*/
+	void GetMaterials(TArray<UMaterialInterface*>& OutMaterials) const;
+
 	/**
 	 * Set an override material for the preview mesh. This material will override all the given materials.
 	 */
diff --git a/Engine/Plugins/Experimental/MeshModelingToolset/Source/ModelingComponents/Public/ToolSetupUtil.h b/Engine/Plugins/Experimental/MeshModelingToolset/Source/ModelingComponents/Public/ToolSetupUtil.h
--- a/Engine/Plugins/Experimental/MeshModelingToolset/Source/ModelingComponents/Public/ToolSetupUtil.h
+++ b/Engine/Plugins/Experimental/MeshModelingToolset/Source/ModelingComponents/Public/ToolSetupUtil.h
@@ -73,7 +73,10 @@ namespace ToolSetupUtil
 	 */
 	MODELINGCOMPONENTS_API UMaterialInterface* GetSelectionMaterial(const FLinearColor& UseColor, UInteractiveToolManager* ToolManager);
 
-
+	/**
+	 * @return Simple material with configurable color and opacity.
+	 */
+	MODELINGCOMPONENTS_API UMaterialInstanceDynamic* GetSimpleCustomMaterial(UInteractiveToolManager* ToolManager, const FLinearColor& Color, float Opacity);
 
 	/**
 	 * @param bRoundPoints true for round points, false for square
diff --git a/Engine/Plugins/Experimental/ModelingToolsEditorMode/Source/ModelingToolsEditorMode/Private/ModelingToolsEditorMode.cpp b/Engine/Plugins/Experimental/ModelingToolsEditorMode/Source/ModelingToolsEditorMode/Private/ModelingToolsEditorMode.cpp
--- a/Engine/Plugins/Experimental/ModelingToolsEditorMode/Source/ModelingToolsEditorMode/Private/ModelingToolsEditorMode.cpp
+++ b/Engine/Plugins/Experimental/ModelingToolsEditorMode/Source/ModelingToolsEditorMode/Private/ModelingToolsEditorMode.cpp
@@ -31,6 +31,7 @@
 #include "PlaneCutTool.h"
 #include "SelfUnionMeshesTool.h"
 #include "CSGMeshesTool.h"
+#include "BspConversionTool.h"
 #include "HoleFillTool.h"
 #include "PolygonOnMeshTool.h"
 #include "DisplaceMeshTool.h"
@@ -603,6 +604,10 @@ void FModelingToolsEditorMode::Enter()
 	CSGMeshesToolBuilder->AssetAPI = ToolsContext->GetAssetAPI();
 	RegisterToolFunc(ToolManagerCommands.BeginMeshBooleanTool, TEXT("CSGMeshesTool"), CSGMeshesToolBuilder);
 
+	auto BspConversionToolBuilder = NewObject<UBspConversionToolBuilder>();
+	BspConversionToolBuilder->AssetAPI = ToolsContext->GetAssetAPI();
+	RegisterToolFunc(ToolManagerCommands.BeginBspConversionTool, TEXT("BspConversionTool"), BspConversionToolBuilder);
+
 	auto PlaneCutToolBuilder = NewObject<UPlaneCutToolBuilder>();
 	PlaneCutToolBuilder->AssetAPI = ToolsContext->GetAssetAPI();
 	RegisterToolFunc(ToolManagerCommands.BeginPlaneCutTool, TEXT("PlaneCutTool"), PlaneCutToolBuilder);
diff --git a/Engine/Plugins/Experimental/ModelingToolsEditorMode/Source/ModelingToolsEditorMode/Private/ModelingToolsEditorModeToolkit.cpp b/Engine/Plugins/Experimental/ModelingToolsEditorMode/Source/ModelingToolsEditorMode/Private/ModelingToolsEditorModeToolkit.cpp
--- a/Engine/Plugins/Experimental/ModelingToolsEditorMode/Source/ModelingToolsEditorMode/Private/ModelingToolsEditorModeToolkit.cpp
+++ b/Engine/Plugins/Experimental/ModelingToolsEditorMode/Source/ModelingToolsEditorMode/Private/ModelingToolsEditorModeToolkit.cpp
@@ -396,6 +396,8 @@ void FModelingToolsEditorModeToolkit::BuildToolPalette_Standard(FName PaletteInd
 		ToolbarBuilder.AddToolBarButton(Commands.BeginVoxelMergeTool);
 		ToolbarBuilder.AddToolBarButton(Commands.BeginVoxelBooleanTool);
 		ToolbarBuilder.AddToolBarButton(Commands.BeginMeshBooleanTool);
+		ToolbarBuilder.AddSeparator();
+		ToolbarBuilder.AddToolBarButton(Commands.BeginBspConversionTool);
 	}
 	else if (PaletteIndex == TransformTabName)
 	{
@@ -512,6 +514,8 @@ void FModelingToolsEditorModeToolkit::BuildToolPalette_Experimental(FName Palett
 		ToolbarBuilder.AddToolBarButton(Commands.BeginVoxelBooleanTool);
 		ToolbarBuilder.AddToolBarButton(Commands.BeginSelfUnionTool);
 		ToolbarBuilder.AddToolBarButton(Commands.BeginMeshBooleanTool);
+		ToolbarBuilder.AddSeparator();
+		ToolbarBuilder.AddToolBarButton(Commands.BeginBspConversionTool);
 	}
 	else if (PaletteIndex == TransformTabName)
 	{
diff --git a/Engine/Plugins/Experimental/ModelingToolsEditorMode/Source/ModelingToolsEditorMode/Private/ModelingToolsManagerActions.cpp b/Engine/Plugins/Experimental/ModelingToolsEditorMode/Source/ModelingToolsEditorMode/Private/ModelingToolsManagerActions.cpp
--- a/Engine/Plugins/Experimental/ModelingToolsEditorMode/Source/ModelingToolsEditorMode/Private/ModelingToolsManagerActions.cpp
+++ b/Engine/Plugins/Experimental/ModelingToolsEditorMode/Source/ModelingToolsEditorMode/Private/ModelingToolsManagerActions.cpp
@@ -63,6 +63,7 @@ void FModelingToolsManagerCommands::RegisterCommands()
 	UI_COMMAND(BeginVoxelBooleanTool, "VoxBoolean", "Start the Voxel Boolean Tool", EUserInterfaceActionType::Button, FInputChord());
 	UI_COMMAND(BeginSelfUnionTool, "Mesh Merge", "Start the Mesh-Based (Intersection-Resolving) Merge Tool", EUserInterfaceActionType::Button, FInputChord());
 	UI_COMMAND(BeginMeshBooleanTool, "Mesh Boolean", "Start the Mesh-Based Boolean Tool", EUserInterfaceActionType::Button, FInputChord());
+	UI_COMMAND(BeginBspConversionTool, "BSP Conversion", "Start the BSP-to-static-mesh Conversion Tool", EUserInterfaceActionType::Button, FInputChord());
 	UI_COMMAND(BeginPlaneCutTool, "PlaneCut", "Start the Plane Cut Tool", EUserInterfaceActionType::Button, FInputChord());
 	UI_COMMAND(BeginHoleFillTool, "HoleFill", "Start the Hole Fill Tool", EUserInterfaceActionType::Button, FInputChord());
 	UI_COMMAND(BeginPolygonCutTool, "PolyCut", "Start the Polygon Cut Tool", EUserInterfaceActionType::Button, FInputChord());
diff --git a/Engine/Plugins/Experimental/ModelingToolsEditorMode/Source/ModelingToolsEditorMode/Public/ModelingToolsManagerActions.h b/Engine/Plugins/Experimental/ModelingToolsEditorMode/Source/ModelingToolsEditorMode/Public/ModelingToolsManagerActions.h
--- a/Engine/Plugins/Experimental/ModelingToolsEditorMode/Source/ModelingToolsEditorMode/Public/ModelingToolsManagerActions.h
+++ b/Engine/Plugins/Experimental/ModelingToolsEditorMode/Source/ModelingToolsEditorMode/Public/ModelingToolsManagerActions.h
@@ -61,6 +61,7 @@ public:
 	TSharedPtr<FUICommandInfo> BeginSelfUnionTool;
 	TSharedPtr<FUICommandInfo> BeginMeshBooleanTool;
 	TSharedPtr<FUICommandInfo> BeginMeshSelectionTool;
+	TSharedPtr<FUICommandInfo> BeginBspConversionTool;
 
 	TSharedPtr<FUICommandInfo> BeginMeshInspectorTool;
 	TSharedPtr<FUICommandInfo> BeginGlobalUVGenerateTool;
diff --git a/Engine/Source/Editor/UnrealEd/Classes/Editor/EditorEngine.h b/Engine/Source/Editor/UnrealEd/Classes/Editor/EditorEngine.h
--- a/Engine/Source/Editor/UnrealEd/Classes/Editor/EditorEngine.h
+++ b/Engine/Source/Editor/UnrealEd/Classes/Editor/EditorEngine.h
@@ -990,6 +990,15 @@ public:
 	 */
 	void RebuildModelFromBrushes(UModel* Model, bool bSelectedBrushesOnly, bool bTreatMovableBrushesAsStatic = false);
 
+	/**
+	 * Builds up a model from a given set of brushes. Used by BspConversionTool to build brushes before converting them
+	 * to static meshes.
+	 *
+	 * @param BrushesToBuild	List of brushes to build.
+	 * @param Model				Model into which to put the output.
+	 */
+	void RebuildModelFromBrushes(TArray<ABrush*>& BrushesToBuild, UModel* Model);
+
 	/**
 	 * Rebuilds levels containing currently selected brushes and should be invoked after a brush has been modified
 	 */
diff --git a/Engine/Source/Editor/UnrealEd/Private/EditorServer.cpp b/Engine/Source/Editor/UnrealEd/Private/EditorServer.cpp
--- a/Engine/Source/Editor/UnrealEd/Private/EditorServer.cpp
+++ b/Engine/Source/Editor/UnrealEd/Private/EditorServer.cpp
@@ -1856,6 +1856,45 @@ void UEditorEngine::RebuildModelFromBrushes(UModel* Model, bool bSelectedBrushes
 	FBspPointsGrid::GBspVectors = nullptr;
 }
 
+void UEditorEngine::RebuildModelFromBrushes(TArray<ABrush*> &BrushesToBuild, UModel* Model)
+{
+	TUniquePtr<FBspPointsGrid> BspPoints = MakeUnique<FBspPointsGrid>(50.0f, THRESH_POINTS_ARE_SAME);
+	TUniquePtr<FBspPointsGrid> BspVectors = MakeUnique<FBspPointsGrid>(1 / 16.0f, FMath::Max(THRESH_NORMALS_ARE_SAME, THRESH_VECTORS_ARE_NEAR));
+	FBspPointsGrid::GBspPoints = BspPoints.Get();
+	FBspPointsGrid::GBspVectors = BspVectors.Get();
+
+	// Empty the model out.
+	const int32 NumPoints = Model->Points.Num();
+	const int32 NumNodes = Model->Nodes.Num();
+	const int32 NumVerts = Model->Verts.Num();
+	const int32 NumVectors = Model->Vectors.Num();
+	const int32 NumSurfs = Model->Surfs.Num();
+
+	Model->Modify();
+	Model->EmptyModel(1, 1);
+
+	// Reserve arrays an eighth bigger than the previous allocation
+	Model->Points.Empty(NumPoints + NumPoints / 8);
+	Model->Nodes.Empty(NumNodes + NumNodes / 8);
+	Model->Verts.Empty(NumVerts + NumVerts / 8);
+	Model->Vectors.Empty(NumVectors + NumVectors / 8);
+	Model->Surfs.Empty(NumSurfs + NumSurfs / 8);
+
+	FScopedSlowTask SlowTask(BrushesToBuild.Num());
+	SlowTask.MakeDialogDelayed(3.0f);
+
+	// Compose all brushes
+	for (ABrush* Brush : BrushesToBuild)
+	{
+		SlowTask.EnterProgressFrame(1);
+		Brush->Modify();
+		bspBrushCSG(Brush, Model, Brush->PolyFlags, (EBrushType)Brush->BrushType, CSG_None, false, true, false, false);
+	}
+
+	FBspPointsGrid::GBspPoints = nullptr;
+	FBspPointsGrid::GBspVectors = nullptr;
+}
+
 
 void UEditorEngine::RebuildAlteredBSP()
 {
