ModelingTools: add BaseMeshProcessingTool, implements standard properties-driven-mesh-edit-with-bgcompute-preview pattern. Rewrote SmoothMeshTool to use this Base Tool.
#rnx

diff --git a/Engine/Plugins/Experimental/MeshModelingToolset/Source/MeshModelingTools/Private/SmoothMeshTool.cpp b/Engine/Plugins/Experimental/MeshModelingToolset/Source/MeshModelingTools/Private/SmoothMeshTool.cpp
--- a/Engine/Plugins/Experimental/MeshModelingToolset/Source/MeshModelingTools/Private/SmoothMeshTool.cpp
+++ b/Engine/Plugins/Experimental/MeshModelingToolset/Source/MeshModelingTools/Private/SmoothMeshTool.cpp
@@ -3,17 +3,8 @@
 #include "SmoothMeshTool.h"
 #include "InteractiveToolManager.h"
 #include "ToolBuilderUtil.h"
-
 #include "DynamicMesh3.h"
 #include "MeshNormals.h"
-#include "MeshTransforms.h"
-#include "DynamicMeshToMeshDescription.h"
-#include "MeshDescriptionToDynamicMesh.h"
-
-#include "SimpleDynamicMeshComponent.h"
-
-#include "SceneManagement.h" // for FPrimitiveDrawInterface
-#include "ToolSetupUtil.h"
 
 // Smoothing operators
 #include "SmoothingOps/IterativeSmoothingOp.h"
@@ -23,233 +14,50 @@
 #define LOCTEXT_NAMESPACE "USmoothMeshTool"
 
 
-/*
- * ToolBuilder
- */
-
-
-bool USmoothMeshToolBuilder::CanBuildTool(const FToolBuilderState& SceneState) const
-{
-	return ToolBuilderUtil::CountComponents(SceneState, CanMakeComponentTarget) == 1;
-}
-
-UInteractiveTool* USmoothMeshToolBuilder::BuildTool(const FToolBuilderState& SceneState) const
-{
-	USmoothMeshTool* NewTool = NewObject<USmoothMeshTool>(SceneState.ToolManager);
-
-	UActorComponent* ActorComponent = ToolBuilderUtil::FindFirstComponent(SceneState, CanMakeComponentTarget);
-	auto* MeshComponent = Cast<UPrimitiveComponent>(ActorComponent);
-	check(MeshComponent != nullptr);
-
-	NewTool->SetSelection(MakeComponentTarget(MeshComponent));
-	NewTool->SetWorld(SceneState.World);
-	NewTool->SetAssetAPI(AssetAPI);
-
-	return NewTool;
-}
-
-
-
 /*
  * Tool
  */
 
-void USmoothMeshTool::SetWorld(UWorld* World)
-{
-	this->TargetWorld = World;
-}
 
-void USmoothMeshTool::SetAssetAPI(IToolsContextAssetAPI* AssetAPIIn)
+void USmoothMeshTool::InitializeProperties()
 {
-	this->AssetAPI = AssetAPIIn;
-}
-
-
-void USmoothMeshTool::Setup()
-{
-	UInteractiveTool::Setup();
-
-	// hide input StaticMeshComponent
-	ComponentTarget->SetOwnerVisibility(false);
-
-	// initialize our properties
-	ToolPropertyObjects.Add(this);
-
-	// populate the SrcDynamicMesh with a conversion of the input mesh.
-	FMeshDescriptionToDynamicMesh Converter;
-	Converter.Convert(ComponentTarget->GetMesh(), SrcDynamicMesh);
-
-	// compute area of the input mesh and compute normalization scaling factor
-	FVector2d VolArea = TMeshQueries<FDynamicMesh3>::GetVolumeArea(SrcDynamicMesh);
-	double UnitScalingMeasure = FMathd::Max(0.01, FMathd::Sqrt(VolArea.Y / 6.0));  // 6.0 is a bit arbitrary here...surface area of unit box
-
-	// translate to origin and then apply inverse of scale
-	FAxisAlignedBox3d Bounds = SrcDynamicMesh.GetCachedBounds();
-	SrcTranslate = Bounds.Center();
-	MeshTransforms::Translate(SrcDynamicMesh, -SrcTranslate);
-	SrcScale = UnitScalingMeasure;
-	MeshTransforms::Scale(SrcDynamicMesh, (1.0/SrcScale)*FVector3d::One(), FVector3d::Zero() );
-
-	// apply that transform to target transform so that visible mesh stays in the same spot
-	OverrideTransform = ComponentTarget->GetWorldTransform();
-	FVector TranslateDelta = OverrideTransform.TransformVector((FVector)SrcTranslate);
-	FVector CurScale = OverrideTransform.GetScale3D();
-	OverrideTransform.AddToTranslation(TranslateDelta);
-	CurScale.X *= (float)SrcScale;
-	CurScale.Y *= (float)SrcScale;
-	CurScale.Z *= (float)SrcScale;
-	OverrideTransform.SetScale3D(CurScale);
-
-	// Initialize the preview mesh with a copy of the source mesh.
-	{
-		// Construct the preview object and set the material on it.
-		Preview = NewObject<UMeshOpPreviewWithBackgroundCompute>(this, "Preview");
-		Preview->Setup(this->TargetWorld, this); // Adds the actual functional tool in the Preview object
-		Preview->PreviewMesh->SetTangentsMode(EDynamicMeshTangentCalcType::AutoCalculated);
-
-		FComponentMaterialSet MaterialSet;
-		ComponentTarget->GetMaterialSet(MaterialSet);
-		Preview->ConfigureMaterials(MaterialSet.Materials,
-			ToolSetupUtil::GetDefaultWorkingMaterial(GetToolManager())
-		);
-		Preview->SetWorkingMaterialDelay(0.75);
-		Preview->PreviewMesh->SetTransform(OverrideTransform);
-		Preview->PreviewMesh->UpdatePreview(&SrcDynamicMesh);
-	}
-
-	// calculate normals
-	BaseNormals = MakeShared<FMeshNormals>(&SrcDynamicMesh);
-	BaseNormals->ComputeVertexNormals();
-
-	// show the preview mesh
-	Preview->SetVisibility(true);
-
 	SmoothProperties = NewObject<USmoothMeshToolProperties>(this);
 	AddToolPropertySource(SmoothProperties);
 	SmoothProperties->RestoreProperties(this);
 	SmoothProperties->WatchProperty(SmoothProperties->SmoothingType,
-		[&](ESmoothMeshToolSmoothType) { UpdateVisiblePropertySets(); InvalidateResult();  });
-
-	IterativeProperties = NewObject<UIterativeSmoothProperties>(this);
-	AddToolPropertySource(IterativeProperties);
-	IterativeProperties->RestoreProperties(this);
-	SetToolPropertySourceEnabled(IterativeProperties, false);
-	IterativeProperties->GetOnModified().AddLambda([this](UObject*, FProperty*) { InvalidateResult(); });
-
-	DiffusionProperties = NewObject<UDiffusionSmoothProperties>(this);
-	AddToolPropertySource(DiffusionProperties);
-	DiffusionProperties->RestoreProperties(this);
-	SetToolPropertySourceEnabled(DiffusionProperties, false);
-	DiffusionProperties->GetOnModified().AddLambda([this](UObject*, FProperty*) { InvalidateResult(); });
+		[&](ESmoothMeshToolSmoothType) { UpdateOptionalPropertyVisibility(); InvalidateResult();  });
 
-	ImplicitProperties = NewObject<UImplicitSmoothProperties>(this);
-	AddToolPropertySource(ImplicitProperties);
-	ImplicitProperties->RestoreProperties(this);
-	SetToolPropertySourceEnabled(ImplicitProperties, false);
-	ImplicitProperties->GetOnModified().AddLambda([this](UObject*, FProperty*) { InvalidateResult(); });
+	IterativeProperties = AddOptionalPropertySet<UIterativeSmoothProperties>(
+		[&]() { return SmoothProperties->SmoothingType == ESmoothMeshToolSmoothType::Iterative; } );
 
-	// start the compute
-	InvalidateResult();
+	DiffusionProperties = AddOptionalPropertySet<UDiffusionSmoothProperties>(
+		[&]() { return SmoothProperties->SmoothingType == ESmoothMeshToolSmoothType::Diffusion; });
 
-	GetToolManager()->DisplayMessage(
-		LOCTEXT("StartSmoothToolMessage", "Smooths the mesh vertex positions using various smoothing methods."),
-		EToolMessageLevel::UserNotification);
+	ImplicitProperties = AddOptionalPropertySet<UImplicitSmoothProperties>(
+		[&]() { return SmoothProperties->SmoothingType == ESmoothMeshToolSmoothType::Implicit; });
 }
 
 
-void USmoothMeshTool::Shutdown(EToolShutdownType ShutdownType)
+void USmoothMeshTool::OnShutdown(EToolShutdownType ShutdownType)
 {
 	SmoothProperties->SaveProperties(this);
-	IterativeProperties->SaveProperties(this);
-	DiffusionProperties->SaveProperties(this);
-	ImplicitProperties->SaveProperties(this);
-
-	// Restore (unhide) the source meshes
-	ComponentTarget->SetOwnerVisibility(true);
-
-	if (Preview != nullptr)
-	{
-		FDynamicMeshOpResult Result = Preview->Shutdown();
-		if (ShutdownType == EToolShutdownType::Accept)
-		{
-			
-			GetToolManager()->BeginUndoTransaction(LOCTEXT("SmoothMeshToolTransactionName", "Smooth Mesh"));
-
-			FDynamicMesh3* DynamicMeshResult = Result.Mesh.Get();
-			check(DynamicMeshResult != nullptr);
-
-			MeshTransforms::Scale(*DynamicMeshResult, FVector3d(SrcScale, SrcScale, SrcScale), FVector3d::Zero());
-			MeshTransforms::Translate(*DynamicMeshResult, SrcTranslate);
-
-			ComponentTarget->CommitMesh([DynamicMeshResult](const FPrimitiveComponentTarget::FCommitParams& CommitParams)
-			{
-				FDynamicMeshToMeshDescription Converter;
-				Converter.Convert(DynamicMeshResult, *CommitParams.MeshDescription);
-			});
-
-
-			GetToolManager()->EndUndoTransaction();
-		}
-	}
-}
-
-void USmoothMeshTool::Render(IToolsContextRenderAPI* RenderAPI)
-{
-	UpdateResult();
 }
 
 
-
-void USmoothMeshTool::OnTick(float DeltaTime)
+FText USmoothMeshTool::GetToolMessageString() const
 {
-	Preview->Tick(DeltaTime);
+	return LOCTEXT("StartSmoothToolMessage", "Smooths the mesh vertex positions using various smoothing methods.");
 }
 
-void USmoothMeshTool::InvalidateResult()
+FText USmoothMeshTool::GetAcceptTransactionName() const
 {
-	Preview->InvalidateResult();
-	bResultValid = false;
+	return LOCTEXT("SmoothMeshToolTransactionName", "Smooth Mesh");
 }
 
-void USmoothMeshTool::UpdateResult()
-{
-	if (bResultValid) 
-	{
-		return;
-	}
 
-	bResultValid = Preview->HaveValidResult();
-}
-
-bool USmoothMeshTool::HasAccept() const
-{
-	return true;
-}
-
-bool USmoothMeshTool::CanAccept() const
-{
-	return bResultValid;
-}
-
-
-void USmoothMeshTool::UpdateVisiblePropertySets()
-{
-	SetToolPropertySourceEnabled(IterativeProperties, false);
-	SetToolPropertySourceEnabled(DiffusionProperties, false);
-	SetToolPropertySourceEnabled(ImplicitProperties, false);
-
-	switch (SmoothProperties->SmoothingType)
-	{
-	case ESmoothMeshToolSmoothType::Iterative:
-		SetToolPropertySourceEnabled(IterativeProperties, true);
-		break;
-	case ESmoothMeshToolSmoothType::Diffusion:
-		SetToolPropertySourceEnabled(DiffusionProperties, true);
-		break;
-	case ESmoothMeshToolSmoothType::Implicit:
-		SetToolPropertySourceEnabled(ImplicitProperties, true);
-		break;
-	}
+bool USmoothMeshTool::HasMeshTopologyChanged() const 
+{ 
+	return false; 
 }
 
 
@@ -257,8 +65,10 @@ TUniquePtr<FDynamicMeshOperator> USmoothMeshTool::MakeNewOperator()
 {
 	TUniquePtr<FSmoothingOpBase> MeshOp;
 	
+	const FDynamicMesh3* Mesh = &GetInitialMesh();
+
 	FSmoothingOpBase::FOptions Options;
-	Options.BaseNormals = this->BaseNormals;
+	Options.BaseNormals = this->GetBaseNormals();
 
 	switch (SmoothProperties->SmoothingType)
 	{
@@ -269,7 +79,7 @@ TUniquePtr<FDynamicMeshOperator> USmoothMeshTool::MakeNewOperator()
 		Options.bSmoothBoundary = IterativeProperties->bSmoothBoundary;
 		Options.bUniform = true;
 		Options.bUseImplicit = false;
-		MeshOp = MakeUnique<FIterativeSmoothingOp>(&SrcDynamicMesh, Options);
+		MeshOp = MakeUnique<FIterativeSmoothingOp>(Mesh, Options);
 		break;
 
 	case ESmoothMeshToolSmoothType::Diffusion:
@@ -277,7 +87,7 @@ TUniquePtr<FDynamicMeshOperator> USmoothMeshTool::MakeNewOperator()
 		Options.Iterations = DiffusionProperties->Steps;
 		Options.bUniform = DiffusionProperties->bPreserveUVs == false;
 		Options.bUseImplicit = true;
-		MeshOp = MakeUnique<FIterativeSmoothingOp>(&SrcDynamicMesh, Options);
+		MeshOp = MakeUnique<FIterativeSmoothingOp>(Mesh, Options);
 		break;
 
 	case ESmoothMeshToolSmoothType::Implicit:
@@ -285,21 +95,20 @@ TUniquePtr<FDynamicMeshOperator> USmoothMeshTool::MakeNewOperator()
 		Options.SmoothAlpha = ImplicitProperties->SmoothSpeed;
 		double NonlinearT = FMathd::Pow(ImplicitProperties->Smoothness, 2.0);
 		// this is an empirically-determined hack that seems to work OK to normalize the smoothing result for variable vertex count...
-		double ScaledPower = (NonlinearT/50.0) * SrcDynamicMesh.VertexCount();
+		double ScaledPower = (NonlinearT/50.0) * Mesh->VertexCount();
 		Options.SmoothPower = ScaledPower;
 		Options.bUniform = ImplicitProperties->bPreserveUVs == false;
 		Options.bUseImplicit = true;
 		Options.NormalOffset = ImplicitProperties->VolumeCorrection;
-		MeshOp = MakeUnique<FCotanSmoothingOp>(&SrcDynamicMesh, Options);
+		MeshOp = MakeUnique<FCotanSmoothingOp>(Mesh, Options);
 		}
 		break;
 	}
 
-	FTransform3d XForm3d(OverrideTransform);
+	FTransform3d XForm3d(GetPreviewTransform());
 	MeshOp->SetTransform(XForm3d);
 
 	return MeshOp;
-
 }
 
 
diff --git a/Engine/Plugins/Experimental/MeshModelingToolset/Source/MeshModelingTools/Public/SmoothMeshTool.h b/Engine/Plugins/Experimental/MeshModelingToolset/Source/MeshModelingTools/Public/SmoothMeshTool.h
--- a/Engine/Plugins/Experimental/MeshModelingToolset/Source/MeshModelingTools/Public/SmoothMeshTool.h
+++ b/Engine/Plugins/Experimental/MeshModelingToolset/Source/MeshModelingTools/Public/SmoothMeshTool.h
@@ -3,42 +3,9 @@
 #pragma once
 
 #include "CoreMinimal.h"
-#include "UObject/NoExportTypes.h"
-#include "SingleSelectionTool.h"
-#include "InteractiveToolBuilder.h"
-#include "DynamicMesh3.h"
-#include "MeshOpPreviewHelpers.h"
+#include "BaseTools/BaseMeshProcessingTool.h"
 #include "SmoothMeshTool.generated.h"
 
-// predeclarations
-struct FMeshDescription;
-class FMeshNormals;
-class USimpleDynamicMeshComponent;
-
-/**
- *
- */
-UCLASS()
-class MESHMODELINGTOOLS_API USmoothMeshToolBuilder : public UInteractiveToolBuilder
-{
-	GENERATED_BODY()
-
-public:
-
-	IToolsContextAssetAPI* AssetAPI;
-
-	USmoothMeshToolBuilder()
-	{
-		AssetAPI = nullptr;
-	}
-
-	virtual bool CanBuildTool(const FToolBuilderState& SceneState) const override;
-	virtual UInteractiveTool* BuildTool(const FToolBuilderState& SceneState) const override;
-};
-
-
-
-
 
 UENUM()
 enum class ESmoothMeshToolSmoothType : uint8
@@ -63,7 +30,7 @@ class MESHMODELINGTOOLS_API USmoothMeshToolProperties : public UInteractiveToolP
 
 public:
 	/** Type of smoothing to apply */
-	UPROPERTY(EditAnywhere, Category = Options)
+	UPROPERTY(EditAnywhere, Category = SmoothingType)
 	ESmoothMeshToolSmoothType SmoothingType = ESmoothMeshToolSmoothType::Iterative;
 };
 
@@ -76,15 +43,15 @@ class MESHMODELINGTOOLS_API UIterativeSmoothProperties : public UInteractiveTool
 	GENERATED_BODY()
 public:
 	/** Amount of smoothing allowed per step. Smaller steps will avoid things like collapse of small/thin features. */
-	UPROPERTY(EditAnywhere, Category = IterativeSmoothing, meta = (UIMin = "0.0", UIMax = "1.0", ClampMin = "0.0", ClampMax = "1.0"))
+	UPROPERTY(EditAnywhere, Category = IterativeSmoothingOptions, meta = (UIMin = "0.0", UIMax = "1.0", ClampMin = "0.0", ClampMax = "1.0"))
 	float SmoothingPerStep = 0.8f;
 
 	/** Number of Smoothing iterations */
-	UPROPERTY(EditAnywhere, Category = IterativeSmoothing, meta = (UIMin = "0", UIMax = "100", ClampMin = "0", ClampMax = "1000"))
+	UPROPERTY(EditAnywhere, Category = IterativeSmoothingOptions, meta = (UIMin = "0", UIMax = "100", ClampMin = "0", ClampMax = "1000"))
 	int32 Steps = 10;
 
 	/** If this is false, the smoother will try to reshape the triangles to be more regular, which will distort UVs */
-	UPROPERTY(EditAnywhere, Category = IterativeSmoothing)
+	UPROPERTY(EditAnywhere, Category = IterativeSmoothingOptions)
 	bool bSmoothBoundary = true;
 };
 
@@ -98,15 +65,15 @@ class MESHMODELINGTOOLS_API UDiffusionSmoothProperties : public UInteractiveTool
 
 public:
 	/** Amount of smoothing allowed per step. Smaller steps will avoid things like collapse of small/thin features. */
-	UPROPERTY(EditAnywhere, Category = DiffusionSmoothing, meta = (UIMin = "0.0", UIMax = "1.0", ClampMin = "0.0", ClampMax = "1.0"))
+	UPROPERTY(EditAnywhere, Category = DiffusionSmoothingOptions, meta = (UIMin = "0.0", UIMax = "1.0", ClampMin = "0.0", ClampMax = "1.0"))
 	float SmoothingPerStep = 0.8f;
 
 	/** Number of Smoothing iterations */
-	UPROPERTY(EditAnywhere, Category = DiffusionSmoothing, meta = (UIMin = "0", UIMax = "100", ClampMin = "0", ClampMax = "1000"))
+	UPROPERTY(EditAnywhere, Category = DiffusionSmoothingOptions, meta = (UIMin = "0", UIMax = "100", ClampMin = "0", ClampMax = "1000"))
 	int32 Steps = 1;
 
 	/** If this is false, the smoother will try to reshape the triangles to be more regular, which will distort UVs */
-	UPROPERTY(EditAnywhere, Category = DiffusionSmoothing)
+	UPROPERTY(EditAnywhere, Category = DiffusionSmoothingOptions)
 	bool bPreserveUVs = true;
 };
 
@@ -126,15 +93,15 @@ public:
 	float SmoothSpeed = 0.1f;
 
 	/** Desired Smoothness. This is not a linear quantity, but larger numbers produce smoother results */
-	UPROPERTY(EditAnywhere, Category = ImplicitSmoothing, meta = (UIMin = "0.0", UIMax = "1.0", ClampMin = "0.0", ClampMax = "100.0"))
+	UPROPERTY(EditAnywhere, Category = ImplicitSmoothingOptions, meta = (UIMin = "0.0", UIMax = "1.0", ClampMin = "0.0", ClampMax = "100.0"))
 	float Smoothness = 0.2f;
 
 	/** If this is false, the smoother will try to reshape the triangles to be more regular, which will distort UVs */
-	UPROPERTY(EditAnywhere, Category = ImplicitSmoothing)
+	UPROPERTY(EditAnywhere, Category = ImplicitSmoothingOptions)
 	bool bPreserveUVs = true;
 
 	/** Magic number that allows you to try to correct for shrinking caused by smoothing */
-	UPROPERTY(EditAnywhere, Category = ImplicitSmoothing, meta = (UIMin = "0.0", UIMax = "1.0", ClampMin = "0.0", ClampMax = "10.0"))
+	UPROPERTY(EditAnywhere, Category = ImplicitSmoothingOptions, meta = (UIMin = "0.0", UIMax = "1.0", ClampMin = "0.0", ClampMax = "10.0"))
 	float VolumeCorrection = 0.0f;
 };
 
@@ -144,31 +111,26 @@ public:
 
 
 /**
- * Simple Mesh Smoothing Tool
+ * Mesh Smoothing Tool
  */
 UCLASS()
-class MESHMODELINGTOOLS_API USmoothMeshTool : public USingleSelectionTool, public IDynamicMeshOperatorFactory
+class MESHMODELINGTOOLS_API USmoothMeshTool : public UBaseMeshProcessingTool
 {
 	GENERATED_BODY()
 
 public:
 	USmoothMeshTool() = default;
 
-	virtual void SetWorld(UWorld* World);
-	virtual void SetAssetAPI(IToolsContextAssetAPI* AssetAPIIn);
+	virtual void InitializeProperties() override;
+	virtual void OnShutdown(EToolShutdownType ShutdownType) override;
 
-	virtual void Setup() override;
-	virtual void Shutdown(EToolShutdownType ShutdownType) override;
+	virtual TUniquePtr<FDynamicMeshOperator> MakeNewOperator() override;
 
-	virtual void OnTick(float DeltaTime) override;
-	virtual void Render(IToolsContextRenderAPI* RenderAPI) override;
+	virtual bool HasMeshTopologyChanged() const override;
 
-	virtual bool HasCancel() const override { return true; }
-	virtual bool HasAccept() const override;
-	virtual bool CanAccept() const override;
+	virtual FText GetToolMessageString() const override;
+	virtual FText GetAcceptTransactionName() const override;
 
-	virtual TUniquePtr<FDynamicMeshOperator> MakeNewOperator() override;
-	
 protected:
 	UPROPERTY()
 	USmoothMeshToolProperties* SmoothProperties = nullptr;
@@ -181,27 +143,20 @@ protected:
 
 	UPROPERTY()
 	UImplicitSmoothProperties* ImplicitProperties = nullptr;
+};
 
-	void UpdateVisiblePropertySets();
-
-protected:
-	void UpdateResult();
-
-	bool bResultValid = false;
-	void InvalidateResult();
-
-	UPROPERTY()
-	UMeshOpPreviewWithBackgroundCompute* Preview = nullptr;
-	FDynamicMesh3 SrcDynamicMesh;
 
-	// scale/translate applied to input mesh to regularize it
-	FVector3d SrcTranslate;
-	double SrcScale;
-	// transform that does the opposite of scale/translate so that mesh stays in the right spot on screen
-	FTransform OverrideTransform;
 
-	TSharedPtr<FMeshNormals> BaseNormals;
 
-	UWorld* TargetWorld = nullptr;
-	IToolsContextAssetAPI* AssetAPI = nullptr;
+/**
+ *
+ */
+UCLASS()
+class MESHMODELINGTOOLS_API USmoothMeshToolBuilder : public UBaseMeshProcessingToolBuilder
+{
+	GENERATED_BODY()
+public:
+	virtual UBaseMeshProcessingTool* MakeNewToolInstance(UObject* Outer) const {
+		return NewObject<USmoothMeshTool>(Outer);
+	}
 };
diff --git a/Engine/Plugins/Experimental/MeshModelingToolset/Source/ModelingComponents/Private/BaseTools/BaseMeshProcessingTool.cpp b/Engine/Plugins/Experimental/MeshModelingToolset/Source/ModelingComponents/Private/BaseTools/BaseMeshProcessingTool.cpp
new file mode 100644
--- /dev/null
+++ b/Engine/Plugins/Experimental/MeshModelingToolset/Source/ModelingComponents/Private/BaseTools/BaseMeshProcessingTool.cpp
@@ -0,0 +1,324 @@
+// Copyright Epic Games, Inc. All Rights Reserved.
+
+#include "BaseTools/BaseMeshProcessingTool.h"
+#include "InteractiveToolManager.h"
+#include "ToolBuilderUtil.h"
+#include "ToolSetupUtil.h"
+#include "SimpleDynamicMeshComponent.h"
+
+#include "MeshNormals.h"
+#include "MeshTransforms.h"
+#include "DynamicMeshToMeshDescription.h"
+#include "MeshDescriptionToDynamicMesh.h"
+
+
+#define LOCTEXT_NAMESPACE "UBaseMeshProcessingTool"
+
+
+/*
+ * ToolBuilder
+ */
+
+
+bool UBaseMeshProcessingToolBuilder::CanBuildTool(const FToolBuilderState& SceneState) const
+{
+	if (SupportsMultipleObjects())
+	{
+		return ToolBuilderUtil::CountComponents(SceneState, CanMakeComponentTarget) > 0;
+	}
+	else
+	{
+		return ToolBuilderUtil::CountComponents(SceneState, CanMakeComponentTarget) == 1;
+	}
+}
+
+UInteractiveTool* UBaseMeshProcessingToolBuilder::BuildTool(const FToolBuilderState& SceneState) const
+{
+	check(SupportsMultipleObjects() == false); // not supported yet
+
+	UBaseMeshProcessingTool* NewTool = MakeNewToolInstance(SceneState.ToolManager);
+
+	TArray<UActorComponent*> Components = ToolBuilderUtil::FindAllComponents(SceneState, CanMakeComponentTarget);
+	check(Components.Num() > 0);
+
+	TArray<TUniquePtr<FPrimitiveComponentTarget>> ComponentTargets;
+	for (UActorComponent* ActorComponent : Components)
+	{
+		auto* MeshComponent = Cast<UPrimitiveComponent>(ActorComponent);
+		if (MeshComponent)
+		{
+			ComponentTargets.Add(MakeComponentTarget(MeshComponent));
+		}
+	}
+
+	NewTool->SetSelection( MoveTemp(ComponentTargets[0]) );
+	NewTool->SetWorld(SceneState.World);
+
+	return NewTool;
+}
+
+
+
+
+/*
+ * Tool
+ */
+
+void UBaseMeshProcessingTool::SetWorld(UWorld* World)
+{
+	this->TargetWorld = World;
+}
+
+
+void UBaseMeshProcessingTool::Setup()
+{
+	UInteractiveTool::Setup();
+
+	// hide input StaticMeshComponent
+	ComponentTarget->SetOwnerVisibility(false);
+
+	// initialize our properties
+	ToolPropertyObjects.Add(this);
+
+	// populate the BaseMesh with a conversion of the input mesh.
+	FMeshDescriptionToDynamicMesh Converter;
+	Converter.Convert(ComponentTarget->GetMesh(), InitialMesh);
+
+	if (RequiresScaleNormalization())
+	{
+		// compute area of the input mesh and compute normalization scaling factor
+		FVector2d VolArea = TMeshQueries<FDynamicMesh3>::GetVolumeArea(InitialMesh);
+		double UnitScalingMeasure = FMathd::Max(0.01, FMathd::Sqrt(VolArea.Y / 6.0));  // 6.0 is a bit arbitrary here...surface area of unit box
+
+		// translate to origin and then apply inverse of scale
+		FAxisAlignedBox3d Bounds = InitialMesh.GetCachedBounds();
+		SrcTranslate = Bounds.Center();
+		MeshTransforms::Translate(InitialMesh, -SrcTranslate);
+		SrcScale = UnitScalingMeasure;
+		MeshTransforms::Scale(InitialMesh, (1.0 / SrcScale) * FVector3d::One(), FVector3d::Zero());
+
+		// apply that transform to target transform so that visible mesh stays in the same spot
+		OverrideTransform = ComponentTarget->GetWorldTransform();
+		FVector TranslateDelta = OverrideTransform.TransformVector((FVector)SrcTranslate);
+		FVector CurScale = OverrideTransform.GetScale3D();
+		OverrideTransform.AddToTranslation(TranslateDelta);
+		CurScale.X *= (float)SrcScale;
+		CurScale.Y *= (float)SrcScale;
+		CurScale.Z *= (float)SrcScale;
+		OverrideTransform.SetScale3D(CurScale);
+
+		bIsScaleNormalizationApplied = true;
+	}
+	else
+	{
+		SrcTranslate = FVector3d::Zero();
+		SrcScale = 1.0;
+		OverrideTransform = ComponentTarget->GetWorldTransform();
+		bIsScaleNormalizationApplied = false;
+	}
+
+	// Construct the preview object and set the material on it.
+	Preview = NewObject<UMeshOpPreviewWithBackgroundCompute>(this, "Preview");
+	Preview->Setup(this->TargetWorld, this); // Adds the actual functional tool in the Preview object
+	Preview->PreviewMesh->SetTangentsMode(EDynamicMeshTangentCalcType::AutoCalculated);
+
+	FComponentMaterialSet MaterialSet;
+	ComponentTarget->GetMaterialSet(MaterialSet);
+	Preview->ConfigureMaterials(MaterialSet.Materials,
+		ToolSetupUtil::GetDefaultWorkingMaterial(GetToolManager())
+	);
+	Preview->SetWorkingMaterialDelay(0.75);
+	Preview->PreviewMesh->SetTransform(OverrideTransform);
+	Preview->PreviewMesh->UpdatePreview(&InitialMesh);
+
+	// calculate base mesh vertex normals if necessary normals
+	if (RequiresBaseNormals())
+	{
+		BaseNormals = MakeShared<FMeshNormals>(&InitialMesh);
+		BaseNormals->ComputeVertexNormals();
+	}
+
+	// show the preview mesh
+	Preview->SetVisibility(true);
+
+
+	InitializeProperties();
+	UpdateOptionalPropertyVisibility();
+
+	// start the compute
+	InvalidateResult();
+
+	GetToolManager()->DisplayMessage( GetToolMessageString(), EToolMessageLevel::UserNotification);
+}
+
+
+
+FText UBaseMeshProcessingTool::GetToolMessageString() const
+{
+	return FText::GetEmpty();
+}
+
+FText UBaseMeshProcessingTool::GetAcceptTransactionName() const
+{
+	return LOCTEXT("BaseMeshProcessingToolTransactionName", "Update Mesh");
+}
+
+
+
+void UBaseMeshProcessingTool::SavePropertySets()
+{
+	for (FOptionalPropertySet& PropStruct : OptionalProperties)
+	{
+		if (PropStruct.PropertySet.IsValid())
+		{
+			PropStruct.PropertySet->SaveProperties(this);
+		}
+	}
+}
+
+
+
+void UBaseMeshProcessingTool::Shutdown(EToolShutdownType ShutdownType)
+{
+	OnShutdown(ShutdownType);
+
+	SavePropertySets();
+
+	// Restore (unhide) the source meshes
+	ComponentTarget->SetOwnerVisibility(true);
+
+	if (Preview != nullptr)
+	{
+		FDynamicMeshOpResult Result = Preview->Shutdown();
+		if (ShutdownType == EToolShutdownType::Accept)
+		{
+			
+			GetToolManager()->BeginUndoTransaction(GetAcceptTransactionName());
+
+			FDynamicMesh3* DynamicMeshResult = Result.Mesh.Get();
+			check(DynamicMeshResult != nullptr);
+
+			// un-apply scale normalization if it was applied
+			if (bIsScaleNormalizationApplied)
+			{
+				MeshTransforms::Scale(*DynamicMeshResult, FVector3d(SrcScale, SrcScale, SrcScale), FVector3d::Zero());
+				MeshTransforms::Translate(*DynamicMeshResult, SrcTranslate);
+			}
+
+			bool bTopologyChanged = HasMeshTopologyChanged();
+			ComponentTarget->CommitMesh([DynamicMeshResult, bTopologyChanged](const FPrimitiveComponentTarget::FCommitParams& CommitParams)
+			{
+				FDynamicMeshToMeshDescription Converter;
+				if (bTopologyChanged)
+				{
+					Converter.Convert(DynamicMeshResult, *CommitParams.MeshDescription);
+				}
+				else
+				{
+					Converter.Update(DynamicMeshResult, *CommitParams.MeshDescription);
+				}
+			});
+
+
+			GetToolManager()->EndUndoTransaction();
+		}
+	}
+}
+
+void UBaseMeshProcessingTool::Render(IToolsContextRenderAPI* RenderAPI)
+{
+	UpdateResult();
+}
+
+void UBaseMeshProcessingTool::OnTick(float DeltaTime)
+{
+	Preview->Tick(DeltaTime);
+}
+
+void UBaseMeshProcessingTool::InvalidateResult()
+{
+	Preview->InvalidateResult();
+	bResultValid = false;
+}
+
+void UBaseMeshProcessingTool::UpdateResult()
+{
+	if (bResultValid) 
+	{
+		return;
+	}
+
+	bResultValid = Preview->HaveValidResult();
+}
+
+bool UBaseMeshProcessingTool::HasAccept() const
+{
+	return true;
+}
+
+bool UBaseMeshProcessingTool::CanAccept() const
+{
+	return bResultValid;
+}
+
+
+
+
+void UBaseMeshProcessingTool::AddOptionalPropertySet(
+	UInteractiveToolPropertySet* PropSet, 
+	TUniqueFunction<bool()> VisibilityFunc, 
+	TUniqueFunction<void()> OnModifiedFunc,
+	bool bChangeInvalidatesResult)
+{
+	AddToolPropertySource(PropSet);
+	PropSet->RestoreProperties(this);
+	SetToolPropertySourceEnabled(PropSet, false);
+
+	FOptionalPropertySet PropSetStruct;
+	PropSetStruct.PropertySet = PropSet;
+	PropSetStruct.IsVisible = MoveTemp(VisibilityFunc);
+	PropSetStruct.OnModifiedFunc = MoveTemp(OnModifiedFunc);
+	PropSetStruct.bInvalidateOnModify = bChangeInvalidatesResult;
+	int32 Index = OptionalProperties.Num();
+	OptionalProperties.Add(MoveTemp(PropSetStruct));
+
+	PropSet->GetOnModified().AddLambda([Index, this](UObject*, FProperty*) { OnOptionalPropSetModified(Index); } );
+}
+
+
+void UBaseMeshProcessingTool::OnOptionalPropSetModified(int32 Index)
+{
+	const FOptionalPropertySet& PropSetStruct = OptionalProperties[Index];
+	PropSetStruct.OnModifiedFunc();
+	if (PropSetStruct.bInvalidateOnModify)
+	{
+		InvalidateResult();
+	}
+}
+
+
+void UBaseMeshProcessingTool::UpdateOptionalPropertyVisibility()
+{
+	for (FOptionalPropertySet& PropStruct : OptionalProperties)
+	{
+		if (PropStruct.PropertySet.IsValid())
+		{
+			bool bVisible = PropStruct.IsVisible();
+			SetToolPropertySourceEnabled(PropStruct.PropertySet.Get(), bVisible);
+		}
+	}
+}
+
+
+
+
+TSharedPtr<FMeshNormals>& UBaseMeshProcessingTool::GetBaseNormals()
+{
+	checkf(BaseNormals.IsValid(), TEXT("Base Normals have not been computed - must return true from RequiresBaseNormals()") );
+	return BaseNormals;
+}
+
+
+
+
+#undef LOCTEXT_NAMESPACE
diff --git a/Engine/Plugins/Experimental/MeshModelingToolset/Source/ModelingComponents/Public/BaseTools/BaseMeshProcessingTool.h b/Engine/Plugins/Experimental/MeshModelingToolset/Source/ModelingComponents/Public/BaseTools/BaseMeshProcessingTool.h
new file mode 100644
--- /dev/null
+++ b/Engine/Plugins/Experimental/MeshModelingToolset/Source/ModelingComponents/Public/BaseTools/BaseMeshProcessingTool.h
@@ -0,0 +1,253 @@
+// Copyright Epic Games, Inc. All Rights Reserved.
+
+#pragma once
+
+#include "CoreMinimal.h"
+#include "UObject/WeakObjectPtr.h"
+#include "SingleSelectionTool.h"
+#include "InteractiveToolBuilder.h"
+#include "DynamicMesh3.h"
+#include "MeshOpPreviewHelpers.h"
+#include "BaseMeshProcessingTool.generated.h"
+
+// predeclarations
+struct FMeshDescription;
+class FMeshNormals;
+class USimpleDynamicMeshComponent;
+class UPreviewMesh;
+
+/**
+ * ToolBuilder for UBaseMeshProcessingTool
+ */
+UCLASS()
+class MODELINGCOMPONENTS_API UBaseMeshProcessingToolBuilder : public UInteractiveToolBuilder
+{
+	GENERATED_BODY()
+
+public:
+	virtual bool CanBuildTool(const FToolBuilderState& SceneState) const override;
+	virtual UInteractiveTool* BuildTool(const FToolBuilderState& SceneState) const override;
+
+public:
+	
+	virtual bool SupportsMultipleObjects() const { return false;  }
+	
+	// subclass must override!
+	virtual UBaseMeshProcessingTool* MakeNewToolInstance(UObject* Outer) const { check(false); return nullptr; }
+
+};
+
+
+
+
+/**
+ * UBaseMeshProcessingTool is a base Tool (ie has no functionality of it's own and must be subclassed) 
+ * that provides the following structure:
+ *   - a Background-Compute-With-Preview Temp Actor/Component is created based on the input mesh
+ *   - The Subclass provides FDynamicMeshOperator instances (via IDynamicMeshOperatorFactory) that process/modify and update this Preview
+ *   - PropertySets with custom visibility can be registered, and on change will invalidate the current computation
+ *   
+ * Most subclasses will only need to define their PropertySets and implement MakeNewOperator(), see eg SmoothMeshTool for a minimal example
+ *
+ * Other functions:
+ *   - GetInitialMesh() : return reference to copy of initial mesh, used to initialize FDynamicMeshOperator
+ *   - GetUPreviewMesh() : return the UPreviewMesh inside the background compute (for configuration/etc - should not directly touch the mesh!)
+ *   - GetPreviewTransform() : return active FTransform on the Preview mesh, should be passed to FDynamicMeshOperator unless it is outputting world position
+ *   - InvalidateResult() : subclasses call this to notify the base class that current result/computation has been invalidated
+ *
+ * The Base tool will do various optional precomputations or changes to the input mesh, which can be configured by
+ * overriding various functions below.
+ * 
+ *   RequiresBaseNormals() : return true (default) to calculate per-vertex normals on the input mesh, returned by GetBaseNormals()
+ * 
+ *   RequiresScaleNormalization() : return true (default) to apply an initial scale to the input mesh so that it has consistent size
+ *     before being sent into the computation. Scaling factor (eg for scaling UI constants) can be accessed via GetScaleNormalizationFactor()
+ *
+ */
+UCLASS()
+class MODELINGCOMPONENTS_API UBaseMeshProcessingTool : public USingleSelectionTool, public IDynamicMeshOperatorFactory
+{
+	GENERATED_BODY()
+
+public:
+	UBaseMeshProcessingTool() = default;
+
+	virtual void SetWorld(UWorld* World);
+
+	//
+	// InteractiveTool API - generally does not need to be modified by subclasses
+	//
+
+	virtual void Setup() override;
+	virtual void Shutdown(EToolShutdownType ShutdownType) override;
+
+	virtual void OnTick(float DeltaTime) override;
+	virtual void Render(IToolsContextRenderAPI* RenderAPI) override;
+
+	virtual bool HasCancel() const override { return true; }
+	virtual bool HasAccept() const override;
+	virtual bool CanAccept() const override;
+
+
+protected:
+
+	//
+	// UBaseMeshProcessingTool REQUIRED API - subclasses must implement these functions
+	//
+
+	/**
+	 * IDynamicMeshOperatorFactory implementation that subclass must override and implement
+	 */
+	virtual TUniquePtr<FDynamicMeshOperator> MakeNewOperator() override
+	{
+		check(false);
+		return TUniquePtr<FDynamicMeshOperator>();
+	}
+
+
+	/**
+	 * called when Tool is Accepted to determine whether it is safe to only update vertex positions, or if entire target mesh must be replaced.
+	 */
+	virtual bool HasMeshTopologyChanged() const
+	{
+		check(false);
+		return true;
+	}
+
+	//
+	// UBaseMeshProcessingTool OPTIONAL API - subclasses may implement these functions
+	//
+
+	/** @return text string shown to user for Accept Transaction that updates input mesh. Subclass should override. */
+	virtual FText GetToolMessageString() const;
+
+	/** @return text string shown to user for Accept Transaction that updates input mesh. Subclass should override. */
+	virtual FText GetAcceptTransactionName() const;
+
+	/** This function is called during ::Setup() to allow subclass to register property sets, before kicking off initial computation */
+	virtual void InitializeProperties() {}
+
+	/** This function is called during ::Shutdown so that subclass may perform final processing and save property sets */
+	virtual void OnShutdown(EToolShutdownType ShutdownType) {}
+
+
+	//
+	// Optional Property Set API - subclasses can use this to manage property sets with configurable visibility that invalidate the precompute
+	//
+
+protected:
+
+	/**
+	 * Register an optional property set with the given VisibilityFunc
+	 */
+	template<class PropSetType>
+	PropSetType* AddOptionalPropertySet(TUniqueFunction<bool()> VisibilityFunc, bool bChangeInvalidatesResult = true)
+	{
+		return AddOptionalPropertySet<PropSetType>(MoveTemp(VisibilityFunc), []() {}, bChangeInvalidatesResult);
+	}
+
+	/**
+	 * Register an optional property set with the given VisibilityFunc, and call OnModifiedFunc if any of the properties change
+	 */
+	template<class PropSetType>
+	PropSetType* AddOptionalPropertySet( TUniqueFunction<bool()> VisibilityFunc, TUniqueFunction<void()> OnModifiedFunc, bool bChangeInvalidatesResult)
+	{
+		PropSetType* PropSet = NewObject<PropSetType>(this);
+		AddOptionalPropertySet(PropSet, MoveTemp(VisibilityFunc), MoveTemp(OnModifiedFunc), bChangeInvalidatesResult);
+		return PropSet;
+	}
+
+	/** Call this function to update optional property sets visibility. Should call base implementation */
+	virtual void UpdateOptionalPropertyVisibility();
+
+protected:
+
+	virtual void AddOptionalPropertySet(UInteractiveToolPropertySet* PropSet, 
+		TUniqueFunction<bool()> VisibilityFunc, 
+		TUniqueFunction<void()> OnModifiedFunc,
+		bool bChangeInvalidatesResult);
+
+	struct FOptionalPropertySet
+	{
+		TUniqueFunction<bool()> IsVisible;
+		TUniqueFunction<void()> OnModifiedFunc;
+		bool bInvalidateOnModify;
+		TWeakObjectPtr<UInteractiveToolPropertySet> PropertySet;
+	};
+
+	TArray<FOptionalPropertySet> OptionalProperties;
+	virtual void OnOptionalPropSetModified(int32 Index);
+	virtual void SavePropertySets();
+
+
+protected:
+	UWorld* TargetWorld = nullptr;
+	IToolsContextAssetAPI* AssetAPI = nullptr;
+
+	// Preview object holds temporary Actor with preview mesh component
+	UPROPERTY()
+	UMeshOpPreviewWithBackgroundCompute* Preview = nullptr;
+
+	UPreviewMesh* GetUPreviewMesh() const { return Preview->PreviewMesh; }
+	const FTransform& GetPreviewTransform() const { return OverrideTransform; }
+
+
+private:
+	bool bResultValid = false;
+protected:
+	virtual void InvalidateResult();
+	virtual void UpdateResult();
+
+
+	//
+	// Initial data, used to initialize background compute / etc
+	//
+
+private:
+	FDynamicMesh3 InitialMesh;
+protected:
+	/** @return duplciate of initial mesh (possibly with optional size normalization) */
+	const FDynamicMesh3& GetInitialMesh() const { return InitialMesh;  }
+	/** @return duplciate of initial mesh (possibly with optional size normalization) */
+	FDynamicMesh3& GetInitialMesh() { return InitialMesh; }
+
+
+
+	//
+	// Optional base mesh per-vertex normals
+	//
+protected:
+	/** 
+	 * If this function returns true, BaseNormals will be initialized in Tool ::Setup(). 
+	 * This has some cost and should be disabled if not necessary. 
+	 */
+	virtual bool RequiresBaseNormals() const { return true; }
+
+	/** @return calculated base normals. This pointer does not change for the lifetime of the Tool. */
+	TSharedPtr<FMeshNormals>& GetBaseNormals();
+
+private:
+	TSharedPtr<FMeshNormals> BaseNormals;
+
+
+	//
+	// Optional uniform scale applied to mesh. Enabled by default
+	// 
+
+protected:
+	/** 
+	 * If this function returns true, input mesh will be scaled to normalized dimension in ::Setup() before any processing begins.
+	 * This scaling will be undone on Accept()
+	 */
+	virtual bool RequiresScaleNormalization() const { return true; }
+
+	double GetScaleNormalizationFactor() const { return 1.0 / SrcScale; };
+
+private:
+	// scale/translate applied to input mesh to regularize it
+	bool bIsScaleNormalizationApplied = false;
+	FVector3d SrcTranslate;
+	double SrcScale;
+	// transform that does the opposite of scale/translate so that mesh stays in the right spot on screen
+	FTransform OverrideTransform;
+};
