Missed file from CL 14247116:

The EditorToolsContext is now owned by the mode manager.

Updated Mesh Paint, geometry, arch gen, hair lab, and sample editor modes to reflect changes to the tools context ownership.

diff --git a/Engine/Plugins/Experimental/MeshModelingToolset/Source/MeshModelingTools/Public/Physics/ExtractCollisionGeometryTool.h b/Engine/Plugins/Experimental/MeshModelingToolset/Source/MeshModelingTools/Public/Physics/ExtractCollisionGeometryTool.h
--- a/Engine/Plugins/Experimental/MeshModelingToolset/Source/MeshModelingTools/Public/Physics/ExtractCollisionGeometryTool.h
+++ b/Engine/Plugins/Experimental/MeshModelingToolset/Source/MeshModelingTools/Public/Physics/ExtractCollisionGeometryTool.h
@@ -11,6 +11,7 @@
 
 class UPreviewGeometry;
 class UPreviewMesh;
+class IAssetGenerationAPI;
 
 UCLASS()
 class MESHMODELINGTOOLS_API UExtractCollisionGeometryToolBuilder : public UInteractiveToolBuilder
@@ -18,7 +19,7 @@ class MESHMODELINGTOOLS_API UExtractCollisionGeometryToolBuilder : public UInter
 	GENERATED_BODY()
 
 public:
-	IToolsContextAssetAPI* AssetAPI = nullptr;
+	IAssetGenerationAPI* AssetAPI = nullptr;
 
 	virtual bool CanBuildTool(const FToolBuilderState& SceneState) const override;
 	virtual UInteractiveTool* BuildTool(const FToolBuilderState& SceneState) const override;
@@ -36,7 +37,7 @@ class MESHMODELINGTOOLS_API UExtractCollisionGeometryTool : public USingleSelect
 	GENERATED_BODY()
 public:
 	virtual void SetWorld(UWorld* World) { this->TargetWorld = World; }
-	virtual void SetAssetAPI(IToolsContextAssetAPI* InAssetAPI) { this->AssetAPI = InAssetAPI; }
+	virtual void SetAssetAPI(IAssetGenerationAPI* InAssetAPI) { this->AssetAPI = InAssetAPI; }
 
 	virtual void Setup() override;
 	virtual void Shutdown(EToolShutdownType ShutdownType) override;
@@ -66,7 +67,7 @@ protected:
 	TSharedPtr<FPhysicsDataCollection> PhysicsInfo;
 
 	UWorld* TargetWorld = nullptr;
-	IToolsContextAssetAPI* AssetAPI = nullptr;
+	IAssetGenerationAPI* AssetAPI = nullptr;
 
 	FDynamicMesh3 CurrentMesh;
 	bool bResultValid = false;
