Make destructor for IAssetGenerationAPI viltual

#rnx

diff --git a/Engine/Plugins/Experimental/MeshModelingToolset/Source/ModelingComponents/Public/AssetGenerationUtil.h b/Engine/Plugins/Experimental/MeshModelingToolset/Source/ModelingComponents/Public/AssetGenerationUtil.h
--- a/Engine/Plugins/Experimental/MeshModelingToolset/Source/ModelingComponents/Public/AssetGenerationUtil.h
+++ b/Engine/Plugins/Experimental/MeshModelingToolset/Source/ModelingComponents/Public/AssetGenerationUtil.h
@@ -16,6 +16,8 @@ class UTexture2D;
 MODELINGCOMPONENTS_API class IAssetGenerationAPI
 {
 public:
+	virtual ~IAssetGenerationAPI(){}
+
 	/**
 	 * Create a new UStaticMeshAsset for the provided mesh, then a new UStaticMeshComponent and AStaticMeshActor in the TargetWorld
 	 * @param TargetWorld world to create new Actor in
