Simple free-form deformation (FFD) lattice tool. Create a hexahedral grid (lattice) around the mesh, allow the user to move the lattice, and deform the mesh accordingly. Cubic and linear interpolation from lattice control points to mesh vertices are both supported, as is user-defined lattice resolution. Control points can be selected individually or with a rectangle selection.

diff --git a/Engine/Plugins/Experimental/GeometryProcessing/Source/DynamicMesh/Private/Operations/FFDLattice.cpp b/Engine/Plugins/Experimental/GeometryProcessing/Source/DynamicMesh/Private/Operations/FFDLattice.cpp
new file mode 100644
--- /dev/null
+++ b/Engine/Plugins/Experimental/GeometryProcessing/Source/DynamicMesh/Private/Operations/FFDLattice.cpp
@@ -0,0 +1,329 @@
+// Copyright Epic Games, Inc. All Rights Reserved.
+
+#include "Operations/FFDLattice.h"
+#include "DynamicMesh3.h"
+#include "Util/ProgressCancel.h"
+
+FFFDLattice::FFFDLattice(const FVector3i& InDims, const FDynamicMesh3& Mesh) :
+	Dims(InDims)
+{
+	InitialBounds = Mesh.GetBounds();
+
+	// Expand the initial bounding box to make the computation of which grid cell a mesh vertex is inside of a little 
+	// less susceptible to numerical error issues.
+	// TODO: derive a better padding value than this magic number, or allow user to set it (UETOOL-2362)
+	constexpr float Padding = 10.0f;
+	InitialBounds.Min -= FVector3d{ Padding, Padding, Padding };
+	InitialBounds.Max += FVector3d{ Padding, Padding, Padding };
+	
+	ComputeInitialEmbedding(Mesh);
+}
+
+void FFFDLattice::GenerateInitialLatticePositions(TArray<FVector3d>& OutLatticePositions) const
+{
+	double DX = (InitialBounds.Max.X - InitialBounds.Min.X) / (Dims.X - 1);
+	double DY = (InitialBounds.Max.Y - InitialBounds.Min.Y) / (Dims.Y - 1);
+	double DZ = (InitialBounds.Max.Z - InitialBounds.Min.Z) / (Dims.Z - 1);
+
+	int TotalNumLatticePoints = Dims.X * Dims.Y * Dims.Z;
+
+	OutLatticePositions.SetNum(TotalNumLatticePoints);
+
+	for (int i = 0; i < Dims.X; ++i)
+	{
+		double X = DX * i;
+		for (int j = 0; j < Dims.Y; ++j)
+		{
+			double Y = DY * j;
+			for (int k = 0; k < Dims.Z; ++k)
+			{
+				int PointID = ControlPointIndexFromCoordinates(i, j, k);
+				double Z = DZ * k;
+
+				OutLatticePositions[PointID] = InitialBounds.Min + FVector3d{ X,Y,Z };
+			}
+		}
+	}
+}
+
+
+void FFFDLattice::GenerateLatticeEdges(TArray<FVector2i>& OutLatticeEdges) const
+{
+	OutLatticeEdges.Reset(3 * Dims.X * Dims.Y * Dims.Z);
+
+	for (int i = 0; i < Dims.X; ++i)
+	{
+		for (int j = 0; j < Dims.Y; ++j)
+		{
+			for (int k = 0; k < Dims.Z; ++k)
+			{
+				int PointID = ControlPointIndexFromCoordinates(i, j, k);
+
+				if (i + 1 < Dims.X)
+				{
+					int IPlusOne = ControlPointIndexFromCoordinates(i + 1, j, k);
+					OutLatticeEdges.Add({ PointID, IPlusOne });
+				}
+				if (j + 1 < Dims.Y)
+				{
+					int JPlusOne = ControlPointIndexFromCoordinates(i, j + 1, k);
+					OutLatticeEdges.Add({ PointID, JPlusOne });
+				}
+				if (k + 1 < Dims.Z)
+				{
+					int KPlusOne = ControlPointIndexFromCoordinates(i, j, k + 1);
+					OutLatticeEdges.Add({ PointID, KPlusOne });
+				}
+			}
+		}
+	}
+}
+
+
+FVector3d FFFDLattice::ComputeTrilinearWeights(const FVector3d& QueryPoint, FVector3i& GridCoordinates) const
+{
+	double DX = (InitialBounds.Max.X - InitialBounds.Min.X) / (Dims.X - 1);
+	double DY = (InitialBounds.Max.Y - InitialBounds.Min.Y) / (Dims.Y - 1);
+	double DZ = (InitialBounds.Max.Z - InitialBounds.Min.Z) / (Dims.Z - 1);
+
+	FVector3d GridPoint{
+		((QueryPoint.X - InitialBounds.Min.X) / DX),
+		((QueryPoint.Y - InitialBounds.Min.Y) / DY),
+		((QueryPoint.Z - InitialBounds.Min.Z) / DZ) };
+
+	// compute integer coordinates
+	int X0 = (int)GridPoint.X;
+	int Y0 = (int)GridPoint.Y;
+	int Z0 = (int)GridPoint.Z;
+
+	GridCoordinates = FVector3i{ X0, Y0, Z0 };
+
+	FVector3d Weights = FVector3d{
+		GridPoint.X - (double)X0,
+		GridPoint.Y - (double)Y0,
+		GridPoint.Z - (double)Z0 };
+
+	return Weights;
+}
+
+void FFFDLattice::ComputeInitialEmbedding(const FDynamicMesh3& Mesh)
+{
+	VertexEmbeddings.Reset(Mesh.VertexCount());
+	for (const FVector3d& VertexPosition : Mesh.VerticesItr())
+	{
+		FEmbedding Embedding;
+		Embedding.CellWeighting = ComputeTrilinearWeights(VertexPosition, Embedding.LatticeCell);
+		VertexEmbeddings.Add(Embedding);
+	}
+
+	check(VertexEmbeddings.Num() == Mesh.VertexCount());
+}
+
+
+void FFFDLattice::GetDeformedMeshVertexPositions(const TArray<FVector3d>& LatticeControlPoints, 
+												 TArray<FVector3d>& OutVertexPositions, 
+												 bool bCubic, 
+												 FProgressCancel* Progress) const
+{
+	int NumVertices = VertexEmbeddings.Num();
+	OutVertexPositions.SetNumZeroed(NumVertices);
+
+	for (int VertexID = 0; VertexID < NumVertices; ++VertexID)
+	{
+		if (bCubic)
+		{
+			OutVertexPositions[VertexID] = InterpolatedPositionCubic(VertexEmbeddings[VertexID], LatticeControlPoints);
+		}
+		else
+		{
+			OutVertexPositions[VertexID] = InterpolatedPosition(VertexEmbeddings[VertexID], LatticeControlPoints);
+		}
+
+		// Every once in a while, check for cancellation
+		if (Progress && (VertexID % 1000 == 0) && Progress->Cancelled())
+		{
+			return;
+		}
+	}
+}
+
+
+void FFFDLattice::GetValuePair(int I, int J, int K, FVector3d& A, FVector3d& B, 
+							   const TArray<FVector3d>& LatticeControlPoints) const
+{
+	int IndexA = ControlPointIndexFromCoordinates(I, J, K);
+	check(IndexA < LatticeControlPoints.Num());
+	check(IndexA >= 0);
+	A = LatticeControlPoints[IndexA];
+
+	int IndexB = ControlPointIndexFromCoordinates(I + 1, J, K);
+	B = LatticeControlPoints[IndexB];
+}
+
+
+static float CubicBSplineKernel(float A)
+{
+	// Using cubic kernel f(a) =
+	//	(4 - 6a^2 + 3|a|^3) / 6		for 0 <= |a| < 1
+	//	(2 - |a|)^3 / 6				for 1 <= |a| < 2
+	//	0							otherwise
+	//
+	// So at a = {-2, -1, 0, 1, 2}, f(a) = {0, 1/6, 4/6, 1/6, 0}, and is piecewise cubic in between.
+
+	float AbsA = FMath::Abs(A);
+	if (AbsA < 1.0f)
+	{
+		float ASquared = AbsA * AbsA;
+		return (4.0f - 6.0f * ASquared + 3.0f * ASquared * AbsA) / 6.0f;
+	}
+	else if (AbsA < 2.0f)
+	{
+		float TwoMinusAbsA = (2.0f - AbsA);
+		return TwoMinusAbsA * TwoMinusAbsA * TwoMinusAbsA / 6.0f;
+	}
+
+	return 0.0f;
+}
+
+
+FVector3d FFFDLattice::InterpolatedPositionCubic(const FEmbedding& VertexEmbedding, const TArray<FVector3d>& LatticeControlPoints) const
+{
+	int X0 = VertexEmbedding.LatticeCell.X;
+	int Y0 = VertexEmbedding.LatticeCell.Y;
+	int Z0 = VertexEmbedding.LatticeCell.Z;
+
+	double T = VertexEmbedding.CellWeighting.X;
+	double U = VertexEmbedding.CellWeighting.Y;
+	double V = VertexEmbedding.CellWeighting.Z;
+
+	double DX = (InitialBounds.Max.X - InitialBounds.Min.X) / (Dims.X - 1);
+	double DY = (InitialBounds.Max.Y - InitialBounds.Min.Y) / (Dims.Y - 1);
+	double DZ = (InitialBounds.Max.Z - InitialBounds.Min.Z) / (Dims.Z - 1);
+
+	FVector3d Sum{ 0.0f, 0.0f, 0.0f };
+
+	// TODO: This can probably be replaced with some relatively simple linear algebra
+
+	for (int DI = -1; DI <= 2; ++DI)
+	{
+		double WeightX = CubicBSplineKernel(T - DI);
+
+		for (int DJ = -1; DJ <= 2; ++DJ)
+		{
+			double WeightY = CubicBSplineKernel(U - DJ);
+
+			for (int DK = -1; DK <= 2; ++DK)
+			{
+				double WeightZ = CubicBSplineKernel(V - DK);
+				double Weight = WeightX * WeightY * WeightZ;
+
+				int i = X0 + DI;
+				int j = Y0 + DJ;
+				int k = Z0 + DK;
+
+				FVector3d LatticePoint;
+				if (i < 0 || i >= Dims.X || j < 0 || j >= Dims.Y || k < 0 || k >= Dims.Z)
+				{
+					// Get the extrapolated position for a "virtual" control point outside of the deformed lattice
+					LatticePoint = ExtrapolatedLatticePosition({ i,j,k }, LatticeControlPoints);
+				}
+				else
+				{
+					int PointIndex = ControlPointIndexFromCoordinates(i, j, k);
+					LatticePoint = LatticeControlPoints[PointIndex];
+				}
+
+				Sum += Weight * LatticePoint;
+			}
+		}
+	}
+
+	return Sum;
+}
+
+FVector3d FFFDLattice::InterpolatedPosition(const FEmbedding& VertexEmbedding, const TArray<FVector3d>& LatticeControlPoints) const
+{
+	// TODO: See if we can refactor TTriLinearGridInterpolant to make that usable in this class
+
+	// Trilinear interpolation:
+	//		V### is grid cell corner index
+	//		AlphaN is [0,1] fraction of point in cell along N'th dimension  
+	// return
+	//		V000 * (1 - AlphaX) * (1 - AlphaY) * (1 - AlphaZ) +
+	//		V001 * (1 - AlphaX) * (1 - AlphaY) * (AlphaZ) +
+	//		V010 * (1 - AlphaX) * (AlphaY) * (1 - AlphaZ) +
+	//		V011 * (1 - AlphaX) * (AlphaY) * (AlphaZ) +
+	//		V100 * (AlphaX) * (1 - AlphaY) * (1 - AlphaZ) +
+	//		V101 * (AlphaX) * (1 - AlphaY) * (AlphaZ) +
+	//		V110 * (AlphaX) * (AlphaY) * (1 - AlphaZ) +
+	//		V111 * (AlphaX) * (AlphaY) * (AlphaZ);
+
+	int X0 = VertexEmbedding.LatticeCell.X;
+	int Y0 = VertexEmbedding.LatticeCell.Y;
+	int Y1 = Y0 + 1;
+	int Z0 = VertexEmbedding.LatticeCell.Z;
+	int Z1 = Z0 + 1;
+
+	double AlphaX = VertexEmbedding.CellWeighting.X;
+	double AlphaY = VertexEmbedding.CellWeighting.Y;
+	double AlphaZ = VertexEmbedding.CellWeighting.Z;
+	double OneMinusAlphaX = 1.0 - AlphaX;
+
+	FVector3d Sum{ 0,0,0 };
+
+	FVector3d FV000, FV100;
+	GetValuePair(X0, Y0, Z0, FV000, FV100, LatticeControlPoints);
+	double YZ = (1 - AlphaY) * (1 - AlphaZ);
+	Sum = (OneMinusAlphaX * FV000 + AlphaX * FV100) * YZ;
+
+	FVector3d FV001, FV101;
+	GetValuePair(X0, Y0, Z1, FV001, FV101, LatticeControlPoints);
+	YZ = (1 - AlphaY) * (AlphaZ);
+	Sum += (OneMinusAlphaX * FV001 + AlphaX * FV101) * YZ;
+
+	FVector3d FV010, FV110;
+	GetValuePair(X0, Y1, Z0, FV010, FV110, LatticeControlPoints);
+	YZ = (AlphaY) * (1 - AlphaZ);
+	Sum += (OneMinusAlphaX * FV010 + AlphaX * FV110) * YZ;
+
+	FVector3d FV011, FV111;
+	GetValuePair(X0, Y1, Z1, FV011, FV111, LatticeControlPoints);
+	YZ = (AlphaY) * (AlphaZ);
+	Sum += (OneMinusAlphaX * FV011 + AlphaX * FV111) * YZ;
+
+	return Sum;
+}
+
+FVector3d FFFDLattice::ClosestLatticePosition(const FVector3i& VirtualControlPointIndex, const TArray<FVector3d>& LatticeControlPoints) const
+{
+	FVector3i NearestControlPointIndex;
+	NearestControlPointIndex.X = FMath::Clamp(VirtualControlPointIndex.X, 0, Dims.X - 1);
+	NearestControlPointIndex.Y = FMath::Clamp(VirtualControlPointIndex.Y, 0, Dims.Y - 1);
+	NearestControlPointIndex.Z = FMath::Clamp(VirtualControlPointIndex.Z, 0, Dims.Z - 1);
+	return LatticeControlPoints[ControlPointIndexFromCoordinates(NearestControlPointIndex)];
+}
+
+FVector3d FFFDLattice::ExtrapolatedLatticePosition(const FVector3i& VirtualControlPointIndex, const TArray<FVector3d>& LatticeControlPoints) const
+{
+	// Use the location of the nearest control point and the location of a control point in the opposite direction of 
+	// the extrapolation to get the extrapolated location.
+
+	FVector3i NearestControlPointIndex;
+	NearestControlPointIndex.X = FMath::Clamp(VirtualControlPointIndex.X, 0, Dims.X - 1);
+	NearestControlPointIndex.Y = FMath::Clamp(VirtualControlPointIndex.Y, 0, Dims.Y - 1);
+	NearestControlPointIndex.Z = FMath::Clamp(VirtualControlPointIndex.Z, 0, Dims.Z - 1);
+
+	FVector3i Delta = VirtualControlPointIndex - NearestControlPointIndex;
+	check(Delta != FVector3i::Zero());
+
+	FVector3i TraceBackControlPointIndex = NearestControlPointIndex - Delta;
+	check(TraceBackControlPointIndex.X >= 0 && TraceBackControlPointIndex.X < Dims.X);
+	check(TraceBackControlPointIndex.Y >= 0 && TraceBackControlPointIndex.Y < Dims.Y);
+	check(TraceBackControlPointIndex.Z >= 0 && TraceBackControlPointIndex.Z < Dims.Z);
+
+	const FVector3d& A = LatticeControlPoints[ControlPointIndexFromCoordinates(TraceBackControlPointIndex)];
+	const FVector3d& B = LatticeControlPoints[ControlPointIndexFromCoordinates(NearestControlPointIndex)];
+
+	FVector3d Position = B + (B - A);
+	return Position;
+}
diff --git a/Engine/Plugins/Experimental/GeometryProcessing/Source/DynamicMesh/Public/Operations/FFDLattice.h b/Engine/Plugins/Experimental/GeometryProcessing/Source/DynamicMesh/Public/Operations/FFDLattice.h
new file mode 100644
--- /dev/null
+++ b/Engine/Plugins/Experimental/GeometryProcessing/Source/DynamicMesh/Public/Operations/FFDLattice.h
@@ -0,0 +1,87 @@
+// Copyright Epic Games, Inc. All Rights Reserved.
+
+#pragma once
+
+#include "VectorTypes.h"
+#include "BoxTypes.h"
+#include "Containers/StaticArray.h"
+
+class FDynamicMesh3;
+class FProgressCancel;
+
+/** 
+ * Free-form deformation lattice. Initialize it with a mesh and desired resolution, get the initial lattice points out.
+ * Then pass in deformed lattice points and it will compute deformed mesh vertex positions.
+ */
+class DYNAMICMESH_API FFFDLattice
+{
+public:
+
+	/// Linear interpolation information. There should be one of these per mesh vertex.
+	struct FEmbedding
+	{
+		FVector3i LatticeCell;		// Original point is in the hexahedron whose Min corner is LatticeCell
+		FVector3d CellWeighting;	// Original linear interpolation weights of the point in the cell
+	};
+
+	/// Create a lattice that fits the given mesh and has the given resolution along each dimension. Also precompute
+	/// the weight of each mesh vertex inside the appropriate lattice cell.
+	FFFDLattice(const FVector3i& InDims, const FDynamicMesh3& InMesh);
+
+	/// Using the Lattice's current control point positions and the original embedding information, compute the
+	/// deformed mesh vertex positions.
+	void GetDeformedMeshVertexPositions(const TArray<FVector3d>& LatticeControlPoints,
+										TArray<FVector3d>& OutVertexPositions, 
+										bool bCubic = false, 
+										FProgressCancel* Progress = nullptr) const;
+
+	void GenerateInitialLatticePositions(TArray<FVector3d>& OutLatticePositions) const;
+	void GenerateLatticeEdges(TArray<FVector2i>& OutLatticeEdges) const;
+
+protected:
+
+	// Number of lattice points in each dimension
+	FVector3i Dims;
+
+	// Extents of the lattice before it is deformed
+	FAxisAlignedBox3d InitialBounds;
+
+	/// Interpolation weights and cell indices per vertex of the input FDynamicMesh3
+	TArray<FEmbedding> VertexEmbeddings;
+
+
+	void InitializeLatticePositions();
+
+	/// Get the index into the flat TArray of positions given the (i,j,k) coordinates in the lattice
+	int ControlPointIndexFromCoordinates(int i, int j, int k) const
+	{
+		int Idx = k + Dims.Z * (j + Dims.Y * i);
+		return Idx;
+	}
+
+	int ControlPointIndexFromCoordinates(const FVector3i& Index) const
+	{
+		return ControlPointIndexFromCoordinates(Index.X, Index.Y, Index.Z);
+	}
+
+	/// For each vertex in Mesh, compute the lattice cell it resides in and its weighting
+	void ComputeInitialEmbedding(const FDynamicMesh3& Mesh);
+
+	/// Compute cell index and linear interpolation weights for a given point
+	FVector3d ComputeTrilinearWeights(const FVector3d& Pt, FVector3i& GridCoordinates) const;
+
+	/// Helper for linear interpolation. Returns lattice points at (I,J,K) and (I+1,J,K).
+	void GetValuePair(int I, int J, int K, FVector3d& A, FVector3d& B, const TArray<FVector3d>& LatticeControlPoints) const;
+
+	/// Use the given cell index and weights to compute the interpolated position
+	FVector3d InterpolatedPosition(const FEmbedding& VertexEmbedding, const TArray<FVector3d>& LatticeControlPoints) const;
+	FVector3d InterpolatedPositionCubic(const FEmbedding& VertexEmbedding, const TArray<FVector3d>& LatticeControlPoints) const;
+
+	/// Clamp the given index to [0, Dims] and return the current lattice point position at the clamped index
+	FVector3d ClosestLatticePosition(const FVector3i& VirtualControlPointIndex, const TArray<FVector3d>& LatticeControlPoints) const;
+
+	/// Compute the extrapolated position for a "virtual" lattice control point outside of the actual lattice
+	FVector3d ExtrapolatedLatticePosition(const FVector3i& VirtualPointIndex, const TArray<FVector3d>& LatticeControlPoints) const;
+
+};
+
diff --git a/Engine/Plugins/Experimental/MeshModelingToolset/Source/MeshModelingTools/Private/LatticeDeformerTool.cpp b/Engine/Plugins/Experimental/MeshModelingToolset/Source/MeshModelingTools/Private/LatticeDeformerTool.cpp
new file mode 100644
--- /dev/null
+++ b/Engine/Plugins/Experimental/MeshModelingToolset/Source/MeshModelingTools/Private/LatticeDeformerTool.cpp
@@ -0,0 +1,230 @@
+// Copyright Epic Games, Inc. All Rights Reserved.
+
+#include "LatticeDeformerTool.h"
+
+#include "Mechanics/LatticeControlPointsMechanic.h"
+#include "MeshDescriptionToDynamicMesh.h"
+#include "BaseGizmos/GizmoRenderingUtil.h"
+#include "BaseBehaviors/SingleClickBehavior.h"
+#include "BaseBehaviors/MouseHoverBehavior.h"
+#include "DeformationOps/LatticeDeformerOp.h"
+#include "Properties/MeshMaterialProperties.h"
+#include "Selection/ToolSelectionUtil.h"
+#include "MeshOpPreviewHelpers.h" //FDynamicMeshOpResult
+#include "ToolSceneQueriesUtil.h"
+#include "ToolBuilderUtil.h"
+#include "DynamicMeshToMeshDescription.h"
+#include "MeshTransforms.h"
+#include "Algo/ForEach.h"
+#include "CanvasItem.h"
+#include "CanvasTypes.h"
+#include "Operations/FFDLattice.h"
+
+#define LOCTEXT_NAMESPACE "ULatticeDeformerTool"
+
+
+// Tool builder
+
+bool ULatticeDeformerToolBuilder::CanBuildTool(const FToolBuilderState& SceneState) const
+{
+	// TODO: Do we want the ability to handle multiple meshes with the same lattice?
+	return ToolBuilderUtil::CountComponents(SceneState, CanMakeComponentTarget) == 1;
+}
+
+UInteractiveTool* ULatticeDeformerToolBuilder::BuildTool(const FToolBuilderState& SceneState) const
+{
+	ULatticeDeformerTool* NewTool = NewObject<ULatticeDeformerTool>(SceneState.ToolManager);
+
+	UActorComponent* ActorComponent = ToolBuilderUtil::FindFirstComponent(SceneState, CanMakeComponentTarget);
+	auto* MeshComponent = Cast<UPrimitiveComponent>(ActorComponent);
+	check(MeshComponent != nullptr);
+	NewTool->SetSelection(MakeComponentTarget(MeshComponent));
+	NewTool->SetWorld(SceneState.World);
+	return NewTool;
+}
+
+
+// Operator factory
+
+TUniquePtr<FDynamicMeshOperator> ULatticeDeformerOperatorFactory::MakeNewOperator()
+{
+	bool bUseCubicInterpolation =
+		(LatticeDeformerTool->Settings->InterpolationType == ELatticeInterpolationType::Cubic);
+
+	TUniquePtr<FLatticeDeformerOp> LatticeDeformOp = MakeUnique<FLatticeDeformerOp>(
+		LatticeDeformerTool->OriginalMesh,
+		LatticeDeformerTool->Lattice,
+		LatticeDeformerTool->ControlPointsMechanic->GetControlPoints(),
+		bUseCubicInterpolation);
+
+	return LatticeDeformOp;
+}
+
+
+// Tool itself
+
+FVector3i ULatticeDeformerTool::GetLatticeResolution() const
+{
+	return FVector3i{ Settings->XAxisResolution, Settings->YAxisResolution, Settings->ZAxisResolution };
+}
+
+void ULatticeDeformerTool::DrawHUD(FCanvas* Canvas)
+{
+	// Draw the drag rectangle if it's active
+	if (ControlPointsMechanic->bIsDragging)
+	{		
+		FVector2D Start = ControlPointsMechanic->DragStartScreenPosition;
+		FVector2D Curr = ControlPointsMechanic->DragCurrentScreenPosition;
+
+		FCanvasBoxItem BoxItem(Start / Canvas->GetDPIScale(), (Curr - Start) / Canvas->GetDPIScale());
+		BoxItem.SetColor(FLinearColor::White);
+		Canvas->DrawItem(BoxItem);
+	}
+}
+
+bool ULatticeDeformerTool::CanAccept() const
+{
+	return Preview != nullptr && Preview->HaveValidResult();
+}
+
+void ULatticeDeformerTool::InitializeLattice(TArray<FVector3d>& OutLatticePoints, TArray<FVector2i>& OutLatticeEdges)
+{
+	Lattice = MakeShared<FFFDLattice>(GetLatticeResolution(), *OriginalMesh);
+
+	Lattice->GenerateInitialLatticePositions(OutLatticePoints);
+
+	// Put the lattice in world space
+	FTransform3d LocalToWorld = FTransform3d(ComponentTarget->GetWorldTransform());
+	Algo::ForEach(OutLatticePoints, [&LocalToWorld](FVector3d& Point) {
+		Point = LocalToWorld.TransformPosition(Point);
+	});
+
+	Lattice->GenerateLatticeEdges(OutLatticeEdges);
+}
+
+void ULatticeDeformerTool::Setup()
+{
+	UInteractiveTool::Setup();
+
+	GetToolManager()->DisplayMessage(LOCTEXT("LatticeDeformerToolMessage", 
+		"Drag the lattice control points to deform the mesh"), EToolMessageLevel::UserNotification);
+
+	OriginalMesh = MakeShared<FDynamicMesh3>();
+	FMeshDescriptionToDynamicMesh Converter;
+	Converter.Convert(ComponentTarget->GetMesh(), *OriginalMesh);
+
+	Settings = NewObject<ULatticeDeformerToolProperties>(this, TEXT("Lattice Deformer Tool Settings"));
+	Settings->RestoreProperties(this);
+	AddToolPropertySource(Settings);
+
+	TArray<FVector3d> LatticePoints;
+	TArray<FVector2i> LatticeEdges;
+	InitializeLattice(LatticePoints, LatticeEdges);
+
+	// Set up control points mechanic
+	ControlPointsMechanic = NewObject<ULatticeControlPointsMechanic>(this);
+	ControlPointsMechanic->Setup(this);
+	ControlPointsMechanic->SetWorld(TargetWorld);
+	ControlPointsMechanic->Initialize(LatticePoints, LatticeEdges);
+
+	auto OnPointsChangedLambda = [this]()
+	{
+		Preview->InvalidateResult();
+		Settings->bCanChangeResolution = !ControlPointsMechanic->bHasChanged;
+	};
+
+	ControlPointsMechanic->OnPointsChanged.AddLambda(OnPointsChangedLambda);
+
+	StartPreview();
+}
+
+void ULatticeDeformerTool::Shutdown(EToolShutdownType ShutdownType)
+{
+	Settings->SaveProperties(this);
+	ControlPointsMechanic->Shutdown();
+
+	if (Preview)
+	{
+		FDynamicMeshOpResult Result = Preview->Shutdown();
+
+		if (ShutdownType == EToolShutdownType::Accept)
+		{
+			GetToolManager()->BeginUndoTransaction(LOCTEXT("LatticeDeformerTool", "Lattice Deformer"));
+
+			FDynamicMesh3* DynamicMeshResult = Result.Mesh.Get();
+			check(DynamicMeshResult != nullptr);
+
+			// The lattice and its output mesh are in world space, so get them in local space.
+			// TODO: Would it make more sense to do all the lattice computation in local space?
+			FTransform3d LocalToWorld = FTransform3d(ComponentTarget->GetWorldTransform());
+			MeshTransforms::ApplyTransformInverse(*DynamicMeshResult, LocalToWorld);
+
+			ComponentTarget->CommitMesh([DynamicMeshResult](const FPrimitiveComponentTarget::FCommitParams& CommitParams)
+			{
+				FDynamicMeshToMeshDescription Converter;
+				Converter.Convert(DynamicMeshResult, *CommitParams.MeshDescription);
+			});
+
+			GetToolManager()->EndUndoTransaction();
+		}
+	}
+
+	ComponentTarget->SetOwnerVisibility(true);
+}
+
+
+void ULatticeDeformerTool::StartPreview()
+{
+	ULatticeDeformerOperatorFactory* LatticeDeformOpCreator = NewObject<ULatticeDeformerOperatorFactory>();
+	LatticeDeformOpCreator->LatticeDeformerTool = this;
+
+	Preview = NewObject<UMeshOpPreviewWithBackgroundCompute>(LatticeDeformOpCreator);
+
+	Preview->Setup(TargetWorld, LatticeDeformOpCreator);
+	Preview->PreviewMesh->SetTangentsMode(EDynamicMeshTangentCalcType::NoTangents);
+	Preview->SetVisibility(true);
+	Preview->InvalidateResult();
+
+	ComponentTarget->SetOwnerVisibility(false);
+}
+
+void ULatticeDeformerTool::OnPropertyModified(UObject* PropertySet, FProperty* Property)
+{
+	if (Preview)
+	{
+		bool bShouldRebuild = Property &&
+			((Property->GetFName() == GET_MEMBER_NAME_CHECKED(ULatticeDeformerToolProperties, XAxisResolution)) ||
+			 (Property->GetFName() == GET_MEMBER_NAME_CHECKED(ULatticeDeformerToolProperties, YAxisResolution)) ||
+			 (Property->GetFName() == GET_MEMBER_NAME_CHECKED(ULatticeDeformerToolProperties, ZAxisResolution)));
+
+		if (bShouldRebuild)
+		{
+			TArray<FVector3d> LatticePoints;
+			TArray<FVector2i> LatticeEdges;
+			InitializeLattice(LatticePoints, LatticeEdges);
+			ControlPointsMechanic->Initialize(LatticePoints, LatticeEdges);
+		}
+
+		Preview->InvalidateResult();
+	}
+}
+
+
+void ULatticeDeformerTool::OnTick(float DeltaTime)
+{
+	if (Preview)
+	{
+		Preview->Tick(DeltaTime);
+	}
+}
+
+
+void ULatticeDeformerTool::Render(IToolsContextRenderAPI* RenderAPI)
+{
+	if (ControlPointsMechanic != nullptr)
+	{
+		ControlPointsMechanic->Render(RenderAPI);
+	}
+}
+
+#undef LOCTEXT_NAMESPACE
diff --git a/Engine/Plugins/Experimental/MeshModelingToolset/Source/MeshModelingTools/Public/LatticeDeformerTool.h b/Engine/Plugins/Experimental/MeshModelingToolset/Source/MeshModelingTools/Public/LatticeDeformerTool.h
new file mode 100644
--- /dev/null
+++ b/Engine/Plugins/Experimental/MeshModelingToolset/Source/MeshModelingTools/Public/LatticeDeformerTool.h
@@ -0,0 +1,137 @@
+// Copyright Epic Games, Inc. All Rights Reserved.
+
+#pragma once
+
+#include "ModelingOperators.h" //IDynamicMeshOperatorFactory
+#include "InteractiveTool.h" //UInteractiveToolPropertySet
+#include "InteractiveToolBuilder.h" //UInteractiveToolBuilder
+#include "InteractiveToolChange.h" //FToolCommandChange
+#include "SingleSelectionTool.h" //USingleSelectionTool
+#include "Operations/FFDLattice.h"
+#include "LatticeDeformerTool.generated.h"
+
+class ULatticeControlPointsMechanic;
+class UMeshOpPreviewWithBackgroundCompute;
+
+UCLASS()
+class MESHMODELINGTOOLS_API ULatticeDeformerToolBuilder : public UInteractiveToolBuilder
+{
+	GENERATED_BODY()
+
+public:
+	virtual bool CanBuildTool(const FToolBuilderState& SceneState) const override;
+	virtual UInteractiveTool* BuildTool(const FToolBuilderState& SceneState) const override;
+};
+
+
+UENUM()
+enum class ELatticeInterpolationType : uint8
+{
+	/** Use trilinear interpolation to get new mesh vertex positions from the lattice */
+	Linear UMETA(DisplayName = "Linear"),
+
+	/** Use tricubic interpolation to get new mesh vertex positions from the lattice */
+	Cubic UMETA(DisplayName = "Cubic")
+};
+
+
+UCLASS()
+class MESHMODELINGTOOLS_API ULatticeDeformerToolProperties : public UInteractiveToolPropertySet
+{
+	GENERATED_BODY()
+
+public:
+
+	/** Number of lattice vertices along the X axis */
+	UPROPERTY(EditAnywhere, Category = Resolution, meta = (UIMin = "2", ClampMin = "2", EditCondition = "bCanChangeResolution", HideEditConditionToggle))
+	int XAxisResolution = 5;
+
+	/** Number of lattice vertices along the Y axis */
+	UPROPERTY(EditAnywhere, Category = Resolution, meta = (UIMin = "2", ClampMin = "2", EditCondition = "bCanChangeResolution", HideEditConditionToggle))
+	int YAxisResolution = 5;
+
+	/** Number of lattice vertices along the Z axis */
+	UPROPERTY(EditAnywhere, Category = Resolution, meta = (UIMin = "2", ClampMin = "2", EditCondition = "bCanChangeResolution", HideEditConditionToggle))
+	int ZAxisResolution = 5;
+
+	/** Whether to use linear or cubic interpolation to get new mesh vertex positions from the lattice */
+	UPROPERTY(EditAnywhere, Category = Interpolation )
+	ELatticeInterpolationType InterpolationType = ELatticeInterpolationType::Linear;
+
+	// Not user visible - used to disallow changing the lattice resolution after deformation
+	UPROPERTY(meta = (TransientToolProperty))
+	bool bCanChangeResolution = true;
+
+};
+
+
+UCLASS()
+class MESHMODELINGTOOLS_API ULatticeDeformerOperatorFactory : public UObject, public IDynamicMeshOperatorFactory
+{
+	GENERATED_BODY()
+
+public:
+	// IDynamicMeshOperatorFactory API
+	virtual TUniquePtr<FDynamicMeshOperator> MakeNewOperator() override;
+
+	UPROPERTY()
+	ULatticeDeformerTool* LatticeDeformerTool;
+};
+
+
+/** Deform a mesh using a regular hexahedral lattice */
+UCLASS()
+class MESHMODELINGTOOLS_API ULatticeDeformerTool : public USingleSelectionTool
+{
+	GENERATED_BODY()
+
+public:
+
+	virtual void DrawHUD(FCanvas* Canvas) override;
+
+	virtual void SetWorld(UWorld* World) { TargetWorld = World; }
+	virtual void SetAssetAPI(IToolsContextAssetAPI* NewAssetApi) { AssetAPI = NewAssetApi; }
+
+	virtual bool HasCancel() const override { return true; }
+	virtual bool HasAccept() const override { return true; }
+	virtual bool CanAccept() const override;
+
+	virtual void Setup() override;
+	virtual void Shutdown(EToolShutdownType ShutdownType) override;
+
+	virtual void OnTick(float DeltaTime) override;
+	virtual void Render(IToolsContextRenderAPI* RenderAPI) override;
+
+	virtual void OnPropertyModified(UObject* PropertySet, FProperty* Property) override;
+
+	FVector3i GetLatticeResolution() const;
+
+protected:
+
+	UWorld* TargetWorld;
+	IToolsContextAssetAPI* AssetAPI;
+
+	// Input mesh
+	TSharedPtr<FDynamicMesh3> OriginalMesh;
+
+	TSharedPtr<FFFDLattice> Lattice;
+
+	UPROPERTY()
+	ULatticeControlPointsMechanic* ControlPointsMechanic = nullptr;
+
+	UPROPERTY()
+	ULatticeDeformerToolProperties* Settings = nullptr;
+
+	UPROPERTY()
+	UMeshOpPreviewWithBackgroundCompute* Preview = nullptr;
+
+	UPROPERTY()
+	bool bLatticeDeformed = false;
+
+	// Create and store an FFFDLattice. Pass out the lattice's positions and edges.
+	void InitializeLattice(TArray<FVector3d>& OutLatticePoints, TArray<FVector2i>& OutLatticeEdges);
+
+	void StartPreview();
+
+	friend class ULatticeDeformerOperatorFactory;
+};
diff --git a/Engine/Plugins/Experimental/MeshModelingToolset/Source/ModelingComponents/Private/Mechanics/LatticeControlPointsMechanic.cpp b/Engine/Plugins/Experimental/MeshModelingToolset/Source/ModelingComponents/Private/Mechanics/LatticeControlPointsMechanic.cpp
new file mode 100644
--- /dev/null
+++ b/Engine/Plugins/Experimental/MeshModelingToolset/Source/ModelingComponents/Private/Mechanics/LatticeControlPointsMechanic.cpp
@@ -0,0 +1,811 @@
+// Copyright Epic Games, Inc. All Rights Reserved.
+
+#include "Mechanics/LatticeControlPointsMechanic.h"
+
+#include "BaseBehaviors/SingleClickBehavior.h"
+#include "BaseBehaviors/MouseHoverBehavior.h"
+#include "BaseBehaviors/ClickDragBehavior.h"
+#include "BaseGizmos/TransformGizmo.h"
+#include "BaseGizmos/TransformProxy.h"
+#include "Drawing/PreviewGeometryActor.h"
+#include "Drawing/LineSetComponent.h"
+#include "Drawing/PointSetComponent.h"
+#include "InteractiveToolManager.h"
+#include "InteractiveToolActionSet.h"
+#include "Polyline3.h"
+#include "ToolSceneQueriesUtil.h"
+#include "ToolSetupUtil.h"
+#include "Transforms/MultiTransformer.h"
+
+#define LOCTEXT_NAMESPACE "ULatticeControlPointsMechanic"
+
+static const FText LatticePointDeselectionTransactionText = LOCTEXT("LatticePointDeselection", "Lattice Point Deselection");
+static const FText LatticePointSelectionTransactionText = LOCTEXT("LatticePointSelection", "Lattice Point Selection");
+static const FText LatticePointMovementTransactionText = LOCTEXT("LatticePointMovement", "Lattice Point Movement");
+
+void ULatticeControlPointsMechanic::Setup(UInteractiveTool* ParentToolIn)
+{
+	UInteractionMechanic::Setup(ParentToolIn);
+
+	USingleClickInputBehavior* ClickBehavior = NewObject<USingleClickInputBehavior>();
+	ClickBehavior->Initialize(this);
+	ClickBehavior->Modifiers.RegisterModifier(CtrlModifierId, FInputDeviceState::IsCtrlKeyDown);
+	ClickBehavior->Modifiers.RegisterModifier(ShiftModifierId, FInputDeviceState::IsShiftKeyDown);
+	ParentTool->AddInputBehavior(ClickBehavior);
+
+	UMouseHoverBehavior* HoverBehavior = NewObject<UMouseHoverBehavior>();
+	HoverBehavior->Initialize(this);
+	HoverBehavior->Modifiers.RegisterModifier(CtrlModifierId, FInputDeviceState::IsCtrlKeyDown);
+	HoverBehavior->Modifiers.RegisterModifier(ShiftModifierId, FInputDeviceState::IsShiftKeyDown);
+	ParentTool->AddInputBehavior(HoverBehavior);
+
+	UClickDragInputBehavior* ClickDragBehavior = NewObject<UClickDragInputBehavior>(this);
+	ClickDragBehavior->Initialize(this);
+	// TODO: Add CTRL support for drag selection?
+	ParentTool->AddInputBehavior(ClickDragBehavior);
+
+	DrawnControlPoints = NewObject<UPointSetComponent>();
+	DrawnControlPoints->SetPointMaterial(
+		LoadObject<UMaterial>(nullptr, TEXT("/MeshModelingToolset/Materials/PointSetComponentMaterial")));
+	DrawnLatticeEdges = NewObject<ULineSetComponent>();
+	DrawnLatticeEdges->SetLineMaterial(
+		LoadObject<UMaterial>(nullptr, TEXT("/MeshModelingToolset/Materials/LineSetComponentMaterial")));
+
+	NormalPointColor = FColor::Red;
+	NormalSegmentColor = FColor::Red;
+	HoverColor = FColor::Green;
+	SelectedColor = FColor::Yellow;
+	SegmentsThickness = 1.0f;
+	PointsSize = 6.0f;
+
+	GeometrySetToleranceTest = [this](const FVector3d& Position1, const FVector3d& Position2) {
+		if (CachedCameraState.bIsOrthographic)
+		{
+			// We could just always use ToolSceneQueriesUtil::PointSnapQuery. But in ortho viewports, we happen to know
+			// that the only points that we will ever give this function will be the closest points between a ray and
+			// some geometry, meaning that the vector between them will be orthogonal to the view ray. With this knowledge,
+			// we can do the tolerance computation more efficiently than PointSnapQuery can, since we don't need to project
+			// down to the view plane.
+			// As in PointSnapQuery, we convert our angle-based tolerance to one we can use in an ortho viewport (instead of
+			// dividing our field of view into 90 visual angle degrees, we divide the plane into 90 units).
+			float OrthoTolerance = ToolSceneQueriesUtil::GetDefaultVisualAngleSnapThreshD() * CachedCameraState.OrthoWorldCoordinateWidth / 90.0;
+			return Position1.DistanceSquared(Position2) < OrthoTolerance * OrthoTolerance;
+		}
+		else
+		{
+			return ToolSceneQueriesUtil::PointSnapQuery(CachedCameraState, Position1, Position2);
+		}
+	};
+
+	UInteractiveGizmoManager* GizmoManager = GetParentTool()->GetToolManager()->GetPairedGizmoManager();
+	PointTransformProxy = NewObject<UTransformProxy>(this);
+	
+	// TODO: Maybe don't have the gizmo's axes flip around when it crosses the origin, if possible?
+	// TODO: Enable local vs world gizmo switching (UETOOL-2356)
+	PointTransformGizmo = GizmoManager->CreateCustomTransformGizmo(ETransformGizmoSubElements::FullTranslateRotateScale,
+																   GetParentTool());
+
+	PointTransformProxy->OnTransformChanged.AddUObject(this, &ULatticeControlPointsMechanic::GizmoTransformChanged);
+	PointTransformProxy->OnBeginTransformEdit.AddUObject(this, &ULatticeControlPointsMechanic::GizmoTransformStarted);
+	PointTransformProxy->OnEndTransformEdit.AddUObject(this, &ULatticeControlPointsMechanic::GizmoTransformEnded);
+	PointTransformGizmo->SetActiveTarget(PointTransformProxy);
+	PointTransformGizmo->SetVisibility(false);
+	PointTransformGizmo->bUseContextCoordinateSystem = false;
+}
+
+void ULatticeControlPointsMechanic::Initialize(const TArray<FVector3d>& Points, const TArray<FVector2i>& Edges)
+{
+	ControlPoints = Points;
+	SelectedPointIDs.Empty();
+	LatticeEdges = Edges;
+	UpdateGizmoLocation();
+	RebuildDrawables();
+	++CurrentChangeStamp;		// If the lattice is potentially changing resolution, make this an undo barrier
+	bHasChanged = false;
+}
+
+void ULatticeControlPointsMechanic::SetWorld(UWorld* World)
+{
+	// It may be unreasonable to worry about SetWorld being called more than once, but let's be safe anyway
+	if (PreviewGeometryActor)
+	{
+		PreviewGeometryActor->Destroy();
+	} 
+
+	// We need the world so we can create the geometry actor in the right place
+	FRotator Rotation(0.0f, 0.0f, 0.0f);
+	FActorSpawnParameters SpawnInfo;
+	PreviewGeometryActor = World->SpawnActor<APreviewGeometryActor>(FVector::ZeroVector, Rotation, SpawnInfo);
+
+	// Attach the rendering components to the actor
+	DrawnControlPoints->Rename(nullptr, PreviewGeometryActor);
+	PreviewGeometryActor->SetRootComponent(DrawnControlPoints);
+	if (DrawnControlPoints->IsRegistered())
+	{
+		DrawnControlPoints->ReregisterComponent();
+	}
+	else
+	{
+		DrawnControlPoints->RegisterComponent();
+	}
+
+	DrawnLatticeEdges->Rename(nullptr, PreviewGeometryActor); 
+	DrawnLatticeEdges->AttachToComponent(DrawnControlPoints, FAttachmentTransformRules::KeepWorldTransform);
+	if (DrawnLatticeEdges->IsRegistered())
+	{
+		DrawnLatticeEdges->ReregisterComponent();
+	}
+	else
+	{
+		DrawnLatticeEdges->RegisterComponent();
+	}
+}
+
+void ULatticeControlPointsMechanic::Shutdown()
+{
+	if (PreviewGeometryActor)
+	{
+		PreviewGeometryActor->Destroy();
+		PreviewGeometryActor = nullptr;
+	}
+
+	if (PointTransformGizmo)
+	{
+		PointTransformGizmo->Shutdown();
+		PointTransformGizmo = nullptr;
+	}
+
+	UInteractiveGizmoManager* GizmoManager = GetParentTool()->GetToolManager()->GetPairedGizmoManager();
+	GizmoManager->DestroyAllGizmosByOwner(GetParentTool());
+}
+
+
+void ULatticeControlPointsMechanic::Render(IToolsContextRenderAPI* RenderAPI)
+{
+	// Cache the camera state
+	GetParentTool()->GetToolManager()->GetContextQueriesAPI()->GetCurrentViewState(CachedCameraState);
+}
+
+void ULatticeControlPointsMechanic::RebuildDrawables()
+{
+	DrawnControlPoints->Clear();
+	GeometrySet.Reset();
+	for (int32 PointIndex = 0; PointIndex < ControlPoints.Num(); ++PointIndex)
+	{
+		const FVector3d& P = ControlPoints[PointIndex];
+		DrawnControlPoints->InsertPoint(PointIndex, FRenderablePoint(static_cast<FVector>(P), NormalPointColor, PointsSize));
+		GeometrySet.AddPoint(PointIndex, P);
+	}
+
+	for (int32 PointID : SelectedPointIDs)
+	{
+		if (DrawnControlPoints->IsPointValid(PointID))
+		{
+			DrawnControlPoints->SetPointColor(PointID, SelectedColor);
+		}
+	}
+
+	DrawnLatticeEdges->Clear();
+	for (int32 EdgeIndex = 0; EdgeIndex < LatticeEdges.Num(); ++EdgeIndex)
+	{
+		const FVector3d& A = ControlPoints[LatticeEdges[EdgeIndex].X];
+		const FVector3d& B = ControlPoints[LatticeEdges[EdgeIndex].Y];
+		int32 SegmentID = DrawnLatticeEdges->AddLine(FVector(A), FVector(B), NormalSegmentColor, SegmentsThickness);
+		check(SegmentID == EdgeIndex);
+	}
+}
+
+
+void ULatticeControlPointsMechanic::UpdateDrawables()
+{
+	for (int32 PointIndex = 0; PointIndex < ControlPoints.Num(); ++PointIndex)
+	{
+		const FVector3d& P = ControlPoints[PointIndex];
+		DrawnControlPoints->SetPointPosition(PointIndex, static_cast<FVector>(P));
+		DrawnControlPoints->SetPointColor(PointIndex, NormalPointColor);
+
+		GeometrySet.UpdatePoint(PointIndex, P);
+	}
+
+	for (int32 PointID : SelectedPointIDs)
+	{
+		if (DrawnControlPoints->IsPointValid(PointID))
+		{
+			DrawnControlPoints->SetPointColor(PointID, SelectedColor);
+		}
+	}
+
+	for (int32 EdgeIndex = 0; EdgeIndex < LatticeEdges.Num(); ++EdgeIndex)
+	{
+		const FVector3d& A = ControlPoints[LatticeEdges[EdgeIndex].X];
+		const FVector3d& B = ControlPoints[LatticeEdges[EdgeIndex].Y];
+		DrawnLatticeEdges->SetLineStart(EdgeIndex, FVector(A));
+		DrawnLatticeEdges->SetLineEnd(EdgeIndex, FVector(B));
+	}
+}
+
+
+void ULatticeControlPointsMechanic::UpdateDrawablesForPoint(int32 PointIndex)
+{
+	const FVector3d& P = ControlPoints[PointIndex];
+	GeometrySet.UpdatePoint(PointIndex, P);
+
+	DrawnControlPoints->SetPointPosition(PointIndex, static_cast<FVector>(P));
+	DrawnControlPoints->SetPointColor(PointIndex, NormalPointColor);
+
+	if (DrawnControlPoints->IsPointValid(PointIndex))
+	{
+		DrawnControlPoints->SetPointColor(PointIndex, SelectedColor);
+	}
+
+	// TODO: Accelerate this somehow. Don't want to search over the entire set of edges any time one point changes.
+	for (int32 EdgeIndex = 0; EdgeIndex < LatticeEdges.Num(); ++EdgeIndex)		
+	{
+		if (LatticeEdges[EdgeIndex].X == PointIndex)
+		{
+			DrawnLatticeEdges->SetLineStart(EdgeIndex, static_cast<FVector>(P));
+		}
+		else if (LatticeEdges[EdgeIndex].Y == PointIndex)
+		{
+			DrawnLatticeEdges->SetLineEnd(EdgeIndex, static_cast<FVector>(P));
+		}
+	}
+}
+
+
+void ULatticeControlPointsMechanic::GizmoTransformStarted(UTransformProxy* Proxy)
+{
+	ParentTool->GetToolManager()->BeginUndoTransaction(LatticePointMovementTransactionText);
+
+	GizmoStartPosition = Proxy->GetTransform().GetTranslation();
+	GizmoStartRotation = Proxy->GetTransform().GetRotation();
+	GizmoStartScale = Proxy->GetTransform().GetScale3D();
+
+	SelectedPointStartPositions.SetNum(SelectedPointIDs.Num());
+	for (int32 i = 0; i < SelectedPointIDs.Num(); ++i)
+	{
+		SelectedPointStartPositions[i] = ControlPoints[SelectedPointIDs[i]];
+	}
+
+	bGizmoBeingDragged = true;
+}
+
+void ULatticeControlPointsMechanic::GizmoTransformChanged(UTransformProxy* Proxy, FTransform Transform)
+{
+	if (SelectedPointIDs.Num() == 0 || !bGizmoBeingDragged)
+	{
+		return;
+	}
+
+	bool bPointsChanged = false;
+
+	FVector Displacement = Transform.GetTranslation() - GizmoStartPosition;
+	FQuaterniond DeltaRotation = FQuaterniond(Transform.GetRotation() * GizmoStartRotation.Inverse());
+	FVector DeltaScale = Transform.GetScale3D() / GizmoStartScale;
+
+	FTransform3d DeltaTransform;
+	DeltaTransform.SetScale(DeltaScale);
+	DeltaTransform.SetRotation(DeltaRotation);
+	DeltaTransform.SetTranslation(Transform.GetTranslation());
+
+	// If any deltas are non-zero
+	if (Displacement != FVector::ZeroVector || !DeltaRotation.EpsilonEqual(FQuaterniond::Identity(), SMALL_NUMBER) || DeltaScale != FVector::OneVector)
+	{
+		for (int32 i = 0; i < SelectedPointIDs.Num(); ++i)
+		{
+			int PointID = SelectedPointIDs[i];
+			FVector3d PointPosition = SelectedPointStartPositions[i];
+
+			// Translate to origin, scale, rotate, and translate back (DeltaTransform has "translate back" baked in.)
+			PointPosition -= GizmoStartPosition;
+			PointPosition = DeltaTransform.TransformPosition(PointPosition);
+
+			ControlPoints[PointID] = PointPosition;
+		}
+		bPointsChanged = true;
+		UpdateDrawables();
+	}
+
+	if (bPointsChanged)
+	{
+		OnPointsChanged.Broadcast();
+	}
+}
+
+void ULatticeControlPointsMechanic::GizmoTransformEnded(UTransformProxy* Proxy)
+{
+	TArray<FVector3d> SelectedPointNewPositions;
+	for (int32 i = 0; i < SelectedPointIDs.Num(); ++i)
+	{
+		SelectedPointNewPositions.Add(ControlPoints[SelectedPointIDs[i]]);
+	}
+
+	bool bFirstMovement = !bHasChanged;
+	bHasChanged = true;
+
+	ParentTool->GetToolManager()->EmitObjectChange(this, MakeUnique<FLatticeControlPointsMechanicMovementChange>(
+		SelectedPointIDs, SelectedPointStartPositions, SelectedPointNewPositions,
+		CurrentChangeStamp, bFirstMovement), LatticePointMovementTransactionText);
+
+	SelectedPointStartPositions.Reset();
+
+	// TODO: When we implement snapping
+	// We may need to reset the gizmo if our snapping caused the final point position to differ from the gizmo position
+	// UpdateGizmoLocation();
+
+	ParentTool->GetToolManager()->EndUndoTransaction();		// was started in GizmoTransformStarted above
+
+	// This just lets the tool know that the gizmo has finished moving and we've added it to the undo stack.
+	// TODO: Add a different callback? "OnGizmoTransformChanged"?
+	OnPointsChanged.Broadcast();
+
+	bGizmoBeingDragged = false;
+}
+
+void ULatticeControlPointsMechanic::UpdatePointLocations(const TArray<int32>& PointIDs, const TArray<FVector3d>& NewLocations)
+{
+	check(NewLocations.Num() == PointIDs.Num());
+	for (int i = 0; i < NewLocations.Num(); ++i)
+	{
+		ControlPoints[PointIDs[i]] = NewLocations[i];
+	}
+	UpdateDrawables();
+}
+
+bool ULatticeControlPointsMechanic::HitTest(const FInputDeviceRay& ClickPos, FInputRayHit& ResultOut)
+{
+	FGeometrySet3::FNearest Nearest;
+
+	// See if we hit a point for selection
+	if (GeometrySet.FindNearestPointToRay(ClickPos.WorldRay, Nearest, GeometrySetToleranceTest))
+	{
+		ResultOut = FInputRayHit(Nearest.RayParam);
+		return true;
+	}
+	return false;
+}
+
+FInputRayHit ULatticeControlPointsMechanic::IsHitByClick(const FInputDeviceRay& ClickPos)
+{
+	FInputRayHit Result;
+	HitTest(ClickPos, Result);
+	return Result;
+}
+
+void ULatticeControlPointsMechanic::OnClicked(const FInputDeviceRay& ClickPos)
+{
+	FGeometrySet3::FNearest Nearest;
+
+	if (GeometrySet.FindNearestPointToRay(ClickPos.WorldRay, Nearest, GeometrySetToleranceTest))
+	{
+		ParentTool->GetToolManager()->BeginUndoTransaction(LatticePointSelectionTransactionText);
+		ChangeSelection(Nearest.ID, bAddToSelectionToggle);
+		ParentTool->GetToolManager()->EndUndoTransaction();
+	}
+
+	bIsDragging = false;
+}
+
+void ULatticeControlPointsMechanic::ChangeSelection(int32 NewPointID, bool bAddToSelection)
+{
+	// If not adding to selection, clear it
+	if (!bAddToSelection && SelectedPointIDs.Num() > 0)
+	{
+		TArray<int32> PointsToDeselect;
+
+		for (int32 PointID : SelectedPointIDs)
+		{
+			// We check for validity here because we'd like to be able to use this function to deselect points after
+			// deleting them.
+			if (DrawnControlPoints->IsPointValid(PointID))
+			{
+				PointsToDeselect.Add(PointID);
+				DrawnControlPoints->SetPointColor(PointID, NormalPointColor);
+			}
+		}
+
+		ParentTool->GetToolManager()->EmitObjectChange(this, MakeUnique<FLatticeControlPointsMechanicSelectionChange>(
+			PointsToDeselect, false, CurrentChangeStamp), LatticePointDeselectionTransactionText);
+
+		SelectedPointIDs.Empty();
+	}
+
+	// We check for validity here because giving an invalid id (such as -1) with bAddToSelection == false
+	// is an easy way to clear the selection.
+	if ((NewPointID >= 0)  && (NewPointID < ControlPoints.Num()))
+	{
+		if (bAddToSelection && DeselectPoint(NewPointID))
+		{
+			ParentTool->GetToolManager()->EmitObjectChange(this, MakeUnique<FLatticeControlPointsMechanicSelectionChange>(
+				NewPointID, false, CurrentChangeStamp), LatticePointDeselectionTransactionText);
+		}
+		else
+		{
+			SelectPoint(NewPointID);
+			ParentTool->GetToolManager()->EmitObjectChange(this, MakeUnique<FLatticeControlPointsMechanicSelectionChange>(
+				NewPointID, true, CurrentChangeStamp), LatticePointSelectionTransactionText);
+		}
+	}
+
+	UpdateGizmoLocation();
+}
+
+void ULatticeControlPointsMechanic::UpdateGizmoLocation()
+{
+	if (!PointTransformGizmo)
+	{
+		return;
+	}
+
+	if (SelectedPointIDs.Num() == 0)
+	{
+		PointTransformGizmo->SetVisibility(false);
+		PointTransformGizmo->ReinitializeGizmoTransform(FTransform());
+	}
+	else
+	{
+		FVector3d NewGizmoLocation(0,0,0);
+		for (int32 PointID : SelectedPointIDs)
+		{
+			NewGizmoLocation += ControlPoints[PointID];
+		}
+		NewGizmoLocation /= SelectedPointIDs.Num();
+
+		// Don't clear the gizmo rotation
+		FQuat OldGizmoRotation = PointTransformProxy->GetTransform().GetRotation();
+
+		PointTransformGizmo->ReinitializeGizmoTransform(FTransform(OldGizmoRotation, (FVector)NewGizmoLocation));
+
+		// Clear the child scale
+		FVector GizmoScale{ 1.0, 1.0, 1.0 };
+		PointTransformGizmo->SetNewChildScale(GizmoScale);
+
+		PointTransformGizmo->SetVisibility(true);
+	}
+}
+
+bool ULatticeControlPointsMechanic::DeselectPoint(int32 PointID)
+{
+	bool PointFound = false;
+	int32 IndexInSelection;
+	// TODO: This might be slow if we have a lot of selected points (UETOOL-2357)
+	if (SelectedPointIDs.Find(PointID, IndexInSelection))
+	{
+		SelectedPointIDs.RemoveAt(IndexInSelection);
+		DrawnControlPoints->SetPointColor(PointID, NormalPointColor);
+		PointFound = true;
+	}
+	
+	return PointFound;
+}
+
+void ULatticeControlPointsMechanic::SelectPoint(int32 PointID)
+{
+	SelectedPointIDs.Add(PointID);
+	DrawnControlPoints->SetPointColor(PointID, SelectedColor);
+}
+
+void ULatticeControlPointsMechanic::ClearSelection()
+{
+	ChangeSelection(-1, false);
+}
+
+FInputRayHit ULatticeControlPointsMechanic::BeginHoverSequenceHitTest(const FInputDeviceRay& PressPos)
+{
+	FInputRayHit Result;
+	HitTest(PressPos, Result);
+	return Result;
+}
+
+void ULatticeControlPointsMechanic::OnBeginHover(const FInputDeviceRay& DevicePos)
+{
+	OnUpdateHover(DevicePos);
+}
+
+void ULatticeControlPointsMechanic::ClearHover()
+{
+	if (HoveredPointID >= 0)
+	{
+		DrawnControlPoints->SetPointColor(HoveredPointID, PreHoverPointColor);
+		HoveredPointID = -1;
+	}
+}
+
+
+bool ULatticeControlPointsMechanic::OnUpdateHover(const FInputDeviceRay& DevicePos)
+{
+	FGeometrySet3::FNearest Nearest;
+
+	// see if we're hovering a point for selection
+	if (GeometrySet.FindNearestPointToRay(DevicePos.WorldRay, Nearest, GeometrySetToleranceTest))
+	{
+		// Only need to update the hover if we changed the point
+		if (Nearest.ID != HoveredPointID)
+		{
+			ClearHover();
+			HoveredPointID = Nearest.ID;
+			PreHoverPointColor = DrawnControlPoints->GetPoint(HoveredPointID).Color;
+			DrawnControlPoints->SetPointColor(HoveredPointID, HoverColor);
+		}
+	}
+	else
+	{
+		// Not hovering anything, so done hovering
+		return false;
+	}
+
+	return true;
+}
+
+void ULatticeControlPointsMechanic::OnEndHover()
+{
+	ClearHover();
+}
+
+// Detects Ctrl key state
+void ULatticeControlPointsMechanic::OnUpdateModifierState(int ModifierID, bool bIsOn)
+{
+	if (ModifierID == CtrlModifierId || ModifierID == ShiftModifierId)
+	{
+		bAddToSelectionToggle = bIsOn;
+	}
+}
+
+
+// ==================== IClickDragBehaviorTarget ====================
+
+FInputRayHit ULatticeControlPointsMechanic::CanBeginClickDragSequence(const FInputDeviceRay& PressPos)
+{
+	FInputRayHit Dummy;
+
+	// This is the boolean that is checked to see if a drag sequence can be started. In our case we want to begin the 
+	// drag sequence even if the first ray doesn't hit anything, so set this to true.
+	Dummy.bHit = true;
+
+	return Dummy;
+}
+
+void ULatticeControlPointsMechanic::OnClickPress(const FInputDeviceRay& PressPos)
+{
+	if (!PressPos.bHas2D)
+	{
+		return;
+	}
+
+	for (int32& PointID : SelectedPointIDs)
+	{
+		DrawnControlPoints->SetPointColor(PointID, NormalPointColor);
+	}
+	CurrentDragSelection.Empty();
+
+	UpdateGizmoLocation();
+
+	DragStartScreenPosition = PressPos.ScreenPosition;
+	DragStartWorldRay = PressPos.WorldRay;
+
+	// Hide gizmo while dragging
+	if (PointTransformGizmo)
+	{
+		PointTransformGizmo->SetVisibility(false);
+		PointTransformGizmo->ReinitializeGizmoTransform(FTransform());
+	}
+}
+
+
+static FVector2D PlaneCoordinates(const FVector& Point, const FPlane& Plane, const FVector& UBasisVector, const FVector& VBasisVector)
+{
+	float U = FVector::DotProduct(Point - Plane.GetOrigin(), UBasisVector);
+	float V = FVector::DotProduct(Point - Plane.GetOrigin(), VBasisVector);
+	return FVector2D{ U,V };
+}
+
+void ULatticeControlPointsMechanic::OnClickDrag(const FInputDeviceRay& DragPos)
+{
+	if (!DragPos.bHas2D)
+	{
+		return;
+	}
+
+	bIsDragging = true;
+	DragCurrentScreenPosition = DragPos.ScreenPosition;
+	DragCurrentWorldRay = DragPos.WorldRay;
+
+	// Intersect the drag rays and and project lattice points all to the same plane in 3D. Then compute 2D coordinates 
+	// and use an AABB test to determine which points are in the drag rectangle.
+
+	FViewCameraState CameraState;
+	GetParentTool()->GetToolManager()->GetContextQueriesAPI()->GetCurrentViewState(CameraState);
+
+	// Create plane in front of camera
+	FPlane CameraPlane(CameraState.Position + CameraState.Forward(), CameraState.Forward());
+
+	FVector UBasisVector = CameraState.Right();
+	FVector VBasisVector = CameraState.Up();
+
+	FVector StartIntersection = FMath::RayPlaneIntersection(DragStartWorldRay.Origin, 
+															DragStartWorldRay.Direction, 
+															CameraPlane);
+	FVector2D Start2D = PlaneCoordinates(StartIntersection, CameraPlane, UBasisVector, VBasisVector);
+
+	FVector CurrentIntersection = FMath::RayPlaneIntersection(DragCurrentWorldRay.Origin, 
+															  DragCurrentWorldRay.Direction, 
+															  CameraPlane);
+	FVector2D Current2D = PlaneCoordinates(CurrentIntersection, CameraPlane, UBasisVector, VBasisVector);
+
+	FBox2D DragBox(Start2D, Start2D);
+	DragBox += Current2D;	// Initialize this way so we don't have to care about min/max
+
+	for (int32 PointID = 0; PointID < ControlPoints.Num(); ++PointID)
+	{
+		FVector PointPosition = DrawnControlPoints->GetPoint(PointID).Position;
+		FVector PointIntersection;
+		if (CameraState.bIsOrthographic)
+		{
+			// project directly to plane
+			PointIntersection = FVector::PointPlaneProject(PointPosition, CameraPlane);
+		}
+		else
+		{
+			// intersect along the eye-to-point ray
+			PointIntersection = FMath::RayPlaneIntersection(CameraState.Position, 
+															PointPosition - CameraState.Position,
+															CameraPlane);
+		}
+
+		FVector2D Point2D = PlaneCoordinates(PointIntersection, CameraPlane, UBasisVector, VBasisVector);
+		if (DragBox.IsInside(Point2D))
+		{
+			CurrentDragSelection.Add(PointID);
+			DrawnControlPoints->SetPointColor(PointID, SelectedColor);
+		}
+		else
+		{
+			DrawnControlPoints->SetPointColor(PointID, NormalPointColor);
+		}
+	}
+}
+
+void ULatticeControlPointsMechanic::OnClickRelease(const FInputDeviceRay& ReleasePos)
+{
+	// Deselect previous SelectedPointIDs and replace it with "drag selection" points. Do this in one Undo transaction.
+
+	ParentTool->GetToolManager()->BeginUndoTransaction(LatticePointSelectionTransactionText);
+
+	if (SelectedPointIDs.Num() > 0)
+	{
+		ParentTool->GetToolManager()->EmitObjectChange(this, MakeUnique<FLatticeControlPointsMechanicSelectionChange>(
+			SelectedPointIDs, false, CurrentChangeStamp), LatticePointDeselectionTransactionText);
+	}
+
+	if (CurrentDragSelection.Num() > 0)
+	{
+		ParentTool->GetToolManager()->EmitObjectChange(this, MakeUnique<FLatticeControlPointsMechanicSelectionChange>(
+			CurrentDragSelection, true, CurrentChangeStamp), LatticePointSelectionTransactionText);
+	}
+
+	ParentTool->GetToolManager()->EndUndoTransaction();
+
+	SelectedPointIDs = CurrentDragSelection;
+	CurrentDragSelection.Empty();
+
+	bIsDragging = false;
+	UpdateGizmoLocation();
+}
+
+
+void ULatticeControlPointsMechanic::OnTerminateDragSequence()
+{
+	// Not sure how this can happen. Pressing escape quits the tool altogether.
+	CurrentDragSelection.Empty();
+	bIsDragging = false;
+	UpdateGizmoLocation();
+}
+
+const TArray<FVector3d>& ULatticeControlPointsMechanic::GetControlPoints() const
+{
+	return ControlPoints;
+}
+
+// ==================== Undo/redo object functions ====================
+
+FLatticeControlPointsMechanicSelectionChange::FLatticeControlPointsMechanicSelectionChange(int32 PointIDIn,
+	bool AddedIn, int32 ChangeStampIn)
+	: PointIDs{ PointIDIn }
+	, Added(AddedIn)
+	, ChangeStamp(ChangeStampIn)
+{}
+
+FLatticeControlPointsMechanicSelectionChange::FLatticeControlPointsMechanicSelectionChange(const TArray<int32>& PointIDsIn, 
+	bool AddedIn, int32 ChangeStampIn)
+	: PointIDs(PointIDsIn)
+	, Added(AddedIn)
+	, ChangeStamp(ChangeStampIn)
+{}
+
+void FLatticeControlPointsMechanicSelectionChange::Apply(UObject* Object)
+{
+	ULatticeControlPointsMechanic* Mechanic = Cast<ULatticeControlPointsMechanic>(Object);
+	
+	for (int32 PointID : PointIDs)
+	{
+		if (Added)
+		{
+			Mechanic->SelectPoint(PointID);
+		}
+		else
+		{
+			Mechanic->DeselectPoint(PointID);
+		}
+	}
+
+	Mechanic->UpdateGizmoLocation();
+}
+
+void FLatticeControlPointsMechanicSelectionChange::Revert(UObject* Object)
+{
+	ULatticeControlPointsMechanic* Mechanic = Cast<ULatticeControlPointsMechanic>(Object);
+	for (int32 PointID : PointIDs)
+	{
+		if (Added)
+		{
+			Mechanic->DeselectPoint(PointID);
+		}
+		else
+		{
+			Mechanic->SelectPoint(PointID);
+		}
+	}
+	Mechanic->UpdateGizmoLocation();
+}
+
+FString FLatticeControlPointsMechanicSelectionChange::ToString() const
+{
+	return TEXT("FLatticeControlPointsMechanicSelectionChange");
+}
+
+
+FLatticeControlPointsMechanicMovementChange::FLatticeControlPointsMechanicMovementChange(const TArray<int32>& PointIDsIn,
+	const TArray<FVector3d>& OriginalPositionsIn,
+	const TArray<FVector3d>& NewPositionsIn,
+	int32 ChangeStampIn,
+	bool bFirstMovementIn)
+	: PointIDs{ PointIDsIn }
+	, OriginalPositions{ OriginalPositionsIn }
+	, NewPositions{ NewPositionsIn }
+	, ChangeStamp(ChangeStampIn)
+	, bFirstMovement(bFirstMovementIn)
+{
+	check(PointIDs.Num() == OriginalPositions.Num());
+	check(PointIDs.Num() == NewPositions.Num());
+}
+
+void FLatticeControlPointsMechanicMovementChange::Apply(UObject* Object)
+{
+	ULatticeControlPointsMechanic* Mechanic = Cast<ULatticeControlPointsMechanic>(Object);
+	check(PointIDs.Num() == NewPositions.Num());
+	Mechanic->UpdatePointLocations(PointIDs, NewPositions);
+	Mechanic->UpdateGizmoLocation();
+	Mechanic->bHasChanged = false;
+	Mechanic->OnPointsChanged.Broadcast();
+}
+
+void FLatticeControlPointsMechanicMovementChange::Revert(UObject* Object)
+{
+	ULatticeControlPointsMechanic* Mechanic = Cast<ULatticeControlPointsMechanic>(Object);
+	check(PointIDs.Num() == OriginalPositions.Num());
+	Mechanic->UpdatePointLocations(PointIDs, OriginalPositions);
+	Mechanic->UpdateGizmoLocation();
+	if (bFirstMovement)
+	{
+		// If we're undoing the first change, make it possible to change the lattice resolution again
+		Mechanic->bHasChanged = false;
+	}
+	Mechanic->OnPointsChanged.Broadcast();
+}
+
+FString FLatticeControlPointsMechanicMovementChange::ToString() const
+{
+	return TEXT("FLatticeControlPointsMechanicMovementChange");
+}
+
+#undef LOCTEXT_NAMESPACE
diff --git a/Engine/Plugins/Experimental/MeshModelingToolset/Source/ModelingComponents/Public/Mechanics/LatticeControlPointsMechanic.h b/Engine/Plugins/Experimental/MeshModelingToolset/Source/ModelingComponents/Public/Mechanics/LatticeControlPointsMechanic.h
new file mode 100644
--- /dev/null
+++ b/Engine/Plugins/Experimental/MeshModelingToolset/Source/ModelingComponents/Public/Mechanics/LatticeControlPointsMechanic.h
@@ -0,0 +1,224 @@
+// Copyright Epic Games, Inc. All Rights Reserved.
+
+#pragma once
+
+#include "CoreMinimal.h"
+
+#include "BaseBehaviors/BehaviorTargetInterfaces.h"
+#include "FrameTypes.h"
+#include "InteractionMechanic.h"
+#include "InteractiveToolChange.h"
+#include "Snapping/PointPlanarSnapSolver.h"
+#include "Spatial/GeometrySet3.h"
+#include "ToolContextInterfaces.h" //FViewCameraState
+#include "VectorTypes.h"
+#include "VectorUtil.h"
+#include "ToolDataVisualizer.h"
+#include "SceneView.h"
+
+#include "LatticeControlPointsMechanic.generated.h"
+
+class APreviewGeometryActor;
+class ULineSetComponent;
+class UMouseHoverBehavior;
+class UPointSetComponent;
+class USingleClickInputBehavior;
+class UTransformGizmo;
+class UTransformProxy;
+
+UCLASS()
+class MODELINGCOMPONENTS_API ULatticeControlPointsMechanic : 
+	public UInteractionMechanic, public IClickBehaviorTarget, public IHoverBehaviorTarget, public IClickDragBehaviorTarget
+{
+	GENERATED_BODY()
+
+public:
+
+	// TODO: Snapping
+
+	// This delegate is called every time the control points are altered.
+	DECLARE_MULTICAST_DELEGATE(OnPointsChangedEvent);
+	OnPointsChangedEvent OnPointsChanged;
+
+	virtual void Initialize(const TArray<FVector3d>& Points, const TArray<FVector2i>& Edges);
+	void SetWorld(UWorld* World);
+	const TArray<FVector3d>& GetControlPoints() const;
+
+	// UInteractionMechanic
+	virtual void Setup(UInteractiveTool* ParentTool) override;
+	virtual void Shutdown() override;
+	virtual void Render(IToolsContextRenderAPI* RenderAPI) override;
+
+	// IClickBehaviorTarget implementation
+	virtual FInputRayHit IsHitByClick(const FInputDeviceRay& ClickPos) override;
+	virtual void OnClicked(const FInputDeviceRay& ClickPos) override;
+
+	// IHoverBehaviorTarget implementation
+	virtual FInputRayHit BeginHoverSequenceHitTest(const FInputDeviceRay& PressPos) override;
+	virtual void OnBeginHover(const FInputDeviceRay& DevicePos) override;
+	virtual bool OnUpdateHover(const FInputDeviceRay& DevicePos) override;
+	virtual void OnEndHover() override;
+	virtual void OnUpdateModifierState(int ModifierID, bool bIsOn) override;
+
+	// IClickDragBehaviorTarget
+	FInputRayHit CanBeginClickDragSequence(const FInputDeviceRay& PressPos) override;
+	void OnClickPress(const FInputDeviceRay& PressPos) override;
+	void OnClickDrag(const FInputDeviceRay& DragPos) override;
+	void OnClickRelease(const FInputDeviceRay& ReleasePos) override;
+	void OnTerminateDragSequence() override;
+
+	// Marquee/rectangle select support
+	FVector2D DragStartScreenPosition;
+	FRay DragStartWorldRay;
+	FVector2D DragCurrentScreenPosition;
+	FRay DragCurrentWorldRay;
+	bool bIsDragging;
+
+	bool bHasChanged = false;
+
+protected:
+
+	TArray<FVector3d> ControlPoints;
+	TArray<FVector2i> LatticeEdges;
+
+	// Used for spatial queries
+	FGeometrySet3 GeometrySet;
+
+	/** Used for displaying points/segments */
+	UPROPERTY()
+	APreviewGeometryActor* PreviewGeometryActor;
+	UPROPERTY()
+	UPointSetComponent* DrawnControlPoints;
+	UPROPERTY()
+	ULineSetComponent* DrawnLatticeEdges;
+
+	// Variables for drawing
+	FColor NormalSegmentColor;
+	FColor NormalPointColor;
+	float SegmentsThickness;
+	float PointsSize;
+	FColor HoverColor;
+	FColor SelectedColor;
+
+	// Cache previous color while temporarily changing the color of a hovered-over point
+	FColor PreHoverPointColor;
+
+	// Support for Ctrl toggle
+	bool bAddToSelectionToggle = false;
+	const int32 CtrlModifierId = 1;
+	const int32 ShiftModifierId = 2;
+
+	// Support for gizmo. Since the points aren't individual components, we don't actually use UTransformProxy
+	// for the transform forwarding- we just use it for the callbacks.
+	UPROPERTY()
+	UTransformProxy* PointTransformProxy;
+	UPROPERTY()
+	UTransformGizmo* PointTransformGizmo;
+
+	// Used to make it easy to tell whether the gizmo was moved by the user or by undo/redo or
+	// some other change that we shouldn't respond to. Basing our movement undo/redo on the
+	// gizmo turns out to be quite a pain, though may someday be easier if the transform proxy
+	// is able to manage arbitrary objects.
+	bool bGizmoBeingDragged = false;
+
+	// Callbacks we'll receive from the gizmo proxy
+	void GizmoTransformChanged(UTransformProxy* Proxy, FTransform Transform);
+	void GizmoTransformStarted(UTransformProxy* Proxy);
+	void GizmoTransformEnded(UTransformProxy* Proxy);
+
+	// Cached render information for viewport/scene interaction
+	FViewCameraState CachedCameraState;
+
+	// Support for hovering
+	TFunction<bool(const FVector3d&, const FVector3d&)> GeometrySetToleranceTest;
+	int32 HoveredPointID = -1;
+	void ClearHover();
+
+	// Support for selection
+	TArray<int32> SelectedPointIDs;
+	TArray<int32> CurrentDragSelection;
+
+	// We need the selected point start positions so we can move multiple points appropriately.
+	TArray<FVector3d> SelectedPointStartPositions;
+
+	// The starting point of the gizmo is needed to determine the offset by which to move the points.
+	FVector GizmoStartPosition;
+	FQuat	GizmoStartRotation;
+	FVector GizmoStartScale;
+
+	// These issue undo/redo change objects, and must therefore not be called in undo/redo code.
+	void ChangeSelection(int32 NewPointID, bool AddToSelection);
+	void ClearSelection();
+
+	// All of the following do not issue undo/redo change objects.
+	bool HitTest(const FInputDeviceRay& ClickPos, FInputRayHit& ResultOut);
+	void SelectPoint(int32 PointID);
+	bool DeselectPoint(int32 PointID);
+	void UpdateGizmoLocation();
+	void UpdatePointLocations(const TArray<int32>& PointIDs, const TArray<FVector3d>& NewLocations);
+
+	void RebuildDrawables();
+	void UpdateDrawables();
+	void UpdateDrawablesForPoint(int32 PointIndex);
+
+	// Used for expiring undo/redo changes, which compare this to their stored value and expire themselves if they do not match.
+	int32 CurrentChangeStamp = 0;
+
+	friend class FLatticeControlPointsMechanicSelectionChange;
+	friend class FLatticeControlPointsMechanicMovementChange;
+};
+
+
+// Undo/redo support:
+
+// Control point selection has changed
+class MODELINGCOMPONENTS_API FLatticeControlPointsMechanicSelectionChange : public FToolCommandChange
+{
+public:
+	FLatticeControlPointsMechanicSelectionChange(int32 PointIDIn, bool AddedIn, int32 ChangeStampIn);
+
+	FLatticeControlPointsMechanicSelectionChange(const TArray<int32>& PointIDsIn, bool AddedIn, int32 ChangeStampIn);
+
+	virtual void Apply(UObject* Object) override;
+	virtual void Revert(UObject* Object) override;
+	virtual bool HasExpired(UObject* Object) const override
+	{
+		return Cast<ULatticeControlPointsMechanic>(Object)->CurrentChangeStamp != ChangeStamp;
+	}
+	virtual FString ToString() const override;
+
+protected:
+	TArray<int32> PointIDs;
+
+	bool Added;
+	int32 ChangeStamp;
+};
+
+
+// Control points have moved
+class MODELINGCOMPONENTS_API FLatticeControlPointsMechanicMovementChange : public FToolCommandChange
+{
+public:
+
+	FLatticeControlPointsMechanicMovementChange(const TArray<int32>& PointIDsIn, 
+		const TArray<FVector3d>& OriginalPositionIn,
+		const TArray<FVector3d>& NewPositionIn, 
+		int32 ChangeStampIn,
+		bool bFirstMovementIn);
+
+	virtual void Apply(UObject* Object) override;
+	virtual void Revert(UObject* Object) override;
+	virtual bool HasExpired(UObject* Object) const override
+	{
+		return Cast<ULatticeControlPointsMechanic>(Object)->CurrentChangeStamp != ChangeStamp;
+	}
+	virtual FString ToString() const override;
+
+protected:
+
+	TArray<int32> PointIDs;
+	TArray<FVector3d> OriginalPositions;
+	TArray<FVector3d> NewPositions;
+	int32 ChangeStamp;
+	bool bFirstMovement;
+};
diff --git a/Engine/Plugins/Experimental/MeshModelingToolset/Source/ModelingOperators/Private/DeformationOps/LatticeDeformerOp.cpp b/Engine/Plugins/Experimental/MeshModelingToolset/Source/ModelingOperators/Private/DeformationOps/LatticeDeformerOp.cpp
new file mode 100644
--- /dev/null
+++ b/Engine/Plugins/Experimental/MeshModelingToolset/Source/ModelingOperators/Private/DeformationOps/LatticeDeformerOp.cpp
@@ -0,0 +1,39 @@
+// Copyright Epic Games, Inc. All Rights Reserved.
+
+#include "DeformationOps/LatticeDeformerOp.h"
+#include "Operations/FFDLattice.h"
+
+void FLatticeDeformerOp::CalculateResult(FProgressCancel* Progress)
+{
+	if (Progress->Cancelled())
+	{
+		return;
+	}
+
+	ResultMesh->Copy(*OriginalMesh);
+
+	if (Progress->Cancelled())
+	{
+		return;
+	}
+
+	TArray<FVector3d> DeformedPositions;
+	Lattice->GetDeformedMeshVertexPositions(LatticeControlPoints, DeformedPositions, bUseCubicInterpolation, Progress);
+
+	check(ResultMesh->VertexCount() == DeformedPositions.Num());
+
+	for (int vid : ResultMesh->VertexIndicesItr())
+	{
+		ResultMesh->SetVertex(vid, DeformedPositions[vid]);
+	}
+}
+
+FLatticeDeformerOp::FLatticeDeformerOp(TSharedPtr<FDynamicMesh3> InOriginalMesh,
+									   TSharedPtr<FFFDLattice> InLattice,
+									   const TArray<FVector3d>& InLatticeControlPoints,
+									   bool InUseCubicInterpolation) :
+	Lattice(InLattice),
+	OriginalMesh(InOriginalMesh),
+	LatticeControlPoints(InLatticeControlPoints),
+	bUseCubicInterpolation(InUseCubicInterpolation)
+{}
diff --git a/Engine/Plugins/Experimental/MeshModelingToolset/Source/ModelingOperators/Public/DeformationOps/LatticeDeformerOp.h b/Engine/Plugins/Experimental/MeshModelingToolset/Source/ModelingOperators/Public/DeformationOps/LatticeDeformerOp.h
new file mode 100644
--- /dev/null
+++ b/Engine/Plugins/Experimental/MeshModelingToolset/Source/ModelingOperators/Public/DeformationOps/LatticeDeformerOp.h
@@ -0,0 +1,35 @@
+// Copyright Epic Games, Inc. All Rights Reserved.
+
+#pragma once
+
+#include "BaseOps/SimpleMeshProcessingBaseOp.h"
+#include "DynamicMeshEditor.h"
+#include "MeshNormals.h"
+#include "MeshWeights.h"
+#include "MeshBoundaryLoops.h"
+#include "Operations/JoinMeshLoops.h"
+#include "Solvers/ConstrainedMeshDeformer.h"
+#include "Implicit/GridInterpolant.h"
+#include "Operations/FFDLattice.h"
+
+class MODELINGOPERATORS_API FLatticeDeformerOp : public FDynamicMeshOperator
+{
+public:
+
+	FLatticeDeformerOp(TSharedPtr<FDynamicMesh3> InOriginalMesh,
+					   TSharedPtr<FFFDLattice> InLattice,
+					   const TArray<FVector3d>& InLatticeControlPoints,
+					   bool InUseCubicInterpolation);
+
+	// FDynamicMeshOperator implementation
+	void CalculateResult(FProgressCancel* Progress) override;
+
+protected:
+
+	// Inputs
+	const TSharedPtr<const FFFDLattice> Lattice;
+	const TSharedPtr<const FDynamicMesh3> OriginalMesh;
+	const TArray<FVector3d> LatticeControlPoints;
+	bool bUseCubicInterpolation;
+};
+
diff --git a/Engine/Plugins/Experimental/ModelingToolsEditorMode/Source/ModelingToolsEditorMode/Private/ModelingToolsEditorMode.cpp b/Engine/Plugins/Experimental/ModelingToolsEditorMode/Source/ModelingToolsEditorMode/Private/ModelingToolsEditorMode.cpp
--- a/Engine/Plugins/Experimental/ModelingToolsEditorMode/Source/ModelingToolsEditorMode/Private/ModelingToolsEditorMode.cpp
+++ b/Engine/Plugins/Experimental/ModelingToolsEditorMode/Source/ModelingToolsEditorMode/Private/ModelingToolsEditorMode.cpp
@@ -65,6 +65,7 @@
 #include "ParameterizeMeshTool.h"
 #include "MeshTangentsTool.h"
 #include "ProjectToTargetTool.h"
+#include "LatticeDeformerTool.h"
 
 #include "Physics/PhysicsInspectorTool.h"
 #include "Physics/SetCollisionGeometryTool.h"
@@ -580,6 +581,7 @@ void FModelingToolsEditorMode::Enter()
 	RegisterToolFunc(ToolManagerCommands.BeginAlignObjectsTool, TEXT("AlignObjects"), NewObject<UAlignObjectsToolBuilder>());
 	RegisterToolFunc(ToolManagerCommands.BeginBakeTransformTool, TEXT("BakeTransformTool"), NewObject<UBakeTransformToolBuilder>());
 	RegisterToolFunc(ToolManagerCommands.BeginTransformUVIslandsTool, TEXT("EditUVIslands"), NewObject<UEditUVIslandsToolBuilder>());
+	RegisterToolFunc(ToolManagerCommands.BeginLatticeDeformerTool, TEXT("LatticeDeformerTool"), NewObject<ULatticeDeformerToolBuilder>());
 
 	UCombineMeshesToolBuilder* CombineMeshesToolBuilder = NewObject<UCombineMeshesToolBuilder>();
 	CombineMeshesToolBuilder->AssetAPI = ModelingModeAssetGenerationAPI.Get();
diff --git a/Engine/Plugins/Experimental/ModelingToolsEditorMode/Source/ModelingToolsEditorMode/Private/ModelingToolsEditorModeToolkit.cpp b/Engine/Plugins/Experimental/ModelingToolsEditorMode/Source/ModelingToolsEditorMode/Private/ModelingToolsEditorModeToolkit.cpp
--- a/Engine/Plugins/Experimental/ModelingToolsEditorMode/Source/ModelingToolsEditorMode/Private/ModelingToolsEditorModeToolkit.cpp
+++ b/Engine/Plugins/Experimental/ModelingToolsEditorMode/Source/ModelingToolsEditorMode/Private/ModelingToolsEditorModeToolkit.cpp
@@ -409,6 +409,7 @@ void FModelingToolsEditorModeToolkit::BuildToolPalette_Standard(FName PaletteInd
 		ToolbarBuilder.AddToolBarButton(Commands.BeginOffsetMeshTool);
 		ToolbarBuilder.AddToolBarButton(Commands.BeginDisplaceMeshTool);
 		ToolbarBuilder.AddToolBarButton(Commands.BeginMeshSpaceDeformerTool);
+		ToolbarBuilder.AddToolBarButton(Commands.BeginLatticeDeformerTool);
 		ToolbarBuilder.AddSeparator();
 		ToolbarBuilder.AddToolBarButton(Commands.BeginRemeshMeshTool);
 		ToolbarBuilder.AddToolBarButton(Commands.BeginMeshInspectorTool);
@@ -549,6 +550,7 @@ void FModelingToolsEditorModeToolkit::BuildToolPalette_Experimental(FName Palett
 		ToolbarBuilder.AddToolBarButton(Commands.BeginOffsetMeshTool);
 		ToolbarBuilder.AddToolBarButton(Commands.BeginDisplaceMeshTool);
 		ToolbarBuilder.AddToolBarButton(Commands.BeginMeshSpaceDeformerTool);
+		ToolbarBuilder.AddToolBarButton(Commands.BeginLatticeDeformerTool);
 		ToolbarBuilder.AddSeparator();
 		ToolbarBuilder.AddToolBarButton(Commands.BeginRemeshMeshTool);
 		ToolbarBuilder.AddToolBarButton(Commands.BeginMeshInspectorTool);
diff --git a/Engine/Plugins/Experimental/ModelingToolsEditorMode/Source/ModelingToolsEditorMode/Private/ModelingToolsManagerActions.cpp b/Engine/Plugins/Experimental/ModelingToolsEditorMode/Source/ModelingToolsEditorMode/Private/ModelingToolsManagerActions.cpp
--- a/Engine/Plugins/Experimental/ModelingToolsEditorMode/Source/ModelingToolsEditorMode/Private/ModelingToolsManagerActions.cpp
+++ b/Engine/Plugins/Experimental/ModelingToolsEditorMode/Source/ModelingToolsEditorMode/Private/ModelingToolsManagerActions.cpp
@@ -79,6 +79,7 @@ void FModelingToolsManagerCommands::RegisterCommands()
 	UI_COMMAND(BeginPlaneCutTool, "PlaneCut", "Start the Plane Cut Tool", EUserInterfaceActionType::Button, FInputChord());
 	UI_COMMAND(BeginMirrorTool, "Mirror", "Start the Mirror Tool", EUserInterfaceActionType::Button, FInputChord());
 	UI_COMMAND(BeginHoleFillTool, "HoleFill", "Start the Hole Fill Tool", EUserInterfaceActionType::Button, FInputChord());
+	UI_COMMAND(BeginLatticeDeformerTool, "LatticeDeformer", "Start the Lattice Deformer Tool", EUserInterfaceActionType::Button, FInputChord());
 	UI_COMMAND(BeginPolygonCutTool, "PolyCut", "Start the Polygon Cut Tool", EUserInterfaceActionType::Button, FInputChord());
 	UI_COMMAND(BeginGlobalUVGenerateTool, "AutoUnwrap", "Start the Global UV Unwrap Tool", EUserInterfaceActionType::Button, FInputChord());
 	UI_COMMAND(BeginGroupUVGenerateTool, "GroupUnwrap", "Stat the PolyGroup UV Unwrap Tool", EUserInterfaceActionType::Button, FInputChord());
diff --git a/Engine/Plugins/Experimental/ModelingToolsEditorMode/Source/ModelingToolsEditorMode/Public/ModelingToolsManagerActions.h b/Engine/Plugins/Experimental/ModelingToolsEditorMode/Source/ModelingToolsEditorMode/Public/ModelingToolsManagerActions.h
--- a/Engine/Plugins/Experimental/ModelingToolsEditorMode/Source/ModelingToolsEditorMode/Public/ModelingToolsManagerActions.h
+++ b/Engine/Plugins/Experimental/ModelingToolsEditorMode/Source/ModelingToolsEditorMode/Public/ModelingToolsManagerActions.h
@@ -62,6 +62,7 @@ public:
 	TSharedPtr<FUICommandInfo> BeginPlaneCutTool;
 	TSharedPtr<FUICommandInfo> BeginMirrorTool;
 	TSharedPtr<FUICommandInfo> BeginHoleFillTool;
+	TSharedPtr<FUICommandInfo> BeginLatticeDeformerTool;
 	TSharedPtr<FUICommandInfo> BeginPolygonCutTool;
 	TSharedPtr<FUICommandInfo> BeginVoxelMergeTool;
 	TSharedPtr<FUICommandInfo> BeginVoxelBooleanTool;
