MeshModelingTools: Added a ComponentTarget that allows our tools to operate on volumes.

#rnx

diff --git a/Engine/Plugins/Experimental/MeshModelingToolset/Source/MeshModelingToolsEditorOnly/Private/ConversionUtils/VolumeComponentTarget.cpp b/Engine/Plugins/Experimental/MeshModelingToolset/Source/MeshModelingToolsEditorOnly/Private/ConversionUtils/VolumeComponentTarget.cpp
new file mode 100644
--- /dev/null
+++ b/Engine/Plugins/Experimental/MeshModelingToolset/Source/MeshModelingToolsEditorOnly/Private/ConversionUtils/VolumeComponentTarget.cpp
@@ -0,0 +1,138 @@
+// Copyright Epic Games, Inc. All Rights Reserved.
+
+#include "ConversionUtils/VolumeComponentTarget.h"
+
+#include "Components/BrushComponent.h"
+#include "DynamicMeshToMeshDescription.h"
+#include "Engine/Brush.h"
+#include "GameFramework/Volume.h"
+#include "MeshDescriptionToDynamicMesh.h"
+#include "MeshNormals.h"
+#include "StaticMeshAttributes.h"
+#include "ToolSetupUtil.h"
+
+
+FVolumeComponentTarget::FVolumeComponentTarget(UPrimitiveComponent* Component)
+	: FPrimitiveComponentTarget(Cast<UBrushComponent>(Component)) 
+{
+	// TODO: These should be user-configurable somewhere
+	VolumeToMeshOptions.bInWorldSpace = false;
+	VolumeToMeshOptions.bSetGroups = true;
+	VolumeToMeshOptions.bMergeVertices = true;
+	VolumeToMeshOptions.bAutoRepairMesh = true;
+	VolumeToMeshOptions.bOptimizeMesh = true;
+}
+
+FMeshDescription* FVolumeComponentTarget::GetMesh()
+{
+	FTransform transform = GetWorldTransform();
+	if (!ConvertedMeshDescription.IsValid())
+	{
+		UBrushComponent* BrushComponent = Cast<UBrushComponent>(Component);
+		if (!BrushComponent)
+		{
+			return nullptr;
+		}
+		AVolume* Volume = Cast<AVolume>(BrushComponent->GetOwner());
+		if (!Volume)
+		{
+			return nullptr;
+		}
+
+		// Note: We can go directly from a volume to a mesh description using GetBrushMesh() in
+		// Editor.h. However, that path doesn't assign polygroups to the result, which we
+		// typically want when using this target, hence the two-step path we use here.
+		// All of this may also someday be fixed if we allow targets to be converted directly
+		// to dynamic meshes rather than mesh descriptions, which is not currently possible
+		// because FDynamicMesh3 is a class specific to the geometry plugin.
+		FDynamicMesh3 DynamicMesh;
+		UE::Conversion::VolumeToDynamicMesh(Volume, DynamicMesh, GetVolumeToMeshOptions());
+		FMeshNormals::InitializeMeshToPerTriangleNormals(&DynamicMesh);
+
+		ConvertedMeshDescription = MakeUnique<FMeshDescription>();
+		FStaticMeshAttributes Attributes(*ConvertedMeshDescription);
+		Attributes.Register();
+		FDynamicMeshToMeshDescription Converter;
+		Converter.Convert(&DynamicMesh, *ConvertedMeshDescription);
+	}
+	FTransform transform2 = GetWorldTransform();
+	return ConvertedMeshDescription.Get();
+}
+
+void FVolumeComponentTarget::GetMaterialSet(FComponentMaterialSet& MaterialSetOut, bool bAssetMaterials) const
+{
+	MaterialSetOut.Materials.Reset();
+	if (IsValid() == false)
+	{
+		return;
+	}
+
+	UMaterialInterface* Material = ToolSetupUtil::GetDefaultEditVolumeMaterial();
+	if (Material)
+	{
+		MaterialSetOut.Materials.Add(Material);
+	}
+}
+
+bool FVolumeComponentTarget::HasSameSourceData(const FPrimitiveComponentTarget& OtherTarget) const
+{
+	if (IsValid())
+	{
+		const UBrushComponent* BrushComponent = Cast<UBrushComponent>(Component);
+		const UBrushComponent* OtherBrushComponent = Cast<UBrushComponent>(OtherTarget.Component);
+		return BrushComponent && BrushComponent == OtherBrushComponent;
+	}
+	return false;
+}
+
+void FVolumeComponentTarget::CommitMesh(const FCommitter& Committer)
+{
+	check(IsValid());
+
+	UBrushComponent* BrushComponent = Cast<UBrushComponent>(Component);
+	check(BrushComponent);
+	AVolume* Volume = Cast<AVolume>(BrushComponent->GetOwner());
+	check(Volume);
+
+	FCommitParams CommitParams;
+	CommitParams.MeshDescription = ConvertedMeshDescription.Get();
+	Committer(CommitParams);
+
+	// TODO: We should probably have a path directly from a mesh description to a volume
+	// rather than going through a dynamic mesh, though this may eventually not be necessary
+	// if our targets are changed to hold dynamic meshes to begin with.
+	FDynamicMesh3 DynamicMesh;
+	FMeshDescriptionToDynamicMesh Converter;
+	Converter.Convert(ConvertedMeshDescription.Get(), DynamicMesh);
+
+	FTransform Transform = GetWorldTransform();
+	UE::Conversion::DynamicMeshToVolume(DynamicMesh, Volume);
+
+	Volume->SetActorTransform(Transform);
+	Volume->PostEditChange();
+}
+
+bool FVolumeComponentTargetFactory::CanBuild(UActorComponent* Component)
+{
+	if (UBrushComponent* BrushComponent = Cast<UBrushComponent>(Component))
+	{
+		return Cast<AVolume>(BrushComponent->GetOwner()) != nullptr;
+	}
+	return false;
+}
+
+TUniquePtr<FPrimitiveComponentTarget> FVolumeComponentTargetFactory::Build(UPrimitiveComponent* Component)
+{
+	UBrushComponent* VolumeBrushComponent = Cast<UBrushComponent>(Component);
+	if (!VolumeBrushComponent)
+	{
+		return {};
+	}
+	AVolume* Volume = Cast<AVolume>(VolumeBrushComponent->GetOwner());
+	if (!Volume)
+	{
+		return {};
+	}
+
+	return TUniquePtr<FPrimitiveComponentTarget> { new FVolumeComponentTarget(Component) };
+}
diff --git a/Engine/Plugins/Experimental/MeshModelingToolset/Source/MeshModelingToolsEditorOnly/Private/ConversionUtils/VolumeDynamicMeshConversion.cpp b/Engine/Plugins/Experimental/MeshModelingToolset/Source/MeshModelingToolsEditorOnly/Private/ConversionUtils/VolumeDynamicMeshConversion.cpp
new file mode 100644
--- /dev/null
+++ b/Engine/Plugins/Experimental/MeshModelingToolset/Source/MeshModelingToolsEditorOnly/Private/ConversionUtils/VolumeDynamicMeshConversion.cpp
@@ -0,0 +1,414 @@
+// Copyright Epic Games, Inc. All Rights Reserved.
+
+#include "ConversionUtils/VolumeDynamicMeshConversion.h"
+
+#include "BSPOps.h"
+#include "DynamicMesh3.h"
+#include "CompGeom/PolygonTriangulation.h"
+#include "Engine/Classes/Engine/Polys.h"
+#include "GameFramework/Volume.h"
+#include "MeshBoundaryLoops.h"
+#include "MeshQueries.h"
+#include "MeshRegionBoundaryLoops.h"
+#include "Model.h"
+#include "Operations/MergeCoincidentMeshEdges.h"
+#include "Operations/MinimalHoleFiller.h"
+#include "Selections/MeshConnectedComponents.h"
+
+namespace UE {
+
+// In a namespace because even anonymous namespaces could collide with 
+namespace VolumeDynamicMeshConversionLocals
+{
+/** @return triangle aspect ratio transformed to be in [0,1] range */
+double UnitAspectRatio(const FVector3d& A, const FVector3d& B, const FVector3d& C)
+{
+	double AspectRatio = VectorUtil::AspectRatio(A, B, C);
+	return (AspectRatio > 1.0) ? FMathd::Clamp(1.0 / AspectRatio, 0.0, 1.0) : AspectRatio;
+}
+/** @return triangle aspect ratio transformed to be in [0,1] range */
+double UnitAspectRatio(const FDynamicMesh3& Mesh, int32 TriangleID)
+{
+	FVector3d A, B, C;
+	Mesh.GetTriVertices(TriangleID, A, B, C);
+	return UnitAspectRatio(A, B, C);
+}
+
+/**
+ * If both triangles on an edge are coplanar, we can arbitrarily flip the interior edge to
+ * improve triangle quality. Similarly if one triangle on an edge is degenerate, we can flip
+ * the edge without affecting the shape to try to remove it. This code does a single pass of
+ * such an optimization.
+ * Note: could be more efficient to do multiple passes internally, would save on the initial computation
+ */
+void PlanarFlipsOptimization(FDynamicMesh3& Mesh, double PlanarDotThresh = 0.99)
+{
+	struct FFlatEdge
+	{
+		int32 eid;
+		double MinAspect;
+	};
+
+	TArray<double> AspectRatios;
+	TArray<FVector3d> Normals;
+	AspectRatios.SetNum(Mesh.MaxTriangleID());
+	Normals.SetNum(Mesh.MaxTriangleID());
+	for (int32 tid : Mesh.TriangleIndicesItr())
+	{
+		FVector3d A, B, C;
+		Mesh.GetTriVertices(tid, A, B, C);
+		AspectRatios[tid] = UnitAspectRatio(A, B, C);
+		Normals[tid] = VectorUtil::Normal(A, B, C);
+	}
+
+	TArray<FFlatEdge> Flips;
+	for (int32 eid : Mesh.EdgeIndicesItr())
+	{
+		if (Mesh.IsBoundaryEdge(eid) == false)
+		{
+			FIndex2i EdgeT = Mesh.GetEdgeT(eid);
+			if (AspectRatios[EdgeT.A] < 0.01 && AspectRatios[EdgeT.B] < 0.01)
+			{
+				continue;		// if both are degenerate we can't fix by flipping edge between them
+			}
+			double MinAspect = FMathd::Min(AspectRatios[EdgeT.A], AspectRatios[EdgeT.B]);
+			double NormDot = Normals[EdgeT.A].Dot(Normals[EdgeT.B]);
+			if (NormDot > PlanarDotThresh)
+			{
+				Flips.Add({ eid, MinAspect });
+			}
+		}
+	}
+
+	Flips.Sort([&](const FFlatEdge& A, const FFlatEdge& B) { return A.MinAspect < B.MinAspect; });
+
+	for (int32 k = 0; k < Flips.Num(); ++k)
+	{
+		int32 eid = Flips[k].eid;
+		FIndex2i EdgeV = Mesh.GetEdgeV(eid);
+		int32 a = EdgeV.A, b = EdgeV.B;
+		FIndex2i EdgeT = Mesh.GetEdgeT(eid);
+		FIndex3i Tri0 = Mesh.GetTriangle(EdgeT.A), Tri1 = Mesh.GetTriangle(EdgeT.B);
+		int32 c = IndexUtil::OrientTriEdgeAndFindOtherVtx(a, b, Tri0);
+		int32 d = IndexUtil::FindTriOtherVtx(a, b, Tri1);
+
+		double AspectA = AspectRatios[EdgeT.A], AspectB = AspectRatios[EdgeT.B];
+		double Metric = FMathd::Min(AspectA, AspectB);
+		FVector3d Normal = (AspectA > AspectB) ? Normals[EdgeT.A] : Normals[EdgeT.B];
+
+		FVector3d A = Mesh.GetVertex(a), B = Mesh.GetVertex(b);
+		FVector3d C = Mesh.GetVertex(c), D = Mesh.GetVertex(d);
+
+		double FlipAspect1 = UnitAspectRatio(C, D, B);
+		double FlipAspect2 = UnitAspectRatio(D, C, A);
+		FVector3d FlipNormal1 = VectorUtil::Normal(C, D, B);
+		FVector3d FlipNormal2 = VectorUtil::Normal(D, C, A);
+		if (FlipNormal1.Dot(Normal) < PlanarDotThresh || FlipNormal2.Dot(Normal) < PlanarDotThresh)
+		{
+			continue;		// should not happen?
+		}
+
+		if (FMathd::Min(FlipAspect1, FlipAspect2) > Metric)
+		{
+			FDynamicMesh3::FEdgeFlipInfo FlipInfo;
+			if (Mesh.FlipEdge(eid, FlipInfo) == EMeshResult::Ok)
+			{
+				AspectRatios[EdgeT.A] = UnitAspectRatio(Mesh, EdgeT.A);
+				AspectRatios[EdgeT.B] = UnitAspectRatio(Mesh, EdgeT.B);
+
+				// safety check - if somehow we flipped the normal, flip it back
+				bool bInvertedNormal = (Mesh.GetTriNormal(EdgeT.A).Dot(Normal) < PlanarDotThresh) ||
+					(Mesh.GetTriNormal(EdgeT.B).Dot(Normal) < PlanarDotThresh);
+				if (bInvertedNormal)
+				{
+					UE_LOG(LogTemp, Warning, TEXT("VolumeDynamicMeshConversion::PlanarFlipsOptimization - Invalid Flip!"));
+					Mesh.FlipEdge(eid, FlipInfo);
+					AspectRatios[EdgeT.A] = UnitAspectRatio(Mesh, EdgeT.A);
+					AspectRatios[EdgeT.B] = UnitAspectRatio(Mesh, EdgeT.B);
+				}
+			}
+		}
+	}
+}
+}
+using namespace VolumeDynamicMeshConversionLocals;
+
+
+namespace Conversion {
+
+void VolumeToDynamicMesh(AVolume* Volume, FDynamicMesh3& Mesh, 
+	const FVolumeToMeshOptions& Options)
+{
+	Mesh.DiscardAttributes();
+	if (Options.bSetGroups)
+	{
+		Mesh.EnableTriangleGroups();
+	}
+
+	UModel* Model = Volume->Brush;
+	FTransform3d XForm = (Options.bInWorldSpace) ? FTransform3d(Volume->GetTransform()) : FTransform3d::Identity();
+
+	// Each "BspNode" is a planar polygon, triangulate each polygon and accumulate in a mesh.
+	// Note that this does not make any attempt to weld vertices/edges
+	for (const FBspNode& Node : Model->Nodes)
+	{
+		FVector3d Normal = (FVector3d)Node.Plane;
+		FFrame3d Plane(Node.Plane.W * Normal, Normal);
+
+		int32 NumVerts = (Node.NodeFlags & PF_TwoSided) ? Node.NumVertices / 2 : Node.NumVertices;  // ??
+		if (NumVerts > 0)
+		{
+			TArray<int32> VertIndices;
+			TArray<FVector2d> VertPositions2d;
+			VertIndices.SetNum(NumVerts);
+			VertPositions2d.SetNum(NumVerts);
+			for (int32 VertexIndex = 0; VertexIndex < NumVerts; ++VertexIndex)
+			{
+				const FVert& Vert = Model->Verts[Node.iVertPool + VertexIndex];
+				FVector3d Point = (FVector3d)Model->Points[Vert.pVertex];
+				Point = XForm.TransformPosition(Point);
+				VertIndices[VertexIndex] = Mesh.AppendVertex(Point);
+				VertPositions2d[VertexIndex] = Plane.ToPlaneUV(Point, 2);
+			}
+
+			TArray<FIndex3i> PolyTriangles;
+			PolygonTriangulation::TriangulateSimplePolygon(VertPositions2d, PolyTriangles);
+
+			int32 GroupID = FDynamicMesh3::InvalidID;
+			if (Options.bSetGroups)
+			{
+				GroupID = Mesh.AllocateTriangleGroup();
+			}
+
+			for (FIndex3i Tri : PolyTriangles)
+			{
+				// flip orientation here...
+				Mesh.AppendTriangle(VertIndices[Tri.A], VertIndices[Tri.C], VertIndices[Tri.B], GroupID);
+			}
+		}
+	}
+
+	if (Options.bMergeVertices)
+	{
+		// Merge the mesh edges to create a closed solid
+		double MinLen, MaxLen, AvgLen;
+		TMeshQueries<FDynamicMesh3>::EdgeLengthStats(Mesh, MinLen, MaxLen, AvgLen);
+		FMergeCoincidentMeshEdges Merge(&Mesh);
+		Merge.MergeVertexTolerance = FMathd::Max(Merge.MergeVertexTolerance, MinLen * 0.1);
+		Merge.Apply();
+
+		// If the mesh is not closed, the merge failed or the volume had cracks/holes. 
+		// Do trivial hole fills to ensure the output is solid   (really want autorepair here)
+		if (Mesh.IsClosed() == false && Options.bAutoRepairMesh)
+		{
+			FMeshBoundaryLoops BoundaryLoops(&Mesh, true);
+			for (FEdgeLoop& Loop : BoundaryLoops.Loops)
+			{
+				FMinimalHoleFiller Filler(&Mesh, Loop);
+				Filler.Fill();
+			}
+		}
+
+
+		// try to flip towards better triangles in planar areas, should reduce/remove degenerate geo
+		if (Options.bOptimizeMesh)
+		{
+			for (int32 k = 0; k < 5; ++k)
+			{
+				PlanarFlipsOptimization(Mesh);
+			}
+		}
+	}
+}
+
+void DynamicMeshToVolume(const FDynamicMesh3& InputMesh, AVolume* TargetVolume)
+{
+	TArray<FDynamicMeshFace> Faces;
+	GetPolygonFaces(InputMesh, Faces);
+	DynamicMeshToVolume(InputMesh, Faces, TargetVolume);
+}
+
+void DynamicMeshToVolume(const FDynamicMesh3& InputMesh, TArray<FDynamicMeshFace>& Faces, AVolume* TargetVolume)
+{
+	check(TargetVolume->Brush);
+
+	UModel* Model = TargetVolume->Brush;
+
+	Model->Modify();
+
+	Model->Initialize(TargetVolume);
+	UPolys* Polys = Model->Polys;
+
+	for (FDynamicMeshFace& Face : Faces)
+	{
+		int32 NumVertices = Face.BoundaryLoop.Num();
+		FVector Normal = (FVector)Face.Plane.Z();
+		FVector U = (FVector)Face.Plane.X();
+		FVector V = (FVector)Face.Plane.Y();
+		FPlane FacePlane = Face.Plane.ToFPlane();
+
+		// create FPoly. This is Editor-only and I'm not entirely sure we need it?
+		int32 PolyIndex = Polys->Element.Num();
+		FPoly NewPoly;
+		NewPoly.Base = (FVector)Face.BoundaryLoop[0];
+		NewPoly.Normal = Normal;
+		NewPoly.TextureU = U;
+		NewPoly.TextureV = V;
+		NewPoly.Vertices.SetNum(NumVertices);
+		for (int32 k = 0; k < NumVertices; ++k)
+		{
+			NewPoly.Vertices[k] = (FVector)Face.BoundaryLoop[k];
+		}
+		NewPoly.PolyFlags = 0;
+		NewPoly.iLink = NewPoly.iLinkSurf = NewPoly.iBrushPoly = -1;
+		NewPoly.SmoothingMask = 0;
+		Polys->Element.Add(NewPoly);
+
+		/*
+
+				// create points for this face in UModel::Points
+				// TODO: can we share points between faces?
+				TArray<int32> PointIndices;
+				PointIndices.SetNum(NumVertices);
+				for (int32 k = 0; k < NumVertices; ++k)
+				{
+					int32 NewIdx = Model->Points.Num();
+					Model->Points.Add((FVector)Face.BoundaryLoop[k]);
+					PointIndices[k] = NewIdx;
+				}
+				int32 BasePointIndex = PointIndices[0];
+
+				// create normal for this face in UModel::Vectors along with U and V direction vectors
+				int32 NormalIdx = Model->Vectors.Num();
+				Model->Vectors.Add(Normal);
+				int32 TextureUIdx = Model->Vectors.Num();
+				Model->Vectors.Add(U);
+				int32 TextureVIdx = Model->Vectors.Num();
+				Model->Vectors.Add(V);
+
+				// create FVerts for this face in UModel::Verts
+				int32 iVertPoolStart = Model->Verts.Num();
+				for (int32 k = 0; k < NumVertices; ++k)
+				{
+					FVert NewVert;
+					NewVert.pVertex = PointIndices[k];		// Index of vertex point.
+					NewVert.iSide = INDEX_NONE;				// If shared, index of unique side. Otherwise INDEX_NONE.
+					NewVert.ShadowTexCoord = FVector2D::ZeroVector;			// The vertex's shadow map coordinate.
+					NewVert.BackfaceShadowTexCoord = FVector2D::ZeroVector;	// The vertex's shadow map coordinate for the backface of the node.
+					Model->Verts.Add(NewVert);
+				}
+
+				// create Surf
+
+				int32 SurfIndex = Model->Surfs.Num();
+				FBspSurf NewSurf;
+				NewSurf.Material = nullptr;			// 4 Material.
+				NewSurf.PolyFlags = 0;				// 4 Polygon flags.
+				NewSurf.pBase = BasePointIndex;		// 4 Polygon & texture base point index (where U,V==0,0).
+				NewSurf.vNormal = NormalIdx;		// 4 Index to polygon normal.
+				NewSurf.vTextureU = TextureUIdx;	// 4 Texture U-vector index.
+				NewSurf.vTextureV = TextureVIdx;	// 4 Texture V-vector index.
+				//NewSurf.iBrushPoly = PolyIndex;		// 4 Editor brush polygon index.
+				NewSurf.iBrushPoly = -1;
+				//NewSurf.Actor = NewVolume;			// 4 Brush actor owning this Bsp surface.
+				NewSurf.Actor = nullptr;
+				NewSurf.Plane = FacePlane;			// 16 The plane this surface lies on.
+				Model->Surfs.Add(NewSurf);
+
+
+				// create nodes for this face in UModel::Nodes
+
+				FBspNode NewNode;
+				NewNode.Plane = FacePlane;					// 16 Plane the node falls into (X, Y, Z, W).
+				NewNode.iVertPool = iVertPoolStart;			// 4  Index of first vertex in vertex pool, =iTerrain if NumVertices==0 and NF_TerrainFront.
+				NewNode.iSurf = SurfIndex;					// 4  Index to surface information.
+				NewNode.iVertexIndex = INDEX_NONE;			// The index of the node's first vertex in the UModel's vertex buffer.
+															// This is initialized by UModel::UpdateVertices()
+				NewNode.NumVertices = NumVertices;			// 1  Number of vertices in node.
+				NewNode.NodeFlags = 0;						// 1  Node flags.
+				Model->Nodes.Add(NewNode);
+		*/
+	}
+
+	// requires ed
+	FBSPOps::csgPrepMovingBrush(TargetVolume);
+
+	// do we need to do this?
+	TargetVolume->MarkPackageDirty();
+}
+
+
+void GetPolygonFaces(const FDynamicMesh3& InputMesh, TArray<FDynamicMeshFace>& Faces)
+{
+	Faces.SetNum(0);
+
+	FMeshNormals Normals(&InputMesh);
+	Normals.ComputeTriangleNormals();
+
+	double PlanarTolerance = FMathf::ZeroTolerance;
+
+	FMeshConnectedComponents Components(&InputMesh);
+	Components.FindConnectedTriangles([&](int32 Triangle0, int32 Triangle1)
+		{
+			FVector3d Origin = InputMesh.GetTriCentroid(Triangle0);
+			FVector3d Normal = Normals[Triangle0];
+
+			FVector3d A, B, C;
+			InputMesh.GetTriVertices(Triangle1, A, B, C);;
+			double DistA = (A - Origin).Dot(Normal);
+			double DistB = (B - Origin).Dot(Normal);
+			double DistC = (C - Origin).Dot(Normal);
+			double MaxDist = FMathd::Max3(FMathd::Abs(DistA), FMathd::Abs(DistB), FMathd::Abs(DistC));
+
+			return MaxDist < PlanarTolerance;
+		});
+
+	for (const FMeshConnectedComponents::FComponent& Component : Components)
+	{
+		FVector3d FaceNormal = Normals[Component.Indices[0]];
+
+		FMeshRegionBoundaryLoops Loops(&InputMesh, Component.Indices);
+		for (const FEdgeLoop& Loop : Loops.Loops)
+		{
+			FDynamicMeshFace Face;
+
+			FVector3d AvgPos(0, 0, 0);
+			for (int32 vid : Loop.Vertices)
+			{
+				FVector3d Position = InputMesh.GetVertex(vid);
+				Face.BoundaryLoop.Add(Position);
+				AvgPos += Position;
+			}
+			AvgPos /= (double)Loop.Vertices.Num();
+			Algo::Reverse(Face.BoundaryLoop);
+
+			Face.Plane = FFrame3d(AvgPos, FaceNormal);
+
+			Faces.Add(Face);
+		}
+	}
+}
+
+void GetTriangleFaces(const FDynamicMesh3& InputMesh, TArray<FDynamicMeshFace>& Faces)
+{
+	Faces.SetNum(0);
+
+	for (int32 tid : InputMesh.TriangleIndicesItr())
+	{
+		FVector3d A, B, C;
+		InputMesh.GetTriVertices(tid, A, B, C);
+		FVector3d Centroid, Normal; double Area;
+		InputMesh.GetTriInfo(tid, Normal, Area, Centroid);
+
+		FDynamicMeshFace Face;
+		Face.Plane = FFrame3d(Centroid, Normal);
+		Face.BoundaryLoop.Add(A);
+		Face.BoundaryLoop.Add(C);
+		Face.BoundaryLoop.Add(B);
+
+		Faces.Add(Face);
+	}
+}
+
+}//end namespace UE::Conversion
+}//end namespace UE
\ No newline at end of file
diff --git a/Engine/Plugins/Experimental/MeshModelingToolset/Source/MeshModelingToolsEditorOnly/Private/MeshModelingToolsEditorOnly.cpp b/Engine/Plugins/Experimental/MeshModelingToolset/Source/MeshModelingToolsEditorOnly/Private/MeshModelingToolsEditorOnly.cpp
--- a/Engine/Plugins/Experimental/MeshModelingToolset/Source/MeshModelingToolsEditorOnly/Private/MeshModelingToolsEditorOnly.cpp
+++ b/Engine/Plugins/Experimental/MeshModelingToolset/Source/MeshModelingToolsEditorOnly/Private/MeshModelingToolsEditorOnly.cpp
@@ -2,11 +2,14 @@
 
 #include "MeshModelingToolsEditorOnly.h"
 
+#include "ConversionUtils/VolumeComponentTarget.h"
+
 #define LOCTEXT_NAMESPACE "FMeshModelingToolsEditorOnlyModule"
 
 void FMeshModelingToolsEditorOnlyModule::StartupModule()
 {
 	// This code will execute after your module is loaded into memory; the exact timing is specified in the .uplugin file per-module
+	AddComponentTargetFactory(TUniquePtr<FComponentTargetFactory>{new FVolumeComponentTargetFactory{} });
 }
 
 void FMeshModelingToolsEditorOnlyModule::ShutdownModule()
diff --git a/Engine/Plugins/Experimental/MeshModelingToolset/Source/MeshModelingToolsEditorOnly/Private/MeshToVolumeTool.cpp b/Engine/Plugins/Experimental/MeshModelingToolset/Source/MeshModelingToolsEditorOnly/Private/MeshToVolumeTool.cpp
--- a/Engine/Plugins/Experimental/MeshModelingToolset/Source/MeshModelingToolsEditorOnly/Private/MeshToVolumeTool.cpp
+++ b/Engine/Plugins/Experimental/MeshModelingToolset/Source/MeshModelingToolsEditorOnly/Private/MeshToVolumeTool.cpp
@@ -11,7 +11,7 @@
 #include "Util/ColorConstants.h"
 #include "ToolSetupUtil.h"
 #include "Selections/MeshConnectedComponents.h"
-
+#include "Selection/ToolSelectionUtil.h"
 
 #include "Engine/Classes/Engine/BlockingVolume.h"
 #include "Engine/Classes/Components/BrushComponent.h"
@@ -29,7 +29,9 @@
 
 bool UMeshToVolumeToolBuilder::CanBuildTool(const FToolBuilderState& SceneState) const
 {
-	return ToolBuilderUtil::CountComponents(SceneState, CanMakeComponentTarget) == 1;
+	// We don't want to allow this tool to run on selected volumes
+	return ToolBuilderUtil::CountSelectedActorsOfType<AVolume>(SceneState) == 0 
+		&& ToolBuilderUtil::CountComponents(SceneState, CanMakeComponentTarget) == 1;
 }
 
 UInteractiveTool* UMeshToVolumeToolBuilder::BuildTool(const FToolBuilderState& SceneState) const
@@ -100,125 +102,6 @@ void UMeshToVolumeTool::Setup()
 		EToolMessageLevel::UserNotification);
 }
 
-
-
-
-
-void UMeshToVolumeTool::BakeToVolume(AVolume* TargetVolume)
-{
-	check(TargetVolume->Brush);
-
-	UModel* Model = TargetVolume->Brush;
-
-	Model->Modify();
-
-	Model->Initialize(TargetVolume);
-	UPolys* Polys = Model->Polys;
-
-	for (FModelFace& Face : Faces)
-	{
-		int32 NumVertices = Face.BoundaryLoop.Num();
-		FVector Normal = (FVector)Face.Plane.Z();
-		FVector U = (FVector)Face.Plane.X();
-		FVector V = (FVector)Face.Plane.Y();
-		FPlane FacePlane = Face.Plane.ToFPlane();
-
-		// create FPoly. This is Editor-only and I'm not entirely sure we need it?
-		int32 PolyIndex = Polys->Element.Num();
-		FPoly NewPoly;
-		NewPoly.Base = (FVector)Face.BoundaryLoop[0];
-		NewPoly.Normal = Normal;
-		NewPoly.TextureU = U;
-		NewPoly.TextureV = V;
-		NewPoly.Vertices.SetNum(NumVertices);
-		for (int32 k = 0; k < NumVertices; ++k)
-		{
-			NewPoly.Vertices[k] = (FVector)Face.BoundaryLoop[k];
-		}
-		NewPoly.PolyFlags = 0;
-		NewPoly.iLink = NewPoly.iLinkSurf = NewPoly.iBrushPoly = -1;
-		NewPoly.SmoothingMask = 0;
-		Polys->Element.Add(NewPoly);
-
-/*
-
-		// create points for this face in UModel::Points
-		// TODO: can we share points between faces?
-		TArray<int32> PointIndices;
-		PointIndices.SetNum(NumVertices);
-		for (int32 k = 0; k < NumVertices; ++k)
-		{
-			int32 NewIdx = Model->Points.Num();
-			Model->Points.Add((FVector)Face.BoundaryLoop[k]);
-			PointIndices[k] = NewIdx;
-		}
-		int32 BasePointIndex = PointIndices[0];
-
-		// create normal for this face in UModel::Vectors along with U and V direction vectors
-		int32 NormalIdx = Model->Vectors.Num();
-		Model->Vectors.Add(Normal);
-		int32 TextureUIdx = Model->Vectors.Num();
-		Model->Vectors.Add(U);
-		int32 TextureVIdx = Model->Vectors.Num();
-		Model->Vectors.Add(V);
-
-		// create FVerts for this face in UModel::Verts
-		int32 iVertPoolStart = Model->Verts.Num();
-		for (int32 k = 0; k < NumVertices; ++k)
-		{
-			FVert NewVert;
-			NewVert.pVertex = PointIndices[k];		// Index of vertex point.
-			NewVert.iSide = INDEX_NONE;				// If shared, index of unique side. Otherwise INDEX_NONE.
-			NewVert.ShadowTexCoord = FVector2D::ZeroVector;			// The vertex's shadow map coordinate.
-			NewVert.BackfaceShadowTexCoord = FVector2D::ZeroVector;	// The vertex's shadow map coordinate for the backface of the node.
-			Model->Verts.Add(NewVert);
-		}
-
-		// create Surf
-
-		int32 SurfIndex = Model->Surfs.Num();
-		FBspSurf NewSurf;
-		NewSurf.Material = nullptr;			// 4 Material.
-		NewSurf.PolyFlags = 0;				// 4 Polygon flags.
-		NewSurf.pBase = BasePointIndex;		// 4 Polygon & texture base point index (where U,V==0,0).
-		NewSurf.vNormal = NormalIdx;		// 4 Index to polygon normal.
-		NewSurf.vTextureU = TextureUIdx;	// 4 Texture U-vector index.
-		NewSurf.vTextureV = TextureVIdx;	// 4 Texture V-vector index.
-		//NewSurf.iBrushPoly = PolyIndex;		// 4 Editor brush polygon index.
-		NewSurf.iBrushPoly = -1;
-		//NewSurf.Actor = NewVolume;			// 4 Brush actor owning this Bsp surface.
-		NewSurf.Actor = nullptr;
-		NewSurf.Plane = FacePlane;			// 16 The plane this surface lies on.
-		Model->Surfs.Add(NewSurf);
-
-
-		// create nodes for this face in UModel::Nodes
-
-		FBspNode NewNode;
-		NewNode.Plane = FacePlane;					// 16 Plane the node falls into (X, Y, Z, W).
-		NewNode.iVertPool = iVertPoolStart;			// 4  Index of first vertex in vertex pool, =iTerrain if NumVertices==0 and NF_TerrainFront.
-		NewNode.iSurf = SurfIndex;					// 4  Index to surface information.
-		NewNode.iVertexIndex = INDEX_NONE;			// The index of the node's first vertex in the UModel's vertex buffer.
-													// This is initialized by UModel::UpdateVertices()
-		NewNode.NumVertices = NumVertices;			// 1  Number of vertices in node.
-		NewNode.NodeFlags = 0;						// 1  Node flags.
-		Model->Nodes.Add(NewNode);
-*/
-	}
-
-	// requires ed
-	FBSPOps::csgPrepMovingBrush(TargetVolume);
-
-	// do we need to do this?
-	TargetVolume->MarkPackageDirty();
-}
-
-
-
-
-
-
-
 void UMeshToVolumeTool::Shutdown(EToolShutdownType ShutdownType)
 {
 	Settings->SaveProperties(this);
@@ -270,10 +153,11 @@ void UMeshToVolumeTool::Shutdown(EToolShutdownType ShutdownType)
 			TargetVolume->GetBrushComponent()->Modify();
 		}
 
-		BakeToVolume(TargetVolume);
+		UE::Conversion::DynamicMeshToVolume(InputMesh, Faces, TargetVolume);
 		TargetVolume->SetActorTransform(SetTransform);
 		TargetVolume->PostEditChange();
 
+		ToolSelectionUtil::SetNewActorSelection(GetToolManager(), TargetVolume);
 
 		TArray<AActor*> Actors;
 		Actors.Add(ComponentTarget->GetOwnerActor());
@@ -285,7 +169,6 @@ void UMeshToVolumeTool::Shutdown(EToolShutdownType ShutdownType)
 	}
 }
 
-
 void UMeshToVolumeTool::OnTick(float DeltaTime)
 {
 	if (bVolumeValid == false)
@@ -310,7 +193,6 @@ bool UMeshToVolumeTool::CanAccept() const
 }
 
 
-
 void UMeshToVolumeTool::UpdateLineSet()
 {
 	FColor BoundaryEdgeColor(240, 15, 15);
@@ -318,7 +200,7 @@ void UMeshToVolumeTool::UpdateLineSet()
 	float BoundaryEdgeDepthBias = 2.0f;
 
 	VolumeEdgesSet->Clear();
-	for (const FModelFace& Face : Faces)
+	for (const UE::Conversion::FDynamicMeshFace& Face : Faces)
 	{
 		int32 NumV = Face.BoundaryLoop.Num();
 		for (int32 k = 0; k < NumV; ++k)
@@ -331,108 +213,20 @@ void UMeshToVolumeTool::UpdateLineSet()
 
 }
 
-
-
 void UMeshToVolumeTool::RecalculateVolume()
 {
 	if (Settings->ConversionMode == EMeshToVolumeMode::MinimalPolygons)
 	{
-		RecalculateVolume_Polygons();
+		UE::Conversion::GetPolygonFaces(InputMesh, Faces);
 	}
 	else
 	{
-		RecalculateVolume_Triangles();
-	}
-}
-
-
-void UMeshToVolumeTool::RecalculateVolume_Polygons()
-{
-	Faces.SetNum(0);
-
-	FMeshNormals Normals(&InputMesh);
-	Normals.ComputeTriangleNormals();
-
-	double PlanarTolerance = FMathf::ZeroTolerance;
-
-	FMeshConnectedComponents Components(&InputMesh);
-	Components.FindConnectedTriangles([&](int32 Triangle0, int32 Triangle1)
-	{
-		FVector3d Origin = InputMesh.GetTriCentroid(Triangle0);
-		FVector3d Normal = Normals[Triangle0];
-
-		FVector3d A, B, C;
-		InputMesh.GetTriVertices(Triangle1, A, B, C);;
-		double DistA = (A - Origin).Dot(Normal);
-		double DistB = (B - Origin).Dot(Normal);
-		double DistC = (C - Origin).Dot(Normal);
-		double MaxDist = FMathd::Max3(FMathd::Abs(DistA), FMathd::Abs(DistB), FMathd::Abs(DistC));
-
-		return MaxDist < PlanarTolerance;
-	});
-
-	for (const FMeshConnectedComponents::FComponent& Component : Components)
-	{
-		FVector3d FaceNormal = Normals[Component.Indices[0]];
-
-		FMeshRegionBoundaryLoops Loops(&InputMesh, Component.Indices);
-		for (const FEdgeLoop& Loop : Loops.Loops)
-		{
-			FModelFace Face;
-
-			FVector3d AvgPos(0, 0, 0);
-			for (int32 vid : Loop.Vertices)
-			{
-				FVector3d Position = InputMesh.GetVertex(vid);
-				Face.BoundaryLoop.Add(Position);
-				AvgPos += Position;
-			}
-			AvgPos /= (double)Loop.Vertices.Num();
-			Algo::Reverse(Face.BoundaryLoop);
-
-			Face.Plane = FFrame3d(AvgPos, FaceNormal);
-
-			Faces.Add(Face);
-		}
-	}
-
-	UpdateLineSet();
-
-	bVolumeValid = true;
-}
-
-
-
-
-
-
-void UMeshToVolumeTool::RecalculateVolume_Triangles()
-{
-	Faces.SetNum(0);
-
-	for (int32 tid : InputMesh.TriangleIndicesItr())
-	{
-		FVector3d A, B, C;
-		InputMesh.GetTriVertices(tid, A, B, C);
-		FVector3d Centroid, Normal; double Area;
-		InputMesh.GetTriInfo(tid, Normal, Area, Centroid);
-
-		FModelFace Face;
-		Face.Plane = FFrame3d(Centroid, Normal);
-		Face.BoundaryLoop.Add(A);
-		Face.BoundaryLoop.Add(C);
-		Face.BoundaryLoop.Add(B);
-
-		Faces.Add(Face);
+		UE::Conversion::GetTriangleFaces(InputMesh, Faces);
 	}
 
 	UpdateLineSet();
-
 	bVolumeValid = true;
 }
 
 
-
-
-
 #undef LOCTEXT_NAMESPACE
diff --git a/Engine/Plugins/Experimental/MeshModelingToolset/Source/MeshModelingToolsEditorOnly/Private/VolumeToMeshTool.cpp b/Engine/Plugins/Experimental/MeshModelingToolset/Source/MeshModelingToolsEditorOnly/Private/VolumeToMeshTool.cpp
--- a/Engine/Plugins/Experimental/MeshModelingToolset/Source/MeshModelingToolsEditorOnly/Private/VolumeToMeshTool.cpp
+++ b/Engine/Plugins/Experimental/MeshModelingToolset/Source/MeshModelingToolsEditorOnly/Private/VolumeToMeshTool.cpp
@@ -4,6 +4,7 @@
 #include "InteractiveToolManager.h"
 #include "ToolBuilderUtil.h"
 
+#include "ConversionUtils/VolumeDynamicMeshConversion.h"
 #include "DynamicMesh3.h"
 #include "MeshNormals.h"
 #include "DynamicMeshEditor.h"
@@ -197,236 +198,17 @@ void UVolumeToMeshTool::UpdateLineSet()
 	}
 }
 
-
-
-
-
-/** @return triangle aspect ratio transformed to be in [0,1] range */
-double UnitAspectRatio(const FVector3d& A, const FVector3d& B, const FVector3d& C)
-{
-	double AspectRatio = VectorUtil::AspectRatio(A, B, C);
-	return (AspectRatio > 1.0) ? FMathd::Clamp(1.0 / AspectRatio, 0.0, 1.0) : AspectRatio;
-}
-/** @return triangle aspect ratio transformed to be in [0,1] range */
-double UnitAspectRatio(const FDynamicMesh3& Mesh, int32 TriangleID)
-{
-	FVector3d A, B, C;
-	Mesh.GetTriVertices(TriangleID, A, B, C);
-	return UnitAspectRatio(A, B, C);
-}
-
-
-
-/**
- * If both triangles on an edge are coplanar, we can arbitrarily flip the interior edge to
- * improve triangle quality. Similarly if one triangle on an edge is degenerate, we can flip
- * the edge without affecting the shape to try to remove it. This code does a single pass of
- * such an optimization.
- * Note: could be more efficient to do multiple passes internally, would save on the initial computation
- */
-void PlanarFlipsOptimization(FDynamicMesh3& Mesh, double PlanarDotThresh = 0.99)
-{
-	struct FFlatEdge
-	{
-		int32 eid;
-		double MinAspect;
-	};
-
-	TArray<double> AspectRatios;
-	TArray<FVector3d> Normals;
-	AspectRatios.SetNum(Mesh.MaxTriangleID());
-	Normals.SetNum(Mesh.MaxTriangleID());
-	for (int32 tid : Mesh.TriangleIndicesItr())
-	{
-		FVector3d A, B, C;
-		Mesh.GetTriVertices(tid, A, B, C);
-		AspectRatios[tid] = UnitAspectRatio(A, B, C);
-		Normals[tid] = VectorUtil::Normal(A, B, C);
-	}
-
-	TArray<FFlatEdge> Flips;
-	for (int32 eid : Mesh.EdgeIndicesItr())
-	{
-		if (Mesh.IsBoundaryEdge(eid) == false)
-		{
-			FIndex2i EdgeT = Mesh.GetEdgeT(eid);
-			if (AspectRatios[EdgeT.A] < 0.01 && AspectRatios[EdgeT.B] < 0.01)
-			{
-				continue;		// if both are degenerate we can't fix by flipping edge between them
-			}
-			double MinAspect = FMathd::Min(AspectRatios[EdgeT.A], AspectRatios[EdgeT.B]);
-			double NormDot = Normals[EdgeT.A].Dot(Normals[EdgeT.B]);
-			if (NormDot > PlanarDotThresh)
-			{
-				Flips.Add({ eid, MinAspect });
-			}
-		}
-	}
-
-	Flips.Sort([&](const FFlatEdge& A, const FFlatEdge& B) { return A.MinAspect < B.MinAspect; });
-
-	for (int32 k = 0; k < Flips.Num(); ++k)
-	{
-		int32 eid = Flips[k].eid;
-		FIndex2i EdgeV = Mesh.GetEdgeV(eid);
-		int32 a = EdgeV.A, b = EdgeV.B;
-		FIndex2i EdgeT = Mesh.GetEdgeT(eid);
-		FIndex3i Tri0 = Mesh.GetTriangle(EdgeT.A), Tri1 = Mesh.GetTriangle(EdgeT.B);
-		int32 c = IndexUtil::OrientTriEdgeAndFindOtherVtx(a, b, Tri0);
-		int32 d = IndexUtil::FindTriOtherVtx(a, b, Tri1);
-
-		double AspectA = AspectRatios[EdgeT.A], AspectB = AspectRatios[EdgeT.B];
-		double Metric = FMathd::Min(AspectA, AspectB);
-		FVector3d Normal = (AspectA > AspectB) ? Normals[EdgeT.A] : Normals[EdgeT.B];
-
-		FVector3d A = Mesh.GetVertex(a), B = Mesh.GetVertex(b);
-		FVector3d C = Mesh.GetVertex(c), D = Mesh.GetVertex(d);
-
-		double FlipAspect1 = UnitAspectRatio(C, D, B);
-		double FlipAspect2 = UnitAspectRatio(D, C, A);
-		FVector3d FlipNormal1 = VectorUtil::Normal(C, D, B);
-		FVector3d FlipNormal2 = VectorUtil::Normal(D, C, A);
-		if (FlipNormal1.Dot(Normal) < PlanarDotThresh || FlipNormal2.Dot(Normal) < PlanarDotThresh)
-		{
-			continue;		// should not happen?
-		}
-
-		if (FMathd::Min(FlipAspect1, FlipAspect2) > Metric)
-		{
-			FDynamicMesh3::FEdgeFlipInfo FlipInfo;
-			if (Mesh.FlipEdge(eid, FlipInfo) == EMeshResult::Ok)
-			{
-				AspectRatios[EdgeT.A] = UnitAspectRatio(Mesh, EdgeT.A);
-				AspectRatios[EdgeT.B] = UnitAspectRatio(Mesh, EdgeT.B);
-
-				// safety check - if somehow we flipped the normal, flip it back
-				bool bInvertedNormal = (Mesh.GetTriNormal(EdgeT.A).Dot(Normal) < PlanarDotThresh) ||
-					(Mesh.GetTriNormal(EdgeT.B).Dot(Normal) < PlanarDotThresh);
-				if (bInvertedNormal)
-				{
-					UE_LOG(LogTemp, Warning, TEXT("UE::Water::PlanarFlipsOptimization - Invalid Flip!"));
-					Mesh.FlipEdge(eid, FlipInfo);
-					AspectRatios[EdgeT.A] = UnitAspectRatio(Mesh, EdgeT.A);
-					AspectRatios[EdgeT.B] = UnitAspectRatio(Mesh, EdgeT.B);
-				}
-			}
-		}
-	}
-}
-
-
-
-struct FVolumeToMeshOptions
-{
-	bool bInWorldSpace = false;
-	bool bSetGroups = true;
-	bool bMergeVertices = true;
-	bool bAutoRepairMesh = true;
-	bool bOptimizeMesh = true;
-};
-
-
-/**
- * Extracts a FDynamicMesh3 from an AVolume
- * The output mesh is in World Space.
- */
-void ExtractMesh(AVolume* Volume, FDynamicMesh3& Mesh, const FVolumeToMeshOptions& Options)
-{
-	Mesh.DiscardAttributes();
-	if (Options.bSetGroups)
-	{
-		Mesh.EnableTriangleGroups();
-	}
-
-	UModel* Model = Volume->Brush;
-	FTransform3d XForm = (Options.bInWorldSpace) ? FTransform3d(Volume->GetTransform()) : FTransform3d::Identity();
-
-	// Each "BspNode" is a planar polygon, triangulate each polygon and accumulate in a mesh.
-	// Note that this does not make any attempt to weld vertices/edges
-	for (const FBspNode& Node : Model->Nodes)
-	{
-		FVector3d Normal = (FVector3d)Node.Plane;
-		FFrame3d Plane(Node.Plane.W * Normal, Normal);
-
-		int32 NumVerts = (Node.NodeFlags & PF_TwoSided) ? Node.NumVertices / 2 : Node.NumVertices;  // ??
-		if (NumVerts > 0)
-		{
-			TArray<int32> VertIndices;
-			TArray<FVector2d> VertPositions2d;
-			VertIndices.SetNum(NumVerts);
-			VertPositions2d.SetNum(NumVerts);
-			for (int32 VertexIndex = 0; VertexIndex < NumVerts; ++VertexIndex)
-			{
-				const FVert& Vert = Model->Verts[Node.iVertPool + VertexIndex];
-				FVector3d Point = (FVector3d)Model->Points[Vert.pVertex];
-				Point = XForm.TransformPosition(Point);
-				VertIndices[VertexIndex] = Mesh.AppendVertex(Point);
-				VertPositions2d[VertexIndex] = Plane.ToPlaneUV(Point, 2);
-			}
-
-			TArray<FIndex3i> PolyTriangles;
-			PolygonTriangulation::TriangulateSimplePolygon(VertPositions2d, PolyTriangles);
-
-			int32 GroupID = FDynamicMesh3::InvalidID;
-			if (Options.bSetGroups)
-			{
-				GroupID = Mesh.AllocateTriangleGroup();
-			}
-
-			for (FIndex3i Tri : PolyTriangles)
-			{
-				// flip orientation here...
-				Mesh.AppendTriangle(VertIndices[Tri.A], VertIndices[Tri.C], VertIndices[Tri.B], GroupID);
-			}
-		}
-	}
-
-	if (Options.bMergeVertices)
-	{
-		// Merge the mesh edges to create a closed solid
-		double MinLen, MaxLen, AvgLen;
-		TMeshQueries<FDynamicMesh3>::EdgeLengthStats(Mesh, MinLen, MaxLen, AvgLen);
-		FMergeCoincidentMeshEdges Merge(&Mesh);
-		Merge.MergeVertexTolerance = FMathd::Max(Merge.MergeVertexTolerance, MinLen * 0.1);
-		Merge.Apply();
-
-		// If the mesh is not closed, the merge failed or the volume had cracks/holes. 
-		// Do trivial hole fills to ensure the output is solid   (really want autorepair here)
-		if (Mesh.IsClosed() == false && Options.bAutoRepairMesh)
-		{
-			FMeshBoundaryLoops BoundaryLoops(&Mesh, true);
-			for (FEdgeLoop& Loop : BoundaryLoops.Loops)
-			{
-				FMinimalHoleFiller Filler(&Mesh, Loop);
-				Filler.Fill();
-			}
-		}
-
-
-		// try to flip towards better triangles in planar areas, should reduce/remove degenerate geo
-		if (Options.bOptimizeMesh)
-		{
-			for (int32 k = 0; k < 5; ++k)
-			{
-				PlanarFlipsOptimization(Mesh);
-			}
-		}
-	}
-}
-
-
-
 void UVolumeToMeshTool::RecalculateMesh()
 {
 	if (TargetVolume.IsValid())
 	{
-		FVolumeToMeshOptions Options;
+		UE::Conversion::FVolumeToMeshOptions Options;
 		Options.bMergeVertices = Settings->bWeldEdges;
 		Options.bAutoRepairMesh = Settings->bAutoRepair;
 		Options.bOptimizeMesh = Settings->bOptimizeMesh;
 
 		CurrentMesh = FDynamicMesh3(EMeshComponents::FaceGroups);
-		ExtractMesh(TargetVolume.Get(), CurrentMesh, Options);
+		UE::Conversion::VolumeToDynamicMesh(TargetVolume.Get(), CurrentMesh, Options);
 		FMeshNormals::InitializeMeshToPerTriangleNormals(&CurrentMesh);
 		PreviewMesh->UpdatePreview(&CurrentMesh);
 	}
diff --git a/Engine/Plugins/Experimental/MeshModelingToolset/Source/MeshModelingToolsEditorOnly/Public/ConversionUtils/VolumeComponentTarget.h b/Engine/Plugins/Experimental/MeshModelingToolset/Source/MeshModelingToolsEditorOnly/Public/ConversionUtils/VolumeComponentTarget.h
new file mode 100644
--- /dev/null
+++ b/Engine/Plugins/Experimental/MeshModelingToolset/Source/MeshModelingToolsEditorOnly/Public/ConversionUtils/VolumeComponentTarget.h
@@ -0,0 +1,38 @@
+// Copyright Epic Games, Inc. All Rights Reserved.
+
+#pragma once
+
+#include "CoreMinimal.h"
+
+#include "ComponentSourceInterfaces.h"
+#include "ConversionUtils/VolumeDynamicMeshConversion.h"
+#include "MeshDescription.h"
+
+class MESHMODELINGTOOLSEDITORONLY_API FVolumeComponentTargetFactory : public FComponentTargetFactory
+{
+public:
+	bool CanBuild(UActorComponent* Candidate) override;
+	TUniquePtr<FPrimitiveComponentTarget> Build(UPrimitiveComponent* PrimitiveComponent) override;
+};
+
+class MESHMODELINGTOOLSEDITORONLY_API FVolumeComponentTarget : public FPrimitiveComponentTarget
+{
+public:
+
+	FVolumeComponentTarget(UPrimitiveComponent* Component);
+
+	FMeshDescription* GetMesh() override;
+	void CommitMesh(const FCommitter&) override;
+
+	virtual void GetMaterialSet(FComponentMaterialSet& MaterialSetOut, bool bAssetMaterials) const override;
+	virtual void CommitMaterialSetUpdate(const FComponentMaterialSet& MaterialSet, bool bApplyToAsset) override {};
+
+	virtual bool HasSameSourceData(const FPrimitiveComponentTarget& OtherTarget) const override;
+
+	const UE::Conversion::FVolumeToMeshOptions& GetVolumeToMeshOptions() { return VolumeToMeshOptions; }
+
+protected:
+	TUniquePtr<FMeshDescription> ConvertedMeshDescription;
+
+	UE::Conversion::FVolumeToMeshOptions VolumeToMeshOptions;
+};
diff --git a/Engine/Plugins/Experimental/MeshModelingToolset/Source/MeshModelingToolsEditorOnly/Public/ConversionUtils/VolumeDynamicMeshConversion.h b/Engine/Plugins/Experimental/MeshModelingToolset/Source/MeshModelingToolsEditorOnly/Public/ConversionUtils/VolumeDynamicMeshConversion.h
new file mode 100644
--- /dev/null
+++ b/Engine/Plugins/Experimental/MeshModelingToolset/Source/MeshModelingToolsEditorOnly/Public/ConversionUtils/VolumeDynamicMeshConversion.h
@@ -0,0 +1,53 @@
+// Copyright Epic Games, Inc. All Rights Reserved.
+
+#pragma once
+
+#include "CoreMinimal.h"
+#include "FrameTypes.h"
+#include "VectorTypes.h"
+
+class AVolume;
+class FDynamicMesh3;
+
+namespace UE {
+namespace Conversion { 
+
+struct FVolumeToMeshOptions
+{
+	bool bInWorldSpace = false;
+	bool bSetGroups = true;
+	bool bMergeVertices = true;
+	bool bAutoRepairMesh = true;
+	bool bOptimizeMesh = true;
+};
+
+struct FDynamicMeshFace
+{
+	FFrame3d Plane;
+	TArray<FVector3d> BoundaryLoop;
+};
+
+/**
+ * Converts a volume to a dynamic mesh. Does not initialize normals and does not delete the volume.
+ */
+void MESHMODELINGTOOLSEDITORONLY_API VolumeToDynamicMesh(AVolume* Volume, FDynamicMesh3& Mesh, const FVolumeToMeshOptions& Options);
+
+/**
+ * Gets an array of face objects that can be used to convert a dynamic mesh to a volume. This version tries to
+ * merge coplanar triangles into polygons.
+ */
+void MESHMODELINGTOOLSEDITORONLY_API GetPolygonFaces(const FDynamicMesh3& InputMesh, TArray<FDynamicMeshFace>& Faces);
+/**
+ * Gets an array of face objects that can be used to convert a dynamic mesh to a volume. This version makes
+ * each triangle its own face.
+ */
+void MESHMODELINGTOOLSEDITORONLY_API GetTriangleFaces(const FDynamicMesh3& InputMesh, TArray<FDynamicMeshFace>& Faces);
+
+/**
+ * Converts a dynamic mesh to a volume.
+ */
+void MESHMODELINGTOOLSEDITORONLY_API DynamicMeshToVolume(const FDynamicMesh3& InputMesh, AVolume* TargetVolume);
+void MESHMODELINGTOOLSEDITORONLY_API DynamicMeshToVolume(const FDynamicMesh3& InputMesh, TArray<FDynamicMeshFace>& Faces, AVolume* TargetVolume);
+
+}//end namespace UE::Conversion
+}//end namespace UE
\ No newline at end of file
diff --git a/Engine/Plugins/Experimental/MeshModelingToolset/Source/MeshModelingToolsEditorOnly/Public/MeshToVolumeTool.h b/Engine/Plugins/Experimental/MeshModelingToolset/Source/MeshModelingToolsEditorOnly/Public/MeshToVolumeTool.h
--- a/Engine/Plugins/Experimental/MeshModelingToolset/Source/MeshModelingToolsEditorOnly/Public/MeshToVolumeTool.h
+++ b/Engine/Plugins/Experimental/MeshModelingToolset/Source/MeshModelingToolsEditorOnly/Public/MeshToVolumeTool.h
@@ -3,6 +3,8 @@
 #pragma once
 
 #include "CoreMinimal.h"
+
+#include "ConversionUtils/VolumeDynamicMeshConversion.h"
 #include "GameFramework/Volume.h"
 #include "Engine/Classes/Engine/BlockingVolume.h"
 #include "SingleSelectionTool.h"
@@ -63,13 +65,6 @@ public:
 };
 
 
-struct FModelFace
-{
-	FFrame3d Plane;
-	TArray<FVector3d> BoundaryLoop;
-};
-
-
 /**
  *
  */
@@ -112,14 +107,9 @@ protected:
 	FDynamicMesh3 InputMesh;
 
 	void RecalculateVolume();
-	void RecalculateVolume_Polygons();
-	void RecalculateVolume_Triangles();
-
 	void UpdateLineSet();
 
 	bool bVolumeValid = false;
-	TArray<FModelFace> Faces;
-
-	void BakeToVolume(AVolume* TargetVolume);
+	TArray<UE::Conversion::FDynamicMeshFace> Faces;
 
 };
diff --git a/Engine/Plugins/Experimental/MeshModelingToolset/Source/ModelingComponents/Private/ToolSetupUtil.cpp b/Engine/Plugins/Experimental/MeshModelingToolset/Source/ModelingComponents/Private/ToolSetupUtil.cpp
--- a/Engine/Plugins/Experimental/MeshModelingToolset/Source/ModelingComponents/Private/ToolSetupUtil.cpp
+++ b/Engine/Plugins/Experimental/MeshModelingToolset/Source/ModelingComponents/Private/ToolSetupUtil.cpp
@@ -178,6 +178,11 @@ UMaterialInstanceDynamic* ToolSetupUtil::GetSimpleCustomMaterial(UInteractiveToo
 	return nullptr;
 }
 
+UMaterialInterface* ToolSetupUtil::GetDefaultEditVolumeMaterial()
+{
+	return LoadObject<UMaterial>(nullptr, TEXT("/MeshModelingToolset/Materials/VolumeEditMaterial"));
+}
+
 
 UMaterialInterface* ToolSetupUtil::GetDefaultPointComponentMaterial(bool bRoundPoints, UInteractiveToolManager* ToolManager)
 {
diff --git a/Engine/Plugins/Experimental/MeshModelingToolset/Source/ModelingComponents/Public/ToolSetupUtil.h b/Engine/Plugins/Experimental/MeshModelingToolset/Source/ModelingComponents/Public/ToolSetupUtil.h
--- a/Engine/Plugins/Experimental/MeshModelingToolset/Source/ModelingComponents/Public/ToolSetupUtil.h
+++ b/Engine/Plugins/Experimental/MeshModelingToolset/Source/ModelingComponents/Public/ToolSetupUtil.h
@@ -48,7 +48,7 @@ namespace ToolSetupUtil
 
 
 	/**
-	 * @return default material to use for brush volume indicators
+	 * @return default material to use for brush volume indicators (for instance, a spherical sculpt brush).
 	 */
 	MODELINGCOMPONENTS_API UMaterialInstanceDynamic* GetDefaultBrushVolumeMaterial(UInteractiveToolManager* ToolManager);
 
@@ -93,6 +93,11 @@ namespace ToolSetupUtil
 	 */
 	MODELINGCOMPONENTS_API UMaterialInstanceDynamic* GetSimpleCustomMaterial(UInteractiveToolManager* ToolManager, const FLinearColor& Color, float Opacity);
 
+	/**
+	 * @return Material used when editing AVolume objects using our tools.
+	 */
+	MODELINGCOMPONENTS_API UMaterialInterface* GetDefaultEditVolumeMaterial();
+
 	/**
 	 * @param bRoundPoints true for round points, false for square
 	 * @return custom material suitable for use with UPointSetComponent
