Add a user-editable relative Padding property to ULatticeDeformerToolProperties to give the user a bit of control over the lattice size. We were just padding by a hardcoded absolute value.

#rnx

diff --git a/Engine/Plugins/Experimental/GeometryProcessing/Source/DynamicMesh/Private/Operations/FFDLattice.cpp b/Engine/Plugins/Experimental/GeometryProcessing/Source/DynamicMesh/Private/Operations/FFDLattice.cpp
--- a/Engine/Plugins/Experimental/GeometryProcessing/Source/DynamicMesh/Private/Operations/FFDLattice.cpp
+++ b/Engine/Plugins/Experimental/GeometryProcessing/Source/DynamicMesh/Private/Operations/FFDLattice.cpp
@@ -4,17 +4,21 @@
 #include "DynamicMesh3.h"
 #include "Util/ProgressCancel.h"
 
-FFFDLattice::FFFDLattice(const FVector3i& InDims, const FDynamicMesh3& Mesh) :
+FFFDLattice::FFFDLattice(const FVector3i& InDims, const FDynamicMesh3& Mesh, float Padding) :
 	Dims(InDims)
 {
 	InitialBounds = Mesh.GetBounds();
 
 	// Expand the initial bounding box to make the computation of which grid cell a mesh vertex is inside of a little 
 	// less susceptible to numerical error issues.
-	// TODO: derive a better padding value than this magic number, or allow user to set it (UETOOL-2362)
-	constexpr float Padding = 10.0f;
-	InitialBounds.Min -= FVector3d{ Padding, Padding, Padding };
-	InitialBounds.Max += FVector3d{ Padding, Padding, Padding };
+
+	const float ClampedPadding = FMath::Clamp(Padding, 0.01f, 5.0f);
+	const FVector3d Center = InitialBounds.Center();
+	FVector3d Extents = InitialBounds.Extents();
+	Extents *= (1.0 + 0.5 * ClampedPadding);
+
+	InitialBounds.Min = Center - Extents;
+	InitialBounds.Max = Center + Extents;
 	
 	ComputeInitialEmbedding(Mesh);
 }
diff --git a/Engine/Plugins/Experimental/GeometryProcessing/Source/DynamicMesh/Public/Operations/FFDLattice.h b/Engine/Plugins/Experimental/GeometryProcessing/Source/DynamicMesh/Public/Operations/FFDLattice.h
--- a/Engine/Plugins/Experimental/GeometryProcessing/Source/DynamicMesh/Public/Operations/FFDLattice.h
+++ b/Engine/Plugins/Experimental/GeometryProcessing/Source/DynamicMesh/Public/Operations/FFDLattice.h
@@ -26,7 +26,7 @@ public:
 
 	/// Create a lattice that fits the given mesh and has the given resolution along each dimension. Also precompute
 	/// the weight of each mesh vertex inside the appropriate lattice cell.
-	FFFDLattice(const FVector3i& InDims, const FDynamicMesh3& InMesh);
+	FFFDLattice(const FVector3i& InDims, const FDynamicMesh3& InMesh, float Padding);
 
 	/// Using the Lattice's current control point positions and the original embedding information, compute the
 	/// deformed mesh vertex positions.
diff --git a/Engine/Plugins/Experimental/MeshModelingToolset/Source/MeshModelingTools/Private/LatticeDeformerTool.cpp b/Engine/Plugins/Experimental/MeshModelingToolset/Source/MeshModelingTools/Private/LatticeDeformerTool.cpp
--- a/Engine/Plugins/Experimental/MeshModelingToolset/Source/MeshModelingTools/Private/LatticeDeformerTool.cpp
+++ b/Engine/Plugins/Experimental/MeshModelingToolset/Source/MeshModelingTools/Private/LatticeDeformerTool.cpp
@@ -89,7 +89,7 @@ bool ULatticeDeformerTool::CanAccept() const
 
 void ULatticeDeformerTool::InitializeLattice(TArray<FVector3d>& OutLatticePoints, TArray<FVector2i>& OutLatticeEdges)
 {
-	Lattice = MakeShared<FFFDLattice>(GetLatticeResolution(), *OriginalMesh);
+	Lattice = MakeShared<FFFDLattice>(GetLatticeResolution(), *OriginalMesh, Settings->Padding);
 
 	Lattice->GenerateInitialLatticePositions(OutLatticePoints);
 
@@ -195,7 +195,8 @@ void ULatticeDeformerTool::OnPropertyModified(UObject* PropertySet, FProperty* P
 		bool bShouldRebuild = Property &&
 			((Property->GetFName() == GET_MEMBER_NAME_CHECKED(ULatticeDeformerToolProperties, XAxisResolution)) ||
 			 (Property->GetFName() == GET_MEMBER_NAME_CHECKED(ULatticeDeformerToolProperties, YAxisResolution)) ||
-			 (Property->GetFName() == GET_MEMBER_NAME_CHECKED(ULatticeDeformerToolProperties, ZAxisResolution)));
+			 (Property->GetFName() == GET_MEMBER_NAME_CHECKED(ULatticeDeformerToolProperties, ZAxisResolution)) ||
+			 (Property->GetFName() == GET_MEMBER_NAME_CHECKED(ULatticeDeformerToolProperties, Padding)));
 
 		if (bShouldRebuild)
 		{
diff --git a/Engine/Plugins/Experimental/MeshModelingToolset/Source/MeshModelingTools/Public/LatticeDeformerTool.h b/Engine/Plugins/Experimental/MeshModelingToolset/Source/MeshModelingTools/Public/LatticeDeformerTool.h
--- a/Engine/Plugins/Experimental/MeshModelingToolset/Source/MeshModelingTools/Public/LatticeDeformerTool.h
+++ b/Engine/Plugins/Experimental/MeshModelingToolset/Source/MeshModelingTools/Public/LatticeDeformerTool.h
@@ -54,6 +54,10 @@ public:
 	UPROPERTY(EditAnywhere, Category = Resolution, meta = (UIMin = "2", ClampMin = "2", EditCondition = "bCanChangeResolution", HideEditConditionToggle))
 	int ZAxisResolution = 5;
 
+	/** Relative distance the lattice extends from the mesh */
+	UPROPERTY(EditAnywhere, Category = Resolution, meta = (UIMin = "0.01", ClampMin = "0.01", UIMax = "2", ClampMax = "5", EditCondition = "bCanChangeResolution", HideEditConditionToggle))
+	float Padding = 0.01;
+
 	/** Whether to use linear or cubic interpolation to get new mesh vertex positions from the lattice */
 	UPROPERTY(EditAnywhere, Category = Interpolation )
 	ELatticeInterpolationType InterpolationType = ELatticeInterpolationType::Linear;
