Lattice deformer: only draw the marquee selection rectangle in the focused viewport.

#rnx

diff --git a/Engine/Plugins/Experimental/MeshModelingToolset/Source/MeshModelingTools/Private/LatticeDeformerTool.cpp b/Engine/Plugins/Experimental/MeshModelingToolset/Source/MeshModelingTools/Private/LatticeDeformerTool.cpp
--- a/Engine/Plugins/Experimental/MeshModelingToolset/Source/MeshModelingTools/Private/LatticeDeformerTool.cpp
+++ b/Engine/Plugins/Experimental/MeshModelingToolset/Source/MeshModelingTools/Private/LatticeDeformerTool.cpp
@@ -70,8 +70,11 @@ FVector3i ULatticeDeformerTool::GetLatticeResolution() const
 
 void ULatticeDeformerTool::DrawHUD(FCanvas* Canvas, IToolsContextRenderAPI* RenderAPI)
 {
+	EViewInteractionState State = RenderAPI->GetViewInteractionState();
+	bool bThisViewHasFocus = !!(State & EViewInteractionState::Focused);
+
 	// Draw the drag rectangle if it's active
-	if (ControlPointsMechanic->bIsDragging)
+	if (bThisViewHasFocus && ControlPointsMechanic->bIsDragging)
 	{		
 		FVector2D Start = ControlPointsMechanic->DragStartScreenPosition;
 		FVector2D Curr = ControlPointsMechanic->DragCurrentScreenPosition;
