Replace a boolean parameter with an enum class. I think it's clearer to pass in "Linear" than "false".

Also some code formatting fixes.

#rnx

diff --git a/Engine/Plugins/Experimental/GeometryProcessing/Source/DynamicMesh/Private/Operations/FFDLattice.cpp b/Engine/Plugins/Experimental/GeometryProcessing/Source/DynamicMesh/Private/Operations/FFDLattice.cpp
--- a/Engine/Plugins/Experimental/GeometryProcessing/Source/DynamicMesh/Private/Operations/FFDLattice.cpp
+++ b/Engine/Plugins/Experimental/GeometryProcessing/Source/DynamicMesh/Private/Operations/FFDLattice.cpp
@@ -112,7 +112,7 @@ void FFFDLattice::ComputeInitialEmbedding(const FDynamicMesh3& Mesh)
 
 void FFFDLattice::GetDeformedMeshVertexPositions(const TArray<FVector3d>& LatticeControlPoints, 
 												 TArray<FVector3d>& OutVertexPositions, 
-												 bool bCubic, 
+												 ELatticeInterpolation Interpolation,
 												 FProgressCancel* Progress) const
 {
 	int NumVertices = VertexEmbeddings.Num();
@@ -120,7 +120,7 @@ void FFFDLattice::GetDeformedMeshVertexPositions(const TArray<FVector3d>& Lattic
 
 	for (int VertexID = 0; VertexID < NumVertices; ++VertexID)
 	{
-		if (bCubic)
+		if (Interpolation == ELatticeInterpolation::Cubic)
 		{
 			OutVertexPositions[VertexID] = InterpolatedPositionCubic(VertexEmbeddings[VertexID], LatticeControlPoints);
 		}
diff --git a/Engine/Plugins/Experimental/GeometryProcessing/Source/DynamicMesh/Public/Operations/FFDLattice.h b/Engine/Plugins/Experimental/GeometryProcessing/Source/DynamicMesh/Public/Operations/FFDLattice.h
--- a/Engine/Plugins/Experimental/GeometryProcessing/Source/DynamicMesh/Public/Operations/FFDLattice.h
+++ b/Engine/Plugins/Experimental/GeometryProcessing/Source/DynamicMesh/Public/Operations/FFDLattice.h
@@ -9,6 +9,14 @@
 class FDynamicMesh3;
 class FProgressCancel;
 
+
+enum class ELatticeInterpolation : uint8
+{
+	Linear = 0,
+	Cubic = 1
+};
+
+
 /** 
  * Free-form deformation lattice. Initialize it with a mesh and desired resolution, get the initial lattice points out.
  * Then pass in deformed lattice points and it will compute deformed mesh vertex positions.
@@ -32,7 +40,7 @@ public:
 	/// deformed mesh vertex positions.
 	void GetDeformedMeshVertexPositions(const TArray<FVector3d>& LatticeControlPoints,
 										TArray<FVector3d>& OutVertexPositions, 
-										bool bCubic = false, 
+										ELatticeInterpolation Interpolation = ELatticeInterpolation::Linear,
 										FProgressCancel* Progress = nullptr) const;
 
 	void GenerateInitialLatticePositions(TArray<FVector3d>& OutLatticePositions) const;
diff --git a/Engine/Plugins/Experimental/MeshModelingToolset/Source/MeshModelingTools/Private/LatticeDeformerTool.cpp b/Engine/Plugins/Experimental/MeshModelingToolset/Source/MeshModelingTools/Private/LatticeDeformerTool.cpp
--- a/Engine/Plugins/Experimental/MeshModelingToolset/Source/MeshModelingTools/Private/LatticeDeformerTool.cpp
+++ b/Engine/Plugins/Experimental/MeshModelingToolset/Source/MeshModelingTools/Private/LatticeDeformerTool.cpp
@@ -48,14 +48,16 @@ UInteractiveTool* ULatticeDeformerToolBuilder::BuildTool(const FToolBuilderState
 
 TUniquePtr<FDynamicMeshOperator> ULatticeDeformerOperatorFactory::MakeNewOperator()
 {
-	bool bUseCubicInterpolation =
-		(LatticeDeformerTool->Settings->InterpolationType == ELatticeInterpolationType::Cubic);
+	ELatticeInterpolation OpInterpolationType =
+		(LatticeDeformerTool->Settings->InterpolationType == ELatticeInterpolationType::Cubic) ?
+		ELatticeInterpolation::Cubic :
+		ELatticeInterpolation::Linear;
 
 	TUniquePtr<FLatticeDeformerOp> LatticeDeformOp = MakeUnique<FLatticeDeformerOp>(
 		LatticeDeformerTool->OriginalMesh,
 		LatticeDeformerTool->Lattice,
 		LatticeDeformerTool->ControlPointsMechanic->GetControlPoints(),
-		bUseCubicInterpolation);
+		OpInterpolationType);
 
 	return LatticeDeformOp;
 }
diff --git a/Engine/Plugins/Experimental/MeshModelingToolset/Source/ModelingOperators/Private/DeformationOps/LatticeDeformerOp.cpp b/Engine/Plugins/Experimental/MeshModelingToolset/Source/ModelingOperators/Private/DeformationOps/LatticeDeformerOp.cpp
--- a/Engine/Plugins/Experimental/MeshModelingToolset/Source/ModelingOperators/Private/DeformationOps/LatticeDeformerOp.cpp
+++ b/Engine/Plugins/Experimental/MeshModelingToolset/Source/ModelingOperators/Private/DeformationOps/LatticeDeformerOp.cpp
@@ -18,7 +18,7 @@ void FLatticeDeformerOp::CalculateResult(FProgressCancel* Progress)
 	}
 
 	TArray<FVector3d> DeformedPositions;
-	Lattice->GetDeformedMeshVertexPositions(LatticeControlPoints, DeformedPositions, bUseCubicInterpolation, Progress);
+	Lattice->GetDeformedMeshVertexPositions(LatticeControlPoints, DeformedPositions, InterpolationType, Progress);
 
 	check(ResultMesh->VertexCount() == DeformedPositions.Num());
 
@@ -31,9 +31,9 @@ void FLatticeDeformerOp::CalculateResult(FProgressCancel* Progress)
 FLatticeDeformerOp::FLatticeDeformerOp(TSharedPtr<FDynamicMesh3> InOriginalMesh,
 									   TSharedPtr<FFFDLattice> InLattice,
 									   const TArray<FVector3d>& InLatticeControlPoints,
-									   bool InUseCubicInterpolation) :
+									   ELatticeInterpolation InInterpolationType) :
 	Lattice(InLattice),
 	OriginalMesh(InOriginalMesh),
 	LatticeControlPoints(InLatticeControlPoints),
-	bUseCubicInterpolation(InUseCubicInterpolation)
+	InterpolationType(InInterpolationType)
 {}
diff --git a/Engine/Plugins/Experimental/MeshModelingToolset/Source/ModelingOperators/Public/DeformationOps/LatticeDeformerOp.h b/Engine/Plugins/Experimental/MeshModelingToolset/Source/ModelingOperators/Public/DeformationOps/LatticeDeformerOp.h
--- a/Engine/Plugins/Experimental/MeshModelingToolset/Source/ModelingOperators/Public/DeformationOps/LatticeDeformerOp.h
+++ b/Engine/Plugins/Experimental/MeshModelingToolset/Source/ModelingOperators/Public/DeformationOps/LatticeDeformerOp.h
@@ -19,7 +19,7 @@ public:
 	FLatticeDeformerOp(TSharedPtr<FDynamicMesh3> InOriginalMesh,
 					   TSharedPtr<FFFDLattice> InLattice,
 					   const TArray<FVector3d>& InLatticeControlPoints,
-					   bool InUseCubicInterpolation);
+					   ELatticeInterpolation InInterpolationType);
 
 	// FDynamicMeshOperator implementation
 	void CalculateResult(FProgressCancel* Progress) override;
@@ -30,6 +30,6 @@ protected:
 	const TSharedPtr<const FFFDLattice> Lattice;
 	const TSharedPtr<const FDynamicMesh3> OriginalMesh;
 	const TArray<FVector3d> LatticeControlPoints;
-	bool bUseCubicInterpolation;
+	ELatticeInterpolation InterpolationType;
 };
 
