MeshModelingTools: Made mechanics only clean up their own gizmos, and fixed potential gizmo leak in CurveControlPointsMechanic (that was being cleaned up by another mechanic).

#rnx

diff --git a/Engine/Plugins/Experimental/MeshModelingToolset/Source/ModelingComponents/Private/Mechanics/ConstructionPlaneMechanic.cpp b/Engine/Plugins/Experimental/MeshModelingToolset/Source/ModelingComponents/Private/Mechanics/ConstructionPlaneMechanic.cpp
--- a/Engine/Plugins/Experimental/MeshModelingToolset/Source/ModelingComponents/Private/Mechanics/ConstructionPlaneMechanic.cpp
+++ b/Engine/Plugins/Experimental/MeshModelingToolset/Source/ModelingComponents/Private/Mechanics/ConstructionPlaneMechanic.cpp
@@ -17,7 +17,7 @@ void UConstructionPlaneMechanic::Setup(UInteractiveTool* ParentToolIn)
 
 void UConstructionPlaneMechanic::Shutdown()
 {
-	GetParentTool()->GetToolManager()->GetPairedGizmoManager()->DestroyAllGizmosByOwner(GetParentTool());
+	GetParentTool()->GetToolManager()->GetPairedGizmoManager()->DestroyAllGizmosByOwner(this);
 }
 
 
@@ -31,7 +31,7 @@ void UConstructionPlaneMechanic::Initialize(UWorld* TargetWorld, const FFrame3d&
 
 	PlaneTransformProxy = NewObject<UTransformProxy>(this);
 	PlaneTransformGizmo = GizmoManager->CreateCustomTransformGizmo(
-		ETransformGizmoSubElements::StandardTranslateRotate, GetParentTool());
+		ETransformGizmoSubElements::StandardTranslateRotate, this);
 	PlaneTransformProxy->OnTransformChanged.AddUObject(this, &UConstructionPlaneMechanic::TransformChanged);
 
 	PlaneTransformGizmo->SetActiveTarget(PlaneTransformProxy, GetParentTool()->GetToolManager());
diff --git a/Engine/Plugins/Experimental/MeshModelingToolset/Source/ModelingComponents/Private/Mechanics/CurveControlPointsMechanic.cpp b/Engine/Plugins/Experimental/MeshModelingToolset/Source/ModelingComponents/Private/Mechanics/CurveControlPointsMechanic.cpp
--- a/Engine/Plugins/Experimental/MeshModelingToolset/Source/ModelingComponents/Private/Mechanics/CurveControlPointsMechanic.cpp
+++ b/Engine/Plugins/Experimental/MeshModelingToolset/Source/ModelingComponents/Private/Mechanics/CurveControlPointsMechanic.cpp
@@ -114,7 +114,7 @@ void UCurveControlPointsMechanic::Setup(UInteractiveTool* ParentToolIn)
 	PointTransformProxy = NewObject<UTransformProxy>(this);
 	PointTransformGizmo = GizmoManager->CreateCustomTransformGizmo(
 		ETransformGizmoSubElements::TranslateAxisX | ETransformGizmoSubElements::TranslateAxisY | ETransformGizmoSubElements::TranslatePlaneXY,
-		GetParentTool());
+		this);
 	PointTransformProxy->OnTransformChanged.AddUObject(this, &UCurveControlPointsMechanic::GizmoTransformChanged);
 	PointTransformProxy->OnBeginTransformEdit.AddUObject(this, &UCurveControlPointsMechanic::GizmoTransformStarted);
 	PointTransformProxy->OnEndTransformEdit.AddUObject(this, &UCurveControlPointsMechanic::GizmoTransformEnded);
@@ -187,17 +187,14 @@ void UCurveControlPointsMechanic::SetWorld(UWorld* World)
 
 void UCurveControlPointsMechanic::Shutdown()
 {
+	// Calls shutdown on gizmo and destroys it.
+	GetParentTool()->GetToolManager()->GetPairedGizmoManager()->DestroyAllGizmosByOwner(this);
+
 	if (PreviewGeometryActor)
 	{
 		PreviewGeometryActor->Destroy();
 		PreviewGeometryActor = nullptr;
 	}
-
-	if (PointTransformGizmo)
-	{
-		PointTransformGizmo->Shutdown();
-		PointTransformGizmo = nullptr;
-	}
 }
 
 void UCurveControlPointsMechanic::Render(IToolsContextRenderAPI* RenderAPI)
diff --git a/Engine/Plugins/Experimental/MeshModelingToolset/Source/ModelingComponents/Private/Mechanics/LatticeControlPointsMechanic.cpp b/Engine/Plugins/Experimental/MeshModelingToolset/Source/ModelingComponents/Private/Mechanics/LatticeControlPointsMechanic.cpp
--- a/Engine/Plugins/Experimental/MeshModelingToolset/Source/ModelingComponents/Private/Mechanics/LatticeControlPointsMechanic.cpp
+++ b/Engine/Plugins/Experimental/MeshModelingToolset/Source/ModelingComponents/Private/Mechanics/LatticeControlPointsMechanic.cpp
@@ -82,8 +82,8 @@ void ULatticeControlPointsMechanic::Setup(UInteractiveTool* ParentToolIn)
 	
 	// TODO: Maybe don't have the gizmo's axes flip around when it crosses the origin, if possible?
 	// TODO: Enable local vs world gizmo switching (UETOOL-2356)
-	PointTransformGizmo = GizmoManager->CreateCustomTransformGizmo(ETransformGizmoSubElements::FullTranslateRotateScale,
-																   GetParentTool());
+	PointTransformGizmo = GizmoManager->CreateCustomTransformGizmo(
+		ETransformGizmoSubElements::FullTranslateRotateScale, this);
 
 	PointTransformProxy->OnTransformChanged.AddUObject(this, &ULatticeControlPointsMechanic::GizmoTransformChanged);
 	PointTransformProxy->OnBeginTransformEdit.AddUObject(this, &ULatticeControlPointsMechanic::GizmoTransformStarted);
@@ -149,14 +149,9 @@ void ULatticeControlPointsMechanic::Shutdown()
 		PreviewGeometryActor = nullptr;
 	}
 
-	if (PointTransformGizmo)
-	{
-		PointTransformGizmo->Shutdown();
-		PointTransformGizmo = nullptr;
-	}
-
+	// Calls shutdown on gizmo and destroys it.
 	UInteractiveGizmoManager* GizmoManager = GetParentTool()->GetToolManager()->GetPairedGizmoManager();
-	GizmoManager->DestroyAllGizmosByOwner(GetParentTool());
+	GizmoManager->DestroyAllGizmosByOwner(this);
 }
 
 
