- Support world and local space modes for lattice deformer gizmos.
- Make gizmo transform change due to point selection change undoable.
- Add a "set pivot" mode to lattice deformer tool.
- Use either ctrl or shift to append points to the current selection.

#rnx

diff --git a/Engine/Plugins/Experimental/MeshModelingToolset/Source/MeshModelingTools/Private/LatticeDeformerTool.cpp b/Engine/Plugins/Experimental/MeshModelingToolset/Source/MeshModelingTools/Private/LatticeDeformerTool.cpp
--- a/Engine/Plugins/Experimental/MeshModelingToolset/Source/MeshModelingTools/Private/LatticeDeformerTool.cpp
+++ b/Engine/Plugins/Experimental/MeshModelingToolset/Source/MeshModelingTools/Private/LatticeDeformerTool.cpp
@@ -13,6 +13,7 @@
 #include "MeshOpPreviewHelpers.h" //FDynamicMeshOpResult
 #include "ToolSceneQueriesUtil.h"
 #include "ToolBuilderUtil.h"
+#include "ToolSetupUtil.h"
 #include "DynamicMeshToMeshDescription.h"
 #include "MeshTransforms.h"
 #include "Algo/ForEach.h"
@@ -122,6 +123,25 @@ void ULatticeDeformerTool::Setup()
 	Settings->RestoreProperties(this);
 	AddToolPropertySource(Settings);
 
+	// Watch for property changes
+	Settings->WatchProperty(Settings->XAxisResolution, [this](int) { bShouldRebuild = true; });
+	Settings->WatchProperty(Settings->YAxisResolution, [this](int) { bShouldRebuild = true; });
+	Settings->WatchProperty(Settings->ZAxisResolution, [this](int) { bShouldRebuild = true; });
+	Settings->WatchProperty(Settings->Padding, [this](float) { bShouldRebuild = true; });
+	Settings->WatchProperty(Settings->InterpolationType, [this](ELatticeInterpolationType)
+	{
+		Preview->InvalidateResult();
+	});
+	Settings->WatchProperty(Settings->GizmoCoordinateSystem, [this](EToolContextCoordinateSystem)
+	{
+		ControlPointsMechanic->SetCoordinateSystem(Settings->GizmoCoordinateSystem);
+	});
+	Settings->WatchProperty(Settings->bSetPivotMode, [this](bool)
+	{
+		ControlPointsMechanic->UpdateSetPivotMode(Settings->bSetPivotMode);
+	});
+	
+
 	TArray<FVector3d> LatticePoints;
 	TArray<FVector2i> LatticeEdges;
 	InitializeLattice(LatticePoints, LatticeEdges);
@@ -130,16 +150,19 @@ void ULatticeDeformerTool::Setup()
 	ControlPointsMechanic = NewObject<ULatticeControlPointsMechanic>(this);
 	ControlPointsMechanic->Setup(this);
 	ControlPointsMechanic->SetWorld(TargetWorld);
-	ControlPointsMechanic->Initialize(LatticePoints, LatticeEdges);
+	FTransform3d LocalToWorld = FTransform3d(ComponentTarget->GetWorldTransform());
+	ControlPointsMechanic->Initialize(LatticePoints, LatticeEdges, LocalToWorld);
 
 	auto OnPointsChangedLambda = [this]()
 	{
 		Preview->InvalidateResult();
 		Settings->bCanChangeResolution = !ControlPointsMechanic->bHasChanged;
 	};
-
 	ControlPointsMechanic->OnPointsChanged.AddLambda(OnPointsChangedLambda);
 
+	ControlPointsMechanic->SetCoordinateSystem(Settings->GizmoCoordinateSystem);
+	ControlPointsMechanic->UpdateSetPivotMode(Settings->bSetPivotMode);
+
 	StartPreview();
 }
 
@@ -186,6 +209,20 @@ void ULatticeDeformerTool::StartPreview()
 	Preview = NewObject<UMeshOpPreviewWithBackgroundCompute>(LatticeDeformOpCreator);
 
 	Preview->Setup(TargetWorld, LatticeDeformOpCreator);
+
+	FComponentMaterialSet MaterialSet;
+	ComponentTarget->GetMaterialSet(MaterialSet);
+	Preview->ConfigureMaterials(MaterialSet.Materials,
+								ToolSetupUtil::GetDefaultWorkingMaterial(GetToolManager())
+	);
+
+	// configure secondary render material
+	UMaterialInterface* SelectionMaterial = ToolSetupUtil::GetSelectionMaterial(FLinearColor(0.8f, 0.75f, 0.0f), GetToolManager());
+	if (SelectionMaterial != nullptr)
+	{
+		Preview->PreviewMesh->SetSecondaryRenderMaterial(SelectionMaterial);
+	}
+
 	Preview->PreviewMesh->SetTangentsMode(EDynamicMeshTangentCalcType::NoTangents);
 	Preview->SetVisibility(true);
 	Preview->InvalidateResult();
@@ -193,33 +230,21 @@ void ULatticeDeformerTool::StartPreview()
 	ComponentTarget->SetOwnerVisibility(false);
 }
 
-void ULatticeDeformerTool::OnPropertyModified(UObject* PropertySet, FProperty* Property)
+void ULatticeDeformerTool::OnTick(float DeltaTime)
 {
 	if (Preview)
 	{
-		bool bShouldRebuild = Property &&
-			((Property->GetFName() == GET_MEMBER_NAME_CHECKED(ULatticeDeformerToolProperties, XAxisResolution)) ||
-			 (Property->GetFName() == GET_MEMBER_NAME_CHECKED(ULatticeDeformerToolProperties, YAxisResolution)) ||
-			 (Property->GetFName() == GET_MEMBER_NAME_CHECKED(ULatticeDeformerToolProperties, ZAxisResolution)) ||
-			 (Property->GetFName() == GET_MEMBER_NAME_CHECKED(ULatticeDeformerToolProperties, Padding)));
-
 		if (bShouldRebuild)
 		{
 			TArray<FVector3d> LatticePoints;
 			TArray<FVector2i> LatticeEdges;
 			InitializeLattice(LatticePoints, LatticeEdges);
-			ControlPointsMechanic->Initialize(LatticePoints, LatticeEdges);
+			FTransform3d LocalToWorld = FTransform3d(ComponentTarget->GetWorldTransform());
+			ControlPointsMechanic->Initialize(LatticePoints, LatticeEdges, LocalToWorld);
+			Preview->InvalidateResult();
+			bShouldRebuild = false;
 		}
 
-		Preview->InvalidateResult();
-	}
-}
-
-
-void ULatticeDeformerTool::OnTick(float DeltaTime)
-{
-	if (Preview)
-	{
 		Preview->Tick(DeltaTime);
 	}
 }
diff --git a/Engine/Plugins/Experimental/MeshModelingToolset/Source/MeshModelingTools/Public/LatticeDeformerTool.h b/Engine/Plugins/Experimental/MeshModelingToolset/Source/MeshModelingTools/Public/LatticeDeformerTool.h
--- a/Engine/Plugins/Experimental/MeshModelingToolset/Source/MeshModelingTools/Public/LatticeDeformerTool.h
+++ b/Engine/Plugins/Experimental/MeshModelingToolset/Source/MeshModelingTools/Public/LatticeDeformerTool.h
@@ -66,6 +66,14 @@ public:
 	UPROPERTY(meta = (TransientToolProperty))
 	bool bCanChangeResolution = true;
 
+	/** Whether the gizmo's axes remain aligned with world axes or rotate as the gizmo is transformed */
+	UPROPERTY(EditAnywhere, Category = Gizmo)
+	EToolContextCoordinateSystem GizmoCoordinateSystem = EToolContextCoordinateSystem::Local;
+
+	/** If Set Pivot Mode is active, the gizmo can be repositioned without moving the selected lattice points */
+	UPROPERTY(EditAnywhere, Category = Gizmo)
+	bool bSetPivotMode = false;
+
 };
 
 
@@ -106,8 +114,6 @@ public:
 	virtual void OnTick(float DeltaTime) override;
 	virtual void Render(IToolsContextRenderAPI* RenderAPI) override;
 
-	virtual void OnPropertyModified(UObject* PropertySet, FProperty* Property) override;
-
 	FVector3i GetLatticeResolution() const;
 
 protected:
@@ -132,6 +138,8 @@ protected:
 	UPROPERTY()
 	bool bLatticeDeformed = false;
 
+	bool bShouldRebuild = false;
+
 	// Create and store an FFFDLattice. Pass out the lattice's positions and edges.
 	void InitializeLattice(TArray<FVector3d>& OutLatticePoints, TArray<FVector2i>& OutLatticeEdges);
 
diff --git a/Engine/Plugins/Experimental/MeshModelingToolset/Source/ModelingComponents/Private/Mechanics/LatticeControlPointsMechanic.cpp b/Engine/Plugins/Experimental/MeshModelingToolset/Source/ModelingComponents/Private/Mechanics/LatticeControlPointsMechanic.cpp
--- a/Engine/Plugins/Experimental/MeshModelingToolset/Source/ModelingComponents/Private/Mechanics/LatticeControlPointsMechanic.cpp
+++ b/Engine/Plugins/Experimental/MeshModelingToolset/Source/ModelingComponents/Private/Mechanics/LatticeControlPointsMechanic.cpp
@@ -29,14 +29,16 @@ void ULatticeControlPointsMechanic::Setup(UInteractiveTool* ParentToolIn)
 
 	USingleClickInputBehavior* ClickBehavior = NewObject<USingleClickInputBehavior>();
 	ClickBehavior->Initialize(this);
-	ClickBehavior->Modifiers.RegisterModifier(CtrlModifierId, FInputDeviceState::IsCtrlKeyDown);
-	ClickBehavior->Modifiers.RegisterModifier(ShiftModifierId, FInputDeviceState::IsShiftKeyDown);
+	auto ShiftOrCtrlDown = [](const FInputDeviceState& State) -> bool {
+		return FInputDeviceState::IsShiftKeyDown(State) || FInputDeviceState::IsCtrlKeyDown(State);
+	};
+	ClickBehavior->Modifiers.RegisterModifier(AddToSelectionModifierId, ShiftOrCtrlDown);
+
 	ParentTool->AddInputBehavior(ClickBehavior);
 
 	UMouseHoverBehavior* HoverBehavior = NewObject<UMouseHoverBehavior>();
 	HoverBehavior->Initialize(this);
-	HoverBehavior->Modifiers.RegisterModifier(CtrlModifierId, FInputDeviceState::IsCtrlKeyDown);
-	HoverBehavior->Modifiers.RegisterModifier(ShiftModifierId, FInputDeviceState::IsShiftKeyDown);
+	HoverBehavior->Modifiers.RegisterModifier(AddToSelectionModifierId, ShiftOrCtrlDown);
 	ParentTool->AddInputBehavior(HoverBehavior);
 
 	UClickDragInputBehavior* ClickDragBehavior = NewObject<UClickDragInputBehavior>(this);
@@ -91,10 +93,32 @@ void ULatticeControlPointsMechanic::Setup(UInteractiveTool* ParentToolIn)
 	PointTransformGizmo->SetActiveTarget(PointTransformProxy);
 	PointTransformGizmo->SetVisibility(false);
 	PointTransformGizmo->bUseContextCoordinateSystem = false;
+	PointTransformGizmo->CurrentCoordinateSystem = EToolContextCoordinateSystem::Local;
+}
+
+void ULatticeControlPointsMechanic::SetCoordinateSystem(EToolContextCoordinateSystem InCoordinateSystem)
+{
+	PointTransformGizmo->CurrentCoordinateSystem = InCoordinateSystem;
+	UpdateGizmoLocation();
+}
+
+EToolContextCoordinateSystem ULatticeControlPointsMechanic::GetCoordinateSystem() const
+{
+	return PointTransformGizmo->CurrentCoordinateSystem;
+}
+
+
+void ULatticeControlPointsMechanic::UpdateSetPivotMode(bool bInSetPivotMode)
+{
+	PointTransformProxy->bSetPivotMode = bInSetPivotMode;
 }
 
-void ULatticeControlPointsMechanic::Initialize(const TArray<FVector3d>& Points, const TArray<FVector2i>& Edges)
+
+void ULatticeControlPointsMechanic::Initialize(const TArray<FVector3d>& Points, 
+											   const TArray<FVector2i>& Edges,
+											   const FTransform3d& InLocalToWorldTransform)
 {
+	LocalToWorldTransform = InLocalToWorldTransform;
 	ControlPoints = Points;
 	SelectedPointIDs.Empty();
 	LatticeEdges = Edges;
@@ -272,6 +296,11 @@ void ULatticeControlPointsMechanic::GizmoTransformChanged(UTransformProxy* Proxy
 		return;
 	}
 
+	if (PointTransformProxy->bSetPivotMode)
+	{
+		return;
+	}
+
 	bool bPointsChanged = false;
 
 	FVector Displacement = Transform.GetTranslation() - GizmoStartPosition;
@@ -399,30 +428,37 @@ void ULatticeControlPointsMechanic::ChangeSelection(int32 NewPointID, bool bAddT
 			}
 		}
 
-		ParentTool->GetToolManager()->EmitObjectChange(this, MakeUnique<FLatticeControlPointsMechanicSelectionChange>(
-			PointsToDeselect, false, CurrentChangeStamp), LatticePointDeselectionTransactionText);
-
+		FTransform PreviousTransform = PointTransformProxy->GetTransform();
 		SelectedPointIDs.Empty();
+		UpdateGizmoLocation();
+		FTransform NewTransform = PointTransformProxy->GetTransform();
+
+		ParentTool->GetToolManager()->EmitObjectChange(this, MakeUnique<FLatticeControlPointsMechanicSelectionChange>(
+			PointsToDeselect, false, PreviousTransform, NewTransform, CurrentChangeStamp), LatticePointDeselectionTransactionText);
 	}
 
 	// We check for validity here because giving an invalid id (such as -1) with bAddToSelection == false
 	// is an easy way to clear the selection.
 	if ((NewPointID >= 0)  && (NewPointID < ControlPoints.Num()))
 	{
+		FTransform PreviousTransform = PointTransformProxy->GetTransform();
+
 		if (bAddToSelection && DeselectPoint(NewPointID))
 		{
+			UpdateGizmoLocation();
+			FTransform NewTransform = PointTransformProxy->GetTransform();
 			ParentTool->GetToolManager()->EmitObjectChange(this, MakeUnique<FLatticeControlPointsMechanicSelectionChange>(
-				NewPointID, false, CurrentChangeStamp), LatticePointDeselectionTransactionText);
+				NewPointID, false, PreviousTransform, NewTransform, CurrentChangeStamp), LatticePointDeselectionTransactionText);
 		}
 		else
 		{
 			SelectPoint(NewPointID);
+			UpdateGizmoLocation();
+			FTransform NewTransform = PointTransformProxy->GetTransform();
 			ParentTool->GetToolManager()->EmitObjectChange(this, MakeUnique<FLatticeControlPointsMechanicSelectionChange>(
-				NewPointID, true, CurrentChangeStamp), LatticePointSelectionTransactionText);
+				NewPointID, true, PreviousTransform, NewTransform, CurrentChangeStamp), LatticePointSelectionTransactionText);
 		}
 	}
-
-	UpdateGizmoLocation();
 }
 
 void ULatticeControlPointsMechanic::UpdateGizmoLocation()
@@ -432,6 +468,8 @@ void ULatticeControlPointsMechanic::UpdateGizmoLocation()
 		return;
 	}
 
+	FVector3d NewGizmoLocation(0, 0, 0);
+
 	if (SelectedPointIDs.Num() == 0)
 	{
 		PointTransformGizmo->SetVisibility(false);
@@ -439,24 +477,21 @@ void ULatticeControlPointsMechanic::UpdateGizmoLocation()
 	}
 	else
 	{
-		FVector3d NewGizmoLocation(0,0,0);
 		for (int32 PointID : SelectedPointIDs)
 		{
 			NewGizmoLocation += ControlPoints[PointID];
 		}
 		NewGizmoLocation /= SelectedPointIDs.Num();
 
-		// Don't clear the gizmo rotation
-		FQuat OldGizmoRotation = PointTransformProxy->GetTransform().GetRotation();
-
-		PointTransformGizmo->ReinitializeGizmoTransform(FTransform(OldGizmoRotation, (FVector)NewGizmoLocation));
-
-		// Clear the child scale
-		FVector GizmoScale{ 1.0, 1.0, 1.0 };
-		PointTransformGizmo->SetNewChildScale(GizmoScale);
-
 		PointTransformGizmo->SetVisibility(true);
 	}
+
+	FTransform NewTransform(FQuat(LocalToWorldTransform.GetRotation()), (FVector)NewGizmoLocation);
+	PointTransformGizmo->ReinitializeGizmoTransform(NewTransform);
+
+	// Clear the child scale
+	FVector GizmoScale{ 1.0, 1.0, 1.0 };
+	PointTransformGizmo->SetNewChildScale(GizmoScale);
 }
 
 bool ULatticeControlPointsMechanic::DeselectPoint(int32 PointID)
@@ -540,7 +575,7 @@ void ULatticeControlPointsMechanic::OnEndHover()
 // Detects Ctrl key state
 void ULatticeControlPointsMechanic::OnUpdateModifierState(int ModifierID, bool bIsOn)
 {
-	if (ModifierID == CtrlModifierId || ModifierID == ShiftModifierId)
+	if (ModifierID == AddToSelectionModifierId)
 	{
 		bAddToSelectionToggle = bIsOn;
 	}
@@ -573,8 +608,6 @@ void ULatticeControlPointsMechanic::OnClickPress(const FInputDeviceRay& PressPos
 	}
 	CurrentDragSelection.Empty();
 
-	UpdateGizmoLocation();
-
 	DragStartScreenPosition = PressPos.ScreenPosition;
 	DragStartWorldRay = PressPos.WorldRay;
 
@@ -582,7 +615,6 @@ void ULatticeControlPointsMechanic::OnClickPress(const FInputDeviceRay& PressPos
 	if (PointTransformGizmo)
 	{
 		PointTransformGizmo->SetVisibility(false);
-		PointTransformGizmo->ReinitializeGizmoTransform(FTransform());
 	}
 }
 
@@ -666,25 +698,29 @@ void ULatticeControlPointsMechanic::OnClickRelease(const FInputDeviceRay& Releas
 
 	ParentTool->GetToolManager()->BeginUndoTransaction(LatticePointSelectionTransactionText);
 
-	if (SelectedPointIDs.Num() > 0)
+	TArray<int32> SavedSelectedIDs = SelectedPointIDs;
+
+	FTransform PreviousTransform = PointTransformProxy->GetTransform();
+	SelectedPointIDs = CurrentDragSelection;
+	UpdateGizmoLocation();
+	FTransform NewTransform = PointTransformProxy->GetTransform();
+
+	if (SavedSelectedIDs.Num() > 0)
 	{
 		ParentTool->GetToolManager()->EmitObjectChange(this, MakeUnique<FLatticeControlPointsMechanicSelectionChange>(
-			SelectedPointIDs, false, CurrentChangeStamp), LatticePointDeselectionTransactionText);
+			SavedSelectedIDs, false, PreviousTransform, NewTransform, CurrentChangeStamp), LatticePointDeselectionTransactionText);
 	}
 
 	if (CurrentDragSelection.Num() > 0)
 	{
 		ParentTool->GetToolManager()->EmitObjectChange(this, MakeUnique<FLatticeControlPointsMechanicSelectionChange>(
-			CurrentDragSelection, true, CurrentChangeStamp), LatticePointSelectionTransactionText);
+			CurrentDragSelection, true, PreviousTransform, NewTransform, CurrentChangeStamp), LatticePointSelectionTransactionText);
 	}
 
 	ParentTool->GetToolManager()->EndUndoTransaction();
 
-	SelectedPointIDs = CurrentDragSelection;
 	CurrentDragSelection.Empty();
-
 	bIsDragging = false;
-	UpdateGizmoLocation();
 }
 
 
@@ -704,16 +740,26 @@ const TArray<FVector3d>& ULatticeControlPointsMechanic::GetControlPoints() const
 // ==================== Undo/redo object functions ====================
 
 FLatticeControlPointsMechanicSelectionChange::FLatticeControlPointsMechanicSelectionChange(int32 PointIDIn,
-	bool AddedIn, int32 ChangeStampIn)
+																						   bool bAddedIn,
+																						   const FTransform& PreviousTransformIn,
+																						   const FTransform& NewTransformIn,
+																						   int32 ChangeStampIn)
 	: PointIDs{ PointIDIn }
-	, Added(AddedIn)
+	, bAdded(bAddedIn)
+	, PreviousTransform(PreviousTransformIn)
+	, NewTransform(NewTransformIn)
 	, ChangeStamp(ChangeStampIn)
 {}
 
 FLatticeControlPointsMechanicSelectionChange::FLatticeControlPointsMechanicSelectionChange(const TArray<int32>& PointIDsIn, 
-	bool AddedIn, int32 ChangeStampIn)
+																						   bool bAddedIn, 
+																						   const FTransform& PreviousTransformIn,
+																						   const FTransform& NewTransformIn, 
+																						   int32 ChangeStampIn)
 	: PointIDs(PointIDsIn)
-	, Added(AddedIn)
+	, bAdded(bAddedIn)
+	, PreviousTransform(PreviousTransformIn)
+	, NewTransform(NewTransformIn)
 	, ChangeStamp(ChangeStampIn)
 {}
 
@@ -723,7 +769,7 @@ void FLatticeControlPointsMechanicSelectionChange::Apply(UObject* Object)
 	
 	for (int32 PointID : PointIDs)
 	{
-		if (Added)
+		if (bAdded)
 		{
 			Mechanic->SelectPoint(PointID);
 		}
@@ -733,7 +779,11 @@ void FLatticeControlPointsMechanicSelectionChange::Apply(UObject* Object)
 		}
 	}
 
-	Mechanic->UpdateGizmoLocation();
+	Mechanic->PointTransformGizmo->ReinitializeGizmoTransform(NewTransform);
+	Mechanic->PointTransformGizmo->SetNewChildScale(FVector{ 1.0, 1.0, 1.0 });	// Clear the child scale
+
+	bool bAnyPointSelected = (Mechanic->SelectedPointIDs.Num() > 0);
+	Mechanic->PointTransformGizmo->SetVisibility(bAnyPointSelected);
 }
 
 void FLatticeControlPointsMechanicSelectionChange::Revert(UObject* Object)
@@ -741,7 +791,7 @@ void FLatticeControlPointsMechanicSelectionChange::Revert(UObject* Object)
 	ULatticeControlPointsMechanic* Mechanic = Cast<ULatticeControlPointsMechanic>(Object);
 	for (int32 PointID : PointIDs)
 	{
-		if (Added)
+		if (bAdded)
 		{
 			Mechanic->DeselectPoint(PointID);
 		}
@@ -750,7 +800,12 @@ void FLatticeControlPointsMechanicSelectionChange::Revert(UObject* Object)
 			Mechanic->SelectPoint(PointID);
 		}
 	}
-	Mechanic->UpdateGizmoLocation();
+	
+	Mechanic->PointTransformGizmo->ReinitializeGizmoTransform(PreviousTransform);
+	Mechanic->PointTransformGizmo->SetNewChildScale(FVector{ 1.0, 1.0, 1.0 });	// Clear the child scale
+
+	bool bAnyPointSelected = (Mechanic->SelectedPointIDs.Num() > 0);
+	Mechanic->PointTransformGizmo->SetVisibility(bAnyPointSelected);
 }
 
 FString FLatticeControlPointsMechanicSelectionChange::ToString() const
@@ -779,7 +834,6 @@ void FLatticeControlPointsMechanicMovementChange::Apply(UObject* Object)
 	ULatticeControlPointsMechanic* Mechanic = Cast<ULatticeControlPointsMechanic>(Object);
 	check(PointIDs.Num() == NewPositions.Num());
 	Mechanic->UpdatePointLocations(PointIDs, NewPositions);
-	Mechanic->UpdateGizmoLocation();
 	Mechanic->bHasChanged = false;
 	Mechanic->OnPointsChanged.Broadcast();
 }
@@ -789,7 +843,6 @@ void FLatticeControlPointsMechanicMovementChange::Revert(UObject* Object)
 	ULatticeControlPointsMechanic* Mechanic = Cast<ULatticeControlPointsMechanic>(Object);
 	check(PointIDs.Num() == OriginalPositions.Num());
 	Mechanic->UpdatePointLocations(PointIDs, OriginalPositions);
-	Mechanic->UpdateGizmoLocation();
 	if (bFirstMovement)
 	{
 		// If we're undoing the first change, make it possible to change the lattice resolution again
diff --git a/Engine/Plugins/Experimental/MeshModelingToolset/Source/ModelingComponents/Public/Mechanics/LatticeControlPointsMechanic.h b/Engine/Plugins/Experimental/MeshModelingToolset/Source/ModelingComponents/Public/Mechanics/LatticeControlPointsMechanic.h
--- a/Engine/Plugins/Experimental/MeshModelingToolset/Source/ModelingComponents/Public/Mechanics/LatticeControlPointsMechanic.h
+++ b/Engine/Plugins/Experimental/MeshModelingToolset/Source/ModelingComponents/Public/Mechanics/LatticeControlPointsMechanic.h
@@ -40,10 +40,18 @@ public:
 	DECLARE_MULTICAST_DELEGATE(OnPointsChangedEvent);
 	OnPointsChangedEvent OnPointsChanged;
 
-	virtual void Initialize(const TArray<FVector3d>& Points, const TArray<FVector2i>& Edges);
+	virtual void Initialize(const TArray<FVector3d>& Points, 
+							const TArray<FVector2i>& Edges,
+							const FTransform3d& LocalToWorldTransform );
+
 	void SetWorld(UWorld* World);
 	const TArray<FVector3d>& GetControlPoints() const;
 
+	void SetCoordinateSystem(EToolContextCoordinateSystem InCoordinateSystem);
+	EToolContextCoordinateSystem GetCoordinateSystem() const;
+
+	void UpdateSetPivotMode(bool bInSetPivotMode);
+
 	// UInteractionMechanic
 	virtual void Setup(UInteractiveTool* ParentTool) override;
 	virtual void Shutdown() override;
@@ -80,6 +88,7 @@ protected:
 
 	TArray<FVector3d> ControlPoints;
 	TArray<FVector2i> LatticeEdges;
+	FTransform3d LocalToWorldTransform;
 
 	// Used for spatial queries
 	FGeometrySet3 GeometrySet;
@@ -103,10 +112,10 @@ protected:
 	// Cache previous color while temporarily changing the color of a hovered-over point
 	FColor PreHoverPointColor;
 
-	// Support for Ctrl toggle
+	// Support for Shift and Ctrl toggle
 	bool bAddToSelectionToggle = false;
-	const int32 CtrlModifierId = 1;
-	const int32 ShiftModifierId = 2;
+	const int32 AddToSelectionModifierId = 1;
+	
 
 	// Support for gizmo. Since the points aren't individual components, we don't actually use UTransformProxy
 	// for the transform forwarding- we just use it for the callbacks.
@@ -142,6 +151,7 @@ protected:
 	TArray<FVector3d> SelectedPointStartPositions;
 
 	// The starting point of the gizmo is needed to determine the offset by which to move the points.
+	// TODO: Replace with single FTransform?
 	FVector GizmoStartPosition;
 	FQuat	GizmoStartRotation;
 	FVector GizmoStartScale;
@@ -175,9 +185,11 @@ protected:
 class MODELINGCOMPONENTS_API FLatticeControlPointsMechanicSelectionChange : public FToolCommandChange
 {
 public:
-	FLatticeControlPointsMechanicSelectionChange(int32 PointIDIn, bool AddedIn, int32 ChangeStampIn);
+	FLatticeControlPointsMechanicSelectionChange(int32 PointIDIn, bool bAddedIn, const FTransform& PreviousTransformIn,
+												 const FTransform& NewTransformIn, int32 ChangeStampIn);
 
-	FLatticeControlPointsMechanicSelectionChange(const TArray<int32>& PointIDsIn, bool AddedIn, int32 ChangeStampIn);
+	FLatticeControlPointsMechanicSelectionChange(const TArray<int32>& PointIDsIn, bool bAddedIn, const FTransform& PreviousTransformIn,
+												 const FTransform& NewTransformIn, int32 ChangeStampIn);
 
 	virtual void Apply(UObject* Object) override;
 	virtual void Revert(UObject* Object) override;
@@ -189,8 +201,11 @@ public:
 
 protected:
 	TArray<int32> PointIDs;
+	bool bAdded;
+
+	const FTransform PreviousTransform;
+	const FTransform NewTransform;
 
-	bool Added;
 	int32 ChangeStamp;
 };
 
diff --git a/Engine/Source/Runtime/Experimental/InteractiveToolsFramework/Private/BaseGizmos/TransformGizmo.cpp b/Engine/Source/Runtime/Experimental/InteractiveToolsFramework/Private/BaseGizmos/TransformGizmo.cpp
--- a/Engine/Source/Runtime/Experimental/InteractiveToolsFramework/Private/BaseGizmos/TransformGizmo.cpp
+++ b/Engine/Source/Runtime/Experimental/InteractiveToolsFramework/Private/BaseGizmos/TransformGizmo.cpp
@@ -527,9 +527,10 @@ void UTransformGizmo::ReinitializeGizmoTransform(const FTransform& NewTransform)
 	GizmoComponent->TransformUpdated = temp;
 
 	// The underlying proxy has an existing way to reinitialize its transform without callbacks.
+	bool bSavedSetPivotMode = ActiveTarget->bSetPivotMode;
 	ActiveTarget->bSetPivotMode = true;
 	ActiveTarget->SetTransform(NewTransform);
-	ActiveTarget->bSetPivotMode = false;
+	ActiveTarget->bSetPivotMode = bSavedSetPivotMode;
 }
 
 void UTransformGizmo::SetNewGizmoTransform(const FTransform& NewTransform)
