Remove an unused function

#rnx

diff --git a/Engine/Plugins/Experimental/MeshModelingToolset/Source/ModelingComponents/Private/Mechanics/LatticeControlPointsMechanic.cpp b/Engine/Plugins/Experimental/MeshModelingToolset/Source/ModelingComponents/Private/Mechanics/LatticeControlPointsMechanic.cpp
--- a/Engine/Plugins/Experimental/MeshModelingToolset/Source/ModelingComponents/Private/Mechanics/LatticeControlPointsMechanic.cpp
+++ b/Engine/Plugins/Experimental/MeshModelingToolset/Source/ModelingComponents/Private/Mechanics/LatticeControlPointsMechanic.cpp
@@ -83,7 +83,6 @@ void ULatticeControlPointsMechanic::Setup(UInteractiveTool* ParentToolIn)
 	PointTransformProxy = NewObject<UTransformProxy>(this);
 	
 	// TODO: Maybe don't have the gizmo's axes flip around when it crosses the origin, if possible?
-	// TODO: Enable local vs world gizmo switching (UETOOL-2356)
 	PointTransformGizmo = GizmoManager->CreateCustomTransformGizmo(
 		ETransformGizmoSubElements::FullTranslateRotateScale, this);
 
@@ -244,34 +243,6 @@ void ULatticeControlPointsMechanic::UpdateDrawables()
 }
 
 
-void ULatticeControlPointsMechanic::UpdateDrawablesForPoint(int32 PointIndex)
-{
-	const FVector3d& P = ControlPoints[PointIndex];
-	GeometrySet.UpdatePoint(PointIndex, P);
-
-	DrawnControlPoints->SetPointPosition(PointIndex, static_cast<FVector>(P));
-	DrawnControlPoints->SetPointColor(PointIndex, NormalPointColor);
-
-	if (DrawnControlPoints->IsPointValid(PointIndex))
-	{
-		DrawnControlPoints->SetPointColor(PointIndex, SelectedColor);
-	}
-
-	// TODO: Accelerate this somehow. Don't want to search over the entire set of edges any time one point changes.
-	for (int32 EdgeIndex = 0; EdgeIndex < LatticeEdges.Num(); ++EdgeIndex)		
-	{
-		if (LatticeEdges[EdgeIndex].X == PointIndex)
-		{
-			DrawnLatticeEdges->SetLineStart(EdgeIndex, static_cast<FVector>(P));
-		}
-		else if (LatticeEdges[EdgeIndex].Y == PointIndex)
-		{
-			DrawnLatticeEdges->SetLineEnd(EdgeIndex, static_cast<FVector>(P));
-		}
-	}
-}
-
-
 void ULatticeControlPointsMechanic::GizmoTransformStarted(UTransformProxy* Proxy)
 {
 	ParentTool->GetToolManager()->BeginUndoTransaction(LatticePointMovementTransactionText);
@@ -498,7 +469,7 @@ bool ULatticeControlPointsMechanic::DeselectPoint(int32 PointID)
 {
 	bool PointFound = false;
 	int32 IndexInSelection;
-	// TODO: This might be slow if we have a lot of selected points (UETOOL-2357)
+	
 	if (SelectedPointIDs.Find(PointID, IndexInSelection))
 	{
 		SelectedPointIDs.RemoveAt(IndexInSelection);
diff --git a/Engine/Plugins/Experimental/MeshModelingToolset/Source/ModelingComponents/Public/Mechanics/LatticeControlPointsMechanic.h b/Engine/Plugins/Experimental/MeshModelingToolset/Source/ModelingComponents/Public/Mechanics/LatticeControlPointsMechanic.h
--- a/Engine/Plugins/Experimental/MeshModelingToolset/Source/ModelingComponents/Public/Mechanics/LatticeControlPointsMechanic.h
+++ b/Engine/Plugins/Experimental/MeshModelingToolset/Source/ModelingComponents/Public/Mechanics/LatticeControlPointsMechanic.h
@@ -169,7 +169,6 @@ protected:
 
 	void RebuildDrawables();
 	void UpdateDrawables();
-	void UpdateDrawablesForPoint(int32 PointIndex);
 
 	// Used for expiring undo/redo changes, which compare this to their stored value and expire themselves if they do not match.
 	int32 CurrentChangeStamp = 0;
