Split out the rectangle marquee selection code from LatticeControlPointsMechanic so it can be used by other mechanics.

#rnx

diff --git a/Engine/Plugins/Experimental/MeshModelingToolset/Source/MeshModelingTools/Private/LatticeDeformerTool.cpp b/Engine/Plugins/Experimental/MeshModelingToolset/Source/MeshModelingTools/Private/LatticeDeformerTool.cpp
--- a/Engine/Plugins/Experimental/MeshModelingToolset/Source/MeshModelingTools/Private/LatticeDeformerTool.cpp
+++ b/Engine/Plugins/Experimental/MeshModelingToolset/Source/MeshModelingTools/Private/LatticeDeformerTool.cpp
@@ -17,8 +17,6 @@
 #include "DynamicMeshToMeshDescription.h"
 #include "MeshTransforms.h"
 #include "Algo/ForEach.h"
-#include "CanvasItem.h"
-#include "CanvasTypes.h"
 #include "Operations/FFDLattice.h"
 
 #define LOCTEXT_NAMESPACE "ULatticeDeformerTool"
@@ -76,15 +74,10 @@ void ULatticeDeformerTool::DrawHUD(FCanvas* Canvas, IToolsContextRenderAPI* Rend
 	EViewInteractionState State = RenderAPI->GetViewInteractionState();
 	bool bThisViewHasFocus = !!(State & EViewInteractionState::Focused);
 
-	// Draw the drag rectangle if it's active
-	if (bThisViewHasFocus && ControlPointsMechanic->bIsDragging)
+	if (bThisViewHasFocus)
 	{		
-		FVector2D Start = ControlPointsMechanic->DragStartScreenPosition;
-		FVector2D Curr = ControlPointsMechanic->DragCurrentScreenPosition;
-
-		FCanvasBoxItem BoxItem(Start / Canvas->GetDPIScale(), (Curr - Start) / Canvas->GetDPIScale());
-		BoxItem.SetColor(FLinearColor::White);
-		Canvas->DrawItem(BoxItem);
+		// Draw the drag rectangle if it's active
+		ControlPointsMechanic->DrawHUD(Canvas, RenderAPI);
 	}
 }
 
diff --git a/Engine/Plugins/Experimental/MeshModelingToolset/Source/ModelingComponents/Private/Mechanics/LatticeControlPointsMechanic.cpp b/Engine/Plugins/Experimental/MeshModelingToolset/Source/ModelingComponents/Private/Mechanics/LatticeControlPointsMechanic.cpp
--- a/Engine/Plugins/Experimental/MeshModelingToolset/Source/ModelingComponents/Private/Mechanics/LatticeControlPointsMechanic.cpp
+++ b/Engine/Plugins/Experimental/MeshModelingToolset/Source/ModelingComponents/Private/Mechanics/LatticeControlPointsMechanic.cpp
@@ -25,7 +25,7 @@ static const FText LatticePointMovementTransactionText = LOCTEXT("LatticePointMo
 
 void ULatticeControlPointsMechanic::Setup(UInteractiveTool* ParentToolIn)
 {
-	UInteractionMechanic::Setup(ParentToolIn);
+	URectangleMarqueeMechanic::Setup(ParentToolIn);
 
 	USingleClickInputBehavior* ClickBehavior = NewObject<USingleClickInputBehavior>();
 	ClickBehavior->Initialize(this);
@@ -41,11 +41,6 @@ void ULatticeControlPointsMechanic::Setup(UInteractiveTool* ParentToolIn)
 	HoverBehavior->Modifiers.RegisterModifier(AddToSelectionModifierId, ShiftOrCtrlDown);
 	ParentTool->AddInputBehavior(HoverBehavior);
 
-	UClickDragInputBehavior* ClickDragBehavior = NewObject<UClickDragInputBehavior>(this);
-	ClickDragBehavior->Initialize(this);
-	// TODO: Add CTRL support for drag selection?
-	ParentTool->AddInputBehavior(ClickDragBehavior);
-
 	DrawnControlPoints = NewObject<UPointSetComponent>();
 	DrawnControlPoints->SetPointMaterial(
 		LoadObject<UMaterial>(nullptr, TEXT("/MeshModelingToolset/Materials/PointSetComponentMaterial")));
@@ -378,7 +373,6 @@ void ULatticeControlPointsMechanic::OnClicked(const FInputDeviceRay& ClickPos)
 		ParentTool->GetToolManager()->EndUndoTransaction();
 	}
 
-	bIsDragging = false;
 }
 
 void ULatticeControlPointsMechanic::ChangeSelection(int32 NewPointID, bool bAddToSelection)
@@ -553,35 +547,16 @@ void ULatticeControlPointsMechanic::OnUpdateModifierState(int ModifierID, bool b
 }
 
 
-// ==================== IClickDragBehaviorTarget ====================
-
-FInputRayHit ULatticeControlPointsMechanic::CanBeginClickDragSequence(const FInputDeviceRay& PressPos)
-{
-	FInputRayHit Dummy;
-
-	// This is the boolean that is checked to see if a drag sequence can be started. In our case we want to begin the 
-	// drag sequence even if the first ray doesn't hit anything, so set this to true.
-	Dummy.bHit = true;
-
-	return Dummy;
-}
+// =========================== URectangleMarqueeMechanic ===========================
 
-void ULatticeControlPointsMechanic::OnClickPress(const FInputDeviceRay& PressPos)
+void ULatticeControlPointsMechanic::OnDragRectangleStarted()
 {
-	if (!PressPos.bHas2D)
-	{
-		return;
-	}
-
 	for (int32& PointID : SelectedPointIDs)
 	{
 		DrawnControlPoints->SetPointColor(PointID, NormalPointColor);
 	}
 	CurrentDragSelection.Empty();
 
-	DragStartScreenPosition = PressPos.ScreenPosition;
-	DragStartWorldRay = PressPos.WorldRay;
-
 	// Hide gizmo while dragging
 	if (PointTransformGizmo)
 	{
@@ -589,69 +564,14 @@ void ULatticeControlPointsMechanic::OnClickPress(const FInputDeviceRay& PressPos
 	}
 }
 
-
-static FVector2D PlaneCoordinates(const FVector& Point, const FPlane& Plane, const FVector& UBasisVector, const FVector& VBasisVector)
-{
-	float U = FVector::DotProduct(Point - Plane.GetOrigin(), UBasisVector);
-	float V = FVector::DotProduct(Point - Plane.GetOrigin(), VBasisVector);
-	return FVector2D{ U,V };
-}
-
-void ULatticeControlPointsMechanic::OnClickDrag(const FInputDeviceRay& DragPos)
+void ULatticeControlPointsMechanic::OnDragRectangleChanged(const FCameraRectangle& Rectangle)
 {
-	if (!DragPos.bHas2D)
-	{
-		return;
-	}
-
-	bIsDragging = true;
-	DragCurrentScreenPosition = DragPos.ScreenPosition;
-	DragCurrentWorldRay = DragPos.WorldRay;
-
-	// Intersect the drag rays and and project lattice points all to the same plane in 3D. Then compute 2D coordinates 
-	// and use an AABB test to determine which points are in the drag rectangle.
-
-	FViewCameraState CameraState;
-	GetParentTool()->GetToolManager()->GetContextQueriesAPI()->GetCurrentViewState(CameraState);
-
-	// Create plane in front of camera
-	FPlane CameraPlane(CameraState.Position + CameraState.Forward(), CameraState.Forward());
-
-	FVector UBasisVector = CameraState.Right();
-	FVector VBasisVector = CameraState.Up();
-
-	FVector StartIntersection = FMath::RayPlaneIntersection(DragStartWorldRay.Origin, 
-															DragStartWorldRay.Direction, 
-															CameraPlane);
-	FVector2D Start2D = PlaneCoordinates(StartIntersection, CameraPlane, UBasisVector, VBasisVector);
-
-	FVector CurrentIntersection = FMath::RayPlaneIntersection(DragCurrentWorldRay.Origin, 
-															  DragCurrentWorldRay.Direction, 
-															  CameraPlane);
-	FVector2D Current2D = PlaneCoordinates(CurrentIntersection, CameraPlane, UBasisVector, VBasisVector);
-
-	FBox2D DragBox(Start2D, Start2D);
-	DragBox += Current2D;	// Initialize this way so we don't have to care about min/max
-
+	CurrentDragSelection.Empty();
 	for (int32 PointID = 0; PointID < ControlPoints.Num(); ++PointID)
 	{
 		FVector PointPosition = DrawnControlPoints->GetPoint(PointID).Position;
-		FVector PointIntersection;
-		if (CameraState.bIsOrthographic)
-		{
-			// project directly to plane
-			PointIntersection = FVector::PointPlaneProject(PointPosition, CameraPlane);
-		}
-		else
-		{
-			// intersect along the eye-to-point ray
-			PointIntersection = FMath::RayPlaneIntersection(CameraState.Position, 
-															PointPosition - CameraState.Position,
-															CameraPlane);
-		}
 
-		FVector2D Point2D = PlaneCoordinates(PointIntersection, CameraPlane, UBasisVector, VBasisVector);
-		if (DragBox.IsInside(Point2D))
+		if (Rectangle.IsProjectedPointInRectangle(PointPosition) )
 		{
 			CurrentDragSelection.Add(PointID);
 			DrawnControlPoints->SetPointColor(PointID, SelectedColor);
@@ -663,7 +583,7 @@ void ULatticeControlPointsMechanic::OnClickDrag(const FInputDeviceRay& DragPos)
 	}
 }
 
-void ULatticeControlPointsMechanic::OnClickRelease(const FInputDeviceRay& ReleasePos)
+void ULatticeControlPointsMechanic::OnDragRectangleFinished()
 {
 	// Deselect previous SelectedPointIDs and replace it with "drag selection" points. Do this in one Undo transaction.
 
@@ -691,16 +611,6 @@ void ULatticeControlPointsMechanic::OnClickRelease(const FInputDeviceRay& Releas
 	ParentTool->GetToolManager()->EndUndoTransaction();
 
 	CurrentDragSelection.Empty();
-	bIsDragging = false;
-}
-
-
-void ULatticeControlPointsMechanic::OnTerminateDragSequence()
-{
-	// Not sure how this can happen. Pressing escape quits the tool altogether.
-	CurrentDragSelection.Empty();
-	bIsDragging = false;
-	UpdateGizmoLocation();
 }
 
 const TArray<FVector3d>& ULatticeControlPointsMechanic::GetControlPoints() const
diff --git a/Engine/Plugins/Experimental/MeshModelingToolset/Source/ModelingComponents/Private/Mechanics/RectangleMarqueeMechanic.cpp b/Engine/Plugins/Experimental/MeshModelingToolset/Source/ModelingComponents/Private/Mechanics/RectangleMarqueeMechanic.cpp
new file mode 100644
--- /dev/null
+++ b/Engine/Plugins/Experimental/MeshModelingToolset/Source/ModelingComponents/Private/Mechanics/RectangleMarqueeMechanic.cpp
@@ -0,0 +1,156 @@
+// Copyright Epic Games, Inc. All Rights Reserved.
+
+#include "Mechanics/RectangleMarqueeMechanic.h"
+#include "BaseBehaviors/ClickDragBehavior.h"
+#include "CanvasTypes.h"
+#include "CanvasItem.h"
+#include "ToolSceneQueriesUtil.h"
+
+#define LOCTEXT_NAMESPACE "URectangleMarqueeMechanic"
+
+
+static FVector2D PlaneCoordinates(const FVector& Point, const FPlane& Plane, const FVector& UBasisVector, const FVector& VBasisVector)
+{
+	float U = FVector::DotProduct(Point - Plane.GetOrigin(), UBasisVector);
+	float V = FVector::DotProduct(Point - Plane.GetOrigin(), VBasisVector);
+	return FVector2D{ U,V };
+}
+
+FCameraRectangle::FCameraRectangle(const FViewCameraState& CachedCameraState,
+								   const FRay& DragStartWorldRay,
+								   const FRay& DragEndWorldRay)
+{
+	// Create a plane just in front of the camera
+	CameraOrigin = CachedCameraState.Position;
+	CameraPlane = FPlane(CachedCameraState.Position + CachedCameraState.Forward(), CachedCameraState.Forward());
+	bCameraIsOrthographic = CachedCameraState.bIsOrthographic;
+
+	// Intersect the drag rays with the camera plane and compute their coordinates in the camera basis
+	UBasisVector = CachedCameraState.Right();
+	VBasisVector = CachedCameraState.Up();
+
+	FVector StartIntersection = FMath::RayPlaneIntersection(DragStartWorldRay.Origin,
+															DragStartWorldRay.Direction,
+															CameraPlane);
+	FVector2D Start2D = PlaneCoordinates(StartIntersection,
+										 CameraPlane,
+										 UBasisVector,
+										 VBasisVector);
+
+	FVector CurrentIntersection = FMath::RayPlaneIntersection(DragEndWorldRay.Origin,
+															  DragEndWorldRay.Direction,
+															  CameraPlane);
+
+	FVector2D Current2D = PlaneCoordinates(CurrentIntersection,
+										   CameraPlane,
+										   UBasisVector,
+										   VBasisVector);
+
+	RectangleCorners = FBox2D(Start2D, Start2D);
+	RectangleCorners += Current2D;	// Initialize this way so we don't have to care about min/max
+}
+
+bool FCameraRectangle::IsProjectedPointInRectangle(const FVector& Point) const
+{
+	FVector ProjectedPoint;
+	if (bCameraIsOrthographic)
+	{
+		// project directly to plane
+		ProjectedPoint = FVector::PointPlaneProject(Point, CameraPlane);
+	}
+	else
+	{
+		// intersect along the eye-to-point ray
+		ProjectedPoint = FMath::RayPlaneIntersection(CameraOrigin,
+													 Point - CameraOrigin,
+													 CameraPlane);
+	}
+
+	FVector2D Point2D = PlaneCoordinates(ProjectedPoint, CameraPlane, UBasisVector, VBasisVector);
+	return RectangleCorners.IsInside(Point2D);
+}
+
+// ---------------------------------------
+
+void URectangleMarqueeMechanic::Setup(UInteractiveTool* ParentToolIn)
+{
+	UInteractionMechanic::Setup(ParentToolIn);
+
+	UClickDragInputBehavior* ClickDragBehavior = NewObject<UClickDragInputBehavior>(this);
+	ClickDragBehavior->Initialize(this);
+	ParentTool->AddInputBehavior(ClickDragBehavior);
+}
+
+void URectangleMarqueeMechanic::Render(IToolsContextRenderAPI* RenderAPI)
+{
+	// Cache the camera state
+	GetParentTool()->GetToolManager()->GetContextQueriesAPI()->GetCurrentViewState(CachedCameraState);
+}
+
+FInputRayHit URectangleMarqueeMechanic::CanBeginClickDragSequence(const FInputDeviceRay& PressPos)
+{
+	FInputRayHit Dummy;
+
+	// This is the boolean that is checked to see if a drag sequence can be started. In our case we want to begin the 
+	// drag sequence even if the first ray doesn't hit anything, so set this to true.
+	Dummy.bHit = true;
+
+	return Dummy;
+}
+
+void URectangleMarqueeMechanic::OnClickPress(const FInputDeviceRay& PressPos)
+{
+	if (!PressPos.bHas2D)
+	{
+		bIsDragging = false;
+		return;
+	}
+
+	DragStartScreenPosition = PressPos.ScreenPosition;
+	DragStartWorldRay = PressPos.WorldRay;
+
+	OnDragRectangleStarted();
+}
+
+void URectangleMarqueeMechanic::OnClickDrag(const FInputDeviceRay& DragPos)
+{
+	if (!DragPos.bHas2D)
+	{
+		return;
+	}
+
+	bIsDragging = true;
+	DragCurrentScreenPosition = DragPos.ScreenPosition;
+	FRay DragCurrentWorldRay = DragPos.WorldRay;
+
+	FCameraRectangle Rectangle(CachedCameraState, DragStartWorldRay, DragCurrentWorldRay);
+
+	OnDragRectangleChanged(Rectangle);
+}
+
+void URectangleMarqueeMechanic::OnClickRelease(const FInputDeviceRay& ReleasePos)
+{
+	bIsDragging = false;
+	OnDragRectangleFinished();
+}
+
+void URectangleMarqueeMechanic::OnTerminateDragSequence()
+{
+	bIsDragging = false;
+	OnDragRectangleFinished();
+}
+
+
+void URectangleMarqueeMechanic::DrawHUD(FCanvas* Canvas, IToolsContextRenderAPI* RenderAPI)
+{
+	if (bIsDragging)
+	{
+		FVector2D Start = DragStartScreenPosition;
+		FVector2D Curr = DragCurrentScreenPosition;
+		FCanvasBoxItem BoxItem(Start / Canvas->GetDPIScale(), (Curr - Start) / Canvas->GetDPIScale());
+		BoxItem.SetColor(FLinearColor::White);
+		Canvas->DrawItem(BoxItem);
+	}
+}
+
+#undef LOCTEXT_NAMESPACE
diff --git a/Engine/Plugins/Experimental/MeshModelingToolset/Source/ModelingComponents/Public/Mechanics/LatticeControlPointsMechanic.h b/Engine/Plugins/Experimental/MeshModelingToolset/Source/ModelingComponents/Public/Mechanics/LatticeControlPointsMechanic.h
--- a/Engine/Plugins/Experimental/MeshModelingToolset/Source/ModelingComponents/Public/Mechanics/LatticeControlPointsMechanic.h
+++ b/Engine/Plugins/Experimental/MeshModelingToolset/Source/ModelingComponents/Public/Mechanics/LatticeControlPointsMechanic.h
@@ -6,7 +6,7 @@
 
 #include "BaseBehaviors/BehaviorTargetInterfaces.h"
 #include "FrameTypes.h"
-#include "InteractionMechanic.h"
+#include "RectangleMarqueeMechanic.h"
 #include "InteractiveToolChange.h"
 #include "Snapping/PointPlanarSnapSolver.h"
 #include "Spatial/GeometrySet3.h"
@@ -28,7 +28,7 @@ class UTransformProxy;
 
 UCLASS()
 class MODELINGCOMPONENTS_API ULatticeControlPointsMechanic : 
-	public UInteractionMechanic, public IClickBehaviorTarget, public IHoverBehaviorTarget, public IClickDragBehaviorTarget
+	public URectangleMarqueeMechanic, public IClickBehaviorTarget, public IHoverBehaviorTarget
 {
 	GENERATED_BODY()
 
@@ -68,19 +68,10 @@ public:
 	virtual void OnEndHover() override;
 	virtual void OnUpdateModifierState(int ModifierID, bool bIsOn) override;
 
-	// IClickDragBehaviorTarget
-	FInputRayHit CanBeginClickDragSequence(const FInputDeviceRay& PressPos) override;
-	void OnClickPress(const FInputDeviceRay& PressPos) override;
-	void OnClickDrag(const FInputDeviceRay& DragPos) override;
-	void OnClickRelease(const FInputDeviceRay& ReleasePos) override;
-	void OnTerminateDragSequence() override;
-
-	// Marquee/rectangle select support
-	FVector2D DragStartScreenPosition;
-	FRay DragStartWorldRay;
-	FVector2D DragCurrentScreenPosition;
-	FRay DragCurrentWorldRay;
-	bool bIsDragging;
+	// URectangleMarqueeMechanic implementation
+	void OnDragRectangleStarted() final;
+	void OnDragRectangleChanged(const FCameraRectangle& Rectangle) final;
+	void OnDragRectangleFinished() final;
 
 	bool bHasChanged = false;
 
@@ -135,9 +126,6 @@ protected:
 	void GizmoTransformStarted(UTransformProxy* Proxy);
 	void GizmoTransformEnded(UTransformProxy* Proxy);
 
-	// Cached render information for viewport/scene interaction
-	FViewCameraState CachedCameraState;
-
 	// Support for hovering
 	TFunction<bool(const FVector3d&, const FVector3d&)> GeometrySetToleranceTest;
 	int32 HoveredPointID = -1;
diff --git a/Engine/Plugins/Experimental/MeshModelingToolset/Source/ModelingComponents/Public/Mechanics/RectangleMarqueeMechanic.h b/Engine/Plugins/Experimental/MeshModelingToolset/Source/ModelingComponents/Public/Mechanics/RectangleMarqueeMechanic.h
new file mode 100644
--- /dev/null
+++ b/Engine/Plugins/Experimental/MeshModelingToolset/Source/ModelingComponents/Public/Mechanics/RectangleMarqueeMechanic.h
@@ -0,0 +1,96 @@
+// Copyright Epic Games, Inc. All Rights Reserved.
+
+#pragma once
+
+#include "CoreMinimal.h"
+#include "InteractionMechanic.h"
+#include "BaseBehaviors/BehaviorTargetInterfaces.h"
+#include "VectorTypes.h"
+
+#include "RectangleMarqueeMechanic.generated.h"
+
+
+/// Struct containing:
+///		- camera information, 
+///		- a 3D plane just in front of the camera,
+///		- a 2D basis for coordinates in this plane, and
+///		- the corners of a rectangle contained in this plane, in this 2D basis
+/// 
+struct FCameraRectangle
+{
+	FCameraRectangle(const FViewCameraState& CachedCameraState,
+					 const FRay& DragStartWorldRay,
+					 const FRay& DragEndWorldRay);
+
+	FVector CameraOrigin;
+	bool bCameraIsOrthographic;
+	FPlane CameraPlane;
+	FVector UBasisVector;
+	FVector VBasisVector;
+	FBox2D RectangleCorners;
+
+	// Project the given 3D point to the camera plane and test if it's in the rectangle
+	bool IsProjectedPointInRectangle(const FVector& Point) const;
+
+	// TODO: Rectangle vs projected triangle, edge, etc.
+};
+
+
+/*
+ * Mechanic for a rectangle "marquee" selection. It creates and maintains the 2D rectangle associated with a mouse drag. 
+ * It does not test against any scene geometry, nor does it maintain any sort of list of selected objects.
+ *
+ * To use this class:
+ *	- create a subclass of this class
+ *  - override the pure virtual functions. These will be called from this class in response to mouse events
+ *  - use the FCameraRectangle passed into RectangleChanged to test against your scene geometry
+ *
+ * See ULatticeControlPointsMechanic for an example.
+ */
+
+UCLASS()
+class MODELINGCOMPONENTS_API URectangleMarqueeMechanic : public UInteractionMechanic, public IClickDragBehaviorTarget
+{
+	GENERATED_BODY()
+
+public:
+
+	// UInteractionMechanic
+	void Setup(UInteractiveTool* ParentTool) override;
+	void Render(IToolsContextRenderAPI* RenderAPI) override;
+
+	// REVIEW: Should this be added to UInteractionMechanic?
+	virtual void DrawHUD(FCanvas* Canvas, IToolsContextRenderAPI* RenderAPI);
+
+protected:
+
+	bool bIsDragging;
+	FViewCameraState CachedCameraState;
+	FVector2D DragStartScreenPosition;
+	FRay DragStartWorldRay;
+	FVector2D DragCurrentScreenPosition;
+
+private:
+
+	//
+	// Override these
+	//
+
+	virtual void OnDragRectangleStarted() 
+	PURE_VIRTUAL(URectangleMarqueeMechanic::OnDragRectangleStarted, );
+
+	virtual void OnDragRectangleChanged(const FCameraRectangle& CurrentRectangle) 
+	PURE_VIRTUAL(URectangleMarqueeMechanic::OnDragRectangleChanged, );
+
+	virtual void OnDragRectangleFinished() 
+	PURE_VIRTUAL(URectangleMarqueeMechanic::OnDragRectangleFinished, );
+
+
+private:
+	
+	FInputRayHit CanBeginClickDragSequence(const FInputDeviceRay& PressPos) final;
+	void OnClickPress(const FInputDeviceRay& PressPos) final;
+	void OnClickDrag(const FInputDeviceRay& DragPos) final;
+	void OnClickRelease(const FInputDeviceRay& ReleasePos) final;
+	void OnTerminateDragSequence() final;
+};
