MeshModelingTools: Added an option in revolve tools to offset the profile curve along the axis to create springs, and gave an automatic way to specify the number of steps based on the revolution extent. Also made it easier to silently update a single watched property.

#rnx

diff --git a/Engine/Plugins/Experimental/GeometryProcessing/Source/GeometricObjects/Private/Util/RevolveUtil.cpp b/Engine/Plugins/Experimental/GeometryProcessing/Source/GeometricObjects/Private/Util/RevolveUtil.cpp
--- a/Engine/Plugins/Experimental/GeometryProcessing/Source/GeometricObjects/Private/Util/RevolveUtil.cpp
+++ b/Engine/Plugins/Experimental/GeometryProcessing/Source/GeometricObjects/Private/Util/RevolveUtil.cpp
@@ -3,7 +3,7 @@
 #include "Util/RevolveUtil.h"
 
 void RevolveUtil::GenerateSweepCurve(const FVector3d& RevolutionAxisOrigin, const FVector3d& RevolutionAxisDirection, 
-	double DegreesOffset, double DegreesPerStep, int TotalNumFrames, bool bWeldFullRevolution, TArray<FFrame3d> &SweepCurveOut)
+	double DegreesOffset, double DegreesPerStep, double DownAxisOffset, int TotalNumFrames, TArray<FFrame3d> &SweepCurveOut)
 {
 	// For a revolve, we need to sweep along a circular path around the axis of rotation. While we could pick
 	// any arbitrary frame to create the curve, we choose the standard world frame since that is what the 
@@ -13,7 +13,7 @@ void RevolveUtil::GenerateSweepCurve(const FVector3d& RevolutionAxisOrigin, cons
 	{
 		// Revolve the origin frame relative the axis
 		FQuaterniond Rotation(RevolutionAxisDirection, DegreesOffset + DegreesPerStep * i, true);
-		FVector3d NewOrigin = Rotation * WorldOriginWrtAxisOrigin + RevolutionAxisOrigin;
+		FVector3d NewOrigin = Rotation * (WorldOriginWrtAxisOrigin + i * DownAxisOffset * RevolutionAxisDirection) + RevolutionAxisOrigin;
 		FFrame3d NewFrame(NewOrigin, Rotation);
 		SweepCurveOut.Add(NewFrame);
 	}
diff --git a/Engine/Plugins/Experimental/GeometryProcessing/Source/GeometricObjects/Public/Util/RevolveUtil.h b/Engine/Plugins/Experimental/GeometryProcessing/Source/GeometricObjects/Public/Util/RevolveUtil.h
--- a/Engine/Plugins/Experimental/GeometryProcessing/Source/GeometricObjects/Public/Util/RevolveUtil.h
+++ b/Engine/Plugins/Experimental/GeometryProcessing/Source/GeometricObjects/Public/Util/RevolveUtil.h
@@ -15,10 +15,18 @@ namespace RevolveUtil
 {
 	/** 
 	 * Generates a sweep curve with the initial frame at the origin and the rest rotated around the given axis.
+	 *
+	 * @param RevolutionAxisOrigin
+	 * @param RevolutionAxisDirection Should be normalized.
+	 * @param DegreesOffset Number of degrees to rotate the first frame.
+	 * @param DegreesPerStep Number of degrees each frame is rotated from the previous
+	 * @param DownAxisOffset Distance to move each frame up the axis from the previous (for spirals)
+	 * @param TotalNumFrames Number of frames to create, including the first frame
+	 * @param SweepCurveOut Output
 	 */
 	GEOMETRICOBJECTS_API
 	void GenerateSweepCurve(const FVector3d& RevolutionAxisOrigin, const FVector3d& RevolutionAxisDirection, double DegreesOffset,
-		double DegreesPerStep, int TotalNumFrames, bool bWeldFullRevolution, TArray<FFrame3d>& SweepCurveOut);
+		double DegreesPerStep, double DownAxisOffset, int TotalNumFrames, TArray<FFrame3d>& SweepCurveOut);
 
 	/**
 	 * Checks the profile curve for points that lie within a particular tolerance of the revolution axis. If they
diff --git a/Engine/Plugins/Experimental/MeshModelingToolset/Source/MeshModelingTools/Private/Properties/RevolveProperties.cpp b/Engine/Plugins/Experimental/MeshModelingToolset/Source/MeshModelingTools/Private/Properties/RevolveProperties.cpp
--- a/Engine/Plugins/Experimental/MeshModelingToolset/Source/MeshModelingTools/Private/Properties/RevolveProperties.cpp
+++ b/Engine/Plugins/Experimental/MeshModelingToolset/Source/MeshModelingTools/Private/Properties/RevolveProperties.cpp
@@ -25,13 +25,18 @@ void URevolveProperties::ApplyToCurveSweepOp(const UNewMeshMaterialProperties& M
 		}
 	}
 
+	double TotalRevolutionDegrees = (DownAxisOffsetPerDegree == 0) ? ClampedRevolutionDegrees : RevolutionDegrees;
+
+	int32 Steps = bExplicitSteps ? NumExplicitSteps : FMath::CeilToInt(TotalRevolutionDegrees / MaxDegreesPerStep);
+	
+	double DegreesPerStep = TotalRevolutionDegrees / Steps;
 	double DegreesOffset = RevolutionDegreesOffset;
-	double DegreesPerStep = RevolutionDegrees / Steps;
 	if (bReverseRevolutionDirection)
 	{
 		DegreesPerStep *= -1;
 		DegreesOffset *= -1;
 	}
+	double DownAxisOffsetPerStep = TotalRevolutionDegrees * DownAxisOffsetPerDegree / Steps;
 
 	if (bProfileIsCrossSectionOfSide && DegreesPerStep != 0 && abs(DegreesPerStep) < 180)
 	{
@@ -39,14 +44,14 @@ void URevolveProperties::ApplyToCurveSweepOp(const UNewMeshMaterialProperties& M
 	}
 
 	// Generate the sweep curve
-	CurveSweepOpOut.bSweepCurveIsClosed = bWeldFullRevolution && RevolutionDegrees == 360;
+	CurveSweepOpOut.bSweepCurveIsClosed = bWeldFullRevolution && DownAxisOffsetPerDegree == 0 && RevolutionDegrees == 360;
 	int32 NumSweepFrames = CurveSweepOpOut.bSweepCurveIsClosed ? Steps : Steps + 1; // If closed, last sweep frame is also first
 	CurveSweepOpOut.SweepCurve.Reserve(NumSweepFrames);
 	RevolveUtil::GenerateSweepCurve(RevolutionAxisOrigin, RevolutionAxisDirection, DegreesOffset,
-		DegreesPerStep, NumSweepFrames, bWeldFullRevolution, CurveSweepOpOut.SweepCurve);
+		DegreesPerStep, DownAxisOffsetPerStep, NumSweepFrames, CurveSweepOpOut.SweepCurve);
 
 	// Weld any vertices that are on the axis
-	if (bWeldVertsOnAxis)
+	if (bWeldVertsOnAxis && DownAxisOffsetPerStep == 0)
 	{
 		RevolveUtil::WeldPointsOnAxis(CurveSweepOpOut.ProfileCurve, RevolutionAxisOrigin,
 			RevolutionAxisDirection, AxisWeldTolerance, CurveSweepOpOut.ProfileVerticesToWeld);
diff --git a/Engine/Plugins/Experimental/MeshModelingToolset/Source/MeshModelingTools/Public/Properties/RevolveProperties.h b/Engine/Plugins/Experimental/MeshModelingToolset/Source/MeshModelingTools/Public/Properties/RevolveProperties.h
--- a/Engine/Plugins/Experimental/MeshModelingToolset/Source/MeshModelingTools/Public/Properties/RevolveProperties.h
+++ b/Engine/Plugins/Experimental/MeshModelingToolset/Source/MeshModelingTools/Public/Properties/RevolveProperties.h
@@ -62,18 +62,60 @@ class MESHMODELINGTOOLS_API URevolveProperties : public UInteractiveToolProperty
 
 public:
 
+	URevolveProperties()
+	{
+		// We want the revolution degrees to be clamped to 360 when there is no offset along axis, and extendable
+		// beyond 360 when there is one (to make springs, etc). There's not currently a way to have a conditional
+		// clamp value, so instead, we swap the variables that we make visible to the user (RevolutionDegrees and
+		// ClampedRevolutionDegrees). However, we have to do some work to keep the values consistent so that they
+		// do not jump when they are swapped.
+		ClampedRevolutionDegreesWatcherIndex = WatchProperty(ClampedRevolutionDegrees,
+			[this](const double& NewClamped) 
+			{ 
+				RevolutionDegrees = NewClamped; 
+				SilentUpdateWatcherAtIndex(RevolutionDegreesWatcherIndex); //avoid triggering that watcher
+			});
+
+		RevolutionDegreesWatcherIndex = WatchProperty(RevolutionDegrees,
+			[this](const double& NewNonClamped)
+			{
+				ClampedRevolutionDegrees = FMath::Min(RevolutionDegrees, 360.0);
+				SilentUpdateWatcherAtIndex(ClampedRevolutionDegreesWatcherIndex); //avoid triggering that watcher
+			});
+	}
+
+	/** Revolution extent. Clamped to a max of 360 when not offsetting along axis.*/
+	UPROPERTY(EditAnywhere, Category = RevolveSettings, meta = (UIMin = "0", UIMax = "360", ClampMin = "0", ClampMax = "360", 
+		EditCondition = "DownAxisOffsetPerDegree == 0", EditConditionHides, DisplayName = "Revolution Degrees"))
+	double ClampedRevolutionDegrees = 360;
+
 	/** Revolution extent. */
-	UPROPERTY(EditAnywhere, Category = RevolveSettings, meta = (UIMin = "0", UIMax = "360", ClampMin = "0", ClampMax = "360"))
+	UPROPERTY(EditAnywhere, Category = RevolveSettings, meta = (UIMin = "0", UIMax = "3600", ClampMin = "0", ClampMax = "360000",
+		EditCondition = "DownAxisOffsetPerDegree != 0", EditConditionHides))
 	double RevolutionDegrees = 360;
 
+	/** When true, the number of steps can be specified explicitly. When false, the number of steps is adjusted automatically. */
+	UPROPERTY(EditAnywhere, Category = RevolveSettings)
+	bool bExplicitSteps = false;
+
+	/** Number of steps to take while revolving. */
+	UPROPERTY(EditAnywhere, Category = RevolveSettings, meta = (UIMin = "1", ClampMin = "1", UIMax = "100", ClampMax = "5000",
+		Display = "Steps", EditCondition = "bExplicitSteps", EditConditionHides))
+	int NumExplicitSteps = 24;
+
+	/** The revolution is split into a number of steps such that each step moves the revolution no more than this number of degrees.  */
+	UPROPERTY(EditAnywhere, Category = RevolveSettings, meta = (UIMin = "1", ClampMin = "1", UIMax = "100", ClampMax = "5000",
+		EditCondition = "!bExplicitSteps", EditConditionHides))
+	double MaxDegreesPerStep = 15;
+
+	/** How far to move each step down the axis (per degree). Used to create spirals. */
+	UPROPERTY(EditAnywhere, Category = RevolveSettings, meta = (UIMin = "-1", UIMax = "1", ClampMin = "-100000", ClampMax = "100000"))
+	double DownAxisOffsetPerDegree = 0;
+
 	/** The angle by which to shift the profile curve around the axis before beginning the revolve */
 	UPROPERTY(EditAnywhere, Category = RevolveSettings, meta = (UIMin = "-360", UIMax = "360", ClampMin = "-36000", ClampMax = "36000"), AdvancedDisplay)
 	double RevolutionDegreesOffset = 0;
 
-	/** Number of steps to take while revolving. */
-	UPROPERTY(EditAnywhere, Category = RevolveSettings, meta = (UIMin = "1", ClampMin = "1", UIMax = "100", ClampMax = "5000"))
-	int Steps = 24;
-
 	/** By default, revolution is done counterclockwise if looking down the revolution axis. This reverses the direction.*/
 	UPROPERTY(EditAnywhere, Category = RevolveSettings)
 	bool bReverseRevolutionDirection = false;
@@ -106,12 +148,14 @@ public:
 
 	/** Determines how caps are created if the revolution is partial. Not relevant if the
 	  revolution is full and welded. */
-	UPROPERTY(EditAnywhere, Category = RevolveSettings, AdvancedDisplay)
+	UPROPERTY(EditAnywhere, Category = RevolveSettings, AdvancedDisplay, meta = (
+		EditCondition = "DownAxisOffsetPerDegree != 0 || RevolutionDegrees != 360", EditConditionHides))
 	ERevolvePropertiesCapFillMode CapFillMode = ERevolvePropertiesCapFillMode::Delaunay;
 
 	/** If true, the ends of a fully revolved profile are welded together, rather than duplicating
 	  vertices at the seam. Not relevant if the revolution is not full. */
-	UPROPERTY(EditAnywhere, Category = RevolveSettings, AdvancedDisplay)
+	UPROPERTY(EditAnywhere, Category = RevolveSettings, AdvancedDisplay, meta = (
+		EditCondition = "DownAxisOffsetPerDegree == 0 && RevolutionDegrees == 360", EditConditionHides))
 	bool bWeldFullRevolution = true;
 
 	/** If true, vertices sufficiently close to the axis will not be replicated, instead reusing
@@ -120,7 +164,8 @@ public:
 	bool bWeldVertsOnAxis = true;
 
 	/** If welding vertices on the axis, the distance that a vertex can be from the axis and still be welded */
-	UPROPERTY(EditAnywhere, Category = RevolveSettings, AdvancedDisplay, meta = (ClampMin = "0.0", ClampMax = "20.0", EditCondition = "bWeldVertsOnAxis"))
+	UPROPERTY(EditAnywhere, Category = RevolveSettings, AdvancedDisplay, meta = (ClampMin = "0.0", ClampMax = "20.0", 
+		EditCondition = "bWeldVertsOnAxis", EditConditionHides))
 	double AxisWeldTolerance = 0.1;
 
 	/** If true, normals are not averaged or shared between triangles with sufficient angle difference. */
@@ -129,7 +174,8 @@ public:
 
 	/** When using sharp normals, the degree difference to accept between adjacent triangle normals to allow them to share
 	 normals at their vertices. */
-	UPROPERTY(EditAnywhere, Category = RevolveSettings, meta = (ClampMin = "0.0", ClampMax = "90.0", EditCondition = "bSharpNormals"))
+	UPROPERTY(EditAnywhere, Category = RevolveSettings, meta = (ClampMin = "0.0", ClampMax = "90.0", 
+		EditCondition = "bSharpNormals", EditConditionHides))
 	double SharpNormalAngleTolerance = 0.1;
 
 	/** If true, UV coordinates will be flipped in the V direction. */
@@ -152,4 +198,8 @@ public:
 	void ApplyToCurveSweepOp(const UNewMeshMaterialProperties& MaterialProperties,
 		const FVector3d& RevolutionAxisOrigin, const FVector3d& RevolutionAxisDirection,
 		FCurveSweepOp& CurveSweepOpOut) const;
+
+	protected:
+		int32 ClampedRevolutionDegreesWatcherIndex;
+		int32 RevolutionDegreesWatcherIndex;
 };
diff --git a/Engine/Source/Runtime/Experimental/InteractiveToolsFramework/Public/InteractiveTool.h b/Engine/Source/Runtime/Experimental/InteractiveToolsFramework/Public/InteractiveTool.h
--- a/Engine/Source/Runtime/Experimental/InteractiveToolsFramework/Public/InteractiveTool.h
+++ b/Engine/Source/Runtime/Experimental/InteractiveToolsFramework/Public/InteractiveTool.h
@@ -101,18 +101,30 @@ public:
 			PropWatcher->SilentUpdate();
 		}
 	}
+	/** 
+	 * Silently updates just a single watcher, using an index gotten from WatchProperty.
+	 * Useful when you want watching to still work for other properties changed in the
+	 * same tick.
+	 */
+	void SilentUpdateWatcherAtIndex(int32 i)
+	{
+		check(i >= 0 && i < PropertyWatchers.Num());
+		PropertyWatchers[i]->SilentUpdate();
+	}
 
+	/** @return Index of the watcher, which can be used in SilentUpdateWatcherAtIndex */
 	template <typename PropType>
-	void WatchProperty(const PropType& ValueIn,
+	int32 WatchProperty(const PropType& ValueIn,
 					   typename TPropertyWatcher<PropType>::FChangedCallback OnChangedIn)
 	{
-		PropertyWatchers.Emplace(MakeUnique<TPropertyWatcher<PropType>>(ValueIn, OnChangedIn));
+		return PropertyWatchers.Emplace(MakeUnique<TPropertyWatcher<PropType>>(ValueIn, OnChangedIn));
 	}
+	/** @return Index of the watcher, which can be used in SilentUpdateWatcherAtIndex */
 	template <typename PropType>
-	void WatchProperty(typename TPropertyWatcher<PropType>::FValueGetter GetValueIn,
+	int32 WatchProperty(typename TPropertyWatcher<PropType>::FValueGetter GetValueIn,
 					   typename TPropertyWatcher<PropType>::FChangedCallback OnChangedIn)
 	{
-		PropertyWatchers.Emplace(MakeUnique<TPropertyWatcher<PropType>>(GetValueIn, OnChangedIn));
+		return PropertyWatchers.Emplace(MakeUnique<TPropertyWatcher<PropType>>(GetValueIn, OnChangedIn));
 	}
 private:
 	TArray<TUniquePtr<FPropertyWatcher>> PropertyWatchers;
