Add an optional IRenderMeshPostProcessor to USimpleDynamicMeshComponent that will be called before sending the mesh for rendering.

#rnx

diff --git a/Engine/Plugins/Experimental/MeshModelingToolset/Source/ModelingComponents/Private/SimpleDynamicMeshComponent.cpp b/Engine/Plugins/Experimental/MeshModelingToolset/Source/ModelingComponents/Private/SimpleDynamicMeshComponent.cpp
--- a/Engine/Plugins/Experimental/MeshModelingToolset/Source/ModelingComponents/Private/SimpleDynamicMeshComponent.cpp
+++ b/Engine/Plugins/Experimental/MeshModelingToolset/Source/ModelingComponents/Private/SimpleDynamicMeshComponent.cpp
@@ -47,6 +47,29 @@ USimpleDynamicMeshComponent::USimpleDynamicMeshComponent(const FObjectInitialize
 }
 
 
+FDynamicMesh3* USimpleDynamicMeshComponent::GetRenderMesh()
+{
+	if (RenderMeshPostProcessor && RenderMesh)
+	{
+		return RenderMesh.Get();
+	}
+	else
+	{
+		return Mesh.Get();
+	}
+}
+
+const FDynamicMesh3* USimpleDynamicMeshComponent::GetRenderMesh() const
+{
+	if (RenderMeshPostProcessor && RenderMesh)
+	{
+		return RenderMesh.Get();
+	}
+	else
+	{
+		return Mesh.Get();
+	}
+}
 
 void USimpleDynamicMeshComponent::InitializeMesh(FMeshDescription* MeshDescription)
 {
@@ -61,6 +84,22 @@ void USimpleDynamicMeshComponent::InitializeMesh(FMeshDescription* MeshDescripti
 	NotifyMeshUpdated();
 }
 
+void USimpleDynamicMeshComponent::SetRenderMeshPostProcessor(TUniquePtr<IRenderMeshPostProcessor> Processor)
+{
+	RenderMeshPostProcessor = MoveTemp(Processor);
+	if (RenderMeshPostProcessor)
+	{
+		if (!RenderMesh)
+		{
+			RenderMesh = MakeUnique<FDynamicMesh3>(*Mesh);
+		}
+	}
+	else
+	{
+		// No post processor, no render mesh
+		RenderMesh = nullptr;
+	}
+}
 
 void USimpleDynamicMeshComponent::UpdateTangents(const FMeshTangentsf* ExternalTangents, bool bFastUpdateIfPossible)
 {
@@ -195,8 +234,7 @@ void USimpleDynamicMeshComponent::SetDrawOnTop(bool bSet)
 	bUseEditorCompositing = bSet;
 }
 
-
-void USimpleDynamicMeshComponent::NotifyMeshUpdated()
+void USimpleDynamicMeshComponent::ResetProxy()
 {
 	// Need to recreate scene proxy to send it over
 	MarkRenderStateDirty();
@@ -209,11 +247,27 @@ void USimpleDynamicMeshComponent::NotifyMeshUpdated()
 	}
 }
 
+void USimpleDynamicMeshComponent::NotifyMeshUpdated()
+{
+	if (RenderMeshPostProcessor)
+	{
+		RenderMeshPostProcessor->ProcessMesh(*Mesh, *RenderMesh);
+	}
+
+	ResetProxy();
+}
+
 
 void USimpleDynamicMeshComponent::FastNotifyColorsUpdated()
 {
+	if (RenderMeshPostProcessor)
+	{
+		RenderMeshPostProcessor->ProcessMesh(*Mesh, *RenderMesh);
+	}
+
 	FSimpleDynamicMeshSceneProxy* Proxy = GetCurrentSceneProxy();
-	if (Proxy != nullptr)
+
+	if (Proxy && Proxy->RenderMeshLayoutMatchesRenderBuffers())
 	{
 		if (TriangleColorFunc != nullptr &&  Proxy->bUsePerTriangleColor == false )
 		{
@@ -231,7 +285,7 @@ void USimpleDynamicMeshComponent::FastNotifyColorsUpdated()
 	}
 	else
 	{
-		NotifyMeshUpdated();
+		ResetProxy();
 	}
 }
 
@@ -239,7 +293,14 @@ void USimpleDynamicMeshComponent::FastNotifyColorsUpdated()
 
 void USimpleDynamicMeshComponent::FastNotifyPositionsUpdated(bool bNormals, bool bColors, bool bUVs)
 {
-	if (GetCurrentSceneProxy() != nullptr)
+	if (RenderMeshPostProcessor)
+	{
+		RenderMeshPostProcessor->ProcessMesh(*Mesh, *RenderMesh);
+	}
+
+	FSimpleDynamicMeshSceneProxy* Proxy = GetCurrentSceneProxy();
+
+	if (Proxy && Proxy->RenderMeshLayoutMatchesRenderBuffers())
 	{
 		GetCurrentSceneProxy()->FastUpdateVertices(true, bNormals, bColors, bUVs);
 		//MarkRenderDynamicDataDirty();
@@ -249,15 +310,22 @@ void USimpleDynamicMeshComponent::FastNotifyPositionsUpdated(bool bNormals, bool
 	}
 	else
 	{
-		NotifyMeshUpdated();
+		ResetProxy();
 	}
 }
 
 
 void USimpleDynamicMeshComponent::FastNotifyVertexAttributesUpdated(bool bNormals, bool bColors, bool bUVs)
 {
+	if (RenderMeshPostProcessor)
+	{
+		RenderMeshPostProcessor->ProcessMesh(*Mesh, *RenderMesh);
+	}
+
+	FSimpleDynamicMeshSceneProxy* Proxy = GetCurrentSceneProxy();
+
 	check(bNormals || bColors || bUVs);
-	if (GetCurrentSceneProxy() != nullptr)
+	if (Proxy && Proxy->RenderMeshLayoutMatchesRenderBuffers())
 	{
 		GetCurrentSceneProxy()->FastUpdateVertices(false, bNormals, bColors, bUVs);
 		//MarkRenderDynamicDataDirty();
@@ -265,15 +333,22 @@ void USimpleDynamicMeshComponent::FastNotifyVertexAttributesUpdated(bool bNormal
 	}
 	else
 	{
-		NotifyMeshUpdated();
+		ResetProxy();
 	}
 }
 
 
 void USimpleDynamicMeshComponent::FastNotifyVertexAttributesUpdated(EMeshRenderAttributeFlags UpdatedAttributes)
 {
+	if (RenderMeshPostProcessor)
+	{
+		RenderMeshPostProcessor->ProcessMesh(*Mesh, *RenderMesh);
+	}
+
+	FSimpleDynamicMeshSceneProxy* Proxy = GetCurrentSceneProxy();
+
 	check(UpdatedAttributes != EMeshRenderAttributeFlags::None);
-	if (GetCurrentSceneProxy() != nullptr)
+	if (Proxy && Proxy->RenderMeshLayoutMatchesRenderBuffers())
 	{
 		bool bPositions = (UpdatedAttributes & EMeshRenderAttributeFlags::Positions) != EMeshRenderAttributeFlags::None;
 		GetCurrentSceneProxy()->FastUpdateVertices(bPositions,
@@ -290,7 +365,7 @@ void USimpleDynamicMeshComponent::FastNotifyVertexAttributesUpdated(EMeshRenderA
 	}
 	else
 	{
-		NotifyMeshUpdated();
+		ResetProxy();
 	}
 }
 
@@ -304,22 +379,36 @@ void USimpleDynamicMeshComponent::FastNotifyUVsUpdated()
 
 void USimpleDynamicMeshComponent::FastNotifySecondaryTrianglesChanged()
 {
-	if (GetCurrentSceneProxy() != nullptr)
+	if (RenderMeshPostProcessor)
+	{
+		RenderMeshPostProcessor->ProcessMesh(*Mesh, *RenderMesh);
+	}
+
+	FSimpleDynamicMeshSceneProxy* Proxy = GetCurrentSceneProxy();
+
+	if (Proxy && Proxy->RenderMeshLayoutMatchesRenderBuffers())
 	{
 		GetCurrentSceneProxy()->FastUpdateAllIndexBuffers();
 	}
 	else
 	{
-		NotifyMeshUpdated();
+		ResetProxy();
 	}
 }
 
 
 void USimpleDynamicMeshComponent::FastNotifyTriangleVerticesUpdated(const TArray<int32>& Triangles, EMeshRenderAttributeFlags UpdatedAttributes)
 {
-	if (GetCurrentSceneProxy() == nullptr)
+	if (RenderMeshPostProcessor)
 	{
-		NotifyMeshUpdated();
+		RenderMeshPostProcessor->ProcessMesh(*Mesh, *RenderMesh);
+	}
+
+	FSimpleDynamicMeshSceneProxy* Proxy = GetCurrentSceneProxy();
+
+	if (!Proxy || !Proxy->RenderMeshLayoutMatchesRenderBuffers())
+	{
+		ResetProxy();
 	}
 	else if ( ! Decomposition )
 	{
@@ -353,9 +442,16 @@ void USimpleDynamicMeshComponent::FastNotifyTriangleVerticesUpdated(const TArray
 
 void USimpleDynamicMeshComponent::FastNotifyTriangleVerticesUpdated(const TSet<int32>& Triangles, EMeshRenderAttributeFlags UpdatedAttributes)
 {
-	if (GetCurrentSceneProxy() == nullptr)
+	if (RenderMeshPostProcessor)
 	{
-		NotifyMeshUpdated();
+		RenderMeshPostProcessor->ProcessMesh(*Mesh, *RenderMesh);
+	}
+
+	FSimpleDynamicMeshSceneProxy* Proxy = GetCurrentSceneProxy();
+
+	if (!Proxy || !Proxy->RenderMeshLayoutMatchesRenderBuffers())
+	{
+		ResetProxy();
 	}
 	else if (!Decomposition)
 	{
diff --git a/Engine/Plugins/Experimental/MeshModelingToolset/Source/ModelingComponents/Private/SimpleDynamicMeshSceneProxy.h b/Engine/Plugins/Experimental/MeshModelingToolset/Source/ModelingComponents/Private/SimpleDynamicMeshSceneProxy.h
--- a/Engine/Plugins/Experimental/MeshModelingToolset/Source/ModelingComponents/Private/SimpleDynamicMeshSceneProxy.h
+++ b/Engine/Plugins/Experimental/MeshModelingToolset/Source/ModelingComponents/Private/SimpleDynamicMeshSceneProxy.h
@@ -68,7 +68,7 @@ public:
 		}
 
 
-		FDynamicMesh3* Mesh = ParentComponent->GetMesh();
+		FDynamicMesh3* Mesh = ParentComponent->GetRenderMesh();
 		if (Mesh->HasAttributes() && Mesh->Attributes()->HasMaterialID() && NumMaterials > 1)
 		{
 			bIsSingleBuffer = false;
@@ -103,7 +103,7 @@ public:
 
 		bIsSingleBuffer = false;
 
-		FDynamicMesh3* Mesh = ParentComponent->GetMesh();
+		FDynamicMesh3* Mesh = ParentComponent->GetRenderMesh();
 		// find suitable overlays
 		FDynamicMeshUVOverlay* UVOverlay = Mesh->Attributes()->PrimaryUV();
 		FDynamicMeshNormalOverlay* NormalOverlay = Mesh->Attributes()->PrimaryNormals();
@@ -152,7 +152,7 @@ public:
 	 */
 	virtual void InitializeSingleBufferSet(FMeshRenderBufferSet* RenderBuffers)
 	{
-		FDynamicMesh3* Mesh = ParentComponent->GetMesh();
+		FDynamicMesh3* Mesh = ParentComponent->GetRenderMesh();
 
 		// find suitable overlays
 		FDynamicMeshUVOverlay* UVOverlay = nullptr;
@@ -191,7 +191,7 @@ public:
 	 */
 	virtual void InitializeByMaterial(TArray<FMeshRenderBufferSet*>& BufferSets)
 	{
-		FDynamicMesh3* Mesh = ParentComponent->GetMesh();
+		FDynamicMesh3* Mesh = ParentComponent->GetRenderMesh();
 		check(Mesh->HasAttributes() && Mesh->Attributes()->HasMaterialID());
 
 		// find suitable overlays
@@ -271,7 +271,67 @@ public:
 	}
 
 
+	bool RenderMeshLayoutMatchesRenderBuffers() const
+	{
+		const FDynamicMesh3* Mesh = ParentComponent->GetRenderMesh();
+
+		auto CheckBufferSet = [](const FDynamicMesh3* Mesh, const FMeshRenderBufferSet* BufferSet, int NumTriangles) -> bool
+		{
+			if (BufferSet->Triangles)
+			{
+				for (int TriangleID : BufferSet->Triangles.GetValue())
+				{
+					if (!Mesh->IsTriangle(TriangleID))
+					{
+						return false;
+					}
+				}
+			}
+
+			int NumVertices = NumTriangles * 3;
+			if (BufferSet->PositionVertexBuffer.GetNumVertices() != NumVertices ||
+				BufferSet->StaticMeshVertexBuffer.GetNumVertices() != NumVertices ||
+				BufferSet->ColorVertexBuffer.GetNumVertices() != NumVertices)
+			{
+				return false;
+			}
+
+			return true;
+		};
+
+		if (bIsSingleBuffer)
+		{
+			check(RenderBufferSets.Num() == 1);
 
+			FMeshRenderBufferSet* BufferSet = RenderBufferSets[0];
+			if (BufferSet->TriangleCount != Mesh->TriangleCount())
+			{
+				return false;
+			}
+
+			int NumTriangles = Mesh->TriangleCount();
+
+			if (!CheckBufferSet(Mesh, BufferSet, NumTriangles))
+			{
+				return false;
+			}
+		}
+		else
+		{
+			for (FMeshRenderBufferSet* BufferSet : RenderBufferSets)
+			{
+				check(BufferSet->Triangles);
+				int NumTriangles = BufferSet->Triangles->Num();
+
+				if (!CheckBufferSet(Mesh, BufferSet, NumTriangles))
+				{
+					return false;
+				}
+			}
+		}
+
+		return true;
+	}
 
 
 
@@ -283,7 +343,7 @@ public:
 		// This needs to be rewritten for split-by-material buffers.
 		// Could store triangle set with each buffer, and then rebuild vtx buffer(s) as needed?
 
-		FDynamicMesh3* Mesh = ParentComponent->GetMesh();
+		FDynamicMesh3* Mesh = ParentComponent->GetRenderMesh();
 
 		// find suitable overlays and attributes
 		FDynamicMeshNormalOverlay* NormalOverlay = nullptr;
@@ -376,7 +436,7 @@ public:
 		// This needs to be rewritten for split-by-material buffers.
 		// Could store triangle set with each buffer, and then rebuild vtx buffer(s) as needed?
 
-		FDynamicMesh3* Mesh = ParentComponent->GetMesh();
+		FDynamicMesh3* Mesh = ParentComponent->GetRenderMesh();
 
 		// find suitable overlays
 		FDynamicMeshNormalOverlay* NormalOverlay = nullptr;
@@ -441,7 +501,7 @@ public:
 	 */
 	virtual void FastUpdateAllIndexBuffers()
 	{
-		const FDynamicMesh3* Mesh = ParentComponent->GetMesh();
+		FDynamicMesh3* Mesh = ParentComponent->GetRenderMesh();
 
 		// have to wait for all outstanding rendering to finish because the index buffers we are about to edit might be in-use
 		FlushRenderingCommands();
diff --git a/Engine/Plugins/Experimental/MeshModelingToolset/Source/ModelingComponents/Public/SimpleDynamicMeshComponent.h b/Engine/Plugins/Experimental/MeshModelingToolset/Source/ModelingComponents/Public/SimpleDynamicMeshComponent.h
--- a/Engine/Plugins/Experimental/MeshModelingToolset/Source/ModelingComponents/Public/SimpleDynamicMeshComponent.h
+++ b/Engine/Plugins/Experimental/MeshModelingToolset/Source/ModelingComponents/Public/SimpleDynamicMeshComponent.h
@@ -18,6 +18,19 @@ struct FMeshDescription;
 class FSimpleDynamicMeshSceneProxy;
 
 
+/**
+ * Interface for a render mesh processor. Use this to process the Mesh stored in USimpleDynamicMeshComponent before
+ * sending it off for rendering.
+ * NOTE: This is called whenever the Mesh is updated and before rendering, so performance matters.
+ */
+class MODELINGCOMPONENTS_API IRenderMeshPostProcessor
+{
+public:
+	virtual ~IRenderMeshPostProcessor() = default;
+
+	virtual void ProcessMesh(const FDynamicMesh3& Mesh, FDynamicMesh3& OutRenderMesh) = 0;
+};
+
 
 /** 
  * USimpleDynamicMeshComponent is a mesh component similar to UProceduralMeshComponent,
@@ -55,6 +68,16 @@ public:
 	 */
 	virtual const FDynamicMesh3* GetMesh() const override { return Mesh.Get(); }
 
+	/*
+	* The SceneProxy should call these functions to get the post-processed RenderMesh. (See IRenderMeshPostProcessor.)
+	*/
+	virtual FDynamicMesh3* GetRenderMesh();
+
+	/*
+	* The SceneProxy should call these functions to get the post-processed RenderMesh. (See IRenderMeshPostProcessor.)
+	*/
+	virtual const FDynamicMesh3* GetRenderMesh() const;
+
 	/**
 	 * @return the current internal mesh, which is replaced with an empty mesh
 	 */
@@ -241,6 +264,10 @@ public:
 	 */
 	virtual void SetExternalDecomposition(TUniquePtr<FMeshRenderDecomposition> Decomposition);
 
+	/**
+	 * Add a render mesh processor, to be called before the mesh is sent for rendering.
+	 */
+	virtual void SetRenderMeshPostProcessor(TUniquePtr<IRenderMeshPostProcessor> Processor);
 
 public:
 
@@ -263,6 +290,9 @@ private:
 
 	FSimpleDynamicMeshSceneProxy* GetCurrentSceneProxy() { return (FSimpleDynamicMeshSceneProxy*)SceneProxy; }
 
+	// Called from NotifyMeshUpdated, as well as the FastNotify functions if needed
+	void ResetProxy();
+
 	//~ Begin UPrimitiveComponent Interface.
 	virtual FPrimitiveSceneProxy* CreateSceneProxy() override;
 	//~ End UPrimitiveComponent Interface.
@@ -285,4 +315,7 @@ private:
 	TUniqueFunction<bool(const FDynamicMesh3*, int32)> SecondaryTriFilterFunc = nullptr;
 
 	TUniquePtr<FMeshRenderDecomposition> Decomposition;
+
+	TUniquePtr<IRenderMeshPostProcessor> RenderMeshPostProcessor;
+	TUniquePtr<FDynamicMesh3> RenderMesh;
 };
