Add a prototype Subdivide tool that converts a Group Topology into a smooth mesh.

#rnx

diff --git a/Engine/Plugins/Experimental/MeshModelingToolset/Source/MeshModelingToolsEditorOnly/MeshModelingToolsEditorOnly.Build.cs b/Engine/Plugins/Experimental/MeshModelingToolset/Source/MeshModelingToolsEditorOnly/MeshModelingToolsEditorOnly.Build.cs
--- a/Engine/Plugins/Experimental/MeshModelingToolset/Source/MeshModelingToolsEditorOnly/MeshModelingToolsEditorOnly.Build.cs
+++ b/Engine/Plugins/Experimental/MeshModelingToolset/Source/MeshModelingToolsEditorOnly/MeshModelingToolsEditorOnly.Build.cs
@@ -60,6 +60,7 @@ public class MeshModelingToolsEditorOnly : ModuleRules
 				"MeshUtilitiesCommon",  
 				"MeshReductionInterface", // for UE4 standard simplification 
                 "ProxyLODMeshReduction", // for mesh merging voxel-based csg
+				"OpenSubdiv",
 				//"Slate",
 				//"SlateCore",
 				// ... add private dependencies that you statically link with here ...	
diff --git a/Engine/Plugins/Experimental/MeshModelingToolset/Source/MeshModelingToolsEditorOnly/Private/SubdividePoly.cpp b/Engine/Plugins/Experimental/MeshModelingToolset/Source/MeshModelingToolsEditorOnly/Private/SubdividePoly.cpp
new file mode 100644
--- /dev/null
+++ b/Engine/Plugins/Experimental/MeshModelingToolset/Source/MeshModelingToolsEditorOnly/Private/SubdividePoly.cpp
@@ -0,0 +1,499 @@
+// Copyright Epic Games, Inc. All Rights Reserved.
+
+#include "SubdividePoly.h"
+#include "GroupTopology.h"
+#include "MeshNormals.h"
+
+// OpenSubdiv needs M_PI defined
+#ifndef M_PI
+#define M_PI PI
+#define LOCAL_M_PI 1
+#endif
+
+#pragma warning(push, 0)     
+#include "far/topologyRefiner.h"
+#include "far/topologyDescriptor.h"
+#include "far/primvarRefiner.h"
+#pragma warning(pop)     
+
+#if LOCAL_M_PI
+#undef M_PI
+#endif
+
+namespace SubdividePolyLocal
+{
+	// Vertices with positions and normals. Used for interpolating from polymesh to subdivision surface mesh.
+	struct SubdVertex
+	{
+		FVertexInfo VertexInfo;
+
+		void Clear()
+		{
+			VertexInfo = FVertexInfo();
+		}
+
+		void AddWithWeight(SubdVertex const& Src, float Weight)
+		{
+			VertexInfo.Position += Weight * Src.VertexInfo.Position;
+			VertexInfo.Normal += Weight * Src.VertexInfo.Normal;
+			VertexInfo.bHaveN = Src.VertexInfo.bHaveN;
+			// TODO: add support for vertex color
+		}
+	};
+
+	// Vertex with only UV data. This is used for "face-varying" interpolation, i.e. vertices can have multiple UV 
+	// values if they are incident on triangles in different UV islands.
+	struct SubdUVVertex
+	{
+		FVector2f VertexUV;
+
+		SubdUVVertex(const FVector2f& InVertexUV) :
+			VertexUV{ InVertexUV }
+		{}
+
+		void Clear()
+		{
+			VertexUV = FVector2f();
+		}
+
+		void AddWithWeight(SubdUVVertex const& Src, float Weight)
+		{
+			VertexUV += Weight * Src.VertexUV;
+		}
+	};
+
+
+	// Compute the average of all normals corresponding to the given VertexID
+	FVector3f GetAverageVertexNormalFromOverlay(const FDynamicMeshNormalOverlay* NormalOverlay,
+												int VertexID)
+	{
+		TArray<int> NormalElements;
+		NormalOverlay->GetVertexElements(VertexID, NormalElements);
+
+		FVector3f SumElements{ 0.0f, 0.0f, 0.0f };
+		float ElementCount = 0.0f;
+		for (int ElementID : NormalElements)
+		{
+			SumElements += NormalOverlay->GetElement(ElementID);
+			ElementCount += 1.0f;
+		}
+
+		return SumElements / ElementCount;
+	}
+
+	// Get the indices of GroupTopology "Corners" for a particular group boundary.
+	void GetBoundaryCorners(const FGroupTopology::FGroupBoundary& Boundary,
+							const FGroupTopology& Topology,
+							TArray<int>& Corners)
+	{
+		int FirstEdgeIndex = Boundary.GroupEdges[0];
+		Corners.Add(Topology.Edges[FirstEdgeIndex].EndpointCorners[0]);
+		Corners.Add(Topology.Edges[FirstEdgeIndex].EndpointCorners[1]);
+
+		int NextEdgeIndex = Boundary.GroupEdges[1];
+		FIndex2i NextEdgeCorners = Topology.Edges[NextEdgeIndex].EndpointCorners;
+		if (Corners[1] != NextEdgeCorners[0] && Corners[1] != NextEdgeCorners[1])
+		{
+			Swap(Corners[0], Corners[1]);
+			check(Corners[1] == NextEdgeCorners[0] || Corners[1] == NextEdgeCorners[1]);
+		}
+
+		for (int i = 1; i < Boundary.GroupEdges.Num() - 1; ++i)
+		{
+			int EdgeIndex = Boundary.GroupEdges[i];
+			FIndex2i CurrEdgeCorners = Topology.Edges[EdgeIndex].EndpointCorners;
+			if (Corners.Last() == CurrEdgeCorners[0])
+			{
+				Corners.Add(CurrEdgeCorners[1]);
+			}
+			else
+			{
+				check(Corners.Last() == CurrEdgeCorners[1]);
+				Corners.Add(CurrEdgeCorners[0]);
+			}
+		}
+	}
+
+	// Treat the given FGroupTopology as a polygonal mesh, and get its vertices
+	void GetGroupPolyMeshVertices(const FDynamicMesh3& Mesh,
+								  const FGroupTopology& Topology,
+								  bool bGetVertexNormals,
+								  TArray<SubdVertex>& OutVertices )
+	{
+		for (const FGroupTopology::FCorner& Corner : Topology.Corners)
+		{
+			SubdVertex Vertex;
+			bool bWantNormals = bGetVertexNormals && Mesh.HasVertexNormals();
+			constexpr bool bWantColors = false;
+			constexpr bool bWantVertexUVs = false;
+			Mesh.GetVertex(Corner.VertexID, Vertex.VertexInfo, bWantNormals, bWantColors, bWantVertexUVs);
+
+			if (bGetVertexNormals && !Vertex.VertexInfo.bHaveN)
+			{
+				// We want vertex normals but the mesh doesn't have them. See if we can get them from an overlay.
+				if (Mesh.HasAttributes() && (Mesh.Attributes()->NumNormalLayers() > 0))
+				{
+					Vertex.VertexInfo.Normal = GetAverageVertexNormalFromOverlay(Mesh.Attributes()->PrimaryNormals(), 
+																				 Corner.VertexID);
+					Vertex.VertexInfo.bHaveN = true;
+				}
+				else
+				{
+					Vertex.VertexInfo.Normal = { 0,1,0 };
+					Vertex.VertexInfo.bHaveN = true;
+				}
+			}
+
+			OutVertices.Add(Vertex);
+		}
+	}
+
+
+	// Given a group and a vertex ID, return:
+	// - a triangle in the group with one of its corners equal to vertex ID
+	// - the (0-2) triangle corner index corresponding to the given vertex ID
+	bool FindTriangleVertex(const FGroupTopology::FGroup& Group,
+							int VertexID,
+							const FDynamicMesh3& Mesh,
+							TTuple<int, int>& OutTriangleVertex)
+	{
+		for (int Tri : Group.Triangles)
+		{
+			FIndex3i TriVertices = Mesh.GetTriangle(Tri);
+			for (int i = 0; i < 3; ++i)
+			{
+				if (TriVertices[i] == VertexID)
+				{
+					OutTriangleVertex = TTuple<int, int>{ Tri, i };
+					return true;
+				}
+			}
+		}
+		return false;
+	}
+
+	// Treating the GroupTopology as a polygonal mesh, get its "face-varying" UV coordinates. Assumes that each 
+	// polygonal face has no UV seams cutting through it, but that UV seams might exist at polygon boundaries.
+	bool GetGroupPolyMeshUVs(const FGroupTopology& Topology,
+							 const FDynamicMesh3& Mesh,
+							 const FDynamicMeshUVOverlay* UVOverlay,
+							 TArray<SubdUVVertex>& OutUVs)		// for each group, all the uvs on the group corners
+	{
+		for (const FGroupTopology::FGroup& Group : Topology.Groups)
+		{
+			check(Group.Boundaries.Num() == 1);
+			check(Group.Triangles.Num() > 0);
+			const FGroupTopology::FGroupBoundary& Bdry = Group.Boundaries[0];
+
+			TArray<int> Corners;
+			GetBoundaryCorners(Bdry, Topology, Corners);
+
+			for (int CornerID : Corners)
+			{
+				int CornerVertexID = Topology.Corners[CornerID].VertexID;
+
+				// Find a triangle in the group that has a vertex ID equal to the given corner
+				TTuple<int, int> TriangleVertex;
+				if (FindTriangleVertex(Group, CornerVertexID, Mesh, TriangleVertex))
+				{
+					int TriangleID = TriangleVertex.Get<0>();
+					int TriVertexIndex = TriangleVertex.Get<1>();	// The (0,1,2) index of the polygon corner wrt the triangle
+
+					TStaticArray<FVector2f, 3> TriangleUVs;
+					UVOverlay->GetTriElements(TriangleID, TriangleUVs[0], TriangleUVs[1], TriangleUVs[2]);
+					FVector2f CornerUV = TriangleUVs[TriVertexIndex];
+
+					OutUVs.Add(SubdUVVertex{ CornerUV });
+				}
+				else
+				{
+					return false;
+				}
+			}
+		}
+
+		return true;
+	}
+
+
+	void InitializeOverlayToFaceVertexUVs(FDynamicMeshUVOverlay* UVOverlay, const TArray<FVector2f>& UVs)
+	{
+		const FDynamicMesh3* Mesh = UVOverlay->GetParentMesh();
+		check(Mesh->IsCompact());
+
+		int NumTriangles = Mesh->TriangleCount();
+		check(NumTriangles == Mesh->MaxTriangleID());
+		check(UVs.Num() == 3 * NumTriangles);
+
+		UVOverlay->ClearElements();
+		UVOverlay->InitializeTriangles(NumTriangles);
+
+		for (int TriangleIndex = 0; TriangleIndex < NumTriangles; ++TriangleIndex)
+		{
+			FIndex3i Tri{ 3 * TriangleIndex + 0, 3 * TriangleIndex + 1,	3 * TriangleIndex + 2 };
+			UVOverlay->AppendElement(UVs[Tri[0]]);
+			UVOverlay->AppendElement(UVs[Tri[1]]);
+			UVOverlay->AppendElement(UVs[Tri[2]]);
+			UVOverlay->SetTriangle(TriangleIndex, Tri);
+		}
+	}
+
+
+}	// namespace SubdivisionSurfaceLocal
+
+
+class FSubdividePoly::RefinerImpl
+{
+public:
+	OpenSubdiv::Far::TopologyRefiner* TopologyRefiner = nullptr;
+};
+
+FSubdividePoly::FSubdividePoly(const FGroupTopology& InTopology,
+										 const FDynamicMesh3& InOriginalMesh,
+										 int InLevel) :
+	GroupTopology(InTopology)
+	, OriginalMesh(InOriginalMesh)
+	, Level(InLevel)
+{
+	Refiner = MakeUnique<RefinerImpl>();
+}
+
+FSubdividePoly::~FSubdividePoly()
+{
+	if (Refiner && Refiner->TopologyRefiner)
+	{
+		// This was created by TopologyRefinerFactory; looks like we are responsible for cleaning it up.
+		delete Refiner->TopologyRefiner;
+	}
+	Refiner = nullptr;
+}
+
+
+bool FSubdividePoly::ComputeTopologySubdivision()
+{
+	TArray<int> BoundaryVertsPerFace;
+	TArray<int> NumVertsPerFace;
+	int TotalNumVertices = 0;
+
+	for (const FGroupTopology::FGroup& Group : GroupTopology.Groups)
+	{
+		if (Group.Boundaries.Num() != 1)
+		{
+			return false;
+		}
+
+		const FGroupTopology::FGroupBoundary& Boundary = Group.Boundaries[0];
+		if (Boundary.GroupEdges.Num() < 2)
+		{
+			return false;
+		}
+
+		TArray<int> Corners;
+		SubdividePolyLocal::GetBoundaryCorners(Boundary, GroupTopology, Corners);
+
+		NumVertsPerFace.Add(Corners.Num());
+		TotalNumVertices += Corners.Num();
+		BoundaryVertsPerFace.Append(Corners);
+	}
+
+	OpenSubdiv::Far::TopologyDescriptor Descriptor;
+	Descriptor.numVertices = TotalNumVertices;
+	Descriptor.numFaces = GroupTopology.Groups.Num();
+	Descriptor.numVertsPerFace = NumVertsPerFace.GetData();		// taking pointer to local memory seems to be ok?
+	Descriptor.vertIndicesPerFace = BoundaryVertsPerFace.GetData();
+
+	OpenSubdiv::Far::TopologyDescriptor::FVarChannel UVChannel;
+	TArray<int> Buffer;
+	if (UVComputationMethod == ESubdivisionOutputUVs::Interpolated)
+	{
+		UVChannel.numValues = BoundaryVertsPerFace.Num();
+		Buffer.SetNum(UVChannel.numValues);
+		for (int i = 0; i < UVChannel.numValues; ++i)
+		{
+			Buffer[i] = i;
+		}
+		UVChannel.valueIndices = Buffer.GetData();
+		Descriptor.numFVarChannels = 1;
+		Descriptor.fvarChannels = &UVChannel;		// taking a pointer to local memory seems to be ok?
+	}
+	else
+	{
+		Descriptor.numFVarChannels = 0;
+	}
+
+	using RefinerFactory = OpenSubdiv::Far::TopologyRefinerFactory<OpenSubdiv::Far::TopologyDescriptor>;
+	RefinerFactory::Options RefinerOptions;
+
+	RefinerOptions.schemeOptions.SetVtxBoundaryInterpolation(OpenSubdiv::Sdc::Options::VTX_BOUNDARY_EDGE_AND_CORNER);
+
+	Refiner->TopologyRefiner = RefinerFactory::Create(Descriptor, RefinerOptions);
+
+	if (Refiner->TopologyRefiner == nullptr)
+	{
+		return false;
+	}
+
+	Refiner->TopologyRefiner->RefineUniform(OpenSubdiv::Far::TopologyRefiner::UniformOptions(Level));
+
+	return true;
+}
+
+
+bool FSubdividePoly::ComputeSubdividedMesh(FDynamicMesh3& OutMesh)
+{
+	if (!Refiner || !(Refiner->TopologyRefiner))
+	{
+		return false;
+	}
+
+	OpenSubdiv::Far::PrimvarRefiner Interpolator(*Refiner->TopologyRefiner);
+
+	//
+	// Interpolate vertex positions from initial Group poly mesh down to refinement level
+	// 
+	TArray<SubdividePolyLocal::SubdVertex> SourceVertices;
+	bool bInterpolateVertexNormals = (NormalComputationMethod == ESubdivisionOutputNormals::Interpolated);
+	GetGroupPolyMeshVertices(OriginalMesh, GroupTopology, bInterpolateVertexNormals, SourceVertices);
+
+	TArray<SubdividePolyLocal::SubdVertex> RefinedVertices;
+	for (int CurrentLevel = 1; CurrentLevel <= Level; ++CurrentLevel)
+	{
+		// TODO: Don't keep resizing -- preallocate one big buffer and move through it
+		RefinedVertices.SetNumUninitialized(Refiner->TopologyRefiner->GetLevel(CurrentLevel).GetNumVertices());
+		SubdividePolyLocal::SubdVertex* s = SourceVertices.GetData();
+		SubdividePolyLocal::SubdVertex* d = RefinedVertices.GetData();
+		Interpolator.Interpolate(CurrentLevel, s, d);
+		SourceVertices = RefinedVertices;
+	}
+
+	//
+	// Interpolate face group IDs
+	// 
+	TArray<int> SourceGroupIDs;
+	for (const FGroupTopology::FGroup& Group : GroupTopology.Groups)
+	{
+		SourceGroupIDs.Add(Group.GroupID);
+	}
+
+	TArray<int> RefinedGroupIDs;
+	for (int CurrentLevel = 1; CurrentLevel <= Level; ++CurrentLevel)
+	{
+		// TODO: Don't keep resizing -- preallocate one big buffer and move through it
+		RefinedGroupIDs.SetNumUninitialized(Refiner->TopologyRefiner->GetLevel(CurrentLevel).GetNumFaces());
+		int* s = SourceGroupIDs.GetData();
+		int* d = RefinedGroupIDs.GetData();
+		Interpolator.InterpolateFaceUniform(CurrentLevel, s, d);
+		SourceGroupIDs = RefinedGroupIDs;
+	}
+
+
+	//
+	// Interpolate UVs
+	// 
+	TArray<SubdividePolyLocal::SubdUVVertex> RefinedUVs;
+
+	if (UVComputationMethod == ESubdivisionOutputUVs::Interpolated)
+	{
+		if (!OriginalMesh.HasAttributes() || !OriginalMesh.Attributes()->PrimaryUV())
+		{
+			return false;
+		}
+
+		TArray<SubdividePolyLocal::SubdUVVertex> SourceUVs;
+		bool bGetUVsOK = SubdividePolyLocal::GetGroupPolyMeshUVs(GroupTopology,
+																	  OriginalMesh,
+																	  OriginalMesh.Attributes()->PrimaryUV(), 
+																	  SourceUVs);
+
+		if (!bGetUVsOK)
+		{
+			return false;
+		}
+
+		check(SourceUVs.Num() == Refiner->TopologyRefiner->GetLevel(0).GetNumFaceVertices());
+
+		for (int CurrentLevel = 1; CurrentLevel <= Level; ++CurrentLevel)
+		{
+			// TODO: Don't keep resizing -- preallocate one big buffer and move through it
+			RefinedUVs.SetNumUninitialized(Refiner->TopologyRefiner->GetLevel(CurrentLevel).GetNumFVarValues());
+			SubdividePolyLocal::SubdUVVertex* s = SourceUVs.GetData();
+			SubdividePolyLocal::SubdUVVertex* d = RefinedUVs.GetData();
+			Interpolator.InterpolateFaceVarying(CurrentLevel, s, d);
+			SourceUVs = RefinedUVs;
+		}
+	}
+
+	// Now transfer to output mesh
+	OutMesh.Clear();
+
+	OutMesh.EnableTriangleGroups();
+	if (bInterpolateVertexNormals)
+	{
+		OutMesh.EnableVertexNormals(FVector3f{ 0,0,0 });
+	}
+	if ((NormalComputationMethod != ESubdivisionOutputNormals::None) || (UVComputationMethod != ESubdivisionOutputUVs::None))
+	{
+		OutMesh.EnableAttributes();
+	}
+
+	// Add the vertices
+	// TODO: Can we resize the Mesh vertex array up front and then populate it?
+	for (const SubdividePolyLocal::SubdVertex& V : RefinedVertices)
+	{
+		OutMesh.AppendVertex(V.VertexInfo);
+	}
+
+	const OpenSubdiv::Far::TopologyLevel& FinalLevel = Refiner->TopologyRefiner->GetLevel(Level);
+
+	check(UVComputationMethod != ESubdivisionOutputUVs::Interpolated || FinalLevel.GetNumFVarValues() == RefinedUVs.Num());
+	check(FinalLevel.GetNumFaces() == RefinedGroupIDs.Num());
+
+	// Add the faces (manually triangulate the output here)
+
+	int NumFaceVertices = 0;
+	TArray<FVector2f> TriangleUVs;
+
+	for (int FaceID = 0; FaceID < FinalLevel.GetNumFaces(); ++FaceID)
+	{
+		int GroupID = RefinedGroupIDs[FaceID];
+
+		OpenSubdiv::Far::ConstIndexArray Face = FinalLevel.GetFaceVertices(FaceID);
+		NumFaceVertices += Face.size();
+
+		check(Face.size() == 4);
+
+		FIndex3i TriA{ Face[0], Face[1], Face[3] };
+		FIndex3i TriB{ Face[2], Face[3], Face[1] };
+		OutMesh.AppendTriangle(TriA, GroupID);
+		OutMesh.AppendTriangle(TriB, GroupID);
+
+		if (UVComputationMethod == ESubdivisionOutputUVs::Interpolated)
+		{
+			OpenSubdiv::Far::ConstIndexArray FaceUVIndices = FinalLevel.GetFaceFVarValues(FaceID);
+
+			FIndex3i UVTriA{ FaceUVIndices[0], FaceUVIndices[1], FaceUVIndices[3] };
+			TriangleUVs.Add(RefinedUVs[UVTriA[0]].VertexUV);
+			TriangleUVs.Add(RefinedUVs[UVTriA[1]].VertexUV);
+			TriangleUVs.Add(RefinedUVs[UVTriA[2]].VertexUV);
+
+			FIndex3i UVTriB{ FaceUVIndices[2], FaceUVIndices[3], FaceUVIndices[1] };
+			TriangleUVs.Add(RefinedUVs[UVTriB[0]].VertexUV);
+			TriangleUVs.Add(RefinedUVs[UVTriB[1]].VertexUV);
+			TriangleUVs.Add(RefinedUVs[UVTriB[2]].VertexUV);
+		}
+	}
+
+	if (NormalComputationMethod != ESubdivisionOutputNormals::None)
+	{
+		const bool bUseExistingMeshVertexNormals = (NormalComputationMethod == ESubdivisionOutputNormals::Interpolated);
+		FMeshNormals::InitializeOverlayToPerVertexNormals(OutMesh.Attributes()->PrimaryNormals(), bUseExistingMeshVertexNormals);
+	}
+
+	if (UVComputationMethod == ESubdivisionOutputUVs::Interpolated)
+	{
+		SubdividePolyLocal::InitializeOverlayToFaceVertexUVs(OutMesh.Attributes()->PrimaryUV(), TriangleUVs);
+	}
+
+	return true;
+}
+
diff --git a/Engine/Plugins/Experimental/MeshModelingToolset/Source/MeshModelingToolsEditorOnly/Private/SubdividePolyTool.cpp b/Engine/Plugins/Experimental/MeshModelingToolset/Source/MeshModelingToolsEditorOnly/Private/SubdividePolyTool.cpp
new file mode 100644
--- /dev/null
+++ b/Engine/Plugins/Experimental/MeshModelingToolset/Source/MeshModelingToolsEditorOnly/Private/SubdividePolyTool.cpp
@@ -0,0 +1,339 @@
+// Copyright Epic Games, Inc. All Rights Reserved.
+
+#include "SubdividePolyTool.h"
+#include "ToolBuilderUtil.h"
+#include "InteractiveToolManager.h"
+#include "GroupTopology.h"
+#include "MeshDescriptionToDynamicMesh.h"
+#include "DynamicMeshToMeshDescription.h"
+#include "ToolSetupUtil.h"
+#include "Util/ColorConstants.h"
+#include "MeshNormals.h"
+#include "SimpleDynamicMeshComponent.h"
+#include "Drawing/PreviewGeometryActor.h"
+
+#define LOCTEXT_NAMESPACE "USubdividePolyTool"
+
+class SubdivPostProcessor : public IRenderMeshPostProcessor
+{
+public:
+
+	SubdivPostProcessor(int InSubdivisionLevel,
+						ESubdivisionOutputNormals InNormalComputationMethod,
+						ESubdivisionOutputUVs InUVComputationMethod) :
+		SubdivisionLevel(InSubdivisionLevel),
+		NormalComputationMethod(InNormalComputationMethod),
+		UVComputationMethod(InUVComputationMethod)
+	{}
+
+	int SubdivisionLevel = 3;
+
+	ESubdivisionOutputNormals NormalComputationMethod = ESubdivisionOutputNormals::Generated;
+
+	ESubdivisionOutputUVs UVComputationMethod = ESubdivisionOutputUVs::Interpolated;
+
+	void ProcessMesh(const FDynamicMesh3& Mesh, FDynamicMesh3& OutRenderMesh) final
+	{
+		constexpr bool bAutoCompute = true;
+		FGroupTopology Topo(&Mesh, bAutoCompute);
+		FSubdividePoly Subd(Topo, Mesh, SubdivisionLevel);
+		Subd.NormalComputationMethod = NormalComputationMethod;
+		Subd.UVComputationMethod = UVComputationMethod;
+
+		ensure(Subd.ComputeTopologySubdivision());
+
+		ensure(Subd.ComputeSubdividedMesh(OutRenderMesh));
+	}
+};
+
+
+// Tool builder
+
+bool USubdividePolyToolBuilder::CanBuildTool(const FToolBuilderState& SceneState) const
+{
+	return ToolBuilderUtil::CountComponents(SceneState, CanMakeComponentTarget) == 1;
+}
+
+UInteractiveTool* USubdividePolyToolBuilder::BuildTool(const FToolBuilderState& SceneState) const
+{
+	UActorComponent* ActorComponent = ToolBuilderUtil::FindFirstComponent(SceneState, CanMakeComponentTarget);
+	auto* MeshComponent = Cast<UPrimitiveComponent>(ActorComponent);
+	check(MeshComponent != nullptr);
+
+	USubdividePolyTool* NewTool = NewObject<USubdividePolyTool>(SceneState.ToolManager);
+	NewTool->SetSelection(MakeComponentTarget(MeshComponent));
+	NewTool->SetWorld(SceneState.World);
+
+	return NewTool;
+}
+
+
+// Tool actual
+
+void USubdividePolyTool::Setup()
+{
+	UInteractiveTool::Setup();
+
+	if (!ComponentTarget)
+	{
+		return;
+	}
+
+	bool bWantVertexNormals = false;
+	OriginalMesh = MakeShared<FDynamicMesh3>(bWantVertexNormals, false, false, false);
+	FMeshDescriptionToDynamicMesh Converter;
+	Converter.Convert(ComponentTarget->GetMesh(), *OriginalMesh);
+
+	//
+	// Error checking
+	//
+
+	FGroupTopology Topo(OriginalMesh.Get(), true);
+	if (Topo.Groups.Num() == 0)
+	{
+		GetToolManager()->DisplayMessage(LOCTEXT("NoGroupsWarning",
+												 "This object has no PolyGroups. Use the PolyGroups or Select Tool to assign PolyGroups."),
+										 EToolMessageLevel::UserWarning);
+		return;
+	}
+
+	if (Topo.Groups.Num() < 2)
+	{
+		// TODO: for an open surface, use the surface boundary as a group boundary?
+
+		GetToolManager()->DisplayMessage(LOCTEXT("SingleGroupsWarning",
+												 "This object has only one PolyGroup. Use the PolyGroups or Select Tool to assign PolyGroups."),
+										 EToolMessageLevel::UserWarning);
+		return;
+	}
+
+	for (const FGroupTopology::FGroup& Group : Topo.Groups)
+	{
+		if (Group.Boundaries.Num() == 0)
+		{
+			// Error: Group has no boundaries, e.g. closed surface component with only one group
+			GetToolManager()->DisplayMessage(LOCTEXT("NoGroupBoundaryWarning",
+													 "Found a PolyGroup with no boundaries. Use the PolyGroups or Select Tool to assign PolyGroups."),
+											 EToolMessageLevel::UserWarning);
+			return;
+		}
+
+		if (Group.Boundaries.Num() > 1)
+		{
+			// Error: Group has multiple boundaries, e.g. nested polygon
+			GetToolManager()->DisplayMessage(LOCTEXT("MultipleGroupBoundaryWarning",
+													 "Found a PolyGroup with multiple boundaries, which is not supported. Use the PolyGroups or Select Tool to assign PolyGroups."),
+											 EToolMessageLevel::UserWarning);
+			return;
+		}
+
+		for ( const FGroupTopology::FGroupBoundary& Boundary : Group.Boundaries )
+		{
+			if (Boundary.GroupEdges.Num() < 3)
+			{
+				GetToolManager()->DisplayMessage(LOCTEXT("DegenerateGroupPolygon",
+														 "One PolyGroup has fewer than three boundary edges. Use the PolyGroups or Select Tool to assign/fix PolyGroups."),
+												 EToolMessageLevel::UserWarning);
+				return;
+			}
+		}
+	}
+
+	//
+	// Finished error checking
+	//
+
+	GetToolManager()->DisplayMessage(LOCTEXT("SubdividePolyToolMessage",
+											 "Set the subdivision level and hit Accept to create a new subdivided mesh"),
+									 EToolMessageLevel::UserNotification);
+
+	Properties = NewObject<USubdividePolyToolProperties>(this, TEXT("Subdivide Mesh Tool Settings"));
+	Properties->RestoreProperties(this);
+	AddToolPropertySource(Properties);
+	SetToolPropertySourceEnabled(Properties, true);
+
+
+	// create dynamic mesh component to use for live preview
+	FActorSpawnParameters SpawnInfo;
+	TemporaryParentActor = TargetWorld->SpawnActor<APreviewMeshActor>(FVector::ZeroVector, FRotator::ZeroRotator, SpawnInfo);
+
+	PreviewDynamicMeshComponent = NewObject<USimpleDynamicMeshComponent>(TemporaryParentActor, "DynamicMesh");
+	PreviewDynamicMeshComponent->SetRenderMeshPostProcessor(MakeUnique<SubdivPostProcessor>(Properties->SubdivisionLevel,
+																					 Properties->NormalComputationMethod,
+																					 Properties->UVComputationMethod));
+
+	TemporaryParentActor->SetRootComponent(PreviewDynamicMeshComponent);
+	PreviewDynamicMeshComponent->RegisterComponent();
+
+	TemporaryParentActor->SetActorTransform(ComponentTarget->GetOwnerActor()->GetRootComponent()->GetComponentTransform());
+
+	// Use the input mesh's material on the preview
+	FComponentMaterialSet MaterialSet;
+	ComponentTarget->GetMaterialSet(MaterialSet);
+
+	for (int k = 0; k < MaterialSet.Materials.Num(); ++k)
+	{
+		PreviewDynamicMeshComponent->SetMaterial(k, MaterialSet.Materials[k]);
+	}
+
+	// configure secondary render material
+	UMaterialInterface* SelectionMaterial = ToolSetupUtil::GetSelectionMaterial(FLinearColor(0.8f, 0.75f, 0.0f), GetToolManager());
+	if (SelectionMaterial != nullptr)
+	{
+		PreviewDynamicMeshComponent->SetSecondaryRenderMaterial(SelectionMaterial);
+	}
+
+	// dynamic mesh configuration settings
+	PreviewDynamicMeshComponent->TangentsType = EDynamicMeshTangentCalcType::NoTangents;
+	PreviewDynamicMeshComponent->InitializeMesh(ComponentTarget->GetMesh());
+	FMeshNormals::QuickComputeVertexNormals(*PreviewDynamicMeshComponent->GetMesh());
+
+	auto RebuildMeshPostProcessor = [this]()
+	{
+		if (PreviewDynamicMeshComponent)
+		{
+			PreviewDynamicMeshComponent->SetRenderMeshPostProcessor(MakeUnique<SubdivPostProcessor>(Properties->SubdivisionLevel,
+																							 Properties->NormalComputationMethod,
+																							 Properties->UVComputationMethod));
+			PreviewDynamicMeshComponent->NotifyMeshUpdated();
+		}
+	};
+
+	// Watch for property changes
+	Properties->WatchProperty(Properties->SubdivisionLevel, [this, RebuildMeshPostProcessor](int)
+	{
+		RebuildMeshPostProcessor();
+	});
+	Properties->WatchProperty(Properties->NormalComputationMethod, [this, RebuildMeshPostProcessor](ESubdivisionOutputNormals)
+	{
+		RebuildMeshPostProcessor();
+	});
+	Properties->WatchProperty(Properties->UVComputationMethod, [this, RebuildMeshPostProcessor](ESubdivisionOutputUVs)
+	{
+		RebuildMeshPostProcessor();
+	});
+
+
+	auto RenderGroupsChanged = [this](bool bNewRenderGroups)
+	{
+		if (PreviewDynamicMeshComponent)
+		{
+			if (bNewRenderGroups)
+			{
+				PreviewDynamicMeshComponent->SetOverrideRenderMaterial(ToolSetupUtil::GetSelectionMaterial(GetToolManager()));
+				PreviewDynamicMeshComponent->TriangleColorFunc = [](const FDynamicMesh3* Mesh, int TriangleID)
+				{
+					return LinearColors::SelectFColor(Mesh->GetTriangleGroup(TriangleID));
+				};
+			}
+			else
+			{
+				PreviewDynamicMeshComponent->SetOverrideRenderMaterial(nullptr);
+				PreviewDynamicMeshComponent->TriangleColorFunc = nullptr;
+			}
+			PreviewDynamicMeshComponent->FastNotifyColorsUpdated();
+		}
+	};
+
+	Properties->WatchProperty(Properties->bRenderGroups, RenderGroupsChanged);
+
+	ComponentTarget->SetOwnerVisibility(false);
+
+	// Render with polygroup colors
+	RenderGroupsChanged(Properties->bRenderGroups);
+
+
+	PreviewGeometry = NewObject<UPreviewGeometry>(this);
+	PreviewGeometry->CreateInWorld(ComponentTarget->GetOwnerActor()->GetWorld(), ComponentTarget->GetWorldTransform());
+	CreateOrUpdatePreviewGeometry();
+
+	// regenerate preview geo if mesh changes due to undo/redo/etc
+	PreviewDynamicMeshComponent->OnMeshChanged.AddLambda([this]() { bPreviewGeometryNeedsUpdate = true; });
+}
+
+void USubdividePolyTool::CreateOrUpdatePreviewGeometry()
+{
+	FGroupTopology Topology(OriginalMesh.Get(), true);
+	int NumEdges = Topology.Edges.Num();
+
+	PreviewGeometry->CreateOrUpdateLineSet(TEXT("TopologyEdges"),
+										   NumEdges,
+										   [&Topology, this](int32 Index, TArray<FRenderableLine>& LinesOut)
+	{
+		const FGroupTopology::FGroupEdge& Edge = Topology.Edges[Index];
+		FIndex2i EdgeCorners = Edge.EndpointCorners;
+
+		if (EdgeCorners[0] == FDynamicMesh3::InvalidID || EdgeCorners[1] == FDynamicMesh3::InvalidID)
+		{
+			return;
+		}
+
+		FIndex2i EdgeVertices{ Topology.Corners[EdgeCorners[0]].VertexID,
+							   Topology.Corners[EdgeCorners[1]].VertexID };
+		FVector A = (FVector)OriginalMesh->GetVertex(EdgeVertices[0]);
+		FVector B = (FVector)OriginalMesh->GetVertex(EdgeVertices[1]);
+
+		const float TopologyLineThickness = 4.0f;
+		const FColor TopologyLineColor(255, 0, 0);
+		LinesOut.Add(FRenderableLine(A, B, TopologyLineColor, TopologyLineThickness));
+	});
+}
+
+
+void USubdividePolyTool::Shutdown(EToolShutdownType ShutdownType)
+{
+	if (Properties)
+	{
+		Properties->SaveProperties(this);
+	}
+
+	check(PreviewGeometry);
+	PreviewGeometry->Disconnect();
+
+	if (PreviewDynamicMeshComponent)
+	{
+		if (ShutdownType == EToolShutdownType::Accept)
+		{
+			GetToolManager()->BeginUndoTransaction(LOCTEXT("USubdividePolyTool", "Subdivide Mesh"));
+
+			FDynamicMesh3* DynamicMeshResult = PreviewDynamicMeshComponent->GetRenderMesh();
+
+			ComponentTarget->CommitMesh([DynamicMeshResult](const FPrimitiveComponentTarget::FCommitParams& CommitParams)
+			{
+				FDynamicMeshToMeshDescription Converter;
+				Converter.Convert(DynamicMeshResult, *CommitParams.MeshDescription);
+			});
+
+			GetToolManager()->EndUndoTransaction();
+		}
+		PreviewDynamicMeshComponent->UnregisterComponent();
+		PreviewDynamicMeshComponent->DestroyComponent();
+		PreviewDynamicMeshComponent = nullptr;
+	}
+
+	if (TemporaryParentActor != nullptr)
+	{
+		TemporaryParentActor->Destroy();
+		TemporaryParentActor = nullptr;
+	}
+
+	ComponentTarget->SetOwnerVisibility(true);
+}
+
+bool USubdividePolyTool::CanAccept() const
+{
+	return PreviewDynamicMeshComponent != nullptr;
+}
+
+
+void USubdividePolyTool::OnTick(float DeltaTime)
+{
+	if (bPreviewGeometryNeedsUpdate)
+	{
+		CreateOrUpdatePreviewGeometry();
+		bPreviewGeometryNeedsUpdate = false;
+	}
+}
+
+
+#undef LOCTEXT_NAMESPACE
diff --git a/Engine/Plugins/Experimental/MeshModelingToolset/Source/MeshModelingToolsEditorOnly/Public/SubdividePoly.h b/Engine/Plugins/Experimental/MeshModelingToolset/Source/MeshModelingToolsEditorOnly/Public/SubdividePoly.h
new file mode 100644
--- /dev/null
+++ b/Engine/Plugins/Experimental/MeshModelingToolset/Source/MeshModelingToolsEditorOnly/Public/SubdividePoly.h
@@ -0,0 +1,78 @@
+// Copyright Epic Games, Inc. All Rights Reserved.
+
+#pragma once
+
+#include "CoreMinimal.h"
+#include "ModelingOperators.h"
+#include "SubdividePoly.generated.h"
+
+class FGroupTopology;
+class FDynamicMesh3;
+class FProgressCancel;
+
+/**
+ * Method for computing normals on the output mesh
+ */
+UENUM()
+enum class ESubdivisionOutputNormals : uint8
+{
+	// Copy normals to poly mesh then interpolate to new triangles as we subdivide.
+	Interpolated,
+
+	// Recompute normals after subdivision
+	Generated,
+
+	// Leave normals uninitialized on output mesh
+	None
+};
+
+/**
+ * Method for computing UVs on the output mesh
+ */
+UENUM()
+enum class ESubdivisionOutputUVs : uint8
+{
+	// Copy UVs to poly mesh vertices then interpolate as we subdivide.
+	Interpolated,
+
+	// Leave UVs uninitialized on output mesh
+	None
+};
+
+
+/**
+ * Interprets an FGroupTopology object as a poly mesh and refines it to a given level. Topological subdivision 
+ * can be computed separately from the interpolation of vertex positions. A typical workflow might be to first 
+ * precompute the refined topology, then compute the interpolated positions (and other vertex values) as the original 
+ * polymesh is deformed.
+ */
+class MESHMODELINGTOOLSEDITORONLY_API FSubdividePoly
+{
+public:
+
+	FSubdividePoly(const FGroupTopology& InTopology,
+				   const FDynamicMesh3& InOriginalMesh,
+				   int InLevel);
+
+	~FSubdividePoly();
+
+	bool ComputeTopologySubdivision();
+
+	bool ComputeSubdividedMesh(FDynamicMesh3& OutMesh);
+
+	ESubdivisionOutputNormals NormalComputationMethod = ESubdivisionOutputNormals::Generated;
+
+	ESubdivisionOutputUVs UVComputationMethod = ESubdivisionOutputUVs::Interpolated;
+
+	// wrapper around Opensubdiv::TopologyRefiner to avoid including OSD headers here. See cpp file.
+	class RefinerImpl;
+
+private:
+
+	const FGroupTopology& GroupTopology;
+	const FDynamicMesh3& OriginalMesh;
+	int Level;
+
+	TUniquePtr<RefinerImpl> Refiner;
+};
+
diff --git a/Engine/Plugins/Experimental/MeshModelingToolset/Source/MeshModelingToolsEditorOnly/Public/SubdividePolyTool.h b/Engine/Plugins/Experimental/MeshModelingToolset/Source/MeshModelingToolsEditorOnly/Public/SubdividePolyTool.h
new file mode 100644
--- /dev/null
+++ b/Engine/Plugins/Experimental/MeshModelingToolset/Source/MeshModelingToolsEditorOnly/Public/SubdividePolyTool.h
@@ -0,0 +1,105 @@
+// Copyright Epic Games, Inc. All Rights Reserved.
+
+#pragma once
+
+#include "SubdividePoly.h"
+#include "CoreMinimal.h"
+#include "InteractiveTool.h"
+#include "InteractiveToolBuilder.h"
+#include "ModelingOperators.h"
+#include "SingleSelectionTool.h"
+#include "SimpleDynamicMeshComponent.h"
+#include "MeshOpPreviewHelpers.h"
+#include "SubdividePolyTool.generated.h"
+
+class USubdividePolyTool;
+class USimpleDynamicMeshComponent;
+class UPreviewGeometry;
+
+/**
+ * Tool builder
+ */
+UCLASS()
+class MESHMODELINGTOOLSEDITORONLY_API USubdividePolyToolBuilder : public UInteractiveToolBuilder
+{
+	GENERATED_BODY()
+public:
+
+	bool CanBuildTool(const FToolBuilderState& SceneState) const override;
+
+	UInteractiveTool* BuildTool(const FToolBuilderState& SceneState) const override;
+};
+
+/**
+ * Properties
+ */
+
+UCLASS()
+class MESHMODELINGTOOLSEDITORONLY_API USubdividePolyToolProperties : public UInteractiveToolPropertySet
+{
+	GENERATED_BODY()
+
+public:
+
+	UPROPERTY(EditAnywhere, Category=Settings, meta = (UIMin = "1", ClampMin = "1"))
+	int SubdivisionLevel = 3;
+
+	UPROPERTY(EditAnywhere, Category=Settings)
+	ESubdivisionOutputNormals NormalComputationMethod = ESubdivisionOutputNormals::Generated;
+
+	UPROPERTY(EditAnywhere, Category=Settings)
+	ESubdivisionOutputUVs UVComputationMethod = ESubdivisionOutputUVs::Interpolated;
+
+	UPROPERTY(EditAnywhere, Category = Settings)
+	bool bRenderGroups = true;
+
+};
+
+
+/**
+ * Tool actual
+ */
+UCLASS()
+class MESHMODELINGTOOLSEDITORONLY_API USubdividePolyTool : public USingleSelectionTool
+{
+	GENERATED_BODY()
+
+public:
+
+	virtual void SetWorld(UWorld* World) { TargetWorld = World; }
+
+	virtual void Setup() override;
+	virtual void Shutdown(EToolShutdownType ShutdownType) override;
+
+	virtual bool HasCancel() const override { return true; }
+	virtual bool HasAccept() const override { return true; }
+	virtual bool CanAccept() const override;
+
+	void OnTick(float DeltaTime);
+
+protected:
+
+	friend class USubdividePolyOperatorFactory;
+
+	UWorld* TargetWorld;
+
+	UPROPERTY()
+	USimpleDynamicMeshComponent* PreviewDynamicMeshComponent = nullptr;
+
+	UPROPERTY()
+	USubdividePolyToolProperties* Properties = nullptr;
+
+	// The temporary actor we create internally to own the preview DynamicMeshComponent
+	UPROPERTY()
+	APreviewMeshActor* TemporaryParentActor = nullptr;
+
+	// Input mesh
+	TSharedPtr<FDynamicMesh3> OriginalMesh;
+
+	UPROPERTY()
+	UPreviewGeometry* PreviewGeometry = nullptr;
+
+	bool bPreviewGeometryNeedsUpdate;
+	void CreateOrUpdatePreviewGeometry();
+};
+
diff --git a/Engine/Plugins/Experimental/ModelingToolsEditorMode/Source/ModelingToolsEditorMode/Private/ModelingToolsEditorMode.cpp b/Engine/Plugins/Experimental/ModelingToolsEditorMode/Source/ModelingToolsEditorMode/Private/ModelingToolsEditorMode.cpp
--- a/Engine/Plugins/Experimental/ModelingToolsEditorMode/Source/ModelingToolsEditorMode/Private/ModelingToolsEditorMode.cpp
+++ b/Engine/Plugins/Experimental/ModelingToolsEditorMode/Source/ModelingToolsEditorMode/Private/ModelingToolsEditorMode.cpp
@@ -17,6 +17,7 @@
 #include "MeshVertexSculptTool.h"
 #include "EditMeshPolygonsTool.h"
 #include "DeformMeshPolygonsTool.h"
+#include "SubdividePolyTool.h"
 #include "GroupEdgeInsertionTool.h"
 #include "EdgeLoopInsertionTool.h"
 #include "ConvertToPolygonsTool.h"
@@ -622,6 +623,8 @@ void FModelingToolsEditorMode::Enter()
 	EdgeLoopInsertionToolBuilder->AssetAPI = ToolsContext->GetAssetAPI();
 	RegisterToolFunc(ToolManagerCommands.BeginEdgeLoopInsertionTool, TEXT("EdgeLoopInsertionTool"), EdgeLoopInsertionToolBuilder);
 
+	RegisterToolFunc(ToolManagerCommands.BeginSubdividePolyTool, TEXT("SubdividePolyTool"), NewObject<USubdividePolyToolBuilder>());
+
 	auto EditNormalsToolBuilder = NewObject<UEditNormalsToolBuilder>();
 	EditNormalsToolBuilder->AssetAPI = ModelingModeAssetGenerationAPI.Get();
 	RegisterToolFunc(ToolManagerCommands.BeginEditNormalsTool, TEXT("EditNormalsTool"), EditNormalsToolBuilder);
diff --git a/Engine/Plugins/Experimental/ModelingToolsEditorMode/Source/ModelingToolsEditorMode/Private/ModelingToolsEditorModeStyle.cpp b/Engine/Plugins/Experimental/ModelingToolsEditorMode/Source/ModelingToolsEditorMode/Private/ModelingToolsEditorModeStyle.cpp
--- a/Engine/Plugins/Experimental/ModelingToolsEditorMode/Source/ModelingToolsEditorMode/Private/ModelingToolsEditorModeStyle.cpp
+++ b/Engine/Plugins/Experimental/ModelingToolsEditorMode/Source/ModelingToolsEditorMode/Private/ModelingToolsEditorModeStyle.cpp
@@ -116,6 +116,8 @@ void FModelingToolsEditorModeStyle::Initialize()
 		StyleSet->Set("ModelingToolsManagerCommands.BeginSculptMeshTool.Small", 		new IMAGE_PLUGIN_BRUSH("Icons/Sculpt_40x", 			Icon20x20));
 		StyleSet->Set("ModelingToolsManagerCommands.BeginPolyEditTool", 				new IMAGE_PLUGIN_BRUSH("Icons/PolyEdit_40x", 		Icon20x20));
 		StyleSet->Set("ModelingToolsManagerCommands.BeginPolyEditTool.Small", 			new IMAGE_PLUGIN_BRUSH("Icons/PolyEdit_40x", 		Icon20x20));
+		//StyleSet->Set("ModelingToolsManagerCommands.BeginSubdividePolyTool",			new IMAGE_PLUGIN_BRUSH("Icons/PolyEdit_40x",		Icon20x20));
+		//StyleSet->Set("ModelingToolsManagerCommands.BeginSubdividePolyTool.Small",		new IMAGE_PLUGIN_BRUSH("Icons/PolyEdit_40x",		Icon20x20));
 		StyleSet->Set("ModelingToolsManagerCommands.BeginTriEditTool", 					new IMAGE_PLUGIN_BRUSH("Icons/TriEdit_40x", 		Icon20x20));
 		StyleSet->Set("ModelingToolsManagerCommands.BeginTriEditTool.Small", 			new IMAGE_PLUGIN_BRUSH("Icons/TriEdit_40x", 		Icon20x20));
 		//StyleSet->Set("ModelingToolsManagerCommands.BeginPolyDeformTool", 				new IMAGE_PLUGIN_BRUSH("Icons/PolyEdit_40x", 		Icon20x20));
diff --git a/Engine/Plugins/Experimental/ModelingToolsEditorMode/Source/ModelingToolsEditorMode/Private/ModelingToolsEditorModeToolkit.cpp b/Engine/Plugins/Experimental/ModelingToolsEditorMode/Source/ModelingToolsEditorMode/Private/ModelingToolsEditorModeToolkit.cpp
--- a/Engine/Plugins/Experimental/ModelingToolsEditorMode/Source/ModelingToolsEditorMode/Private/ModelingToolsEditorModeToolkit.cpp
+++ b/Engine/Plugins/Experimental/ModelingToolsEditorMode/Source/ModelingToolsEditorMode/Private/ModelingToolsEditorModeToolkit.cpp
@@ -659,6 +659,7 @@ void FModelingToolsEditorModeToolkit::BuildToolPalette_Experimental(FName Palett
 		ToolbarBuilder.AddSeparator();
 		ToolbarBuilder.AddToolBarButton(Commands.BeginAddPatchTool);
 		ToolbarBuilder.AddToolBarButton(Commands.BeginShapeSprayTool);
+		ToolbarBuilder.AddToolBarButton(Commands.BeginSubdividePolyTool);
 	}
 }
 
diff --git a/Engine/Plugins/Experimental/ModelingToolsEditorMode/Source/ModelingToolsEditorMode/Private/ModelingToolsManagerActions.cpp b/Engine/Plugins/Experimental/ModelingToolsEditorMode/Source/ModelingToolsEditorMode/Private/ModelingToolsManagerActions.cpp
--- a/Engine/Plugins/Experimental/ModelingToolsEditorMode/Source/ModelingToolsEditorMode/Private/ModelingToolsManagerActions.cpp
+++ b/Engine/Plugins/Experimental/ModelingToolsEditorMode/Source/ModelingToolsEditorMode/Private/ModelingToolsManagerActions.cpp
@@ -106,6 +106,7 @@ void FModelingToolsManagerCommands::RegisterCommands()
 	UI_COMMAND(BeginGroomCardsToMeshTool, "CardsToMesh", "Start the Hair Cards to Mesh Tool", EUserInterfaceActionType::ToggleButton, FInputChord());
 	UI_COMMAND(BeginGroomMeshToCardsTool, "MeshToCards", "Start the Mesh to Hair Cards Tool", EUserInterfaceActionType::ToggleButton, FInputChord());
 	UI_COMMAND(BeginGroomCardsEditorTool, "CardsEditor", "Start the Cards Editor Tool", EUserInterfaceActionType::ToggleButton, FInputChord());
+	UI_COMMAND(BeginSubdividePolyTool, "Subdivide", "Start the SubdividePoly Tool", EUserInterfaceActionType::ToggleButton, FInputChord());
 
 	UI_COMMAND(AcceptActiveTool, "Accept", "Accept the active tool", EUserInterfaceActionType::Button, FInputChord());
 	UI_COMMAND(CancelActiveTool, "Cancel", "Cancel the active tool", EUserInterfaceActionType::Button, FInputChord());
diff --git a/Engine/Plugins/Experimental/ModelingToolsEditorMode/Source/ModelingToolsEditorMode/Public/ModelingToolsManagerActions.h b/Engine/Plugins/Experimental/ModelingToolsEditorMode/Source/ModelingToolsEditorMode/Public/ModelingToolsManagerActions.h
--- a/Engine/Plugins/Experimental/ModelingToolsEditorMode/Source/ModelingToolsEditorMode/Public/ModelingToolsManagerActions.h
+++ b/Engine/Plugins/Experimental/ModelingToolsEditorMode/Source/ModelingToolsEditorMode/Public/ModelingToolsManagerActions.h
@@ -36,6 +36,7 @@ public:
 
 	TSharedPtr<FUICommandInfo> BeginSculptMeshTool;
 	TSharedPtr<FUICommandInfo> BeginPolyEditTool;
+	TSharedPtr<FUICommandInfo> BeginSubdividePolyTool;
 	TSharedPtr<FUICommandInfo> BeginGroupEdgeInsertionTool;
 	TSharedPtr<FUICommandInfo> BeginEdgeLoopInsertionTool;
 	TSharedPtr<FUICommandInfo> BeginTriEditTool;
