Fix possible dereferencing a nullptr

diff --git a/Engine/Plugins/Experimental/MeshModelingToolset/Source/MeshModelingToolsEditorOnly/Private/SubdividePolyTool.cpp b/Engine/Plugins/Experimental/MeshModelingToolset/Source/MeshModelingToolsEditorOnly/Private/SubdividePolyTool.cpp
--- a/Engine/Plugins/Experimental/MeshModelingToolset/Source/MeshModelingToolsEditorOnly/Private/SubdividePolyTool.cpp
+++ b/Engine/Plugins/Experimental/MeshModelingToolset/Source/MeshModelingToolsEditorOnly/Private/SubdividePolyTool.cpp
@@ -158,6 +158,10 @@ void USubdividePolyTool::Setup()
 	TemporaryParentActor = TargetWorld->SpawnActor<APreviewMeshActor>(FVector::ZeroVector, FRotator::ZeroRotator, SpawnInfo);
 
 	PreviewDynamicMeshComponent = NewObject<USimpleDynamicMeshComponent>(TemporaryParentActor, "DynamicMesh");
+	if (PreviewDynamicMeshComponent == nullptr)
+	{
+		return;
+	}
 	PreviewDynamicMeshComponent->SetRenderMeshPostProcessor(MakeUnique<SubdivPostProcessor>(Properties->SubdivisionLevel,
 																					 Properties->NormalComputationMethod,
 																					 Properties->UVComputationMethod));
@@ -190,13 +194,10 @@ void USubdividePolyTool::Setup()
 
 	auto RebuildMeshPostProcessor = [this]()
 	{
-		if (PreviewDynamicMeshComponent)
-		{
-			PreviewDynamicMeshComponent->SetRenderMeshPostProcessor(MakeUnique<SubdivPostProcessor>(Properties->SubdivisionLevel,
-																							 Properties->NormalComputationMethod,
-																							 Properties->UVComputationMethod));
-			PreviewDynamicMeshComponent->NotifyMeshUpdated();
-		}
+		PreviewDynamicMeshComponent->SetRenderMeshPostProcessor(MakeUnique<SubdivPostProcessor>(Properties->SubdivisionLevel,
+																							Properties->NormalComputationMethod,
+																							Properties->UVComputationMethod));
+		PreviewDynamicMeshComponent->NotifyMeshUpdated();
 	};
 
 	// Watch for property changes
@@ -216,23 +217,20 @@ void USubdividePolyTool::Setup()
 
 	auto RenderGroupsChanged = [this](bool bNewRenderGroups)
 	{
-		if (PreviewDynamicMeshComponent)
+		if (bNewRenderGroups)
 		{
-			if (bNewRenderGroups)
-			{
-				PreviewDynamicMeshComponent->SetOverrideRenderMaterial(ToolSetupUtil::GetSelectionMaterial(GetToolManager()));
-				PreviewDynamicMeshComponent->TriangleColorFunc = [](const FDynamicMesh3* Mesh, int TriangleID)
-				{
-					return LinearColors::SelectFColor(Mesh->GetTriangleGroup(TriangleID));
-				};
-			}
-			else
+			PreviewDynamicMeshComponent->SetOverrideRenderMaterial(ToolSetupUtil::GetSelectionMaterial(GetToolManager()));
+			PreviewDynamicMeshComponent->TriangleColorFunc = [](const FDynamicMesh3* Mesh, int TriangleID)
 			{
-				PreviewDynamicMeshComponent->SetOverrideRenderMaterial(nullptr);
-				PreviewDynamicMeshComponent->TriangleColorFunc = nullptr;
-			}
-			PreviewDynamicMeshComponent->FastNotifyColorsUpdated();
+				return LinearColors::SelectFColor(Mesh->GetTriangleGroup(TriangleID));
+			};
+		}
+		else
+		{
+			PreviewDynamicMeshComponent->SetOverrideRenderMaterial(nullptr);
+			PreviewDynamicMeshComponent->TriangleColorFunc = nullptr;
 		}
+		PreviewDynamicMeshComponent->FastNotifyColorsUpdated();
 	};
 
 	Properties->WatchProperty(Properties->bRenderGroups, RenderGroupsChanged);
@@ -287,8 +285,10 @@ void USubdividePolyTool::Shutdown(EToolShutdownType ShutdownType)
 		Properties->SaveProperties(this);
 	}
 
-	check(PreviewGeometry);
-	PreviewGeometry->Disconnect();
+	if (PreviewGeometry)
+	{
+		PreviewGeometry->Disconnect();
+	}
 
 	if (PreviewDynamicMeshComponent)
 	{
