MeshModelingTools: Added marquee select of vertices and edges in PolyEdit tool. Also refactored the selection mechanic to handle its own behaviors and adusted the using tools.

#rnx

diff --git a/Engine/Plugins/Experimental/GeometryProcessing/Source/GeometricObjects/Public/Spatial/GeometrySet3.h b/Engine/Plugins/Experimental/GeometryProcessing/Source/GeometricObjects/Public/Spatial/GeometrySet3.h
--- a/Engine/Plugins/Experimental/GeometryProcessing/Source/GeometricObjects/Public/Spatial/GeometrySet3.h
+++ b/Engine/Plugins/Experimental/GeometryProcessing/Source/GeometricObjects/Public/Spatial/GeometrySet3.h
@@ -2,6 +2,7 @@
 
 #pragma once
 
+#include "Async/ParallelFor.h"
 #include "BoxTypes.h"
 #include "RayTypes.h"
 #include "Polyline3.h"
@@ -35,7 +36,7 @@ public:
 	void UpdatePoint(int PointID, const FVector3d& Position);
 	/** Update the Polyline of previously-added CurveID */
 	void UpdateCurve(int CurveID, const FPolyline3d& Polyline);
-	
+
 	/**
 	 * FNearest is returned by nearest-point queries
 	 */
@@ -105,6 +106,100 @@ public:
 	bool CollectCurvesNearRay(const FRay3d& Ray, TArray<FNearest>& ResultsOut,
 		TFunction<bool(const FVector3d&, const FVector3d&)> PointWithinToleranceTest) const;
 
+	/**
+	 * Fills PointIDsOut with point IDs of points that satisfy the given predicate.
+	 *
+	 * @param Predicate A lambda or function that takes in a const FVector3d& and returns
+	 *  true when the point's ID should be added to PointIDsOut.
+	 * @param PointIDsOut Output list or set. Must have an Add(int) method.
+	 */
+	template <typename PredicateType, typename IntContainerType>
+	bool FindAllPointsSatisfying(PredicateType Predicate, IntContainerType& PointIDsOut) const
+	{
+		PointIDsOut.Reset();
+		for (const FPoint& Point : Points)
+		{
+			if (Invoke(Predicate, Point.Position))
+			{
+				PointIDsOut.Add(Point.ID);
+			}
+		}
+		return !PointIDsOut.IsEmpty();
+	}
+
+	/**
+	 * Like FindAllPointsSatisfying, but parallel, so predicate must be safe to call in parallel.
+	 */
+	template <typename PredicateType, typename IntContainerType>
+	bool ParallelFindAllPointsSatisfying(PredicateType Predicate, IntContainerType& PointIDsOut) const
+	{
+		TArray<bool> Flags;
+		Flags.SetNumZeroed(Points.Num());
+		ParallelFor(Points.Num(), [this, &Predicate, &Flags](int32 i)
+			{
+				if (Invoke(Predicate, Points[i].Position))
+				{
+					Flags[i] = true;
+				}
+			});
+
+		for (int32 i = 0; i < Points.Num(); ++i)
+		{
+			if (Flags[i])
+			{
+				PointIDsOut.Add(Points[i].ID);
+			}
+		}
+		return !PointIDsOut.IsEmpty();
+	}
+
+	/**
+	 * Fills CurveIDsOut with IDs of curves that satisfy the given predicate.
+	 *
+	 * @param Predicate A lambda or function that takes in a const FPolyline3d& and returns
+	 *  true when the curve's ID should be added to PointIDsOut.
+	 * @param CurveIDsOut Output list or set. Must have an Add(int) method.
+	 */
+	template <typename PredicateType, typename IntContainerType>
+	bool FindAllCurvesSatisfying(PredicateType Predicate, IntContainerType& CurveIDsOut) const
+	{
+		CurveIDsOut.Reset();
+		for (const FCurve& Curve : Curves)
+		{
+			if (Invoke(Predicate, Curve.Geometry))
+			{
+				CurveIDsOut.Add(Curve.ID);
+			}
+		}
+		return !CurveIDsOut.IsEmpty();
+	}
+
+	/**
+	 * Like FindAllCurvesSatisfying, but parallel, so predicate must be safe to call in parallel.
+	 */
+	template <typename PredicateType, typename IntContainerType>
+	bool ParallelFindAllCurvesSatisfying(PredicateType Predicate, IntContainerType& CurveIDsOut) const
+	{
+		TArray<bool> Flags;
+		Flags.SetNumZeroed(Curves.Num());
+		ParallelFor(Curves.Num(), [this, &Predicate, &Flags](int32 i)
+		{
+			if (Invoke(Predicate, Curves[i].Geometry))
+			{
+				Flags[i] = true;
+			}
+		});
+
+		for (int32 i = 0; i < Curves.Num(); ++i)
+		{
+			if (Flags[i])
+			{
+				CurveIDsOut.Add(Curves[i].ID);
+			}
+		}
+		return !CurveIDsOut.IsEmpty();
+	}
+
 protected:
 
 	struct FPoint
@@ -119,7 +214,7 @@ protected:
 	{
 		int ID;
 		FPolyline3d Geometry;
-		
+
 		FAxisAlignedBox3d Bounds;
 	};
 	TArray<FCurve> Curves;
diff --git a/Engine/Plugins/Experimental/MeshModelingToolset/Source/MeshModelingTools/Private/EditMeshPolygonsTool.cpp b/Engine/Plugins/Experimental/MeshModelingToolset/Source/MeshModelingTools/Private/EditMeshPolygonsTool.cpp
--- a/Engine/Plugins/Experimental/MeshModelingToolset/Source/MeshModelingTools/Private/EditMeshPolygonsTool.cpp
+++ b/Engine/Plugins/Experimental/MeshModelingToolset/Source/MeshModelingTools/Private/EditMeshPolygonsTool.cpp
@@ -11,6 +11,7 @@
 #include "ToolSceneQueriesUtil.h"
 #include "Intersection/IntersectionUtil.h"
 #include "Transforms/MultiTransformer.h"
+#include "TransformTypes.h"
 #include "BaseBehaviors/SingleClickBehavior.h"
 #include "Util/ColorConstants.h"
 #include "ToolSetupUtil.h"
@@ -352,6 +353,7 @@ FDynamicMeshAABBTree3& UEditMeshPolygonsTool::GetSpatial()
 
 bool UEditMeshPolygonsTool::HitTest(const FRay& WorldRay, FHitResult& OutHit)
 {
+	// If we're in the middle of an action, take the click (to finish an inset, etc).
 	if (CurrentToolMode != ECurrentToolMode::TransformSelection)
 	{
 		OutHit.Distance = 100.0;
@@ -359,8 +361,9 @@ bool UEditMeshPolygonsTool::HitTest(const FRay& WorldRay, FHitResult& OutHit)
 		return true;
 	}
 
-	// disable hit test
-	return SelectionMechanic->TopologyHitTest(WorldRay, OutHit);
+	// The selection mechanic and gizmo will take care of the TransformSelection state.
+
+	return false;
 }
 
 
@@ -372,9 +375,7 @@ FInputRayHit UEditMeshPolygonsTool::IsHitByClick(const FInputDeviceRay& ClickPos
 	{
 		return FInputRayHit(OutHit.Distance);
 	}
-
-	// background capture, if nothing else is hit
-	return FInputRayHit(TNumericLimits<float>::Max());
+	return FInputRayHit(); // bHit is set to false.
 }
 
 void UEditMeshPolygonsTool::OnClicked(const FInputDeviceRay& ClickPos)
@@ -417,19 +418,9 @@ void UEditMeshPolygonsTool::OnClicked(const FInputDeviceRay& ClickPos)
 		return;
 	}
 
-	// update selection
-	GetToolManager()->BeginUndoTransaction(LOCTEXT("PolyMeshSelectionChange", "Selection"));
-	SelectionMechanic->BeginChange();
-	FVector3d LocalHitPosition, LocalHitNormal;
-	bool bSelectionModified = SelectionMechanic->UpdateSelection(ClickPos.WorldRay, LocalHitPosition, LocalHitNormal);
-	if (bSelectionModified) 
-	{
-		FFrame3d LocalFrame(LocalHitPosition, LocalHitNormal);
-		LastGeometryFrame = SelectionMechanic->GetSelectionFrame(true, &LocalFrame);
-		UpdateMultiTransformerFrame();
-	}
-	SelectionMechanic->EndChangeAndEmitIfModified();
-	GetToolManager()->EndUndoTransaction();
+	ensureMsgf(CurrentToolMode != ECurrentToolMode::TransformSelection, TEXT("EditMeshPolygonsTool: "
+		"Should not receive click requests in transform mode- they should have been handled by "
+		"selection mechanic or gizmo."));
 }
 
 
@@ -465,6 +456,11 @@ void UEditMeshPolygonsTool::UpdateMultiTransformerFrame(const FFrame3d* UseFrame
 
 void UEditMeshPolygonsTool::OnSelectionModifiedEvent()
 {
+	FVector3d LocalLastHitPosition, LocalLastHitNormal;
+	SelectionMechanic->GetClickedHitPosition(LocalLastHitPosition, LocalLastHitNormal);
+	FFrame3d LocalFrame(LocalLastHitPosition, LocalLastHitNormal);
+	LastGeometryFrame = SelectionMechanic->GetSelectionFrame(true, &LocalFrame);
+	UpdateMultiTransformerFrame();
 	bSelectionStateDirty = true;
 }
 
@@ -561,10 +557,9 @@ bool UEditMeshPolygonsTool::OnUpdateHover(const FInputDeviceRay& DevicePos)
 		return true;
 	}
 
-	if (ActiveVertexChange == nullptr && MultiTransformer->InGizmoEdit() == false )
-	{
-		SelectionMechanic->UpdateHighlight(DevicePos.WorldRay);
-	}
+	ensureMsgf(CurrentToolMode != ECurrentToolMode::TransformSelection, TEXT("EditMeshPolygonsTool: "
+		"Should not receive hover requests in transform mode- they should have been handled by "
+		"selection mechanic or gizmo."));
 	return true;
 }
 
@@ -825,8 +820,7 @@ void UEditMeshPolygonsTool::PrecomputeTopology()
 
 	// update selection mechanic
 	SelectionMechanic->Initialize(DynamicMeshComponent, Topology.Get(),
-		[this]() { return &GetSpatial(); },
-		[this]() { return GetShiftToggle(); }
+		[this]() { return &GetSpatial(); }
 		);
 
 	LinearDeformer.Initialize(Mesh, Topology.Get());
@@ -856,6 +850,11 @@ void UEditMeshPolygonsTool::Render(IToolsContextRenderAPI* RenderAPI)
 	}
 }
 
+void UEditMeshPolygonsTool::DrawHUD(FCanvas* Canvas, IToolsContextRenderAPI* RenderAPI)
+{
+	SelectionMechanic->DrawHUD(Canvas, RenderAPI);
+}
+
 
 
 
diff --git a/Engine/Plugins/Experimental/MeshModelingToolset/Source/MeshModelingTools/Private/EditUVIslandsTool.cpp b/Engine/Plugins/Experimental/MeshModelingToolset/Source/MeshModelingTools/Private/EditUVIslandsTool.cpp
--- a/Engine/Plugins/Experimental/MeshModelingToolset/Source/MeshModelingTools/Private/EditUVIslandsTool.cpp
+++ b/Engine/Plugins/Experimental/MeshModelingToolset/Source/MeshModelingTools/Private/EditUVIslandsTool.cpp
@@ -40,11 +40,6 @@ void UEditUVIslandsTool::Setup()
 {
 	UMeshSurfacePointTool::Setup();
 
-	// register click behavior
-	USingleClickInputBehavior* ClickBehavior = NewObject<USingleClickInputBehavior>();
-	ClickBehavior->Initialize(this);
-	AddInputBehavior(ClickBehavior);
-
 	// create dynamic mesh component to use for live preview
 	DynamicMeshComponent = NewObject<USimpleDynamicMeshComponent>(ComponentTarget->GetOwnerActor(), "DynamicMesh");
 	DynamicMeshComponent->SetupAttachment(ComponentTarget->GetOwnerActor()->GetRootComponent());
@@ -78,7 +73,9 @@ void UEditUVIslandsTool::Setup()
 	SelectionMechanic = NewObject<UPolygonSelectionMechanic>(this);
 	SelectionMechanic->bAddSelectionFilterPropertiesToParentTool = false;
 	SelectionMechanic->Setup(this);
-	SelectionMechanic->Properties->bSelectEdges = SelectionMechanic->Properties->bSelectVertices = false;
+	SelectionMechanic->Properties->bSelectEdges = false;
+	SelectionMechanic->Properties->bSelectVertices = false;
+	SelectionMechanic->Properties->bEnableMarquee = false;
 	SelectionMechanic->OnSelectionChanged.AddUObject(this, &UEditUVIslandsTool::OnSelectionModifiedEvent);
 
 	// initialize AABBTree
@@ -194,63 +191,25 @@ FDynamicMeshAABBTree3& UEditUVIslandsTool::GetSpatial()
 	return MeshSpatial;
 }
 
-
-
-
-
-
-
-bool UEditUVIslandsTool::HitTest(const FRay& WorldRay, FHitResult& OutHit)
-{
-	// disable hit test
-	return SelectionMechanic->TopologyHitTest(WorldRay, OutHit);
-}
-
-
-
-FInputRayHit UEditUVIslandsTool::IsHitByClick(const FInputDeviceRay& ClickPos)
-{
-	FHitResult OutHit;
-	if (HitTest(ClickPos.WorldRay, OutHit))
-	{
-		return FInputRayHit(OutHit.Distance);
-	}
-
-	// background capture, if nothing else is hit
-	return FInputRayHit(TNumericLimits<float>::Max());
-}
-
-void UEditUVIslandsTool::OnClicked(const FInputDeviceRay& ClickPos)
+void UEditUVIslandsTool::OnSelectionModifiedEvent()
 {
-	// update selection
-	GetToolManager()->BeginUndoTransaction(LOCTEXT("PolyMeshSelectionChange", "Selection"));
-	SelectionMechanic->BeginChange();
-	FVector3d LocalHitPos, LocalHitNormal;
-	bool bSelectionModified = SelectionMechanic->UpdateSelection(ClickPos.WorldRay, LocalHitPos, LocalHitNormal);
-
-	if (bSelectionModified && SelectionMechanic->GetActiveSelection().IsEmpty() == false)
+	bSelectionStateDirty = true;
+	if (!SelectionMechanic->GetActiveSelection().IsEmpty())
 	{
 		FFrame3d UseFrame = Topology.GetIslandFrame(
 			SelectionMechanic->GetActiveSelection().GetASelectedGroupID(), GetSpatial());
 		UseFrame.Transform(WorldTransform);
-		MultiTransformer->UpdateGizmoPositionFromWorldFrame(UseFrame, true);
-		//MultiTransformer->SetGizmoPositionFromWorldFrame(SelectionMechanic->GetSelectionFrame(true), true);
-
+		MultiTransformer->InitializeGizmoPositionFromWorldFrame(UseFrame, true);
 	}
-
-	SelectionMechanic->EndChangeAndEmitIfModified();
-	GetToolManager()->EndUndoTransaction();
 }
 
 
-void UEditUVIslandsTool::OnSelectionModifiedEvent()
+
+bool UEditUVIslandsTool::HitTest(const FRay& WorldRay, FHitResult& OutHit)
 {
-	bSelectionStateDirty = true;
+	return SelectionMechanic->TopologyHitTest(WorldRay, OutHit);
 }
 
-
-
-
 FInputRayHit UEditUVIslandsTool::CanBeginClickDragSequence(const FInputDeviceRay& PressPos)
 {
 	// disable this for now
@@ -258,14 +217,10 @@ FInputRayHit UEditUVIslandsTool::CanBeginClickDragSequence(const FInputDeviceRay
 	//return UMeshSurfacePointTool::CanBeginClickDragSequence(PressPos);
 }
 
-
-
 void UEditUVIslandsTool::OnBeginDrag(const FRay& WorldRay)
 {
 }
 
-
-
 void UEditUVIslandsTool::OnUpdateDrag(const FRay& Ray)
 {
 	check(false);
@@ -518,8 +473,7 @@ void UEditUVIslandsTool::PrecomputeTopology()
 
 	// update selection mechanic
 	SelectionMechanic->Initialize(DynamicMeshComponent, &Topology,
-		[this]() { return &GetSpatial(); },
-		[this]() { return GetShiftToggle(); }
+		[this]() { return &GetSpatial(); }
 		);
 }
 
diff --git a/Engine/Plugins/Experimental/MeshModelingToolset/Source/MeshModelingTools/Private/HoleFillTool.cpp b/Engine/Plugins/Experimental/MeshModelingToolset/Source/MeshModelingTools/Private/HoleFillTool.cpp
--- a/Engine/Plugins/Experimental/MeshModelingToolset/Source/MeshModelingTools/Private/HoleFillTool.cpp
+++ b/Engine/Plugins/Experimental/MeshModelingToolset/Source/MeshModelingTools/Private/HoleFillTool.cpp
@@ -161,16 +161,6 @@ void UHoleFillTool::Setup()
 
 	ToolPropertyObjects.Add(this);
 
-	// click behavior
-	USingleClickInputBehavior* ClickBehavior = NewObject<USingleClickInputBehavior>();
-	ClickBehavior->Initialize(this);
-	AddInputBehavior(ClickBehavior);
-
-	// hover behavior
-	UMouseHoverBehavior* HoverBehavior = NewObject<UMouseHoverBehavior>();
-	HoverBehavior->Initialize(this);
-	AddInputBehavior(HoverBehavior);
-
 	// initialize hit query
 	MeshSpatial.SetMesh(OriginalMesh.Get());
 
@@ -189,9 +179,13 @@ void UHoleFillTool::Setup()
 		ComponentTarget->GetWorldTransform(),
 		TargetWorld,
 		Topology.Get(),
-		[this]() { return &MeshSpatial; },
-		[]() { return true; }	// allow adding to selection without modifier key
+		[this]() { return &MeshSpatial; }
 	);
+	// allow toggling selection without modifier key
+	SelectionMechanic->SetShouldAddToSelectionFunc([]() {return true; });
+	SelectionMechanic->SetShouldRemoveFromSelectionFunc([]() {return true; });
+
+	SelectionMechanic->OnSelectionChanged.AddUObject(this, &UHoleFillTool::OnSelectionModified);
 	
 	// Store a UV scale based on the original mesh bounds
 	MeshUVScaleFactor = (1.0 / OriginalMesh->GetBounds().MaxDim());
@@ -292,56 +286,10 @@ void UHoleFillTool::Shutdown(EToolShutdownType ShutdownType)
 	}
 }
 
-FInputRayHit UHoleFillTool::IsHitByClick(const FInputDeviceRay& ClickPos)
-{
-	FHitResult OutHit;
-	if (SelectionMechanic->TopologyHitTest(ClickPos.WorldRay, OutHit))
-	{
-		return FInputRayHit(OutHit.Distance);
-	}
-
-	return FInputRayHit(TNumericLimits<float>::Max());
-}
-
-void UHoleFillTool::OnClicked(const FInputDeviceRay& ClickPos)
-{
-	// update selection
-	GetToolManager()->BeginUndoTransaction(LOCTEXT("PolyMeshSelectionChange", "Selection"));
-	SelectionMechanic->BeginChange();
-
-	FVector3d LocalHitPosition, LocalHitNormal;
-	bool bSelectionModified = SelectionMechanic->UpdateSelection(ClickPos.WorldRay, LocalHitPosition, LocalHitNormal);
-
-	if (bSelectionModified)
-	{
-		UpdateActiveBoundaryLoopSelection();
-		InvalidatePreviewResult();
-	}
-
-	SelectionMechanic->EndChangeAndEmitIfModified();
-	GetToolManager()->EndUndoTransaction();
-}
-
-
-FInputRayHit UHoleFillTool::BeginHoverSequenceHitTest(const FInputDeviceRay& PressPos)
-{
-	FHitResult OutHit;
-	if (SelectionMechanic->TopologyHitTest(PressPos.WorldRay, OutHit))
-	{
-		return FInputRayHit(OutHit.Distance);
-	}
-	return FInputRayHit();
-}
-
-bool UHoleFillTool::OnUpdateHover(const FInputDeviceRay& DevicePos) 
+void UHoleFillTool::OnSelectionModified()
 {
-	SelectionMechanic->UpdateHighlight(DevicePos.WorldRay);
-	return true;
-}
-
-void UHoleFillTool::OnEndHover()
-{
-	SelectionMechanic->ClearHighlight();
+	UpdateActiveBoundaryLoopSelection();
+	InvalidatePreviewResult();
 }
 
 void UHoleFillTool::RequestAction(EHoleFillToolActions ActionType)
diff --git a/Engine/Plugins/Experimental/MeshModelingToolset/Source/MeshModelingTools/Private/LatticeDeformerTool.cpp b/Engine/Plugins/Experimental/MeshModelingToolset/Source/MeshModelingTools/Private/LatticeDeformerTool.cpp
--- a/Engine/Plugins/Experimental/MeshModelingToolset/Source/MeshModelingTools/Private/LatticeDeformerTool.cpp
+++ b/Engine/Plugins/Experimental/MeshModelingToolset/Source/MeshModelingTools/Private/LatticeDeformerTool.cpp
@@ -71,14 +71,7 @@ FVector3i ULatticeDeformerTool::GetLatticeResolution() const
 
 void ULatticeDeformerTool::DrawHUD(FCanvas* Canvas, IToolsContextRenderAPI* RenderAPI)
 {
-	EViewInteractionState State = RenderAPI->GetViewInteractionState();
-	bool bThisViewHasFocus = !!(State & EViewInteractionState::Focused);
-
-	if (bThisViewHasFocus)
-	{		
-		// Draw the drag rectangle if it's active
-		ControlPointsMechanic->DrawHUD(Canvas, RenderAPI);
-	}
+	ControlPointsMechanic->DrawHUD(Canvas, RenderAPI);
 }
 
 bool ULatticeDeformerTool::CanAccept() const
diff --git a/Engine/Plugins/Experimental/MeshModelingToolset/Source/MeshModelingTools/Private/MeshBoundaryToolBase.cpp b/Engine/Plugins/Experimental/MeshModelingToolset/Source/MeshModelingTools/Private/MeshBoundaryToolBase.cpp
--- a/Engine/Plugins/Experimental/MeshModelingToolset/Source/MeshModelingTools/Private/MeshBoundaryToolBase.cpp
+++ b/Engine/Plugins/Experimental/MeshModelingToolset/Source/MeshModelingTools/Private/MeshBoundaryToolBase.cpp
@@ -25,14 +25,6 @@ void UMeshBoundaryToolBase::Setup()
 	FMeshDescriptionToDynamicMesh Converter;
 	Converter.Convert(ComponentTarget->GetMesh(), *OriginalMesh);
 
-	LoopSelectClickBehavior = NewObject<USingleClickInputBehavior>();
-	LoopSelectClickBehavior->Initialize(this);
-	AddInputBehavior(LoopSelectClickBehavior);
-
-	UMouseHoverBehavior* HoverBehavior = NewObject<UMouseHoverBehavior>();
-	HoverBehavior->Initialize(this);
-	AddInputBehavior(HoverBehavior);
-
 	// initialize hit query
 	MeshSpatial.SetMesh(OriginalMesh.Get());
 
@@ -51,9 +43,9 @@ void UMeshBoundaryToolBase::Setup()
 		ComponentTarget->GetWorldTransform(),
 		TargetWorld,
 		Topology.Get(),
-		[this]() { return &MeshSpatial; },
-		[this]() { return ShouldSelectionAppend(); }
+		[this]() { return &MeshSpatial; }
 	);
+	SelectionMechanic->OnSelectionChanged.AddUObject(this, &UMeshBoundaryToolBase::OnSelectionChanged);
 }
 
 void UMeshBoundaryToolBase::Shutdown(EToolShutdownType ShutdownType)
@@ -72,56 +64,4 @@ void UMeshBoundaryToolBase::Render(IToolsContextRenderAPI* RenderAPI)
 	}
 }
 
-
-FInputRayHit UMeshBoundaryToolBase::IsHitByClick(const FInputDeviceRay& ClickPos)
-{
-	FHitResult OutHit;
-	if (SelectionMechanic->TopologyHitTest(ClickPos.WorldRay, OutHit))
-	{
-		return FInputRayHit(OutHit.Distance);
-	}
-	FInputRayHit Result(TNumericLimits<float>::Max());
-	Result.bHit = false;
-	return Result;
-}
-
-void UMeshBoundaryToolBase::OnClicked(const FInputDeviceRay& ClickPos)
-{
-	// update selection
-	GetToolManager()->BeginUndoTransaction(LOCTEXT("BoundarySelectionChange", "Selection"));
-	SelectionMechanic->BeginChange();
-
-	FVector3d LocalHitPosition, LocalHitNormal;
-	bool bSelectionModified = SelectionMechanic->UpdateSelection(ClickPos.WorldRay, LocalHitPosition, LocalHitNormal);
-
-	if (bSelectionModified)
-	{
-		OnSelectionChanged();
-	}
-
-	SelectionMechanic->EndChangeAndEmitIfModified();
-	GetToolManager()->EndUndoTransaction();
-}
-
-FInputRayHit UMeshBoundaryToolBase::BeginHoverSequenceHitTest(const FInputDeviceRay& PressPos)
-{
-	FHitResult OutHit;
-	if (SelectionMechanic->TopologyHitTest(PressPos.WorldRay, OutHit))
-	{
-		return FInputRayHit(OutHit.Distance);
-	}
-	return FInputRayHit();
-}
-
-bool UMeshBoundaryToolBase::OnUpdateHover(const FInputDeviceRay& DevicePos)
-{
-	SelectionMechanic->UpdateHighlight(DevicePos.WorldRay);
-	return true;
-}
-
-void UMeshBoundaryToolBase::OnEndHover()
-{
-	SelectionMechanic->ClearHighlight();
-}
-
 #undef LOCTEXT_NAMESPACE
diff --git a/Engine/Plugins/Experimental/MeshModelingToolset/Source/MeshModelingTools/Private/RevolveBoundaryTool.cpp b/Engine/Plugins/Experimental/MeshModelingToolset/Source/MeshModelingTools/Private/RevolveBoundaryTool.cpp
--- a/Engine/Plugins/Experimental/MeshModelingToolset/Source/MeshModelingTools/Private/RevolveBoundaryTool.cpp
+++ b/Engine/Plugins/Experimental/MeshModelingToolset/Source/MeshModelingTools/Private/RevolveBoundaryTool.cpp
@@ -81,6 +81,18 @@ void URevolveBoundaryTool::Setup()
 {
 	UMeshBoundaryToolBase::Setup();
 
+	// We're actually going to handle the selection clicks ourselves so that we can align axis if
+	// we want to.
+	SelectionMechanic->DisableBehaviors(this);
+	SelectionMechanic->SetShouldAddToSelectionFunc([]() {return false; });
+	SelectionMechanic->SetShouldRemoveFromSelectionFunc([]() {return false; });
+
+	USingleClickInputBehavior* ClickBehavior = NewObject<USingleClickInputBehavior>();
+	ClickBehavior->Initialize(this);
+	ClickBehavior->Modifiers.RegisterModifier(AlignAxisModifier, FInputDeviceState::IsCtrlKeyDown);
+	AddInputBehavior(ClickBehavior);
+	
+
 	Settings = NewObject<URevolveBoundaryToolProperties>(this);
 	Settings->RestoreProperties(this);
 	AddToolPropertySource(Settings);
@@ -96,7 +108,7 @@ void URevolveBoundaryTool::Setup()
 	PlaneMechanic->Setup(this);
 	PlaneMechanic->Initialize(TargetWorld, FFrame3d(Settings->AxisOrigin, 
 		FRotator(Settings->AxisPitch, Settings->AxisYaw, 0).Quaternion()));
-	PlaneMechanic->UpdateClickPriority(LoopSelectClickBehavior->GetPriority().MakeLower());
+	PlaneMechanic->UpdateClickPriority(ClickBehavior->GetPriority().MakeLower());
 	PlaneMechanic->bShowGrid = false;
 	PlaneMechanic->OnPlaneChanged.AddLambda([this]() {
 		Settings->AxisOrigin = (FVector)PlaneMechanic->Plane.Origin;
@@ -108,8 +120,6 @@ void URevolveBoundaryTool::Setup()
 
 	PlaneMechanic->SetEnableGridSnaping(Settings->bSnapToWorldGrid);
 
-	LoopSelectClickBehavior->Modifiers.RegisterModifier(AlignAxisModifier, FInputDeviceState::IsCtrlKeyDown);
-
 	ComponentTarget->SetOwnerVisibility(Settings->bDisplayOriginalMesh);
 
 	GetToolManager()->DisplayMessage(
@@ -144,6 +154,16 @@ void URevolveBoundaryTool::OnUpdateModifierState(int ModifierId, bool bIsOn)
 	}
 }
 
+FInputRayHit URevolveBoundaryTool::IsHitByClick(const FInputDeviceRay& ClickPos)
+{
+	FHitResult OutHit;
+	if (SelectionMechanic->TopologyHitTest(ClickPos.WorldRay, OutHit))
+	{
+		return FInputRayHit(OutHit.Distance);
+	}
+	return FInputRayHit(); // bHit is false
+}
+
 void URevolveBoundaryTool::OnClicked(const FInputDeviceRay& ClickPos)
 {
 	// Update selection only if we clicked on something. We don't want to be able to
diff --git a/Engine/Plugins/Experimental/MeshModelingToolset/Source/MeshModelingTools/Public/EditMeshPolygonsTool.h b/Engine/Plugins/Experimental/MeshModelingToolset/Source/MeshModelingTools/Public/EditMeshPolygonsTool.h
--- a/Engine/Plugins/Experimental/MeshModelingToolset/Source/MeshModelingTools/Public/EditMeshPolygonsTool.h
+++ b/Engine/Plugins/Experimental/MeshModelingToolset/Source/MeshModelingTools/Public/EditMeshPolygonsTool.h
@@ -447,6 +447,7 @@ public:
 
 	virtual void OnTick(float DeltaTime) override;
 	virtual void Render(IToolsContextRenderAPI* RenderAPI) override;
+	virtual void DrawHUD(FCanvas* Canvas, IToolsContextRenderAPI* RenderAPI) override;
 
 	virtual bool HasCancel() const override { return true; }
 	virtual bool HasAccept() const override { return true; }
diff --git a/Engine/Plugins/Experimental/MeshModelingToolset/Source/MeshModelingTools/Public/EditUVIslandsTool.h b/Engine/Plugins/Experimental/MeshModelingToolset/Source/MeshModelingTools/Public/EditUVIslandsTool.h
--- a/Engine/Plugins/Experimental/MeshModelingToolset/Source/MeshModelingTools/Public/EditUVIslandsTool.h
+++ b/Engine/Plugins/Experimental/MeshModelingToolset/Source/MeshModelingTools/Public/EditUVIslandsTool.h
@@ -57,7 +57,7 @@ public:
  *
  */
 UCLASS()
-class MESHMODELINGTOOLS_API UEditUVIslandsTool : public UMeshSurfacePointTool, public IClickBehaviorTarget
+class MESHMODELINGTOOLS_API UEditUVIslandsTool : public UMeshSurfacePointTool
 {
 	GENERATED_BODY()
 
@@ -88,10 +88,6 @@ public:
 	// IClickDragBehaviorTarget API
 	virtual FInputRayHit CanBeginClickDragSequence(const FInputDeviceRay& PressPos) override;
 
-	// IClickBehaviorTarget API
-	virtual FInputRayHit IsHitByClick(const FInputDeviceRay& ClickPos) override;
-	virtual void OnClicked(const FInputDeviceRay& ClickPos) override;
-
 
 public:
 	UPROPERTY()
diff --git a/Engine/Plugins/Experimental/MeshModelingToolset/Source/MeshModelingTools/Public/HoleFillTool.h b/Engine/Plugins/Experimental/MeshModelingToolset/Source/MeshModelingTools/Public/HoleFillTool.h
--- a/Engine/Plugins/Experimental/MeshModelingToolset/Source/MeshModelingTools/Public/HoleFillTool.h
+++ b/Engine/Plugins/Experimental/MeshModelingToolset/Source/MeshModelingTools/Public/HoleFillTool.h
@@ -211,7 +211,7 @@ public:
  */
 
 UCLASS()
-class MESHMODELINGTOOLS_API UHoleFillTool : public USingleSelectionTool, public IClickBehaviorTarget, public IHoverBehaviorTarget
+class MESHMODELINGTOOLS_API UHoleFillTool : public USingleSelectionTool
 {
 	GENERATED_BODY()
 
@@ -226,15 +226,7 @@ public:
 	bool CanAccept() const override;
 	void Shutdown(EToolShutdownType ShutdownType) override;
 
-	// IClickBehaviorTarget
-	FInputRayHit IsHitByClick(const FInputDeviceRay& ClickPos) override;
-	void OnClicked(const FInputDeviceRay& ClickPos) override;
-
-	// IHoverBehaviorTarget
-	FInputRayHit BeginHoverSequenceHitTest(const FInputDeviceRay& PressPos) override;
-	void OnBeginHover(const FInputDeviceRay& DevicePos) override {}
-	bool OnUpdateHover(const FInputDeviceRay& DevicePos) override;
-	void OnEndHover() override;
+	void OnSelectionModified();
 
 	virtual void RequestAction(EHoleFillToolActions Action);
 
diff --git a/Engine/Plugins/Experimental/MeshModelingToolset/Source/MeshModelingTools/Public/MeshBoundaryToolBase.h b/Engine/Plugins/Experimental/MeshModelingToolset/Source/MeshModelingTools/Public/MeshBoundaryToolBase.h
--- a/Engine/Plugins/Experimental/MeshModelingToolset/Source/MeshModelingTools/Public/MeshBoundaryToolBase.h
+++ b/Engine/Plugins/Experimental/MeshModelingToolset/Source/MeshModelingTools/Public/MeshBoundaryToolBase.h
@@ -22,7 +22,7 @@ class USingleClickInputBehavior;
   *	TODO: We can refactor to make the HoleFiller tool inherit from this.
   */
 UCLASS()
-class MESHMODELINGTOOLS_API UMeshBoundaryToolBase : public USingleSelectionTool, public IClickBehaviorTarget, public IHoverBehaviorTarget
+class MESHMODELINGTOOLS_API UMeshBoundaryToolBase : public USingleSelectionTool
 {
 	GENERATED_BODY()
 
@@ -38,19 +38,6 @@ public:
 
 	virtual void Render(IToolsContextRenderAPI* RenderAPI) override;
 
-
-	// IClickBehaviorTarget
-	virtual FInputRayHit IsHitByClick(const FInputDeviceRay& ClickPos) override;
-	// This may need overriding, depending on the child class, though the child could also just
-	// override OnSelectionChanged().
-	virtual void OnClicked(const FInputDeviceRay& ClickPos) override;
-
-	// IHoverBehaviorTarget
-	virtual FInputRayHit BeginHoverSequenceHitTest(const FInputDeviceRay& PressPos) override;
-	virtual void OnBeginHover(const FInputDeviceRay& DevicePos) override {};
-	virtual bool OnUpdateHover(const FInputDeviceRay& DevicePos) override;
-	virtual void OnEndHover() override;
-
 protected:
 
 	UWorld* TargetWorld;
@@ -63,10 +50,6 @@ protected:
 	UPROPERTY()
 	UPolygonSelectionMechanic* SelectionMechanic = nullptr;
 
-	// Behavior that wraps this class's selection click operation. Useful to have a pointer to so that its
-	// priority can be worked with.
-	UPROPERTY()
-	USingleClickInputBehavior* LoopSelectClickBehavior = nullptr;
 
 	// A variant of group topology that considers all triangles one group, so that group edges are boundary
 	// edges in the mesh.
@@ -84,9 +67,6 @@ protected:
 	};
 	TUniquePtr<FBasicTopology> Topology;
 
-	// Override in the child to determine whether new selection should attempt to append
-	virtual bool ShouldSelectionAppend() const { return true; }
-
 	// Override in the child to respond to new loop selections.
 	virtual void OnSelectionChanged() {};
 };
diff --git a/Engine/Plugins/Experimental/MeshModelingToolset/Source/MeshModelingTools/Public/RevolveBoundaryTool.h b/Engine/Plugins/Experimental/MeshModelingToolset/Source/MeshModelingTools/Public/RevolveBoundaryTool.h
--- a/Engine/Plugins/Experimental/MeshModelingToolset/Source/MeshModelingTools/Public/RevolveBoundaryTool.h
+++ b/Engine/Plugins/Experimental/MeshModelingToolset/Source/MeshModelingTools/Public/RevolveBoundaryTool.h
@@ -3,6 +3,8 @@
 #pragma once
 
 #include "CoreMinimal.h"
+
+#include "BaseBehaviors/BehaviorTargetInterfaces.h"
 #include "InteractiveToolBuilder.h"
 #include "SingleSelectionTool.h"
 #include "Mechanics/ConstructionPlaneMechanic.h"
@@ -79,7 +81,7 @@ public:
  * revolving planar meshes. 
  */
 UCLASS()
-class MESHMODELINGTOOLS_API URevolveBoundaryTool : public UMeshBoundaryToolBase
+class MESHMODELINGTOOLS_API URevolveBoundaryTool : public UMeshBoundaryToolBase, public IClickBehaviorTarget
 {
 	GENERATED_BODY()
 
@@ -119,12 +121,13 @@ protected:
 	FVector3d RevolutionAxisOrigin;
 	FVector3d RevolutionAxisDirection;
 
-	virtual void OnClicked(const FInputDeviceRay& ClickPos) override;
-	virtual bool ShouldSelectionAppend() const override { return false; }
-
 	void GenerateAsset(const FDynamicMeshOpResult& Result);
 	void UpdateRevolutionAxis();
 	void StartPreview();
 
+	// IClickBehaviorTarget API
+	virtual FInputRayHit IsHitByClick(const FInputDeviceRay& ClickPos) override;
+	virtual void OnClicked(const FInputDeviceRay& ClickPos) override;
+
 	friend class URevolveBoundaryOperatorFactory;
 };
\ No newline at end of file
diff --git a/Engine/Plugins/Experimental/MeshModelingToolset/Source/MeshModelingToolsEditorOnly/Private/Hair/GroomCardsEditorTool.cpp b/Engine/Plugins/Experimental/MeshModelingToolset/Source/MeshModelingToolsEditorOnly/Private/Hair/GroomCardsEditorTool.cpp
--- a/Engine/Plugins/Experimental/MeshModelingToolset/Source/MeshModelingToolsEditorOnly/Private/Hair/GroomCardsEditorTool.cpp
+++ b/Engine/Plugins/Experimental/MeshModelingToolset/Source/MeshModelingToolsEditorOnly/Private/Hair/GroomCardsEditorTool.cpp
@@ -94,12 +94,15 @@ void UGroomCardsEditorTool::Setup()
 	UVMaterial = ToolSetupUtil::GetUVCheckerboardMaterial(50.0);
 	PreviewMesh->SetMaterial(MeshMaterial);
 
-	// set up SelectionMechanic
+	// Set up SelectionMechanic but remove its behaviors because we'd rather deal with clicks
+	// ourselves. Specifically, we want to be able to start new transactions on clicks, which
+	// we currently shouldn't do in an OnSelectionModifiedEvent call.
 	CardMeshSelectionMechanic = NewObject<UPolygonSelectionMechanic>(this);
 	CardMeshSelectionMechanic->bAddSelectionFilterPropertiesToParentTool = false;
 	CardMeshSelectionMechanic->Setup(this);
 	CardMeshSelectionMechanic->OnSelectionChanged.AddUObject(this, &UGroomCardsEditorTool::OnSelectionModifiedEvent);
 	CardMeshSelectionMechanic->PolyEdgesRenderer.LineThickness = 1.0;
+	CardMeshSelectionMechanic->DisableBehaviors(this);
 
 	// configure secondary render material
 	UMaterialInterface* SelectionMaterial = ToolSetupUtil::GetSelectionMaterial(FLinearColor(0.9f, 0.1f, 0.1f), GetToolManager());
@@ -864,8 +867,7 @@ void UGroomCardsEditorTool::RecomputeTopology()
 	CardMeshSelectionMechanic->Initialize(
 		PreviewMesh->GetMesh(), PreviewMesh->GetTransform(), TargetWorld,
 		Topology.Get(),
-		[this]() { return PreviewMesh->GetSpatial(); },
-		[this]() { return GetShiftToggle(); }
+		[this]() { return PreviewMesh->GetSpatial(); }
 	);
 }
 
diff --git a/Engine/Plugins/Experimental/MeshModelingToolset/Source/ModelingComponents/Private/Mechanics/LatticeControlPointsMechanic.cpp b/Engine/Plugins/Experimental/MeshModelingToolset/Source/ModelingComponents/Private/Mechanics/LatticeControlPointsMechanic.cpp
--- a/Engine/Plugins/Experimental/MeshModelingToolset/Source/ModelingComponents/Private/Mechanics/LatticeControlPointsMechanic.cpp
+++ b/Engine/Plugins/Experimental/MeshModelingToolset/Source/ModelingComponents/Private/Mechanics/LatticeControlPointsMechanic.cpp
@@ -25,7 +25,13 @@ static const FText LatticePointMovementTransactionText = LOCTEXT("LatticePointMo
 
 void ULatticeControlPointsMechanic::Setup(UInteractiveTool* ParentToolIn)
 {
-	URectangleMarqueeMechanic::Setup(ParentToolIn);
+	UInteractionMechanic::Setup(ParentToolIn);
+
+	MarqueeMechanic = NewObject<URectangleMarqueeMechanic>(this);
+	MarqueeMechanic->Setup(ParentToolIn);
+	MarqueeMechanic->OnDragRectangleStarted.AddUObject(this, &ULatticeControlPointsMechanic::OnDragRectangleStarted);
+	MarqueeMechanic->OnDragRectangleChanged.AddUObject(this, &ULatticeControlPointsMechanic::OnDragRectangleChanged);
+	MarqueeMechanic->OnDragRectangleFinished.AddUObject(this, &ULatticeControlPointsMechanic::OnDragRectangleFinished);
 
 	USingleClickInputBehavior* ClickBehavior = NewObject<USingleClickInputBehavior>();
 	ClickBehavior->Initialize(this);
@@ -177,8 +183,16 @@ void ULatticeControlPointsMechanic::Render(IToolsContextRenderAPI* RenderAPI)
 {
 	// Cache the camera state
 	GetParentTool()->GetToolManager()->GetContextQueriesAPI()->GetCurrentViewState(CachedCameraState);
+
+	MarqueeMechanic->Render(RenderAPI);
+}
+
+void ULatticeControlPointsMechanic::DrawHUD(FCanvas* Canvas, IToolsContextRenderAPI* RenderAPI)
+{
+	MarqueeMechanic->DrawHUD(Canvas, RenderAPI);
 }
 
+
 void ULatticeControlPointsMechanic::RebuildDrawables()
 {
 	DrawnControlPoints->Clear();
@@ -546,9 +560,7 @@ void ULatticeControlPointsMechanic::OnUpdateModifierState(int ModifierID, bool b
 	}
 }
 
-
-// =========================== URectangleMarqueeMechanic ===========================
-
+// These get bound to the delegates on the marquee mechanic.
 void ULatticeControlPointsMechanic::OnDragRectangleStarted()
 {
 	for (int32& PointID : SelectedPointIDs)
@@ -563,7 +575,6 @@ void ULatticeControlPointsMechanic::OnDragRectangleStarted()
 		PointTransformGizmo->SetVisibility(false);
 	}
 }
-
 void ULatticeControlPointsMechanic::OnDragRectangleChanged(const FCameraRectangle& Rectangle)
 {
 	CurrentDragSelection.Empty();
@@ -582,7 +593,6 @@ void ULatticeControlPointsMechanic::OnDragRectangleChanged(const FCameraRectangl
 		}
 	}
 }
-
 void ULatticeControlPointsMechanic::OnDragRectangleFinished()
 {
 	// Deselect previous SelectedPointIDs and replace it with "drag selection" points. Do this in one Undo transaction.
diff --git a/Engine/Plugins/Experimental/MeshModelingToolset/Source/ModelingComponents/Private/Mechanics/RectangleMarqueeMechanic.cpp b/Engine/Plugins/Experimental/MeshModelingToolset/Source/ModelingComponents/Private/Mechanics/RectangleMarqueeMechanic.cpp
--- a/Engine/Plugins/Experimental/MeshModelingToolset/Source/ModelingComponents/Private/Mechanics/RectangleMarqueeMechanic.cpp
+++ b/Engine/Plugins/Experimental/MeshModelingToolset/Source/ModelingComponents/Private/Mechanics/RectangleMarqueeMechanic.cpp
@@ -5,6 +5,7 @@
 #include "CanvasTypes.h"
 #include "CanvasItem.h"
 #include "ToolSceneQueriesUtil.h"
+#include "TransformTypes.h"
 
 #define LOCTEXT_NAMESPACE "URectangleMarqueeMechanic"
 
@@ -70,15 +71,83 @@ bool FCameraRectangle::IsProjectedPointInRectangle(const FVector& Point) const
 	return RectangleCorners.IsInside(Point2D);
 }
 
+bool FCameraRectangle::IsProjectedSegmentIntersectingRectangle(const FVector& Endpoint1, const FVector& Endpoint2) const
+{
+	FVector ProjectedEndpoint1 = bCameraIsOrthographic ? FVector::PointPlaneProject(Endpoint1, CameraPlane)
+		: FMath::RayPlaneIntersection(CameraOrigin, Endpoint1 - CameraOrigin, CameraPlane);
+	FVector ProjectedEndpoint2 = bCameraIsOrthographic ? FVector::PointPlaneProject(Endpoint2, CameraPlane)
+		: FMath::RayPlaneIntersection(CameraOrigin, Endpoint2 - CameraOrigin, CameraPlane);
+
+	FVector2D Endpoint1PlaneCoord = PlaneCoordinates(ProjectedEndpoint1, CameraPlane, UBasisVector, VBasisVector);
+	FVector2D Endpoint2PlaneCoord = PlaneCoordinates(ProjectedEndpoint2, CameraPlane, UBasisVector, VBasisVector);
+
+	// If either endpoint is inside, then definitely (at least partially) contained
+	if (RectangleCorners.IsInside(Endpoint1PlaneCoord) || RectangleCorners.IsInside(Endpoint2PlaneCoord))
+	{
+		return true;
+	}
+
+	// If both outside, have to do some intersections with the box sides. The function we have for this
+	// uses FVectors instead of FVector2Ds.
+	ProjectedEndpoint1 = FVector(Endpoint1PlaneCoord, 0);
+	ProjectedEndpoint2 = FVector(Endpoint2PlaneCoord, 0);
+	FVector Throwaway;
+	return FMath::SegmentIntersection2D(ProjectedEndpoint1, ProjectedEndpoint2, 
+		FVector(RectangleCorners.Min, 0), FVector(RectangleCorners.Max.X, RectangleCorners.Min.Y, 0), 
+		Throwaway)
+
+		|| FMath::SegmentIntersection2D(ProjectedEndpoint1, ProjectedEndpoint2, 
+			FVector(RectangleCorners.Max.X, RectangleCorners.Min.Y, 0), FVector(RectangleCorners.Max, 0), 
+			Throwaway)
+
+		|| FMath::SegmentIntersection2D(ProjectedEndpoint1, ProjectedEndpoint2,
+			FVector(RectangleCorners.Max, 0), FVector(RectangleCorners.Min.X, RectangleCorners.Max.Y, 0), 
+			Throwaway);
+
+	// Don't need to intersect with the fourth side because segment would have to intersect two sides
+	// of box if both endpoints are outside the box.
+}
+
 // ---------------------------------------
 
 void URectangleMarqueeMechanic::Setup(UInteractiveTool* ParentToolIn)
 {
 	UInteractionMechanic::Setup(ParentToolIn);
 
-	UClickDragInputBehavior* ClickDragBehavior = NewObject<UClickDragInputBehavior>(this);
+	ClickDragBehavior = NewObject<UClickDragInputBehavior>(this);
+	ClickDragBehavior->SetDefaultPriority(BasePriority);
 	ClickDragBehavior->Initialize(this);
-	ParentTool->AddInputBehavior(ClickDragBehavior);
+	ParentTool->AddInputBehavior(ClickDragBehavior, this);
+	SetIsEnabled(true);
+}
+
+bool URectangleMarqueeMechanic::IsEnabled()
+{
+	return bIsEnabled;
+}
+
+void URectangleMarqueeMechanic::SetIsEnabled(bool bOn)
+{
+	if (bIsDragging && !bOn)
+	{
+		OnTerminateDragSequence();
+	}
+
+	bIsEnabled = bOn;
+}
+
+void URectangleMarqueeMechanic::SetBasePriority(const FInputCapturePriority& Priority)
+{
+	BasePriority = Priority;
+	if (ClickDragBehavior)
+	{
+		ClickDragBehavior->SetDefaultPriority(Priority);
+	}
+}
+
+TPair<FInputCapturePriority, FInputCapturePriority> URectangleMarqueeMechanic::GetPriorityRange() const
+{
+	return TPair<FInputCapturePriority, FInputCapturePriority>(BasePriority, BasePriority);
 }
 
 void URectangleMarqueeMechanic::Render(IToolsContextRenderAPI* RenderAPI)
@@ -89,13 +158,9 @@ void URectangleMarqueeMechanic::Render(IToolsContextRenderAPI* RenderAPI)
 
 FInputRayHit URectangleMarqueeMechanic::CanBeginClickDragSequence(const FInputDeviceRay& PressPos)
 {
-	FInputRayHit Dummy;
-
-	// This is the boolean that is checked to see if a drag sequence can be started. In our case we want to begin the 
-	// drag sequence even if the first ray doesn't hit anything, so set this to true.
-	Dummy.bHit = true;
-
-	return Dummy;
+	return bIsEnabled ? 
+		FInputRayHit(TNumericLimits<float>::Max()) // bHit is true. Depth is max to lose the standard tiebreaker.
+		: FInputRayHit(); // bHit is false
 }
 
 void URectangleMarqueeMechanic::OnClickPress(const FInputDeviceRay& PressPos)
@@ -109,7 +174,7 @@ void URectangleMarqueeMechanic::OnClickPress(const FInputDeviceRay& PressPos)
 	DragStartScreenPosition = PressPos.ScreenPosition;
 	DragStartWorldRay = PressPos.WorldRay;
 
-	OnDragRectangleStarted();
+	OnDragRectangleStarted.Broadcast();
 }
 
 void URectangleMarqueeMechanic::OnClickDrag(const FInputDeviceRay& DragPos)
@@ -125,25 +190,27 @@ void URectangleMarqueeMechanic::OnClickDrag(const FInputDeviceRay& DragPos)
 
 	FCameraRectangle Rectangle(CachedCameraState, DragStartWorldRay, DragCurrentWorldRay);
 
-	OnDragRectangleChanged(Rectangle);
+	OnDragRectangleChanged.Broadcast(Rectangle);
 }
 
 void URectangleMarqueeMechanic::OnClickRelease(const FInputDeviceRay& ReleasePos)
 {
 	bIsDragging = false;
-	OnDragRectangleFinished();
+	OnDragRectangleFinished.Broadcast();
 }
 
 void URectangleMarqueeMechanic::OnTerminateDragSequence()
 {
 	bIsDragging = false;
-	OnDragRectangleFinished();
+	OnDragRectangleFinished.Broadcast();
 }
 
 
 void URectangleMarqueeMechanic::DrawHUD(FCanvas* Canvas, IToolsContextRenderAPI* RenderAPI)
 {
-	if (bIsDragging)
+	EViewInteractionState State = RenderAPI->GetViewInteractionState();
+	bool bThisViewHasFocus = !!(State & EViewInteractionState::Focused);
+	if (bThisViewHasFocus && bIsDragging)
 	{
 		FVector2D Start = DragStartScreenPosition;
 		FVector2D Curr = DragCurrentScreenPosition;
diff --git a/Engine/Plugins/Experimental/MeshModelingToolset/Source/ModelingComponents/Private/Selection/GroupTopologySelector.cpp b/Engine/Plugins/Experimental/MeshModelingToolset/Source/ModelingComponents/Private/Selection/GroupTopologySelector.cpp
--- a/Engine/Plugins/Experimental/MeshModelingToolset/Source/ModelingComponents/Private/Selection/GroupTopologySelector.cpp
+++ b/Engine/Plugins/Experimental/MeshModelingToolset/Source/ModelingComponents/Private/Selection/GroupTopologySelector.cpp
@@ -2,12 +2,14 @@
 
 
 #include "Selection/GroupTopologySelector.h"
+#include "Mechanics/RectangleMarqueeMechanic.h"
 #include "MeshQueries.h"
 #include "ToolDataVisualizer.h"
 #include "ToolSceneQueriesUtil.h"
 
 // Local utility function forward declarations
 bool IsOccluded(const FGeometrySet3::FNearest& ClosestElement, const FVector3d& ViewOrigin, const FDynamicMeshAABBTree3* Spatial);
+bool IsOccluded(const FVector3d& Point, const FVector3d& ViewOrigin, const FDynamicMeshAABBTree3* Spatial);
 void AddNewEdgeLoopEdgesFromCorner(const FGroupTopology& Topology, int32 EdgeID, int32 CornerID, TSet<int32>& EdgeSet);
 bool GetNextEdgeLoopEdge(const FGroupTopology& Topology, int32 IncomingEdgeID, int32 CornerID, int32& NextEdgeIDOut);
 void AddNewEdgeRingEdges(const FGroupTopology& Topology, int32 StartEdgeID, int32 ForwardGroupID, TSet<int32>& EdgeSet);
@@ -566,9 +568,14 @@ bool FGroupTopologySelector::DoEdgeBasedSelection(const FSelectionSettings& Sett
 }
 
 bool IsOccluded(const FGeometrySet3::FNearest& ClosestElement, const FVector3d& ViewOrigin, const FDynamicMeshAABBTree3* Spatial)
+{
+	return IsOccluded(ClosestElement.NearestGeoPoint, ViewOrigin, Spatial);
+}
+
+bool IsOccluded(const FVector3d& Point, const FVector3d& ViewOrigin, const FDynamicMeshAABBTree3* Spatial)
 {
 	// Shoot ray backwards to see if we hit something. 
-	FRay3d ToEyeRay(ClosestElement.NearestGeoPoint, (ViewOrigin - ClosestElement.NearestGeoPoint).Normalized(), true);
+	FRay3d ToEyeRay(Point, (ViewOrigin - Point).Normalized(), true);
 	ToEyeRay.Origin += (double)(100 * FMathf::ZeroTolerance) * ToEyeRay.Direction;
 	if (Spatial->FindNearestHitTriangle(ToEyeRay) >= 0)
 	{
@@ -577,6 +584,66 @@ bool IsOccluded(const FGeometrySet3::FNearest& ClosestElement, const FVector3d&
 	return false;
 }
 
+bool FGroupTopologySelector::FindSelectedElement(const FSelectionSettings& Settings, 
+	const FCameraRectangle& CameraRectangle, FTransform3d TargetTransform, FGroupTopologySelection& ResultOut)
+{
+	// One minor easy thing we can do to speed up the below is to detect cases where transform does not have a non-uniform
+	// scale, and in those cases we transform the rectangle once rather than transforming all the query points/curves. 
+	// This isn't worth doing for non-uniform scale transforms because the rectangle basis stops being orthonormal and 
+	// becomes a problem to work with.
+
+	ResultOut.Clear();
+	FDynamicMeshAABBTree3* Spatial = GetSpatial();
+
+	// Needed for occlusion test, which happens in local space.
+	FVector3d LocalCameraOrigin = TargetTransform.InverseTransformPosition(CameraRectangle.CameraOrigin);
+	
+	// Corner selection takes priority over edges.
+	if (Settings.bEnableCornerHits)
+	{
+		GeometrySet.ParallelFindAllPointsSatisfying(
+			[&CameraRectangle, &TargetTransform, &Settings, &LocalCameraOrigin, Spatial](const FVector3d& PointPosition)
+			{
+				return CameraRectangle.IsProjectedPointInRectangle((FVector)TargetTransform.TransformPosition(PointPosition))
+					&& (Settings.bIgnoreOcclusion || !IsOccluded(PointPosition, LocalCameraOrigin, Spatial));
+			},
+			ResultOut.SelectedCornerIDs);
+	}
+
+	// If we didn't get corners, look for edges.
+	if (ResultOut.SelectedCornerIDs.IsEmpty() && Settings.bEnableEdgeHits)
+	{
+		GeometrySet.ParallelFindAllCurvesSatisfying(
+			[&CameraRectangle, &TargetTransform, &Settings, &LocalCameraOrigin, Spatial](const FPolyline3d& Curve)
+			{
+				// Testing occlusion properly seems like it would be a pain, so we'll just consider something occluded
+				// if one of the endpoints is occluded. This will handle the common case of not wanting to select a
+				// hidden edge that is connected to a visible corner.
+				if (!Settings.bIgnoreOcclusion && (IsOccluded(Curve.Start(), LocalCameraOrigin, Spatial)
+						|| IsOccluded(Curve.End(), LocalCameraOrigin, Spatial)))
+				{
+					return false;
+				}
+
+				// Check whether any of the component segments intersect the rectangle
+				const TArray<FVector3d>& Verts = Curve.GetVertices();
+				FVector3d CurrentVert = TargetTransform.TransformPosition(Verts[0]);
+				for (int32 i = 1; i < Verts.Num(); ++i)
+				{
+					FVector3d NextVert = TargetTransform.TransformPosition(Verts[i]);
+					if (CameraRectangle.IsProjectedSegmentIntersectingRectangle((FVector)CurrentVert, (FVector)NextVert))
+					{
+						return true;
+					}
+					CurrentVert = NextVert;
+				}
+				return false;
+			},
+			ResultOut.SelectedEdgeIDs);
+	}
+	
+	return ResultOut.IsEmpty();
+}
 
 bool FGroupTopologySelector::ExpandSelectionByEdgeLoops(FGroupTopologySelection& Selection)
 {
diff --git a/Engine/Plugins/Experimental/MeshModelingToolset/Source/ModelingComponents/Private/Selection/PolygonSelectionMechanic.cpp b/Engine/Plugins/Experimental/MeshModelingToolset/Source/ModelingComponents/Private/Selection/PolygonSelectionMechanic.cpp
--- a/Engine/Plugins/Experimental/MeshModelingToolset/Source/ModelingComponents/Private/Selection/PolygonSelectionMechanic.cpp
+++ b/Engine/Plugins/Experimental/MeshModelingToolset/Source/ModelingComponents/Private/Selection/PolygonSelectionMechanic.cpp
@@ -1,6 +1,10 @@
 // Copyright Epic Games, Inc. All Rights Reserved.
 
 #include "Selection/PolygonSelectionMechanic.h"
+
+#include "BaseBehaviors/ClickDragBehavior.h"
+#include "BaseBehaviors/MouseHoverBehavior.h"
+#include "BaseBehaviors/SingleClickBehavior.h"
 #include "InteractiveToolManager.h"
 #include "Util/ColorConstants.h"
 #include "ToolSceneQueriesUtil.h"
@@ -17,11 +21,36 @@ void UPolygonSelectionMechanic::Setup(UInteractiveTool* ParentToolIn)
 {
 	UInteractionMechanic::Setup(ParentToolIn);
 
+	ClickBehavior = NewObject<USingleClickInputBehavior>();
+	ClickBehavior->Initialize(this);
+	ClickBehavior->SetDefaultPriority(BasePriority);
+	ParentToolIn->AddInputBehavior(ClickBehavior, this);
+
+	HoverBehavior = NewObject<UMouseHoverBehavior>();
+	HoverBehavior->Modifiers.RegisterModifier(ShiftModifierID, FInputDeviceState::IsShiftKeyDown);
+	HoverBehavior->Modifiers.RegisterModifier(CtrlModifierID, FInputDeviceState::IsCtrlKeyDown);
+	HoverBehavior->Initialize(this);
+	HoverBehavior->SetDefaultPriority(BasePriority);
+	ParentToolIn->AddInputBehavior(HoverBehavior, this);
+
+	MarqueeMechanic = NewObject<URectangleMarqueeMechanic>();
+	MarqueeMechanic->Setup(ParentToolIn);
+	MarqueeMechanic->OnDragRectangleStarted.AddUObject(this, &UPolygonSelectionMechanic::OnDragRectangleStarted);
+	MarqueeMechanic->OnDragRectangleChanged.AddUObject(this, &UPolygonSelectionMechanic::OnDragRectangleChanged);
+	MarqueeMechanic->OnDragRectangleFinished.AddUObject(this, &UPolygonSelectionMechanic::OnDragRectangleFinished);
+	MarqueeMechanic->SetBasePriority(BasePriority.MakeLower());
+
 	Properties = NewObject<UPolygonSelectionMechanicProperties>(this);
 	if (bAddSelectionFilterPropertiesToParentTool)
 	{
 		AddToolPropertySource(Properties);
 	}
+	Properties->WatchProperty(Properties->bSelectVertices, [this](bool bSelectVertices) { 
+		MarqueeMechanic->SetIsEnabled(Properties->bEnableMarquee && (Properties->bSelectVertices || Properties->bSelectEdges)); });
+	Properties->WatchProperty(Properties->bSelectEdges, [this](bool bSelectVertices) { 
+		MarqueeMechanic->SetIsEnabled(Properties->bEnableMarquee && (Properties->bSelectVertices || Properties->bSelectEdges)); });
+	Properties->WatchProperty(Properties->bEnableMarquee, [this](bool bEnableMarquee) { 
+		MarqueeMechanic->SetIsEnabled(Properties->bEnableMarquee && (Properties->bSelectVertices || Properties->bSelectEdges)); });
 
 	// set up visualizers
 	PolyEdgesRenderer.LineColor = FLinearColor::Red;
@@ -50,8 +79,7 @@ void UPolygonSelectionMechanic::Initialize(
 	FTransform TargetTransformIn,
 	UWorld* WorldIn,
 	const FGroupTopology* TopologyIn,
-	TFunction<FDynamicMeshAABBTree3 * ()> GetSpatialSourceFuncIn,
-	TFunction<bool(void)> GetAddToSelectionModifierStateFuncIn)
+	TFunction<FDynamicMeshAABBTree3 * ()> GetSpatialSourceFuncIn)
 {
 	this->Mesh = MeshIn;
 	this->Topology = TopologyIn;
@@ -82,8 +110,6 @@ void UPolygonSelectionMechanic::Initialize(
 		}
 	};
 
-	GetAddToSelectionModifierStateFunc = GetAddToSelectionModifierStateFuncIn;
-
 	// Set up the component we use to draw highlighted triangles. Only needs to be done once, not when the mesh
 	// changes (we are assuming that we won't swap worlds without creating a new mechanic).
 	if (PreviewGeometryActor == nullptr)
@@ -106,25 +132,57 @@ void UPolygonSelectionMechanic::Initialize(
 void UPolygonSelectionMechanic::Initialize(
 	USimpleDynamicMeshComponent* MeshComponentIn,
 	const FGroupTopology* TopologyIn,
-	TFunction<FDynamicMeshAABBTree3 * ()> GetSpatialSourceFuncIn,
-	TFunction<bool()> GetAddToSelectionModifierStateFuncIn)
+	TFunction<FDynamicMeshAABBTree3 * ()> GetSpatialSourceFuncIn)
 {
 
 	Initialize(MeshComponentIn->GetMesh(),
 		MeshComponentIn->GetComponentTransform(),
 		MeshComponentIn->GetWorld(),
 		TopologyIn,
-		GetSpatialSourceFuncIn,
-		GetAddToSelectionModifierStateFuncIn);
+		GetSpatialSourceFuncIn);
+}
+
+void UPolygonSelectionMechanic::DisableBehaviors(UInteractiveTool* ParentToolIn)
+{
+	ParentToolIn->RemoveInputBehaviorsBySource(this);
+	ParentToolIn->RemoveInputBehaviorsBySource(MarqueeMechanic);
+
+	// TODO: Is it worth adding a way to remove the property watchers for marquee?
+}
+
+void UPolygonSelectionMechanic::SetBasePriority(const FInputCapturePriority &Priority)
+{
+	BasePriority = Priority;
+	if (ClickBehavior)
+	{
+		ClickBehavior->SetDefaultPriority(Priority);
+	}
+	if (HoverBehavior)
+	{
+		HoverBehavior->SetDefaultPriority(Priority);
+	}
+	if (MarqueeMechanic)
+	{
+		MarqueeMechanic->SetBasePriority(Priority.MakeLower());
+	}
+}
+
+TPair<FInputCapturePriority, FInputCapturePriority> UPolygonSelectionMechanic::GetPriorityRange() const
+{
+	TPair<FInputCapturePriority, FInputCapturePriority> Result;
+	Result.Key = BasePriority;
+	Result.Value = MarqueeMechanic->GetPriorityRange().Value;
+	return Result;
 }
 
 void UPolygonSelectionMechanic::Render(IToolsContextRenderAPI* RenderAPI)
 {
+	MarqueeMechanic->Render(RenderAPI);
+
 	// Cache the view camera state so we can use for snapping/etc.
 	// This should not happen in Render() though...
 	GetParentTool()->GetToolManager()->GetContextQueriesAPI()->GetCurrentViewState(CameraState);
 
-
 	FViewCameraState RenderCameraState = RenderAPI->GetCameraState();
 
 	const FDynamicMesh3* TargetMesh = this->Mesh;
@@ -157,6 +215,10 @@ void UPolygonSelectionMechanic::Render(IToolsContextRenderAPI* RenderAPI)
 	HilightRenderer.EndFrame();
 }
 
+void UPolygonSelectionMechanic::DrawHUD(FCanvas* Canvas, IToolsContextRenderAPI* RenderAPI)
+{
+	MarqueeMechanic->DrawHUD(Canvas, RenderAPI);
+}
 
 
 
@@ -252,7 +314,7 @@ FGroupTopologySelector::FSelectionSettings UPolygonSelectionMechanic::GetTopoSel
 	Settings.bEnableEdgeHits = Properties->bSelectEdges;
 	Settings.bEnableCornerHits = Properties->bSelectVertices;
 
-	if (PersistentSelection.IsEmpty() == false && GetAddToSelectionModifierStateFunc() == true)
+	if (!PersistentSelection.IsEmpty() && (ShouldAddToSelectionFunc() || ShouldRemoveFromSelectionFunc()))
 	{
 		Settings.bEnableFaceHits = Settings.bEnableFaceHits && PersistentSelection.SelectedGroupIDs.Num() > 0;
 		Settings.bEnableEdgeHits = Settings.bEnableEdgeHits && PersistentSelection.SelectedEdgeIDs.Num() > 0;
@@ -294,6 +356,12 @@ bool UPolygonSelectionMechanic::UpdateHighlight(const FRay& WorldRay)
 		TopoSelector.ExpandSelectionByEdgeLoops(HilightSelection);
 	}
 
+	// Don't hover highlight a selection that we already selected, because people didn't like that
+	if (PersistentSelection.Contains(HilightSelection))
+	{
+		HilightSelection.Clear();
+	}
+
 	// Currently we draw highlighted edges/vertices differently from highlighted faces. Edges/vertices
 	// get drawn in the Render() call, so it is sufficient to just update HighlightSelection above.
 	// Faces, meanwhile, get placed into a Component that is rendered through the normal rendering system.
@@ -371,22 +439,21 @@ bool UPolygonSelectionMechanic::UpdateSelection(const FRay& WorldRay, FVector3d&
 			TopoSelector.ExpandSelectionByEdgeLoops(Selection);
 		}
 
-		if (GetAddToSelectionModifierStateFunc())
+		if (ShouldAddToSelectionFunc())
 		{
-			// We don't toggle because in cases where we are trying to add multiple elements,
-			// we want to remove the selection only if it was all selected to begin with,
-			// otherwise we want to add to it.
-			// At the moment the only way to add multiple elements at a time is through edge
-			// loop/ring selection, but we will someday have marquee selection and face ring selection.
-			if (PersistentSelection.Contains(Selection))
+			if (ShouldRemoveFromSelectionFunc())
 			{
-				PersistentSelection.Remove(Selection);
+				PersistentSelection.Toggle(Selection);
 			}
 			else
 			{
 				PersistentSelection.Append(Selection);
 			}
 		}
+		else if (ShouldRemoveFromSelectionFunc())
+		{
+			PersistentSelection.Remove(Selection);
+		}
 		else
 		{
 			PersistentSelection = Selection;
@@ -410,11 +477,14 @@ bool UPolygonSelectionMechanic::UpdateSelection(const FRay& WorldRay, FVector3d&
 }
 
 
-void UPolygonSelectionMechanic::SetSelection(const FGroupTopologySelection& Selection)
+void UPolygonSelectionMechanic::SetSelection(const FGroupTopologySelection& Selection, bool bBroadcast)
 {
 	PersistentSelection = Selection;
 	SelectionTimestamp++;
-	OnSelectionChanged.Broadcast();
+	if (bBroadcast)
+	{
+		OnSelectionChanged.Broadcast();
+	}
 }
 
 
@@ -425,6 +495,117 @@ void UPolygonSelectionMechanic::ClearSelection()
 	OnSelectionChanged.Broadcast();
 }
 
+FInputRayHit UPolygonSelectionMechanic::IsHitByClick(const FInputDeviceRay& ClickPos)
+{
+	FHitResult OutHit;
+	FGroupTopologySelection Selection;
+	if (TopologyHitTest(ClickPos.WorldRay, OutHit, Selection, CameraState.bIsOrthographic))
+	{
+		return FInputRayHit(OutHit.Distance);
+	}
+
+	// We didn't hit selectable geometry. If we're using the marquee mechanic, give it a
+	// chance to capture. Otherwise still capture ourselves so that we are able to clear selections.
+	return MarqueeMechanic->IsEnabled() ?
+		FInputRayHit() // bHit is false
+		: FInputRayHit(TNumericLimits<float>::Max()); // bHit is true
+}
+
+void UPolygonSelectionMechanic::OnClicked(const FInputDeviceRay& ClickPos)
+{
+	// update selection
+	ParentTool->GetToolManager()->BeginUndoTransaction(LOCTEXT("SelectionChange", "Selection"));
+	BeginChange();
+
+	// This will fire off a OnSelectionChanged delegate.
+	UpdateSelection(ClickPos.WorldRay, LastClickedHitPosition, LastClickedHitNormal);
+
+	EndChangeAndEmitIfModified();
+	ParentTool->GetToolManager()->EndUndoTransaction();
+}
+
+FInputRayHit UPolygonSelectionMechanic::BeginHoverSequenceHitTest(const FInputDeviceRay& PressPos)
+{
+	FHitResult OutHit;
+	if (TopologyHitTest(PressPos.WorldRay, OutHit))
+	{
+		return FInputRayHit(OutHit.Distance);
+	}
+	return FInputRayHit(); // bHit is false
+}
+
+void UPolygonSelectionMechanic::OnBeginHover(const FInputDeviceRay& DevicePos)
+{
+}
+
+bool UPolygonSelectionMechanic::OnUpdateHover(const FInputDeviceRay& DevicePos)
+{
+	UpdateHighlight(DevicePos.WorldRay);
+	return true;
+}
+
+void UPolygonSelectionMechanic::OnEndHover()
+{
+	ClearHighlight();
+}
+
+void UPolygonSelectionMechanic::OnUpdateModifierState(int ModifierID, bool bIsOn)
+{
+	switch (ModifierID)
+	{
+	case ShiftModifierID:
+		bShiftToggle = bIsOn;
+		break;
+	case CtrlModifierID:
+		bCtrlToggle = bIsOn;
+		break;
+	}
+}
+
+void UPolygonSelectionMechanic::OnDragRectangleStarted()
+{
+	PreDragPersistentSelection = PersistentSelection;
+}
+
+void UPolygonSelectionMechanic::OnDragRectangleChanged(const FCameraRectangle& CurrentRectangle)
+{
+	FGroupTopologySelection RectangleSelection;
+
+	FGroupTopologySelector::FSelectionSettings TopoSelectorSettings = GetTopoSelectorSettings(false);
+	TopoSelectorSettings.bIgnoreOcclusion = Properties->bMarqueeIgnoreOcclusion; // uses a separate setting for marquee
+	
+	TopoSelector.FindSelectedElement(TopoSelectorSettings, CurrentRectangle, TargetTransform,
+		RectangleSelection);
+
+	if (ShouldAddToSelectionFunc())
+	{
+		PersistentSelection = PreDragPersistentSelection;
+		if (ShouldRemoveFromSelectionFunc())
+		{
+			PersistentSelection.Toggle(RectangleSelection);
+		}
+		else
+		{
+			PersistentSelection.Append(RectangleSelection);
+		}
+	}
+	else if (ShouldRemoveFromSelectionFunc())
+	{
+		PersistentSelection = PreDragPersistentSelection;
+		PersistentSelection.Remove(RectangleSelection);
+	}
+	else
+	{
+		// Neither key pressed.
+		PersistentSelection = RectangleSelection;
+	}
+}
+
+void UPolygonSelectionMechanic::OnDragRectangleFinished()
+{
+	SelectionTimestamp++;
+	OnSelectionChanged.Broadcast();
+}
 
 
 void UPolygonSelectionMechanic::BeginChange()
@@ -461,6 +642,11 @@ bool UPolygonSelectionMechanic::EndChangeAndEmitIfModified()
 	return false;
 }
 
+void UPolygonSelectionMechanic::GetClickedHitPosition(FVector3d& PositionOut, FVector3d& NormalOut) const
+{
+	PositionOut = LastClickedHitPosition;
+	NormalOut = LastClickedHitNormal;
+}
 
 FFrame3d UPolygonSelectionMechanic::GetSelectionFrame(bool bWorld, FFrame3d* InitialLocalFrame) const
 {
diff --git a/Engine/Plugins/Experimental/MeshModelingToolset/Source/ModelingComponents/Public/Mechanics/LatticeControlPointsMechanic.h b/Engine/Plugins/Experimental/MeshModelingToolset/Source/ModelingComponents/Public/Mechanics/LatticeControlPointsMechanic.h
--- a/Engine/Plugins/Experimental/MeshModelingToolset/Source/ModelingComponents/Public/Mechanics/LatticeControlPointsMechanic.h
+++ b/Engine/Plugins/Experimental/MeshModelingToolset/Source/ModelingComponents/Public/Mechanics/LatticeControlPointsMechanic.h
@@ -28,7 +28,7 @@ class UTransformProxy;
 
 UCLASS()
 class MODELINGCOMPONENTS_API ULatticeControlPointsMechanic : 
-	public URectangleMarqueeMechanic, public IClickBehaviorTarget, public IHoverBehaviorTarget
+	public UInteractionMechanic, public IClickBehaviorTarget, public IHoverBehaviorTarget
 {
 	GENERATED_BODY()
 
@@ -52,6 +52,8 @@ public:
 
 	void UpdateSetPivotMode(bool bInSetPivotMode);
 
+	virtual void DrawHUD(FCanvas* Canvas, IToolsContextRenderAPI* RenderAPI);
+
 	// UInteractionMechanic
 	virtual void Setup(UInteractiveTool* ParentTool) override;
 	virtual void Shutdown() override;
@@ -68,11 +70,6 @@ public:
 	virtual void OnEndHover() override;
 	virtual void OnUpdateModifierState(int ModifierID, bool bIsOn) override;
 
-	// URectangleMarqueeMechanic implementation
-	void OnDragRectangleStarted() final;
-	void OnDragRectangleChanged(const FCameraRectangle& Rectangle) final;
-	void OnDragRectangleFinished() final;
-
 	bool bHasChanged = false;
 
 protected:
@@ -83,6 +80,7 @@ protected:
 
 	// Used for spatial queries
 	FGeometrySet3 GeometrySet;
+	FViewCameraState CachedCameraState;
 
 	/** Used for displaying points/segments */
 	UPROPERTY()
@@ -133,7 +131,11 @@ protected:
 
 	// Support for selection
 	TArray<int32> SelectedPointIDs;
+	URectangleMarqueeMechanic* MarqueeMechanic;
 	TArray<int32> CurrentDragSelection;
+	void OnDragRectangleStarted();
+	void OnDragRectangleChanged(const FCameraRectangle& Rectangle);
+	void OnDragRectangleFinished();
 
 	// We need the selected point start positions so we can move multiple points appropriately.
 	TArray<FVector3d> SelectedPointStartPositions;
diff --git a/Engine/Plugins/Experimental/MeshModelingToolset/Source/ModelingComponents/Public/Mechanics/RectangleMarqueeMechanic.h b/Engine/Plugins/Experimental/MeshModelingToolset/Source/ModelingComponents/Public/Mechanics/RectangleMarqueeMechanic.h
--- a/Engine/Plugins/Experimental/MeshModelingToolset/Source/ModelingComponents/Public/Mechanics/RectangleMarqueeMechanic.h
+++ b/Engine/Plugins/Experimental/MeshModelingToolset/Source/ModelingComponents/Public/Mechanics/RectangleMarqueeMechanic.h
@@ -3,12 +3,15 @@
 #pragma once
 
 #include "CoreMinimal.h"
+#include "InputBehavior.h"
 #include "InteractionMechanic.h"
 #include "BaseBehaviors/BehaviorTargetInterfaces.h"
+#include "TransformTypes.h"
 #include "VectorTypes.h"
 
 #include "RectangleMarqueeMechanic.generated.h"
 
+class UClickDragInputBehavior;
 
 /// Struct containing:
 ///		- camera information, 
@@ -32,7 +35,13 @@ struct FCameraRectangle
 	// Project the given 3D point to the camera plane and test if it's in the rectangle
 	bool IsProjectedPointInRectangle(const FVector& Point) const;
 
-	// TODO: Rectangle vs projected triangle, edge, etc.
+	// Project the given segment to the camera plane and test if it intersects the rectangle
+	bool IsProjectedSegmentIntersectingRectangle(const FVector& Endpoint1, const FVector& Endpoint2) const;
+
+	// TODO: Add a way to test rectangle against triangles.
+
+private:
+	FCameraRectangle() {};
 };
 
 
@@ -40,12 +49,10 @@ struct FCameraRectangle
  * Mechanic for a rectangle "marquee" selection. It creates and maintains the 2D rectangle associated with a mouse drag. 
  * It does not test against any scene geometry, nor does it maintain any sort of list of selected objects.
  *
- * To use this class:
- *	- create a subclass of this class
- *  - override the pure virtual functions. These will be called from this class in response to mouse events
- *  - use the FCameraRectangle passed into RectangleChanged to test against your scene geometry
+ * When using this mechanic, you should call Render() on it in the tool's Render() call so that it can cache
+ * necessary camera state, and DrawHUD() in the tool's DrawHUD() call so that it can draw the box.
  *
- * See ULatticeControlPointsMechanic for an example.
+ * Attach to the mechanic's delegates and use the passed rectangle to test against your geometry. 
  */
 
 UCLASS()
@@ -59,38 +66,65 @@ public:
 	void Setup(UInteractiveTool* ParentTool) override;
 	void Render(IToolsContextRenderAPI* RenderAPI) override;
 
-	// REVIEW: Should this be added to UInteractionMechanic?
 	virtual void DrawHUD(FCanvas* Canvas, IToolsContextRenderAPI* RenderAPI);
 
+	bool IsEnabled();
+	void SetIsEnabled(bool bOn);
+
+	/**
+	 * Sets the base priority so that users can make sure that their own behaviors are higher
+	 * priority. The mechanic will not use any priority value higher than this.
+	 * Mechanics could use lower priorities (and their range could be inspected with
+	 * GetPriorityRange), but marquee mechanic doesn't.
+	 *
+	 * Can be called before or after Setup().
+	 */
+	void SetBasePriority(const FInputCapturePriority& Priority);
+
+	/**
+	 * Gets the current priority range used by behaviors in the mechanic, higher
+	 * priority to lower.
+	 *
+	 * For marquee mechanic, the range will be [BasePriority, BasePriority] since
+	 * it only uses one priority.
+	 */
+	TPair<FInputCapturePriority, FInputCapturePriority> GetPriorityRange() const;
+
+	/**
+	 * Called when user starts dragging a new rectangle.
+	 */
+	FSimpleMulticastDelegate OnDragRectangleStarted;
+
+	/**
+	 * Called as the user drags the other corner of the rectangle around.
+	 */
+	DECLARE_MULTICAST_DELEGATE_OneParam(OnDragRectangleChangedEvent, const FCameraRectangle&);
+	OnDragRectangleChangedEvent OnDragRectangleChanged;
+
+	/**
+	 * Called once the user lets go of the mouse button after dragging out a rectangle.
+	 */
+	FSimpleMulticastDelegate OnDragRectangleFinished;
+	
 protected:
+	UPROPERTY()
+	UClickDragInputBehavior* ClickDragBehavior = nullptr;
 
-	bool bIsDragging;
 	FViewCameraState CachedCameraState;
-	FVector2D DragStartScreenPosition;
-	FRay DragStartWorldRay;
-	FVector2D DragCurrentScreenPosition;
-
-private:
-
-	//
-	// Override these
-	//
-
-	virtual void OnDragRectangleStarted() 
-	PURE_VIRTUAL(URectangleMarqueeMechanic::OnDragRectangleStarted, );
-
-	virtual void OnDragRectangleChanged(const FCameraRectangle& CurrentRectangle) 
-	PURE_VIRTUAL(URectangleMarqueeMechanic::OnDragRectangleChanged, );
-
-	virtual void OnDragRectangleFinished() 
-	PURE_VIRTUAL(URectangleMarqueeMechanic::OnDragRectangleFinished, );
 
+	FInputCapturePriority BasePriority = FInputCapturePriority(FInputCapturePriority::DEFAULT_TOOL_PRIORITY);
 
 private:
-	
+
 	FInputRayHit CanBeginClickDragSequence(const FInputDeviceRay& PressPos) final;
 	void OnClickPress(const FInputDeviceRay& PressPos) final;
 	void OnClickDrag(const FInputDeviceRay& DragPos) final;
 	void OnClickRelease(const FInputDeviceRay& ReleasePos) final;
 	void OnTerminateDragSequence() final;
+
+	bool bIsEnabled;
+	bool bIsDragging;
+	FVector2D DragStartScreenPosition;
+	FRay DragStartWorldRay;
+	FVector2D DragCurrentScreenPosition;
 };
diff --git a/Engine/Plugins/Experimental/MeshModelingToolset/Source/ModelingComponents/Public/Selection/GroupTopologySelector.h b/Engine/Plugins/Experimental/MeshModelingToolset/Source/ModelingComponents/Public/Selection/GroupTopologySelector.h
--- a/Engine/Plugins/Experimental/MeshModelingToolset/Source/ModelingComponents/Public/Selection/GroupTopologySelector.h
+++ b/Engine/Plugins/Experimental/MeshModelingToolset/Source/ModelingComponents/Public/Selection/GroupTopologySelector.h
@@ -7,6 +7,7 @@
 #include "DynamicMeshAABBTree3.h"
 #include "RayTypes.h"
 
+struct FCameraRectangle;
 class FToolDataVisualizer;
 struct FViewCameraState;
 
@@ -110,6 +111,23 @@ public:
 	bool FindSelectedElement(const FSelectionSettings& Settings, const FRay3d& Ray, FGroupTopologySelection& ResultOut,
 		FVector3d& SelectedPositionOut, FVector3d& SelectedNormalOut, int32* EdgeSegmentIdOut = nullptr);
 
+	/**
+	 * Given a camera rectangle (from a marquee selection) gives the group corners or edges that are contained in
+	 * the rectangle. Currently does not select faces.
+	 * 
+	 * Uses the bEnableEdgeHits, bEnableCornerHits, and bIgnoreOcclusion members of the FSelectionSettings argument.
+	 * For group edges, occlusion is checked at endpoints, and if either is occluded, the edge is considered occluded.
+	 *
+	 * @param Settings Settings that determine what is selected and whether occlusion is tested.
+	 * @param CameraRectangle The camera rectangle.
+	 * @param TargetTransform Transform from points to world space, since the rectangle is in world space.
+	 * @param ResultOut Output selection (cleared before use).
+	 *
+	 * @return true if something was selected.
+	 */
+	bool FindSelectedElement(const FSelectionSettings& Settings, const FCameraRectangle& CameraRectangle, 
+		FTransform3d TargetTransform, FGroupTopologySelection& ResultOut);
+
 	/**
 	 * Using the edges in the given selection as starting points, add any "edge loops" containing the edges. An 
 	 * edge loop is a sequence of edges that passes through valence-4 corners through the opposite edge, and may
diff --git a/Engine/Plugins/Experimental/MeshModelingToolset/Source/ModelingComponents/Public/Selection/PolygonSelectionMechanic.h b/Engine/Plugins/Experimental/MeshModelingToolset/Source/ModelingComponents/Public/Selection/PolygonSelectionMechanic.h
--- a/Engine/Plugins/Experimental/MeshModelingToolset/Source/ModelingComponents/Public/Selection/PolygonSelectionMechanic.h
+++ b/Engine/Plugins/Experimental/MeshModelingToolset/Source/ModelingComponents/Public/Selection/PolygonSelectionMechanic.h
@@ -5,7 +5,9 @@
 #include "CoreMinimal.h"
 #include "Drawing/PreviewGeometryActor.h"
 #include "Drawing/TriangleSetComponent.h"
+#include "InputBehavior.h"
 #include "InteractiveTool.h"
+#include "Mechanics/RectangleMarqueeMechanic.h"
 #include "SimpleDynamicMeshComponent.h"
 #include "Selection/GroupTopologySelector.h"
 #include "TransformTypes.h"
@@ -14,6 +16,9 @@
 #include "PolygonSelectionMechanic.generated.h"
 
 class FPolygonSelectionMechanicSelectionChange;
+class UMouseHoverBehavior;
+class URectangleMarqueeMechanic;
+class USingleClickInputBehavior;
 
 UCLASS()
 class MODELINGCOMPONENTS_API UPolygonSelectionMechanicProperties : public UInteractiveToolPropertySet
@@ -38,6 +43,12 @@ public:
 	UPROPERTY(EditAnywhere, Category = SelectionFilter, meta = (EditCondition = "bSelectEdges"))
 	bool bSelectEdgeRings = false;
 
+	UPROPERTY(EditAnywhere, Category = SelectionFilter, AdvancedDisplay)
+	bool bEnableMarquee = true;
+
+	/** Determines whether vertices should be checked for occlusion in marquee select (Note: marquee select currently only works with edges and vertices) */
+	UPROPERTY(EditAnywhere, Category = SelectionFilter, meta = (EditCondition = "bEnableMarquee", EditConditionHides))
+	bool bMarqueeIgnoreOcclusion = true;
 
 	// The following were originally in their own category, all marked as AdvancedDisplay. However, since there wasn't a non-AdvancedDisplay
 	// property in the category, they started out as expanded and could not be collapsed.
@@ -63,7 +74,7 @@ public:
  * from a FGroupTopology on a USimpleDynamicMeshComponent. 
  */
 UCLASS()
-class MODELINGCOMPONENTS_API UPolygonSelectionMechanic : public UInteractionMechanic
+class MODELINGCOMPONENTS_API UPolygonSelectionMechanic : public UInteractionMechanic, public IClickBehaviorTarget, public IHoverBehaviorTarget
 {
 	GENERATED_BODY()
 public:
@@ -77,6 +88,7 @@ public:
 	virtual void Shutdown() override;
 
 	virtual void Render(IToolsContextRenderAPI* RenderAPI) override;
+	virtual void DrawHUD(FCanvas* Canvas, IToolsContextRenderAPI* RenderAPI);
 
 	/**
 	 * Initializes the mechanic.
@@ -86,22 +98,44 @@ public:
 	 * @param World World in which we are operating, used to add drawing components that draw highlighted edges.
 	 * @param Topology Group topology of the mesh.
 	 * @param GetSpatialSourceFunc Function that returns an AABB tree for the mesh.
-	 * @param GetAddToSelectionModifierStateFunc Functions that returns whether new selection should be trying to append to an existing 
-	     selection, usually by checking whether a particular modifier key is currently pressed.
 	 */
 	void Initialize(const FDynamicMesh3* Mesh,
 		FTransform TargetTransform,
 		UWorld* World,
 		const FGroupTopology* Topology,
-		TFunction<FDynamicMeshAABBTree3*()> GetSpatialSourceFunc,
-		TFunction<bool(void)> GetAddToSelectionModifierStateFunc = []() {return false; }
+		TFunction<FDynamicMeshAABBTree3*()> GetSpatialSourceFunc
 		);
 
 	void Initialize(USimpleDynamicMeshComponent* MeshComponent, const FGroupTopology* Topology,
-		TFunction<FDynamicMeshAABBTree3 * ()> GetSpatialSourceFunc,
-		TFunction<bool(void)> GetAddToSelectionModifierStateFunc = []() {return false; }
+		TFunction<FDynamicMeshAABBTree3 * ()> GetSpatialSourceFunc
 	);
 
+	/**
+	 * Removes the mechanic's own click/hover handlers, which means that the parent tool
+	 * will need to call UpdateSelection(), UpdateHighlight(), ClearHighlight(), and 
+	 * ClearSelection() from its own hover/click handlers.
+	 *
+	 * Must be called after Setup() to have an effect.
+	 *
+	 * @param ParentToolIn The parent tool, needed to be able to remove the behaviors.
+	 */
+	void DisableBehaviors(UInteractiveTool* ParentToolIn);
+
+	/**
+	 * Sets the base priority so that tools can make sure that their own behaviors are higher
+	 * priority. The mechanic will not use any priority value higher than this, but it may use
+	 * lower if it needs to stagger the priorities of behaviors it uses.
+	 * Can be called before or after Setup().
+	 */
+	void SetBasePriority(const FInputCapturePriority& Priority);
+
+	/**
+	 * Gets the current priority range used by behaviors in the mechanic. The returned pair will
+	 * have the base (highest) priority as the key, and the lowest priority as the value.
+	 */
+	TPair<FInputCapturePriority, FInputCapturePriority> GetPriorityRange() const;
+
+
 	void SetShouldSelectEdgeLoopsFunc(TFunction<bool(void)> Func)
 	{
 		ShouldSelectEdgeLoopsFunc = Func;
@@ -112,6 +146,24 @@ public:
 		ShouldSelectEdgeRingsFunc = Func;
 	}
 
+	/**
+	 * By default, the shift key will cause new clicks to add to the selection. However, this
+	 * can be changed by supplying a different function to check here.
+	 */
+	void SetShouldAddToSelectionFunc(TFunction<bool(void)> Func)
+	{
+		ShouldAddToSelectionFunc = Func;
+	}
+
+	/**
+	 * By default, the Ctrl key will cause new clicks to remove from the existing selection.
+	 * However, this can be changed by supplying a different function to check here.
+	 */
+	void SetShouldRemoveFromSelectionFunc(TFunction<bool(void)> Func)
+	{
+		ShouldRemoveFromSelectionFunc = Func;
+	}
+
 	/**
 	 * Notify internal data structures that the associated MeshComponent has been modified.
 	 * @param bTopologyModified if true, the underlying mesh topology has been changed. This clears the current selection.
@@ -160,11 +212,13 @@ public:
 	/**
 	 * Replace the current selection with an external selection. 
 	 * @warning does not check that the selection is valid!
+	 *
+	 * @param bBroadcast If true, issues an OnSelectionChanged delegate broadcast.
 	 */
-	void SetSelection(const FGroupTopologySelection& Selection);
+	void SetSelection(const FGroupTopologySelection& Selection, bool bBroadcast = true);
 
 	/**
-	 * Clear the current selection
+	 * Clear the current selection.
 	 */
 	void ClearSelection();
 
@@ -178,13 +232,18 @@ public:
 	 */
 	const FGroupTopologySelection& GetActiveSelection() const { return PersistentSelection; }
 
+	/**
+	 * Can be used by in an OnSelectionChanged event to inspect the clicked location (i.e., the
+	 * values returned by the UpdateSelection() function when the click happened).
+	 */
+	void GetClickedHitPosition(FVector3d& PositionOut, FVector3d& NormalOut) const;
+
 	/**
 	 * @return The best-guess 3D frame for the current select
 	 * @param bWorld if true, local-to-world transform of the target MeshComponent is applied to the frame
 	 */
 	FFrame3d GetSelectionFrame(bool bWorld, FFrame3d* InitialLocalFrame = nullptr) const;
 
-
 	//
 	// Change Tracking
 	//
@@ -204,10 +263,33 @@ public:
 	 */
 	bool EndChangeAndEmitIfModified();
 
-	/** OnSelectionChanged is broadcast whenever the selection is modified (including by FChanges) */
-	FSimpleMulticastDelegate OnSelectionChanged;
+	// IClickBehaviorTarget implementation
+	virtual FInputRayHit IsHitByClick(const FInputDeviceRay& ClickPos) override;
+	virtual void OnClicked(const FInputDeviceRay& ClickPos) override;
+
+	// IHoverBehaviorTarget implementation
+	virtual FInputRayHit BeginHoverSequenceHitTest(const FInputDeviceRay& PressPos) override;
+	virtual void OnBeginHover(const FInputDeviceRay& DevicePos) override;
+	virtual bool OnUpdateHover(const FInputDeviceRay& DevicePos) override;
+	virtual void OnEndHover() override;
+	virtual void OnUpdateModifierState(int ModifierID, bool bIsOn) override;
+
+protected:
+	// These get bound to marquee mechanic delegates.
+	virtual void OnDragRectangleStarted();
+	virtual void OnDragRectangleChanged(const FCameraRectangle& CurrentRectangle);
+	virtual void OnDragRectangleFinished();
 
 public:
+	/** 
+	 * OnSelectionChanged is broadcast whenever the selection is modified (including by FChanges, which
+	 * means that called functions should not issue undo transactions. 
+	 */
+	FSimpleMulticastDelegate OnSelectionChanged;
+
+	// TODO: Is it worth issuing separate callbacks in normal selection changes and in FChange ones, to
+	// allow the user to bundle in some FChanges into the normal callback?
+
 	UPROPERTY()
 	UPolygonSelectionMechanicProperties* Properties;
 
@@ -216,7 +298,16 @@ protected:
 	const FGroupTopology* Topology;
 	TFunction<FDynamicMeshAABBTree3*()> GetSpatialFunc;
 
-	TFunction<bool(void)> GetAddToSelectionModifierStateFunc;
+	UPROPERTY()
+	UMouseHoverBehavior* HoverBehavior;
+
+	UPROPERTY()
+	USingleClickInputBehavior* ClickBehavior;
+
+	UPROPERTY()
+	URectangleMarqueeMechanic* MarqueeMechanic;
+
+	FInputCapturePriority BasePriority = FInputCapturePriority(FInputCapturePriority::DEFAULT_TOOL_PRIORITY);
 
 	// When bSelectEdgeLoops is true, this function is tested to see if we should select edge loops,
 	// to allow edge loop selection to be toggled with some key (setting bSelectEdgeLoops to
@@ -228,6 +319,9 @@ protected:
 	// false overrides this function).
 	TFunction<bool(void)> ShouldSelectEdgeRingsFunc = []() {return true; };
 
+	TFunction<bool(void)> ShouldAddToSelectionFunc = [this]() {return bShiftToggle; };
+	TFunction<bool(void)> ShouldRemoveFromSelectionFunc = [this]() {return bCtrlToggle; };;
+
 	FTransform3d TargetTransform;
 
 	FGroupTopologySelector TopoSelector;
@@ -245,6 +339,10 @@ protected:
 	int32 SelectionTimestamp = 0;
 	TUniquePtr<FPolygonSelectionMechanicSelectionChange> ActiveChange;
 
+	FGroupTopologySelection PreDragPersistentSelection; // Used for box selection
+	FVector3d LastClickedHitPosition;
+	FVector3d LastClickedHitNormal;
+
 	/** The actor we create internally to own the DrawnTriangleSetComponent */
 	UPROPERTY()
 	APreviewGeometryActor* PreviewGeometryActor;
@@ -258,6 +356,11 @@ protected:
 	UMaterialInterface* HighlightedFaceMaterial;
 
 	FViewCameraState CameraState;
+
+	bool bShiftToggle = false;
+	bool bCtrlToggle = false;
+	static const int32 ShiftModifierID = 1;
+	static const int32 CtrlModifierID = 2;
 public:
 	FToolDataVisualizer PolyEdgesRenderer;
 	FToolDataVisualizer HilightRenderer;
diff --git a/Engine/Source/Runtime/Experimental/InteractiveToolsFramework/Private/InteractiveTool.cpp b/Engine/Source/Runtime/Experimental/InteractiveToolsFramework/Private/InteractiveTool.cpp
--- a/Engine/Source/Runtime/Experimental/InteractiveToolsFramework/Private/InteractiveTool.cpp
+++ b/Engine/Source/Runtime/Experimental/InteractiveToolsFramework/Private/InteractiveTool.cpp
@@ -53,9 +53,14 @@ void UInteractiveTool::DrawHUD(FCanvas* Canvas, IToolsContextRenderAPI* RenderAP
 {
 }
 
-void UInteractiveTool::AddInputBehavior(UInputBehavior* Behavior)
+void UInteractiveTool::AddInputBehavior(UInputBehavior* Behavior, void* Source)
 {
-	InputBehaviors->Add(Behavior);
+	InputBehaviors->Add(Behavior, Source);
+}
+
+void UInteractiveTool::RemoveInputBehaviorsBySource(void* Source)
+{
+	InputBehaviors->RemoveBySource(Source);
 }
 
 const UInputBehaviorSet* UInteractiveTool::GetInputBehaviors() const
diff --git a/Engine/Source/Runtime/Experimental/InteractiveToolsFramework/Public/InteractiveTool.h b/Engine/Source/Runtime/Experimental/InteractiveToolsFramework/Public/InteractiveTool.h
--- a/Engine/Source/Runtime/Experimental/InteractiveToolsFramework/Public/InteractiveTool.h
+++ b/Engine/Source/Runtime/Experimental/InteractiveToolsFramework/Public/InteractiveTool.h
@@ -322,8 +322,15 @@ public:
 	/**
 	 * Add an input behavior for this Tool
 	 * @param Behavior behavior to add
+	 * @param Source Optional pointer that could be used to identify the behavior for removal later.
 	 */
-	virtual void AddInputBehavior(UInputBehavior* Behavior);
+	virtual void AddInputBehavior(UInputBehavior* Behavior, void* Source = nullptr);
+
+	/**
+	 * Remove all input behaviors that had the given source pointer set during their addition.
+	 * @param Source Identifying pointer
+	 */
+	virtual void RemoveInputBehaviorsBySource(void* Source);
 
 	/**
 	 * @return Current input behavior set.
