Replace DynamicMeshComponent and PreviewMeshActor with a UPreviewMesh.

#rnx

diff --git a/Engine/Plugins/Experimental/MeshModelingToolset/Source/MeshModelingToolsEditorOnly/Private/SubdividePolyTool.cpp b/Engine/Plugins/Experimental/MeshModelingToolset/Source/MeshModelingToolsEditorOnly/Private/SubdividePolyTool.cpp
--- a/Engine/Plugins/Experimental/MeshModelingToolset/Source/MeshModelingToolsEditorOnly/Private/SubdividePolyTool.cpp
+++ b/Engine/Plugins/Experimental/MeshModelingToolset/Source/MeshModelingToolsEditorOnly/Private/SubdividePolyTool.cpp
@@ -152,48 +152,45 @@ void USubdividePolyTool::Setup()
 	AddToolPropertySource(Properties);
 	SetToolPropertySourceEnabled(Properties, true);
 
+	PreviewMesh = NewObject<UPreviewMesh>(this);
+	if (PreviewMesh == nullptr)
+	{
+		return;
+	}
+	PreviewMesh->CreateInWorld(TargetWorld, FTransform::Identity);
 
-	// create dynamic mesh component to use for live preview
-	FActorSpawnParameters SpawnInfo;
-	TemporaryParentActor = TargetWorld->SpawnActor<APreviewMeshActor>(FVector::ZeroVector, FRotator::ZeroRotator, SpawnInfo);
+	PreviewMesh->SetTransform(ComponentTarget->GetWorldTransform());
+	PreviewMesh->UpdatePreview(OriginalMesh.Get());
 
-	PreviewDynamicMeshComponent = NewObject<USimpleDynamicMeshComponent>(TemporaryParentActor, "DynamicMesh");
+	USimpleDynamicMeshComponent* PreviewDynamicMeshComponent = (USimpleDynamicMeshComponent*)PreviewMesh->GetRootComponent();
 	if (PreviewDynamicMeshComponent == nullptr)
 	{
 		return;
 	}
+
 	PreviewDynamicMeshComponent->SetRenderMeshPostProcessor(MakeUnique<SubdivPostProcessor>(Properties->SubdivisionLevel,
 																					 Properties->NormalComputationMethod,
 																					 Properties->UVComputationMethod));
 
-	TemporaryParentActor->SetRootComponent(PreviewDynamicMeshComponent);
-	PreviewDynamicMeshComponent->RegisterComponent();
-
-	TemporaryParentActor->SetActorTransform(ComponentTarget->GetOwnerActor()->GetRootComponent()->GetComponentTransform());
-
 	// Use the input mesh's material on the preview
 	FComponentMaterialSet MaterialSet;
 	ComponentTarget->GetMaterialSet(MaterialSet);
-
 	for (int k = 0; k < MaterialSet.Materials.Num(); ++k)
 	{
-		PreviewDynamicMeshComponent->SetMaterial(k, MaterialSet.Materials[k]);
+		PreviewMesh->SetMaterial(k, MaterialSet.Materials[k]);
 	}
 
 	// configure secondary render material
 	UMaterialInterface* SelectionMaterial = ToolSetupUtil::GetSelectionMaterial(FLinearColor(0.8f, 0.75f, 0.0f), GetToolManager());
 	if (SelectionMaterial != nullptr)
 	{
-		PreviewDynamicMeshComponent->SetSecondaryRenderMaterial(SelectionMaterial);
+		PreviewMesh->SetSecondaryRenderMaterial(SelectionMaterial);
 	}
 
 	// dynamic mesh configuration settings
-	PreviewDynamicMeshComponent->TangentsType = EDynamicMeshTangentCalcType::NoTangents;
-	PreviewDynamicMeshComponent->InitializeMesh(ComponentTarget->GetMesh());
-	FMeshNormals::QuickComputeVertexNormals(*PreviewDynamicMeshComponent->GetMesh());
-
 	auto RebuildMeshPostProcessor = [this]()
 	{
+		USimpleDynamicMeshComponent* PreviewDynamicMeshComponent = (USimpleDynamicMeshComponent*)PreviewMesh->GetRootComponent();
 		PreviewDynamicMeshComponent->SetRenderMeshPostProcessor(MakeUnique<SubdivPostProcessor>(Properties->SubdivisionLevel,
 																							Properties->NormalComputationMethod,
 																							Properties->UVComputationMethod));
@@ -219,17 +216,19 @@ void USubdividePolyTool::Setup()
 	{
 		if (bNewRenderGroups)
 		{
-			PreviewDynamicMeshComponent->SetOverrideRenderMaterial(ToolSetupUtil::GetSelectionMaterial(GetToolManager()));
-			PreviewDynamicMeshComponent->TriangleColorFunc = [](const FDynamicMesh3* Mesh, int TriangleID)
+			PreviewMesh->SetOverrideRenderMaterial(ToolSetupUtil::GetSelectionMaterial(GetToolManager()));
+			PreviewMesh->SetTriangleColorFunction([](const FDynamicMesh3* Mesh, int TriangleID)
 			{
 				return LinearColors::SelectFColor(Mesh->GetTriangleGroup(TriangleID));
-			};
+			});
 		}
 		else
 		{
-			PreviewDynamicMeshComponent->SetOverrideRenderMaterial(nullptr);
-			PreviewDynamicMeshComponent->TriangleColorFunc = nullptr;
+			PreviewMesh->SetOverrideRenderMaterial(nullptr);
+			PreviewMesh->SetTriangleColorFunction(nullptr);
 		}
+
+		USimpleDynamicMeshComponent* PreviewDynamicMeshComponent = (USimpleDynamicMeshComponent*)PreviewMesh->GetRootComponent();
 		PreviewDynamicMeshComponent->FastNotifyColorsUpdated();
 	};
 
@@ -247,6 +246,8 @@ void USubdividePolyTool::Setup()
 
 	// regenerate preview geo if mesh changes due to undo/redo/etc
 	PreviewDynamicMeshComponent->OnMeshChanged.AddLambda([this]() { bPreviewGeometryNeedsUpdate = true; });
+
+	PreviewMesh->SetVisible(true);
 }
 
 void USubdividePolyTool::CreateOrUpdatePreviewGeometry()
@@ -290,12 +291,13 @@ void USubdividePolyTool::Shutdown(EToolShutdownType ShutdownType)
 		PreviewGeometry->Disconnect();
 	}
 
-	if (PreviewDynamicMeshComponent)
+	if (PreviewMesh)
 	{
 		if (ShutdownType == EToolShutdownType::Accept)
 		{
 			GetToolManager()->BeginUndoTransaction(LOCTEXT("USubdividePolyTool", "Subdivide Mesh"));
 
+			USimpleDynamicMeshComponent* PreviewDynamicMeshComponent = (USimpleDynamicMeshComponent*)PreviewMesh->GetRootComponent();
 			FDynamicMesh3* DynamicMeshResult = PreviewDynamicMeshComponent->GetRenderMesh();
 
 			ComponentTarget->CommitMesh([DynamicMeshResult](const FPrimitiveComponentTarget::FCommitParams& CommitParams)
@@ -306,15 +308,8 @@ void USubdividePolyTool::Shutdown(EToolShutdownType ShutdownType)
 
 			GetToolManager()->EndUndoTransaction();
 		}
-		PreviewDynamicMeshComponent->UnregisterComponent();
-		PreviewDynamicMeshComponent->DestroyComponent();
-		PreviewDynamicMeshComponent = nullptr;
-	}
-
-	if (TemporaryParentActor != nullptr)
-	{
-		TemporaryParentActor->Destroy();
-		TemporaryParentActor = nullptr;
+		PreviewMesh->Disconnect();
+		PreviewMesh = nullptr;
 	}
 
 	ComponentTarget->SetOwnerVisibility(true);
@@ -322,7 +317,7 @@ void USubdividePolyTool::Shutdown(EToolShutdownType ShutdownType)
 
 bool USubdividePolyTool::CanAccept() const
 {
-	return PreviewDynamicMeshComponent != nullptr;
+	return PreviewMesh != nullptr;
 }
 
 
diff --git a/Engine/Plugins/Experimental/MeshModelingToolset/Source/MeshModelingToolsEditorOnly/Public/SubdividePolyTool.h b/Engine/Plugins/Experimental/MeshModelingToolset/Source/MeshModelingToolsEditorOnly/Public/SubdividePolyTool.h
--- a/Engine/Plugins/Experimental/MeshModelingToolset/Source/MeshModelingToolsEditorOnly/Public/SubdividePolyTool.h
+++ b/Engine/Plugins/Experimental/MeshModelingToolset/Source/MeshModelingToolsEditorOnly/Public/SubdividePolyTool.h
@@ -13,7 +13,6 @@
 #include "SubdividePolyTool.generated.h"
 
 class USubdividePolyTool;
-class USimpleDynamicMeshComponent;
 class UPreviewGeometry;
 
 /**
@@ -84,15 +83,11 @@ protected:
 	UWorld* TargetWorld;
 
 	UPROPERTY()
-	USimpleDynamicMeshComponent* PreviewDynamicMeshComponent = nullptr;
+	UPreviewMesh* PreviewMesh = nullptr;
 
 	UPROPERTY()
 	USubdividePolyToolProperties* Properties = nullptr;
 
-	// The temporary actor we create internally to own the preview DynamicMeshComponent
-	UPROPERTY()
-	APreviewMeshActor* TemporaryParentActor = nullptr;
-
 	// Input mesh
 	TSharedPtr<FDynamicMesh3> OriginalMesh;
 
