Add support for modifier keys to LatticeControlPointsMechanic. Behavior is consistent with PolygonSelectionMechanic.
Also changed some Arrays for Sets/Maps.

#rnx

diff --git a/Engine/Plugins/Experimental/MeshModelingToolset/Source/ModelingComponents/Private/Mechanics/LatticeControlPointsMechanic.cpp b/Engine/Plugins/Experimental/MeshModelingToolset/Source/ModelingComponents/Private/Mechanics/LatticeControlPointsMechanic.cpp
--- a/Engine/Plugins/Experimental/MeshModelingToolset/Source/ModelingComponents/Private/Mechanics/LatticeControlPointsMechanic.cpp
+++ b/Engine/Plugins/Experimental/MeshModelingToolset/Source/ModelingComponents/Private/Mechanics/LatticeControlPointsMechanic.cpp
@@ -23,6 +23,21 @@ static const FText LatticePointDeselectionTransactionText = LOCTEXT("LatticePoin
 static const FText LatticePointSelectionTransactionText = LOCTEXT("LatticePointSelection", "Lattice Point Selection");
 static const FText LatticePointMovementTransactionText = LOCTEXT("LatticePointMovement", "Lattice Point Movement");
 
+namespace LatticeControlPointMechanicLocals
+{
+	void Toggle(TSet<int>& Selection, const TSet<int>& NewSelection)
+	{
+		for (const int& NewElement : NewSelection)
+		{
+			if (Selection.Remove(NewElement) == 0)
+			{
+				Selection.Add(NewElement);
+			}
+		}
+	}
+}
+
+
 void ULatticeControlPointsMechanic::Setup(UInteractiveTool* ParentToolIn)
 {
 	UInteractionMechanic::Setup(ParentToolIn);
@@ -35,16 +50,12 @@ void ULatticeControlPointsMechanic::Setup(UInteractiveTool* ParentToolIn)
 
 	USingleClickInputBehavior* ClickBehavior = NewObject<USingleClickInputBehavior>();
 	ClickBehavior->Initialize(this);
-	auto ShiftOrCtrlDown = [](const FInputDeviceState& State) -> bool {
-		return FInputDeviceState::IsShiftKeyDown(State) || FInputDeviceState::IsCtrlKeyDown(State);
-	};
-	ClickBehavior->Modifiers.RegisterModifier(AddToSelectionModifierId, ShiftOrCtrlDown);
-
 	ParentTool->AddInputBehavior(ClickBehavior);
 
 	UMouseHoverBehavior* HoverBehavior = NewObject<UMouseHoverBehavior>();
 	HoverBehavior->Initialize(this);
-	HoverBehavior->Modifiers.RegisterModifier(AddToSelectionModifierId, ShiftOrCtrlDown);
+	HoverBehavior->Modifiers.RegisterModifier(ShiftModifierID, FInputDeviceState::IsShiftKeyDown);
+	HoverBehavior->Modifiers.RegisterModifier(CtrlModifierID, FInputDeviceState::IsCtrlKeyDown);
 	ParentTool->AddInputBehavior(HoverBehavior);
 
 	DrawnControlPoints = NewObject<UPointSetComponent>();
@@ -260,10 +271,10 @@ void ULatticeControlPointsMechanic::GizmoTransformStarted(UTransformProxy* Proxy
 	GizmoStartRotation = Proxy->GetTransform().GetRotation();
 	GizmoStartScale = Proxy->GetTransform().GetScale3D();
 
-	SelectedPointStartPositions.SetNum(SelectedPointIDs.Num());
-	for (int32 i = 0; i < SelectedPointIDs.Num(); ++i)
+	SelectedPointStartPositions.Empty();
+	for (int32 PointID : SelectedPointIDs)
 	{
-		SelectedPointStartPositions[i] = ControlPoints[SelectedPointIDs[i]];
+		SelectedPointStartPositions.Add(PointID, ControlPoints[PointID]);
 	}
 
 	bGizmoBeingDragged = true;
@@ -295,10 +306,9 @@ void ULatticeControlPointsMechanic::GizmoTransformChanged(UTransformProxy* Proxy
 	// If any deltas are non-zero
 	if (Displacement != FVector::ZeroVector || !DeltaRotation.EpsilonEqual(FQuaterniond::Identity(), SMALL_NUMBER) || DeltaScale != FVector::OneVector)
 	{
-		for (int32 i = 0; i < SelectedPointIDs.Num(); ++i)
+		for (int32 PointID : SelectedPointIDs)
 		{
-			int PointID = SelectedPointIDs[i];
-			FVector3d PointPosition = SelectedPointStartPositions[i];
+			FVector3d PointPosition = SelectedPointStartPositions[PointID];
 
 			// Translate to origin, scale, rotate, and translate back (DeltaTransform has "translate back" baked in.)
 			PointPosition -= GizmoStartPosition;
@@ -318,17 +328,17 @@ void ULatticeControlPointsMechanic::GizmoTransformChanged(UTransformProxy* Proxy
 
 void ULatticeControlPointsMechanic::GizmoTransformEnded(UTransformProxy* Proxy)
 {
-	TArray<FVector3d> SelectedPointNewPositions;
-	for (int32 i = 0; i < SelectedPointIDs.Num(); ++i)
+	TMap<int32, FVector3d> SelectedPointNewPositions;
+	for (int32 PointID : SelectedPointIDs)
 	{
-		SelectedPointNewPositions.Add(ControlPoints[SelectedPointIDs[i]]);
+		SelectedPointNewPositions.Add(PointID, ControlPoints[PointID]);
 	}
 
 	bool bFirstMovement = !bHasChanged;
 	bHasChanged = true;
 
 	ParentTool->GetToolManager()->EmitObjectChange(this, MakeUnique<FLatticeControlPointsMechanicMovementChange>(
-		SelectedPointIDs, SelectedPointStartPositions, SelectedPointNewPositions,
+		SelectedPointStartPositions, SelectedPointNewPositions,
 		CurrentChangeStamp, bFirstMovement), LatticePointMovementTransactionText);
 
 	SelectedPointStartPositions.Reset();
@@ -346,12 +356,11 @@ void ULatticeControlPointsMechanic::GizmoTransformEnded(UTransformProxy* Proxy)
 	bGizmoBeingDragged = false;
 }
 
-void ULatticeControlPointsMechanic::UpdatePointLocations(const TArray<int32>& PointIDs, const TArray<FVector3d>& NewLocations)
+void ULatticeControlPointsMechanic::UpdatePointLocations(const TMap<int32, FVector3d>& NewLocations)
 {
-	check(NewLocations.Num() == PointIDs.Num());
-	for (int i = 0; i < NewLocations.Num(); ++i)
+	for ( const TPair<int32, FVector3d>& NewLocationPair : NewLocations )
 	{
-		ControlPoints[PointIDs[i]] = NewLocations[i];
+		ControlPoints[NewLocationPair.Key] = NewLocationPair.Value;
 	}
 	UpdateDrawables();
 }
@@ -383,7 +392,8 @@ void ULatticeControlPointsMechanic::OnClicked(const FInputDeviceRay& ClickPos)
 	if (GeometrySet.FindNearestPointToRay(ClickPos.WorldRay, Nearest, GeometrySetToleranceTest))
 	{
 		ParentTool->GetToolManager()->BeginUndoTransaction(LatticePointSelectionTransactionText);
-		ChangeSelection(Nearest.ID, bAddToSelectionToggle);
+		// TODO: Have the modifier keys for single-select match the behavior of marquee select
+		ChangeSelection(Nearest.ID, ShouldAddToSelectionFunc());
 		ParentTool->GetToolManager()->EndUndoTransaction();
 	}
 
@@ -394,7 +404,7 @@ void ULatticeControlPointsMechanic::ChangeSelection(int32 NewPointID, bool bAddT
 	// If not adding to selection, clear it
 	if (!bAddToSelection && SelectedPointIDs.Num() > 0)
 	{
-		TArray<int32> PointsToDeselect;
+		TSet<int32> PointsToDeselect;
 
 		for (int32 PointID : SelectedPointIDs)
 		{
@@ -476,11 +486,9 @@ void ULatticeControlPointsMechanic::UpdateGizmoLocation()
 bool ULatticeControlPointsMechanic::DeselectPoint(int32 PointID)
 {
 	bool PointFound = false;
-	int32 IndexInSelection;
 	
-	if (SelectedPointIDs.Find(PointID, IndexInSelection))
+	if (SelectedPointIDs.Remove(PointID) != 0)
 	{
-		SelectedPointIDs.RemoveAt(IndexInSelection);
 		DrawnControlPoints->SetPointColor(PointID, NormalPointColor);
 		PointFound = true;
 	}
@@ -554,20 +562,21 @@ void ULatticeControlPointsMechanic::OnEndHover()
 // Detects Ctrl key state
 void ULatticeControlPointsMechanic::OnUpdateModifierState(int ModifierID, bool bIsOn)
 {
-	if (ModifierID == AddToSelectionModifierId)
+	switch (ModifierID)
 	{
-		bAddToSelectionToggle = bIsOn;
+	case ShiftModifierID:
+		bShiftToggle = bIsOn;
+		break;
+	case CtrlModifierID:
+		bCtrlToggle = bIsOn;
+		break;
 	}
 }
 
 // These get bound to the delegates on the marquee mechanic.
 void ULatticeControlPointsMechanic::OnDragRectangleStarted()
 {
-	for (int32& PointID : SelectedPointIDs)
-	{
-		DrawnControlPoints->SetPointColor(PointID, NormalPointColor);
-	}
-	CurrentDragSelection.Empty();
+	PreDragSelection = SelectedPointIDs;
 
 	// Hide gizmo while dragging
 	if (PointTransformGizmo)
@@ -575,52 +584,69 @@ void ULatticeControlPointsMechanic::OnDragRectangleStarted()
 		PointTransformGizmo->SetVisibility(false);
 	}
 }
+
 void ULatticeControlPointsMechanic::OnDragRectangleChanged(const FCameraRectangle& Rectangle)
 {
-	CurrentDragSelection.Empty();
+	TSet<int> DragSelection;
 	for (int32 PointID = 0; PointID < ControlPoints.Num(); ++PointID)
 	{
 		FVector PointPosition = DrawnControlPoints->GetPoint(PointID).Position;
+		if (Rectangle.IsProjectedPointInRectangle(PointPosition))
+		{
+			DragSelection.Add(PointID);
+		}
+	}
 
-		if (Rectangle.IsProjectedPointInRectangle(PointPosition) )
+	if (ShouldAddToSelectionFunc())
+	{
+		SelectedPointIDs = PreDragSelection;
+		if (ShouldRemoveFromSelectionFunc())
 		{
-			CurrentDragSelection.Add(PointID);
-			DrawnControlPoints->SetPointColor(PointID, SelectedColor);
+			LatticeControlPointMechanicLocals::Toggle(SelectedPointIDs, DragSelection);
 		}
 		else
 		{
-			DrawnControlPoints->SetPointColor(PointID, NormalPointColor);
+			SelectedPointIDs.Append(DragSelection);
 		}
 	}
+	else if (ShouldRemoveFromSelectionFunc())
+	{
+		SelectedPointIDs = PreDragSelection.Difference(DragSelection);
+	}
+	else
+	{
+		// Neither key pressed.
+		SelectedPointIDs = DragSelection;
+	}
+
+	UpdateDrawables();
 }
+
 void ULatticeControlPointsMechanic::OnDragRectangleFinished()
 {
-	// Deselect previous SelectedPointIDs and replace it with "drag selection" points. Do this in one Undo transaction.
-
 	ParentTool->GetToolManager()->BeginUndoTransaction(LatticePointSelectionTransactionText);
 
-	TArray<int32> SavedSelectedIDs = SelectedPointIDs;
-
 	FTransform PreviousTransform = PointTransformProxy->GetTransform();
-	SelectedPointIDs = CurrentDragSelection;
+	
 	UpdateGizmoLocation();
+
 	FTransform NewTransform = PointTransformProxy->GetTransform();
 
-	if (SavedSelectedIDs.Num() > 0)
+	if (PreDragSelection.Num() > 0)
 	{
 		ParentTool->GetToolManager()->EmitObjectChange(this, MakeUnique<FLatticeControlPointsMechanicSelectionChange>(
-			SavedSelectedIDs, false, PreviousTransform, NewTransform, CurrentChangeStamp), LatticePointDeselectionTransactionText);
+			PreDragSelection, false, PreviousTransform, NewTransform, CurrentChangeStamp), LatticePointDeselectionTransactionText);
 	}
 
-	if (CurrentDragSelection.Num() > 0)
+	if (SelectedPointIDs.Num() > 0)
 	{
 		ParentTool->GetToolManager()->EmitObjectChange(this, MakeUnique<FLatticeControlPointsMechanicSelectionChange>(
-			CurrentDragSelection, true, PreviousTransform, NewTransform, CurrentChangeStamp), LatticePointSelectionTransactionText);
+			SelectedPointIDs, true, PreviousTransform, NewTransform, CurrentChangeStamp), LatticePointSelectionTransactionText);
 	}
 
 	ParentTool->GetToolManager()->EndUndoTransaction();
 
-	CurrentDragSelection.Empty();
+	UpdateDrawables();
 }
 
 const TArray<FVector3d>& ULatticeControlPointsMechanic::GetControlPoints() const
@@ -642,7 +668,7 @@ FLatticeControlPointsMechanicSelectionChange::FLatticeControlPointsMechanicSelec
 	, ChangeStamp(ChangeStampIn)
 {}
 
-FLatticeControlPointsMechanicSelectionChange::FLatticeControlPointsMechanicSelectionChange(const TArray<int32>& PointIDsIn, 
+FLatticeControlPointsMechanicSelectionChange::FLatticeControlPointsMechanicSelectionChange(const TSet<int32>& PointIDsIn, 
 																						   bool bAddedIn, 
 																						   const FTransform& PreviousTransformIn,
 																						   const FTransform& NewTransformIn, 
@@ -705,26 +731,22 @@ FString FLatticeControlPointsMechanicSelectionChange::ToString() const
 }
 
 
-FLatticeControlPointsMechanicMovementChange::FLatticeControlPointsMechanicMovementChange(const TArray<int32>& PointIDsIn,
-	const TArray<FVector3d>& OriginalPositionsIn,
-	const TArray<FVector3d>& NewPositionsIn,
+FLatticeControlPointsMechanicMovementChange::FLatticeControlPointsMechanicMovementChange(
+	const TMap<int32, FVector3d>& OriginalPositionsIn,
+	const TMap<int32, FVector3d>& NewPositionsIn,
 	int32 ChangeStampIn,
-	bool bFirstMovementIn)
-	: PointIDs{ PointIDsIn }
-	, OriginalPositions{ OriginalPositionsIn }
+	bool bFirstMovementIn) :
+	OriginalPositions{ OriginalPositionsIn }
 	, NewPositions{ NewPositionsIn }
 	, ChangeStamp(ChangeStampIn)
 	, bFirstMovement(bFirstMovementIn)
 {
-	check(PointIDs.Num() == OriginalPositions.Num());
-	check(PointIDs.Num() == NewPositions.Num());
 }
 
 void FLatticeControlPointsMechanicMovementChange::Apply(UObject* Object)
 {
 	ULatticeControlPointsMechanic* Mechanic = Cast<ULatticeControlPointsMechanic>(Object);
-	check(PointIDs.Num() == NewPositions.Num());
-	Mechanic->UpdatePointLocations(PointIDs, NewPositions);
+	Mechanic->UpdatePointLocations(NewPositions);
 	Mechanic->bHasChanged = false;
 	Mechanic->OnPointsChanged.Broadcast();
 }
@@ -732,8 +754,7 @@ void FLatticeControlPointsMechanicMovementChange::Apply(UObject* Object)
 void FLatticeControlPointsMechanicMovementChange::Revert(UObject* Object)
 {
 	ULatticeControlPointsMechanic* Mechanic = Cast<ULatticeControlPointsMechanic>(Object);
-	check(PointIDs.Num() == OriginalPositions.Num());
-	Mechanic->UpdatePointLocations(PointIDs, OriginalPositions);
+	Mechanic->UpdatePointLocations(OriginalPositions);
 	if (bFirstMovement)
 	{
 		// If we're undoing the first change, make it possible to change the lattice resolution again
diff --git a/Engine/Plugins/Experimental/MeshModelingToolset/Source/ModelingComponents/Public/Mechanics/LatticeControlPointsMechanic.h b/Engine/Plugins/Experimental/MeshModelingToolset/Source/ModelingComponents/Public/Mechanics/LatticeControlPointsMechanic.h
--- a/Engine/Plugins/Experimental/MeshModelingToolset/Source/ModelingComponents/Public/Mechanics/LatticeControlPointsMechanic.h
+++ b/Engine/Plugins/Experimental/MeshModelingToolset/Source/ModelingComponents/Public/Mechanics/LatticeControlPointsMechanic.h
@@ -102,9 +102,14 @@ protected:
 	FColor PreHoverPointColor;
 
 	// Support for Shift and Ctrl toggle
-	bool bAddToSelectionToggle = false;
-	const int32 AddToSelectionModifierId = 1;
-	
+	bool bShiftToggle = false;
+	bool bCtrlToggle = false;
+	static const int32 ShiftModifierID = 1;
+	static const int32 CtrlModifierID = 2;
+
+	// Default modifier key behavior is consistent with PolygonSelectionMechanic
+	TFunction<bool(void)> ShouldAddToSelectionFunc = [this]() {return bShiftToggle; };
+	TFunction<bool(void)> ShouldRemoveFromSelectionFunc = [this]() {return bCtrlToggle; };
 
 	// Support for gizmo. Since the points aren't individual components, we don't actually use UTransformProxy
 	// for the transform forwarding- we just use it for the callbacks.
@@ -130,7 +135,8 @@ protected:
 	void ClearHover();
 
 	// Support for selection
-	TArray<int32> SelectedPointIDs;
+	TSet<int32> SelectedPointIDs;
+	TSet<int32> PreDragSelection;
 	URectangleMarqueeMechanic* MarqueeMechanic;
 	TArray<int32> CurrentDragSelection;
 	void OnDragRectangleStarted();
@@ -138,7 +144,7 @@ protected:
 	void OnDragRectangleFinished();
 
 	// We need the selected point start positions so we can move multiple points appropriately.
-	TArray<FVector3d> SelectedPointStartPositions;
+	TMap<int32, FVector3d> SelectedPointStartPositions;
 
 	// The starting point of the gizmo is needed to determine the offset by which to move the points.
 	// TODO: Replace with single FTransform?
@@ -155,7 +161,7 @@ protected:
 	void SelectPoint(int32 PointID);
 	bool DeselectPoint(int32 PointID);
 	void UpdateGizmoLocation();
-	void UpdatePointLocations(const TArray<int32>& PointIDs, const TArray<FVector3d>& NewLocations);
+	void UpdatePointLocations(const TMap<int32, FVector3d>& NewLocations);
 
 	void RebuildDrawables();
 	void UpdateDrawables();
@@ -177,7 +183,7 @@ public:
 	FLatticeControlPointsMechanicSelectionChange(int32 PointIDIn, bool bAddedIn, const FTransform& PreviousTransformIn,
 												 const FTransform& NewTransformIn, int32 ChangeStampIn);
 
-	FLatticeControlPointsMechanicSelectionChange(const TArray<int32>& PointIDsIn, bool bAddedIn, const FTransform& PreviousTransformIn,
+	FLatticeControlPointsMechanicSelectionChange(const TSet<int32>& PointIDsIn, bool bAddedIn, const FTransform& PreviousTransformIn,
 												 const FTransform& NewTransformIn, int32 ChangeStampIn);
 
 	virtual void Apply(UObject* Object) override;
@@ -189,7 +195,7 @@ public:
 	virtual FString ToString() const override;
 
 protected:
-	TArray<int32> PointIDs;
+	TSet<int32> PointIDs;
 	bool bAdded;
 
 	const FTransform PreviousTransform;
@@ -204,9 +210,9 @@ class MODELINGCOMPONENTS_API FLatticeControlPointsMechanicMovementChange : publi
 {
 public:
 
-	FLatticeControlPointsMechanicMovementChange(const TArray<int32>& PointIDsIn, 
-		const TArray<FVector3d>& OriginalPositionIn,
-		const TArray<FVector3d>& NewPositionIn, 
+	FLatticeControlPointsMechanicMovementChange(
+		const TMap<int32, FVector3d>& OriginalPositionsIn,
+		const TMap<int32, FVector3d>& NewPositionsIn,
 		int32 ChangeStampIn,
 		bool bFirstMovementIn);
 
@@ -220,9 +226,8 @@ public:
 
 protected:
 
-	TArray<int32> PointIDs;
-	TArray<FVector3d> OriginalPositions;
-	TArray<FVector3d> NewPositions;
+	TMap<int32, FVector3d> OriginalPositions;
+	TMap<int32, FVector3d> NewPositions;
 	int32 ChangeStamp;
 	bool bFirstMovement;
 };
