MeshModelingTools: Split VolumeDynamicMeshConversion into runtime and editor-only parts. Also moved VolumeToMeshTool to not be editor-only now that it is safe to do so, and moved the PlanarFlipsOptimization function into an accessible file.

#rnx

diff --git a/Engine/Plugins/Experimental/GeometryProcessing/Source/DynamicMesh/Private/Operations/PlanarFlipsOptimization.cpp b/Engine/Plugins/Experimental/GeometryProcessing/Source/DynamicMesh/Private/Operations/PlanarFlipsOptimization.cpp
new file mode 100644
--- /dev/null
+++ b/Engine/Plugins/Experimental/GeometryProcessing/Source/DynamicMesh/Private/Operations/PlanarFlipsOptimization.cpp
@@ -0,0 +1,128 @@
+// Copyright Epic Games, Inc. All Rights Reserved.
+
+#include "Operations/PlanarFlipsOptimization.h"
+
+#include "DynamicMesh3.h"
+#include "MathUtil.h"
+#include "VectorTypes.h"
+
+namespace UE {
+namespace PlanarFlipsOptimizationLocals{
+
+/** @return triangle aspect ratio transformed to be in [0,1] range */
+double UnitAspectRatio(const FVector3d& A, const FVector3d& B, const FVector3d& C)
+{
+	double AspectRatio = VectorUtil::AspectRatio(A, B, C);
+	return (AspectRatio > 1.0) ? FMathd::Clamp(1.0 / AspectRatio, 0.0, 1.0) : AspectRatio;
+}
+/** @return triangle aspect ratio transformed to be in [0,1] range */
+double UnitAspectRatio(const FDynamicMesh3& Mesh, int32 TriangleID)
+{
+	FVector3d A, B, C;
+	Mesh.GetTriVertices(TriangleID, A, B, C);
+	return UnitAspectRatio(A, B, C);
+}
+}}//end namespace UE::PlanarFlipsOptimizationLocals
+
+void FPlanarFlipsOptimization::Apply()
+{
+	check(Mesh);
+	for (int32 i = 0; i < NumPasses; ++i)
+	{
+		ApplySinglePass();
+	}
+}
+
+void FPlanarFlipsOptimization::ApplySinglePass()
+{
+	using namespace UE::PlanarFlipsOptimizationLocals;
+
+	//Note: could be more efficient to do multiple passes directly in Apply(), would save on the initial computation
+
+	struct FFlatEdge
+	{
+		int32 eid;
+		double MinAspect;
+	};
+
+	TArray<double> AspectRatios;
+	TArray<FVector3d> Normals;
+	AspectRatios.SetNum(Mesh->MaxTriangleID());
+	Normals.SetNum(Mesh->MaxTriangleID());
+	for (int32 tid : Mesh->TriangleIndicesItr())
+	{
+		FVector3d A, B, C;
+		Mesh->GetTriVertices(tid, A, B, C);
+		AspectRatios[tid] = UnitAspectRatio(A, B, C);
+		Normals[tid] = VectorUtil::Normal(A, B, C);
+	}
+
+	TArray<FFlatEdge> Flips;
+	for (int32 eid : Mesh->EdgeIndicesItr())
+	{
+		if (Mesh->IsBoundaryEdge(eid) == false)
+		{
+			FIndex2i EdgeT = Mesh->GetEdgeT(eid);
+			if (AspectRatios[EdgeT.A] < 0.01 && AspectRatios[EdgeT.B] < 0.01)
+			{
+				continue;		// if both are degenerate we can't fix by flipping edge between them
+			}
+			double MinAspect = FMathd::Min(AspectRatios[EdgeT.A], AspectRatios[EdgeT.B]);
+			double NormDot = Normals[EdgeT.A].Dot(Normals[EdgeT.B]);
+			if (NormDot > PlanarDotThresh)
+			{
+				Flips.Add({ eid, MinAspect });
+			}
+		}
+	}
+
+	Flips.Sort([&](const FFlatEdge& A, const FFlatEdge& B) { return A.MinAspect < B.MinAspect; });
+
+	for (int32 k = 0; k < Flips.Num(); ++k)
+	{
+		int32 eid = Flips[k].eid;
+		FIndex2i EdgeV = Mesh->GetEdgeV(eid);
+		int32 a = EdgeV.A, b = EdgeV.B;
+		FIndex2i EdgeT = Mesh->GetEdgeT(eid);
+		FIndex3i Tri0 = Mesh->GetTriangle(EdgeT.A), Tri1 = Mesh->GetTriangle(EdgeT.B);
+		int32 c = IndexUtil::OrientTriEdgeAndFindOtherVtx(a, b, Tri0);
+		int32 d = IndexUtil::FindTriOtherVtx(a, b, Tri1);
+
+		double AspectA = AspectRatios[EdgeT.A], AspectB = AspectRatios[EdgeT.B];
+		double Metric = FMathd::Min(AspectA, AspectB);
+		FVector3d Normal = (AspectA > AspectB) ? Normals[EdgeT.A] : Normals[EdgeT.B];
+
+		FVector3d A = Mesh->GetVertex(a), B = Mesh->GetVertex(b);
+		FVector3d C = Mesh->GetVertex(c), D = Mesh->GetVertex(d);
+
+		double FlipAspect1 = UnitAspectRatio(C, D, B);
+		double FlipAspect2 = UnitAspectRatio(D, C, A);
+		FVector3d FlipNormal1 = VectorUtil::Normal(C, D, B);
+		FVector3d FlipNormal2 = VectorUtil::Normal(D, C, A);
+		if (FlipNormal1.Dot(Normal) < PlanarDotThresh || FlipNormal2.Dot(Normal) < PlanarDotThresh)
+		{
+			continue;		// should not happen?
+		}
+
+		if (FMathd::Min(FlipAspect1, FlipAspect2) > Metric)
+		{
+			FDynamicMesh3::FEdgeFlipInfo FlipInfo;
+			if (Mesh->FlipEdge(eid, FlipInfo) == EMeshResult::Ok)
+			{
+				AspectRatios[EdgeT.A] = UnitAspectRatio(*Mesh, EdgeT.A);
+				AspectRatios[EdgeT.B] = UnitAspectRatio(*Mesh, EdgeT.B);
+
+				// safety check - if somehow we flipped the normal, flip it back
+				bool bInvertedNormal = (Mesh->GetTriNormal(EdgeT.A).Dot(Normal) < PlanarDotThresh) ||
+					(Mesh->GetTriNormal(EdgeT.B).Dot(Normal) < PlanarDotThresh);
+				if (bInvertedNormal)
+				{
+					UE_LOG(LogTemp, Warning, TEXT("FPlanarFlipsOptimization: Invalid Flip!"));
+					Mesh->FlipEdge(eid, FlipInfo);
+					AspectRatios[EdgeT.A] = UnitAspectRatio(*Mesh, EdgeT.A);
+					AspectRatios[EdgeT.B] = UnitAspectRatio(*Mesh, EdgeT.B);
+				}
+			}
+		}
+	}
+}
\ No newline at end of file
diff --git a/Engine/Plugins/Experimental/GeometryProcessing/Source/DynamicMesh/Public/Operations/PlanarFlipsOptimization.h b/Engine/Plugins/Experimental/GeometryProcessing/Source/DynamicMesh/Public/Operations/PlanarFlipsOptimization.h
new file mode 100644
--- /dev/null
+++ b/Engine/Plugins/Experimental/GeometryProcessing/Source/DynamicMesh/Public/Operations/PlanarFlipsOptimization.h
@@ -0,0 +1,37 @@
+// Copyright Epic Games, Inc. All Rights Reserved.
+
+#pragma once
+
+#include "CoreMinimal.h"
+
+class FDynamicMesh3;
+
+/**
+ * If both triangles on an edge are coplanar, we can arbitrarily flip the interior edge to
+ * improve triangle quality. Similarly if one triangle on an edge is degenerate, we can flip
+ * the edge without affecting the shape to try to remove it. This code does one or more passes of
+ * such an optimization.
+ */
+class DYNAMICMESH_API FPlanarFlipsOptimization
+{
+public:
+
+	explicit FPlanarFlipsOptimization(FDynamicMesh3* MeshIn, int32 NumPassesIn = 1, double PlanarDotThreshIn = 0.99)
+		: Mesh(MeshIn)
+		, NumPasses(NumPassesIn)
+		, PlanarDotThresh(PlanarDotThreshIn)
+	{
+	}
+
+	FDynamicMesh3* Mesh;
+	int32 NumPasses;
+	double PlanarDotThresh;
+
+	/**
+	 * Apply the operation to the mesh.
+	 */
+	void Apply();
+
+protected:
+	void ApplySinglePass();
+};
\ No newline at end of file
diff --git a/Engine/Plugins/Experimental/MeshModelingToolset/Source/MeshModelingTools/Private/ConversionUtils/VolumeToDynamicMesh.cpp b/Engine/Plugins/Experimental/MeshModelingToolset/Source/MeshModelingTools/Private/ConversionUtils/VolumeToDynamicMesh.cpp
new file mode 100644
--- /dev/null
+++ b/Engine/Plugins/Experimental/MeshModelingToolset/Source/MeshModelingTools/Private/ConversionUtils/VolumeToDynamicMesh.cpp
@@ -0,0 +1,101 @@
+// Copyright Epic Games, Inc. All Rights Reserved.
+
+#include "ConversionUtils/VolumeToDynamicMesh.h"
+
+#include "CompGeom/PolygonTriangulation.h"
+#include "DynamicMesh3.h"
+#include "GameFramework/Volume.h"
+#include "MeshBoundaryLoops.h"
+#include "MeshQueries.h"
+#include "MeshRegionBoundaryLoops.h"
+#include "Model.h"
+#include "Operations/MergeCoincidentMeshEdges.h"
+#include "Operations/MinimalHoleFiller.h"
+#include "Operations/PlanarFlipsOptimization.h"
+
+namespace UE {
+namespace Conversion {
+
+void VolumeToDynamicMesh(AVolume* Volume, FDynamicMesh3& Mesh, 
+	const FVolumeToMeshOptions& Options)
+{
+	Mesh.DiscardAttributes();
+	if (Options.bSetGroups)
+	{
+		Mesh.EnableTriangleGroups();
+	}
+
+	UModel* Model = Volume->Brush;
+	FTransform3d XForm = (Options.bInWorldSpace) ? FTransform3d(Volume->GetTransform()) : FTransform3d::Identity();
+
+	// Each "BspNode" is a planar polygon, triangulate each polygon and accumulate in a mesh.
+	// Note that this does not make any attempt to weld vertices/edges
+	for (const FBspNode& Node : Model->Nodes)
+	{
+		FVector3d Normal = (FVector3d)Node.Plane;
+		FFrame3d Plane(Node.Plane.W * Normal, Normal);
+
+		int32 NumVerts = (Node.NodeFlags & PF_TwoSided) ? Node.NumVertices / 2 : Node.NumVertices;  // ??
+		if (NumVerts > 0)
+		{
+			TArray<int32> VertIndices;
+			TArray<FVector2d> VertPositions2d;
+			VertIndices.SetNum(NumVerts);
+			VertPositions2d.SetNum(NumVerts);
+			for (int32 VertexIndex = 0; VertexIndex < NumVerts; ++VertexIndex)
+			{
+				const FVert& Vert = Model->Verts[Node.iVertPool + VertexIndex];
+				FVector3d Point = (FVector3d)Model->Points[Vert.pVertex];
+				Point = XForm.TransformPosition(Point);
+				VertIndices[VertexIndex] = Mesh.AppendVertex(Point);
+				VertPositions2d[VertexIndex] = Plane.ToPlaneUV(Point, 2);
+			}
+
+			TArray<FIndex3i> PolyTriangles;
+			PolygonTriangulation::TriangulateSimplePolygon(VertPositions2d, PolyTriangles);
+
+			int32 GroupID = FDynamicMesh3::InvalidID;
+			if (Options.bSetGroups)
+			{
+				GroupID = Mesh.AllocateTriangleGroup();
+			}
+
+			for (FIndex3i Tri : PolyTriangles)
+			{
+				// flip orientation here...
+				Mesh.AppendTriangle(VertIndices[Tri.A], VertIndices[Tri.C], VertIndices[Tri.B], GroupID);
+			}
+		}
+	}
+
+	if (Options.bMergeVertices)
+	{
+		// Merge the mesh edges to create a closed solid
+		double MinLen, MaxLen, AvgLen;
+		TMeshQueries<FDynamicMesh3>::EdgeLengthStats(Mesh, MinLen, MaxLen, AvgLen);
+		FMergeCoincidentMeshEdges Merge(&Mesh);
+		Merge.MergeVertexTolerance = FMathd::Max(Merge.MergeVertexTolerance, MinLen * 0.1);
+		Merge.Apply();
+
+		// If the mesh is not closed, the merge failed or the volume had cracks/holes. 
+		// Do trivial hole fills to ensure the output is solid   (really want autorepair here)
+		if (Mesh.IsClosed() == false && Options.bAutoRepairMesh)
+		{
+			FMeshBoundaryLoops BoundaryLoops(&Mesh, true);
+			for (FEdgeLoop& Loop : BoundaryLoops.Loops)
+			{
+				FMinimalHoleFiller Filler(&Mesh, Loop);
+				Filler.Fill();
+			}
+		}
+
+
+		// try to flip towards better triangles in planar areas, should reduce/remove degenerate geo
+		if (Options.bOptimizeMesh)
+		{
+			FPlanarFlipsOptimization(&Mesh, 5).Apply(); // Do five passes
+		}
+	}
+}
+
+}}//end namespace UE::Conversion
\ No newline at end of file
diff --git a/Engine/Plugins/Experimental/MeshModelingToolset/Source/MeshModelingTools/Private/VolumeToMeshTool.cpp b/Engine/Plugins/Experimental/MeshModelingToolset/Source/MeshModelingTools/Private/VolumeToMeshTool.cpp
new file mode 100644
--- /dev/null
+++ b/Engine/Plugins/Experimental/MeshModelingToolset/Source/MeshModelingTools/Private/VolumeToMeshTool.cpp
@@ -0,0 +1,224 @@
+// Copyright Epic Games, Inc. All Rights Reserved.
+
+#include "VolumeToMeshTool.h"
+#include "InteractiveToolManager.h"
+#include "ToolBuilderUtil.h"
+
+#include "ConversionUtils/VolumeToDynamicMesh.h"
+#include "DynamicMesh3.h"
+#include "MeshNormals.h"
+#include "DynamicMeshEditor.h"
+#include "Util/ColorConstants.h"
+#include "ToolSetupUtil.h"
+#include "Selection/ToolSelectionUtil.h"
+#include "AssetGenerationUtil.h"
+
+#include "Operations/MergeCoincidentMeshEdges.h"
+#include "CompGeom/PolygonTriangulation.h"
+#include "MeshBoundaryLoops.h"
+#include "Operations/MinimalHoleFiller.h"
+
+#include "Model.h"
+
+#define LOCTEXT_NAMESPACE "UVolumeToMeshTool"
+
+
+/*
+ * ToolBuilder
+ */
+
+bool UVolumeToMeshToolBuilder::CanBuildTool(const FToolBuilderState& SceneState) const
+{
+	return ToolBuilderUtil::CountSelectedActorsOfType<AVolume>(SceneState) == 1;
+}
+
+UInteractiveTool* UVolumeToMeshToolBuilder::BuildTool(const FToolBuilderState& SceneState) const
+{
+	UVolumeToMeshTool* NewTool = NewObject<UVolumeToMeshTool>(SceneState.ToolManager);
+
+	NewTool->SetWorld(SceneState.World);
+	check(AssetAPI);
+	NewTool->SetAssetAPI(AssetAPI);
+
+	AVolume* Volume = ToolBuilderUtil::FindFirstActorOfType<AVolume>(SceneState);
+	check(Volume != nullptr);
+	NewTool->SetSelection(Volume);
+
+	return NewTool;
+}
+
+
+
+
+/*
+ * Tool
+ */
+UVolumeToMeshTool::UVolumeToMeshTool()
+{
+	SetToolDisplayName(LOCTEXT("VolumeToMeshToolName", "Volume to Mesh"));
+}
+
+
+void UVolumeToMeshTool::SetSelection(AVolume* Volume)
+{
+	TargetVolume = Volume;
+}
+
+
+void UVolumeToMeshTool::Setup()
+{
+	UInteractiveTool::Setup();
+
+	PreviewMesh = NewObject<UPreviewMesh>(this);
+	PreviewMesh->bBuildSpatialDataStructure = false;
+	PreviewMesh->CreateInWorld(TargetVolume->GetWorld(), FTransform::Identity);
+	PreviewMesh->SetTransform(TargetVolume->GetActorTransform());
+
+	PreviewMesh->SetMaterial( ToolSetupUtil::GetDefaultSculptMaterial(GetToolManager()) );
+
+	PreviewMesh->SetOverrideRenderMaterial(ToolSetupUtil::GetSelectionMaterial(GetToolManager()));
+	PreviewMesh->SetTriangleColorFunction([this](const FDynamicMesh3* Mesh, int TriangleID)
+	{
+		return LinearColors::SelectFColor(Mesh->GetTriangleGroup(TriangleID));
+	});
+
+	VolumeEdgesSet = NewObject<ULineSetComponent>(PreviewMesh->GetRootComponent());
+	VolumeEdgesSet->SetupAttachment(PreviewMesh->GetRootComponent());
+	VolumeEdgesSet->SetLineMaterial(ToolSetupUtil::GetDefaultLineComponentMaterial(GetToolManager()));
+	VolumeEdgesSet->RegisterComponent();
+
+
+	Settings = NewObject<UVolumeToMeshToolProperties>(this);
+	Settings->RestoreProperties(this);
+	AddToolPropertySource(Settings);
+
+	Settings->WatchProperty(Settings->bWeldEdges, [this](bool) { bResultValid = false; });
+	Settings->WatchProperty(Settings->bAutoRepair, [this](bool) { bResultValid = false; });
+	Settings->WatchProperty(Settings->bOptimizeMesh, [this](bool) { bResultValid = false; });
+	Settings->WatchProperty(Settings->bShowWireframe, [this](bool) { bResultValid = false; });
+
+	bResultValid = false;
+
+	GetToolManager()->DisplayMessage( 
+		LOCTEXT("OnStartTool", "Convert a Volume to a Static Mesh"),
+		EToolMessageLevel::UserNotification);
+}
+
+
+
+void UVolumeToMeshTool::Shutdown(EToolShutdownType ShutdownType)
+{
+	Settings->SaveProperties(this);
+
+	FTransform3d Transform(PreviewMesh->GetTransform());
+	PreviewMesh->SetVisible(false);
+	PreviewMesh->Disconnect();
+	PreviewMesh = nullptr;
+
+	if (ShutdownType == EToolShutdownType::Accept )
+	{
+		UMaterialInterface* UseMaterial = UMaterial::GetDefaultMaterial(MD_Surface);
+
+		FString NewName = TargetVolume.IsValid() ?
+			FString::Printf(TEXT("%sMesh"), *TargetVolume->GetName()) : TEXT("Volume Mesh");
+
+		GetToolManager()->BeginUndoTransaction(LOCTEXT("CreateMeshVolume", "Volume To Mesh"));
+
+		AActor* NewActor = AssetGenerationUtil::GenerateStaticMeshActor(
+			AssetAPI, TargetWorld,
+			&CurrentMesh, Transform, NewName, UseMaterial);
+		if (NewActor != nullptr)
+		{
+			ToolSelectionUtil::SetNewActorSelection(GetToolManager(), NewActor);
+		}
+
+		GetToolManager()->EndUndoTransaction();
+	}
+
+
+}
+
+
+void UVolumeToMeshTool::OnTick(float DeltaTime)
+{
+	if (bResultValid == false)
+	{
+		RecalculateMesh();
+	}
+}
+
+void UVolumeToMeshTool::Render(IToolsContextRenderAPI* RenderAPI)
+{
+}
+
+
+bool UVolumeToMeshTool::HasAccept() const
+{
+	return true;
+}
+
+bool UVolumeToMeshTool::CanAccept() const
+{
+	return true;
+}
+
+
+
+void UVolumeToMeshTool::UpdateLineSet()
+{
+	VolumeEdgesSet->Clear();
+
+	FColor BoundaryEdgeColor = LinearColors::VideoRed3b();
+	float BoundaryEdgeThickness = 1.0;
+	float BoundaryEdgeDepthBias = 2.0f;
+
+	FColor WireEdgeColor = LinearColors::Gray3b();
+	float WireEdgeThickness = 0.1;
+	float WireEdgeDepthBias = 1.0f;
+
+	if (Settings->bShowWireframe)
+	{
+		VolumeEdgesSet->ReserveLines(CurrentMesh.EdgeCount());
+
+		for (int32 eid : CurrentMesh.EdgeIndicesItr())
+		{
+			FVector3d A, B;
+			CurrentMesh.GetEdgeV(eid, A, B);
+			if (CurrentMesh.IsBoundaryEdge(eid))
+			{
+				VolumeEdgesSet->AddLine((FVector)A, (FVector)B,
+					BoundaryEdgeColor, BoundaryEdgeThickness, BoundaryEdgeDepthBias);
+			}
+			else
+			{
+				VolumeEdgesSet->AddLine((FVector)A, (FVector)B,
+					WireEdgeColor, WireEdgeThickness, WireEdgeDepthBias);
+			}
+		}
+	}
+}
+
+void UVolumeToMeshTool::RecalculateMesh()
+{
+	if (TargetVolume.IsValid())
+	{
+		UE::Conversion::FVolumeToMeshOptions Options;
+		Options.bMergeVertices = Settings->bWeldEdges;
+		Options.bAutoRepairMesh = Settings->bAutoRepair;
+		Options.bOptimizeMesh = Settings->bOptimizeMesh;
+
+		CurrentMesh = FDynamicMesh3(EMeshComponents::FaceGroups);
+		UE::Conversion::VolumeToDynamicMesh(TargetVolume.Get(), CurrentMesh, Options);
+		FMeshNormals::InitializeMeshToPerTriangleNormals(&CurrentMesh);
+		PreviewMesh->UpdatePreview(&CurrentMesh);
+	}
+
+	UpdateLineSet();
+
+	bResultValid = true;
+}
+
+
+
+
+#undef LOCTEXT_NAMESPACE
diff --git a/Engine/Plugins/Experimental/MeshModelingToolset/Source/MeshModelingTools/Public/ConversionUtils/VolumeToDynamicMesh.h b/Engine/Plugins/Experimental/MeshModelingToolset/Source/MeshModelingTools/Public/ConversionUtils/VolumeToDynamicMesh.h
new file mode 100644
--- /dev/null
+++ b/Engine/Plugins/Experimental/MeshModelingToolset/Source/MeshModelingTools/Public/ConversionUtils/VolumeToDynamicMesh.h
@@ -0,0 +1,29 @@
+// Copyright Epic Games, Inc. All Rights Reserved.
+
+#pragma once
+
+#include "CoreMinimal.h"
+#include "FrameTypes.h"
+#include "VectorTypes.h"
+
+class AVolume;
+class FDynamicMesh3;
+
+namespace UE {
+namespace Conversion {
+
+struct FVolumeToMeshOptions
+{
+	bool bInWorldSpace = false;
+	bool bSetGroups = true;
+	bool bMergeVertices = true;
+	bool bAutoRepairMesh = true;
+	bool bOptimizeMesh = true;
+};
+
+/**
+ * Converts a volume to a dynamic mesh. Does not initialize normals and does not delete the volume.
+ */
+void MESHMODELINGTOOLS_API VolumeToDynamicMesh(AVolume* Volume, FDynamicMesh3& Mesh, const FVolumeToMeshOptions& Options);
+
+}}//end namespace UE::Conversion
\ No newline at end of file
diff --git a/Engine/Plugins/Experimental/MeshModelingToolset/Source/MeshModelingTools/Public/VolumeToMeshTool.h b/Engine/Plugins/Experimental/MeshModelingToolset/Source/MeshModelingTools/Public/VolumeToMeshTool.h
new file mode 100644
--- /dev/null
+++ b/Engine/Plugins/Experimental/MeshModelingToolset/Source/MeshModelingTools/Public/VolumeToMeshTool.h
@@ -0,0 +1,115 @@
+// Copyright Epic Games, Inc. All Rights Reserved.
+
+#pragma once
+
+#include "CoreMinimal.h"
+#include "GameFramework/Volume.h"
+#include "SingleSelectionTool.h"
+#include "InteractiveToolBuilder.h"
+#include "DynamicMesh3.h"
+#include "PreviewMesh.h"
+#include "Drawing/LineSetComponent.h"
+#include "VolumeToMeshTool.generated.h"
+
+class IAssetGenerationAPI;
+
+/**
+ *
+ */
+UCLASS()
+class MESHMODELINGTOOLS_API UVolumeToMeshToolBuilder : public UInteractiveToolBuilder
+{
+	GENERATED_BODY()
+public:
+	IAssetGenerationAPI* AssetAPI = nullptr;
+
+	virtual bool CanBuildTool(const FToolBuilderState& SceneState) const override;
+	virtual UInteractiveTool* BuildTool(const FToolBuilderState& SceneState) const override;
+};
+
+
+
+
+UENUM()
+enum class EVolumeToMeshMode
+{
+	/** Convert based on Angle Tolerance between Face Normals */
+	TriangulatePolygons,
+	/** Create PolyGroups based on UV Islands */
+	MinimalPolygons
+};
+
+
+
+UCLASS()
+class MESHMODELINGTOOLS_API UVolumeToMeshToolProperties : public UInteractiveToolPropertySet
+{
+	GENERATED_BODY()
+
+public:
+	UPROPERTY(EditAnywhere, Category = Options)
+	bool bWeldEdges = true;
+
+	UPROPERTY(EditAnywhere, Category = Options)
+	bool bAutoRepair = true;
+
+	UPROPERTY(EditAnywhere, Category = Options)
+	bool bOptimizeMesh = true;
+
+	UPROPERTY(EditAnywhere, Category = Options)
+	bool bShowWireframe = true;
+};
+
+
+
+/**
+ *
+ */
+UCLASS()
+class MESHMODELINGTOOLS_API UVolumeToMeshTool : public UInteractiveTool
+{
+	GENERATED_BODY()
+
+public:
+	UVolumeToMeshTool();
+
+	virtual void SetWorld(UWorld* World) { this->TargetWorld = World; }
+	virtual void SetAssetAPI(IAssetGenerationAPI* InAssetAPI) { this->AssetAPI = InAssetAPI; }
+	virtual void SetSelection(AVolume* Volume);
+
+	virtual void Setup() override;
+	virtual void Shutdown(EToolShutdownType ShutdownType) override;
+
+	virtual void Render(IToolsContextRenderAPI* RenderAPI) override;
+	virtual void OnTick(float DeltaTime) override;
+
+	virtual bool HasCancel() const override { return true; }
+	virtual bool HasAccept() const override;
+	virtual bool CanAccept() const override;
+
+protected:
+	UPROPERTY()
+	UVolumeToMeshToolProperties* Settings;
+
+	UPROPERTY()
+	UPreviewMesh* PreviewMesh;
+
+	UPROPERTY()
+	TLazyObjectPtr<AVolume> TargetVolume;
+
+	UPROPERTY()
+	ULineSetComponent* VolumeEdgesSet;
+
+protected:
+	UWorld* TargetWorld = nullptr;
+	IAssetGenerationAPI* AssetAPI = nullptr;
+
+	FDynamicMesh3 CurrentMesh;
+
+	void RecalculateMesh();
+
+	void UpdateLineSet();
+
+	bool bResultValid = false;
+
+};
diff --git a/Engine/Plugins/Experimental/MeshModelingToolset/Source/MeshModelingToolsEditorOnly/Private/ConversionUtils/DynamicMeshToVolume.cpp b/Engine/Plugins/Experimental/MeshModelingToolset/Source/MeshModelingToolsEditorOnly/Private/ConversionUtils/DynamicMeshToVolume.cpp
new file mode 100644
--- /dev/null
+++ b/Engine/Plugins/Experimental/MeshModelingToolset/Source/MeshModelingToolsEditorOnly/Private/ConversionUtils/DynamicMeshToVolume.cpp
@@ -0,0 +1,208 @@
+// Copyright Epic Games, Inc. All Rights Reserved.
+
+#include "ConversionUtils/DynamicMeshToVolume.h"
+
+#include "BSPOps.h"
+#include "DynamicMesh3.h"
+#include "Engine/Classes/Engine/Polys.h"
+#include "GameFramework/Volume.h"
+#include "MeshNormals.h"
+#include "MeshRegionBoundaryLoops.h"
+#include "Model.h"
+#include "Selections/MeshConnectedComponents.h"
+
+
+namespace UE {
+namespace Conversion {
+
+void DynamicMeshToVolume(const FDynamicMesh3& InputMesh, AVolume* TargetVolume)
+{
+	TArray<FDynamicMeshFace> Faces;
+	GetPolygonFaces(InputMesh, Faces);
+	DynamicMeshToVolume(InputMesh, Faces, TargetVolume);
+}
+
+void DynamicMeshToVolume(const FDynamicMesh3& InputMesh, TArray<FDynamicMeshFace>& Faces, AVolume* TargetVolume)
+{
+	check(TargetVolume->Brush);
+
+	UModel* Model = TargetVolume->Brush;
+
+	Model->Modify();
+
+	Model->Initialize(TargetVolume);
+	UPolys* Polys = Model->Polys;
+
+	for (FDynamicMeshFace& Face : Faces)
+	{
+		int32 NumVertices = Face.BoundaryLoop.Num();
+		FVector Normal = (FVector)Face.Plane.Z();
+		FVector U = (FVector)Face.Plane.X();
+		FVector V = (FVector)Face.Plane.Y();
+		FPlane FacePlane = Face.Plane.ToFPlane();
+
+		// create FPoly. This is Editor-only and I'm not entirely sure we need it?
+		int32 PolyIndex = Polys->Element.Num();
+		FPoly NewPoly;
+		NewPoly.Base = (FVector)Face.BoundaryLoop[0];
+		NewPoly.Normal = Normal;
+		NewPoly.TextureU = U;
+		NewPoly.TextureV = V;
+		NewPoly.Vertices.SetNum(NumVertices);
+		for (int32 k = 0; k < NumVertices; ++k)
+		{
+			NewPoly.Vertices[k] = (FVector)Face.BoundaryLoop[k];
+		}
+		NewPoly.PolyFlags = 0;
+		NewPoly.iLink = NewPoly.iLinkSurf = NewPoly.iBrushPoly = -1;
+		NewPoly.SmoothingMask = 0;
+		Polys->Element.Add(NewPoly);
+
+		/*
+
+				// create points for this face in UModel::Points
+				// TODO: can we share points between faces?
+				TArray<int32> PointIndices;
+				PointIndices.SetNum(NumVertices);
+				for (int32 k = 0; k < NumVertices; ++k)
+				{
+					int32 NewIdx = Model->Points.Num();
+					Model->Points.Add((FVector)Face.BoundaryLoop[k]);
+					PointIndices[k] = NewIdx;
+				}
+				int32 BasePointIndex = PointIndices[0];
+
+				// create normal for this face in UModel::Vectors along with U and V direction vectors
+				int32 NormalIdx = Model->Vectors.Num();
+				Model->Vectors.Add(Normal);
+				int32 TextureUIdx = Model->Vectors.Num();
+				Model->Vectors.Add(U);
+				int32 TextureVIdx = Model->Vectors.Num();
+				Model->Vectors.Add(V);
+
+				// create FVerts for this face in UModel::Verts
+				int32 iVertPoolStart = Model->Verts.Num();
+				for (int32 k = 0; k < NumVertices; ++k)
+				{
+					FVert NewVert;
+					NewVert.pVertex = PointIndices[k];		// Index of vertex point.
+					NewVert.iSide = INDEX_NONE;				// If shared, index of unique side. Otherwise INDEX_NONE.
+					NewVert.ShadowTexCoord = FVector2D::ZeroVector;			// The vertex's shadow map coordinate.
+					NewVert.BackfaceShadowTexCoord = FVector2D::ZeroVector;	// The vertex's shadow map coordinate for the backface of the node.
+					Model->Verts.Add(NewVert);
+				}
+
+				// create Surf
+
+				int32 SurfIndex = Model->Surfs.Num();
+				FBspSurf NewSurf;
+				NewSurf.Material = nullptr;			// 4 Material.
+				NewSurf.PolyFlags = 0;				// 4 Polygon flags.
+				NewSurf.pBase = BasePointIndex;		// 4 Polygon & texture base point index (where U,V==0,0).
+				NewSurf.vNormal = NormalIdx;		// 4 Index to polygon normal.
+				NewSurf.vTextureU = TextureUIdx;	// 4 Texture U-vector index.
+				NewSurf.vTextureV = TextureVIdx;	// 4 Texture V-vector index.
+				//NewSurf.iBrushPoly = PolyIndex;		// 4 Editor brush polygon index.
+				NewSurf.iBrushPoly = -1;
+				//NewSurf.Actor = NewVolume;			// 4 Brush actor owning this Bsp surface.
+				NewSurf.Actor = nullptr;
+				NewSurf.Plane = FacePlane;			// 16 The plane this surface lies on.
+				Model->Surfs.Add(NewSurf);
+
+
+				// create nodes for this face in UModel::Nodes
+
+				FBspNode NewNode;
+				NewNode.Plane = FacePlane;					// 16 Plane the node falls into (X, Y, Z, W).
+				NewNode.iVertPool = iVertPoolStart;			// 4  Index of first vertex in vertex pool, =iTerrain if NumVertices==0 and NF_TerrainFront.
+				NewNode.iSurf = SurfIndex;					// 4  Index to surface information.
+				NewNode.iVertexIndex = INDEX_NONE;			// The index of the node's first vertex in the UModel's vertex buffer.
+															// This is initialized by UModel::UpdateVertices()
+				NewNode.NumVertices = NumVertices;			// 1  Number of vertices in node.
+				NewNode.NodeFlags = 0;						// 1  Node flags.
+				Model->Nodes.Add(NewNode);
+		*/
+	}
+
+	// requires ed
+	FBSPOps::csgPrepMovingBrush(TargetVolume);
+
+	// do we need to do this?
+	TargetVolume->MarkPackageDirty();
+}
+
+
+void GetPolygonFaces(const FDynamicMesh3& InputMesh, TArray<FDynamicMeshFace>& Faces)
+{
+	Faces.SetNum(0);
+
+	FMeshNormals Normals(&InputMesh);
+	Normals.ComputeTriangleNormals();
+
+	double PlanarTolerance = FMathf::ZeroTolerance;
+
+	FMeshConnectedComponents Components(&InputMesh);
+	Components.FindConnectedTriangles([&](int32 Triangle0, int32 Triangle1)
+		{
+			FVector3d Origin = InputMesh.GetTriCentroid(Triangle0);
+			FVector3d Normal = Normals[Triangle0];
+
+			FVector3d A, B, C;
+			InputMesh.GetTriVertices(Triangle1, A, B, C);;
+			double DistA = (A - Origin).Dot(Normal);
+			double DistB = (B - Origin).Dot(Normal);
+			double DistC = (C - Origin).Dot(Normal);
+			double MaxDist = FMathd::Max3(FMathd::Abs(DistA), FMathd::Abs(DistB), FMathd::Abs(DistC));
+
+			return MaxDist < PlanarTolerance;
+		});
+
+	for (const FMeshConnectedComponents::FComponent& Component : Components)
+	{
+		FVector3d FaceNormal = Normals[Component.Indices[0]];
+
+		FMeshRegionBoundaryLoops Loops(&InputMesh, Component.Indices);
+		for (const FEdgeLoop& Loop : Loops.Loops)
+		{
+			FDynamicMeshFace Face;
+
+			FVector3d AvgPos(0, 0, 0);
+			for (int32 vid : Loop.Vertices)
+			{
+				FVector3d Position = InputMesh.GetVertex(vid);
+				Face.BoundaryLoop.Add(Position);
+				AvgPos += Position;
+			}
+			AvgPos /= (double)Loop.Vertices.Num();
+			Algo::Reverse(Face.BoundaryLoop);
+
+			Face.Plane = FFrame3d(AvgPos, FaceNormal);
+
+			Faces.Add(Face);
+		}
+	}
+}
+
+void GetTriangleFaces(const FDynamicMesh3& InputMesh, TArray<FDynamicMeshFace>& Faces)
+{
+	Faces.SetNum(0);
+
+	for (int32 tid : InputMesh.TriangleIndicesItr())
+	{
+		FVector3d A, B, C;
+		InputMesh.GetTriVertices(tid, A, B, C);
+		FVector3d Centroid, Normal; double Area;
+		InputMesh.GetTriInfo(tid, Normal, Area, Centroid);
+
+		FDynamicMeshFace Face;
+		Face.Plane = FFrame3d(Centroid, Normal);
+		Face.BoundaryLoop.Add(A);
+		Face.BoundaryLoop.Add(C);
+		Face.BoundaryLoop.Add(B);
+
+		Faces.Add(Face);
+	}
+}
+
+}//end namespace UE::Conversion
+}//end namespace UE
\ No newline at end of file
diff --git a/Engine/Plugins/Experimental/MeshModelingToolset/Source/MeshModelingToolsEditorOnly/Private/ConversionUtils/VolumeComponentTarget.cpp b/Engine/Plugins/Experimental/MeshModelingToolset/Source/MeshModelingToolsEditorOnly/Private/ConversionUtils/VolumeComponentTarget.cpp
--- a/Engine/Plugins/Experimental/MeshModelingToolset/Source/MeshModelingToolsEditorOnly/Private/ConversionUtils/VolumeComponentTarget.cpp
+++ b/Engine/Plugins/Experimental/MeshModelingToolset/Source/MeshModelingToolsEditorOnly/Private/ConversionUtils/VolumeComponentTarget.cpp
@@ -3,6 +3,8 @@
 #include "ConversionUtils/VolumeComponentTarget.h"
 
 #include "Components/BrushComponent.h"
+#include "ConversionUtils/VolumeToDynamicMesh.h"
+#include "ConversionUtils/DynamicMeshToVolume.h"
 #include "DynamicMeshToMeshDescription.h"
 #include "Engine/Brush.h"
 #include "GameFramework/Volume.h"
diff --git a/Engine/Plugins/Experimental/MeshModelingToolset/Source/MeshModelingToolsEditorOnly/Private/ConversionUtils/VolumeDynamicMeshConversion.cpp b/Engine/Plugins/Experimental/MeshModelingToolset/Source/MeshModelingToolsEditorOnly/Private/ConversionUtils/VolumeDynamicMeshConversion.cpp
deleted file mode 100644
diff --git a/Engine/Plugins/Experimental/MeshModelingToolset/Source/MeshModelingToolsEditorOnly/Private/VolumeToMeshTool.cpp b/Engine/Plugins/Experimental/MeshModelingToolset/Source/MeshModelingToolsEditorOnly/Private/VolumeToMeshTool.cpp
deleted file mode 100644
diff --git a/Engine/Plugins/Experimental/MeshModelingToolset/Source/MeshModelingToolsEditorOnly/Public/ConversionUtils/DynamicMeshToVolume.h b/Engine/Plugins/Experimental/MeshModelingToolset/Source/MeshModelingToolsEditorOnly/Public/ConversionUtils/DynamicMeshToVolume.h
new file mode 100644
--- /dev/null
+++ b/Engine/Plugins/Experimental/MeshModelingToolset/Source/MeshModelingToolsEditorOnly/Public/ConversionUtils/DynamicMeshToVolume.h
@@ -0,0 +1,43 @@
+// Copyright Epic Games, Inc. All Rights Reserved.
+
+#pragma once
+
+#include "CoreMinimal.h"
+#include "FrameTypes.h"
+#include "VectorTypes.h"
+
+// NOTE: The current implementation of DynamicMeshToVolume is editor-only,
+// and is therefore split here from VolumeToDynamicMesh. If it ever becomes
+// safe for runtime, we should move it to the same place (or combine it in
+// one file).
+
+class AVolume;
+class FDynamicMesh3;
+
+namespace UE {
+namespace Conversion { 
+
+struct FDynamicMeshFace
+{
+	FFrame3d Plane;
+	TArray<FVector3d> BoundaryLoop;
+};
+
+/**
+ * Gets an array of face objects that can be used to convert a dynamic mesh to a volume. This version tries to
+ * merge coplanar triangles into polygons.
+ */
+void MESHMODELINGTOOLSEDITORONLY_API GetPolygonFaces(const FDynamicMesh3& InputMesh, TArray<FDynamicMeshFace>& Faces);
+/**
+ * Gets an array of face objects that can be used to convert a dynamic mesh to a volume. This version makes
+ * each triangle its own face.
+ */
+void MESHMODELINGTOOLSEDITORONLY_API GetTriangleFaces(const FDynamicMesh3& InputMesh, TArray<FDynamicMeshFace>& Faces);
+
+/**
+ * Converts a dynamic mesh to a volume.
+ */
+void MESHMODELINGTOOLSEDITORONLY_API DynamicMeshToVolume(const FDynamicMesh3& InputMesh, AVolume* TargetVolume);
+void MESHMODELINGTOOLSEDITORONLY_API DynamicMeshToVolume(const FDynamicMesh3& InputMesh, TArray<FDynamicMeshFace>& Faces, AVolume* TargetVolume);
+
+}}//end namespace UE::Conversion
\ No newline at end of file
diff --git a/Engine/Plugins/Experimental/MeshModelingToolset/Source/MeshModelingToolsEditorOnly/Public/ConversionUtils/VolumeComponentTarget.h b/Engine/Plugins/Experimental/MeshModelingToolset/Source/MeshModelingToolsEditorOnly/Public/ConversionUtils/VolumeComponentTarget.h
--- a/Engine/Plugins/Experimental/MeshModelingToolset/Source/MeshModelingToolsEditorOnly/Public/ConversionUtils/VolumeComponentTarget.h
+++ b/Engine/Plugins/Experimental/MeshModelingToolset/Source/MeshModelingToolsEditorOnly/Public/ConversionUtils/VolumeComponentTarget.h
@@ -5,7 +5,7 @@
 #include "CoreMinimal.h"
 
 #include "ComponentSourceInterfaces.h"
-#include "ConversionUtils/VolumeDynamicMeshConversion.h"
+#include "ConversionUtils/VolumeToDynamicMesh.h"
 #include "MeshDescription.h"
 
 class MESHMODELINGTOOLSEDITORONLY_API FVolumeComponentTargetFactory : public FComponentTargetFactory
diff --git a/Engine/Plugins/Experimental/MeshModelingToolset/Source/MeshModelingToolsEditorOnly/Public/ConversionUtils/VolumeDynamicMeshConversion.h b/Engine/Plugins/Experimental/MeshModelingToolset/Source/MeshModelingToolsEditorOnly/Public/ConversionUtils/VolumeDynamicMeshConversion.h
deleted file mode 100644
diff --git a/Engine/Plugins/Experimental/MeshModelingToolset/Source/MeshModelingToolsEditorOnly/Public/MeshToVolumeTool.h b/Engine/Plugins/Experimental/MeshModelingToolset/Source/MeshModelingToolsEditorOnly/Public/MeshToVolumeTool.h
--- a/Engine/Plugins/Experimental/MeshModelingToolset/Source/MeshModelingToolsEditorOnly/Public/MeshToVolumeTool.h
+++ b/Engine/Plugins/Experimental/MeshModelingToolset/Source/MeshModelingToolsEditorOnly/Public/MeshToVolumeTool.h
@@ -4,7 +4,7 @@
 
 #include "CoreMinimal.h"
 
-#include "ConversionUtils/VolumeDynamicMeshConversion.h"
+#include "ConversionUtils/DynamicMeshToVolume.h"
 #include "GameFramework/Volume.h"
 #include "Engine/Classes/Engine/BlockingVolume.h"
 #include "SingleSelectionTool.h"
@@ -66,7 +66,10 @@ public:
 
 
 /**
+ * Converts a mesh to a volume.
  *
+ * Note: If ConversionUtils/DynamicMeshToVolume is rewritten to be safe for runtime, this
+ * tool can be moved out of the editor-only section and put with VolumeToMeshTool.
  */
 UCLASS()
 class MESHMODELINGTOOLSEDITORONLY_API UMeshToVolumeTool : public USingleSelectionTool
diff --git a/Engine/Plugins/Experimental/MeshModelingToolset/Source/MeshModelingToolsEditorOnly/Public/VolumeToMeshTool.h b/Engine/Plugins/Experimental/MeshModelingToolset/Source/MeshModelingToolsEditorOnly/Public/VolumeToMeshTool.h
deleted file mode 100644
