Add bilinear and Loop subdivision options to subdivide poly tool.
Add option to toggle cage rendering.

diff --git a/Engine/Plugins/Experimental/MeshModelingToolset/Source/MeshModelingToolsEditorOnly/Private/SubdividePoly.cpp b/Engine/Plugins/Experimental/MeshModelingToolset/Source/MeshModelingToolsEditorOnly/Private/SubdividePoly.cpp
--- a/Engine/Plugins/Experimental/MeshModelingToolset/Source/MeshModelingToolsEditorOnly/Private/SubdividePoly.cpp
+++ b/Engine/Plugins/Experimental/MeshModelingToolset/Source/MeshModelingToolsEditorOnly/Private/SubdividePoly.cpp
@@ -114,40 +114,49 @@ namespace SubdividePolyLocal
 		}
 	}
 
+	SubdVertex GetVertexInfo(int32 VertexID, const FDynamicMesh3& Mesh, bool bGetNormals)
+	{
+		SubdVertex Vertex;
+		constexpr bool bWantColors = false;
+		constexpr bool bWantVertexUVs = false;
+		constexpr bool bWantNormals = false;
+
+		Mesh.GetVertex(VertexID, Vertex.VertexInfo, bWantNormals, bWantColors, bWantVertexUVs);
+
+		if (bGetNormals)
+		{
+			if (Mesh.HasAttributes() && (Mesh.Attributes()->NumNormalLayers() > 0))
+			{
+				Vertex.VertexInfo.Normal = GetAverageVertexNormalFromOverlay(Mesh.Attributes()->PrimaryNormals(),
+																			 VertexID);
+				Vertex.VertexInfo.bHaveN = true;
+			}
+		}
+
+		return Vertex;
+	}
+
 	// Treat the given FGroupTopology as a polygonal mesh, and get its vertices
 	void GetGroupPolyMeshVertices(const FDynamicMesh3& Mesh,
 								  const FGroupTopology& Topology,
-								  bool bGetVertexNormals,
+								  bool bGetNormals,
 								  TArray<SubdVertex>& OutVertices )
 	{
 		for (const FGroupTopology::FCorner& Corner : Topology.Corners)
 		{
-			SubdVertex Vertex;
-			bool bWantNormals = bGetVertexNormals && Mesh.HasVertexNormals();
-			constexpr bool bWantColors = false;
-			constexpr bool bWantVertexUVs = false;
-			Mesh.GetVertex(Corner.VertexID, Vertex.VertexInfo, bWantNormals, bWantColors, bWantVertexUVs);
-
-			if (bGetVertexNormals && !Vertex.VertexInfo.bHaveN)
-			{
-				// We want vertex normals but the mesh doesn't have them. See if we can get them from an overlay.
-				if (Mesh.HasAttributes() && (Mesh.Attributes()->NumNormalLayers() > 0))
-				{
-					Vertex.VertexInfo.Normal = GetAverageVertexNormalFromOverlay(Mesh.Attributes()->PrimaryNormals(), 
-																				 Corner.VertexID);
-					Vertex.VertexInfo.bHaveN = true;
-				}
-				else
-				{
-					Vertex.VertexInfo.Normal = { 0,1,0 };
-					Vertex.VertexInfo.bHaveN = true;
-				}
-			}
-
-			OutVertices.Add(Vertex);
+			OutVertices.Add(GetVertexInfo(Corner.VertexID, Mesh, bGetNormals));
 		}
 	}
 
+	void GetAllMeshVertices(const FDynamicMesh3& Mesh,
+							bool bGetNormals,
+							TArray<SubdVertex>& OutVertices)
+	{
+		for (int32 VertexID : Mesh.VertexIndicesItr())
+		{
+			OutVertices.Add(GetVertexInfo(VertexID, Mesh, bGetNormals));
+		}
+	}
 
 	// Given a group and a vertex ID, return:
 	// - a triangle in the group with one of its corners equal to vertex ID
@@ -216,6 +225,22 @@ namespace SubdividePolyLocal
 	}
 
 
+	bool GetMeshUVs(const FDynamicMesh3& Mesh,
+					const FDynamicMeshUVOverlay* UVOverlay,
+					TArray<SubdUVVertex>& OutUVs)		// for each triangle, all the uvs
+	{
+		for (int32 TriangleID : Mesh.TriangleIndicesItr())
+		{
+			TStaticArray<FVector2f, 3> TriangleUVs;
+			UVOverlay->GetTriElements(TriangleID, TriangleUVs[0], TriangleUVs[1], TriangleUVs[2]);
+			OutUVs.Add(SubdUVVertex{ TriangleUVs[0] });
+			OutUVs.Add(SubdUVVertex{ TriangleUVs[1] });
+			OutUVs.Add(SubdUVVertex{ TriangleUVs[2] });
+		}
+
+		return true;
+	}
+
 	void InitializeOverlayToFaceVertexUVs(FDynamicMeshUVOverlay* UVOverlay, const TArray<FVector2f>& UVs)
 	{
 		const FDynamicMesh3* Mesh = UVOverlay->GetParentMesh();
@@ -249,8 +274,8 @@ public:
 };
 
 FSubdividePoly::FSubdividePoly(const FGroupTopology& InTopology,
-										 const FDynamicMesh3& InOriginalMesh,
-										 int InLevel) :
+							   const FDynamicMesh3& InOriginalMesh,
+							   int InLevel) :
 	GroupTopology(InTopology)
 	, OriginalMesh(InOriginalMesh)
 	, Level(InLevel)
@@ -271,54 +296,123 @@ FSubdividePoly::~FSubdividePoly()
 
 bool FSubdividePoly::ComputeTopologySubdivision()
 {
+	if (Level < 1)
+	{
+		return false;
+	}
+
 	TArray<int> BoundaryVertsPerFace;
 	TArray<int> NumVertsPerFace;
 	int TotalNumVertices = 0;
+	OpenSubdiv::Far::TopologyDescriptor::FVarChannel UVChannel;
+	TArray<int> UVBuffer;
 
-	for (const FGroupTopology::FGroup& Group : GroupTopology.Groups)
+	auto DescriptorFromTriangleMesh = [this, &NumVertsPerFace, &BoundaryVertsPerFace, &UVChannel, &UVBuffer]
+	(OpenSubdiv::Far::TopologyDescriptor& Descriptor)
 	{
-		if (Group.Boundaries.Num() != 1)
+		int TotalNumVertices = 0;
+
+		for (auto TriangleID : OriginalMesh.TriangleIndicesItr())
 		{
-			return false;
+			FIndex3i TriangleVertices = OriginalMesh.GetTriangle(TriangleID);
+			NumVertsPerFace.Add(3);
+			TotalNumVertices += 3;
+			BoundaryVertsPerFace.Add(TriangleVertices[0]);
+			BoundaryVertsPerFace.Add(TriangleVertices[1]);
+			BoundaryVertsPerFace.Add(TriangleVertices[2]);
 		}
 
-		const FGroupTopology::FGroupBoundary& Boundary = Group.Boundaries[0];
-		if (Boundary.GroupEdges.Num() < 2)
+		Descriptor.numVertices = TotalNumVertices;
+		Descriptor.numFaces = OriginalMesh.TriangleCount();
+		Descriptor.numVertsPerFace = NumVertsPerFace.GetData();
+		Descriptor.vertIndicesPerFace = BoundaryVertsPerFace.GetData();
+
+		if (UVComputationMethod == ESubdivisionOutputUVs::Interpolated)
 		{
-			return false;
+			UVChannel.numValues = BoundaryVertsPerFace.Num();
+			UVBuffer.SetNum(UVChannel.numValues);
+			for (int i = 0; i < UVChannel.numValues; ++i)
+			{
+				UVBuffer[i] = i;
+			}
+			UVChannel.valueIndices = UVBuffer.GetData();
+			Descriptor.numFVarChannels = 1;
+			Descriptor.fvarChannels = &UVChannel;
+		}
+		else
+		{
+			Descriptor.numFVarChannels = 0;
 		}
 
-		TArray<int> Corners;
-		SubdividePolyLocal::GetBoundaryCorners(Boundary, GroupTopology, Corners);
+		return true;
+	};
 
-		NumVertsPerFace.Add(Corners.Num());
-		TotalNumVertices += Corners.Num();
-		BoundaryVertsPerFace.Append(Corners);
-	}
+	auto DescriptorFromGroupTopology = [this, &NumVertsPerFace, &BoundaryVertsPerFace, &UVChannel, &UVBuffer]
+	(OpenSubdiv::Far::TopologyDescriptor& Descriptor)
+	{
+		int TotalNumVertices = 0;
+
+		for (const FGroupTopology::FGroup& Group : GroupTopology.Groups)
+		{
+			if (Group.Boundaries.Num() != 1)
+			{
+				return false;
+			}
+
+			const FGroupTopology::FGroupBoundary& Boundary = Group.Boundaries[0];
+			if (Boundary.GroupEdges.Num() < 2)
+			{
+				return false;
+			}
+
+			TArray<int> Corners;
+			SubdividePolyLocal::GetBoundaryCorners(Boundary, GroupTopology, Corners);
+
+			NumVertsPerFace.Add(Corners.Num());
+			TotalNumVertices += Corners.Num();
+			BoundaryVertsPerFace.Append(Corners);
+		}
+
+		Descriptor.numVertices = TotalNumVertices;
+		Descriptor.numFaces = GroupTopology.Groups.Num();
+		Descriptor.numVertsPerFace = NumVertsPerFace.GetData();
+		Descriptor.vertIndicesPerFace = BoundaryVertsPerFace.GetData();
+
+		if (UVComputationMethod == ESubdivisionOutputUVs::Interpolated)
+		{
+			UVChannel.numValues = BoundaryVertsPerFace.Num();
+			UVBuffer.SetNum(UVChannel.numValues);
+			for (int i = 0; i < UVChannel.numValues; ++i)
+			{
+				UVBuffer[i] = i;
+			}
+			UVChannel.valueIndices = UVBuffer.GetData();
+			Descriptor.numFVarChannels = 1;
+			Descriptor.fvarChannels = &UVChannel;
+		}
+		else
+		{
+			Descriptor.numFVarChannels = 0;
+		}
+
+		return true;
+	};
 
 	OpenSubdiv::Far::TopologyDescriptor Descriptor;
-	Descriptor.numVertices = TotalNumVertices;
-	Descriptor.numFaces = GroupTopology.Groups.Num();
-	Descriptor.numVertsPerFace = NumVertsPerFace.GetData();		// taking pointer to local memory seems to be ok?
-	Descriptor.vertIndicesPerFace = BoundaryVertsPerFace.GetData();
 
-	OpenSubdiv::Far::TopologyDescriptor::FVarChannel UVChannel;
-	TArray<int> Buffer;
-	if (UVComputationMethod == ESubdivisionOutputUVs::Interpolated)
+	if (SubdivisionScheme == ESubdivisionScheme::Loop)
 	{
-		UVChannel.numValues = BoundaryVertsPerFace.Num();
-		Buffer.SetNum(UVChannel.numValues);
-		for (int i = 0; i < UVChannel.numValues; ++i)
+		if (!DescriptorFromTriangleMesh(Descriptor))
 		{
-			Buffer[i] = i;
+			return false;
 		}
-		UVChannel.valueIndices = Buffer.GetData();
-		Descriptor.numFVarChannels = 1;
-		Descriptor.fvarChannels = &UVChannel;		// taking a pointer to local memory seems to be ok?
 	}
 	else
 	{
-		Descriptor.numFVarChannels = 0;
+		if (!DescriptorFromGroupTopology(Descriptor))
+		{
+			return false;
+		}
 	}
 
 	using RefinerFactory = OpenSubdiv::Far::TopologyRefinerFactory<OpenSubdiv::Far::TopologyDescriptor>;
@@ -326,6 +420,19 @@ bool FSubdividePoly::ComputeTopologySubdivision()
 
 	RefinerOptions.schemeOptions.SetVtxBoundaryInterpolation(OpenSubdiv::Sdc::Options::VTX_BOUNDARY_EDGE_AND_CORNER);
 
+	switch (SubdivisionScheme)
+	{
+	case ESubdivisionScheme::Bilinear:
+		RefinerOptions.schemeType = OpenSubdiv::Sdc::SchemeType::SCHEME_BILINEAR;
+		break;
+	case ESubdivisionScheme::CatmullClark:
+		RefinerOptions.schemeType = OpenSubdiv::Sdc::SchemeType::SCHEME_CATMARK;
+		break;
+	case ESubdivisionScheme::Loop:
+		RefinerOptions.schemeType = OpenSubdiv::Sdc::SchemeType::SCHEME_LOOP;
+		break;
+	}
+
 	Refiner->TopologyRefiner = RefinerFactory::Create(Descriptor, RefinerOptions);
 
 	if (Refiner->TopologyRefiner == nullptr)
@@ -341,6 +448,11 @@ bool FSubdividePoly::ComputeTopologySubdivision()
 
 bool FSubdividePoly::ComputeSubdividedMesh(FDynamicMesh3& OutMesh)
 {
+	if (Level < 1)
+	{
+		return false;
+	}
+
 	if (!Refiner || !(Refiner->TopologyRefiner))
 	{
 		return false;
@@ -353,13 +465,21 @@ bool FSubdividePoly::ComputeSubdividedMesh(FDynamicMesh3& OutMesh)
 	// 
 	TArray<SubdividePolyLocal::SubdVertex> SourceVertices;
 	bool bInterpolateVertexNormals = (NormalComputationMethod == ESubdivisionOutputNormals::Interpolated);
-	GetGroupPolyMeshVertices(OriginalMesh, GroupTopology, bInterpolateVertexNormals, SourceVertices);
+	if (SubdivisionScheme == ESubdivisionScheme::Loop)
+	{
+		GetAllMeshVertices(OriginalMesh, bInterpolateVertexNormals, SourceVertices);
+	}
+	else
+	{
+		GetGroupPolyMeshVertices(OriginalMesh, GroupTopology, bInterpolateVertexNormals, SourceVertices);
+	}
 
 	TArray<SubdividePolyLocal::SubdVertex> RefinedVertices;
 	for (int CurrentLevel = 1; CurrentLevel <= Level; ++CurrentLevel)
 	{
 		// TODO: Don't keep resizing -- preallocate one big buffer and move through it
-		RefinedVertices.SetNumUninitialized(Refiner->TopologyRefiner->GetLevel(CurrentLevel).GetNumVertices());
+		int NumVertices = Refiner->TopologyRefiner->GetLevel(CurrentLevel).GetNumVertices();
+		RefinedVertices.SetNumUninitialized(NumVertices);
 		SubdividePolyLocal::SubdVertex* s = SourceVertices.GetData();
 		SubdividePolyLocal::SubdVertex* d = RefinedVertices.GetData();
 		Interpolator.Interpolate(CurrentLevel, s, d);
@@ -370,10 +490,21 @@ bool FSubdividePoly::ComputeSubdividedMesh(FDynamicMesh3& OutMesh)
 	// Interpolate face group IDs
 	// 
 	TArray<int> SourceGroupIDs;
-	for (const FGroupTopology::FGroup& Group : GroupTopology.Groups)
+	if (SubdivisionScheme == ESubdivisionScheme::Loop)
+	{
+		for ( int TriangleID : OriginalMesh.TriangleIndicesItr() )
+		{
+			SourceGroupIDs.Add(OriginalMesh.GetTriangleGroup(TriangleID));
+		}
+	}
+	else
 	{
-		SourceGroupIDs.Add(Group.GroupID);
+		for (const FGroupTopology::FGroup& Group : GroupTopology.Groups)
+		{
+			SourceGroupIDs.Add(Group.GroupID);
+		}
 	}
+	check(SourceGroupIDs.Num() == Refiner->TopologyRefiner->GetLevel(0).GetNumFaces());
 
 	TArray<int> RefinedGroupIDs;
 	for (int CurrentLevel = 1; CurrentLevel <= Level; ++CurrentLevel)
@@ -400,10 +531,19 @@ bool FSubdividePoly::ComputeSubdividedMesh(FDynamicMesh3& OutMesh)
 		}
 
 		TArray<SubdividePolyLocal::SubdUVVertex> SourceUVs;
-		bool bGetUVsOK = SubdividePolyLocal::GetGroupPolyMeshUVs(GroupTopology,
-																	  OriginalMesh,
-																	  OriginalMesh.Attributes()->PrimaryUV(), 
-																	  SourceUVs);
+		bool bGetUVsOK = false;
+
+		if (SubdivisionScheme == ESubdivisionScheme::Loop)
+		{
+			bGetUVsOK = GetMeshUVs(OriginalMesh, OriginalMesh.Attributes()->PrimaryUV(), SourceUVs);
+		}
+		else
+		{
+			bGetUVsOK = GetGroupPolyMeshUVs(GroupTopology,
+											OriginalMesh,
+											OriginalMesh.Attributes()->PrimaryUV(),
+											SourceUVs);
+		}
 
 		if (!bGetUVsOK)
 		{
@@ -460,26 +600,40 @@ bool FSubdividePoly::ComputeSubdividedMesh(FDynamicMesh3& OutMesh)
 		OpenSubdiv::Far::ConstIndexArray Face = FinalLevel.GetFaceVertices(FaceID);
 		NumFaceVertices += Face.size();
 
-		check(Face.size() == 4);
+		if (Face.size() == 4)
+		{
+			FIndex3i TriA{ Face[0], Face[1], Face[3] };
+			FIndex3i TriB{ Face[2], Face[3], Face[1] };
+			OutMesh.AppendTriangle(TriA, GroupID);
+			OutMesh.AppendTriangle(TriB, GroupID);
 
-		FIndex3i TriA{ Face[0], Face[1], Face[3] };
-		FIndex3i TriB{ Face[2], Face[3], Face[1] };
-		OutMesh.AppendTriangle(TriA, GroupID);
-		OutMesh.AppendTriangle(TriB, GroupID);
+			if (UVComputationMethod == ESubdivisionOutputUVs::Interpolated)
+			{
+				OpenSubdiv::Far::ConstIndexArray FaceUVIndices = FinalLevel.GetFaceFVarValues(FaceID);
 
-		if (UVComputationMethod == ESubdivisionOutputUVs::Interpolated)
-		{
-			OpenSubdiv::Far::ConstIndexArray FaceUVIndices = FinalLevel.GetFaceFVarValues(FaceID);
+				FIndex3i UVTriA{ FaceUVIndices[0], FaceUVIndices[1], FaceUVIndices[3] };
+				TriangleUVs.Add(RefinedUVs[UVTriA[0]].VertexUV);
+				TriangleUVs.Add(RefinedUVs[UVTriA[1]].VertexUV);
+				TriangleUVs.Add(RefinedUVs[UVTriA[2]].VertexUV);
 
-			FIndex3i UVTriA{ FaceUVIndices[0], FaceUVIndices[1], FaceUVIndices[3] };
-			TriangleUVs.Add(RefinedUVs[UVTriA[0]].VertexUV);
-			TriangleUVs.Add(RefinedUVs[UVTriA[1]].VertexUV);
-			TriangleUVs.Add(RefinedUVs[UVTriA[2]].VertexUV);
+				FIndex3i UVTriB{ FaceUVIndices[2], FaceUVIndices[3], FaceUVIndices[1] };
+				TriangleUVs.Add(RefinedUVs[UVTriB[0]].VertexUV);
+				TriangleUVs.Add(RefinedUVs[UVTriB[1]].VertexUV);
+				TriangleUVs.Add(RefinedUVs[UVTriB[2]].VertexUV);
+			}
+		}
+		else
+		{
+			check(Face.size() == 3);
+			OutMesh.AppendTriangle(FIndex3i{ Face[0], Face[1], Face[2] }, GroupID);
 
-			FIndex3i UVTriB{ FaceUVIndices[2], FaceUVIndices[3], FaceUVIndices[1] };
-			TriangleUVs.Add(RefinedUVs[UVTriB[0]].VertexUV);
-			TriangleUVs.Add(RefinedUVs[UVTriB[1]].VertexUV);
-			TriangleUVs.Add(RefinedUVs[UVTriB[2]].VertexUV);
+			if (UVComputationMethod == ESubdivisionOutputUVs::Interpolated)
+			{
+				OpenSubdiv::Far::ConstIndexArray FaceUVIndices = FinalLevel.GetFaceFVarValues(FaceID);
+				TriangleUVs.Add(RefinedUVs[FaceUVIndices[0]].VertexUV);
+				TriangleUVs.Add(RefinedUVs[FaceUVIndices[1]].VertexUV);
+				TriangleUVs.Add(RefinedUVs[FaceUVIndices[2]].VertexUV);
+			}
 		}
 	}
 
@@ -494,6 +648,15 @@ bool FSubdividePoly::ComputeSubdividedMesh(FDynamicMesh3& OutMesh)
 		SubdividePolyLocal::InitializeOverlayToFaceVertexUVs(OutMesh.Attributes()->PrimaryUV(), TriangleUVs);
 	}
 
+	// Remove any vertices that are not referenced by a face
+	for (int Vid = 0; Vid < OutMesh.MaxVertexID(); ++Vid)
+	{
+		if (OutMesh.IsVertex(Vid) && !OutMesh.IsReferencedVertex(Vid))
+		{
+			OutMesh.RemoveVertex(Vid, false, false);
+		}
+	}
+
 	return true;
 }
 
diff --git a/Engine/Plugins/Experimental/MeshModelingToolset/Source/MeshModelingToolsEditorOnly/Private/SubdividePolyTool.cpp b/Engine/Plugins/Experimental/MeshModelingToolset/Source/MeshModelingToolsEditorOnly/Private/SubdividePolyTool.cpp
--- a/Engine/Plugins/Experimental/MeshModelingToolset/Source/MeshModelingToolsEditorOnly/Private/SubdividePolyTool.cpp
+++ b/Engine/Plugins/Experimental/MeshModelingToolset/Source/MeshModelingToolsEditorOnly/Private/SubdividePolyTool.cpp
@@ -19,17 +19,18 @@ class SubdivPostProcessor : public IRenderMeshPostProcessor
 public:
 
 	SubdivPostProcessor(int InSubdivisionLevel,
+						ESubdivisionScheme InSubdivisionScheme,
 						ESubdivisionOutputNormals InNormalComputationMethod,
 						ESubdivisionOutputUVs InUVComputationMethod) :
 		SubdivisionLevel(InSubdivisionLevel),
+		SubdivisionScheme(InSubdivisionScheme),
 		NormalComputationMethod(InNormalComputationMethod),
 		UVComputationMethod(InUVComputationMethod)
 	{}
 
 	int SubdivisionLevel = 3;
-
+	ESubdivisionScheme SubdivisionScheme = ESubdivisionScheme::CatmullClark;
 	ESubdivisionOutputNormals NormalComputationMethod = ESubdivisionOutputNormals::Generated;
-
 	ESubdivisionOutputUVs UVComputationMethod = ESubdivisionOutputUVs::Interpolated;
 
 	void ProcessMesh(const FDynamicMesh3& Mesh, FDynamicMesh3& OutRenderMesh) final
@@ -37,6 +38,7 @@ public:
 		constexpr bool bAutoCompute = true;
 		FGroupTopology Topo(&Mesh, bAutoCompute);
 		FSubdividePoly Subd(Topo, Mesh, SubdivisionLevel);
+		Subd.SubdivisionScheme = SubdivisionScheme;
 		Subd.NormalComputationMethod = NormalComputationMethod;
 		Subd.UVComputationMethod = UVComputationMethod;
 
@@ -168,9 +170,12 @@ void USubdividePolyTool::Setup()
 		return;
 	}
 
+	check(Properties->SubdivisionLevel >= 1);	// Should be enforced by UPROPERTY meta tags
+
 	PreviewDynamicMeshComponent->SetRenderMeshPostProcessor(MakeUnique<SubdivPostProcessor>(Properties->SubdivisionLevel,
-																					 Properties->NormalComputationMethod,
-																					 Properties->UVComputationMethod));
+																							Properties->SubdivisionScheme,
+																							Properties->NormalComputationMethod,
+																							Properties->UVComputationMethod));
 
 	// Use the input mesh's material on the preview
 	FComponentMaterialSet MaterialSet;
@@ -192,8 +197,9 @@ void USubdividePolyTool::Setup()
 	{
 		USimpleDynamicMeshComponent* PreviewDynamicMeshComponent = (USimpleDynamicMeshComponent*)PreviewMesh->GetRootComponent();
 		PreviewDynamicMeshComponent->SetRenderMeshPostProcessor(MakeUnique<SubdivPostProcessor>(Properties->SubdivisionLevel,
-																							Properties->NormalComputationMethod,
-																							Properties->UVComputationMethod));
+																								Properties->SubdivisionScheme,
+																								Properties->NormalComputationMethod,
+																								Properties->UVComputationMethod));
 		PreviewDynamicMeshComponent->NotifyMeshUpdated();
 	};
 
@@ -202,6 +208,11 @@ void USubdividePolyTool::Setup()
 	{
 		RebuildMeshPostProcessor();
 	});
+	Properties->WatchProperty(Properties->SubdivisionScheme, [this, RebuildMeshPostProcessor](ESubdivisionScheme)
+	{
+		RebuildMeshPostProcessor();
+		bPreviewGeometryNeedsUpdate = true;		// Switch from rendering poly cage to all triangle edges
+	});
 	Properties->WatchProperty(Properties->NormalComputationMethod, [this, RebuildMeshPostProcessor](ESubdivisionOutputNormals)
 	{
 		RebuildMeshPostProcessor();
@@ -234,11 +245,13 @@ void USubdividePolyTool::Setup()
 
 	Properties->WatchProperty(Properties->bRenderGroups, RenderGroupsChanged);
 
-	ComponentTarget->SetOwnerVisibility(false);
-
 	// Render with polygroup colors
 	RenderGroupsChanged(Properties->bRenderGroups);
 
+	Properties->WatchProperty(Properties->bRenderCage, [this](bool bNewRenderCage)
+	{
+		bPreviewGeometryNeedsUpdate = true;
+	});
 
 	PreviewGeometry = NewObject<UPreviewGeometry>(this);
 	PreviewGeometry->CreateInWorld(ComponentTarget->GetOwnerActor()->GetWorld(), ComponentTarget->GetWorldTransform());
@@ -247,35 +260,71 @@ void USubdividePolyTool::Setup()
 	// regenerate preview geo if mesh changes due to undo/redo/etc
 	PreviewDynamicMeshComponent->OnMeshChanged.AddLambda([this]() { bPreviewGeometryNeedsUpdate = true; });
 
+	ComponentTarget->SetOwnerVisibility(false);
 	PreviewMesh->SetVisible(true);
 }
 
 void USubdividePolyTool::CreateOrUpdatePreviewGeometry()
 {
-	FGroupTopology Topology(OriginalMesh.Get(), true);
-	int NumEdges = Topology.Edges.Num();
+	if (!Properties->bRenderCage)
+	{
+		PreviewGeometry->RemoveLineSet(TEXT("TopologyEdges"));
+		PreviewGeometry->RemoveLineSet(TEXT("AllEdges"));
+		return;
+	}
 
-	PreviewGeometry->CreateOrUpdateLineSet(TEXT("TopologyEdges"),
-										   NumEdges,
-										   [&Topology, this](int32 Index, TArray<FRenderableLine>& LinesOut)
+	if (Properties->SubdivisionScheme == ESubdivisionScheme::Loop)
 	{
-		const FGroupTopology::FGroupEdge& Edge = Topology.Edges[Index];
-		FIndex2i EdgeCorners = Edge.EndpointCorners;
+		int NumEdges = OriginalMesh->EdgeCount();
 
-		if (EdgeCorners[0] == FDynamicMesh3::InvalidID || EdgeCorners[1] == FDynamicMesh3::InvalidID)
+		PreviewGeometry->RemoveLineSet(TEXT("TopologyEdges"));
+
+		PreviewGeometry->CreateOrUpdateLineSet(TEXT("AllEdges"),
+											   NumEdges,
+											   [this](int32 Index, TArray<FRenderableLine>& LinesOut)
 		{
-			return;
-		}
+			FIndex2i EdgeVertices = OriginalMesh->GetEdgeV(Index);
 
-		FIndex2i EdgeVertices{ Topology.Corners[EdgeCorners[0]].VertexID,
-							   Topology.Corners[EdgeCorners[1]].VertexID };
-		FVector A = (FVector)OriginalMesh->GetVertex(EdgeVertices[0]);
-		FVector B = (FVector)OriginalMesh->GetVertex(EdgeVertices[1]);
+			if (EdgeVertices[0] == FDynamicMesh3::InvalidID || EdgeVertices[1] == FDynamicMesh3::InvalidID)
+			{
+				return;
+			}
+			FVector A = (FVector)OriginalMesh->GetVertex(EdgeVertices[0]);
+			FVector B = (FVector)OriginalMesh->GetVertex(EdgeVertices[1]);
+			const float TopologyLineThickness = 4.0f;
+			const FColor TopologyLineColor(255, 0, 0);
+			LinesOut.Add(FRenderableLine(A, B, TopologyLineColor, TopologyLineThickness));
+		});
+	}
+	else
+	{
+		FGroupTopology Topology(OriginalMesh.Get(), true);
+		int NumEdges = Topology.Edges.Num();
 
-		const float TopologyLineThickness = 4.0f;
-		const FColor TopologyLineColor(255, 0, 0);
-		LinesOut.Add(FRenderableLine(A, B, TopologyLineColor, TopologyLineThickness));
-	});
+		PreviewGeometry->RemoveLineSet(TEXT("AllEdges"));
+
+		PreviewGeometry->CreateOrUpdateLineSet(TEXT("TopologyEdges"),
+											   NumEdges,
+											   [&Topology, this](int32 Index, TArray<FRenderableLine>& LinesOut)
+		{
+			const FGroupTopology::FGroupEdge& Edge = Topology.Edges[Index];
+			FIndex2i EdgeCorners = Edge.EndpointCorners;
+
+			if (EdgeCorners[0] == FDynamicMesh3::InvalidID || EdgeCorners[1] == FDynamicMesh3::InvalidID)
+			{
+				return;
+			}
+
+			FIndex2i EdgeVertices{ Topology.Corners[EdgeCorners[0]].VertexID,
+									Topology.Corners[EdgeCorners[1]].VertexID };
+			FVector A = (FVector)OriginalMesh->GetVertex(EdgeVertices[0]);
+			FVector B = (FVector)OriginalMesh->GetVertex(EdgeVertices[1]);
+
+			const float TopologyLineThickness = 4.0f;
+			const FColor TopologyLineColor(255, 0, 0);
+			LinesOut.Add(FRenderableLine(A, B, TopologyLineColor, TopologyLineThickness));
+		});
+	}
 }
 
 
diff --git a/Engine/Plugins/Experimental/MeshModelingToolset/Source/MeshModelingToolsEditorOnly/Public/SubdividePoly.h b/Engine/Plugins/Experimental/MeshModelingToolset/Source/MeshModelingToolsEditorOnly/Public/SubdividePoly.h
--- a/Engine/Plugins/Experimental/MeshModelingToolset/Source/MeshModelingToolsEditorOnly/Public/SubdividePoly.h
+++ b/Engine/Plugins/Experimental/MeshModelingToolset/Source/MeshModelingToolsEditorOnly/Public/SubdividePoly.h
@@ -10,6 +10,18 @@ class FGroupTopology;
 class FDynamicMesh3;
 class FProgressCancel;
 
+/**
+ * Subdivision scheme
+ */
+UENUM()
+enum class ESubdivisionScheme : uint8
+{
+	Bilinear,
+	CatmullClark,
+	Loop
+};
+
+
 /**
  * Method for computing normals on the output mesh
  */
@@ -60,6 +72,8 @@ public:
 
 	bool ComputeSubdividedMesh(FDynamicMesh3& OutMesh);
 
+	ESubdivisionScheme SubdivisionScheme = ESubdivisionScheme::CatmullClark;
+
 	ESubdivisionOutputNormals NormalComputationMethod = ESubdivisionOutputNormals::Generated;
 
 	ESubdivisionOutputUVs UVComputationMethod = ESubdivisionOutputUVs::Interpolated;
diff --git a/Engine/Plugins/Experimental/MeshModelingToolset/Source/MeshModelingToolsEditorOnly/Public/SubdividePolyTool.h b/Engine/Plugins/Experimental/MeshModelingToolset/Source/MeshModelingToolsEditorOnly/Public/SubdividePolyTool.h
--- a/Engine/Plugins/Experimental/MeshModelingToolset/Source/MeshModelingToolsEditorOnly/Public/SubdividePolyTool.h
+++ b/Engine/Plugins/Experimental/MeshModelingToolset/Source/MeshModelingToolsEditorOnly/Public/SubdividePolyTool.h
@@ -43,6 +43,9 @@ public:
 	UPROPERTY(EditAnywhere, Category=Settings, meta = (UIMin = "1", ClampMin = "1"))
 	int SubdivisionLevel = 3;
 
+	UPROPERTY(EditAnywhere, Category = Settings)
+	ESubdivisionScheme SubdivisionScheme = ESubdivisionScheme::CatmullClark;
+
 	UPROPERTY(EditAnywhere, Category=Settings)
 	ESubdivisionOutputNormals NormalComputationMethod = ESubdivisionOutputNormals::Generated;
 
@@ -52,6 +55,9 @@ public:
 	UPROPERTY(EditAnywhere, Category = Settings)
 	bool bRenderGroups = true;
 
+	UPROPERTY(EditAnywhere, Category = Settings)
+	bool bRenderCage = true;
+
 };
 
 
