Fix VolumeToMeshTool to be able to accept after merge of CL 14602422 into main (through 14811440).

#rnx

diff --git a/Engine/Plugins/Experimental/MeshModelingToolset/Source/MeshModelingTools/Public/VolumeToMeshTool.h b/Engine/Plugins/Experimental/MeshModelingToolset/Source/MeshModelingTools/Public/VolumeToMeshTool.h
--- a/Engine/Plugins/Experimental/MeshModelingToolset/Source/MeshModelingTools/Public/VolumeToMeshTool.h
+++ b/Engine/Plugins/Experimental/MeshModelingToolset/Source/MeshModelingTools/Public/VolumeToMeshTool.h
@@ -85,6 +85,7 @@ public:
 
 	virtual bool HasCancel() const override { return true; }
 	virtual bool HasAccept() const override { return true; }
+	virtual bool CanAccept() const override { return true; }
 
 protected:
 	UPROPERTY()
