MeshModelingTools: Fix simple clicks away from the mesh not clearing selection when using marquee in selection mechanic.

#rnx

diff --git a/Engine/Plugins/Experimental/MeshModelingToolset/Source/ModelingComponents/Private/Selection/PolygonSelectionMechanic.cpp b/Engine/Plugins/Experimental/MeshModelingToolset/Source/ModelingComponents/Private/Selection/PolygonSelectionMechanic.cpp
--- a/Engine/Plugins/Experimental/MeshModelingToolset/Source/ModelingComponents/Private/Selection/PolygonSelectionMechanic.cpp
+++ b/Engine/Plugins/Experimental/MeshModelingToolset/Source/ModelingComponents/Private/Selection/PolygonSelectionMechanic.cpp
@@ -583,6 +583,7 @@ void UPolygonSelectionMechanic::OnUpdateModifierState(int ModifierID, bool bIsOn
 void UPolygonSelectionMechanic::OnDragRectangleStarted()
 {
 	PreDragPersistentSelection = PersistentSelection;
+	PersistentSelection.Clear();
 }
 
 void UPolygonSelectionMechanic::OnDragRectangleChanged(const FCameraRectangle& CurrentRectangle)
