Subdivide tool: If the mesh to be subdivided doesn't have polygroup topology suitable for Catmull-Clark, display a warning and allow the user to only use Loop subdivision. (Previous to this change, the tool would just display a warning and fail to initialize.)

diff --git a/Engine/Plugins/Experimental/MeshModelingToolset/Source/MeshModelingToolsEditorOnly/Private/SubdividePolyTool.cpp b/Engine/Plugins/Experimental/MeshModelingToolset/Source/MeshModelingToolsEditorOnly/Private/SubdividePolyTool.cpp
--- a/Engine/Plugins/Experimental/MeshModelingToolset/Source/MeshModelingToolsEditorOnly/Private/SubdividePolyTool.cpp
+++ b/Engine/Plugins/Experimental/MeshModelingToolset/Source/MeshModelingToolsEditorOnly/Private/SubdividePolyTool.cpp
@@ -70,43 +70,23 @@ UInteractiveTool* USubdividePolyToolBuilder::BuildTool(const FToolBuilderState&
 }
 
 
-// Tool actual
-
-void USubdividePolyTool::Setup()
+bool USubdividePolyTool::CheckGroupTopology(FText& Message)
 {
-	UInteractiveTool::Setup();
-
-	if (!ComponentTarget)
-	{
-		return;
-	}
-
-	bool bWantVertexNormals = false;
-	OriginalMesh = MakeShared<FDynamicMesh3>(bWantVertexNormals, false, false, false);
-	FMeshDescriptionToDynamicMesh Converter;
-	Converter.Convert(ComponentTarget->GetMesh(), *OriginalMesh);
-
-	//
-	// Error checking
-	//
-
 	FGroupTopology Topo(OriginalMesh.Get(), true);
+
 	if (Topo.Groups.Num() == 0)
 	{
-		GetToolManager()->DisplayMessage(LOCTEXT("NoGroupsWarning",
-												 "This object has no PolyGroups. Use the PolyGroups or Select Tool to assign PolyGroups."),
-										 EToolMessageLevel::UserWarning);
-		return;
+		Message = LOCTEXT("NoGroupsWarning",
+						  "This object has no PolyGroups.\nUse the PolyGroups or Select Tool to assign PolyGroups.\nTool will be limited to Loop subdivision scheme.");
+		return false;
 	}
 
 	if (Topo.Groups.Num() < 2)
 	{
 		// TODO: for an open surface, use the surface boundary as a group boundary?
-
-		GetToolManager()->DisplayMessage(LOCTEXT("SingleGroupsWarning",
-												 "This object has only one PolyGroup. Use the PolyGroups or Select Tool to assign PolyGroups."),
-										 EToolMessageLevel::UserWarning);
-		return;
+		Message = LOCTEXT("SingleGroupsWarning",
+						  "This object has only one PolyGroup.\nUse the PolyGroups or Select Tool to assign PolyGroups.\nTool will be limited to Loop subdivision scheme.");
+		return false;
 	}
 
 	for (const FGroupTopology::FGroup& Group : Topo.Groups)
@@ -114,43 +94,70 @@ void USubdividePolyTool::Setup()
 		if (Group.Boundaries.Num() == 0)
 		{
 			// Error: Group has no boundaries, e.g. closed surface component with only one group
-			GetToolManager()->DisplayMessage(LOCTEXT("NoGroupBoundaryWarning",
-													 "Found a PolyGroup with no boundaries. Use the PolyGroups or Select Tool to assign PolyGroups."),
-											 EToolMessageLevel::UserWarning);
-			return;
+			Message = LOCTEXT("NoGroupBoundaryWarning",
+							  "Found a PolyGroup with no boundaries.\nUse the PolyGroups or Select Tool to assign PolyGroups.\nTool will be limited to Loop subdivision scheme.");
+			return false;
 		}
 
 		if (Group.Boundaries.Num() > 1)
 		{
 			// Error: Group has multiple boundaries, e.g. nested polygon
-			GetToolManager()->DisplayMessage(LOCTEXT("MultipleGroupBoundaryWarning",
-													 "Found a PolyGroup with multiple boundaries, which is not supported. Use the PolyGroups or Select Tool to assign PolyGroups."),
-											 EToolMessageLevel::UserWarning);
-			return;
+			Message = LOCTEXT("MultipleGroupBoundaryWarning",
+							  "Found a PolyGroup with multiple boundaries, which is not supported.\nUse the PolyGroups or Select Tool to assign PolyGroups.\nTool will be limited to Loop subdivision scheme.");
+			return false;
 		}
 
-		for ( const FGroupTopology::FGroupBoundary& Boundary : Group.Boundaries )
+		for (const FGroupTopology::FGroupBoundary& Boundary : Group.Boundaries)
 		{
 			if (Boundary.GroupEdges.Num() < 3)
 			{
-				GetToolManager()->DisplayMessage(LOCTEXT("DegenerateGroupPolygon",
-														 "One PolyGroup has fewer than three boundary edges. Use the PolyGroups or Select Tool to assign/fix PolyGroups."),
-												 EToolMessageLevel::UserWarning);
-				return;
+				Message = LOCTEXT("DegenerateGroupPolygon",
+								  "One PolyGroup has fewer than three boundary edges.\nUse the PolyGroups or Select Tool to assign/fix PolyGroups.\nTool will be limited to Loop subdivision scheme.");
+				return false;
 			}
 		}
 	}
 
-	//
-	// Finished error checking
-	//
+	return true;
+}
 
-	GetToolManager()->DisplayMessage(LOCTEXT("SubdividePolyToolMessage",
-											 "Set the subdivision level and hit Accept to create a new subdivided mesh"),
-									 EToolMessageLevel::UserNotification);
+void USubdividePolyTool::Setup()
+{
+	UInteractiveTool::Setup();
+
+	if (!ComponentTarget)
+	{
+		return;
+	}
+
+	bool bWantVertexNormals = false;
+	OriginalMesh = MakeShared<FDynamicMesh3>(bWantVertexNormals, false, false, false);
+	FMeshDescriptionToDynamicMesh Converter;
+	Converter.Convert(ComponentTarget->GetMesh(), *OriginalMesh);
+
+	FText ErrorMessage;
+	bool bCatmullClarkOK = CheckGroupTopology(ErrorMessage);
+
+	if (bCatmullClarkOK)
+	{
+		GetToolManager()->DisplayMessage(LOCTEXT("SubdividePolyToolMessage",
+												 "Set the subdivision level and hit Accept to create a new subdivided mesh"),
+										 EToolMessageLevel::UserNotification);
+	}
+	else
+	{
+		GetToolManager()->DisplayMessage(ErrorMessage, EToolMessageLevel::UserWarning);
+	}
 
 	Properties = NewObject<USubdividePolyToolProperties>(this, TEXT("Subdivide Mesh Tool Settings"));
 	Properties->RestoreProperties(this);
+
+	Properties->bCatmullClarkOK = bCatmullClarkOK;
+	if (!bCatmullClarkOK)
+	{
+		Properties->SubdivisionScheme = ESubdivisionScheme::Loop;
+	}
+
 	AddToolPropertySource(Properties);
 	SetToolPropertySourceEnabled(Properties, true);
 
diff --git a/Engine/Plugins/Experimental/MeshModelingToolset/Source/MeshModelingToolsEditorOnly/Public/SubdividePolyTool.h b/Engine/Plugins/Experimental/MeshModelingToolset/Source/MeshModelingToolsEditorOnly/Public/SubdividePolyTool.h
--- a/Engine/Plugins/Experimental/MeshModelingToolset/Source/MeshModelingToolsEditorOnly/Public/SubdividePolyTool.h
+++ b/Engine/Plugins/Experimental/MeshModelingToolset/Source/MeshModelingToolsEditorOnly/Public/SubdividePolyTool.h
@@ -43,7 +43,11 @@ public:
 	UPROPERTY(EditAnywhere, Category=Settings, meta = (UIMin = "1", ClampMin = "1"))
 	int SubdivisionLevel = 3;
 
-	UPROPERTY(EditAnywhere, Category = Settings)
+	// Controls whether the user can select Catmull-Clark or is forced to use Loop
+	UPROPERTY(meta = (TransientToolProperty))
+	bool bCatmullClarkOK = true;
+
+	UPROPERTY(EditAnywhere, Category = Settings, meta = (EditCondition = "bCatmullClarkOK", HideEditConditionToggle))
 	ESubdivisionScheme SubdivisionScheme = ESubdivisionScheme::CatmullClark;
 
 	UPROPERTY(EditAnywhere, Category=Settings)
@@ -102,5 +106,9 @@ protected:
 
 	bool bPreviewGeometryNeedsUpdate;
 	void CreateOrUpdatePreviewGeometry();
+
+	// Determine if the mesh group topology can be used for Catmull-Clark or Bilinear subdivision. If not, we can only 
+	// Loop subdivision on the original triangle mesh.
+	bool CheckGroupTopology(FText& Message);
 };
 
