split trim options off the boolean tool into a trim-specific tool w/ trim-specific options
+ add trim inside vs outside option

diff --git a/Engine/Plugins/Experimental/GeometryProcessing/Source/DynamicMesh/Private/Operations/MeshBoolean.cpp b/Engine/Plugins/Experimental/GeometryProcessing/Source/DynamicMesh/Private/Operations/MeshBoolean.cpp
--- a/Engine/Plugins/Experimental/GeometryProcessing/Source/DynamicMesh/Private/Operations/MeshBoolean.cpp
+++ b/Engine/Plugins/Experimental/GeometryProcessing/Source/DynamicMesh/Private/Operations/MeshBoolean.cpp
@@ -256,10 +256,12 @@ bool FMeshBoolean::Compute()
 		return false;
 	}
 
+	bool bOpIsATrim = Operation == EBooleanOp::TrimInside || Operation == EBooleanOp::TrimOutside;
+
 	// cut the meshes
 	FMeshMeshCut Cut(CutMesh[0], CutMesh[1]);
 	Cut.bTrackInsertedVertices = bCollapseDegenerateEdgesOnCut; // to collect candidates to collapse
-	Cut.bMutuallyCut = Operation != EBooleanOp::Trim;
+	Cut.bMutuallyCut = !bOpIsATrim;
 	Cut.SnapTolerance = SnapTolerance;
 	Cut.Cut(Intersections);
 
@@ -268,7 +270,7 @@ bool FMeshBoolean::Compute()
 		return false;
 	}
 
-	int NumMeshesToProcess = Operation == EBooleanOp::Trim ? 1 : 2;
+	int NumMeshesToProcess = bOpIsATrim ? 1 : 2;
 
 	// collapse tiny edges along cut boundary
 	if (bCollapseDegenerateEdgesOnCut)
@@ -361,9 +363,9 @@ bool FMeshBoolean::Compute()
 			FDynamicMesh3& ProcessMesh = *CutMesh[MeshIdx];
 			int MaxTriID = ProcessMesh.MaxTriangleID();
 			KeepTri[MeshIdx].SetNumUninitialized(MaxTriID);
-			bool bCoplanarKeepSameDir = Operation != EBooleanOp::Difference;
+			bool bCoplanarKeepSameDir = (Operation != EBooleanOp::Difference && Operation != EBooleanOp::TrimInside);
 			bool bRemoveInside = 1; // whether to remove the inside triangles (e.g. for union) or the outside ones (e.g. for intersection)
-			if (Operation == EBooleanOp::Trim || Operation == EBooleanOp::Intersect || (Operation == EBooleanOp::Difference && MeshIdx == 1))
+			if (Operation == EBooleanOp::TrimOutside || Operation == EBooleanOp::Intersect || (Operation == EBooleanOp::Difference && MeshIdx == 1))
 			{
 				bRemoveInside = 0;
 			}
diff --git a/Engine/Plugins/Experimental/GeometryProcessing/Source/DynamicMesh/Public/Operations/MeshBoolean.h b/Engine/Plugins/Experimental/GeometryProcessing/Source/DynamicMesh/Public/Operations/MeshBoolean.h
--- a/Engine/Plugins/Experimental/GeometryProcessing/Source/DynamicMesh/Public/Operations/MeshBoolean.h
+++ b/Engine/Plugins/Experimental/GeometryProcessing/Source/DynamicMesh/Public/Operations/MeshBoolean.h
@@ -38,7 +38,8 @@ public:
 		Union,
 		Difference,
 		Intersect,
-		Trim
+		TrimInside,
+		TrimOutside
 	};
 	EBooleanOp Operation;
 	
diff --git a/Engine/Plugins/Experimental/MeshModelingToolset/Source/MeshModelingTools/Private/CSGMeshesTool.cpp b/Engine/Plugins/Experimental/MeshModelingToolset/Source/MeshModelingTools/Private/CSGMeshesTool.cpp
--- a/Engine/Plugins/Experimental/MeshModelingToolset/Source/MeshModelingTools/Private/CSGMeshesTool.cpp
+++ b/Engine/Plugins/Experimental/MeshModelingToolset/Source/MeshModelingTools/Private/CSGMeshesTool.cpp
@@ -12,26 +12,51 @@
 
 #define LOCTEXT_NAMESPACE "UCSGMeshesTool"
 
-
+void UCSGMeshesTool::EnableTrimMode()
+{
+	check(OriginalDynamicMeshes.Num() == 0);		// must not have been initialized!
+	bTrimMode = true;
+}
 
 void UCSGMeshesTool::SetupProperties()
 {
 	Super::SetupProperties();
 
-	CSGProperties = NewObject<UCSGMeshesToolProperties>(this);
-	CSGProperties->RestoreProperties(this);
-	AddToolPropertySource(CSGProperties);
-
-	SetToolDisplayName(LOCTEXT("CSGMeshesToolName", "Mesh Boolean Tool"));
-	GetToolManager()->DisplayMessage(
-		LOCTEXT("OnStartTool", "Compute CSG Booleans on the input meshes. Use the transform gizmos to tweak the positions of the input objects (can help to resolve errors/failures)"),
-		EToolMessageLevel::UserNotification);
+	if (bTrimMode)
+	{
+		TrimProperties = NewObject<UTrimMeshesToolProperties>(this);
+		TrimProperties->RestoreProperties(this);
+		AddToolPropertySource(TrimProperties);
+
+		SetToolDisplayName(LOCTEXT("TrimMeshesToolName", "Trim Meshes Tool"));
+		GetToolManager()->DisplayMessage(
+			LOCTEXT("OnStartTrimTool", "Trim one mesh with another. Use the transform gizmos to tweak the positions of the input objects (can help to resolve errors/failures)"),
+			EToolMessageLevel::UserNotification);
+	}
+	else
+	{
+		CSGProperties = NewObject<UCSGMeshesToolProperties>(this);
+		CSGProperties->RestoreProperties(this);
+		AddToolPropertySource(CSGProperties);
+
+		SetToolDisplayName(LOCTEXT("CSGMeshesToolName", "Mesh Boolean Tool"));
+		GetToolManager()->DisplayMessage(
+			LOCTEXT("OnStartTool", "Compute CSG Booleans on the input meshes. Use the transform gizmos to tweak the positions of the input objects (can help to resolve errors/failures)"),
+			EToolMessageLevel::UserNotification);
+	}
 }
 
 void UCSGMeshesTool::SaveProperties()
 {
 	Super::SaveProperties();
-	CSGProperties->SaveProperties(this);
+	if (bTrimMode)
+	{
+		TrimProperties->SaveProperties(this);
+	}
+	else
+	{
+		CSGProperties->SaveProperties(this);
+	}
 }
 
 
@@ -42,7 +67,7 @@ void UCSGMeshesTool::ConvertInputsAndSetPreviewMaterials(bool bSetPreviewMesh)
 	TMap<UMaterialInterface*, int> KnownMaterials;
 	TArray<TArray<int>> MaterialRemap; MaterialRemap.SetNum(ComponentTargets.Num());
 
-	if (!CSGProperties->bOnlyUseFirstMeshMaterials)
+	if (bTrimMode || !CSGProperties->bOnlyUseFirstMeshMaterials)
 	{
 		for (int ComponentIdx = 0; ComponentIdx < ComponentTargets.Num(); ComponentIdx++)
 		{
@@ -131,7 +156,7 @@ void UCSGMeshesTool::UpdateVisualization()
 	FVector3d A, B;
 
 	DrawnLineSet->Clear();
-	if (CSGProperties->bShowNewBoundaryEdges)
+	if (!bTrimMode && CSGProperties->bShowNewBoundaryEdges)
 	{
 		for (int EID : CreatedBoundaryEdges)
 		{
@@ -146,8 +171,18 @@ TUniquePtr<FDynamicMeshOperator> UCSGMeshesTool::MakeNewOperator()
 {
 	TUniquePtr<FBooleanMeshesOp> BooleanOp = MakeUnique<FBooleanMeshesOp>();
 	
-	BooleanOp->Operation = CSGProperties->Operation;
-	BooleanOp->bAttemptFixHoles = CSGProperties->bAttemptFixHoles;
+	BooleanOp->bTrimMode = bTrimMode;
+	if (bTrimMode)
+	{
+		BooleanOp->TrimOperation = TrimProperties->WhichMesh;
+		BooleanOp->TrimSide = TrimProperties->TrimSide;
+		BooleanOp->bAttemptFixHoles = false;
+	}
+	else
+	{
+		BooleanOp->CSGOperation = CSGProperties->Operation;
+		BooleanOp->bAttemptFixHoles = CSGProperties->bAttemptFixHoles;
+	}
 
 	check(OriginalDynamicMeshes.Num() == 2);
 	check(ComponentTargets.Num() == 2);
@@ -192,13 +227,27 @@ void UCSGMeshesTool::OnPropertyModified(UObject* PropertySet, FProperty* Propert
 
 FString UCSGMeshesTool::GetCreatedAssetName() const
 {
-	return TEXT("Boolean");
+	if (bTrimMode)
+	{
+		return TEXT("Trim");
+	}
+	else
+	{
+		return TEXT("Boolean");
+	}
 }
 
 
 FText UCSGMeshesTool::GetActionName() const
 {
-	return LOCTEXT("CSGMeshes", "Boolean Meshes");
+	if (bTrimMode)
+	{
+		return LOCTEXT("CSGMeshes", "Trim Meshes");
+	}
+	else
+	{
+		return LOCTEXT("CSGMeshes", "Boolean Meshes");
+	}
 }
 
 
diff --git a/Engine/Plugins/Experimental/MeshModelingToolset/Source/MeshModelingTools/Public/CSGMeshesTool.h b/Engine/Plugins/Experimental/MeshModelingToolset/Source/MeshModelingTools/Public/CSGMeshesTool.h
--- a/Engine/Plugins/Experimental/MeshModelingToolset/Source/MeshModelingTools/Public/CSGMeshesTool.h
+++ b/Engine/Plugins/Experimental/MeshModelingToolset/Source/MeshModelingTools/Public/CSGMeshesTool.h
@@ -42,6 +42,24 @@ public:
 };
 
 
+/**
+ * Properties of the trim mode
+ */
+UCLASS()
+class MESHMODELINGTOOLS_API UTrimMeshesToolProperties : public UInteractiveToolPropertySet
+{
+	GENERATED_BODY()
+public:
+	/** Which object to trim */
+	UPROPERTY(EditAnywhere, Category = Options)
+	ETrimOperation WhichMesh = ETrimOperation::TrimA;
+
+	/** Whether to remove the surface inside or outside of the trimming geometry */
+	UPROPERTY(EditAnywhere, Category = Options)
+	ETrimSide TrimSide = ETrimSide::RemoveInside;
+};
+
+
 
 /**
  * Simple Mesh Plane Cutting Tool
@@ -55,6 +73,8 @@ public:
 
 	UCSGMeshesTool() {}
 
+	void EnableTrimMode();
+
 protected:
 
 	virtual void OnPropertyModified(UObject* PropertySet, FProperty* Property) override;
@@ -78,6 +98,9 @@ protected:
 	UPROPERTY()
 	UCSGMeshesToolProperties* CSGProperties;
 
+	UPROPERTY()
+	UTrimMeshesToolProperties* TrimProperties;
+
 	TArray<TSharedPtr<FDynamicMesh3>> OriginalDynamicMeshes;
 
 	UPROPERTY()
@@ -85,6 +108,8 @@ protected:
 
 	// for visualization of any errors in the currently-previewed CSG operation
 	TArray<int> CreatedBoundaryEdges;
+
+	bool bTrimMode = false;
 };
 
 
@@ -97,12 +122,13 @@ class MESHMODELINGTOOLS_API UCSGMeshesToolBuilder : public UBaseCreateFromSelect
 
 public:
 
+	bool bTrimMode = false;
+
 	virtual TOptional<int32> MaxComponentsSupported() const override
 	{
 		return TOptional<int32>(2);
 	}
 
-
 	virtual int32 MinComponentsSupported() const override
 	{
 		return 2;
@@ -110,7 +136,12 @@ public:
 
 	virtual UBaseCreateFromSelectedTool* MakeNewToolInstance(UObject* Outer) const override
 	{
-		return NewObject<UCSGMeshesTool>(Outer);
+		UCSGMeshesTool* Tool = NewObject<UCSGMeshesTool>(Outer);
+		if (bTrimMode)
+		{
+			Tool->EnableTrimMode();
+		}
+		return Tool;
 	}
 };
 
diff --git a/Engine/Plugins/Experimental/MeshModelingToolset/Source/ModelingOperators/Private/CompositionOps/BooleanMeshesOp.cpp b/Engine/Plugins/Experimental/MeshModelingToolset/Source/ModelingOperators/Private/CompositionOps/BooleanMeshesOp.cpp
--- a/Engine/Plugins/Experimental/MeshModelingToolset/Source/ModelingOperators/Private/CompositionOps/BooleanMeshesOp.cpp
+++ b/Engine/Plugins/Experimental/MeshModelingToolset/Source/ModelingOperators/Private/CompositionOps/BooleanMeshesOp.cpp
@@ -20,7 +20,7 @@ void FBooleanMeshesOp::CalculateResult(FProgressCancel* Progress)
 	check(Meshes.Num() == 2 && Transforms.Num() == 2);
 	
 	int FirstIdx = 0;
-	if (Operation == ECSGOperation::DifferenceBA || Operation == ECSGOperation::TrimB)
+	if ((!bTrimMode && CSGOperation == ECSGOperation::DifferenceBA) || (bTrimMode && TrimOperation == ETrimOperation::TrimB))
 	{
 		FirstIdx = 1;
 	}
@@ -28,25 +28,39 @@ void FBooleanMeshesOp::CalculateResult(FProgressCancel* Progress)
 
 	FMeshBoolean::EBooleanOp Op;
 	// convert UI enum to algorithm enum
-	switch (Operation)
+	if (bTrimMode)
 	{
-	case ECSGOperation::DifferenceAB:
-	case ECSGOperation::DifferenceBA:
-		Op = FMeshBoolean::EBooleanOp::Difference;
-		break;
-	case ECSGOperation::TrimA:
-	case ECSGOperation::TrimB:
-		Op = FMeshBoolean::EBooleanOp::Trim;
-		break;
-	case ECSGOperation::Union:
-		Op = FMeshBoolean::EBooleanOp::Union;
-		break;
-	case ECSGOperation::Intersect:
-		Op = FMeshBoolean::EBooleanOp::Intersect;
-		break;
-	default:
-		check(false); // all conversion cases should be implemented
-		Op = FMeshBoolean::EBooleanOp::Union;
+		switch (TrimSide)
+		{
+		case ETrimSide::RemoveInside:
+			Op = FMeshBoolean::EBooleanOp::TrimInside;
+			break;
+		case ETrimSide::RemoveOutside:
+			Op = FMeshBoolean::EBooleanOp::TrimOutside;
+			break;
+		default:
+			check(false);
+			Op = FMeshBoolean::EBooleanOp::TrimInside;
+		}
+	}
+	else
+	{
+		switch (CSGOperation)
+		{
+		case ECSGOperation::DifferenceAB:
+		case ECSGOperation::DifferenceBA:
+			Op = FMeshBoolean::EBooleanOp::Difference;
+			break;
+		case ECSGOperation::Union:
+			Op = FMeshBoolean::EBooleanOp::Union;
+			break;
+		case ECSGOperation::Intersect:
+			Op = FMeshBoolean::EBooleanOp::Intersect;
+			break;
+		default:
+			check(false); // all conversion cases should be implemented
+			Op = FMeshBoolean::EBooleanOp::Union;
+		}
 	}
 
 	FMeshBoolean MeshBoolean(Meshes[FirstIdx].Get(), (FTransform3d)Transforms[FirstIdx], Meshes[OtherIdx].Get(), (FTransform3d)Transforms[OtherIdx], ResultMesh.Get(), Op);
diff --git a/Engine/Plugins/Experimental/MeshModelingToolset/Source/ModelingOperators/Public/CompositionOps/BooleanMeshesOp.h b/Engine/Plugins/Experimental/MeshModelingToolset/Source/ModelingOperators/Public/CompositionOps/BooleanMeshesOp.h
--- a/Engine/Plugins/Experimental/MeshModelingToolset/Source/ModelingOperators/Public/CompositionOps/BooleanMeshesOp.h
+++ b/Engine/Plugins/Experimental/MeshModelingToolset/Source/ModelingOperators/Public/CompositionOps/BooleanMeshesOp.h
@@ -25,16 +25,26 @@ enum class ECSGOperation : uint8
 
 	/** union of two objects */
 	Union = 3 UMETA(DisplayName = "Union"),
+};
 
-	/** Trim the first object using the second */
-	TrimA = 4 UMETA(DisplayName = "Trim A"),
-
-	/** Trim the second object using the first */
-	TrimB = 5 UMETA(DisplayName = "Trim B"),
+UENUM()
+enum class ETrimOperation : uint8
+{
+	/** Remove geometry from the first object using the second */
+	TrimA = 0 UMETA(DisplayName = "Trim A"),
 
+	/** Remove geometry from the second object using the first */
+	TrimB = 1 UMETA(DisplayName = "Trim B"),
 };
 
 
+UENUM()
+enum class ETrimSide : uint8
+{
+	RemoveInside = 0,
+	RemoveOutside = 1,
+};
+
 
 class MODELINGOPERATORS_API FBooleanMeshesOp : public FDynamicMeshOperator
 {
@@ -42,7 +52,10 @@ public:
 	virtual ~FBooleanMeshesOp() {}
 
 	// inputs
-	ECSGOperation Operation;
+	ECSGOperation CSGOperation;
+	ETrimOperation TrimOperation;
+	ETrimSide TrimSide;
+	bool bTrimMode = false; // if true, do a trim operation instead of a boolean
 	TArray<TSharedPtr<const FDynamicMesh3>> Meshes;
 	TArray<FTransform> Transforms; // 1:1 with Meshes
 	bool bAttemptFixHoles;
diff --git a/Engine/Plugins/Experimental/ModelingToolsEditorMode/Source/ModelingToolsEditorMode/Private/ModelingToolsEditorMode.cpp b/Engine/Plugins/Experimental/ModelingToolsEditorMode/Source/ModelingToolsEditorMode/Private/ModelingToolsEditorMode.cpp
--- a/Engine/Plugins/Experimental/ModelingToolsEditorMode/Source/ModelingToolsEditorMode/Private/ModelingToolsEditorMode.cpp
+++ b/Engine/Plugins/Experimental/ModelingToolsEditorMode/Source/ModelingToolsEditorMode/Private/ModelingToolsEditorMode.cpp
@@ -476,6 +476,11 @@ void UModelingToolsEditorMode::Enter()
 	CSGMeshesToolBuilder->AssetAPI = ModelingModeAssetGenerationAPI.Get();
 	RegisterTool(ToolManagerCommands.BeginMeshBooleanTool, TEXT("CSGMeshesTool"), CSGMeshesToolBuilder);
 
+	auto TrimMeshesToolBuilder = NewObject<UCSGMeshesToolBuilder>();
+	TrimMeshesToolBuilder->bTrimMode = true;
+	TrimMeshesToolBuilder->AssetAPI = ModelingModeAssetGenerationAPI.Get();
+	RegisterTool(ToolManagerCommands.BeginMeshTrimTool, TEXT("TrimMeshesTool"), TrimMeshesToolBuilder);
+
 	auto BspConversionToolBuilder = NewObject<UBspConversionToolBuilder>();
 	BspConversionToolBuilder->AssetAPI = ModelingModeAssetGenerationAPI.Get();
 	RegisterTool(ToolManagerCommands.BeginBspConversionTool, TEXT("BspConversionTool"), BspConversionToolBuilder);
diff --git a/Engine/Plugins/Experimental/ModelingToolsEditorMode/Source/ModelingToolsEditorMode/Private/ModelingToolsEditorModeStyle.cpp b/Engine/Plugins/Experimental/ModelingToolsEditorMode/Source/ModelingToolsEditorMode/Private/ModelingToolsEditorModeStyle.cpp
--- a/Engine/Plugins/Experimental/ModelingToolsEditorMode/Source/ModelingToolsEditorMode/Private/ModelingToolsEditorModeStyle.cpp
+++ b/Engine/Plugins/Experimental/ModelingToolsEditorMode/Source/ModelingToolsEditorMode/Private/ModelingToolsEditorModeStyle.cpp
@@ -236,6 +236,8 @@ void FModelingToolsEditorModeStyle::Initialize()
 		StyleSet->Set("ModelingToolsManagerCommands.BeginRevolveBoundaryTool.Small",         new IMAGE_PLUGIN_BRUSH("Icons/ModelingRevolveBoundary_x20", Icon20x20));
 		StyleSet->Set("ModelingToolsManagerCommands.BeginMeshBooleanTool",                   new IMAGE_PLUGIN_BRUSH("Icons/ModelingMeshBoolean_x40", Icon20x20));
 		StyleSet->Set("ModelingToolsManagerCommands.BeginMeshBooleanTool.Small",             new IMAGE_PLUGIN_BRUSH("Icons/ModelingMeshBoolean_x20", Icon20x20));
+		//StyleSet->Set("ModelingToolsManagerCommands.BeginMeshBooleanTool",				 new IMAGE_PLUGIN_BRUSH("Icons/ModelingMeshTrim_x40", Icon20x20));
+		//StyleSet->Set("ModelingToolsManagerCommands.BeginMeshBooleanTool.Small",			 new IMAGE_PLUGIN_BRUSH("Icons/ModelingMeshTrim_x20", Icon20x20));
 		StyleSet->Set("ModelingToolsManagerCommands.BeginSelfUnionTool",                     new IMAGE_PLUGIN_BRUSH("Icons/ModelingSelfUnion_x40", Icon20x20));
 		StyleSet->Set("ModelingToolsManagerCommands.BeginSelfUnionTool.Small",               new IMAGE_PLUGIN_BRUSH("Icons/ModelingSelfUnion_x20", Icon20x20));
 		StyleSet->Set("ModelingToolsManagerCommands.BeginVoxelSolidifyTool",                 new IMAGE_PLUGIN_BRUSH("Icons/ModelingVoxSolidify_x40", Icon20x20));
diff --git a/Engine/Plugins/Experimental/ModelingToolsEditorMode/Source/ModelingToolsEditorMode/Private/ModelingToolsEditorModeToolkit.cpp b/Engine/Plugins/Experimental/ModelingToolsEditorMode/Source/ModelingToolsEditorMode/Private/ModelingToolsEditorModeToolkit.cpp
--- a/Engine/Plugins/Experimental/ModelingToolsEditorMode/Source/ModelingToolsEditorMode/Private/ModelingToolsEditorModeToolkit.cpp
+++ b/Engine/Plugins/Experimental/ModelingToolsEditorMode/Source/ModelingToolsEditorMode/Private/ModelingToolsEditorModeToolkit.cpp
@@ -345,6 +345,7 @@ void FModelingToolsEditorModeToolkit::BuildToolPalette_Experimental(FName Palett
 		ToolbarBuilder.AddToolBarButton(Commands.BeginRevolveBoundaryTool);
 		ToolbarBuilder.AddSeparator();
 		ToolbarBuilder.AddToolBarButton(Commands.BeginMeshBooleanTool);
+		ToolbarBuilder.AddToolBarButton(Commands.BeginMeshTrimTool);
 		ToolbarBuilder.AddToolBarButton(Commands.BeginSelfUnionTool);
 		ToolbarBuilder.AddSeparator();
 		ToolbarBuilder.AddToolBarButton(Commands.BeginVoxelSolidifyTool);
diff --git a/Engine/Plugins/Experimental/ModelingToolsEditorMode/Source/ModelingToolsEditorMode/Private/ModelingToolsManagerActions.cpp b/Engine/Plugins/Experimental/ModelingToolsEditorMode/Source/ModelingToolsEditorMode/Private/ModelingToolsManagerActions.cpp
--- a/Engine/Plugins/Experimental/ModelingToolsEditorMode/Source/ModelingToolsEditorMode/Private/ModelingToolsManagerActions.cpp
+++ b/Engine/Plugins/Experimental/ModelingToolsEditorMode/Source/ModelingToolsEditorMode/Private/ModelingToolsManagerActions.cpp
@@ -74,6 +74,7 @@ void FModelingToolsManagerCommands::RegisterCommands()
 	UI_COMMAND(BeginVoxelMorphologyTool, "VoxMorphology", "Start the Voxel Morphology Tool", EUserInterfaceActionType::ToggleButton, FInputChord());
 	UI_COMMAND(BeginSelfUnionTool, "SelfUnion", "Start the Mesh-Based Self-Union Tool, which will resolve self-intersections", EUserInterfaceActionType::ToggleButton, FInputChord());
 	UI_COMMAND(BeginMeshBooleanTool, "Mesh Boolean", "Start the Mesh-Based Boolean Tool", EUserInterfaceActionType::ToggleButton, FInputChord());
+	UI_COMMAND(BeginMeshTrimTool, "Mesh Trim", "Start the Mesh-Based Trim Tool", EUserInterfaceActionType::ToggleButton, FInputChord());
 	UI_COMMAND(BeginBspConversionTool, "BSP Conversion", "Start the BSP-to-StaticMesh Conversion Tool", EUserInterfaceActionType::ToggleButton, FInputChord());
 	UI_COMMAND(BeginMeshToVolumeTool, "Mesh2Vol", "Start the Mesh to Volume Conversion Tool", EUserInterfaceActionType::ToggleButton, FInputChord());
 	UI_COMMAND(BeginVolumeToMeshTool, "Vol2Mesh", "Start the Volume to Mesh Conversion Tool", EUserInterfaceActionType::ToggleButton, FInputChord());
diff --git a/Engine/Plugins/Experimental/ModelingToolsEditorMode/Source/ModelingToolsEditorMode/Public/ModelingToolsManagerActions.h b/Engine/Plugins/Experimental/ModelingToolsEditorMode/Source/ModelingToolsEditorMode/Public/ModelingToolsManagerActions.h
--- a/Engine/Plugins/Experimental/ModelingToolsEditorMode/Source/ModelingToolsEditorMode/Public/ModelingToolsManagerActions.h
+++ b/Engine/Plugins/Experimental/ModelingToolsEditorMode/Source/ModelingToolsEditorMode/Public/ModelingToolsManagerActions.h
@@ -73,6 +73,7 @@ public:
 	TSharedPtr<FUICommandInfo> BeginVoxelMorphologyTool;
 	TSharedPtr<FUICommandInfo> BeginSelfUnionTool;
 	TSharedPtr<FUICommandInfo> BeginMeshBooleanTool;
+	TSharedPtr<FUICommandInfo> BeginMeshTrimTool;
 	TSharedPtr<FUICommandInfo> BeginMeshSelectionTool;
 	TSharedPtr<FUICommandInfo> BeginBspConversionTool;
 	TSharedPtr<FUICommandInfo> BeginMeshToVolumeTool;
