ModelingComponents: Add FMeshPolygroupChange and supporting classes for creating Polygroup layer changes. Add TSimpleValueLambdaChange for implementing basic value-swap changes. Add FMeshChange::ApplyChangeToMesh() API which allows for subclassing FMeshChange to modify its behavior.
#rnx

diff --git a/Engine/Plugins/Experimental/MeshModelingToolset/Source/ModelingComponents/Private/Changes/MeshChange.cpp b/Engine/Plugins/Experimental/MeshModelingToolset/Source/ModelingComponents/Private/Changes/MeshChange.cpp
--- a/Engine/Plugins/Experimental/MeshModelingToolset/Source/ModelingComponents/Private/Changes/MeshChange.cpp
+++ b/Engine/Plugins/Experimental/MeshModelingToolset/Source/ModelingComponents/Private/Changes/MeshChange.cpp
@@ -40,3 +40,10 @@ FString FMeshChange::ToString() const
 	return FString(TEXT("Mesh Change"));
 }
 
+
+void FMeshChange::ApplyChangeToMesh(FDynamicMesh3* Mesh, bool bRevert) const
+{
+	DynamicMeshChange->Apply(Mesh, bRevert);
+}
+
+
diff --git a/Engine/Plugins/Experimental/MeshModelingToolset/Source/ModelingComponents/Private/Changes/MeshPolygroupChange.cpp b/Engine/Plugins/Experimental/MeshModelingToolset/Source/ModelingComponents/Private/Changes/MeshPolygroupChange.cpp
new file mode 100644
--- /dev/null
+++ b/Engine/Plugins/Experimental/MeshModelingToolset/Source/ModelingComponents/Private/Changes/MeshPolygroupChange.cpp
@@ -0,0 +1,126 @@
+// Copyright Epic Games, Inc. All Rights Reserved.
+
+#include "Changes/MeshPolygroupChange.h"
+#include "DynamicMesh3.h"
+#include "DynamicMeshAttributeSet.h"
+#include "Polygroups/PolygroupUtil.h"
+
+
+void FDynamicMeshGroupEdit::ApplyToMesh(FDynamicMesh3* Mesh, bool bRevert)
+{
+	int32 NumTris = Triangles.Num();
+	if (GroupLayerIndex < 0)
+	{
+		if (ensure(Mesh->HasTriangleGroups()) == false) return;
+
+		for (int32 k = 0; k < NumTris; ++k)
+		{
+			int32 tid = Triangles[k];
+			if (Mesh->IsTriangle(tid))
+			{
+				int32 SetGroup = (bRevert) ? OldGroups[k] : NewGroups[k];
+				Mesh->SetTriangleGroup(tid, SetGroup);
+			}
+		}
+	}
+	else
+	{
+		if ( ensure(Mesh->HasAttributes()) == false || 
+			 ensure(GroupLayerIndex <= Mesh->Attributes()->NumPolygroupLayers()) == false ) return;
+
+		FDynamicMeshPolygroupAttribute* PolygroupLayer = Mesh->Attributes()->GetPolygroupLayer(GroupLayerIndex);
+		if (ensure(PolygroupLayer))
+		{
+			for (int32 k = 0; k < NumTris; ++k)
+			{
+				int32 tid = Triangles[k];
+				if (Mesh->IsTriangle(tid))
+				{
+					int32 SetGroup = (bRevert) ? OldGroups[k] : NewGroups[k];
+					PolygroupLayer->SetValue(tid, SetGroup);
+				}
+			}
+		}
+	}
+}
+
+
+
+
+FDynamicMeshGroupEditBuilder::FDynamicMeshGroupEditBuilder(FDynamicMesh3* Mesh)
+{
+	PolygroupSet = MakePimpl<UE::Geometry::FPolygroupSet>(Mesh);
+	Edit = MakeUnique<FDynamicMeshGroupEdit>();
+}
+
+FDynamicMeshGroupEditBuilder::FDynamicMeshGroupEditBuilder(FDynamicMesh3* Mesh, int32 GroupLayer)
+{
+	PolygroupSet = MakePimpl<UE::Geometry::FPolygroupSet>(Mesh, GroupLayer);
+
+	Edit = MakeUnique<FDynamicMeshGroupEdit>();
+	if (ensure(PolygroupSet->GetPolygroupIndex() != -1))
+	{
+		Edit->GroupLayerIndex = PolygroupSet->GetPolygroupIndex();
+	}
+}
+
+FDynamicMeshGroupEditBuilder::FDynamicMeshGroupEditBuilder(UE::Geometry::FPolygroupSet* PolygroupSetIn)
+{
+	PolygroupSet = MakePimpl<UE::Geometry::FPolygroupSet>(*PolygroupSetIn);
+
+	Edit = MakeUnique<FDynamicMeshGroupEdit>();
+	if (PolygroupSet->GetPolygroupIndex() != -1)
+	{
+		Edit->GroupLayerIndex = PolygroupSet->GetPolygroupIndex();
+	}
+}
+
+
+
+void FDynamicMeshGroupEditBuilder::SaveTriangle(int32 TriangleID)
+{
+	int32 NewGroupID = PolygroupSet->GetGroup(TriangleID);
+	const int32* NewIndex = SavedIndexMap.Find(TriangleID);
+	if (NewIndex == nullptr)
+	{
+		int32 Index = Edit->Triangles.Num();
+		SavedIndexMap.Add(TriangleID, Index);
+		Edit->Triangles.Add(TriangleID);
+		Edit->OldGroups.Add(NewGroupID);
+		Edit->NewGroups.Add(NewGroupID);
+	}
+	else
+	{
+		Edit->NewGroups[*NewIndex] = PolygroupSet->GetGroup(TriangleID);
+	}
+}
+
+
+void FDynamicMeshGroupEditBuilder::SaveTriangle(int32 TriangleID, int32 OldGroup, int32 NewGroup)
+{
+	const int32* NewIndex = SavedIndexMap.Find(TriangleID);
+	if (NewIndex == nullptr)
+	{
+		int32 Index = Edit->Triangles.Num();
+		SavedIndexMap.Add(TriangleID, Index);
+		Edit->Triangles.Add(TriangleID);
+		Edit->OldGroups.Add(OldGroup);
+		Edit->NewGroups.Add(NewGroup);
+	}
+	else
+	{
+		Edit->NewGroups[*NewIndex] = NewGroup;
+	}
+}
+
+
+
+FMeshPolygroupChange::FMeshPolygroupChange(TUniquePtr<FDynamicMeshGroupEdit>&& GroupEditIn)
+{
+	GroupEdit = MoveTemp(GroupEditIn);
+}
+
+void FMeshPolygroupChange::ApplyChangeToMesh(FDynamicMesh3* Mesh, bool bRevert) const
+{
+	GroupEdit->ApplyToMesh(Mesh, bRevert);
+}
diff --git a/Engine/Plugins/Experimental/MeshModelingToolset/Source/ModelingComponents/Public/Changes/BasicChanges.h b/Engine/Plugins/Experimental/MeshModelingToolset/Source/ModelingComponents/Public/Changes/BasicChanges.h
new file mode 100644
--- /dev/null
+++ b/Engine/Plugins/Experimental/MeshModelingToolset/Source/ModelingComponents/Public/Changes/BasicChanges.h
@@ -0,0 +1,32 @@
+// Copyright Epic Games, Inc. All Rights Reserved.
+
+#pragma once
+
+#include "CoreMinimal.h"
+#include "InteractiveToolChange.h"
+
+
+/**
+ * TSimpleValueLambdaChange is a generic FToolCommandChange that swaps between two (template-type) Values.
+ * The swap is applied via a lambda provided by the creator.
+ */
+template<typename ValueType>
+class TSimpleValueLambdaChange : public FToolCommandChange
+{
+public:
+	ValueType FromValue;
+	ValueType ToValue;
+	TUniqueFunction<void(UObject*, const ValueType& From, const ValueType& To, bool bIsRevert)> ValueChangeFunc;
+
+	/** Makes the change to the object */
+	virtual void Apply(UObject* Object) override
+	{
+		ValueChangeFunc(Object, FromValue, ToValue, false);
+	}
+
+	/** Reverts change to the object */
+	virtual void Revert(UObject* Object) override
+	{
+		ValueChangeFunc(Object, ToValue, FromValue, true);
+	}
+};
diff --git a/Engine/Plugins/Experimental/MeshModelingToolset/Source/ModelingComponents/Public/Changes/MeshChange.h b/Engine/Plugins/Experimental/MeshModelingToolset/Source/ModelingComponents/Public/Changes/MeshChange.h
--- a/Engine/Plugins/Experimental/MeshModelingToolset/Source/ModelingComponents/Public/Changes/MeshChange.h
+++ b/Engine/Plugins/Experimental/MeshModelingToolset/Source/ModelingComponents/Public/Changes/MeshChange.h
@@ -37,6 +37,15 @@ public:
 
 	/** Describes this change (for debugging) */
 	virtual FString ToString() const override;
+
+
+	/** 
+	 * Apply embedded DynamicMeshChange to given Mesh. This function is for 
+	 * change-targets to call, when passed a FMeshChange to apply to a Mesh they own. 
+	 * This allows FMeshChange subclasses to customize the change behavior if necessary.
+	 * The default behavior just forwards the call to DynamicMeshChange->Apply(Mesh, bRevert).
+	 */
+	virtual void ApplyChangeToMesh(FDynamicMesh3* Mesh, bool bRevert) const;
 };
 
 
diff --git a/Engine/Plugins/Experimental/MeshModelingToolset/Source/ModelingComponents/Public/Changes/MeshPolygroupChange.h b/Engine/Plugins/Experimental/MeshModelingToolset/Source/ModelingComponents/Public/Changes/MeshPolygroupChange.h
new file mode 100644
--- /dev/null
+++ b/Engine/Plugins/Experimental/MeshModelingToolset/Source/ModelingComponents/Public/Changes/MeshPolygroupChange.h
@@ -0,0 +1,81 @@
+// Copyright Epic Games, Inc. All Rights Reserved.
+
+#pragma once
+
+#include "CoreMinimal.h"
+#include "Templates/PimplPtr.h"
+#include "InteractiveToolChange.h"
+#include "VectorTypes.h"
+#include "Changes/MeshChange.h"
+#include "Polygroups/PolygroupSet.h"
+
+
+/**
+ * FDynamicMeshGroupEdit stores a modification of polygroup IDs on a set of triangles
+ */
+class MODELINGCOMPONENTS_API FDynamicMeshGroupEdit
+{
+public:
+	/** GroupLayerIndex indicates which AttributeSet PolyGroup layer to apply the change to, or if -1, apply to the default on the Mesh */
+	int32 GroupLayerIndex = -1;
+	/** IDs of triangles that are modified */
+	TArray<int32> Triangles;
+	/** Old PolygroupID for each triangle */
+	TArray<int32> OldGroups;
+	/** New PolygroupID for each triangle */
+	TArray<int32> NewGroups;
+
+	virtual ~FDynamicMeshGroupEdit() {}
+
+	virtual void ApplyToMesh(FDynamicMesh3* Mesh, bool bRevert);
+};
+
+
+
+/**
+ * FDynamicMeshGroupEditBuilder builds up a FDynamicMeshGroupEdit incrementally.
+ */
+class MODELINGCOMPONENTS_API FDynamicMeshGroupEditBuilder
+{
+public:
+	FDynamicMeshGroupEditBuilder(FDynamicMesh3* Mesh);
+	FDynamicMeshGroupEditBuilder(FDynamicMesh3* Mesh, int32 GroupLayer);
+	FDynamicMeshGroupEditBuilder(UE::Geometry::FPolygroupSet* PolygroupSet);
+
+	TUniquePtr<FDynamicMeshGroupEdit> ExtractResult() { return MoveTemp(Edit); }
+
+	void SaveTriangle(int32 TriangleID);
+	void SaveTriangle(int32 TriangleID, int32 OldGroup, int32 NewGroup);
+
+	template<typename EnumerableType>
+	void SaveTriangles(EnumerableType Enumerable)
+	{
+		for (int32 tid : Enumerable)
+		{
+			SaveTriangle(tid);
+		}
+	}
+
+protected:
+	TPimplPtr<UE::Geometry::FPolygroupSet> PolygroupSet;
+	TUniquePtr<FDynamicMeshGroupEdit> Edit;
+	TMap<int32, int32> SavedIndexMap;
+};
+
+
+
+
+
+/**
+ * FMeshPolygroupChange stores a change to Polygroup IDs on a set of triangles, as a FDynamicMeshGroupEdit.
+ */
+class MODELINGCOMPONENTS_API FMeshPolygroupChange : public FMeshChange
+{
+public:
+	FMeshPolygroupChange(TUniquePtr<FDynamicMeshGroupEdit>&& GroupEditIn);
+
+	TUniquePtr<FDynamicMeshGroupEdit> GroupEdit;
+
+	virtual void ApplyChangeToMesh(FDynamicMesh3* Mesh, bool bRevert) const override;
+};
+
