Made Interactive Tools Framework capable of operating on arbitrary tool targets that implement the proper interfaces. Made MeshToVolume tool use this system.

#rnx

diff --git a/Engine/Plugins/Experimental/MeshModelingToolset/Source/MeshModelingToolsEditorOnly/Private/ConversionUtils/VolumeMeshDescriptionToolTarget.cpp b/Engine/Plugins/Experimental/MeshModelingToolset/Source/MeshModelingToolsEditorOnly/Private/ConversionUtils/VolumeMeshDescriptionToolTarget.cpp
new file mode 100644
--- /dev/null
+++ b/Engine/Plugins/Experimental/MeshModelingToolset/Source/MeshModelingToolsEditorOnly/Private/ConversionUtils/VolumeMeshDescriptionToolTarget.cpp
@@ -0,0 +1,138 @@
+// Copyright Epic Games, Inc. All Rights Reserved.
+
+#pragma once
+
+#include "ConversionUtils/VolumeMeshDescriptionToolTarget.h"
+
+#include "Components/BrushComponent.h"
+#include "ConversionUtils/VolumeToDynamicMesh.h"
+#include "ConversionUtils/DynamicMeshToVolume.h"
+#include "DynamicMeshToMeshDescription.h"
+#include "Engine/Brush.h"
+#include "GameFramework/Volume.h"
+#include "MeshDescriptionToDynamicMesh.h"
+#include "MeshNormals.h"
+#include "StaticMeshAttributes.h"
+#include "ToolSetupUtil.h"
+
+UVolumeMeshDescriptionToolTarget::UVolumeMeshDescriptionToolTarget()
+{
+	// TODO: These should be user-configurable somewhere
+	VolumeToMeshOptions.bInWorldSpace = false;
+	VolumeToMeshOptions.bSetGroups = true;
+	VolumeToMeshOptions.bMergeVertices = true;
+	VolumeToMeshOptions.bAutoRepairMesh = true;
+	VolumeToMeshOptions.bOptimizeMesh = true;
+}
+
+void UVolumeMeshDescriptionToolTarget::GetMaterialSet(FComponentMaterialSet& MaterialSetOut) const
+{
+	MaterialSetOut.Materials.Reset();
+	if (IsValid() == false)
+	{
+		return;
+	}
+
+	UMaterialInterface* Material = ToolSetupUtil::GetDefaultEditVolumeMaterial();
+	if (Material)
+	{
+		MaterialSetOut.Materials.Add(Material);
+	}
+}
+
+void UVolumeMeshDescriptionToolTarget::CommitMaterialSetUpdate(const FComponentMaterialSet& MaterialSet)
+{
+	check(IsValid());
+	// Do nothing
+}
+
+FMeshDescription* UVolumeMeshDescriptionToolTarget::GetMeshDescription()
+{
+	FTransform transform = GetWorldTransform();
+	if (!ConvertedMeshDescription.IsValid())
+	{
+		UBrushComponent* BrushComponent = Cast<UBrushComponent>(Component);
+		if (!BrushComponent)
+		{
+			return nullptr;
+		}
+		AVolume* Volume = Cast<AVolume>(BrushComponent->GetOwner());
+		if (!Volume)
+		{
+			return nullptr;
+		}
+
+		// Note: We can go directly from a volume to a mesh description using GetBrushMesh() in
+		// Editor.h. However, that path doesn't assign polygroups to the result, which we
+		// typically want when using this target, hence the two-step path we use here.
+		FDynamicMesh3 DynamicMesh;
+		UE::Conversion::VolumeToDynamicMesh(Volume, DynamicMesh, GetVolumeToMeshOptions());
+		FMeshNormals::InitializeMeshToPerTriangleNormals(&DynamicMesh);
+
+		ConvertedMeshDescription = MakeShared<FMeshDescription>();
+		FStaticMeshAttributes Attributes(*ConvertedMeshDescription);
+		Attributes.Register();
+		FDynamicMeshToMeshDescription Converter;
+		Converter.Convert(&DynamicMesh, *ConvertedMeshDescription);
+	}
+	FTransform transform2 = GetWorldTransform();
+	return ConvertedMeshDescription.Get();
+}
+
+void UVolumeMeshDescriptionToolTarget::CommitMeshDescription(const FCommitter& Committer)
+{
+	check(IsValid());
+
+	UBrushComponent* BrushComponent = Cast<UBrushComponent>(Component);
+	check(BrushComponent);
+	AVolume* Volume = Cast<AVolume>(BrushComponent->GetOwner());
+	check(Volume);
+
+	// Let the user fill our mesh description with the Committer
+	if (!ConvertedMeshDescription.IsValid())
+	{
+		ConvertedMeshDescription = MakeShared<FMeshDescription>();
+		FStaticMeshAttributes Attributes(*ConvertedMeshDescription);
+		Attributes.Register();
+	}
+	FCommitterParams CommitParams;
+	CommitParams.MeshDescriptionOut = ConvertedMeshDescription.Get();
+	Committer(CommitParams);
+
+	// The conversion we have right now is from dynamic mesh to volume, so we convert
+	// to dynamic mesh first.
+	FDynamicMesh3 DynamicMesh;
+	FMeshDescriptionToDynamicMesh Converter;
+	Converter.Convert(ConvertedMeshDescription.Get(), DynamicMesh);
+
+	FTransform Transform = GetWorldTransform();
+	UE::Conversion::DynamicMeshToVolume(DynamicMesh, Volume);
+
+	Volume->SetActorTransform(Transform);
+	Volume->PostEditChange();
+}
+
+
+// Factory
+
+bool UVolumeMeshDescriptionToolTargetFactory::CanBuildTarget(UObject* SourceObject, const FToolTargetTypeRequirements& Requirements) const
+{
+	UBrushComponent* BrushComponent = Cast<UBrushComponent>(SourceObject);
+	if (!BrushComponent)
+	{
+		return false;
+	}
+
+	return Cast<AVolume>(BrushComponent->GetOwner()) && Requirements.AreSatisfiedBy(UVolumeMeshDescriptionToolTarget::StaticClass());
+}
+
+UToolTarget* UVolumeMeshDescriptionToolTargetFactory::BuildTarget(UObject* SourceObject, const FToolTargetTypeRequirements& Requirements)
+{
+	UVolumeMeshDescriptionToolTarget* Target = NewObject<UVolumeMeshDescriptionToolTarget>();
+	Target->Component = Cast<UBrushComponent>(SourceObject);
+
+	check(Target->Component && Requirements.AreSatisfiedBy(Target)
+		&& Cast<AVolume>(Cast<UBrushComponent>(SourceObject)->GetOwner()));
+
+	return Target;
+}
\ No newline at end of file
diff --git a/Engine/Plugins/Experimental/MeshModelingToolset/Source/MeshModelingToolsEditorOnly/Private/MeshToVolumeTool.cpp b/Engine/Plugins/Experimental/MeshModelingToolset/Source/MeshModelingToolsEditorOnly/Private/MeshToVolumeTool.cpp
--- a/Engine/Plugins/Experimental/MeshModelingToolset/Source/MeshModelingToolsEditorOnly/Private/MeshToVolumeTool.cpp
+++ b/Engine/Plugins/Experimental/MeshModelingToolset/Source/MeshModelingToolsEditorOnly/Private/MeshToVolumeTool.cpp
@@ -12,6 +12,9 @@
 #include "ToolSetupUtil.h"
 #include "Selections/MeshConnectedComponents.h"
 #include "Selection/ToolSelectionUtil.h"
+#include "TargetInterfaces/MeshDescriptionProvider.h"
+#include "TargetInterfaces/PrimitiveComponentBackedTarget.h"
+#include "ToolTargetManager.h"
 
 #include "Engine/Classes/Engine/BlockingVolume.h"
 #include "Engine/Classes/Components/BrushComponent.h"
@@ -26,22 +29,29 @@
  * ToolBuilder
  */
 
+const FToolTargetTypeRequirements& UMeshToVolumeToolBuilder::GetTargetRequirements() const
+{
+	static FToolTargetTypeRequirements TypeRequirements(nullptr, TArray<const UClass*>({
+		UMeshDescriptionProvider::StaticClass(),
+		UPrimitiveComponentBackedTarget::StaticClass() })
+		);
+	return TypeRequirements;
+}
 
 bool UMeshToVolumeToolBuilder::CanBuildTool(const FToolBuilderState& SceneState) const
 {
 	// We don't want to allow this tool to run on selected volumes
 	return ToolBuilderUtil::CountSelectedActorsOfType<AVolume>(SceneState) == 0 
-		&& ToolBuilderUtil::CountComponents(SceneState, CanMakeComponentTarget) == 1;
+		&& SceneState.TargetManager->CountSelectedAndTargetable(SceneState, GetTargetRequirements()) == 1;
 }
 
 UInteractiveTool* UMeshToVolumeToolBuilder::BuildTool(const FToolBuilderState& SceneState) const
 {
 	UMeshToVolumeTool* NewTool = NewObject<UMeshToVolumeTool>(SceneState.ToolManager);
 
-	UActorComponent* ActorComponent = ToolBuilderUtil::FindFirstComponent(SceneState, CanMakeComponentTarget);
-	UPrimitiveComponent* MeshComponent = Cast<UPrimitiveComponent>(ActorComponent);
-	check(MeshComponent != nullptr);
-	NewTool->SetSelection(MakeComponentTarget(MeshComponent));
+	UToolTarget* Target = SceneState.TargetManager->BuildFirstSelectedTargetable(SceneState, GetTargetRequirements());
+	check(Target);
+	NewTool->SetTarget(Target);
 
 	return NewTool;
 }
@@ -59,17 +69,25 @@ void UMeshToVolumeTool::Setup()
 {
 	UInteractiveTool::Setup();
 
+	IPrimitiveComponentBackedTarget* TargetComponent = Cast<IPrimitiveComponentBackedTarget>(Target);
+	check(TargetComponent);
+
 	PreviewMesh = NewObject<UPreviewMesh>(this);
 	PreviewMesh->bBuildSpatialDataStructure = false;
-	PreviewMesh->CreateInWorld(ComponentTarget->GetOwnerActor()->GetWorld(), FTransform::Identity);
-	PreviewMesh->SetTransform(ComponentTarget->GetWorldTransform());
+	PreviewMesh->CreateInWorld(TargetComponent->GetOwnerActor()->GetWorld(), FTransform::Identity);
+	PreviewMesh->SetTransform(TargetComponent->GetWorldTransform());
 
 	FComponentMaterialSet MaterialSet;
-	ComponentTarget->GetMaterialSet(MaterialSet);
+	TargetComponent->GetMaterialSet(MaterialSet);
 	PreviewMesh->SetMaterials(MaterialSet.Materials);
 
 	PreviewMesh->SetTangentsMode(EDynamicMeshTangentCalcType::ExternallyCalculated);
-	PreviewMesh->InitializeMesh(ComponentTarget->GetMesh());
+
+	// TODO: We really just need a dynamic mesh, not a mesh description (since we're converting the mesh
+	// description to a dynamic mesh), but it requires refactoring PreviewMesh to accept a dynamic mesh. 
+	// Once we do, however, the requirement for this tool will change to IDynamicMeshProvider, and we'll
+	// use that instead.
+	PreviewMesh->InitializeMesh(Cast<IMeshDescriptionProvider>(Target)->GetMeshDescription());
 
 	InputMesh.Copy(*PreviewMesh->GetMesh());
 
@@ -79,7 +97,7 @@ void UMeshToVolumeTool::Setup()
 	VolumeEdgesSet->RegisterComponent();
 
 	// hide input StaticMeshComponent
-	ComponentTarget->SetOwnerVisibility(false);
+	TargetComponent->SetOwnerVisibility(false);
 
 	Settings = NewObject<UMeshToVolumeToolProperties>(this);
 	Settings->RestoreProperties(this);
@@ -111,17 +129,17 @@ void UMeshToVolumeTool::Shutdown(EToolShutdownType ShutdownType)
 	PreviewMesh->Disconnect();
 	PreviewMesh = nullptr;
 
-	ComponentTarget->SetOwnerVisibility(true);
-
+	IPrimitiveComponentBackedTarget* TargetComponent = Cast<IPrimitiveComponentBackedTarget>(Target);
+	TargetComponent->SetOwnerVisibility(true);
 
-	UWorld* TargetWorld = ComponentTarget->GetOwnerActor()->GetWorld();
+	UWorld* TargetWorld = TargetComponent->GetOwnerActor()->GetWorld();
 
 
 	if (ShutdownType == EToolShutdownType::Accept)
 	{
 		GetToolManager()->BeginUndoTransaction(LOCTEXT("MeshToVolumeToolTransactionName", "Create Volume"));
 
-		FTransform TargetTransform = ComponentTarget->GetWorldTransform();
+		FTransform TargetTransform = TargetComponent->GetWorldTransform();
 
 		AVolume* TargetVolume = nullptr;
 
@@ -160,7 +178,7 @@ void UMeshToVolumeTool::Shutdown(EToolShutdownType ShutdownType)
 		ToolSelectionUtil::SetNewActorSelection(GetToolManager(), TargetVolume);
 
 		TArray<AActor*> Actors;
-		Actors.Add(ComponentTarget->GetOwnerActor());
+		Actors.Add(TargetComponent->GetOwnerActor());
 		HandleSourcesProperties->ApplyMethod(Actors, GetToolManager());
 
 
diff --git a/Engine/Plugins/Experimental/MeshModelingToolset/Source/MeshModelingToolsEditorOnly/Public/ConversionUtils/VolumeComponentTarget.h b/Engine/Plugins/Experimental/MeshModelingToolset/Source/MeshModelingToolsEditorOnly/Public/ConversionUtils/VolumeComponentTarget.h
--- a/Engine/Plugins/Experimental/MeshModelingToolset/Source/MeshModelingToolsEditorOnly/Public/ConversionUtils/VolumeComponentTarget.h
+++ b/Engine/Plugins/Experimental/MeshModelingToolset/Source/MeshModelingToolsEditorOnly/Public/ConversionUtils/VolumeComponentTarget.h
@@ -15,6 +15,7 @@ public:
 	TUniquePtr<FPrimitiveComponentTarget> Build(UPrimitiveComponent* PrimitiveComponent) override;
 };
 
+/** Deprecated. Use the tool target instead. */
 class MESHMODELINGTOOLSEDITORONLY_API FVolumeComponentTarget : public FPrimitiveComponentTarget
 {
 public:
diff --git a/Engine/Plugins/Experimental/MeshModelingToolset/Source/MeshModelingToolsEditorOnly/Public/ConversionUtils/VolumeMeshDescriptionToolTarget.h b/Engine/Plugins/Experimental/MeshModelingToolset/Source/MeshModelingToolsEditorOnly/Public/ConversionUtils/VolumeMeshDescriptionToolTarget.h
new file mode 100644
--- /dev/null
+++ b/Engine/Plugins/Experimental/MeshModelingToolset/Source/MeshModelingToolsEditorOnly/Public/ConversionUtils/VolumeMeshDescriptionToolTarget.h
@@ -0,0 +1,57 @@
+// Copyright Epic Games, Inc. All Rights Reserved.
+
+#pragma once
+
+#include "CoreMinimal.h"
+
+#include "ConversionUtils/VolumeToDynamicMesh.h"
+#include "TargetInterfaces/MeshDescriptionCommitter.h"
+#include "TargetInterfaces/MeshDescriptionProvider.h"
+#include "ToolTargets/PrimitiveComponentToolTarget.h"
+
+#include "VolumeMeshDescriptionToolTarget.generated.h"
+
+/**
+ * A tool target backed by AVolume
+ */
+UCLASS(Transient)
+class MESHMODELINGTOOLSEDITORONLY_API UVolumeMeshDescriptionToolTarget : public UPrimitiveComponentToolTarget,
+	public IMeshDescriptionCommitter, public IMeshDescriptionProvider
+{
+	GENERATED_BODY()
+
+public:
+	UVolumeMeshDescriptionToolTarget();
+
+	const UE::Conversion::FVolumeToMeshOptions& GetVolumeToMeshOptions() { return VolumeToMeshOptions; }
+
+	// IMeshDescriptionProvider implementation
+	FMeshDescription* GetMeshDescription() override;
+
+	// IMeshDescritpionCommitter implementation
+	virtual void CommitMeshDescription(const FCommitter& Committer) override;
+
+	// IPrimitiveComponentBackedTarget implementation
+	virtual void GetMaterialSet(FComponentMaterialSet& MaterialSetOut) const override;
+	virtual void CommitMaterialSetUpdate(const FComponentMaterialSet& MaterialSet) override;
+	// Rest provided by parent class
+
+protected:
+	TSharedPtr<FMeshDescription> ConvertedMeshDescription;
+	UE::Conversion::FVolumeToMeshOptions VolumeToMeshOptions;
+
+	friend class UVolumeMeshDescriptionToolTargetFactory;
+};
+
+/** Factory for UStaticMeshComponentToolTarget to be used by the target manager. */
+UCLASS(Transient)
+class MESHMODELINGTOOLSEDITORONLY_API UVolumeMeshDescriptionToolTargetFactory : public UToolTargetFactory
+{
+	GENERATED_BODY()
+
+public:
+
+	virtual bool CanBuildTarget(UObject* SourceObject, const FToolTargetTypeRequirements& TargetTypeInfo) const override;
+
+	virtual UToolTarget* BuildTarget(UObject* SourceObject, const FToolTargetTypeRequirements& TargetTypeInfo) override;
+};
\ No newline at end of file
diff --git a/Engine/Plugins/Experimental/MeshModelingToolset/Source/MeshModelingToolsEditorOnly/Public/MeshToVolumeTool.h b/Engine/Plugins/Experimental/MeshModelingToolset/Source/MeshModelingToolsEditorOnly/Public/MeshToVolumeTool.h
--- a/Engine/Plugins/Experimental/MeshModelingToolset/Source/MeshModelingToolsEditorOnly/Public/MeshToVolumeTool.h
+++ b/Engine/Plugins/Experimental/MeshModelingToolset/Source/MeshModelingToolsEditorOnly/Public/MeshToVolumeTool.h
@@ -29,6 +29,9 @@ public:
 
 	virtual bool CanBuildTool(const FToolBuilderState & SceneState) const override;
 	virtual UInteractiveTool* BuildTool(const FToolBuilderState & SceneState) const override;
+
+protected:
+	virtual const FToolTargetTypeRequirements& GetTargetRequirements() const override;
 };
 
 
diff --git a/Engine/Plugins/Experimental/ModelingToolsEditorMode/Source/ModelingToolsEditorMode/Private/ModelingToolsEditorMode.cpp b/Engine/Plugins/Experimental/ModelingToolsEditorMode/Source/ModelingToolsEditorMode/Private/ModelingToolsEditorMode.cpp
--- a/Engine/Plugins/Experimental/ModelingToolsEditorMode/Source/ModelingToolsEditorMode/Private/ModelingToolsEditorMode.cpp
+++ b/Engine/Plugins/Experimental/ModelingToolsEditorMode/Source/ModelingToolsEditorMode/Private/ModelingToolsEditorMode.cpp
@@ -5,6 +5,9 @@
 #include "InteractiveToolsSelectionStoreSubsystem.h"
 #include "ModelingToolsEditorModeToolkit.h"
 #include "Toolkits/ToolkitManager.h"
+#include "ToolTargetManager.h"
+#include "ToolTargets/StaticMeshComponentToolTarget.h"
+#include "ConversionUtils/VolumeMeshDescriptionToolTarget.h"
 #include "Framework/Commands/UICommandList.h"
 #include "Framework/Application/SlateApplication.h"
 #include "Framework/MultiBox/MultiBoxBuilder.h"
@@ -285,6 +288,10 @@ void UModelingToolsEditorMode::Enter()
 
 	ModelingModeAssetGenerationAPI = MakeShareable(new FModelingModeAssetAPI(ToolsContext->GetAssetAPI()));
 
+	// Register builders for tool targets that the mode uses.
+	ToolsContext->TargetManager->AddTargetFactory(NewObject<UStaticMeshComponentToolTargetFactory>(ToolsContext->TargetManager));
+	ToolsContext->TargetManager->AddTargetFactory(NewObject<UVolumeMeshDescriptionToolTargetFactory>(ToolsContext->TargetManager));
+
 	// register stylus event handler
 	StylusStateTracker = MakeUnique<FStylusStateTracker>();
 
diff --git a/Engine/Plugins/Runtime/ProceduralMeshComponent/Source/ProceduralMeshComponent/Private/ProceduralMeshComponentToolTarget.cpp b/Engine/Plugins/Runtime/ProceduralMeshComponent/Source/ProceduralMeshComponent/Private/ProceduralMeshComponentToolTarget.cpp
new file mode 100644
--- /dev/null
+++ b/Engine/Plugins/Runtime/ProceduralMeshComponent/Source/ProceduralMeshComponent/Private/ProceduralMeshComponentToolTarget.cpp
@@ -0,0 +1,40 @@
+// Copyright Epic Games, Inc. All Rights Reserved.
+
+#include "ProceduralMeshComponentToolTarget.h"
+
+#include "MeshDescription.h"
+#include "ProceduralMeshComponent.h"
+#include "ProceduralMeshConversion.h"
+
+FMeshDescription* UProceduralMeshComponentToolTarget::GetMeshDescription()
+{
+	if (!MeshDescription.IsValid())
+	{
+		MeshDescription = MakeShared<FMeshDescription>();
+	}
+	*MeshDescription = BuildMeshDescription(Cast<UProceduralMeshComponent>(Component));
+	return MeshDescription.Get();
+}
+
+void UProceduralMeshComponentToolTarget::CommitMeshDescription(const FCommitter& Committer)
+{
+	if (!MeshDescription.IsValid())
+	{
+		MeshDescription = MakeShared<FMeshDescription>();
+	}
+	Committer({ MeshDescription.Get() });
+	MeshDescriptionToProcMesh(*MeshDescription, Cast<UProceduralMeshComponent>(Component));
+}
+
+bool UProceduralMeshComponentToolTargetFactory::CanBuildTarget(UObject* SourceObject, const FToolTargetTypeRequirements& Requirements) const
+{
+	return Cast<UProceduralMeshComponent>(SourceObject) && Requirements.AreSatisfiedBy(UProceduralMeshComponentToolTarget::StaticClass());
+}
+
+UToolTarget* UProceduralMeshComponentToolTargetFactory::BuildTarget(UObject* SourceObject, const FToolTargetTypeRequirements& Requirements)
+{
+	UProceduralMeshComponentToolTarget* Target = NewObject<UProceduralMeshComponentToolTarget>();
+	Target->Component = Cast<UProceduralMeshComponent>(SourceObject);
+	check(Target->Component && Requirements.AreSatisfiedBy(Target));
+	return Target;
+}
\ No newline at end of file
diff --git a/Engine/Plugins/Runtime/ProceduralMeshComponent/Source/ProceduralMeshComponent/Public/ProceduralMeshBridge.h b/Engine/Plugins/Runtime/ProceduralMeshComponent/Source/ProceduralMeshComponent/Public/ProceduralMeshBridge.h
--- a/Engine/Plugins/Runtime/ProceduralMeshComponent/Source/ProceduralMeshComponent/Public/ProceduralMeshBridge.h
+++ b/Engine/Plugins/Runtime/ProceduralMeshComponent/Source/ProceduralMeshComponent/Public/ProceduralMeshBridge.h
@@ -14,6 +14,7 @@ public:
 	TUniquePtr<FPrimitiveComponentTarget> Build( UPrimitiveComponent* PrimitiveComponent ) override;
 };
 
+/** @deprecated Use UProceduralMeshComponentToolTarget instead. */
 class PROCEDURALMESHCOMPONENT_API FProceduralMeshComponentTarget : public FPrimitiveComponentTarget
 {
 public:
diff --git a/Engine/Plugins/Runtime/ProceduralMeshComponent/Source/ProceduralMeshComponent/Public/ProceduralMeshComponentToolTarget.h b/Engine/Plugins/Runtime/ProceduralMeshComponent/Source/ProceduralMeshComponent/Public/ProceduralMeshComponentToolTarget.h
new file mode 100644
--- /dev/null
+++ b/Engine/Plugins/Runtime/ProceduralMeshComponent/Source/ProceduralMeshComponent/Public/ProceduralMeshComponentToolTarget.h
@@ -0,0 +1,35 @@
+// Copyright Epic Games, Inc. All Rights Reserved.
+
+#pragma once
+
+#include "CoreMinimal.h"
+
+#include "TargetInterfaces/MeshDescriptionCommitter.h"
+#include "TargetInterfaces/MeshDescriptionProvider.h"
+#include "ToolTargets/PrimitiveComponentToolTarget.h"
+
+#include "ProceduralMeshComponentToolTarget.generated.h"
+UCLASS(Transient)
+class PROCEDURALMESHCOMPONENT_API UProceduralMeshComponentToolTarget : public UPrimitiveComponentToolTarget,
+	public IMeshDescriptionCommitter, public IMeshDescriptionProvider
+{
+	GENERATED_BODY()
+
+public:
+	virtual FMeshDescription* GetMeshDescription() override;
+	virtual void CommitMeshDescription(const FCommitter&) override;
+private:
+	TSharedPtr<FMeshDescription> MeshDescription;
+
+	friend class UProceduralMeshComponentToolTargetFactory;
+};
+
+UCLASS(Transient)
+class PROCEDURALMESHCOMPONENT_API  UProceduralMeshComponentToolTargetFactory : public UToolTargetFactory
+{
+	GENERATED_BODY()
+
+public:
+	virtual bool CanBuildTarget(UObject* SourceObject, const FToolTargetTypeRequirements& TargetTypeInfo) const override;
+	virtual UToolTarget* BuildTarget(UObject* SourceObject, const FToolTargetTypeRequirements& TargetTypeInfo) override;
+};
diff --git a/Engine/Source/Editor/Experimental/EditorInteractiveToolsFramework/EditorInteractiveToolsFramework.Build.cs b/Engine/Source/Editor/Experimental/EditorInteractiveToolsFramework/EditorInteractiveToolsFramework.Build.cs
--- a/Engine/Source/Editor/Experimental/EditorInteractiveToolsFramework/EditorInteractiveToolsFramework.Build.cs
+++ b/Engine/Source/Editor/Experimental/EditorInteractiveToolsFramework/EditorInteractiveToolsFramework.Build.cs
@@ -37,7 +37,8 @@ public class EditorInteractiveToolsFramework : ModuleRules
 			{
 				"CoreUObject",
 				"Engine",
-                "Slate",
+				"RenderCore",
+				"Slate",
                 "SlateCore",
                 "InputCore",
 				"EditorFramework",
diff --git a/Engine/Source/Editor/Experimental/EditorInteractiveToolsFramework/Private/EdModeInteractiveToolsContext.cpp b/Engine/Source/Editor/Experimental/EditorInteractiveToolsFramework/Private/EdModeInteractiveToolsContext.cpp
--- a/Engine/Source/Editor/Experimental/EditorInteractiveToolsFramework/Private/EdModeInteractiveToolsContext.cpp
+++ b/Engine/Source/Editor/Experimental/EditorInteractiveToolsFramework/Private/EdModeInteractiveToolsContext.cpp
@@ -24,7 +24,6 @@
 
 #include "ToolContextInterfaces.h"
 #include "Tools/EditorToolAssetAPI.h"
-#include "Tools/EditorComponentSourceFactory.h"
 #include "InteractiveToolObjects.h"
 #include "InteractiveToolsSelectionStoreSubsystem.h"
 #include "BaseBehaviors/ClickDragBehavior.h"
@@ -176,6 +175,7 @@ public:
 	virtual void GetCurrentSelectionState(FToolBuilderState& StateOut) const override
 	{
 		StateOut.ToolManager = ToolsContext->ToolManager;
+		StateOut.TargetManager = ToolsContext->TargetManager;
 		StateOut.GizmoManager = ToolsContext->GizmoManager;
 		StateOut.World = EditorModeManager->GetWorld();
 		EditorModeManager->GetSelectedActors()->GetSelectedObjects(StateOut.SelectedActors);
diff --git a/Engine/Source/Editor/Experimental/EditorInteractiveToolsFramework/Private/ToolTargets/StaticMeshComponentToolTarget.cpp b/Engine/Source/Editor/Experimental/EditorInteractiveToolsFramework/Private/ToolTargets/StaticMeshComponentToolTarget.cpp
new file mode 100644
--- /dev/null
+++ b/Engine/Source/Editor/Experimental/EditorInteractiveToolsFramework/Private/ToolTargets/StaticMeshComponentToolTarget.cpp
@@ -0,0 +1,105 @@
+// Copyright Epic Games, Inc. All Rights Reserved.
+
+#pragma once
+
+#include "ToolTargets/StaticMeshComponentToolTarget.h"
+
+#include "ComponentReregisterContext.h"
+#include "Components/StaticMeshComponent.h"
+#include "Engine/StaticMesh.h"
+#include "RenderingThread.h"
+
+namespace UStaticMeshComponentToolTargetLocals
+{
+	int32 LODIndex = 0;
+}
+
+void UStaticMeshComponentToolTarget::CommitMaterialSetUpdate(const FComponentMaterialSet& MaterialSet)
+{
+	check(IsValid());
+
+	UStaticMesh* StaticMesh = Cast<UStaticMeshComponent>(Component)->GetStaticMesh();
+
+	// flush any pending rendering commands, which might touch this component while we are rebuilding it's mesh
+	FlushRenderingCommands();
+
+	// unregister the component while we update it's static mesh
+	TUniquePtr<FComponentReregisterContext> ComponentReregisterContext = MakeUnique<FComponentReregisterContext>(Component);
+
+	// make sure transactional flag is on
+	StaticMesh->SetFlags(RF_Transactional);
+
+	bool bSavedToTransactionBuffer = StaticMesh->Modify();
+	check(bSavedToTransactionBuffer);
+
+	int NewNumMaterials = MaterialSet.Materials.Num();
+	if (NewNumMaterials != StaticMesh->GetStaticMaterials().Num())
+	{
+		StaticMesh->GetStaticMaterials().SetNum(NewNumMaterials);
+	}
+	for (int k = 0; k < NewNumMaterials; ++k)
+	{
+		if (StaticMesh->GetMaterial(k) != MaterialSet.Materials[k])
+		{
+			StaticMesh->SetMaterial(k, MaterialSet.Materials[k]);
+		}
+	}
+
+	// right?
+	StaticMesh->PostEditChange();
+}
+
+FMeshDescription* UStaticMeshComponentToolTarget::GetMeshDescription()
+{
+	using namespace UStaticMeshComponentToolTargetLocals;
+
+	return IsValid() ? Cast<UStaticMeshComponent>(Component)->GetStaticMesh()->GetMeshDescription(LODIndex) : nullptr;
+}
+
+void UStaticMeshComponentToolTarget::CommitMeshDescription(const FCommitter& Committer)
+{
+	using namespace UStaticMeshComponentToolTargetLocals;
+	check(IsValid());
+
+	UStaticMesh* StaticMesh = Cast<UStaticMeshComponent>(Component)->GetStaticMesh();
+
+	// flush any pending rendering commands, which might touch this component while we are rebuilding it's mesh
+	FlushRenderingCommands();
+
+	// unregister the component while we update its static mesh
+	FComponentReregisterContext ComponentReregisterContext(Component);
+
+	// make sure transactional flag is on for this asset
+	StaticMesh->SetFlags(RF_Transactional);
+
+	bool bSavedToTransactionBuffer = StaticMesh->Modify();
+	check(bSavedToTransactionBuffer);
+
+	FCommitterParams CommitterParams;
+	CommitterParams.MeshDescriptionOut = StaticMesh->GetMeshDescription(LODIndex);
+
+	Committer(CommitterParams);
+
+	StaticMesh->CommitMeshDescription(LODIndex);
+	StaticMesh->PostEditChange();
+
+	// this rebuilds physics, but it doesn't undo!
+	Component->RecreatePhysicsState();
+}
+
+
+// Factory
+
+bool UStaticMeshComponentToolTargetFactory::CanBuildTarget(UObject* SourceObject, const FToolTargetTypeRequirements& Requirements) const
+{
+	return Cast<UStaticMeshComponent>(SourceObject) && Requirements.AreSatisfiedBy(UStaticMeshComponentToolTarget::StaticClass());
+}
+
+UToolTarget* UStaticMeshComponentToolTargetFactory::BuildTarget(UObject* SourceObject, const FToolTargetTypeRequirements& Requirements)
+{
+	UStaticMeshComponentToolTarget* Target = NewObject<UStaticMeshComponentToolTarget>();// TODO: Should we set an outer here?
+	Target->Component = Cast<UStaticMeshComponent>(SourceObject);
+	check(Target->Component && Requirements.AreSatisfiedBy(Target));
+
+	return Target;
+}
\ No newline at end of file
diff --git a/Engine/Source/Editor/Experimental/EditorInteractiveToolsFramework/Public/ToolTargets/StaticMeshComponentToolTarget.h b/Engine/Source/Editor/Experimental/EditorInteractiveToolsFramework/Public/ToolTargets/StaticMeshComponentToolTarget.h
new file mode 100644
--- /dev/null
+++ b/Engine/Source/Editor/Experimental/EditorInteractiveToolsFramework/Public/ToolTargets/StaticMeshComponentToolTarget.h
@@ -0,0 +1,51 @@
+// Copyright Epic Games, Inc. All Rights Reserved.
+
+#pragma once
+
+#include "CoreMinimal.h"
+
+#include "TargetInterfaces/MeshDescriptionCommitter.h"
+#include "TargetInterfaces/MeshDescriptionProvider.h"
+#include "ToolTargets/PrimitiveComponentToolTarget.h"
+
+#include "StaticMeshComponentToolTarget.generated.h"
+
+/**
+ * A tool target backed by a static mesh component that can provide and take a mesh
+ * description.
+ */
+UCLASS(Transient)
+class EDITORINTERACTIVETOOLSFRAMEWORK_API UStaticMeshComponentToolTarget : public UPrimitiveComponentToolTarget,
+	public IMeshDescriptionCommitter, public IMeshDescriptionProvider
+{
+	GENERATED_BODY()
+
+public:
+
+	// IMeshDescriptionProvider implementation
+	FMeshDescription* GetMeshDescription() override;
+
+	// IMeshDescritpionCommitter implementation
+	virtual void CommitMeshDescription(const FCommitter& Committer) override;
+
+	// IPrimitiveComponentBackedTarget implementation
+	virtual void CommitMaterialSetUpdate(const FComponentMaterialSet& MaterialSet) override;
+	// Rest provided by parent class
+
+protected:
+
+	friend class UStaticMeshComponentToolTargetFactory;
+};
+
+/** Factory for UStaticMeshComponentToolTarget to be used by the target manager. */
+UCLASS(Transient)
+class EDITORINTERACTIVETOOLSFRAMEWORK_API UStaticMeshComponentToolTargetFactory : public UToolTargetFactory
+{
+	GENERATED_BODY()
+
+public:
+
+	virtual bool CanBuildTarget(UObject* SourceObject, const FToolTargetTypeRequirements& TargetTypeInfo) const override;
+
+	virtual UToolTarget* BuildTarget(UObject* SourceObject, const FToolTargetTypeRequirements& TargetTypeInfo) override;
+};
\ No newline at end of file
diff --git a/Engine/Source/Runtime/Experimental/InteractiveToolsFramework/Private/ComponentSourceInterfaces.cpp b/Engine/Source/Runtime/Experimental/InteractiveToolsFramework/Private/ComponentSourceInterfaces.cpp
--- a/Engine/Source/Runtime/Experimental/InteractiveToolsFramework/Private/ComponentSourceInterfaces.cpp
+++ b/Engine/Source/Runtime/Experimental/InteractiveToolsFramework/Private/ComponentSourceInterfaces.cpp
@@ -42,25 +42,6 @@ TUniquePtr<FPrimitiveComponentTarget> MakeComponentTarget(UPrimitiveComponent* C
 
 
 
-
-bool FComponentMaterialSet::operator!=(const FComponentMaterialSet& Other) const
-{
-	int32 Num = Materials.Num();
-	if (Other.Materials.Num() != Num)
-	{
-		return true;
-	}
-	for (int32 j = 0; j < Num; ++j)
-	{
-		if (Other.Materials[j] != Materials[j])
-		{
-			return true;
-		}
-	}
-	return false;
-}
-
-
 bool FPrimitiveComponentTarget::IsValid() const
 {
 	return (Component->IsPendingKillOrUnreachable() == false) && Component->IsValidLowLevel();
diff --git a/Engine/Source/Runtime/Experimental/InteractiveToolsFramework/Private/InteractiveToolsContext.cpp b/Engine/Source/Runtime/Experimental/InteractiveToolsFramework/Private/InteractiveToolsContext.cpp
--- a/Engine/Source/Runtime/Experimental/InteractiveToolsFramework/Private/InteractiveToolsContext.cpp
+++ b/Engine/Source/Runtime/Experimental/InteractiveToolsFramework/Private/InteractiveToolsContext.cpp
@@ -2,12 +2,13 @@
 
 
 #include "InteractiveToolsContext.h"
-
+#include "ToolTargetManager.h"
 
 UInteractiveToolsContext::UInteractiveToolsContext()
 {
 	InputRouter = nullptr;
 	ToolManager = nullptr;
+	TargetManager = nullptr;
 	ToolManagerClass = UInteractiveToolManager::StaticClass();
 }
 
@@ -17,9 +18,11 @@ void UInteractiveToolsContext::Initialize(IToolsContextQueriesAPI* QueriesAPI, I
 	InputRouter->Initialize(TransactionsAPI);
 
 	ToolManager = NewObject<UInteractiveToolManager>(this, ToolManagerClass.Get());
-
 	ToolManager->Initialize(QueriesAPI, TransactionsAPI, InputRouter);
 
+	TargetManager = NewObject<UToolTargetManager>(this);
+	TargetManager->Initialize();
+
 	GizmoManager = NewObject<UInteractiveGizmoManager>(this);
 	GizmoManager->Initialize(QueriesAPI, TransactionsAPI, InputRouter);
 
diff --git a/Engine/Source/Runtime/Experimental/InteractiveToolsFramework/Private/ToolTargetManager.cpp b/Engine/Source/Runtime/Experimental/InteractiveToolsFramework/Private/ToolTargetManager.cpp
new file mode 100644
--- /dev/null
+++ b/Engine/Source/Runtime/Experimental/InteractiveToolsFramework/Private/ToolTargetManager.cpp
@@ -0,0 +1,68 @@
+// Copyright Epic Games, Inc. All Rights Reserved.
+
+
+#include "ToolTargetManager.h"
+#include "InteractiveToolsContext.h"
+#include "ToolBuilderUtil.h"
+
+void UToolTargetManager::Initialize()
+{
+	bIsActive = true;
+}
+
+void UToolTargetManager::Shutdown()
+{
+	Factories.Empty();
+	bIsActive = false;
+}
+
+void UToolTargetManager::AddTargetFactory(UToolTargetFactory* Factory)
+{
+	Factories.AddUnique(Factory);
+}
+
+bool UToolTargetManager::CanBuildTarget(UObject* SourceObject, const FToolTargetTypeRequirements& TargetType) const
+{
+	for (UToolTargetFactory* Factory : Factories)
+	{
+		if (Factory->CanBuildTarget(SourceObject, TargetType))
+		{
+			return true;
+		}
+	}
+	return false;
+}
+
+UToolTarget* UToolTargetManager::BuildTarget(UObject* SourceObject, const FToolTargetTypeRequirements& TargetType)
+{
+	for (UToolTargetFactory* Factory : Factories)
+	{
+		if (Factory->CanBuildTarget(SourceObject, TargetType))
+		{
+			UToolTarget* Result = Factory->BuildTarget(SourceObject, TargetType);
+			if (Result != nullptr)
+			{
+				return Result;
+			}
+		}
+	}
+	return nullptr;
+}
+
+int32 UToolTargetManager::CountSelectedAndTargetable(const FToolBuilderState& SceneState, const FToolTargetTypeRequirements& TargetType) const
+{
+	return ToolBuilderUtil::CountComponents(SceneState, [&](UActorComponent* Object)
+		{
+			return CanBuildTarget(Object, TargetType);
+		});
+}
+
+UToolTarget* UToolTargetManager::BuildFirstSelectedTargetable(const FToolBuilderState& SceneState, const FToolTargetTypeRequirements& TargetType)
+{
+	return BuildTarget(
+		ToolBuilderUtil::FindFirstComponent(SceneState, [&](UActorComponent* Object)
+		{
+			return CanBuildTarget(Object, TargetType);
+		}),
+		TargetType);
+}
diff --git a/Engine/Source/Runtime/Experimental/InteractiveToolsFramework/Private/ToolTargets/PrimitiveComponentToolTarget.cpp b/Engine/Source/Runtime/Experimental/InteractiveToolsFramework/Private/ToolTargets/PrimitiveComponentToolTarget.cpp
new file mode 100644
--- /dev/null
+++ b/Engine/Source/Runtime/Experimental/InteractiveToolsFramework/Private/ToolTargets/PrimitiveComponentToolTarget.cpp
@@ -0,0 +1,66 @@
+// Copyright Epic Games, Inc. All Rights Reserved.
+
+#include "ToolTargets/PrimitiveComponentToolTarget.h"
+
+#include "Components/PrimitiveComponent.h"
+
+bool UPrimitiveComponentToolTarget::IsValid() const
+{
+	return Component && !Component->IsPendingKillOrUnreachable() && Component->IsValidLowLevel();
+}
+
+UPrimitiveComponent* UPrimitiveComponentToolTarget::GetOwnerComponent() const
+{
+	return IsValid() ? Component : nullptr;
+}
+
+AActor* UPrimitiveComponentToolTarget::GetOwnerActor() const
+{
+	return IsValid() ? Component->GetOwner() : nullptr;
+}
+
+void UPrimitiveComponentToolTarget::SetOwnerVisibility(bool bVisible) const
+{
+	if (IsValid())
+	{
+		Component->SetVisibility(bVisible);
+	}
+}
+
+int32 UPrimitiveComponentToolTarget::GetNumMaterials() const
+{
+	return IsValid() ? Component->GetNumMaterials() : 0;
+}
+
+UMaterialInterface* UPrimitiveComponentToolTarget::GetMaterial(int32 MaterialIndex) const
+{
+	return IsValid() ? Component->GetMaterial(MaterialIndex) : nullptr;
+}
+
+void UPrimitiveComponentToolTarget::GetMaterialSet(FComponentMaterialSet& MaterialSetOut) const
+{
+	if (IsValid())
+	{
+		int32 NumMaterials = Component->GetNumMaterials();
+		MaterialSetOut.Materials.SetNum(NumMaterials);
+		for (int32 k = 0; k < NumMaterials; ++k)
+		{
+			MaterialSetOut.Materials[k] = Component->GetMaterial(k);
+		}
+	}
+}
+
+FTransform UPrimitiveComponentToolTarget::GetWorldTransform() const
+{
+	return IsValid() ? Component->GetComponentTransform() : FTransform::Identity;
+}
+
+bool UPrimitiveComponentToolTarget::HitTestComponent(const FRay& WorldRay, FHitResult& OutHit) const
+{
+	FVector End = WorldRay.PointAt(HALF_WORLD_MAX);
+	if (IsValid() && Component->LineTraceComponent(OutHit, WorldRay.Origin, End, FCollisionQueryParams(SCENE_QUERY_STAT(HitTest), true)))
+	{
+		return true;
+	}
+	return false;
+}
diff --git a/Engine/Source/Runtime/Experimental/InteractiveToolsFramework/Private/ToolTargets/ToolTarget.cpp b/Engine/Source/Runtime/Experimental/InteractiveToolsFramework/Private/ToolTargets/ToolTarget.cpp
new file mode 100644
--- /dev/null
+++ b/Engine/Source/Runtime/Experimental/InteractiveToolsFramework/Private/ToolTargets/ToolTarget.cpp
@@ -0,0 +1,27 @@
+// Copyright Epic Games, Inc. All Rights Reserved.
+
+#include "ToolTargets/ToolTarget.h"
+
+bool FToolTargetTypeRequirements::AreSatisfiedBy(UClass* Class) const
+{
+	// Required type must either be null (no requirement) or our type, or our parent type.
+	if (!BaseType || BaseType == Class ||
+		Class->IsChildOf(BaseType))
+	{
+		// we have to support all the required interfaces
+		for (const UClass* Interface : Interfaces)
+		{
+			if (!Class->ImplementsInterface(Interface))
+			{
+				return false;
+			}
+		}
+		return true;
+	}
+	return false;
+}
+
+bool FToolTargetTypeRequirements::AreSatisfiedBy(UToolTarget* ToolTarget) const
+{
+	return !ToolTarget ? false : AreSatisfiedBy(ToolTarget->GetClass());
+}
\ No newline at end of file
diff --git a/Engine/Source/Runtime/Experimental/InteractiveToolsFramework/Public/ComponentSourceInterfaces.h b/Engine/Source/Runtime/Experimental/InteractiveToolsFramework/Public/ComponentSourceInterfaces.h
--- a/Engine/Source/Runtime/Experimental/InteractiveToolsFramework/Public/ComponentSourceInterfaces.h
+++ b/Engine/Source/Runtime/Experimental/InteractiveToolsFramework/Public/ComponentSourceInterfaces.h
@@ -6,6 +6,7 @@
 #include "Math/UnrealMath.h"
 #include "Misc/Optional.h"
 #include "Templates/Function.h"
+#include "TargetInterfaces/PrimitiveComponentBackedTarget.h" // FComponentMaterialSet
 #include "Engine/EngineTypes.h"    // FHitResult
 
 // predeclarations
@@ -16,16 +17,8 @@ struct FMeshDescription;
 class UMaterialInterface;
 
 /**
- * FComponentMaterialSet is the set of materials assigned to a component (ie Material Slots on a StaticMesh)
- */
-struct INTERACTIVETOOLSFRAMEWORK_API FComponentMaterialSet
-{
-	TArray<UMaterialInterface*> Materials;
-	bool operator!=(const FComponentMaterialSet& Other) const;
-};
-
-
-/**
+ * @deprecated Use tool targets instead, such as UPrimitiveComponentToolTarget
+ *
  * Wrapper around a UObject Component that can provide a MeshDescription, and
  * (optionally) bake a modified MeshDescription back to this Component.
  * An example of a Source might be a StaticMeshComponent. How a modified
@@ -116,6 +109,7 @@ public:
 	UPrimitiveComponent* Component{};
 };
 
+/** @deprecated Use tool target factories instead. */
 class INTERACTIVETOOLSFRAMEWORK_API FComponentTargetFactory
 {
 public:
@@ -125,6 +119,8 @@ public:
 };
 
 /**
+ * @deprecated Add tool target factories to the tool target manager instead.
+ *
  * Add a factory method to make ComponentTarget from UPrimitiveComponent*
  * @param Factory The ComponentTargetFactory
  * @return void
@@ -133,6 +129,8 @@ INTERACTIVETOOLSFRAMEWORK_API void AddComponentTargetFactory( TUniquePtr<FCompon
 
 
 /**
+ * @deprecated Ask your tool target manager to build a target for you instead.
+ *
  * Create a TargetComponent for the given Component
  * @param Component A UObject that we would like to use as tool target. This must presently descend from
  * UPrimitiveComponent
@@ -143,6 +141,8 @@ INTERACTIVETOOLSFRAMEWORK_API TUniquePtr<FPrimitiveComponentTarget> MakeComponen
 
 
 /**
+ * @deprecated Ask your tool target manager instead.
+ *
  * Determine whether a TargetComponent can be created for the given Component
  * @param Component A UObject that we would like to use as tool target.
  * @return bool signifying whether or not a ComponentTarget can be built
diff --git a/Engine/Source/Runtime/Experimental/InteractiveToolsFramework/Public/InteractiveToolBuilder.h b/Engine/Source/Runtime/Experimental/InteractiveToolsFramework/Public/InteractiveToolBuilder.h
--- a/Engine/Source/Runtime/Experimental/InteractiveToolsFramework/Public/InteractiveToolBuilder.h
+++ b/Engine/Source/Runtime/Experimental/InteractiveToolsFramework/Public/InteractiveToolBuilder.h
@@ -3,8 +3,9 @@
 #pragma once
 
 #include "CoreMinimal.h"
-#include "UObject/Interface.h"
 #include "InteractiveTool.h"
+#include "ToolTargets/ToolTarget.h"
+
 #include "InteractiveToolBuilder.generated.h"
 
 
@@ -43,4 +44,15 @@ public:
 		check(false);
 		return nullptr;
 	}
+
+protected:
+	/**
+	 * Gives the target requirements of the associated tool. Usually, it is the tool builder
+	 * will use this function in CanBuildTool and BuildTool to find and create any necessary targets.
+	 */
+	virtual const FToolTargetTypeRequirements& GetTargetRequirements() const
+	{
+		static FToolTargetTypeRequirements TypeRequirements; // Default initialized to no requirements.
+		return TypeRequirements;
+	}
 };
\ No newline at end of file
diff --git a/Engine/Source/Runtime/Experimental/InteractiveToolsFramework/Public/InteractiveToolsContext.h b/Engine/Source/Runtime/Experimental/InteractiveToolsFramework/Public/InteractiveToolsContext.h
--- a/Engine/Source/Runtime/Experimental/InteractiveToolsFramework/Public/InteractiveToolsContext.h
+++ b/Engine/Source/Runtime/Experimental/InteractiveToolsFramework/Public/InteractiveToolsContext.h
@@ -7,6 +7,8 @@
 #include "InteractiveGizmoManager.h"
 #include "InteractiveToolsContext.generated.h"
 
+class UToolTargetManager;
+
 /**
  * InteractiveToolsContext owns a ToolManager and an InputRouter. This is just a top-level 
  * UObject container, however implementations like UEdModeInteractiveToolsContext extend
@@ -58,6 +60,10 @@ public:
 	UPROPERTY()
 	UInputRouter* InputRouter;	
 
+	/** current UToolTargetManager for this Context */
+	UPROPERTY()
+	UToolTargetManager* TargetManager;
+
 	/** current UInteractiveToolManager for this Context */
 	UPROPERTY()
 	UInteractiveToolManager* ToolManager;	
diff --git a/Engine/Source/Runtime/Experimental/InteractiveToolsFramework/Public/SingleSelectionTool.h b/Engine/Source/Runtime/Experimental/InteractiveToolsFramework/Public/SingleSelectionTool.h
--- a/Engine/Source/Runtime/Experimental/InteractiveToolsFramework/Public/SingleSelectionTool.h
+++ b/Engine/Source/Runtime/Experimental/InteractiveToolsFramework/Public/SingleSelectionTool.h
@@ -4,6 +4,8 @@
 
 #include "InteractiveTool.h"
 #include "ComponentSourceInterfaces.h"
+#include "ToolTargets/ToolTarget.h"
+
 #include "SingleSelectionTool.generated.h"
 
 UCLASS(Transient)
@@ -11,17 +13,24 @@ class INTERACTIVETOOLSFRAMEWORK_API USingleSelectionTool : public UInteractiveTo
 {
 GENERATED_BODY()
 public:
+	// @deprecated Use SetTarget instead, and don't use FPrimitiveComponentTarget.
 	void SetSelection(TUniquePtr<FPrimitiveComponentTarget> ComponentTargetIn)
     {
 		ComponentTarget = MoveTemp(ComponentTargetIn);
 	}
 
+	void SetTarget(UToolTarget* TargetIn)
+	{
+		Target = TargetIn;
+	}
+
 	/**
 	 * @return true if all ComponentTargets of this tool are still valid
 	 */
 	virtual bool AreAllTargetsValid() const
 	{
-		return ComponentTarget->IsValid();
+		// TODO: This needs to be updated once tools no longer use ComponentTarget.
+		return Target ? Target->IsValid() : ComponentTarget->IsValid();
 	}
 
 
@@ -32,7 +41,9 @@ public:
 	}
 
 protected:
+	/** @deprecated Tools should use Target instead. */
 	TUniquePtr<FPrimitiveComponentTarget> ComponentTarget{};
 
-
+	UPROPERTY()
+	UToolTarget* Target;
 };
diff --git a/Engine/Source/Runtime/Experimental/InteractiveToolsFramework/Public/TargetInterfaces/MeshDescriptionCommitter.h b/Engine/Source/Runtime/Experimental/InteractiveToolsFramework/Public/TargetInterfaces/MeshDescriptionCommitter.h
new file mode 100644
--- /dev/null
+++ b/Engine/Source/Runtime/Experimental/InteractiveToolsFramework/Public/TargetInterfaces/MeshDescriptionCommitter.h
@@ -0,0 +1,59 @@
+// Copyright Epic Games, Inc. All Rights Reserved.
+
+#pragma once
+
+#include "CoreMinimal.h"
+
+#include "MeshDescription.h"
+#include "UObject/Interface.h"
+
+#include "MeshDescriptionCommitter.generated.h"
+
+UINTERFACE()
+class INTERACTIVETOOLSFRAMEWORK_API UMeshDescriptionCommitter : public UInterface
+{
+	GENERATED_BODY()
+};
+
+class INTERACTIVETOOLSFRAMEWORK_API IMeshDescriptionCommitter
+{
+	GENERATED_BODY()
+
+public:
+
+	struct FCommitterParams
+	{
+		/** 
+		 * Mesh description that should be populated/updated by the passed-in function and which
+		 * will be committed to the target.
+		 */
+		FMeshDescription* MeshDescriptionOut  = nullptr;
+	};
+	using  FCommitter = TFunction<void(const FCommitterParams&)>;
+
+	/** 
+	 * Commit a mesh description. The mesh description to be committed will be passed to the
+	 * given function as a parameter, and it is up to the function to update it properly.
+	 *
+	 * @param Committer A function that takes in const IMeshDescriptionCommitter::FCommitParams&
+	 *  and populates the FMeshDescription pointed to by the MeshDescription pointer inside.
+	 */
+	virtual void CommitMeshDescription(const FCommitter& Committer) = 0;
+
+	/**
+	 * Commits the given FMeshDescription.
+	 */
+	virtual void CommitMeshDescription(const FMeshDescription& Mesh)
+	{
+		// It seems reasonable to have this function, but we'll go ahead and give a default implementation
+		// if users want to just implement the other one.
+
+		CommitMeshDescription([&Mesh](const FCommitterParams& CommitParams)
+		{
+			if (CommitParams.MeshDescriptionOut)
+			{
+				*CommitParams.MeshDescriptionOut = Mesh;
+			}
+		});
+	}
+};
\ No newline at end of file
diff --git a/Engine/Source/Runtime/Experimental/InteractiveToolsFramework/Public/TargetInterfaces/MeshDescriptionProvider.h b/Engine/Source/Runtime/Experimental/InteractiveToolsFramework/Public/TargetInterfaces/MeshDescriptionProvider.h
new file mode 100644
--- /dev/null
+++ b/Engine/Source/Runtime/Experimental/InteractiveToolsFramework/Public/TargetInterfaces/MeshDescriptionProvider.h
@@ -0,0 +1,28 @@
+// Copyright Epic Games, Inc. All Rights Reserved.
+
+#pragma once
+
+#include "CoreMinimal.h"
+#include "UObject/Interface.h"
+
+#include "MeshDescriptionProvider.generated.h"
+
+struct FMeshDescription;
+
+UINTERFACE()
+class INTERACTIVETOOLSFRAMEWORK_API UMeshDescriptionProvider : public UInterface
+{
+	GENERATED_BODY()
+};
+
+class INTERACTIVETOOLSFRAMEWORK_API IMeshDescriptionProvider
+{
+	GENERATED_BODY()
+
+public:
+	virtual FMeshDescription* GetMeshDescription() = 0;
+
+	//TODO: Should this be providing a shared pointer instead? That seems like generally
+	// the better idea, but some refactoring may be necessary to avoid having to copy
+	// when getting from a static mesh.
+};
\ No newline at end of file
diff --git a/Engine/Source/Runtime/Experimental/InteractiveToolsFramework/Public/TargetInterfaces/PrimitiveComponentBackedTarget.h b/Engine/Source/Runtime/Experimental/InteractiveToolsFramework/Public/TargetInterfaces/PrimitiveComponentBackedTarget.h
new file mode 100644
--- /dev/null
+++ b/Engine/Source/Runtime/Experimental/InteractiveToolsFramework/Public/TargetInterfaces/PrimitiveComponentBackedTarget.h
@@ -0,0 +1,90 @@
+// Copyright Epic Games, Inc. All Rights Reserved.
+
+#pragma once
+
+#include "CoreMinimal.h"
+#include "Engine/EngineTypes.h"    // FHitResult
+#include "UObject/Interface.h"
+
+#include "PrimitiveComponentBackedTarget.generated.h"
+
+class UMaterialInterface;
+class UPrimitiveComponent;
+class AActor;
+struct FTransform;
+
+/**
+ * FComponentMaterialSet is the set of materials assigned to a component (ie Material Slots on a StaticMesh)
+ */
+struct INTERACTIVETOOLSFRAMEWORK_API FComponentMaterialSet
+{
+	TArray<UMaterialInterface*> Materials;
+
+	bool operator!=(const FComponentMaterialSet& Other) const
+	{
+		return Materials != Other.Materials;
+	}
+};
+
+UINTERFACE()
+class INTERACTIVETOOLSFRAMEWORK_API UPrimitiveComponentBackedTarget : public UInterface
+{
+	GENERATED_BODY()
+};
+
+class INTERACTIVETOOLSFRAMEWORK_API IPrimitiveComponentBackedTarget
+{
+	GENERATED_BODY()
+
+public:
+
+	/** @return the Component this is a Source for */
+	virtual UPrimitiveComponent* GetOwnerComponent() const = 0;
+
+	/** @return the Actor that owns this Component */
+	virtual AActor* GetOwnerActor() const = 0;
+
+	/**
+	 * Set the visibility of the Component associated with this Source (ie to hide during Tool usage)
+	 * @param bVisible desired visibility
+	 */
+	virtual void SetOwnerVisibility(bool bVisible) const = 0;
+
+	/** @return number of material indices in use by this Component */
+	virtual int32 GetNumMaterials() const = 0;
+
+	/**
+	 * Get pointer to a Material provided by this Source
+	 * @param MaterialIndex index of the material
+	 * @return MaterialInterface pointer, or null if MaterialIndex is invalid
+	 */
+	virtual UMaterialInterface* GetMaterial(int32 MaterialIndex) const = 0;
+
+	/**
+	 * Get material set provided by this source
+	 * @param MaterialSetOut returned material set
+	 */
+	virtual void GetMaterialSet(FComponentMaterialSet& MaterialSetOut) const = 0;
+
+	/**
+	 * @return the transform on this component
+	 * @todo Do we need to return a list of transforms here?
+	 */
+	virtual FTransform GetWorldTransform() const = 0;
+
+	/**
+	 * Compute ray intersection with the MeshDescription this Source is providing
+	 * @param WorldRay ray in world space
+	 * @param OutHit hit test data
+	 * @return true if ray intersected Component
+	 */
+	virtual bool HitTestComponent(const FRay& WorldRay, FHitResult& OutHit) const = 0;
+
+
+	// TODO: Should this be in a separate interface?
+	/**
+	 * Commit an update to the material set. This may generate a transaction.
+	 * @param MaterialSet new list of materials
+	 */
+	virtual void CommitMaterialSetUpdate(const FComponentMaterialSet& MaterialSet) = 0;
+};
\ No newline at end of file
diff --git a/Engine/Source/Runtime/Experimental/InteractiveToolsFramework/Public/ToolContextInterfaces.h b/Engine/Source/Runtime/Experimental/InteractiveToolsFramework/Public/ToolContextInterfaces.h
--- a/Engine/Source/Runtime/Experimental/InteractiveToolsFramework/Public/ToolContextInterfaces.h
+++ b/Engine/Source/Runtime/Experimental/InteractiveToolsFramework/Public/ToolContextInterfaces.h
@@ -18,6 +18,7 @@ class FPrimitiveDrawInterface;
 class FSceneView;
 class UInteractiveToolManager;
 class UInteractiveGizmoManager;
+class UToolTargetManager;
 struct FMeshDescription;
 class UTexture2D;
 class UInteractiveToolStorableSelection;
@@ -37,6 +38,8 @@ struct INTERACTIVETOOLSFRAMEWORK_API FToolBuilderState
 	UWorld* World = nullptr;
 	/** The current ToolManager */
 	UInteractiveToolManager* ToolManager = nullptr;
+	/** The current TargetManager */
+	UToolTargetManager* TargetManager = nullptr;
 	/** The current GizmoManager */
 	UInteractiveGizmoManager* GizmoManager = nullptr;
 
diff --git a/Engine/Source/Runtime/Experimental/InteractiveToolsFramework/Public/ToolTargetManager.h b/Engine/Source/Runtime/Experimental/InteractiveToolsFramework/Public/ToolTargetManager.h
new file mode 100644
--- /dev/null
+++ b/Engine/Source/Runtime/Experimental/InteractiveToolsFramework/Public/ToolTargetManager.h
@@ -0,0 +1,94 @@
+// Copyright Epic Games, Inc. All Rights Reserved.
+
+#pragma once
+
+#include "CoreMinimal.h"
+#include "ToolTargets/ToolTarget.h"
+#include "ToolContextInterfaces.h"
+
+#include "ToolTargetManager.generated.h"
+
+// TODO: Do we need more control over factory order to prioritize which gets used if we can
+// make multiple qualifying targets? It should theoretically not matter, but in practice, one
+// target or another could be more efficient for certain tasks. This is probably only worth
+// thinking about once we encounter it.
+
+/**
+ * The tool target manager converts input objects into tool targets- objects that
+ * can expose various interfaces that tools might expect but which the original
+ * objects may not know about.
+ *
+ * Someday, the tool target manager may implement caching of targets.
+ */
+UCLASS(Transient)
+class INTERACTIVETOOLSFRAMEWORK_API UToolTargetManager : public UObject
+{
+	GENERATED_BODY()
+
+public:
+	/**
+	 * @return true if ToolManager is currently active, ie between Initialize() and Shutdown()
+	 */
+	bool IsActive() const { return bIsActive; }
+
+	virtual void AddTargetFactory(UToolTargetFactory* Factory);
+
+	/** Examines stored target factories to see if one can build the requested type of target. */
+	virtual bool CanBuildTarget(UObject* SourceObject, const FToolTargetTypeRequirements& TargetRequirements) const;
+
+	/** 
+	 * Uses one of the stored factories to build a tool target out of the given input object
+	 * that satisfies the given requirements. If multiple factories are capable of building a 
+	 * qualifying target, the first encountered one will be used. If none are capable, a nullptr 
+	 * is returned.
+	 * 
+	 * @return Tool target that staisfies given requirements, or nullptr if none could be created.
+	 */
+	virtual UToolTarget* BuildTarget(UObject* SourceObject, const FToolTargetTypeRequirements& TargetRequirements);
+
+	/** Much like BuildTarget, but casts the target to the template argument before returning. */
+	template<typename CastToType>
+	CastToType* BuildTarget(UObject* SourceObject, const FToolTargetTypeRequirements& TargetType)
+	{
+		UToolTarget* Result = BuildTarget(SourceObject, TargetType);
+		return (Result != nullptr) ? Cast<CastToType>(Result) : nullptr;
+	}
+
+	/**
+	 * Looks through the currently selected components and actors and counts the number of
+	 * inputs that could be used to create qualifying tool targets.
+	 */
+	virtual int32 CountSelectedAndTargetable(const FToolBuilderState& SceneState,
+		const FToolTargetTypeRequirements& TargetRequirements) const;
+	
+	/**
+	 * Looks through the currently selected components and actors and builds a target out of
+	 * the first encountered element that satisfies the requirements. 
+	 */
+	virtual UToolTarget* BuildFirstSelectedTargetable(const FToolBuilderState& SceneState,
+		const FToolTargetTypeRequirements& TargetRequirements);
+
+protected:
+	friend class UInteractiveToolsContext;		// to call Initialize/Shutdown
+
+	// This should be removed once tools are transitioned to using tool targets, and the
+	// functions in ComponentSourceInterfaces.h do not exist. For now, this is here to call
+	// Initialize().
+	friend INTERACTIVETOOLSFRAMEWORK_API void AddFactoryToDeprecatedToolTargetManager(UToolTargetFactory* Factory);
+
+	UToolTargetManager(){};
+
+	/** Initialize the ToolTargetManager. UInteractiveToolsContext calls this, you should not. */
+	virtual void Initialize();
+
+	/** Shutdown the ToolTargetManager. Called by UInteractiveToolsContext. */
+	virtual void Shutdown();
+
+	/** This flag is set to true on Initialize() and false on Shutdown(). */
+	bool bIsActive = false;
+
+	UPROPERTY()
+	TArray<UToolTargetFactory*> Factories;
+
+	// More state will go here if the manager deals with target caching.
+};
\ No newline at end of file
diff --git a/Engine/Source/Runtime/Experimental/InteractiveToolsFramework/Public/ToolTargets/PrimitiveComponentToolTarget.h b/Engine/Source/Runtime/Experimental/InteractiveToolsFramework/Public/ToolTargets/PrimitiveComponentToolTarget.h
new file mode 100644
--- /dev/null
+++ b/Engine/Source/Runtime/Experimental/InteractiveToolsFramework/Public/ToolTargets/PrimitiveComponentToolTarget.h
@@ -0,0 +1,37 @@
+// Copyright Epic Games, Inc. All Rights Reserved.
+
+#pragma once
+
+#include "CoreMinimal.h"
+
+#include "TargetInterfaces/PrimitiveComponentBackedTarget.h"
+#include "ToolTargets/ToolTarget.h"
+
+#include "PrimitiveComponentToolTarget.generated.h"
+
+/** 
+ * An abstract tool target to share some reusable code for tool targets that are
+ * backed by primitive components. 
+ */
+UCLASS(Transient, Abstract)
+class INTERACTIVETOOLSFRAMEWORK_API UPrimitiveComponentToolTarget : public UToolTarget, public IPrimitiveComponentBackedTarget
+{
+	GENERATED_BODY()
+public:
+	virtual bool IsValid() const override;
+
+	// IPrimitiveComponentBackedTarget implementation
+	virtual UPrimitiveComponent* GetOwnerComponent() const override;
+	virtual AActor* GetOwnerActor() const override;
+	virtual void SetOwnerVisibility(bool bVisible) const override;
+	virtual int32 GetNumMaterials() const override;
+	virtual UMaterialInterface* GetMaterial(int32 MaterialIndex) const override;
+	virtual void GetMaterialSet(FComponentMaterialSet& MaterialSetOut) const override;
+	virtual FTransform GetWorldTransform() const override;
+	virtual bool HitTestComponent(const FRay& WorldRay, FHitResult& OutHit) const override;
+
+	// Not implemented by this class:
+	virtual void CommitMaterialSetUpdate(const FComponentMaterialSet& MaterialSet) override PURE_VIRTUAL(UPrimitiveComponentToolTarget::CommitMaterialSetUpdate);
+protected:
+	UPrimitiveComponent* Component;
+};
\ No newline at end of file
diff --git a/Engine/Source/Runtime/Experimental/InteractiveToolsFramework/Public/ToolTargets/ToolTarget.h b/Engine/Source/Runtime/Experimental/InteractiveToolsFramework/Public/ToolTargets/ToolTarget.h
new file mode 100644
--- /dev/null
+++ b/Engine/Source/Runtime/Experimental/InteractiveToolsFramework/Public/ToolTargets/ToolTarget.h
@@ -0,0 +1,80 @@
+// Copyright Epic Games, Inc. All Rights Reserved.
+
+#pragma once
+
+#include "CoreMinimal.h"
+
+#include "ToolTarget.generated.h"
+
+class FToolTargetTypeRequirements;
+
+/**
+ * A tool target is a stand-in object that a tool can operate on. It can implement any
+ * interface(s) that a tool requires without having to implement those interfaces in
+ * specific actor classes, and it allows the tools to work on anything that can provide
+ * a qualifying tool target.
+ */
+UCLASS(Transient, Abstract)
+class INTERACTIVETOOLSFRAMEWORK_API UToolTarget : public UObject
+{
+	GENERATED_BODY()
+public:
+
+	/** @return true if target is still valid. May become invalid for various reasons (eg Component was deleted out from under us) */
+	virtual bool IsValid() const PURE_VIRTUAL(UToolTarget::IsValid, return false;);
+};
+
+
+/**
+ * A structure used to specify the requirements of a tool for its target. E.g., a tool
+ * may need a target that has base type x and interfaces w,y,z.
+ */
+class INTERACTIVETOOLSFRAMEWORK_API FToolTargetTypeRequirements
+{
+public:
+
+	const UClass* BaseType = nullptr; //TODO: Do we even need this? Seems like we shouldn't use it.
+	TArray<const UClass*, TInlineAllocator<2>> Interfaces;
+
+	FToolTargetTypeRequirements()
+	{
+	}
+
+	explicit FToolTargetTypeRequirements(const UClass* BaseTypeIn)
+	{
+		BaseType = BaseTypeIn;
+	}
+
+	explicit FToolTargetTypeRequirements(const UClass* BaseTypeIn, const UClass* Interface0)
+	{
+		BaseType = BaseTypeIn;
+		Interfaces.Add(Interface0);
+	}
+
+	explicit FToolTargetTypeRequirements(const UClass* BaseTypeIn, const TArray<const UClass*>& InterfacesIn)
+	{
+		BaseType = BaseTypeIn;
+		Interfaces = InterfacesIn;
+	}
+
+	bool AreSatisfiedBy(UClass* Class) const;
+
+	bool AreSatisfiedBy(UToolTarget* ToolTarget) const;
+};
+
+
+/**
+ * Base class for factories of tool targets, which let a tool manager build targets
+ * out of inputs without knowing anything itself.
+ */
+UCLASS(Transient)
+class INTERACTIVETOOLSFRAMEWORK_API UToolTargetFactory : public UObject
+{
+	GENERATED_BODY()
+
+public:
+
+	virtual bool CanBuildTarget(UObject* SourceObject, const FToolTargetTypeRequirements& TargetTypeInfo) const PURE_VIRTUAL(UToolTargetFactory::CanBuildTarget, return false;);
+
+	virtual UToolTarget* BuildTarget(UObject* SourceObject, const FToolTargetTypeRequirements& TargetTypeInfo) PURE_VIRTUAL(UToolTargetFactory::BuildTarget, return nullptr;);
+};
