ModelingComponents: don't try to fast update if a MeshPostProcessor is enabled on SimpleDynamicMeshComponent
#rnx

diff --git a/Engine/Plugins/Experimental/MeshModelingToolset/Source/MeshModelingToolsEditorOnly/Private/SubdividePolyTool.cpp b/Engine/Plugins/Experimental/MeshModelingToolset/Source/MeshModelingToolsEditorOnly/Private/SubdividePolyTool.cpp
--- a/Engine/Plugins/Experimental/MeshModelingToolset/Source/MeshModelingToolsEditorOnly/Private/SubdividePolyTool.cpp
+++ b/Engine/Plugins/Experimental/MeshModelingToolset/Source/MeshModelingToolsEditorOnly/Private/SubdividePolyTool.cpp
@@ -238,16 +238,13 @@ void USubdividePolyTool::Setup()
 			PreviewMesh->SetTriangleColorFunction([](const FDynamicMesh3* Mesh, int TriangleID)
 			{
 				return LinearColors::SelectFColor(Mesh->GetTriangleGroup(TriangleID));
-			});
+			}, UPreviewMesh::ERenderUpdateMode::FullUpdate);
 		}
 		else
 		{
 			PreviewMesh->SetOverrideRenderMaterial(nullptr);
-			PreviewMesh->SetTriangleColorFunction(nullptr);
+			PreviewMesh->SetTriangleColorFunction(nullptr, UPreviewMesh::ERenderUpdateMode::FullUpdate);
 		}
-
-		USimpleDynamicMeshComponent* PreviewDynamicMeshComponent = (USimpleDynamicMeshComponent*)PreviewMesh->GetRootComponent();
-		PreviewDynamicMeshComponent->FastNotifyColorsUpdated();
 	};
 
 	Properties->WatchProperty(Properties->bRenderGroups, RenderGroupsChanged);
diff --git a/Engine/Plugins/Experimental/MeshModelingToolset/Source/ModelingComponents/Private/SimpleDynamicMeshComponent.cpp b/Engine/Plugins/Experimental/MeshModelingToolset/Source/ModelingComponents/Private/SimpleDynamicMeshComponent.cpp
--- a/Engine/Plugins/Experimental/MeshModelingToolset/Source/ModelingComponents/Private/SimpleDynamicMeshComponent.cpp
+++ b/Engine/Plugins/Experimental/MeshModelingToolset/Source/ModelingComponents/Private/SimpleDynamicMeshComponent.cpp
@@ -219,11 +219,11 @@ const FMeshTangentsf* USimpleDynamicMeshComponent::GetTangents()
 	}
 
 	// in this mode we assume the tangents are valid
-	check(TangentsType == EDynamicMeshTangentCalcType::ExternallyCalculated);
+	ensure(TangentsType == EDynamicMeshTangentCalcType::ExternallyCalculated);
 	if (TangentsType == EDynamicMeshTangentCalcType::ExternallyCalculated)
 	{
 		// if you hit this, you did not request ExternallyCalculated tangents before initializing this PreviewMesh
-		check(Tangents.GetTangents().Num() > 0)
+		ensure(Tangents.GetTangents().Num() > 0);
 	}
 
 	return &Tangents;
@@ -239,6 +239,8 @@ void USimpleDynamicMeshComponent::SetDrawOnTop(bool bSet)
 
 void USimpleDynamicMeshComponent::ResetProxy()
 {
+	bProxyValid = false;
+
 	// Need to recreate scene proxy to send it over
 	MarkRenderStateDirty();
 	LocalBounds = Mesh->GetCachedBounds();
@@ -263,14 +265,16 @@ void USimpleDynamicMeshComponent::NotifyMeshUpdated()
 
 void USimpleDynamicMeshComponent::FastNotifyColorsUpdated()
 {
-	if (RenderMeshPostProcessor)
+	// should not be using fast paths if we have to run mesh postprocessor
+	if (ensure(!RenderMeshPostProcessor) == false)
 	{
 		RenderMeshPostProcessor->ProcessMesh(*Mesh, *RenderMesh);
+		ResetProxy();
+		return;
 	}
 
 	FSimpleDynamicMeshSceneProxy* Proxy = GetCurrentSceneProxy();
-
-	if (Proxy && Proxy->RenderMeshLayoutMatchesRenderBuffers())
+	if (Proxy)
 	{
 		if (TriangleColorFunc != nullptr &&  Proxy->bUsePerTriangleColor == false )
 		{
@@ -296,14 +300,16 @@ void USimpleDynamicMeshComponent::FastNotifyColorsUpdated()
 
 void USimpleDynamicMeshComponent::FastNotifyPositionsUpdated(bool bNormals, bool bColors, bool bUVs)
 {
-	if (RenderMeshPostProcessor)
+	// should not be using fast paths if we have to run mesh postprocessor
+	if (ensure(!RenderMeshPostProcessor) == false)
 	{
 		RenderMeshPostProcessor->ProcessMesh(*Mesh, *RenderMesh);
+		ResetProxy();
+		return;
 	}
 
 	FSimpleDynamicMeshSceneProxy* Proxy = GetCurrentSceneProxy();
-
-	if (Proxy && Proxy->RenderMeshLayoutMatchesRenderBuffers())
+	if (Proxy)
 	{
 		GetCurrentSceneProxy()->FastUpdateVertices(true, bNormals, bColors, bUVs);
 		//MarkRenderDynamicDataDirty();
@@ -320,15 +326,16 @@ void USimpleDynamicMeshComponent::FastNotifyPositionsUpdated(bool bNormals, bool
 
 void USimpleDynamicMeshComponent::FastNotifyVertexAttributesUpdated(bool bNormals, bool bColors, bool bUVs)
 {
-	if (RenderMeshPostProcessor)
+	// should not be using fast paths if we have to run mesh postprocessor
+	if (ensure(!RenderMeshPostProcessor) == false)
 	{
 		RenderMeshPostProcessor->ProcessMesh(*Mesh, *RenderMesh);
+		ResetProxy();
+		return;
 	}
 
 	FSimpleDynamicMeshSceneProxy* Proxy = GetCurrentSceneProxy();
-
-	check(bNormals || bColors || bUVs);
-	if (Proxy && Proxy->RenderMeshLayoutMatchesRenderBuffers())
+	if (Proxy && ensure(bNormals || bColors || bUVs) )
 	{
 		GetCurrentSceneProxy()->FastUpdateVertices(false, bNormals, bColors, bUVs);
 		//MarkRenderDynamicDataDirty();
@@ -343,15 +350,16 @@ void USimpleDynamicMeshComponent::FastNotifyVertexAttributesUpdated(bool bNormal
 
 void USimpleDynamicMeshComponent::FastNotifyVertexAttributesUpdated(EMeshRenderAttributeFlags UpdatedAttributes)
 {
-	if (RenderMeshPostProcessor)
+	// should not be using fast paths if we have to run mesh postprocessor
+	if (ensure(!RenderMeshPostProcessor) == false)
 	{
 		RenderMeshPostProcessor->ProcessMesh(*Mesh, *RenderMesh);
+		ResetProxy();
+		return;
 	}
 
 	FSimpleDynamicMeshSceneProxy* Proxy = GetCurrentSceneProxy();
-
-	check(UpdatedAttributes != EMeshRenderAttributeFlags::None);
-	if (Proxy && Proxy->RenderMeshLayoutMatchesRenderBuffers())
+	if (Proxy && ensure(UpdatedAttributes != EMeshRenderAttributeFlags::None))
 	{
 		bool bPositions = (UpdatedAttributes & EMeshRenderAttributeFlags::Positions) != EMeshRenderAttributeFlags::None;
 		GetCurrentSceneProxy()->FastUpdateVertices(bPositions,
@@ -382,14 +390,16 @@ void USimpleDynamicMeshComponent::FastNotifyUVsUpdated()
 
 void USimpleDynamicMeshComponent::FastNotifySecondaryTrianglesChanged()
 {
-	if (RenderMeshPostProcessor)
+	// should not be using fast paths if we have to run mesh postprocessor
+	if (ensure(!RenderMeshPostProcessor) == false)
 	{
 		RenderMeshPostProcessor->ProcessMesh(*Mesh, *RenderMesh);
+		ResetProxy();
+		return;
 	}
 
 	FSimpleDynamicMeshSceneProxy* Proxy = GetCurrentSceneProxy();
-
-	if (Proxy && Proxy->RenderMeshLayoutMatchesRenderBuffers())
+	if (Proxy)
 	{
 		GetCurrentSceneProxy()->FastUpdateAllIndexBuffers();
 	}
@@ -402,14 +412,16 @@ void USimpleDynamicMeshComponent::FastNotifySecondaryTrianglesChanged()
 
 void USimpleDynamicMeshComponent::FastNotifyTriangleVerticesUpdated(const TArray<int32>& Triangles, EMeshRenderAttributeFlags UpdatedAttributes)
 {
-	if (RenderMeshPostProcessor)
+	// should not be using fast paths if we have to run mesh postprocessor
+	if (ensure(!RenderMeshPostProcessor) == false)
 	{
 		RenderMeshPostProcessor->ProcessMesh(*Mesh, *RenderMesh);
+		ResetProxy();
+		return;
 	}
 
 	FSimpleDynamicMeshSceneProxy* Proxy = GetCurrentSceneProxy();
-
-	if (!Proxy || !Proxy->RenderMeshLayoutMatchesRenderBuffers())
+	if (!Proxy)
 	{
 		ResetProxy();
 	}
@@ -445,14 +457,17 @@ void USimpleDynamicMeshComponent::FastNotifyTriangleVerticesUpdated(const TArray
 
 void USimpleDynamicMeshComponent::FastNotifyTriangleVerticesUpdated(const TSet<int32>& Triangles, EMeshRenderAttributeFlags UpdatedAttributes)
 {
-	if (RenderMeshPostProcessor)
+	// should not be using fast paths if we have to run mesh postprocessor
+	if (ensure(!RenderMeshPostProcessor) == false)
 	{
 		RenderMeshPostProcessor->ProcessMesh(*Mesh, *RenderMesh);
+		ResetProxy();
+		return;
 	}
 
 	FSimpleDynamicMeshSceneProxy* Proxy = GetCurrentSceneProxy();
 
-	if (!Proxy || !Proxy->RenderMeshLayoutMatchesRenderBuffers())
+	if (!Proxy)
 	{
 		ResetProxy();
 	}
@@ -509,7 +524,8 @@ void USimpleDynamicMeshComponent::FastNotifyTriangleVerticesUpdated(const TSet<i
 
 FPrimitiveSceneProxy* USimpleDynamicMeshComponent::CreateSceneProxy()
 {
-	check(GetCurrentSceneProxy() == nullptr);
+	// if this is not always the case, we have made incorrect assumptions
+	ensure(GetCurrentSceneProxy() == nullptr);
 
 	FSimpleDynamicMeshSceneProxy* NewProxy = nullptr;
 	if (Mesh->TriangleCount() > 0)
@@ -540,10 +556,13 @@ FPrimitiveSceneProxy* USimpleDynamicMeshComponent::CreateSceneProxy()
 			NewProxy->Initialize();
 		}
 	}
+
+	bProxyValid = true;
 	return NewProxy;
 }
 
 
+
 void USimpleDynamicMeshComponent::NotifyMaterialSetUpdated()
 {
 	if (GetCurrentSceneProxy() != nullptr)
@@ -653,7 +672,7 @@ void USimpleDynamicMeshComponent::ApplyChange(const FMeshVertexChange* Change, b
 
 void USimpleDynamicMeshComponent::ApplyChange(const FMeshChange* Change, bool bRevert)
 {
-	Change->DynamicMeshChange->Apply(Mesh.Get(), bRevert);
+	Change->ApplyChangeToMesh(Mesh.Get(), bRevert);
 
 	if (bInvalidateProxyOnChange)
 	{
diff --git a/Engine/Plugins/Experimental/MeshModelingToolset/Source/ModelingComponents/Public/SimpleDynamicMeshComponent.h b/Engine/Plugins/Experimental/MeshModelingToolset/Source/ModelingComponents/Public/SimpleDynamicMeshComponent.h
--- a/Engine/Plugins/Experimental/MeshModelingToolset/Source/ModelingComponents/Public/SimpleDynamicMeshComponent.h
+++ b/Engine/Plugins/Experimental/MeshModelingToolset/Source/ModelingComponents/Public/SimpleDynamicMeshComponent.h
@@ -288,7 +288,16 @@ protected:
 
 private:
 
-	FSimpleDynamicMeshSceneProxy* GetCurrentSceneProxy() { return (FSimpleDynamicMeshSceneProxy*)SceneProxy; }
+	/**
+	 * If the render proxy is invalidated (eg by MarkRenderStateDirty()), it will be destroyed at the end of
+	 * the frame, but the base SceneProxy pointer is not nulled out immediately. As a result if we call various
+	 * partial-update functions after invalidating the proxy, they may be operating on an invalid proxy.
+	 * So we have to keep track of proxy-valid state ourselves.
+	 */
+	bool bProxyValid = false;
+
+	/** @return current render proxy, if valid, otherwise nullptr */
+	FSimpleDynamicMeshSceneProxy* GetCurrentSceneProxy() { return (bProxyValid) ? (FSimpleDynamicMeshSceneProxy*)SceneProxy : nullptr; }
 
 	// Called from NotifyMeshUpdated, as well as the FastNotify functions if needed
 	void ResetProxy();
