ModelingToolset: export static utility function
#rnx

diff --git a/Engine/Plugins/Experimental/MeshModelingToolset/Source/MeshModelingTools/Public/Physics/CollisionGeometryVisualization.h b/Engine/Plugins/Experimental/MeshModelingToolset/Source/MeshModelingTools/Public/Physics/CollisionGeometryVisualization.h
--- a/Engine/Plugins/Experimental/MeshModelingToolset/Source/MeshModelingTools/Public/Physics/CollisionGeometryVisualization.h
+++ b/Engine/Plugins/Experimental/MeshModelingToolset/Source/MeshModelingTools/Public/Physics/CollisionGeometryVisualization.h
@@ -13,7 +13,7 @@ namespace UE
 		 * Create line sets in a UPreviewGeometry for all the elements in a Physics Data Collection.
 		 * Spheres and Capsules are drawn as 3-axis wireframes. Convexes are added as wireframes.
 		 */
-		void InitializePreviewGeometryLines(const FPhysicsDataCollection& PhysicsData, UPreviewGeometry* PreviewGeom,
+		void MESHMODELINGTOOLS_API InitializePreviewGeometryLines(const FPhysicsDataCollection& PhysicsData, UPreviewGeometry* PreviewGeom,
 			const FColor& LineColor, float LineThickness, float DepthBias = 0.0, int32 CircleStepResolution = 16 );
 	}
 }
