ModelingComponents:
Add UMeshWireframeComponent, this is a component that renders mesh edges  as lines, with line color/thickness varying depending on line type. Accesses mesh via an API so client does not have to handle iterations/etc (derived from LineSetComponent).
Add UMeshElementsVisualizer, subclass of UPreviewGeometry that displays mesh information (currently using UMeshWireframeComponent)
Small optimization to ULineSetComponent to avoid calling MarkRenderStateDirty() for each line when adding many lines.
#rnx

diff --git a/Engine/Plugins/Experimental/MeshModelingToolset/Source/ModelingComponents/Private/Drawing/LineSetComponent.cpp b/Engine/Plugins/Experimental/MeshModelingToolset/Source/ModelingComponents/Private/Drawing/LineSetComponent.cpp
--- a/Engine/Plugins/Experimental/MeshModelingToolset/Source/ModelingComponents/Private/Drawing/LineSetComponent.cpp
+++ b/Engine/Plugins/Experimental/MeshModelingToolset/Source/ModelingComponents/Private/Drawing/LineSetComponent.cpp
@@ -155,7 +155,7 @@ public:
 					Mesh.MaterialRenderProxy = MeshBatchData.MaterialProxy;
 
 					FDynamicPrimitiveUniformBuffer& DynamicPrimitiveUniformBuffer = Collector.AllocateOneFrameResource<FDynamicPrimitiveUniformBuffer>();
-					DynamicPrimitiveUniformBuffer.Set(GetLocalToWorld(), GetLocalToWorld(), GetBounds(), GetLocalBounds(), true, false, DrawsVelocity(), false);
+					DynamicPrimitiveUniformBuffer.Set(GetLocalToWorld(), GetLocalToWorld(), GetBounds(), GetLocalBounds(), false, false, DrawsVelocity(), false);
 					BatchElement.PrimitiveUniformBufferResource = &DynamicPrimitiveUniformBuffer.UniformBuffer;
 
 					BatchElement.FirstIndex = MeshBatchData.StartIndex;
@@ -241,12 +241,18 @@ void ULineSetComponent::ReserveLines(const int32 MaxID)
 
 int32 ULineSetComponent::AddLine(const FRenderableLine& OverlayLine)
 {
-	const int32 ID(Lines.Add(OverlayLine));
 	MarkRenderStateDirty();
+	return AddLineInternal(OverlayLine);
+}
+
+int32 ULineSetComponent::AddLineInternal(const FRenderableLine& Line)
+{
+	const int32 ID(Lines.Add(Line));
 	bBoundsDirty = true;
 	return ID;
 }
 
+
 void ULineSetComponent::InsertLine(const int32 ID, const FRenderableLine& OverlayLine)
 {
 	Lines.Insert(ID, OverlayLine);
@@ -370,7 +376,8 @@ FBoxSphereBounds ULineSetComponent::CalcBounds(const FTransform& LocalToWorld) c
 void ULineSetComponent::AddLines(
 	int32 NumIndices,
 	TFunctionRef<void(int32 Index, TArray<FRenderableLine>& LinesOut)> LineGenFunc,
-	int32 LinesPerIndexHint)
+	int32 LinesPerIndexHint,
+	bool bDeferRenderStateDirty)
 {
 	TArray<FRenderableLine> TempLines;
 	if (LinesPerIndexHint > 0)
@@ -385,7 +392,12 @@ void ULineSetComponent::AddLines(
 		LineGenFunc(k, TempLines);
 		for (const FRenderableLine& Line : TempLines)
 		{
-			AddLine(Line);
+			AddLineInternal(Line);
 		}
 	}
+
+	if (!bDeferRenderStateDirty)
+	{
+		MarkRenderStateDirty();
+	}
 }
diff --git a/Engine/Plugins/Experimental/MeshModelingToolset/Source/ModelingComponents/Private/Drawing/MeshElementsVisualizer.cpp b/Engine/Plugins/Experimental/MeshModelingToolset/Source/ModelingComponents/Private/Drawing/MeshElementsVisualizer.cpp
new file mode 100644
--- /dev/null
+++ b/Engine/Plugins/Experimental/MeshModelingToolset/Source/ModelingComponents/Private/Drawing/MeshElementsVisualizer.cpp
@@ -0,0 +1,181 @@
+// Copyright Epic Games, Inc. All Rights Reserved.
+
+#include "Drawing/MeshElementsVisualizer.h"
+#include "DynamicMesh3.h"
+#include "DynamicMeshAttributeSet.h"
+#include "Drawing/LineSetComponent.h"
+#include "Drawing/MeshWireframeComponent.h"
+#include "ToolSetupUtil.h"
+#include "Async/Async.h"
+
+
+
+
+/**
+ * IMeshWireframeSource implementation for a FDynamicMesh3
+ */
+class FDynamicMeshWireframeSource : public IMeshWireframeSource
+{
+public:
+	const FDynamicMesh3* Mesh = nullptr;
+
+	virtual bool IsValid() const { return Mesh != nullptr; }
+
+	virtual FBoxSphereBounds GetBounds() const 
+	{ 
+		FAxisAlignedBox3d Bounds = Mesh->GetBounds();
+		return FBoxSphereBounds((FBox)Bounds);
+	}
+
+	virtual FVector GetVertex(int32 Index) const
+	{
+		return (FVector)Mesh->GetVertex(Index);
+	}
+
+	virtual int32 GetEdgeCount() const
+	{
+		return Mesh->EdgeCount();
+	}
+
+	virtual int32 GetMaxEdgeIndex() const
+	{
+		return Mesh->MaxEdgeID();
+	}
+
+	virtual bool IsEdge(int32 Index) const
+	{
+		return Mesh->IsEdge(Index);
+	}
+
+	virtual void GetEdge(int32 EdgeIndex, int32& VertIndexAOut, int32& VertIndexBOut, EMeshEdgeType& TypeOut) const
+	{
+		FIndex2i EdgeV = Mesh->GetEdgeV(EdgeIndex);
+		VertIndexAOut = EdgeV.A;
+		VertIndexBOut = EdgeV.B;
+		int32 EdgeType = (int32)EMeshEdgeType::Regular;
+		if (Mesh->IsBoundaryEdge(EdgeIndex))
+		{
+			EdgeType |= (int32)EMeshEdgeType::MeshBoundary;
+		}
+		if (Mesh->HasAttributes())
+		{
+			bool bIsUVSeam = false, bIsNormalSeam = false;
+			if (Mesh->Attributes()->IsSeamEdge(EdgeIndex, bIsUVSeam, bIsNormalSeam))
+			{
+				if (bIsUVSeam)
+				{
+					EdgeType |= (int32)EMeshEdgeType::UVSeam;
+				}
+				if (bIsNormalSeam)
+				{
+					EdgeType |= (int32)EMeshEdgeType::NormalSeam;
+				}
+			}
+		}
+		TypeOut = (EMeshEdgeType)EdgeType;
+	}
+};
+
+
+class FDynamicMeshWireframeSourceProvider : public IMeshWireframeSourceProvider
+{
+public:
+	TUniqueFunction<const FDynamicMesh3* (void)> MeshAccessFunction;
+
+	FDynamicMeshWireframeSourceProvider(TUniqueFunction<const FDynamicMesh3* (void)>&& MeshAccessFuncIn)
+	{
+		MeshAccessFunction = MoveTemp(MeshAccessFuncIn);
+	}
+
+	virtual void AccessMesh(TFunctionRef<void(const IMeshWireframeSource&)> ProcessingFunc) override
+	{
+		const FDynamicMesh3* Mesh = MeshAccessFunction();
+		check(Mesh);
+		FDynamicMeshWireframeSource WireSource;
+		WireSource.Mesh = Mesh;
+		ProcessingFunc(WireSource);
+	}
+};
+
+
+
+
+void UMeshElementsVisualizer::OnCreated()
+{
+	Settings = NewObject<UMeshElementsVisualizerProperties>(this);
+	Settings->WatchProperty(Settings->bVisible, [this](bool) { bSettingsModified = true; });
+	Settings->WatchProperty(Settings->bShowWireframe, [this](bool) { bSettingsModified = true; });
+	Settings->WatchProperty(Settings->bShowBorders, [this](bool) { bSettingsModified = true; });
+	Settings->WatchProperty(Settings->bShowUVSeams, [this](bool) { bSettingsModified = true; });
+	Settings->WatchProperty(Settings->bShowNormalSeams, [this](bool) { bSettingsModified = true; });
+	Settings->WatchProperty(Settings->ThicknessScale, [this](float) { bSettingsModified = true; });
+	Settings->WatchProperty(Settings->DepthBias, [this](float) { bSettingsModified = true; });
+	Settings->WatchProperty(Settings->WireframeColor, [this](FColor) { bSettingsModified = true; });
+	Settings->WatchProperty(Settings->BoundaryEdgeColor, [this](FColor) { bSettingsModified = true; });
+	Settings->WatchProperty(Settings->UVSeamColor, [this](FColor) { bSettingsModified = true; });
+	Settings->WatchProperty(Settings->NormalSeamColor, [this](FColor) { bSettingsModified = true; });
+	bSettingsModified = false;
+
+	WireframeComponent = NewObject<UMeshWireframeComponent>(GetActor());
+	WireframeComponent->SetLineMaterial(ToolSetupUtil::GetDefaultLineComponentMaterial(nullptr));
+	WireframeComponent->SetupAttachment(GetActor()->GetRootComponent());	
+	WireframeComponent->RegisterComponent();
+}
+
+
+void UMeshElementsVisualizer::SetMeshAccessFunction(TUniqueFunction<const FDynamicMesh3* (void)>&& MeshAccessFunctionIn)
+{
+	const FDynamicMesh3* Mesh = MeshAccessFunctionIn();
+	FAxisAlignedBox3d Bounds = Mesh->GetBounds();
+	float DepthBiasDim = Bounds.DiagonalLength() * 0.01;
+	WireframeComponent->LineDepthBiasSizeScale = DepthBiasDim;
+
+	WireframeSourceProvider = MakeShared<FDynamicMeshWireframeSourceProvider>(MoveTemp(MeshAccessFunctionIn));
+
+	WireframeComponent->SetWireframeSourceProvider(WireframeSourceProvider);
+}
+
+
+
+
+void UMeshElementsVisualizer::OnTick(float DeltaTime)
+{
+	if (bSettingsModified)
+	{
+		UpdateVisibility();
+		bSettingsModified = false;
+	}
+}
+
+
+void UMeshElementsVisualizer::UpdateVisibility()
+{
+	if (Settings->bVisible == false)
+	{
+		WireframeComponent->SetVisibility(false);
+		return;
+	}
+
+	WireframeComponent->SetVisibility(true);
+
+	WireframeComponent->LineDepthBias = Settings->DepthBias;
+	WireframeComponent->ThicknessScale = Settings->ThicknessScale;
+
+	WireframeComponent->bEnableWireframe = Settings->bShowWireframe;
+	WireframeComponent->bEnableBoundaryEdges = Settings->bShowBorders;
+	WireframeComponent->bEnableUVSeams = Settings->bShowUVSeams;
+	WireframeComponent->bEnableNormalSeams = Settings->bShowNormalSeams;
+
+	WireframeComponent->WireframeColor = Settings->WireframeColor;
+	WireframeComponent->BoundaryEdgeColor = Settings->BoundaryEdgeColor;
+	WireframeComponent->UVSeamColor = Settings->UVSeamColor;
+	WireframeComponent->NormalSeamColor = Settings->NormalSeamColor;
+
+	WireframeComponent->MarkRenderStateDirty();
+}
+
+
+void UMeshElementsVisualizer::NotifyMeshChanged()
+{
+	WireframeComponent->MarkRenderStateDirty();
+}
\ No newline at end of file
diff --git a/Engine/Plugins/Experimental/MeshModelingToolset/Source/ModelingComponents/Private/Drawing/MeshWireframeComponent.cpp b/Engine/Plugins/Experimental/MeshModelingToolset/Source/ModelingComponents/Private/Drawing/MeshWireframeComponent.cpp
new file mode 100644
--- /dev/null
+++ b/Engine/Plugins/Experimental/MeshModelingToolset/Source/ModelingComponents/Private/Drawing/MeshWireframeComponent.cpp
@@ -0,0 +1,352 @@
+// Copyright Epic Games, Inc. All Rights Reserved.
+
+#include "Drawing/MeshWireframeComponent.h"
+
+#include "Engine/CollisionProfile.h"
+#include "LocalVertexFactory.h"
+#include "MaterialShared.h"
+#include "Materials/MaterialInterface.h"
+#include "Materials/Material.h"
+#include "PrimitiveSceneProxy.h"
+#include "PrimitiveViewRelevance.h"
+#include "StaticMeshResources.h"
+#include "IndexTypes.h"
+#include "Async/ParallelFor.h"
+
+struct FWireframeLinesMeshBatchData
+{
+	FWireframeLinesMeshBatchData()
+		: MaterialProxy(nullptr)
+	{}
+
+	FMaterialRenderProxy* MaterialProxy;
+	int32 StartIndex;
+	int32 NumPrimitives;
+	int32 MinVertexIndex;
+	int32 MaxVertexIndex;
+};
+
+
+
+
+
+/** Class for the MeshWireframeComponent data passed to the render thread. */
+class FMeshWireframeSceneProxy final : public FPrimitiveSceneProxy
+{
+public:
+
+	FMeshWireframeSceneProxy(UMeshWireframeComponent* Component, const IMeshWireframeSource* WireSource)
+		: FPrimitiveSceneProxy(Component),
+		MaterialRelevance(Component->GetMaterialRelevance(GetScene().GetFeatureLevel())),
+		VertexFactory(GetScene().GetFeatureLevel(), "FPointSetSceneProxy")
+	{
+		if (!ensure(WireSource && WireSource->IsValid())) return;
+
+		if (WireSource->GetEdgeCount() <= 0) return;
+
+		// count visible edges and remap so we can build in parallel below
+		int32 MaxEdgeIndex = WireSource->GetMaxEdgeIndex();
+		CurrentEdgeSet.Reserve(WireSource->GetEdgeCount());
+		for (int32 li = 0; li < MaxEdgeIndex; ++li)
+		{
+			if (WireSource->IsEdge(li) == false) continue;
+
+			int32 VertIndexA, VertIndexB;
+			IMeshWireframeSource::EMeshEdgeType EdgeType;
+			WireSource->GetEdge(li, VertIndexA, VertIndexB, EdgeType);
+
+			bool bEdgeIsVisible = false;
+			if (Component->bEnableWireframe)
+			{
+				bEdgeIsVisible = true;
+			}
+			else if (((int)EdgeType & (int)IMeshWireframeSource::EMeshEdgeType::MeshBoundary) != 0 && Component->bEnableBoundaryEdges)
+			{
+				bEdgeIsVisible = true;
+			}
+			else if (((int)EdgeType & (int)IMeshWireframeSource::EMeshEdgeType::UVSeam) != 0 && Component->bEnableUVSeams)
+			{
+				bEdgeIsVisible = true;
+			}
+			else if (((int)EdgeType & (int)IMeshWireframeSource::EMeshEdgeType::NormalSeam) != 0 && Component->bEnableNormalSeams)
+			{
+				bEdgeIsVisible = true;
+			}
+
+			if (bEdgeIsVisible)
+			{
+				CurrentEdgeSet.Add(FIndex4i(li, VertIndexA, VertIndexB, (int)EdgeType));
+			}
+		};
+		int32 NumEdges = CurrentEdgeSet.Num();
+		if (NumEdges == 0)
+		{
+			return;
+		}
+
+		const int32 NumLineVertices = NumEdges * 4;
+		const int32 NumLineIndices = NumEdges * 6;
+		const int32 NumTextureCoordinates = 1;
+
+		VertexBuffers.PositionVertexBuffer.Init(NumLineVertices);
+		VertexBuffers.StaticMeshVertexBuffer.Init(NumLineVertices, NumTextureCoordinates);
+		VertexBuffers.ColorVertexBuffer.Init(NumLineVertices);
+		IndexBuffer.Indices.SetNumUninitialized(NumLineIndices);
+
+		MeshBatchDatas.Emplace();
+		FWireframeLinesMeshBatchData& MeshBatchData = MeshBatchDatas.Last();
+		MeshBatchData.MinVertexIndex = 0;
+		MeshBatchData.MaxVertexIndex = 0 + NumLineVertices - 1;
+		MeshBatchData.StartIndex = 0;
+		MeshBatchData.NumPrimitives = NumEdges * 2;
+		if (Component->GetMaterial(0) != nullptr)
+		{
+			MeshBatchData.MaterialProxy = Component->GetMaterial(0)->GetRenderProxy();
+		}
+		else
+		{
+			MeshBatchData.MaterialProxy = UMaterial::GetDefaultMaterial(MD_Surface)->GetRenderProxy();
+		}
+		
+		FColor RegularEdgeColor = FLinearColor::FromSRGBColor(Component->WireframeColor).ToFColor(false);
+		float RegularEdgeThickness = Component->ThicknessScale * Component->WireframeThickness;
+		FColor BoundaryEdgeColor = FLinearColor::FromSRGBColor(Component->BoundaryEdgeColor).ToFColor(false);
+		float BoundaryEdgeThickness = Component->ThicknessScale * Component->BoundaryEdgeThickness;
+		FColor UVSeamColor = FLinearColor::FromSRGBColor(Component->UVSeamColor).ToFColor(false);
+		float UVSeamThickness = Component->ThicknessScale * Component->UVSeamThickness;
+		FColor NormalSeamColor = FLinearColor::FromSRGBColor(Component->NormalSeamColor).ToFColor(false);
+		float NormalSeamThickness = Component->ThicknessScale * Component->NormalSeamThickness;
+
+		float LineDepthBias = Component->LineDepthBias * Component->LineDepthBiasSizeScale;
+
+		// Initialize lines.
+		// Lines are represented as two tris of zero thickness. The UV's stored at vertices are actually (lineThickness, depthBias), 
+		// which the material unpacks and uses to thicken the polygons and set the pixel depth bias.
+		ParallelFor(NumEdges, [&](int32 idx)
+		{
+			int32 VertexBufferIndex = idx * 4;
+			int32 IndexBufferIndex = idx * 6;
+
+			FIndex4i EdgeInfo = CurrentEdgeSet[idx];
+			IMeshWireframeSource::EMeshEdgeType EdgeType = (IMeshWireframeSource::EMeshEdgeType)EdgeInfo.D;
+
+			float UseThickness = RegularEdgeThickness;
+			FColor UseColor = RegularEdgeColor;
+
+			bool bIsRegularEdge = (EdgeType == IMeshWireframeSource::EMeshEdgeType::Regular);
+			bool bIsBoundaryEdge = (((int)EdgeType & (int)IMeshWireframeSource::EMeshEdgeType::MeshBoundary) != 0);
+			if (!bIsRegularEdge)
+			{
+				if (bIsBoundaryEdge && Component->bEnableBoundaryEdges)
+				{
+					UseThickness = BoundaryEdgeThickness;
+					UseColor = BoundaryEdgeColor;
+				}
+				else if (((int)EdgeType & (int)IMeshWireframeSource::EMeshEdgeType::UVSeam) != 0 && Component->bEnableUVSeams)
+				{
+					UseThickness = (bIsBoundaryEdge) ? BoundaryEdgeThickness : UVSeamThickness;
+					UseColor = UVSeamColor;
+				}
+				else if (((int)EdgeType & (int)IMeshWireframeSource::EMeshEdgeType::NormalSeam) != 0 && Component->bEnableNormalSeams)
+				{
+					UseThickness = (bIsBoundaryEdge) ? BoundaryEdgeThickness : NormalSeamThickness;
+					UseColor = NormalSeamColor;
+				}
+			}
+
+			const FVector A = WireSource->GetVertex(EdgeInfo.B);
+			const FVector B = WireSource->GetVertex(EdgeInfo.C);
+			const FVector LineDirection = (B - A).GetSafeNormal();
+			const FVector2D UV(UseThickness, LineDepthBias);
+
+			VertexBuffers.PositionVertexBuffer.VertexPosition(VertexBufferIndex + 0) = A;
+			VertexBuffers.PositionVertexBuffer.VertexPosition(VertexBufferIndex + 1) = B;
+			VertexBuffers.PositionVertexBuffer.VertexPosition(VertexBufferIndex + 2) = B;
+			VertexBuffers.PositionVertexBuffer.VertexPosition(VertexBufferIndex + 3) = A;
+
+			VertexBuffers.StaticMeshVertexBuffer.SetVertexTangents(VertexBufferIndex + 0, FVector::ZeroVector, FVector::ZeroVector, -LineDirection);
+			VertexBuffers.StaticMeshVertexBuffer.SetVertexTangents(VertexBufferIndex + 1, FVector::ZeroVector, FVector::ZeroVector, -LineDirection);
+			VertexBuffers.StaticMeshVertexBuffer.SetVertexTangents(VertexBufferIndex + 2, FVector::ZeroVector, FVector::ZeroVector, LineDirection);
+			VertexBuffers.StaticMeshVertexBuffer.SetVertexTangents(VertexBufferIndex + 3, FVector::ZeroVector, FVector::ZeroVector, LineDirection);
+
+			VertexBuffers.StaticMeshVertexBuffer.SetVertexUV(VertexBufferIndex + 0, 0, UV);
+			VertexBuffers.StaticMeshVertexBuffer.SetVertexUV(VertexBufferIndex + 1, 0, UV);
+			VertexBuffers.StaticMeshVertexBuffer.SetVertexUV(VertexBufferIndex + 2, 0, UV);
+			VertexBuffers.StaticMeshVertexBuffer.SetVertexUV(VertexBufferIndex + 3, 0, UV);
+
+			// The color stored in the vertices actually gets interpreted as a linear color by the material,
+			// whereas it is more convenient for the user of the MeshWireframe to specify colors as sRGB. So we actually
+			// have to convert it back to linear. The ToFColor(false) call just scales back into 0-255 space.
+			VertexBuffers.ColorVertexBuffer.VertexColor(VertexBufferIndex + 0) = UseColor;
+			VertexBuffers.ColorVertexBuffer.VertexColor(VertexBufferIndex + 1) = UseColor;
+			VertexBuffers.ColorVertexBuffer.VertexColor(VertexBufferIndex + 2) = UseColor;
+			VertexBuffers.ColorVertexBuffer.VertexColor(VertexBufferIndex + 3) = UseColor;
+
+			IndexBuffer.Indices[IndexBufferIndex + 0] = VertexBufferIndex + 0;
+			IndexBuffer.Indices[IndexBufferIndex + 1] = VertexBufferIndex + 1;
+			IndexBuffer.Indices[IndexBufferIndex + 2] = VertexBufferIndex + 2;
+			IndexBuffer.Indices[IndexBufferIndex + 3] = VertexBufferIndex + 2;
+			IndexBuffer.Indices[IndexBufferIndex + 4] = VertexBufferIndex + 3;
+			IndexBuffer.Indices[IndexBufferIndex + 5] = VertexBufferIndex + 0;
+		});
+
+		ENQUEUE_RENDER_COMMAND(MeshWireframeVertexBuffersInit)(
+			[this](FRHICommandListImmediate& RHICmdList)
+		{
+			VertexBuffers.PositionVertexBuffer.InitResource();
+			VertexBuffers.StaticMeshVertexBuffer.InitResource();
+			VertexBuffers.ColorVertexBuffer.InitResource();
+
+			FLocalVertexFactory::FDataType Data;
+			VertexBuffers.PositionVertexBuffer.BindPositionVertexBuffer(&VertexFactory, Data);
+			VertexBuffers.StaticMeshVertexBuffer.BindTangentVertexBuffer(&VertexFactory, Data);
+			VertexBuffers.StaticMeshVertexBuffer.BindTexCoordVertexBuffer(&VertexFactory, Data);
+			VertexBuffers.ColorVertexBuffer.BindColorVertexBuffer(&VertexFactory, Data);
+			VertexFactory.SetData(Data);
+
+			VertexFactory.InitResource();
+			IndexBuffer.InitResource();
+		});
+	}
+
+	virtual ~FMeshWireframeSceneProxy()
+	{
+		VertexBuffers.PositionVertexBuffer.ReleaseResource();
+		VertexBuffers.StaticMeshVertexBuffer.ReleaseResource();
+		VertexBuffers.ColorVertexBuffer.ReleaseResource();
+		IndexBuffer.ReleaseResource();
+		VertexFactory.ReleaseResource();
+	}
+
+	virtual void GetDynamicMeshElements(const TArray<const FSceneView*>& Views, 
+		const FSceneViewFamily& ViewFamily, uint32 VisibilityMap, FMeshElementCollector& Collector) const override
+	{
+		QUICK_SCOPE_CYCLE_COUNTER(STAT_OverlaySceneProxy_GetDynamicMeshElements);
+
+		for (int32 ViewIndex = 0; ViewIndex < Views.Num(); ViewIndex++)
+		{
+			if (VisibilityMap & (1 << ViewIndex))
+			{
+				for (const FWireframeLinesMeshBatchData& MeshBatchData : MeshBatchDatas)
+				{
+					const FSceneView* View = Views[ViewIndex];
+					FMeshBatch& Mesh = Collector.AllocateMesh();
+					FMeshBatchElement& BatchElement = Mesh.Elements[0];
+					BatchElement.IndexBuffer = &IndexBuffer;
+					Mesh.bWireframe = false;
+					Mesh.VertexFactory = &VertexFactory;
+					Mesh.MaterialRenderProxy = MeshBatchData.MaterialProxy;
+
+					FDynamicPrimitiveUniformBuffer& DynamicPrimitiveUniformBuffer = Collector.AllocateOneFrameResource<FDynamicPrimitiveUniformBuffer>();
+					DynamicPrimitiveUniformBuffer.Set(GetLocalToWorld(), GetLocalToWorld(), GetBounds(), GetLocalBounds(), false, false, DrawsVelocity(), false);
+					BatchElement.PrimitiveUniformBufferResource = &DynamicPrimitiveUniformBuffer.UniformBuffer;
+
+					BatchElement.FirstIndex = MeshBatchData.StartIndex;
+					BatchElement.NumPrimitives = MeshBatchData.NumPrimitives;
+					BatchElement.MinVertexIndex = MeshBatchData.MinVertexIndex;
+					BatchElement.MaxVertexIndex = MeshBatchData.MaxVertexIndex;
+					Mesh.ReverseCulling = IsLocalToWorldDeterminantNegative();
+					Mesh.Type = PT_TriangleList;
+					Mesh.DepthPriorityGroup = SDPG_World;
+					Mesh.bCanApplyViewModeOverrides = false;
+					Collector.AddMesh(ViewIndex, Mesh);
+				}
+			}
+		}
+	}
+
+	virtual FPrimitiveViewRelevance GetViewRelevance(const FSceneView* View) const override
+	{
+		FPrimitiveViewRelevance Result;
+		Result.bDrawRelevance = IsShown(View);
+		Result.bShadowRelevance = IsShadowCast(View);
+		Result.bDynamicRelevance = true;
+		Result.bRenderInMainPass = ShouldRenderInMainPass();
+		Result.bUsesLightingChannels = GetLightingChannelMask() != GetDefaultLightingChannelMask();
+		Result.bRenderCustomDepth = ShouldRenderCustomDepth();
+		Result.bTranslucentSelfShadow = bCastVolumetricTranslucentShadow;
+		MaterialRelevance.SetPrimitiveViewRelevance(Result);
+		Result.bVelocityRelevance = IsMovable() && Result.bOpaque && Result.bRenderInMainPass;
+		return Result;
+	}
+
+	virtual bool CanBeOccluded() const override
+	{
+		return !MaterialRelevance.bDisableDepthTest;
+	}
+
+	virtual uint32 GetMemoryFootprint() const override { return sizeof(*this) + GetAllocatedSize(); }
+
+	uint32 GetAllocatedSize() const { return FPrimitiveSceneProxy::GetAllocatedSize(); }
+
+	virtual SIZE_T GetTypeHash() const override
+	{
+		static SIZE_T UniquePointer;
+		return reinterpret_cast<SIZE_T>(&UniquePointer);
+	}
+
+private:
+	TArray<FWireframeLinesMeshBatchData> MeshBatchDatas;
+	FMaterialRelevance MaterialRelevance;
+	FLocalVertexFactory VertexFactory;
+	FStaticMeshVertexBuffers VertexBuffers;
+	FDynamicMeshIndexBuffer32 IndexBuffer;
+
+	TArray<FIndex4i> CurrentEdgeSet;
+};
+
+
+UMeshWireframeComponent::UMeshWireframeComponent()
+{
+	CastShadow = false;
+	bSelectable = false;
+	PrimaryComponentTick.bCanEverTick = false;
+
+	SetCollisionProfileName(UCollisionProfile::NoCollision_ProfileName);
+}
+
+void UMeshWireframeComponent::SetWireframeSourceProvider(TSharedPtr<IMeshWireframeSourceProvider> Provider)
+{
+	SourceProvider = Provider;
+	if (SourceProvider)
+	{
+		SourceProvider->AccessMesh([&](const IMeshWireframeSource& Source)
+		{
+			this->LocalBounds = Source.GetBounds();
+		});
+	}
+	MarkRenderStateDirty();
+}
+
+void UMeshWireframeComponent::SetLineMaterial(UMaterialInterface* InLineMaterial)
+{
+	LineMaterial = InLineMaterial;
+	SetMaterial(0, InLineMaterial);
+}
+
+
+
+FPrimitiveSceneProxy* UMeshWireframeComponent::CreateSceneProxy()
+{
+	if (SourceProvider)
+	{
+		FMeshWireframeSceneProxy* NewProxy = nullptr;
+		SourceProvider->AccessMesh([this, &NewProxy](const IMeshWireframeSource& Source)
+		{
+			NewProxy = new FMeshWireframeSceneProxy(this, &Source);
+		});
+		return NewProxy;
+	}
+	return nullptr;
+}
+
+int32 UMeshWireframeComponent::GetNumMaterials() const
+{
+	return 1;
+}
+
+FBoxSphereBounds UMeshWireframeComponent::CalcBounds(const FTransform& LocalToWorld) const
+{
+	return LocalBounds.TransformBy(LocalToWorld);
+}
+
diff --git a/Engine/Plugins/Experimental/MeshModelingToolset/Source/ModelingComponents/Private/Drawing/PreviewGeometryActor.cpp b/Engine/Plugins/Experimental/MeshModelingToolset/Source/ModelingComponents/Private/Drawing/PreviewGeometryActor.cpp
--- a/Engine/Plugins/Experimental/MeshModelingToolset/Source/ModelingComponents/Private/Drawing/PreviewGeometryActor.cpp
+++ b/Engine/Plugins/Experimental/MeshModelingToolset/Source/ModelingComponents/Private/Drawing/PreviewGeometryActor.cpp
@@ -29,11 +29,15 @@ void UPreviewGeometry::CreateInWorld(UWorld* World, const FTransform& WithTransf
 	SphereComponent->SetCollisionProfileName(UCollisionProfile::NoCollision_ProfileName);
 
 	ParentActor->SetActorTransform(WithTransform);
+
+	OnCreated();
 }
 
 
 void UPreviewGeometry::Disconnect()
 {
+	OnDisconnected();
+
 	if (ParentActor != nullptr)
 	{
 		ParentActor->Destroy();
diff --git a/Engine/Plugins/Experimental/MeshModelingToolset/Source/ModelingComponents/Public/Drawing/LineSetComponent.h b/Engine/Plugins/Experimental/MeshModelingToolset/Source/ModelingComponents/Public/Drawing/LineSetComponent.h
--- a/Engine/Plugins/Experimental/MeshModelingToolset/Source/ModelingComponents/Public/Drawing/LineSetComponent.h
+++ b/Engine/Plugins/Experimental/MeshModelingToolset/Source/ModelingComponents/Public/Drawing/LineSetComponent.h
@@ -109,8 +109,8 @@ public:
 	void AddLines(
 		int32 NumIndices,
 		TFunctionRef<void(int32 Index, TArray<FRenderableLine>& LinesOut)> LineGenFunc,
-		int32 LinesPerIndexHint = -1);
-
+		int32 LinesPerIndexHint = -1,
+		bool bDeferRenderStateDirty = false);
 
 private:
 
@@ -137,5 +137,7 @@ private:
 
 	TSparseArray<FRenderableLine> Lines;
 
+	int32 AddLineInternal(const FRenderableLine& Line);
+
 	friend class FLineSetSceneProxy;
 };
\ No newline at end of file
diff --git a/Engine/Plugins/Experimental/MeshModelingToolset/Source/ModelingComponents/Public/Drawing/MeshElementsVisualizer.h b/Engine/Plugins/Experimental/MeshModelingToolset/Source/ModelingComponents/Public/Drawing/MeshElementsVisualizer.h
new file mode 100644
--- /dev/null
+++ b/Engine/Plugins/Experimental/MeshModelingToolset/Source/ModelingComponents/Public/Drawing/MeshElementsVisualizer.h
@@ -0,0 +1,123 @@
+// Copyright Epic Games, Inc. All Rights Reserved.
+
+#pragma once
+
+#include "CoreMinimal.h"
+#include "Drawing/PreviewGeometryActor.h"
+#include "InteractiveTool.h"
+#include "MeshWireframeComponent.h"
+#include "MeshElementsVisualizer.generated.h"
+
+class FDynamicMesh3;
+
+/**
+ * Visualization settings UMeshElementsVisualizer
+ */
+UCLASS()
+class MODELINGCOMPONENTS_API UMeshElementsVisualizerProperties : public UInteractiveToolPropertySet
+{
+	GENERATED_BODY()
+public:
+	/** Should any be mesh elements be shown */
+	UPROPERTY(EditAnywhere, Category = MeshElementVisualization)
+	bool bVisible = true;
+
+	/** Should mesh wireframe be shown */
+	UPROPERTY(EditAnywhere, Category = MeshElementVisualization)
+	bool bShowWireframe = false;
+
+	/** Should mesh boundary edges be shown */
+	UPROPERTY(EditAnywhere, Category = MeshElementVisualization)
+	bool bShowBorders = true;
+
+	/** Should mesh uv seam edges be shown */
+	UPROPERTY(EditAnywhere, Category = MeshElementVisualization)
+	bool bShowUVSeams = true;
+
+	/** Should mesh normal seam edges be shown */
+	UPROPERTY(EditAnywhere, Category = MeshElementVisualization)
+	bool bShowNormalSeams = true;
+
+	/** multiplier on edge thicknesses */
+	UPROPERTY(EditAnywhere, Category = MeshElementVisualization, AdvancedDisplay, meta = (UIMin = 0.1, UIMax = 10.0))
+	float ThicknessScale = 1.0;
+
+	/** Color of mesh wireframe */
+	UPROPERTY(EditAnywhere, Category = MeshElementVisualization, AdvancedDisplay)
+	FColor WireframeColor = FColor(128, 128, 128);
+
+	/** Color of mesh boundary edges */
+	UPROPERTY(EditAnywhere, Category = MeshElementVisualization, AdvancedDisplay)
+	FColor BoundaryEdgeColor = FColor(245, 15, 15);
+
+	/** Color of mesh UV seam edges */
+	UPROPERTY(EditAnywhere, Category = MeshElementVisualization, AdvancedDisplay)
+	FColor UVSeamColor = FColor(240, 160, 15);
+
+	/** Color of mesh normal seam edges */
+	UPROPERTY(EditAnywhere, Category = MeshElementVisualization, AdvancedDisplay)
+	FColor NormalSeamColor = FColor(128, 128, 240);
+
+	/** depth bias used to slightly shift depth of lines */
+	UPROPERTY(EditAnywhere, Category = MeshElementVisualization, AdvancedDisplay, meta = (UIMin = -2.0, UIMax = 2.0))
+	float DepthBias = 0.2;
+};
+
+
+/**
+ * UMeshElementsVisualizer is a subclass of UPreviewGeometry that displays mesh elements.
+ * Currently supports wireframe, boundary edges, UV seams, and Normal seams.
+ *
+ * UMeshElementsVisualizer initializes an instance of UMeshElementsVisualizerProperties
+ * as its .Settings value, and will watch for changes in these properties.
+ *
+ * Mesh is accessed via lambda callback provided by creator/client. See SetMeshAccessFunction() comments
+ */
+UCLASS(Transient)
+class MODELINGCOMPONENTS_API UMeshElementsVisualizer : public UPreviewGeometry
+{
+	GENERATED_BODY()
+
+public:
+	/**
+	 * UMeshElementsVisualizer must be provided with a callback function that can be used to
+	 * access the mesh. UMeshElementsVisualizer will hold onto this callback, and wrap it in
+	 * an API that is passed to child Components used to render mesh geometry. Those Components
+	 * will access the mesh when creating their Scene Proxies/etc. The callback can return null
+	 * if the mesh has become invalid/etc
+	 */
+	void SetMeshAccessFunction(TUniqueFunction<const FDynamicMesh3* (void)>&& MeshAccessFunction);
+
+	/**
+	 * Call if mesh provided by MeshAccessFunction has been modified, will cause a full recomputation
+	 * of all rendering data structures.
+	 */
+	void NotifyMeshChanged();
+
+
+	/**
+	 * Client must call this every frame for changes to .Settings to be reflected in rendered result.
+	 */
+	void OnTick(float DeltaTime);
+
+
+public:
+	/** Visualization settings */
+	UPROPERTY()
+	UMeshElementsVisualizerProperties* Settings;
+
+	/** Mesh Wireframe component, draws wireframe, boundaries, UV seams, normal seams */
+	UPROPERTY()
+	UMeshWireframeComponent* WireframeComponent;
+
+
+protected:
+	virtual void OnCreated() override;
+
+	bool bSettingsModified = false;
+
+	void UpdateVisibility();
+
+	TSharedPtr<IMeshWireframeSourceProvider> WireframeSourceProvider;
+};
+
diff --git a/Engine/Plugins/Experimental/MeshModelingToolset/Source/ModelingComponents/Public/Drawing/MeshWireframeComponent.h b/Engine/Plugins/Experimental/MeshModelingToolset/Source/ModelingComponents/Public/Drawing/MeshWireframeComponent.h
new file mode 100644
--- /dev/null
+++ b/Engine/Plugins/Experimental/MeshModelingToolset/Source/ModelingComponents/Public/Drawing/MeshWireframeComponent.h
@@ -0,0 +1,196 @@
+// Copyright Epic Games, Inc. All Rights Reserved.
+
+#pragma once
+
+#include "CoreMinimal.h"
+#include "Components/MeshComponent.h"
+#include "Materials/MaterialInterface.h"
+#include "PrimitiveSceneProxy.h"
+
+#include "MeshWireframeComponent.generated.h"
+
+class FPrimitiveSceneProxy;
+
+
+/**
+ * IMeshWireframeSource is an abstract interface to a class that provides
+ * a set of edges. Edges are identified by an integer ID, which may be invalid
+ * for some edges in the ID range, and each edge has two vertex indices.
+ *
+ * The class is not specific to mesh wireframes, and could be used to draw any edge set, 
+ * but the EdgeType is based on triangle/polygon mesh edges.
+ */
+class IMeshWireframeSource
+{
+public:
+	enum class EMeshEdgeType
+	{
+		Regular = 1,
+		MeshBoundary = 1 << 2,
+		UVSeam = 1 << 3,
+		NormalSeam = 1 << 4,
+		GroupBoundary = 1 << 5
+	};
+
+public:
+	virtual ~IMeshWireframeSource() {}
+
+	/** @return true if this source is valid */
+	virtual bool IsValid() const = 0;
+
+	/** @return bounds of all possible edges */
+	virtual FBoxSphereBounds GetBounds() const = 0;
+	/** @return position of vertex at given index */
+	virtual FVector GetVertex(int32 Index) const = 0;
+
+	/** @return number of valid edge indices */
+	virtual int32 GetEdgeCount() const = 0;
+	/** @return maximum edge index (may be greater than count, if some indices are invalid) */
+	virtual int32 GetMaxEdgeIndex() const = 0;
+	/** @return true if edge index is valid, ie refers to valid edge */
+	virtual bool IsEdge(int32 Index) const = 0;
+	/** @return vertex indices and type for a given edge index */
+	virtual void GetEdge(int32 Index, int32& VertexAOut, int32& VertexBOut, EMeshEdgeType& TypeOut) const = 0;
+};
+
+/**
+ * IMeshWireframeSourceProvider is an abstract interface to some implementation that
+ * can provide an IMeshWireframeSource implementation on demand. This is used
+ * to populate/update a UMeshWireframeComponent when necessary.
+ */
+class IMeshWireframeSourceProvider
+{
+public:
+	virtual ~IMeshWireframeSourceProvider() {}
+
+	/**
+	 * Request that the given function ProcessingFunc be called with an IMeshWireframeSource implementation.
+	 * The IMeshWireframeSource implementation may be invalid, ie IsValid() == false, ProcessingFunc must handle this case.
+	 * The AccessMesh implementation may choose to not call the lambda in such cases, as well
+	 */
+	virtual void AccessMesh(TFunctionRef<void(const IMeshWireframeSource&)> ProcessingFunc) = 0;
+};
+
+
+/**
+ * UMeshWireframeComponent draws a mesh wireframe as lines, with line color/thickness
+ * varying depending on line type and the configuration settings.
+ * Currently can draw:
+ *    - all mesh edges
+ *    - boundary edges
+ *    - UV seam edges
+ *    - Normal seam edges
+ *
+ * Client must provide a IMeshWireframeSourceProvider implementation that provides the 
+ * edge set, vertices, and edge type
+ */
+UCLASS()
+class MODELINGCOMPONENTS_API UMeshWireframeComponent : public UMeshComponent
+{
+	GENERATED_BODY()
+
+public:
+	UMeshWireframeComponent();
+
+	/** Specify material which handles lines */
+	void SetLineMaterial(UMaterialInterface* InLineMaterial);
+
+	/**
+	 * Set the wireframe source
+	 */
+	void SetWireframeSourceProvider(TSharedPtr<IMeshWireframeSourceProvider> Provider);
+
+
+public:
+	
+	/**
+	 * Depth bias of the lines, used to offset in depth to avoid z-fighting
+	 */
+	UPROPERTY(EditAnywhere, Category=MeshWireframe)
+	float LineDepthBias = 0.1f;
+
+	/**
+	 * Target-size depth bias scale. This is multiplied by LineDepthBias.
+	 * Client of UMeshWireframeComponent can set this if desired, eg to fraction of target object bounding box size, etc.
+	 */
+	UPROPERTY(EditAnywhere, Category=MeshWireframe)
+	float LineDepthBiasSizeScale = 1.0f;
+
+	/**
+	 * Scaling factor applied to the per-edge-type thicknesses defined below
+	 */
+	UPROPERTY(EditAnywhere, Category = MeshWireframe)
+	float ThicknessScale = 1.0f;
+
+
+	// Wireframe properties
+
+	UPROPERTY(EditAnywhere, Category = MeshWireframe)
+	bool bEnableWireframe = true;
+
+	UPROPERTY(EditAnywhere, Category = MeshWireframe)
+	FColor WireframeColor = FColor(128, 128, 128);
+
+	UPROPERTY(EditAnywhere, Category = MeshWireframe)
+	float WireframeThickness = 1.0f;
+
+
+	// Boundary Edge properties
+
+	UPROPERTY(EditAnywhere, Category = MeshWireframe)
+	bool bEnableBoundaryEdges = true;
+
+	UPROPERTY(EditAnywhere, Category = MeshWireframe)
+	FColor BoundaryEdgeColor = FColor(245, 15, 15);
+
+	UPROPERTY(EditAnywhere, Category = MeshWireframe)
+	float BoundaryEdgeThickness = 4.0f;
+
+
+	// UV seam properties
+
+	UPROPERTY(EditAnywhere, Category = MeshWireframe)
+	bool bEnableUVSeams = true;
+
+	UPROPERTY(EditAnywhere, Category = MeshWireframe)
+	FColor UVSeamColor = FColor(240, 160, 15);
+
+	UPROPERTY(EditAnywhere, Category = MeshWireframe)
+	float UVSeamThickness = 2.0f;
+
+
+	// normal seam properties
+
+	UPROPERTY(EditAnywhere, Category = MeshWireframe)
+	bool bEnableNormalSeams = true;
+
+	UPROPERTY(EditAnywhere, Category = MeshWireframe)
+	FColor NormalSeamColor = FColor(128, 128, 240);
+
+	UPROPERTY(EditAnywhere, Category = MeshWireframe)
+	float NormalSeamThickness = 2.0f;
+
+private:
+
+	//~ Begin UPrimitiveComponent Interface.
+	virtual FPrimitiveSceneProxy* CreateSceneProxy() override;
+	//~ End UPrimitiveComponent Interface.
+
+	//~ Begin UMeshComponent Interface.
+	virtual int32 GetNumMaterials() const override;
+	//~ End UMeshComponent Interface.
+
+	//~ Begin USceneComponent Interface.
+	virtual FBoxSphereBounds CalcBounds(const FTransform& LocalToWorld) const override;
+	//~ End USceneComponent Interface.
+
+	UPROPERTY()
+	const UMaterialInterface* LineMaterial;
+
+	UPROPERTY()
+	mutable FBoxSphereBounds LocalBounds;
+
+	TSharedPtr<IMeshWireframeSourceProvider> SourceProvider;
+
+	friend class FMeshWireframeSceneProxy;
+};
\ No newline at end of file
diff --git a/Engine/Plugins/Experimental/MeshModelingToolset/Source/ModelingComponents/Public/Drawing/PreviewGeometryActor.h b/Engine/Plugins/Experimental/MeshModelingToolset/Source/ModelingComponents/Public/Drawing/PreviewGeometryActor.h
--- a/Engine/Plugins/Experimental/MeshModelingToolset/Source/ModelingComponents/Public/Drawing/PreviewGeometryActor.h
+++ b/Engine/Plugins/Experimental/MeshModelingToolset/Source/ModelingComponents/Public/Drawing/PreviewGeometryActor.h
@@ -160,4 +160,14 @@ public:
 	UPROPERTY()
 	TMap<FString, ULineSetComponent*> LineSets;
 
+
+
+protected:
+
+	// called at the end of CreateInWorld() to allow subclasses to add additional setup
+	virtual void OnCreated() {}
+
+	// called at the beginning of Disconnect() to allow subclasses to perform additional cleanup
+	virtual void OnDisconnected() {}
+
 };
\ No newline at end of file
