MeshModelingToolset
- add support for VertexColor shader mode in EMeshEditingMaterialModes, handle in sculpt tools
- MeshSculptToolBase: add ability to hide view and brushop properties, add ability to show surface or volumetric indicator, add tracking of last-hit triangle ID
- add UPolygroupLayersProperties, standard property set for listing and picking from available Polygroup layers for a mesh
- add MeshGroupPaintTool, uses above extensions to allow for painting per-triangle integer polygroups
#rnx

diff --git a/Engine/Plugins/Experimental/MeshModelingToolset/Source/MeshModelingTools/Private/DynamicMeshSculptTool.cpp b/Engine/Plugins/Experimental/MeshModelingToolset/Source/MeshModelingTools/Private/DynamicMeshSculptTool.cpp
--- a/Engine/Plugins/Experimental/MeshModelingToolset/Source/MeshModelingTools/Private/DynamicMeshSculptTool.cpp
+++ b/Engine/Plugins/Experimental/MeshModelingToolset/Source/MeshModelingTools/Private/DynamicMeshSculptTool.cpp
@@ -1260,6 +1260,7 @@ bool UDynamicMeshSculptTool::UpdateBrushPositionOnSculptMesh(const FRay& WorldRa
 
 		LastBrushPosNormalWorld = CurTargetTransform.TransformNormal(SculptMesh->GetTriNormal(HitTID));
 		LastBrushPosWorld = CurTargetTransform.TransformPosition(LocalRay.PointAt(Query.RayParameter));
+		LastBrushTriangleID = HitTID;
 		return true;
 	}
 
@@ -1270,6 +1271,7 @@ bool UDynamicMeshSculptTool::UpdateBrushPositionOnSculptMesh(const FRay& WorldRa
 		BrushPlane.RayPlaneIntersection(WorldRay.Origin, WorldRay.Direction, 2, NewHitPosWorld);
 		LastBrushPosWorld = NewHitPosWorld;
 		LastBrushPosNormalWorld = ActiveDragPlane.Z();
+		LastBrushTriangleID = -1;
 		return true;
 	}
 
@@ -2240,6 +2242,9 @@ void UDynamicMeshSculptTool::UpdateMaterialMode(EMeshEditingMaterialModes Materi
 			case EMeshEditingMaterialModes::TangentNormal:
 				SculptMaterial = ToolSetupUtil::GetImageBasedSculptMaterial(GetToolManager(), ToolSetupUtil::ImageMaterialType::TangentNormalFromView);
 				break;
+			case EMeshEditingMaterialModes::VertexColor:
+				SculptMaterial = ToolSetupUtil::GetVertexColorMaterial(GetToolManager());
+				break;
 			}
 			if (SculptMaterial != nullptr)
 			{
diff --git a/Engine/Plugins/Experimental/MeshModelingToolset/Source/MeshModelingTools/Private/MeshGroupPaintTool.cpp b/Engine/Plugins/Experimental/MeshModelingToolset/Source/MeshModelingTools/Private/MeshGroupPaintTool.cpp
new file mode 100644
--- /dev/null
+++ b/Engine/Plugins/Experimental/MeshModelingToolset/Source/MeshModelingTools/Private/MeshGroupPaintTool.cpp
@@ -0,0 +1,1086 @@
+// Copyright Epic Games, Inc. All Rights Reserved.
+
+#include "MeshGroupPaintTool.h"
+#include "InteractiveToolManager.h"
+#include "InteractiveGizmoManager.h"
+#include "Drawing/MeshElementsVisualizer.h"
+#include "Async/ParallelFor.h"
+#include "Async/Async.h"
+
+#include "MeshWeights.h"
+#include "MeshNormals.h"
+#include "MeshIndexUtil.h"
+#include "Util/BufferUtil.h"
+#include "Util/ColorConstants.h"
+#include "Selections/MeshConnectedComponents.h"
+#include "Selections/MeshFaceSelection.h"
+#include "Polygroups/PolygroupUtil.h"
+
+#include "Changes/MeshVertexChange.h"
+#include "Changes/MeshPolygroupChange.h"
+#include "Changes/BasicChanges.h"
+
+#include "Sculpting/MeshGroupPaintBrushOps.h"
+#include "Sculpting/StampFalloffs.h"
+#include "Sculpting/MeshSculptUtil.h"
+
+#include "ToolSetupUtil.h"
+
+#define LOCTEXT_NAMESPACE "UMeshGroupPaintTool"
+
+namespace
+{
+	// probably should be something defined for the whole tool framework...
+#if WITH_EDITOR
+	static EAsyncExecution GroupPaintToolAsyncExecTarget = EAsyncExecution::LargeThreadPool;
+#else
+	static EAsyncExecution GroupPaintToolAsyncExecTarget = EAsyncExecution::ThreadPool;
+#endif
+}
+
+
+/*
+ * ToolBuilder
+ */
+UMeshSurfacePointTool* UMeshGroupPaintToolBuilder::CreateNewTool(const FToolBuilderState& SceneState) const
+{
+	UMeshGroupPaintTool* SculptTool = NewObject<UMeshGroupPaintTool>(SceneState.ToolManager);
+	SculptTool->SetWorld(SceneState.World);
+	return SculptTool;
+}
+
+
+/*
+ * Properties
+ */
+void UMeshGroupPaintToolActionPropertySet::PostAction(EMeshGroupPaintToolActions Action)
+{
+	if (ParentTool.IsValid())
+	{
+		ParentTool->RequestAction(Action);
+	}
+}
+
+
+
+
+
+/*
+ * Tool
+ */
+
+void UMeshGroupPaintTool::Setup()
+{
+	UMeshSculptToolBase::Setup();
+
+	// create dynamic mesh component to use for live preview
+	DynamicMeshComponent = NewObject<USimpleDynamicMeshComponent>(ComponentTarget->GetOwnerActor());
+	InitializeSculptMeshComponent(DynamicMeshComponent);
+
+	// assign materials
+	FComponentMaterialSet MaterialSet;
+	ComponentTarget->GetMaterialSet(MaterialSet);
+	for (int k = 0; k < MaterialSet.Materials.Num(); ++k)
+	{
+		DynamicMeshComponent->SetMaterial(k, MaterialSet.Materials[k]);
+	}
+
+	DynamicMeshComponent->bInvalidateProxyOnChange = false;
+	OnDynamicMeshComponentChangedHandle = DynamicMeshComponent->OnMeshVerticesChanged.AddUObject(this, &UMeshGroupPaintTool::OnDynamicMeshComponentChanged);
+
+	FDynamicMesh3* Mesh = GetSculptMesh();
+	Mesh->EnableVertexColors(FVector3f::One());
+	FAxisAlignedBox3d Bounds = Mesh->GetCachedBounds();
+
+	TFuture<void> PrecomputeFuture = Async(GroupPaintToolAsyncExecTarget, [&]()
+	{
+		PrecomputeFilterData();
+	});
+
+	TFuture<void> OctreeFuture = Async(GroupPaintToolAsyncExecTarget, [&]()
+	{
+		// initialize dynamic octree
+		if (Mesh->TriangleCount() > 100000)
+		{
+			Octree.RootDimension = Bounds.MaxDim() / 10.0;
+			Octree.SetMaxTreeDepth(4);
+		}
+		else
+		{
+			Octree.RootDimension = Bounds.MaxDim();
+			Octree.SetMaxTreeDepth(8);
+		}
+		Octree.Initialize(Mesh);
+		//Octree.CheckValidity(EValidityCheckFailMode::Check, true, true);
+		//FDynamicMeshOctree3::FStatistics Stats;
+		//Octree.ComputeStatistics(Stats);
+		//UE_LOG(LogTemp, Warning, TEXT("Octree Stats: %s"), *Stats.ToString());
+	});
+
+	// initialize render decomposition
+	TUniquePtr<FMeshRenderDecomposition> Decomp = MakeUnique<FMeshRenderDecomposition>();
+	FMeshRenderDecomposition::BuildChunkedDecomposition(Mesh, &MaterialSet, *Decomp);
+	Decomp->BuildAssociations(Mesh);
+	//UE_LOG(LogTemp, Warning, TEXT("Decomposition has %d groups"), Decomp->Num());
+	DynamicMeshComponent->SetExternalDecomposition(MoveTemp(Decomp));
+
+	// initialize brush radius range interval, brush properties
+	UMeshSculptToolBase::InitializeBrushSizeRange(Bounds);
+
+	PolygroupLayerProperties = NewObject<UPolygroupLayersProperties>(this);
+	PolygroupLayerProperties->RestoreProperties(this);
+	PolygroupLayerProperties->InitializeGroupLayers(GetSculptMesh());
+	PolygroupLayerProperties->WatchProperty(PolygroupLayerProperties->ActiveGroupLayer, [&](FName) { OnSelectedGroupLayerChanged(); });
+	UpdateActiveGroupLayer();
+	AddToolPropertySource(PolygroupLayerProperties);
+
+	ToolProperties = NewObject<UGroupPaintToolProperties>(this);
+	AddToolPropertySource(ToolProperties);
+	ToolProperties->WatchProperty(ToolProperties->SubToolType,
+		[this](EMeshGroupPaintInteractionType NewType) { UpdateSubToolType(NewType); });
+	ToolProperties->RestoreProperties(this);
+
+	// initialize other properties
+	FilterProperties = NewObject<UGroupPaintBrushFilterProperties>(this);
+
+	InitializeIndicator();
+
+	// initialize our properties
+	AddToolPropertySource(UMeshSculptToolBase::BrushProperties);
+	UMeshSculptToolBase::BrushProperties->bShowPerBrushProps = false;
+	UMeshSculptToolBase::BrushProperties->bShowFalloff = false;
+	CalculateBrushRadius();
+	FilterProperties->RestoreProperties(this);
+
+	PaintBrushOpOperties = NewObject<UGroupPaintBrushOpProps>(this);
+	RegisterBrushType((int32)EMeshGroupPaintBrushType::Paint,
+		MakeUnique<FLambdaMeshSculptBrushOpFactory>([this]() { return MakeUnique<FGroupPaintBrushOp>(); }),
+		PaintBrushOpOperties);
+
+	//RegisterBrushType((int32)EMeshGroupPaintBrushType::Erase,
+	//	MakeUnique<TBasicMeshSculptBrushOpFactory<FGroupEraseBrushOp>>(),
+	//	EraseBrushOpOperties);
+
+	// secondary brushes
+	EraseBrushOpOperties = NewObject<UGroupEraseBrushOpProps>(this);
+	EraseBrushOpOperties->GetCurrentGroupLambda = [this]() { return PaintBrushOpOperties->GetGroup(); };
+
+	RegisterSecondaryBrushType((int32)EMeshGroupPaintBrushType::Erase,
+		MakeUnique<TBasicMeshSculptBrushOpFactory<FGroupEraseBrushOp>>(),
+		EraseBrushOpOperties);
+
+	AddToolPropertySource(FilterProperties);
+	AddToolPropertySource(UMeshSculptToolBase::ViewProperties);
+
+	AddToolPropertySource(UMeshSculptToolBase::GizmoProperties);
+	SetToolPropertySourceEnabled(UMeshSculptToolBase::GizmoProperties, false);
+
+
+	// register watchers
+	FilterProperties->WatchProperty( FilterProperties->PrimaryBrushType,
+		[this](EMeshGroupPaintBrushType NewType) { UpdateBrushType(NewType); });
+
+	// must call before updating brush type so that we register all brush properties?
+	UMeshSculptToolBase::OnCompleteSetup();
+
+	UpdateBrushType(FilterProperties->PrimaryBrushType);
+	SetActiveSecondaryBrushType((int32)EMeshGroupPaintBrushType::Erase);
+
+	UpdateSubToolType(ToolProperties->SubToolType);
+
+	FreezeActions = NewObject<UMeshGroupPaintToolFreezeActions>(this);
+	FreezeActions->Initialize(this);
+	AddToolPropertySource(FreezeActions);
+
+	MeshElementsDisplay = NewObject<UMeshElementsVisualizer>(this);
+	MeshElementsDisplay->CreateInWorld(DynamicMeshComponent->GetWorld(), DynamicMeshComponent->GetComponentTransform());
+	if (ensure(MeshElementsDisplay->Settings))
+	{
+		MeshElementsDisplay->Settings->RestoreProperties(this);
+		AddToolPropertySource(MeshElementsDisplay->Settings);
+	}
+	MeshElementsDisplay->SetMeshAccessFunction([&](void) { return GetSculptMesh(); });
+
+	// force colors update... ?
+	DynamicMeshComponent->TriangleColorFunc = [this](const FDynamicMesh3* Mesh, int TriangleID)
+	{
+		return GetColorForGroup(ActiveGroupSet->GetGroup(TriangleID));
+	};
+	DynamicMeshComponent->FastNotifyColorsUpdated();
+
+	// disable view properties
+	SetViewPropertiesEnabled(false);
+	UpdateMaterialMode(EMeshEditingMaterialModes::VertexColor);
+	UpdateWireframeVisibility(false);
+	UpdateFlatShadingSetting(true);
+
+	PrecomputeFuture.Wait();
+	OctreeFuture.Wait();
+}
+
+void UMeshGroupPaintTool::Shutdown(EToolShutdownType ShutdownType)
+{
+	if (DynamicMeshComponent != nullptr)
+	{
+		DynamicMeshComponent->OnMeshChanged.Remove(OnDynamicMeshComponentChangedHandle);
+	}
+
+	if (ensure(MeshElementsDisplay->Settings))
+	{
+		MeshElementsDisplay->Settings->SaveProperties(this);
+	}
+	MeshElementsDisplay->Disconnect();
+
+	FilterProperties->SaveProperties(this);
+	ToolProperties->SaveProperties(this);
+	PolygroupLayerProperties->SaveProperties(this);
+
+
+	// TODO
+	// TODO
+	// TODO
+	// Bake should not have to replace entire mesh just to update groups...
+
+	// do our own bake
+	if (ShutdownType == EToolShutdownType::Accept)
+	{
+		GetToolManager()->BeginUndoTransaction(LOCTEXT("GroupPaintMeshToolTransactionName", "Paint Groups"));
+		ComponentTarget->CommitMesh([=](const FPrimitiveComponentTarget::FCommitParams& CommitParams)
+		{
+			FConversionToMeshDescriptionOptions ConversionOptions;
+			DynamicMeshComponent->Bake(CommitParams.MeshDescription, true, ConversionOptions);
+		});
+		GetToolManager()->EndUndoTransaction();
+	}
+
+
+	// this call will unregister and destroy DynamicMeshComponent
+	UMeshSculptToolBase::Shutdown(EToolShutdownType::Completed);
+}
+
+
+void UMeshGroupPaintTool::RegisterActions(FInteractiveToolActionSet& ActionSet)
+{
+	UMeshSculptToolBase::RegisterActions(ActionSet);
+
+	ActionSet.RegisterAction(this, (int32)EStandardToolActions::BaseClientDefinedActionID + 500,
+		TEXT("PickGroupColorUnderCursor"),
+		LOCTEXT("PickGroupColorUnderCursor", "Pick PolyGroup"),
+		LOCTEXT("PickGroupColorUnderCursorTooltip", "Switch the active PolyGroup to the group currently under the cursor"),
+		EModifierKey::Shift, EKeys::G,
+		[this]() { bPendingPickGroup = true; });
+
+	ActionSet.RegisterAction(this, (int32)EStandardToolActions::BaseClientDefinedActionID + 501,
+		TEXT("ToggleFrozenGroup"),
+		LOCTEXT("ToggleFrozenGroup", "Toggle Group Frozen State"),
+		LOCTEXT("ToggleFrozenGroupTooltip", "Toggle Group Frozen State"),
+		EModifierKey::Shift, EKeys::F,
+		[this]() { bPendingToggleFreezeGroup = true; });
+
+
+	ActionSet.RegisterAction(this, (int32)EStandardToolActions::BaseClientDefinedActionID + 502,
+		TEXT("CreateNewGroup"),
+		LOCTEXT("CreateNewGroup", "New Group"),
+		LOCTEXT("CreateNewGroupTooltip", "Allocate a new Polygroup and set as Current"),
+		EModifierKey::Shift, EKeys::Q,
+		[this]() { AllocateNewGroupAndSetAsCurrentAction(); });
+};
+
+
+void UMeshGroupPaintTool::OnPropertyModified(UObject* PropertySet, FProperty* Property)
+{
+	CalculateBrushRadius();
+}
+
+
+void UMeshGroupPaintTool::OnBeginStroke(const FRay& WorldRay)
+{
+	UpdateBrushPosition(WorldRay);
+
+	// initialize first "Last Stamp", so that we can assume all stamps in stroke have a valid previous stamp
+	LastStamp.WorldFrame = GetBrushFrameWorld();
+	LastStamp.LocalFrame = GetBrushFrameLocal();
+	LastStamp.Radius = GetCurrentBrushRadius();
+	LastStamp.Falloff = GetCurrentBrushFalloff();
+	LastStamp.Direction = GetInInvertStroke() ? -1.0 : 1.0;
+	LastStamp.Depth = GetCurrentBrushDepth();
+	LastStamp.Power = GetActivePressure() * GetCurrentBrushStrength();
+	LastStamp.TimeStamp = FDateTime::Now();
+
+	FSculptBrushOptions SculptOptions;
+	//SculptOptions.bPreserveUVFlow = false; // FilterProperties->bPreserveUVFlow;
+	SculptOptions.ConstantReferencePlane = GetCurrentStrokeReferencePlane();
+
+	TUniquePtr<FMeshSculptBrushOp>& UseBrushOp = GetActiveBrushOp();
+	UseBrushOp->ConfigureOptions(SculptOptions);
+	UseBrushOp->BeginStroke(GetSculptMesh(), LastStamp, VertexROI);
+
+	AccumulatedTriangleROI.Reset();
+
+	// begin change here? or wait for first stamp?
+	BeginChange();
+}
+
+void UMeshGroupPaintTool::OnEndStroke()
+{
+	GetActiveBrushOp()->EndStroke(GetSculptMesh(), LastStamp, VertexROI);
+
+	// close change record
+	EndChange();
+}
+
+
+
+void UMeshGroupPaintTool::UpdateROI(const FSculptBrushStamp& BrushStamp)
+{
+	SCOPE_CYCLE_COUNTER(GroupPaintTool_UpdateROI);
+
+	const FVector3d& BrushPos = BrushStamp.LocalFrame.Origin;
+	const FDynamicMesh3* Mesh = GetSculptMesh();
+	float RadiusSqr = GetCurrentBrushRadius() * GetCurrentBrushRadius();
+	FAxisAlignedBox3d BrushBox(
+		BrushPos - GetCurrentBrushRadius() * FVector3d::One(),
+		BrushPos + GetCurrentBrushRadius() * FVector3d::One());
+
+	TriangleROI.Reset();
+
+	int32 CenterTID = GetBrushTriangleID();
+	if (Mesh->IsTriangle(CenterTID))
+	{
+		TriangleROI.Add(CenterTID);
+	}
+
+	if (FilterProperties->bVolumetric)
+	{
+		Octree.RangeQuery(BrushBox,
+			[&](int TriIdx) {
+
+			if ((Mesh->GetTriCentroid(TriIdx) - BrushPos).SquaredLength() < RadiusSqr)
+			{
+				TriangleROI.Add(TriIdx);
+			}
+		});
+	}
+	else
+	{
+		if (Mesh->IsTriangle(CenterTID))
+		{
+			FVector3d CenterNormal = TriNormals[CenterTID];
+			bool bUseAngleThreshold = (FilterProperties->AngleThreshold < 180.0f);
+			double DotAngleThreshold = FMathd::Cos(FilterProperties->AngleThreshold * FMathd::DegToRad);
+
+			bool bStopAtUVSeams = FilterProperties->bUVSeams;
+			bool bStopAtNormalSeams = FilterProperties->bNormalSeams;
+
+			TArray<int32> StartROI;
+			StartROI.Add(CenterTID);
+			FMeshConnectedComponents::GrowToConnectedTriangles(Mesh, StartROI, TriangleROI, &TempROIBuffer,
+				[&](int t1, int t2) 
+			{ 
+				if ((Mesh->GetTriCentroid(t2) - BrushPos).SquaredLength() < RadiusSqr)
+				{
+					if (bUseAngleThreshold == false || CenterNormal.Dot(TriNormals[t2]) > DotAngleThreshold)
+					{
+						int32 eid = Mesh->FindEdgeFromTriPair(t1, t2);
+						if (bStopAtUVSeams == false || UVSeamEdges[eid] == false)
+						{
+							if (bStopAtNormalSeams == false || NormalSeamEdges[eid] == false)
+							{
+								return true;
+							}
+						}
+					}
+				}
+				return false;
+			});
+		}
+	}
+
+
+	// apply visibility filter
+	if (FilterProperties->VisibilityFilter != EMeshGroupPaintVisibilityType::None)
+	{
+		FViewCameraState StateOut;
+		GetToolManager()->GetContextQueriesAPI()->GetCurrentViewState(StateOut);
+		FVector3d LocalEyePosition(ComponentTarget->GetWorldTransform().InverseTransformPosition(StateOut.Position));
+		TempROIBuffer.SetNum(0, false);
+		for (int32 tid : TriangleROI)
+		{
+			TempROIBuffer.Add(tid);
+		}
+		ParallelFor(TempROIBuffer.Num(), [&](int32 idx)
+		{
+			FVector3d Centroid = Mesh->GetTriCentroid(TempROIBuffer[idx]);
+			FVector3d FaceNormal = Mesh->GetTriNormal(TempROIBuffer[idx]);
+			if (FaceNormal.Dot((Centroid - LocalEyePosition)) > 0)
+			{
+				TempROIBuffer[idx] = -1;
+			}
+			if (FilterProperties->VisibilityFilter == EMeshGroupPaintVisibilityType::Unoccluded)
+			{
+				int32 HitTID = Octree.FindNearestHitObject(FRay3d(LocalEyePosition, (Centroid - LocalEyePosition).Normalized()));
+				if (HitTID != TempROIBuffer[idx])
+				{
+					TempROIBuffer[idx] = -1;
+				}
+			}
+		});
+		TriangleROI.Reset();
+		for (int32 tid : TempROIBuffer)
+		{
+			if (tid >= 0)
+			{
+				TriangleROI.Add(tid);
+			}
+		}
+	}
+
+
+	VertexSetBuffer.Reset();
+	for (int32 tid : TriangleROI)
+	{
+		FIndex3i Tri = Mesh->GetTriangle(tid);
+		VertexSetBuffer.Add(Tri.A);  VertexSetBuffer.Add(Tri.B);  VertexSetBuffer.Add(Tri.C);
+	}
+	VertexROI.SetNum(0, false);
+	BufferUtil::AppendElements(VertexROI, VertexSetBuffer);
+
+
+	ROITriangleBuffer.Reserve(TriangleROI.Num());
+	ROITriangleBuffer.SetNum(0, false);
+	for (int32 tid : TriangleROI)
+	{
+		ROITriangleBuffer.Add(tid);
+	}
+	ROIGroupBuffer.SetNum(ROITriangleBuffer.Num(), false);
+}
+
+bool UMeshGroupPaintTool::UpdateStampPosition(const FRay& WorldRay)
+{
+	CalculateBrushRadius();
+
+	TUniquePtr<FMeshSculptBrushOp>& UseBrushOp = GetActiveBrushOp();
+
+	ESculptBrushOpTargetType TargetType = UseBrushOp->GetBrushTargetType();
+	switch (TargetType)
+	{
+	case ESculptBrushOpTargetType::SculptMesh:
+	case ESculptBrushOpTargetType::TargetMesh:
+		UpdateBrushPositionOnSculptMesh(WorldRay, true);
+		break;
+	case ESculptBrushOpTargetType::ActivePlane:
+		check(false);
+		UpdateBrushPositionOnActivePlane(WorldRay);
+		break;
+	}
+
+	if (UseBrushOp->GetAlignStampToView())
+	{
+		AlignBrushToView();
+	}
+
+	CurrentStamp = LastStamp;
+	CurrentStamp.DeltaTime = FMathd::Min((FDateTime::Now() - LastStamp.TimeStamp).GetTotalSeconds(), 1.0);
+	CurrentStamp.WorldFrame = GetBrushFrameWorld();
+	CurrentStamp.LocalFrame = GetBrushFrameLocal();
+	CurrentStamp.Power = GetActivePressure() * GetCurrentBrushStrength();
+
+	CurrentStamp.PrevLocalFrame = LastStamp.LocalFrame;
+	CurrentStamp.PrevWorldFrame = LastStamp.WorldFrame;
+
+	FVector3d MoveDelta = CurrentStamp.LocalFrame.Origin - CurrentStamp.PrevLocalFrame.Origin;
+	if (UseBrushOp->IgnoreZeroMovements() && MoveDelta.SquaredLength() < FMathd::ZeroTolerance)
+	{
+		return false;
+	}
+
+	return true;
+}
+
+
+void UMeshGroupPaintTool::ApplyStamp()
+{
+	SCOPE_CYCLE_COUNTER(GroupPaintToolApplyStamp);
+
+	TUniquePtr<FMeshSculptBrushOp>& UseBrushOp = GetActiveBrushOp();
+
+	// yuck
+	FMeshTriangleGroupEditBrushOp* GroupBrushOp = (FMeshTriangleGroupEditBrushOp*)UseBrushOp.Get();
+
+	FDynamicMesh3* Mesh = GetSculptMesh();
+	GroupBrushOp->ApplyStampByTriangles(Mesh, CurrentStamp, ROITriangleBuffer, ROIGroupBuffer);
+
+	SyncMeshWithGroupBuffer(Mesh);
+
+	LastStamp = CurrentStamp;
+	LastStamp.TimeStamp = FDateTime::Now();
+}
+
+
+
+
+void UMeshGroupPaintTool::SyncMeshWithGroupBuffer(FDynamicMesh3* Mesh)
+{
+	const int32 NumT = ROITriangleBuffer.Num();
+	// change update could be async here if we collected array of <idx,orig,new> and dispatched independenlty
+	for ( int32 k = 0; k < NumT; ++k)
+	{
+		int TriIdx = ROITriangleBuffer[k];
+		int32 CurGroupID = ActiveGroupSet->GetGroup(TriIdx);
+
+		if (FrozenGroups.Contains(CurGroupID))		// skip frozen groups
+		{
+			continue;
+		}
+
+		ActiveGroupEditBuilder->SaveTriangle(TriIdx, CurGroupID, ROIGroupBuffer[k]);
+
+		ActiveGroupSet->SetGroup(TriIdx, ROIGroupBuffer[k]);
+
+		//ActiveVertexChange->UpdateVertexColor(VertIdx, OrigColor, NewColor);
+	}
+}
+
+
+
+
+
+
+int32 UMeshGroupPaintTool::FindHitSculptMeshTriangle(const FRay3d& LocalRay)
+{
+	if (GetBrushCanHitBackFaces())
+	{
+		return Octree.FindNearestHitObject(LocalRay);
+	}
+	else
+	{
+		FDynamicMesh3* Mesh = GetSculptMesh();
+
+		FViewCameraState StateOut;
+		GetToolManager()->GetContextQueriesAPI()->GetCurrentViewState(StateOut);
+		FVector3d LocalEyePosition(CurTargetTransform.InverseTransformPosition(StateOut.Position));
+		int HitTID = Octree.FindNearestHitObject(LocalRay,
+			[this, Mesh, &LocalEyePosition](int TriangleID) {
+			FVector3d Normal, Centroid;
+			double Area;
+			Mesh->GetTriInfo(TriangleID, Normal, Area, Centroid);
+			return Normal.Dot((Centroid - LocalEyePosition)) < 0;
+		});
+		return HitTID;
+	}
+}
+
+int32 UMeshGroupPaintTool::FindHitTargetMeshTriangle(const FRay3d& LocalRay)
+{
+	check(false);
+	return IndexConstants::InvalidID;
+}
+
+
+
+bool UMeshGroupPaintTool::UpdateBrushPosition(const FRay& WorldRay)
+{
+	TUniquePtr<FMeshSculptBrushOp>& UseBrushOp = GetActiveBrushOp();
+
+	bool bHit = false; 
+	ESculptBrushOpTargetType TargetType = UseBrushOp->GetBrushTargetType();
+	switch (TargetType)
+	{
+	case ESculptBrushOpTargetType::SculptMesh:
+	case ESculptBrushOpTargetType::TargetMesh:
+		bHit = UpdateBrushPositionOnSculptMesh(WorldRay, false);
+		break;
+	case ESculptBrushOpTargetType::ActivePlane:
+		check(false);
+		bHit = UpdateBrushPositionOnSculptMesh(WorldRay, false);
+		break;
+	}
+
+	if (bHit && UseBrushOp->GetAlignStampToView())
+	{
+		AlignBrushToView();
+	}
+
+	return bHit;
+}
+
+
+
+
+bool UMeshGroupPaintTool::OnUpdateHover(const FInputDeviceRay& DevicePos)
+{
+	PendingStampType = FilterProperties->PrimaryBrushType;
+
+	if(ensure(InStroke() == false))
+	{
+		UpdateBrushPosition(DevicePos.WorldRay);
+	}
+	return true;
+}
+
+
+
+void UMeshGroupPaintTool::OnTick(float DeltaTime)
+{
+	UMeshSculptToolBase::OnTick(DeltaTime);
+	MeshElementsDisplay->OnTick(DeltaTime);
+
+	ConfigureIndicator(FilterProperties->bVolumetric);
+
+	if (bHavePendingAction)
+	{
+		ApplyAction(PendingAction);
+		bHavePendingAction = false;
+		PendingAction = EMeshGroupPaintToolActions::NoAction;
+	}
+
+	SCOPE_CYCLE_COUNTER(GroupPaintToolTick);
+
+	// process the undo update
+	if (bUndoUpdatePending)
+	{
+		// wait for updates
+		WaitForPendingUndoRedo();
+
+		// post rendering update
+		DynamicMeshComponent->FastNotifyTriangleVerticesUpdated(AccumulatedTriangleROI, EMeshRenderAttributeFlags::VertexColors);
+		GetToolManager()->PostInvalidation();
+
+		// ignore stamp and wait for next tick to do anything else
+		bUndoUpdatePending = false;
+		return;
+	}
+
+	if (bPendingPickGroup || bPendingToggleFreezeGroup)
+	{
+		if (GetBrushTriangleID() >= 0 && IsStampPending() == false )
+		{
+			if (GetSculptMesh()->IsTriangle(GetBrushTriangleID()))
+			{
+				int32 HitGroupID = ActiveGroupSet->GetGroup(GetBrushTriangleID());
+				if (bPendingPickGroup)
+				{
+					PaintBrushOpOperties->Group = HitGroupID;
+				}
+				else if (bPendingToggleFreezeGroup)
+				{
+					ToggleFrozenGroup(HitGroupID);
+				}
+			}
+		}
+		bPendingPickGroup = bPendingToggleFreezeGroup = false;
+	}
+
+
+	if (ToolProperties->SubToolType == EMeshGroupPaintInteractionType::Brush)
+	{
+		if (IsStampPending())
+		{
+			//UE_LOG(LogTemp, Warning, TEXT("dt is %.3f, tick fps %.2f - roi size %d/%d"), DeltaTime, 1.0 / DeltaTime, VertexROI.Num(), TriangleROI.Num());
+			SCOPE_CYCLE_COUNTER(GroupPaintTool_Tick_ApplyStampBlock);
+
+			ApplyStrokeFlowInTick();
+
+			// update brush position
+			if (UpdateStampPosition(GetPendingStampRayWorld()) == false)
+			{
+				return;
+			}
+
+			// update sculpt ROI
+			UpdateROI(CurrentStamp);
+
+			// append updated ROI to modified region (async)
+			TFuture<void> AccumulateROI = Async(GroupPaintToolAsyncExecTarget, [&]()
+			{
+				AccumulatedTriangleROI.Append(TriangleROI);
+			});
+
+			// apply the stamp
+			ApplyStamp();
+
+			{
+				SCOPE_CYCLE_COUNTER(GroupPaintTool_Tick_UpdateMeshBlock);
+				DynamicMeshComponent->FastNotifyTriangleVerticesUpdated(TriangleROI, EMeshRenderAttributeFlags::VertexColors);
+				GetToolManager()->PostInvalidation();
+			}
+
+			// we don't really need to wait for these to happen to end Tick()...
+			AccumulateROI.Wait();
+		}
+	}
+
+}
+
+
+
+
+
+
+
+void UMeshGroupPaintTool::AllocateNewGroupAndSetAsCurrentAction()
+{
+	int32 NewGroupID = ActiveGroupSet->AllocateNewGroupID();
+	PaintBrushOpOperties->Group = NewGroupID;
+}
+
+
+
+FColor UMeshGroupPaintTool::GetColorForGroup(int32 GroupID)
+{
+	FColor Color = LinearColors::SelectFColor(GroupID);
+	if (FrozenGroups.Contains(GroupID))
+	{
+		int32 GrayValue = (Color.R + Color.G + Color.B) / 3;
+		Color.R = Color.G = Color.B = FMath::Clamp(GrayValue, 0, 255);
+	}
+	return Color;
+}
+
+void UMeshGroupPaintTool::ToggleFrozenGroup(int32 FreezeGroupID)
+{
+	if (FreezeGroupID == 0) return;
+
+	TArray<int32> InitialFrozenGroups = FrozenGroups;
+	if (FrozenGroups.Contains(FreezeGroupID))
+	{
+		FrozenGroups.Remove(FreezeGroupID);
+	}
+	else
+	{
+		FrozenGroups.Add(FreezeGroupID);
+	}
+
+	const FDynamicMesh3* Mesh = DynamicMeshComponent->GetMesh();
+	TempROIBuffer.SetNum(0, false);
+	for (int32 tid : Mesh->TriangleIndicesItr())
+	{
+		int32 TriGroupID = ActiveGroupSet->GetGroup(tid);
+		if (TriGroupID == FreezeGroupID)
+		{
+			TempROIBuffer.Add(tid);
+		}
+	}
+	EmitFrozenGroupsChange(InitialFrozenGroups, FrozenGroups, LOCTEXT("ToggleFrozenGroup", "Toggle Frozen Group"));
+	DynamicMeshComponent->FastNotifyTriangleVerticesUpdated(TempROIBuffer, EMeshRenderAttributeFlags::VertexColors);
+	GetToolManager()->PostInvalidation();
+}
+
+void UMeshGroupPaintTool::FreezeOtherGroups(int32 KeepGroupID)
+{
+	TArray<int32> InitialFrozenGroups = FrozenGroups;
+	FrozenGroups.Reset();
+	const FDynamicMesh3* Mesh = DynamicMeshComponent->GetMesh();
+	TempROIBuffer.SetNum(0, false);
+	for (int32 tid : Mesh->TriangleIndicesItr())
+	{
+		int32 GroupID = ActiveGroupSet->GetGroup(tid);
+		if ( GroupID != 0 && GroupID != KeepGroupID)
+		{
+			FrozenGroups.AddUnique(ActiveGroupSet->GetGroup(tid));
+			TempROIBuffer.Add(tid);
+		}
+	}
+	EmitFrozenGroupsChange(InitialFrozenGroups, FrozenGroups, LOCTEXT("FreezeOtherGroups", "Freeze Other Groups"));
+	DynamicMeshComponent->FastNotifyTriangleVerticesUpdated(TempROIBuffer, EMeshRenderAttributeFlags::VertexColors);
+	GetToolManager()->PostInvalidation();
+}
+
+void UMeshGroupPaintTool::ClearAllFrozenGroups()
+{
+	TArray<int32> InitialFrozenGroups = FrozenGroups;
+	const FDynamicMesh3* Mesh = DynamicMeshComponent->GetMesh();
+	TempROIBuffer.SetNum(0, false);
+	for (int32 tid : Mesh->TriangleIndicesItr())
+	{
+		if ( FrozenGroups.Contains(ActiveGroupSet->GetGroup(tid)) )
+		{
+			TempROIBuffer.Add(tid);
+		}
+	}
+	FrozenGroups.Reset();
+	EmitFrozenGroupsChange(InitialFrozenGroups, FrozenGroups, LOCTEXT("ClearAllFrozenGroups", "Clear Frozen Groups"));
+	DynamicMeshComponent->FastNotifyTriangleVerticesUpdated(TempROIBuffer, EMeshRenderAttributeFlags::VertexColors);
+	GetToolManager()->PostInvalidation();
+}
+
+
+
+void UMeshGroupPaintTool::EmitFrozenGroupsChange(const TArray<int32>& FromGroups, const TArray<int32>& ToGroups, const FText& ChangeText)
+{
+	if (FromGroups != ToGroups)
+	{
+		TUniquePtr<TSimpleValueLambdaChange<TArray<int32>>> FrozenGroupsChange = MakeUnique<TSimpleValueLambdaChange<TArray<int32>>>();
+		FrozenGroupsChange->FromValue = FromGroups;
+		FrozenGroupsChange->ToValue = ToGroups;
+		FrozenGroupsChange->ValueChangeFunc = [this](UObject*, const TArray<int32>& FromGroups, const TArray<int32>& ToGroups, bool)
+		{	
+			FrozenGroups = ToGroups;
+			DynamicMeshComponent->FastNotifyVertexAttributesUpdated(EMeshRenderAttributeFlags::VertexColors);
+		};
+		GetToolManager()->EmitObjectChange(this, MoveTemp(FrozenGroupsChange), ChangeText);
+	}
+}
+
+
+void UMeshGroupPaintTool::GrowCurrentGroupAction()
+{
+	BeginChange();
+
+	int32 CurrentGroupID = PaintBrushOpOperties->Group;
+	const FDynamicMesh3* Mesh = DynamicMeshComponent->GetMesh();
+	FMeshFaceSelection InitialSelection(Mesh);
+	InitialSelection.Select([&](int32 tid) { return ActiveGroupSet->GetGroup(tid) == CurrentGroupID; });
+	FMeshFaceSelection ExpandSelection(InitialSelection);
+	ExpandSelection.ExpandToOneRingNeighbours([&](int32 tid) { return FrozenGroups.Contains(ActiveGroupSet->GetGroup(tid)) == false; });
+	TempROIBuffer.SetNum(0, false);
+	ExpandSelection.SetDifference(InitialSelection, TempROIBuffer);
+
+	ActiveGroupEditBuilder->SaveTriangles(TempROIBuffer);
+	for (int32 tid : TempROIBuffer)
+	{
+		ActiveGroupSet->SetGroup(tid, CurrentGroupID);
+	}
+	ActiveGroupEditBuilder->SaveTriangles(TempROIBuffer);
+
+	DynamicMeshComponent->FastNotifyTriangleVerticesUpdated(TempROIBuffer, EMeshRenderAttributeFlags::VertexColors);
+	GetToolManager()->PostInvalidation();
+	EndChange();
+}
+
+
+void UMeshGroupPaintTool::ShrinkCurrentGroupAction()
+{
+	BeginChange();
+
+	int32 CurrentGroupID = PaintBrushOpOperties->Group;
+	const FDynamicMesh3* Mesh = DynamicMeshComponent->GetMesh();
+	FMeshFaceSelection InitialSelection(Mesh);
+	InitialSelection.Select([&](int32 tid) { return ActiveGroupSet->GetGroup(tid) == CurrentGroupID; });
+	FMeshFaceSelection ContractSelection(InitialSelection);
+	ContractSelection.ContractBorderByOneRingNeighbours();
+	TempROIBuffer.SetNum(0, false);
+	InitialSelection.SetDifference(ContractSelection, TempROIBuffer);
+
+	ActiveGroupEditBuilder->SaveTriangles(TempROIBuffer);
+	for (int32 tid : TempROIBuffer)
+	{
+		// todo: could probably guess boundary groups here...
+		ActiveGroupSet->SetGroup(tid, 0);
+	}
+	ActiveGroupEditBuilder->SaveTriangles(TempROIBuffer);
+
+	DynamicMeshComponent->FastNotifyTriangleVerticesUpdated(TempROIBuffer, EMeshRenderAttributeFlags::VertexColors);
+	GetToolManager()->PostInvalidation();
+	EndChange();
+}
+
+
+
+//
+// Change Tracking
+//
+void UMeshGroupPaintTool::BeginChange()
+{
+	check(ActiveGroupEditBuilder == nullptr);
+	ActiveGroupEditBuilder = MakeUnique<FDynamicMeshGroupEditBuilder>(ActiveGroupSet.Get());
+}
+
+void UMeshGroupPaintTool::EndChange()
+{
+	check(ActiveGroupEditBuilder);
+
+	TUniquePtr<FDynamicMeshGroupEdit> EditResult = ActiveGroupEditBuilder->ExtractResult();
+	ActiveGroupEditBuilder = nullptr;
+
+	TUniquePtr<TWrappedToolCommandChange<FMeshPolygroupChange>> NewChange = MakeUnique<TWrappedToolCommandChange<FMeshPolygroupChange>>();
+	NewChange->WrappedChange = MakeUnique<FMeshPolygroupChange>(MoveTemp(EditResult));
+	NewChange->BeforeModify = [this](bool bRevert)
+	{
+		this->WaitForPendingUndoRedo();
+	};
+
+	GetToolManager()->EmitObjectChange(DynamicMeshComponent, MoveTemp(NewChange), LOCTEXT("GroupPaintChange", "Group Stroke"));
+}
+
+
+void UMeshGroupPaintTool::WaitForPendingUndoRedo()
+{
+	if (bUndoUpdatePending)
+	{
+		bUndoUpdatePending = false;
+	}
+}
+
+void UMeshGroupPaintTool::OnDynamicMeshComponentChanged(USimpleDynamicMeshComponent* Component, const FMeshVertexChange* Change, bool bRevert)
+{
+	// update octree
+	FDynamicMesh3* Mesh = GetSculptMesh();
+
+	// make sure any previous async computations are done, and update the undo ROI
+	if (bUndoUpdatePending)
+	{
+		// we should never hit this anymore, because of pre-change calling WaitForPendingUndoRedo()
+		WaitForPendingUndoRedo();
+
+		// this is not right because now we are going to do extra recomputation, but it's very messy otherwise...
+		MeshIndexUtil::VertexToTriangleOneRing(Mesh, Change->Vertices, AccumulatedTriangleROI);
+	}
+	else
+	{
+		AccumulatedTriangleROI.Reset();
+		MeshIndexUtil::VertexToTriangleOneRing(Mesh, Change->Vertices, AccumulatedTriangleROI);
+	}
+
+	// note that we have a pending update
+	bUndoUpdatePending = true;
+}
+
+
+void UMeshGroupPaintTool::PrecomputeFilterData()
+{
+	const FDynamicMesh3* Mesh = GetSculptMesh();
+	
+	TriNormals.SetNum(Mesh->MaxTriangleID());
+	ParallelFor(Mesh->MaxTriangleID(), [&](int32 tid)
+	{
+		if (Mesh->IsTriangle(tid))
+		{
+			TriNormals[tid] = Mesh->GetTriNormal(tid);
+		}
+	});
+
+	const FDynamicMeshNormalOverlay* Normals = Mesh->Attributes()->PrimaryNormals();
+	const FDynamicMeshUVOverlay* UVs = Mesh->Attributes()->PrimaryUV();
+	UVSeamEdges.SetNum(Mesh->MaxEdgeID());
+	NormalSeamEdges.SetNum(Mesh->MaxEdgeID());
+	ParallelFor(Mesh->MaxEdgeID(), [&](int32 eid)
+	{
+		if (Mesh->IsEdge(eid))
+		{
+			UVSeamEdges[eid] = UVs->IsSeamEdge(eid);
+			NormalSeamEdges[eid] = Normals->IsSeamEdge(eid);
+		}
+	});
+}
+
+
+void UMeshGroupPaintTool::OnSelectedGroupLayerChanged()
+{
+	GetToolManager()->BeginUndoTransaction(LOCTEXT("ChangeActiveGroupLayer", "Change Polygroup Layer"));
+
+	TArray<int32> InitialFrozenGroups = FrozenGroups;
+
+	int32 ActiveLayerIndex = (ActiveGroupSet) ? ActiveGroupSet->GetPolygroupIndex() : -1;
+	UpdateActiveGroupLayer();
+	int32 NewLayerIndex = (ActiveGroupSet) ? ActiveGroupSet->GetPolygroupIndex() : -1;
+
+	if (ActiveLayerIndex != NewLayerIndex)
+	{
+		// clear frozen groups
+		EmitFrozenGroupsChange(InitialFrozenGroups, FrozenGroups, LOCTEXT("ClearAllFrozenGroups", "Clear Frozen Groups"));
+
+		TUniquePtr<TSimpleValueLambdaChange<int32>> GroupLayerChange = MakeUnique<TSimpleValueLambdaChange<int32>>();
+		GroupLayerChange->FromValue = ActiveLayerIndex;
+		GroupLayerChange->ToValue = NewLayerIndex;
+		GroupLayerChange->ValueChangeFunc = [this](UObject*, int32 FromIndex, int32 ToIndex, bool)
+		{
+			this->PolygroupLayerProperties->SetSelectedFromPolygroupIndex(ToIndex);
+			this->PolygroupLayerProperties->SilentUpdateWatched();		// to prevent OnSelectedGroupLayerChanged() from being called immediately
+			this->UpdateActiveGroupLayer();
+		};
+		GetToolManager()->EmitObjectChange(this, MoveTemp(GroupLayerChange), LOCTEXT("ChangeActiveGroupLayer", "Change Polygroup Layer"));
+	}
+
+	GetToolManager()->EndUndoTransaction();
+}
+
+
+void UMeshGroupPaintTool::UpdateActiveGroupLayer()
+{
+	if (PolygroupLayerProperties->HasSelectedPolygroup() == false)
+	{
+		ActiveGroupSet = MakeUnique<UE::Geometry::FPolygroupSet>(GetSculptMesh());
+	}
+	else
+	{
+		FName SelectedName = PolygroupLayerProperties->ActiveGroupLayer;
+		FDynamicMeshPolygroupAttribute* FoundAttrib = UE::Geometry::FindPolygroupLayerByName(*GetSculptMesh(), SelectedName);
+		ensureMsgf(FoundAttrib, TEXT("Selected Attribute Not Found! Falling back to Default group layer."));
+		ActiveGroupSet = MakeUnique<UE::Geometry::FPolygroupSet>(GetSculptMesh(), FoundAttrib);
+	}
+
+	// need to reset everything here...
+	FrozenGroups.Reset();
+
+	// update colors
+	DynamicMeshComponent->FastNotifyVertexAttributesUpdated(EMeshRenderAttributeFlags::VertexColors);
+	GetToolManager()->PostInvalidation();
+}
+
+
+
+void UMeshGroupPaintTool::UpdateSubToolType(EMeshGroupPaintInteractionType NewType)
+{
+	bool bSculptPropsVisible = (NewType == EMeshGroupPaintInteractionType::Brush);
+	SetToolPropertySourceEnabled(FilterProperties, bSculptPropsVisible);
+	SetToolPropertySourceEnabled(UMeshSculptToolBase::BrushProperties, bSculptPropsVisible);
+	SetBrushOpPropsVisibility(bSculptPropsVisible);
+}
+
+
+void UMeshGroupPaintTool::UpdateBrushType(EMeshGroupPaintBrushType BrushType)
+{
+	static const FText BaseMessage = LOCTEXT("OnStartTool", "Hold Shift to Erase. [/] and S/D change Size (+Shift to small-step). Shift+Q for New Group, Shift+G to pick Group, Shift+F to Freeze Group.");
+	FTextBuilder Builder;
+	Builder.AppendLine(BaseMessage);
+
+	SetActivePrimaryBrushType((int32)BrushType);
+
+	SetToolPropertySourceEnabled(GizmoProperties, false);
+
+	GetToolManager()->DisplayMessage(Builder.ToText(), EToolMessageLevel::UserNotification);
+}
+
+
+
+
+void UMeshGroupPaintTool::RequestAction(EMeshGroupPaintToolActions ActionType)
+{
+	if (!bHavePendingAction)
+	{
+		PendingAction = ActionType;
+		bHavePendingAction = true;
+	}
+}
+
+
+void UMeshGroupPaintTool::ApplyAction(EMeshGroupPaintToolActions ActionType)
+{
+	switch (ActionType)
+	{
+	case EMeshGroupPaintToolActions::ClearFrozen:
+		ClearAllFrozenGroups();
+		break;
+
+	case EMeshGroupPaintToolActions::FreezeCurrent:
+		ToggleFrozenGroup(PaintBrushOpOperties->Group);
+		break;
+
+	case EMeshGroupPaintToolActions::FreezeOthers:
+		FreezeOtherGroups(PaintBrushOpOperties->Group);
+		break;
+
+	case EMeshGroupPaintToolActions::GrowCurrent:
+		GrowCurrentGroupAction();
+		break;
+
+	case EMeshGroupPaintToolActions::ShrinkCurrent:
+		ShrinkCurrentGroupAction();
+		break;
+	}
+}
+
+
+
+#undef LOCTEXT_NAMESPACE
diff --git a/Engine/Plugins/Experimental/MeshModelingToolset/Source/MeshModelingTools/Private/Sculpting/MeshSculptToolBase.cpp b/Engine/Plugins/Experimental/MeshModelingToolset/Source/MeshModelingTools/Private/Sculpting/MeshSculptToolBase.cpp
--- a/Engine/Plugins/Experimental/MeshModelingToolset/Source/MeshModelingTools/Private/Sculpting/MeshSculptToolBase.cpp
+++ b/Engine/Plugins/Experimental/MeshModelingToolset/Source/MeshModelingTools/Private/Sculpting/MeshSculptToolBase.cpp
@@ -359,7 +359,17 @@ TUniquePtr<FMeshSculptBrushOp>& UMeshSculptToolBase::GetActiveBrushOp()
 	}
 }
 
-
+void UMeshSculptToolBase::SetBrushOpPropsVisibility(bool bVisible)
+{
+	if (PrimaryVisiblePropSet)
+	{
+		SetToolPropertySourceEnabled(PrimaryVisiblePropSet, bVisible);
+	}
+	if (SecondaryVisiblePropSet)
+	{
+		SetToolPropertySourceEnabled(SecondaryVisiblePropSet, bVisible);
+	}
+}
 
 
 void UMeshSculptToolBase::SetPrimaryFalloffType(EMeshSculptFalloffType FalloffType)
@@ -513,6 +523,8 @@ void UMeshSculptToolBase::UpdateBrushTargetPlaneFromHit(const FRay& WorldRay, co
 
 bool UMeshSculptToolBase::UpdateBrushPositionOnActivePlane(const FRay& WorldRay)
 {
+	LastBrushTriangleID = IndexConstants::InvalidID;
+
 	FVector3d NewHitPosWorld;
 	ActiveBrushTargetPlaneWorld.RayPlaneIntersection(WorldRay.Origin, WorldRay.Direction, 2, NewHitPosWorld);
 	UpdateBrushFrameWorld(NewHitPosWorld, ActiveBrushTargetPlaneWorld.Z());
@@ -521,6 +533,8 @@ bool UMeshSculptToolBase::UpdateBrushPositionOnActivePlane(const FRay& WorldRay)
 
 bool UMeshSculptToolBase::UpdateBrushPositionOnTargetMesh(const FRay& WorldRay, bool bFallbackToViewPlane)
 {
+	LastBrushTriangleID = IndexConstants::InvalidID;
+
 	FRay3d LocalRay = GetLocalRay(WorldRay);
 	int32 HitTID = FindHitTargetMeshTriangle(LocalRay);
 	if (HitTID != IndexConstants::InvalidID)
@@ -529,6 +543,7 @@ bool UMeshSculptToolBase::UpdateBrushPositionOnTargetMesh(const FRay& WorldRay,
 		FIntrRay3Triangle3d Query = TMeshQueries<FDynamicMesh3>::TriangleIntersection(*BaseMesh, HitTID, LocalRay);
 		FVector3d WorldNormal = CurTargetTransform.TransformNormal(BaseMesh->GetTriNormal(HitTID));
 		FVector3d WorldPos = CurTargetTransform.TransformPosition(LocalRay.PointAt(Query.RayParameter));
+		LastBrushTriangleID = HitTID;
 		UpdateBrushFrameWorld(WorldPos, WorldNormal);
 		return true;
 	}
@@ -547,6 +562,8 @@ bool UMeshSculptToolBase::UpdateBrushPositionOnTargetMesh(const FRay& WorldRay,
 
 bool UMeshSculptToolBase::UpdateBrushPositionOnSculptMesh(const FRay& WorldRay, bool bFallbackToViewPlane)
 {
+	LastBrushTriangleID = IndexConstants::InvalidID;
+
 	FRay3d LocalRay = GetLocalRay(WorldRay);
 	int32 HitTID = FindHitSculptMeshTriangle(LocalRay);
 	if (HitTID != IndexConstants::InvalidID)
@@ -555,6 +572,7 @@ bool UMeshSculptToolBase::UpdateBrushPositionOnSculptMesh(const FRay& WorldRay,
 		FIntrRay3Triangle3d Query = TMeshQueries<FDynamicMesh3>::TriangleIntersection(*SculptMesh, HitTID, LocalRay);
 		FVector3d WorldNormal = CurTargetTransform.TransformNormal(SculptMesh->GetTriNormal(HitTID));
 		FVector3d WorldPos = CurTargetTransform.TransformPosition(LocalRay.PointAt(Query.RayParameter));
+		LastBrushTriangleID = HitTID;
 		UpdateBrushFrameWorld(WorldPos, WorldNormal);
 		return true;
 	}
@@ -744,7 +762,13 @@ void UMeshSculptToolBase::DecreaseBrushRadiusSmallStepAction()
 
 
 
-
+void UMeshSculptToolBase::SetViewPropertiesEnabled(bool bNewValue)
+{
+	if (ViewProperties)
+	{
+		SetToolPropertySourceEnabled(ViewProperties, bNewValue);
+	}
+}
 
 void UMeshSculptToolBase::UpdateWireframeVisibility(bool bNewValue)
 {
@@ -795,6 +819,10 @@ void UMeshSculptToolBase::UpdateMaterialMode(EMeshEditingMaterialModes MaterialM
 				ActiveOverrideMaterial->SetTextureParameterValue(TEXT("ImageTexture"), ViewProperties->Image);
 			}
 		}
+		else if (MaterialMode == EMeshEditingMaterialModes::VertexColor)
+		{
+			ActiveOverrideMaterial = ToolSetupUtil::GetVertexColorMaterial(GetToolManager());
+		}
 		else
 		{
 			UMaterialInterface* SculptMaterial = nullptr;
@@ -813,7 +841,7 @@ void UMeshSculptToolBase::UpdateMaterialMode(EMeshEditingMaterialModes MaterialM
 				SculptMaterial = ToolSetupUtil::GetImageBasedSculptMaterial(GetToolManager(), ToolSetupUtil::ImageMaterialType::TangentNormalFromView);
 				break;
 			}
-			if (SculptMaterial != nullptr)
+			if (SculptMaterial != nullptr )
 			{
 				ActiveOverrideMaterial = UMaterialInstanceDynamic::Create(SculptMaterial, this);
 			}
@@ -848,6 +876,28 @@ void UMeshSculptToolBase::InitializeIndicator()
 	BrushIndicator->bDrawIndicatorLines = true;
 	BrushIndicator->bDrawRadiusCircle = false;
 	BrushIndicator->LineColor = FLinearColor(0.9f, 0.4f, 0.4f);
+
+	bIsVolumetricIndicator = true;
+}
+
+bool UMeshSculptToolBase::GetIsVolumetricIndicator()
+{
+	return bIsVolumetricIndicator;
+}
+
+void UMeshSculptToolBase::ConfigureIndicator(bool bVolumetric)
+{
+	if (bIsVolumetricIndicator == bVolumetric) return;
+	bIsVolumetricIndicator = bVolumetric;
+	BrushIndicatorMesh->SetVisible(bVolumetric);
+	if (bVolumetric)
+	{
+		BrushIndicator->bDrawRadiusCircle = false;
+	}
+	else
+	{
+		BrushIndicator->bDrawRadiusCircle = true;
+	}
 }
 
 
diff --git a/Engine/Plugins/Experimental/MeshModelingToolset/Source/MeshModelingTools/Private/Sculpting/MeshSculptUtil.cpp b/Engine/Plugins/Experimental/MeshModelingToolset/Source/MeshModelingTools/Private/Sculpting/MeshSculptUtil.cpp
new file mode 100644
--- /dev/null
+++ b/Engine/Plugins/Experimental/MeshModelingToolset/Source/MeshModelingTools/Private/Sculpting/MeshSculptUtil.cpp
@@ -0,0 +1,85 @@
+// Copyright Epic Games, Inc. All Rights Reserved.
+
+#include "MeshSculptUtil.h"
+
+void UE::SculptUtil::RecalculateNormals_Overlay(
+	FDynamicMesh3* Mesh, 
+	const TSet<int32>& ModifiedTris, 
+	TSet<int32>& VertexSetBuffer, 
+	TArray<int32>& NormalsBuffer)
+{
+	FDynamicMeshNormalOverlay* Normals = Mesh->HasAttributes() ? Mesh->Attributes()->PrimaryNormals() : nullptr;
+	check(Normals != nullptr);
+
+	NormalsBuffer.Reset();
+	VertexSetBuffer.Reset();
+	for (int32 TriangleID : ModifiedTris)
+	{
+		FIndex3i TriElems = Normals->GetTriangle(TriangleID);
+		for (int32 j = 0; j < 3; ++j)
+		{
+			int32 elemid = TriElems[j];
+			if (VertexSetBuffer.Contains(elemid) == false)
+			{
+				VertexSetBuffer.Add(elemid);
+				NormalsBuffer.Add(elemid);
+			}
+		}
+	}
+
+	ParallelFor(NormalsBuffer.Num(), [&](int32 k) {
+		int32 elemid = NormalsBuffer[k];
+		FVector3d NewNormal = FMeshNormals::ComputeOverlayNormal(*Mesh, Normals, elemid);
+		Normals->SetElement(elemid, (FVector3f)NewNormal);
+	});
+}
+
+
+
+void UE::SculptUtil::RecalculateNormals_PerVertex(
+	FDynamicMesh3* Mesh, 
+	const TSet<int32>& ModifiedTris, 
+	TSet<int32>& VertexSetBuffer, 
+	TArray<int32>& NormalsBuffer)
+{
+	NormalsBuffer.Reset();
+	VertexSetBuffer.Reset();
+	for (int32 TriangleID : ModifiedTris)
+	{
+		FIndex3i TriV = Mesh->GetTriangle(TriangleID);
+		for (int32 j = 0; j < 3; ++j)
+		{
+			int32 vid = TriV[j];
+			if (VertexSetBuffer.Contains(vid) == false)
+			{
+				VertexSetBuffer.Add(vid);
+				NormalsBuffer.Add(vid);
+			}
+		}
+	}
+
+	ParallelFor(NormalsBuffer.Num(), [&](int32 k) {
+		int32 vid = NormalsBuffer[k];
+		FVector3d NewNormal = FMeshNormals::ComputeVertexNormal(*Mesh, vid);
+		Mesh->SetVertexNormal(vid, (FVector3f)NewNormal);
+	});
+}
+
+
+
+void UE::SculptUtil::RecalculateROINormals(
+	FDynamicMesh3* Mesh, 
+	const TSet<int32>& TriangleROI, 
+	TSet<int32>& VertexSetBuffer, 
+	TArray<int32>& NormalsBuffer, 
+	bool bForceVertex)
+{
+	if (Mesh->HasAttributes() && Mesh->Attributes()->PrimaryNormals() && bForceVertex == false)
+	{
+		RecalculateNormals_Overlay(Mesh, TriangleROI, VertexSetBuffer, NormalsBuffer);
+	}
+	else
+	{
+		RecalculateNormals_PerVertex(Mesh, TriangleROI, VertexSetBuffer, NormalsBuffer);
+	}
+}
\ No newline at end of file
diff --git a/Engine/Plugins/Experimental/MeshModelingToolset/Source/MeshModelingTools/Private/Sculpting/MeshSculptUtil.h b/Engine/Plugins/Experimental/MeshModelingToolset/Source/MeshModelingTools/Private/Sculpting/MeshSculptUtil.h
--- a/Engine/Plugins/Experimental/MeshModelingToolset/Source/MeshModelingTools/Private/Sculpting/MeshSculptUtil.h
+++ b/Engine/Plugins/Experimental/MeshModelingToolset/Source/MeshModelingTools/Private/Sculpting/MeshSculptUtil.h
@@ -8,81 +8,19 @@
 #include "MeshNormals.h"
 
 
-namespace UE { namespace SculptUtil {
+namespace UE 
+{ 
+namespace SculptUtil 
+{
 
+	void RecalculateNormals_Overlay(FDynamicMesh3* Mesh, const TSet<int32>& ModifiedTris, TSet<int32>& VertexSetBuffer, TArray<int32>& NormalsBuffer);
 
 
-	void RecalculateNormals_Overlay(FDynamicMesh3* Mesh, const TSet<int32>& ModifiedTris, TSet<int32>& VertexSetBuffer, TArray<int32>& NormalsBuffer)
-	{
-		FDynamicMeshNormalOverlay* Normals = Mesh->HasAttributes() ? Mesh->Attributes()->PrimaryNormals() : nullptr;
-		check(Normals != nullptr);
+	void RecalculateNormals_PerVertex(FDynamicMesh3* Mesh, const TSet<int32>& ModifiedTris, TSet<int32>& VertexSetBuffer, TArray<int32>& NormalsBuffer);
 
-		NormalsBuffer.Reset();
-		VertexSetBuffer.Reset();
-		for (int32 TriangleID : ModifiedTris)
-		{
-			FIndex3i TriElems = Normals->GetTriangle(TriangleID);
-			for (int32 j = 0; j < 3; ++j)
-			{
-				int32 elemid = TriElems[j];
-				if (VertexSetBuffer.Contains(elemid) == false)
-				{
-					VertexSetBuffer.Add(elemid);
-					NormalsBuffer.Add(elemid);
-				}
-			}
-		}
 
-		ParallelFor(NormalsBuffer.Num(), [&](int32 k) {
-			int32 elemid = NormalsBuffer[k];
-			FVector3d NewNormal = FMeshNormals::ComputeOverlayNormal(*Mesh, Normals, elemid);
-			Normals->SetElement(elemid, (FVector3f)NewNormal);
-		});
-	}
-
-
-
-
-	void RecalculateNormals_PerVertex(FDynamicMesh3* Mesh, const TSet<int32>& ModifiedTris, TSet<int32>& VertexSetBuffer, TArray<int32>& NormalsBuffer)
-	{
-		NormalsBuffer.Reset();
-		VertexSetBuffer.Reset();
-		for (int32 TriangleID : ModifiedTris)
-		{
-			FIndex3i TriV = Mesh->GetTriangle(TriangleID);
-			for (int32 j = 0; j < 3; ++j)
-			{
-				int32 vid = TriV[j];
-				if (VertexSetBuffer.Contains(vid) == false)
-				{
-					VertexSetBuffer.Add(vid);
-					NormalsBuffer.Add(vid);
-				}
-			}
-		}
-
-		ParallelFor(NormalsBuffer.Num(), [&](int32 k) {
-			int32 vid = NormalsBuffer[k];
-			FVector3d NewNormal = FMeshNormals::ComputeVertexNormal(*Mesh, vid);
-			Mesh->SetVertexNormal(vid, (FVector3f)NewNormal);
-		});
-	}
-
-
-
-
-	void RecalculateROINormals(FDynamicMesh3* Mesh, const TSet<int32>& TriangleROI, TSet<int32>& VertexSetBuffer, TArray<int32>& NormalsBuffer, bool bForceVertex = false)
-	{
-		if (Mesh->HasAttributes() && Mesh->Attributes()->PrimaryNormals() && bForceVertex == false)
-		{
-			RecalculateNormals_Overlay(Mesh, TriangleROI, VertexSetBuffer, NormalsBuffer);
-		}
-		else
-		{
-			RecalculateNormals_PerVertex(Mesh, TriangleROI, VertexSetBuffer, NormalsBuffer);
-		}
-	}
 
+	void RecalculateROINormals(FDynamicMesh3* Mesh, const TSet<int32>& TriangleROI, TSet<int32>& VertexSetBuffer, TArray<int32>& NormalsBuffer, bool bForceVertex = false);
 
 
 /* end namespace UE::SculptUtil */  } }
\ No newline at end of file
diff --git a/Engine/Plugins/Experimental/MeshModelingToolset/Source/MeshModelingTools/Public/DynamicMeshSculptTool.h b/Engine/Plugins/Experimental/MeshModelingToolset/Source/MeshModelingTools/Public/DynamicMeshSculptTool.h
--- a/Engine/Plugins/Experimental/MeshModelingToolset/Source/MeshModelingTools/Public/DynamicMeshSculptTool.h
+++ b/Engine/Plugins/Experimental/MeshModelingToolset/Source/MeshModelingTools/Public/DynamicMeshSculptTool.h
@@ -350,6 +350,7 @@ private:
 	FVector3d LastBrushPosWorld;
 	FVector3d LastBrushPosNormalWorld;
 	FVector3d LastSmoothBrushPosLocal;
+	int32 LastBrushTriangleID = -1;
 
 	TArray<int> VertexROI;
 	TSet<int> VertexSetBuffer;
diff --git a/Engine/Plugins/Experimental/MeshModelingToolset/Source/MeshModelingTools/Public/MeshGroupPaintTool.h b/Engine/Plugins/Experimental/MeshModelingToolset/Source/MeshModelingTools/Public/MeshGroupPaintTool.h
new file mode 100644
--- /dev/null
+++ b/Engine/Plugins/Experimental/MeshModelingToolset/Source/MeshModelingTools/Public/MeshGroupPaintTool.h
@@ -0,0 +1,374 @@
+// Copyright Epic Games, Inc. All Rights Reserved.
+
+#pragma once
+
+#include "CoreMinimal.h"
+#include "HAL/Platform.h"
+#include "UObject/NoExportTypes.h"
+#include "BaseTools/BaseBrushTool.h"
+#include "SimpleDynamicMeshComponent.h"
+#include "PropertySets/PolygroupLayersProperties.h"
+
+#include "Sculpting/MeshSculptToolBase.h"
+#include "Sculpting/MeshBrushOpBase.h"
+
+#include "DynamicMeshAABBTree3.h"
+#include "DynamicMeshOctree3.h"
+#include "MeshNormals.h"
+#include "TransformTypes.h"
+#include "Changes/MeshPolygroupChange.h"
+#include "Polygroups/PolygroupSet.h"
+
+#include "MeshGroupPaintTool.generated.h"
+
+class UMeshElementsVisualizer;
+class UGroupEraseBrushOpProps;
+class UGroupPaintBrushOpProps;
+
+DECLARE_STATS_GROUP(TEXT("GroupPaintTool"), STATGROUP_GroupPaintTool, STATCAT_Advanced);
+DECLARE_CYCLE_STAT(TEXT("GroupPaintTool_UpdateROI"), GroupPaintTool_UpdateROI, STATGROUP_GroupPaintTool);
+DECLARE_CYCLE_STAT(TEXT("GroupPaintTool_ApplyStamp"), GroupPaintToolApplyStamp, STATGROUP_GroupPaintTool );
+DECLARE_CYCLE_STAT(TEXT("GroupPaintTool_Tick"), GroupPaintToolTick, STATGROUP_GroupPaintTool);
+DECLARE_CYCLE_STAT(TEXT("GroupPaintTool_Tick_ApplyStampBlock"), GroupPaintTool_Tick_ApplyStampBlock, STATGROUP_GroupPaintTool);
+DECLARE_CYCLE_STAT(TEXT("GroupPaintTool_Tick_ApplyStamp_Remove"), GroupPaintTool_Tick_ApplyStamp_Remove, STATGROUP_GroupPaintTool);
+DECLARE_CYCLE_STAT(TEXT("GroupPaintTool_Tick_ApplyStamp_Insert"), GroupPaintTool_Tick_ApplyStamp_Insert, STATGROUP_GroupPaintTool);
+DECLARE_CYCLE_STAT(TEXT("GroupPaintTool_Tick_NormalsBlock"), GroupPaintTool_Tick_NormalsBlock, STATGROUP_GroupPaintTool);
+DECLARE_CYCLE_STAT(TEXT("GroupPaintTool_Tick_UpdateMeshBlock"), GroupPaintTool_Tick_UpdateMeshBlock, STATGROUP_GroupPaintTool);
+DECLARE_CYCLE_STAT(TEXT("GroupPaintTool_Tick_UpdateTargetBlock"), GroupPaintTool_Tick_UpdateTargetBlock, STATGROUP_GroupPaintTool);
+DECLARE_CYCLE_STAT(TEXT("GroupPaintTool_Normals_Collect"), GroupPaintTool_Normals_Collect, STATGROUP_GroupPaintTool);
+DECLARE_CYCLE_STAT(TEXT("GroupPaintTool_Normals_Compute"), GroupPaintTool_Normals_Compute, STATGROUP_GroupPaintTool);
+
+
+
+
+
+/**
+ * Tool Builder
+ */
+UCLASS()
+class MESHMODELINGTOOLS_API UMeshGroupPaintToolBuilder : public UMeshSurfacePointToolBuilder
+{
+	GENERATED_BODY()
+
+public:
+	virtual UMeshSurfacePointTool* CreateNewTool(const FToolBuilderState& SceneState) const override;
+};
+
+
+
+
+/** Mesh Sculpting Brush Types */
+UENUM()
+enum class EMeshGroupPaintInteractionType : uint8
+{
+	Brush,
+	Fill,
+	PolyLasso,
+
+	LastValue UMETA(Hidden)
+};
+
+
+
+
+
+
+UCLASS()
+class MESHMODELINGTOOLS_API UGroupPaintToolProperties : public UInteractiveToolPropertySet
+{
+	GENERATED_BODY()
+
+public:
+	//UPROPERTY(EditAnywhere, Category = "GroupPaint")
+	UPROPERTY()
+	EMeshGroupPaintInteractionType SubToolType = EMeshGroupPaintInteractionType::Brush;
+};
+
+
+
+
+
+/** Mesh Sculpting Brush Types */
+UENUM()
+enum class EMeshGroupPaintBrushType : uint8
+{
+	/** Paint active group */
+	Paint UMETA(DisplayName = "Paint"),
+
+	/** Erase active group */
+	Erase UMETA(DisplayName = "Erase"),
+
+	LastValue UMETA(Hidden)
+};
+
+
+/** Mesh Sculpting Brush Types */
+UENUM()
+enum class EMeshGroupPaintVisibilityType : uint8
+{
+	None,
+	FrontFacing,
+	Unoccluded
+};
+
+
+
+
+UCLASS()
+class MESHMODELINGTOOLS_API UGroupPaintBrushFilterProperties : public UInteractiveToolPropertySet
+{
+	GENERATED_BODY()
+
+public:
+	/** Primary Brush Mode */
+	//UPROPERTY(EditAnywhere, Category = Brush2, meta = (DisplayName = "Brush Type"))
+	UPROPERTY()
+	EMeshGroupPaintBrushType PrimaryBrushType = EMeshGroupPaintBrushType::Paint;
+
+	UPROPERTY(EditAnywhere, Category = BrushFilters, meta = (DisplayName = "Volumetric"))
+	bool bVolumetric = false;
+
+	UPROPERTY(EditAnywhere, Category = BrushFilters)
+	EMeshGroupPaintVisibilityType VisibilityFilter = EMeshGroupPaintVisibilityType::None;
+
+	UPROPERTY(EditAnywhere, Category = BrushFilters, meta = (UIMin = "0.0", UIMax = "180.0", EditCondition = "bVolumetric == false"))
+	float AngleThreshold = 180.0f;
+
+	UPROPERTY(EditAnywhere, Category = BrushFilters, meta = (UIMin = "0.0", UIMax = "180.0", EditCondition = "bVolumetric == false"))
+	bool bUVSeams = false;
+
+	UPROPERTY(EditAnywhere, Category = BrushFilters, meta = (UIMin = "0.0", UIMax = "180.0", EditCondition = "bVolumetric == false"))
+	bool bNormalSeams = false;
+
+};
+
+
+
+
+
+UENUM()
+enum class EMeshGroupPaintToolActions
+{
+	NoAction,
+
+	ClearFrozen,
+	FreezeCurrent,
+	FreezeOthers,
+
+	GrowCurrent,
+	ShrinkCurrent
+};
+
+
+UCLASS()
+class MESHMODELINGTOOLS_API UMeshGroupPaintToolActionPropertySet : public UInteractiveToolPropertySet
+{
+	GENERATED_BODY()
+
+public:
+	TWeakObjectPtr<UMeshGroupPaintTool> ParentTool;
+
+	void Initialize(UMeshGroupPaintTool* ParentToolIn) { ParentTool = ParentToolIn; }
+
+	void PostAction(EMeshGroupPaintToolActions Action);
+};
+
+
+
+UCLASS()
+class MESHMODELINGTOOLS_API UMeshGroupPaintToolFreezeActions : public UMeshGroupPaintToolActionPropertySet
+{
+	GENERATED_BODY()
+
+public:
+	UFUNCTION(CallInEditor, Category = Freezing, meta = (DisplayPriority = 1))
+	void UnfreezeAll()
+	{
+		PostAction(EMeshGroupPaintToolActions::ClearFrozen);
+	}
+
+	UFUNCTION(CallInEditor, Category = Freezing, meta = (DisplayPriority = 2))
+	void FreezeCurrent()
+	{
+		PostAction(EMeshGroupPaintToolActions::FreezeCurrent);
+	}
+
+	UFUNCTION(CallInEditor, Category = Freezing, meta = (DisplayPriority = 3))
+	void FreezeOthers()
+	{
+		PostAction(EMeshGroupPaintToolActions::FreezeOthers);
+	}
+
+	UFUNCTION(CallInEditor, Category = Operations, meta = (DisplayPriority = 4))
+	void GrowCurrent()
+	{
+		PostAction(EMeshGroupPaintToolActions::GrowCurrent);
+	}
+
+	UFUNCTION(CallInEditor, Category = Operations, meta = (DisplayPriority = 5))
+	void ShrinkCurrent()
+	{
+		PostAction(EMeshGroupPaintToolActions::ShrinkCurrent);
+	}
+};
+
+
+
+
+/**
+ * Mesh Element Paint Tool Class
+ */
+UCLASS()
+class MESHMODELINGTOOLS_API UMeshGroupPaintTool : public UMeshSculptToolBase
+{
+	GENERATED_BODY()
+
+public:
+	virtual void RegisterActions(FInteractiveToolActionSet& ActionSet) override;
+
+	virtual void Setup() override;
+	virtual void Shutdown(EToolShutdownType ShutdownType) override;
+
+	virtual void OnTick(float DeltaTime) override;
+
+	virtual bool HasCancel() const override { return true; }
+	virtual bool HasAccept() const override { return true; }
+	virtual bool CanAccept() const override { return true; }
+
+	virtual bool OnUpdateHover(const FInputDeviceRay& DevicePos) override;
+
+	virtual void OnPropertyModified(UObject* PropertySet, FProperty* Property) override;
+
+public:
+
+	UPROPERTY()
+	UPolygroupLayersProperties* PolygroupLayerProperties;
+
+	UPROPERTY()
+	UGroupPaintToolProperties* ToolProperties;
+
+	/** Filters on paint brush */
+	UPROPERTY()
+	UGroupPaintBrushFilterProperties* FilterProperties;
+
+
+private:
+	// This will be of type UGroupPaintBrushOpProps, we keep a ref so we can change active group ID on pick
+	UPROPERTY()
+	UGroupPaintBrushOpProps* PaintBrushOpOperties;
+
+	UPROPERTY()
+	UGroupEraseBrushOpProps* EraseBrushOpOperties;
+
+public:
+	void AllocateNewGroupAndSetAsCurrentAction();
+	void GrowCurrentGroupAction();
+	void ShrinkCurrentGroupAction();
+
+protected:
+	// UMeshSculptToolBase API
+	virtual UBaseDynamicMeshComponent* GetSculptMeshComponent() { return DynamicMeshComponent; }
+	virtual FDynamicMesh3* GetBaseMesh() { check(false); return nullptr; }
+	virtual const FDynamicMesh3* GetBaseMesh() const { check(false); return nullptr; }
+
+	virtual int32 FindHitSculptMeshTriangle(const FRay3d& LocalRay) override;
+	virtual int32 FindHitTargetMeshTriangle(const FRay3d& LocalRay) override;
+
+	virtual void OnBeginStroke(const FRay& WorldRay) override;
+	virtual void OnEndStroke() override;
+	// end UMeshSculptToolBase API
+
+
+
+	//
+	// Action support
+	//
+
+public:
+	virtual void RequestAction(EMeshGroupPaintToolActions ActionType);
+
+	UPROPERTY()
+	UMeshGroupPaintToolFreezeActions* FreezeActions;
+
+protected:
+	bool bHavePendingAction = false;
+	EMeshGroupPaintToolActions PendingAction;
+	virtual void ApplyAction(EMeshGroupPaintToolActions ActionType);
+
+
+	//
+	// Internals
+	//
+
+protected:
+	UPROPERTY()
+	USimpleDynamicMeshComponent* DynamicMeshComponent;
+
+	UPROPERTY()
+	UMeshElementsVisualizer* MeshElementsDisplay;
+
+	// realtime visualization
+	void OnDynamicMeshComponentChanged(USimpleDynamicMeshComponent* Component, const FMeshVertexChange* Change, bool bRevert);
+	FDelegateHandle OnDynamicMeshComponentChangedHandle;
+
+	TUniquePtr<UE::Geometry::FPolygroupSet> ActiveGroupSet;
+	void OnSelectedGroupLayerChanged();
+	void UpdateActiveGroupLayer();
+
+	void UpdateSubToolType(EMeshGroupPaintInteractionType NewType);
+
+	void UpdateBrushType(EMeshGroupPaintBrushType BrushType);
+
+	TSet<int32> AccumulatedTriangleROI;
+	bool bUndoUpdatePending = false;
+	TArray<int> NormalsBuffer;
+	void WaitForPendingUndoRedo();
+
+	TArray<int> TempROIBuffer;
+	TArray<int> VertexROI;
+	TSet<int> VertexSetBuffer;
+	TSet<int> TriangleROI;
+	void UpdateROI(const FSculptBrushStamp& CurrentStamp);
+
+	EMeshGroupPaintBrushType PendingStampType = EMeshGroupPaintBrushType::Paint;
+
+	bool UpdateStampPosition(const FRay& WorldRay);
+	void ApplyStamp();
+
+	FDynamicMeshOctree3 Octree;
+
+	bool UpdateBrushPosition(const FRay& WorldRay);
+
+
+	bool bPendingPickGroup = false;
+	bool bPendingToggleFreezeGroup = false;
+
+
+	TArray<int32> ROITriangleBuffer;
+	TArray<int32> ROIGroupBuffer;
+	void SyncMeshWithGroupBuffer(FDynamicMesh3* Mesh);
+
+	TUniquePtr<FDynamicMeshGroupEditBuilder> ActiveGroupEditBuilder;
+	void BeginChange();
+	void EndChange();
+
+	TArray<int32> FrozenGroups;
+	void ToggleFrozenGroup(int32 GroupID);
+	void FreezeOtherGroups(int32 GroupID);
+	void ClearAllFrozenGroups();
+	void EmitFrozenGroupsChange(const TArray<int32>& FromGroups, const TArray<int32>& ToGroups, const FText& ChangeText);
+
+	FColor GetColorForGroup(int32 GroupID);
+
+	TArray<FVector3d> TriNormals;
+	TArray<int32> UVSeamEdges;
+	TArray<int32> NormalSeamEdges;
+	void PrecomputeFilterData();
+
+
+protected:
+	virtual bool ShowWorkPlane() const override { return false; }
+};
+
+
+
diff --git a/Engine/Plugins/Experimental/MeshModelingToolset/Source/MeshModelingTools/Public/Properties/MeshMaterialProperties.h b/Engine/Plugins/Experimental/MeshModelingToolset/Source/MeshModelingTools/Public/Properties/MeshMaterialProperties.h
--- a/Engine/Plugins/Experimental/MeshModelingToolset/Source/MeshModelingTools/Public/Properties/MeshMaterialProperties.h
+++ b/Engine/Plugins/Experimental/MeshModelingToolset/Source/MeshModelingTools/Public/Properties/MeshMaterialProperties.h
@@ -97,7 +97,8 @@ enum class EMeshEditingMaterialModes
 	Grey = 2,
 	Soft = 3,
 	TangentNormal = 4,
-	Custom = 5
+	VertexColor = 5,
+	Custom = 6
 };
 
 
diff --git a/Engine/Plugins/Experimental/MeshModelingToolset/Source/MeshModelingTools/Public/Sculpting/MeshGroupPaintBrushOps.h b/Engine/Plugins/Experimental/MeshModelingToolset/Source/MeshModelingTools/Public/Sculpting/MeshGroupPaintBrushOps.h
new file mode 100644
--- /dev/null
+++ b/Engine/Plugins/Experimental/MeshModelingToolset/Source/MeshModelingTools/Public/Sculpting/MeshGroupPaintBrushOps.h
@@ -0,0 +1,149 @@
+// Copyright Epic Games, Inc. All Rights Reserved.
+
+#pragma once
+
+#include "Sculpting/MeshBrushOpBase.h"
+#include "DynamicMesh3.h"
+#include "MeshWeights.h"
+#include "Async/ParallelFor.h"
+#include "MeshGroupPaintBrushOps.generated.h"
+
+
+
+
+class MESHMODELINGTOOLS_API FMeshTriangleGroupEditBrushOp : public FMeshSculptBrushOp
+{
+public:
+
+	// not supported for this kind of brush op
+	virtual void ApplyStamp(
+		const FDynamicMesh3* Mesh,
+		const FSculptBrushStamp& Stamp,
+		const TArray<int32>& Vertices,
+		TArray<FVector3d>& NewColorsOut) override
+	{
+		check(false);
+	}
+
+
+	virtual void ApplyStampByTriangles(
+		const FDynamicMesh3* Mesh,
+		const FSculptBrushStamp& Stamp,
+		const TArray<int32>& Triangles,
+		TArray<int32>& NewGroupsOut) = 0;
+
+};
+
+
+
+
+
+
+UCLASS()
+class MESHMODELINGTOOLS_API UGroupEraseBrushOpProps : public UMeshSculptBrushOpProps
+{
+	GENERATED_BODY()
+public:
+	/** Group to set as Erased value */
+	UPROPERTY(EditAnywhere, Category = EraseBrush, meta = (DisplayName = "Erase Group", UIMin = 1))
+	int32 Group = 0;
+
+	/** When enabled, only the current group configured in the Paint brush is erased */
+	UPROPERTY(EditAnywhere, Category = EraseBrush, meta = (DisplayName = "Only Current"))
+	bool bOnlyEraseCurrent = false;
+
+	virtual int32 GetGroup() { return Group; }
+
+	// this lambda defines the "current" group for bOnlyEraseCurrent mode
+	TUniqueFunction<int32()> GetCurrentGroupLambda = []() { return -1; };
+};
+
+
+
+class FGroupEraseBrushOp : public FMeshTriangleGroupEditBrushOp
+{
+public:
+
+	virtual void ApplyStampByTriangles(
+		const FDynamicMesh3* Mesh,
+		const FSculptBrushStamp& Stamp,
+		const TArray<int32>& Triangles,
+		TArray<int32>& NewGroupsOut) override
+	{
+		const FVector3d& StampPos = Stamp.LocalFrame.Origin;
+
+		UGroupEraseBrushOpProps* Props = GetPropertySetAs<UGroupEraseBrushOpProps>();
+		int32 EraseGroup = (int32)Props->GetGroup();
+		bool bOnlyEraseCurrent = Props->bOnlyEraseCurrent;
+		int32 EraseCurrentGroup = Props->GetCurrentGroupLambda();
+
+		int32 NumTriangles = Triangles.Num();
+		for (int32 k = 0; k < NumTriangles; ++k)
+		{
+			int32 TriIdx = Triangles[k];
+			int32 CurGroupID = Mesh->GetTriangleGroup(TriIdx);
+			if (bOnlyEraseCurrent)
+			{
+				NewGroupsOut[k] = (CurGroupID == EraseCurrentGroup) ? EraseGroup : CurGroupID;
+			}
+			else
+			{
+				NewGroupsOut[k] = EraseGroup;
+			}
+		}
+	}
+};
+
+
+
+
+UCLASS()
+class MESHMODELINGTOOLS_API UGroupPaintBrushOpProps : public UMeshSculptBrushOpProps
+{
+	GENERATED_BODY()
+public:
+	/** The group that will be assigned to triangles within the brush region */
+	UPROPERTY(EditAnywhere, Category = PaintBrush, meta = (DisplayName = "Group", UIMin = 0))
+	int32 Group = 1;
+
+	/** If true, only triangles with no group assigned will be painted */
+	UPROPERTY(EditAnywhere, Category = PaintBrush, meta = (DisplayName = "Only Ungrouped"))
+	bool bOnlyPaintUngrouped = false;
+
+	virtual int32 GetGroup() { return Group; }
+};
+
+
+
+class FGroupPaintBrushOp : public FMeshTriangleGroupEditBrushOp
+{
+public:
+
+	virtual void ApplyStampByTriangles(
+		const FDynamicMesh3* Mesh,
+		const FSculptBrushStamp& Stamp,
+		const TArray<int32>& Triangles,
+		TArray<int32>& NewGroupsOut) override
+	{
+		const FVector3d& StampPos = Stamp.LocalFrame.Origin;
+
+		UGroupPaintBrushOpProps* Props = GetPropertySetAs<UGroupPaintBrushOpProps>();
+		int32 SetToGroup = (int32)Props->GetGroup();
+		bool bOnlyUngrouped = Props->bOnlyPaintUngrouped;
+
+		int32 NumTriangles = Triangles.Num();
+		for ( int32 k = 0; k < NumTriangles; ++k)
+		{
+			int32 TriIdx = Triangles[k];
+			int32 OrigGroupID = Mesh->GetTriangleGroup(TriIdx);
+			if (bOnlyUngrouped == false || OrigGroupID == 0)
+			{
+				NewGroupsOut[k] = SetToGroup;
+			}
+			else
+			{
+				NewGroupsOut[k] = OrigGroupID;
+			}
+		}
+	}
+};
diff --git a/Engine/Plugins/Experimental/MeshModelingToolset/Source/MeshModelingTools/Public/Sculpting/MeshSculptToolBase.h b/Engine/Plugins/Experimental/MeshModelingToolset/Source/MeshModelingTools/Public/Sculpting/MeshSculptToolBase.h
--- a/Engine/Plugins/Experimental/MeshModelingToolset/Source/MeshModelingTools/Public/Sculpting/MeshSculptToolBase.h
+++ b/Engine/Plugins/Experimental/MeshModelingToolset/Source/MeshModelingTools/Public/Sculpting/MeshSculptToolBase.h
@@ -242,7 +242,7 @@ protected:
 	void SetActivePrimaryBrushType(int32 Identifier);
 	void SetActiveSecondaryBrushType(int32 Identifier);
 	TUniquePtr<FMeshSculptBrushOp>& GetActiveBrushOp();
-
+	void SetBrushOpPropsVisibility(bool bVisible);
 
 	//
 	// Falloff types
@@ -295,9 +295,11 @@ public:
 private:
 	FFrame3d LastBrushFrameWorld;
 	FFrame3d LastBrushFrameLocal;
+	int32 LastBrushTriangleID;
 protected:
 	const FFrame3d& GetBrushFrameWorld() const { return LastBrushFrameWorld; }
 	const FFrame3d& GetBrushFrameLocal() const { return LastBrushFrameLocal; }
+	int32 GetBrushTriangleID() const { return LastBrushTriangleID; }
 	void UpdateBrushFrameWorld(const FVector3d& NewPosition, const FVector3d& NewNormal);
 	void AlignBrushToView();
 
@@ -384,6 +386,7 @@ public:
 	UMaterialInstanceDynamic* ActiveOverrideMaterial;
 
 protected:
+	virtual void SetViewPropertiesEnabled(bool bNewValue);
 	virtual void UpdateWireframeVisibility(bool bNewValue);
 	virtual void UpdateMaterialMode(EMeshEditingMaterialModes NewMode);
 	virtual void UpdateFlatShadingSetting(bool bNewValue);
@@ -399,10 +402,18 @@ protected:
 	// subclasses should call this to create indicator in their ::Setup()
 	virtual void InitializeIndicator();
 
+	virtual bool GetIsVolumetricIndicator();
+
+	virtual void ConfigureIndicator(bool bVolumetric);
+
 protected:
 	UPROPERTY()
 	UBrushStampIndicator* BrushIndicator;
 
+	UPROPERTY()
+	bool bIsVolumetricIndicator;
+
+
 	UPROPERTY()
 	UMaterialInstanceDynamic* BrushIndicatorMaterial;
 
diff --git a/Engine/Plugins/Experimental/MeshModelingToolset/Source/ModelingComponents/Private/PropertySets/PolygroupLayersProperties.cpp b/Engine/Plugins/Experimental/MeshModelingToolset/Source/ModelingComponents/Private/PropertySets/PolygroupLayersProperties.cpp
new file mode 100644
--- /dev/null
+++ b/Engine/Plugins/Experimental/MeshModelingToolset/Source/ModelingComponents/Private/PropertySets/PolygroupLayersProperties.cpp
@@ -0,0 +1,43 @@
+// Copyright Epic Games, Inc. All Rights Reserved.
+
+#include "PropertySets/PolygroupLayersProperties.h"
+#include "DynamicMesh3.h"
+#include "DynamicMeshAttributeSet.h"
+
+
+void UPolygroupLayersProperties::InitializeGroupLayers(const FDynamicMesh3* Mesh)
+{
+	GroupLayersList.Add(TEXT("Default"));		// always have standard group
+	if (Mesh->Attributes())
+	{
+		for (int32 k = 0; k < Mesh->Attributes()->NumPolygroupLayers(); k++)
+		{
+			FName Name = Mesh->Attributes()->GetPolygroupLayer(k)->GetName();
+			GroupLayersList.Add(Name.ToString());
+		}
+	}
+
+	if (GroupLayersList.Contains(ActiveGroupLayer.ToString()) == false)		// discard restored value if it doesn't apply
+	{
+		ActiveGroupLayer = FName(GroupLayersList[0]);
+	}
+}
+
+
+bool UPolygroupLayersProperties::HasSelectedPolygroup() const
+{
+	return ActiveGroupLayer != FName(GroupLayersList[0]);
+}
+
+
+void UPolygroupLayersProperties::SetSelectedFromPolygroupIndex(int32 Index)
+{
+	if (Index < 0)
+	{
+		ActiveGroupLayer = FName(GroupLayersList[0]);
+	}
+	else
+	{
+		ActiveGroupLayer = FName(GroupLayersList[Index+1]);
+	}
+}
\ No newline at end of file
diff --git a/Engine/Plugins/Experimental/MeshModelingToolset/Source/ModelingComponents/Public/PropertySets/PolygroupLayersProperties.h b/Engine/Plugins/Experimental/MeshModelingToolset/Source/ModelingComponents/Public/PropertySets/PolygroupLayersProperties.h
new file mode 100644
--- /dev/null
+++ b/Engine/Plugins/Experimental/MeshModelingToolset/Source/ModelingComponents/Public/PropertySets/PolygroupLayersProperties.h
@@ -0,0 +1,39 @@
+// Copyright Epic Games, Inc. All Rights Reserved.
+
+#pragma once
+
+#include "CoreMinimal.h"
+#include "InteractiveTool.h"
+#include "PolygroupLayersProperties.generated.h"
+
+class FDynamicMesh3;
+
+/**
+ * Basic Tool Property Set that allows for selecting from a list of FNames (that we assume are Polygroup Layers)
+ */
+UCLASS()
+class MODELINGCOMPONENTS_API UPolygroupLayersProperties : public UInteractiveToolPropertySet
+{
+	GENERATED_BODY()
+public:
+
+	/** Select vertex weight map. If configured, the weight map value will be sampled to modulate displacement intensity. */
+	UPROPERTY(EditAnywhere, Category = PolygroupLayers, meta = (GetOptions = GetGroupLayersFunc))
+	FName ActiveGroupLayer;
+
+	// this function is called provide set of available group layers
+	UFUNCTION()
+	TArray<FString> GetGroupLayersFunc() { return GroupLayersList; }
+
+	// internal list used to implement above
+	UPROPERTY(meta = (TransientToolProperty))
+	TArray<FString> GroupLayersList;
+
+	void InitializeGroupLayers(const FDynamicMesh3* Mesh);
+
+	// return true if any option other than "None" is selected
+	bool HasSelectedPolygroup() const;
+
+	void SetSelectedFromPolygroupIndex(int32 Index);
+};
+
