GeometryProcessing: add DynamicMeshAttributeSet::SplitAllBowties() function, splits bowties in NormalOverlay and UVOverlays
ModelingTools: Expose Simplify To Minimal Planar simplification mode in Simplify Tool. Add edge filter to this simplification mode in TMeshSimplification. Add support for returning all edges along intersection curves in FMeshBoolean and FMeshSelfUnion. Add option BooleanMeshesOp and SelfUnionMeshesOp to apply minimal-planar simplification along these cut/intersection edges. Expose as toggles in CSGMeshesTool and SelfUnionMeshesTool
#rnx

diff --git a/Engine/Plugins/Experimental/GeometryProcessing/Source/DynamicMesh/Private/DynamicMeshAttributeSet.cpp b/Engine/Plugins/Experimental/GeometryProcessing/Source/DynamicMesh/Private/DynamicMeshAttributeSet.cpp
--- a/Engine/Plugins/Experimental/GeometryProcessing/Source/DynamicMesh/Private/DynamicMeshAttributeSet.cpp
+++ b/Engine/Plugins/Experimental/GeometryProcessing/Source/DynamicMesh/Private/DynamicMeshAttributeSet.cpp
@@ -2,6 +2,7 @@
 
 #include "DynamicMeshAttributeSet.h"
 #include "IndexTypes.h"
+#include "Async/ParallelFor.h"
 
 
 
@@ -125,6 +126,24 @@ void FDynamicMeshAttributeSet::CompactInPlace(const FCompactMaps& CompactMaps)
 
 
 
+void FDynamicMeshAttributeSet::SplitAllBowties(bool bParallel)
+{
+	int32 UVLayerCount = NumUVLayers();
+	ParallelFor(UVLayerCount + 1, [&](int32 idx)
+	{
+		if (idx == 0)
+		{
+			PrimaryNormals()->SplitBowties();
+		}
+		else
+		{
+			GetUVLayer(idx - 1)->SplitBowties();
+		}
+	}, (bParallel) ? EParallelForFlags::Unbalanced : EParallelForFlags::ForceSingleThread );
+}
+
+
+
 void FDynamicMeshAttributeSet::EnableMatchingAttributes(const FDynamicMeshAttributeSet& ToMatch)
 {
 	SetNumUVLayers(ToMatch.NumUVLayers());
diff --git a/Engine/Plugins/Experimental/GeometryProcessing/Source/DynamicMesh/Private/MeshSimplification.cpp b/Engine/Plugins/Experimental/GeometryProcessing/Source/DynamicMesh/Private/MeshSimplification.cpp
--- a/Engine/Plugins/Experimental/GeometryProcessing/Source/DynamicMesh/Private/MeshSimplification.cpp
+++ b/Engine/Plugins/Experimental/GeometryProcessing/Source/DynamicMesh/Private/MeshSimplification.cpp
@@ -837,6 +837,10 @@ static bool IsCollapsableDevelopableEdge(const FDynamicMesh3& Mesh, int32 Collap
 		if (eid != CollapseEdgeID)
 		{
 			FIndex2i EdgeT = Mesh.GetEdgeT(eid);
+			if (EdgeT.B == IndexConstants::InvalidID)
+			{
+				return false;		// abort if one of the edges of RemoveV is a boundary edge (?)
+			}
 			FVector3d Normal3 = GetTriNormalFunc(EdgeT.A);
 			FVector3d Normal4 = GetTriNormalFunc(EdgeT.B);
 
@@ -866,7 +870,9 @@ static bool IsCollapsableDevelopableEdge(const FDynamicMesh3& Mesh, int32 Collap
 
 
 template <typename QuadricErrorType>
-void TMeshSimplification<QuadricErrorType>::SimplifyToMinimalPlanar(double CoplanarAngleTolDeg)
+void TMeshSimplification<QuadricErrorType>::SimplifyToMinimalPlanar(
+	double CoplanarAngleTolDeg,
+	TFunctionRef<bool(int32 EdgeID)> EdgeFilterPredicate)
 {
 #define RETURN_IF_CANCELLED 	if (Cancelled()) { return; }
 
@@ -927,6 +933,11 @@ void TMeshSimplification<QuadricErrorType>::SimplifyToMinimalPlanar(double Copla
 		CollapseEdges.Reset();
 		for (int32 eid : Mesh->EdgeIndicesItr())
 		{
+			if (EdgeFilterPredicate(eid) == false)
+			{
+				continue;
+			}
+
 			FIndex2i ev = Mesh->GetEdgeV(eid);
 			if (DevelopableVerts[ev.A] || DevelopableVerts[ev.B])
 			{
diff --git a/Engine/Plugins/Experimental/GeometryProcessing/Source/DynamicMesh/Private/Operations/MeshBoolean.cpp b/Engine/Plugins/Experimental/GeometryProcessing/Source/DynamicMesh/Private/Operations/MeshBoolean.cpp
--- a/Engine/Plugins/Experimental/GeometryProcessing/Source/DynamicMesh/Private/Operations/MeshBoolean.cpp
+++ b/Engine/Plugins/Experimental/GeometryProcessing/Source/DynamicMesh/Private/Operations/MeshBoolean.cpp
@@ -673,6 +673,14 @@ bool FMeshBoolean::Compute()
 		CreatedBoundaryEdges = CutBoundaryEdges[0];
 	}
 
+	if (bTrackAllNewEdges)
+	{
+		for (int32 eid : CreatedBoundaryEdges)
+		{
+			AllNewEdges.Add(eid);
+		}
+	}
+
 	if (bPutResultInInputSpace)
 	{
 		MeshTransforms::ApplyTransform(*Result, ResultTransform);
@@ -741,6 +749,13 @@ bool FMeshBoolean::MergeEdges(const FMeshIndexMappings& IndexMaps, FDynamicMesh3
 		{
 			UnmatchedEdges.Add(Candidate.A);
 		}
+		else
+		{
+			if (bTrackAllNewEdges)
+			{
+				AllNewEdges.Add(Candidate.A);
+			}
+		}
 	}
 
 	// filter matched edges from the edge array for the other mesh
@@ -781,6 +796,10 @@ bool FMeshBoolean::MergeEdges(const FMeshIndexMappings& IndexMaps, FDynamicMesh3
 					if (EdgeMergeResult == EMeshResult::Ok)
 					{
 						UnmatchedEdges.RemoveAtSwap(UnmatchedIdx, 1, false);
+						if (bTrackAllNewEdges)
+						{
+							AllNewEdges.Add(EID);
+						}
 						break;
 					}
 				}
diff --git a/Engine/Plugins/Experimental/GeometryProcessing/Source/DynamicMesh/Private/Operations/MeshSelfUnion.cpp b/Engine/Plugins/Experimental/GeometryProcessing/Source/DynamicMesh/Private/Operations/MeshSelfUnion.cpp
--- a/Engine/Plugins/Experimental/GeometryProcessing/Source/DynamicMesh/Private/Operations/MeshSelfUnion.cpp
+++ b/Engine/Plugins/Experimental/GeometryProcessing/Source/DynamicMesh/Private/Operations/MeshSelfUnion.cpp
@@ -447,6 +447,14 @@ bool FMeshSelfUnion::Compute()
 		bWeldSuccess = MergeEdges(CutBoundaryEdges, FoundMatches);
 	}
 
+	if (bTrackAllNewEdges)
+	{
+		for (int32 eid : CreatedBoundaryEdges)
+		{
+			AllNewEdges.Add(eid);
+		}
+	}
+
 	MeshTransforms::ApplyTransform(*Mesh, ResultTransform);
 
 	return bWeldSuccess;
@@ -489,6 +497,14 @@ bool FMeshSelfUnion::MergeEdges(const TArray<int>& CutBoundaryEdges, const TMap<
 
 		FDynamicMesh3::FMergeEdgesInfo MergeInfo;
 		EMeshResult EdgeMergeResult = Mesh->MergeEdges(Candidate.A, Candidate.B, MergeInfo);
+
+		if (EdgeMergeResult == EMeshResult::Ok)
+		{
+			if (bTrackAllNewEdges)
+			{
+				AllNewEdges.Add(Candidate.A);
+			}
+		}
 	}
 
 	// collect remaining unmatched edges
@@ -531,6 +547,10 @@ bool FMeshSelfUnion::MergeEdges(const TArray<int>& CutBoundaryEdges, const TMap<
 				if (EdgeMergeResult == EMeshResult::Ok)
 				{
 					UnmatchedEdges.RemoveAtSwap(OtherIdx, 1, false);
+					if (bTrackAllNewEdges)
+					{
+						AllNewEdges.Add(EID);
+					}
 					break;
 				}
 			}
diff --git a/Engine/Plugins/Experimental/GeometryProcessing/Source/DynamicMesh/Public/DynamicMeshAttributeSet.h b/Engine/Plugins/Experimental/GeometryProcessing/Source/DynamicMesh/Public/DynamicMeshAttributeSet.h
--- a/Engine/Plugins/Experimental/GeometryProcessing/Source/DynamicMesh/Public/DynamicMeshAttributeSet.h
+++ b/Engine/Plugins/Experimental/GeometryProcessing/Source/DynamicMesh/Public/DynamicMeshAttributeSet.h
@@ -66,6 +66,13 @@ public:
 	void CompactInPlace(const FCompactMaps& CompactMaps);
 
 
+	/**
+	 * Split all bowtie vertices in all layers
+	 * @param bParallel if true, layers are processed in parallel
+	 */
+	void SplitAllBowties(bool bParallel = true);
+
+
 	/**
 	 * Enable the matching attributes and overlay layers as the reference Copy set, but do not copy any data across
 	 */
diff --git a/Engine/Plugins/Experimental/GeometryProcessing/Source/DynamicMesh/Public/MeshSimplification.h b/Engine/Plugins/Experimental/GeometryProcessing/Source/DynamicMesh/Public/MeshSimplification.h
--- a/Engine/Plugins/Experimental/GeometryProcessing/Source/DynamicMesh/Public/MeshSimplification.h
+++ b/Engine/Plugins/Experimental/GeometryProcessing/Source/DynamicMesh/Public/MeshSimplification.h
@@ -104,10 +104,14 @@ public:
 
 	/**
 	 * Maximally collapse mesh in a way that does not change shape at all.
-	 * This process does not invove quadric error at all.
+	 * This process does not involve quadric error at all.
 	 * @param AngleTolDeg two triangles are considered coplanar if their normals are within this angle tolerance
+	 * @param EdgeFilterPredicate only edges that pass this predicate will be considered for collapse. Default all true.
 	 */
-	virtual void SimplifyToMinimalPlanar(double CoplanarAngleTolDeg = 0.001);
+	virtual void SimplifyToMinimalPlanar(
+		double CoplanarAngleTolDeg = 0.001,
+		TFunctionRef<bool(int32 EdgeID)> EdgeFilterPredicate = [](int32) { return true; }
+		);
 
 
 	/** 
diff --git a/Engine/Plugins/Experimental/GeometryProcessing/Source/DynamicMesh/Public/Operations/MeshBoolean.h b/Engine/Plugins/Experimental/GeometryProcessing/Source/DynamicMesh/Public/Operations/MeshBoolean.h
--- a/Engine/Plugins/Experimental/GeometryProcessing/Source/DynamicMesh/Public/Operations/MeshBoolean.h
+++ b/Engine/Plugins/Experimental/GeometryProcessing/Source/DynamicMesh/Public/Operations/MeshBoolean.h
@@ -64,6 +64,9 @@ public:
 	/** Weld newly-created cut edges where the input meshes meet.  If false, the input meshes will remain topologically disconnected. */
 	bool bWeldSharedEdges = true;
 
+	/** Control whether new edges should be tracked */
+	bool bTrackAllNewEdges = false;
+
 	/** Set this to be able to cancel running operation */
 	FProgressCancel* Progress = nullptr;
 
@@ -85,6 +88,9 @@ public:
 	/** Boundary edges created by the mesh boolean algorithm failing to cleanly weld (doesn't include boundaries that already existed in source meshes) */
 	TArray<int> CreatedBoundaryEdges;
 
+	/** All edges created by mesh boolean algorithm. Only populated if bTrackAllNewEdges = true */
+	TSet<int32> AllNewEdges;
+
 public:
 
 	/**
diff --git a/Engine/Plugins/Experimental/GeometryProcessing/Source/DynamicMesh/Public/Operations/MeshSelfUnion.h b/Engine/Plugins/Experimental/GeometryProcessing/Source/DynamicMesh/Public/Operations/MeshSelfUnion.h
--- a/Engine/Plugins/Experimental/GeometryProcessing/Source/DynamicMesh/Public/Operations/MeshSelfUnion.h
+++ b/Engine/Plugins/Experimental/GeometryProcessing/Source/DynamicMesh/Public/Operations/MeshSelfUnion.h
@@ -50,6 +50,8 @@ public:
 	/** Weld newly-created cut edges where the mesh is unioned with itself.  If false, newly joined surfaces remain topologically disconnected. */
 	bool bWeldSharedEdges = true;
 
+	/** Control whether new edges should be tracked */
+	bool bTrackAllNewEdges = false;
 
 	/** Set this to be able to cancel running operation */
 	FProgressCancel* Progress = nullptr;
@@ -69,6 +71,9 @@ public:
 	/** Boundary edges created by the mesh boolean algorithm failing to cleanly weld (doesn't include boundaries that already existed in source mesh) */
 	TArray<int> CreatedBoundaryEdges;
 
+	/** All edges created by mesh boolean algorithm. Only populated if bTrackAllNewEdges = true */
+	TSet<int32> AllNewEdges;
+
 public:
 
 	FMeshSelfUnion(FDynamicMesh3* MeshIn)
diff --git a/Engine/Plugins/Experimental/MeshModelingToolset/Source/MeshModelingTools/Private/CSGMeshesTool.cpp b/Engine/Plugins/Experimental/MeshModelingToolset/Source/MeshModelingTools/Private/CSGMeshesTool.cpp
--- a/Engine/Plugins/Experimental/MeshModelingToolset/Source/MeshModelingTools/Private/CSGMeshesTool.cpp
+++ b/Engine/Plugins/Experimental/MeshModelingToolset/Source/MeshModelingTools/Private/CSGMeshesTool.cpp
@@ -177,11 +177,13 @@ TUniquePtr<FDynamicMeshOperator> UCSGMeshesTool::MakeNewOperator()
 		BooleanOp->TrimOperation = TrimProperties->WhichMesh;
 		BooleanOp->TrimSide = TrimProperties->TrimSide;
 		BooleanOp->bAttemptFixHoles = false;
+		BooleanOp->bTryCollapseExtraEdges = false;
 	}
 	else
 	{
 		BooleanOp->CSGOperation = CSGProperties->Operation;
 		BooleanOp->bAttemptFixHoles = CSGProperties->bAttemptFixHoles;
+		BooleanOp->bTryCollapseExtraEdges = CSGProperties->bCollapseExtraEdges;
 	}
 
 	check(OriginalDynamicMeshes.Num() == 2);
diff --git a/Engine/Plugins/Experimental/MeshModelingToolset/Source/MeshModelingTools/Private/SelfUnionMeshesTool.cpp b/Engine/Plugins/Experimental/MeshModelingToolset/Source/MeshModelingTools/Private/SelfUnionMeshesTool.cpp
--- a/Engine/Plugins/Experimental/MeshModelingToolset/Source/MeshModelingTools/Private/SelfUnionMeshesTool.cpp
+++ b/Engine/Plugins/Experimental/MeshModelingToolset/Source/MeshModelingTools/Private/SelfUnionMeshesTool.cpp
@@ -188,6 +188,7 @@ TUniquePtr<FDynamicMeshOperator> USelfUnionMeshesTool::MakeNewOperator()
 	TUniquePtr<FSelfUnionMeshesOp> Op = MakeUnique<FSelfUnionMeshesOp>();
 	
 	Op->bAttemptFixHoles = Properties->bAttemptFixHoles;
+	Op->bTryCollapseExtraEdges = Properties->bCollapseExtraEdges;
 	Op->WindingNumberThreshold = Properties->WindingNumberThreshold;
 	Op->bTrimFlaps = Properties->bTrimFlaps;
 
diff --git a/Engine/Plugins/Experimental/MeshModelingToolset/Source/MeshModelingTools/Public/CSGMeshesTool.h b/Engine/Plugins/Experimental/MeshModelingToolset/Source/MeshModelingTools/Public/CSGMeshesTool.h
--- a/Engine/Plugins/Experimental/MeshModelingToolset/Source/MeshModelingTools/Public/CSGMeshesTool.h
+++ b/Engine/Plugins/Experimental/MeshModelingToolset/Source/MeshModelingTools/Public/CSGMeshesTool.h
@@ -36,6 +36,10 @@ public:
 	UPROPERTY(EditAnywhere, Category = Options)
 	bool bAttemptFixHoles = false;
 
+	/** Try to collapse extra edges created by the Boolean operation */
+	UPROPERTY(EditAnywhere, Category = Options)
+	bool bCollapseExtraEdges = true;
+
 	/** If true, only the first mesh will keep its materials assignments; all other triangles will be assigned material 0 */
 	UPROPERTY(EditAnywhere, Category = Materials)
 	bool bOnlyUseFirstMeshMaterials = false;
diff --git a/Engine/Plugins/Experimental/MeshModelingToolset/Source/MeshModelingTools/Public/SelfUnionMeshesTool.h b/Engine/Plugins/Experimental/MeshModelingToolset/Source/MeshModelingTools/Public/SelfUnionMeshesTool.h
--- a/Engine/Plugins/Experimental/MeshModelingToolset/Source/MeshModelingTools/Public/SelfUnionMeshesTool.h
+++ b/Engine/Plugins/Experimental/MeshModelingToolset/Source/MeshModelingTools/Public/SelfUnionMeshesTool.h
@@ -28,6 +28,10 @@ public:
 	UPROPERTY(EditAnywhere, Category = Options)
 	bool bAttemptFixHoles = false;
 
+	/** Try to collapse extra edges created by the Boolean operation */
+	UPROPERTY(EditAnywhere, Category = Options)
+	bool bCollapseExtraEdges = true;
+
 	/** Show boundary edges created by the union operation -- often due to numerical error */
 	UPROPERTY(EditAnywhere, Category = Options)
 	bool bShowNewBoundaryEdges = true;
diff --git a/Engine/Plugins/Experimental/MeshModelingToolset/Source/ModelingOperators/Private/CompositionOps/BooleanMeshesOp.cpp b/Engine/Plugins/Experimental/MeshModelingToolset/Source/ModelingOperators/Private/CompositionOps/BooleanMeshesOp.cpp
--- a/Engine/Plugins/Experimental/MeshModelingToolset/Source/ModelingOperators/Private/CompositionOps/BooleanMeshesOp.cpp
+++ b/Engine/Plugins/Experimental/MeshModelingToolset/Source/ModelingOperators/Private/CompositionOps/BooleanMeshesOp.cpp
@@ -4,6 +4,8 @@
 
 #include "Operations/MeshBoolean.h"
 
+#include "MeshSimplification.h"
+
 #include "MeshBoundaryLoops.h"
 #include "Operations/MinimalHoleFiller.h"
 
@@ -70,6 +72,7 @@ void FBooleanMeshesOp::CalculateResult(FProgressCancel* Progress)
 	}
 
 	MeshBoolean.bPutResultInInputSpace = false;
+	MeshBoolean.bTrackAllNewEdges = (bTryCollapseExtraEdges);
 	MeshBoolean.Progress = Progress;
 	bool bSuccess = MeshBoolean.Compute();
 	ResultTransform = MeshBoolean.ResultTransform;
@@ -79,10 +82,47 @@ void FBooleanMeshesOp::CalculateResult(FProgressCancel* Progress)
 		return;
 	}
 
-	if (MeshBoolean.CreatedBoundaryEdges.Num() > 0 && bAttemptFixHoles)
+	CreatedBoundaryEdges = MeshBoolean.CreatedBoundaryEdges;
+
+	// Boolean operation is based on edge splits, which results in spurious vertices
+	// along straight intersection edges. Try to collapse away those extra vertices.
+	if (bTryCollapseExtraEdges)
+	{
+		FDynamicMesh3* TargetMesh = MeshBoolean.Result;
+
+		FQEMSimplification Simplifier(TargetMesh);
+		Simplifier.bAllowSeamCollapse = true;
+		if (TargetMesh->Attributes())
+		{
+			TargetMesh->Attributes()->SplitAllBowties();		// eliminate any bowties that might have formed on UV seams.
+		}
+
+		FMeshConstraints Constraints;
+		FMeshConstraintsUtil::ConstrainAllBoundariesAndSeams(Constraints, *TargetMesh,
+			EEdgeRefineFlags::NoConstraint, EEdgeRefineFlags::NoConstraint, EEdgeRefineFlags::NoConstraint,
+			true, true, true);
+		Simplifier.SetExternalConstraints(MoveTemp(Constraints));
+
+		Simplifier.SimplifyToMinimalPlanar( TryCollapseExtraEdgesPlanarThresh,
+			[&MeshBoolean](int32 eid) { return MeshBoolean.AllNewEdges.Contains(eid); } );
+
+		// update boundary-edge set
+		TArray<int32> UpdatedBoundaryEdges;
+		for (int32 eid : CreatedBoundaryEdges)
+		{
+			if (MeshBoolean.Result->IsEdge(eid))
+			{
+				UpdatedBoundaryEdges.Add(eid);
+			}
+		}
+		CreatedBoundaryEdges = MoveTemp(UpdatedBoundaryEdges);
+	}
+
+	// try to fill cracks/holes in boolean result
+	if (CreatedBoundaryEdges.Num() > 0 && bAttemptFixHoles)
 	{
 		FMeshBoundaryLoops OpenBoundary(MeshBoolean.Result, false);
-		TSet<int> ConsiderEdges(MeshBoolean.CreatedBoundaryEdges);
+		TSet<int> ConsiderEdges(CreatedBoundaryEdges);
 		OpenBoundary.EdgeFilterFunc = [&ConsiderEdges](int EID)
 		{
 			return ConsiderEdges.Contains(EID);
@@ -99,16 +139,15 @@ void FBooleanMeshesOp::CalculateResult(FProgressCancel* Progress)
 			FMinimalHoleFiller Filler(MeshBoolean.Result, Loop);
 			Filler.Fill();
 		}
-		for (int EID : MeshBoolean.CreatedBoundaryEdges)
+
+		TArray<int32> UpdatedBoundaryEdges;
+		for (int EID : CreatedBoundaryEdges)
 		{
 			if (MeshBoolean.Result->IsEdge(EID) && MeshBoolean.Result->IsBoundaryEdge(EID))
 			{
-				CreatedBoundaryEdges.Add(EID);
+				UpdatedBoundaryEdges.Add(EID);
 			}
 		}
-	}
-	else
-	{
-		CreatedBoundaryEdges = MeshBoolean.CreatedBoundaryEdges;
+		CreatedBoundaryEdges = MoveTemp(UpdatedBoundaryEdges);
 	}
 }
\ No newline at end of file
diff --git a/Engine/Plugins/Experimental/MeshModelingToolset/Source/ModelingOperators/Private/CompositionOps/SelfUnionMeshesOp.cpp b/Engine/Plugins/Experimental/MeshModelingToolset/Source/ModelingOperators/Private/CompositionOps/SelfUnionMeshesOp.cpp
--- a/Engine/Plugins/Experimental/MeshModelingToolset/Source/ModelingOperators/Private/CompositionOps/SelfUnionMeshesOp.cpp
+++ b/Engine/Plugins/Experimental/MeshModelingToolset/Source/ModelingOperators/Private/CompositionOps/SelfUnionMeshesOp.cpp
@@ -4,6 +4,8 @@
 
 #include "Operations/MeshSelfUnion.h"
 
+#include "MeshSimplification.h"
+
 #include "MeshBoundaryLoops.h"
 #include "Operations/MinimalHoleFiller.h"
 
@@ -28,6 +30,7 @@ void FSelfUnionMeshesOp::CalculateResult(FProgressCancel* Progress)
 	FMeshSelfUnion Union(ResultMesh.Get());
 	Union.WindingThreshold = WindingNumberThreshold;
 	Union.bTrimFlaps = bTrimFlaps;
+	Union.bTrackAllNewEdges = (bTryCollapseExtraEdges);
 	bool bSuccess = Union.Compute();
 
 	if (Progress->Cancelled())
@@ -35,10 +38,46 @@ void FSelfUnionMeshesOp::CalculateResult(FProgressCancel* Progress)
 		return;
 	}
 
+	CreatedBoundaryEdges = Union.CreatedBoundaryEdges;
+
+	// Boolean operation is based on edge splits, which results in spurious vertices
+	// along straight intersection edges. Try to collapse away those extra vertices.
+	if (bTryCollapseExtraEdges)
+	{
+		FDynamicMesh3* TargetMesh = ResultMesh.Get();
+
+		FQEMSimplification Simplifier(TargetMesh);
+		Simplifier.bAllowSeamCollapse = true;
+		if (TargetMesh->Attributes())
+		{
+			TargetMesh->Attributes()->SplitAllBowties();		// eliminate any bowties that might have formed on UV seams.
+		}
+
+		FMeshConstraints Constraints;
+		FMeshConstraintsUtil::ConstrainAllBoundariesAndSeams(Constraints, *TargetMesh,
+			EEdgeRefineFlags::NoConstraint, EEdgeRefineFlags::NoConstraint, EEdgeRefineFlags::NoConstraint,
+			true, true, true);
+		Simplifier.SetExternalConstraints(MoveTemp(Constraints));
+
+		Simplifier.SimplifyToMinimalPlanar(TryCollapseExtraEdgesPlanarThresh,
+			[&Union](int32 eid) { return Union.AllNewEdges.Contains(eid); });
+
+		// update boundary-edge set as we may have collapsed some during this process
+		TArray<int32> UpdatedBoundaryEdges;
+		for (int32 eid : CreatedBoundaryEdges)
+		{
+			if (ResultMesh->IsEdge(eid))
+			{
+				UpdatedBoundaryEdges.Add(eid);
+			}
+		}
+		CreatedBoundaryEdges = MoveTemp(UpdatedBoundaryEdges);
+	}
+
 	if (!bSuccess && bAttemptFixHoles)
 	{
 		FMeshBoundaryLoops OpenBoundary(ResultMesh.Get(), false);
-		TSet<int> ConsiderEdges(Union.CreatedBoundaryEdges);
+		TSet<int> ConsiderEdges(CreatedBoundaryEdges);
 		OpenBoundary.EdgeFilterFunc = [&ConsiderEdges](int EID)
 		{
 			return ConsiderEdges.Contains(EID);
@@ -55,16 +94,16 @@ void FSelfUnionMeshesOp::CalculateResult(FProgressCancel* Progress)
 			FMinimalHoleFiller Filler(ResultMesh.Get(), Loop);
 			Filler.Fill();
 		}
+
+		TArray<int32> UpdatedBoundaryEdges;
 		for (int EID : Union.CreatedBoundaryEdges)
 		{
 			if (ResultMesh->IsEdge(EID) && ResultMesh->IsBoundaryEdge(EID))
 			{
-				CreatedBoundaryEdges.Add(EID);
+				UpdatedBoundaryEdges.Add(EID);
 			}
 		}
+		CreatedBoundaryEdges = MoveTemp(UpdatedBoundaryEdges);
 	}
-	else
-	{
-		CreatedBoundaryEdges = Union.CreatedBoundaryEdges;
-	}
+
 }
\ No newline at end of file
diff --git a/Engine/Plugins/Experimental/MeshModelingToolset/Source/ModelingOperators/Public/CompositionOps/BooleanMeshesOp.h b/Engine/Plugins/Experimental/MeshModelingToolset/Source/ModelingOperators/Public/CompositionOps/BooleanMeshesOp.h
--- a/Engine/Plugins/Experimental/MeshModelingToolset/Source/ModelingOperators/Public/CompositionOps/BooleanMeshesOp.h
+++ b/Engine/Plugins/Experimental/MeshModelingToolset/Source/ModelingOperators/Public/CompositionOps/BooleanMeshesOp.h
@@ -58,7 +58,12 @@ public:
 	bool bTrimMode = false; // if true, do a trim operation instead of a boolean
 	TArray<TSharedPtr<const FDynamicMesh3>> Meshes;
 	TArray<FTransform> Transforms; // 1:1 with Meshes
-	bool bAttemptFixHoles;
+	bool bAttemptFixHoles = false;
+
+	/** If true, try to do edge-collapses along cut edges to remove unnecessary edges inserted by cut */
+	bool bTryCollapseExtraEdges = false;
+	/** Angle threshold in degrees used for testing if two triangles should be considered coplanar, or two lines collinear */
+	float TryCollapseExtraEdgesPlanarThresh = 0.01f;
 
 	void SetTransform(const FTransform& Transform);
 
diff --git a/Engine/Plugins/Experimental/MeshModelingToolset/Source/ModelingOperators/Public/CompositionOps/SelfUnionMeshesOp.h b/Engine/Plugins/Experimental/MeshModelingToolset/Source/ModelingOperators/Public/CompositionOps/SelfUnionMeshesOp.h
--- a/Engine/Plugins/Experimental/MeshModelingToolset/Source/ModelingOperators/Public/CompositionOps/SelfUnionMeshesOp.h
+++ b/Engine/Plugins/Experimental/MeshModelingToolset/Source/ModelingOperators/Public/CompositionOps/SelfUnionMeshesOp.h
@@ -14,9 +14,14 @@ public:
 
 	// inputs
 	TSharedPtr<const FDynamicMesh3> CombinedMesh;
-	bool bAttemptFixHoles;
+	bool bAttemptFixHoles = false;
 	double WindingNumberThreshold = .5;
-	bool bTrimFlaps;
+	bool bTrimFlaps = false;
+
+	/** If true, try to do edge-collapses along cut edges to remove unnecessary edges inserted by cut */
+	bool bTryCollapseExtraEdges = false;
+	/** Angle threshold in degrees used for testing if two triangles should be considered coplanar, or two lines collinear */
+	float TryCollapseExtraEdgesPlanarThresh = 0.01f;
 
 	void SetTransform(const FTransform& Transform);
 
diff --git a/Engine/Plugins/Experimental/MeshModelingToolset/Source/ModelingOperatorsEditorOnly/Private/CleaningOps/SimplifyMeshOp.cpp b/Engine/Plugins/Experimental/MeshModelingToolset/Source/ModelingOperatorsEditorOnly/Private/CleaningOps/SimplifyMeshOp.cpp
--- a/Engine/Plugins/Experimental/MeshModelingToolset/Source/ModelingOperatorsEditorOnly/Private/CleaningOps/SimplifyMeshOp.cpp
+++ b/Engine/Plugins/Experimental/MeshModelingToolset/Source/ModelingOperatorsEditorOnly/Private/CleaningOps/SimplifyMeshOp.cpp
@@ -28,7 +28,8 @@ void ComputeSimplify(FDynamicMesh3* TargetMesh, const bool bReproject,
 					 EEdgeRefineFlags MaterialBoundaryConstraint,
 					 bool bPreserveSharpEdges, bool bAllowSeamCollapse,
 					 const ESimplifyTargetType TargetMode,
-					 const float TargetPercentage, const int TargetCount, const float TargetEdgeLength)
+					 const float TargetPercentage, const int TargetCount, const float TargetEdgeLength,
+					 const float AngleThreshold)
 {
 	SimplificationType Reducer(TargetMesh);
 
@@ -44,16 +45,9 @@ void ComputeSimplify(FDynamicMesh3* TargetMesh, const bool bReproject,
 		Reducer.SetEdgeFlipTolerance(1.e-5);
 
 		// eliminate any bowties that might have formed on UV seams.
-		if (FDynamicMeshAttributeSet* Attributes = TargetMesh->Attributes())
+		if (TargetMesh->Attributes())
 		{
-			// @todo parallelize over NumUVLayers?
-			int NumUVLayers = Attributes->NumUVLayers();
-			for (int i = 0; i < NumUVLayers; ++i)
-			{
-				Attributes->GetUVLayer(i)->SplitBowties();
-			}
-			Attributes->PrimaryNormals()->SplitBowties();
-
+			TargetMesh->Attributes()->SplitAllBowties();
 		}
 	}
 
@@ -89,6 +83,10 @@ void ComputeSimplify(FDynamicMesh3* TargetMesh, const bool bReproject,
 	{
 		Reducer.SimplifyToEdgeLength(TargetEdgeLength);
 	}
+	else if (TargetMode == ESimplifyTargetType::MinimalPlanar)
+	{
+		Reducer.SimplifyToMinimalPlanar(AngleThreshold);
+	}
 }
 
 
@@ -110,29 +108,35 @@ void FSimplifyMeshOp::CalculateResult(FProgressCancel* Progress)
 	}
 
 	int OriginalTriCount = OriginalMesh->TriangleCount();
-	if (SimplifierType != ESimplifyType::UE4Standard)
+	if (SimplifierType == ESimplifyType::QEM)
 	{
-		// GeometryProcessing-specific methods that operate directly on a DynamicMesh3
 		ResultMesh->Copy(*OriginalMesh, true, true, true, !bDiscardAttributes);
-
-		if (SimplifierType == ESimplifyType::QEM)
-		{
-			ComputeSimplify<FQEMSimplification>(TargetMesh, bReproject, OriginalTriCount, *OriginalMesh, *OriginalMeshSpatial,
-												MeshBoundaryConstraint,
-												GroupBoundaryConstraint,
-												MaterialBoundaryConstraint,
-												bPreserveSharpEdges, bAllowSeamCollapse,
-												TargetMode, TargetPercentage, TargetCount, TargetEdgeLength);
-		}
-		else if (SimplifierType == ESimplifyType::Attribute)
-		{
-			ComputeSimplify<FAttrMeshSimplification>(TargetMesh, bReproject, OriginalTriCount, *OriginalMesh, *OriginalMeshSpatial,
-													 MeshBoundaryConstraint,
-													 GroupBoundaryConstraint,
-													 MaterialBoundaryConstraint,
-													 bPreserveSharpEdges, bAllowSeamCollapse,
-													 TargetMode, TargetPercentage, TargetCount, TargetEdgeLength);
-		}
+		ComputeSimplify<FQEMSimplification>(TargetMesh, bReproject, OriginalTriCount, *OriginalMesh, *OriginalMeshSpatial,
+											MeshBoundaryConstraint,
+											GroupBoundaryConstraint,
+											MaterialBoundaryConstraint,
+											bPreserveSharpEdges, bAllowSeamCollapse,
+											TargetMode, TargetPercentage, TargetCount, TargetEdgeLength, MinimalPlanarAngleThresh);
+	}
+	else if (SimplifierType == ESimplifyType::Attribute)
+	{
+		ResultMesh->Copy(*OriginalMesh, true, true, true, !bDiscardAttributes);
+		ComputeSimplify<FAttrMeshSimplification>(TargetMesh, bReproject, OriginalTriCount, *OriginalMesh, *OriginalMeshSpatial,
+													MeshBoundaryConstraint,
+													GroupBoundaryConstraint,
+													MaterialBoundaryConstraint,
+													bPreserveSharpEdges, bAllowSeamCollapse,
+													TargetMode, TargetPercentage, TargetCount, TargetEdgeLength, MinimalPlanarAngleThresh);
+	}
+	else if (SimplifierType == ESimplifyType::MinimalPlanar)
+	{
+		ResultMesh->Copy(*OriginalMesh, true, true, true, !bDiscardAttributes);
+		ComputeSimplify<FAttrMeshSimplification>(TargetMesh, bReproject, OriginalTriCount, *OriginalMesh, *OriginalMeshSpatial,
+			MeshBoundaryConstraint,
+			GroupBoundaryConstraint,
+			MaterialBoundaryConstraint,
+			bPreserveSharpEdges, bAllowSeamCollapse,
+			ESimplifyTargetType::MinimalPlanar, TargetPercentage, TargetCount, TargetEdgeLength, MinimalPlanarAngleThresh);
 	}
 	else // SimplifierType == ESimplifyType::UE4Standard
 	{
diff --git a/Engine/Plugins/Experimental/MeshModelingToolset/Source/ModelingOperatorsEditorOnly/Public/CleaningOps/SimplifyMeshOp.h b/Engine/Plugins/Experimental/MeshModelingToolset/Source/ModelingOperatorsEditorOnly/Public/CleaningOps/SimplifyMeshOp.h
--- a/Engine/Plugins/Experimental/MeshModelingToolset/Source/ModelingOperatorsEditorOnly/Public/CleaningOps/SimplifyMeshOp.h
+++ b/Engine/Plugins/Experimental/MeshModelingToolset/Source/ModelingOperatorsEditorOnly/Public/CleaningOps/SimplifyMeshOp.h
@@ -29,7 +29,10 @@ enum class ESimplifyTargetType : uint8
 	VertexCount = 2 UMETA(DisplayName = "Vertex Count"),
 
 	/** Target edge length */
-	EdgeLength = 3 UMETA(DisplayName = "Edge Length")
+	EdgeLength = 3 UMETA(DisplayName = "Edge Length"),
+
+	/** Apply all allowable edge collapses that do not change the shape */
+	MinimalPlanar = 4 UMETA(Hidden)
 };
 
 UENUM()
@@ -43,6 +46,9 @@ enum class ESimplifyType : uint8
 
 	/** Highest quality reduction.  Will simplify UV boundaries. */
 	UE4Standard = 2 UMETA(DisplayName = "UE4 Standard"),
+
+	/** Collapse any spurious edges but do not change the 3D shape */
+	MinimalPlanar = 3 UMETA(DisplayName = "Minimal Shape-Preserving"),
 };
 
 class MODELINGOPERATORSEDITORONLY_API FSimplifyMeshOp : public FDynamicMeshOperator
@@ -59,6 +65,8 @@ public:
 	float TargetEdgeLength;
 	bool bDiscardAttributes, bReproject, bPreventNormalFlips, bPreserveSharpEdges, bAllowSeamCollapse;
 	EEdgeRefineFlags MeshBoundaryConstraint, GroupBoundaryConstraint, MaterialBoundaryConstraint;
+	/** Angle threshold in degrees used for testing if two triangles should be considered coplanar, or two lines collinear */
+	float MinimalPlanarAngleThresh = 0.01f;
 
 	// stored for the UE4 Standard path
 	TSharedPtr<FMeshDescription> OriginalMeshDescription;
