MeshModelingTools: Updated displace mesh tool. Made it work with async texture loading, added parameters for UV offset and scale and for an offset base value, updated tooltips, made non-square textures be handled properly, and fixed bug with reverting vector parameters.

#rnx

diff --git a/Engine/Plugins/Experimental/MeshModelingToolset/Source/MeshModelingTools/Private/DisplaceMeshTool.cpp b/Engine/Plugins/Experimental/MeshModelingToolset/Source/MeshModelingTools/Private/DisplaceMeshTool.cpp
--- a/Engine/Plugins/Experimental/MeshModelingToolset/Source/MeshModelingTools/Private/DisplaceMeshTool.cpp
+++ b/Engine/Plugins/Experimental/MeshModelingToolset/Source/MeshModelingTools/Private/DisplaceMeshTool.cpp
@@ -1,6 +1,8 @@
 // Copyright Epic Games, Inc. All Rights Reserved.
 
 #include "DisplaceMeshTool.h"
+
+#include "AssetUtils/Texture2DUtil.h"
 #include "InteractiveToolManager.h"
 #include "ToolBuilderUtil.h"
 #include "ToolSetupUtil.h"
@@ -13,7 +15,7 @@
 #include "MeshDescription.h"
 #define LOCTEXT_NAMESPACE "UDisplaceMeshTool"
 
-namespace {
+namespace DisplaceMeshToolLocals{
 
 	void SubdivideMesh(FDynamicMesh3& Mesh)
 	{
@@ -170,9 +172,18 @@ namespace {
 			const FMeshNormals& Normals,
 			TFunctionRef<float(const FVector3d&, const FVector3d&)> IntensityFunc,
 			const FSampledScalarField2f& DisplaceField,
-			TArray<FVector3d>& DisplacedPositions)
+			TArray<FVector3d>& DisplacedPositions,
+			float DisplaceFieldBaseValue = 128.0/255, // value that corresponds to zero displacement
+			FVector2f UVScale = FVector2f(1, 1),
+			FVector2f UVOffset = FVector2f(0,0))
 		{
 			const FDynamicMeshUVOverlay* UVOverlay = Mesh.Attributes()->GetUVLayer(0);
+
+			// We set things up such that DisplaceField goes from 0 to 1 in the U direction,
+			// but the V direction may be shorter or longer if the texture is not square
+			// (it will be 1/AspectRatio)
+			float VHeight = DisplaceField.Height() * DisplaceField.CellDimensions.Y;
+
 			for (int tid : Mesh.TriangleIndicesItr())
 			{
 				FIndex3i Tri = Mesh.GetTriangle(tid);
@@ -181,7 +192,18 @@ namespace {
 				{
 					int vid = Tri[j];
 					FVector2f UV = UVOverlay->GetElement(UVTri[j]);
-					double Offset = DisplaceField.BilinearSampleClamped(UV);
+
+					// Adjust UV value and tile it. 
+					// Note that we're effectively stretching the texture to be square before tiling, since this
+					// seems to be what non square textures do by default in UE. If we decide to tile without 
+					// stretching by default someday, we'd do UV - FVector2f(FMath::Floor(UV.X), FMath:Floor(UV.Y/VHeight)*VHeight)
+					// without multiplying by VHeight afterward.
+					UV = UV * UVScale + UVOffset;
+					UV = UV - FVector2f(FMath::Floor(UV.X), FMath::Floor(UV.Y));
+					UV.Y *= VHeight;
+
+					double Offset = DisplaceField.BilinearSampleClamped(UV) - DisplaceFieldBaseValue;
+
 					double Intensity = IntensityFunc(Positions[vid], Normals[vid]);
 					DisplacedPositions[vid] = Positions[vid] + (Offset * Intensity * Normals[vid]);
 				}
@@ -213,62 +235,6 @@ namespace {
 
 	}
 
-	class FTextureAccess
-	{
-	public:
-		FTextureAccess(UTexture2D* DisplacementMap)
-			:DisplacementMap(DisplacementMap)
-		{
-			check(DisplacementMap);
-			OldCompressionSettings = DisplacementMap->CompressionSettings;
-			bOldSRGB = DisplacementMap->SRGB;
-#if WITH_EDITOR
-			OldMipGenSettings = DisplacementMap->MipGenSettings;
-#endif
-			DisplacementMap->CompressionSettings = TextureCompressionSettings::TC_VectorDisplacementmap;
-			DisplacementMap->SRGB = false;
-#if WITH_EDITOR
-			DisplacementMap->MipGenSettings = TextureMipGenSettings::TMGS_NoMipmaps;
-#endif
-			DisplacementMap->UpdateResource();
-
-			FormattedImageData = reinterpret_cast<const FColor*>(DisplacementMap->PlatformData->Mips[0].BulkData.LockReadOnly());
-		}
-		FTextureAccess(const FTextureAccess&) = delete;
-		FTextureAccess(FTextureAccess&&) = delete;
-		void operator=(const FTextureAccess&) = delete;
-		void operator=(FTextureAccess&&) = delete;
-
-		~FTextureAccess()
-		{
-			DisplacementMap->PlatformData->Mips[0].BulkData.Unlock();
-
-			DisplacementMap->CompressionSettings = OldCompressionSettings;
-			DisplacementMap->SRGB = bOldSRGB;
-#if WITH_EDITOR
-			DisplacementMap->MipGenSettings = OldMipGenSettings;
-#endif
-
-			DisplacementMap->UpdateResource();
-		}
-
-		bool HasData() const
-		{
-			return FormattedImageData != nullptr;
-		}
-		const FColor* GetData() const
-		{
-			return FormattedImageData;
-		}
-
-	private:
-		UTexture2D* DisplacementMap{ nullptr };
-		TextureCompressionSettings OldCompressionSettings{};
-		TextureMipGenSettings OldMipGenSettings{};
-		bool bOldSRGB{ false };
-		const FColor* FormattedImageData{ nullptr };
-	};
-
 	class FSubdivideMeshOp : public FDynamicMeshOperator
 	{
 	public:
@@ -286,6 +252,8 @@ namespace {
 
 	void FSubdivideMeshOp::CalculateResult(FProgressCancel* ProgressCancel)
 	{
+		using namespace DisplaceMeshToolLocals;
+
 		// calculate subdivisions (todo: move to elsewhere)
 		for (int ri = 0; ri < SubdivisionsCount; ri++)
 		{
@@ -340,6 +308,11 @@ namespace {
 		FSampledScalarField2f DisplaceField;
 		TArray<FPerlinLayerProperties> PerlinLayerProperties;
 
+		// Used in texture map displacement
+		float DisplacementMapBaseValue = 128.0/255; // i.e., what constitutes no displacement
+		FVector2f UVScale = FVector2f(1,1);
+		FVector2f UVOffset = FVector2f(0, 0);
+
 		TSharedPtr<FIndexedWeightMap> WeightMap;
 		TFunction<float(const FVector3d&, const FIndexedWeightMap)> WeightMapQueryFunc;
 	};
@@ -376,6 +349,12 @@ namespace {
 		ResultMesh->Copy(*SourceMesh);
 
 		if (Progress->Cancelled()) return;
+
+		if (DisplacementType == EDisplaceMeshToolDisplaceType::DisplacementMap && !Parameters.DisplacementMap)
+		{
+			return;
+		}
+
 		SourceNormals = FMeshNormals(SourceMesh.Get());
 		SourceNormals.ComputeVertexNormals();
 
@@ -445,7 +424,10 @@ namespace {
 				SourceNormals,
 				IntensityFunc,
 				Parameters.DisplaceField, 
-				DisplacedPositions);
+				DisplacedPositions,
+				Parameters.DisplacementMapBaseValue,
+				Parameters.UVScale,
+				Parameters.UVOffset);
 			break;
 
 		case EDisplaceMeshToolDisplaceType::SineWave:
@@ -502,10 +484,16 @@ namespace {
 
 			Parameters.WeightMap = DisplaceParametersIn.WeightMap;
 			Parameters.WeightMapQueryFunc = DisplaceParametersIn.WeightMapQueryFunc;
+
+			Parameters.DisplacementMapBaseValue = DisplaceParametersIn.DisplacementMapBaseValue;
+			Parameters.UVScale = DisplaceParametersIn.UVScale;
+			Parameters.UVOffset = DisplaceParametersIn.UVOffset;
 		}
 		void SetIntensity(float IntensityIn);
 		void SetRandomSeed(int RandomSeedIn);
 		void SetDisplacementMap(UTexture2D* DisplacementMapIn);
+		void SetDisplacementMapUVAdjustment(const FVector2f& UVScale, const FVector2f& UVOffset);
+		void SetDisplacementMapBaseValue(float DisplacementMapBaseValue);
 		void SetFrequency(float FrequencyIn);
 		void SetPhaseShift(float PhaseShiftIn);
 		void SetSineWaveDirection(const FVector& Direction);
@@ -548,6 +536,19 @@ namespace {
 		}
 	}
 
+	void FDisplaceMeshOpFactory::SetDisplacementMapUVAdjustment(const FVector2f& UVScale, const FVector2f& UVOffset)
+	{
+		Parameters.UVScale = UVScale;
+		Parameters.UVOffset = UVOffset;
+	}
+
+	void FDisplaceMeshOpFactory::SetDisplacementMapBaseValue(float DisplacementMapBaseValue)
+	{
+		// We could bake this into the displacement field, but that would require calling UpdateMap with
+		// every slider change, which is slow. So we'll just pass this down to the calculation.
+		Parameters.DisplacementMapBaseValue = DisplacementMapBaseValue;
+	}
+
 	void FDisplaceMeshOpFactory::SetFrequency(float FrequencyIn)
 	{
 		Parameters.SineWaveFrequency = FrequencyIn;
@@ -560,7 +561,7 @@ namespace {
 
 	void FDisplaceMeshOpFactory::SetSineWaveDirection(const FVector& Direction)
 	{
-		Parameters.SineWaveDirection = Direction;
+		Parameters.SineWaveDirection = Direction.GetSafeNormal();
 	}
 
 	void FDisplaceMeshOpFactory::SetDisplacementType(EDisplaceMeshToolDisplaceType TypeIn)
@@ -575,29 +576,31 @@ namespace {
 			Parameters.DisplacementMap->PlatformData->Mips.Num() < 1)
 		{
 			Parameters.DisplaceField = FSampledScalarField2f();
+			Parameters.DisplaceField.GridValues.AssignAll(0);
 			return;
 		}
 
-		FTextureAccess TextureAccess(Parameters.DisplacementMap);
-		if (!TextureAccess.HasData())
+		TImageBuilder<FVector4f> DisplacementMapValues;
+		FImageDimensions DisplacementMapDimensions;
+		if (!UE::AssetUtils::ReadTexture(Parameters.DisplacementMap, DisplacementMapDimensions, DisplacementMapValues))
 		{
 			Parameters.DisplaceField = FSampledScalarField2f();
+			Parameters.DisplaceField.GridValues.AssignAll(0);
 		}
 		else
 		{
-			int64 TextureWidth = Parameters.DisplacementMap->GetSizeX();
-			int64 TextureHeight = Parameters.DisplacementMap->GetSizeY();
+			int64 TextureWidth = DisplacementMapDimensions.GetWidth();
+			int64 TextureHeight = DisplacementMapDimensions.GetHeight();
 			Parameters.DisplaceField.Resize(TextureWidth, TextureHeight, 0.0f);
+
+			// Note that the height of the texture will not be 1.0 if it was not square. This should be kept in mind when sampling it later.
 			Parameters.DisplaceField.SetCellSize(1.0f / (float)TextureWidth);
 
-			const FColor* FormattedData = TextureAccess.GetData();
 			for (int64 y = 0; y < TextureHeight; ++y)
 			{
 				for (int64 x = 0; x < TextureWidth; ++x)
 				{
-					FColor PixelColor = FormattedData[y * TextureWidth + x];
-					float Value = PixelColor.R / 255.0;
-					Parameters.DisplaceField.GridValues[y * TextureWidth + x] = Value;
+					Parameters.DisplaceField.GridValues[y * TextureWidth + x] = DisplacementMapValues.GetPixel(y * TextureWidth + x).X;
 				}
 			}
 		}
@@ -610,7 +613,7 @@ namespace {
 
 	void FDisplaceMeshOpFactory::SetFilterDirection(const FVector& Direction)
 	{
-		Parameters.FilterDirection = Direction;
+		Parameters.FilterDirection = Direction.GetSafeNormal();
 	}
 
 	void FDisplaceMeshOpFactory::SetFilterFalloffWidth(float FalloffWidth)
@@ -661,6 +664,8 @@ TArray<FString> UDisplaceMeshCommonProperties::GetWeightMapsFunc()
 
 void UDisplaceMeshTool::Setup()
 {
+	using namespace DisplaceMeshToolLocals;
+
 	UInteractiveTool::Setup();
 
 	// UInteractiveToolPropertySets
@@ -722,9 +727,9 @@ void UDisplaceMeshTool::Setup()
 	Parameters.DisplacementMap = TextureMapProperties->DisplacementMap;
 	Parameters.SineWaveFrequency = SineWaveProperties->SineWaveFrequency;
 	Parameters.SineWavePhaseShift = SineWaveProperties->SineWavePhaseShift;
-	Parameters.SineWaveDirection = SineWaveProperties->SineWaveDirection;
+	Parameters.SineWaveDirection = SineWaveProperties->SineWaveDirection.GetSafeNormal();
 	Parameters.bEnableFilter = DirectionalFilterProperties->bEnableFilter;
-	Parameters.FilterDirection = DirectionalFilterProperties->FilterDirection;
+	Parameters.FilterDirection = DirectionalFilterProperties->FilterDirection.GetSafeNormal();
 	Parameters.FilterWidth = DirectionalFilterProperties->FilterWidth;
 	Parameters.PerlinLayerProperties = NoiseProperties->PerlinLayerProperties;
 	Parameters.WeightMap = ActiveWeightMap;
@@ -812,7 +817,7 @@ void UDisplaceMeshTool::ValidateSubdivisions()
 	int MaxSubdivisions = (int)floor(log2(MaxTriangles / NumTriangles) / 2.0);
 	if (CommonProperties->Subdivisions > MaxSubdivisions)
 	{
-		FText WarningText = FText::Format(LOCTEXT("SubdivisionsTooHigh", "Desired number of Subdivisions ({0}) exceeds maximum number of {1}"), 
+		FText WarningText = FText::Format(LOCTEXT("SubdivisionsTooHigh", "Desired number of Subdivisions ({0}) exceeds maximum number ({1}) for a mesh of this number of triangles."), 
 			FText::AsNumber(CommonProperties->Subdivisions), 
 			FText::AsNumber(MaxSubdivisions));
 		GetToolManager()->DisplayMessage(WarningText, EToolMessageLevel::UserWarning);
@@ -833,6 +838,8 @@ void UDisplaceMeshTool::ValidateSubdivisions()
 
 void UDisplaceMeshTool::OnPropertyModified(UObject* PropertySet, FProperty* Property)
 {
+	using namespace DisplaceMeshToolLocals;
+
 	if (PropertySet && Property)
 	{
 		FDisplaceMeshOpFactory* DisplacerDownCast = static_cast<FDisplaceMeshOpFactory*>(Displacer.Get());
@@ -880,6 +887,10 @@ void UDisplaceMeshTool::OnPropertyModified(UObject* PropertySet, FProperty* Prop
 		{
 			DisplacerDownCast->SetDisplacementMap(TextureMapProperties->DisplacementMap);
 		}
+		else if (PropName == GET_MEMBER_NAME_CHECKED(UDisplaceMeshTextureMapProperties, DisplacementMapBaseValue))
+		{
+			DisplacerDownCast->SetDisplacementMapBaseValue(TextureMapProperties->DisplacementMapBaseValue);
+		}
 		else if (PropName == GET_MEMBER_NAME_CHECKED(UDisplaceMeshCommonProperties, WeightMap) 
 				 || PropName == GET_MEMBER_NAME_CHECKED(UDisplaceMeshCommonProperties, bInvertWeightMap))
 		{
@@ -898,21 +909,20 @@ void UDisplaceMeshTool::OnPropertyModified(UObject* PropertySet, FProperty* Prop
 		{
 			DisplacerDownCast->SetPerlinNoiseLayerProperties(NoiseProperties->PerlinLayerProperties);
 		}
-		else if (PropName == "X" || PropName == "Y" || PropName == "Z")
+		// The FName we get for the individual vector elements are all the same, whereas resetting with the "revert
+		// to default" arrow gets us the name of the vector itself. We'll just update all of them if any of them
+		// change.
+		else if (PropName == "X" || PropName == "Y" || PropName == "Z"
+			|| PropName == GET_MEMBER_NAME_CHECKED(UDisplaceMeshDirectionalFilterProperties, FilterDirection)
+			|| PropName == GET_MEMBER_NAME_CHECKED(UDisplaceMeshSineWaveProperties, SineWaveDirection)
+			|| PropName == GET_MEMBER_NAME_CHECKED(UDisplaceMeshTextureMapProperties, UVScale)
+			|| PropName == GET_MEMBER_NAME_CHECKED(UDisplaceMeshTextureMapProperties, UVOffset))
 		{
-			if (PropertySetName == "DisplaceMeshDirectionalFilterProperties")
-			{
-				DirectionalFilterProperties->FilterDirection.Normalize();
-				DisplacerDownCast->SetFilterDirection(DirectionalFilterProperties->FilterDirection);
-			}
-			else if (PropertySetName == "DisplaceMeshSineWaveProperties")
-			{
-				SineWaveProperties->SineWaveDirection.Normalize();
-				DisplacerDownCast->SetSineWaveDirection(SineWaveProperties->SineWaveDirection);
-			}
+			DisplacerDownCast->SetFilterDirection(DirectionalFilterProperties->FilterDirection);
+			DisplacerDownCast->SetSineWaveDirection(SineWaveProperties->SineWaveDirection);
+			DisplacerDownCast->SetDisplacementMapUVAdjustment(TextureMapProperties->UVScale, TextureMapProperties->UVOffset);
 		}
 
-
 		StartComputation();
 	}
 }
diff --git a/Engine/Plugins/Experimental/MeshModelingToolset/Source/MeshModelingTools/Private/Tests/DisplaceMeshTool_Tests.inl b/Engine/Plugins/Experimental/MeshModelingToolset/Source/MeshModelingTools/Private/Tests/DisplaceMeshTool_Tests.inl
--- a/Engine/Plugins/Experimental/MeshModelingToolset/Source/MeshModelingTools/Private/Tests/DisplaceMeshTool_Tests.inl
+++ b/Engine/Plugins/Experimental/MeshModelingToolset/Source/MeshModelingTools/Private/Tests/DisplaceMeshTool_Tests.inl
@@ -9,8 +9,10 @@
 #include "Generators/GridBoxMeshGenerator.h"
 
 // Test helpers
-namespace
+namespace DisplaceMeshToolTestsLocals
 {
+	using namespace DisplaceMeshToolLocals;
+
 	TSharedPtr<FDynamicMesh3> CreateSingleTriangleMesh()
 	{
 		TSharedPtr<FDynamicMesh3> Mesh = MakeShareable<FDynamicMesh3>(new FDynamicMesh3);
@@ -43,6 +45,8 @@ IMPLEMENT_SIMPLE_AUTOMATION_TEST(UDisplaceMeshTestConstant,
 
 	bool UDisplaceMeshTestConstant::RunTest(const FString& Parameters)
 {
+	using namespace DisplaceMeshToolTestsLocals;
+
 	// ---- Setup
 	TSharedPtr<FDynamicMesh3> InputMesh = CreateSingleTriangleMesh();
 
@@ -88,6 +92,8 @@ IMPLEMENT_SIMPLE_AUTOMATION_TEST(UDisplaceMeshTestRandom,
 
 	bool UDisplaceMeshTestRandom::RunTest(const FString& Parameters)
 {
+	using namespace DisplaceMeshToolTestsLocals;
+
 	// ---- Setup
 	TSharedPtr<FDynamicMesh3> InputMesh = CreateSingleTriangleMesh();
 
@@ -122,6 +128,8 @@ IMPLEMENT_SIMPLE_AUTOMATION_TEST(UDisplaceMeshTestPerlin,
 
 	bool UDisplaceMeshTestPerlin::RunTest(const FString& Parameters)
 {
+	using namespace DisplaceMeshToolTestsLocals;
+
 	// ---- Setup
 	TSharedPtr<FDynamicMesh3> InputMesh = CreateSingleTriangleMesh();
 
@@ -157,6 +165,8 @@ IMPLEMENT_SIMPLE_AUTOMATION_TEST(UDisplaceMeshTestSine,
 
 	bool UDisplaceMeshTestSine::RunTest(const FString& Parameters)
 {
+	using namespace DisplaceMeshToolTestsLocals;
+
 	// ---- Setup
 	TSharedPtr<FDynamicMesh3> InputMesh = CreateSingleTriangleMesh();
 
@@ -200,6 +210,8 @@ IMPLEMENT_SIMPLE_AUTOMATION_TEST(UDisplaceMeshTestSineDirection,
 
 	bool UDisplaceMeshTestSineDirection::RunTest(const FString& Parameters)
 {
+	using namespace DisplaceMeshToolTestsLocals;
+
 	// ---- Setup
 	TSharedPtr<FDynamicMesh3> InputMesh = CreateSingleTriangleMesh();
 
@@ -250,6 +262,8 @@ IMPLEMENT_SIMPLE_AUTOMATION_TEST(UDisplaceMeshTestDirectionalFilter,
 
 	bool UDisplaceMeshTestDirectionalFilter::RunTest(const FString& Parameters)
 {
+	using namespace DisplaceMeshToolTestsLocals;
+
 	// ---- Setup
 	// Create a box mesh, random noise, and a directional filter pointing in the +Y direction
 	TSharedPtr<FDynamicMesh3> InputMesh = CreateTestBoxMesh();
diff --git a/Engine/Plugins/Experimental/MeshModelingToolset/Source/MeshModelingTools/Public/DisplaceMeshTool.h b/Engine/Plugins/Experimental/MeshModelingToolset/Source/MeshModelingTools/Public/DisplaceMeshTool.h
--- a/Engine/Plugins/Experimental/MeshModelingToolset/Source/MeshModelingTools/Public/DisplaceMeshTool.h
+++ b/Engine/Plugins/Experimental/MeshModelingToolset/Source/MeshModelingTools/Public/DisplaceMeshTool.h
@@ -20,10 +20,13 @@ class USimpleDynamicMeshComponent;
 UENUM()
 enum class EDisplaceMeshToolDisplaceType : uint8
 {
-	/** Displace with N iterations */
+	/** Offset a set distance in the normal direction. */
 	Constant UMETA(DisplayName = "Constant"),
 
-	/** Displace with N iterations */
+	/** Offset in the normal direction using the first channel of a 2D texture. */
+	DisplacementMap UMETA(DisplayName = "Texture2D Map"),
+
+	/** Offset vertices randomly. */
 	RandomNoise UMETA(DisplayName = "Random Noise"),
 
 	/** Offset in the normal direction weighted by Perlin noise. 
@@ -34,9 +37,6 @@ enum class EDisplaceMeshToolDisplaceType : uint8
 	*/
 	PerlinNoise UMETA(DisplayName = "Perlin Noise"),
 
-	/** Displace with N iterations */
-	DisplacementMap UMETA(DisplayName = "Texture2D Map"),
-
 	/** Move vertices in spatial sine wave pattern */
 	SineWave UMETA(DisplayName = "Sine Wave"),
 };
@@ -62,7 +62,7 @@ public:
 		meta = (EditCondition = "DisplacementType == EDisplaceMeshToolDisplaceType::RandomNoise || DisplacementType == EDisplaceMeshToolDisplaceType::PerlinNoise"))
 	int RandomSeed = 31337;
 
-	/** Subdivision iterations for mesh */
+	/** Number of times to subdivide the mesh before displacing it. */
 	UPROPERTY(EditAnywhere, Category = Options,
 		meta = (UIMin = "0", UIMax = "10", ClampMin = "0", ClampMax = "100"))
 	int Subdivisions = 4;
@@ -93,9 +93,26 @@ class MESHMODELINGTOOLS_API UDisplaceMeshTextureMapProperties : public UInteract
 	GENERATED_BODY()
 			
 public:
-	/** Displacement map */
+	/** Displacement map. Only the first channel is used. */
 	UPROPERTY(EditAnywhere, Category = Options)
 	UTexture2D* DisplacementMap = nullptr;
+
+	/** The value in the texture map that corresponds to no displacement. For instance, if set to 0, then all
+	 displacement will be positive. If set to 0.5, displacement below 0.5 will be negative, and above will be
+	 positive. Default is for 128/255 to be no displacement. */
+	UPROPERTY(EditAnywhere, Category = Options, meta = (UIMin = "0", UIMax = "1", ClampMin = "0", ClampMax = "1"))
+	float DisplacementMapBaseValue = 128.0/255;
+
+	/** When sampling from the texture map, how to scale the mesh UV's in the x and y directions. For a mesh with
+	UV's in the range 0 to 1, setting a scale above 1 will result in tiling the texture map, and scaling below
+	1 will result in using only part of the texture map.*/
+	UPROPERTY(EditAnywhere, Category = Options)
+	FVector2D UVScale = FVector2D(1,1);
+
+	/** When sampling from the texture map, how to offset the mesh UV's. This will result in offsetting the
+	tiling of the texture map across the mesh. */
+	UPROPERTY(EditAnywhere, Category = Options, meta = (UIMin = "-1", UIMax = "1"))
+	FVector2D UVOffset = FVector2D(0, 0);
 };
 
 
