MeshModelingTools: Added contrast adjustment curve to displacement tool (for displacement maps). Fixed the displacement texture values always being interpreted as gamma compressed. Added option to show wireframe.

#rnx

diff --git a/Engine/Plugins/Experimental/MeshModelingToolset/Source/MeshModelingTools/Private/DisplaceMeshTool.cpp b/Engine/Plugins/Experimental/MeshModelingToolset/Source/MeshModelingTools/Private/DisplaceMeshTool.cpp
--- a/Engine/Plugins/Experimental/MeshModelingToolset/Source/MeshModelingTools/Private/DisplaceMeshTool.cpp
+++ b/Engine/Plugins/Experimental/MeshModelingToolset/Source/MeshModelingTools/Private/DisplaceMeshTool.cpp
@@ -3,6 +3,8 @@
 #include "DisplaceMeshTool.h"
 
 #include "AssetUtils/Texture2DUtil.h"
+#include "Curves/CurveFloat.h"
+#include "Curves/RichCurve.h"
 #include "InteractiveToolManager.h"
 #include "ToolBuilderUtil.h"
 #include "ToolSetupUtil.h"
@@ -175,7 +177,8 @@ namespace DisplaceMeshToolLocals{
 			TArray<FVector3d>& DisplacedPositions,
 			float DisplaceFieldBaseValue = 128.0/255, // value that corresponds to zero displacement
 			FVector2f UVScale = FVector2f(1, 1),
-			FVector2f UVOffset = FVector2f(0,0))
+			FVector2f UVOffset = FVector2f(0,0),
+			FRichCurve* AdjustmentCurve = nullptr)
 		{
 			const FDynamicMeshUVOverlay* UVOverlay = Mesh.Attributes()->GetUVLayer(0);
 
@@ -202,7 +205,12 @@ namespace DisplaceMeshToolLocals{
 					UV = UV - FVector2f(FMath::Floor(UV.X), FMath::Floor(UV.Y));
 					UV.Y *= VHeight;
 
-					double Offset = DisplaceField.BilinearSampleClamped(UV) - DisplaceFieldBaseValue;
+					double Offset = DisplaceField.BilinearSampleClamped(UV);
+					if (AdjustmentCurve)
+					{
+						Offset = AdjustmentCurve->Eval(Offset);
+					}
+					Offset -= DisplaceFieldBaseValue;
 
 					double Intensity = IntensityFunc(Positions[vid], Normals[vid]);
 					DisplacedPositions[vid] = Positions[vid] + (Offset * Intensity * Normals[vid]);
@@ -312,6 +320,9 @@ namespace DisplaceMeshToolLocals{
 		float DisplacementMapBaseValue = 128.0/255; // i.e., what constitutes no displacement
 		FVector2f UVScale = FVector2f(1,1);
 		FVector2f UVOffset = FVector2f(0, 0);
+		// This gets used by worker threads, so do not try to change an existing curve- make
+		// a new one each time.
+		TSharedPtr<FRichCurve> AdjustmentCurve;
 
 		TSharedPtr<FIndexedWeightMap> WeightMap;
 		TFunction<float(const FVector3d&, const FIndexedWeightMap)> WeightMapQueryFunc;
@@ -427,7 +438,8 @@ namespace DisplaceMeshToolLocals{
 				DisplacedPositions,
 				Parameters.DisplacementMapBaseValue,
 				Parameters.UVScale,
-				Parameters.UVOffset);
+				Parameters.UVOffset,
+				Parameters.AdjustmentCurve.Get());
 			break;
 
 		case EDisplaceMeshToolDisplaceType::SineWave:
@@ -472,7 +484,7 @@ namespace DisplaceMeshToolLocals{
 		{
 			SetIntensity(DisplaceParametersIn.DisplaceIntensity);
 			SetRandomSeed(DisplaceParametersIn.RandomSeed);
-			SetDisplacementMap(DisplaceParametersIn.DisplacementMap);
+			SetDisplacementMap(DisplaceParametersIn.DisplacementMap); // Calls UpdateMap
 			SetFrequency(DisplaceParametersIn.SineWaveFrequency);
 			SetPhaseShift(DisplaceParametersIn.SineWavePhaseShift);
 			SetSineWaveDirection(DisplaceParametersIn.SineWaveDirection);
@@ -488,12 +500,15 @@ namespace DisplaceMeshToolLocals{
 			Parameters.DisplacementMapBaseValue = DisplaceParametersIn.DisplacementMapBaseValue;
 			Parameters.UVScale = DisplaceParametersIn.UVScale;
 			Parameters.UVOffset = DisplaceParametersIn.UVOffset;
+
+			Parameters.AdjustmentCurve = DisplaceParametersIn.AdjustmentCurve;
 		}
 		void SetIntensity(float IntensityIn);
 		void SetRandomSeed(int RandomSeedIn);
 		void SetDisplacementMap(UTexture2D* DisplacementMapIn);
 		void SetDisplacementMapUVAdjustment(const FVector2f& UVScale, const FVector2f& UVOffset);
 		void SetDisplacementMapBaseValue(float DisplacementMapBaseValue);
+		void SetAdjustmentCurve(UCurveFloat* CurveFloat);
 		void SetFrequency(float FrequencyIn);
 		void SetPhaseShift(float PhaseShiftIn);
 		void SetSineWaveDirection(const FVector& Direction);
@@ -529,11 +544,11 @@ namespace DisplaceMeshToolLocals{
 
 	void FDisplaceMeshOpFactory::SetDisplacementMap(UTexture2D* DisplacementMapIn)
 	{
-		if (Parameters.DisplacementMap != DisplacementMapIn)
-		{
-			Parameters.DisplacementMap = DisplacementMapIn;
-			UpdateMap();
-		}
+		Parameters.DisplacementMap = DisplacementMapIn;
+
+		// Note that we do the update even if we got the same pointer, because the texture
+		// may have been changed in the editor.
+		UpdateMap();
 	}
 
 	void FDisplaceMeshOpFactory::SetDisplacementMapUVAdjustment(const FVector2f& UVScale, const FVector2f& UVOffset)
@@ -549,6 +564,13 @@ namespace DisplaceMeshToolLocals{
 		Parameters.DisplacementMapBaseValue = DisplacementMapBaseValue;
 	}
 
+	void FDisplaceMeshOpFactory::SetAdjustmentCurve(UCurveFloat* CurveFloat)
+	{
+		Parameters.AdjustmentCurve = CurveFloat ? TSharedPtr<FRichCurve>(
+			static_cast<FRichCurve*>(CurveFloat->FloatCurve.Duplicate()))
+			: nullptr;
+	}
+
 	void FDisplaceMeshOpFactory::SetFrequency(float FrequencyIn)
 	{
 		Parameters.SineWaveFrequency = FrequencyIn;
@@ -582,7 +604,9 @@ namespace DisplaceMeshToolLocals{
 
 		TImageBuilder<FVector4f> DisplacementMapValues;
 		FImageDimensions DisplacementMapDimensions;
-		if (!UE::AssetUtils::ReadTexture(Parameters.DisplacementMap, DisplacementMapDimensions, DisplacementMapValues))
+		if (!UE::AssetUtils::ReadTexture(Parameters.DisplacementMap, DisplacementMapDimensions, DisplacementMapValues,
+			// need bPreferPlatformData to be true to respond to non-destructive changes to the texture in the editor
+			true)) 
 		{
 			Parameters.DisplaceField = FSampledScalarField2f();
 			Parameters.DisplaceField.GridValues.AssignAll(0);
@@ -684,7 +708,25 @@ void UDisplaceMeshTool::Setup()
 	{
 		TextureMapProperties->DisplacementMap = nullptr;
 	}
+	TextureMapProperties->AdjustmentCurve = ToolSetupUtil::GetContrastAdjustmentCurve(GetToolManager());
 
+	// In editor, we can respond directly to curve updates.
+#if WITH_EDITORONLY_DATA
+	if (TextureMapProperties->AdjustmentCurve)
+	{
+		TextureMapProperties->AdjustmentCurve->OnUpdateCurve.AddWeakLambda(this,
+			[this](UCurveBase* Curve, EPropertyChangeType::Type ChangeType) {
+				if (TextureMapProperties->bApplyAdjustmentCurve)
+				{
+					FDisplaceMeshOpFactory* DisplacerDownCast = static_cast<FDisplaceMeshOpFactory*>(Displacer.Get());
+					DisplacerDownCast->SetAdjustmentCurve(TextureMapProperties->AdjustmentCurve);
+					bNeedsDisplaced = true;
+					StartComputation();
+				}
+			});
+	}
+#endif
+	
 	// populate weight maps list
 	TArray<FName> WeightMaps;
 	UE::WeightMaps::FindVertexWeightMaps(ComponentTarget->GetMesh(), WeightMaps);
@@ -705,6 +747,7 @@ void UDisplaceMeshTool::Setup()
 	DynamicMeshComponent->SetupAttachment(ComponentTarget->GetOwnerActor()->GetRootComponent());
 	DynamicMeshComponent->RegisterComponent();
 	DynamicMeshComponent->SetWorldTransform(ComponentTarget->GetWorldTransform());
+	DynamicMeshComponent->bExplicitShowWireframe = CommonProperties->bShowWireframe;
 
 	// transfer materials
 	FComponentMaterialSet MaterialSet;
@@ -777,6 +820,13 @@ void UDisplaceMeshTool::Setup()
 
 void UDisplaceMeshTool::Shutdown(EToolShutdownType ShutdownType)
 {
+#if WITH_EDITORONLY_DATA
+	if (TextureMapProperties->AdjustmentCurve)
+	{
+		TextureMapProperties->AdjustmentCurve->OnUpdateCurve.RemoveAll(this);
+	}
+#endif
+
 	CommonProperties->SaveProperties(this);
 	NoiseProperties->SaveProperties(this);
 	DirectionalFilterProperties->SaveProperties(this);
@@ -875,6 +925,10 @@ void UDisplaceMeshTool::OnPropertyModified(UObject* PropertySet, FProperty* Prop
 		{
 			DisplacerDownCast->SetIntensity(CommonProperties->DisplaceIntensity);
 		}
+		else if (PropName == GET_MEMBER_NAME_CHECKED(UDisplaceMeshCommonProperties, bShowWireframe))
+		{
+			DynamicMeshComponent->bExplicitShowWireframe = CommonProperties->bShowWireframe;
+		}
 		else if (PropName == GET_MEMBER_NAME_CHECKED(UDisplaceMeshSineWaveProperties, SineWaveFrequency))
 		{
 			DisplacerDownCast->SetFrequency(SineWaveProperties->SineWaveFrequency);
@@ -891,6 +945,11 @@ void UDisplaceMeshTool::OnPropertyModified(UObject* PropertySet, FProperty* Prop
 		{
 			DisplacerDownCast->SetDisplacementMapBaseValue(TextureMapProperties->DisplacementMapBaseValue);
 		}
+		else if (PropName == GET_MEMBER_NAME_CHECKED(UDisplaceMeshTextureMapProperties, bApplyAdjustmentCurve)
+			|| PropName == GET_MEMBER_NAME_CHECKED(UDisplaceMeshTextureMapProperties, AdjustmentCurve))
+		{
+			DisplacerDownCast->SetAdjustmentCurve(TextureMapProperties->bApplyAdjustmentCurve ? TextureMapProperties->AdjustmentCurve : nullptr);
+		}
 		else if (PropName == GET_MEMBER_NAME_CHECKED(UDisplaceMeshCommonProperties, WeightMap) 
 				 || PropName == GET_MEMBER_NAME_CHECKED(UDisplaceMeshCommonProperties, bInvertWeightMap))
 		{
diff --git a/Engine/Plugins/Experimental/MeshModelingToolset/Source/MeshModelingTools/Public/DisplaceMeshTool.h b/Engine/Plugins/Experimental/MeshModelingToolset/Source/MeshModelingTools/Public/DisplaceMeshTool.h
--- a/Engine/Plugins/Experimental/MeshModelingToolset/Source/MeshModelingTools/Public/DisplaceMeshTool.h
+++ b/Engine/Plugins/Experimental/MeshModelingToolset/Source/MeshModelingTools/Public/DisplaceMeshTool.h
@@ -15,6 +15,7 @@
 #include "DisplaceMeshTool.generated.h"
 
 struct FMeshDescription;
+class UCurveFloat;
 class USimpleDynamicMeshComponent;
 
 UENUM()
@@ -80,6 +81,9 @@ public:
 	UPROPERTY(EditAnywhere, Category = Options)
 	bool bInvertWeightMap = false;
 
+	UPROPERTY(EditAnywhere, Category = Options)
+	bool bShowWireframe = false;
+
 	UPROPERTY(EditAnywhere, Category = Options, AdvancedDisplay, meta = (TransientToolProperty) )
 	bool bDisableSizeWarning = false;
 };
@@ -113,6 +117,16 @@ public:
 	tiling of the texture map across the mesh. */
 	UPROPERTY(EditAnywhere, Category = Options, meta = (UIMin = "-1", UIMax = "1"))
 	FVector2D UVOffset = FVector2D(0, 0);
+
+	/** When true, applies a function to remap the values in the displacement map, which can be used 
+	 for contrast adjustment. The texture map values are converted to the range [0,1] before applying
+	 the remapping. */
+	UPROPERTY(EditAnywhere, Category = Options)
+	bool bApplyAdjustmentCurve = false;
+
+	/** This curve is queried in the range [0,1] to adjust contrast of the displacement map. */
+	UPROPERTY(VisibleAnywhere, Category = Options, meta = (EditCondition = "bApplyAdjustmentCurve", EditConditionHides))
+	UCurveFloat* AdjustmentCurve = nullptr;
 };
 
 
diff --git a/Engine/Plugins/Experimental/MeshModelingToolset/Source/ModelingComponents/Private/AssetUtils/Texture2DUtil.cpp b/Engine/Plugins/Experimental/MeshModelingToolset/Source/ModelingComponents/Private/AssetUtils/Texture2DUtil.cpp
--- a/Engine/Plugins/Experimental/MeshModelingToolset/Source/ModelingComponents/Private/AssetUtils/Texture2DUtil.cpp
+++ b/Engine/Plugins/Experimental/MeshModelingToolset/Source/ModelingComponents/Private/AssetUtils/Texture2DUtil.cpp
@@ -47,8 +47,7 @@ static bool ReadTexture_PlatformData(
 	for (int32 i = 0; i < Num; ++i)
 	{
 		FColor ByteColor = FormattedImageData[i];
-		FLinearColor FloatColor(ByteColor);
-		DestImage.SetPixel(i, FVector4f(FloatColor));
+		DestImage.SetPixel(i, FVector4f(ByteColor.ReinterpretAsLinear()));
 	}
 
 	// restore built platform texture data to initial state
@@ -95,7 +94,7 @@ static bool ReadTexture_SourceData(
 
 			FLinearColor FloatColor = (TextureMap->SRGB) ?
 				FLinearColor::FromSRGBColor(PixelColor) :
-				FLinearColor(float(PixelColor.R), float(PixelColor.G), float(PixelColor.B), float(PixelColor.A)) / 255.0f;
+				PixelColor.ReinterpretAsLinear();
 
 			DestImage.SetPixel(i, FVector4f(FloatColor));
 		}
diff --git a/Engine/Plugins/Experimental/MeshModelingToolset/Source/ModelingComponents/Private/ToolSetupUtil.cpp b/Engine/Plugins/Experimental/MeshModelingToolset/Source/ModelingComponents/Private/ToolSetupUtil.cpp
--- a/Engine/Plugins/Experimental/MeshModelingToolset/Source/ModelingComponents/Private/ToolSetupUtil.cpp
+++ b/Engine/Plugins/Experimental/MeshModelingToolset/Source/ModelingComponents/Private/ToolSetupUtil.cpp
@@ -2,6 +2,8 @@
 
 
 #include "ToolSetupUtil.h"
+
+#include "Curves/CurveFloat.h"
 #include "InteractiveTool.h"
 #include "InteractiveToolManager.h"
 #include "Materials/Material.h"
@@ -207,4 +209,12 @@ UMaterialInterface* ToolSetupUtil::GetDefaultLineComponentMaterial(UInteractiveT
 		return ToolManager->GetContextQueriesAPI()->GetStandardMaterial(EStandardToolContextMaterials::VertexColorMaterial);
 	}
 	return Material;
+}
+
+UCurveFloat* ToolSetupUtil::GetContrastAdjustmentCurve(UInteractiveToolManager* ToolManager)
+{
+	// This curve would currently be shared across any tools that need such a curve. We'll probably want to revisit this
+	// once it is used in multiple tools.
+	UCurveFloat* Curve = LoadObject<UCurveFloat>(nullptr, TEXT("/MeshModelingToolset/Curves/ContrastAdjustmentCurve"));
+	return Curve;
 }
\ No newline at end of file
diff --git a/Engine/Plugins/Experimental/MeshModelingToolset/Source/ModelingComponents/Public/ToolSetupUtil.h b/Engine/Plugins/Experimental/MeshModelingToolset/Source/ModelingComponents/Public/ToolSetupUtil.h
--- a/Engine/Plugins/Experimental/MeshModelingToolset/Source/ModelingComponents/Public/ToolSetupUtil.h
+++ b/Engine/Plugins/Experimental/MeshModelingToolset/Source/ModelingComponents/Public/ToolSetupUtil.h
@@ -6,6 +6,7 @@
 #include "Materials/Material.h"
 #include "Materials/MaterialInstanceDynamic.h"
 
+class UCurveFloat;
 class UTexture;
 class UInteractiveToolManager;
 
@@ -108,4 +109,9 @@ namespace ToolSetupUtil
 	 * @return custom material suitable for use with ULineSetComponent
 	 */
 	MODELINGCOMPONENTS_API UMaterialInterface* GetDefaultLineComponentMaterial(UInteractiveToolManager* ToolManager, bool bDepthTested = true);
+
+	/**
+	 * @return a curve asset used for contrast adjustments when using a texture map for displacements.
+	 */
+	MODELINGCOMPONENTS_API UCurveFloat* GetContrastAdjustmentCurve(UInteractiveToolManager* ToolManager);
 }
\ No newline at end of file
