ModelingToolset: add PolyLassoMarqueeMechanic, which implements freehand and multi-click 2D polygon drawing interaction. Add this as an optional input modality to UMeshGroupPaintTool. Add support for hiding brush indicator to MeshSculptToolBase and BrushStampIndicator
#rnx

diff --git a/Engine/Plugins/Experimental/MeshModelingToolset/Source/MeshModelingTools/Private/MeshGroupPaintTool.cpp b/Engine/Plugins/Experimental/MeshModelingToolset/Source/MeshModelingTools/Private/MeshGroupPaintTool.cpp
--- a/Engine/Plugins/Experimental/MeshModelingToolset/Source/MeshModelingTools/Private/MeshGroupPaintTool.cpp
+++ b/Engine/Plugins/Experimental/MeshModelingToolset/Source/MeshModelingTools/Private/MeshGroupPaintTool.cpp
@@ -14,7 +14,11 @@
 #include "Util/ColorConstants.h"
 #include "Selections/MeshConnectedComponents.h"
 #include "Selections/MeshFaceSelection.h"
+#include "Selections/MeshVertexSelection.h"
 #include "Polygroups/PolygroupUtil.h"
+#include "Polygon2.h"
+#include "Intersection/IntrLine2Line2.h"
+#include "Intersection/IntrSegment2Segment2.h"
 
 #include "Changes/MeshVertexChange.h"
 #include "Changes/MeshPolygroupChange.h"
@@ -24,6 +28,10 @@
 #include "Sculpting/StampFalloffs.h"
 #include "Sculpting/MeshSculptUtil.h"
 
+
+#include "CanvasTypes.h"
+#include "CanvasItem.h"
+
 #include "ToolSetupUtil.h"
 
 #define LOCTEXT_NAMESPACE "UMeshGroupPaintTool"
@@ -127,6 +135,13 @@ void UMeshGroupPaintTool::Setup()
 	// initialize brush radius range interval, brush properties
 	UMeshSculptToolBase::InitializeBrushSizeRange(Bounds);
 
+	// Set up control points mechanic
+	PolyLassoMechanic = NewObject<UPolyLassoMarqueeMechanic>(this);
+	PolyLassoMechanic->Setup(this);
+	PolyLassoMechanic->SetIsEnabled(false);
+	PolyLassoMechanic->SpacingTolerance = 10.0f;
+	PolyLassoMechanic->OnDrawPolyLassoFinished.AddUObject(this, &UMeshGroupPaintTool::OnPolyLassoFinished);
+
 	PolygroupLayerProperties = NewObject<UPolygroupLayersProperties>(this);
 	PolygroupLayerProperties->RestoreProperties(this);
 	PolygroupLayerProperties->InitializeGroupLayers(GetSculptMesh());
@@ -134,14 +149,15 @@ void UMeshGroupPaintTool::Setup()
 	UpdateActiveGroupLayer();
 	AddToolPropertySource(PolygroupLayerProperties);
 
-	ToolProperties = NewObject<UGroupPaintToolProperties>(this);
-	AddToolPropertySource(ToolProperties);
-	ToolProperties->WatchProperty(ToolProperties->SubToolType,
-		[this](EMeshGroupPaintInteractionType NewType) { UpdateSubToolType(NewType); });
-	ToolProperties->RestoreProperties(this);
-
 	// initialize other properties
 	FilterProperties = NewObject<UGroupPaintBrushFilterProperties>(this);
+	FilterProperties->WatchProperty(FilterProperties->SubToolType,
+		[this](EMeshGroupPaintInteractionType NewType) { UpdateSubToolType(NewType); });
+	FilterProperties->WatchProperty(FilterProperties->BrushSize,
+		[this](float NewSize) { UMeshSculptToolBase::BrushProperties->BrushSize = NewSize; });
+	FilterProperties->RestoreProperties(this);
+	FilterProperties->BrushSize = UMeshSculptToolBase::BrushProperties->BrushSize;
+	AddToolPropertySource(FilterProperties);
 
 	InitializeIndicator();
 
@@ -149,27 +165,22 @@ void UMeshGroupPaintTool::Setup()
 	AddToolPropertySource(UMeshSculptToolBase::BrushProperties);
 	UMeshSculptToolBase::BrushProperties->bShowPerBrushProps = false;
 	UMeshSculptToolBase::BrushProperties->bShowFalloff = false;
+	UMeshSculptToolBase::BrushProperties->bShowLazyness = false;
 	CalculateBrushRadius();
-	FilterProperties->RestoreProperties(this);
 
-	PaintBrushOpOperties = NewObject<UGroupPaintBrushOpProps>(this);
+	PaintBrushOpProperties = NewObject<UGroupPaintBrushOpProps>(this);
 	RegisterBrushType((int32)EMeshGroupPaintBrushType::Paint,
 		MakeUnique<FLambdaMeshSculptBrushOpFactory>([this]() { return MakeUnique<FGroupPaintBrushOp>(); }),
-		PaintBrushOpOperties);
-
-	//RegisterBrushType((int32)EMeshGroupPaintBrushType::Erase,
-	//	MakeUnique<TBasicMeshSculptBrushOpFactory<FGroupEraseBrushOp>>(),
-	//	EraseBrushOpOperties);
+		PaintBrushOpProperties);
 
 	// secondary brushes
-	EraseBrushOpOperties = NewObject<UGroupEraseBrushOpProps>(this);
-	EraseBrushOpOperties->GetCurrentGroupLambda = [this]() { return PaintBrushOpOperties->GetGroup(); };
+	EraseBrushOpProperties = NewObject<UGroupEraseBrushOpProps>(this);
+	EraseBrushOpProperties->GetCurrentGroupLambda = [this]() { return PaintBrushOpProperties->GetGroup(); };
 
 	RegisterSecondaryBrushType((int32)EMeshGroupPaintBrushType::Erase,
 		MakeUnique<TBasicMeshSculptBrushOpFactory<FGroupEraseBrushOp>>(),
-		EraseBrushOpOperties);
+		EraseBrushOpProperties);
 
-	AddToolPropertySource(FilterProperties);
 	AddToolPropertySource(UMeshSculptToolBase::ViewProperties);
 
 	AddToolPropertySource(UMeshSculptToolBase::GizmoProperties);
@@ -186,8 +197,6 @@ void UMeshGroupPaintTool::Setup()
 	UpdateBrushType(FilterProperties->PrimaryBrushType);
 	SetActiveSecondaryBrushType((int32)EMeshGroupPaintBrushType::Erase);
 
-	UpdateSubToolType(ToolProperties->SubToolType);
-
 	FreezeActions = NewObject<UMeshGroupPaintToolFreezeActions>(this);
 	FreezeActions->Initialize(this);
 	AddToolPropertySource(FreezeActions);
@@ -214,6 +223,9 @@ void UMeshGroupPaintTool::Setup()
 	UpdateWireframeVisibility(false);
 	UpdateFlatShadingSetting(true);
 
+	// configure panels
+	UpdateSubToolType(FilterProperties->SubToolType);
+
 	PrecomputeFuture.Wait();
 	OctreeFuture.Wait();
 }
@@ -232,7 +244,6 @@ void UMeshGroupPaintTool::Shutdown(EToolShutdownType ShutdownType)
 	MeshElementsDisplay->Disconnect();
 
 	FilterProperties->SaveProperties(this);
-	ToolProperties->SaveProperties(this);
 	PolygroupLayerProperties->SaveProperties(this);
 
 
@@ -293,10 +304,28 @@ void UMeshGroupPaintTool::OnPropertyModified(UObject* PropertySet, FProperty* Pr
 }
 
 
+bool UMeshGroupPaintTool::IsInBrushSubMode() const
+{
+	return FilterProperties->SubToolType == EMeshGroupPaintInteractionType::Brush
+		|| FilterProperties->SubToolType == EMeshGroupPaintInteractionType::Fill;
+}
+
+
 void UMeshGroupPaintTool::OnBeginStroke(const FRay& WorldRay)
 {
 	UpdateBrushPosition(WorldRay);
 
+	if (PaintBrushOpProperties)
+	{
+		PaintBrushOpProperties->Group = FilterProperties->SetGroup;
+		PaintBrushOpProperties->bOnlyPaintUngrouped = FilterProperties->bOnlySetUngrouped;
+	}
+	if (EraseBrushOpProperties)
+	{
+		EraseBrushOpProperties->Group = FilterProperties->EraseGroup;
+		EraseBrushOpProperties->bOnlyEraseCurrent = FilterProperties->bOnlyEraseCurrent;
+	}
+
 	// initialize first "Last Stamp", so that we can assume all stamps in stroke have a valid previous stamp
 	LastStamp.WorldFrame = GetBrushFrameWorld();
 	LastStamp.LocalFrame = GetBrushFrameLocal();
@@ -331,6 +360,7 @@ void UMeshGroupPaintTool::OnEndStroke()
 
 
 
+
 void UMeshGroupPaintTool::UpdateROI(const FSculptBrushStamp& BrushStamp)
 {
 	SCOPE_CYCLE_COUNTER(GroupPaintTool_UpdateROI);
@@ -350,7 +380,9 @@ void UMeshGroupPaintTool::UpdateROI(const FSculptBrushStamp& BrushStamp)
 		TriangleROI.Add(CenterTID);
 	}
 
-	if (FilterProperties->bVolumetric)
+	bool bFill = (FilterProperties->SubToolType == EMeshGroupPaintInteractionType::Fill);
+
+	if (FilterProperties->BrushAreaMode == EMeshGroupPaintBrushAreaType::Volumetric)
 	{
 		Octree.RangeQuery(BrushBox,
 			[&](int TriIdx) {
@@ -360,6 +392,20 @@ void UMeshGroupPaintTool::UpdateROI(const FSculptBrushStamp& BrushStamp)
 				TriangleROI.Add(TriIdx);
 			}
 		});
+
+		if (bFill)
+		{
+			TArray<int32> StartROI;
+			for (int32 tid : TriangleROI)
+			{
+				StartROI.Add(tid);
+			}
+			FMeshConnectedComponents::GrowToConnectedTriangles(Mesh, StartROI, TriangleROI, &TempROIBuffer,
+				[&](int t1, int t2)
+			{
+				return true;
+			});
+		}
 	}
 	else
 	{
@@ -377,7 +423,7 @@ void UMeshGroupPaintTool::UpdateROI(const FSculptBrushStamp& BrushStamp)
 			FMeshConnectedComponents::GrowToConnectedTriangles(Mesh, StartROI, TriangleROI, &TempROIBuffer,
 				[&](int t1, int t2) 
 			{ 
-				if ((Mesh->GetTriCentroid(t2) - BrushPos).SquaredLength() < RadiusSqr)
+				if (bFill || (Mesh->GetTriCentroid(t2) - BrushPos).SquaredLength() < RadiusSqr)
 				{
 					if (bUseAngleThreshold == false || CenterNormal.Dot(TriNormals[t2]) > DotAngleThreshold)
 					{
@@ -400,42 +446,11 @@ void UMeshGroupPaintTool::UpdateROI(const FSculptBrushStamp& BrushStamp)
 	// apply visibility filter
 	if (FilterProperties->VisibilityFilter != EMeshGroupPaintVisibilityType::None)
 	{
-		FViewCameraState StateOut;
-		GetToolManager()->GetContextQueriesAPI()->GetCurrentViewState(StateOut);
-		FVector3d LocalEyePosition(ComponentTarget->GetWorldTransform().InverseTransformPosition(StateOut.Position));
-		TempROIBuffer.SetNum(0, false);
-		for (int32 tid : TriangleROI)
-		{
-			TempROIBuffer.Add(tid);
-		}
-		ParallelFor(TempROIBuffer.Num(), [&](int32 idx)
-		{
-			FVector3d Centroid = Mesh->GetTriCentroid(TempROIBuffer[idx]);
-			FVector3d FaceNormal = Mesh->GetTriNormal(TempROIBuffer[idx]);
-			if (FaceNormal.Dot((Centroid - LocalEyePosition)) > 0)
-			{
-				TempROIBuffer[idx] = -1;
-			}
-			if (FilterProperties->VisibilityFilter == EMeshGroupPaintVisibilityType::Unoccluded)
-			{
-				int32 HitTID = Octree.FindNearestHitObject(FRay3d(LocalEyePosition, (Centroid - LocalEyePosition).Normalized()));
-				if (HitTID != TempROIBuffer[idx])
-				{
-					TempROIBuffer[idx] = -1;
-				}
-			}
-		});
-		TriangleROI.Reset();
-		for (int32 tid : TempROIBuffer)
-		{
-			if (tid >= 0)
-			{
-				TriangleROI.Add(tid);
-			}
-		}
+		TArray<int32> ResultBuffer;
+		ApplyVisibilityFilter(TriangleROI, TempROIBuffer, ResultBuffer);
 	}
 
-
+	// construct ROI vertex set
 	VertexSetBuffer.Reset();
 	for (int32 tid : TriangleROI)
 	{
@@ -445,7 +460,7 @@ void UMeshGroupPaintTool::UpdateROI(const FSculptBrushStamp& BrushStamp)
 	VertexROI.SetNum(0, false);
 	BufferUtil::AppendElements(VertexROI, VertexSetBuffer);
 
-
+	// construct ROI triangle and group buffers
 	ROITriangleBuffer.Reserve(TriangleROI.Num());
 	ROITriangleBuffer.SetNum(0, false);
 	for (int32 tid : TriangleROI)
@@ -543,11 +558,364 @@ void UMeshGroupPaintTool::SyncMeshWithGroupBuffer(FDynamicMesh3* Mesh)
 
 
 
+template<typename RealType>
+static bool FindPolylineSelfIntersection(
+	const TArray<FVector2<RealType>>& Polyline, 
+	FVector2<RealType>& IntersectionPointOut, 
+	FIndex2i& IntersectionIndexOut,
+	bool bParallel = true)
+{
+	int32 N = Polyline.Num();
+	std::atomic<bool> bSelfIntersects(false);
+	ParallelFor(N - 1, [&](int32 i)
+	{
+		TSegment2<RealType> SegA(Polyline[i], Polyline[i + 1]);
+		for (int32 j = i + 2; j < N - 1 && bSelfIntersects == false; ++j)
+		{
+			TSegment2<RealType> SegB(Polyline[j], Polyline[j + 1]);
+			if (SegA.Intersects(SegB) && bSelfIntersects == false)		
+			{
+				bool ExpectedValue = false;
+				if (std::atomic_compare_exchange_strong(&bSelfIntersects, &ExpectedValue, true))
+				{
+					TIntrSegment2Segment2<RealType> Intersection(SegA, SegB);
+					Intersection.Find();
+					IntersectionPointOut = Intersection.Point0;
+					IntersectionIndexOut = FIndex2i(i, j);
+					return;
+				}
+			}
+		}
+	}, (bParallel) ? EParallelForFlags::None : EParallelForFlags::ForceSingleThread );
+
+	return bSelfIntersects;
+}
+
+
+
+template<typename RealType>
+static bool FindPolylineSegmentIntersection(
+	const TArray<FVector2<RealType>>& Polyline,
+	const TSegment2<RealType>& Segment,
+	FVector2<RealType>& IntersectionPointOut,
+	int& IntersectionIndexOut)
+{
+
+	int32 N = Polyline.Num();
+	for (int32 i = 0; i < N-1; ++i)
+	{
+		TSegment2<RealType> PolySeg(Polyline[i], Polyline[i + 1]);
+		if (Segment.Intersects(PolySeg))
+		{
+			TIntrSegment2Segment2<RealType> Intersection(Segment, PolySeg);
+			Intersection.Find();
+			IntersectionPointOut = Intersection.Point0;
+			IntersectionIndexOut = i;
+			return true;
+		}
+	}
+	return false;
+}
+
+
+
+bool ApproxSelfClipPolyline(TArray<FVector2f>& Polyline)
+{
+	int32 N = Polyline.Num();
+
+	// handle already-closed polylines
+	if (Polyline[0].Distance(Polyline[N-1]) < 0.0001f)
+	{
+		return true;
+	}
+
+	FVector2f IntersectPoint;
+	FIndex2i IntersectionIndex(-1, -1);
+	bool bSelfIntersects = FindPolylineSelfIntersection(Polyline, IntersectPoint, IntersectionIndex);
+	if (bSelfIntersects)
+	{
+		TArray<FVector2f> NewPolyline;
+		NewPolyline.Add(IntersectPoint);
+		for (int32 i = IntersectionIndex.A; i <= IntersectionIndex.B; ++i)
+		{
+			NewPolyline.Add(Polyline[i]);
+		}
+		NewPolyline.Add(IntersectPoint);
+		Polyline = MoveTemp(NewPolyline);
+		return true;
+	}
+
+
+	FVector2f StartDirOut = (Polyline[0] - Polyline[1]).Normalized();
+	FLine2f StartLine(Polyline[0], StartDirOut);
+	FVector2f EndDirOut = (Polyline[N - 1] - Polyline[N - 2]).Normalized();
+	FLine2f EndLine(Polyline[N - 1], EndDirOut);
+	FIntrLine2Line2f LineIntr(StartLine, EndLine);
+	bool bIntersects = false;
+	if (LineIntr.Find())
+	{
+		bIntersects = LineIntr.IsSimpleIntersection() && (LineIntr.Segment1Parameter > 0) && (LineIntr.Segment2Parameter > 0);
+		if (bIntersects)
+		{
+			Polyline.Add(StartLine.PointAt(LineIntr.Segment1Parameter));
+			Polyline.Add(StartLine.Origin);
+			return true;
+		}
+	}
+
+
+	FAxisAlignedBox2f Bounds;
+	for (const FVector2f& P : Polyline)
+	{
+		Bounds.Contain(P);
+	}
+	float Size = Bounds.DiagonalLength();
+
+	FVector2f StartPos = Polyline[0] + 0.001 * StartDirOut;
+	if (FindPolylineSegmentIntersection(Polyline, FSegment2f(StartPos, StartPos + 2*Size*StartDirOut), IntersectPoint, IntersectionIndex.A))
+	{
+		//TArray<FVector2f> NewPolyline;
+		//for (int32 i = 0; i <= IntersectionIndex.A; ++i)
+		//{
+		//	NewPolyline.Add(Polyline[i]);
+		//}
+		//NewPolyline.Add(IntersectPoint);
+		//NewPolyline.Add(Polyline[0]);
+		//Polyline = MoveTemp(NewPolyline);
+		return true;
+	}
+
+	FVector2f EndPos = Polyline[N-1] + 0.001 * EndDirOut;
+	if (FindPolylineSegmentIntersection(Polyline, FSegment2f(EndPos, EndPos + 2*Size*EndDirOut), IntersectPoint, IntersectionIndex.A))
+	{
+		//TArray<FVector2f> NewPolyline;
+		//NewPolyline.Add(IntersectPoint);
+		//for (int32 i = IntersectionIndex.A+1; i < N; ++i)
+		//{
+		//	NewPolyline.Add(Polyline[i]);
+		//}
+		//NewPolyline.Add(Polyline[0]);
+		//NewPolyline.Add(IntersectPoint);
+		//Polyline = MoveTemp(NewPolyline);
+		return true;
+	}
+
+	return false;
+}
+
+
+
+void UMeshGroupPaintTool::OnPolyLassoFinished(const FCameraPolyLasso& Lasso, bool bCanceled)
+{
+	// construct polyline
+	TArray<FVector2f> Polyline;
+	for (FVector2D Pos : Lasso.Polyline)
+	{
+		Polyline.Add((FVector2f)Pos);
+	}
+	int32 N = Polyline.Num();
+	if (N < 2)
+	{
+		return;
+	}
+
+	// Try to clip polyline to be closed, or closed-enough for winding evaluation to work.
+	// If that returns false, the polyline is "too open". In that case we will extend
+	// outwards from the endpoints and then try to create a closed very large polygon
+	if (ApproxSelfClipPolyline(Polyline) == false)
+	{
+		FVector2f StartDirOut = (Polyline[0] - Polyline[1]).Normalized();
+		FLine2f StartLine(Polyline[0], StartDirOut);
+		FVector2f EndDirOut = (Polyline[N-1] - Polyline[N-2]).Normalized();
+		FLine2f EndLine(Polyline[N-1], EndDirOut);
+
+		// if we did not intersect, we are in ambiguous territory. Check if a segment along either end-direction
+		// intersects the polyline. If it does, we have something like a spiral and will be OK. 
+		// If not, make a closed polygon by interpolating outwards from each endpoint, and then in perp-directions.
+		FPolygon2f Polygon(Polyline);
+		float PerpSign = Polygon.IsClockwise() ? -1.0 : 1.0;
+
+		Polyline.Insert(StartLine.PointAt(10000.0f), 0);
+		Polyline.Insert(Polyline[0] + 1000 * PerpSign * StartDirOut.Perp(), 0);
+
+		Polyline.Add(EndLine.PointAt(10000.0f));
+		Polyline.Add(Polyline.Last() + 1000 * PerpSign * EndDirOut.Perp());
+		FVector2f StartPos = Polyline[0];
+		Polyline.Add(StartPos);		// close polyline (cannot use Polyline[0] in case Add resizes!)
+	}
+
+	N = Polyline.Num();
+
+	// project each mesh vertex to view plane and evaluate winding integral of polyline
+	const FDynamicMesh3* Mesh = GetSculptMesh();
+	TempROIBuffer.SetNum(Mesh->MaxVertexID());
+	ParallelFor(Mesh->MaxVertexID(), [&](int32 vid)
+	{
+		if (Mesh->IsVertex(vid))
+		{
+			FVector3d WorldPos = CurTargetTransform.TransformPosition(Mesh->GetVertex(vid));
+			FVector2f PlanePos = (FVector2f)Lasso.GetProjectedPoint((FVector)WorldPos);
+
+			double WindingSum = 0;
+			FVector2f a = Polyline[0] - PlanePos, b = FVector2f::Zero();
+			for (int32 i = 1; i < N; ++i)
+			{
+				b = Polyline[i] - PlanePos;
+				WindingSum += (double)FMathf::Atan2(a.X*b.Y - a.Y*b.X, a.X*b.X + a.Y*b.Y);
+				a = b;
+			}
+			WindingSum /= FMathd::TwoPi;
+			bool bInside = FMathd::Abs(WindingSum) > 0.3;
+			TempROIBuffer[vid] = bInside ? 1 : 0;
+		}
+		else
+		{
+			TempROIBuffer[vid] = -1;
+		}
+	});
+
+	// convert to vertex selection, and then select fully-enclosed faces
+	FMeshVertexSelection VertexSelection(Mesh);
+	VertexSelection.SelectByVertexID([&](int32 vid) { return TempROIBuffer[vid] == 1; });
+ 	FMeshFaceSelection FaceSelection(Mesh, VertexSelection, FilterProperties->MinTriVertCount);
+	if (FaceSelection.Num() == 0)
+	{
+		return;
+	}
+
+	int32 SetGroupID = GetInInvertStroke() ? FilterProperties->EraseGroup : FilterProperties->SetGroup;
+	SetTrianglesToGroupID(FaceSelection.AsSet(), SetGroupID, GetInInvertStroke());
+}
+
+
+
+
+void UMeshGroupPaintTool::SetTrianglesToGroupID(const TSet<int32>& Triangles, int32 ToGroupID, bool bIsErase)
+{
+	BeginChange();
+
+	TempROIBuffer.SetNum(0, false);
+	for (int32 tid : Triangles)
+	{
+		int32 CurGroupID = ActiveGroupSet->GetGroup(tid);
+		if (CurGroupID == ToGroupID || FrozenGroups.Contains(CurGroupID))		// skip frozen groups
+		{
+			continue;
+		}
+		if (bIsErase == false && FilterProperties->bOnlySetUngrouped && CurGroupID != 0)
+		{
+			continue;
+		}
+		if (bIsErase && FilterProperties->bOnlyEraseCurrent && CurGroupID != FilterProperties->SetGroup)
+		{
+			continue;
+		}
+
+		TempROIBuffer.Add(tid);
+	}
+
+	if (HaveVisibilityFilter())
+	{
+		TArray<int32> VisibleTriangles;
+		VisibleTriangles.Reserve(TempROIBuffer.Num());
+		ApplyVisibilityFilter(TempROIBuffer, VisibleTriangles);
+		TempROIBuffer = MoveTemp(VisibleTriangles);
+	}
+
+	ActiveGroupEditBuilder->SaveTriangles(TempROIBuffer);
+	for (int32 tid : TempROIBuffer)
+	{
+		ActiveGroupSet->SetGroup(tid, ToGroupID);
+	}
+	ActiveGroupEditBuilder->SaveTriangles(TempROIBuffer);
+
+	DynamicMeshComponent->FastNotifyTriangleVerticesUpdated(TempROIBuffer, EMeshRenderAttributeFlags::VertexColors);
+	GetToolManager()->PostInvalidation();
+	EndChange();
+}
+
+
+
+bool UMeshGroupPaintTool::HaveVisibilityFilter() const
+{
+	return FilterProperties->VisibilityFilter != EMeshGroupPaintVisibilityType::None;
+}
+
+
+void UMeshGroupPaintTool::ApplyVisibilityFilter(TSet<int32>& Triangles, TArray<int32>& ROIBuffer, TArray<int32>& OutputBuffer)
+{
+	ROIBuffer.SetNum(0, false);
+	ROIBuffer.Reserve(Triangles.Num());
+	for (int32 tid : Triangles)
+	{
+		ROIBuffer.Add(tid);
+	}
+	
+	OutputBuffer.Reset();
+	ApplyVisibilityFilter(TempROIBuffer, OutputBuffer);
+
+	Triangles.Reset();
+	for (int32 tid : OutputBuffer)
+	{
+		TriangleROI.Add(tid);
+	}
+}
+
+void UMeshGroupPaintTool::ApplyVisibilityFilter(const TArray<int32>& Triangles, TArray<int32>& VisibleTriangles)
+{
+	if (!HaveVisibilityFilter())
+	{
+		VisibleTriangles = Triangles;
+		return;
+	}
+
+	FViewCameraState StateOut;
+	GetToolManager()->GetContextQueriesAPI()->GetCurrentViewState(StateOut);
+	FVector3d LocalEyePosition(ComponentTarget->GetWorldTransform().InverseTransformPosition(StateOut.Position));
+
+	const FDynamicMesh3* Mesh = GetSculptMesh();
+
+	int32 NumTriangles = Triangles.Num();
+
+	VisibilityFilterBuffer.SetNum(NumTriangles, false);
+	ParallelFor(NumTriangles, [&](int32 idx)
+	{
+		VisibilityFilterBuffer[idx] = true;
+		FVector3d Centroid = Mesh->GetTriCentroid(Triangles[idx]);
+		FVector3d FaceNormal = Mesh->GetTriNormal(Triangles[idx]);
+		if (FaceNormal.Dot((Centroid - LocalEyePosition)) > 0)
+		{
+			VisibilityFilterBuffer[idx] = false;
+		}
+		if (FilterProperties->VisibilityFilter == EMeshGroupPaintVisibilityType::Unoccluded)
+		{
+			int32 HitTID = Octree.FindNearestHitObject(FRay3d(LocalEyePosition, (Centroid - LocalEyePosition).Normalized()));
+			if (HitTID != Triangles[idx])
+			{
+				VisibilityFilterBuffer[idx] = false;
+			}
+		}
+	});
+
+	VisibleTriangles.Reset();
+	for (int32 k = 0; k < NumTriangles; ++k)
+	{
+		if (VisibilityFilterBuffer[k])
+		{
+			VisibleTriangles.Add(Triangles[k]);
+		}
+	}
+}
 
 
 
 int32 UMeshGroupPaintTool::FindHitSculptMeshTriangle(const FRay3d& LocalRay)
 {
+	if (!IsInBrushSubMode())
+	{
+		return IndexConstants::InvalidID;
+	}
+
 	if (GetBrushCanHitBackFaces())
 	{
 		return Octree.FindNearestHitObject(LocalRay);
@@ -619,13 +987,25 @@ bool UMeshGroupPaintTool::OnUpdateHover(const FInputDeviceRay& DevicePos)
 }
 
 
+void UMeshGroupPaintTool::DrawHUD(FCanvas* Canvas, IToolsContextRenderAPI* RenderAPI)
+{
+	if (PolyLassoMechanic)
+	{
+		PolyLassoMechanic->DrawHUD(Canvas, RenderAPI);
+	}
+}
+
 
 void UMeshGroupPaintTool::OnTick(float DeltaTime)
 {
 	UMeshSculptToolBase::OnTick(DeltaTime);
 	MeshElementsDisplay->OnTick(DeltaTime);
 
-	ConfigureIndicator(FilterProperties->bVolumetric);
+	bool bIsLasso = (FilterProperties->SubToolType == EMeshGroupPaintInteractionType::PolyLasso);
+	PolyLassoMechanic->SetIsEnabled(bIsLasso);
+
+	ConfigureIndicator(FilterProperties->BrushAreaMode == EMeshGroupPaintBrushAreaType::Volumetric);
+	SetIndicatorVisibility(bIsLasso == false);
 
 	if (bHavePendingAction)
 	{
@@ -660,7 +1040,7 @@ void UMeshGroupPaintTool::OnTick(float DeltaTime)
 				int32 HitGroupID = ActiveGroupSet->GetGroup(GetBrushTriangleID());
 				if (bPendingPickGroup)
 				{
-					PaintBrushOpOperties->Group = HitGroupID;
+					FilterProperties->SetGroup = HitGroupID;
 				}
 				else if (bPendingToggleFreezeGroup)
 				{
@@ -672,11 +1052,10 @@ void UMeshGroupPaintTool::OnTick(float DeltaTime)
 	}
 
 
-	if (ToolProperties->SubToolType == EMeshGroupPaintInteractionType::Brush)
+	if (IsInBrushSubMode())
 	{
 		if (IsStampPending())
 		{
-			//UE_LOG(LogTemp, Warning, TEXT("dt is %.3f, tick fps %.2f - roi size %d/%d"), DeltaTime, 1.0 / DeltaTime, VertexROI.Num(), TriangleROI.Num());
 			SCOPE_CYCLE_COUNTER(GroupPaintTool_Tick_ApplyStampBlock);
 
 			ApplyStrokeFlowInTick();
@@ -720,8 +1099,7 @@ void UMeshGroupPaintTool::OnTick(float DeltaTime)
 
 void UMeshGroupPaintTool::AllocateNewGroupAndSetAsCurrentAction()
 {
-	int32 NewGroupID = ActiveGroupSet->AllocateNewGroupID();
-	PaintBrushOpOperties->Group = NewGroupID;
+	FilterProperties->SetGroup = ActiveGroupSet->AllocateNewGroupID();
 }
 
 
@@ -827,7 +1205,7 @@ void UMeshGroupPaintTool::GrowCurrentGroupAction()
 {
 	BeginChange();
 
-	int32 CurrentGroupID = PaintBrushOpOperties->Group;
+	int32 CurrentGroupID = FilterProperties->SetGroup;
 	const FDynamicMesh3* Mesh = DynamicMeshComponent->GetMesh();
 	FMeshFaceSelection InitialSelection(Mesh);
 	InitialSelection.Select([&](int32 tid) { return ActiveGroupSet->GetGroup(tid) == CurrentGroupID; });
@@ -853,7 +1231,7 @@ void UMeshGroupPaintTool::ShrinkCurrentGroupAction()
 {
 	BeginChange();
 
-	int32 CurrentGroupID = PaintBrushOpOperties->Group;
+	int32 CurrentGroupID = FilterProperties->SetGroup;
 	const FDynamicMesh3* Mesh = DynamicMeshComponent->GetMesh();
 	FMeshFaceSelection InitialSelection(Mesh);
 	InitialSelection.Select([&](int32 tid) { return ActiveGroupSet->GetGroup(tid) == CurrentGroupID; });
@@ -876,6 +1254,91 @@ void UMeshGroupPaintTool::ShrinkCurrentGroupAction()
 }
 
 
+void UMeshGroupPaintTool::ClearCurrentGroupAction()
+{
+	BeginChange();
+
+	int32 CurrentGroupID = FilterProperties->SetGroup;
+	const FDynamicMesh3* Mesh = DynamicMeshComponent->GetMesh();
+	TempROIBuffer.SetNum(0, false);
+	for (int32 tid : Mesh->TriangleIndicesItr())
+	{
+		if (ActiveGroupSet->GetGroup(tid) == CurrentGroupID)
+		{
+			TempROIBuffer.Add(tid);
+		}
+	}
+
+	ActiveGroupEditBuilder->SaveTriangles(TempROIBuffer);
+	for (int32 tid : TempROIBuffer)
+	{
+		ActiveGroupSet->SetGroup(tid, 0);
+	}
+	ActiveGroupEditBuilder->SaveTriangles(TempROIBuffer);
+
+	DynamicMeshComponent->FastNotifyTriangleVerticesUpdated(TempROIBuffer, EMeshRenderAttributeFlags::VertexColors);
+	GetToolManager()->PostInvalidation();
+	EndChange();
+}
+
+
+void UMeshGroupPaintTool::FloodFillCurrentGroupAction()
+{
+	BeginChange();
+
+	int32 SetGroupID = FilterProperties->SetGroup;
+	const FDynamicMesh3* Mesh = DynamicMeshComponent->GetMesh();
+	TempROIBuffer.SetNum(0, false);
+	for (int32 tid : Mesh->TriangleIndicesItr())
+	{
+		int32 GroupID = ActiveGroupSet->GetGroup(tid);
+		if (GroupID == 0 && GroupID != SetGroupID && FrozenGroups.Contains(GroupID) == false)
+		{
+			TempROIBuffer.Add(tid);
+		}
+	}
+
+	ActiveGroupEditBuilder->SaveTriangles(TempROIBuffer);
+	for (int32 tid : TempROIBuffer)
+	{
+		ActiveGroupSet->SetGroup(tid, SetGroupID);
+	}
+	ActiveGroupEditBuilder->SaveTriangles(TempROIBuffer);
+
+	DynamicMeshComponent->FastNotifyTriangleVerticesUpdated(TempROIBuffer, EMeshRenderAttributeFlags::VertexColors);
+	GetToolManager()->PostInvalidation();
+	EndChange();
+}
+
+
+void UMeshGroupPaintTool::ClearAllGroupsAction()
+{
+	BeginChange();
+
+	const FDynamicMesh3* Mesh = DynamicMeshComponent->GetMesh();
+	TempROIBuffer.SetNum(0, false);
+	for (int32 tid : Mesh->TriangleIndicesItr())
+	{
+		if (ActiveGroupSet->GetGroup(tid) != 0)
+		{
+			TempROIBuffer.Add(tid);
+		}
+	}
+
+	ActiveGroupEditBuilder->SaveTriangles(TempROIBuffer);
+	for (int32 tid : TempROIBuffer)
+	{
+		ActiveGroupSet->SetGroup(tid, 0);
+	}
+	ActiveGroupEditBuilder->SaveTriangles(TempROIBuffer);
+
+	DynamicMeshComponent->FastNotifyTriangleVerticesUpdated(TempROIBuffer, EMeshRenderAttributeFlags::VertexColors);
+	GetToolManager()->PostInvalidation();
+	EndChange();
+}
+
+
+
 
 //
 // Change Tracking
@@ -1022,10 +1485,12 @@ void UMeshGroupPaintTool::UpdateActiveGroupLayer()
 
 void UMeshGroupPaintTool::UpdateSubToolType(EMeshGroupPaintInteractionType NewType)
 {
-	bool bSculptPropsVisible = (NewType == EMeshGroupPaintInteractionType::Brush);
-	SetToolPropertySourceEnabled(FilterProperties, bSculptPropsVisible);
-	SetToolPropertySourceEnabled(UMeshSculptToolBase::BrushProperties, bSculptPropsVisible);
-	SetBrushOpPropsVisibility(bSculptPropsVisible);
+	//bool bSculptPropsVisible = (NewType == EMeshGroupPaintInteractionType::Brush);
+	//SetToolPropertySourceEnabled(UMeshSculptToolBase::BrushProperties, bSculptPropsVisible);
+	SetToolPropertySourceEnabled(UMeshSculptToolBase::BrushProperties, false);
+
+	SetToolPropertySourceEnabled(FilterProperties, true);
+	SetBrushOpPropsVisibility(false);
 }
 
 
@@ -1064,11 +1529,11 @@ void UMeshGroupPaintTool::ApplyAction(EMeshGroupPaintToolActions ActionType)
 		break;
 
 	case EMeshGroupPaintToolActions::FreezeCurrent:
-		ToggleFrozenGroup(PaintBrushOpOperties->Group);
+		ToggleFrozenGroup(FilterProperties->SetGroup);
 		break;
 
 	case EMeshGroupPaintToolActions::FreezeOthers:
-		FreezeOtherGroups(PaintBrushOpOperties->Group);
+		FreezeOtherGroups(FilterProperties->SetGroup);
 		break;
 
 	case EMeshGroupPaintToolActions::GrowCurrent:
@@ -1078,9 +1543,23 @@ void UMeshGroupPaintTool::ApplyAction(EMeshGroupPaintToolActions ActionType)
 	case EMeshGroupPaintToolActions::ShrinkCurrent:
 		ShrinkCurrentGroupAction();
 		break;
+
+	case EMeshGroupPaintToolActions::ClearCurrent:
+		ClearCurrentGroupAction();
+		break;
+
+	case EMeshGroupPaintToolActions::FloodFillCurrent:
+		FloodFillCurrentGroupAction();
+		break;
+
+	case EMeshGroupPaintToolActions::ClearAll:
+		ClearAllGroupsAction();
+		break;
+
 	}
 }
 
 
 
+
 #undef LOCTEXT_NAMESPACE
diff --git a/Engine/Plugins/Experimental/MeshModelingToolset/Source/MeshModelingTools/Private/Sculpting/MeshSculptToolBase.cpp b/Engine/Plugins/Experimental/MeshModelingToolset/Source/MeshModelingTools/Private/Sculpting/MeshSculptToolBase.cpp
--- a/Engine/Plugins/Experimental/MeshModelingToolset/Source/MeshModelingTools/Private/Sculpting/MeshSculptToolBase.cpp
+++ b/Engine/Plugins/Experimental/MeshModelingToolset/Source/MeshModelingTools/Private/Sculpting/MeshSculptToolBase.cpp
@@ -306,7 +306,7 @@ void UMeshSculptToolBase::SetActivePrimaryBrushType(int32 Identifier)
 	UMeshSculptBrushOpProps** FoundProps = BrushOpPropSets.Find(Identifier);
 	if (FoundProps != nullptr)
 	{
-		SetToolPropertySourceEnabled(*FoundProps, true);
+		SetToolPropertySourceEnabled(*FoundProps, bBrushOpPropsVisible);
 		PrimaryVisiblePropSet = *FoundProps;
 
 		PrimaryBrushOp->PropertySet = PrimaryVisiblePropSet;
@@ -338,7 +338,7 @@ void UMeshSculptToolBase::SetActiveSecondaryBrushType(int32 Identifier)
 	UMeshSculptBrushOpProps** FoundProps = SecondaryBrushOpPropSets.Find(Identifier);
 	if (FoundProps != nullptr)
 	{
-		SetToolPropertySourceEnabled(*FoundProps, true);
+		SetToolPropertySourceEnabled(*FoundProps, bBrushOpPropsVisible);
 		SecondaryVisiblePropSet = *FoundProps;
 
 		SecondaryBrushOp->PropertySet = SecondaryVisiblePropSet;
@@ -361,6 +361,7 @@ TUniquePtr<FMeshSculptBrushOp>& UMeshSculptToolBase::GetActiveBrushOp()
 
 void UMeshSculptToolBase::SetBrushOpPropsVisibility(bool bVisible)
 {
+	bBrushOpPropsVisible = bVisible;
 	if (PrimaryVisiblePropSet)
 	{
 		SetToolPropertySourceEnabled(PrimaryVisiblePropSet, bVisible);
@@ -889,7 +890,7 @@ void UMeshSculptToolBase::ConfigureIndicator(bool bVolumetric)
 {
 	if (bIsVolumetricIndicator == bVolumetric) return;
 	bIsVolumetricIndicator = bVolumetric;
-	BrushIndicatorMesh->SetVisible(bVolumetric);
+	BrushIndicatorMesh->SetVisible(GetIndicatorVisibility() && bIsVolumetricIndicator);
 	if (bVolumetric)
 	{
 		BrushIndicator->bDrawRadiusCircle = false;
@@ -901,6 +902,20 @@ void UMeshSculptToolBase::ConfigureIndicator(bool bVolumetric)
 }
 
 
+void UMeshSculptToolBase::SetIndicatorVisibility(bool bVisible)
+{
+	if (GetIndicatorVisibility() != bVisible)
+	{
+		BrushIndicator->bVisible = bVisible;
+		BrushIndicatorMesh->SetVisible(bVisible && bIsVolumetricIndicator);
+	}
+}
+
+bool UMeshSculptToolBase::GetIndicatorVisibility() const
+{
+	return BrushIndicator->bVisible;
+}
+
 UPreviewMesh* UMeshSculptToolBase::MakeDefaultIndicatorSphereMesh(UObject* Parent, UWorld* World, int Resolution /*= 32*/)
 {
 	UPreviewMesh* SphereMesh = NewObject<UPreviewMesh>(Parent);
diff --git a/Engine/Plugins/Experimental/MeshModelingToolset/Source/MeshModelingTools/Public/MeshGroupPaintTool.h b/Engine/Plugins/Experimental/MeshModelingToolset/Source/MeshModelingTools/Public/MeshGroupPaintTool.h
--- a/Engine/Plugins/Experimental/MeshModelingToolset/Source/MeshModelingTools/Public/MeshGroupPaintTool.h
+++ b/Engine/Plugins/Experimental/MeshModelingToolset/Source/MeshModelingTools/Public/MeshGroupPaintTool.h
@@ -8,6 +8,7 @@
 #include "BaseTools/BaseBrushTool.h"
 #include "SimpleDynamicMeshComponent.h"
 #include "PropertySets/PolygroupLayersProperties.h"
+#include "Mechanics/PolyLassoMarqueeMechanic.h"
 
 #include "Sculpting/MeshSculptToolBase.h"
 #include "Sculpting/MeshBrushOpBase.h"
@@ -73,20 +74,6 @@ enum class EMeshGroupPaintInteractionType : uint8
 
 
 
-UCLASS()
-class MESHMODELINGTOOLS_API UGroupPaintToolProperties : public UInteractiveToolPropertySet
-{
-	GENERATED_BODY()
-
-public:
-	//UPROPERTY(EditAnywhere, Category = "GroupPaint")
-	UPROPERTY()
-	EMeshGroupPaintInteractionType SubToolType = EMeshGroupPaintInteractionType::Brush;
-};
-
-
-
-
 
 /** Mesh Sculpting Brush Types */
 UENUM()
@@ -102,6 +89,14 @@ enum class EMeshGroupPaintBrushType : uint8
 };
 
 
+/** Mesh Sculpting Brush Area Types */
+UENUM()
+enum class EMeshGroupPaintBrushAreaType : uint8
+{
+	Connected,
+	Volumetric
+};
+
 /** Mesh Sculpting Brush Types */
 UENUM()
 enum class EMeshGroupPaintVisibilityType : uint8
@@ -125,21 +120,57 @@ public:
 	UPROPERTY()
 	EMeshGroupPaintBrushType PrimaryBrushType = EMeshGroupPaintBrushType::Paint;
 
-	UPROPERTY(EditAnywhere, Category = BrushFilters, meta = (DisplayName = "Volumetric"))
-	bool bVolumetric = false;
+	UPROPERTY(EditAnywhere, Category = ActionType, meta = (DisplayName = "Action"))
+	EMeshGroupPaintInteractionType SubToolType = EMeshGroupPaintInteractionType::Brush;
+
+	/** Relative size of brush */
+	UPROPERTY(EditAnywhere, Category = ActionType, meta = (DisplayName = "Brush Size", UIMin = "0.0", UIMax = "1.0", ClampMin = "0.0", ClampMax = "10.0", 
+		HideEditConditionToggle, EditConditionHides, EditCondition = "SubToolType == EMeshGroupPaintInteractionType::Brush"))
+	float BrushSize;
+
+	/** When Volumetric, all faces inside the brush sphere are selected, otherwise only connected faces are selected */
+	UPROPERTY(EditAnywhere, Category = ActionType, meta = (DisplayName = "Brush Area Mode",
+		HideEditConditionToggle, EditConditionHides, EditCondition = "SubToolType == EMeshGroupPaintInteractionType::Brush"))
+	EMeshGroupPaintBrushAreaType BrushAreaMode = EMeshGroupPaintBrushAreaType::Connected;
+
+	/** The group that will be assigned to triangles */
+	UPROPERTY(EditAnywhere, Category = ActionType)
+	int32 SetGroup = 1;
 
-	UPROPERTY(EditAnywhere, Category = BrushFilters)
+	/** If true, only triangles with no group assigned will be painted */
+	UPROPERTY(EditAnywhere, Category = ActionType)
+	bool bOnlySetUngrouped = false;
+
+	/** Group to set as Erased value */
+	UPROPERTY(EditAnywhere, Category = ActionType)
+	int32 EraseGroup = 0;
+
+	/** When enabled, only the current group configured in the Paint brush is erased */
+	UPROPERTY(EditAnywhere, Category = ActionType)
+	bool bOnlyEraseCurrent = false;
+
+	
+	/** Control which triangles can be affected by the current operation based on visibility */
+	UPROPERTY(EditAnywhere, Category = Filters)
 	EMeshGroupPaintVisibilityType VisibilityFilter = EMeshGroupPaintVisibilityType::None;
 
-	UPROPERTY(EditAnywhere, Category = BrushFilters, meta = (UIMin = "0.0", UIMax = "180.0", EditCondition = "bVolumetric == false"))
+	/** The Region affected by the current operation will be bounded by edge angles larger than this threshold */
+	UPROPERTY(EditAnywhere, Category = Filters, meta = (UIMin = "0.0", UIMax = "180.0",
+		EditCondition = "bVolumetric == false"))
 	float AngleThreshold = 180.0f;
 
-	UPROPERTY(EditAnywhere, Category = BrushFilters, meta = (UIMin = "0.0", UIMax = "180.0", EditCondition = "bVolumetric == false"))
+	/** The Region affected by the current operation will be bounded by UV borders/seams */
+	UPROPERTY(EditAnywhere, Category = Filters, meta = (EditCondition = "bVolumetric == false"))
 	bool bUVSeams = false;
 
-	UPROPERTY(EditAnywhere, Category = BrushFilters, meta = (UIMin = "0.0", UIMax = "180.0", EditCondition = "bVolumetric == false"))
+	/** The Region affected by the current operation will be bounded by Hard Normal edges/seams */
+	UPROPERTY(EditAnywhere, Category = Filters, meta = (EditCondition = "bVolumetric == false"))
 	bool bNormalSeams = false;
 
+
+	/** Number of vertices in a triangle the Lasso must hit to be counted as "inside" */
+	UPROPERTY(EditAnywhere, Category = Filters, AdvancedDisplay, meta = (UIMin = 1, UIMax = 3, EditCondition = "SubToolType == EMeshGroupPaintInteractionType::PolyLasso"))
+	int MinTriVertCount = 1;
 };
 
 
@@ -156,7 +187,10 @@ enum class EMeshGroupPaintToolActions
 	FreezeOthers,
 
 	GrowCurrent,
-	ShrinkCurrent
+	ShrinkCurrent,
+	ClearCurrent,
+	FloodFillCurrent,
+	ClearAll
 };
 
 
@@ -199,17 +233,38 @@ public:
 		PostAction(EMeshGroupPaintToolActions::FreezeOthers);
 	}
 
-	UFUNCTION(CallInEditor, Category = Operations, meta = (DisplayPriority = 4))
+	UFUNCTION(CallInEditor, Category = Operations, meta = (DisplayPriority = 10))
+	void ClearAll()
+	{
+		PostAction(EMeshGroupPaintToolActions::ClearAll);
+	}
+
+	UFUNCTION(CallInEditor, Category = Operations, meta = (DisplayPriority = 11))
+	void ClearCurrent()
+	{
+		PostAction(EMeshGroupPaintToolActions::ClearCurrent);
+	}
+
+	UFUNCTION(CallInEditor, Category = Operations, meta = (DisplayPriority = 12))
+	void FloodFillCurrent()
+	{
+		PostAction(EMeshGroupPaintToolActions::FloodFillCurrent);
+	}
+
+
+	UFUNCTION(CallInEditor, Category = Operations, meta = (DisplayPriority = 20))
 	void GrowCurrent()
 	{
 		PostAction(EMeshGroupPaintToolActions::GrowCurrent);
 	}
 
-	UFUNCTION(CallInEditor, Category = Operations, meta = (DisplayPriority = 5))
+	UFUNCTION(CallInEditor, Category = Operations, meta = (DisplayPriority = 21))
 	void ShrinkCurrent()
 	{
 		PostAction(EMeshGroupPaintToolActions::ShrinkCurrent);
 	}
+
+
 };
 
 
@@ -228,6 +283,7 @@ public:
 
 	virtual void Setup() override;
 	virtual void Shutdown(EToolShutdownType ShutdownType) override;
+	virtual void DrawHUD(FCanvas* Canvas, IToolsContextRenderAPI* RenderAPI) override;
 
 	virtual void OnTick(float DeltaTime) override;
 
@@ -239,14 +295,13 @@ public:
 
 	virtual void OnPropertyModified(UObject* PropertySet, FProperty* Property) override;
 
+	bool IsInBrushSubMode() const;
+
 public:
 
 	UPROPERTY()
 	UPolygroupLayersProperties* PolygroupLayerProperties;
 
-	UPROPERTY()
-	UGroupPaintToolProperties* ToolProperties;
-
 	/** Filters on paint brush */
 	UPROPERTY()
 	UGroupPaintBrushFilterProperties* FilterProperties;
@@ -255,15 +310,24 @@ public:
 private:
 	// This will be of type UGroupPaintBrushOpProps, we keep a ref so we can change active group ID on pick
 	UPROPERTY()
-	UGroupPaintBrushOpProps* PaintBrushOpOperties;
+	UGroupPaintBrushOpProps* PaintBrushOpProperties;
 
 	UPROPERTY()
-	UGroupEraseBrushOpProps* EraseBrushOpOperties;
+	UGroupEraseBrushOpProps* EraseBrushOpProperties;
 
 public:
 	void AllocateNewGroupAndSetAsCurrentAction();
 	void GrowCurrentGroupAction();
 	void ShrinkCurrentGroupAction();
+	void ClearCurrentGroupAction();
+	void FloodFillCurrentGroupAction();
+	void ClearAllGroupsAction();
+
+	void SetTrianglesToGroupID(const TSet<int32>& Triangles, int32 ToGroupID, bool bIsErase);
+
+	bool HaveVisibilityFilter() const;
+	void ApplyVisibilityFilter(const TArray<int32>& Triangles, TArray<int32>& VisibleTriangles);
+	void ApplyVisibilityFilter(TSet<int32>& Triangles, TArray<int32>& ROIBuffer, TArray<int32>& OutputBuffer);
 
 protected:
 	// UMeshSculptToolBase API
@@ -296,6 +360,18 @@ protected:
 	virtual void ApplyAction(EMeshGroupPaintToolActions ActionType);
 
 
+
+	//
+	// Marquee Support
+	//
+public:
+	UPROPERTY()
+	UPolyLassoMarqueeMechanic* PolyLassoMechanic;
+
+protected:
+	void OnPolyLassoFinished(const FCameraPolyLasso& Lasso, bool bCanceled);
+
+
 	//
 	// Internals
 	//
@@ -326,6 +402,7 @@ protected:
 
 	TArray<int> TempROIBuffer;
 	TArray<int> VertexROI;
+	TArray<bool> VisibilityFilterBuffer;
 	TSet<int> VertexSetBuffer;
 	TSet<int> TriangleROI;
 	void UpdateROI(const FSculptBrushStamp& CurrentStamp);
@@ -368,6 +445,8 @@ protected:
 
 protected:
 	virtual bool ShowWorkPlane() const override { return false; }
+
+
 };
 
 
diff --git a/Engine/Plugins/Experimental/MeshModelingToolset/Source/MeshModelingTools/Public/Sculpting/MeshSculptToolBase.h b/Engine/Plugins/Experimental/MeshModelingToolset/Source/MeshModelingTools/Public/Sculpting/MeshSculptToolBase.h
--- a/Engine/Plugins/Experimental/MeshModelingToolset/Source/MeshModelingTools/Public/Sculpting/MeshSculptToolBase.h
+++ b/Engine/Plugins/Experimental/MeshModelingToolset/Source/MeshModelingTools/Public/Sculpting/MeshSculptToolBase.h
@@ -58,12 +58,16 @@ public:
 	bool bHitBackFaces = true;
 
 
-	UPROPERTY(EditAnywhere, Category = Brush, meta = (UIMin = "0.0", UIMax = "1.0", ClampMin = "0.0", ClampMax = "1.0", DisplayPriority = 7))
+	UPROPERTY(EditAnywhere, Category = Brush, meta = (UIMin = "0.0", UIMax = "1.0", ClampMin = "0.0", ClampMax = "1.0", DisplayPriority = 7, HideEditConditionToggle, EditConditionHides, EditCondition = "bShowLazyness"))
 	float Lazyness = 0;
 
 	/**  */
 	UPROPERTY( meta = (TransientToolProperty))
 	bool bShowPerBrushProps = true;
+
+	/**  */
+	UPROPERTY(meta = (TransientToolProperty))
+	bool bShowLazyness = true;
 };
 
 
@@ -238,6 +242,8 @@ private:
 
 	TUniquePtr<FMeshSculptBrushOp> SecondaryBrushOp;
 	UMeshSculptBrushOpProps* SecondaryVisiblePropSet = nullptr;
+
+	bool bBrushOpPropsVisible = true;
 protected:
 	void SetActivePrimaryBrushType(int32 Identifier);
 	void SetActiveSecondaryBrushType(int32 Identifier);
@@ -406,6 +412,9 @@ protected:
 
 	virtual void ConfigureIndicator(bool bVolumetric);
 
+	virtual void SetIndicatorVisibility(bool bVisible);
+	virtual bool GetIndicatorVisibility() const;
+
 protected:
 	UPROPERTY()
 	UBrushStampIndicator* BrushIndicator;
diff --git a/Engine/Plugins/Experimental/MeshModelingToolset/Source/ModelingComponents/Private/Mechanics/PolyLassoMarqueeMechanic.cpp b/Engine/Plugins/Experimental/MeshModelingToolset/Source/ModelingComponents/Private/Mechanics/PolyLassoMarqueeMechanic.cpp
new file mode 100644
--- /dev/null
+++ b/Engine/Plugins/Experimental/MeshModelingToolset/Source/ModelingComponents/Private/Mechanics/PolyLassoMarqueeMechanic.cpp
@@ -0,0 +1,306 @@
+// Copyright Epic Games, Inc. All Rights Reserved.
+
+#include "Mechanics/PolyLassoMarqueeMechanic.h"
+#include "BaseBehaviors/ClickDragBehavior.h"
+#include "BaseBehaviors/MouseHoverBehavior.h"
+#include "CanvasTypes.h"
+#include "CanvasItem.h"
+#include "ToolSceneQueriesUtil.h"
+#include "TransformTypes.h"
+
+#define LOCTEXT_NAMESPACE "UPolyLassoMarqueeMechanic"
+
+
+
+FCameraPolyLasso::FCameraPolyLasso(const FViewCameraState& CachedCameraState)
+{
+	// Create a plane just in front of the camera
+	CameraOrigin = CachedCameraState.Position;
+	CameraPlane = FPlane(CachedCameraState.Position + CachedCameraState.Forward(), CachedCameraState.Forward());
+	bCameraIsOrthographic = CachedCameraState.bIsOrthographic;
+
+	// Intersect the drag rays with the camera plane and compute their coordinates in the camera basis
+	UBasisVector = CachedCameraState.Right();
+	VBasisVector = CachedCameraState.Up();
+}
+
+
+void FCameraPolyLasso::AddPoint(const FRay& WorldRay)
+{
+	FVector Intersection = FMath::RayPlaneIntersection(WorldRay.Origin, WorldRay.Direction, CameraPlane);
+	Polyline.Add(PlaneCoordinates(Intersection));
+}
+
+
+FVector2D FCameraPolyLasso::GetProjectedPoint(const FVector& Point) const
+{
+	FVector ProjectedPoint;
+	if (bCameraIsOrthographic)
+	{
+		// project directly to plane
+		ProjectedPoint = FVector::PointPlaneProject(Point, CameraPlane);
+	}
+	else
+	{
+		// intersect along the eye-to-point ray
+		ProjectedPoint = FMath::RayPlaneIntersection(CameraOrigin,
+			Point - CameraOrigin,
+			CameraPlane);
+	}
+
+	return PlaneCoordinates(ProjectedPoint);
+}
+
+
+
+// ---------------------------------------
+
+void UPolyLassoMarqueeMechanic::Setup(UInteractiveTool* ParentToolIn)
+{
+	UInteractionMechanic::Setup(ParentToolIn);
+
+	ClickDragBehavior = NewObject<UClickDragInputBehavior>(this);
+	ClickDragBehavior->SetDefaultPriority(BasePriority);
+	ClickDragBehavior->Initialize(this);
+	ParentTool->AddInputBehavior(ClickDragBehavior, this);
+
+	HoverBehavior = NewObject<UMouseHoverBehavior>();
+	HoverBehavior->SetDefaultPriority(BasePriority);
+	HoverBehavior->Initialize(this);
+	ParentTool->AddInputBehavior(HoverBehavior, this);
+
+	SetIsEnabled(true);
+}
+
+bool UPolyLassoMarqueeMechanic::IsEnabled() const
+{
+	return bIsEnabled;
+}
+
+void UPolyLassoMarqueeMechanic::SetIsEnabled(bool bOn)
+{
+	if ((bIsDragging || bIsInMultiClickPolygon) && !bOn)
+	{
+		OnTerminateDragSequence();
+	}
+
+	bIsEnabled = bOn;
+}
+
+void UPolyLassoMarqueeMechanic::SetBasePriority(const FInputCapturePriority& Priority)
+{
+	BasePriority = Priority;
+	if (ClickDragBehavior)
+	{
+		ClickDragBehavior->SetDefaultPriority(Priority);
+	}
+}
+
+TPair<FInputCapturePriority, FInputCapturePriority> UPolyLassoMarqueeMechanic::GetPriorityRange() const
+{
+	return TPair<FInputCapturePriority, FInputCapturePriority>(BasePriority, BasePriority);
+}
+
+
+
+FInputRayHit UPolyLassoMarqueeMechanic::CanBeginClickDragSequence(const FInputDeviceRay& PressPos)
+{
+	if (bIsEnabled && (bEnableFreehandPolygons || bEnableMultiClickPolygons))
+	{
+		return FInputRayHit(TNumericLimits<float>::Max()); // bHit is true. Depth is max to lose the standard tiebreaker.
+	}
+	else
+	{
+		return FInputRayHit(); // bHit is false
+	}
+}
+
+
+void UPolyLassoMarqueeMechanic::OnClickPress(const FInputDeviceRay& PressPos)
+{
+	if (!PressPos.bHas2D)
+	{
+		bIsDragging = false;
+		return;
+	}
+
+	if (bIsInMultiClickPolygon)
+	{
+		// nothing to do on click-down
+		DragCurrentScreenPosition = PressPos.ScreenPosition;
+	}
+	else
+	{
+		PolyPathPoints.Reset();
+		PolyPathPoints.Add(PressPos);
+		DragCurrentScreenPosition = PressPos.ScreenPosition;
+
+		GetParentTool()->GetToolManager()->GetContextQueriesAPI()->GetCurrentViewState(CachedCameraState);
+		CurrentPolyLasso = FCameraPolyLasso(CachedCameraState);
+		CurrentPolyLasso.AddPoint(PressPos.WorldRay);
+
+		if (bEnableFreehandPolygons)
+		{
+			OnDrawPolyLassoStarted.Broadcast();
+		}
+	}
+}
+
+void UPolyLassoMarqueeMechanic::OnClickDrag(const FInputDeviceRay& DragPos)
+{
+	if (!DragPos.bHas2D)
+	{
+		return;
+	}
+
+	if (bIsInMultiClickPolygon)
+	{
+		DragCurrentScreenPosition = DragPos.ScreenPosition;
+	}
+	else
+	{
+
+		bIsDragging = true;
+		DragCurrentScreenPosition = DragPos.ScreenPosition;
+
+		if (bEnableFreehandPolygons)
+		{
+			float DistSqr = FVector2D::DistSquared(DragPos.ScreenPosition, PolyPathPoints.Last().ScreenPosition);
+			if (DistSqr > SpacingTolerance * SpacingTolerance)
+			{
+				PolyPathPoints.Add(DragPos);
+				CurrentPolyLasso.AddPoint(DragPos.WorldRay);
+			}
+
+			OnDrawPolyLassoChanged.Broadcast(CurrentPolyLasso);
+		}
+	}
+}
+
+void UPolyLassoMarqueeMechanic::OnClickRelease(const FInputDeviceRay& ReleasePos)
+{
+	// if we are in multiclick draw polygon, we either close poly or continue  
+	// (could also check for intersection?)
+	if (bIsInMultiClickPolygon)
+	{
+		if (WillCloseCurrentLasso(ReleasePos))
+		{
+			CurrentPolyLasso.AddPoint(PolyPathPoints[0].WorldRay);
+			OnDrawPolyLassoFinished.Broadcast(CurrentPolyLasso, false);
+			bIsInMultiClickPolygon = false;
+		}
+		else
+		{
+			PolyPathPoints.Add(ReleasePos);
+			CurrentPolyLasso.AddPoint(ReleasePos.WorldRay);
+			OnDrawPolyLassoChanged.Broadcast(CurrentPolyLasso);
+		}
+	
+	} 
+	else
+	{
+		bIsDragging = false;
+
+		// if we clicked down and up in (approximately) the same spot, begin multi-click polygon mode
+		if (PolyPathPoints.Num() == 1 && this->bEnableMultiClickPolygons)
+		{
+			float DistSqr = FVector2D::DistSquared(ReleasePos.ScreenPosition, PolyPathPoints.Last().ScreenPosition);
+			if (DistSqr < SpacingTolerance * SpacingTolerance)
+			{
+				bIsInMultiClickPolygon = true;
+				return;
+			}
+		}
+
+		if (bEnableFreehandPolygons)
+		{
+			OnDrawPolyLassoFinished.Broadcast(CurrentPolyLasso, false);
+		}
+	}
+}
+
+void UPolyLassoMarqueeMechanic::OnTerminateDragSequence()
+{
+	bIsDragging = false;
+	if (bIsInMultiClickPolygon)
+	{
+		bIsInMultiClickPolygon = false;
+		// do not emit lasso?
+	}
+	else
+	{
+		OnDrawPolyLassoFinished.Broadcast(CurrentPolyLasso, true);
+	}
+}
+
+
+
+FInputRayHit UPolyLassoMarqueeMechanic::BeginHoverSequenceHitTest(const FInputDeviceRay& PressPos)
+{
+	if (bIsInMultiClickPolygon)
+	{
+		// abort multiclick polygon if camera moves
+		FViewCameraState CurCameraState;
+		GetParentTool()->GetToolManager()->GetContextQueriesAPI()->GetCurrentViewState(CurCameraState);
+		if (CurCameraState.Position != CachedCameraState.Position || CurCameraState.Orientation != CachedCameraState.Orientation)
+		{
+			OnTerminateDragSequence();
+			return FInputRayHit();
+		}
+
+		return FInputRayHit(0);
+	}
+	return FInputRayHit();
+}
+
+
+bool UPolyLassoMarqueeMechanic::OnUpdateHover(const FInputDeviceRay& DevicePos)
+{
+	if (bIsInMultiClickPolygon)
+	{
+		DragCurrentScreenPosition = DevicePos.ScreenPosition;
+		bIsMultiClickPolygonClosed = WillCloseCurrentLasso(DevicePos);
+		return true;
+	}
+	return false;
+}
+
+
+bool UPolyLassoMarqueeMechanic::WillCloseCurrentLasso(const FInputDeviceRay& DevicePos) const
+{
+	float UseTol = 2.0 * SpacingTolerance;
+	float DistSqr = FVector2D::DistSquared(DevicePos.ScreenPosition, PolyPathPoints[0].ScreenPosition);
+	return (DistSqr < UseTol* UseTol);
+}
+
+
+void UPolyLassoMarqueeMechanic::DrawHUD(FCanvas* Canvas, IToolsContextRenderAPI* RenderAPI)
+{
+	EViewInteractionState State = RenderAPI->GetViewInteractionState();
+	bool bThisViewHasFocus = !!(State & EViewInteractionState::Focused);
+	if (bThisViewHasFocus && (bIsDragging || bIsInMultiClickPolygon) )
+	{
+		FLinearColor UseColor = (bIsInMultiClickPolygon && bIsMultiClickPolygonClosed) ? this->ClosedColor : this->LineColor;
+
+		float DPIScale = Canvas->GetDPIScale();
+		int32 NumPoints = PolyPathPoints.Num();
+		FVector2D LastPoint = PolyPathPoints[0].ScreenPosition / DPIScale;
+		for (int32 k = 1; k < NumPoints; ++k)
+		{
+			FVector2D NextPoint = PolyPathPoints[k].ScreenPosition / DPIScale;
+			FCanvasLineItem LineItem(LastPoint, NextPoint);
+			LineItem.LineThickness = this->LineThickness;
+			LineItem.SetColor(UseColor);
+			Canvas->DrawItem(LineItem);
+			LastPoint = NextPoint;
+		}
+
+		FVector2D CurPoint = DragCurrentScreenPosition / DPIScale;
+		FCanvasLineItem CurLineItem(LastPoint, CurPoint);
+		CurLineItem.LineThickness = this->LineThickness;
+		CurLineItem.SetColor(UseColor);
+		Canvas->DrawItem(CurLineItem);
+	}
+}
+
+#undef LOCTEXT_NAMESPACE
diff --git a/Engine/Plugins/Experimental/MeshModelingToolset/Source/ModelingComponents/Private/Mechanics/RectangleMarqueeMechanic.cpp b/Engine/Plugins/Experimental/MeshModelingToolset/Source/ModelingComponents/Private/Mechanics/RectangleMarqueeMechanic.cpp
--- a/Engine/Plugins/Experimental/MeshModelingToolset/Source/ModelingComponents/Private/Mechanics/RectangleMarqueeMechanic.cpp
+++ b/Engine/Plugins/Experimental/MeshModelingToolset/Source/ModelingComponents/Private/Mechanics/RectangleMarqueeMechanic.cpp
@@ -10,13 +10,6 @@
 #define LOCTEXT_NAMESPACE "URectangleMarqueeMechanic"
 
 
-static FVector2D PlaneCoordinates(const FVector& Point, const FPlane& Plane, const FVector& UBasisVector, const FVector& VBasisVector)
-{
-	float U = FVector::DotProduct(Point - Plane.GetOrigin(), UBasisVector);
-	float V = FVector::DotProduct(Point - Plane.GetOrigin(), VBasisVector);
-	return FVector2D{ U,V };
-}
-
 FCameraRectangle::FCameraRectangle(const FViewCameraState& CachedCameraState,
 								   const FRay& DragStartWorldRay,
 								   const FRay& DragEndWorldRay)
@@ -33,19 +26,13 @@ FCameraRectangle::FCameraRectangle(const FViewCameraState& CachedCameraState,
 	FVector StartIntersection = FMath::RayPlaneIntersection(DragStartWorldRay.Origin,
 															DragStartWorldRay.Direction,
 															CameraPlane);
-	FVector2D Start2D = PlaneCoordinates(StartIntersection,
-										 CameraPlane,
-										 UBasisVector,
-										 VBasisVector);
+	FVector2D Start2D = PlaneCoordinates(StartIntersection);
 
 	FVector CurrentIntersection = FMath::RayPlaneIntersection(DragEndWorldRay.Origin,
 															  DragEndWorldRay.Direction,
 															  CameraPlane);
 
-	FVector2D Current2D = PlaneCoordinates(CurrentIntersection,
-										   CameraPlane,
-										   UBasisVector,
-										   VBasisVector);
+	FVector2D Current2D = PlaneCoordinates(CurrentIntersection);
 
 	RectangleCorners = FBox2D(Start2D, Start2D);
 	RectangleCorners += Current2D;	// Initialize this way so we don't have to care about min/max
@@ -67,7 +54,7 @@ bool FCameraRectangle::IsProjectedPointInRectangle(const FVector& Point) const
 													 CameraPlane);
 	}
 
-	FVector2D Point2D = PlaneCoordinates(ProjectedPoint, CameraPlane, UBasisVector, VBasisVector);
+	FVector2D Point2D = PlaneCoordinates(ProjectedPoint);
 	return RectangleCorners.IsInside(Point2D);
 }
 
@@ -78,8 +65,8 @@ bool FCameraRectangle::IsProjectedSegmentIntersectingRectangle(const FVector& En
 	FVector ProjectedEndpoint2 = bCameraIsOrthographic ? FVector::PointPlaneProject(Endpoint2, CameraPlane)
 		: FMath::RayPlaneIntersection(CameraOrigin, Endpoint2 - CameraOrigin, CameraPlane);
 
-	FVector2D Endpoint1PlaneCoord = PlaneCoordinates(ProjectedEndpoint1, CameraPlane, UBasisVector, VBasisVector);
-	FVector2D Endpoint2PlaneCoord = PlaneCoordinates(ProjectedEndpoint2, CameraPlane, UBasisVector, VBasisVector);
+	FVector2D Endpoint1PlaneCoord = PlaneCoordinates(ProjectedEndpoint1);
+	FVector2D Endpoint2PlaneCoord = PlaneCoordinates(ProjectedEndpoint2);
 
 	// If either endpoint is inside, then definitely (at least partially) contained
 	if (RectangleCorners.IsInside(Endpoint1PlaneCoord) || RectangleCorners.IsInside(Endpoint2PlaneCoord))
diff --git a/Engine/Plugins/Experimental/MeshModelingToolset/Source/ModelingComponents/Public/Mechanics/PolyLassoMarqueeMechanic.h b/Engine/Plugins/Experimental/MeshModelingToolset/Source/ModelingComponents/Public/Mechanics/PolyLassoMarqueeMechanic.h
new file mode 100644
--- /dev/null
+++ b/Engine/Plugins/Experimental/MeshModelingToolset/Source/ModelingComponents/Public/Mechanics/PolyLassoMarqueeMechanic.h
@@ -0,0 +1,185 @@
+// Copyright Epic Games, Inc. All Rights Reserved.
+
+#pragma once
+
+#include "CoreMinimal.h"
+#include "InputBehavior.h"
+#include "InteractionMechanic.h"
+#include "BaseBehaviors/BehaviorTargetInterfaces.h"
+#include "TransformTypes.h"
+#include "VectorTypes.h"
+
+#include "PolyLassoMarqueeMechanic.generated.h"
+
+class UClickDragInputBehavior;
+class UMouseHoverBehavior;
+
+/// Struct containing:
+///		- camera information, 
+///		- a 3D plane just in front of the camera,
+///		- a 2D basis for coordinates in this plane, and
+///		- the vertices of a PolyLasso contained in this plane, in this 2D basis
+/// 
+struct MODELINGCOMPONENTS_API FCameraPolyLasso
+{
+	FVector CameraOrigin;
+	bool bCameraIsOrthographic;
+	FPlane CameraPlane;
+	FVector UBasisVector;
+	FVector VBasisVector;
+	TArray<FVector2D> Polyline;
+
+	FCameraPolyLasso() {};
+	FCameraPolyLasso(const FViewCameraState& CachedCameraState);
+
+	/** Append a point to the PolyLasso polyine using a world-space ray */
+	void AddPoint(const FRay& WorldRay);
+
+	/** Project a world-space point into the same space as the 2D PolyLine */
+	FVector2D GetProjectedPoint(const FVector& Point) const;
+
+	FVector2D PlaneCoordinates(const FVector& Point) const
+	{
+		float U = FVector::DotProduct(Point - CameraPlane.GetOrigin(), UBasisVector);
+		float V = FVector::DotProduct(Point - CameraPlane.GetOrigin(), VBasisVector);
+		return FVector2D{ U,V };
+	}
+};
+
+
+/*
+ * Mechanic for a PolyLasso "marquee" selection. It creates and maintains the 2D PolyLasso associated with a mouse drag. 
+ * It does not test against any scene geometry, nor does it maintain any sort of list of selected objects.
+ *
+ * The PolyLasso has two potential modes, a freehand polyline and a multi-click polygon. By default both are enabled
+ * but this can be selectively controlled with flags below. If a click-and-release is within a small distance tolerance,
+ * then a multi-click polygon is entered, and must be exited by clicking again at the start point. The freehand polyline
+ * is drawn by click-dragging, and exited by releasing the mouse.
+ *
+ * When using this mechanic, you should call DrawHUD() in the tool's DrawHUD() call so that it can draw the box.
+ *
+ * Attach to the mechanic's delegates and use the passed PolyLasso to test against your geometry. 
+ */
+UCLASS()
+class MODELINGCOMPONENTS_API UPolyLassoMarqueeMechanic : public UInteractionMechanic, public IClickDragBehaviorTarget, public IHoverBehaviorTarget
+{
+	GENERATED_BODY()
+
+public:
+
+	// UInteractionMechanic
+	virtual void Setup(UInteractiveTool* ParentTool) override;
+
+	virtual void DrawHUD(FCanvas* Canvas, IToolsContextRenderAPI* RenderAPI);
+
+	/** @return true if PolyLasso mechanic is enabled */
+	bool IsEnabled() const;
+
+	/** enable/disable the PolyLasso mechanic (effectively controls whether the mechanic's behaviors will be allowed to capture the mouse) */
+	void SetIsEnabled(bool bOn);
+
+	/**
+	 * Sets the base priority so that users can make sure that their own behaviors are higher
+	 * priority. The mechanic will not use any priority value higher than this.
+	 * Mechanics could use lower priorities (and their range could be inspected with
+	 * GetPriorityRange), but marquee mechanic doesn't.
+	 *
+	 * Can be called before or after Setup().
+	 */
+	void SetBasePriority(const FInputCapturePriority& Priority);
+
+	/**
+	 * Gets the current priority range used by behaviors in the mechanic, higher
+	 * priority to lower.
+	 *
+	 * For marquee mechanic, the range will be [BasePriority, BasePriority] since
+	 * it only uses one priority.
+	 */
+	TPair<FInputCapturePriority, FInputCapturePriority> GetPriorityRange() const;
+
+	/**
+	 * Called when user starts dragging a new PolyLasso.
+	 */
+	FSimpleMulticastDelegate OnDrawPolyLassoStarted;
+
+	/**
+	 * Called as the user drags the other corner of the PolyLasso around.
+	 */
+	DECLARE_MULTICAST_DELEGATE_OneParam(OnDrawPolyLassoChangedEvent, const FCameraPolyLasso&);
+	OnDrawPolyLassoChangedEvent OnDrawPolyLassoChanged;
+
+	/**
+	 * Called once the user lets go of the mouse button after dragging out a PolyLasso.
+	 */
+	DECLARE_MULTICAST_DELEGATE_TwoParams(OnDrawPolyLassoFinishedEvent, const FCameraPolyLasso&, bool bCanceled);
+	OnDrawPolyLassoFinishedEvent OnDrawPolyLassoFinished;
+
+
+	/** Tolerance for PolyLasso points and closure test */
+	UPROPERTY()
+	float SpacingTolerance = 2.0f;
+
+	/** Thickness of the 2D PolyLasso drawing path */
+	UPROPERTY()
+	float LineThickness = 2.0f;
+
+	/** PolyLasso path is drawn in this color */
+	UPROPERTY()
+	FLinearColor LineColor = FLinearColor::Red;
+
+	/** PolyLasso path is drawn in this color if the cursor is at a point that would close the loop */
+	UPROPERTY()
+	FLinearColor ClosedColor = FLinearColor(1.0f, 0.549019607843f, 0.0f);
+
+	/**
+	 * If true, freehand polygons can be drawn by click-dragging the mouse
+	 */
+	UPROPERTY()
+	bool bEnableFreehandPolygons = true;
+
+	/**
+	 * If true, if click and release are within SpacingTolerance, mechanic enters a multi-click mode,
+	 * where each click adds a vertex to a polygon. Polygon must be closed by clicking within 2*SpacingTolerance
+	 * of the initial point to exit this mode. However, rotating the camera will also cancel out of the interaction.
+	 */
+	UPROPERTY()
+	bool bEnableMultiClickPolygons = true;
+
+
+protected:
+	UPROPERTY()
+	UClickDragInputBehavior* ClickDragBehavior = nullptr;
+
+	UPROPERTY()
+	UMouseHoverBehavior* HoverBehavior = nullptr;
+
+	// camera state for currently active PolyLasso
+	FViewCameraState CachedCameraState;
+
+	FInputCapturePriority BasePriority = FInputCapturePriority(FInputCapturePriority::DEFAULT_TOOL_PRIORITY);
+
+private:
+
+	// IClickDragBehaviorTarget implementation
+	FInputRayHit CanBeginClickDragSequence(const FInputDeviceRay& PressPos) final;
+	void OnClickPress(const FInputDeviceRay& PressPos) final;
+	void OnClickDrag(const FInputDeviceRay& DragPos) final;
+	void OnClickRelease(const FInputDeviceRay& ReleasePos) final;
+	void OnTerminateDragSequence() final;
+
+	// IHoverBehaviorTarget implementation
+	virtual FInputRayHit BeginHoverSequenceHitTest(const FInputDeviceRay& PressPos) override;
+	virtual void OnBeginHover(const FInputDeviceRay& DevicePos) override {}
+	virtual bool OnUpdateHover(const FInputDeviceRay& DevicePos) override;
+	virtual void OnEndHover() override {}
+
+	bool bIsEnabled = false;
+	bool bIsDragging = false;
+	bool bIsInMultiClickPolygon = false;
+	bool bIsMultiClickPolygonClosed = false;
+	TArray<FInputDeviceRay> PolyPathPoints;
+	FVector2D DragCurrentScreenPosition;
+
+	FCameraPolyLasso CurrentPolyLasso;
+	bool WillCloseCurrentLasso(const FInputDeviceRay& DevicePos) const;
+};
diff --git a/Engine/Plugins/Experimental/MeshModelingToolset/Source/ModelingComponents/Public/Mechanics/RectangleMarqueeMechanic.h b/Engine/Plugins/Experimental/MeshModelingToolset/Source/ModelingComponents/Public/Mechanics/RectangleMarqueeMechanic.h
--- a/Engine/Plugins/Experimental/MeshModelingToolset/Source/ModelingComponents/Public/Mechanics/RectangleMarqueeMechanic.h
+++ b/Engine/Plugins/Experimental/MeshModelingToolset/Source/ModelingComponents/Public/Mechanics/RectangleMarqueeMechanic.h
@@ -40,6 +40,13 @@ struct FCameraRectangle
 
 	// TODO: Add a way to test rectangle against triangles.
 
+	FVector2D PlaneCoordinates(const FVector& Point) const
+	{
+		float U = FVector::DotProduct(Point - CameraPlane.GetOrigin(), UBasisVector);
+		float V = FVector::DotProduct(Point - CameraPlane.GetOrigin(), VBasisVector);
+		return FVector2D{ U,V };
+	}
+
 private:
 	FCameraRectangle() {};
 };
diff --git a/Engine/Source/Runtime/Experimental/InteractiveToolsFramework/Private/BaseGizmos/BrushStampIndicator.cpp b/Engine/Source/Runtime/Experimental/InteractiveToolsFramework/Private/BaseGizmos/BrushStampIndicator.cpp
--- a/Engine/Source/Runtime/Experimental/InteractiveToolsFramework/Private/BaseGizmos/BrushStampIndicator.cpp
+++ b/Engine/Source/Runtime/Experimental/InteractiveToolsFramework/Private/BaseGizmos/BrushStampIndicator.cpp
@@ -24,6 +24,11 @@ void UBrushStampIndicator::Shutdown()
 
 void UBrushStampIndicator::Render(IToolsContextRenderAPI* RenderAPI)
 {
+	if (bVisible == false)
+	{
+		return;
+	}
+
 	if (bDrawIndicatorLines)
 	{
 		FToolDataVisualizer Draw;
diff --git a/Engine/Source/Runtime/Experimental/InteractiveToolsFramework/Public/BaseGizmos/BrushStampIndicator.h b/Engine/Source/Runtime/Experimental/InteractiveToolsFramework/Public/BaseGizmos/BrushStampIndicator.h
--- a/Engine/Source/Runtime/Experimental/InteractiveToolsFramework/Public/BaseGizmos/BrushStampIndicator.h
+++ b/Engine/Source/Runtime/Experimental/InteractiveToolsFramework/Public/BaseGizmos/BrushStampIndicator.h
@@ -44,6 +44,10 @@ public:
 
 public:
 
+	/** Controls whether Gizmo will draw visual elements. Does not currently affect AttachedComponent. */
+	UPROPERTY()
+	bool bVisible = true;
+
 	UPROPERTY()
 	float BrushRadius = 1.0f;
 
