When starting the Paint Maps tool, don't crash if there are no suitable vertex attributes on the mesh.

diff --git a/Engine/Plugins/Experimental/MeshModelingToolset/Source/MeshModelingTools/Private/MeshAttributePaintTool.cpp b/Engine/Plugins/Experimental/MeshModelingToolset/Source/MeshModelingTools/Private/MeshAttributePaintTool.cpp
--- a/Engine/Plugins/Experimental/MeshModelingToolset/Source/MeshModelingTools/Private/MeshAttributePaintTool.cpp
+++ b/Engine/Plugins/Experimental/MeshModelingToolset/Source/MeshModelingTools/Private/MeshAttributePaintTool.cpp
@@ -191,6 +191,14 @@ void UMeshAttributePaintTool::Setup()
 
 	AttributeSource = MakeUnique<FMeshDescriptionVertexAttributeSource>(EditedMesh.Get());
 	AttribProps->Attributes = AttributeSource->GetAttributeList();
+
+	if (AttribProps->Attributes.Num() == 0)
+	{
+		GetToolManager()->DisplayMessage(
+			LOCTEXT("StartAttribPaintFailed", "No Float attributes exist for this mesh. Use the Attribute Editor to create one."),
+			EToolMessageLevel::UserWarning);
+	}
+
 	InitializeAttributes();
 	PendingNewSelectedIndex = 0;
 
@@ -391,21 +399,25 @@ void UMeshAttributePaintTool::UpdateVisibleAttribute()
 	StoreCurrentAttribute();
 
 	CurrentAttributeIndex = AttribProps->SelectedAttribute;
-	FAttributeData& AttribData = Attributes[CurrentAttributeIndex];
-	CurrentValueRange = AttribData.Attribute->GetValueRange();
 
-	// update mesh with new value colors
-	PreviewMesh->EditMesh([&](FDynamicMesh3& Mesh)
+	if (CurrentAttributeIndex >= 0)
 	{
-		for (int32 vid : Mesh.VertexIndicesItr())
+		FAttributeData& AttribData = Attributes[CurrentAttributeIndex];
+		CurrentValueRange = AttribData.Attribute->GetValueRange();
+
+		// update mesh with new value colors
+		PreviewMesh->EditMesh([&](FDynamicMesh3& Mesh)
 		{
-			float Value = AttribData.CurrentValues[vid];
-			FVector3f Color = ColorMapper->ToColor(Value);
-			Mesh.SetVertexColor(vid, Color);
-		}
-	});
+			for (int32 vid : Mesh.VertexIndicesItr())
+			{
+				float Value = AttribData.CurrentValues[vid];
+				FVector3f Color = ColorMapper->ToColor(Value);
+				Mesh.SetVertexColor(vid, Color);
+			}
+		});
 
-	AttribProps->AttributeName = AttribData.Name.ToString();
+		AttribProps->AttributeName = AttribData.Name.ToString();
+	}
 }
 
 
@@ -436,6 +448,11 @@ double UMeshAttributePaintTool::CalculateBrushFalloff(double Distance)
 
 void UMeshAttributePaintTool::ApplyStamp(const FBrushStampData& Stamp)
 {
+	if (CurrentAttributeIndex < 0)
+	{
+		return;
+	}
+
 	FTransform3d Transform(ComponentTarget->GetWorldTransform());
 	FVector3d StampPosLocal = Transform.InverseTransformPosition(Stamp.WorldPosition);
 
