Geometry Tools: Mesh Simplification.  Expose option to only use existing vertices in simplification.

diff --git a/Engine/Plugins/Experimental/MeshModelingToolset/Source/ModelingOperatorsEditorOnly/Private/CleaningOps/SimplifyMeshOp.cpp b/Engine/Plugins/Experimental/MeshModelingToolset/Source/ModelingOperatorsEditorOnly/Private/CleaningOps/SimplifyMeshOp.cpp
--- a/Engine/Plugins/Experimental/MeshModelingToolset/Source/ModelingOperatorsEditorOnly/Private/CleaningOps/SimplifyMeshOp.cpp
+++ b/Engine/Plugins/Experimental/MeshModelingToolset/Source/ModelingOperatorsEditorOnly/Private/CleaningOps/SimplifyMeshOp.cpp
@@ -29,7 +29,8 @@ void ComputeSimplify(FDynamicMesh3* TargetMesh, const bool bReproject,
 					 bool bPreserveSharpEdges, bool bAllowSeamCollapse,
 					 const ESimplifyTargetType TargetMode,
 					 const float TargetPercentage, const int TargetCount, const float TargetEdgeLength,
-					 const float AngleThreshold)
+					 const float AngleThreshold,
+	                 typename SimplificationType::ESimplificationCollapseModes CollapseMode = typename SimplificationType::ESimplificationCollapseModes::MinimalQuadricPositionError)
 {
 	SimplificationType Reducer(TargetMesh);
 
@@ -65,6 +66,12 @@ void ComputeSimplify(FDynamicMesh3* TargetMesh, const bool bReproject,
 		Reducer.SetProjectionTarget(&ProjTarget);
 	}
 
+	// MinimalPlanar is a special path that ignores the collapse mode (and quadric error ), so don't bother setting it.
+	if (TargetMode != ESimplifyTargetType::MinimalPlanar)
+	{
+		Reducer.CollapseMode = CollapseMode;
+	}
+
 	if (TargetMode == ESimplifyTargetType::Percentage)
 	{
 		double Ratio = (double)TargetPercentage / 100.0;
@@ -138,6 +145,17 @@ void FSimplifyMeshOp::CalculateResult(FProgressCancel* Progress)
 			bPreserveSharpEdges, bAllowSeamCollapse,
 			ESimplifyTargetType::MinimalPlanar, TargetPercentage, TargetCount, TargetEdgeLength, MinimalPlanarAngleThresh);
 	}
+	else if (SimplifierType == ESimplifyType::MinimalExistingVertex)
+	{
+		ResultMesh->Copy(*OriginalMesh, true, true, true, !bDiscardAttributes);
+		ComputeSimplify<FQEMSimplification>(TargetMesh, bReproject, OriginalTriCount, *OriginalMesh, *OriginalMeshSpatial,
+			MeshBoundaryConstraint,
+			GroupBoundaryConstraint,
+			MaterialBoundaryConstraint,
+			bPreserveSharpEdges, bAllowSeamCollapse,
+			TargetMode, TargetPercentage, TargetCount, TargetEdgeLength, MinimalPlanarAngleThresh, 
+			FQEMSimplification::ESimplificationCollapseModes::MinimalExistingVertexError);
+	}
 	else // SimplifierType == ESimplifyType::UE4Standard
 	{
 		const FMeshDescription* SrcMeshDescription = OriginalMeshDescription.Get();
diff --git a/Engine/Plugins/Experimental/MeshModelingToolset/Source/ModelingOperatorsEditorOnly/Public/CleaningOps/SimplifyMeshOp.h b/Engine/Plugins/Experimental/MeshModelingToolset/Source/ModelingOperatorsEditorOnly/Public/CleaningOps/SimplifyMeshOp.h
--- a/Engine/Plugins/Experimental/MeshModelingToolset/Source/ModelingOperatorsEditorOnly/Public/CleaningOps/SimplifyMeshOp.h
+++ b/Engine/Plugins/Experimental/MeshModelingToolset/Source/ModelingOperatorsEditorOnly/Public/CleaningOps/SimplifyMeshOp.h
@@ -38,17 +38,21 @@ enum class ESimplifyTargetType : uint8
 UENUM()
 enum class ESimplifyType : uint8
 {
-	/** Fastest. Standard quadric error metric. Will not simplify UV boundaries.*/
+	/** Fastest. Standard quadric error metric.*/
 	QEM = 0 UMETA(DisplayName = "QEM"),
 
-	/** Potentially higher quality. Takes the normal into account. Will not simplify UV bounaries. */
+	/** Potentially higher quality. Takes the normal into account. */
 	Attribute = 1 UMETA(DisplayName = "Normal Aware"),
 
-	/** Highest quality reduction.  Will simplify UV boundaries. */
+	/** Highest quality reduction. */
 	UE4Standard = 2 UMETA(DisplayName = "UE4 Standard"),
 
-	/** Collapse any spurious edges but do not change the 3D shape */
-	MinimalPlanar = 3 UMETA(DisplayName = "Minimal Shape-Preserving"),
+	/** Edge collapse to existing vertices only.  Quality may suffer.*/
+	MinimalExistingVertex = 3 UMETA(DisplayName = "Existing Positions"),
+
+	/** Collapse any spurious edges but do not change the 3D shape. */
+	MinimalPlanar = 4 UMETA(DisplayName = "Minimal Shape-Preserving"),
+
 };
 
 class MODELINGOPERATORSEDITORONLY_API FSimplifyMeshOp : public FDynamicMeshOperator
