ModelingToolset: fix threshold in WeldMeshEdgesTool. Default vertex threshold is below floating-point epsilon, and UI parameter should affect vertex weld threshold as well as edge-midpoint matching threshold
#rnx

diff --git a/Engine/Plugins/Experimental/MeshModelingToolset/Source/MeshModelingTools/Private/WeldMeshEdgesTool.cpp b/Engine/Plugins/Experimental/MeshModelingToolset/Source/MeshModelingTools/Private/WeldMeshEdgesTool.cpp
--- a/Engine/Plugins/Experimental/MeshModelingToolset/Source/MeshModelingTools/Private/WeldMeshEdgesTool.cpp
+++ b/Engine/Plugins/Experimental/MeshModelingToolset/Source/MeshModelingTools/Private/WeldMeshEdgesTool.cpp
@@ -164,7 +164,8 @@ void UWeldMeshEdgesTool::UpdateResult()
 	TargetMesh->Copy(OriginalMesh);
 
 	FMergeCoincidentMeshEdges Merger(TargetMesh);
-	Merger.MergeSearchTolerance = this->Tolerance;
+	Merger.MergeVertexTolerance = this->Tolerance;
+	Merger.MergeSearchTolerance = 2*Merger.MergeVertexTolerance;
 	Merger.OnlyUniquePairs = this->bOnlyUnique;
 		
 	if (Merger.Apply() == false)
