ModelingTools: early-out in PolyEdit Weld Edges when we have 2 selected edges but one is invalid for welding
#rnx

diff --git a/Engine/Plugins/Experimental/MeshModelingToolset/Source/MeshModelingTools/Private/EditMeshPolygonsTool.cpp b/Engine/Plugins/Experimental/MeshModelingToolset/Source/MeshModelingTools/Private/EditMeshPolygonsTool.cpp
--- a/Engine/Plugins/Experimental/MeshModelingToolset/Source/MeshModelingTools/Private/EditMeshPolygonsTool.cpp
+++ b/Engine/Plugins/Experimental/MeshModelingToolset/Source/MeshModelingTools/Private/EditMeshPolygonsTool.cpp
@@ -1808,6 +1808,7 @@ void UEditMeshPolygonsTool::ApplyCollapseEdge()
 void UEditMeshPolygonsTool::ApplyWeldEdges()
 {
 	bool bValidInput = SelectionMechanic->GetActiveSelection().SelectedEdgeIDs.Num() == 2 && BeginMeshBoundaryEdgeEditChange(true);
+	bValidInput = bValidInput && ActiveEdgeSelection.Num() == 2;		// one of the initial edges may not have been valid
 	if ( bValidInput == false )
 	{
 		GetToolManager()->DisplayMessage( LOCTEXT("OnWeldEdgesFailed", "Cannot Weld current selection"), EToolMessageLevel::UserWarning);
