ModelingTools: add target component as invisible hit-test target in PolyCut, so that ctrl-click-to-set-plane mechanic works on selected object
#rnx

diff --git a/Engine/Plugins/Experimental/MeshModelingToolset/Source/MeshModelingToolsEditorOnly/Private/PolygonOnMeshTool.cpp b/Engine/Plugins/Experimental/MeshModelingToolset/Source/MeshModelingToolsEditorOnly/Private/PolygonOnMeshTool.cpp
--- a/Engine/Plugins/Experimental/MeshModelingToolset/Source/MeshModelingToolsEditorOnly/Private/PolygonOnMeshTool.cpp
+++ b/Engine/Plugins/Experimental/MeshModelingToolset/Source/MeshModelingToolsEditorOnly/Private/PolygonOnMeshTool.cpp
@@ -131,6 +131,7 @@ void UPolygonOnMeshTool::Setup()
 		DrawPlaneWorld = PlaneMechanic->Plane;
 		UpdateDrawPlane();
 	});
+	PlaneMechanic->SetPlaneCtrlClickBehaviorTarget->InvisibleComponentsToHitTest.Add(ComponentTarget->GetOwnerComponent());
 
 	// Convert input mesh description to dynamic mesh
 	OriginalDynamicMesh = MakeShared<FDynamicMesh3>();
