Fix compile issue in SimplifyMeshOp.
Seems like older versions of MSVC have a problem with the default argument that uses a template dependent name.

diff --git a/Engine/Plugins/Experimental/MeshModelingToolset/Source/ModelingOperatorsEditorOnly/Private/CleaningOps/SimplifyMeshOp.cpp b/Engine/Plugins/Experimental/MeshModelingToolset/Source/ModelingOperatorsEditorOnly/Private/CleaningOps/SimplifyMeshOp.cpp
--- a/Engine/Plugins/Experimental/MeshModelingToolset/Source/ModelingOperatorsEditorOnly/Private/CleaningOps/SimplifyMeshOp.cpp
+++ b/Engine/Plugins/Experimental/MeshModelingToolset/Source/ModelingOperatorsEditorOnly/Private/CleaningOps/SimplifyMeshOp.cpp
@@ -30,7 +30,7 @@ void ComputeSimplify(FDynamicMesh3* TargetMesh, const bool bReproject,
 					 const ESimplifyTargetType TargetMode,
 					 const float TargetPercentage, const int TargetCount, const float TargetEdgeLength,
 					 const float AngleThreshold,
-	                 typename SimplificationType::ESimplificationCollapseModes CollapseMode = typename SimplificationType::ESimplificationCollapseModes::MinimalQuadricPositionError)
+	                 typename SimplificationType::ESimplificationCollapseModes CollapseMode)
 {
 	SimplificationType Reducer(TargetMesh);
 
@@ -123,7 +123,8 @@ void FSimplifyMeshOp::CalculateResult(FProgressCancel* Progress)
 											GroupBoundaryConstraint,
 											MaterialBoundaryConstraint,
 											bPreserveSharpEdges, bAllowSeamCollapse,
-											TargetMode, TargetPercentage, TargetCount, TargetEdgeLength, MinimalPlanarAngleThresh);
+											TargetMode, TargetPercentage, TargetCount, TargetEdgeLength, MinimalPlanarAngleThresh,
+											FQEMSimplification::ESimplificationCollapseModes::MinimalQuadricPositionError);
 	}
 	else if (SimplifierType == ESimplifyType::Attribute)
 	{
@@ -133,7 +134,8 @@ void FSimplifyMeshOp::CalculateResult(FProgressCancel* Progress)
 													GroupBoundaryConstraint,
 													MaterialBoundaryConstraint,
 													bPreserveSharpEdges, bAllowSeamCollapse,
-													TargetMode, TargetPercentage, TargetCount, TargetEdgeLength, MinimalPlanarAngleThresh);
+													TargetMode, TargetPercentage, TargetCount, TargetEdgeLength, MinimalPlanarAngleThresh,
+													FAttrMeshSimplification::ESimplificationCollapseModes::MinimalQuadricPositionError);
 	}
 	else if (SimplifierType == ESimplifyType::MinimalPlanar)
 	{
@@ -143,7 +145,8 @@ void FSimplifyMeshOp::CalculateResult(FProgressCancel* Progress)
 			GroupBoundaryConstraint,
 			MaterialBoundaryConstraint,
 			bPreserveSharpEdges, bAllowSeamCollapse,
-			ESimplifyTargetType::MinimalPlanar, TargetPercentage, TargetCount, TargetEdgeLength, MinimalPlanarAngleThresh);
+			ESimplifyTargetType::MinimalPlanar, TargetPercentage, TargetCount, TargetEdgeLength, MinimalPlanarAngleThresh,
+			FAttrMeshSimplification::ESimplificationCollapseModes::MinimalQuadricPositionError);
 	}
 	else if (SimplifierType == ESimplifyType::MinimalExistingVertex)
 	{
