fix tooltip for DrawPolygonTool's "slices" option to refer to vertices instead of sections, to match its actual behavior

#rnx

diff --git a/Engine/Plugins/Experimental/MeshModelingToolset/Source/MeshModelingTools/Public/DrawPolygonTool.h b/Engine/Plugins/Experimental/MeshModelingToolset/Source/MeshModelingTools/Public/DrawPolygonTool.h
--- a/Engine/Plugins/Experimental/MeshModelingToolset/Source/MeshModelingTools/Public/DrawPolygonTool.h
+++ b/Engine/Plugins/Experimental/MeshModelingToolset/Source/MeshModelingTools/Public/DrawPolygonTool.h
@@ -109,7 +109,7 @@ public:
 									EditCondition = "OutputMode == EDrawPolygonOutputMode::ExtrudedConstant", EditConditionHides))
 	float ExtrudeHeight = 100.0f;
 
-	/** Number of sections in round features */
+	/** Number of vertices in round features */
 	UPROPERTY(EditAnywhere, NonTransactional, Category = Polygon, meta = (UIMin = "3", UIMax = "100", ClampMin = "3", ClampMax = "10000",
 				  EditCondition = "PolygonType == EDrawPolygonDrawMode::Circle || PolygonType == EDrawPolygonDrawMode::RoundedRectangle || PolygonType == EDrawPolygonDrawMode::HoleyCircle", EditConditionHides))
 	int Steps = 16;
