ModelingTools: prevent crashes due to transactions being applied to Primitive Tools that have been Shutdown
#rnx

diff --git a/Engine/Plugins/Experimental/MeshModelingToolset/Source/MeshModelingTools/Private/AddPrimitiveTool.cpp b/Engine/Plugins/Experimental/MeshModelingToolset/Source/MeshModelingTools/Private/AddPrimitiveTool.cpp
--- a/Engine/Plugins/Experimental/MeshModelingToolset/Source/MeshModelingTools/Private/AddPrimitiveTool.cpp
+++ b/Engine/Plugins/Experimental/MeshModelingToolset/Source/MeshModelingTools/Private/AddPrimitiveTool.cpp
@@ -171,9 +171,15 @@ void UAddPrimitiveTool::Render(IToolsContextRenderAPI* RenderAPI)
 
 void UAddPrimitiveTool::OnPropertyModified(UObject* PropertySet, FProperty* Property)
 {
-	PreviewMesh->EnableWireframe(MaterialProperties->bWireframe);
-	PreviewMesh->SetMaterial(MaterialProperties->Material.Get());
-	UpdatePreviewMesh();
+	// Because of how the ShapeSettings property set is implemented in this Tool, changes to it are transacted,
+	// and if the user exits the Tool and then tries to undo/redo those transactions, this function will end up being called.
+	// So we need to ensure that we handle this case.
+	if (PreviewMesh)
+	{
+		PreviewMesh->EnableWireframe(MaterialProperties->bWireframe);
+		PreviewMesh->SetMaterial(MaterialProperties->Material.Get());
+		UpdatePreviewMesh();
+	}
 }
 
 
