MeshModelingTools: Added ability to speed up tool preview mesh updates when mesh topology remains constant. Used it in LatticeDeformer tool. Also made preview object move the result mesh rather than do a copy.

#rnx

diff --git a/Engine/Plugins/Experimental/MeshModelingToolset/Source/MeshModelingTools/Private/LatticeDeformerTool.cpp b/Engine/Plugins/Experimental/MeshModelingToolset/Source/MeshModelingTools/Private/LatticeDeformerTool.cpp
--- a/Engine/Plugins/Experimental/MeshModelingToolset/Source/MeshModelingTools/Private/LatticeDeformerTool.cpp
+++ b/Engine/Plugins/Experimental/MeshModelingToolset/Source/MeshModelingTools/Private/LatticeDeformerTool.cpp
@@ -196,6 +196,8 @@ void ULatticeDeformerTool::StartPreview()
 
 	Preview->Setup(TargetWorld, LatticeDeformOpCreator);
 
+	Preview->SetIsMeshTopologyConstant(true, EMeshRenderAttributeFlags::Positions | EMeshRenderAttributeFlags::VertexNormals);
+
 	FComponentMaterialSet MaterialSet;
 	ComponentTarget->GetMaterialSet(MaterialSet);
 	Preview->ConfigureMaterials(MaterialSet.Materials,
diff --git a/Engine/Plugins/Experimental/MeshModelingToolset/Source/ModelingComponents/Private/MeshOpPreviewHelpers.cpp b/Engine/Plugins/Experimental/MeshModelingToolset/Source/ModelingComponents/Private/MeshOpPreviewHelpers.cpp
--- a/Engine/Plugins/Experimental/MeshModelingToolset/Source/ModelingComponents/Private/MeshOpPreviewHelpers.cpp
+++ b/Engine/Plugins/Experimental/MeshModelingToolset/Source/ModelingComponents/Private/MeshOpPreviewHelpers.cpp
@@ -10,6 +10,7 @@ void UMeshOpPreviewWithBackgroundCompute::Setup(UWorld* InWorld, IDynamicMeshOpe
 
 	BackgroundCompute = MakeUnique<FBackgroundDynamicMeshComputeSource>(OpGenerator);
 	bResultValid = false;
+	bMeshInitialized = false;
 }
 
 FDynamicMeshOpResult UMeshOpPreviewWithBackgroundCompute::Shutdown()
@@ -76,7 +77,14 @@ void UMeshOpPreviewWithBackgroundCompute::UpdateResults()
 
 		TUniquePtr<FDynamicMesh3> ResultMesh = MeshOp->ExtractResult();
 		PreviewMesh->SetTransform((FTransform)MeshOp->GetResultTransform());
-		PreviewMesh->UpdatePreview(ResultMesh.Get());  // copies the mesh @todo we could just give ownership to the Preview!
+
+		UPreviewMesh::ERenderUpdateMode UpdateType = (bMeshTopologyIsConstant && bMeshInitialized) ?
+			UPreviewMesh::ERenderUpdateMode::FastUpdate
+			: UPreviewMesh::ERenderUpdateMode::FullUpdate;
+		
+		PreviewMesh->UpdatePreview(MoveTemp(*ResultMesh), UpdateType, ChangingAttributeFlags);
+		bMeshInitialized = true;
+
 		PreviewMesh->SetVisible(bVisible);
 		bResultValid = true;
 
@@ -129,6 +137,12 @@ void UMeshOpPreviewWithBackgroundCompute::SetVisibility(bool bVisibleIn)
 	PreviewMesh->SetVisible(bVisible);
 }
 
+void UMeshOpPreviewWithBackgroundCompute::SetIsMeshTopologyConstant(bool bOn, EMeshRenderAttributeFlags ChangingAttributesIn)
+{
+	bMeshTopologyIsConstant = bOn;
+	ChangingAttributeFlags = ChangingAttributesIn;
+}
+
 
 bool UMeshOpPreviewWithBackgroundCompute::IsUsingWorkingMaterial()
 {
diff --git a/Engine/Plugins/Experimental/MeshModelingToolset/Source/ModelingComponents/Private/PreviewMesh.cpp b/Engine/Plugins/Experimental/MeshModelingToolset/Source/ModelingComponents/Private/PreviewMesh.cpp
--- a/Engine/Plugins/Experimental/MeshModelingToolset/Source/ModelingComponents/Private/PreviewMesh.cpp
+++ b/Engine/Plugins/Experimental/MeshModelingToolset/Source/ModelingComponents/Private/PreviewMesh.cpp
@@ -224,30 +224,24 @@ void UPreviewMesh::ClearPreview()
 }
 
 
-void UPreviewMesh::UpdatePreview(const FDynamicMesh3* Mesh)
+void UPreviewMesh::UpdatePreview(const FDynamicMesh3* Mesh, ERenderUpdateMode UpdateMode,
+	EMeshRenderAttributeFlags ModifiedAttribs)
 {
 	DynamicMeshComponent->SetDrawOnTop(this->bDrawOnTop);
 
 	DynamicMeshComponent->GetMesh()->Copy(*Mesh);
-	DynamicMeshComponent->NotifyMeshUpdated();
 
-	if (bBuildSpatialDataStructure)
-	{
-		MeshAABBTree.SetMesh(DynamicMeshComponent->GetMesh(), true);
-	}
+	NotifyDeferredEditCompleted(UpdateMode, ModifiedAttribs, bBuildSpatialDataStructure);
 }
 
-void UPreviewMesh::UpdatePreview(FDynamicMesh3&& Mesh)
+void UPreviewMesh::UpdatePreview(FDynamicMesh3&& Mesh, ERenderUpdateMode UpdateMode, 
+	EMeshRenderAttributeFlags ModifiedAttribs)
 {
 	DynamicMeshComponent->SetDrawOnTop(this->bDrawOnTop);
 
 	*(DynamicMeshComponent->GetMesh()) = MoveTemp(Mesh);
-	DynamicMeshComponent->NotifyMeshUpdated();
 
-	if (bBuildSpatialDataStructure)
-	{
-		MeshAABBTree.SetMesh(DynamicMeshComponent->GetMesh(), true);
-	}
+	NotifyDeferredEditCompleted(UpdateMode, ModifiedAttribs, bBuildSpatialDataStructure);
 }
 
 
diff --git a/Engine/Plugins/Experimental/MeshModelingToolset/Source/ModelingComponents/Public/MeshOpPreviewHelpers.h b/Engine/Plugins/Experimental/MeshModelingToolset/Source/ModelingComponents/Public/MeshOpPreviewHelpers.h
--- a/Engine/Plugins/Experimental/MeshModelingToolset/Source/ModelingComponents/Public/MeshOpPreviewHelpers.h
+++ b/Engine/Plugins/Experimental/MeshModelingToolset/Source/ModelingComponents/Public/MeshOpPreviewHelpers.h
@@ -123,6 +123,18 @@ public:
 	 */
 	void SetVisibility(bool bVisible);
 
+	/**
+	 * If set to true, then it will be assumed that the mesh topology (i.e. triangle/edge connectivity) remains constant,
+	 * which will allow updates after the first one to modify existing render proxy buffers rather than creating entirely
+	 * new ones. This will give a significant speedup to tools that do not add/remove vertices or affect their connectivity.
+	 *
+	 * @param bOn
+	 * @param ChangingAttributes If bOn is set to true, determines which attributes need updating. For instance, if a tool 
+	 *  moves verts without changing their UV's, then one would probably pass "EMeshRenderAttributeFlags::Positions | EMeshRenderAttributeFlags::VertexNormals".
+	 *  Has no effect if bOn was set to false.
+	 */
+	void SetIsMeshTopologyConstant(bool bOn, EMeshRenderAttributeFlags ChangingAttributes = EMeshRenderAttributeFlags::All);
+
 	/**
 	 * Set time that Preview will wait before showing working material
 	 */
@@ -167,6 +179,11 @@ protected:
 
 	bool bVisible = true;
 
+	// Used for partial/fast updates of the preview mesh render proxy.
+	bool bMeshTopologyIsConstant = false;
+	EMeshRenderAttributeFlags ChangingAttributeFlags = EMeshRenderAttributeFlags::All;
+	bool bMeshInitialized = false;
+
 	float SecondsBeforeWorkingMaterial = 2.0;
 
 	// this object manages the background computes
diff --git a/Engine/Plugins/Experimental/MeshModelingToolset/Source/ModelingComponents/Public/PreviewMesh.h b/Engine/Plugins/Experimental/MeshModelingToolset/Source/ModelingComponents/Public/PreviewMesh.h
--- a/Engine/Plugins/Experimental/MeshModelingToolset/Source/ModelingComponents/Public/PreviewMesh.h
+++ b/Engine/Plugins/Experimental/MeshModelingToolset/Source/ModelingComponents/Public/PreviewMesh.h
@@ -277,13 +277,21 @@ public:
 
 	/**
 	 * Update the internal mesh by copying the given Mesh
+	 * @param Mesh to copy.
+	 * @param UpdateMode Type of rendering update required. Should be FullUpdate if topology changes, otherwise can be FastUpdate.
+	 * @param ModifiedAttribs Only relevant in case of FastUpdate- determines which attributes actually changed.
 	 */
-	void UpdatePreview(const FDynamicMesh3* Mesh);
+	void UpdatePreview(const FDynamicMesh3* Mesh, ERenderUpdateMode UpdateMode = ERenderUpdateMode::FullUpdate, 
+		EMeshRenderAttributeFlags ModifiedAttribs = EMeshRenderAttributeFlags::All);
 
 	/**
 	 * Update the internal mesh by moving in the given Mesh
+	 * @param Mesh to move.
+	 * @param UpdateMode Type of rendering update required. Should be FullUpdate if topology changes, otherwise can be FastUpdate.
+	 * @param ModifiedAttribs Only relevant in case of FastUpdate- determines which attributes actually changed.
 	 */
-	void UpdatePreview(FDynamicMesh3&& Mesh);
+	void UpdatePreview(FDynamicMesh3&& Mesh, ERenderUpdateMode UpdateMode = ERenderUpdateMode::FullUpdate,
+		EMeshRenderAttributeFlags ModifiedAttribs = EMeshRenderAttributeFlags::All);
 
 	/**
 	 * Initialize the internal mesh based on the given MeshDescription
@@ -339,7 +347,7 @@ public:
 	 * but do not update/rebuild rendering data structures. NotifyDeferredEditOcurred() must be
 	 * called to complete a deferred edit, this will update the rendering mesh.
 	 * DeferredEditMesh can be called multiple times before NotifyDeferredEditCompleted() is called.
-	 * @param EditFunc function that is applied to the internal mes
+	 * @param EditFunc function that is applied to the internal mesh
 	 * @param bRebuildSpatial if true, and internal spatial data structure is enabled, rebuild it for updated mesh
 	 */
 	void DeferredEditMesh(TFunctionRef<void(FDynamicMesh3&)> EditFunc, bool bRebuildSpatial);
