MeshModelingTools: Modifications to SpaceDeformerTool to make it more responsive. Moved mesh copy over to background thread, parallelized operations, and (most importantly) used the fast render proxy update path.

#rnx

diff --git a/Engine/Plugins/Experimental/MeshModelingToolset/Source/MeshModelingTools/Private/MeshSpaceDeformerTool.cpp b/Engine/Plugins/Experimental/MeshModelingToolset/Source/MeshModelingTools/Private/MeshSpaceDeformerTool.cpp
--- a/Engine/Plugins/Experimental/MeshModelingToolset/Source/MeshModelingTools/Private/MeshSpaceDeformerTool.cpp
+++ b/Engine/Plugins/Experimental/MeshModelingToolset/Source/MeshModelingTools/Private/MeshSpaceDeformerTool.cpp
@@ -103,7 +103,7 @@ UMeshSpaceDeformerTool::UMeshSpaceDeformerTool()
 {
 	GizmoCenter = FVector::ZeroVector;
 	GizmoOrientation = FQuat::Identity;
-
+	GizmoFrame = FFrame3d(FVector3d(GizmoCenter), FQuaterniond(GizmoOrientation));
 }
 
 void UMeshSpaceDeformerTool::SetWorld(UWorld* World)
@@ -172,6 +172,7 @@ void UMeshSpaceDeformerTool::Setup()
 
 	AABBHalfExtents = 0.5 * FVector3d(Extents.X, Extents.Y, Extents.Z);
 	GizmoCenter = 0.5 * (AABBMin + AABBMax);
+	GizmoFrame.Origin = GizmoCenter;
 
 
 	// add click to set plane behavior
@@ -277,6 +278,8 @@ void UMeshSpaceDeformerTool::Setup()
 		Preview->Setup(this->TargetWorld, DeformerOperatorFactory);
 		Preview->PreviewMesh->SetTangentsMode(EDynamicMeshTangentCalcType::AutoCalculated);
 
+		Preview->SetIsMeshTopologyConstant(true, EMeshRenderAttributeFlags::Positions | EMeshRenderAttributeFlags::VertexNormals);
+
 		// Give the preview something to display
 		Preview->PreviewMesh->UpdatePreview(OriginalDynamicMesh.Get());
 		Preview->PreviewMesh->SetTransform(ComponentTarget->GetWorldTransform());
@@ -423,69 +426,9 @@ void  UMeshSpaceDeformerTool::OnPropertyModified(UObject* PropertySet, FProperty
 
 void UMeshSpaceDeformerTool::UpdateOpParameters(FMeshSpaceDeformerOp& MeshSpaceDeformerOp) const
 {
-	
-	
-
-	
-	FMatrix WorldToGizmo(EForceInit::ForceInitToZero);
-	{
-
-		// Matrix that rotates to the gizmo frame 
-		const FMatrix3d ToGizmo(GizmoFrame.GetAxis(0),
-		                    	GizmoFrame.GetAxis(1),
-			                    GizmoFrame.GetAxis(2),
-			                    true /* row constructor*/);
-
-		// Gizmo To World is rotation followed by translation of the form
-		// world_vec = FromGizmo * local_vec + GizmoCenter
-		// so the inverse will be
-		// local_vec = FromGizmo^(-1) * world_vec - FromGizmo^(-1) * GizmoCenter
-		// 
-
-		// Copy the rotation
-		for (int i = 0; i < 3; ++i)
-		{
-			for (int j = 0; j < 3; ++j)
-			{
-				WorldToGizmo.M[i][j] = ToGizmo(i, j);
-			}
-		}
-		// Add the translation 
-		FVector3d Translation = ToGizmo * GizmoCenter;
-		WorldToGizmo.M[0][3] = -Translation[0];
-		WorldToGizmo.M[1][3] = -Translation[1];
-		WorldToGizmo.M[2][3] = -Translation[2];
-		WorldToGizmo.M[3][3] = 1.f;
-	}
-	// Get the transform as a matrix.  Note may need to transpose this due
-	// to the vec * Matrix  vs Matrix * vec usage in engine.
-
-	FMatrix ObjectToWorld = ComponentTarget->GetWorldTransform().ToMatrixWithScale().GetTransposed();
-
-	FMatrix ObjectToGizmo(EForceInit::ForceInitToZero);
-	{
-		// ObjectToGizmo = WorldToGizmo * ObjectToWorld
-		for (int i = 0; i < 4; ++i)
-		{
-			for (int j = 0; j < 4; ++j)
-			{
-				for (int k = 0; k < 4; ++k)
-				{
-					ObjectToGizmo.M[i][j] += WorldToGizmo.M[i][k] * ObjectToWorld.M[k][j];
-				}
-			}
-		}
-	}
-
-	//This allows the object to move, or the handle position to move and have the changes reflected based on their relative world-space positions
-	const FVector3d AxisCentroidObjectSpace = ComponentTarget->GetWorldTransform().InverseTransformPositionNoScale(GizmoCenter);
-
-	//Pass the relevant updated data to the operator then run the operation.
-	/**
-	const double LowerBoundsInterval = this->LowerBoundsInterval;
-	const double UpperBoundsInterval = this->UpperBoundsInterval;
-	const double ModifierPercent     = this->ModifierPercent;
-	*/
+	MeshSpaceDeformerOp.OriginalMesh = OriginalDynamicMesh;
+	MeshSpaceDeformerOp.SetTransform(ComponentTarget->GetWorldTransform());
+	MeshSpaceDeformerOp.GizmoFrame = GizmoFrame;
 
 	// Set half axis length to be 1/2 the major axis of the bbox.
 	double LengthScale = FMath::Max( AABBHalfExtents.MaxAbsElement(), 1.e-3);
@@ -497,15 +440,6 @@ void UMeshSpaceDeformerTool::UpdateOpParameters(FMeshSpaceDeformerOp& MeshSpaceD
 
 	// percent to apply.
 	MeshSpaceDeformerOp.ModifierPercent = ModifierPercent;
-
-
-	// the transform to gizmo space
-	MeshSpaceDeformerOp.ObjectToGizmo = ObjectToGizmo;
-
-	//MeshSpaceDeformerOp.AxisOriginObjectSpace = AxisCentroidObjectSpace; // gizmo origin in object space.
-
-	// copy the mesh
-	MeshSpaceDeformerOp.CopySource(*OriginalDynamicMesh, ComponentTarget->GetWorldTransform());
 }
 
 void UMeshSpaceDeformerTool::OnTick(float DeltaTime)
diff --git a/Engine/Plugins/Experimental/MeshModelingToolset/Source/ModelingOperators/Private/MeshSpaceDeformerOps/BendMeshOp.cpp b/Engine/Plugins/Experimental/MeshModelingToolset/Source/ModelingOperators/Private/MeshSpaceDeformerOps/BendMeshOp.cpp
--- a/Engine/Plugins/Experimental/MeshModelingToolset/Source/ModelingOperators/Private/MeshSpaceDeformerOps/BendMeshOp.cpp
+++ b/Engine/Plugins/Experimental/MeshModelingToolset/Source/ModelingOperators/Private/MeshSpaceDeformerOps/BendMeshOp.cpp
@@ -1,12 +1,19 @@
 // Copyright Epic Games, Inc. All Rights Reserved.
-#include "ModelingOperators\Public\SpaceDeformerOps\BendMeshOp.h"
+#include "SpaceDeformerOps\BendMeshOp.h"
+
+#include "Async/ParallelFor.h"
 #include "DynamicMeshAttributeSet.h"
 
 
 //Bends along the Y-axis
 void FBendMeshOp::CalculateResult(FProgressCancel* Progress)
 {
+	FMeshSpaceDeformerOp::CalculateResult(Progress);
 
+	if (!OriginalMesh || (Progress && Progress->Cancelled()))
+	{
+		return;
+	}
 
 	// Matrix from gizmo space (z-up) to a y-up space
 	FMatrix ToYUp(EForceInit::ForceInitToZero);
@@ -30,7 +37,6 @@ void FBendMeshOp::CalculateResult(FProgressCancel* Progress)
 		}
 	}
 
-
 	float Det = ObjectToYUpGizmo.Determinant();
 
 	// Check if the transform is nearly singular
@@ -66,17 +72,26 @@ void FBendMeshOp::CalculateResult(FProgressCancel* Progress)
 	const double K = ThetaRadians / AxesHalfLength;
 	const double Ik = 1.0 / K;
 
+	if (Progress && Progress->Cancelled())
+	{
+		return;
+	}
 
 	if (ResultMesh->HasAttributes())
 	{
 		// Fix the normals first if they exist.
 
 		FDynamicMeshNormalOverlay* Normals = ResultMesh->Attributes()->PrimaryNormals();
-		for (int ElID : Normals->ElementIndicesItr())
+		ParallelFor(Normals->MaxElementID(), [this, Normals, &ObjectToYUpGizmo, &YUpGizmoToObject, YMin, YMax, Y0, K, Ik](int32 ElID)
 		{
+			if (!Normals->IsElement(ElID))
+			{
+				return;
+			}
+
 			// get the vertex
 			auto VertexID = Normals->GetParentVertex(ElID);
-			const FVector3d& SrcPos = TargetMesh->GetVertex(VertexID);
+			const FVector3d& SrcPos = ResultMesh->GetVertex(VertexID);
 
 			FVector3f SrcNormalF = Normals->GetElement(ElID);
 			FVector3d SrcNormal;
@@ -147,13 +162,23 @@ void FBendMeshOp::CalculateResult(FProgressCancel* Progress)
 			FVector3f RotatedNoramlF;
 			RotatedNoramlF[0] = RotatedNormal[0]; RotatedNoramlF[1] = RotatedNormal[1]; RotatedNoramlF[2] = RotatedNormal[2];
 			Normals->SetElement(ElID, RotatedNoramlF);
-		}
+		});
+	}
+
+	if (Progress && Progress->Cancelled())
+	{
+		return;
 	}
 
-	for (int VertexID : TargetMesh->VertexIndicesItr())
+	ParallelFor(ResultMesh->MaxVertexID(), [this, &ObjectToYUpGizmo, &YUpGizmoToObject, YMin, YMax, Y0, K, Ik](int32 VertexID)
 	{
+		if (!ResultMesh->IsVertex(VertexID))
+		{
+			return;
+		}
+
 		//const FVector3d& SrcPos = OriginalPositions[VertexID];
-		const FVector3d SrcPos = TargetMesh->GetVertex(VertexID);
+		const FVector3d SrcPos = ResultMesh->GetVertex(VertexID);
 		const double SrcPos4[4] = { SrcPos[0], SrcPos[1], SrcPos[2], 1.0};
 
 		// Position in gizmo space
@@ -208,7 +233,6 @@ void FBendMeshOp::CalculateResult(FProgressCancel* Progress)
 			}
 		}
 
-		TargetMesh->SetVertex(VertexID, FVector3d(DstPos4[0], DstPos4[1], DstPos4[2]));
-	}
-
+		ResultMesh->SetVertex(VertexID, FVector3d(DstPos4[0], DstPos4[1], DstPos4[2]));
+	});
 }
\ No newline at end of file
diff --git a/Engine/Plugins/Experimental/MeshModelingToolset/Source/ModelingOperators/Private/MeshSpaceDeformerOps/FlareMeshOp.cpp b/Engine/Plugins/Experimental/MeshModelingToolset/Source/ModelingOperators/Private/MeshSpaceDeformerOps/FlareMeshOp.cpp
--- a/Engine/Plugins/Experimental/MeshModelingToolset/Source/ModelingOperators/Private/MeshSpaceDeformerOps/FlareMeshOp.cpp
+++ b/Engine/Plugins/Experimental/MeshModelingToolset/Source/ModelingOperators/Private/MeshSpaceDeformerOps/FlareMeshOp.cpp
@@ -1,5 +1,7 @@
 // Copyright Epic Games, Inc. All Rights Reserved.
-#include "ModelingOperators\Public\SpaceDeformerOps\FlareMeshOp.h"
+#include "SpaceDeformerOps\FlareMeshOp.h"
+
+#include "Async/ParallelFor.h"
 #include "DynamicMeshAttributeSet.h"
 
 
@@ -34,6 +36,12 @@ double  inverseSinerp(const T& valueBetween, const T& value1, const T& value2)
 //Flares along the Z axis
 void FFlareMeshOp::CalculateResult(FProgressCancel* Progress)
 {
+	FMeshSpaceDeformerOp::CalculateResult(Progress);
+
+	if (!OriginalMesh || (Progress && Progress->Cancelled()))
+	{
+		return;
+	}
 
 	float Det = ObjectToGizmo.Determinant();
 
@@ -49,18 +57,21 @@ void FFlareMeshOp::CalculateResult(FProgressCancel* Progress)
 	const double ZMin = -LowerBoundsInterval * AxesHalfLength;
 	const double ZMax =  UpperBoundsInterval * AxesHalfLength;
 
-
-
 	if (ResultMesh->HasAttributes())
 	{
 		// Fix the normals first if they exist.
 
 		FDynamicMeshNormalOverlay* Normals = ResultMesh->Attributes()->PrimaryNormals();
-		for (int ElID : Normals->ElementIndicesItr())
+		ParallelFor(Normals->MaxElementID(), [this, Normals, &GizmoToObject, ZMin, ZMax](int32 ElID)
 		{
+			if (!Normals->IsElement(ElID))
+			{
+				return;
+			}
+
 			// get the vertex
 			auto VertexID = Normals->GetParentVertex(ElID);
-			const FVector3d& SrcPos = TargetMesh->GetVertex(VertexID);
+			const FVector3d& SrcPos = ResultMesh->GetVertex(VertexID);
 
 			FVector3f SrcNormalF = Normals->GetElement(ElID);
 			FVector3d SrcNormal; 
@@ -124,13 +135,22 @@ void FFlareMeshOp::CalculateResult(FProgressCancel* Progress)
 			FVector3f RotatedNoramlF;
 			RotatedNoramlF[0] = RotatedNormal[0]; RotatedNoramlF[1] = RotatedNormal[1]; RotatedNoramlF[2] = RotatedNormal[2];
 			Normals->SetElement(ElID, RotatedNoramlF);
-		}
+		});
 	}
 
-	for (int VertexID : TargetMesh->VertexIndicesItr())
+	if (Progress && Progress->Cancelled())
 	{
+		return;
+	}
+
+	ParallelFor(ResultMesh->MaxVertexID(), [this, &GizmoToObject, ZMin, ZMax](int32 VertexID)
+	{
+		if (!ResultMesh->IsVertex(VertexID))
+		{
+			return;
+		}
 		
-		const FVector3d SrcPos = TargetMesh->GetVertex(VertexID);
+		const FVector3d SrcPos = ResultMesh->GetVertex(VertexID);
 
 		const double SrcPos4[4] = { SrcPos[0], SrcPos[1], SrcPos[2], 1.0 };
 
@@ -165,7 +185,7 @@ void FFlareMeshOp::CalculateResult(FProgressCancel* Progress)
 		}
 
 		// set the position
-		TargetMesh->SetVertex(VertexID, FVector3d(DstPos4[0], DstPos4[1], DstPos4[2]));
-	}
+		ResultMesh->SetVertex(VertexID, FVector3d(DstPos4[0], DstPos4[1], DstPos4[2]));
+	});
 
 }
\ No newline at end of file
diff --git a/Engine/Plugins/Experimental/MeshModelingToolset/Source/ModelingOperators/Private/MeshSpaceDeformerOps/MeshSpaceDeformerOp.cpp b/Engine/Plugins/Experimental/MeshModelingToolset/Source/ModelingOperators/Private/MeshSpaceDeformerOps/MeshSpaceDeformerOp.cpp
--- a/Engine/Plugins/Experimental/MeshModelingToolset/Source/ModelingOperators/Private/MeshSpaceDeformerOps/MeshSpaceDeformerOp.cpp
+++ b/Engine/Plugins/Experimental/MeshModelingToolset/Source/ModelingOperators/Private/MeshSpaceDeformerOps/MeshSpaceDeformerOp.cpp
@@ -3,15 +3,73 @@
 #include "SpaceDeformerOps\MeshSpaceDeformerOp.h"
 #include "DynamicMesh3.h"
 
-void FMeshSpaceDeformerOp::CalculateResult(FProgressCancel* Progress)
+void FMeshSpaceDeformerOp::SetTransform(const FTransform& Transform)
 {
+	ResultTransform = (FTransform3d)Transform;
 }
 
-void FMeshSpaceDeformerOp::CopySource(const FDynamicMesh3& MeshIn, const FTransform& XForm)
+void FMeshSpaceDeformerOp::CalculateResult(FProgressCancel* Progress)
 {
-	ResultMesh = MakeUnique < FDynamicMesh3>(MeshIn);
-	ResultTransform = FTransform3d(XForm);
+	if (Progress && Progress->Cancelled() || !OriginalMesh)
+	{
+		return;
+	}
+
+	ResultMesh->Copy(*OriginalMesh, true, true, true, true);
+
+	if (Progress && Progress->Cancelled())
+	{
+		return;
+	}
+
+	FMatrix WorldToGizmo(EForceInit::ForceInitToZero);
+	{
+		// Matrix that rotates to the gizmo frame 
+		const FMatrix3d ToGizmo(GizmoFrame.GetAxis(0),
+			GizmoFrame.GetAxis(1),
+			GizmoFrame.GetAxis(2),
+			true /* row constructor*/);
+
+		// Gizmo To World is rotation followed by translation of the form
+		// world_vec = FromGizmo * local_vec + GizmoCenter
+		// so the inverse will be
+		// local_vec = FromGizmo^(-1) * world_vec - FromGizmo^(-1) * GizmoCenter
+		// 
+
+		// Copy the rotation
+		for (int i = 0; i < 3; ++i)
+		{
+			for (int j = 0; j < 3; ++j)
+			{
+				WorldToGizmo.M[i][j] = ToGizmo(i, j);
+			}
+		}
+		// Add the translation 
+		FVector3d Translation = ToGizmo * GizmoFrame.Origin;
+		WorldToGizmo.M[0][3] = -Translation[0];
+		WorldToGizmo.M[1][3] = -Translation[1];
+		WorldToGizmo.M[2][3] = -Translation[2];
+		WorldToGizmo.M[3][3] = 1.f;
+	}
+	// Get the transform as a matrix.  We need to transpose this due
+	// to the vec * Matrix  vs Matrix * vec usage in engine.
+
+	FMatrix ObjectToWorld = FTransform(ResultTransform).ToMatrixWithScale().GetTransposed();
 
-	TargetMesh = ResultMesh.Get();
+	ObjectToGizmo = FMatrix(EForceInit::ForceInitToZero);
+	{
+		// ObjectToGizmo = WorldToGizmo * ObjectToWorld
+		for (int i = 0; i < 4; ++i)
+		{
+			for (int j = 0; j < 4; ++j)
+			{
+				for (int k = 0; k < 4; ++k)
+				{
+					ObjectToGizmo.M[i][j] += WorldToGizmo.M[i][k] * ObjectToWorld.M[k][j];
+				}
+			}
+		}
+	}
 
+	// The rest is up to the subclass.
 }
diff --git a/Engine/Plugins/Experimental/MeshModelingToolset/Source/ModelingOperators/Private/MeshSpaceDeformerOps/TwistMeshOp.cpp b/Engine/Plugins/Experimental/MeshModelingToolset/Source/ModelingOperators/Private/MeshSpaceDeformerOps/TwistMeshOp.cpp
--- a/Engine/Plugins/Experimental/MeshModelingToolset/Source/ModelingOperators/Private/MeshSpaceDeformerOps/TwistMeshOp.cpp
+++ b/Engine/Plugins/Experimental/MeshModelingToolset/Source/ModelingOperators/Private/MeshSpaceDeformerOps/TwistMeshOp.cpp
@@ -1,11 +1,20 @@
 // Copyright Epic Games, Inc. All Rights Reserved.
 
-#include "ModelingOperators\Public\SpaceDeformerOps\TwistMeshOp.h"
+#include "SpaceDeformerOps\TwistMeshOp.h"
+
+#include "Async/ParallelFor.h"
 #include "DynamicMeshAttributeSet.h"
 
 //Twists along the Z axis
 void FTwistMeshOp::CalculateResult(FProgressCancel* Progress)
 {
+	FMeshSpaceDeformerOp::CalculateResult(Progress);
+
+	if (!OriginalMesh || (Progress && Progress->Cancelled()))
+	{
+		return;
+	}
+
 	float Det = ObjectToGizmo.Determinant();
 
 	// Check if the transform is nearly singular
@@ -30,17 +39,21 @@ void FTwistMeshOp::CalculateResult(FProgressCancel* Progress)
 	const double DegreesToRadians = 0.017453292519943295769236907684886; // Pi / 180
 	const double ThetaRadians = DegreesToRadians * GetModifierValue();
 
-
 	if (ResultMesh->HasAttributes())
 	{
 		// Fix the normals first if they exist.
 
 		FDynamicMeshNormalOverlay* Normals = ResultMesh->Attributes()->PrimaryNormals();
-		for (int ElID : Normals->ElementIndicesItr())
+		ParallelFor(Normals->MaxElementID(), [this, Normals, &GizmoToObject, ZMin, ZMax, ThetaRadians](int32 ElID)
 		{
+			if (!Normals->IsElement(ElID))
+			{
+				return;
+			}
+
 			// get the vertex
 			auto VertexID = Normals->GetParentVertex(ElID);
-			const FVector3d& SrcPos = TargetMesh->GetVertex(VertexID);
+			const FVector3d& SrcPos = ResultMesh->GetVertex(VertexID);
 
 			FVector3f SrcNormalF = Normals->GetElement(ElID);
 			FVector3d SrcNormal;
@@ -109,14 +122,23 @@ void FTwistMeshOp::CalculateResult(FProgressCancel* Progress)
 			FVector3f RotatedNoramlF;
 			RotatedNoramlF[0] = RotatedNormal[0]; RotatedNoramlF[1] = RotatedNormal[1]; RotatedNoramlF[2] = RotatedNormal[2];
 			Normals->SetElement(ElID, RotatedNoramlF);
-		}
+		});
+	}
+
+	if (Progress && Progress->Cancelled())
+	{
+		return;
 	}
 
 	// now fix the vertex positions
-	for (int VertexID : TargetMesh->VertexIndicesItr())
+	ParallelFor(ResultMesh->MaxVertexID(), [this, &GizmoToObject, ZMin, ZMax, ThetaRadians](int32 VertexID)
 	{
+		if (!ResultMesh->IsVertex(VertexID))
+		{
+			return;
+		}
 
-		const FVector3d SrcPos = TargetMesh->GetVertex(VertexID);
+		const FVector3d SrcPos = ResultMesh->GetVertex(VertexID);
 		const double SrcPos4[4] = { SrcPos[0], SrcPos[1], SrcPos[2], 1.0 };
 
 		// Position in gizmo space
@@ -158,7 +180,7 @@ void FTwistMeshOp::CalculateResult(FProgressCancel* Progress)
 		}
 
 		// set the position
-		TargetMesh->SetVertex(VertexID, FVector3d(DstPos4[0], DstPos4[1], DstPos4[2]));
+		ResultMesh->SetVertex(VertexID, FVector3d(DstPos4[0], DstPos4[1], DstPos4[2]));
 		
-	}
+	});
 }
\ No newline at end of file
diff --git a/Engine/Plugins/Experimental/MeshModelingToolset/Source/ModelingOperators/Public/SpaceDeformerOps/MeshSpaceDeformerOp.h b/Engine/Plugins/Experimental/MeshModelingToolset/Source/ModelingOperators/Public/SpaceDeformerOps/MeshSpaceDeformerOp.h
--- a/Engine/Plugins/Experimental/MeshModelingToolset/Source/ModelingOperators/Public/SpaceDeformerOps/MeshSpaceDeformerOp.h
+++ b/Engine/Plugins/Experimental/MeshModelingToolset/Source/ModelingOperators/Public/SpaceDeformerOps/MeshSpaceDeformerOp.h
@@ -12,43 +12,41 @@ public:
 	FMeshSpaceDeformerOp(double Min, double Max) : ModifierMin{ Min }, ModifierMax{ Max }{};
 	virtual ~FMeshSpaceDeformerOp() {}
 
-	FDynamicMesh3* TargetMesh;
-
-	// set ability on protected transform.
+	// Inputs
+	TSharedPtr<const FDynamicMesh3> OriginalMesh;
+	FFrame3d GizmoFrame = GizmoFrame;
+	void SetTransform(const FTransform& Transform);
 
 	//Converts the given modifier percent and returns the value between -1 and 1
 	virtual double GetModifierValue() { return FMath::Lerp(0.01 * ModifierPercent, ModifierMin, ModifierMax); }
 
-	/**
-	* \brief Updates the pointer to the target mesh
-	*/
-	virtual void CopySource(const FDynamicMesh3& Mesh, const FTransform& XFrom);
-
-
 	// half the major axis of the bounding box for the source geometry
 	double AxesHalfLength;
 
-	// transform, including translation, to gizmo space
-	FMatrix ObjectToGizmo;
-
 	//Percent input from the tool (between -100 and 100), as each operator will have a min/max value i.e. the angle of curvature, twist or scale. This is the percent used to interpolate between the min/max value
 	double ModifierPercent;
 
-	//The minimum value provided by the constructor so that each operator can set it's own interpolation min and max
-	double ModifierMin{};
-
-	//The minimum value provided by the constructor so that each operator can set it's own interpolation min and max
-	double ModifierMax{};
-
 	//The interval usually from [-1,0] upon which this operator will work. This corresponds to the region of space which this operator affects.
 	double LowerBoundsInterval;
 
 	//The interval usually from [0,1] upon which this operator will work. This corresponds to the region of space which this operator affects.
 	double UpperBoundsInterval;
 
+	/**
+	 * Copies over the original mesh into result mesh, and initializes ObjectToGizmo in preparation to whatever work the base class does.
+	 * Note that the function will return if OriginalMesh was null but doesn't have a way to log the error, so the base class should
+	 * check OriginalMesh itself as well.
+	 */
 	virtual void CalculateResult(FProgressCancel* Progress) override;
 
 
 protected:
+	// transform, including translation, to gizmo space
+	FMatrix ObjectToGizmo;
+
+	//The minimum value provided by the constructor so that each operator can set its own interpolation min and max
+	double ModifierMin{};
 
+	//The minimum value provided by the constructor so that each operator can set its own interpolation min and max
+	double ModifierMax{};
 };
\ No newline at end of file
