MeshModelingTools: Fix marquee selection in PolyEdit not issuing transactions.

#rnx

diff --git a/Engine/Plugins/Experimental/GeometryProcessing/Source/DynamicMesh/Public/GroupTopology.h b/Engine/Plugins/Experimental/GeometryProcessing/Source/DynamicMesh/Public/GroupTopology.h
--- a/Engine/Plugins/Experimental/GeometryProcessing/Source/DynamicMesh/Public/GroupTopology.h
+++ b/Engine/Plugins/Experimental/GeometryProcessing/Source/DynamicMesh/Public/GroupTopology.h
@@ -97,6 +97,14 @@ struct DYNAMICMESH_API FGroupTopologySelection
 			&& SelectedGroupIDs.Includes(Selection.SelectedGroupIDs);
 	}
 
+	bool operator==(const FGroupTopologySelection& Selection) const
+	{
+		return SelectedCornerIDs.Num() == Selection.SelectedCornerIDs.Num()
+			&& SelectedEdgeIDs.Num() == Selection.SelectedEdgeIDs.Num()
+			&& SelectedGroupIDs.Num() == Selection.SelectedGroupIDs.Num()
+			&& Contains(Selection);
+	}
+
 	inline bool IsSelectedTriangle(const FDynamicMesh3* Mesh, const FGroupTopology* Topology, int32 TriangleID) const;
 
 
diff --git a/Engine/Plugins/Experimental/MeshModelingToolset/Source/ModelingComponents/Private/Selection/PolygonSelectionMechanic.cpp b/Engine/Plugins/Experimental/MeshModelingToolset/Source/ModelingComponents/Private/Selection/PolygonSelectionMechanic.cpp
--- a/Engine/Plugins/Experimental/MeshModelingToolset/Source/ModelingComponents/Private/Selection/PolygonSelectionMechanic.cpp
+++ b/Engine/Plugins/Experimental/MeshModelingToolset/Source/ModelingComponents/Private/Selection/PolygonSelectionMechanic.cpp
@@ -582,6 +582,9 @@ void UPolygonSelectionMechanic::OnUpdateModifierState(int ModifierID, bool bIsOn
 
 void UPolygonSelectionMechanic::OnDragRectangleStarted()
 {
+	ParentTool->GetToolManager()->BeginUndoTransaction(LOCTEXT("SelectionChange", "Selection"));
+	BeginChange();
+
 	PreDragPersistentSelection = PersistentSelection;
 	PersistentSelection.Clear();
 }
@@ -622,8 +625,14 @@ void UPolygonSelectionMechanic::OnDragRectangleChanged(const FCameraRectangle& C
 
 void UPolygonSelectionMechanic::OnDragRectangleFinished()
 {
-	SelectionTimestamp++;
-	OnSelectionChanged.Broadcast();
+	if (!(PersistentSelection == PreDragPersistentSelection))
+	{
+		SelectionTimestamp++;
+		OnSelectionChanged.Broadcast();
+	}
+
+	EndChangeAndEmitIfModified();
+	ParentTool->GetToolManager()->EndUndoTransaction();
 }
 
 
