MeshModelingTools: Fix marquee mechanic selecting points behind camera.

#rnx

diff --git a/Engine/Plugins/Experimental/MeshModelingToolset/Source/ModelingComponents/Private/Mechanics/RectangleMarqueeMechanic.cpp b/Engine/Plugins/Experimental/MeshModelingToolset/Source/ModelingComponents/Private/Mechanics/RectangleMarqueeMechanic.cpp
--- a/Engine/Plugins/Experimental/MeshModelingToolset/Source/ModelingComponents/Private/Mechanics/RectangleMarqueeMechanic.cpp
+++ b/Engine/Plugins/Experimental/MeshModelingToolset/Source/ModelingComponents/Private/Mechanics/RectangleMarqueeMechanic.cpp
@@ -48,6 +48,12 @@ bool FCameraRectangle::IsProjectedPointInRectangle(const FVector& Point) const
 	}
 	else
 	{
+		// If it's behind the camera rectangle, then not contained.
+		if (CameraPlane.PlaneDot(Point) < 0)
+		{
+			return false;
+		}
+
 		// intersect along the eye-to-point ray
 		ProjectedPoint = FMath::RayPlaneIntersection(CameraOrigin,
 													 Point - CameraOrigin,
@@ -60,10 +66,37 @@ bool FCameraRectangle::IsProjectedPointInRectangle(const FVector& Point) const
 
 bool FCameraRectangle::IsProjectedSegmentIntersectingRectangle(const FVector& Endpoint1, const FVector& Endpoint2) const
 {
-	FVector ProjectedEndpoint1 = bCameraIsOrthographic ? FVector::PointPlaneProject(Endpoint1, CameraPlane)
-		: FMath::RayPlaneIntersection(CameraOrigin, Endpoint1 - CameraOrigin, CameraPlane);
-	FVector ProjectedEndpoint2 = bCameraIsOrthographic ? FVector::PointPlaneProject(Endpoint2, CameraPlane)
-		: FMath::RayPlaneIntersection(CameraOrigin, Endpoint2 - CameraOrigin, CameraPlane);
+	FVector ProjectedEndpoint1;
+	FVector ProjectedEndpoint2;
+
+	if (bCameraIsOrthographic)
+	{
+		ProjectedEndpoint1 = FVector::PointPlaneProject(Endpoint1, CameraPlane);
+		ProjectedEndpoint2 = FVector::PointPlaneProject(Endpoint2, CameraPlane);
+	}
+	else
+	{
+		// We'll have to crop the segment to the portion in front of the camera plane
+		bool bPoint1InFrontOfCamera = CameraPlane.PlaneDot(Endpoint1) > 0;
+		bool bPoint2InFrontOfCamera = CameraPlane.PlaneDot(Endpoint2) > 0;
+		
+		if (!bPoint1InFrontOfCamera && !bPoint2InFrontOfCamera)
+		{
+			return false; // Segment is behind the camera plane
+		}
+
+		FVector IntersectionPoint;
+		if (bPoint1InFrontOfCamera != bPoint2InFrontOfCamera)
+		{
+			// Get the intersection point so we can replace the endpoint that is behind the camera with it.
+			ensure(FMath::SegmentPlaneIntersection(Endpoint1, Endpoint2, CameraPlane, IntersectionPoint));
+		}
+
+		ProjectedEndpoint1 = bPoint1InFrontOfCamera ? FMath::RayPlaneIntersection(CameraOrigin, Endpoint1 - CameraOrigin, CameraPlane)
+			: IntersectionPoint;
+		ProjectedEndpoint2 = bPoint2InFrontOfCamera ? FMath::RayPlaneIntersection(CameraOrigin, Endpoint2 - CameraOrigin, CameraPlane)
+			: IntersectionPoint;
+	}
 
 	FVector2D Endpoint1PlaneCoord = PlaneCoordinates(ProjectedEndpoint1);
 	FVector2D Endpoint2PlaneCoord = PlaneCoordinates(ProjectedEndpoint2);
