MeshModelingTools: Fixing trivial self-assignment that accidentally got commited in 15066707.

#rnx

diff --git a/Engine/Plugins/Experimental/MeshModelingToolset/Source/ModelingOperators/Public/SpaceDeformerOps/MeshSpaceDeformerOp.h b/Engine/Plugins/Experimental/MeshModelingToolset/Source/ModelingOperators/Public/SpaceDeformerOps/MeshSpaceDeformerOp.h
--- a/Engine/Plugins/Experimental/MeshModelingToolset/Source/ModelingOperators/Public/SpaceDeformerOps/MeshSpaceDeformerOp.h
+++ b/Engine/Plugins/Experimental/MeshModelingToolset/Source/ModelingOperators/Public/SpaceDeformerOps/MeshSpaceDeformerOp.h
@@ -14,7 +14,7 @@ public:
 
 	// Inputs
 	TSharedPtr<const FDynamicMesh3> OriginalMesh;
-	FFrame3d GizmoFrame = GizmoFrame;
+	FFrame3d GizmoFrame;
 	void SetTransform(const FTransform& Transform);
 
 	//Converts the given modifier percent and returns the value between -1 and 1
