Every Progress->Cancelled() gets a nullptr check

#rnx

diff --git a/Engine/Plugins/Experimental/MeshModelingToolset/Source/MeshModelingTools/Private/DisplaceMeshTool.cpp b/Engine/Plugins/Experimental/MeshModelingToolset/Source/MeshModelingTools/Private/DisplaceMeshTool.cpp
--- a/Engine/Plugins/Experimental/MeshModelingToolset/Source/MeshModelingTools/Private/DisplaceMeshTool.cpp
+++ b/Engine/Plugins/Experimental/MeshModelingToolset/Source/MeshModelingTools/Private/DisplaceMeshTool.cpp
@@ -356,10 +356,10 @@ namespace DisplaceMeshToolLocals{
 
 	void FDisplaceMeshOp::CalculateResult(FProgressCancel* Progress)
 	{
-		if (Progress->Cancelled()) return;
+		if (Progress && Progress->Cancelled()) return;
 		ResultMesh->Copy(*SourceMesh);
 
-		if (Progress->Cancelled()) return;
+		if (Progress && Progress->Cancelled()) return;
 
 		if (DisplacementType == EDisplaceMeshToolDisplaceType::DisplacementMap && !Parameters.DisplacementMap)
 		{
@@ -369,7 +369,7 @@ namespace DisplaceMeshToolLocals{
 		SourceNormals = FMeshNormals(SourceMesh.Get());
 		SourceNormals.ComputeVertexNormals();
 
-		if (Progress->Cancelled()) return;
+		if (Progress && Progress->Cancelled()) return;
 		// cache initial positions
 		SourcePositions.SetNum(SourceMesh->MaxVertexID());
 		for (int vid : SourceMesh->VertexIndicesItr())
@@ -377,10 +377,10 @@ namespace DisplaceMeshToolLocals{
 			SourcePositions[vid] = SourceMesh->GetVertex(vid);
 		}
 
-		if (Progress->Cancelled()) return;
+		if (Progress && Progress->Cancelled()) return;
 		DisplacedPositions.SetNum(SourceMesh->MaxVertexID());
 
-		if (Progress->Cancelled()) return;
+		if (Progress && Progress->Cancelled()) return;
 
 		ComputeDisplacement::FDirectionalFilter DirectionalFilter{ Parameters.bEnableFilter,
 			FVector3d(Parameters.FilterDirection),
diff --git a/Engine/Plugins/Experimental/MeshModelingToolset/Source/ModelingOperators/Private/BaseOps/VoxelBaseOp.cpp b/Engine/Plugins/Experimental/MeshModelingToolset/Source/ModelingOperators/Private/BaseOps/VoxelBaseOp.cpp
--- a/Engine/Plugins/Experimental/MeshModelingToolset/Source/ModelingOperators/Private/BaseOps/VoxelBaseOp.cpp
+++ b/Engine/Plugins/Experimental/MeshModelingToolset/Source/ModelingOperators/Private/BaseOps/VoxelBaseOp.cpp
@@ -11,7 +11,7 @@
 
 void FVoxelBaseOp::PostProcessResult(FProgressCancel* Progress, double MeshCellSize)
 {
-	if (Progress->Cancelled())
+	if (Progress && Progress->Cancelled())
 	{
 		return;
 	}
@@ -21,6 +21,10 @@ void FVoxelBaseOp::PostProcessResult(FProgressCancel* Progress, double MeshCellS
 		UE::MeshAutoRepair::RemoveInternalTriangles(*ResultMesh.Get(), true, EOcclusionTriangleSampling::Centroids, EOcclusionCalculationMode::FastWindingNumber);
 	}
 
+	if (Progress && Progress->Cancelled())
+	{
+		return;
+	}
 
 	bool bFixNormals = bAutoSimplify;
 	{
@@ -35,6 +39,11 @@ void FVoxelBaseOp::PostProcessResult(FProgressCancel* Progress, double MeshCellS
 		}
 	}
 
+	if (Progress && Progress->Cancelled())
+	{
+		return;
+	}
+
 	if (bFixNormals)
 	{
 		TSharedPtr<FDynamicMesh3> OpResultMesh(ExtractResult().Release()); // moved the unique pointer
@@ -62,6 +71,11 @@ void FVoxelBaseOp::PostProcessResult(FProgressCancel* Progress, double MeshCellS
 		FMeshNormals::QuickComputeVertexNormals(*ResultMesh.Get());
 	}
 
+	if (Progress && Progress->Cancelled())
+	{
+		return;
+	}
+
 	if (MinComponentVolume > 0 || MinComponentArea > 0)
 	{
 		FDynamicMeshEditor ComponentRemover(ResultMesh.Get());
diff --git a/Engine/Plugins/Experimental/MeshModelingToolset/Source/ModelingOperators/Private/CleaningOps/EditNormalsOp.cpp b/Engine/Plugins/Experimental/MeshModelingToolset/Source/ModelingOperators/Private/CleaningOps/EditNormalsOp.cpp
--- a/Engine/Plugins/Experimental/MeshModelingToolset/Source/ModelingOperators/Private/CleaningOps/EditNormalsOp.cpp
+++ b/Engine/Plugins/Experimental/MeshModelingToolset/Source/ModelingOperators/Private/CleaningOps/EditNormalsOp.cpp
@@ -19,7 +19,7 @@ void FEditNormalsOp::SetTransform(const FTransform& Transform) {
 
 void FEditNormalsOp::CalculateResult(FProgressCancel* Progress)
 {
-	if (Progress->Cancelled())
+	if (Progress && Progress->Cancelled())
 	{
 		return;
 	}
@@ -31,7 +31,7 @@ void FEditNormalsOp::CalculateResult(FProgressCancel* Progress)
 		ResultMesh->EnableAttributes();
 	}
 
-	if (Progress->Cancelled())
+	if (Progress && Progress->Cancelled())
 	{
 		return;
 	}
@@ -44,7 +44,7 @@ void FEditNormalsOp::CalculateResult(FProgressCancel* Progress)
 		FMeshRepairOrientation Repair(ResultMesh.Get());
 		Repair.OrientComponents();
 
-		if (Progress->Cancelled())
+		if (Progress && Progress->Cancelled())
 		{
 			return;
 		}
@@ -52,7 +52,7 @@ void FEditNormalsOp::CalculateResult(FProgressCancel* Progress)
 		Repair.SolveGlobalOrientation(&Tree);
 	}
 
-	if (Progress->Cancelled())
+	if (Progress && Progress->Cancelled())
 	{
 		return;
 	}
@@ -76,7 +76,7 @@ void FEditNormalsOp::CalculateResult(FProgressCancel* Progress)
 		}
 	}
 
-	if (Progress->Cancelled())
+	if (Progress && Progress->Cancelled())
 	{
 		return;
 	}
@@ -112,7 +112,7 @@ void FEditNormalsOp::CalculateResult(FProgressCancel* Progress)
 		}
 	}
 
-	if (Progress->Cancelled())
+	if (Progress && Progress->Cancelled())
 	{
 		return;
 	}
@@ -127,7 +127,7 @@ void FEditNormalsOp::CalculateResult(FProgressCancel* Progress)
 		MeshNormals.CopyToOverlay(ResultMesh->Attributes()->PrimaryNormals(), false);
 	}
 
-	if (Progress->Cancelled())
+	if (Progress && Progress->Cancelled())
 	{
 		return;
 	}
diff --git a/Engine/Plugins/Experimental/MeshModelingToolset/Source/ModelingOperators/Private/CleaningOps/HoleFillOp.cpp b/Engine/Plugins/Experimental/MeshModelingToolset/Source/ModelingOperators/Private/CleaningOps/HoleFillOp.cpp
--- a/Engine/Plugins/Experimental/MeshModelingToolset/Source/ModelingOperators/Private/CleaningOps/HoleFillOp.cpp
+++ b/Engine/Plugins/Experimental/MeshModelingToolset/Source/ModelingOperators/Private/CleaningOps/HoleFillOp.cpp
@@ -86,7 +86,7 @@ void FHoleFillOp::CalculateResult(FProgressCancel* Progress)
 {
 	NumFailedLoops = 0;
 
-	if (Progress->Cancelled())
+	if (Progress && Progress->Cancelled())
 	{
 		return;
 	}
@@ -98,7 +98,7 @@ void FHoleFillOp::CalculateResult(FProgressCancel* Progress)
 		return;
 	}
 
-	if (Progress->Cancelled())
+	if (Progress && Progress->Cancelled())
 	{
 		return;
 	}
@@ -182,7 +182,7 @@ void FHoleFillOp::CalculateResult(FProgressCancel* Progress)
 			Editor.SetTriangleUVsFromProjection(Filler->NewTriangles, ProjectionFrame, MeshUVScaleFactor);
 		}
 
-		if (Progress->Cancelled())
+		if (Progress && Progress->Cancelled())
 		{
 			return;
 		}
diff --git a/Engine/Plugins/Experimental/MeshModelingToolset/Source/ModelingOperators/Private/CleaningOps/RemeshMeshOp.cpp b/Engine/Plugins/Experimental/MeshModelingToolset/Source/ModelingOperators/Private/CleaningOps/RemeshMeshOp.cpp
--- a/Engine/Plugins/Experimental/MeshModelingToolset/Source/ModelingOperators/Private/CleaningOps/RemeshMeshOp.cpp
+++ b/Engine/Plugins/Experimental/MeshModelingToolset/Source/ModelingOperators/Private/CleaningOps/RemeshMeshOp.cpp
@@ -47,7 +47,7 @@ void FRemeshMeshOp::SetTransform(const FTransform& Transform)
 
 void FRemeshMeshOp::CalculateResult(FProgressCancel* Progress)
 {
-	if (Progress->Cancelled())
+	if (Progress && Progress->Cancelled())
 	{
 		return;
 	}
@@ -55,7 +55,7 @@ void FRemeshMeshOp::CalculateResult(FProgressCancel* Progress)
 	bool bDiscardAttributesImmediately = bDiscardAttributes && !bPreserveSharpEdges;
 	ResultMesh->Copy(*OriginalMesh, true, true, true, !bDiscardAttributesImmediately);
 
-	if (Progress->Cancelled())
+	if (Progress && Progress->Cancelled())
 	{
 		return;
 	}
diff --git a/Engine/Plugins/Experimental/MeshModelingToolset/Source/ModelingOperators/Private/CleaningOps/RemoveOccludedTrianglesOp.cpp b/Engine/Plugins/Experimental/MeshModelingToolset/Source/ModelingOperators/Private/CleaningOps/RemoveOccludedTrianglesOp.cpp
--- a/Engine/Plugins/Experimental/MeshModelingToolset/Source/ModelingOperators/Private/CleaningOps/RemoveOccludedTrianglesOp.cpp
+++ b/Engine/Plugins/Experimental/MeshModelingToolset/Source/ModelingOperators/Private/CleaningOps/RemoveOccludedTrianglesOp.cpp
@@ -49,7 +49,7 @@ void FRemoveOccludedTrianglesOp::SetTransform(const FTransform& Transform)
 
 void FRemoveOccludedTrianglesOp::CalculateResult(FProgressCancel* Progress)
 {
-	if (Progress->Cancelled())
+	if (Progress && Progress->Cancelled())
 	{
 		return;
 	}
@@ -70,7 +70,7 @@ void FRemoveOccludedTrianglesOp::CalculateResult(FProgressCancel* Progress)
 		TRemoveOccludedTriangles<FDynamicMesh3> SelfJacket(ResultMesh.Get());
 		FDynamicMeshAABBTree3 SelfAABB(ResultMesh.Get());
 		TFastWindingTree<FDynamicMesh3> SelfFWTree(&SelfAABB);
-		if (Progress->Cancelled())
+		if (Progress && Progress->Cancelled())
 		{
 			return;
 		}
@@ -91,7 +91,7 @@ void FRemoveOccludedTrianglesOp::CalculateResult(FProgressCancel* Progress)
 	{
 		TRemoveOccludedTriangles<TIndexMeshArrayAdapter<int, double, FVector3d>> Jacket(ResultMesh.Get());
 
-		if (Progress->Cancelled())
+		if (Progress && Progress->Cancelled())
 		{
 			return;
 		}
diff --git a/Engine/Plugins/Experimental/MeshModelingToolset/Source/ModelingOperators/Private/CompositionOps/BooleanMeshesOp.cpp b/Engine/Plugins/Experimental/MeshModelingToolset/Source/ModelingOperators/Private/CompositionOps/BooleanMeshesOp.cpp
--- a/Engine/Plugins/Experimental/MeshModelingToolset/Source/ModelingOperators/Private/CompositionOps/BooleanMeshesOp.cpp
+++ b/Engine/Plugins/Experimental/MeshModelingToolset/Source/ModelingOperators/Private/CompositionOps/BooleanMeshesOp.cpp
@@ -15,7 +15,7 @@ void FBooleanMeshesOp::SetTransform(const FTransform& Transform) {
 
 void FBooleanMeshesOp::CalculateResult(FProgressCancel* Progress)
 {
-	if (Progress->Cancelled())
+	if (Progress && Progress->Cancelled())
 	{
 		return;
 	}
@@ -66,7 +66,7 @@ void FBooleanMeshesOp::CalculateResult(FProgressCancel* Progress)
 	}
 
 	FMeshBoolean MeshBoolean(Meshes[FirstIdx].Get(), (FTransform3d)Transforms[FirstIdx], Meshes[OtherIdx].Get(), (FTransform3d)Transforms[OtherIdx], ResultMesh.Get(), Op);
-	if (Progress->Cancelled())
+	if (Progress && Progress->Cancelled())
 	{
 		return;
 	}
@@ -77,7 +77,7 @@ void FBooleanMeshesOp::CalculateResult(FProgressCancel* Progress)
 	bool bSuccess = MeshBoolean.Compute();
 	ResultTransform = MeshBoolean.ResultTransform;
 
-	if (Progress->Cancelled())
+	if (Progress && Progress->Cancelled())
 	{
 		return;
 	}
@@ -129,7 +129,7 @@ void FBooleanMeshesOp::CalculateResult(FProgressCancel* Progress)
 		};
 		OpenBoundary.Compute();
 
-		if (Progress->Cancelled())
+		if (Progress && Progress->Cancelled())
 		{
 			return;
 		}
diff --git a/Engine/Plugins/Experimental/MeshModelingToolset/Source/ModelingOperators/Private/CompositionOps/CurveSweepOp.cpp b/Engine/Plugins/Experimental/MeshModelingToolset/Source/ModelingOperators/Private/CompositionOps/CurveSweepOp.cpp
--- a/Engine/Plugins/Experimental/MeshModelingToolset/Source/ModelingOperators/Private/CompositionOps/CurveSweepOp.cpp
+++ b/Engine/Plugins/Experimental/MeshModelingToolset/Source/ModelingOperators/Private/CompositionOps/CurveSweepOp.cpp
@@ -14,7 +14,7 @@ void FCurveSweepOp::CalculateResult(FProgressCancel* Progress)
 {
 	check(Progress != nullptr);
 
-	if (Progress->Cancelled())
+	if (Progress && Progress->Cancelled())
 	{
 		return;
 	}
@@ -50,7 +50,7 @@ void FCurveSweepOp::CalculateResult(FProgressCancel* Progress)
 	CurveSweeper.DiagonalTolerance = DiagonalTolerance;
 	CurveSweeper.Generate();
 
-	if (Progress->Cancelled())
+	if (Progress && Progress->Cancelled())
 	{
 		return;
 	}
@@ -62,7 +62,7 @@ void FCurveSweepOp::CalculateResult(FProgressCancel* Progress)
 	{
 		for (int CapIndex = 0; CapIndex < 2; ++CapIndex)
 		{
-			if (Progress->Cancelled())
+			if (Progress && Progress->Cancelled())
 			{
 				return;
 			}
@@ -154,7 +154,7 @@ void FCurveSweepOp::CalculateResult(FProgressCancel* Progress)
 
 	if (bSharpNormals)
 	{
-		if (Progress->Cancelled())
+		if (Progress && Progress->Cancelled())
 		{
 			return;
 		}
@@ -172,7 +172,7 @@ void FCurveSweepOp::CalculateResult(FProgressCancel* Progress)
 		NormalsUtility.ComputeTriangleNormals();
 		const TArray<FVector3d>& Normals = NormalsUtility.GetNormals();
 
-		if (Progress->Cancelled())
+		if (Progress && Progress->Cancelled())
 		{
 			return;
 		}
@@ -184,7 +184,7 @@ void FCurveSweepOp::CalculateResult(FProgressCancel* Progress)
 				return Normals[TA].Dot(Normals[TB]) > NormalDotProdThreshold;
 			}, 0);
 
-		if (Progress->Cancelled())
+		if (Progress && Progress->Cancelled())
 		{
 			return;
 		}
@@ -194,7 +194,7 @@ void FCurveSweepOp::CalculateResult(FProgressCancel* Progress)
 		NormalsUtility.CopyToOverlay(ResultMesh->Attributes()->PrimaryNormals(), false);
 	}
 
-	if (Progress->Cancelled())
+	if (Progress && Progress->Cancelled())
 	{
 		return;
 	}
diff --git a/Engine/Plugins/Experimental/MeshModelingToolset/Source/ModelingOperators/Private/CompositionOps/MirrorOp.cpp b/Engine/Plugins/Experimental/MeshModelingToolset/Source/ModelingOperators/Private/CompositionOps/MirrorOp.cpp
--- a/Engine/Plugins/Experimental/MeshModelingToolset/Source/ModelingOperators/Private/CompositionOps/MirrorOp.cpp
+++ b/Engine/Plugins/Experimental/MeshModelingToolset/Source/ModelingOperators/Private/CompositionOps/MirrorOp.cpp
@@ -11,14 +11,14 @@ void FMirrorOp::SetTransform(const FTransform& Transform) {
 
 void FMirrorOp::CalculateResult(FProgressCancel* Progress)
 {
-	if (Progress->Cancelled())
+	if (Progress && Progress->Cancelled())
 	{
 		return;
 	}
 
 	ResultMesh->Copy(*OriginalMesh, true, true, true, true);
 	
-	if (Progress->Cancelled())
+	if (Progress && Progress->Cancelled())
 	{
 		return;
 	}
@@ -36,7 +36,7 @@ void FMirrorOp::CalculateResult(FProgressCancel* Progress)
 		Cutter.PlaneTolerance = PlaneTolerance;
 		Cutter.Cut();
 
-		if (Progress->Cancelled())
+		if (Progress && Progress->Cancelled())
 		{
 			return;
 		}
diff --git a/Engine/Plugins/Experimental/MeshModelingToolset/Source/ModelingOperators/Private/CompositionOps/SelfUnionMeshesOp.cpp b/Engine/Plugins/Experimental/MeshModelingToolset/Source/ModelingOperators/Private/CompositionOps/SelfUnionMeshesOp.cpp
--- a/Engine/Plugins/Experimental/MeshModelingToolset/Source/ModelingOperators/Private/CompositionOps/SelfUnionMeshesOp.cpp
+++ b/Engine/Plugins/Experimental/MeshModelingToolset/Source/ModelingOperators/Private/CompositionOps/SelfUnionMeshesOp.cpp
@@ -15,14 +15,14 @@ void FSelfUnionMeshesOp::SetTransform(const FTransform& Transform) {
 
 void FSelfUnionMeshesOp::CalculateResult(FProgressCancel* Progress)
 {
-	if (Progress->Cancelled())
+	if (Progress && Progress->Cancelled())
 	{
 		return;
 	}
 
 	*ResultMesh = *CombinedMesh;
 
-	if (Progress->Cancelled())
+	if (Progress && Progress->Cancelled())
 	{
 		return;
 	}
@@ -33,7 +33,7 @@ void FSelfUnionMeshesOp::CalculateResult(FProgressCancel* Progress)
 	Union.bTrackAllNewEdges = (bTryCollapseExtraEdges);
 	bool bSuccess = Union.Compute();
 
-	if (Progress->Cancelled())
+	if (Progress && Progress->Cancelled())
 	{
 		return;
 	}
@@ -84,7 +84,7 @@ void FSelfUnionMeshesOp::CalculateResult(FProgressCancel* Progress)
 		};
 		OpenBoundary.Compute();
 
-		if (Progress->Cancelled())
+		if (Progress && Progress->Cancelled())
 		{
 			return;
 		}
diff --git a/Engine/Plugins/Experimental/MeshModelingToolset/Source/ModelingOperators/Private/CompositionOps/VoxelBlendMeshesOp.cpp b/Engine/Plugins/Experimental/MeshModelingToolset/Source/ModelingOperators/Private/CompositionOps/VoxelBlendMeshesOp.cpp
--- a/Engine/Plugins/Experimental/MeshModelingToolset/Source/ModelingOperators/Private/CompositionOps/VoxelBlendMeshesOp.cpp
+++ b/Engine/Plugins/Experimental/MeshModelingToolset/Source/ModelingOperators/Private/CompositionOps/VoxelBlendMeshesOp.cpp
@@ -22,7 +22,7 @@ void FVoxelBlendMeshesOp::SetTransform(const FTransform& Transform) {
 
 void FVoxelBlendMeshesOp::CalculateResult(FProgressCancel* Progress)
 {
-	if (Progress->Cancelled())
+	if (Progress && Progress->Cancelled())
 	{
 		return;
 	}
@@ -78,10 +78,10 @@ void FVoxelBlendMeshesOp::CalculateResult(FProgressCancel* Progress)
 
 	ImplicitBlend.CancelF = [&Progress]()
 	{
-		return Progress->Cancelled();
+		return Progress && Progress->Cancelled();
 	};
 
-	if (Progress->Cancelled())
+	if (Progress && Progress->Cancelled())
 	{
 		return;
 	}
diff --git a/Engine/Plugins/Experimental/MeshModelingToolset/Source/ModelingOperators/Private/CompositionOps/VoxelMorphologyMeshesOp.cpp b/Engine/Plugins/Experimental/MeshModelingToolset/Source/ModelingOperators/Private/CompositionOps/VoxelMorphologyMeshesOp.cpp
--- a/Engine/Plugins/Experimental/MeshModelingToolset/Source/ModelingOperators/Private/CompositionOps/VoxelMorphologyMeshesOp.cpp
+++ b/Engine/Plugins/Experimental/MeshModelingToolset/Source/ModelingOperators/Private/CompositionOps/VoxelMorphologyMeshesOp.cpp
@@ -23,7 +23,7 @@ void FVoxelMorphologyMeshesOp::SetTransform(const FTransform& Transform) {
 
 void FVoxelMorphologyMeshesOp::CalculateResult(FProgressCancel* Progress)
 {
-	if (Progress->Cancelled())
+	if (Progress && Progress->Cancelled())
 	{
 		return;
 	}
@@ -112,12 +112,12 @@ void FVoxelMorphologyMeshesOp::CalculateResult(FProgressCancel* Progress)
 	
 	ImplicitMorphology.CancelF = [&Progress]()
 	{
-		return Progress->Cancelled();
+		return Progress && Progress->Cancelled();
 	};
 
 	;
 
-	if (Progress->Cancelled())
+	if (Progress && Progress->Cancelled())
 	{
 		return;
 	}
diff --git a/Engine/Plugins/Experimental/MeshModelingToolset/Source/ModelingOperators/Private/CompositionOps/VoxelSolidifyMeshesOp.cpp b/Engine/Plugins/Experimental/MeshModelingToolset/Source/ModelingOperators/Private/CompositionOps/VoxelSolidifyMeshesOp.cpp
--- a/Engine/Plugins/Experimental/MeshModelingToolset/Source/ModelingOperators/Private/CompositionOps/VoxelSolidifyMeshesOp.cpp
+++ b/Engine/Plugins/Experimental/MeshModelingToolset/Source/ModelingOperators/Private/CompositionOps/VoxelSolidifyMeshesOp.cpp
@@ -21,7 +21,7 @@ void FVoxelSolidifyMeshesOp::SetTransform(const FTransform& Transform) {
 
 void FVoxelSolidifyMeshesOp::CalculateResult(FProgressCancel* Progress)
 {
-	if (Progress->Cancelled())
+	if (Progress && Progress->Cancelled())
 	{
 		return;
 	}
@@ -75,7 +75,7 @@ void FVoxelSolidifyMeshesOp::CalculateResult(FProgressCancel* Progress)
 	Solidify.bSolidAtBoundaries = bSolidAtBoundaries;
 	Solidify.ExtendBounds = ExtendBounds;
 
-	if (Progress->Cancelled())
+	if (Progress && Progress->Cancelled())
 	{
 		return;
 	}
diff --git a/Engine/Plugins/Experimental/MeshModelingToolset/Source/ModelingOperators/Private/CuttingOps/EdgeLoopInsertionOp.cpp b/Engine/Plugins/Experimental/MeshModelingToolset/Source/ModelingOperators/Private/CuttingOps/EdgeLoopInsertionOp.cpp
--- a/Engine/Plugins/Experimental/MeshModelingToolset/Source/ModelingOperators/Private/CuttingOps/EdgeLoopInsertionOp.cpp
+++ b/Engine/Plugins/Experimental/MeshModelingToolset/Source/ModelingOperators/Private/CuttingOps/EdgeLoopInsertionOp.cpp
@@ -23,7 +23,7 @@ void FEdgeLoopInsertionOp::CalculateResult(FProgressCancel* Progress)
 {
 	bSucceeded = false;
 
-	if (Progress->Cancelled())
+	if (Progress && Progress->Cancelled())
 	{
 		return;
 	}
@@ -33,7 +33,7 @@ void FEdgeLoopInsertionOp::CalculateResult(FProgressCancel* Progress)
 	*ResultTopology = *OriginalTopology;
 	ResultTopology->RetargetOnClonedMesh(ResultMesh.Get());
 
-	if (Progress->Cancelled())
+	if (Progress && Progress->Cancelled())
 	{
 		return;
 	}
diff --git a/Engine/Plugins/Experimental/MeshModelingToolset/Source/ModelingOperators/Private/CuttingOps/PlaneCutOp.cpp b/Engine/Plugins/Experimental/MeshModelingToolset/Source/ModelingOperators/Private/CuttingOps/PlaneCutOp.cpp
--- a/Engine/Plugins/Experimental/MeshModelingToolset/Source/ModelingOperators/Private/CuttingOps/PlaneCutOp.cpp
+++ b/Engine/Plugins/Experimental/MeshModelingToolset/Source/ModelingOperators/Private/CuttingOps/PlaneCutOp.cpp
@@ -18,7 +18,7 @@ void FPlaneCutOp::SetTransform(const FTransform& Transform) {
 
 void FPlaneCutOp::CalculateResult(FProgressCancel* Progress)
 {
-	if (Progress->Cancelled())
+	if (Progress && Progress->Cancelled())
 	{
 		return;
 	}
@@ -28,14 +28,14 @@ void FPlaneCutOp::CalculateResult(FProgressCancel* Progress)
 	
 	FMeshPlaneCut Cut(ResultMesh.Get(), LocalPlaneOrigin, LocalPlaneNormal);
 
-	if (Progress->Cancelled())
+	if (Progress && Progress->Cancelled())
 	{
 		return;
 	}
 	Cut.UVScaleFactor = UVScaleFactor;
 
 
-	if (Progress->Cancelled())
+	if (Progress && Progress->Cancelled())
 	{
 		return;
 	}
@@ -82,7 +82,7 @@ void FPlaneCutOp::CalculateResult(FProgressCancel* Progress)
 	}
 
 
-	if (Progress->Cancelled())
+	if (Progress && Progress->Cancelled())
 	{
 		return;
 	}
@@ -91,7 +91,7 @@ void FPlaneCutOp::CalculateResult(FProgressCancel* Progress)
 		Cut.HoleFill(ConstrainedDelaunayTriangulate<double>, bFillSpans);
 	}
 
-	if (Progress->Cancelled())
+	if (Progress && Progress->Cancelled())
 	{
 		return;
 	}
diff --git a/Engine/Plugins/Experimental/MeshModelingToolset/Source/ModelingOperators/Private/DeformationOps/LatticeDeformerOp.cpp b/Engine/Plugins/Experimental/MeshModelingToolset/Source/ModelingOperators/Private/DeformationOps/LatticeDeformerOp.cpp
--- a/Engine/Plugins/Experimental/MeshModelingToolset/Source/ModelingOperators/Private/DeformationOps/LatticeDeformerOp.cpp
+++ b/Engine/Plugins/Experimental/MeshModelingToolset/Source/ModelingOperators/Private/DeformationOps/LatticeDeformerOp.cpp
@@ -5,14 +5,14 @@
 
 void FLatticeDeformerOp::CalculateResult(FProgressCancel* Progress)
 {
-	if (Progress->Cancelled())
+	if (Progress && Progress->Cancelled())
 	{
 		return;
 	}
 
 	ResultMesh->Copy(*OriginalMesh);
 
-	if (Progress->Cancelled())
+	if (Progress && Progress->Cancelled())
 	{
 		return;
 	}
diff --git a/Engine/Plugins/Experimental/MeshModelingToolset/Source/ModelingOperators/Private/ParameterizationOps/UVProjectionOp.cpp b/Engine/Plugins/Experimental/MeshModelingToolset/Source/ModelingOperators/Private/ParameterizationOps/UVProjectionOp.cpp
--- a/Engine/Plugins/Experimental/MeshModelingToolset/Source/ModelingOperators/Private/ParameterizationOps/UVProjectionOp.cpp
+++ b/Engine/Plugins/Experimental/MeshModelingToolset/Source/ModelingOperators/Private/ParameterizationOps/UVProjectionOp.cpp
@@ -21,7 +21,7 @@ void FUVProjectionOp::SetTransform(const FTransform& Transform) {
 
 void FUVProjectionOp::CalculateResult(FProgressCancel* Progress)
 {
-	if (Progress->Cancelled())
+	if (Progress && Progress->Cancelled())
 	{
 		return;
 	}
@@ -59,6 +59,11 @@ void FUVProjectionOp::CalculateResult(FProgressCancel* Progress)
 		return FVector2f(float(P[Ax1]) * Ax1Scale * Scale.X + Offset.X, float(P[Ax2]) * Ax2Scale * Scale.Y + Offset.Y);
 	};
 
+	if (Progress && Progress->Cancelled())
+	{
+		return;
+	}
+
 	int MaxVertexID = ResultMesh->MaxVertexID();
 	int MaxTriangleID = ResultMesh->MaxTriangleID();
 	if (ProjectionMethod == EUVProjectionMethod::Plane)
diff --git a/Engine/Plugins/Experimental/MeshModelingToolset/Source/ModelingOperatorsEditorOnly/Private/CleaningOps/SimplifyMeshOp.cpp b/Engine/Plugins/Experimental/MeshModelingToolset/Source/ModelingOperatorsEditorOnly/Private/CleaningOps/SimplifyMeshOp.cpp
--- a/Engine/Plugins/Experimental/MeshModelingToolset/Source/ModelingOperatorsEditorOnly/Private/CleaningOps/SimplifyMeshOp.cpp
+++ b/Engine/Plugins/Experimental/MeshModelingToolset/Source/ModelingOperatorsEditorOnly/Private/CleaningOps/SimplifyMeshOp.cpp
@@ -100,7 +100,7 @@ void ComputeSimplify(FDynamicMesh3* TargetMesh, const bool bReproject,
 
 void FSimplifyMeshOp::CalculateResult(FProgressCancel* Progress)
 {
-	if (Progress->Cancelled())
+	if (Progress && Progress->Cancelled())
 	{
 		return;
 	}
@@ -109,7 +109,7 @@ void FSimplifyMeshOp::CalculateResult(FProgressCancel* Progress)
 	const FMeshDescription* MeshDescription = OriginalMeshDescription.Get();
 	FDynamicMesh3* TargetMesh = ResultMesh.Get();
 
-	if (Progress->Cancelled())
+	if (Progress && Progress->Cancelled())
 	{
 		return;
 	}
@@ -164,7 +164,7 @@ void FSimplifyMeshOp::CalculateResult(FProgressCancel* Progress)
 		const FMeshDescription* SrcMeshDescription = OriginalMeshDescription.Get();
 		FMeshDescription DstMeshDescription(*SrcMeshDescription);
 
-		if (Progress->Cancelled())
+		if (Progress && Progress->Cancelled())
 		{
 			return;
 		}
@@ -172,7 +172,7 @@ void FSimplifyMeshOp::CalculateResult(FProgressCancel* Progress)
 		FOverlappingCorners OverlappingCorners;
 		FStaticMeshOperations::FindOverlappingCorners(OverlappingCorners, *SrcMeshDescription, 1.e-5);
 
-		if (Progress->Cancelled())
+		if (Progress && Progress->Cancelled())
 		{
 			return;
 		}
@@ -210,7 +210,7 @@ void FSimplifyMeshOp::CalculateResult(FProgressCancel* Progress)
 			MeshReduction->ReduceMeshDescription(DstMeshDescription, Error, *SrcMeshDescription, OverlappingCorners, ReductionSettings);
 		}
 
-		if (Progress->Cancelled())
+		if (Progress && Progress->Cancelled())
 		{
 			return;
 		}
@@ -237,7 +237,7 @@ void FSimplifyMeshOp::CalculateResult(FProgressCancel* Progress)
 				bFailedModifyNeedsRegen = true;
 			}
 
-			if (Progress->Cancelled())
+			if (Progress && Progress->Cancelled())
 			{
 				return;
 			}
@@ -247,7 +247,7 @@ void FSimplifyMeshOp::CalculateResult(FProgressCancel* Progress)
 				bFailedModifyNeedsRegen = true;
 			}
 
-			if (Progress->Cancelled())
+			if (Progress && Progress->Cancelled())
 			{
 				return;
 			}
@@ -266,7 +266,7 @@ void FSimplifyMeshOp::CalculateResult(FProgressCancel* Progress)
 	}
 
 
-	if (Progress->Cancelled())
+	if (Progress && Progress->Cancelled())
 	{
 		return;
 	}
diff --git a/Engine/Plugins/Experimental/MeshModelingToolset/Source/ModelingOperatorsEditorOnly/Private/CompositionOps/VoxelBooleanMeshesOp.cpp b/Engine/Plugins/Experimental/MeshModelingToolset/Source/ModelingOperatorsEditorOnly/Private/CompositionOps/VoxelBooleanMeshesOp.cpp
--- a/Engine/Plugins/Experimental/MeshModelingToolset/Source/ModelingOperatorsEditorOnly/Private/CompositionOps/VoxelBooleanMeshesOp.cpp
+++ b/Engine/Plugins/Experimental/MeshModelingToolset/Source/ModelingOperatorsEditorOnly/Private/CompositionOps/VoxelBooleanMeshesOp.cpp
@@ -73,7 +73,7 @@ void FVoxelBooleanMeshesOp::CalculateResult(FProgressCancel* Progress)
 
 	ResultTransform = FTransform3d(MergedOrigin);
 
-	if (Progress->Cancelled())
+	if (Progress && Progress->Cancelled())
 	{
 		return;
 	}
@@ -90,7 +90,7 @@ void FVoxelBooleanMeshesOp::CalculateResult(FProgressCancel* Progress)
 
 		for (int32 i = 0; i < NumRemeshes; ++i)
 		{
-			if (Progress->Cancelled())
+			if (Progress && Progress->Cancelled())
 			{
 				return;
 			}
@@ -114,7 +114,7 @@ void FVoxelBooleanMeshesOp::CalculateResult(FProgressCancel* Progress)
 	FMeshDescriptionToDynamicMesh Converter;
 	Converter.Convert(&ResultMeshDescription, *ResultMesh);
 
-	if (Progress->Cancelled())
+	if (Progress && Progress->Cancelled())
 	{
 		return;
 	}
diff --git a/Engine/Plugins/Experimental/MeshModelingToolset/Source/ModelingOperatorsEditorOnly/Private/CompositionOps/VoxelMergeMeshesOp.cpp b/Engine/Plugins/Experimental/MeshModelingToolset/Source/ModelingOperatorsEditorOnly/Private/CompositionOps/VoxelMergeMeshesOp.cpp
--- a/Engine/Plugins/Experimental/MeshModelingToolset/Source/ModelingOperatorsEditorOnly/Private/CompositionOps/VoxelMergeMeshesOp.cpp
+++ b/Engine/Plugins/Experimental/MeshModelingToolset/Source/ModelingOperatorsEditorOnly/Private/CompositionOps/VoxelMergeMeshesOp.cpp
@@ -34,7 +34,7 @@ void FVoxelMergeMeshesOp::CalculateResult(FProgressCancel* Progress)
 	FVector MergedOrigin = VoxelCSGTool->ComputeUnion(*InputMeshArray, MergedMeshesDescription, AdaptivityD, CSGIsoSurface);
 	ResultTransform = FTransform3d(MergedOrigin);
 
-	if (Progress->Cancelled())
+	if (Progress && Progress->Cancelled())
 	{
 		return;
 	}
@@ -51,7 +51,7 @@ void FVoxelMergeMeshesOp::CalculateResult(FProgressCancel* Progress)
 
 		for (int32 i = 0; i < NumRemeshes; ++i)
 		{
-			if (Progress->Cancelled())
+			if (Progress && Progress->Cancelled())
 			{
 				return;
 			}
@@ -74,7 +74,7 @@ void FVoxelMergeMeshesOp::CalculateResult(FProgressCancel* Progress)
 	FMeshDescriptionToDynamicMesh Converter;
 	Converter.Convert(&MergedMeshesDescription, *ResultMesh);
 
-	if (Progress->Cancelled())
+	if (Progress && Progress->Cancelled())
 	{
 		return;
 	}
diff --git a/Engine/Plugins/Experimental/MeshModelingToolset/Source/ModelingOperatorsEditorOnly/Private/CuttingOps/EmbedPolygonsOp.cpp b/Engine/Plugins/Experimental/MeshModelingToolset/Source/ModelingOperatorsEditorOnly/Private/CuttingOps/EmbedPolygonsOp.cpp
--- a/Engine/Plugins/Experimental/MeshModelingToolset/Source/ModelingOperatorsEditorOnly/Private/CuttingOps/EmbedPolygonsOp.cpp
+++ b/Engine/Plugins/Experimental/MeshModelingToolset/Source/ModelingOperatorsEditorOnly/Private/CuttingOps/EmbedPolygonsOp.cpp
@@ -166,7 +166,7 @@ void FEmbedPolygonsOp::BooleanPath(FProgressCancel* Progress)
 		unimplemented();
 	}
 
-	if (Progress->Cancelled())
+	if (Progress && Progress->Cancelled())
 	{
 		return;
 	}
@@ -177,7 +177,7 @@ void FEmbedPolygonsOp::BooleanPath(FProgressCancel* Progress)
 	Boolean.bPutResultInInputSpace = true;
 	bool bBoolSuccess = Boolean.Compute();
 
-	if (Progress->Cancelled())
+	if (Progress && Progress->Cancelled())
 	{
 		return;
 	}
@@ -192,7 +192,7 @@ void FEmbedPolygonsOp::BooleanPath(FProgressCancel* Progress)
 		};
 		OpenBoundary.Compute();
 
-		if (Progress->Cancelled())
+		if (Progress && Progress->Cancelled())
 		{
 			return;
 		}
@@ -234,7 +234,7 @@ void FEmbedPolygonsOp::CalculateResult(FProgressCancel* Progress)
 {
 	bOperationSucceeded = false;
 
-	if (Progress->Cancelled())
+	if (Progress && Progress->Cancelled())
 	{
 		return;
 	}
@@ -280,7 +280,7 @@ void FEmbedPolygonsOp::CalculateResult(FProgressCancel* Progress)
 		}
 	}
 
-	if (Progress->Cancelled())
+	if (Progress && Progress->Cancelled())
 	{
 		return;
 	}
@@ -418,7 +418,7 @@ void FEmbedPolygonsOp::CalculateResult(FProgressCancel* Progress)
 		return bResult;
 	};
 
-	if (Progress->Cancelled())
+	if (Progress && Progress->Cancelled())
 	{
 		return;
 	}
diff --git a/Engine/Plugins/Experimental/MeshModelingToolset/Source/ModelingOperatorsEditorOnly/Private/ParameterizationOps/CalculateTangentsOp.cpp b/Engine/Plugins/Experimental/MeshModelingToolset/Source/ModelingOperatorsEditorOnly/Private/ParameterizationOps/CalculateTangentsOp.cpp
--- a/Engine/Plugins/Experimental/MeshModelingToolset/Source/ModelingOperatorsEditorOnly/Private/ParameterizationOps/CalculateTangentsOp.cpp
+++ b/Engine/Plugins/Experimental/MeshModelingToolset/Source/ModelingOperatorsEditorOnly/Private/ParameterizationOps/CalculateTangentsOp.cpp
@@ -12,7 +12,7 @@
 
 void FCalculateTangentsOp::CalculateResult(FProgressCancel* Progress)
 {
-	if (Progress->Cancelled())
+	if (Progress && Progress->Cancelled())
 	{
 		return;
 	}
diff --git a/Engine/Plugins/Experimental/MeshModelingToolset/Source/ModelingOperatorsEditorOnly/Private/ParameterizationOps/ParameterizeMeshOp.cpp b/Engine/Plugins/Experimental/MeshModelingToolset/Source/ModelingOperatorsEditorOnly/Private/ParameterizationOps/ParameterizeMeshOp.cpp
--- a/Engine/Plugins/Experimental/MeshModelingToolset/Source/ModelingOperatorsEditorOnly/Private/ParameterizationOps/ParameterizeMeshOp.cpp
+++ b/Engine/Plugins/Experimental/MeshModelingToolset/Source/ModelingOperatorsEditorOnly/Private/ParameterizationOps/ParameterizeMeshOp.cpp
@@ -376,7 +376,7 @@ void FParameterizeMeshOp::CalculateResult(FProgressCancel* Progress)
 
 	ResultMesh = MakeUnique<FDynamicMesh3>(*InputMesh);
 
-	if (Progress->Cancelled())
+	if (Progress && Progress->Cancelled())
 	{
 		return;
 	}
@@ -394,7 +394,7 @@ void FParameterizeMeshOp::CalculateResult(FProgressCancel* Progress)
 	const bool bSplitIntoMultiple = true;
 
 	// The UV atlas callback uses a float progress-based interrupter. 
-	TFunction<bool(float)> Iterrupter = [Progress](float)->bool {return !Progress->Cancelled(); };
+	TFunction<bool(float)> Iterrupter = [Progress](float)->bool {return !(Progress && Progress->Cancelled()); };
 
 	if (UnwrapType == EParamOpUnwrapType::MinStretch && IslandMode == EParamOpIslandMode::Auto)
 	{
@@ -457,7 +457,7 @@ void FParameterizeMeshOp::CalculateResult(FProgressCancel* Progress)
 			if (bComputedUVs) SuccessCount++;
 		}
 
-		if (Progress->Cancelled())
+		if (Progress && Progress->Cancelled())
 		{
 			return;
 		}
@@ -504,7 +504,7 @@ void FParameterizeMeshOp::CalculateResult(FProgressCancel* Progress)
 				}
 			}
 
-			if (Progress->Cancelled())
+			if (Progress && Progress->Cancelled())
 			{
 				return;
 			}
