Fix crash when setting subdivision level too high. Cap the subdivision level in SubdividePolyTool.

Also fix a bug applying Loop subdivision to a non-compact mesh.

#rnx

diff --git a/Engine/Plugins/Experimental/MeshModelingToolset/Source/MeshModelingToolsEditorOnly/Private/SubdividePoly.cpp b/Engine/Plugins/Experimental/MeshModelingToolset/Source/MeshModelingToolsEditorOnly/Private/SubdividePoly.cpp
--- a/Engine/Plugins/Experimental/MeshModelingToolset/Source/MeshModelingToolsEditorOnly/Private/SubdividePoly.cpp
+++ b/Engine/Plugins/Experimental/MeshModelingToolset/Source/MeshModelingToolsEditorOnly/Private/SubdividePoly.cpp
@@ -303,26 +303,23 @@ bool FSubdividePoly::ComputeTopologySubdivision()
 
 	TArray<int> BoundaryVertsPerFace;
 	TArray<int> NumVertsPerFace;
-	int TotalNumVertices = 0;
 	OpenSubdiv::Far::TopologyDescriptor::FVarChannel UVChannel;
 	TArray<int> UVBuffer;
 
 	auto DescriptorFromTriangleMesh = [this, &NumVertsPerFace, &BoundaryVertsPerFace, &UVChannel, &UVBuffer]
 	(OpenSubdiv::Far::TopologyDescriptor& Descriptor)
 	{
-		int TotalNumVertices = 0;
-
 		for (auto TriangleID : OriginalMesh.TriangleIndicesItr())
 		{
 			FIndex3i TriangleVertices = OriginalMesh.GetTriangle(TriangleID);
 			NumVertsPerFace.Add(3);
-			TotalNumVertices += 3;
 			BoundaryVertsPerFace.Add(TriangleVertices[0]);
 			BoundaryVertsPerFace.Add(TriangleVertices[1]);
 			BoundaryVertsPerFace.Add(TriangleVertices[2]);
 		}
 
-		Descriptor.numVertices = TotalNumVertices;
+		// TODO: We should probably create a compact mesh descriptor for subdivision to operate on. UETOOL-2944
+		Descriptor.numVertices = OriginalMesh.MaxVertexID();
 		Descriptor.numFaces = OriginalMesh.TriangleCount();
 		Descriptor.numVertsPerFace = NumVertsPerFace.GetData();
 		Descriptor.vertIndicesPerFace = BoundaryVertsPerFace.GetData();
diff --git a/Engine/Plugins/Experimental/MeshModelingToolset/Source/MeshModelingToolsEditorOnly/Private/SubdividePolyTool.cpp b/Engine/Plugins/Experimental/MeshModelingToolset/Source/MeshModelingToolsEditorOnly/Private/SubdividePolyTool.cpp
--- a/Engine/Plugins/Experimental/MeshModelingToolset/Source/MeshModelingToolsEditorOnly/Private/SubdividePolyTool.cpp
+++ b/Engine/Plugins/Experimental/MeshModelingToolset/Source/MeshModelingToolsEditorOnly/Private/SubdividePolyTool.cpp
@@ -121,6 +121,42 @@ bool USubdividePolyTool::CheckGroupTopology(FText& Message)
 	return true;
 }
 
+
+void USubdividePolyTool::CapSubdivisionLevel(ESubdivisionScheme Scheme, int DesiredLevel)
+{
+	// Stolen from UDisplaceMeshTool::ValidateSubdivisions
+	constexpr int MaxFaces = 3000000;
+
+	int NumOriginalFaces = MaxFaces;
+	if (Scheme == ESubdivisionScheme::Loop)
+	{
+		NumOriginalFaces = OriginalMesh->TriangleCount();
+	}
+	else
+	{
+		constexpr bool bAutoCompute = true;
+		FGroupTopology Topo(OriginalMesh.Get(), bAutoCompute);
+		NumOriginalFaces = Topo.Groups.Num();
+	}
+	int MaxLevel = (int)floor(log2(MaxFaces / NumOriginalFaces) / 2.0);
+
+	if (DesiredLevel > MaxLevel)
+	{
+		FText WarningText = FText::Format(LOCTEXT("SubdivisionLevelTooHigh", "Subdivision level clamped: desired subdivision level ({0}) exceeds maximum level ({1}) for a mesh with this number of faces."),
+										  FText::AsNumber(DesiredLevel),
+										  FText::AsNumber(MaxLevel));
+		GetToolManager()->DisplayMessage(WarningText, EToolMessageLevel::UserWarning);
+		Properties->SubdivisionLevel = MaxLevel;
+		Properties->SilentUpdateWatched();		// Don't trigger this function again due to setting SubdivisionLevel above
+	}
+	else
+	{
+		// Clear possible lingering warning message
+		GetToolManager()->DisplayMessage(FText(), EToolMessageLevel::UserWarning);
+	}
+}
+
+
 void USubdividePolyTool::Setup()
 {
 	UInteractiveTool::Setup();
@@ -178,7 +214,8 @@ void USubdividePolyTool::Setup()
 	}
 
 	check(Properties->SubdivisionLevel >= 1);	// Should be enforced by UPROPERTY meta tags
-
+	
+	CapSubdivisionLevel(Properties->SubdivisionScheme, Properties->SubdivisionLevel);
 	PreviewDynamicMeshComponent->SetRenderMeshPostProcessor(MakeUnique<SubdivPostProcessor>(Properties->SubdivisionLevel,
 																							Properties->SubdivisionScheme,
 																							Properties->NormalComputationMethod,
@@ -211,15 +248,19 @@ void USubdividePolyTool::Setup()
 	};
 
 	// Watch for property changes
-	Properties->WatchProperty(Properties->SubdivisionLevel, [this, RebuildMeshPostProcessor](int)
+	Properties->WatchProperty(Properties->SubdivisionLevel, [this, RebuildMeshPostProcessor](int NewSubdLevel)
 	{
+		CapSubdivisionLevel(Properties->SubdivisionScheme, NewSubdLevel);
 		RebuildMeshPostProcessor();
 	});
-	Properties->WatchProperty(Properties->SubdivisionScheme, [this, RebuildMeshPostProcessor](ESubdivisionScheme)
+
+	Properties->WatchProperty(Properties->SubdivisionScheme, [this, RebuildMeshPostProcessor](ESubdivisionScheme NewScheme)
 	{
+		CapSubdivisionLevel(NewScheme, Properties->SubdivisionLevel);
 		RebuildMeshPostProcessor();
 		bPreviewGeometryNeedsUpdate = true;		// Switch from rendering poly cage to all triangle edges
 	});
+
 	Properties->WatchProperty(Properties->NormalComputationMethod, [this, RebuildMeshPostProcessor](ESubdivisionOutputNormals)
 	{
 		RebuildMeshPostProcessor();
diff --git a/Engine/Plugins/Experimental/MeshModelingToolset/Source/MeshModelingToolsEditorOnly/Public/SubdividePolyTool.h b/Engine/Plugins/Experimental/MeshModelingToolset/Source/MeshModelingToolsEditorOnly/Public/SubdividePolyTool.h
--- a/Engine/Plugins/Experimental/MeshModelingToolset/Source/MeshModelingToolsEditorOnly/Public/SubdividePolyTool.h
+++ b/Engine/Plugins/Experimental/MeshModelingToolset/Source/MeshModelingToolsEditorOnly/Public/SubdividePolyTool.h
@@ -110,5 +110,7 @@ protected:
 	// Determine if the mesh group topology can be used for Catmull-Clark or Bilinear subdivision. If not, we can only 
 	// Loop subdivision on the original triangle mesh.
 	bool CheckGroupTopology(FText& Message);
+
+	void CapSubdivisionLevel(ESubdivisionScheme Scheme, int DesiredLevel);
 };
 
