For Jacketing tool, add option to set a new polygroup ID for occluded triangles, instead of deleting the triangles

diff --git a/Engine/Plugins/Experimental/MeshModelingToolset/Source/MeshModelingTools/Private/RemoveOccludedTrianglesTool.cpp b/Engine/Plugins/Experimental/MeshModelingToolset/Source/MeshModelingTools/Private/RemoveOccludedTrianglesTool.cpp
--- a/Engine/Plugins/Experimental/MeshModelingToolset/Source/MeshModelingTools/Private/RemoveOccludedTrianglesTool.cpp
+++ b/Engine/Plugins/Experimental/MeshModelingToolset/Source/MeshModelingTools/Private/RemoveOccludedTrianglesTool.cpp
@@ -7,6 +7,7 @@
 #include "ToolSetupUtil.h"
 
 #include "DynamicMesh3.h"
+#include "Polygroups/PolygroupUtil.h"
 #include "BaseBehaviors/MultiClickSequenceInputBehavior.h"
 #include "Selection/SelectClickedAction.h"
 
@@ -136,15 +137,26 @@ void URemoveOccludedTrianglesTool::Setup()
 			EToolMessageLevel::UserWarning);
 	}
 
+	// initialize the PreviewMesh+BackgroundCompute object
+	SetupPreviews();
+
 	BasicProperties = NewObject<URemoveOccludedTrianglesToolProperties>(this, TEXT("Remove Occluded Triangle Settings"));
 	AdvancedProperties = NewObject<URemoveOccludedTrianglesAdvancedProperties>(this, TEXT("Advanced Settings"));
+	MakePolygroupLayerProperties();
+
+	BasicProperties->WatchProperty(BasicProperties->Action,
+		[this](EOccludedAction Action) { SetToolPropertySourceEnabled(PolygroupLayersProperties, Action == EOccludedAction::SetNewGroup); }
+	);
 
 	// initialize our properties
 	AddToolPropertySource(BasicProperties);
+	AddToolPropertySource(PolygroupLayersProperties);
 	AddToolPropertySource(AdvancedProperties);
 
-	// initialize the PreviewMesh+BackgroundCompute object
-	SetupPreviews();
+	for (UMeshOpPreviewWithBackgroundCompute* Preview : Previews)
+	{
+		Preview->InvalidateResult();
+	}
 
 	GetToolManager()->DisplayMessage(
 		LOCTEXT("RemoveOccludedTrianglesToolDescription", "Remove triangles that are fully contained within the selected Meshes, and hence cannot be visible with opaque shading."),
@@ -152,6 +164,34 @@ void URemoveOccludedTrianglesTool::Setup()
 }
 
 
+void URemoveOccludedTrianglesTool::MakePolygroupLayerProperties()
+{
+	PolygroupLayersProperties = NewObject<UPolygroupLayersProperties>(this, TEXT("Polygroup Layer"));
+	auto GetGroupLayerNames = [](const FDynamicMesh3& Mesh)
+	{
+		TSet<FName> Names;
+
+		if (Mesh.Attributes())
+		{
+			for (int32 LayerIdx = 0; LayerIdx < Mesh.Attributes()->NumPolygroupLayers(); LayerIdx++)
+			{
+				FName Name = Mesh.Attributes()->GetPolygroupLayer(LayerIdx)->GetName();
+				Names.Add(Name);
+			}
+		}
+
+		return Names;
+	};
+	check(OriginalDynamicMeshes.Num() > 0);
+	TSet<FName> CommonLayerNames = GetGroupLayerNames(*OriginalDynamicMeshes[0]);
+	for (int32 Idx = 1; Idx < OriginalDynamicMeshes.Num(); Idx++)
+	{
+		CommonLayerNames = CommonLayerNames.Intersect(GetGroupLayerNames(*OriginalDynamicMeshes[Idx]));
+	}
+	PolygroupLayersProperties->InitializeGroupLayers(CommonLayerNames);
+}
+
+
 void URemoveOccludedTrianglesTool::SetupPreviews()
 {	
 	int32 NumTargets = ComponentTargets.Num();
@@ -174,6 +214,12 @@ void URemoveOccludedTrianglesTool::SetupPreviews()
 	auto EnterProgressFrame = [](float Progress) {};
 #endif
 
+	// create a "magic pink" secondary material to mark occluded faces (to be used if we are setting a new triangle group instead of removing faces)
+	UMaterialInterface* OccludedMaterial = ToolSetupUtil::GetSelectionMaterial(FLinearColor(0.9f, 0.1f, 0.9f), GetToolManager());
+
+	OccludedGroupIDs.Init(-1, NumPreviews);
+	OccludedGroupLayers.Init(-1, NumPreviews);
+
 	OriginalDynamicMeshes.SetNum(NumPreviews);
 	PreviewToCopyIdx.Reset(); PreviewToCopyIdx.SetNum(NumPreviews);
 	for (int32 TargetIdx = 0; TargetIdx < NumTargets; TargetIdx++)
@@ -182,6 +228,26 @@ void URemoveOccludedTrianglesTool::SetupPreviews()
 		
 		int PreviewIdx = TargetToPreviewIdx[TargetIdx];
 
+		// used to choose which triangles need to use the special "OccludedMaterial" secondary material, when the Occluded Action == SetNewGroup
+		auto IsOccludedGroupFn = [this, PreviewIdx](const FDynamicMesh3* Mesh, int32 TriangleID)
+		{
+			int GroupID = OccludedGroupIDs[PreviewIdx];
+			if (GroupID >= 0)
+			{
+				int LayerIndex = OccludedGroupLayers[PreviewIdx];
+				if (LayerIndex < 0)
+				{
+					return Mesh->GetTriangleGroup(TriangleID) == GroupID;
+				}
+				else
+				{
+					check(Mesh->HasAttributes() && Mesh->Attributes()->NumPolygroupLayers() > LayerIndex);
+					return Mesh->Attributes()->GetPolygroupLayer(LayerIndex)->GetValue(TriangleID) == GroupID;
+				}
+			}
+			return false;
+		};
+
 		bool bHasConverted = OriginalDynamicMeshes[PreviewIdx].IsValid();
 
 		if (!bHasConverted)
@@ -206,6 +272,19 @@ void URemoveOccludedTrianglesTool::SetupPreviews()
 			Preview->PreviewMesh->SetTransform(ComponentTargets[TargetIdx]->GetWorldTransform());
 			Preview->PreviewMesh->UpdatePreview(OriginalDynamicMeshes[PreviewIdx].Get());
 			Preview->SetVisibility(true);
+
+			// configure secondary render material
+			Preview->SecondaryMaterial = OccludedMaterial;
+
+			// Set occluded layer index and group IDs
+			Previews[PreviewIdx]->OnOpCompleted.AddLambda([this, PreviewIdx](const FDynamicMeshOperator* UncastOp) {
+				const FRemoveOccludedTrianglesOp* Op = static_cast<const FRemoveOccludedTrianglesOp*>(UncastOp);
+				OccludedGroupIDs[PreviewIdx] = Op->CreatedGroupID;
+				OccludedGroupLayers[PreviewIdx] = Op->CreatedGroupLayerIndex;
+			});
+			
+			// enable secondary triangle buffers
+			Preview->PreviewMesh->EnableSecondaryTriangleBuffers(MoveTemp(IsOccludedGroupFn));
 		}
 		else
 		{
@@ -224,6 +303,9 @@ void URemoveOccludedTrianglesTool::SetupPreviews()
 			
 			PreviewMesh->SetVisible(true);
 
+			PreviewMesh->SetSecondaryRenderMaterial(OccludedMaterial);
+			PreviewMesh->EnableSecondaryTriangleBuffers(MoveTemp(IsOccludedGroupFn));
+
 			Previews[PreviewIdx]->OnMeshUpdated.AddLambda([this, CopyIdx](UMeshOpPreviewWithBackgroundCompute* Compute) {
 				PreviewCopies[CopyIdx]->UpdatePreview(Compute->PreviewMesh->GetPreviewDynamicMesh());
 			});
@@ -233,11 +315,6 @@ void URemoveOccludedTrianglesTool::SetupPreviews()
 	}
 
 	CombinedMeshTrees->BuildAcceleration();
-
-	for (UMeshOpPreviewWithBackgroundCompute* Preview : Previews)
-	{
-		Preview->InvalidateResult();
-	}
 }
 
 
@@ -285,6 +362,10 @@ TUniquePtr<FDynamicMeshOperator> URemoveOccludedTrianglesOperatorFactory::MakeNe
 {
 	TUniquePtr<FRemoveOccludedTrianglesOp> Op = MakeUnique<FRemoveOccludedTrianglesOp>();
 	Op->NormalOffset = Tool->AdvancedProperties->NormalOffset;
+	Op->bSetTriangleGroupInsteadOfRemoving = Tool->BasicProperties->Action == EOccludedAction::SetNewGroup;
+	Op->ActiveGroupLayer = Tool->PolygroupLayersProperties->ActiveGroupLayer;
+	Op->bActiveGroupLayerIsDefault = !Tool->PolygroupLayersProperties->HasSelectedPolygroup();
+
 	switch (Tool->BasicProperties->OcclusionTestMethod)
 	{
 	case EOcclusionCalculationUIMode::GeneralizedWindingNumber:
@@ -356,12 +437,14 @@ void URemoveOccludedTrianglesTool::OnTick(float DeltaTime)
 		for (int CopyIdx : PreviewToCopyIdx[PreviewIdx])
 		{
 			if (bIsWorking)
-			{		
+			{
 				PreviewCopies[CopyIdx]->SetOverrideRenderMaterial(Preview->WorkingMaterial);
+				PreviewCopies[CopyIdx]->ClearSecondaryRenderMaterial();
 			}
 			else
 			{
 				PreviewCopies[CopyIdx]->ClearOverrideRenderMaterial();
+				PreviewCopies[CopyIdx]->SetSecondaryRenderMaterial(Preview->SecondaryMaterial);
 			}
 		}
 	}
diff --git a/Engine/Plugins/Experimental/MeshModelingToolset/Source/MeshModelingTools/Public/RemoveOccludedTrianglesTool.h b/Engine/Plugins/Experimental/MeshModelingToolset/Source/MeshModelingTools/Public/RemoveOccludedTrianglesTool.h
--- a/Engine/Plugins/Experimental/MeshModelingToolset/Source/MeshModelingTools/Public/RemoveOccludedTrianglesTool.h
+++ b/Engine/Plugins/Experimental/MeshModelingToolset/Source/MeshModelingTools/Public/RemoveOccludedTrianglesTool.h
@@ -11,6 +11,8 @@
 #include "DynamicMesh3.h"
 #include "MeshAdapter.h"
 #include "BaseTools/SingleClickTool.h"
+#include "PropertySets/PolygroupLayersProperties.h"
+
 #include "RemoveOccludedTrianglesTool.generated.h"
 
 class IAssetGenerationAPI;
@@ -48,7 +50,12 @@ enum class EOcclusionCalculationUIMode : uint8
 	RaycastOcclusionSamples
 };
 
-
+UENUM()
+enum class EOccludedAction : uint8
+{
+	Remove,
+	SetNewGroup
+};
 
 /**
  * Standard properties
@@ -94,6 +101,10 @@ public:
 
 	UPROPERTY(EditAnywhere, Category = RemoveIslands, meta = (UIMin = "0", ClampMin = "0"))
 	int MinTriCountIsland = 0;
+
+	/** What action to perform on occluded triangles */
+	UPROPERTY(EditAnywhere, Category = Action)
+	EOccludedAction Action = EOccludedAction::Remove;
 };
 
 
@@ -174,6 +185,9 @@ protected:
 	UPROPERTY()
 	URemoveOccludedTrianglesToolProperties* BasicProperties;
 
+	UPROPERTY()
+	UPolygroupLayersProperties* PolygroupLayersProperties;
+
 	UPROPERTY()
 	URemoveOccludedTrianglesAdvancedProperties* AdvancedProperties;
 
@@ -196,12 +210,18 @@ protected:
 	TArray<int32> PreviewToTargetIdx;
 	TArray<int32> TargetToPreviewIdx;
 
+	// Group IDs for occluded triangles, per mesh in the Previews array, if OccludedAction is SetNewGroup
+	TArray<int32> OccludedGroupIDs;
+	// Selected layer indices for Group IDs, per mesh in the Previews array
+	TArray<int32> OccludedGroupLayers;
+
 	UWorld* TargetWorld;
 	IAssetGenerationAPI* AssetAPI;
 
 	FViewCameraState CameraState;
 
 	void SetupPreviews();
+	void MakePolygroupLayerProperties();
 
 	void GenerateAsset(const TArray<FDynamicMeshOpResult>& Results);
 };
diff --git a/Engine/Plugins/Experimental/MeshModelingToolset/Source/ModelingComponents/Private/MeshOpPreviewHelpers.cpp b/Engine/Plugins/Experimental/MeshModelingToolset/Source/ModelingComponents/Private/MeshOpPreviewHelpers.cpp
--- a/Engine/Plugins/Experimental/MeshModelingToolset/Source/ModelingComponents/Private/MeshOpPreviewHelpers.cpp
+++ b/Engine/Plugins/Experimental/MeshModelingToolset/Source/ModelingComponents/Private/MeshOpPreviewHelpers.cpp
@@ -59,10 +59,20 @@ void UMeshOpPreviewWithBackgroundCompute::Tick(float DeltaTime)
 		{
 			PreviewMesh->ClearOverrideRenderMaterial();
 		}
+
+		if (SecondaryMaterial != nullptr)
+		{
+			PreviewMesh->SetSecondaryRenderMaterial(SecondaryMaterial);
+		}
+		else
+		{
+			PreviewMesh->ClearSecondaryRenderMaterial();
+		}
 	}
 	else
 	{
 		PreviewMesh->SetOverrideRenderMaterial(WorkingMaterial);
+		PreviewMesh->ClearSecondaryRenderMaterial();
 	}
 }
 
diff --git a/Engine/Plugins/Experimental/MeshModelingToolset/Source/ModelingComponents/Private/PropertySets/PolygroupLayersProperties.cpp b/Engine/Plugins/Experimental/MeshModelingToolset/Source/ModelingComponents/Private/PropertySets/PolygroupLayersProperties.cpp
--- a/Engine/Plugins/Experimental/MeshModelingToolset/Source/ModelingComponents/Private/PropertySets/PolygroupLayersProperties.cpp
+++ b/Engine/Plugins/Experimental/MeshModelingToolset/Source/ModelingComponents/Private/PropertySets/PolygroupLayersProperties.cpp
@@ -7,6 +7,7 @@
 
 void UPolygroupLayersProperties::InitializeGroupLayers(const FDynamicMesh3* Mesh)
 {
+	GroupLayersList.Reset();
 	GroupLayersList.Add(TEXT("Default"));		// always have standard group
 	if (Mesh->Attributes())
 	{
@@ -23,6 +24,21 @@ void UPolygroupLayersProperties::InitializeGroupLayers(const FDynamicMesh3* Mesh
 	}
 }
 
+void UPolygroupLayersProperties::InitializeGroupLayers(const TSet<FName>& LayerNames)
+{
+	GroupLayersList.Reset();
+	GroupLayersList.Add(TEXT("Default"));		// always have standard group
+	for (const FName& Name : LayerNames)
+	{
+		GroupLayersList.Add(Name.ToString());
+	}
+
+	if (GroupLayersList.Contains(ActiveGroupLayer.ToString()) == false)		// discard restored value if it doesn't apply
+	{
+		ActiveGroupLayer = FName(GroupLayersList[0]);
+	}
+}
+
 
 bool UPolygroupLayersProperties::HasSelectedPolygroup() const
 {
diff --git a/Engine/Plugins/Experimental/MeshModelingToolset/Source/ModelingComponents/Public/MeshOpPreviewHelpers.h b/Engine/Plugins/Experimental/MeshModelingToolset/Source/ModelingComponents/Public/MeshOpPreviewHelpers.h
--- a/Engine/Plugins/Experimental/MeshModelingToolset/Source/ModelingComponents/Public/MeshOpPreviewHelpers.h
+++ b/Engine/Plugins/Experimental/MeshModelingToolset/Source/ModelingComponents/Public/MeshOpPreviewHelpers.h
@@ -172,6 +172,10 @@ public:
 	// if non-null, this material is swapped in when a background compute is active
 	UPROPERTY()
 	UMaterialInterface* WorkingMaterial = nullptr;
+	
+	// secondary render material to forward to PreviewMesh if set
+	UPROPERTY()
+	UMaterialInterface* SecondaryMaterial = nullptr;
 
 protected:
 	// state flag, if true then we have valid result
diff --git a/Engine/Plugins/Experimental/MeshModelingToolset/Source/ModelingComponents/Public/PropertySets/PolygroupLayersProperties.h b/Engine/Plugins/Experimental/MeshModelingToolset/Source/ModelingComponents/Public/PropertySets/PolygroupLayersProperties.h
--- a/Engine/Plugins/Experimental/MeshModelingToolset/Source/ModelingComponents/Public/PropertySets/PolygroupLayersProperties.h
+++ b/Engine/Plugins/Experimental/MeshModelingToolset/Source/ModelingComponents/Public/PropertySets/PolygroupLayersProperties.h
@@ -17,9 +17,9 @@ class MODELINGCOMPONENTS_API UPolygroupLayersProperties : public UInteractiveToo
 	GENERATED_BODY()
 public:
 
-	/** Select vertex weight map. If configured, the weight map value will be sampled to modulate displacement intensity. */
+	/** Select polygroup layer. */
 	UPROPERTY(EditAnywhere, Category = PolygroupLayers, meta = (GetOptions = GetGroupLayersFunc))
-	FName ActiveGroupLayer;
+	FName ActiveGroupLayer = "Default";
 
 	// this function is called provide set of available group layers
 	UFUNCTION()
@@ -31,7 +31,9 @@ public:
 
 	void InitializeGroupLayers(const FDynamicMesh3* Mesh);
 
-	// return true if any option other than "None" is selected
+	void InitializeGroupLayers(const TSet<FName>& LayerNames);
+
+	// return true if any option other than "Default" is selected
 	bool HasSelectedPolygroup() const;
 
 	void SetSelectedFromPolygroupIndex(int32 Index);
diff --git a/Engine/Plugins/Experimental/MeshModelingToolset/Source/ModelingOperators/Private/CleaningOps/RemoveOccludedTrianglesOp.cpp b/Engine/Plugins/Experimental/MeshModelingToolset/Source/ModelingOperators/Private/CleaningOps/RemoveOccludedTrianglesOp.cpp
--- a/Engine/Plugins/Experimental/MeshModelingToolset/Source/ModelingOperators/Private/CleaningOps/RemoveOccludedTrianglesOp.cpp
+++ b/Engine/Plugins/Experimental/MeshModelingToolset/Source/ModelingOperators/Private/CleaningOps/RemoveOccludedTrianglesOp.cpp
@@ -10,7 +10,7 @@
 #include "Operations/RemoveOccludedTriangles.h"
 #include "Selections/MeshConnectedComponents.h"
 #include "Selections/MeshFaceSelection.h"
-
+#include "Polygroups/PolygroupSet.h"
 
 #include "Async/ParallelFor.h"
 
@@ -65,6 +65,38 @@ void FRemoveOccludedTrianglesOp::CalculateResult(FProgressCancel* Progress)
 		SelectedTris = Selection.AsArray();
 	};
 
+	auto SetNewGroupSelection = [](FDynamicMesh3& Mesh, const TArray<int>& SelectedTris, FName LayerName, bool bActiveGroupIsDefault)
+	{
+		// don't add any new groups if there's nothing to select
+		if (SelectedTris.Num() == 0)
+		{
+			return FIndex2i::Invalid();
+		}
+
+		auto SetGroup = [](UE::Geometry::FPolygroupSet& ActiveGroupSet, const TArray<int>& SelectedTris)
+		{
+			int32 JacketGroupID = ActiveGroupSet.AllocateNewGroupID();
+
+			for (int TID : SelectedTris)
+			{
+				ActiveGroupSet.SetGroup(TID, JacketGroupID);
+			}
+
+			return FIndex2i(JacketGroupID, ActiveGroupSet.GroupLayerIndex);
+		};
+
+		if (!bActiveGroupIsDefault)
+		{
+			UE::Geometry::FPolygroupSet ActiveGroupSet(&Mesh, LayerName);
+			return SetGroup(ActiveGroupSet, SelectedTris);
+		}
+		else
+		{
+			UE::Geometry::FPolygroupSet ActiveGroupSet(&Mesh);
+			return SetGroup(ActiveGroupSet, SelectedTris);
+		}
+	};
+
 	if (bOnlySelfOcclude)
 	{
 		TRemoveOccludedTriangles<FDynamicMesh3> SelfJacket(ResultMesh.Get());
@@ -85,7 +117,17 @@ void FRemoveOccludedTrianglesOp::CalculateResult(FProgressCancel* Progress)
 		{
 			ShrinkSelection(*ResultMesh.Get(), SelfJacket.RemovedT, ShrinkRemoval);
 		}
-		SelfJacket.RemoveSelected();
+
+		if (bSetTriangleGroupInsteadOfRemoving)
+		{
+			FIndex2i GroupIDAndLayerIndex = SetNewGroupSelection(*ResultMesh.Get(), SelfJacket.RemovedT, ActiveGroupLayer, bActiveGroupLayerIsDefault);
+			CreatedGroupID = GroupIDAndLayerIndex.A;
+			CreatedGroupLayerIndex = GroupIDAndLayerIndex.B;
+		}
+		else
+		{
+			SelfJacket.RemoveSelected();
+		}
 	}
 	else
 	{
@@ -107,7 +149,17 @@ void FRemoveOccludedTrianglesOp::CalculateResult(FProgressCancel* Progress)
 		{
 			ShrinkSelection(*ResultMesh.Get(), Jacket.RemovedT, ShrinkRemoval);
 		}
-		Jacket.RemoveSelected();
+
+		if (bSetTriangleGroupInsteadOfRemoving)
+		{
+			FIndex2i GroupIDAndLayerIndex = SetNewGroupSelection(*ResultMesh.Get(), Jacket.RemovedT, ActiveGroupLayer, bActiveGroupLayerIsDefault);
+			CreatedGroupID = GroupIDAndLayerIndex.A;
+			CreatedGroupLayerIndex = GroupIDAndLayerIndex.B;
+		}
+		else
+		{
+			Jacket.RemoveSelected();
+		}
 	}
 
 	if (MinTriCountConnectedComponent > 0 || MinAreaConnectedComponent > 0)
diff --git a/Engine/Plugins/Experimental/MeshModelingToolset/Source/ModelingOperators/Public/CleaningOps/RemoveOccludedTrianglesOp.h b/Engine/Plugins/Experimental/MeshModelingToolset/Source/ModelingOperators/Public/CleaningOps/RemoveOccludedTrianglesOp.h
--- a/Engine/Plugins/Experimental/MeshModelingToolset/Source/ModelingOperators/Public/CleaningOps/RemoveOccludedTrianglesOp.h
+++ b/Engine/Plugins/Experimental/MeshModelingToolset/Source/ModelingOperators/Public/CleaningOps/RemoveOccludedTrianglesOp.h
@@ -68,10 +68,24 @@ public:
 
 	int MinTriCountConnectedComponent = 0;
 
+	// if true, we will set triangle group IDs for occluded triangles, rather than deleting the triangles
+	bool bSetTriangleGroupInsteadOfRemoving = false;
+
+	// name of the group layer to use if we are setting triangle groups instead of removing occluded triangles
+	FName ActiveGroupLayer;
+
+	// if true, the ActiveGroupLayer is the name of the "default" layer, so we'll use the built-in group IDs
+	bool bActiveGroupLayerIsDefault = true;
+
 	EOcclusionCalculationMode InsideMode =
 		EOcclusionCalculationMode::FastWindingNumber;
 
 
+	// outputs, used when bSetTriangleGroupInsteadOfRemoving is true
+	int CreatedGroupID = -1;
+	int CreatedGroupLayerIndex = -1;
+
+
 	void SetTransform(const FTransform& Transform);
 
 	//
