fix crash in voxel blend and voxel morphology tools when offset is too small and deletes input mesh

diff --git a/Engine/Plugins/Experimental/MeshModelingToolset/Source/ModelingOperators/Private/CompositionOps/VoxelBlendMeshesOp.cpp b/Engine/Plugins/Experimental/MeshModelingToolset/Source/ModelingOperators/Private/CompositionOps/VoxelBlendMeshesOp.cpp
--- a/Engine/Plugins/Experimental/MeshModelingToolset/Source/ModelingOperators/Private/CompositionOps/VoxelBlendMeshesOp.cpp
+++ b/Engine/Plugins/Experimental/MeshModelingToolset/Source/ModelingOperators/Private/CompositionOps/VoxelBlendMeshesOp.cpp
@@ -39,6 +39,11 @@ void FVoxelBlendMeshesOp::CalculateResult(FProgressCancel* Progress)
 	for (int MeshIdx = 0; MeshIdx < Meshes.Num(); MeshIdx++)
 	{
 		TransformedMeshes.Emplace(*Meshes[MeshIdx]);
+		if (TransformedMeshes[MeshIdx].TriangleCount() == 0)
+		{
+			continue;
+		}
+		
 		if (Transforms[MeshIdx].GetDeterminant() < 0)
 		{
 			TransformedMeshes[MeshIdx].ReverseOrientation(false);
@@ -49,9 +54,13 @@ void FVoxelBlendMeshesOp::CalculateResult(FProgressCancel* Progress)
 		{
 			if (OffsetSolidifySurface > 0)
 			{
+				// positive offsets should be at least a cell wide so we don't end up deleting a bunch of the input surface
+				double CellSize = TransformedMeshes[MeshIdx].GetCachedBounds().MaxDim() / InputVoxelCount;
+				double SafeOffset = FMathd::Max(CellSize * 2, OffsetSolidifySurface);
+
 				FMeshNormals::QuickComputeVertexNormals(TransformedMeshes[MeshIdx]);
 				FExtrudeMesh Extrude(&TransformedMeshes[MeshIdx]);
-				Extrude.DefaultExtrudeDistance = -OffsetSolidifySurface;
+				Extrude.DefaultExtrudeDistance = -SafeOffset;
 				Extrude.IsPositiveOffset = false;
 				Extrude.Apply();
 			}
@@ -68,11 +77,20 @@ void FVoxelBlendMeshesOp::CalculateResult(FProgressCancel* Progress)
 			}
 		}
 
+		if (TransformedMeshes[MeshIdx].TriangleCount() == 0)
+		{
+			continue;
+		}
 		ImplicitBlend.Sources.Add(&TransformedMeshes[MeshIdx]);
 		FAxisAlignedBox3d& SourceBounds = ImplicitBlend.SourceBounds.Add_GetRef(TransformedMeshes[MeshIdx].GetCachedBounds());
 		CombinedBounds.Contain(SourceBounds);
 	}
 
+	if (ImplicitBlend.Sources.Num() == 0)
+	{
+		return;
+	}
+
 	ImplicitBlend.SetCellSizesAndFalloff(CombinedBounds, BlendFalloff, InputVoxelCount, OutputVoxelCount);
 	ImplicitBlend.BlendPower = BlendPower;
 
diff --git a/Engine/Plugins/Experimental/MeshModelingToolset/Source/ModelingOperators/Private/CompositionOps/VoxelMorphologyMeshesOp.cpp b/Engine/Plugins/Experimental/MeshModelingToolset/Source/ModelingOperators/Private/CompositionOps/VoxelMorphologyMeshesOp.cpp
--- a/Engine/Plugins/Experimental/MeshModelingToolset/Source/ModelingOperators/Private/CompositionOps/VoxelMorphologyMeshesOp.cpp
+++ b/Engine/Plugins/Experimental/MeshModelingToolset/Source/ModelingOperators/Private/CompositionOps/VoxelMorphologyMeshesOp.cpp
@@ -80,11 +80,20 @@ void FVoxelMorphologyMeshesOp::CalculateResult(FProgressCancel* Progress)
 		}
 	}
 
+	if (CombinedMesh.TriangleCount() == 0)
+	{
+		return;
+	}
+
 	if (bSolidifyInput && OffsetSolidifySurface > 0)
 	{
+		// positive offsets should be at least a cell wide so we don't end up deleting a bunch of the input surface
+		double CellSize = CombinedMesh.GetCachedBounds().MaxDim() / InputVoxelCount;
+		double SafeOffset = FMathd::Max(CellSize * 2, OffsetSolidifySurface);
+
 		FMeshNormals::QuickComputeVertexNormals(CombinedMesh);
 		FExtrudeMesh Extrude(&CombinedMesh);
-		Extrude.DefaultExtrudeDistance = -OffsetSolidifySurface;
+		Extrude.DefaultExtrudeDistance = -SafeOffset;
 		Extrude.IsPositiveOffset = false;
 		Extrude.Apply();
 	}
@@ -107,6 +116,11 @@ void FVoxelMorphologyMeshesOp::CalculateResult(FProgressCancel* Progress)
 		Spatial.Build(); // rebuild w/ updated mesh
 	}
 
+	if (CombinedMesh.TriangleCount() == 0)
+	{
+		return;
+	}
+
 	ImplicitMorphology.SourceSpatial = &Spatial;
 	ImplicitMorphology.SetCellSizesAndDistance(CombinedMesh.GetCachedBounds(), Distance, InputVoxelCount, OutputVoxelCount);
 	
