Clean up of VoxWrap, VoxBlend, and VoxMorphology
 - Rename "offset" to the more accurate "thicken shells"
 - More consistent clamp min thickening based on voxel resolution
 - Only apply thickening to meshes that are actually shells (i.e. have open boundaries)
 - Improve related comments and variable names

diff --git a/Engine/Plugins/Experimental/GeometryProcessing/Source/DynamicMesh/Private/Operations/ExtrudeMesh.cpp b/Engine/Plugins/Experimental/GeometryProcessing/Source/DynamicMesh/Private/Operations/ExtrudeMesh.cpp
--- a/Engine/Plugins/Experimental/GeometryProcessing/Source/DynamicMesh/Private/Operations/ExtrudeMesh.cpp
+++ b/Engine/Plugins/Experimental/GeometryProcessing/Source/DynamicMesh/Private/Operations/ExtrudeMesh.cpp
@@ -57,6 +57,10 @@ bool FExtrudeMesh::ApplyExtrude(FExtrusionInfo& Region, FMeshNormals* UseNormals
 		return false;
 	}
 	int NumInitialLoops = Region.InitialLoops.GetLoopCount();
+	if (bSkipClosedComponents && NumInitialLoops == 0)
+	{
+		return true;
+	}
 
 	MeshIndexUtil::TriangleToVertexIDs(Mesh, Region.InitialTriangles, Region.InitialVertices);
 
diff --git a/Engine/Plugins/Experimental/GeometryProcessing/Source/DynamicMesh/Public/Operations/ExtrudeMesh.h b/Engine/Plugins/Experimental/GeometryProcessing/Source/DynamicMesh/Public/Operations/ExtrudeMesh.h
--- a/Engine/Plugins/Experimental/GeometryProcessing/Source/DynamicMesh/Public/Operations/ExtrudeMesh.h
+++ b/Engine/Plugins/Experimental/GeometryProcessing/Source/DynamicMesh/Public/Operations/ExtrudeMesh.h
@@ -44,6 +44,9 @@ public:
 	/** quads on the stitch loop are planar-projected and scaled by this amount */
 	float UVScaleFactor = 1.0f;
 
+	/** If true, skip closed components */
+	bool bSkipClosedComponents = false;
+
 
 	//
 	// Outputs
diff --git a/Engine/Plugins/Experimental/MeshModelingToolset/Source/MeshModelingTools/Private/VoxelBlendMeshesTool.cpp b/Engine/Plugins/Experimental/MeshModelingToolset/Source/MeshModelingTools/Private/VoxelBlendMeshesTool.cpp
--- a/Engine/Plugins/Experimental/MeshModelingToolset/Source/MeshModelingTools/Private/VoxelBlendMeshesTool.cpp
+++ b/Engine/Plugins/Experimental/MeshModelingToolset/Source/MeshModelingTools/Private/VoxelBlendMeshesTool.cpp
@@ -69,9 +69,9 @@ TUniquePtr<FDynamicMeshOperator> UVoxelBlendMeshesTool::MakeNewOperator()
 
 	Op->BlendFalloff = BlendProperties->BlendFalloff;
 	Op->BlendPower = BlendProperties->BlendPower;
-	Op->bSolidifyInput = BlendProperties->bSolidifyInput;
-	Op->bRemoveInternalsAfterSolidify = BlendProperties->bRemoveInternalsAfterSolidify;
-	Op->OffsetSolidifySurface = BlendProperties->OffsetSolidifySurface;
+	Op->bVoxWrap = BlendProperties->bVoxWrap;
+	Op->bRemoveInternalsAfterVoxWrap = BlendProperties->bRemoveInternalsAfterVoxWrap;
+	Op->ThickenShells = BlendProperties->ThickenShells;
 
 	VoxProperties->SetPropertiesOnOp(*Op);
 
diff --git a/Engine/Plugins/Experimental/MeshModelingToolset/Source/MeshModelingTools/Private/VoxelMorphologyMeshesTool.cpp b/Engine/Plugins/Experimental/MeshModelingToolset/Source/MeshModelingTools/Private/VoxelMorphologyMeshesTool.cpp
--- a/Engine/Plugins/Experimental/MeshModelingToolset/Source/MeshModelingTools/Private/VoxelMorphologyMeshesTool.cpp
+++ b/Engine/Plugins/Experimental/MeshModelingToolset/Source/MeshModelingTools/Private/VoxelMorphologyMeshesTool.cpp
@@ -69,9 +69,9 @@ TUniquePtr<FDynamicMeshOperator> UVoxelMorphologyMeshesTool::MakeNewOperator()
 
 	VoxProperties->SetPropertiesOnOp(*Op);
 	
-	Op->bSolidifyInput = MorphologyProperties->bSolidifyInput;
-	Op->OffsetSolidifySurface = MorphologyProperties->OffsetSolidifySurface;
-	Op->bRemoveInternalsAfterSolidify = MorphologyProperties->bRemoveInternalsAfterSolidify;
+	Op->bVoxWrapInput = MorphologyProperties->bVoxWrap;
+	Op->ThickenShells = MorphologyProperties->ThickenShells;
+	Op->bRemoveInternalsAfterVoxWrap = MorphologyProperties->bRemoveInternalsAfterVoxWrap;
 	Op->Distance = MorphologyProperties->Distance;
 	Op->Operation = MorphologyProperties->Operation;
 
diff --git a/Engine/Plugins/Experimental/MeshModelingToolset/Source/MeshModelingTools/Private/VoxelSolidifyMeshesTool.cpp b/Engine/Plugins/Experimental/MeshModelingToolset/Source/MeshModelingTools/Private/VoxelSolidifyMeshesTool.cpp
--- a/Engine/Plugins/Experimental/MeshModelingToolset/Source/MeshModelingTools/Private/VoxelSolidifyMeshesTool.cpp
+++ b/Engine/Plugins/Experimental/MeshModelingToolset/Source/MeshModelingTools/Private/VoxelSolidifyMeshesTool.cpp
@@ -64,8 +64,8 @@ TUniquePtr<FDynamicMeshOperator> UVoxelSolidifyMeshesTool::MakeNewOperator()
 
 	Op->bSolidAtBoundaries = SolidifyProperties->bSolidAtBoundaries;
 	Op->WindingThreshold = SolidifyProperties->WindingThreshold;
-	Op->bMakeOffsetSurfaces = SolidifyProperties->bMakeOffsetSurfaces;
-	Op->OffsetThickness = SolidifyProperties->OffsetThickness;
+	Op->bApplyThickenShells = SolidifyProperties->bApplyThickenShells;
+	Op->ThickenShells = SolidifyProperties->ThickenShells;
 	Op->SurfaceSearchSteps = SolidifyProperties->SurfaceSearchSteps;
 	Op->ExtendBounds = SolidifyProperties->ExtendBounds;
 	VoxProperties->SetPropertiesOnOp(*Op);
diff --git a/Engine/Plugins/Experimental/MeshModelingToolset/Source/MeshModelingTools/Public/VoxelBlendMeshesTool.h b/Engine/Plugins/Experimental/MeshModelingToolset/Source/MeshModelingTools/Public/VoxelBlendMeshesTool.h
--- a/Engine/Plugins/Experimental/MeshModelingToolset/Source/MeshModelingTools/Public/VoxelBlendMeshesTool.h
+++ b/Engine/Plugins/Experimental/MeshModelingToolset/Source/MeshModelingTools/Public/VoxelBlendMeshesTool.h
@@ -26,17 +26,17 @@ public:
 	UPROPERTY(EditAnywhere, Category = Blend, meta = (UIMin = ".1", UIMax = "100", ClampMin = ".001", ClampMax = "1000"))
 	double BlendFalloff = 10;
 
-	/** Solidify the input mesh(es) before processing, fixing results for inputs with holes and/or self-intersections */
-	UPROPERTY(EditAnywhere, Category = Solidify)
-	bool bSolidifyInput = false;
+	/** Apply a "VoxWrap" operation to input mesh(es) before computing the blend.  Fixes results for inputs with holes and/or self-intersections. */
+	UPROPERTY(EditAnywhere, Category = VoxWrapPreprocess)
+	bool bVoxWrap = false;
 
-	/** Remove internal surfaces from the solidified input */
-	UPROPERTY(EditAnywhere, Category = Solidify, meta = (EditCondition = "bSolidifyInput == true"))
-	bool bRemoveInternalsAfterSolidify = false;
+	/** Remove internal surfaces from the VoxWrap output, before computing the blend. */
+	UPROPERTY(EditAnywhere, Category = VoxWrapPreprocess, meta = (EditCondition = "bVoxWrap == true"))
+	bool bRemoveInternalsAfterVoxWrap = false;
 
-	/** Offset surface to create when solidifying any open-boundary inputs; if 0 then no offset surfaces are created */
-	UPROPERTY(EditAnywhere, Category = Solidify, meta = (EditCondition = "bSolidifyInput == true", UIMin = "0", UIMax = "100", ClampMin = "0", ClampMax = "1000"))
-	double OffsetSolidifySurface = 0.0;
+	/** Thicken open-boundary surfaces (extrude them inwards) before VoxWrapping them. Units are in world space. If 0 then no extrusion is applied. */
+	UPROPERTY(EditAnywhere, Category = VoxWrapPreprocess, meta = (EditCondition = "bVoxWrap == true", UIMin = "0", UIMax = "100", ClampMin = "0", ClampMax = "1000"))
+	double ThickenShells = 0.0;
 };
 
 
diff --git a/Engine/Plugins/Experimental/MeshModelingToolset/Source/MeshModelingTools/Public/VoxelMorphologyMeshesTool.h b/Engine/Plugins/Experimental/MeshModelingToolset/Source/MeshModelingTools/Public/VoxelMorphologyMeshesTool.h
--- a/Engine/Plugins/Experimental/MeshModelingToolset/Source/MeshModelingTools/Public/VoxelMorphologyMeshesTool.h
+++ b/Engine/Plugins/Experimental/MeshModelingToolset/Source/MeshModelingTools/Public/VoxelMorphologyMeshesTool.h
@@ -27,17 +27,17 @@ public:
 	UPROPERTY(EditAnywhere, Category = Morphology, meta = (UIMin = ".1", UIMax = "100", ClampMin = ".001", ClampMax = "1000"))
 	double Distance = 5;
 
-	/** Solidify the input mesh(es) before processing, fixing results for inputs with holes and/or self-intersections */
-	UPROPERTY(EditAnywhere, Category = Solidify)
-	bool bSolidifyInput = false;
+	/** Apply a "VoxWrap" operation to input mesh(es) before computing the morphology.  Fixes results for inputs with holes and/or self-intersections. */
+	UPROPERTY(EditAnywhere, Category = VoxWrapPreprocess)
+	bool bVoxWrap = false;
 
-	/** Remove internal surfaces from the solidified input, before running morphology */
-	UPROPERTY(EditAnywhere, Category = Solidify, meta = (EditCondition = "bSolidifyInput == true"))
-	bool bRemoveInternalsAfterSolidify = false;
+	/** Remove internal surfaces from the VoxWrap output, before computing the morphology. */
+	UPROPERTY(EditAnywhere, Category = VoxWrapPreprocess, meta = (EditCondition = "bVoxWrap == true"))
+	bool bRemoveInternalsAfterVoxWrap = false;
 
-	/** Offset surface to create when solidifying any open-boundary inputs; if 0 then no offset surfaces are created */
-	UPROPERTY(EditAnywhere, Category = Solidify, meta = (EditCondition = "bSolidifyInput == true", UIMin = "0", UIMax = "100", ClampMin = "0", ClampMax = "1000"))
-	double OffsetSolidifySurface = 0.0;
+	/** Thicken open-boundary surfaces (extrude them inwards) before VoxWrapping them. Units are in world space. If 0, no extrusion is applied. */
+	UPROPERTY(EditAnywhere, Category = VoxWrapPreprocess, meta = (EditCondition = "bVoxWrap == true", UIMin = "0", UIMax = "100", ClampMin = "0", ClampMax = "1000"))
+	double ThickenShells = 0.0;
 };
 
 
diff --git a/Engine/Plugins/Experimental/MeshModelingToolset/Source/MeshModelingTools/Public/VoxelSolidifyMeshesTool.h b/Engine/Plugins/Experimental/MeshModelingToolset/Source/MeshModelingTools/Public/VoxelSolidifyMeshesTool.h
--- a/Engine/Plugins/Experimental/MeshModelingToolset/Source/MeshModelingTools/Public/VoxelSolidifyMeshesTool.h
+++ b/Engine/Plugins/Experimental/MeshModelingToolset/Source/MeshModelingTools/Public/VoxelSolidifyMeshesTool.h
@@ -37,13 +37,13 @@ public:
 	UPROPERTY(EditAnywhere, Category = Options)
 	bool bSolidAtBoundaries = true;
 
-	/** If true, treats mesh surfaces with open boundaries as having a fixed, user-defined thickness */
-	UPROPERTY(EditAnywhere, Category = Options)
-	bool bMakeOffsetSurfaces = false;
+	/** If true, uses the ThickenShells setting */
+	UPROPERTY()
+	bool bApplyThickenShells = false;
 
-	/** Thickness of offset surfaces */
-	UPROPERTY(EditAnywhere, Category = Options, meta = (UIMin = ".1", UIMax = "100", ClampMin = ".001", ClampMax = "1000", EditCondition = "bMakeOffsetSurfaces == true"))
-	double OffsetThickness = 5;
+	/** Thicken open-boundary surfaces (extrude them inwards) to ensure they are captured in the VoxWrap output. Units are in world space. */
+	UPROPERTY(EditAnywhere, Category = Options, meta = (UIMin = ".1", UIMax = "100", ClampMin = ".001", ClampMax = "1000", EditCondition = "bApplyThickenShells == true"))
+	double ThickenShells = 5;
 };
 
 
diff --git a/Engine/Plugins/Experimental/MeshModelingToolset/Source/ModelingOperators/Private/CompositionOps/VoxelBlendMeshesOp.cpp b/Engine/Plugins/Experimental/MeshModelingToolset/Source/ModelingOperators/Private/CompositionOps/VoxelBlendMeshesOp.cpp
--- a/Engine/Plugins/Experimental/MeshModelingToolset/Source/ModelingOperators/Private/CompositionOps/VoxelBlendMeshesOp.cpp
+++ b/Engine/Plugins/Experimental/MeshModelingToolset/Source/ModelingOperators/Private/CompositionOps/VoxelBlendMeshesOp.cpp
@@ -50,17 +50,18 @@ void FVoxelBlendMeshesOp::CalculateResult(FProgressCancel* Progress)
 		}
 		MeshTransforms::ApplyTransform(TransformedMeshes[MeshIdx], (FTransform3d)Transforms[MeshIdx]);
 
-		if (bSolidifyInput)
+		if (bVoxWrap)
 		{
-			if (OffsetSolidifySurface > 0)
+			if (ThickenShells > 0 && !TransformedMeshes[MeshIdx].IsClosed())
 			{
-				// positive offsets should be at least a cell wide so we don't end up deleting a bunch of the input surface
+				// thickness should be at least a cell wide so we don't end up deleting a bunch of the input surface
 				double CellSize = TransformedMeshes[MeshIdx].GetCachedBounds().MaxDim() / InputVoxelCount;
-				double SafeOffset = FMathd::Max(CellSize * 2, OffsetSolidifySurface);
+				double SafeThickness = FMathd::Max(CellSize * 2, ThickenShells);
 
 				FMeshNormals::QuickComputeVertexNormals(TransformedMeshes[MeshIdx]);
 				FExtrudeMesh Extrude(&TransformedMeshes[MeshIdx]);
-				Extrude.DefaultExtrudeDistance = -SafeOffset;
+				Extrude.bSkipClosedComponents = true;
+				Extrude.DefaultExtrudeDistance = -SafeThickness;
 				Extrude.IsPositiveOffset = false;
 				Extrude.Apply();
 			}
@@ -71,7 +72,7 @@ void FVoxelBlendMeshesOp::CalculateResult(FProgressCancel* Progress)
 			Solidify.SetCellSizeAndExtendBounds(Spatial.GetBoundingBox(), 0, InputVoxelCount);
 			TransformedMeshes[MeshIdx].Copy(&Solidify.Generate());
 
-			if (bRemoveInternalsAfterSolidify)
+			if (bRemoveInternalsAfterVoxWrap)
 			{
 				UE::MeshAutoRepair::RemoveInternalTriangles(TransformedMeshes[MeshIdx], true, EOcclusionTriangleSampling::Centroids, EOcclusionCalculationMode::FastWindingNumber);
 			}
diff --git a/Engine/Plugins/Experimental/MeshModelingToolset/Source/ModelingOperators/Private/CompositionOps/VoxelMorphologyMeshesOp.cpp b/Engine/Plugins/Experimental/MeshModelingToolset/Source/ModelingOperators/Private/CompositionOps/VoxelMorphologyMeshesOp.cpp
--- a/Engine/Plugins/Experimental/MeshModelingToolset/Source/ModelingOperators/Private/CompositionOps/VoxelMorphologyMeshesOp.cpp
+++ b/Engine/Plugins/Experimental/MeshModelingToolset/Source/ModelingOperators/Private/CompositionOps/VoxelMorphologyMeshesOp.cpp
@@ -70,7 +70,7 @@ void FVoxelMorphologyMeshesOp::CalculateResult(FProgressCancel* Progress)
 			{
 				return MeshTransform.TransformPosition(Pos);
 			}, nullptr
-			);
+		);
 		if (bReverseOrientation)
 		{
 			for (int TID : Meshes[MeshIdx]->TriangleIndicesItr())
@@ -85,30 +85,31 @@ void FVoxelMorphologyMeshesOp::CalculateResult(FProgressCancel* Progress)
 		return;
 	}
 
-	if (bSolidifyInput && OffsetSolidifySurface > 0)
+	if (bVoxWrapInput && ThickenShells > 0)
 	{
 		// positive offsets should be at least a cell wide so we don't end up deleting a bunch of the input surface
 		double CellSize = CombinedMesh.GetCachedBounds().MaxDim() / InputVoxelCount;
-		double SafeOffset = FMathd::Max(CellSize * 2, OffsetSolidifySurface);
+		double SafeThickness = FMathd::Max(CellSize * 2, ThickenShells);
 
 		FMeshNormals::QuickComputeVertexNormals(CombinedMesh);
 		FExtrudeMesh Extrude(&CombinedMesh);
-		Extrude.DefaultExtrudeDistance = -SafeOffset;
+		Extrude.bSkipClosedComponents = true;
+		Extrude.DefaultExtrudeDistance = -SafeThickness;
 		Extrude.IsPositiveOffset = false;
 		Extrude.Apply();
 	}
-
+	
 	ImplicitMorphology.Source = &CombinedMesh;
 	FDynamicMeshAABBTree3 Spatial(&CombinedMesh, true);
 
-	if (bSolidifyInput)
+	if (bVoxWrapInput)
 	{
 		TFastWindingTree<FDynamicMesh3> Winding(&Spatial);
 		TImplicitSolidify<FDynamicMesh3> Solidify(&CombinedMesh, &Spatial, &Winding);
 		Solidify.SetCellSizeAndExtendBounds(Spatial.GetBoundingBox(), 0, InputVoxelCount);
 		CombinedMesh.Copy(&Solidify.Generate());
 
-		if (bRemoveInternalsAfterSolidify)
+		if (bRemoveInternalsAfterVoxWrap)
 		{
 			UE::MeshAutoRepair::RemoveInternalTriangles(CombinedMesh, true, EOcclusionTriangleSampling::Centroids, EOcclusionCalculationMode::FastWindingNumber);
 		}
diff --git a/Engine/Plugins/Experimental/MeshModelingToolset/Source/ModelingOperators/Private/CompositionOps/VoxelSolidifyMeshesOp.cpp b/Engine/Plugins/Experimental/MeshModelingToolset/Source/ModelingOperators/Private/CompositionOps/VoxelSolidifyMeshesOp.cpp
--- a/Engine/Plugins/Experimental/MeshModelingToolset/Source/ModelingOperators/Private/CompositionOps/VoxelSolidifyMeshesOp.cpp
+++ b/Engine/Plugins/Experimental/MeshModelingToolset/Source/ModelingOperators/Private/CompositionOps/VoxelSolidifyMeshesOp.cpp
@@ -55,11 +55,16 @@ void FVoxelSolidifyMeshesOp::CalculateResult(FProgressCancel* Progress)
 		}
 	}
 
-	if (bMakeOffsetSurfaces && OffsetThickness > 0)
+	if (bApplyThickenShells)
 	{
+		// thickness should be at least a cell wide so we don't end up deleting a bunch of the input surface
+		double CellSize = CombinedMesh.GetCachedBounds().MaxDim() / InputVoxelCount;
+		double SafeThickness = FMathd::Max(CellSize * 2, ThickenShells);
+
 		FMeshNormals::QuickComputeVertexNormals(CombinedMesh);
 		FExtrudeMesh Extrude(&CombinedMesh);
-		Extrude.DefaultExtrudeDistance = -OffsetThickness;
+		Extrude.bSkipClosedComponents = true;
+		Extrude.DefaultExtrudeDistance = -SafeThickness;
 		Extrude.IsPositiveOffset = false;
 		Extrude.Apply();
 	}
diff --git a/Engine/Plugins/Experimental/MeshModelingToolset/Source/ModelingOperators/Public/CompositionOps/VoxelBlendMeshesOp.h b/Engine/Plugins/Experimental/MeshModelingToolset/Source/ModelingOperators/Public/CompositionOps/VoxelBlendMeshesOp.h
--- a/Engine/Plugins/Experimental/MeshModelingToolset/Source/ModelingOperators/Public/CompositionOps/VoxelBlendMeshesOp.h
+++ b/Engine/Plugins/Experimental/MeshModelingToolset/Source/ModelingOperators/Public/CompositionOps/VoxelBlendMeshesOp.h
@@ -22,9 +22,9 @@ public:
 	double BlendFalloff = 10;
 	double BlendPower = 2;
 
-	bool bSolidifyInput = false;
-	bool bRemoveInternalsAfterSolidify = false;
-	double OffsetSolidifySurface = 0.0;
+	bool bVoxWrap = false;
+	bool bRemoveInternalsAfterVoxWrap = false;
+	double ThickenShells = 0.0;
 
 	void SetTransform(const FTransform& Transform);
 
diff --git a/Engine/Plugins/Experimental/MeshModelingToolset/Source/ModelingOperators/Public/CompositionOps/VoxelMorphologyMeshesOp.h b/Engine/Plugins/Experimental/MeshModelingToolset/Source/ModelingOperators/Public/CompositionOps/VoxelMorphologyMeshesOp.h
--- a/Engine/Plugins/Experimental/MeshModelingToolset/Source/ModelingOperators/Public/CompositionOps/VoxelMorphologyMeshesOp.h
+++ b/Engine/Plugins/Experimental/MeshModelingToolset/Source/ModelingOperators/Public/CompositionOps/VoxelMorphologyMeshesOp.h
@@ -41,9 +41,9 @@ public:
 	double Distance = 1.0;
 	EMorphologyOperation Operation;
 
-	bool bSolidifyInput = false;
-	bool bRemoveInternalsAfterSolidify = false;
-	double OffsetSolidifySurface = 0.0;
+	bool bVoxWrapInput = false;
+	bool bRemoveInternalsAfterVoxWrap = false;
+	double ThickenShells = 0.0;
 
 	void SetTransform(const FTransform& Transform);
 
diff --git a/Engine/Plugins/Experimental/MeshModelingToolset/Source/ModelingOperators/Public/CompositionOps/VoxelSolidifyMeshesOp.h b/Engine/Plugins/Experimental/MeshModelingToolset/Source/ModelingOperators/Public/CompositionOps/VoxelSolidifyMeshesOp.h
--- a/Engine/Plugins/Experimental/MeshModelingToolset/Source/ModelingOperators/Public/CompositionOps/VoxelSolidifyMeshesOp.h
+++ b/Engine/Plugins/Experimental/MeshModelingToolset/Source/ModelingOperators/Public/CompositionOps/VoxelSolidifyMeshesOp.h
@@ -24,8 +24,8 @@ public:
 	bool bSolidAtBoundaries = true;
 	int SurfaceSearchSteps = 3;
 
-	bool bMakeOffsetSurfaces = false;
-	double OffsetThickness = 5;
+	bool bApplyThickenShells = false;
+	double ThickenShells = 5;
 
 	void SetTransform(const FTransform& Transform);
 
