ModelingTools: improve performance in DisplaceMeshTool.
- Avoid launching level-4 subdivision initially, which currently happens before subdivision level is clamped, and will compute at least one subdivision before cancelling (which can take a long time on a mesh with millions of triangles)
- support ProgressCancel inside SubdivideMesh() function, which can take significant compute/memory resources for a single call
- Pass VertexID through various weightmap functions, so we can directly sample displacement at vertex when it is known
- store weightmap values in per-vertex UVs before subdivision. These are linearly interpolated during subdivision, so output mesh has weightmap value stored in VertexUV.X, and can be directly sampled, instead of using AABBTree nearest-point query
- due to above, need to invalidate subdivision if weightmap changes.
- If weight map is "None", set null ActiveWeightMap, this is cleanly handled and avoids potentially-costly weightmap sampling if it is not needed
- fix issue with subdivision-steps clamping message being shown at Tool startup, only show after the user has explicitly change the value
#rnx

diff --git a/Engine/Plugins/Experimental/MeshModelingToolset/Source/MeshModelingTools/Private/DisplaceMeshTool.cpp b/Engine/Plugins/Experimental/MeshModelingToolset/Source/MeshModelingTools/Private/DisplaceMeshTool.cpp
--- a/Engine/Plugins/Experimental/MeshModelingToolset/Source/MeshModelingTools/Private/DisplaceMeshTool.cpp
+++ b/Engine/Plugins/Experimental/MeshModelingToolset/Source/MeshModelingTools/Private/DisplaceMeshTool.cpp
@@ -19,7 +19,7 @@
 
 namespace DisplaceMeshToolLocals{
 
-	void SubdivideMesh(FDynamicMesh3& Mesh)
+	void SubdivideMesh(FDynamicMesh3& Mesh, FProgressCancel* ProgressCancel)
 	{
 		TArray<int> EdgesToProcess;
 		for (int tid : Mesh.EdgeIndicesItr())
@@ -28,6 +28,8 @@ namespace DisplaceMeshToolLocals{
 		}
 		int MaxTriangleID = Mesh.MaxTriangleID();
 
+		if (ProgressCancel && ProgressCancel->Cancelled()) return;
+
 		TArray<int> TriSplitEdges;
 		TriSplitEdges.Init(-1, Mesh.MaxTriangleID());
 
@@ -51,6 +53,8 @@ namespace DisplaceMeshToolLocals{
 					}
 				}
 			}
+
+			if (ProgressCancel && ProgressCancel->Cancelled()) return;
 		}
 
 		for (int eid : TriSplitEdges)
@@ -59,6 +63,8 @@ namespace DisplaceMeshToolLocals{
 			{
 				FDynamicMesh3::FEdgeFlipInfo FlipInfo;
 				Mesh.FlipEdge(eid, FlipInfo);
+
+				if (ProgressCancel && ProgressCancel->Cancelled()) return;
 			}
 		}
 	}
@@ -115,13 +121,13 @@ namespace DisplaceMeshToolLocals{
 		void Constant(const FDynamicMesh3& Mesh,
 			const TArray<FVector3d>& Positions, 
 			const FMeshNormals& Normals, 
-			TFunctionRef<float(const FVector3d&, const FVector3d&)> IntensityFunc,
+			TFunctionRef<float(int32, const FVector3d&, const FVector3d&)> IntensityFunc,
 			TArray<FVector3d>& DisplacedPositions)
 		{
 			ParallelDisplace(Mesh, Positions, Normals, DisplacedPositions,
 				[&](int32 vid, const FVector3d& Position, const FVector3d& Normal)
 			{
-				double Intensity = IntensityFunc(Position, Normal);
+				double Intensity = IntensityFunc(vid, Position, Normal);
 				return Position + (Intensity * Normal);
 			});
 		}
@@ -130,7 +136,7 @@ namespace DisplaceMeshToolLocals{
 		void RandomNoise(const FDynamicMesh3& Mesh,
 			const TArray<FVector3d>& Positions, 
 			const FMeshNormals& Normals,
-			TFunctionRef<float(const FVector3d&, const FVector3d&)> IntensityFunc,
+			TFunctionRef<float(int32, const FVector3d&, const FVector3d&)> IntensityFunc,
 			int RandomSeed, 
 			TArray<FVector3d>& DisplacedPositions)
 		{
@@ -138,7 +144,7 @@ namespace DisplaceMeshToolLocals{
 			for (int vid : Mesh.VertexIndicesItr())
 			{
 				double RandVal = 2.0 * (FMath::SRand() - 0.5);
-				double Intensity = IntensityFunc(Positions[vid], Normals[vid]);
+				double Intensity = IntensityFunc(vid, Positions[vid], Normals[vid]);
 				DisplacedPositions[vid] = Positions[vid] + (Normals[vid] * RandVal * Intensity);
 			}
 		}
@@ -146,7 +152,7 @@ namespace DisplaceMeshToolLocals{
 		void PerlinNoise(const FDynamicMesh3& Mesh,
 			const TArray<FVector3d>& Positions,
 			const FMeshNormals& Normals,
-			TFunctionRef<float(const FVector3d&, const FVector3d&)> IntensityFunc,
+			TFunctionRef<float(int32, const FVector3d&, const FVector3d&)> IntensityFunc,
 			const TArray<FPerlinLayerProperties>& PerlinLayerProperties,
 			int RandomSeed,
 			TArray<FVector3d>& DisplacedPositions)
@@ -164,7 +170,7 @@ namespace DisplaceMeshToolLocals{
 				{
 					TotalNoiseValue += PerlinLayerProperties[Layer].Intensity * FMath::PerlinNoise3D(PerlinLayerProperties[Layer].Frequency * EvalLocation);
 				}
-				double Intensity = IntensityFunc(Position, Normal);
+				double Intensity = IntensityFunc(vid, Position, Normal);
 				return Position + (TotalNoiseValue * Intensity * Normal);
 			});
 		}
@@ -172,7 +178,7 @@ namespace DisplaceMeshToolLocals{
 		void Map(const FDynamicMesh3& Mesh,
 			const TArray<FVector3d>& Positions, 
 			const FMeshNormals& Normals,
-			TFunctionRef<float(const FVector3d&, const FVector3d&)> IntensityFunc,
+			TFunctionRef<float(int32, const FVector3d&, const FVector3d&)> IntensityFunc,
 			const FSampledScalarField2f& DisplaceField,
 			TArray<FVector3d>& DisplacedPositions,
 			float DisplaceFieldBaseValue = 128.0/255, // value that corresponds to zero displacement
@@ -212,7 +218,7 @@ namespace DisplaceMeshToolLocals{
 					}
 					Offset -= DisplaceFieldBaseValue;
 
-					double Intensity = IntensityFunc(Positions[vid], Normals[vid]);
+					double Intensity = IntensityFunc(vid, Positions[vid], Normals[vid]);
 					DisplacedPositions[vid] = Positions[vid] + (Offset * Intensity * Normals[vid]);
 				}
 			}
@@ -221,7 +227,7 @@ namespace DisplaceMeshToolLocals{
 		void Sine(const FDynamicMesh3& Mesh,
 			const TArray<FVector3d>& Positions,
 			const FMeshNormals& Normals,
-			TFunctionRef<float(const FVector3d&, const FVector3d&)> IntensityFunc,
+			TFunctionRef<float(int32, const FVector3d&, const FVector3d&)> IntensityFunc,
 			double Frequency,
 			double PhaseShift,
 			const FVector3d& Direction,
@@ -234,7 +240,7 @@ namespace DisplaceMeshToolLocals{
 			{
 				FVector3d RotatedPosition = RotateToDirection * Position;
 				double DistXY = FMath::Sqrt(RotatedPosition.X * RotatedPosition.X + RotatedPosition.Y * RotatedPosition.Y);
-				double Intensity = IntensityFunc(Position, Normal);
+				double Intensity = IntensityFunc(vid, Position, Normal);
 				FVector3d Offset = Intensity * FMath::Sin(Frequency * DistXY + PhaseShift) * Direction;
 				return Position + Offset;
 
@@ -246,16 +252,36 @@ namespace DisplaceMeshToolLocals{
 	class FSubdivideMeshOp : public FDynamicMeshOperator
 	{
 	public:
-		FSubdivideMeshOp(const FDynamicMesh3& SourceMesh, int SubdivisionsCountIn);
+		FSubdivideMeshOp(const FDynamicMesh3& SourceMesh, int SubdivisionsCountIn, TSharedPtr<FIndexedWeightMap> WeightMap);
 		void CalculateResult(FProgressCancel* Progress) final;
 	private:
 		int SubdivisionsCount;
 	};
 
-	FSubdivideMeshOp::FSubdivideMeshOp(const FDynamicMesh3& SourceMesh, int SubdivisionsCountIn)
+	FSubdivideMeshOp::FSubdivideMeshOp(const FDynamicMesh3& SourceMesh, int SubdivisionsCountIn, TSharedPtr<FIndexedWeightMap> WeightMap)
 		: SubdivisionsCount(SubdivisionsCountIn)
 	{
 		ResultMesh->Copy(SourceMesh);
+
+		// If we have a WeightMap, initialize VertexUV.X with weightmap value. Note that we are going to process .Y anyway,
+		// we could (for exmaple) speculatively compute another weightmap, or store previous weightmap values there, to support
+		// fast switching between two...
+		ResultMesh->EnableVertexUVs(FVector2f::Zero());
+		if (WeightMap != nullptr)
+		{
+			for (int32 vid : ResultMesh->VertexIndicesItr())
+			{
+				ResultMesh->SetVertexUV(vid, FVector2f(WeightMap->GetValue(vid), 0));
+			}
+		}
+		else
+		{
+			for (int32 vid : ResultMesh->VertexIndicesItr())
+			{
+				ResultMesh->SetVertexUV(vid, FVector2f::One());
+			}
+		}
+
 	}
 
 	void FSubdivideMeshOp::CalculateResult(FProgressCancel* ProgressCancel)
@@ -265,8 +291,8 @@ namespace DisplaceMeshToolLocals{
 		// calculate subdivisions (todo: move to elsewhere)
 		for (int ri = 0; ri < SubdivisionsCount; ri++)
 		{
-			if (ProgressCancel->Cancelled()) return;
-			SubdivideMesh(*ResultMesh);
+			if (ProgressCancel && ProgressCancel->Cancelled()) return;
+			SubdivideMesh(*ResultMesh, ProgressCancel);
 		}
 	}
 
@@ -274,20 +300,24 @@ namespace DisplaceMeshToolLocals{
 	{
 	public:
 		FSubdivideMeshOpFactory(FDynamicMesh3& SourceMeshIn,
-			int SubdivisionsCountIn)
-			: SourceMesh(SourceMeshIn), SubdivisionsCount(SubdivisionsCountIn)
+			int SubdivisionsCountIn,
+			TSharedPtr<FIndexedWeightMap> WeightMapIn)
+			: SourceMesh(SourceMeshIn), SubdivisionsCount(SubdivisionsCountIn), WeightMap(WeightMapIn)
 		{
 		}
 		void SetSubdivisionsCount(int SubdivisionsCountIn);
 		int  GetSubdivisionsCount();
 
+		void SetWeightMap(TSharedPtr<FIndexedWeightMap> WeightMapIn);
+
 		TUniquePtr<FDynamicMeshOperator> MakeNewOperator() final
 		{
-			return MakeUnique<FSubdivideMeshOp>(SourceMesh, SubdivisionsCount);
+			return MakeUnique<FSubdivideMeshOp>(SourceMesh, SubdivisionsCount, WeightMap);
 		}
 	private:
 		const FDynamicMesh3& SourceMesh;
 		int SubdivisionsCount;
+		TSharedPtr<FIndexedWeightMap> WeightMap;
 	};
 
 	void FSubdivideMeshOpFactory::SetSubdivisionsCount(int SubdivisionsCountIn)
@@ -300,6 +330,10 @@ namespace DisplaceMeshToolLocals{
 		return SubdivisionsCount;
 	}
 
+	void FSubdivideMeshOpFactory::SetWeightMap(TSharedPtr<FIndexedWeightMap> WeightMapIn)
+	{
+		WeightMap = WeightMapIn;
+	}
 
 	// A collection of parameters to avoid having excess function parameters
 	struct DisplaceMeshParameters
@@ -387,14 +421,23 @@ namespace DisplaceMeshToolLocals{
 			Parameters.FilterWidth };
 		double Intensity = Parameters.DisplaceIntensity;
 
-		TUniqueFunction<float(const FVector3d&)> WeightMapQueryFunc = [&](const FVector3d&) { return 1.0f; };
+		TUniqueFunction<float(int32 vid, const FVector3d&)> WeightMapQueryFunc = [&](int32, const FVector3d&) { return 1.0f; };
 		if (Parameters.WeightMap.IsValid())
 		{
-			WeightMapQueryFunc = [&](const FVector3d& Pos) { return Parameters.WeightMapQueryFunc(Pos, *Parameters.WeightMap); };
+			if (SourceMesh->IsCompactV() && SourceMesh->VertexCount() == Parameters.WeightMap->Num())
+			{
+				WeightMapQueryFunc = [&](int32 vid, const FVector3d& Pos) { return Parameters.WeightMap->GetValue(vid); };
+			}
+			else
+			{
+				// disable input query function as it uses expensive AABBTree lookup
+				//WeightMapQueryFunc = [&](int32 vid, const FVector3d& Pos) { return Parameters.WeightMapQueryFunc(Pos, *Parameters.WeightMap); };
+				WeightMapQueryFunc = [&](int32 vid, const FVector3d& Pos) { return SourceMesh->GetVertexUV(vid).X; };
+			}
 		}
-		auto IntensityFunc = [&](const FVector3d& Position, const FVector3d& Normal) 
+		auto IntensityFunc = [&](int32 vid, const FVector3d& Position, const FVector3d& Normal) 
 		{
-			return Intensity * DirectionalFilter.FilterValue(Normal) * WeightMapQueryFunc(Position);
+			return Intensity * DirectionalFilter.FilterValue(Normal) * WeightMapQueryFunc(vid, Position);
 		};
 
 
@@ -762,8 +805,6 @@ void UDisplaceMeshTool::Setup()
 	OriginalMesh.Copy(*DynamicMeshComponent->GetMesh());
 	OriginalMeshSpatial.SetMesh(&OriginalMesh, true);
 
-	Subdivider = MakeUnique<FSubdivideMeshOpFactory>(OriginalMesh, CommonProperties->Subdivisions);
-	
 	DisplaceMeshParameters Parameters;
 	Parameters.DisplaceIntensity = CommonProperties->DisplaceIntensity;
 	Parameters.RandomSeed = CommonProperties->RandomSeed;
@@ -811,6 +852,8 @@ void UDisplaceMeshTool::Setup()
 									} );
 
 	ValidateSubdivisions();
+	Subdivider = MakeUnique<FSubdivideMeshOpFactory>(OriginalMesh, CommonProperties->Subdivisions, ActiveWeightMap);
+
 	StartComputation();
 
 	GetToolManager()->DisplayMessage(
@@ -862,15 +905,20 @@ void UDisplaceMeshTool::ValidateSubdivisions()
 		return;
 	}
 
+	bool bIsInitialized = (Subdivider != nullptr);
+
 	constexpr int MaxTriangles = 3000000;
 	double NumTriangles = OriginalMesh.MaxTriangleID();
 	int MaxSubdivisions = (int)floor(log2(MaxTriangles / NumTriangles) / 2.0);
 	if (CommonProperties->Subdivisions > MaxSubdivisions)
 	{
-		FText WarningText = FText::Format(LOCTEXT("SubdivisionsTooHigh", "Desired number of Subdivisions ({0}) exceeds maximum number ({1}) for a mesh of this number of triangles."), 
-			FText::AsNumber(CommonProperties->Subdivisions), 
-			FText::AsNumber(MaxSubdivisions));
-		GetToolManager()->DisplayMessage(WarningText, EToolMessageLevel::UserWarning);
+		if (bIsInitialized)		// only show warning after initial tool startup
+		{
+			FText WarningText = FText::Format(LOCTEXT("SubdivisionsTooHigh", "Desired number of Subdivisions ({0}) exceeds maximum number ({1}) for a mesh of this number of triangles."),
+				FText::AsNumber(CommonProperties->Subdivisions),
+				FText::AsNumber(MaxSubdivisions));
+			GetToolManager()->DisplayMessage(WarningText, EToolMessageLevel::UserWarning);
+		}
 		CommonProperties->Subdivisions = MaxSubdivisions;
 	}
 	else
@@ -954,7 +1002,9 @@ void UDisplaceMeshTool::OnPropertyModified(UObject* PropertySet, FProperty* Prop
 				 || PropName == GET_MEMBER_NAME_CHECKED(UDisplaceMeshCommonProperties, bInvertWeightMap))
 		{
 			UpdateActiveWeightMap();
+			SubdividerDownCast->SetWeightMap(ActiveWeightMap);
 			DisplacerDownCast->SetWeightMap(ActiveWeightMap);
+			bNeedsSubdivided = true;
 		}
 		else if (PropName == GET_MEMBER_NAME_CHECKED(UDisplaceMeshDirectionalFilterProperties, bEnableFilter))
 		{
@@ -1045,13 +1095,20 @@ void UDisplaceMeshTool::AdvanceComputation()
 
 void UDisplaceMeshTool::UpdateActiveWeightMap()
 {
-	TSharedPtr<FIndexedWeightMap> NewWeightMap = MakeShared<FIndexedWeightMap>();
-	UE::WeightMaps::GetVertexWeightMap(ComponentTarget->GetMesh(), CommonProperties->WeightMap, *NewWeightMap, 1.0f);
-	if (CommonProperties->bInvertWeightMap)
+	if (CommonProperties->WeightMap == FName(TEXT("None")))
+	{
+		ActiveWeightMap = nullptr;
+	}
+	else
 	{
-		NewWeightMap->InvertWeightMap();
+		TSharedPtr<FIndexedWeightMap> NewWeightMap = MakeShared<FIndexedWeightMap>();
+		UE::WeightMaps::GetVertexWeightMap(ComponentTarget->GetMesh(), CommonProperties->WeightMap, *NewWeightMap, 1.0f);
+		if (CommonProperties->bInvertWeightMap)
+		{
+			NewWeightMap->InvertWeightMap();
+		}
+		ActiveWeightMap = NewWeightMap;
 	}
-	ActiveWeightMap = NewWeightMap;
 }
 
 
