MeshModelingTools: Fixed possible crash in revolve tools when not taking any steps. Fixed bug where revolve was incorrectly considered full. Changed CutRevolve transaction name, changed MaxDegreesPerStep upper bound, changed name of DownAxisOffsetPerDegree.

#rnx

diff --git a/Engine/Plugins/Experimental/MeshModelingToolset/Source/MeshModelingTools/Private/DrawAndRevolveTool.cpp b/Engine/Plugins/Experimental/MeshModelingToolset/Source/MeshModelingTools/Private/DrawAndRevolveTool.cpp
--- a/Engine/Plugins/Experimental/MeshModelingToolset/Source/MeshModelingTools/Private/DrawAndRevolveTool.cpp
+++ b/Engine/Plugins/Experimental/MeshModelingToolset/Source/MeshModelingTools/Private/DrawAndRevolveTool.cpp
@@ -218,8 +218,12 @@ void UDrawAndRevolveTool::Shutdown(EToolShutdownType ShutdownType)
 
 void UDrawAndRevolveTool::GenerateAsset(const FDynamicMeshOpResult& Result)
 {
-	GetToolManager()->BeginUndoTransaction(LOCTEXT("RevolveToolTransactionName", "Revolve Tool"));
+	if (Result.Mesh->TriangleCount() <= 0)
+	{
+		return;
+	}
 
+	GetToolManager()->BeginUndoTransaction(LOCTEXT("RevolveToolTransactionName", "Revolve Tool"));
 
 	AActor* NewActor = AssetGenerationUtil::GenerateStaticMeshActor(
 		AssetAPI, TargetWorld, Result.Mesh.Get(), Result.Transform, TEXT("RevolveResult"), MaterialProperties->Material.Get());
diff --git a/Engine/Plugins/Experimental/MeshModelingToolset/Source/MeshModelingTools/Private/Properties/RevolveProperties.cpp b/Engine/Plugins/Experimental/MeshModelingToolset/Source/MeshModelingTools/Private/Properties/RevolveProperties.cpp
--- a/Engine/Plugins/Experimental/MeshModelingToolset/Source/MeshModelingTools/Private/Properties/RevolveProperties.cpp
+++ b/Engine/Plugins/Experimental/MeshModelingToolset/Source/MeshModelingTools/Private/Properties/RevolveProperties.cpp
@@ -25,7 +25,7 @@ void URevolveProperties::ApplyToCurveSweepOp(const UNewMeshMaterialProperties& M
 		}
 	}
 
-	double TotalRevolutionDegrees = (DownAxisOffsetPerDegree == 0) ? ClampedRevolutionDegrees : RevolutionDegrees;
+	double TotalRevolutionDegrees = (AlongAxisOffsetPerDegree == 0) ? ClampedRevolutionDegrees : RevolutionDegrees;
 
 	int32 Steps = bExplicitSteps ? NumExplicitSteps : FMath::CeilToInt(TotalRevolutionDegrees / MaxDegreesPerStep);
 	
@@ -36,7 +36,7 @@ void URevolveProperties::ApplyToCurveSweepOp(const UNewMeshMaterialProperties& M
 		DegreesPerStep *= -1;
 		DegreesOffset *= -1;
 	}
-	double DownAxisOffsetPerStep = TotalRevolutionDegrees * DownAxisOffsetPerDegree / Steps;
+	double DownAxisOffsetPerStep = TotalRevolutionDegrees * AlongAxisOffsetPerDegree / Steps;
 
 	if (bProfileIsCrossSectionOfSide && DegreesPerStep != 0 && abs(DegreesPerStep) < 180)
 	{
@@ -44,7 +44,7 @@ void URevolveProperties::ApplyToCurveSweepOp(const UNewMeshMaterialProperties& M
 	}
 
 	// Generate the sweep curve
-	CurveSweepOpOut.bSweepCurveIsClosed = bWeldFullRevolution && DownAxisOffsetPerDegree == 0 && RevolutionDegrees == 360;
+	CurveSweepOpOut.bSweepCurveIsClosed = bWeldFullRevolution && AlongAxisOffsetPerDegree == 0 && TotalRevolutionDegrees == 360;
 	int32 NumSweepFrames = CurveSweepOpOut.bSweepCurveIsClosed ? Steps : Steps + 1; // If closed, last sweep frame is also first
 	CurveSweepOpOut.SweepCurve.Reserve(NumSweepFrames);
 	RevolveUtil::GenerateSweepCurve(RevolutionAxisOrigin, RevolutionAxisDirection, DegreesOffset,
diff --git a/Engine/Plugins/Experimental/MeshModelingToolset/Source/MeshModelingTools/Private/RevolveBoundaryTool.cpp b/Engine/Plugins/Experimental/MeshModelingToolset/Source/MeshModelingTools/Private/RevolveBoundaryTool.cpp
--- a/Engine/Plugins/Experimental/MeshModelingToolset/Source/MeshModelingTools/Private/RevolveBoundaryTool.cpp
+++ b/Engine/Plugins/Experimental/MeshModelingToolset/Source/MeshModelingTools/Private/RevolveBoundaryTool.cpp
@@ -275,7 +275,12 @@ void URevolveBoundaryTool::Shutdown(EToolShutdownType ShutdownType)
 
 void URevolveBoundaryTool::GenerateAsset(const FDynamicMeshOpResult& Result)
 {
-	GetToolManager()->BeginUndoTransaction(LOCTEXT("RevolveBoundaryToolTransactionName", "Revolve Boundary Tool"));
+	if (Result.Mesh->TriangleCount() <= 0)
+	{
+		return;
+	}
+
+	GetToolManager()->BeginUndoTransaction(LOCTEXT("RevolveBoundaryToolTransactionName", "Revolve Tool"));
 
 	AActor* NewActor = AssetGenerationUtil::GenerateStaticMeshActor(
 		AssetAPI, TargetWorld, Result.Mesh.Get(), Result.Transform, TEXT("RevolveBoundaryResult"), MaterialProperties->Material.Get());
diff --git a/Engine/Plugins/Experimental/MeshModelingToolset/Source/MeshModelingTools/Public/DrawAndRevolveTool.h b/Engine/Plugins/Experimental/MeshModelingToolset/Source/MeshModelingTools/Public/DrawAndRevolveTool.h
--- a/Engine/Plugins/Experimental/MeshModelingToolset/Source/MeshModelingTools/Public/DrawAndRevolveTool.h
+++ b/Engine/Plugins/Experimental/MeshModelingToolset/Source/MeshModelingTools/Public/DrawAndRevolveTool.h
@@ -60,7 +60,7 @@ public:
 		UIMin = -180, UIMax = 180, ClampMin = -180000, ClampMax = 18000))
 	FRotator DrawPlaneOrientation = FRotator(90, 0, 0);
 
-	/** Enables/disables snapping while editing the profile curve. */
+	/** Enables/disables snapping while editing the profile curve. Also toggled with Shift. */
 	UPROPERTY(EditAnywhere, Category = ProfileCurve)
 	bool bEnableSnapping = true;
 
diff --git a/Engine/Plugins/Experimental/MeshModelingToolset/Source/MeshModelingTools/Public/Properties/RevolveProperties.h b/Engine/Plugins/Experimental/MeshModelingToolset/Source/MeshModelingTools/Public/Properties/RevolveProperties.h
--- a/Engine/Plugins/Experimental/MeshModelingToolset/Source/MeshModelingTools/Public/Properties/RevolveProperties.h
+++ b/Engine/Plugins/Experimental/MeshModelingToolset/Source/MeshModelingTools/Public/Properties/RevolveProperties.h
@@ -86,12 +86,12 @@ public:
 
 	/** Revolution extent. Clamped to a max of 360 when not offsetting along axis.*/
 	UPROPERTY(EditAnywhere, Category = RevolveSettings, meta = (UIMin = "0", UIMax = "360", ClampMin = "0", ClampMax = "360", 
-		EditCondition = "DownAxisOffsetPerDegree == 0", EditConditionHides, DisplayName = "Revolution Degrees"))
+		EditCondition = "AlongAxisOffsetPerDegree == 0", EditConditionHides, DisplayName = "Revolution Degrees"))
 	double ClampedRevolutionDegrees = 360;
 
 	/** Revolution extent. */
 	UPROPERTY(EditAnywhere, Category = RevolveSettings, meta = (UIMin = "0", UIMax = "3600", ClampMin = "0", ClampMax = "360000",
-		EditCondition = "DownAxisOffsetPerDegree != 0", EditConditionHides))
+		EditCondition = "AlongAxisOffsetPerDegree != 0", EditConditionHides))
 	double RevolutionDegrees = 360;
 
 	/** When true, the number of steps can be specified explicitly. When false, the number of steps is adjusted automatically. */
@@ -104,13 +104,13 @@ public:
 	int NumExplicitSteps = 24;
 
 	/** The revolution is split into a number of steps such that each step moves the revolution no more than this number of degrees.  */
-	UPROPERTY(EditAnywhere, Category = RevolveSettings, meta = (UIMin = "1", ClampMin = "1", UIMax = "100", ClampMax = "5000",
+	UPROPERTY(EditAnywhere, Category = RevolveSettings, meta = (UIMin = "1", ClampMin = "1", UIMax = "120", ClampMax = "180",
 		EditCondition = "!bExplicitSteps", EditConditionHides))
 	double MaxDegreesPerStep = 15;
 
-	/** How far to move each step down the axis (per degree). Used to create spirals. */
+	/** How far to move each step along the axis (per degree). Used to create spirals. */
 	UPROPERTY(EditAnywhere, Category = RevolveSettings, meta = (UIMin = "-1", UIMax = "1", ClampMin = "-100000", ClampMax = "100000"))
-	double DownAxisOffsetPerDegree = 0;
+	double AlongAxisOffsetPerDegree = 0;
 
 	/** The angle by which to shift the profile curve around the axis before beginning the revolve */
 	UPROPERTY(EditAnywhere, Category = RevolveSettings, meta = (UIMin = "-360", UIMax = "360", ClampMin = "-36000", ClampMax = "36000"), AdvancedDisplay)
@@ -149,13 +149,13 @@ public:
 	/** Determines how caps are created if the revolution is partial. Not relevant if the
 	  revolution is full and welded. */
 	UPROPERTY(EditAnywhere, Category = RevolveSettings, AdvancedDisplay, meta = (
-		EditCondition = "DownAxisOffsetPerDegree != 0 || RevolutionDegrees != 360", EditConditionHides))
+		EditCondition = "AlongAxisOffsetPerDegree != 0 || RevolutionDegrees != 360", EditConditionHides))
 	ERevolvePropertiesCapFillMode CapFillMode = ERevolvePropertiesCapFillMode::Delaunay;
 
 	/** If true, the ends of a fully revolved profile are welded together, rather than duplicating
 	  vertices at the seam. Not relevant if the revolution is not full. */
 	UPROPERTY(EditAnywhere, Category = RevolveSettings, AdvancedDisplay, meta = (
-		EditCondition = "DownAxisOffsetPerDegree == 0 && RevolutionDegrees == 360", EditConditionHides))
+		EditCondition = "AlongAxisOffsetPerDegree == 0 && RevolutionDegrees == 360", EditConditionHides))
 	bool bWeldFullRevolution = true;
 
 	/** If true, vertices sufficiently close to the axis will not be replicated, instead reusing
