ModelingTools: add option to Displace Tool to enable/disable normals recomputation on the displaced mesh. This is for use with an existing normal/displacement map set, where the normal map will be used on the displaced mesh.
#rnx

diff --git a/Engine/Plugins/Experimental/MeshModelingToolset/Source/MeshModelingTools/Private/DisplaceMeshTool.cpp b/Engine/Plugins/Experimental/MeshModelingToolset/Source/MeshModelingTools/Private/DisplaceMeshTool.cpp
--- a/Engine/Plugins/Experimental/MeshModelingToolset/Source/MeshModelingTools/Private/DisplaceMeshTool.cpp
+++ b/Engine/Plugins/Experimental/MeshModelingToolset/Source/MeshModelingTools/Private/DisplaceMeshTool.cpp
@@ -15,6 +15,13 @@
 #include "Async/ParallelFor.h"
 #include "ProfilingDebugging/ScopedTimers.h"
 #include "MeshDescription.h"
+
+// needed to disable normals recalculation on the underlying asset
+#include "AssetUtils/MeshDescriptionUtil.h"
+#include "Engine/Classes/Engine/StaticMesh.h"
+#include "Engine/Classes/Components/StaticMeshComponent.h"
+
+
 #define LOCTEXT_NAMESPACE "UDisplaceMeshTool"
 
 namespace DisplaceMeshToolLocals{
@@ -349,6 +356,7 @@ namespace DisplaceMeshToolLocals{
 		float FilterWidth = 0.0f;
 		FSampledScalarField2f DisplaceField;
 		TArray<FPerlinLayerProperties> PerlinLayerProperties;
+		bool bRecalculateNormals = true;
 
 		// Used in texture map displacement
 		float DisplacementMapBaseValue = 128.0/255; // i.e., what constitutes no displacement
@@ -504,16 +512,19 @@ namespace DisplaceMeshToolLocals{
 		}
 
 		// recalculate normals
-		if (ResultMesh->HasAttributes())
+		if (Parameters.bRecalculateNormals)
 		{
-			FMeshNormals Normals(ResultMesh.Get());
-			FDynamicMeshNormalOverlay* NormalOverlay = ResultMesh->Attributes()->PrimaryNormals();
-			Normals.RecomputeOverlayNormals(NormalOverlay);
-			Normals.CopyToOverlay(NormalOverlay);
-		}
-		else
-		{
-			FMeshNormals::QuickComputeVertexNormals(*ResultMesh);
+			if (ResultMesh->HasAttributes())
+			{
+				FMeshNormals Normals(ResultMesh.Get());
+				FDynamicMeshNormalOverlay* NormalOverlay = ResultMesh->Attributes()->PrimaryNormals();
+				Normals.RecomputeOverlayNormals(NormalOverlay);
+				Normals.CopyToOverlay(NormalOverlay);
+			}
+			else
+			{
+				FMeshNormals::QuickComputeVertexNormals(*ResultMesh);
+			}
 		}
 	}
 
@@ -561,6 +572,7 @@ namespace DisplaceMeshToolLocals{
 		void SetFilterFalloffWidth(float FalloffWidth);
 		void SetPerlinNoiseLayerProperties(const TArray<FPerlinLayerProperties>& PerlinLayerProperties);
 		void SetWeightMap(TSharedPtr<FIndexedWeightMap> WeightMap);
+		void SetRecalculateNormals(bool bRecalculateNormals);
 
 		TUniquePtr<FDynamicMeshOperator> MakeNewOperator() final
 		{
@@ -698,6 +710,11 @@ namespace DisplaceMeshToolLocals{
 		Parameters.WeightMap = WeightMap;
 	}
 
+	void FDisplaceMeshOpFactory::SetRecalculateNormals(bool RecalcNormalsIn)
+	{
+		Parameters.bRecalculateNormals = RecalcNormalsIn;
+	}
+
 } // namespace
 
 /*
@@ -884,6 +901,27 @@ void UDisplaceMeshTool::Shutdown(EToolShutdownType ShutdownType)
 		{
 			// this block bakes the modified DynamicMeshComponent back into the StaticMeshComponent inside an undo transaction
 			GetToolManager()->BeginUndoTransaction(LOCTEXT("DisplaceMeshToolTransactionName", "Displace Mesh"));
+
+			// if we are applying a map and not recalculating normals, we need to make sure normals recalculation is disabled
+			// on the underlying StaticMesh Asset, or it will run on the Bake() below and the output result will not be the same as the preview
+			if (CommonProperties->DisplacementType == EDisplaceMeshToolDisplaceType::DisplacementMap && TextureMapProperties->bRecalcNormals == false)
+			{
+				UStaticMeshComponent* StaticMeshComponent = Cast<UStaticMeshComponent>(ComponentTarget->GetOwnerComponent());
+				if (StaticMeshComponent)
+				{
+					if (UStaticMesh* StaticMesh = StaticMeshComponent->GetStaticMesh())
+					{
+						StaticMesh->Modify();
+
+						// disable auto-generated normals and tangents build settings
+						UE::MeshDescription::FStaticMeshBuildSettingChange SettingsChange;
+						SettingsChange.AutoGeneratedNormals = UE::MeshDescription::EBuildSettingBoolChange::Disable;
+						UE::MeshDescription::ConfigureBuildSettings(StaticMesh, 0, SettingsChange);
+					}
+				}
+			}
+
+
 			ComponentTarget->CommitMesh([=](const FPrimitiveComponentTarget::FCommitParams& CommitParams)
 			{
 				DynamicMeshComponent->Bake(CommitParams.MeshDescription, CommonProperties->Subdivisions > 0);
@@ -987,12 +1025,29 @@ void UDisplaceMeshTool::OnPropertyModified(UObject* PropertySet, FProperty* Prop
 		}
 		else if (PropName == GET_MEMBER_NAME_CHECKED(UDisplaceMeshTextureMapProperties, DisplacementMap))
 		{
+			if (TextureMapProperties->DisplacementMap != nullptr
+				&& TextureMapProperties->DisplacementMap->VirtualTextureStreaming)
+			{
+				GetToolManager()->DisplayMessage(
+					LOCTEXT("DisplaceToolVirtualTextureMessage", "Virtual Texture must be disabled on the selected Texture2D to use it as a Displacement Map in this Tool"),
+					EToolMessageLevel::UserWarning);
+			}
+			else
+			{
+				GetToolManager()->DisplayMessage(FText::GetEmpty(), EToolMessageLevel::UserWarning);
+			}
+
 			DisplacerDownCast->SetDisplacementMap(TextureMapProperties->DisplacementMap);
 		}
 		else if (PropName == GET_MEMBER_NAME_CHECKED(UDisplaceMeshTextureMapProperties, DisplacementMapBaseValue))
 		{
 			DisplacerDownCast->SetDisplacementMapBaseValue(TextureMapProperties->DisplacementMapBaseValue);
 		}
+		else if (PropName == GET_MEMBER_NAME_CHECKED(UDisplaceMeshTextureMapProperties, bRecalcNormals))
+		{
+			DisplacerDownCast->SetRecalculateNormals(TextureMapProperties->bRecalcNormals);
+		}
+
 		else if (PropName == GET_MEMBER_NAME_CHECKED(UDisplaceMeshTextureMapProperties, bApplyAdjustmentCurve)
 			|| PropName == GET_MEMBER_NAME_CHECKED(UDisplaceMeshTextureMapProperties, AdjustmentCurve))
 		{
diff --git a/Engine/Plugins/Experimental/MeshModelingToolset/Source/MeshModelingTools/Public/DisplaceMeshTool.h b/Engine/Plugins/Experimental/MeshModelingToolset/Source/MeshModelingTools/Public/DisplaceMeshTool.h
--- a/Engine/Plugins/Experimental/MeshModelingToolset/Source/MeshModelingTools/Public/DisplaceMeshTool.h
+++ b/Engine/Plugins/Experimental/MeshModelingToolset/Source/MeshModelingTools/Public/DisplaceMeshTool.h
@@ -127,6 +127,10 @@ public:
 	/** This curve is queried in the range [0,1] to adjust contrast of the displacement map. */
 	UPROPERTY(VisibleAnywhere, Category = Options, meta = (EditCondition = "bApplyAdjustmentCurve", EditConditionHides))
 	UCurveFloat* AdjustmentCurve = nullptr;
+
+	/** Recalculate normals from displaced mesh. Disable this if you are applying Displacements that are paired with an existing Normal Map in your Material. */
+	UPROPERTY(EditAnywhere, Category = Options)
+	bool bRecalcNormals = true;
 };
 
 
