fix inconsistency w/ UV scaling on polygon when extruded down instead of up

#rnx

diff --git a/Engine/Plugins/Experimental/MeshModelingToolset/Source/MeshModelingTools/Private/DrawPolygonTool.cpp b/Engine/Plugins/Experimental/MeshModelingToolset/Source/MeshModelingTools/Private/DrawPolygonTool.cpp
--- a/Engine/Plugins/Experimental/MeshModelingToolset/Source/MeshModelingTools/Private/DrawPolygonTool.cpp
+++ b/Engine/Plugins/Experimental/MeshModelingToolset/Source/MeshModelingTools/Private/DrawPolygonTool.cpp
@@ -1230,7 +1230,7 @@ bool UDrawPolygonTool::GeneratePolygonMesh(const TArray<FVector3d>& Polygon, con
 	FDynamicMeshEditor Editor(ResultMeshOut);
 	float InitialUVScale = 1.0 / Bounds.MaxDim(); // this is the UV scale used by both the polymeshgen and the extruder above
 	// default global rescale -- initial scale doesn't factor in extrude distance; rescale so UVScale of 1.0 fits in the unit square texture
-	float GlobalUVRescale = MaterialProperties->UVScale / FMathf::Max(1.0f, ExtrudeDistance * InitialUVScale);
+	float GlobalUVRescale = MaterialProperties->UVScale / FMathf::Max(1.0f, FMathd::Abs(ExtrudeDistance) * InitialUVScale);
 	if (MaterialProperties->bWorldSpaceUVScale)
 	{
 		// since we know the initial uv scale, directly compute the global scale (relative to 1 meter as a standard scale)
