ModelingTools: improve performance in DynamicMeshSculptTool by making some of the sub-computations asynchronous. Add button inside the Tool to discard attributes, rather than forcing user to round-trip through another Tool.
#rnx

diff --git a/Engine/Plugins/Experimental/GeometryProcessing/Source/DynamicMesh/Public/DynamicMeshOctree3.h b/Engine/Plugins/Experimental/GeometryProcessing/Source/DynamicMesh/Public/DynamicMeshOctree3.h
--- a/Engine/Plugins/Experimental/GeometryProcessing/Source/DynamicMesh/Public/DynamicMeshOctree3.h
+++ b/Engine/Plugins/Experimental/GeometryProcessing/Source/DynamicMesh/Public/DynamicMeshOctree3.h
@@ -57,9 +57,12 @@ public:
 	 */
 	void InsertTriangle(int32 TriangleID)
 	{
-		FAxisAlignedBox3d Bounds = Mesh->GetTriBounds(TriangleID);
-		ModifiedBounds.Contain(Bounds);
-		InsertObject(TriangleID, Bounds);
+		if (Mesh->IsTriangle(TriangleID))
+		{
+			FAxisAlignedBox3d Bounds = Mesh->GetTriBounds(TriangleID);
+			ModifiedBounds.Contain(Bounds);
+			InsertObject(TriangleID, Bounds);
+		}
 	}
 
 	/**
@@ -70,9 +73,12 @@ public:
 		int N = Triangles.Num();
 		for (int i = 0; i < N; ++i)
 		{
-			FAxisAlignedBox3d Bounds = Mesh->GetTriBounds(Triangles[i]);
-			ModifiedBounds.Contain(Bounds);
-			InsertObject(Triangles[i], Bounds);
+			if (Mesh->IsTriangle(Triangles[i]))
+			{
+				FAxisAlignedBox3d Bounds = Mesh->GetTriBounds(Triangles[i]);
+				ModifiedBounds.Contain(Bounds);
+				InsertObject(Triangles[i], Bounds);
+			}
 		}
 	}
 
@@ -83,9 +89,12 @@ public:
 	{
 		for (int TriangleID : Triangles)
 		{
-			FAxisAlignedBox3d Bounds = Mesh->GetTriBounds(TriangleID);
-			ModifiedBounds.Contain(Bounds);
-			InsertObject(TriangleID, Bounds);
+			if (Mesh->IsTriangle(TriangleID))
+			{
+				FAxisAlignedBox3d Bounds = Mesh->GetTriBounds(TriangleID);
+				ModifiedBounds.Contain(Bounds);
+				InsertObject(TriangleID, Bounds);
+			}
 		}
 	}
 
@@ -95,9 +104,12 @@ public:
 	 */
 	bool RemoveTriangle(int32 TriangleID)
 	{
-		FAxisAlignedBox3d Bounds = Mesh->GetTriBounds(TriangleID);
-		ModifiedBounds.Contain(Bounds);
-		return RemoveObject(TriangleID);
+		if (Mesh->IsTriangle(TriangleID))
+		{
+			FAxisAlignedBox3d Bounds = Mesh->GetTriBounds(TriangleID);
+			ModifiedBounds.Contain(Bounds);
+		}
+		return RemoveObject(TriangleID);		// will ignore if we do not contain this triangle
 	}
 
 	/**
@@ -108,9 +120,12 @@ public:
 		int N = Triangles.Num();
 		for ( int i = 0; i < N; ++i )
 		{
-			FAxisAlignedBox3d Bounds = Mesh->GetTriBounds(Triangles[i]);
-			ModifiedBounds.Contain(Bounds);
-			RemoveObject(Triangles[i]);
+			if (Mesh->IsTriangle(Triangles[i]))
+			{
+				FAxisAlignedBox3d Bounds = Mesh->GetTriBounds(Triangles[i]);
+				ModifiedBounds.Contain(Bounds);
+			}
+			RemoveObject(Triangles[i]);		// will ignore if we do not contain this triangle
 		}
 	}
 
@@ -121,9 +136,12 @@ public:
 	{
 		for (int TriangleID : Triangles)
 		{
-			FAxisAlignedBox3d Bounds = Mesh->GetTriBounds(TriangleID);
-			ModifiedBounds.Contain(Bounds);
-			RemoveObject(TriangleID);
+			if (Mesh->IsTriangle(TriangleID))
+			{
+				FAxisAlignedBox3d Bounds = Mesh->GetTriBounds(TriangleID);
+				ModifiedBounds.Contain(Bounds);
+			}
+			RemoveObject(TriangleID);		// will ignore if we do not contain this triangle
 		}
 	}
 
@@ -135,9 +153,16 @@ public:
 	{
 		for (int TriangleID : Triangles)
 		{
-			FAxisAlignedBox3d Bounds = Mesh->GetTriBounds(TriangleID);
-			ModifiedBounds.Contain(Bounds);
-			ReinsertObject(TriangleID, Bounds);
+			if (Mesh->IsTriangle(TriangleID))
+			{
+				FAxisAlignedBox3d Bounds = Mesh->GetTriBounds(TriangleID);
+				ModifiedBounds.Contain(Bounds);
+				ReinsertObject(TriangleID, Bounds);
+			}
+			else
+			{
+				RemoveObject(TriangleID);		// can only remove, will ignore if we do not contain this triangle
+			}
 		}
 	}
 
@@ -147,8 +172,11 @@ public:
 	 */
 	void NotifyPendingModification(int TriangleID)
 	{
-		FAxisAlignedBox3d Bounds = Mesh->GetTriBounds(TriangleID);
-		ModifiedBounds.Contain(Bounds);
+		if (Mesh->IsTriangle(TriangleID))
+		{
+			FAxisAlignedBox3d Bounds = Mesh->GetTriBounds(TriangleID);
+			ModifiedBounds.Contain(Bounds);
+		}
 	}
 
 	/**
@@ -158,8 +186,11 @@ public:
 	{
 		for (int TriangleID : Triangles)
 		{
-			FAxisAlignedBox3d Bounds = Mesh->GetTriBounds(TriangleID);
-			ModifiedBounds.Contain(Bounds);
+			if (Mesh->IsTriangle(TriangleID))
+			{
+				FAxisAlignedBox3d Bounds = Mesh->GetTriBounds(TriangleID);
+				ModifiedBounds.Contain(Bounds);
+			}
 		}
 	}
 
diff --git a/Engine/Plugins/Experimental/GeometryProcessing/Source/DynamicMesh/Public/SubRegionRemesher.h b/Engine/Plugins/Experimental/GeometryProcessing/Source/DynamicMesh/Public/SubRegionRemesher.h
--- a/Engine/Plugins/Experimental/GeometryProcessing/Source/DynamicMesh/Public/SubRegionRemesher.h
+++ b/Engine/Plugins/Experimental/GeometryProcessing/Source/DynamicMesh/Public/SubRegionRemesher.h
@@ -215,6 +215,20 @@ public:
 	}
 
 
+
+	/**
+	 * Extract final triangle ROI. After calling this function, the ROI will be cleared, and must be re-initialized.
+	 * This can be used to avoid having to copy the TSet.
+	 * @return set of triangles that contains edge ROI (note: may also contain additional triangles)
+	 */
+	TSet<int>&& ExtractFinalTriangleROI()
+	{
+		return MoveTemp(TriangleROI);
+		EdgeROI.Reset();
+		VertexROI.Reset();
+	}
+
+
 	//
 	// specialization of Remesher functionality
 	// 
diff --git a/Engine/Plugins/Experimental/MeshModelingToolset/Source/MeshModelingTools/Private/DynamicMeshSculptTool.cpp b/Engine/Plugins/Experimental/MeshModelingToolset/Source/MeshModelingTools/Private/DynamicMeshSculptTool.cpp
--- a/Engine/Plugins/Experimental/MeshModelingToolset/Source/MeshModelingTools/Private/DynamicMeshSculptTool.cpp
+++ b/Engine/Plugins/Experimental/MeshModelingToolset/Source/MeshModelingTools/Private/DynamicMeshSculptTool.cpp
@@ -2,6 +2,7 @@
 
 #include "DynamicMeshSculptTool.h"
 #include "Containers/Map.h"
+#include "Async/Async.h"
 #include "InteractiveToolManager.h"
 #include "InteractiveGizmoManager.h"
 #include "ToolBuilderUtil.h"
@@ -38,6 +39,17 @@
 #define LOCTEXT_NAMESPACE "UDynamicMeshSculptTool"
 
 
+namespace
+{
+	// probably should be something defined for the whole tool framework...
+#if WITH_EDITOR
+	static EAsyncExecution DynamicSculptToolAsyncExecTarget = EAsyncExecution::LargeThreadPool;
+#else
+	static EAsyncExecution DynamicSculptToolAsyncExecTarget = EAsyncExecution::ThreadPool;
+#endif
+}
+
+
 
 /*
  * ToolBuilder
@@ -51,6 +63,14 @@ UMeshSurfacePointTool* UDynamicMeshSculptToolBuilder::CreateNewTool(const FToolB
 }
 
 
+
+void UDynamicSculptToolActions::DiscardAttributes()
+{
+	ParentTool->DiscardAttributes();
+}
+
+
+
 /*
  * Tool
  */
@@ -114,6 +134,7 @@ void UDynamicMeshSculptTool::Setup()
 	// initialize target mesh
 	UpdateTarget();
 	bTargetDirty = false;
+	PendingTargetUpdate.Wait();
 
 	// initialize brush radius range interval, brush properties
 	double MaxDimension = DynamicMeshComponent->GetMesh()->GetCachedBounds().MaxDim();
@@ -171,6 +192,10 @@ void UDynamicMeshSculptTool::Setup()
 	{
 		SculptProperties->bIsRemeshingEnabled = true;
 		AddToolPropertySource(RemeshProperties);
+
+		SculptToolActions = NewObject<UDynamicSculptToolActions>();
+		SculptToolActions->Initialize(this);
+		AddToolPropertySource(SculptToolActions);
 	}
 
 	BrushProperties->RestoreProperties(this);
@@ -1176,6 +1201,8 @@ int UDynamicMeshSculptTool::FindHitSculptMeshTriangle(const FRay3d& LocalRay)
 
 int UDynamicMeshSculptTool::FindHitTargetMeshTriangle(const FRay3d& LocalRay)
 {
+	PendingTargetUpdate.Wait();
+
 	IMeshSpatial::FQueryOptions RaycastOptions;
 
 	if (BrushProperties->bHitBackFaces == false)
@@ -1210,6 +1237,8 @@ bool UDynamicMeshSculptTool::UpdateBrushPositionOnActivePlane(const FRay& WorldR
 
 bool UDynamicMeshSculptTool::UpdateBrushPositionOnTargetMesh(const FRay& WorldRay, bool bFallbackToViewPlane)
 {
+	PendingTargetUpdate.Wait();
+
 	FRay3d LocalRay(CurTargetTransform.InverseTransformPosition(WorldRay.Origin),
 		CurTargetTransform.InverseTransformVector(WorldRay.Direction));
 	LocalRay.Direction.Normalize();
@@ -1444,6 +1473,7 @@ void UDynamicMeshSculptTool::OnTick(float DeltaTime)
 	if (bCachedFreezeTarget != SculptProperties->bFreezeTarget)
 	{
 		UpdateTarget();
+		PendingTargetUpdate.Wait();
 	}
 
 	bool bMeshModified = false;
@@ -1458,30 +1488,24 @@ void UDynamicMeshSculptTool::OnTick(float DeltaTime)
 	// Apply stamp
 	//
 
+	bool bROIUpdatePending = false;
+	bool bOctreeUpdatePending = false;
+
 	if (bStampPending)
 	{
 		SCOPE_CYCLE_COUNTER(STAT_SculptTool_Tick_ApplyStampBlock);
 
-		// would this ever be true? does stamp require this?
-		bool bRemoveTrianglesBeforeStamp = false;
-
-		if ( bRemoveTrianglesBeforeStamp )
-		{
-			SCOPE_CYCLE_COUNTER(SculptTool_Tick_ApplyStamp_Remove);
-			Octree->RemoveTriangles(TriangleROI);
-		}
-		else
-		{
-			SCOPE_CYCLE_COUNTER(SculptTool_Tick_ApplyStamp_Remove);
-			Octree->NotifyPendingModification(TriangleROI);		// to mark initial positions
-		}
+		SCOPE_CYCLE_COUNTER(SculptTool_Tick_ApplyStamp_Remove);
+		Octree->NotifyPendingModification(TriangleROI);		// to mark initial positions
 
 		ApplyStamp(PendingStampRay);
 		bStampPending = false;
 
 		bNormalUpdatePending = true;
+		bROIUpdatePending = true;
 		bMeshModified = true;
 		bMeshShapeModified = true;
+		bOctreeUpdatePending = true;
 
 		// flow
 		if (bInDrag)
@@ -1489,37 +1513,55 @@ void UDynamicMeshSculptTool::OnTick(float DeltaTime)
 			bStampPending = true;
 		}
 
-		if (bRemoveTrianglesBeforeStamp)
-		{
-			SCOPE_CYCLE_COUNTER(SculptTool_Tick_ApplyStamp_Insert);
-			Octree->InsertTriangles(TriangleROI);
-		}
-		else
-		{
-			SCOPE_CYCLE_COUNTER(SculptTool_Tick_ApplyStamp_Insert);
-			Octree->ReinsertTriangles(TriangleROI);
-		}
+		SCOPE_CYCLE_COUNTER(SculptTool_Tick_ApplyStamp_Insert);
 
+		//UpdateOctreeFuture = Async(DynamicSculptToolAsyncExecTarget, [&]() {
+		//	Octree->ReinsertTriangles(TriangleROI);
+		//});
 	}
 
+	//UpdateOctreeFuture.Wait();
 
-	bool bUpdatedROIInRemesh = false;
 	if (bRemeshPending)
 	{
 		SCOPE_CYCLE_COUNTER(STAT_SculptTool_Tick_RemeshBlock);
 
 		check(bInDrag == true);    // this would break undo otherwise!
 
+		// only removes from octree
 		RemeshROIPass();
 
+		// accumulate new triangles into TriangleROI
+		for (int32 tid : RemeshFinalTriangleROI)
+		{
+			TriangleROI.Add(tid);
+		}
+
 		bMeshModified = true;
 		bMeshShapeModified = true;
 		bRemeshPending = false;
 		bNormalUpdatePending = true;
+		bROIUpdatePending = true;
+		bOctreeUpdatePending = true;
 		bHaveRemeshed = true;
-		bUpdatedROIInRemesh = true;
 	}
 
+
+	// launch octree update
+	TFuture<void> UpdateOctreeFuture;
+	if (bOctreeUpdatePending)
+	{
+		SCOPE_CYCLE_COUNTER(STAT_SculptTool_Remesh_Finish);
+		// reinsert new ROI into octree
+		UpdateOctreeFuture = Async(DynamicSculptToolAsyncExecTarget, [&]() {
+			Octree->ReinsertTriangles(TriangleROI);
+			//Octree->CheckValidity(EValidityCheckFailMode::Check, false, true);
+		});
+		bOctreeUpdatePending = false;
+	}
+
+	//UpdateOctreeFuture.Wait();
+
 	//Octree->CheckValidity(EValidityCheckFailMode::Check, false, true);
 
 	if (bNormalUpdatePending)
@@ -1528,38 +1570,49 @@ void UDynamicMeshSculptTool::OnTick(float DeltaTime)
 
 		if (Mesh->HasAttributes() && Mesh->Attributes()->PrimaryNormals() != nullptr)
 		{
-			RecalculateNormals_Overlay();
+			RecalculateNormals_Overlay(TriangleROI);
 		}
 		else
 		{
-			RecalculateNormals_PerVertex();
+			RecalculateNormals_PerVertex(TriangleROI);
 		}
 		bNormalUpdatePending = false;
 		bMeshModified = true;
 	}
 
+	// next steps need to wait for octree update to finish
+	UpdateOctreeFuture.Wait();
+
+	TFuture<void> ROIUpdateFuture;
+	if (bROIUpdatePending)
+	{
+		ROIUpdateFuture = Async(DynamicSculptToolAsyncExecTarget, [&]() 
+		{
+			UpdateROI(LastBrushPosLocal);
+		});
+		bROIUpdatePending = false;
+	}
+
+	// launch async target update task
+	if (bTargetDirty)
+	{
+		UpdateTarget();		// TODO: does this make sense to do as a full mesh copy??
+		bTargetDirty = false;
+	}
+
 	if (bMeshModified)
 	{
 		SCOPE_CYCLE_COUNTER(STAT_SculptTool_Tick_UpdateMeshBlock);
 
-		if (bMeshModified)
-
 		DynamicMeshComponent->NotifyMeshUpdated();
 		GetToolManager()->PostInvalidation();
 
 		bMeshModified = false;
-
-		if (bUpdatedROIInRemesh == false)
-		{
-			UpdateROI(LastBrushPosLocal);
-		}
 	}
 
-	if (bTargetDirty)
-	{
-		UpdateTarget();
-		bTargetDirty = false;
-	}
+
+	// wait for ROI update to complete (can defer but it will immediately happen on hover...)
+	ROIUpdateFuture.Wait();
 }
 
 void UDynamicMeshSculptTool::PrecomputeRemeshInfo()
@@ -1611,7 +1664,6 @@ Ignore	TRUE	TRUE	FALSE	TRUE
 void UDynamicMeshSculptTool::ConfigureRemesher(FSubRegionRemesher& Remesher)
 {
 	FDynamicMesh3* Mesh = DynamicMeshComponent->GetMesh();
-	FDynamicMeshOctree3* Octree = DynamicMeshComponent->GetOctree();
 
 	const double SizeRange = 5;
 	double LengthMultiplier = (RemeshProperties->TriangleSize >= 0) ?
@@ -1683,7 +1735,6 @@ void UDynamicMeshSculptTool::ConfigureRemesher(FSubRegionRemesher& Remesher)
 		}
 		Remesher.InitializeFromVertexROI();
 		Remesher.UpdateROI();	// required to use roi in constraints fn below
-		Octree->RemoveTriangles(Remesher.GetCurrentTriangleROI());
 	}
 
 	FMeshConstraints Constraints;
@@ -1699,10 +1750,10 @@ void UDynamicMeshSculptTool::ConfigureRemesher(FSubRegionRemesher& Remesher)
 															 (EEdgeRefineFlags)RemeshProperties->GroupBoundaryConstraint,
 															 (EEdgeRefineFlags)RemeshProperties->MaterialBoundaryConstraint,
 															 bConstraintAllowSplits, !RemeshProperties->bPreserveSharpEdges);
-		if ( Constraints.HasConstraints() )
-		{
+		//if ( Constraints.HasConstraints() )
+		//{
 			Remesher.SetExternalConstraints(MoveTemp(Constraints));
-		}
+		//}
 	}
 }
 
@@ -1710,9 +1761,14 @@ void UDynamicMeshSculptTool::RemeshROIPass()
 {
 	FDynamicMesh3* Mesh = DynamicMeshComponent->GetMesh();
 	FDynamicMeshOctree3* Octree = DynamicMeshComponent->GetOctree();
+
 	FSubRegionRemesher Remesher(Mesh);
 	ConfigureRemesher(Remesher);
 
+	// remove initial triangles
+	RemeshRemovedTriangles = Remesher.GetCurrentTriangleROI();
+	Octree->RemoveTriangles(RemeshRemovedTriangles);
+
 	if (ActiveMeshChange != nullptr)
 	{
 		Remesher.SetMeshChangeTracker(ActiveMeshChange);
@@ -1752,26 +1808,21 @@ void UDynamicMeshSculptTool::RemeshROIPass()
 			SCOPE_CYCLE_COUNTER(STAT_SculptTool_Remesh_PassOctreeUpdate);
 			const TSet<int32>& TrisRemovedInPass = Remesher.EndTrackRemovedTrisInPass();
 			Octree->RemoveTriangles(TrisRemovedInPass);
+			for (int32 tid : TrisRemovedInPass)
+			{
+				RemeshRemovedTriangles.Add(tid);
+			}
 		}
 	}
 	//UE_LOG(LogTemp, Warning, TEXT("Triangle Count %d after update"), Mesh->TriangleCount());
 
-	{
-		SCOPE_CYCLE_COUNTER(STAT_SculptTool_Remesh_Finish);
-
-		// reinsert new ROI into octree
-		Octree->ReinsertTriangles(Remesher.GetCurrentTriangleROI());
-
-		//Octree->CheckValidity(EValidityCheckFailMode::Check, false, true);
-
-		UpdateROI(LastBrushPosLocal);
-	}
+	RemeshFinalTriangleROI = Remesher.ExtractFinalTriangleROI();
 }
 
-void UDynamicMeshSculptTool::RecalculateNormals_PerVertex()
+void UDynamicMeshSculptTool::RecalculateNormals_PerVertex(const TSet<int32>& Triangles)
 {
 	FDynamicMesh3* Mesh = DynamicMeshComponent->GetMesh();
-	FDynamicMeshOctree3* Octree = DynamicMeshComponent->GetOctree();
+	//FDynamicMeshOctree3* Octree = DynamicMeshComponent->GetOctree();
 
 	int MaxVertexID = Mesh->MaxVertexID();
 	if (NormalsVertexFlags.Num() < MaxVertexID)
@@ -1779,23 +1830,26 @@ void UDynamicMeshSculptTool::RecalculateNormals_PerVertex()
 		NormalsVertexFlags.Init(false, MaxVertexID * 2);
 	}
 
-
 	{
 		SCOPE_CYCLE_COUNTER(SculptTool_Normals_Collect);
 
-		TrianglesBuffer.Reset();
+		//TrianglesBuffer.Reset();
+		//Octree->RangeQuery(Octree->ModifiedBounds, TrianglesBuffer);
+
 		NormalsBuffer.Reset();
-		Octree->RangeQuery(Octree->ModifiedBounds, TrianglesBuffer);
-		for (int TriangleID : TrianglesBuffer)
+		for (int TriangleID : Triangles)
 		{
-			FIndex3i TriV = Mesh->GetTriangle(TriangleID);
-			for (int j = 0; j < 3; ++j)
+			if (Mesh->IsTriangle(TriangleID))
 			{
-				int vid = TriV[j];
-				if (NormalsVertexFlags[vid] == false)
+				FIndex3i TriV = Mesh->GetTriangle(TriangleID);
+				for (int j = 0; j < 3; ++j)
 				{
-					NormalsBuffer.Add(vid);
-					NormalsVertexFlags[vid] = true;
+					int vid = TriV[j];
+					if (NormalsVertexFlags[vid] == false)
+					{
+						NormalsBuffer.Add(vid);
+						NormalsVertexFlags[vid] = true;
+					}
 				}
 			}
 		}
@@ -1815,15 +1869,13 @@ void UDynamicMeshSculptTool::RecalculateNormals_PerVertex()
 	}
 }
 
-void UDynamicMeshSculptTool::RecalculateNormals_Overlay()
+void UDynamicMeshSculptTool::RecalculateNormals_Overlay(const TSet<int32>& Triangles)
 {
 	FDynamicMesh3* Mesh = DynamicMeshComponent->GetMesh();
+	//FDynamicMeshOctree3* Octree = DynamicMeshComponent->GetOctree();
 	FDynamicMeshNormalOverlay* Normals = Mesh->HasAttributes() ? Mesh->Attributes()->PrimaryNormals() : nullptr;
 	check(Normals != nullptr);
 
-
-	FDynamicMeshOctree3* Octree = DynamicMeshComponent->GetOctree();
-
 	int MaxElementID = Normals->MaxElementID();
 	if (NormalsVertexFlags.Num() < MaxElementID)
 	{
@@ -1834,19 +1886,23 @@ void UDynamicMeshSculptTool::RecalculateNormals_Overlay()
 	{
 		SCOPE_CYCLE_COUNTER(SculptTool_Normals_Collect);
 
-		TrianglesBuffer.Reset();
+		//TrianglesBuffer.Reset();
+		//Octree->RangeQuery(Octree->ModifiedBounds, TrianglesBuffer);
+
 		NormalsBuffer.Reset();
-		Octree->RangeQuery(Octree->ModifiedBounds, TrianglesBuffer);
-		for (int TriangleID : TrianglesBuffer)
+		for (int TriangleID : Triangles)
 		{
-			FIndex3i TriElems = Normals->GetTriangle(TriangleID);
-			for (int j = 0; j < 3; ++j)
+			if (Mesh->IsTriangle(TriangleID))
 			{
-				int elemid = TriElems[j];
-				if (NormalsVertexFlags[elemid] == false)
+				FIndex3i TriElems = Normals->GetTriangle(TriangleID);
+				for (int j = 0; j < 3; ++j)
 				{
-					NormalsBuffer.Add(elemid);
-					NormalsVertexFlags[elemid] = true;
+					int elemid = TriElems[j];
+					if (NormalsVertexFlags[elemid] == false)
+					{
+						NormalsBuffer.Add(elemid);
+						NormalsVertexFlags[elemid] = true;
+					}
 				}
 			}
 		}
@@ -1878,15 +1934,29 @@ void UDynamicMeshSculptTool::UpdateTarget()
 		}
 	}
 
-	BrushTargetMesh.Copy(*DynamicMeshComponent->GetMesh(), false, false, false, false);
-	BrushTargetMeshSpatial.SetMesh(&BrushTargetMesh, true);
+	// TODO: could have a second set of target meshes, and swap between them while we update the other one
+	PendingTargetUpdate = Async(DynamicSculptToolAsyncExecTarget, [this]()
+	{
+		BrushTargetMesh.Copy(*DynamicMeshComponent->GetMesh(), false, false, false, false);
+
+		TFuture<void> TargetSpatialUpdate = Async(DynamicSculptToolAsyncExecTarget, [this]() {
+			BrushTargetMeshSpatial.SetMesh(&BrushTargetMesh, true);
+		});
+		TFuture<void> TargetNormalsUpdate = Async(DynamicSculptToolAsyncExecTarget, [this]() {
+			BrushTargetNormals.SetMesh(&BrushTargetMesh);
+			BrushTargetNormals.ComputeVertexNormals();
+		});
+
+		TargetSpatialUpdate.Wait();
+		TargetNormalsUpdate.Wait();
+	});
 
-	BrushTargetNormals.SetMesh(&BrushTargetMesh);
-	BrushTargetNormals.ComputeVertexNormals();
 }
 
 bool UDynamicMeshSculptTool::GetTargetMeshNearest(const FVector3d& Position, double SearchRadius, FVector3d& TargetPosOut, FVector3d& TargetNormalOut)
 {
+	PendingTargetUpdate.Wait();
+
 	double fDistSqr;
 	int NearTID = BrushTargetMeshSpatial.FindNearestTriangle(Position, fDistSqr, SearchRadius);
 	if (NearTID <= 0)
@@ -2170,7 +2240,7 @@ void UDynamicMeshSculptTool::EndChange()
 {
 	if (ActiveVertexChange != nullptr)
 	{
-		GetToolManager()->EmitObjectChange(DynamicMeshComponent, MoveTemp(ActiveVertexChange->Change), LOCTEXT("VertexSculptChange", "Brush Stroke"));
+		GetToolManager()->EmitObjectChange(DynamicMeshComponent, MoveTemp(ActiveVertexChange->Change), LOCTEXT("MeshSculptChange", "Brush Stroke"));
 
 		delete ActiveVertexChange;
 		ActiveVertexChange = nullptr;
@@ -2390,4 +2460,21 @@ void UDynamicMeshSculptTool::UpdateFixedPlaneGizmoVisibility(bool bVisible)
 	}
 }
 
+
+
+void UDynamicMeshSculptTool::DiscardAttributes()
+{
+	TSharedPtr<FDynamicMesh3> BeforeMesh = MakeShared<FDynamicMesh3>(*DynamicMeshComponent->GetMesh());
+	TSharedPtr<FDynamicMesh3> AfterMesh = MakeShared<FDynamicMesh3>(*BeforeMesh);
+	AfterMesh->DiscardAttributes();
+	FMeshNormals::QuickComputeVertexNormals(*AfterMesh);
+
+	TUniquePtr<FMeshReplacementChange> ReplaceChange = MakeUnique<FMeshReplacementChange>(BeforeMesh, AfterMesh);
+	DynamicMeshComponent->ApplyChange(ReplaceChange.Get(), false);
+
+	// other internals are still valid?
+
+	GetToolManager()->EmitObjectChange(DynamicMeshComponent, MoveTemp(ReplaceChange), LOCTEXT("SculptDiscardAttribsChange", "Discard Attributes"));
+}
+
 #undef LOCTEXT_NAMESPACE
diff --git a/Engine/Plugins/Experimental/MeshModelingToolset/Source/MeshModelingTools/Public/DynamicMeshSculptTool.h b/Engine/Plugins/Experimental/MeshModelingToolset/Source/MeshModelingTools/Public/DynamicMeshSculptTool.h
--- a/Engine/Plugins/Experimental/MeshModelingToolset/Source/MeshModelingTools/Public/DynamicMeshSculptTool.h
+++ b/Engine/Plugins/Experimental/MeshModelingToolset/Source/MeshModelingTools/Public/DynamicMeshSculptTool.h
@@ -19,6 +19,7 @@
 #include "Properties/RemeshProperties.h"
 #include "TransformTypes.h"
 #include "Sculpting/MeshSculptToolBase.h"
+#include "Async/Async.h"
 #include "DynamicMeshSculptTool.generated.h"
 
 class UTransformGizmo;
@@ -163,6 +164,21 @@ public:
 };
 
 
+UCLASS()
+class MESHMODELINGTOOLS_API UDynamicSculptToolActions : public UInteractiveToolPropertySet
+{
+	GENERATED_BODY()
+public:
+	TWeakObjectPtr<UDynamicMeshSculptTool> ParentTool;
+
+	void Initialize(UDynamicMeshSculptTool* ParentToolIn) { ParentTool = ParentToolIn; }
+
+	UFUNCTION(CallInEditor, Category = MeshEdits)
+	void DiscardAttributes();
+};
+
+
+
 
 UCLASS()
 class MESHMODELINGTOOLS_API UBrushRemeshProperties : public URemeshProperties
@@ -248,6 +264,7 @@ public:
 	virtual void SetEnableRemeshing(bool bEnable) { bEnableRemeshing = bEnable; }
 	virtual bool GetEnableRemeshing() const { return bEnableRemeshing; }
 
+	virtual void DiscardAttributes();
 
 	virtual void OnPropertyModified(UObject* PropertySet, FProperty* Property) override;
 
@@ -276,6 +293,9 @@ public:
 	UPROPERTY()
 	UMeshEditingViewProperties* ViewProperties;
 
+	UPROPERTY()
+	UDynamicSculptToolActions* SculptToolActions;
+
 public:
 	virtual void IncreaseBrushRadiusAction();
 	virtual void DecreaseBrushRadiusAction();
@@ -359,7 +379,9 @@ private:
 
 	bool bRemeshPending;
 	bool bNormalUpdatePending;
+	
 	bool bTargetDirty;
+	TFuture<void> PendingTargetUpdate;
 
 	bool bSmoothing;
 	bool bInvert;
@@ -417,12 +439,14 @@ private:
 	TArray<int> TrianglesBuffer;
 	TArray<int> NormalsBuffer;
 	TArray<bool> NormalsVertexFlags;
-	void RecalculateNormals_PerVertex();
-	void RecalculateNormals_Overlay();
+	void RecalculateNormals_PerVertex(const TSet<int32>& Triangles);
+	void RecalculateNormals_Overlay(const TSet<int32>& Triangles);
 
 	bool bHaveMeshBoundaries;
 	bool bHaveUVSeams;
 	bool bHaveNormalSeams;
+	TSet<int32> RemeshRemovedTriangles;
+	TSet<int32> RemeshFinalTriangleROI;
 	void PrecomputeRemeshInfo();
 	void RemeshROIPass();
 
diff --git a/Engine/Plugins/Experimental/MeshModelingToolset/Source/ModelingComponents/Private/OctreeDynamicMeshComponent.cpp b/Engine/Plugins/Experimental/MeshModelingToolset/Source/ModelingComponents/Private/OctreeDynamicMeshComponent.cpp
--- a/Engine/Plugins/Experimental/MeshModelingToolset/Source/ModelingComponents/Private/OctreeDynamicMeshComponent.cpp
+++ b/Engine/Plugins/Experimental/MeshModelingToolset/Source/ModelingComponents/Private/OctreeDynamicMeshComponent.cpp
@@ -432,9 +432,11 @@ void UOctreeDynamicMeshComponent::ApplyChange(const FMeshReplacementChange* Chan
 	Mesh->Copy(*Change->GetMesh(bRevert));
 	Octree = MakeUnique<FDynamicMeshOctree3>();
 	Octree->Initialize(GetMesh());
+	
 	OctreeCut = MakeUnique<FDynamicMeshOctree3::FTreeCutSet>();
 
-	//NotifyMeshUpdated();
+	// need to force proxy recreation because we changed the Octree pointer...
+	MarkRenderStateDirty();
 
 	OnMeshChanged.Broadcast();
 }
