Geometry Tools: Simplification Tool -  Respect UV mirroring seams.  This is done by detecting attribute seams associated all "normal" attributes, these include tangent or bitangent attributes if they exist.

#rnx

diff --git a/Engine/Plugins/Experimental/GeometryProcessing/Source/DynamicMesh/Private/MeshConstraintsUtil.cpp b/Engine/Plugins/Experimental/GeometryProcessing/Source/DynamicMesh/Private/MeshConstraintsUtil.cpp
--- a/Engine/Plugins/Experimental/GeometryProcessing/Source/DynamicMesh/Private/MeshConstraintsUtil.cpp
+++ b/Engine/Plugins/Experimental/GeometryProcessing/Source/DynamicMesh/Private/MeshConstraintsUtil.cpp
@@ -137,13 +137,18 @@ FMeshConstraintsUtil::ConstrainAllBoundariesAndSeams(FMeshConstraints& Constrain
 					};
 
 					bool bHasSeamEnd = false;
-					for (int i = 0; i < Attributes->NumUVLayers(); ++i)
+					for (int32 i = 0; !bHasSeamEnd && (i < Attributes->NumUVLayers()); ++i)
 					{
 						bool bIsEnd = IsSeamWithEnd(*Attributes->GetUVLayer(i));
 
 						bHasSeamEnd = bHasSeamEnd || bIsEnd;
 					}
-					bHasSeamEnd = bHasSeamEnd || IsSeamWithEnd(*Attributes->PrimaryNormals());
+					for (int32 i = 0; !bHasSeamEnd && (i < Attributes->NumNormalLayers()); ++i)
+					{
+						bool bIsEnd = IsSeamWithEnd(*Attributes->GetNormalLayer(i));
+
+						bHasSeamEnd = bHasSeamEnd || bIsEnd;
+					}
 
 					if (bHasSeamEnd)
 					{
diff --git a/Engine/Plugins/Experimental/MeshModelingToolset/Source/MeshModelingToolsEditorOnly/Private/SimplifyMeshTool.cpp b/Engine/Plugins/Experimental/MeshModelingToolset/Source/MeshModelingToolsEditorOnly/Private/SimplifyMeshTool.cpp
--- a/Engine/Plugins/Experimental/MeshModelingToolset/Source/MeshModelingToolsEditorOnly/Private/SimplifyMeshTool.cpp
+++ b/Engine/Plugins/Experimental/MeshModelingToolset/Source/MeshModelingToolsEditorOnly/Private/SimplifyMeshTool.cpp
@@ -12,6 +12,7 @@
 
 #include "MeshDescriptionToDynamicMesh.h"
 #include "DynamicMeshToMeshDescription.h"
+#include "AssetUtils/MeshDescriptionUtil.h"
 
 #include "AssetGenerationUtil.h"
 
@@ -118,10 +119,21 @@ void USimplifyMeshTool::Setup()
 		auto EnterProgressFrame = [](int Progress) {};
 #endif
 		OriginalMeshDescription = MakeShared<FMeshDescription>(*ComponentTarget->GetMesh());
+		// make sure the mesh description has tangents. 
+		{
+			// if the source was imported with auto-generate tangents, the mesh description tangents won't have been populated
+			FMeshBuildSettings BuildSettings; 
+			BuildSettings.bRecomputeNormals = true;
+			BuildSettings.bRecomputeTangents = true;
+			// this will only generate tangents if the mesh description was missing them
+			UE::MeshDescription::InitializeAutoGeneratedAttributes(*OriginalMeshDescription, &BuildSettings);
+		}
 		EnterProgressFrame(1);
 		OriginalMesh = MakeShared<FDynamicMesh3>();
-		FMeshDescriptionToDynamicMesh Converter;
-		Converter.Convert(ComponentTarget->GetMesh(), *OriginalMesh);
+		FMeshDescriptionToDynamicMesh Converter; 
+		// convert with tangent overlay
+		const bool bUseTangents = true;
+		Converter.Convert(ComponentTarget->GetMesh(), *OriginalMesh, bUseTangents);
 		EnterProgressFrame(2);
 		OriginalMeshSpatial = MakeShared<FDynamicMeshAABBTree3>(OriginalMesh.Get(), true);
 	}
