ModelingComponents: use full-precision UVs in Simple and Octree DynamicMeshComponents, so that UV previews are always at maximum precision.
#rnx

diff --git a/Engine/Plugins/Experimental/MeshModelingToolset/Source/ModelingComponents/Private/BaseDynamicMeshSceneProxy.h b/Engine/Plugins/Experimental/MeshModelingToolset/Source/ModelingComponents/Private/BaseDynamicMeshSceneProxy.h
--- a/Engine/Plugins/Experimental/MeshModelingToolset/Source/ModelingComponents/Private/BaseDynamicMeshSceneProxy.h
+++ b/Engine/Plugins/Experimental/MeshModelingToolset/Source/ModelingComponents/Private/BaseDynamicMeshSceneProxy.h
@@ -62,7 +62,8 @@ public:
 	FMeshRenderBufferSet(ERHIFeatureLevel::Type FeatureLevelType)
 		: VertexFactory(FeatureLevelType, "FMeshRenderBufferSet")
 	{
-
+		StaticMeshVertexBuffer.SetUseFullPrecisionUVs(true);
+		StaticMeshVertexBuffer.SetUseHighPrecisionTangentBasis(true);
 	}
 
 
