ModelingTools: properly initialize the recalculate normals flag in DisplaceMeshTool
#rnx

diff --git a/Engine/Plugins/Experimental/MeshModelingToolset/Source/MeshModelingTools/Private/DisplaceMeshTool.cpp b/Engine/Plugins/Experimental/MeshModelingToolset/Source/MeshModelingTools/Private/DisplaceMeshTool.cpp
--- a/Engine/Plugins/Experimental/MeshModelingToolset/Source/MeshModelingTools/Private/DisplaceMeshTool.cpp
+++ b/Engine/Plugins/Experimental/MeshModelingToolset/Source/MeshModelingTools/Private/DisplaceMeshTool.cpp
@@ -826,6 +826,7 @@ void UDisplaceMeshTool::Setup()
 	Parameters.DisplaceIntensity = CommonProperties->DisplaceIntensity;
 	Parameters.RandomSeed = CommonProperties->RandomSeed;
 	Parameters.DisplacementMap = TextureMapProperties->DisplacementMap;
+	Parameters.bRecalculateNormals = TextureMapProperties->bRecalcNormals;
 	Parameters.SineWaveFrequency = SineWaveProperties->SineWaveFrequency;
 	Parameters.SineWavePhaseShift = SineWaveProperties->SineWavePhaseShift;
 	Parameters.SineWaveDirection = SineWaveProperties->SineWaveDirection.GetSafeNormal();
