ModelingComponents: mark raw UObject pointers as UProperties in UBaseDynamicMeshComponent
#rnx

diff --git a/Engine/Plugins/Experimental/MeshModelingToolset/Source/ModelingComponents/Public/BaseDynamicMeshComponent.h b/Engine/Plugins/Experimental/MeshModelingToolset/Source/ModelingComponents/Public/BaseDynamicMeshComponent.h
--- a/Engine/Plugins/Experimental/MeshModelingToolset/Source/ModelingComponents/Public/BaseDynamicMeshComponent.h
+++ b/Engine/Plugins/Experimental/MeshModelingToolset/Source/ModelingComponents/Public/BaseDynamicMeshComponent.h
@@ -250,9 +250,13 @@ public:
 
 protected:
 
+	UPROPERTY()
 	TArray<UMaterialInterface*> BaseMaterials;
 
+	UPROPERTY()
 	UMaterialInterface* OverrideRenderMaterial = nullptr;
+
+	UPROPERTY()
 	UMaterialInterface* SecondaryRenderMaterial = nullptr;
 
 	bool bDrawSecondaryBuffers = true;
