ModelingComponents: Add support for fast update of index buffers in decomposed buffer sets in DynamicMeshComponent/Proxy. Add support for MeshRenderDecomposition to UPreviewMesh. Add UPreviewMesh::NotifyRegionDeferredEditCompleted() functions to allow triangle subsets to be updated.
MeshSelectionTool: enable new decomposition support to improve painting performance on large meshes. Fix code that was doing repeated linear searches to remove elements from a potentially-very-large array. Make brush-ROI tri centroid point rendering optional.
#rnx

diff --git a/Engine/Plugins/Experimental/MeshModelingToolset/Source/MeshModelingTools/Private/MeshSelectionTool.cpp b/Engine/Plugins/Experimental/MeshModelingToolset/Source/MeshModelingTools/Private/MeshSelectionTool.cpp
--- a/Engine/Plugins/Experimental/MeshModelingToolset/Source/MeshModelingTools/Private/MeshSelectionTool.cpp
+++ b/Engine/Plugins/Experimental/MeshModelingToolset/Source/MeshModelingTools/Private/MeshSelectionTool.cpp
@@ -106,6 +106,9 @@ void UMeshSelectionTool::Setup()
 		return SelectedTriangles[TriangleID] ? true : false;
 	});
 
+	// enable auto-chunking of mesh into separate render buffers, so that partial updates can be done during selection paint
+	PreviewMesh->SetEnableRenderMeshDecomposition(true);
+
 	const FDynamicMesh3* Mesh = PreviewMesh->GetPreviewDynamicMesh();
 	SelectedVertices = TBitArray<>(false, Mesh->MaxVertexID());
 	SelectedTriangles = TBitArray<>(false, Mesh->MaxTriangleID());
@@ -257,11 +260,22 @@ void UMeshSelectionTool::RegisterActions(FInteractiveToolActionSet& ActionSet)
 
 void UMeshSelectionTool::OnExternalSelectionChange()
 {
+	int32 NumTriangles = SelectedTriangles.Num();
+	TSet<int32> AllModifiedTriangles;
+	for (int32 k = 0; k < NumTriangles; ++k)
+	{
+		if (SelectedTriangles[k])
+		{
+			AllModifiedTriangles.Add(k);
+		}
+	}
+
 	SelectedVertices.SetRange(0, SelectedVertices.Num(), false);
 	SelectedTriangles.SetRange(0, SelectedTriangles.Num(), false);
 
 	if (SelectionType == EMeshSelectionElementType::Vertex)
 	{
+		ensure(false);		// not supported
 		for (int VertIdx : Selection->Vertices)
 		{
 			SelectedVertices[VertIdx] = true;
@@ -272,10 +286,11 @@ void UMeshSelectionTool::OnExternalSelectionChange()
 		for (int FaceIdx : Selection->Faces)
 		{
 			SelectedTriangles[FaceIdx] = true;
+			AllModifiedTriangles.Add(FaceIdx);
 		}
+		// todo: more efficient for rendering update to compute the delta set here...
+		OnRegionHighlightUpdated(AllModifiedTriangles);
 	}
-
-	OnSelectionUpdated();
 }
 
 
@@ -451,9 +466,9 @@ void UMeshSelectionTool::ApplyStamp(const FBrushStampData& Stamp)
 				}
 			}
 		}
-	}
 
-	OnSelectionUpdated();
+		ensure(false);		// not supported yet
+	}
 }
 
 
@@ -536,20 +551,53 @@ void UMeshSelectionTool::UpdateFaceSelection(const FBrushStampData& Stamp, const
 		UseROI = &LocalROI;
 	}
 
-	bool bDesiredValue = bInRemoveStroke ? false : true;
-	for (int TriIdx : *UseROI)
+	// separate paths for add and remove because we can't efficiently remove from Selection->Faces TArray
+	if (bInRemoveStroke == false)
 	{
-		if (SelectedTriangles[TriIdx] != bDesiredValue)
+		for (int TriIdx : *UseROI)
+		{
+			if (SelectedTriangles[TriIdx] == false)
+			{
+				SelectedTriangles[TriIdx] = true;
+				Selection->Faces.Add(TriIdx);
+				if (ActiveSelectionChange != nullptr)
+				{
+					ActiveSelectionChange->Add(TriIdx);
+				}
+			}
+		}
+	}
+	else
+	{
+		bool bModified = false;
+		TArray<int32> CurSelection = Selection->Faces;
+		for (int TriIdx : *UseROI)
+		{
+			if (SelectedTriangles[TriIdx] == true)
+			{
+				SelectedTriangles[TriIdx] = false;
+				bModified = true;
+				if (ActiveSelectionChange != nullptr)
+				{
+					ActiveSelectionChange->Add(TriIdx);
+				}
+			}
+		}
+		// rebuild selection
+		if (bModified)
 		{
-			SelectedTriangles[TriIdx] = bDesiredValue;
-			UpdateList(Selection->Faces, TriIdx, bDesiredValue);
-			if (ActiveSelectionChange != nullptr)
+			Selection->Faces.Reset();
+			for (int32 tid : CurSelection)
 			{
-				ActiveSelectionChange->Add(TriIdx);
+				if (SelectedTriangles[tid])
+				{
+					Selection->Faces.Add(tid);
+				}
 			}
 		}
 	}
 
+	OnRegionHighlightUpdated(*UseROI);
 }
 
 
@@ -593,13 +641,22 @@ bool UMeshSelectionTool::OnUpdateHover(const FInputDeviceRay& DevicePos)
 
 
 
+void UMeshSelectionTool::OnRegionHighlightUpdated()
+{
+	PreviewMesh->NotifyDeferredEditCompleted(UPreviewMesh::ERenderUpdateMode::FastUpdate, EMeshRenderAttributeFlags::SecondaryIndexBuffers, false);
 
-
-void UMeshSelectionTool::OnSelectionUpdated()
+}
+void UMeshSelectionTool::OnRegionHighlightUpdated(const TArray<int32>& Triangles)
 {
-	UpdateVisualization(true);
+	PreviewMesh->NotifyRegionDeferredEditCompleted(Triangles, EMeshRenderAttributeFlags::SecondaryIndexBuffers);
+}
+void UMeshSelectionTool::OnRegionHighlightUpdated(const TSet<int32>& Triangles)
+{
+	PreviewMesh->NotifyRegionDeferredEditCompleted(Triangles, EMeshRenderAttributeFlags::SecondaryIndexBuffers);
 }
 
+
+
 void UMeshSelectionTool::UpdateVisualization(bool bSelectionModified)
 {
 	check(SelectionType == EMeshSelectionElementType::Face);  // only face selection supported so far
@@ -609,7 +666,7 @@ void UMeshSelectionTool::UpdateVisualization(bool bSelectionModified)
 	// force an update of renderbuffers
 	if (bSelectionModified)
 	{
-		PreviewMesh->NotifyDeferredEditCompleted(UPreviewMesh::ERenderUpdateMode::FullUpdate, EMeshRenderAttributeFlags::All, true);
+		PreviewMesh->NotifyDeferredEditCompleted(UPreviewMesh::ERenderUpdateMode::FullUpdate, EMeshRenderAttributeFlags::AllVertexAttribs, true);
 	}
 
 	if (bColorsUpdatePending)
@@ -686,21 +743,21 @@ void UMeshSelectionTool::Render(IToolsContextRenderAPI* RenderAPI)
 	FTransform WorldTransform = ComponentTarget->GetWorldTransform();
 	const FDynamicMesh3* Mesh = PreviewMesh->GetMesh();
 
-	float PDIScale = RenderAPI->GetCameraState().GetPDIScalingFactor();
-	if (SelectionType == EMeshSelectionElementType::Vertex)
-	{
-		MeshDebugDraw::DrawVertices(Mesh, Selection->Vertices,
-			12.0f*PDIScale, FColor::Orange, RenderAPI->GetPrimitiveDrawInterface(), WorldTransform);
-		MeshDebugDraw::DrawVertices(Mesh, PreviewBrushROI,
-			8.0f*PDIScale, FColor(40, 200, 40), RenderAPI->GetPrimitiveDrawInterface(), WorldTransform);
-	}
-	else
+	if (SelectionProps->bShowPoints)
 	{
-		// drawn via material
-		//MeshDebugDraw::DrawTriCentroids(Mesh, Selection->Faces,
-		//	12.0f*PDIScale, FColor::Green, RenderAPI->GetPrimitiveDrawInterface(), WorldTransform);
-		MeshDebugDraw::DrawTriCentroids(Mesh, PreviewBrushROI,
-			4.0f*PDIScale, FColor(40, 200, 40), RenderAPI->GetPrimitiveDrawInterface(), WorldTransform);
+		float PDIScale = RenderAPI->GetCameraState().GetPDIScalingFactor();
+		if (SelectionType == EMeshSelectionElementType::Vertex)
+		{
+			MeshDebugDraw::DrawVertices(Mesh, Selection->Vertices,
+				12.0f*PDIScale, FColor::Orange, RenderAPI->GetPrimitiveDrawInterface(), WorldTransform);
+			MeshDebugDraw::DrawVertices(Mesh, PreviewBrushROI,
+				8.0f*PDIScale, FColor(40, 200, 40), RenderAPI->GetPrimitiveDrawInterface(), WorldTransform);
+		}
+		else
+		{
+			MeshDebugDraw::DrawTriCentroids(Mesh, PreviewBrushROI,
+				4.0f*PDIScale, FColor(40, 200, 40), RenderAPI->GetPrimitiveDrawInterface(), WorldTransform);
+		}
 	}
 }
 
diff --git a/Engine/Plugins/Experimental/MeshModelingToolset/Source/MeshModelingTools/Public/MeshSelectionTool.h b/Engine/Plugins/Experimental/MeshModelingToolset/Source/MeshModelingTools/Public/MeshSelectionTool.h
--- a/Engine/Plugins/Experimental/MeshModelingToolset/Source/MeshModelingTools/Public/MeshSelectionTool.h
+++ b/Engine/Plugins/Experimental/MeshModelingToolset/Source/MeshModelingTools/Public/MeshSelectionTool.h
@@ -249,6 +249,10 @@ public:
 	UPROPERTY(EditAnywhere, Category = Selection)
 	bool bShowWireframe = true;
 
+	/** Toggle drawing of highlight points on/off */
+	UPROPERTY(EditAnywhere, Category = Selection)
+	bool bShowPoints = true;
+
 	UPROPERTY(EditAnywhere, Category = Selection)
 	EMeshFacesColorMode FaceColorMode = EMeshFacesColorMode::None;
 };
@@ -347,7 +351,9 @@ protected:
 	TBitArray<> SelectedTriangles;
 	void OnExternalSelectionChange();
 
-	void OnSelectionUpdated();
+	void OnRegionHighlightUpdated(const TArray<int32>& Triangles);
+	void OnRegionHighlightUpdated(const TSet<int32>& Triangles);
+	void OnRegionHighlightUpdated();
 	void UpdateVisualization(bool bSelectionModified);
 	bool bFullMeshInvalidationPending = false;
 	bool bColorsUpdatePending = false;
diff --git a/Engine/Plugins/Experimental/MeshModelingToolset/Source/ModelingComponents/Private/BaseDynamicMeshSceneProxy.h b/Engine/Plugins/Experimental/MeshModelingToolset/Source/ModelingComponents/Private/BaseDynamicMeshSceneProxy.h
--- a/Engine/Plugins/Experimental/MeshModelingToolset/Source/ModelingComponents/Private/BaseDynamicMeshSceneProxy.h
+++ b/Engine/Plugins/Experimental/MeshModelingToolset/Source/ModelingComponents/Private/BaseDynamicMeshSceneProxy.h
@@ -536,11 +536,11 @@ public:
 
 
 	/**
-	 * FastUpdateTriangleBuffers re-sorts the existing set of triangles in a FMeshRenderBufferSet
+	 * RecomputeRenderBufferTriangleIndexSets re-sorts the existing set of triangles in a FMeshRenderBufferSet
 	 * into primary and secondary index buffers. Note that UploadIndexBufferUpdate() must be called
 	 * after this function!
 	 */
-	void FastUpdateIndexBuffers(
+	void RecomputeRenderBufferTriangleIndexSets(
 		FMeshRenderBufferSet* RenderBuffers, 
 		const FDynamicMesh3* Mesh)
 	{
diff --git a/Engine/Plugins/Experimental/MeshModelingToolset/Source/ModelingComponents/Private/Drawing/MeshRenderDecomposition.cpp b/Engine/Plugins/Experimental/MeshModelingToolset/Source/ModelingComponents/Private/Drawing/MeshRenderDecomposition.cpp
--- a/Engine/Plugins/Experimental/MeshModelingToolset/Source/ModelingComponents/Private/Drawing/MeshRenderDecomposition.cpp
+++ b/Engine/Plugins/Experimental/MeshModelingToolset/Source/ModelingComponents/Private/Drawing/MeshRenderDecomposition.cpp
@@ -25,10 +25,10 @@ void FMeshRenderDecomposition::BuildAssociations(const FDynamicMesh3* Mesh)
 
 
 
-void FMeshRenderDecomposition::BuildMaterialDecomposition(FDynamicMesh3* Mesh, const FComponentMaterialSet* MaterialSet, FMeshRenderDecomposition& Decomp)
+void FMeshRenderDecomposition::BuildMaterialDecomposition(const FDynamicMesh3* Mesh, const FComponentMaterialSet* MaterialSet, FMeshRenderDecomposition& Decomp)
 {
 	check(MaterialSet);
-	FDynamicMeshMaterialAttribute* MaterialID = Mesh->Attributes()->GetMaterialID();
+	const FDynamicMeshMaterialAttribute* MaterialID = Mesh->Attributes()->GetMaterialID();
 	check(MaterialID);
 
 	int32 NumMaterials = MaterialSet->Materials.Num();
@@ -57,7 +57,7 @@ void FMeshRenderDecomposition::BuildMaterialDecomposition(FDynamicMesh3* Mesh, c
 
 
 static void CollectSubDecomposition(
-	FDynamicMesh3* Mesh, 
+	const FDynamicMesh3* Mesh, 
 	const TArray<int32>& Triangles, 
 	UMaterialInterface* Material, 
 	FMeshRenderDecomposition& Decomp, 
@@ -131,7 +131,7 @@ static void CollectSubDecomposition(
 
 
 
-void FMeshRenderDecomposition::BuildChunkedDecomposition(FDynamicMesh3* Mesh, const FComponentMaterialSet* MaterialSet, FMeshRenderDecomposition& Decomp, int32 MaxChunkSize)
+void FMeshRenderDecomposition::BuildChunkedDecomposition(const FDynamicMesh3* Mesh, const FComponentMaterialSet* MaterialSet, FMeshRenderDecomposition& Decomp, int32 MaxChunkSize)
 {
 	TUniquePtr<FMeshRenderDecomposition> MaterialDecomp = MakeUnique<FMeshRenderDecomposition>();
 	BuildMaterialDecomposition(Mesh, MaterialSet, *MaterialDecomp);
diff --git a/Engine/Plugins/Experimental/MeshModelingToolset/Source/ModelingComponents/Private/PreviewMesh.cpp b/Engine/Plugins/Experimental/MeshModelingToolset/Source/ModelingComponents/Private/PreviewMesh.cpp
--- a/Engine/Plugins/Experimental/MeshModelingToolset/Source/ModelingComponents/Private/PreviewMesh.cpp
+++ b/Engine/Plugins/Experimental/MeshModelingToolset/Source/ModelingComponents/Private/PreviewMesh.cpp
@@ -4,6 +4,7 @@
 
 #include "Containers/StaticArray.h"
 #include "Engine/Classes/Materials/Material.h"
+#include "TargetInterfaces/PrimitiveComponentBackedTarget.h"
 
 
 APreviewMeshActor::APreviewMeshActor()
@@ -73,6 +74,12 @@ void UPreviewMesh::SetMaterial(int MaterialIndex, UMaterialInterface* Material)
 
 	// force rebuild because we can't change materials yet - surprisingly complicated
 	DynamicMeshComponent->NotifyMeshUpdated();
+
+	// if we change materials we have to force a decomposition update because it decomposes by material
+	if (bDecompositionEnabled)
+	{
+		UpdateRenderMeshDecomposition();
+	}
 }
 
 void UPreviewMesh::SetMaterials(const TArray<UMaterialInterface*>& Materials)
@@ -85,6 +92,12 @@ void UPreviewMesh::SetMaterials(const TArray<UMaterialInterface*>& Materials)
 
 	// force rebuild because we can't change materials yet - surprisingly complicated
 	DynamicMeshComponent->NotifyMeshUpdated();
+
+	// if we change materials we have to force a decomposition update because it decomposes by material
+	if (bDecompositionEnabled)
+	{
+		UpdateRenderMeshDecomposition();
+	}
 }
 
 int32 UPreviewMesh::GetNumMaterials() const
@@ -353,6 +366,10 @@ void UPreviewMesh::InitializeMesh(FMeshDescription* MeshDescription)
 	{
 		MeshAABBTree.SetMesh(DynamicMeshComponent->GetMesh(), true);
 	}
+	if (bDecompositionEnabled)
+	{
+		UpdateRenderMeshDecomposition();
+	}
 }
 
 
@@ -367,6 +384,10 @@ void UPreviewMesh::ReplaceMesh(FDynamicMesh3&& NewMesh)
 	{
 		MeshAABBTree.SetMesh(DynamicMeshComponent->GetMesh(), true);
 	}
+	if (bDecompositionEnabled)
+	{
+		UpdateRenderMeshDecomposition();
+	}
 }
 
 
@@ -381,6 +402,10 @@ void UPreviewMesh::EditMesh(TFunctionRef<void(FDynamicMesh3&)> EditFunc)
 	{
 		MeshAABBTree.SetMesh(DynamicMeshComponent->GetMesh(), true);
 	}
+	if (bDecompositionEnabled)
+	{
+		UpdateRenderMeshDecomposition();
+	}
 }
 
 
@@ -415,6 +440,10 @@ void UPreviewMesh::NotifyDeferredEditCompleted(ERenderUpdateMode UpdateMode, EMe
 	if (UpdateMode == ERenderUpdateMode::FullUpdate)
 	{
 		DynamicMeshComponent->NotifyMeshUpdated();
+		if (bDecompositionEnabled)
+		{
+			UpdateRenderMeshDecomposition();
+		}
 	}
 	else if (UpdateMode == ERenderUpdateMode::FastUpdate)
 	{
@@ -448,6 +477,10 @@ TUniquePtr<FMeshChange> UPreviewMesh::TrackedEditMesh(TFunctionRef<void(FDynamic
 	{
 		MeshAABBTree.SetMesh(DynamicMeshComponent->GetMesh(), true);
 	}
+	if (bDecompositionEnabled)
+	{
+		UpdateRenderMeshDecomposition();
+	}
 
 	return MoveTemp(Change);
 }
@@ -461,6 +494,7 @@ void UPreviewMesh::ApplyChange(const FMeshVertexChange* Change, bool bRevert)
 	{
 		MeshAABBTree.SetMesh(DynamicMeshComponent->GetMesh(), true);
 	}
+	// should not need to update render mesh decomposition here...
 }
 void UPreviewMesh::ApplyChange(const FMeshChange* Change, bool bRevert)
 {
@@ -470,6 +504,10 @@ void UPreviewMesh::ApplyChange(const FMeshChange* Change, bool bRevert)
 	{
 		MeshAABBTree.SetMesh(DynamicMeshComponent->GetMesh(), true);
 	}
+	if (bDecompositionEnabled)
+	{
+		UpdateRenderMeshDecomposition();
+	}
 }
 void UPreviewMesh::ApplyChange(const FMeshReplacementChange* Change, bool bRevert)
 {
@@ -479,6 +517,10 @@ void UPreviewMesh::ApplyChange(const FMeshReplacementChange* Change, bool bRever
 	{
 		MeshAABBTree.SetMesh(DynamicMeshComponent->GetMesh(), true);
 	}
+	if (bDecompositionEnabled)
+	{
+		UpdateRenderMeshDecomposition();
+	}
 }
 
 FSimpleMulticastDelegate& UPreviewMesh::GetOnMeshChanged()
@@ -521,4 +563,48 @@ void UPreviewMesh::ClearTriangleColorFunction(ERenderUpdateMode UpdateMode)
 			DynamicMeshComponent->NotifyMeshUpdated();
 		}
 	}
+}
+
+
+void UPreviewMesh::SetEnableRenderMeshDecomposition(bool bEnable)
+{
+	if (bDecompositionEnabled != bEnable)
+	{
+		bDecompositionEnabled = bEnable;
+
+		if (bDecompositionEnabled)
+		{
+			UpdateRenderMeshDecomposition();
+		}
+		else
+		{
+			DynamicMeshComponent->SetExternalDecomposition(nullptr);
+		}
+	}
+}
+
+
+void UPreviewMesh::UpdateRenderMeshDecomposition()
+{
+	check(bDecompositionEnabled);
+
+	const FDynamicMesh3* Mesh = DynamicMeshComponent->GetMesh();
+	FComponentMaterialSet MaterialSet;
+	GetMaterials(MaterialSet.Materials);
+
+	TUniquePtr<FMeshRenderDecomposition> Decomp = MakeUnique<FMeshRenderDecomposition>();
+	FMeshRenderDecomposition::BuildChunkedDecomposition(Mesh, &MaterialSet, *Decomp);
+	Decomp->BuildAssociations(Mesh);
+	DynamicMeshComponent->SetExternalDecomposition(MoveTemp(Decomp));
+}
+
+
+void UPreviewMesh::NotifyRegionDeferredEditCompleted(const TArray<int32>& Triangles, EMeshRenderAttributeFlags ModifiedAttribs)
+{
+	DynamicMeshComponent->FastNotifyTriangleVerticesUpdated(Triangles, ModifiedAttribs);
+}
+
+void UPreviewMesh::NotifyRegionDeferredEditCompleted(const TSet<int32>& Triangles, EMeshRenderAttributeFlags ModifiedAttribs)
+{
+	DynamicMeshComponent->FastNotifyTriangleVerticesUpdated(Triangles, ModifiedAttribs);
 }
\ No newline at end of file
diff --git a/Engine/Plugins/Experimental/MeshModelingToolset/Source/ModelingComponents/Private/SimpleDynamicMeshComponent.cpp b/Engine/Plugins/Experimental/MeshModelingToolset/Source/ModelingComponents/Private/SimpleDynamicMeshComponent.cpp
--- a/Engine/Plugins/Experimental/MeshModelingToolset/Source/ModelingComponents/Private/SimpleDynamicMeshComponent.cpp
+++ b/Engine/Plugins/Experimental/MeshModelingToolset/Source/ModelingComponents/Private/SimpleDynamicMeshComponent.cpp
@@ -420,6 +420,9 @@ void USimpleDynamicMeshComponent::FastNotifyTriangleVerticesUpdated(const TArray
 		return;
 	}
 
+	bool bUpdateSecondarySort = (SecondaryTriFilterFunc) &&
+		((UpdatedAttributes & EMeshRenderAttributeFlags::SecondaryIndexBuffers) != EMeshRenderAttributeFlags::None);
+
 	FSimpleDynamicMeshSceneProxy* Proxy = GetCurrentSceneProxy();
 	if (!Proxy)
 	{
@@ -428,6 +431,10 @@ void USimpleDynamicMeshComponent::FastNotifyTriangleVerticesUpdated(const TArray
 	else if ( ! Decomposition )
 	{
 		FastNotifyVertexAttributesUpdated(UpdatedAttributes);
+		if (bUpdateSecondarySort)
+		{
+			Proxy->FastUpdateAllIndexBuffers();
+		}
 	}
 	else
 	{
@@ -439,11 +446,16 @@ void USimpleDynamicMeshComponent::FastNotifyTriangleVerticesUpdated(const TArray
 		}
 
 		bool bPositions = (UpdatedAttributes & EMeshRenderAttributeFlags::Positions) != EMeshRenderAttributeFlags::None;
-		GetCurrentSceneProxy()->FastUpdateVertices(UpdatedSets, bPositions,
+		Proxy->FastUpdateVertices(UpdatedSets, bPositions,
 			(UpdatedAttributes & EMeshRenderAttributeFlags::VertexNormals) != EMeshRenderAttributeFlags::None,
 			(UpdatedAttributes & EMeshRenderAttributeFlags::VertexColors) != EMeshRenderAttributeFlags::None,
 			(UpdatedAttributes & EMeshRenderAttributeFlags::VertexUVs) != EMeshRenderAttributeFlags::None);
 
+		if (bUpdateSecondarySort)
+		{
+			Proxy->FastUpdateIndexBuffers(UpdatedSets);
+		}
+
 		if (bPositions)
 		{
 			MarkRenderTransformDirty();
@@ -465,8 +477,10 @@ void USimpleDynamicMeshComponent::FastNotifyTriangleVerticesUpdated(const TSet<i
 		return;
 	}
 
-	FSimpleDynamicMeshSceneProxy* Proxy = GetCurrentSceneProxy();
+	bool bUpdateSecondarySort = (SecondaryTriFilterFunc) &&
+		((UpdatedAttributes & EMeshRenderAttributeFlags::SecondaryIndexBuffers) != EMeshRenderAttributeFlags::None);
 
+	FSimpleDynamicMeshSceneProxy* Proxy = GetCurrentSceneProxy();
 	if (!Proxy)
 	{
 		ResetProxy();
@@ -474,6 +488,10 @@ void USimpleDynamicMeshComponent::FastNotifyTriangleVerticesUpdated(const TSet<i
 	else if (!Decomposition)
 	{
 		FastNotifyVertexAttributesUpdated(UpdatedAttributes);
+		if (bUpdateSecondarySort)
+		{
+			Proxy->FastUpdateAllIndexBuffers();
+		}
 	}
 	else
 	{
@@ -509,6 +527,11 @@ void USimpleDynamicMeshComponent::FastNotifyTriangleVerticesUpdated(const TSet<i
 			(UpdatedAttributes & EMeshRenderAttributeFlags::VertexColors) != EMeshRenderAttributeFlags::None,
 			(UpdatedAttributes & EMeshRenderAttributeFlags::VertexUVs) != EMeshRenderAttributeFlags::None);
 
+		if (bUpdateSecondarySort)
+		{
+			Proxy->FastUpdateIndexBuffers(UpdatedSets);
+		}
+
 		if (bPositions)
 		{
 			MarkRenderTransformDirty();
diff --git a/Engine/Plugins/Experimental/MeshModelingToolset/Source/ModelingComponents/Private/SimpleDynamicMeshSceneProxy.h b/Engine/Plugins/Experimental/MeshModelingToolset/Source/ModelingComponents/Private/SimpleDynamicMeshSceneProxy.h
--- a/Engine/Plugins/Experimental/MeshModelingToolset/Source/ModelingComponents/Private/SimpleDynamicMeshSceneProxy.h
+++ b/Engine/Plugins/Experimental/MeshModelingToolset/Source/ModelingComponents/Private/SimpleDynamicMeshSceneProxy.h
@@ -433,6 +433,12 @@ public:
 	 */
 	void FastUpdateVertices(const TArray<int32>& WhichBuffers, bool bPositions, bool bNormals, bool bColors, bool bUVs)
 	{
+		// skip if we have no updates
+		if (bPositions == false && bNormals == false && bColors == false && bUVs == false)
+		{
+			return;
+		}
+
 		// This needs to be rewritten for split-by-material buffers.
 		// Could store triangle set with each buffer, and then rebuild vtx buffer(s) as needed?
 
@@ -506,7 +512,7 @@ public:
 		// have to wait for all outstanding rendering to finish because the index buffers we are about to edit might be in-use
 		FlushRenderingCommands();
 
-		ParallelFor(RenderBufferSets.Num(), [&](int i)
+		ParallelFor(RenderBufferSets.Num(), [this, &Mesh](int i)
 		{
 			FMeshRenderBufferSet* Buffers = RenderBufferSets[i];
 
@@ -514,7 +520,7 @@ public:
 
 			if (Buffers->TriangleCount > 0)
 			{
-				FastUpdateIndexBuffers(Buffers, Mesh);
+				RecomputeRenderBufferTriangleIndexSets(Buffers, Mesh);
 			}
 
 			ENQUEUE_RENDER_COMMAND(FSimpleDynamicMeshSceneProxyFastUpdateAllIndexBuffers)(
@@ -528,6 +534,41 @@ public:
 
 
 
+	/**
+	 * Update index buffers inside each RenderBuffer set
+	 */
+	void FastUpdateIndexBuffers(const TArray<int32>& WhichBuffers)
+	{
+		FDynamicMesh3* Mesh = ParentComponent->GetRenderMesh();
+
+		// have to wait for all outstanding rendering to finish because the index buffers we are about to edit might be in-use
+		FlushRenderingCommands();
+
+		ParallelFor(WhichBuffers.Num(), [this, &WhichBuffers, &Mesh](int i)
+		{
+			int32 BufferIndex = WhichBuffers[i];
+			if (RenderBufferSets.IsValidIndex(BufferIndex) == false)
+			{
+				return;
+			}
+
+			FMeshRenderBufferSet* Buffers = RenderBufferSets[BufferIndex];
+			FScopeLock BuffersLock(&Buffers->BuffersLock);
+			if (Buffers->TriangleCount > 0)
+			{
+				RecomputeRenderBufferTriangleIndexSets(Buffers, Mesh);
+			}
+
+			ENQUEUE_RENDER_COMMAND(FSimpleDynamicMeshSceneProxyFastUpdateSomeIndexBuffers)(
+				[Buffers](FRHICommandListImmediate& RHICmdList)
+			{
+				Buffers->UploadIndexBufferUpdate();
+			});
+		});
+	}
+
+
+
 public:
 
 
diff --git a/Engine/Plugins/Experimental/MeshModelingToolset/Source/ModelingComponents/Public/BaseDynamicMeshComponent.h b/Engine/Plugins/Experimental/MeshModelingToolset/Source/ModelingComponents/Public/BaseDynamicMeshComponent.h
--- a/Engine/Plugins/Experimental/MeshModelingToolset/Source/ModelingComponents/Public/BaseDynamicMeshComponent.h
+++ b/Engine/Plugins/Experimental/MeshModelingToolset/Source/ModelingComponents/Public/BaseDynamicMeshComponent.h
@@ -28,12 +28,13 @@ class FMeshChange;
 enum class EMeshRenderAttributeFlags : uint8
 {
 	None = 0,
-	Positions = 0x1,
-	VertexColors = 0x2,
-	VertexNormals = 0x4,
-	VertexUVs = 0x8,
+	Positions = 1,
+	VertexColors = 1<<1,
+	VertexNormals = 1<<2,
+	VertexUVs = 1<<3,
+	SecondaryIndexBuffers = 1<<4,
 
-	All = 0xFF
+	AllVertexAttribs = Positions | VertexColors | VertexNormals | VertexUVs
 };
 ENUM_CLASS_FLAGS(EMeshRenderAttributeFlags);
 
diff --git a/Engine/Plugins/Experimental/MeshModelingToolset/Source/ModelingComponents/Public/Drawing/MeshRenderDecomposition.h b/Engine/Plugins/Experimental/MeshModelingToolset/Source/ModelingComponents/Public/Drawing/MeshRenderDecomposition.h
--- a/Engine/Plugins/Experimental/MeshModelingToolset/Source/ModelingComponents/Public/Drawing/MeshRenderDecomposition.h
+++ b/Engine/Plugins/Experimental/MeshModelingToolset/Source/ModelingComponents/Public/Drawing/MeshRenderDecomposition.h
@@ -81,11 +81,11 @@ public:
 	/**
 	 * Build decomposition with one group for each MaterialID of mesh
 	 */
-	static void BuildMaterialDecomposition(FDynamicMesh3* Mesh, const FComponentMaterialSet* MaterialSet, FMeshRenderDecomposition& Decomp);
+	static void BuildMaterialDecomposition(const FDynamicMesh3* Mesh, const FComponentMaterialSet* MaterialSet, FMeshRenderDecomposition& Decomp);
 
 	/**
 	 * Build per-material decomposition, and then split each of those into chunks of at most MaxChunkSize
 	 * (actual chunk sizes will be highly variable and some may be very small...)
 	 */
-	static void BuildChunkedDecomposition(FDynamicMesh3* Mesh, const FComponentMaterialSet* MaterialSet, FMeshRenderDecomposition& Decomp, int32 MaxChunkSize = 1 << 14 /* 16k */ );
+	static void BuildChunkedDecomposition(const FDynamicMesh3* Mesh, const FComponentMaterialSet* MaterialSet, FMeshRenderDecomposition& Decomp, int32 MaxChunkSize = 1 << 14 /* 16k */ );
 };
diff --git a/Engine/Plugins/Experimental/MeshModelingToolset/Source/ModelingComponents/Public/MeshOpPreviewHelpers.h b/Engine/Plugins/Experimental/MeshModelingToolset/Source/ModelingComponents/Public/MeshOpPreviewHelpers.h
--- a/Engine/Plugins/Experimental/MeshModelingToolset/Source/ModelingComponents/Public/MeshOpPreviewHelpers.h
+++ b/Engine/Plugins/Experimental/MeshModelingToolset/Source/ModelingComponents/Public/MeshOpPreviewHelpers.h
@@ -133,7 +133,7 @@ public:
 	 *  moves verts without changing their UV's, then one would probably pass "EMeshRenderAttributeFlags::Positions | EMeshRenderAttributeFlags::VertexNormals".
 	 *  Has no effect if bOn was set to false.
 	 */
-	void SetIsMeshTopologyConstant(bool bOn, EMeshRenderAttributeFlags ChangingAttributes = EMeshRenderAttributeFlags::All);
+	void SetIsMeshTopologyConstant(bool bOn, EMeshRenderAttributeFlags ChangingAttributes = EMeshRenderAttributeFlags::AllVertexAttribs);
 
 	/**
 	 * Set time that Preview will wait before showing working material
@@ -185,7 +185,7 @@ protected:
 
 	// Used for partial/fast updates of the preview mesh render proxy.
 	bool bMeshTopologyIsConstant = false;
-	EMeshRenderAttributeFlags ChangingAttributeFlags = EMeshRenderAttributeFlags::All;
+	EMeshRenderAttributeFlags ChangingAttributeFlags = EMeshRenderAttributeFlags::AllVertexAttribs;
 	bool bMeshInitialized = false;
 
 	float SecondsBeforeWorkingMaterial = 2.0;
diff --git a/Engine/Plugins/Experimental/MeshModelingToolset/Source/ModelingComponents/Public/PreviewMesh.h b/Engine/Plugins/Experimental/MeshModelingToolset/Source/ModelingComponents/Public/PreviewMesh.h
--- a/Engine/Plugins/Experimental/MeshModelingToolset/Source/ModelingComponents/Public/PreviewMesh.h
+++ b/Engine/Plugins/Experimental/MeshModelingToolset/Source/ModelingComponents/Public/PreviewMesh.h
@@ -282,7 +282,7 @@ public:
 	 * @param ModifiedAttribs Only relevant in case of FastUpdate- determines which attributes actually changed.
 	 */
 	void UpdatePreview(const FDynamicMesh3* Mesh, ERenderUpdateMode UpdateMode = ERenderUpdateMode::FullUpdate, 
-		EMeshRenderAttributeFlags ModifiedAttribs = EMeshRenderAttributeFlags::All);
+		EMeshRenderAttributeFlags ModifiedAttribs = EMeshRenderAttributeFlags::AllVertexAttribs);
 
 	/**
 	 * Update the internal mesh by moving in the given Mesh
@@ -291,7 +291,7 @@ public:
 	 * @param ModifiedAttribs Only relevant in case of FastUpdate- determines which attributes actually changed.
 	 */
 	void UpdatePreview(FDynamicMesh3&& Mesh, ERenderUpdateMode UpdateMode = ERenderUpdateMode::FullUpdate,
-		EMeshRenderAttributeFlags ModifiedAttribs = EMeshRenderAttributeFlags::All);
+		EMeshRenderAttributeFlags ModifiedAttribs = EMeshRenderAttributeFlags::AllVertexAttribs);
 
 	/**
 	 * Initialize the internal mesh based on the given MeshDescription
@@ -360,6 +360,21 @@ public:
 	 */
 	void NotifyDeferredEditCompleted(ERenderUpdateMode UpdateMode, EMeshRenderAttributeFlags ModifiedAttribs, bool bRebuildSpatial);
 
+	/**
+	 * Notify that a deferred edit is completed and cause update of rendering data structures for modified Triangles.
+	 * This can reduce the cost of mesh updates, but only if SetEnableRenderMeshDecomposition(true) has been called
+	 * @param ModifiedAttribs which mesh attributes have been modified and need to be updated
+	 */
+	void NotifyRegionDeferredEditCompleted(const TArray<int32>& Triangles, EMeshRenderAttributeFlags ModifiedAttribs);
+
+	/**
+	 * Notify that a deferred edit is completed and cause update of rendering data structures for modified Triangles.
+	 * This can reduce the cost of mesh updates, but only if SetEnableRenderMeshDecomposition(true) has been called
+	 * @param ModifiedAttribs which mesh attributes have been modified and need to be updated
+	 */
+	void NotifyRegionDeferredEditCompleted(const TSet<int32>& Triangles, EMeshRenderAttributeFlags ModifiedAttribs);
+
+
 
 	/**
 	 * Apply EditFunc to the internal mesh and update internal data structures as necessary.
@@ -394,6 +409,16 @@ public:
 	void ForceRebuildSpatial();
 
 
+	/**
+	 * Enable automatically-computed decomposition of internal mesh into subregions when rendering (ie inside the Component).
+	 * This allows for faster local updates via NotifyRegionDeferredEditCompleted() functions above.
+	 * Decomposition will be automatically recomputed as necessary when internal mesh is modified via changes, edits, etc
+	 */
+	virtual void SetEnableRenderMeshDecomposition(bool bEnable);
+
+	/** @return true if SetEnableRenderMeshDecomposition(true) has been called. */
+	bool GetIsRenderMeshDecompositionEnabled() const { return bDecompositionEnabled; }
+
 
 
 public:
@@ -415,6 +440,12 @@ protected:
 
 	/** Spatial data structure that is initialized if bBuildSpatialDataStructure = true when UpdatePreview() is called */
 	FDynamicMeshAABBTree3 MeshAABBTree;
+
+	/** If true, mesh will be chunked into multiple render buffers inside the DynamicMeshComponent */
+	bool bDecompositionEnabled = false;
+
+	/** Update chunk decomposition */
+	void UpdateRenderMeshDecomposition();
 };
 
 
