ModelingTools: disable point highlights by default in MeshSelectionTool. Add s/d hotkeys for brush resize, similar to sculpt/grouppaint
#rnx

diff --git a/Engine/Plugins/Experimental/MeshModelingToolset/Source/MeshModelingTools/Private/MeshSelectionTool.cpp b/Engine/Plugins/Experimental/MeshModelingToolset/Source/MeshModelingTools/Private/MeshSelectionTool.cpp
--- a/Engine/Plugins/Experimental/MeshModelingToolset/Source/MeshModelingTools/Private/MeshSelectionTool.cpp
+++ b/Engine/Plugins/Experimental/MeshModelingToolset/Source/MeshModelingTools/Private/MeshSelectionTool.cpp
@@ -187,6 +187,22 @@ void UMeshSelectionTool::RegisterActions(FInteractiveToolActionSet& ActionSet)
 {
 	UDynamicMeshBrushTool::RegisterActions(ActionSet);
 
+
+	ActionSet.RegisterAction(this, (int32)EStandardToolActions::BaseClientDefinedActionID + 50,
+		TEXT("TriSelectIncreaseSize"),
+		LOCTEXT("TriSelectIncreaseSize", "Increase Size"),
+		LOCTEXT("TriSelectIncreaseSizeTooltip", "Increase Brush Size"),
+		EModifierKey::None, EKeys::D,
+		[this]() { IncreaseBrushSizeAction(); });
+
+	ActionSet.RegisterAction(this, (int32)EStandardToolActions::BaseClientDefinedActionID + 51,
+		TEXT("TriSelectDecreaseSize"),
+		LOCTEXT("TriSelectDecreaseSize", "Decrease Size"),
+		LOCTEXT("TriSelectDecreaseSizeTooltip", "Decrease Brush Size"),
+		EModifierKey::None, EKeys::S,
+		[this]() { DecreaseBrushSizeAction(); });
+
+
 	ActionSet.RegisterAction(this, (int32)EStandardToolActions::BaseClientDefinedActionID + 1,
 		TEXT("MeshSelectionToolDelete"),
 		LOCTEXT("MeshSelectionToolDelete", "Delete"),
diff --git a/Engine/Plugins/Experimental/MeshModelingToolset/Source/MeshModelingTools/Public/MeshSelectionTool.h b/Engine/Plugins/Experimental/MeshModelingToolset/Source/MeshModelingTools/Public/MeshSelectionTool.h
--- a/Engine/Plugins/Experimental/MeshModelingToolset/Source/MeshModelingTools/Public/MeshSelectionTool.h
+++ b/Engine/Plugins/Experimental/MeshModelingToolset/Source/MeshModelingTools/Public/MeshSelectionTool.h
@@ -251,7 +251,7 @@ public:
 
 	/** Toggle drawing of highlight points on/off */
 	UPROPERTY(EditAnywhere, Category = Selection)
-	bool bShowPoints = true;
+	bool bShowPoints = false;
 
 	UPROPERTY(EditAnywhere, Category = Selection)
 	EMeshFacesColorMode FaceColorMode = EMeshFacesColorMode::None;
