GeometryProcessing:
Add insights profiling tags to RemeshMeshOp and FRemesher. Small optimizations to FRemesher::ProcessEdge(), which has surprising overhead for no-ops (ie edge is not modified).
Add FUniqueIndexSet, alternative to a TSet for constructing list of unique indices (eg when finding set of one-ring tris of a list of vertices). Generally more efficient, particularly for large sets.
Improve FSubRegionRemesher. Make VertexROI private, provide accessors and update usage sites. Optimize InitializeFromVertexROI(), which is an expensive part of dynamic mesh sculpting. Make FullSmoothPass_Buffer() compute in parallel.

ModelingTools:
Major refactor of DynamicMeshSculptTool to improve performance (2-4x improvement on large brush sizes on 1m triangle mesh). A single FSubRegionRemesher is now created and maintained for entire brush stroke, rather than recreating every stamp. This avoids the high cost of initialzing full-mesh constraint set, and allows re-use of various allocated memory buffers. Reorganize major sub-steps computed in a single brush-stamp tick to allow Async computation where possible. Replace some ROI TSets with FUniqueIndexSet. Instrument for Insights, remove old cycle counting.

#rnx

diff --git a/Engine/Plugins/Experimental/GeometryProcessing/Source/DynamicMesh/Private/Operations/SmoothHoleFiller.cpp b/Engine/Plugins/Experimental/GeometryProcessing/Source/DynamicMesh/Private/Operations/SmoothHoleFiller.cpp
--- a/Engine/Plugins/Experimental/GeometryProcessing/Source/DynamicMesh/Private/Operations/SmoothHoleFiller.cpp
+++ b/Engine/Plugins/Experimental/GeometryProcessing/Source/DynamicMesh/Private/Operations/SmoothHoleFiller.cpp
@@ -274,7 +274,7 @@ bool FSmoothHoleFiller::Fill(int32 GroupID)
 		{
 			Remesher = MakeUnique<FSubRegionRemesher>(&Mesh);
 			FMeshVertexSelection VertexSelection(&Mesh, TriangleSelection);	// All triangle vertices
-			Remesher->VertexROI = VertexSelection.AsSet();
+			Remesher->SetInitialVertexROI(VertexSelection.AsSet());
 			Remesher->InitializeFromVertexROI();
 		}
 
@@ -390,7 +390,7 @@ void FSmoothHoleFiller::SmoothAndRemesh(FMeshFaceSelection& TriangleSelection)
 	{
 		FSubRegionRemesher Remesher(&Mesh);
 		FMeshVertexSelection VertexSelection(&Mesh, TriangleSelection);
-		Remesher.VertexROI = VertexSelection.AsSet();
+		Remesher.SetInitialVertexROI(VertexSelection.AsSet());
 		Remesher.InitializeFromVertexROI();
 		Remesher.UpdateROI();
 
diff --git a/Engine/Plugins/Experimental/GeometryProcessing/Source/DynamicMesh/Private/Remesher.cpp b/Engine/Plugins/Experimental/GeometryProcessing/Source/DynamicMesh/Private/Remesher.cpp
--- a/Engine/Plugins/Experimental/GeometryProcessing/Source/DynamicMesh/Private/Remesher.cpp
+++ b/Engine/Plugins/Experimental/GeometryProcessing/Source/DynamicMesh/Private/Remesher.cpp
@@ -58,26 +58,29 @@ void FRemesher::BasicRemeshPass()
 	// some new edges. Can't see how we could efficiently prevent this.
 	//
 	ProfileBeginOps();
-
-	int cur_eid = StartEdges();
-	bool done = false;
 	ModifiedEdgesLastPass = 0;
-	do 
 	{
-		if (Mesh->IsEdge(cur_eid)) 
+		int cur_eid = StartEdges();
+		bool done = false;
+		int IterationCount = 0;
+		TRACE_CPUPROFILER_EVENT_SCOPE(Remesher_EdgesPass);
+		do
 		{
-			EProcessResult result = ProcessEdge(cur_eid);
-			if (result == EProcessResult::Ok_Collapsed || result == EProcessResult::Ok_Flipped || result == EProcessResult::Ok_Split)
+			if (Mesh->IsEdge(cur_eid))
 			{
-				ModifiedEdgesLastPass++;
+				EProcessResult result = ProcessEdge(cur_eid);
+				if (result == EProcessResult::Ok_Collapsed || result == EProcessResult::Ok_Flipped || result == EProcessResult::Ok_Split)
+				{
+					ModifiedEdgesLastPass++;
+				}
 			}
-		}
-		if (Cancelled())        // expensive to check every iter?
-		{
-			return;
-		}
-		cur_eid = GetNextEdge(cur_eid, done);
-	} while (done == false);
+			if (IterationCount++ % 1000 == 0 && Cancelled())        // expensive to check every iter?
+			{
+				return;
+			}
+			cur_eid = GetNextEdge(cur_eid, done);
+		} while (done == false);
+	}
 	ProfileEndOps();
 
 	if (Cancelled())
@@ -88,6 +91,7 @@ void FRemesher::BasicRemeshPass()
 	ProfileBeginSmooth();
 	if (bEnableSmoothing && (SmoothSpeedT > 0 || CustomSmoothSpeedF) )
 	{
+		TRACE_CPUPROFILER_EVENT_SCOPE(Remesher_SmoothingPass);
 		if (bEnableSmoothInPlace) 
 		{
 			//FullSmoothPass_InPlace(EnableParallelSmooth);
@@ -109,6 +113,7 @@ void FRemesher::BasicRemeshPass()
 	ProfileBeginProject();
 	if (ProjTarget != nullptr && ProjectionMode == ETargetProjectionMode::AfterRefinement)
 	{
+		TRACE_CPUPROFILER_EVENT_SCOPE(Remesher_ProjectionPass);
 		FullProjectionPass(bEnableParallelProjection);
 		DoDebugChecks();
 	}
@@ -140,41 +145,36 @@ FRemesher::EProcessResult FRemesher::ProcessEdge(int edgeID)
 	}
 
 	// look up verts and tris for this edge
-	if (Mesh->IsEdge(edgeID) == false)
-	{
-		return EProcessResult::Failed_NotAnEdge;
-	}
-	const FDynamicMesh3::FEdge Edge = Mesh->GetEdge(edgeID);
+	FDynamicMesh3::FEdge Edge(Mesh->GetEdge(edgeID));
 	int a = Edge.Vert[0], b = Edge.Vert[1];
 	int t0 = Edge.Tri[0], t1 = Edge.Tri[1];
 	bool bIsBoundaryEdge = (t1 == IndexConstants::InvalidID);
-
-	// look up 'other' verts c (from t0) and d (from t1, if it exists)
-	FIndex2i ov = Mesh->GetEdgeOpposingV(edgeID);
-	int c = ov[0], d = ov[1];
-
-	FVector3d vA = Mesh->GetVertex(a);
-	FVector3d vB = Mesh->GetVertex(b);
-	double edge_len_sqr = (vA - vB).SquaredLength();
-
-	ProfileBeginCollapse();
+	FVector3d vA(Mesh->GetVertex(a));
+	FVector3d vB(Mesh->GetVertex(b));
+	double edge_len_sqr = vA.DistanceSquared(vB);
 
 	// check if we should collapse, and also find which vertex we should collapse to,
 	// in cases where we have constraints/etc
 	int collapse_to = -1;
 	bool bCanCollapse = bEnableCollapses
 		&& constraint.CanCollapse()
-		&& edge_len_sqr < MinEdgeLength*MinEdgeLength
-		&& CanCollapseEdge(edgeID, a, b, c, d, t0, t1, collapse_to);
-
-	// optimization: if edge cd exists, we cannot collapse or flip. look that up here?
-	//  funcs will do it internally...
-	//  (or maybe we can collapse if cd exists? edge-collapse doesn't check for it explicitly...)
+		&& edge_len_sqr < MinEdgeLength* MinEdgeLength;
 
 	// if edge length is too short, we want to collapse it
 	bool bTriedCollapse = false;
 	if (bCanCollapse) 
 	{
+		// look up 'other' verts c (from t0) and d (from t1, if it exists)
+		FIndex2i ov = Mesh->GetEdgeOpposingV(edgeID);
+		int c = ov[0], d = ov[1];
+		if (CanCollapseEdge(edgeID, a, b, c, d, t0, t1, collapse_to) == false)
+		{
+			goto abort_collapse;
+		}
+		// optimization: if edge cd exists, we cannot collapse or flip. look that up here?
+		//  funcs will do it internally...
+		//  (or maybe we can collapse if cd exists? edge-collapse doesn't check for it explicitly...)
+
 		double collapse_t = 0.5;		// need to know t-value along edge to update lerpable attributes properly
 		FVector3d vNewPos = (vA + vB) * collapse_t;
 
@@ -241,13 +241,14 @@ FRemesher::EProcessResult FRemesher::ProcessEdge(int edgeID)
 	}
 abort_collapse:
 
-	ProfileEndCollapse();
-	ProfileBeginFlip();
-
 	// if this is not a boundary edge, maybe we want to flip
 	bool bTriedFlip = false;
 	if (bEnableFlips && constraint.CanFlip() && bIsBoundaryEdge == false) 
 	{
+		// look up 'other' verts c (from t0) and d (from t1, if it exists)
+		FIndex2i ov = Mesh->GetEdgeOpposingV(edgeID);
+		int c = ov[0], d = ov[1];
+
 		// can we do this more efficiently somehow?
 		bool bTryFlip = false;
 		if (FlipMetric == EFlipMetric::OptimalValence)
@@ -303,9 +304,6 @@ abort_collapse:
 		}
 	}
 
-	ProfileEndFlip();
-	ProfileBeginSplit();
-
 	// if edge length is too long, we want to split it
 	bool bTriedSplit = false;
 	if (bEnableSplits && constraint.CanSplit() && edge_len_sqr > MaxEdgeLength*MaxEdgeLength) 
@@ -327,9 +325,6 @@ abort_collapse:
 		}
 	}
 
-	ProfileEndSplit();
-
-
 	if (bTriedFlip || bTriedSplit || bTriedCollapse)
 	{
 		return EProcessResult::Failed_OpNotSuccessful;
diff --git a/Engine/Plugins/Experimental/GeometryProcessing/Source/DynamicMesh/Public/SubRegionRemesher.h b/Engine/Plugins/Experimental/GeometryProcessing/Source/DynamicMesh/Public/SubRegionRemesher.h
--- a/Engine/Plugins/Experimental/GeometryProcessing/Source/DynamicMesh/Public/SubRegionRemesher.h
+++ b/Engine/Plugins/Experimental/GeometryProcessing/Source/DynamicMesh/Public/SubRegionRemesher.h
@@ -4,8 +4,10 @@
 
 #include "Remesher.h"
 #include "Util/BufferUtil.h"
+#include "Util/UniqueIndexSet.h"
 #include "DynamicMeshChangeTracker.h"
 #include "MeshWeights.h"
+#include "Async/ParallelFor.h"
 
 /**
  * FSubRegionRemesher is an extension of FRemesher that allows for constraining remeshing to
@@ -18,6 +20,13 @@
 class FSubRegionRemesher : public FRemesher
 {
 protected:
+
+	/**
+	 * Set of vertices in ROI. You add vertices here initially, then
+	 * we will update the list during each Remesh pass
+	 */
+	TSet<int> VertexROI;
+
 	// active set of edges we will consider for remeshing. This set is updated on 
 	// each edge flip/split/collapse, but is not used during the pass.
 	TSet<int> EdgeROI;
@@ -31,15 +40,22 @@ protected:
 	TArray<int> Edges;
 
 	// index of current edge in .Edges we are at in StartEdges/GetNextEdge modulo-iteration
-	int CurEdge;
+	int CurEdge = -1;
 
 	// set of triangles removed in last pass. Enable this by calling BeginTrackRemovedTrisInPass()
 	TSet<int> RemovedLastPassTris;
 
+	// temporary data structure used to quickly convert vertex set to edge set. Currenlty it is faster
+	// to construct this and then copy to a properly-sized TSet, than to incrementally build the TSet
+	FUniqueIndexSet InitialROIEdgeVertices;
+
+	// temporary data structure used to quickly convert vertex set to triangle set. Currenlty it is faster
+	// to construct this and then copy to a properly-sized TSet, than to incrementally build the TSet
+	FUniqueIndexSet InitialROITriangles;
+
 	// controls whether RemovedLastPass will be populated
 	bool bTrackRemoved = false;
 
-
 	// counters for making sure that UpdateROI has been called
 	uint32 LastUpdateROICounter = 0;
 	uint32 LastRemeshPassCounter = 0;
@@ -57,54 +73,92 @@ public:
 
 	virtual ~FSubRegionRemesher() {}
 
-	/**
-	 * Set of vertices in ROI. You add vertices here initially, then 
-	 * we will update the list during each Remesh pass
-	 */
-	TSet<int> VertexROI;
 
+	void Reset()
+	{
+		VertexROI.Reset();
+		EdgeROI.Reset();
+		TriangleROI.Reset();
+		Edges.Reset();
+		RemovedLastPassTris.Reset();
+		CurEdge = -1;
+		LastUpdateROICounter = LastRemeshPassCounter = 0;
+	}
+
+
+	void SetInitialVertexROI(const TSet<int>& Vertices)
+	{
+		VertexROI = Vertices;
+	}
+
+	void SetInitialVertexROI(TSet<int>&& Vertices)
+	{
+		VertexROI = MoveTemp(Vertices);
+	}
+
+	void SetInitialVertexROI(const TArray<int>& Vertices)
+	{
+		VertexROI.Reset();
+		VertexROI.Reserve(Vertices.Num());
+		for (int32 tid : Vertices)
+		{
+			VertexROI.Add(tid);
+		}
+	}
 
 	/**
 	 * Initialize edge-subregion ROI from the VertexROI member that has been externally initialized
 	 */
 	void InitializeFromVertexROI()
 	{
-		EdgeROI.Reset();
-		TriangleROI.Reset();
+		const FDynamicMesh3* UseMesh = GetMesh();
 
-		// to get active edge set
-		for (int VertIdx : VertexROI)
 		{
-			for (int eid : GetMesh()->VtxEdgesItr(VertIdx))
+			TRACE_CPUPROFILER_EVENT_SCOPE(InitializeFromVertexROI_1Init);
+			EdgeROI.Reset();
+			TriangleROI.Reset();
+			InitialROIEdgeVertices.Initialize(UseMesh->MaxVertexID());
+
+			Edges.Reset();
+			LastUpdateROICounter++;
+
+			// convert vertex ROI to set of one-ring edges. Note that this effectively does
+			// UpdateROI(), and if that function changes, this may need to be modified.
+			// New set of ROI vertices is also constructed simultaneously
+			for (int VertIdx : VertexROI)
 			{
-				EdgeROI.Add(eid);
+				InitialROIEdgeVertices.Add(VertIdx);
+				UseMesh->EnumerateVertexEdges(VertIdx, [this, UseMesh, VertIdx](int32 eid)
+				{
+					bool bNewEdge = false;
+					EdgeROI.Add(eid, &bNewEdge);
+					if (bNewEdge)
+					{
+						FIndex2i EdgeV = UseMesh->GetEdgeV(eid);
+						InitialROIEdgeVertices.Add((EdgeV.A == VertIdx) ? EdgeV.B : EdgeV.A);
+						Edges.Add(eid);
+					}
+				});
 			}
 		}
-		UpdateROI();
-
-		// there is quite a bit of overhead here...perhaps remesher could just save triangles
-		// itself before it touches them?
 
-		TArray<int> OneRingTris; OneRingTris.Reserve(32);
-
-		// figuring out unique verts means we don't do each vertex N~=valence times, 
-		// which saves a lot of one-ring iterations that are somewhat expensive...
-		TSet<int> Vertices;
-		const FDynamicMesh3* UseMesh = GetMesh();
-		for (int eid : EdgeROI)
-		{
-			FIndex2i EdgeVerts = UseMesh->GetEdgeV(eid);
-			Vertices.Add(EdgeVerts.A);
-			Vertices.Add(EdgeVerts.B);
-		}
-		for (int vid : Vertices)
 		{
-			OneRingTris.Reset();
-			UseMesh->GetVertexOneRingTriangles(vid, OneRingTris);
-			for (int tid : OneRingTris)
+			TRACE_CPUPROFILER_EVENT_SCOPE(InitializeFromVertexROI_4Triangles);
+
+			// Convert new vertex ROI built above into new one-ring triangle ROI. 
+			// This conservatively guarantees that any triangle affected by a mesh edge op 
+			// (eg split/collapse/flip) is contained within this triangle set
+			InitialROITriangles.Initialize(UseMesh->MaxTriangleID());
+			for (int vid : InitialROIEdgeVertices)
 			{
-				TriangleROI.Add(tid);
+				UseMesh->EnumerateVertexEdges(vid, [this](int32 eid)
+				{
+					FIndex2i EdgeT = Mesh->GetEdgeT(eid);
+					InitialROITriangles.Add(EdgeT.A);
+					if (EdgeT.B != IndexConstants::InvalidID) InitialROITriangles.Add(EdgeT.B);
+				});
 			}
+			InitialROITriangles.Collect(TriangleROI);
 		}
 	}
 
@@ -134,10 +188,11 @@ public:
 	 */
 	void UpdateROI()
 	{
+		// note: if you change this function, you also need to change InitializeFromVertexROI() because it
+		// does this internally (or call this afterwards)
 		Edges.Reset();
 		for (int eid : EdgeROI)
 		{
-			//check(GetMesh()->IsEdge(eid));
 			Edges.Add(eid);
 		}
 		LastUpdateROICounter++;
@@ -189,6 +244,14 @@ public:
 		}
 	}
 
+	/**
+	 * @return 
+	 */
+	const TSet<int>& GetCurrentVertexROI() const
+	{
+		return VertexROI;
+	}
+
 	/**
 	 * @return set of triangles that contains edge ROI (note: may also contain additional triangles)
 	 */
@@ -339,10 +402,10 @@ protected:
 
 	void AddVertexToTriangleROI(int VertexID)
 	{
-		for (int tid : GetMesh()->VtxTrianglesItr(VertexID))
+		GetMesh()->EnumerateVertexTriangles(VertexID, [this](int tid)
 		{
 			TriangleROI.Add(tid);
-		}
+		});
 	}
 
 	void AddEdgeToTriangleROI(int EdgeID)
@@ -367,32 +430,39 @@ protected:
 	// localized smoothing
 	// 
 	
-	TMap<int, FVector3d> SmoothedPositions;
+	TArray<int32> SmoothIndices;
+	TArray<FVector3d> SmoothedPositions;
 
 	virtual void FullSmoothPass_Buffer(bool bParallel) override
 	{
-		SmoothedPositions.Reset();
+		SmoothIndices.Reset();
+		SmoothIndices.Reserve(VertexROI.Num());
+		for (int32 vid : VertexROI)
+		{
+			SmoothIndices.Add(vid);
+		}
+		int32 NumVertices = SmoothIndices.Num();
+
+		SmoothedPositions.SetNum(NumVertices, false);
 
 		TFunction<FVector3d(const FDynamicMesh3&, int, double)> UseSmoothFunc = GetSmoothFunction();
-		auto SmoothAndUpdateFunc = [&](int vid)
+		
+		ParallelFor(NumVertices, [&](int32 i)
 		{
+			int32 VertIdx = SmoothIndices[i];
 			bool bModified = false;
-			FVector3d SmoothedPosition = ComputeSmoothedVertexPos(vid, UseSmoothFunc, bModified);
-			if (bModified)
-			{
-				SmoothedPositions.Add(vid, SmoothedPosition);
-			}
-		};
-
-		for (int vid : VertexROI)
-		{
-			SmoothAndUpdateFunc(vid);
-		}
+			SmoothedPositions[i] = ComputeSmoothedVertexPos(VertIdx, UseSmoothFunc, bModified);
+			SmoothIndices[i] = (bModified) ? SmoothIndices[i] : -1;
+		});
 
-		for (auto Pair : SmoothedPositions)
+		ParallelFor(NumVertices, [&](int32 i)
 		{
-			Mesh->SetVertex(Pair.Key, Pair.Value);
-		}
+			if (SmoothIndices[i] >= 0)
+			{
+				Mesh->SetVertex_NoTimeStampUpdate(SmoothIndices[i], SmoothedPositions[i]);
+			}
+		});
+		Mesh->IncrementTimeStamps(1, true, false);
 	}
 };
 
diff --git a/Engine/Plugins/Experimental/GeometryProcessing/Source/GeometricObjects/Private/Util/UniqueIndexSet.cpp b/Engine/Plugins/Experimental/GeometryProcessing/Source/GeometricObjects/Private/Util/UniqueIndexSet.cpp
new file mode 100644
--- /dev/null
+++ b/Engine/Plugins/Experimental/GeometryProcessing/Source/GeometricObjects/Private/Util/UniqueIndexSet.cpp
@@ -0,0 +1,11 @@
+// Copyright Epic Games, Inc. All Rights Reserved.
+
+#include "Util/UniqueIndexSet.h"
+
+FUniqueIndexSet::~FUniqueIndexSet()
+{
+	if (Bits != nullptr)
+	{
+		delete[] Bits;
+	}
+}
diff --git a/Engine/Plugins/Experimental/GeometryProcessing/Source/GeometricObjects/Public/Util/UniqueIndexSet.h b/Engine/Plugins/Experimental/GeometryProcessing/Source/GeometricObjects/Public/Util/UniqueIndexSet.h
new file mode 100644
--- /dev/null
+++ b/Engine/Plugins/Experimental/GeometryProcessing/Source/GeometricObjects/Public/Util/UniqueIndexSet.h
@@ -0,0 +1,129 @@
+// Copyright Epic Games, Inc. All Rights Reserved.
+
+#pragma once
+
+#include "CoreMinimal.h"
+
+/**
+ * FUniqueIndexSet is used to construct a list of unique elements of integers in range [0, MaxValue].
+ * General usage is for things like converting a vertex set to the set of connected one-ring triangles.
+ * The approach is to have a bit-array for set membership, and then accumulate unique entries in a TArray.
+ * 
+ * This is faster for construction than a TSet, but does require allocating and/or clearing the bit-array.
+ * Still generally quite a bit more efficient, particularly for large sets.
+ */
+class GEOMETRICOBJECTS_API FUniqueIndexSet
+{
+	// TODO: second bit set that tracks which blocks of ints have been set.
+	// This would allow for clears to be much quicker for subsets.
+	// (would want to profile to see that it helps, but the bit set is still 0.125mb per million indices,
+	// which isn't a trivial amount to set to zero...)
+
+	// TODO: alternately, use a TSet instead of Bits array if we know that the number 
+	// of values we are going to add is very small? This would avoid clears.
+
+	// TODO: AtomicAdd()
+
+	// TODO: support growing bit set
+
+public:
+	~FUniqueIndexSet();
+
+	/**
+	 * Initialize the set with maximum index. 
+	 * @param MaxValuesHint if non-zero, we pre-allocate this much memory for the index array
+	 */
+	inline void Initialize(int32 MaxIndexIn, int32 MaxValuesHint = 0)
+	{
+		this->MaxIndex = MaxIndexIn;
+		int32 NeedWords = (MaxIndex / 64) + 1;
+		if (NumWords < NeedWords)
+		{
+			if (Bits != nullptr)
+			{
+				delete[] Bits;
+			}
+			NumWords = NeedWords;
+			Bits = new int64[NumWords];
+		}
+		FMemory::Memset(&Bits[0], (uint8)0, NeedWords * 8);
+		Values.Reset();
+		Values.Reserve(MaxValuesHint);
+	}
+
+	/** @return number of indices in set */
+	inline const int32 Num() const { return Values.Num(); }
+
+	/** @return array of unique indices */
+	inline const TArray<int32>& Indices() const { return Values; }
+
+	/**
+	 * Add Index to the set
+	 */
+	inline void Add(int32 Index)
+	{
+		int64& Word = Bits[(int64)(Index / 64)];
+		int64 Offset = (int64)(Index % 64);
+		int64 Mask = (int64)1 << Offset;
+		if ( (Word & Mask) == 0 )
+		{
+			Word |= Mask;
+			Values.Add(Index);
+		}
+	}
+
+	/**
+	 * @return true if Index is in the current set
+	 */
+	inline bool Contains(int32 Index) const
+	{
+		const int64& Word = Bits[(int64)(Index / 64)];
+		int64 Offset = (int64)(Index % 64);
+		int64 Mask = (int64)1 << Offset;
+		return (Word & Mask) != 0;
+	}
+
+	/**
+	 * Add the members of the current index set to an object that supports .Reserve(int) and .Add(int)
+	 */
+	template<typename ArrayType>
+	void Collect(ArrayType& Storage) const
+	{
+		Storage.Reserve(Values.Num());
+		for (int32 Value : Values)
+		{
+			Storage.Add(Value);
+		}
+	}
+
+	/**
+	 * @return the internal TArray. After calling this, the set is invalid and needs to be Initialize()'d again
+	 */
+	inline TArray<int32>&& TakeValues()
+	{
+		return MoveTemp(Values);
+	}
+
+	/**
+	 * Swap the internal TArray with another TArray. After calling this, the set is invalid and needs to be Initialize()'d again
+	 */
+	inline void SwapValuesWith(TArray<int32>& OtherArray)
+	{
+		TArray<int32> Tmp = MoveTemp(Values);
+		Values = MoveTemp(OtherArray);
+		OtherArray = MoveTemp(Tmp);
+	}
+
+
+public:
+	// range-for support. Do not directly call these functions.
+	inline TArray<int32>::RangedForConstIteratorType begin() const { return Values.begin(); }
+	inline TArray<int32>::RangedForConstIteratorType end() const { return Values.end(); }
+
+
+protected:
+	int32 MaxIndex;
+	TArray<int32> Values;
+	int64* Bits = nullptr;
+	int32 NumWords = 0;
+};
\ No newline at end of file
diff --git a/Engine/Plugins/Experimental/MeshModelingToolset/Source/MeshModelingTools/Private/DynamicMeshSculptTool.cpp b/Engine/Plugins/Experimental/MeshModelingToolset/Source/MeshModelingTools/Private/DynamicMeshSculptTool.cpp
--- a/Engine/Plugins/Experimental/MeshModelingToolset/Source/MeshModelingTools/Private/DynamicMeshSculptTool.cpp
+++ b/Engine/Plugins/Experimental/MeshModelingToolset/Source/MeshModelingTools/Private/DynamicMeshSculptTool.cpp
@@ -387,11 +387,10 @@ void UDynamicMeshSculptTool::OnBeginDrag(const FRay& Ray)
 	}
 }
 
+
 void UDynamicMeshSculptTool::UpdateROI(const FVector3d& BrushPos)
 {
-	SCOPE_CYCLE_COUNTER(SculptTool_UpdateROI);
-
-	// TODO: need dynamic vertex hash table!
+	TRACE_CPUPROFILER_EVENT_SCOPE(DynamicMeshSculptTool_UpdateROI);
 
 	float RadiusSqr = CurrentBrushRadius * CurrentBrushRadius;
 
@@ -399,28 +398,56 @@ void UDynamicMeshSculptTool::UpdateROI(const FVector3d& BrushPos)
 		BrushPos - CurrentBrushRadius * FVector3d::One(),
 		BrushPos + CurrentBrushRadius * FVector3d::One());
 
-	VertexSetBuffer.Reset();
 	FDynamicMesh3* Mesh = DynamicMeshComponent->GetMesh();
 	FDynamicMeshOctree3* Octree = DynamicMeshComponent->GetOctree();
-	Octree->RangeQuery(BrushBox,
-		[&](int TriIdx) {
-		FIndex3i TriV = Mesh->GetTriangle(TriIdx);
-		for (int j = 0; j < 3; ++j)
+
+	// find set of triangles in brush bounding box
+	UpdateROITriBuffer.Reset();
+	{
+		TRACE_CPUPROFILER_EVENT_SCOPE(DynamicMeshSculptTool_UpdateROI_1RangeQuery);
+		Octree->ParallelRangeQuery(BrushBox, UpdateROITriBuffer);
+	}
+
+	// collect set of vertices inside brush sphere, from that box
+	VertexROIBuilder.Initialize(Mesh->MaxVertexID());
+	{
+		TRACE_CPUPROFILER_EVENT_SCOPE(DynamicMeshSculptTool_UpdateROI_2Collect);
+		for (int32 TriIdx : UpdateROITriBuffer)
 		{
-			FVector3d Position = Mesh->GetVertex(TriV[j]);
-			if ((Position - BrushPos).SquaredLength() < RadiusSqr)
+			FIndex3i TriV = Mesh->GetTriangle(TriIdx);
+			for (int j = 0; j < 3; ++j)
 			{
-				VertexSetBuffer.Add(TriV[j]);
+				if (VertexROIBuilder.Contains(TriV[j]) == false)
+				{
+					//const FVector3d& Position = Mesh->GetVertexRef(TriV[j]);
+					FVector3d Position = Mesh->GetVertex(TriV[j]);
+					if (BrushPos.DistanceSquared(Position) < RadiusSqr)
+					{
+						VertexROIBuilder.Add(TriV[j]);
+					}
+				}
 			}
 		}
+		VertexROI.Reset();
+		VertexROIBuilder.SwapValuesWith(VertexROI);
+	}
 
-	});
-
-	VertexROI.SetNum(0, false);
-	BufferUtil::AppendElements(VertexROI, VertexSetBuffer);
-
-	TriangleROI.Reset();
-	MeshIndexUtil::VertexToTriangleOneRing(Mesh, VertexROI, TriangleROI);
+	// find triangle ROI as full one-rings of all vertices (this is surprisingly expensive...)
+	{
+		TRACE_CPUPROFILER_EVENT_SCOPE(DynamicMeshSculptTool_UpdateROI_3TriangleROI);
+		TriangleROIBuilder.Initialize(Mesh->MaxTriangleID());
+		for (int32 vid : VertexROI)
+		{
+			Mesh->EnumerateVertexEdges(vid, [&](int32 eid)
+			{
+				FDynamicMesh3::FEdge Edge = Mesh->GetEdge(eid);
+				TriangleROIBuilder.Add(Edge.Tri.A);
+				if (Edge.Tri.B != IndexConstants::InvalidID) TriangleROIBuilder.Add(Edge.Tri.B);
+			});
+		}
+		TriangleROI.Reset();
+		TriangleROIBuilder.Collect(TriangleROI);
+	}
 }
 
 void UDynamicMeshSculptTool::OnUpdateDrag(const FRay& WorldRay)
@@ -447,7 +474,10 @@ void UDynamicMeshSculptTool::CalculateBrushRadius()
 
 void UDynamicMeshSculptTool::ApplyStamp(const FRay& WorldRay)
 {
-	SCOPE_CYCLE_COUNTER(STAT_SculptToolApplyStamp);
+	TRACE_CPUPROFILER_EVENT_SCOPE(DynamicMeshSculptTool_ApplyStamp);
+
+	FDynamicMesh3* Mesh = DynamicMeshComponent->GetMesh();
+	FDynamicMeshOctree3* Octree = DynamicMeshComponent->GetOctree();
 
 	// update brush type history. apologies for convoluted logic.
 	StampTimestamp++;
@@ -471,67 +501,100 @@ void UDynamicMeshSculptTool::ApplyStamp(const FRay& WorldRay)
 
 	CalculateBrushRadius();
 
-	SaveActiveROI();
+	TFuture<void> DirtyOctreeFuture = Async(DynamicSculptToolAsyncExecTarget, [&]()
+	{
+		TRACE_CPUPROFILER_EVENT_SCOPE(DynamicMeshSculptTool_ApplyStamp_DirtyOctree);
+		Octree->NotifyPendingModification(TriangleROI);
+	});
 
-	if (bSmoothing)
+	TFuture<void> SaveROIFuture = Async(DynamicSculptToolAsyncExecTarget, [&]() 
 	{
-		ApplySmoothBrush(WorldRay);
-		return;
-	}
+		TRACE_CPUPROFILER_EVENT_SCOPE(DynamicMeshSculptTool_ApplyStamp_SaveActiveROI);
+		SaveActiveROI();
+	});
 
-	switch (SculptProperties->PrimaryBrushType)
+	// TODO: 
+	//   - we can begin Octree->RemoveTriangles below as soon as we have (1) finished
+	//     marking dirty box in DirtyOctreeFuture and (2) updated brush position. Unfortunately
+	//     right now that happens inside each brush func :(
+
+	EDynamicMeshSculptBrushType ApplyBrushType = (bSmoothing) ?
+		EDynamicMeshSculptBrushType::Smooth : SculptProperties->PrimaryBrushType;
+
+	bool bBrushApplied = false;
+	switch (ApplyBrushType)
 	{
 		case EDynamicMeshSculptBrushType::Offset:
-			ApplyOffsetBrush(WorldRay, false);
+			bBrushApplied = ApplyOffsetBrush(WorldRay, false);
 			break;
 		case EDynamicMeshSculptBrushType::SculptView:
-			ApplyOffsetBrush(WorldRay, true);
+			bBrushApplied = ApplyOffsetBrush(WorldRay, true);
 			break;
 		case EDynamicMeshSculptBrushType::SculptMax:
-			ApplySculptMaxBrush(WorldRay);
+			bBrushApplied = ApplySculptMaxBrush(WorldRay);
 			break;
 		case EDynamicMeshSculptBrushType::Move:
-			ApplyMoveBrush(WorldRay);
+			bBrushApplied = ApplyMoveBrush(WorldRay);
 			break;
 		case EDynamicMeshSculptBrushType::PullKelvin:
-			ApplyPullKelvinBrush(WorldRay);
+			bBrushApplied = ApplyPullKelvinBrush(WorldRay);
 			break;
 		case EDynamicMeshSculptBrushType::PullSharpKelvin:
-			ApplyPullSharpKelvinBrush(WorldRay);
+			bBrushApplied = ApplyPullSharpKelvinBrush(WorldRay);
 			break;
 		case EDynamicMeshSculptBrushType::Smooth:
-			ApplySmoothBrush(WorldRay);
+			bBrushApplied = ApplySmoothBrush(WorldRay);
 			break;
 		case EDynamicMeshSculptBrushType::Pinch:
-			ApplyPinchBrush(WorldRay);
+			bBrushApplied = ApplyPinchBrush(WorldRay);
 			break;
 		case EDynamicMeshSculptBrushType::TwistKelvin:
-			ApplyTwistKelvinBrush(WorldRay);
+			bBrushApplied = ApplyTwistKelvinBrush(WorldRay);
 			break;
 		case EDynamicMeshSculptBrushType::Inflate:
-			ApplyInflateBrush(WorldRay);
+			bBrushApplied = ApplyInflateBrush(WorldRay);
 			break;
 		case EDynamicMeshSculptBrushType::ScaleKelvin:
-			ApplyScaleKelvinBrush(WorldRay);
+			bBrushApplied = ApplyScaleKelvinBrush(WorldRay);
 			break;
 		case EDynamicMeshSculptBrushType::Flatten:
-			ApplyFlattenBrush(WorldRay);
+			bBrushApplied = ApplyFlattenBrush(WorldRay);
 			break;
 		case EDynamicMeshSculptBrushType::Plane:
-			ApplyPlaneBrush(WorldRay);
+			bBrushApplied = ApplyPlaneBrush(WorldRay);
 			break;
 		case EDynamicMeshSculptBrushType::PlaneViewAligned:
-			ApplyPlaneBrush(WorldRay);
+			bBrushApplied = ApplyPlaneBrush(WorldRay);
 			break;
 		case EDynamicMeshSculptBrushType::FixedPlane:
-			ApplyFixedPlaneBrush(WorldRay);
+			bBrushApplied = ApplyFixedPlaneBrush(WorldRay);
 			break;
 		case EDynamicMeshSculptBrushType::Resample:
-			ApplyResampleBrush(WorldRay);
+			bBrushApplied = ApplyResampleBrush(WorldRay);
 			break;
 		case EDynamicMeshSculptBrushType::LastValue:
 			break;
 	}
+
+	// wait for ROI to finish saving before we update positions
+	SaveROIFuture.Wait();
+	DirtyOctreeFuture.Wait();
+
+	// we are going to reinsert these later
+	TFuture<void> OctreeRemoveFuture = Async(DynamicSculptToolAsyncExecTarget, [&]()
+	{
+		TRACE_CPUPROFILER_EVENT_SCOPE(DynamicMeshSculptTool_ApplyStamp_OctreeRemove);
+		Octree->RemoveTriangles(TriangleROI, false);	// already marked dirty above
+	});
+
+	// Update the mesh positions to match those in the position buffer
+	if (bBrushApplied)
+	{
+		SyncMeshWithPositionBuffer(Mesh);
+	}
+
+	// we don't stricty have to wait here, we could return this future
+	OctreeRemoveFuture.Wait();
 }
 
 double UDynamicMeshSculptTool::CalculateBrushFalloff(double Distance)
@@ -555,22 +618,31 @@ void UDynamicMeshSculptTool::SyncMeshWithPositionBuffer(FDynamicMesh3* Mesh)
 	const int NumV = ROIPositionBuffer.Num();
 	checkSlow(VertexROI.Num() <= NumV);
 
-	for (int k = 0; k < NumV; ++k)
+	// only if remeshing is disabled?
+	if (ActiveVertexChange != nullptr)
+	{
+		for (int k = 0; k < NumV; ++k)
+		{
+			int VertIdx = VertexROI[k];
+			ActiveVertexChange->UpdateVertex(VertIdx, Mesh->GetVertex(VertIdx), ROIPositionBuffer[k]);
+		}
+	}
+
+	ParallelFor(NumV, [&](int32 k)
 	{
 		int VertIdx = VertexROI[k];
 		const FVector3d& NewPos = ROIPositionBuffer[k];
-		FVector3d OrigPos = Mesh->GetVertex(VertIdx);
-		Mesh->SetVertex(VertIdx, NewPos);
-		UpdateSavedVertex(VertIdx, OrigPos, NewPos);
-	}
+		Mesh->SetVertex_NoTimeStampUpdate(VertIdx, NewPos);
+	});
+	Mesh->IncrementTimeStamps(1, true, false);
 }
 
-void UDynamicMeshSculptTool::ApplySmoothBrush(const FRay& WorldRay)
+bool UDynamicMeshSculptTool::ApplySmoothBrush(const FRay& WorldRay)
 {
 	bool bHit = UpdateBrushPositionOnSculptMesh(WorldRay, true);
 	if (bHit == false)
 	{
-		return;
+		return false;
 	}
 
 	FVector3d NewBrushPosLocal = CurTargetTransform.InverseTransformPosition(LastBrushPosWorld);
@@ -595,13 +667,12 @@ void UDynamicMeshSculptTool::ApplySmoothBrush(const FRay& WorldRay)
 		ROIPositionBuffer[k] = NewPos;
 	});
 
-	SyncMeshWithPositionBuffer(Mesh);
-
 	ScheduleRemeshPass();
 	LastBrushPosLocal = NewBrushPosLocal;
+	return true;
 }
 
-void UDynamicMeshSculptTool::ApplyMoveBrush(const FRay& WorldRay)
+bool UDynamicMeshSculptTool::ApplyMoveBrush(const FRay& WorldRay)
 {
 	UpdateBrushPositionOnActivePlane(WorldRay);
 
@@ -611,7 +682,7 @@ void UDynamicMeshSculptTool::ApplyMoveBrush(const FRay& WorldRay)
 	if (MoveVec.SquaredLength() <= 0)
 	{
 		LastBrushPosLocal = NewBrushPosLocal;
-		return;
+		return false;
 	}
 
 	FDynamicMesh3* Mesh = DynamicMeshComponent->GetMesh();
@@ -633,14 +704,12 @@ void UDynamicMeshSculptTool::ApplyMoveBrush(const FRay& WorldRay)
 		ROIPositionBuffer[k] = NewPos;
 	});
 
-	// Update the mesh positions to match those in the position buffer
-	SyncMeshWithPositionBuffer(Mesh);
-
 	ScheduleRemeshPass();
 	LastBrushPosLocal = NewBrushPosLocal;
+	return true;
 }
 
-void UDynamicMeshSculptTool::ApplyOffsetBrush(const FRay& WorldRay, bool bUseViewDirection)
+bool UDynamicMeshSculptTool::ApplyOffsetBrush(const FRay& WorldRay, bool bUseViewDirection)
 {
 	UpdateBrushPositionOnTargetMesh(WorldRay, bUseViewDirection);
 	if (bUseViewDirection)
@@ -678,14 +747,12 @@ void UDynamicMeshSculptTool::ApplyOffsetBrush(const FRay& WorldRay, bool bUseVie
 		}
 	});
 
-	// Update the mesh positions to match those in the position buffer
-	SyncMeshWithPositionBuffer(Mesh);
-
 	ScheduleRemeshPass();
 	LastBrushPosLocal = NewBrushPosLocal;
+	return true;
 }
 
-void UDynamicMeshSculptTool::ApplySculptMaxBrush(const FRay& WorldRay)
+bool UDynamicMeshSculptTool::ApplySculptMaxBrush(const FRay& WorldRay)
 {
 	UpdateBrushPositionOnTargetMesh(WorldRay, true);
 	FVector3d NewBrushPosLocal = CurTargetTransform.InverseTransformPosition(LastBrushPosWorld);
@@ -731,14 +798,12 @@ void UDynamicMeshSculptTool::ApplySculptMaxBrush(const FRay& WorldRay)
 		}
 	});
 
-	// Update the mesh positions to match those in the position buffer
-	SyncMeshWithPositionBuffer(Mesh);
-
 	ScheduleRemeshPass();
 	LastBrushPosLocal = NewBrushPosLocal;
+	return true;
 }
 
-void UDynamicMeshSculptTool::ApplyPinchBrush(const FRay& WorldRay)
+bool UDynamicMeshSculptTool::ApplyPinchBrush(const FRay& WorldRay)
 {
 	UpdateBrushPositionOnTargetMesh(WorldRay, true);
 	FVector3d NewBrushPosLocal = CurTargetTransform.InverseTransformPosition(LastBrushPosWorld);
@@ -784,13 +849,10 @@ void UDynamicMeshSculptTool::ApplyPinchBrush(const FRay& WorldRay)
 		ROIPositionBuffer[k] = NewPos;
 	});
 
-	// Update the mesh positions to match those in the position buffer
-	SyncMeshWithPositionBuffer(Mesh);
-
-
 	ScheduleRemeshPass();
 	LastBrushPosLocal = NewBrushPosLocal;
 	LastSmoothBrushPosLocal = NewSmoothBrushPosLocal;
+	return true;
 }
 
 FFrame3d UDynamicMeshSculptTool::ComputeROIBrushPlane(const FVector3d& BrushCenter, bool bIgnoreDepth, bool bViewAligned)
@@ -825,12 +887,12 @@ FFrame3d UDynamicMeshSculptTool::ComputeROIBrushPlane(const FVector3d& BrushCent
 	return Result;
 }
 
-void UDynamicMeshSculptTool::ApplyPlaneBrush(const FRay& WorldRay)
+bool UDynamicMeshSculptTool::ApplyPlaneBrush(const FRay& WorldRay)
 {
 	bool bHit = UpdateBrushPositionOnSculptMesh(WorldRay, true);
 	if (bHit == false)
 	{
-		return;
+		return false;
 	}
 
 	static const double PlaneSigns[3] = { 0, -1, 1 };
@@ -861,19 +923,17 @@ void UDynamicMeshSculptTool::ApplyPlaneBrush(const FRay& WorldRay)
 		ROIPositionBuffer[k] = NewPos;
 	});
 
-	// Update the mesh positions to match those in the position buffer
-	SyncMeshWithPositionBuffer(Mesh);
-
 	ScheduleRemeshPass();
 	LastBrushPosLocal = NewBrushPosLocal;
+	return true;
 }
 
-void UDynamicMeshSculptTool::ApplyFixedPlaneBrush(const FRay& WorldRay)
+bool UDynamicMeshSculptTool::ApplyFixedPlaneBrush(const FRay& WorldRay)
 {
 	bool bHit = UpdateBrushPositionOnSculptMesh(WorldRay, true);
 	if (bHit == false)
 	{
-		return;
+		return false;
 	}
 
 	static const double PlaneSigns[3] = { 0, -1, 1 };
@@ -910,19 +970,17 @@ void UDynamicMeshSculptTool::ApplyFixedPlaneBrush(const FRay& WorldRay)
 		ROIPositionBuffer[k] = NewPos;
 	});
 
-	// Update the mesh positions to match those in the position buffer
-	SyncMeshWithPositionBuffer(Mesh);
-
 	ScheduleRemeshPass();
 	LastBrushPosLocal = NewBrushPosLocal;
+	return true;
 }
 
-void UDynamicMeshSculptTool::ApplyFlattenBrush(const FRay& WorldRay)
+bool UDynamicMeshSculptTool::ApplyFlattenBrush(const FRay& WorldRay)
 {
 	bool bHit = UpdateBrushPositionOnSculptMesh(WorldRay, true);
 	if (bHit == false)
 	{
-		return;
+		return false;
 	}
 
 	static const double PlaneSigns[3] = { 0, -1, 1 };
@@ -960,19 +1018,17 @@ void UDynamicMeshSculptTool::ApplyFlattenBrush(const FRay& WorldRay)
 		ROIPositionBuffer[k] = NewPos;
 	});
 
-	// Update the mesh positions to match those in the position buffer
-	SyncMeshWithPositionBuffer(Mesh);
-
 	ScheduleRemeshPass();
 	LastBrushPosLocal = NewBrushPosLocal;
+	return true;
 }
 
-void UDynamicMeshSculptTool::ApplyInflateBrush(const FRay& WorldRay)
+bool UDynamicMeshSculptTool::ApplyInflateBrush(const FRay& WorldRay)
 {
 	bool bHit = UpdateBrushPositionOnSculptMesh(WorldRay, true);
 	if (bHit == false)
 	{
-		return;
+		return false;
 	}
 
 	FVector3d NewBrushPosLocal = CurTargetTransform.InverseTransformPosition(LastBrushPosWorld);
@@ -1006,14 +1062,13 @@ void UDynamicMeshSculptTool::ApplyInflateBrush(const FRay& WorldRay)
 		ROIPositionBuffer[k] = NewPos;
 	});
 
-	// Update the mesh positions to match those in the position buffer
-	SyncMeshWithPositionBuffer(Mesh);
 	ScheduleRemeshPass();
 	LastBrushPosLocal = NewBrushPosLocal;
+	return true;
 }
 
 
-void UDynamicMeshSculptTool::ApplyResampleBrush(const FRay& WorldRay)
+bool UDynamicMeshSculptTool::ApplyResampleBrush(const FRay& WorldRay)
 {
 	UpdateBrushPositionOnTargetMesh(WorldRay, true);
 	FVector3d NewBrushPosLocal = CurTargetTransform.InverseTransformPosition(LastBrushPosWorld);
@@ -1026,29 +1081,16 @@ void UDynamicMeshSculptTool::ApplyResampleBrush(const FRay& WorldRay)
 	ROIPositionBuffer.SetNum(NumV, false);
 	ParallelFor(NumV, [&](int k)
 	{
-		int VertIdx = VertexROI[k];
-		FVector3d OrigPos = Mesh->GetVertex(VertIdx);
-
-		FVector3d BasePos, BaseNormal;
-		if (GetTargetMeshNearest(OrigPos, (double)(4 * CurrentBrushRadius), BasePos, BaseNormal) == false)
-		{
-			ROIPositionBuffer[k] = OrigPos;
-		}
-		else
-		{
-			double Falloff = CalculateBrushFalloff(OrigPos.Distance(NewBrushPosLocal));
-			FVector3d NewPos = BasePos;// FVector3d::Lerp(OrigPos, BasePos, Falloff);
-			ROIPositionBuffer[k] = NewPos;
-		}
+		ROIPositionBuffer[k] = Mesh->GetVertex(VertexROI[k]);
 	});
 
-	SyncMeshWithPositionBuffer(Mesh);
 	ScheduleRemeshPass();
 	LastBrushPosLocal = NewBrushPosLocal;
+	return true;
 }
 
 
-void UDynamicMeshSculptTool::ApplyPullKelvinBrush(const FRay& WorldRay)
+bool UDynamicMeshSculptTool::ApplyPullKelvinBrush(const FRay& WorldRay)
 {
 	UpdateBrushPositionOnActivePlane(WorldRay);
 	FVector3d BrushNormalLocal = CurTargetTransform.InverseTransformNormal(LastBrushPosNormalWorld);
@@ -1058,7 +1100,7 @@ void UDynamicMeshSculptTool::ApplyPullKelvinBrush(const FRay& WorldRay)
 	if (MoveVec.SquaredLength() <= 0)
 	{
 		LastBrushPosLocal = NewBrushPosLocal;
-		return;
+		return false;
 	}
 	
 	FDynamicMesh3* Mesh = DynamicMeshComponent->GetMesh();
@@ -1069,20 +1111,18 @@ void UDynamicMeshSculptTool::ApplyPullKelvinBrush(const FRay& WorldRay)
 	FKelvinletBrushOp::FKelvinletBrushOpProperties  KelvinletBrushOpProperties(KelvinMode, *KelvinBrushProperties, *BrushProperties);
 	KelvinletBrushOpProperties.Direction = FVector(MoveVec.X, MoveVec.Y, MoveVec.Z);  //FVector(BrushNormalLocal.X, BrushNormalLocal.Y, BrushNormalLocal.Z);
 	KelvinletBrushOpProperties.Size *= 0.6;
-	
 
 	FMatrix ToBrush; ToBrush.SetIdentity();  ToBrush.SetOrigin(-FVector(NewBrushPosLocal.X, NewBrushPosLocal.Y, NewBrushPosLocal.Z)); //  ToBrush.
 
 	KelvinBrushOp.ApplyBrush(KelvinletBrushOpProperties, ToBrush, VertexROI, ROIPositionBuffer);
 
-	// Update the mesh positions to match those in the position buffer
-	SyncMeshWithPositionBuffer(Mesh);
 	ScheduleRemeshPass();
 	LastBrushPosLocal = NewBrushPosLocal;
+	return true;
 }
 
 
-void UDynamicMeshSculptTool::ApplyPullSharpKelvinBrush(const FRay& WorldRay)
+bool UDynamicMeshSculptTool::ApplyPullSharpKelvinBrush(const FRay& WorldRay)
 {
 	UpdateBrushPositionOnActivePlane(WorldRay);
 	FVector3d BrushNormalLocal = CurTargetTransform.InverseTransformNormal(LastBrushPosNormalWorld);
@@ -1092,7 +1132,7 @@ void UDynamicMeshSculptTool::ApplyPullSharpKelvinBrush(const FRay& WorldRay)
 	if (MoveVec.SquaredLength() <= 0)
 	{
 		LastBrushPosLocal = NewBrushPosLocal;
-		return;
+		return false;
 	}
 
 	FDynamicMesh3* Mesh = DynamicMeshComponent->GetMesh();
@@ -1104,33 +1144,28 @@ void UDynamicMeshSculptTool::ApplyPullSharpKelvinBrush(const FRay& WorldRay)
 	KelvinletBrushOpProperties.Direction = FVector(MoveVec.X, MoveVec.Y, MoveVec.Z);  //FVector(BrushNormalLocal.X, BrushNormalLocal.Y, BrushNormalLocal.Z);
 	KelvinletBrushOpProperties.Size *= 0.6;
 
-
 	FMatrix ToBrush; ToBrush.SetIdentity();  ToBrush.SetOrigin(-FVector(NewBrushPosLocal.X, NewBrushPosLocal.Y, NewBrushPosLocal.Z)); //  ToBrush.
 
 	KelvinBrushOp.ApplyBrush(KelvinletBrushOpProperties, ToBrush, VertexROI, ROIPositionBuffer);
 
-	// Update the mesh positions to match those in the position buffer
-	SyncMeshWithPositionBuffer(Mesh);
 	ScheduleRemeshPass();
 	LastBrushPosLocal = NewBrushPosLocal;
+	return true;
 }
 
-void UDynamicMeshSculptTool::ApplyTwistKelvinBrush(const FRay& WorldRay)
+bool UDynamicMeshSculptTool::ApplyTwistKelvinBrush(const FRay& WorldRay)
 {
 	UpdateBrushPositionOnTargetMesh(WorldRay, true);
 	
 	FVector3d BrushNormalLocal = CurTargetTransform.InverseTransformNormal(LastBrushPosNormalWorld);
 	FVector3d NewBrushPosLocal = CurTargetTransform.InverseTransformPosition(LastBrushPosWorld);
 	
-
-
 	double Direction = (bInvert) ? -1.0 : 1.0;
 	double UseSpeed  = Direction * FMathd::Sqrt(CurrentBrushRadius) * (SculptProperties->PrimaryBrushSpeed) * ActivePressure;
 
 	FDynamicMesh3* Mesh = DynamicMeshComponent->GetMesh();
 	FKelvinletBrushOp KelvinBrushOp(*Mesh);
 
-
 	FKelvinletBrushOp::FKelvinletBrushOpProperties  KelvinletBrushOpProperties(EKelvinletBrushMode::TwistKelvinlet, *KelvinBrushProperties, *BrushProperties);
 	KelvinletBrushOpProperties.Direction = UseSpeed * FVector(BrushNormalLocal.X, BrushNormalLocal.Y, BrushNormalLocal.Z); // twist about local normal
 	KelvinletBrushOpProperties.Size *= 0.35; // reduce the core size of this brush.
@@ -1139,20 +1174,18 @@ void UDynamicMeshSculptTool::ApplyTwistKelvinBrush(const FRay& WorldRay)
 
 	KelvinBrushOp.ApplyBrush(KelvinletBrushOpProperties, ToBrush, VertexROI, ROIPositionBuffer);
 
-	// Update the mesh positions to match those in the position buffer
-	SyncMeshWithPositionBuffer(Mesh);
 	ScheduleRemeshPass();
 	LastBrushPosLocal = NewBrushPosLocal;
+	return true;
 }
 
-void UDynamicMeshSculptTool::ApplyScaleKelvinBrush(const FRay& WorldRay)
+bool UDynamicMeshSculptTool::ApplyScaleKelvinBrush(const FRay& WorldRay)
 {
 	UpdateBrushPositionOnSculptMesh(WorldRay, true);
 
 	FVector3d NewBrushPosLocal = CurTargetTransform.InverseTransformPosition(LastBrushPosWorld);
 	FVector3d BrushNormalLocal = CurTargetTransform.InverseTransformNormal(LastBrushPosNormalWorld);
 	FVector3d OffsetBrushPosLocal = NewBrushPosLocal - BrushProperties->Depth * CurrentBrushRadius * BrushNormalLocal;
-	
 
 	double Direction = (bInvert) ? -1.0 : 1.0;
 	double UseSpeed = Direction * FMath::Sqrt(CurrentBrushRadius) * SculptProperties->PrimaryBrushSpeed * 0.025 * ActivePressure; ; 
@@ -1160,7 +1193,6 @@ void UDynamicMeshSculptTool::ApplyScaleKelvinBrush(const FRay& WorldRay)
 	FDynamicMesh3* Mesh = DynamicMeshComponent->GetMesh();
 	FKelvinletBrushOp KelvinBrushOp(*Mesh);
 
-
 	FKelvinletBrushOp::FKelvinletBrushOpProperties  KelvinletBrushOpProperties(EKelvinletBrushMode::ScaleKelvinlet, *KelvinBrushProperties, *BrushProperties);
 	KelvinletBrushOpProperties.Direction = FVector(UseSpeed, 0., 0.); // it is a bit iffy, but we only use the first component for the scale
 	KelvinletBrushOpProperties.Size *= 0.35;
@@ -1169,10 +1201,9 @@ void UDynamicMeshSculptTool::ApplyScaleKelvinBrush(const FRay& WorldRay)
 
 	KelvinBrushOp.ApplyBrush(KelvinletBrushOpProperties, ToBrush, VertexROI, ROIPositionBuffer);
 
-	// Update the mesh positions to match those in the position buffer
-	SyncMeshWithPositionBuffer(Mesh);
 	ScheduleRemeshPass();
 	LastBrushPosLocal = NewBrushPosLocal;
+	return true;
 }
 
 int UDynamicMeshSculptTool::FindHitSculptMeshTriangle(const FRay3d& LocalRay)
@@ -1375,6 +1406,12 @@ void UDynamicMeshSculptTool::OnEndDrag(const FRay& Ray)
 
 	// close change record
 	EndChange();
+
+	// destroy active remesher. Should we do this every stroke?? need to do it on undo/redo...
+	if (ActiveRemesher)
+	{
+		ActiveRemesher = nullptr;
+	}
 }
 
 
@@ -1437,17 +1474,16 @@ void UDynamicMeshSculptTool::Render(IToolsContextRenderAPI* RenderAPI)
 
 void UDynamicMeshSculptTool::OnTick(float DeltaTime)
 {
-	SCOPE_CYCLE_COUNTER(STAT_SculptToolTick);
+	TRACE_CPUPROFILER_EVENT_SCOPE(DynamicMeshSculptTool_OnTick);
 
 	ActivePressure = GetCurrentDevicePressure();
 
-	// Allow a tick to pass between application of brush stamps. Bizarrely this
-	// improves responsiveness in the Editor...
-	static int TICK_SKIP_HACK = 0;
-	if (TICK_SKIP_HACK++ % 2 == 0)
-	{
-		return;
-	}
+	// Allow a tick to pass between application of brush stamps. Bizarrely this improves responsiveness in the Editor...
+	//static int TICK_SKIP_HACK = 0;
+	//if (TICK_SKIP_HACK++ % 2 == 0)
+	//{
+	//	return;
+	//}
 
 	ShowWireframeWatcher.CheckAndUpdate();
 	MaterialModeWatcher.CheckAndUpdate();
@@ -1491,15 +1527,38 @@ void UDynamicMeshSculptTool::OnTick(float DeltaTime)
 	bool bROIUpdatePending = false;
 	bool bOctreeUpdatePending = false;
 
+
+	TFuture<void> InitializeRemesher;
 	if (bStampPending)
 	{
-		SCOPE_CYCLE_COUNTER(STAT_SculptTool_Tick_ApplyStampBlock);
+		// if we don't have an active remesher for this brush stroke, create one
+		if (ActiveRemesher == nullptr)
+		{
+			InitializeRemesher = Async(DynamicSculptToolAsyncExecTarget, [&]()
+			{
+				InitializeActiveRemesher();
+			});
+		}
+
+		// initialize ROI for current brush position
+		// TODO: does this break move brush??
+		UpdateBrushPosition(PendingStampRay);
+		FVector3d BrushPos = CurTargetTransform.InverseTransformPosition(LastBrushPosWorld);
+		UpdateROI(BrushPos);
 
-		SCOPE_CYCLE_COUNTER(SculptTool_Tick_ApplyStamp_Remove);
-		Octree->NotifyPendingModification(TriangleROI);		// to mark initial positions
 
+		// once we know ROI we can speculatively start initializing remesh ROI
+		TFuture<void> PrecomputeRemeshROI = Async(DynamicSculptToolAsyncExecTarget, [&]()
+		{
+			// make sure our remesher is initialized
+			InitializeRemesher.Wait();
+			// initialize the ROI
+			PrecomputeRemesherROI();
+		});
+
+		// apply brush stamp to ROI
 		ApplyStamp(PendingStampRay);
-		bStampPending = false;
+		bStampPending = (bInDrag) ? true : false;
 
 		bNormalUpdatePending = true;
 		bROIUpdatePending = true;
@@ -1507,66 +1566,61 @@ void UDynamicMeshSculptTool::OnTick(float DeltaTime)
 		bMeshShapeModified = true;
 		bOctreeUpdatePending = true;
 
-		// flow
-		if (bInDrag)
+		if (bRemeshPending)
 		{
-			bStampPending = true;
-		}
+			check(bInDrag == true);    // this would break undo otherwise!
 
-		SCOPE_CYCLE_COUNTER(SculptTool_Tick_ApplyStamp_Insert);
+			// make sure our remesher is initialized
+			InitializeRemesher.Wait();
+			// make sure our ROI is computed
+			PrecomputeRemeshROI.Wait();
 
-		//UpdateOctreeFuture = Async(DynamicSculptToolAsyncExecTarget, [&]() {
-		//	Octree->ReinsertTriangles(TriangleROI);
-		//});
-	}
-
-	//UpdateOctreeFuture.Wait();
-
-	if (bRemeshPending)
-	{
-		SCOPE_CYCLE_COUNTER(STAT_SculptTool_Tick_RemeshBlock);
-
-		check(bInDrag == true);    // this would break undo otherwise!
+			// remesh the ROI (this removes all ROI triangles from octree)
+			if (ActiveRemesher)
+			{
+				RemeshROIPass_ActiveRemesher(true);
+			}
+			else
+			{
+				check(false);		// broken now
+				RemeshROIPass();
+			}
 
-		// only removes from octree
-		RemeshROIPass();
+			// accumulate new triangles into TriangleROI
+			for (int32 tid : RemeshFinalTriangleROI)
+			{
+				TriangleROI.Add(tid);
+			}
 
-		// accumulate new triangles into TriangleROI
-		for (int32 tid : RemeshFinalTriangleROI)
-		{
-			TriangleROI.Add(tid);
+			bMeshModified = true;
+			bMeshShapeModified = true;
+			bRemeshPending = false;
+			bNormalUpdatePending = true;
+			bROIUpdatePending = true;
+			bOctreeUpdatePending = true;
+			bHaveRemeshed = true;
 		}
-
-		bMeshModified = true;
-		bMeshShapeModified = true;
-		bRemeshPending = false;
-		bNormalUpdatePending = true;
-		bROIUpdatePending = true;
-		bOctreeUpdatePending = true;
-		bHaveRemeshed = true;
 	}
+	check(bRemeshPending == false);		// should never happen...
 
-
-	// launch octree update
+	// launch octree update that inserts/reinserts all triangles in TriangleROI
 	TFuture<void> UpdateOctreeFuture;
 	if (bOctreeUpdatePending)
 	{
-		SCOPE_CYCLE_COUNTER(STAT_SculptTool_Remesh_Finish);
 		// reinsert new ROI into octree
 		UpdateOctreeFuture = Async(DynamicSculptToolAsyncExecTarget, [&]() {
+			TRACE_CPUPROFILER_EVENT_SCOPE(DynamicMeshSculptTool_OctreeReinsert);
 			Octree->ReinsertTriangles(TriangleROI);
 			//Octree->CheckValidity(EValidityCheckFailMode::Check, false, true);
 		});
 		bOctreeUpdatePending = false;
 	}
 
-	//UpdateOctreeFuture.Wait();
-
 	//Octree->CheckValidity(EValidityCheckFailMode::Check, false, true);
 
 	if (bNormalUpdatePending)
 	{
-		SCOPE_CYCLE_COUNTER(STAT_SculptTool_Tick_NormalsBlock);
+		TRACE_CPUPROFILER_EVENT_SCOPE(DynamicMeshSculptTool_UpdateNormals);
 
 		if (Mesh->HasAttributes() && Mesh->Attributes()->PrimaryNormals() != nullptr)
 		{
@@ -1583,36 +1637,23 @@ void UDynamicMeshSculptTool::OnTick(float DeltaTime)
 	// next steps need to wait for octree update to finish
 	UpdateOctreeFuture.Wait();
 
-	TFuture<void> ROIUpdateFuture;
-	if (bROIUpdatePending)
-	{
-		ROIUpdateFuture = Async(DynamicSculptToolAsyncExecTarget, [&]() 
-		{
-			UpdateROI(LastBrushPosLocal);
-		});
-		bROIUpdatePending = false;
-	}
-
 	// launch async target update task
 	if (bTargetDirty)
 	{
-		UpdateTarget();		// TODO: does this make sense to do as a full mesh copy??
+		UpdateTarget();
 		bTargetDirty = false;
 	}
 
+	// update render data
 	if (bMeshModified)
 	{
-		SCOPE_CYCLE_COUNTER(STAT_SculptTool_Tick_UpdateMeshBlock);
-
+		TRACE_CPUPROFILER_EVENT_SCOPE(DynamicMeshSculptTool_UpdateRenderMesh);
 		DynamicMeshComponent->NotifyMeshUpdated();
 		GetToolManager()->PostInvalidation();
 
 		bMeshModified = false;
 	}
 
-
-	// wait for ROI update to complete (can defer but it will immediately happen on hover...)
-	ROIUpdateFuture.Wait();
 }
 
 void UDynamicMeshSculptTool::PrecomputeRemeshInfo()
@@ -1653,6 +1694,31 @@ void UDynamicMeshSculptTool::ScheduleRemeshPass()
 	}
 }
 
+
+
+
+
+
+/**
+ * This is an internal class we will use to just hold onto a FSubRegionRemesher instance
+ * between brush stamps. Perhaps does not need to exist.
+ */
+class FPersistentStampRemesher
+{
+public:
+	FDynamicMesh3* Mesh;
+	TUniquePtr<FSubRegionRemesher> Remesher;
+
+	FPersistentStampRemesher(FDynamicMesh3* MeshIn)
+	{
+		this->Mesh = MeshIn;
+		Remesher = MakeUnique<FSubRegionRemesher>(MeshIn);
+	}
+};
+
+
+
+
 /*
         Split	Collapse	Vertices Pinned	Flip
 Fixed	FALSE	FALSE	TRUE	FALSE
@@ -1727,23 +1793,10 @@ void UDynamicMeshSculptTool::ConfigureRemesher(FSubRegionRemesher& Remesher)
 	Remesher.bEnableSplits = RemeshProperties->bSplits;
 	Remesher.bPreventNormalFlips = RemeshProperties->bPreventNormalFlips;
 
-	{
-		SCOPE_CYCLE_COUNTER(STAT_SculptTool_Remesh_Setup);
-		for (int VertIdx : VertexROI)
-		{
-			Remesher.VertexROI.Add(VertIdx);
-		}
-		Remesher.InitializeFromVertexROI();
-		Remesher.UpdateROI();	// required to use roi in constraints fn below
-	}
-
 	FMeshConstraints Constraints;
 	bool bConstraintAllowSplits = true;
 	{
-		SCOPE_CYCLE_COUNTER(STAT_SculptTool_Remesh_Constraints);
-
-		// TODO: only constrain in ROI. This is quite difficult to do externally because we need to update based on
-		// the changing triangle set in Remesher. Perhaps FSubRegionRemesher should update the constraints itself?
+		TRACE_CPUPROFILER_EVENT_SCOPE(DynamicMeshSculptTool_RemeshROI_Configure_Constraints);
 
 		FMeshConstraintsUtil::ConstrainAllBoundariesAndSeams(Constraints, *Mesh,
 															 (EEdgeRefineFlags)RemeshProperties->MeshBoundaryConstraint,
@@ -1755,13 +1808,122 @@ void UDynamicMeshSculptTool::ConfigureRemesher(FSubRegionRemesher& Remesher)
 	}
 }
 
+
+void UDynamicMeshSculptTool::InitializeRemesherROI(FSubRegionRemesher& Remesher)
+{
+	TRACE_CPUPROFILER_EVENT_SCOPE(DynamicMeshSculptTool_RemeshROI_Configure_InitializeROI);
+	Remesher.SetInitialVertexROI(this->VertexROI);
+	Remesher.InitializeFromVertexROI();
+}
+
+
+
+void UDynamicMeshSculptTool::InitializeActiveRemesher()
+{
+	TRACE_CPUPROFILER_EVENT_SCOPE(DynamicMeshSculptTool_InitializeActiveRemesher);
+
+	FDynamicMesh3* Mesh = DynamicMeshComponent->GetMesh();
+	ActiveRemesher = MakeShared<FPersistentStampRemesher>(Mesh);
+	ConfigureRemesher(* ActiveRemesher->Remesher);
+}
+
+void UDynamicMeshSculptTool::PrecomputeRemesherROI()
+{
+	FSubRegionRemesher& Remesher = *ActiveRemesher->Remesher;
+	Remesher.Reset();
+	InitializeRemesherROI(Remesher);
+}
+
+
+void UDynamicMeshSculptTool::RemeshROIPass_ActiveRemesher(bool bHasPrecomputedROI)
+{
+	TRACE_CPUPROFILER_EVENT_SCOPE(DynamicMeshSculptTool_RemeshROIActive);
+
+	FDynamicMesh3* Mesh = DynamicMeshComponent->GetMesh();
+	FDynamicMeshOctree3* Octree = DynamicMeshComponent->GetOctree();
+
+	FSubRegionRemesher& Remesher = *ActiveRemesher->Remesher;
+	
+	if (bHasPrecomputedROI == false)
+	{
+		PrecomputeRemesherROI();
+	}
+
+	// remove initial triangles
+	{
+		TRACE_CPUPROFILER_EVENT_SCOPE(DynamicMeshSculptTool_RemeshROI_CopyROIStep);
+		RemeshRemovedTriangles = Remesher.GetCurrentTriangleROI();
+	}
+	{
+		TRACE_CPUPROFILER_EVENT_SCOPE(DynamicMeshSculptTool_RemeshROI_OctreeRemoveStep);
+		Octree->RemoveTriangles(RemeshRemovedTriangles, true);
+	}
+
+	if (ActiveMeshChange != nullptr)
+	{
+		Remesher.SetMeshChangeTracker(ActiveMeshChange);
+	}
+
+	bool bIsUniformSmooth = (Remesher.SmoothType == FRemesher::ESmoothTypes::Uniform);
+	for (int k = 0; k < RemeshProperties->Iterations; ++k)
+	{
+		if ((bIsUniformSmooth == false) && (k > 1))
+		{
+			Remesher.bEnableFlips = false;
+		}
+
+		{
+			TRACE_CPUPROFILER_EVENT_SCOPE(DynamicMeshSculptTool_RemeshROI_UpdateStep);
+
+			Remesher.UpdateROI();
+
+			if (ActiveMeshChange != nullptr)
+			{
+				// [TODO] would like to only save vertices here, as triangles will be saved by Remesher as necessary.
+				// However currently FDynamicMeshChangeTracker cannot independently save vertices, only vertices 
+				// that are part of saved triangles will be included in the output FDynamicMeshChange
+				Remesher.SaveActiveROI(ActiveMeshChange);
+				//ActiveMeshChange->VerifySaveState();    // useful for debugging
+			}
+
+			Remesher.BeginTrackRemovedTrisInPass();
+		}
+
+		{
+			TRACE_CPUPROFILER_EVENT_SCOPE(DynamicMeshSculptTool_RemeshROI_RemeshStep);
+			Remesher.BasicRemeshPass();
+		}
+
+		{
+			TRACE_CPUPROFILER_EVENT_SCOPE(DynamicMeshSculptTool_RemeshROI_OctreeStep);
+			const TSet<int32>& TrisRemovedInPass = Remesher.EndTrackRemovedTrisInPass();
+			Octree->RemoveTriangles(TrisRemovedInPass);
+			for (int32 tid : TrisRemovedInPass)
+			{
+				RemeshRemovedTriangles.Add(tid);
+			}
+		}
+	}
+	//UE_LOG(LogTemp, Warning, TEXT("Triangle Count %d after update"), Mesh->TriangleCount());
+
+	RemeshFinalTriangleROI = Remesher.ExtractFinalTriangleROI();
+}
+
+
+
 void UDynamicMeshSculptTool::RemeshROIPass()
 {
+	TRACE_CPUPROFILER_EVENT_SCOPE(DynamicMeshSculptTool_RemeshROI);
+
 	FDynamicMesh3* Mesh = DynamicMeshComponent->GetMesh();
 	FDynamicMeshOctree3* Octree = DynamicMeshComponent->GetOctree();
 
 	FSubRegionRemesher Remesher(Mesh);
-	ConfigureRemesher(Remesher);
+	{
+		TRACE_CPUPROFILER_EVENT_SCOPE(DynamicMeshSculptTool_RemeshROI_Configure);
+		ConfigureRemesher(Remesher);
+		InitializeRemesherROI(Remesher);
+	}
 
 	// remove initial triangles
 	RemeshRemovedTriangles = Remesher.GetCurrentTriangleROI();
@@ -1781,7 +1943,7 @@ void UDynamicMeshSculptTool::RemeshROIPass()
 		}
 
 		{
-			SCOPE_CYCLE_COUNTER(STAT_SculptTool_Remesh_RemeshROIUpdate);
+			TRACE_CPUPROFILER_EVENT_SCOPE(DynamicMeshSculptTool_RemeshROI_UpdateStep);
 
 			Remesher.UpdateROI();
 
@@ -1798,12 +1960,11 @@ void UDynamicMeshSculptTool::RemeshROIPass()
 		}
 
 		{
-			SCOPE_CYCLE_COUNTER(STAT_SculptTool_Remesh_RemeshPass);
+			TRACE_CPUPROFILER_EVENT_SCOPE(DynamicMeshSculptTool_RemeshROI_RemeshStep);
 			Remesher.BasicRemeshPass();
 		}
 
 		{
-			SCOPE_CYCLE_COUNTER(STAT_SculptTool_Remesh_PassOctreeUpdate);
 			const TSet<int32>& TrisRemovedInPass = Remesher.EndTrackRemovedTrisInPass();
 			Octree->RemoveTriangles(TrisRemovedInPass);
 			for (int32 tid : TrisRemovedInPass)
@@ -1820,7 +1981,6 @@ void UDynamicMeshSculptTool::RemeshROIPass()
 void UDynamicMeshSculptTool::RecalculateNormals_PerVertex(const TSet<int32>& Triangles)
 {
 	FDynamicMesh3* Mesh = DynamicMeshComponent->GetMesh();
-	//FDynamicMeshOctree3* Octree = DynamicMeshComponent->GetOctree();
 
 	int MaxVertexID = Mesh->MaxVertexID();
 	if (NormalsVertexFlags.Num() < MaxVertexID)
@@ -1829,11 +1989,6 @@ void UDynamicMeshSculptTool::RecalculateNormals_PerVertex(const TSet<int32>& Tri
 	}
 
 	{
-		SCOPE_CYCLE_COUNTER(SculptTool_Normals_Collect);
-
-		//TrianglesBuffer.Reset();
-		//Octree->RangeQuery(Octree->ModifiedBounds, TrianglesBuffer);
-
 		NormalsBuffer.Reset();
 		for (int TriangleID : Triangles)
 		{
@@ -1853,11 +2008,7 @@ void UDynamicMeshSculptTool::RecalculateNormals_PerVertex(const TSet<int32>& Tri
 		}
 	}
 
-	//UE_LOG(LogTemp, Warning, TEXT("Computing %d normals"), NormalsBuffer.Num());
-
 	{
-		SCOPE_CYCLE_COUNTER(SculptTool_Normals_Compute);
-
 		ParallelFor(NormalsBuffer.Num(), [&](int k) {
 			int vid = NormalsBuffer[k];
 			FVector3d NewNormal = FMeshNormals::ComputeVertexNormal(*Mesh, vid);
@@ -1870,7 +2021,6 @@ void UDynamicMeshSculptTool::RecalculateNormals_PerVertex(const TSet<int32>& Tri
 void UDynamicMeshSculptTool::RecalculateNormals_Overlay(const TSet<int32>& Triangles)
 {
 	FDynamicMesh3* Mesh = DynamicMeshComponent->GetMesh();
-	//FDynamicMeshOctree3* Octree = DynamicMeshComponent->GetOctree();
 	FDynamicMeshNormalOverlay* Normals = Mesh->HasAttributes() ? Mesh->Attributes()->PrimaryNormals() : nullptr;
 	check(Normals != nullptr);
 
@@ -1880,13 +2030,7 @@ void UDynamicMeshSculptTool::RecalculateNormals_Overlay(const TSet<int32>& Trian
 		NormalsVertexFlags.Init(false, MaxElementID * 2);
 	}
 
-
 	{
-		SCOPE_CYCLE_COUNTER(SculptTool_Normals_Collect);
-
-		//TrianglesBuffer.Reset();
-		//Octree->RangeQuery(Octree->ModifiedBounds, TrianglesBuffer);
-
 		NormalsBuffer.Reset();
 		for (int TriangleID : Triangles)
 		{
@@ -1906,11 +2050,7 @@ void UDynamicMeshSculptTool::RecalculateNormals_Overlay(const TSet<int32>& Trian
 		}
 	}
 
-	//UE_LOG(LogTemp, Warning, TEXT("Computing %d normals"), NormalsBuffer.Num());
-
 	{
-		SCOPE_CYCLE_COUNTER(SculptTool_Normals_Compute);
-
 		ParallelFor(NormalsBuffer.Num(), [&](int k) {
 			int elemid = NormalsBuffer[k];
 			FVector3d NewNormal = FMeshNormals::ComputeOverlayNormal(*Mesh, Normals, elemid);
@@ -1918,9 +2058,9 @@ void UDynamicMeshSculptTool::RecalculateNormals_Overlay(const TSet<int32>& Trian
 			NormalsVertexFlags[elemid] = false;
 		});
 	}
-
 }
 
+
 void UDynamicMeshSculptTool::UpdateTarget()
 {
 	if (SculptProperties != nullptr )
@@ -1935,12 +2075,15 @@ void UDynamicMeshSculptTool::UpdateTarget()
 	// TODO: could have a second set of target meshes, and swap between them while we update the other one
 	PendingTargetUpdate = Async(DynamicSculptToolAsyncExecTarget, [this]()
 	{
+		TRACE_CPUPROFILER_EVENT_SCOPE(DynamicMeshSculptTool_UpdateTarget);
 		BrushTargetMesh.Copy(*DynamicMeshComponent->GetMesh(), false, false, false, false);
 
 		TFuture<void> TargetSpatialUpdate = Async(DynamicSculptToolAsyncExecTarget, [this]() {
+			TRACE_CPUPROFILER_EVENT_SCOPE(DynamicMeshSculptTool_UpdateTarget_Spatial);
 			BrushTargetMeshSpatial.SetMesh(&BrushTargetMesh, true);
 		});
 		TFuture<void> TargetNormalsUpdate = Async(DynamicSculptToolAsyncExecTarget, [this]() {
+			TRACE_CPUPROFILER_EVENT_SCOPE(DynamicMeshSculptTool_UpdateTarget_Normals);
 			BrushTargetNormals.SetMesh(&BrushTargetMesh);
 			BrushTargetNormals.ComputeVertexNormals();
 		});
@@ -2265,15 +2408,6 @@ void UDynamicMeshSculptTool::SaveActiveROI()
 	}
 }
 
-void UDynamicMeshSculptTool::UpdateSavedVertex(int vid, const FVector3d& OldPosition, const FVector3d& NewPosition)
-{
-	if (ActiveVertexChange != nullptr)
-	{
-		UpdateSavedVertexLock.Lock();
-		ActiveVertexChange->UpdateVertex(vid, OldPosition, NewPosition);
-		UpdateSavedVertexLock.Unlock();
-	}
-}
 
 void UDynamicMeshSculptTool::UpdateMaterialMode(EMeshEditingMaterialModes MaterialMode)
 {
diff --git a/Engine/Plugins/Experimental/MeshModelingToolset/Source/MeshModelingTools/Public/DynamicMeshSculptTool.h b/Engine/Plugins/Experimental/MeshModelingToolset/Source/MeshModelingTools/Public/DynamicMeshSculptTool.h
--- a/Engine/Plugins/Experimental/MeshModelingToolset/Source/MeshModelingTools/Public/DynamicMeshSculptTool.h
+++ b/Engine/Plugins/Experimental/MeshModelingToolset/Source/MeshModelingTools/Public/DynamicMeshSculptTool.h
@@ -20,35 +20,18 @@
 #include "TransformTypes.h"
 #include "Sculpting/MeshSculptToolBase.h"
 #include "Async/Async.h"
+#include "Util/UniqueIndexSet.h"
 #include "DynamicMeshSculptTool.generated.h"
 
 class UTransformGizmo;
 class UTransformProxy;
 class UMaterialInstanceDynamic;
 
-DECLARE_STATS_GROUP(TEXT("SculptTool"), STATGROUP_SculptTool, STATCAT_Advanced);
-DECLARE_CYCLE_STAT(TEXT("SculptTool_UpdateROI"), SculptTool_UpdateROI, STATGROUP_SculptTool);
-DECLARE_CYCLE_STAT(TEXT("SculptTool_ApplyStamp"), STAT_SculptToolApplyStamp, STATGROUP_SculptTool );
-DECLARE_CYCLE_STAT(TEXT("SculptTool_Tick"), STAT_SculptToolTick, STATGROUP_SculptTool);
-DECLARE_CYCLE_STAT(TEXT("SculptTool_Tick_ApplyStampBlock"), STAT_SculptTool_Tick_ApplyStampBlock, STATGROUP_SculptTool);
-DECLARE_CYCLE_STAT(TEXT("SculptTool_Tick_ApplyStamp_Remove"), SculptTool_Tick_ApplyStamp_Remove, STATGROUP_SculptTool);
-DECLARE_CYCLE_STAT(TEXT("SculptTool_Tick_ApplyStamp_Insert"), SculptTool_Tick_ApplyStamp_Insert, STATGROUP_SculptTool);
-DECLARE_CYCLE_STAT(TEXT("SculptTool_Tick_RemeshBlock"), STAT_SculptTool_Tick_RemeshBlock, STATGROUP_SculptTool);
-DECLARE_CYCLE_STAT(TEXT("SculptTool_Tick_NormalsBlock"), STAT_SculptTool_Tick_NormalsBlock, STATGROUP_SculptTool);
-DECLARE_CYCLE_STAT(TEXT("SculptTool_Tick_UpdateMeshBlock"), STAT_SculptTool_Tick_UpdateMeshBlock, STATGROUP_SculptTool);
-DECLARE_CYCLE_STAT(TEXT("SculptTool_Normals_Collect"), SculptTool_Normals_Collect, STATGROUP_SculptTool);
-DECLARE_CYCLE_STAT(TEXT("SculptTool_Normals_Compute"), SculptTool_Normals_Compute, STATGROUP_SculptTool);
-DECLARE_CYCLE_STAT(TEXT("SculptTool_Remesh_1Setup"), STAT_SculptTool_Remesh_Setup, STATGROUP_SculptTool);
-DECLARE_CYCLE_STAT(TEXT("SculptTool_Remesh_2Constraints"), STAT_SculptTool_Remesh_Constraints, STATGROUP_SculptTool);
-DECLARE_CYCLE_STAT(TEXT("SculptTool_Remesh_3RemeshROIUpdate"), STAT_SculptTool_Remesh_RemeshROIUpdate, STATGROUP_SculptTool);
-DECLARE_CYCLE_STAT(TEXT("SculptTool_Remesh_4RemeshPass"), STAT_SculptTool_Remesh_RemeshPass, STATGROUP_SculptTool);
-DECLARE_CYCLE_STAT(TEXT("SculptTool_Remesh_5PassOctreeUpdate"), STAT_SculptTool_Remesh_PassOctreeUpdate, STATGROUP_SculptTool);
-DECLARE_CYCLE_STAT(TEXT("SculptTool_Remesh_6Finish"), STAT_SculptTool_Remesh_Finish, STATGROUP_SculptTool);
-
 class FMeshVertexChangeBuilder;
 class FDynamicMeshChangeTracker;
 class UPreviewMesh;
 class FSubRegionRemesher;
+class FPersistentStampRemesher;
 
 /** Mesh Sculpting Brush Types */
 UENUM()
@@ -198,6 +181,9 @@ public:
 	/** Control the amount of simplification during sculpting. Higher values will avoid wiping out fine details on the mesh. */
 	UPROPERTY(EditAnywhere, Category = Remeshing, meta = (UIMin = "0", UIMax = "5", ClampMin = "0", ClampMax = "5", DisplayPriority = 3))
 	int PreserveDetail = 0;
+
+	UPROPERTY(EditAnywhere, Category = Remeshing, AdvancedDisplay)
+	int Iterations = 5;
 };
 
 
@@ -358,7 +344,14 @@ private:
 	bool bEnableRemeshing;
 	double InitialEdgeLength;
 	void ScheduleRemeshPass();
-	void ConfigureRemesher(FSubRegionRemesher&);
+	void ConfigureRemesher(FSubRegionRemesher& Remesher);
+	void InitializeRemesherROI(FSubRegionRemesher& Remesher);
+
+	TSharedPtr<FPersistentStampRemesher> ActiveRemesher;
+	void InitializeActiveRemesher();
+	void PrecomputeRemesherROI();
+	void RemeshROIPass_ActiveRemesher(bool bHasPrecomputedROI);
+
 
 	bool bInDrag;
 
@@ -372,9 +365,12 @@ private:
 	FVector3d LastSmoothBrushPosLocal;
 	int32 LastBrushTriangleID = -1;
 
+	TArray<int> UpdateROITriBuffer;
+	FUniqueIndexSet VertexROIBuilder;
 	TArray<int> VertexROI;
-	TSet<int> VertexSetBuffer;
+	FUniqueIndexSet TriangleROIBuilder;
 	TSet<int> TriangleROI;
+	//TSet<int> TriangleROI;
 	void UpdateROI(const FVector3d& BrushPos);
 
 	bool bRemeshPending;
@@ -412,20 +408,20 @@ private:
 	bool UpdateBrushPositionOnSculptMesh(const FRay& WorldRay, bool bFallbackToViewPlane);
 	void AlignBrushToView();
 
-	void ApplySmoothBrush(const FRay& WorldRay);
-	void ApplyMoveBrush(const FRay& WorldRay);
-	void ApplyOffsetBrush(const FRay& WorldRay, bool bUseViewDirection);
-	void ApplySculptMaxBrush(const FRay& WorldRay);
-	void ApplyPinchBrush(const FRay& WorldRay);
-	void ApplyInflateBrush(const FRay& WorldRay);
-	void ApplyPlaneBrush(const FRay& WorldRay);
-	void ApplyFixedPlaneBrush(const FRay& WorldRay);
-	void ApplyFlattenBrush(const FRay& WorldRay);
-	void ApplyResampleBrush(const FRay& WorldRay);
-	void ApplyPullKelvinBrush(const FRay& WorldRay);
-	void ApplyPullSharpKelvinBrush(const FRay& WorldRay);
-	void ApplyTwistKelvinBrush(const FRay& WorldRay);
-	void ApplyScaleKelvinBrush(const FRay& WorldRay);
+	bool ApplySmoothBrush(const FRay& WorldRay);
+	bool ApplyMoveBrush(const FRay& WorldRay);
+	bool ApplyOffsetBrush(const FRay& WorldRay, bool bUseViewDirection);
+	bool ApplySculptMaxBrush(const FRay& WorldRay);
+	bool ApplyPinchBrush(const FRay& WorldRay);
+	bool ApplyInflateBrush(const FRay& WorldRay);
+	bool ApplyPlaneBrush(const FRay& WorldRay);
+	bool ApplyFixedPlaneBrush(const FRay& WorldRay);
+	bool ApplyFlattenBrush(const FRay& WorldRay);
+	bool ApplyResampleBrush(const FRay& WorldRay);
+	bool ApplyPullKelvinBrush(const FRay& WorldRay);
+	bool ApplyPullSharpKelvinBrush(const FRay& WorldRay);
+	bool ApplyTwistKelvinBrush(const FRay& WorldRay);
+	bool ApplyScaleKelvinBrush(const FRay& WorldRay);
 
 	double SculptMaxFixedHeight = -1.0;
 
@@ -436,7 +432,6 @@ private:
 	FFrame3d ActiveFixedBrushPlane;
 	FFrame3d ComputeROIBrushPlane(const FVector3d& BrushCenter, bool bIgnoreDepth, bool bViewAligned);
 
-	TArray<int> TrianglesBuffer;
 	TArray<int> NormalsBuffer;
 	TArray<bool> NormalsVertexFlags;
 	void RecalculateNormals_PerVertex(const TSet<int32>& Triangles);
@@ -456,9 +451,6 @@ private:
 	void EndChange();
 	void SaveActiveROI();
 
-	void UpdateSavedVertex(int vid, const FVector3d& OldPosition, const FVector3d& NewPosition);
-	FCriticalSection UpdateSavedVertexLock;
-
 	double EstimateIntialSafeTargetLength(const FDynamicMesh3& Mesh, int MinTargetTriCount);
 
 	TArray<EDynamicMeshSculptBrushType> BrushTypeHistory;
diff --git a/Engine/Plugins/Experimental/MeshModelingToolset/Source/ModelingComponents/Private/OctreeDynamicMeshComponent.cpp b/Engine/Plugins/Experimental/MeshModelingToolset/Source/ModelingComponents/Private/OctreeDynamicMeshComponent.cpp
--- a/Engine/Plugins/Experimental/MeshModelingToolset/Source/ModelingComponents/Private/OctreeDynamicMeshComponent.cpp
+++ b/Engine/Plugins/Experimental/MeshModelingToolset/Source/ModelingComponents/Private/OctreeDynamicMeshComponent.cpp
@@ -55,7 +55,7 @@ void UOctreeDynamicMeshComponent::InitializeMesh(FMeshDescription* MeshDescripti
 
 	FAxisAlignedBox3d MeshBounds = Mesh->GetCachedBounds();
 	Octree = MakeUnique<FDynamicMeshOctree3>();
-	Octree->RootDimension = MeshBounds.MaxDim() * 0.5;
+	Octree->RootDimension = MeshBounds.MaxDim() * 0.25;
 
 	Octree->Initialize(GetMesh());
 	
@@ -115,7 +115,7 @@ void UOctreeDynamicMeshComponent::ApplyTransform(const FTransform3d& Transform,
 	{
 		FAxisAlignedBox3d MeshBounds = Mesh->GetCachedBounds();
 		Octree = MakeUnique<FDynamicMeshOctree3>();
-		Octree->RootDimension = MeshBounds.MaxDim() * 0.5;
+		Octree->RootDimension = MeshBounds.MaxDim() * 0.25;
 		Octree->Initialize(GetMesh());
 		OctreeCut = MakeUnique<FDynamicMeshOctree3::FTreeCutSet>();
 	}
diff --git a/Engine/Plugins/Experimental/MeshModelingToolset/Source/ModelingOperators/Private/CleaningOps/RemeshMeshOp.cpp b/Engine/Plugins/Experimental/MeshModelingToolset/Source/ModelingOperators/Private/CleaningOps/RemeshMeshOp.cpp
--- a/Engine/Plugins/Experimental/MeshModelingToolset/Source/ModelingOperators/Private/CleaningOps/RemeshMeshOp.cpp
+++ b/Engine/Plugins/Experimental/MeshModelingToolset/Source/ModelingOperators/Private/CleaningOps/RemeshMeshOp.cpp
@@ -47,6 +47,8 @@ void FRemeshMeshOp::SetTransform(const FTransform& Transform)
 
 void FRemeshMeshOp::CalculateResult(FProgressCancel* Progress)
 {
+	TRACE_CPUPROFILER_EVENT_SCOPE(RemeshMeshOp);
+
 	if (Progress && Progress->Cancelled())
 	{
 		return;
@@ -129,42 +131,48 @@ void FRemeshMeshOp::CalculateResult(FProgressCancel* Progress)
 		TargetMesh->DiscardAttributes();
 	}
 
-	if (RemeshType == ERemeshType::FullPass)
 	{
-		// Run a fixed number of iterations
-		for (int k = 0; k < RemeshIterations; ++k)
+		TRACE_CPUPROFILER_EVENT_SCOPE(RemeshMeshOp_Remesh);
+		if (RemeshType == ERemeshType::FullPass)
 		{
-			// If we are not uniform smoothing, then flips seem to often make things worse.
-			// Possibly this is because without the tangential flow, we won't get to the nice tris.
-			// In this case we are better off basically not flipping, and just letting collapses resolve things
-			// regular-valence polygons - things stay "stuck". 
-			// @todo try implementing edge-length flip criteria instead of valence-flip
-			if (bIsUniformSmooth == false)
+			// Run a fixed number of iterations
+			for (int k = 0; k < RemeshIterations; ++k)
 			{
-				bool bUseFlipsThisPass = (k % 2 == 0 && k < RemeshIterations / 2);
-				Remesher->bEnableFlips = bUseFlipsThisPass && bFlips;
+				// If we are not uniform smoothing, then flips seem to often make things worse.
+				// Possibly this is because without the tangential flow, we won't get to the nice tris.
+				// In this case we are better off basically not flipping, and just letting collapses resolve things
+				// regular-valence polygons - things stay "stuck". 
+				// @todo try implementing edge-length flip criteria instead of valence-flip
+				if (bIsUniformSmooth == false)
+				{
+					bool bUseFlipsThisPass = (k % 2 == 0 && k < RemeshIterations / 2);
+					Remesher->bEnableFlips = bUseFlipsThisPass && bFlips;
+				}
+
+				Remesher->BasicRemeshPass();
 			}
-
+		}
+		else if (RemeshType == ERemeshType::Standard || RemeshType == ERemeshType::NormalFlow)
+		{
+			// Run to convergence
 			Remesher->BasicRemeshPass();
 		}
-	}
-	else if (RemeshType == ERemeshType::Standard || RemeshType == ERemeshType::NormalFlow)
-	{
-		// Run to convergence
-		Remesher->BasicRemeshPass();
-	}
-	else 
-	{
-		check(!"Encountered unexpected Remesh Type");
+		else 
+		{
+			check(!"Encountered unexpected Remesh Type");
+		}
 	}
 
-	if (!TargetMesh->HasAttributes())
-	{
-		FMeshNormals::QuickComputeVertexNormals(*TargetMesh);
-	}
-	else
 	{
-		FMeshNormals::QuickRecomputeOverlayNormals(*TargetMesh);
+		TRACE_CPUPROFILER_EVENT_SCOPE(RemeshMeshOp_Normals);
+		if (!TargetMesh->HasAttributes())
+		{
+			FMeshNormals::QuickComputeVertexNormals(*TargetMesh);
+		}
+		else
+		{
+			FMeshNormals::QuickRecomputeOverlayNormals(*TargetMesh);
+		}
 	}
 
 }
