ModelingComponents: Add USimpleDynamicMeshComponent:: FastNotifyTriangleVerticesUpdated_TryPrecompute() and FastNotifyTriangleVerticesUpdated_ApplyPrecompute(). These are split-up pieces of the existing FastNotifyTriangleVerticesUpdated(), which are useful in cases where we can do the precompute Async() while other pending computation happens (eg in sculpting). Change bounding-box computations to use new parallel mesh bounds option.
#rnx

diff --git a/Engine/Plugins/Experimental/MeshModelingToolset/Source/ModelingComponents/Private/SimpleDynamicMeshComponent.cpp b/Engine/Plugins/Experimental/MeshModelingToolset/Source/ModelingComponents/Private/SimpleDynamicMeshComponent.cpp
--- a/Engine/Plugins/Experimental/MeshModelingToolset/Source/ModelingComponents/Private/SimpleDynamicMeshComponent.cpp
+++ b/Engine/Plugins/Experimental/MeshModelingToolset/Source/ModelingComponents/Private/SimpleDynamicMeshComponent.cpp
@@ -34,6 +34,17 @@
 #include "SimpleDynamicMeshSceneProxy.h"
 
 
+namespace
+{
+	// probably should be something defined for the whole tool framework...
+#if WITH_EDITOR
+	static EAsyncExecution SimpleDynamicMeshComponentAsyncExecTarget = EAsyncExecution::LargeThreadPool;
+#else
+	static EAsyncExecution SimpleDynamicMeshComponentAsyncExecTarget = EAsyncExecution::ThreadPool;
+#endif
+}
+
+
 
 
 USimpleDynamicMeshComponent::USimpleDynamicMeshComponent(const FObjectInitializer& ObjectInitializer)
@@ -248,7 +259,7 @@ void USimpleDynamicMeshComponent::ResetProxy()
 
 	// Need to recreate scene proxy to send it over
 	MarkRenderStateDirty();
-	LocalBounds = Mesh->GetCachedBounds();
+	LocalBounds = Mesh->GetBounds(true);
 	UpdateBounds();
 
 	if (TangentsType != EDynamicMeshTangentCalcType::ExternallyCalculated)
@@ -316,10 +327,18 @@ void USimpleDynamicMeshComponent::FastNotifyPositionsUpdated(bool bNormals, bool
 	FSimpleDynamicMeshSceneProxy* Proxy = GetCurrentSceneProxy();
 	if (Proxy)
 	{
+		// calculate bounds while we are updating vertices
+		TFuture<void> UpdateBoundsCalc;
+		UpdateBoundsCalc = Async(SimpleDynamicMeshComponentAsyncExecTarget, [&]()
+		{
+			TRACE_CPUPROFILER_EVENT_SCOPE(SimpleDynamicMeshComponent_FastPositionsUpdate_AsyncBoundsUpdate);
+			LocalBounds = Mesh->GetBounds(true);
+		});
+
 		GetCurrentSceneProxy()->FastUpdateVertices(true, bNormals, bColors, bUVs);
 		//MarkRenderDynamicDataDirty();
 		MarkRenderTransformDirty();
-		LocalBounds = Mesh->GetCachedBounds();
+		UpdateBoundsCalc.Wait();
 		UpdateBounds();
 	}
 	else
@@ -367,6 +386,18 @@ void USimpleDynamicMeshComponent::FastNotifyVertexAttributesUpdated(EMeshRenderA
 	if (Proxy && ensure(UpdatedAttributes != EMeshRenderAttributeFlags::None))
 	{
 		bool bPositions = (UpdatedAttributes & EMeshRenderAttributeFlags::Positions) != EMeshRenderAttributeFlags::None;
+
+		// calculate bounds while we are updating vertices
+		TFuture<void> UpdateBoundsCalc;
+		if (bPositions)
+		{
+			UpdateBoundsCalc = Async(SimpleDynamicMeshComponentAsyncExecTarget, [&]()
+			{
+				TRACE_CPUPROFILER_EVENT_SCOPE(SimpleDynamicMeshComponent_FastVertexAttribUpdate_AsyncBoundsUpdate);
+				LocalBounds = Mesh->GetBounds(true);
+			});
+		}
+
 		GetCurrentSceneProxy()->FastUpdateVertices(bPositions,
 			(UpdatedAttributes & EMeshRenderAttributeFlags::VertexNormals) != EMeshRenderAttributeFlags::None,
 			(UpdatedAttributes & EMeshRenderAttributeFlags::VertexColors) != EMeshRenderAttributeFlags::None,
@@ -375,7 +406,7 @@ void USimpleDynamicMeshComponent::FastNotifyVertexAttributesUpdated(EMeshRenderA
 		if (bPositions)
 		{
 			MarkRenderTransformDirty();
-			LocalBounds = Mesh->GetCachedBounds();
+			UpdateBoundsCalc.Wait();
 			UpdateBounds();
 		}
 	}
@@ -443,18 +474,38 @@ void USimpleDynamicMeshComponent::FastNotifyTriangleVerticesUpdated(const TArray
 	}
 	else
 	{
+		// compute list of sets to update
 		TArray<int32> UpdatedSets;
-		for (int32 tid : Triangles)
 		{
-			int32 SetID = Decomposition->GetGroupForTriangle(tid);
-			UpdatedSets.AddUnique(SetID);
+			TRACE_CPUPROFILER_EVENT_SCOPE(SimpleDynamicMeshComponent_FastVertexUpdate_FindSets);
+			for (int32 tid : Triangles)
+			{
+				int32 SetID = Decomposition->GetGroupForTriangle(tid);
+				UpdatedSets.AddUnique(SetID);
+			}
 		}
 
 		bool bPositions = (UpdatedAttributes & EMeshRenderAttributeFlags::Positions) != EMeshRenderAttributeFlags::None;
-		Proxy->FastUpdateVertices(UpdatedSets, bPositions,
-			(UpdatedAttributes & EMeshRenderAttributeFlags::VertexNormals) != EMeshRenderAttributeFlags::None,
-			(UpdatedAttributes & EMeshRenderAttributeFlags::VertexColors) != EMeshRenderAttributeFlags::None,
-			(UpdatedAttributes & EMeshRenderAttributeFlags::VertexUVs) != EMeshRenderAttributeFlags::None);
+
+		// calculate bounds while we are updating vertices
+		TFuture<void> UpdateBoundsCalc;
+		if (bPositions)
+		{
+			UpdateBoundsCalc = Async(SimpleDynamicMeshComponentAsyncExecTarget, [&]()
+			{
+				TRACE_CPUPROFILER_EVENT_SCOPE(SimpleDynamicMeshComponent_FastVertexUpdate_AsyncBoundsUpdate);
+				LocalBounds = Mesh->GetBounds(true);
+			});
+		}
+
+		// update the render buffers
+		{
+			TRACE_CPUPROFILER_EVENT_SCOPE(SimpleDynamicMeshComponent_FastVertexUpdate_ApplyUpdate);
+			Proxy->FastUpdateVertices(UpdatedSets, bPositions,
+				(UpdatedAttributes & EMeshRenderAttributeFlags::VertexNormals) != EMeshRenderAttributeFlags::None,
+				(UpdatedAttributes & EMeshRenderAttributeFlags::VertexColors) != EMeshRenderAttributeFlags::None,
+				(UpdatedAttributes & EMeshRenderAttributeFlags::VertexUVs) != EMeshRenderAttributeFlags::None);
+		}
 
 		if (bUpdateSecondarySort)
 		{
@@ -463,8 +514,9 @@ void USimpleDynamicMeshComponent::FastNotifyTriangleVerticesUpdated(const TArray
 
 		if (bPositions)
 		{
+			TRACE_CPUPROFILER_EVENT_SCOPE(SimpleDynamicMeshComponent_FastVertexUpdate_FinalPositionsUpdate);
 			MarkRenderTransformDirty();
-			LocalBounds = Mesh->GetCachedBounds();
+			UpdateBoundsCalc.Wait();
 			UpdateBounds();
 		}
 	}
@@ -472,6 +524,7 @@ void USimpleDynamicMeshComponent::FastNotifyTriangleVerticesUpdated(const TArray
 
 
 
+
 void USimpleDynamicMeshComponent::FastNotifyTriangleVerticesUpdated(const TSet<int32>& Triangles, EMeshRenderAttributeFlags UpdatedAttributes)
 {
 	// should not be using fast paths if we have to run mesh postprocessor
@@ -500,48 +553,53 @@ void USimpleDynamicMeshComponent::FastNotifyTriangleVerticesUpdated(const TSet<i
 	}
 	else
 	{
+		// compute list of sets to update
 		TArray<int32> UpdatedSets;
-		for (int32 tid : Triangles)
-		{
-			int32 SetID = Decomposition->GetGroupForTriangle(tid);
-			UpdatedSets.AddUnique(SetID);
-		}
-
-		int32 TotalTris = 0;
-		for (int32 SetID : UpdatedSets)
 		{
-			TotalTris += Decomposition->GetGroup(SetID).Triangles.Num();
+			TRACE_CPUPROFILER_EVENT_SCOPE(SimpleDynamicMeshComponent_FastVertexUpdate_FindSets);
+			for (int32 tid : Triangles)
+			{
+				int32 SetID = Decomposition->GetGroupForTriangle(tid);
+				UpdatedSets.AddUnique(SetID);
+			}
 		}
 
-		//UE_LOG(LogTemp, Warning, TEXT("Updating %d groups with %d tris"), UpdatedSets.Num(), TotalTris);
-
 		bool bPositions = (UpdatedAttributes & EMeshRenderAttributeFlags::Positions) != EMeshRenderAttributeFlags::None;
 
 		// calculate bounds while we are updating vertices
 		TFuture<void> UpdateBoundsCalc;
 		if (bPositions)
 		{
-			UpdateBoundsCalc = Async(EAsyncExecution::ThreadPool, [&]()
+			UpdateBoundsCalc = Async(SimpleDynamicMeshComponentAsyncExecTarget, [&]()
 			{
-				LocalBounds = Mesh->GetCachedBounds();
+				TRACE_CPUPROFILER_EVENT_SCOPE(SimpleDynamicMeshComponent_FastVertexUpdate_AsyncBoundsUpdate);
+				LocalBounds = Mesh->GetBounds(true);
 			});
 		}
 
-		GetCurrentSceneProxy()->FastUpdateVertices(UpdatedSets, bPositions,
-			(UpdatedAttributes & EMeshRenderAttributeFlags::VertexNormals) != EMeshRenderAttributeFlags::None,
-			(UpdatedAttributes & EMeshRenderAttributeFlags::VertexColors) != EMeshRenderAttributeFlags::None,
-			(UpdatedAttributes & EMeshRenderAttributeFlags::VertexUVs) != EMeshRenderAttributeFlags::None);
+		// update the render buffers
+		{
+			TRACE_CPUPROFILER_EVENT_SCOPE(SimpleDynamicMeshComponent_FastVertexUpdate_ApplyUpdate);
+			Proxy->FastUpdateVertices(UpdatedSets, bPositions,
+				(UpdatedAttributes & EMeshRenderAttributeFlags::VertexNormals) != EMeshRenderAttributeFlags::None,
+				(UpdatedAttributes & EMeshRenderAttributeFlags::VertexColors) != EMeshRenderAttributeFlags::None,
+				(UpdatedAttributes & EMeshRenderAttributeFlags::VertexUVs) != EMeshRenderAttributeFlags::None);
+		}
 
-		if (bUpdateSecondarySort)
 		{
-			Proxy->FastUpdateIndexBuffers(UpdatedSets);
+			TRACE_CPUPROFILER_EVENT_SCOPE(SimpleDynamicMeshComponent_FastVertexUpdate_UpdateIndexBuffers);
+			if (bUpdateSecondarySort)
+			{
+				Proxy->FastUpdateIndexBuffers(UpdatedSets);
+			}
 		}
 
+		// finish up, have to wait for background bounds recalculation here
 		if (bPositions)
 		{
+			TRACE_CPUPROFILER_EVENT_SCOPE(SimpleDynamicMeshComponent_FastVertexUpdate_FinalPositionsUpdate);
 			MarkRenderTransformDirty();
 			UpdateBoundsCalc.Wait();
-			//LocalBounds = Mesh->GetCachedBounds();
 			UpdateBounds();
 		}
 	}
@@ -549,6 +607,130 @@ void USimpleDynamicMeshComponent::FastNotifyTriangleVerticesUpdated(const TSet<i
 
 
 
+/**
+ * Compute the combined bounding-box of the Triangles array in parallel, by computing
+ * partial boxes for subsets of this array, and then combining those boxes.
+ * TODO: this should move to a pulbic utility function, and possibly the block-based ParallelFor
+ * should be refactored out into something more general, as this pattern is useful in many places...
+ */
+static FAxisAlignedBox3d ParallelComputeROIBounds(const FDynamicMesh3& Mesh, const TArray<int32>& Triangles)
+{
+	FAxisAlignedBox3d FinalBounds = FAxisAlignedBox3d::Empty();
+	FCriticalSection FinalBoundsLock;
+	int32 N = Triangles.Num();
+	constexpr int32 BlockSize = 4096;
+	int32 Blocks = (N / BlockSize) + 1;
+	ParallelFor(Blocks, [&](int bi)
+	{
+		FAxisAlignedBox3d BlockBounds = FAxisAlignedBox3d::Empty();
+		for (int32 k = 0; k < BlockSize; ++k)
+		{
+			int32 i = bi * BlockSize + k;
+			if (i < N)
+			{
+				int32 tid = Triangles[i];
+				const FIndex3i& TriV = Mesh.GetTriangleRef(tid);
+				BlockBounds.Contain(Mesh.GetVertexRef(TriV.A));
+				BlockBounds.Contain(Mesh.GetVertexRef(TriV.B));
+				BlockBounds.Contain(Mesh.GetVertexRef(TriV.C));
+			}
+		}
+		FinalBoundsLock.Lock();
+		FinalBounds.Contain(BlockBounds);
+		FinalBoundsLock.Unlock();
+	});
+	return FinalBounds;
+}
+
+
+
+TFuture<bool> USimpleDynamicMeshComponent::FastNotifyTriangleVerticesUpdated_TryPrecompute(
+	const TArray<int32>& Triangles,
+	TArray<int32>& UpdateSetsOut,
+	FAxisAlignedBox3d& BoundsOut)
+{
+	if ((!!RenderMeshPostProcessor) || (GetCurrentSceneProxy() == nullptr) || (!Decomposition))
+	{
+		// is there a simpler way to do this? cannot seem to just make a TFuture<bool>...
+		return Async(SimpleDynamicMeshComponentAsyncExecTarget, [&]() { return false; });
+	}
+
+	return Async(SimpleDynamicMeshComponentAsyncExecTarget, [&]() 
+	{
+		TFuture<void> ComputeBounds = Async(SimpleDynamicMeshComponentAsyncExecTarget, [this, &BoundsOut, &Triangles]()
+		{
+			TRACE_CPUPROFILER_EVENT_SCOPE(SimpleDynamicMeshComponent_FastVertexUpdatePrecomp_CalcBounds);
+			BoundsOut = ParallelComputeROIBounds(*Mesh, Triangles);
+		});
+
+		TFuture<void> ComputeSets = Async(SimpleDynamicMeshComponentAsyncExecTarget, [this, &UpdateSetsOut, &Triangles]()
+		{
+			TRACE_CPUPROFILER_EVENT_SCOPE(SimpleDynamicMeshComponent_FastVertexUpdatePrecomp_FindSets);
+			UpdateSetsOut.Reset();
+			for (int32 tid : Triangles)
+			{
+				int32 SetID = Decomposition->GetGroupForTriangle(tid);
+				UpdateSetsOut.AddUnique(SetID);
+			}
+			UE_LOG(LogTemp, Warning, TEXT("FOUND %d dirty sets of %d"), UpdateSetsOut.Num(), Decomposition->Num());
+		});
+
+		ComputeSets.Wait();
+		ComputeBounds.Wait();
+		return true;
+	});
+}
+
+
+void USimpleDynamicMeshComponent::FastNotifyTriangleVerticesUpdated_ApplyPrecompute(
+	const TArray<int32>& Triangles,
+	EMeshRenderAttributeFlags UpdatedAttributes, 
+	TFuture<bool>& Precompute, 
+	const TArray<int32>& UpdateSets, 
+	const FAxisAlignedBox3d& UpdateSetBounds)
+{
+	Precompute.Wait();
+
+	bool bPrecomputeOK = Precompute.Get();
+	if (bPrecomputeOK == false || GetCurrentSceneProxy() == nullptr )
+	{
+		FastNotifyTriangleVerticesUpdated(Triangles, UpdatedAttributes);
+		return;
+	}
+
+	FSimpleDynamicMeshSceneProxy* Proxy = GetCurrentSceneProxy();
+	bool bPositions = (UpdatedAttributes & EMeshRenderAttributeFlags::Positions) != EMeshRenderAttributeFlags::None;
+	bool bUpdateSecondarySort = (SecondaryTriFilterFunc) &&
+		((UpdatedAttributes & EMeshRenderAttributeFlags::SecondaryIndexBuffers) != EMeshRenderAttributeFlags::None);
+
+	{
+		TRACE_CPUPROFILER_EVENT_SCOPE(SimpleDynamicMeshComponent_FastVertexUpdate_ApplyUpdate);
+		Proxy->FastUpdateVertices(UpdateSets, bPositions,
+			(UpdatedAttributes & EMeshRenderAttributeFlags::VertexNormals) != EMeshRenderAttributeFlags::None,
+			(UpdatedAttributes & EMeshRenderAttributeFlags::VertexColors) != EMeshRenderAttributeFlags::None,
+			(UpdatedAttributes & EMeshRenderAttributeFlags::VertexUVs) != EMeshRenderAttributeFlags::None);
+	}
+
+	{
+		TRACE_CPUPROFILER_EVENT_SCOPE(SimpleDynamicMeshComponent_FastVertexUpdate_UpdateIndexBuffers);
+		if (bUpdateSecondarySort)
+		{
+			Proxy->FastUpdateIndexBuffers(UpdateSets);
+		}
+	}
+
+	if (bPositions)
+	{
+		TRACE_CPUPROFILER_EVENT_SCOPE(SimpleDynamicMeshComponent_FastVertexUpdate_FinalPositionsUpdate);
+		MarkRenderTransformDirty();
+		LocalBounds.Contain(UpdateSetBounds);
+		UpdateBounds();
+	}
+}
+
+
+
+
 
 FPrimitiveSceneProxy* USimpleDynamicMeshComponent::CreateSceneProxy()
 {
diff --git a/Engine/Plugins/Experimental/MeshModelingToolset/Source/ModelingComponents/Private/SimpleDynamicMeshSceneProxy.h b/Engine/Plugins/Experimental/MeshModelingToolset/Source/ModelingComponents/Private/SimpleDynamicMeshSceneProxy.h
--- a/Engine/Plugins/Experimental/MeshModelingToolset/Source/ModelingComponents/Private/SimpleDynamicMeshSceneProxy.h
+++ b/Engine/Plugins/Experimental/MeshModelingToolset/Source/ModelingComponents/Private/SimpleDynamicMeshSceneProxy.h
@@ -433,6 +433,8 @@ public:
 	 */
 	void FastUpdateVertices(const TArray<int32>& WhichBuffers, bool bPositions, bool bNormals, bool bColors, bool bUVs)
 	{
+		TRACE_CPUPROFILER_EVENT_SCOPE(SimpleDynamicMeshProxy_FastUpdateVertices);
+
 		// skip if we have no updates
 		if (bPositions == false && bNormals == false && bColors == false && bUVs == false)
 		{
diff --git a/Engine/Plugins/Experimental/MeshModelingToolset/Source/ModelingComponents/Public/SimpleDynamicMeshComponent.h b/Engine/Plugins/Experimental/MeshModelingToolset/Source/ModelingComponents/Public/SimpleDynamicMeshComponent.h
--- a/Engine/Plugins/Experimental/MeshModelingToolset/Source/ModelingComponents/Public/SimpleDynamicMeshComponent.h
+++ b/Engine/Plugins/Experimental/MeshModelingToolset/Source/ModelingComponents/Public/SimpleDynamicMeshComponent.h
@@ -8,6 +8,7 @@
 #include "Drawing/MeshRenderDecomposition.h"
 #include "MeshTangents.h"
 #include "TransformTypes.h"
+#include "Async/Future.h"
 
 #include "SimpleDynamicMeshComponent.generated.h"
 
@@ -189,6 +190,26 @@ public:
 	void FastNotifyTriangleVerticesUpdated(const TSet<int32>& Triangles, EMeshRenderAttributeFlags UpdatedAttributes);
 
 
+	/**
+	 * If a Decomposition is set on this Component, and everything is currently valid (proxy/etc), precompute the set of
+	 * buffers that will be modified, as well as the bounds of the modified region. These are both computed in parallel.
+	 * @return a future that will (eventually) return true if the precompute is OK, and (immediately) false if it is not
+	 */
+	TFuture<bool> FastNotifyTriangleVerticesUpdated_TryPrecompute(const TArray<int32>& Triangles, TArray<int32>& UpdateSetsOut, FAxisAlignedBox3d& BoundsOut);
+
+
+	/**
+	 * This function updates vertex positions/attributes of existing SceneProxy render buffers if possible, for the given triangles.
+	 * The assumption is that FastNotifyTriangleVerticesUpdated_TryPrecompute() was used to get the Precompute future, this function
+	 * will Wait() until it is done and then use the UpdateSets and UpdateSetBounds that were computed (must be the same variables
+	 * passed to FastNotifyTriangleVerticesUpdated_TryPrecompute). 
+	 * If the Precompute future returns false, then we forward the call to FastNotifyTriangleVerticesUpdated(), which will do more work.
+	 */
+	void FastNotifyTriangleVerticesUpdated_ApplyPrecompute(const TArray<int32>& Triangles, EMeshRenderAttributeFlags UpdatedAttributes,
+		TFuture<bool>& Precompute, const TArray<int32>& UpdateSets, const FAxisAlignedBox3d& UpdateSetBounds);
+
+
+
 	/** If false, we don't completely invalidate the RenderProxy when ApplyChange() is called (assumption is it will be handled elsewhere) */
 	UPROPERTY()
 	bool bInvalidateProxyOnChange = true;
