MeshModelingTools: Allowing the drag alignment mechanic to include invisible components in hit test, and allowing it to work with interval gizmos.

#rnx

diff --git a/Engine/Plugins/Experimental/MeshModelingToolset/Source/ModelingComponents/Private/Mechanics/DragAlignmentMechanic.cpp b/Engine/Plugins/Experimental/MeshModelingToolset/Source/ModelingComponents/Private/Mechanics/DragAlignmentMechanic.cpp
--- a/Engine/Plugins/Experimental/MeshModelingToolset/Source/ModelingComponents/Private/Mechanics/DragAlignmentMechanic.cpp
+++ b/Engine/Plugins/Experimental/MeshModelingToolset/Source/ModelingComponents/Private/Mechanics/DragAlignmentMechanic.cpp
@@ -3,6 +3,7 @@
 #include "Mechanics/DragAlignmentMechanic.h"
 
 #include "BaseBehaviors/KeyAsModifierInputBehavior.h"
+#include "BaseGizmos/IntervalGizmo.h"
 #include "BaseGizmos/RepositionableTransformGizmo.h"
 #include "BaseGizmos/TransformGizmo.h"
 #include "BaseGizmos/TransformProxy.h"
@@ -24,7 +25,8 @@ void UDragAlignmentMechanic::Setup(UInteractiveTool* ParentToolIn)
 
 	// Set up the function that casts rays into the world.
 	WorldRayCast = [this, ParentToolIn](const FRay& WorldRay, FHitResult& HitResult, 
-		const TArray<const UPrimitiveComponent*>* ComponentsToIgnore)
+		const TArray<const UPrimitiveComponent*>* ComponentsToIgnore,
+		const TArray<const UPrimitiveComponent*>* InvisibleComponentsToInclude)
 	{
 		bool bSuccess = false;
 		FVector3d QueryPoint = WorldRay.PointAt(1);
@@ -41,6 +43,7 @@ void UDragAlignmentMechanic::Setup(UInteractiveTool* ParentToolIn)
 		SnapParams.VisualAngleThreshold = VisualAngleSnapThreshold;
 		SnapParams.SnapGeometryOut = &SnapGeometry;
 		SnapParams.ComponentsToIgnore = ComponentsToIgnore;
+		SnapParams.InvisibleComponentsToInclude = InvisibleComponentsToInclude;
 
 		if (ToolSceneQueriesUtil::FindSceneSnapPoint(SnapParams))
 		{
@@ -65,21 +68,27 @@ void UDragAlignmentMechanic::Shutdown()
 {
 }
 
-void UDragAlignmentMechanic::AddToGizmo(UTransformGizmo* TransformGizmo, const TArray<const UPrimitiveComponent*>* ComponentsToIgnoreInAlignment)
+void UDragAlignmentMechanic::AddToGizmo(UTransformGizmo* TransformGizmo, const TArray<const UPrimitiveComponent*>* ComponentsToIgnoreInAlignment, 
+	const TArray<const UPrimitiveComponent*>* InvisibleComponentsToIncludeInAlignment)
 {
-	// If we have components to ignore, we need a copy of the array so that the alignment
+	// If we have components to ignore/include, we need a copy of the array so that the alignment
 	// functions can use them later.
 	TSharedPtr<TArray<const UPrimitiveComponent*>> ComponentsToIgnorePersistent;
 	if (ComponentsToIgnoreInAlignment)
 	{
 		ComponentsToIgnorePersistent = MakeShared<TArray<const UPrimitiveComponent*>>(*ComponentsToIgnoreInAlignment);
 	}
+	TSharedPtr<TArray<const UPrimitiveComponent*>> ComponentsToIncludePersistent;
+	if (InvisibleComponentsToIncludeInAlignment)
+	{
+		ComponentsToIncludePersistent = MakeShared<TArray<const UPrimitiveComponent*>>(*InvisibleComponentsToIncludeInAlignment);
+	}
 
 	// Set the alignment functions.
 	TransformGizmo->SetWorldAlignmentFunctions(
 		[this]() { return bAlignmentToggle; },
-		[this, ComponentsToIgnorePersistent](const FRay& WorldRay, FVector& OutputPoint) { 
-			return CastRay(WorldRay, OutputPoint, ComponentsToIgnorePersistent.Get(), true); 
+		[this, ComponentsToIgnorePersistent, ComponentsToIncludePersistent](const FRay& WorldRay, FVector& OutputPoint) {
+			return CastRay(WorldRay, OutputPoint, ComponentsToIgnorePersistent.Get(), ComponentsToIncludePersistent.Get(), true);
 		});
 
 	// If we're adding to a repositionable gizmo, we probably don't want to ignore components
@@ -91,8 +100,8 @@ void UDragAlignmentMechanic::AddToGizmo(UTransformGizmo* TransformGizmo, const T
 	{
 		RepositionableGizmo->SetPivotAlignmentFunctions(
 			[this]() { return bAlignmentToggle; },
-			[this](const FRay& WorldRay, FVector& OutputPoint) {
-				return CastRay(WorldRay, OutputPoint, nullptr, false); // don't ignore anything
+			[this, ComponentsToIncludePersistent](const FRay& WorldRay, FVector& OutputPoint) {
+				return CastRay(WorldRay, OutputPoint, nullptr, ComponentsToIncludePersistent.Get(), false); // don't ignore anything
 			});
 	}
 
@@ -116,6 +125,44 @@ void UDragAlignmentMechanic::AddToGizmo(UTransformGizmo* TransformGizmo, const T
 		});
 }
 
+void UDragAlignmentMechanic::AddToGizmo(UIntervalGizmo* IntervalGizmo, const TArray<const UPrimitiveComponent*>* ComponentsToIgnoreInAlignment,
+	const TArray<const UPrimitiveComponent*>* InvisibleComponentsToIncludeInAlignment)
+{
+	// If we have components to ignore/include, we need a copy of the array so that the alignment
+	// functions can use them later.
+	TSharedPtr<TArray<const UPrimitiveComponent*>> ComponentsToIgnorePersistent;
+	if (ComponentsToIgnoreInAlignment)
+	{
+		ComponentsToIgnorePersistent = MakeShared<TArray<const UPrimitiveComponent*>>(*ComponentsToIgnoreInAlignment);
+	}
+	TSharedPtr<TArray<const UPrimitiveComponent*>> ComponentsToIncludePersistent;
+	if (InvisibleComponentsToIncludeInAlignment)
+	{
+		ComponentsToIncludePersistent = MakeShared<TArray<const UPrimitiveComponent*>>(*InvisibleComponentsToIncludeInAlignment);
+	}
+
+	// Set the alignment functions.
+	IntervalGizmo->SetWorldAlignmentFunctions(
+		[this]() { return bAlignmentToggle; },
+		[this, ComponentsToIgnorePersistent, ComponentsToIncludePersistent](const FRay& WorldRay, FVector& OutputPoint) {
+			return CastRay(WorldRay, OutputPoint, ComponentsToIgnorePersistent.Get(), ComponentsToIncludePersistent.Get(), true);
+		});
+
+	// Register listener to help with rendering. It gets us the final location of the interval endpoint
+	// so that we can draw a line from the hit point to that location.
+	IntervalGizmo->OnIntervalChanged.AddWeakLambda(this, [this](UIntervalGizmo* Gizmo, FVector Direction, float NewParam) {
+		FTransform Transform = Gizmo->GetGizmoTransform();
+		Transform.AddToTranslation(Transform.Rotator().RotateVector(Direction) * NewParam);
+		OnGizmoTransformChanged(Transform);
+		});
+
+	// Register a listener to stop drawing once the gizmo is done moving.
+	IntervalGizmo->OnEndIntervalGizmoEdit.AddWeakLambda(this, [this](UIntervalGizmo*) {
+		bPreviewEndpointsValid = false;
+		bWaitingOnProjectedResult = false;
+		});
+}
+
 void UDragAlignmentMechanic::OnGizmoTransformChanged(FTransform NewTransform)
 {
 	// Ignore changes that weren't preceded by a successful query (likely programatic)
@@ -227,7 +274,9 @@ void UDragAlignmentMechanic::InitializeDeformedMeshRayCast(TFunction<FDynamicMes
  * @return true if something was hit.
  */
 bool UDragAlignmentMechanic::CastRay(const FRay& WorldRay, FVector& OutputPoint,
-	const TArray<const UPrimitiveComponent*>* ComponentsToIgnore, bool bUseFilter)
+	const TArray<const UPrimitiveComponent*>* ComponentsToIgnore, 
+	const TArray<const UPrimitiveComponent*>* InvisibleComponentsToInclude,
+	bool bUseFilter)
 {
 	bool bHit = false;
 	FHitResult MeshHitResult;
@@ -239,8 +288,9 @@ bool UDragAlignmentMechanic::CastRay(const FRay& WorldRay, FVector& OutputPoint,
 	}
 
 	FHitResult WorldHitResult;
-	if (WorldRayCast && WorldRayCast(WorldRay, WorldHitResult, ComponentsToIgnore) &&
-		(!bHit || WorldHitResult.Distance < MeshHitResult.Distance))
+	if (WorldRayCast 
+		&& WorldRayCast(WorldRay, WorldHitResult, ComponentsToIgnore, InvisibleComponentsToInclude) 
+		&& (!bHit || WorldHitResult.Distance < MeshHitResult.Distance))
 	{
 		bHit = true;
 		LastHitResult = WorldHitResult;
diff --git a/Engine/Plugins/Experimental/MeshModelingToolset/Source/ModelingComponents/Private/ToolSceneQueriesUtil.cpp b/Engine/Plugins/Experimental/MeshModelingToolset/Source/ModelingComponents/Private/ToolSceneQueriesUtil.cpp
--- a/Engine/Plugins/Experimental/MeshModelingToolset/Source/ModelingComponents/Private/ToolSceneQueriesUtil.cpp
+++ b/Engine/Plugins/Experimental/MeshModelingToolset/Source/ModelingComponents/Private/ToolSceneQueriesUtil.cpp
@@ -176,6 +176,7 @@ bool ToolSceneQueriesUtil::FindSceneSnapPoint(FFindSceneSnapPointParams& Params)
 	Request.Position = (FVector)*Params.Point;
 	Request.VisualAngleThresholdDegrees = UseThreshold;
 	Request.ComponentsToIgnore = Params.ComponentsToIgnore;
+	Request.InvisibleComponentsToInclude = Params.InvisibleComponentsToInclude;
 
 	TArray<FSceneSnapQueryResult> Results;
 	if (QueryAPI->ExecuteSceneSnapQuery(Request, Results))
diff --git a/Engine/Plugins/Experimental/MeshModelingToolset/Source/ModelingComponents/Public/Mechanics/DragAlignmentMechanic.h b/Engine/Plugins/Experimental/MeshModelingToolset/Source/ModelingComponents/Public/Mechanics/DragAlignmentMechanic.h
--- a/Engine/Plugins/Experimental/MeshModelingToolset/Source/ModelingComponents/Public/Mechanics/DragAlignmentMechanic.h
+++ b/Engine/Plugins/Experimental/MeshModelingToolset/Source/ModelingComponents/Public/Mechanics/DragAlignmentMechanic.h
@@ -16,6 +16,7 @@
 
 class FGroupTopologyDeformer;
 class UPrimitiveComponent;
+class UIntervalGizmo;
 class UTransformGizmo;
 
 /**
@@ -36,7 +37,14 @@ public:
 	 * target object(s)). For movement, the gizmo will ignore ComponentsToIgnore, as well as modified triangles if used
 	 * for deformation.
 	 */
-	void AddToGizmo(UTransformGizmo* TransformGizmo, const TArray<const UPrimitiveComponent*>* ComponentsToIgnore = nullptr);
+	void AddToGizmo(UTransformGizmo* TransformGizmo, const TArray<const UPrimitiveComponent*>* ComponentsToIgnore = nullptr,
+		const TArray<const UPrimitiveComponent*>* InvisibleComponentsToInclude = nullptr);
+
+	/**
+	 * Like the other AddToGizmo, but acts on interval gizmos.
+	 */
+	void AddToGizmo(UIntervalGizmo* IntervalGizmo, const TArray<const UPrimitiveComponent*>* ComponentsToIgnoreInAlignment = nullptr,
+		const TArray<const UPrimitiveComponent*>* InvisibleComponentsToInclude = nullptr);
 
 	/**
 	 * Optional- allows the use of the mechanic in mesh deformation. Specifically, alignment will attempt to align
@@ -62,7 +70,9 @@ public:
 
 protected:
 	// These are modifiable in case we someday want to make it easier to customize them further.
-	TUniqueFunction<bool(const FRay& InputRay, FHitResult& OutputHit, const TArray<const UPrimitiveComponent*>* ComponentsToIgnore)> WorldRayCast = nullptr;
+	TUniqueFunction<bool(const FRay& InputRay, FHitResult& OutputHit, 
+		const TArray<const UPrimitiveComponent*>* ComponentsToIgnore,
+		const TArray<const UPrimitiveComponent*>* InvisibleComponentsToInclude)> WorldRayCast = nullptr;
 	TUniqueFunction<bool(const FRay&, FHitResult&, bool bUseFilter)> MeshRayCast = nullptr;
 	TUniqueFunction<bool(int32 Tid)> MeshRayCastTriangleFilter = [](int32 Tid) { return true; };
 
@@ -86,6 +96,8 @@ protected:
 	int32 AlignmentModifierID = 1;
 
 	virtual bool CastRay(const FRay& WorldRay, FVector& OutputPoint,
-		const TArray<const UPrimitiveComponent*>* ComponentsToIgnore, bool bUseFilter);
+		const TArray<const UPrimitiveComponent*>* ComponentsToIgnore, 
+		const TArray<const UPrimitiveComponent*>* InvisibleComponentsToInclude,
+		bool bUseFilter);
 	void OnGizmoTransformChanged(FTransform NewTransform);
 };
\ No newline at end of file
diff --git a/Engine/Plugins/Experimental/MeshModelingToolset/Source/ModelingComponents/Public/ToolSceneQueriesUtil.h b/Engine/Plugins/Experimental/MeshModelingToolset/Source/ModelingComponents/Public/ToolSceneQueriesUtil.h
--- a/Engine/Plugins/Experimental/MeshModelingToolset/Source/ModelingComponents/Public/ToolSceneQueriesUtil.h
+++ b/Engine/Plugins/Experimental/MeshModelingToolset/Source/ModelingComponents/Public/ToolSceneQueriesUtil.h
@@ -101,6 +101,7 @@ namespace ToolSceneQueriesUtil
 		bool bEdges = false; // If true, try to snap to edges
 		double VisualAngleThreshold = 0; // Visual angle threshold to use. If 0, GetDefaultVisualAngleSnapThresh() is used
 		const TArray<const UPrimitiveComponent*>* ComponentsToIgnore = nullptr;
+		const TArray<const UPrimitiveComponent*>* InvisibleComponentsToInclude = nullptr;
 		FSnapGeometry* SnapGeometryOut = nullptr; // world-space position of the snap geometry(point / line / polygon)
 		FVector* DebugTriangleOut = nullptr; // if non - null, triangle containing snap is returned if a snap is found
 	};
diff --git a/Engine/Source/Editor/Experimental/EditorInteractiveToolsFramework/Private/EdModeInteractiveToolsContext.cpp b/Engine/Source/Editor/Experimental/EditorInteractiveToolsFramework/Private/EdModeInteractiveToolsContext.cpp
--- a/Engine/Source/Editor/Experimental/EditorInteractiveToolsFramework/Private/EdModeInteractiveToolsContext.cpp
+++ b/Engine/Source/Editor/Experimental/EditorInteractiveToolsFramework/Private/EdModeInteractiveToolsContext.cpp
@@ -72,7 +72,8 @@ static bool IsVisibleObjectHit_Internal(const FHitResult& HitResult)
 }
 
 static bool FindNearestVisibleObjectHit_Internal(UWorld* World, FHitResult& HitResultOut, const FVector& Start, const FVector& End,
-	bool bIsSceneGeometrySnapQuery, const TArray<const UPrimitiveComponent*>* ComponentsToIgnore)
+	bool bIsSceneGeometrySnapQuery, const TArray<const UPrimitiveComponent*>* ComponentsToIgnore, 
+	const TArray<const UPrimitiveComponent*>* InvisibleComponentsToInclude)
 {
 	FCollisionObjectQueryParams ObjectQueryParams(FCollisionObjectQueryParams::AllObjects);
 	FCollisionQueryParams QueryParams = FCollisionQueryParams::DefaultQueryParam;
@@ -91,10 +92,13 @@ static bool FindNearestVisibleObjectHit_Internal(UWorld* World, FHitResult& HitR
 		if (CurResult.Distance < NearestVisible)
 		{
 			if (IsVisibleObjectHit_Internal(CurResult) 
-				&& (!ComponentsToIgnore || !ComponentsToIgnore->Contains(CurResult.Component.Get())))
+				|| (InvisibleComponentsToInclude && InvisibleComponentsToInclude->Contains(CurResult.GetComponent())))
 			{
-				HitResultOut = CurResult;
-				NearestVisible = CurResult.Distance;
+				if (!ComponentsToIgnore || !ComponentsToIgnore->Contains(CurResult.GetComponent()))
+				{
+					HitResultOut = CurResult;
+					NearestVisible = CurResult.Distance;
+				}
 			}
 		}
 	}
@@ -252,7 +256,7 @@ public:
 		FVector RayEnd = RayStart + HALF_WORLD_MAX * RayDirection;
 		FHitResult HitResult;
 		bool bHitWorld = FindNearestVisibleObjectHit_Internal(EditorModeManager->GetWorld(), HitResult, 
-			RayStart, RayEnd, true, Request.ComponentsToIgnore);
+			RayStart, RayEnd, true, Request.ComponentsToIgnore, Request.InvisibleComponentsToInclude);
 		if (bHitWorld && HitResult.FaceIndex >= 0)
 		{
 			float VisualAngle = OpeningAngleDeg(Request.Position, HitResult.ImpactPoint, RayStart);
diff --git a/Engine/Source/Runtime/Experimental/InteractiveToolsFramework/Private/BaseGizmos/AxisPositionGizmo.cpp b/Engine/Source/Runtime/Experimental/InteractiveToolsFramework/Private/BaseGizmos/AxisPositionGizmo.cpp
--- a/Engine/Source/Runtime/Experimental/InteractiveToolsFramework/Private/BaseGizmos/AxisPositionGizmo.cpp
+++ b/Engine/Source/Runtime/Experimental/InteractiveToolsFramework/Private/BaseGizmos/AxisPositionGizmo.cpp
@@ -106,7 +106,7 @@ void UAxisPositionGizmo::OnClickDrag(const FInputDeviceRay& DragPos)
 	{
 		GizmoMath::NearestPointOnLine(InteractionOrigin, InteractionAxis, HitPoint,
 			InteractionCurPoint, InteractionCurParameter);
-		InteractionCurParameter += InteractionStartAxisOriginParameterOffset;
+		InteractionCurParameter += bCustomDestinationAlignsAxisOrigin ? InteractionStartAxisOriginParameterOffset : 0;
 	}
 	else
 	{
diff --git a/Engine/Source/Runtime/Experimental/InteractiveToolsFramework/Private/BaseGizmos/IntervalGizmo.cpp b/Engine/Source/Runtime/Experimental/InteractiveToolsFramework/Private/BaseGizmos/IntervalGizmo.cpp
--- a/Engine/Source/Runtime/Experimental/InteractiveToolsFramework/Private/BaseGizmos/IntervalGizmo.cpp
+++ b/Engine/Source/Runtime/Experimental/InteractiveToolsFramework/Private/BaseGizmos/IntervalGizmo.cpp
@@ -86,6 +86,45 @@ public:
 	}
 };
 
+/**
+ * This change source doesn't actually issue any valid transactions. Instead, it is a helper class 
+ * that can get attached to the interval gizmo's state target to fire off BeginEditSequence and 
+ * EndEditSequence on the start/end of a drag.
+ */
+class FIntervalGizmoChangeBroadcaster : public IToolCommandChangeSource
+{
+public:
+	FIntervalGizmoChangeBroadcaster(UIntervalGizmo* IntervalGizmoIn) : IntervalGizmo(IntervalGizmoIn) {}
+
+	virtual ~FIntervalGizmoChangeBroadcaster() {}
+
+	TWeakObjectPtr<UIntervalGizmo> IntervalGizmo;
+
+	virtual void BeginChange() override
+	{
+		if (IntervalGizmo.IsValid())
+		{
+			IntervalGizmo->BeginEditSequence();
+		}
+	}
+	virtual TUniquePtr<FToolCommandChange> EndChange() override
+	{
+		if (IntervalGizmo.IsValid())
+		{
+			IntervalGizmo->EndEditSequence();
+		}
+		return TUniquePtr<FToolCommandChange>();
+	}
+	virtual UObject* GetChangeTarget() override
+	{
+		return IntervalGizmo.Get();
+	}
+	virtual FText GetChangeDescription() override
+	{
+		return LOCTEXT("FIntervalGizmoChangeBroadcaster", "IntervalGizmoEdit");
+	}
+};
+
 
 AIntervalGizmoActor::AIntervalGizmoActor()
 {
@@ -160,6 +199,26 @@ void UIntervalGizmo::SetUpdateCoordSystemFunction(TFunction<void(UPrimitiveCompo
 	UpdateCoordSystemFunction = CoordSysFunction;
 }
 
+void UIntervalGizmo::SetWorldAlignmentFunctions(TUniqueFunction<bool()>&& ShouldAlignDestinationIn, TUniqueFunction<bool(const FRay&, FVector&)>&& DestinationAlignmentRayCasterIn)
+{
+	// Save these so that any later gizmo resets (using SetActiveTarget) keep the settings.
+	ShouldAlignDestination = MoveTemp(ShouldAlignDestinationIn);
+	DestinationAlignmentRayCaster = MoveTemp(DestinationAlignmentRayCasterIn);
+
+	for (UInteractiveGizmo* SubGizmo : this->ActiveGizmos)
+	{
+		if (UAxisPositionGizmo* CastGizmo = Cast<UAxisPositionGizmo>(SubGizmo))
+		{
+			CastGizmo->ShouldUseCustomDestinationFunc = [this]() { return ShouldAlignDestination(); };
+			CastGizmo->CustomDestinationFunc =
+				[this](const UAxisPositionGizmo::FCustomDestinationParams& Params, FVector& OutputPoint) {
+				return DestinationAlignmentRayCaster(*Params.WorldRay, OutputPoint);
+			};
+			CastGizmo->bCustomDestinationAlignsAxisOrigin = false; // We're aligning the endpoints of the intervals
+		}
+	}
+}
+
 void UIntervalGizmo::Setup()
 {
 	UInteractiveGizmo::Setup();
@@ -261,6 +320,16 @@ void UIntervalGizmo::SetActiveTarget(UTransformProxy* TransformTargetIn, UGizmoL
 	DownIntervalSource    = DownInterval;
 	ForwardIntervalSource = ForwardInterval;
 
+	// Get the parameter source to notify our delegate of any changes
+	UpIntervalSource->OnParameterChanged.AddWeakLambda(this, [this](IGizmoFloatParameterSource*, FGizmoFloatParameterChange Change) {
+		OnIntervalChanged.Broadcast(this, FVector::ZAxisVector, Change.CurrentValue);
+		});
+	DownIntervalSource->OnParameterChanged.AddWeakLambda(this, [this](IGizmoFloatParameterSource*, FGizmoFloatParameterChange Change) {
+		OnIntervalChanged.Broadcast(this, -FVector::ZAxisVector, -Change.CurrentValue);
+		});
+	ForwardIntervalSource->OnParameterChanged.AddWeakLambda(this, [this](IGizmoFloatParameterSource*, FGizmoFloatParameterChange Change) {
+		OnIntervalChanged.Broadcast(this, FVector::YAxisVector, Change.CurrentValue);
+		});
 
 	USceneComponent* GizmoComponent = GizmoActor->GetRootComponent();
 
@@ -291,9 +360,8 @@ void UIntervalGizmo::SetActiveTarget(UTransformProxy* TransformTargetIn, UGizmoL
 	StateTarget->DependentChangeSources.Add(MakeUnique<FGizmoFloatParameterChangeSource>(DownIntervalSource));
 	StateTarget->DependentChangeSources.Add(MakeUnique<FGizmoFloatParameterChangeSource>(ForwardIntervalSource));
 
-
-	UGizmoTransformProxyTransformSource* TransformSource =
-		UGizmoTransformProxyTransformSource::Construct(TransformProxy, this);
+	// Have the state target notify us of the start/end of drags
+	StateTarget->DependentChangeSources.Add(MakeUnique<FIntervalGizmoChangeBroadcaster>(this));
 
 	// root component provides local Y/Z axis, identified by AxisIndex
 	AxisYSource = UGizmoComponentAxisSource::Construct(GizmoComponent, 1, true, this);
@@ -348,6 +416,10 @@ void UIntervalGizmo::ClearActiveTarget()
 	TransformProxy = nullptr;
 }
 
+FTransform UIntervalGizmo::GetGizmoTransform() const
+{
+	return TransformProxy->GetTransform();
+}
 
 UInteractiveGizmo* UIntervalGizmo::AddIntervalHandleGizmo(
 	USceneComponent* RootComponent,
@@ -380,6 +452,12 @@ UInteractiveGizmo* UIntervalGizmo::AddIntervalHandleGizmo(
 
 	IntervalGizmo->StateTarget = Cast<UObject>(StateTargetIn);
 
+	IntervalGizmo->ShouldUseCustomDestinationFunc = [this]() { return ShouldAlignDestination(); };
+	IntervalGizmo->CustomDestinationFunc =
+		[this](const UAxisPositionGizmo::FCustomDestinationParams& Params, FVector& OutputPoint) {
+		return DestinationAlignmentRayCaster(*Params.WorldRay, OutputPoint);
+	};
+
 	ActiveGizmos.Add(IntervalGizmo);
 
 	return IntervalGizmo;
diff --git a/Engine/Source/Runtime/Experimental/InteractiveToolsFramework/Public/BaseGizmos/AxisPositionGizmo.h b/Engine/Source/Runtime/Experimental/InteractiveToolsFramework/Public/BaseGizmos/AxisPositionGizmo.h
--- a/Engine/Source/Runtime/Experimental/InteractiveToolsFramework/Public/BaseGizmos/AxisPositionGizmo.h
+++ b/Engine/Source/Runtime/Experimental/InteractiveToolsFramework/Public/BaseGizmos/AxisPositionGizmo.h
@@ -106,12 +106,24 @@ public:
 
 	/**
 	 * If ShouldUseCustomDestinationFunc() returns true, this function gets queried to get a destination point.
-	 * The gizmo parameter will then be picked in such a way that the gizmo origin moves to the closest point
-	 * on the axis to the destination point. Can be used, for example, to align to items in the scene.
+	 * The gizmo parameter will then be picked in such a way that the dragged location moves to the closest point
+	 * on the axis to the destination point (optionally offset by the start click position relative the axis origin,
+	 * if bCustomDestinationAlignsAxisOrigin is true, to align the axis origin itself). Can be used, for example, 
+	 * to align to items in the scene.
 	 */
 	TUniqueFunction<bool(const FCustomDestinationParams& WorldRay, FVector& OutputPoint)> CustomDestinationFunc =
 		[](const FCustomDestinationParams& Params, FVector& OutputPoint) { return false; };
 
+	/**
+	 * Only used when a custom destination is obtained from CustomDestinationFunc. When false, the custom destination
+	 * is simply projected back to the axis to get the resulting parameter. This is useful when trying to align the
+	 * dragged point on the gizmo to the destination point.
+	 * If true, the same projection is performed, but the final parameter is adjusted to make the axis origin align
+	 * to the destination instead, rather than the aligning the grabbed point (since the user probably grabbed the
+	 * gizmo some distance away from the axis origin).
+	 */
+	bool bCustomDestinationAlignsAxisOrigin = true;
+
 public:
 	/** If true, we are in an active click+drag interaction, otherwise we are not */
 	UPROPERTY()
diff --git a/Engine/Source/Runtime/Experimental/InteractiveToolsFramework/Public/BaseGizmos/IntervalGizmo.h b/Engine/Source/Runtime/Experimental/InteractiveToolsFramework/Public/BaseGizmos/IntervalGizmo.h
--- a/Engine/Source/Runtime/Experimental/InteractiveToolsFramework/Public/BaseGizmos/IntervalGizmo.h
+++ b/Engine/Source/Runtime/Experimental/InteractiveToolsFramework/Public/BaseGizmos/IntervalGizmo.h
@@ -147,6 +147,16 @@ public:
 	virtual void SetActiveTarget(UTransformProxy* TransformTargetIn, UGizmoLocalFloatParameterSource* UpInterval, UGizmoLocalFloatParameterSource* DownInterval, UGizmoLocalFloatParameterSource* ForwardInterval,
 		                         IToolContextTransactionProvider* TransactionProvider = nullptr);
 
+	/** Sets functions that allow the endpoints of the intervals to be snapped to world geometry when ShouldAlignDestination is true */
+	virtual void SetWorldAlignmentFunctions(
+		TUniqueFunction<bool()>&& ShouldAlignDestination,
+		TUniqueFunction<bool(const FRay&, FVector&)>&& DestinationAlignmentRayCaster
+	);
+
+	/** Notifies listeners that a sequence of edits to the gizmo is beginning/ending (at the start/end of a drag). */
+	virtual void BeginEditSequence() { OnBeginIntervalGizmoEdit.Broadcast(this); }
+	virtual void EndEditSequence() { OnEndIntervalGizmoEdit.Broadcast(this); }
+
 	/**
 	* Clear the parameter sources for this gizmo
 	*/
@@ -157,11 +167,29 @@ public:
 	 */
 	virtual void ClearActiveTarget();
 
+	/**
+	 * Gets the location and orientation of the interval gizmo.
+	 */
+	virtual FTransform GetGizmoTransform() const;
+
 	/** State target is shared across gizmos, and created internally during SetActiveTarget() */
 	UPROPERTY()
 	TObjectPtr<UGizmoTransformChangeStateTarget> StateTarget = nullptr;
 
+	/** 
+	 * Called when an interval is changed. The delegate gives a pointer to the gizmo, the direction of the interval in
+	 * gizmo space, and the new parameter value.
+	 */
+	DECLARE_MULTICAST_DELEGATE_ThreeParams(FOnIntervalChanged, UIntervalGizmo*, const FVector&, float);
+	FOnIntervalChanged OnIntervalChanged;
+
+	/** Called when the gizmo is notified about the start of a sequence of interval changes. */
+	DECLARE_MULTICAST_DELEGATE_OneParam(FOnBeginIntervalEdit, UIntervalGizmo*);
+	FOnBeginIntervalEdit OnBeginIntervalGizmoEdit;
 
+	/** Called when the gizmo is notified about the end of a sequence of interval changes. */
+	DECLARE_MULTICAST_DELEGATE_OneParam(FOnEndIntervalEdit, UIntervalGizmo*);
+	FOnEndIntervalEdit OnEndIntervalGizmoEdit;
 protected:
 
 	/** GizmoActors will be spawned in this World */
@@ -169,7 +197,7 @@ protected:
 
 	AIntervalGizmoActor* GizmoActor = nullptr;
 
-	/** The active target object for the Gizmo */
+	/** The gizmo tracks the location and orientation of the transform in this TransformProxy. */
 	UPROPERTY()
 	TObjectPtr<UTransformProxy> TransformProxy = nullptr;
 
@@ -207,6 +235,8 @@ protected:
 	// of this function via the ToolBuilder
 	TFunction<void(UPrimitiveComponent*, EToolContextCoordinateSystem)> UpdateCoordSystemFunction;
 
+	TUniqueFunction<bool()> ShouldAlignDestination = []() { return false; };
+	TUniqueFunction<bool(const FRay&, FVector&)> DestinationAlignmentRayCaster = [](const FRay&, FVector&) {return false; };
 protected:
 
 	/** @return a new instance of the standard axis-handle Gizmo */
@@ -219,10 +249,8 @@ protected:
 		float MaxParameter,
 		IGizmoStateTarget* StateTargetIn);
 
-
 };
 
-
 /**
  * UGizmoAxisIntervalParameterSource is an IGizmoFloatParameterSource implementation that
  * interprets the float value as the parameter of a line equation, and maps this parameter to a 3D translation
diff --git a/Engine/Source/Runtime/Experimental/InteractiveToolsFramework/Public/ToolContextInterfaces.h b/Engine/Source/Runtime/Experimental/InteractiveToolsFramework/Public/ToolContextInterfaces.h
--- a/Engine/Source/Runtime/Experimental/InteractiveToolsFramework/Public/ToolContextInterfaces.h
+++ b/Engine/Source/Runtime/Experimental/InteractiveToolsFramework/Public/ToolContextInterfaces.h
@@ -150,6 +150,8 @@ struct INTERACTIVETOOLSFRAMEWORK_API FSceneSnapQueryRequest
 
 	/** Optional: components to consider invisible even if they aren't. */
 	const TArray<const UPrimitiveComponent*>* ComponentsToIgnore = nullptr;
+	/** Optional: components to consider visible even if they aren't. */
+	const TArray<const UPrimitiveComponent*>* InvisibleComponentsToInclude = nullptr;
 };
 
 
