MeshModelingTools: SpaceDeform Tool improvements. Fixed modifier intensity and active range errors in bend operation, added normal normalization in flare op, made parameters be operation-specific, added nonuniform flaring and "smooth ends" option to flare op, added a "shift to center" button, added transparent original mesh view, added alignment of gizmo and intervals to world, added "lock bottom" option to twist and bend operations, added line visualization to bend operation, added option for ctrl click to not use the normal.

#rnx

diff --git a/Engine/Plugins/Experimental/MeshModelingToolset/Source/MeshModelingTools/Private/MeshSpaceDeformerTool.cpp b/Engine/Plugins/Experimental/MeshModelingToolset/Source/MeshModelingTools/Private/MeshSpaceDeformerTool.cpp
--- a/Engine/Plugins/Experimental/MeshModelingToolset/Source/MeshModelingTools/Private/MeshSpaceDeformerTool.cpp
+++ b/Engine/Plugins/Experimental/MeshModelingToolset/Source/MeshModelingTools/Private/MeshSpaceDeformerTool.cpp
@@ -7,6 +7,7 @@
 
 #include "SegmentTypes.h"
 #include "DynamicMeshAttributeSet.h"
+#include "Mechanics/DragAlignmentMechanic.h"
 #include "MeshNormals.h"
 #include "MeshOpPreviewHelpers.h"
 #include "ToolSceneQueriesUtil.h"
@@ -34,7 +35,7 @@
 /*
  * ToolBuilder
  */
-UMeshSurfacePointTool* UMeshSpaceDeformerToolBuilder::CreateNewTool(const FToolBuilderState& SceneState) const
+UInteractiveTool* UMeshSpaceDeformerToolBuilder::BuildTool(const FToolBuilderState& SceneState) const
 {
 	UMeshSpaceDeformerTool* MeshSpaceDeformerTool = NewObject<UMeshSpaceDeformerTool>(SceneState.ToolManager);
 	
@@ -61,7 +62,7 @@ TUniquePtr<FDynamicMeshOperator> USpaceDeformerOperatorFactory::MakeNewOperator(
 
 	check(SpaceDeformerTool);
 
-	const ENonlinearOperationType OperationType = SpaceDeformerTool->SelectedOperationType;
+	const ENonlinearOperationType OperationType = SpaceDeformerTool->Settings->SelectedOperationType;
 	
 	// Create the actual operator type based on the requested operation
 	TUniquePtr<FMeshSpaceDeformerOp>  DeformerOp;
@@ -71,16 +72,24 @@ TUniquePtr<FDynamicMeshOperator> USpaceDeformerOperatorFactory::MakeNewOperator(
 		case ENonlinearOperationType::Bend: 
 		{
 			DeformerOp = MakeUnique<FBendMeshOp>();
+			static_cast<FBendMeshOp*>(DeformerOp.Get())->BendDegrees = SpaceDeformerTool->Settings->BendDegrees;
+			static_cast<FBendMeshOp*>(DeformerOp.Get())->bLockBottom = SpaceDeformerTool->Settings->bLockBottom;
 			break;
 		}
 		case ENonlinearOperationType::Flare:
 		{
 			DeformerOp = MakeUnique<FFlareMeshOp>();
+			static_cast<FFlareMeshOp*>(DeformerOp.Get())->FlarePercentY = SpaceDeformerTool->Settings->FlarePercentY;
+			static_cast<FFlareMeshOp*>(DeformerOp.Get())->FlarePercentX = SpaceDeformerTool->Settings->bLockXAndYFlaring ? 
+				SpaceDeformerTool->Settings->FlarePercentY : SpaceDeformerTool->Settings->FlarePercentX;
+			static_cast<FFlareMeshOp*>(DeformerOp.Get())->bSmoothEnds = SpaceDeformerTool->Settings->bSmoothEnds;
 			break;
 		}
 		case ENonlinearOperationType::Twist:
 		{
 			DeformerOp = MakeUnique<FTwistMeshOp>();
+			static_cast<FTwistMeshOp*>(DeformerOp.Get())->TwistDegrees = SpaceDeformerTool->Settings->TwistDegrees;
+			static_cast<FTwistMeshOp*>(DeformerOp.Get())->bLockBottom = SpaceDeformerTool->Settings->bLockBottom;
 			break;
 		}
 	default:
@@ -96,14 +105,20 @@ TUniquePtr<FDynamicMeshOperator> USpaceDeformerOperatorFactory::MakeNewOperator(
 }
 
 
+void UMeshSpaceDeformerToolActionPropertySet::PostAction(EMeshSpaceDeformerToolAction Action)
+{
+	if (ParentTool.IsValid())
+	{
+		ParentTool->RequestAction(Action);
+	}
+}
+
+
 /*
  * Tool
  */
 UMeshSpaceDeformerTool::UMeshSpaceDeformerTool()
 {
-	GizmoCenter = FVector::ZeroVector;
-	GizmoOrientation = FQuat::Identity;
-	GizmoFrame = FFrame3d(FVector3d(GizmoCenter), FQuaterniond(GizmoOrientation));
 }
 
 void UMeshSpaceDeformerTool::SetWorld(UWorld* World)
@@ -123,27 +138,17 @@ bool UMeshSpaceDeformerTool::CanAccept() const
 	return Super::CanAccept() && (Preview == nullptr || Preview->HaveValidResult());
 }
 
-void UMeshSpaceDeformerTool::ComputeAABB(const FDynamicMesh3& MeshIn, const FTransform& XFormIn, FVector& BBoxMin, FVector& BBoxMax) const 
-{
-	
-
-	BBoxMin = FVector(FLT_MAX, FLT_MAX, FLT_MAX);
-	BBoxMax = -BBoxMin;
-	
-	for (const auto VertID : MeshIn.VertexIndicesItr())
-	{
-		const FVector3d LocalPosition = MeshIn.GetVertex(VertID);
-		const FVector WorldPosition = XFormIn.TransformPosition(FVector(LocalPosition.X, LocalPosition.Y, LocalPosition.Z));
-		for (int i = 0; i < 3; ++i) BBoxMin[i] = FMath::Min(BBoxMin[i], WorldPosition[i]);
-		for (int i = 0; i < 3; ++i) BBoxMax[i] = FMath::Max(BBoxMax[i], WorldPosition[i]);
-	}
-
-}
-
 void UMeshSpaceDeformerTool::Setup()
 {
-	UMeshSurfacePointTool::Setup();
+	UInteractiveTool::Setup();
+
+	Settings = NewObject<UMeshSpaceDeformerToolProperties>(this);
+	Settings->RestoreProperties(this);
+	AddToolPropertySource(Settings);
 
+	ToolActions = NewObject<UMeshSpaceDeformerToolActionPropertySet>(this);
+	ToolActions->Initialize(this);
+	AddToolPropertySource(ToolActions);
 
 	// populate the OriginalDynamicMesh with a conversion of the input mesh.
 	{
@@ -152,28 +157,52 @@ void UMeshSpaceDeformerTool::Setup()
 		Converter.Convert(ComponentTarget->GetMesh(), *OriginalDynamicMesh);
 	}
 
-	// hide input StaticMeshComponent
-	ComponentTarget->SetOwnerVisibility(false);
-
+	FTransform MeshTransform = ComponentTarget->GetWorldTransform();
 
+	// Hide the mesh, and potentially put a semi-transparent copy in its place. We could
+	// update the materials and restore them later, but this seems safer.
+	ComponentTarget->SetOwnerVisibility(false);
 
-	//// add properties
-
-	AddToolPropertySource(this);
-
-
-	// compute the bounding box for the input mesh
-	// AABB for the source geometry
-	FVector AABBMin;
-	FVector AABBMax;
-	ComputeAABB(*OriginalDynamicMesh, ComponentTarget->GetWorldTransform(), AABBMin, AABBMax);
-	FVector Extents = AABBMax - AABBMin;
-
+	OriginalMeshPreview = NewObject<UPreviewMesh>();
+	OriginalMeshPreview->CreateInWorld(TargetWorld, MeshTransform);
+	OriginalMeshPreview->UpdatePreview(OriginalDynamicMesh.Get());
+	OriginalMeshPreview->SetMaterial(0, ToolSetupUtil::GetCustomDepthOffsetMaterial(GetToolManager(), FLinearColor::White,
+		-0.5,  // depth offset, 0.5% inward
+		0.4)); // opacity
+	OriginalMeshPreview->SetVisible(Settings->bShowOriginalMesh);
+
+	// The gizmo gets initialized to the center point of the object-space bounding box,
+	// with the Z axis (along which the deformation acts) aligned to the longest of the
+	// bounding box dimensions, after scaling them with the transform.
+	FAxisAlignedBox3d BBox = OriginalDynamicMesh->GetBounds();
+	FVector3d Dimensions = BBox.Max - BBox.Min;
+	MeshCenter = MeshTransform.TransformPosition((FVector)BBox.Center());
+
+	Dimensions = FVector3d(MeshTransform.GetScale3D().GetAbs()) * Dimensions;
+	double WorldMajorLength = 0;
+
+	// Prefer being aligned with the Z axis.
+	if (Dimensions.Z >= Dimensions.Y && Dimensions.Z >= Dimensions.X)
+	{
+		WorldMajorLength = Dimensions.Z;
+		GizmoFrame = FFrame3d(MeshCenter, FVector3d::UnitX(), FVector3d::UnitY(), FVector3d::UnitZ());
+	}
+	else if (Dimensions.Y >= Dimensions.X)
+	{
+		WorldMajorLength = Dimensions.Y;
+		GizmoFrame = FFrame3d(MeshCenter, FVector3d::UnitZ(), FVector3d::UnitX(), FVector3d::UnitY());
+	}
+	else
+	{
+		WorldMajorLength = Dimensions.X;
+		GizmoFrame = FFrame3d(MeshCenter, FVector3d::UnitY(), FVector3d::UnitZ(), FVector3d::UnitX());
+	}
 
-	AABBHalfExtents = 0.5 * FVector3d(Extents.X, Extents.Y, Extents.Z);
-	GizmoCenter = 0.5 * (AABBMin + AABBMax);
-	GizmoFrame.Origin = GizmoCenter;
+	GizmoFrame.Rotate((FQuaterniond)MeshTransform.GetRotation());
 
+	// The scaling of the modifier gizmo is somewhat arbitrary. We choose for it to be
+	// related to the major axis length.
+	ModifierGizmoLength = WorldMajorLength;
 
 	// add click to set plane behavior
 	SetPointInWorldConnector = MakePimpl<FSelectClickedAction>();
@@ -181,7 +210,7 @@ void UMeshSpaceDeformerTool::Setup()
 	SetPointInWorldConnector->InvisibleComponentsToHitTest.Add(ComponentTarget->GetOwnerComponent());
 	SetPointInWorldConnector->OnClickedPositionFunc = [this](const FHitResult& Hit)
 	{
-		SetGizmoPlaneFromWorldPos(Hit.ImpactPoint, -Hit.ImpactNormal, false);
+		SetGizmoFrameFromWorldPos(Hit.ImpactPoint, Hit.ImpactNormal, Settings->bAlignToNormalOnCtrlClick);
 	};
 
 	USingleClickInputBehavior* ClickToSetPlaneBehavior = NewObject<USingleClickInputBehavior>();
@@ -193,72 +222,50 @@ void UMeshSpaceDeformerTool::Setup()
 	// Create a new TransformGizmo and associated TransformProxy. The TransformProxy will not be the
 	// parent of any Components in this case, we just use it's transform and change delegate.
 	TransformProxy = NewObject<UTransformProxy>(this);
-	TransformProxy->SetTransform(FTransform(GizmoOrientation, GizmoCenter));
+	TransformProxy->SetTransform(GizmoFrame.ToFTransform());
 	TransformGizmo = GetToolManager()->GetPairedGizmoManager()->CreateCustomTransformGizmo(
 		ETransformGizmoSubElements::StandardTranslateRotate, this);
+	TransformGizmo->bSnapToWorldGrid = Settings->bSnapToWorldGrid;
 
 	TransformGizmo->SetActiveTarget(TransformProxy, GetToolManager());
+
 	// listen for changes to the proxy and update the preview when that happens
 	TransformProxy->OnTransformChanged.AddUObject(this, &UMeshSpaceDeformerTool::TransformProxyChanged);
 
-	// wire these to the input 
-
-	// The initial range for the intervals.
-	const float InitialVerticalIntervalSize = AABBHalfExtents.Length();
-
 	// create sources for the interval parameters
-
 	UpIntervalSource      = NewObject< UGizmoLocalFloatParameterSource >(this);
 	DownIntervalSource    = NewObject< UGizmoLocalFloatParameterSource >(this);
 	ForwardIntervalSource = NewObject< UGizmoLocalFloatParameterSource >(this);
 
-	// Initial Lengths for the interval handles
-
-	UpIntervalSource->Value      = InitialVerticalIntervalSize;
-	DownIntervalSource->Value    = -InitialVerticalIntervalSize;
-	ForwardIntervalSource->Value = 20.f;
+	// Initial Lengths for the interval handle
+	UpIntervalSource->Value      = WorldMajorLength/2;
+	DownIntervalSource->Value    = -WorldMajorLength / 2;
+	ForwardIntervalSource->Value = GetModifierGizmoValue();
 
 	// Sync the properties panel to the interval handles.
-
-	this->UpperBoundsInterval = UpIntervalSource->Value;
-	this->LowerBoundsInterval = DownIntervalSource->Value;
-	this->ModifierPercent     = ForwardIntervalSource->Value;
+	Settings->UpperBoundsInterval = UpIntervalSource->Value;
+	Settings->AbsLowerBoundsInterval = FMath::Abs(DownIntervalSource->Value);
 
 	// Wire up callbacks to update result mesh and the properties panel when these parameters are changed (by gizmo manipulation in viewport).  Note this is just a one-way
 	// coupling (Sources to Properties). The OnPropertyModified() method provides the Properties to Souces coupling 
-
 	UpIntervalSource->OnParameterChanged.AddLambda([this](IGizmoFloatParameterSource* ParamSource, FGizmoFloatParameterChange Change)->void
 	{
-		this->UpperBoundsInterval = Change.CurrentValue;
-
-		if (this->Preview != nullptr)
-		{
-			this->Preview->InvalidateResult();
-		}
+		Settings->UpperBoundsInterval = Change.CurrentValue;
+		UpdatePreview();
 	});
 
 	DownIntervalSource->OnParameterChanged.AddLambda([this](IGizmoFloatParameterSource* ParamSource, FGizmoFloatParameterChange Change)->void
 	{
-		this->LowerBoundsInterval = Change.CurrentValue;
-
-		if (this->Preview != nullptr)
-		{
-			this->Preview->InvalidateResult();
-		}
+		Settings->AbsLowerBoundsInterval = FMath::Abs(Change.CurrentValue);
+		UpdatePreview();
 	});
 
 	ForwardIntervalSource->OnParameterChanged.AddLambda([this](IGizmoFloatParameterSource* ParamSource, FGizmoFloatParameterChange Change)->void
 	{
-		this->ModifierPercent = Change.CurrentValue;
-
-		if (this->Preview != nullptr)
-		{
-			this->Preview->InvalidateResult();
-		}
+		ApplyModifierGizmoValue(Change.CurrentValue);
+		UpdatePreview();
 	});
 
-
-
 	// add the interval gizmo
 	IntervalGizmo = GetToolManager()->GetPairedGizmoManager()->CreateGizmo<UIntervalGizmo>(UIntervalGizmo::GizmoName, TEXT("MeshSpaceDefomerInterval"), this);
 
@@ -267,6 +274,13 @@ void UMeshSpaceDeformerTool::Setup()
 
 	// use the statetarget to track details changes
 	StateTarget = IntervalGizmo->StateTarget;
+
+	// Set up the bent line visualizer
+	VisualizationRenderer.bDepthTested = false;
+	VisualizationRenderer.LineColor = FLinearColor::Yellow;
+	VisualizationRenderer.LineThickness = 4.0;
+	VisualizationRenderer.SetTransform(GizmoFrame.ToFTransform());
+
 	// Set up the preview object
 	{
 		// create the operator factory
@@ -282,11 +296,9 @@ void UMeshSpaceDeformerTool::Setup()
 
 		// Give the preview something to display
 		Preview->PreviewMesh->UpdatePreview(OriginalDynamicMesh.Get());
-		Preview->PreviewMesh->SetTransform(ComponentTarget->GetWorldTransform());
+		Preview->PreviewMesh->SetTransform(MeshTransform);
 
-		FComponentMaterialSet MaterialSet;
-		ComponentTarget->GetMaterialSet(MaterialSet);
-		Preview->ConfigureMaterials(MaterialSet.Materials,
+		Preview->ConfigureMaterials(StoredMaterialSet.Materials,
 			ToolSetupUtil::GetDefaultWorkingMaterial(GetToolManager())
 		);
 
@@ -294,9 +306,18 @@ void UMeshSpaceDeformerTool::Setup()
 		Preview->SetVisibility(true);
 
 		// start the compute
-		Preview->InvalidateResult();
+		UpdatePreview();
 	}
 
+	DragAlignmentMechanic = NewObject<UDragAlignmentMechanic>(this);
+	DragAlignmentMechanic->Setup(this);
+
+	// We want to align to the original mesh, even though it is hidden (our stand-in preview mesh that
+	// we use to display a transparent version does not get hit tested by our normal raycasts into
+	// the world).
+	TArray<const UPrimitiveComponent*> ComponentsToInclude{ ComponentTarget->GetOwnerComponent() } ;
+	DragAlignmentMechanic->AddToGizmo(TransformGizmo, nullptr, &ComponentsToInclude);
+	DragAlignmentMechanic->AddToGizmo(IntervalGizmo, nullptr, &ComponentsToInclude);
 
 	GetToolManager()->DisplayMessage(
 		LOCTEXT("MeshSpaceDeformerToolDescription", "Deform the vertices of the selected Mesh using various spatial deformations. Use the in-viewport Gizmo to control the extents/strength of the deformation."),
@@ -305,9 +326,13 @@ void UMeshSpaceDeformerTool::Setup()
 
 void UMeshSpaceDeformerTool::Shutdown(EToolShutdownType ShutdownType)
 {
-	
-	// Restore (unhide) the source meshes
+	Settings->SaveProperties(this);
+
+	// Restore source mesh and remove our stand-in
 	ComponentTarget->SetOwnerVisibility(true);
+	OriginalMeshPreview->SetVisible(false);
+	OriginalMeshPreview->Disconnect();
+	OriginalMeshPreview = nullptr;
 
 	if (Preview != nullptr)
 	{
@@ -331,153 +356,189 @@ void UMeshSpaceDeformerTool::Shutdown(EToolShutdownType ShutdownType)
 		}
 	}
 
+	DragAlignmentMechanic->Shutdown();
+
 	UInteractiveGizmoManager* GizmoManager = GetToolManager()->GetPairedGizmoManager();
 	GizmoManager->DestroyAllGizmosByOwner(this);
-
 }
 
 
 
 void  UMeshSpaceDeformerTool::TransformProxyChanged(UTransformProxy* Proxy, FTransform Transform)
 {
-	GizmoOrientation = Transform.GetRotation();
-	GizmoCenter = Transform.GetLocation();
-	GizmoFrame = FFrame3d( FVector3d(GizmoCenter), FQuaterniond(GizmoOrientation) );
+	GizmoFrame = FFrame3d(Transform.GetLocation(), Transform.GetRotation());
+	VisualizationRenderer.SetTransform(GizmoFrame.ToFTransform());
+	UpdatePreview();
+}
 
-	if (Preview != nullptr)
-	{
-		Preview->InvalidateResult();
-	}
+void  UMeshSpaceDeformerTool::OnPropertyModified(UObject* PropertySet, FProperty* Property)
+{
+	UpIntervalSource->Value = Settings->UpperBoundsInterval;
+	DownIntervalSource->Value = -Settings->AbsLowerBoundsInterval;
+	ForwardIntervalSource->Value = GetModifierGizmoValue();
+
+	UpdatePreview();
+
+	TransformGizmo->bSnapToWorldGrid = Settings->bSnapToWorldGrid;
+
+	OriginalMeshPreview->SetVisible(Settings->bShowOriginalMesh);
 }
 
+void UMeshSpaceDeformerTool::UpdateOpParameters(FMeshSpaceDeformerOp& MeshSpaceDeformerOp) const
+{
+	MeshSpaceDeformerOp.OriginalMesh = OriginalDynamicMesh;
+	MeshSpaceDeformerOp.SetTransform(ComponentTarget->GetWorldTransform());
+	MeshSpaceDeformerOp.GizmoFrame = GizmoFrame;
 
+	// set the bound range
+	MeshSpaceDeformerOp.UpperBoundsInterval =  Settings->UpperBoundsInterval;
+	MeshSpaceDeformerOp.LowerBoundsInterval = -Settings->AbsLowerBoundsInterval;
+}
 
-#if WITH_EDITOR
-void UMeshSpaceDeformerTool::PostEditChangeProperty(FPropertyChangedEvent & PropertyChangedEvent)
+void UMeshSpaceDeformerTool::RequestAction(EMeshSpaceDeformerToolAction ActionType)
 {
-	if (PropertyChangedEvent.ChangeType & EPropertyChangeType::Interactive)
+	if (PendingAction == EMeshSpaceDeformerToolAction::NoAction)
 	{
-		if (!bHasBegin)
-		{
-			StateTarget->BeginUpdate();
-			bHasBegin = true;
-		}
-
-		// Propagates change to the ParameterSource by triggering UMeshSpaceDeformerTool::OnPropertyModified
-		OnPropertyModified(this, PropertyChangedEvent.Property);
+		PendingAction = ActionType;
 	}
+}
 
-	if (PropertyChangedEvent.ChangeType & EPropertyChangeType::ValueSet)
+void UMeshSpaceDeformerTool::ApplyAction(EMeshSpaceDeformerToolAction ActionType)
+{
+	if (PendingAction == EMeshSpaceDeformerToolAction::ShiftToCenter)
 	{
-		if (!bHasBegin)
-		{
-			StateTarget->BeginUpdate();
-		}
-
-		// Propagates change to the ParameterSource by triggering UMeshSpaceDeformerTool::OnPropertyModified
-		OnPropertyModified(this, PropertyChangedEvent.Property);
-
-		StateTarget->EndUpdate();
-		// reset the bHasBegin
-		bHasBegin = false;
+		SetGizmoFrameFromWorldPos((FVector)MeshCenter);
 	}
-		
 }
 
-
-#endif
-
-void  UMeshSpaceDeformerTool::OnPropertyModified(UObject* PropertySet, FProperty* Property)
+void UMeshSpaceDeformerTool::OnTick(float DeltaTime)
 {
-	// Update the property sources to reflect any changes.
-	// NB: due to callback, the SetParameter will trigger a recompute of the result via  Preview->InvalidateResult();
+	// Deal with clicked button
+	if (PendingAction != EMeshSpaceDeformerToolAction::NoAction)
+	{
+		ApplyAction(PendingAction);
+		PendingAction = EMeshSpaceDeformerToolAction::NoAction;
+	}
 
-	if (PropertySet == this)
+	if (Preview != nullptr)
 	{
-		if (Property && (Property->GetFName() == GET_MEMBER_NAME_CHECKED(UMeshSpaceDeformerTool, UpperBoundsInterval)))
-		{
-			UpIntervalSource->BeginModify();
-			UpIntervalSource->SetParameter(this->UpperBoundsInterval);
-			UpIntervalSource->EndModify();
-		}
-		if (Property && (Property->GetFName() == GET_MEMBER_NAME_CHECKED(UMeshSpaceDeformerTool, LowerBoundsInterval)))
-		{
-			DownIntervalSource->BeginModify();
-			DownIntervalSource->SetParameter(this->LowerBoundsInterval);
-			DownIntervalSource->EndModify();
-		}
-		if (Property && (Property->GetFName() == GET_MEMBER_NAME_CHECKED(UMeshSpaceDeformerTool, ModifierPercent)))
-		{
-			ForwardIntervalSource->BeginModify();
-			ForwardIntervalSource->SetParameter(this->ModifierPercent);
-			ForwardIntervalSource->EndModify();
-		}
+		Preview->Tick(DeltaTime);
+	}
+}
 
+void UMeshSpaceDeformerTool::Render(IToolsContextRenderAPI* RenderAPI)
+{
+	DragAlignmentMechanic->Render(RenderAPI);
 
-		if (Property && (Property->GetFName() == GET_MEMBER_NAME_CHECKED(UMeshSpaceDeformerTool, SelectedOperationType)))
+	if (Settings->bDrawVisualization && Settings->SelectedOperationType == ENonlinearOperationType::Bend)
+	{
+		VisualizationRenderer.BeginFrame(RenderAPI, RenderAPI->GetCameraState());
+		for (int32 i = 1; i < VisualizationPoints.Num(); ++i)
 		{
-			if (Preview != nullptr)
-			{
-				Preview->InvalidateResult();
-			}
+			VisualizationRenderer.DrawLine(VisualizationPoints[i - 1], VisualizationPoints[i]);
 		}
+		VisualizationRenderer.EndFrame();
+	}
+	else
+	{
 	}
 }
 
-void UMeshSpaceDeformerTool::UpdateOpParameters(FMeshSpaceDeformerOp& MeshSpaceDeformerOp) const
+void UMeshSpaceDeformerTool::UpdatePreview()
 {
-	MeshSpaceDeformerOp.OriginalMesh = OriginalDynamicMesh;
-	MeshSpaceDeformerOp.SetTransform(ComponentTarget->GetWorldTransform());
-	MeshSpaceDeformerOp.GizmoFrame = GizmoFrame;
+	if (Settings->SelectedOperationType == ENonlinearOperationType::Bend)
+	{
+		const int32 NUM_RENDER_POINTS = 30;
 
-	// Set half axis length to be 1/2 the major axis of the bbox.
-	double LengthScale = FMath::Max( AABBHalfExtents.MaxAbsElement(), 1.e-3);
-	MeshSpaceDeformerOp.AxesHalfLength = LengthScale;
+		VisualizationPoints.SetNumUninitialized(NUM_RENDER_POINTS);
 
-	// set the bound range
-	MeshSpaceDeformerOp.UpperBoundsInterval =  UpperBoundsInterval / LengthScale;
-	MeshSpaceDeformerOp.LowerBoundsInterval = -LowerBoundsInterval / LengthScale;
+		double BentLength = Settings->UpperBoundsInterval + Settings->AbsLowerBoundsInterval;
+		double ArcAngle = Settings->BendDegrees * PI / 180;
+		double ArcRadius = BentLength / ArcAngle;
 
-	// percent to apply.
-	MeshSpaceDeformerOp.ModifierPercent = ModifierPercent;
-}
+		double RotationCenterZ = Settings->bLockBottom ? -Settings->AbsLowerBoundsInterval : 0;
+		FVector2d RotationCenterYZ(ArcRadius, RotationCenterZ);
 
-void UMeshSpaceDeformerTool::OnTick(float DeltaTime)
-{
-	if (TransformGizmo != nullptr)
-	{
-		TransformGizmo->bSnapToWorldGrid = this->bSnapToWorldGrid;
+		double PointSpacing = BentLength / (NUM_RENDER_POINTS - 1);
+
+		for (int32 i = 0; i < NUM_RENDER_POINTS; ++i)
+		{
+			double OriginalZ = -Settings->AbsLowerBoundsInterval + PointSpacing * i;
+			FVector2d YZToRotate(0, RotationCenterZ);
+
+			// The negative here is because we are rotating clockwise in the direction of the positive Y axis
+			double AngleToRotate = -ArcAngle * (OriginalZ - RotationCenterZ) / BentLength;
+
+			FMatrix2d RotationMatrix = FMatrix2d::RotationRad(AngleToRotate);
+			FVector2d RotatedYZ = RotationMatrix * (YZToRotate - RotationCenterYZ) + RotationCenterYZ;
+
+			VisualizationPoints[i] = FVector3d(0, RotatedYZ.X, RotatedYZ.Y);
+		}
 	}
 
-	if (Preview != nullptr)
+	if (Preview)
 	{
-		Preview->Tick(DeltaTime);
+		Preview->InvalidateResult();
 	}
 }
 
 
-
-void UMeshSpaceDeformerTool::SetGizmoPlaneFromWorldPos(const FVector& Position, const FVector& Normal, bool bIsInitializing)
+void UMeshSpaceDeformerTool::SetGizmoFrameFromWorldPos(const FVector& Position, const FVector& Normal, bool bAlignNormal)
 {
-	GizmoCenter = Position;
-
-	FFrame3f GizmoPlane(Position, Normal);
-	GizmoOrientation = (FQuat)GizmoPlane.Rotation;
-	GizmoFrame = FFrame3d(FVector3d(GizmoCenter), FQuaterniond(GizmoOrientation));
-
-	TransformGizmo->SetActiveTarget(TransformProxy, GetToolManager());
-	if (bIsInitializing)
+	GizmoFrame.Origin = Position;
+	if (bAlignNormal)
 	{
-		TransformGizmo->ReinitializeGizmoTransform(GizmoPlane.ToFTransform());
+		// It's not clear whether aligning the Z axis to the normal is the right idea here. The Z axis
+		// is the main axis on which we operate. On the one hand, setting it to the normal gives the user
+		// greater control over its alignment. On the other hand, it seems likely that when clicking the object,
+		// the user would want the axis to lie along the object on the side they clicked, not pierce inwards.
+		// Still, it's hard to come up with a clean alternative.
+		FVector3d FrameZ(Normal);
+		FVector3d FrameY = FrameZ.Cross(FVector3d::UnitZ()).Normalized(); // orthogonal to world Z and frame Z 
+		FVector3d FrameX = FrameY.Cross(Normal); // safe to not normalize because already orthogonal
+		GizmoFrame = FFrame3d(Position, FrameX, FrameY, Normal);
 	}
-	else
+
+	TransformGizmo->ReinitializeGizmoTransform(GizmoFrame.ToFTransform());
+	UpdatePreview();
+}
+
+/** 
+ * These two functions translate to and from the modifier gizmo length
+ * to the relevant operator parameters. They should be matched to each
+ * other.
+ */
+double UMeshSpaceDeformerTool::GetModifierGizmoValue() const
+{
+	switch (Settings->SelectedOperationType)
 	{
-		TransformGizmo->SetNewGizmoTransform(GizmoPlane.ToFTransform());
+	case ENonlinearOperationType::Bend:
+		return Settings->BendDegrees * ModifierGizmoLength / 180;
+	case ENonlinearOperationType::Flare:
+		return Settings->FlarePercentY * ModifierGizmoLength / 100;
+	case ENonlinearOperationType::Twist:
+		return Settings->TwistDegrees * ModifierGizmoLength / 360;
 	}
 
-	if (Preview != nullptr)
+	// Shouldn't get here
+	check(false);
+	return 0;
+}
+void UMeshSpaceDeformerTool::ApplyModifierGizmoValue(double Value)
+{
+	switch (Settings->SelectedOperationType)
 	{
-		Preview->InvalidateResult();
+	case ENonlinearOperationType::Bend:
+		Settings->BendDegrees = 180 * Value / ModifierGizmoLength;
+		break;
+	case ENonlinearOperationType::Flare:
+		Settings->FlarePercentY = 100 * Value / ModifierGizmoLength;
+		break;
+	case ENonlinearOperationType::Twist:
+		Settings->TwistDegrees = 360 * Value/ ModifierGizmoLength;
+		break;
+	default:
+		check(false);
 	}
 }
 
diff --git a/Engine/Plugins/Experimental/MeshModelingToolset/Source/MeshModelingTools/Public/MeshSpaceDeformerTool.h b/Engine/Plugins/Experimental/MeshModelingToolset/Source/MeshModelingTools/Public/MeshSpaceDeformerTool.h
--- a/Engine/Plugins/Experimental/MeshModelingToolset/Source/MeshModelingTools/Public/MeshSpaceDeformerTool.h
+++ b/Engine/Plugins/Experimental/MeshModelingToolset/Source/MeshModelingTools/Public/MeshSpaceDeformerTool.h
@@ -11,10 +11,14 @@
 #include "SpaceDeformerOps/TwistMeshOp.h"
 #include "SpaceDeformerOps/FlareMeshOp.h"
 #include "SimpleDynamicMeshComponent.h"
+#include "ToolDataVisualizer.h"
 #include "BaseGizmos/GizmoInterfaces.h"
 
 #include "MeshSpaceDeformerTool.generated.h"
 
+class UMeshSpaceDeformerTool;
+
+class UDragAlignmentMechanic;
 class UPreviewMesh;
 class UTransformGizmo;
 class UTransformProxy;
@@ -28,7 +32,7 @@ class FSelectClickedAction;
  * ToolBuilder
  */
 UCLASS()
-class MESHMODELINGTOOLS_API UMeshSpaceDeformerToolBuilder : public UMeshSurfacePointToolBuilder
+class MESHMODELINGTOOLS_API UMeshSpaceDeformerToolBuilder : public UInteractiveToolBuilder
 {
 	GENERATED_BODY()
 
@@ -38,22 +42,139 @@ public:
 	}
 
 	virtual bool CanBuildTool(const FToolBuilderState& SceneState) const override;
-
-	virtual UMeshSurfacePointTool* CreateNewTool(const FToolBuilderState& SceneState) const override;
+	virtual UInteractiveTool* BuildTool(const FToolBuilderState& SceneState) const override;
 };
 
 /** ENonlinearOperation determines which type of nonlinear deformation will be applied*/
 UENUM()
 enum class  ENonlinearOperationType : int8
 {
-	Bend		UMETA(DisplayName = "Bend"),
-	Flare		UMETA(DisplayName = "Flare"),
-	Twist		UMETA(DisplayName = "Twist"),
-	//Sinusoid	UMETA(DisplayName = "Sinusoid"),
-	//Wave		UMETA(DisplayName = "Wave"),
-	//Squish	UMETA(DisplayName = "Squish")
+	/** 
+	 * Will bend the mesh in the direction of the gizmo Y axis along the Z axis. A line along the Z
+	 * axis from the lower bound to the upper bound would not change length as it bends.
+	 */
+	Bend,
+
+	/**
+	 * Depening on 'Flare Percent", will either flare or squish the mesh along the Gizmo Z axis,
+	 * from lower bound to upper bound.
+	 */
+	Flare		UMETA(DisplayName = "Flare/Squish"),
+
+	/** Twists the mesh along the gizmo Z axis, from lower bound to upper bound. */
+	Twist
 };
 
+UCLASS()
+class MESHMODELINGTOOLS_API UMeshSpaceDeformerToolProperties : public UInteractiveToolPropertySet
+{
+	GENERATED_BODY()
+
+public:
+	UPROPERTY(EditAnywhere, Category = Options, meta = (DisplayName = "Operation Type"))
+	ENonlinearOperationType SelectedOperationType = ENonlinearOperationType::Bend;
+
+	/** The upper bound to the region of space which the operation will affect. Measured upward along the gizmo Z axis from the gizmo center. */
+	UPROPERTY(EditAnywhere, Category = Options, meta = (DisplayName = "Up Extent", UIMin = "0.0", ClampMin = "0.0"))
+	float UpperBoundsInterval = 10.0;
+
+	//~ Note that the value here is the negative of what we use elsewhere, because otherwise the user has to reverse
+	//~ their scrubbing direction compared to UpperBoundsInterval for increasing/decreasing the interval.
+	/** The lower bound to the region of space which the operation will affect. Measured downward along the gizmo Z axis from the gizmo center. */
+	UPROPERTY(EditAnywhere, Category = Options, meta = (DisplayName = "Down Extent", UIMin = "0", ClampMin = "0"))
+	float AbsLowerBoundsInterval = 10.0;
+
+	//~ The two degrees properties (bend vs twist) are separate because they have different clamp values. 
+	//~ Bending a "negative" amount probably won't do what the user wants, whereas twisting in the opposite
+	//~ direction makes sense.
+	/** 
+	 * A line along the Z axis of the gizmo from lower bound to upper bound will be bent into a perfect arc of this
+	 * many degrees in the direction of the Y axis without changing length.
+	 */
+	UPROPERTY(EditAnywhere, Category = Options, meta = (UIMin = "0", UIMax = "360", ClampMin = "0", ClampMax = "3600", 
+		EditCondition = "SelectedOperationType == ENonlinearOperationType::Bend", EditConditionHides))
+	float BendDegrees = 90;
+
+	/** Degrees of twist to from the lower bound to the upper bound along the gizmo Z axis. */
+	UPROPERTY(EditAnywhere, Category = Options, meta = (UIMin = "-360", UIMax = "360", ClampMin = "-3600", ClampMax = "3600", 
+		EditCondition = "SelectedOperationType == ENonlinearOperationType::Twist", EditConditionHides))
+	float TwistDegrees = 180;
+
+	/**
+	 * Determines how much to flare perpendicular to the Z axis. When set to 100%, points are moved double the distance
+	 * away from the gizmo Z axis at the most extreme flare point. 0% does not flare at all, whereas -100% pinches all
+	 * the way to the gizmo Z axis at the most extreme flare point.
+	 */
+	UPROPERTY(EditAnywhere, Category = Options, meta = (UIMin = "-100", UIMax = "200", ClampMin = "-1000", ClampMax = "2000", 
+		EditCondition = "SelectedOperationType == ENonlinearOperationType::Flare", EditConditionHides))
+	float FlarePercentY = 100;
+
+	/**
+	 * If true, flaring is applied along the gizmo X and Y axis the same amount.
+	 */
+	UPROPERTY(EditAnywhere, Category = Options, meta = (EditCondition = "SelectedOperationType == ENonlinearOperationType::Flare", EditConditionHides))
+	bool bLockXAndYFlaring = true;
+
+	/**
+	 * Determines how much to flare perpendicular to the Z axis in the X direction if the flaring is not locked
+	 * to be the same in the X and Y directions.
+	 */
+	UPROPERTY(EditAnywhere, Category = Options, meta = (UIMin = "-100", UIMax = "200", ClampMin = "-1000", ClampMax = "2000",
+	EditCondition = "SelectedOperationType == ENonlinearOperationType::Flare && !bLockXAndYFlaring", EditConditionHides))
+	float FlarePercentX = 100;
+
+	/**
+	 * If true, the "bottom" of the mesh relative to the gizmo Z axis will stay in place while the rest bends or twists. If false, the bend
+	 * or twist will happen around the gizmo location.
+	 */
+	UPROPERTY(EditAnywhere, Category = Options, meta = (
+		EditCondition = "SelectedOperationType == ENonlinearOperationType::Bend || SelectedOperationType == ENonlinearOperationType::Twist", EditConditionHides))
+	bool bLockBottom = true;
+
+	/**
+	 * When true, changes the flare/pinch function to have C1 continuity at the upper and lower bounds.
+	 */
+	UPROPERTY(EditAnywhere, Category = Options, meta = (EditCondition = "SelectedOperationType == ENonlinearOperationType::Flare", EditConditionHides))
+	bool bSmoothEnds = true;
+
+	UPROPERTY(EditAnywhere, Category = Options)
+	bool bShowOriginalMesh = true;
+
+	UPROPERTY(EditAnywhere, Category = Options, meta = (EditCondition = "SelectedOperationType == ENonlinearOperationType::Bend", EditConditionHides))
+	bool bDrawVisualization = true;
+
+	/** Snap the deformer gizmo to the world grid */
+	UPROPERTY(EditAnywhere, Category = Gizmo)
+	bool bSnapToWorldGrid = false;
+
+	/** When true, Ctrl+click not only moves the gizmo to the clicked location, but also aligns the Z axis with the normal at that point. */
+	UPROPERTY(EditAnywhere, Category = Gizmo)
+	bool bAlignToNormalOnCtrlClick = false;
+};
+
+UENUM()
+enum class EMeshSpaceDeformerToolAction
+{
+	NoAction,
+
+	ShiftToCenter
+};
+
+UCLASS()
+class MESHMODELINGTOOLS_API UMeshSpaceDeformerToolActionPropertySet : public UInteractiveToolPropertySet
+{
+	GENERATED_BODY()
+
+public:
+	TWeakObjectPtr<UMeshSpaceDeformerTool> ParentTool;
+
+	void Initialize(UMeshSpaceDeformerTool* ParentToolIn) { ParentTool = ParentToolIn; }
+	void PostAction(EMeshSpaceDeformerToolAction Action);
+
+	/** Move the gizmo to the center of the object without changing the orientation. */
+	UFUNCTION(CallInEditor, Category = Gizmo)
+	void ShiftToCenter() { PostAction(EMeshSpaceDeformerToolAction::ShiftToCenter); }
+};
 
 
 UCLASS()
@@ -74,99 +195,64 @@ public:
  * Applies non-linear deformations to a mesh 
  */
 UCLASS()
-class MESHMODELINGTOOLS_API UMeshSpaceDeformerTool : public UMeshSurfacePointTool
+class MESHMODELINGTOOLS_API UMeshSpaceDeformerTool : public USingleSelectionTool
 {
 	GENERATED_BODY()
 
 public:
 	UMeshSpaceDeformerTool();
 
-
 	virtual void Setup() override;
 	virtual void Shutdown(EToolShutdownType ShutdownType) override;
 
 	virtual void SetWorld(UWorld* World);
 	virtual void SetAssetAPI(IToolsContextAssetAPI* AssetAPI);
 
+	void RequestAction(EMeshSpaceDeformerToolAction ActionType);
 
 	virtual void OnTick(float DeltaTime) override;
-	virtual void Render(IToolsContextRenderAPI* RenderAPI) override {};
+	virtual void Render(IToolsContextRenderAPI* RenderAPI) override;
 
 	virtual bool HasCancel() const override { return true; }
 	virtual bool HasAccept() const override { return true; }
 	virtual bool CanAccept() const override;
 
-
-	// disable UMeshSurfacePointTool hits as tool is not using that interaction (subclass should be changed)
-	virtual bool HitTest(const FRay& Ray, FHitResult& OutHit) override { return false; }
-
-
-#if WITH_EDITOR
-	virtual void PostEditChangeProperty(FPropertyChangedEvent& PropertyChangedEvent) override;
-#endif
-
 	virtual void OnPropertyModified(UObject* PropertySet, FProperty* Property) override;
 
 	// sync the parameters owned by the MeshSpaceDeformerOp 
 	void UpdateOpParameters(FMeshSpaceDeformerOp& MeshSpaceDeformerOp) const;
 
-public:
-
-	UPROPERTY(EditAnywhere, Category = Options, meta = (DisplayName = "Operation Type"))
-	ENonlinearOperationType SelectedOperationType = ENonlinearOperationType::Bend;
-
-	/** The upper bounds interval corresponds to the region of space which the selected operator will affect. A setting of 1.0 should envelope all points in the "upper" half of the mesh given the axis has been auto-detected. The corresponding lower value of -1 will cover the entire mesh. */
-	UPROPERTY(EditAnywhere, Category = Options, meta = (DisplayName = "Upper Bound", UIMin = "0.0", ClampMin = "0.0"))
-	float UpperBoundsInterval = 10.0;
-
-	/** The upper bounds interval corresponds to the region of space which the selected operator will affect. A setting of -1.0 should envelope all points in the "lower" half of the mesh given the axis has been auto-detected. The corresponding upper value of 1 will cover the entire mesh. */
-	UPROPERTY(EditAnywhere, Category = Options, meta = (DisplayName = "Lower Bound", UIMax = "0", ClampMax = "0"))
-	float LowerBoundsInterval = -10.0;
-
-	/** As each operator has a range of values (i.e. curvature, angle of twist, scale), this represents the percentage passed to the operator as a parameter. In the future, for more control, this should be separated into individual settings for each operator for more precise control */
-	UPROPERTY(EditAnywhere, Category = Options, meta = (DisplayName = "Modifier Percent"))
-	float ModifierPercent = 20.0;
-
-	/** Snap the deformer gizmo to the world grid */
-	UPROPERTY(EditAnywhere, Category = Snapping)
-	bool bSnapToWorldGrid = false;
-
+protected:
 
-	//UPROPERTY()
-	//UMeshSpaceDeformerToolStandardProperties* SpaceDeformerProperties;
+	UPROPERTY()
+	UMeshSpaceDeformerToolProperties* Settings;
 
-protected:
+	UPROPERTY()
+	UMeshSpaceDeformerToolActionPropertySet* ToolActions;
 
 	UPROPERTY()
 	UGizmoTransformChangeStateTarget* StateTarget = nullptr;
 	
-	// used to coordinate undo for the detail panel. 
+	UPROPERTY()
+	UDragAlignmentMechanic* DragAlignmentMechanic;
 
-	bool bHasBegin = false;
 protected:
 
 	UPROPERTY()
 	UMeshOpPreviewWithBackgroundCompute* Preview = nullptr;   
 
-
 protected:	
 
 	TSharedPtr<FDynamicMesh3> OriginalDynamicMesh;
 
-	UWorld* TargetWorld;
-	IToolsContextAssetAPI* AssetAPI;
-
-
-	// drawing plane control
-
-	/** offset to center of gizmo*/
 	UPROPERTY()
-	FVector GizmoCenter;
+	UPreviewMesh* OriginalMeshPreview;
 
-	/** Gizmo Plane Orientation */
-	UPROPERTY()
-	FQuat GizmoOrientation;
+	UWorld* TargetWorld;
+	IToolsContextAssetAPI* AssetAPI;
 
+	// Stores the materials of the original object, used to make the object semi-transparent.
+	FComponentMaterialSet StoredMaterialSet;
 
 	UPROPERTY()
 	UIntervalGizmo* IntervalGizmo;
@@ -188,18 +274,32 @@ protected:
 	UPROPERTY()
 	UGizmoLocalFloatParameterSource* ForwardIntervalSource;
 
+	// Button click support
+	EMeshSpaceDeformerToolAction PendingAction;
+	FVector3d MeshCenter;
+
 	FFrame3d GizmoFrame;
 
+	// The length of the third interval gizmo (which sets the intensity of the deformation)
+	// when the modifier is set to some reasonable maximum.
+	double ModifierGizmoLength;
+
 	TPimplPtr<FSelectClickedAction> SetPointInWorldConnector;
 
-	FVector3d AABBHalfExtents; // 1/2 the extents of the bbox
+	TArray<FVector3d> VisualizationPoints;
+	FToolDataVisualizer VisualizationRenderer;
 
 	void TransformProxyChanged(UTransformProxy* Proxy, FTransform Transform);
-	void SetGizmoPlaneFromWorldPos(const FVector& Position, const FVector& Normal, bool bIsInitializing);
-
-	/** Compute the axis aligned abounding box for the source mesh */
-	void ComputeAABB(const FDynamicMesh3& MeshIn, const FTransform& XFormIn, FVector& BBoxMin, FVector& BBoxMax) const;
-	
+	void SetGizmoFrameFromWorldPos(const FVector& Position, const FVector& Normal = FVector(), bool bAlignNormal = false);
 	
+	double GetModifierGizmoValue() const;
+	void ApplyModifierGizmoValue(double Value);
+
+	// Apply clicked action
+	void ApplyAction(EMeshSpaceDeformerToolAction Action);
+
+	void UpdatePreview();
+
+	friend USpaceDeformerOperatorFactory;
 };
 
diff --git a/Engine/Plugins/Experimental/MeshModelingToolset/Source/ModelingComponents/Private/ToolSetupUtil.cpp b/Engine/Plugins/Experimental/MeshModelingToolset/Source/ModelingComponents/Private/ToolSetupUtil.cpp
--- a/Engine/Plugins/Experimental/MeshModelingToolset/Source/ModelingComponents/Private/ToolSetupUtil.cpp
+++ b/Engine/Plugins/Experimental/MeshModelingToolset/Source/ModelingComponents/Private/ToolSetupUtil.cpp
@@ -180,6 +180,20 @@ UMaterialInstanceDynamic* ToolSetupUtil::GetSimpleCustomMaterial(UInteractiveToo
 	return nullptr;
 }
 
+UMaterialInstanceDynamic* ToolSetupUtil::GetCustomDepthOffsetMaterial(UInteractiveToolManager* ToolManager, const FLinearColor& Color, float PercentDepthOffset, float Opacity)
+{
+	UMaterialInterface* Material = LoadObject<UMaterial>(nullptr, TEXT("/MeshModelingToolset/Materials/SimpleOffsetMaterial"));
+	if (Material != nullptr)
+	{
+		UMaterialInstanceDynamic* MatInstance = UMaterialInstanceDynamic::Create(Material, ToolManager);
+		MatInstance->SetVectorParameterValue(TEXT("Color"), Color);
+		MatInstance->SetScalarParameterValue(TEXT("PercentDepthOffset"), PercentDepthOffset);
+		MatInstance->SetScalarParameterValue(TEXT("Opacity"), Opacity);
+		return MatInstance;
+	}
+	return nullptr;
+}
+
 UMaterialInterface* ToolSetupUtil::GetDefaultEditVolumeMaterial()
 {
 	return LoadObject<UMaterial>(nullptr, TEXT("/MeshModelingToolset/Materials/VolumeEditMaterial"));
diff --git a/Engine/Plugins/Experimental/MeshModelingToolset/Source/ModelingComponents/Public/ToolSetupUtil.h b/Engine/Plugins/Experimental/MeshModelingToolset/Source/ModelingComponents/Public/ToolSetupUtil.h
--- a/Engine/Plugins/Experimental/MeshModelingToolset/Source/ModelingComponents/Public/ToolSetupUtil.h
+++ b/Engine/Plugins/Experimental/MeshModelingToolset/Source/ModelingComponents/Public/ToolSetupUtil.h
@@ -94,6 +94,11 @@ namespace ToolSetupUtil
 	 */
 	MODELINGCOMPONENTS_API UMaterialInstanceDynamic* GetSimpleCustomMaterial(UInteractiveToolManager* ToolManager, const FLinearColor& Color, float Opacity);
 
+	/**
+	 * @return Simple material with configurable depth offset, color, and opacity.
+	 */
+	MODELINGCOMPONENTS_API UMaterialInstanceDynamic* GetCustomDepthOffsetMaterial(UInteractiveToolManager* ToolManager, const FLinearColor& Color, float DepthOffset, float Opacity);
+
 	/**
 	 * @return Material used when editing AVolume objects using our tools.
 	 */
diff --git a/Engine/Plugins/Experimental/MeshModelingToolset/Source/ModelingOperators/Private/MeshSpaceDeformerOps/BendMeshOp.cpp b/Engine/Plugins/Experimental/MeshModelingToolset/Source/ModelingOperators/Private/MeshSpaceDeformerOps/BendMeshOp.cpp
--- a/Engine/Plugins/Experimental/MeshModelingToolset/Source/ModelingOperators/Private/MeshSpaceDeformerOps/BendMeshOp.cpp
+++ b/Engine/Plugins/Experimental/MeshModelingToolset/Source/ModelingOperators/Private/MeshSpaceDeformerOps/BendMeshOp.cpp
@@ -15,6 +15,18 @@ void FBendMeshOp::CalculateResult(FProgressCancel* Progress)
 		return;
 	}
 
+	// You can think of the bend as follows: pick a range along the Z axis, 
+	// and pick a rotation point with its Z coordinate within that range.
+	// Take the mesh region in that range and compress it entirely to the
+	// rotation center Z. Then rotate the points out from here, choosing 
+	// each angle of rotation based on the original Z difference of the 
+	// point from the rotation center Z, but clamped to min/max.
+
+	// The math below follows the paper "Global and Local Deformations Of Solid Primitives"
+	// by Alan H. Barr, that uses a right-handed, y-up coordinate system and bends
+	// a line parallel to the y axis. We translate our system to make things follow 
+	// along smoother.
+
 	// Matrix from gizmo space (z-up) to a y-up space
 	FMatrix ToYUp(EForceInit::ForceInitToZero);
 	ToYUp.M[0][0] =  1.f;
@@ -39,37 +51,44 @@ void FBendMeshOp::CalculateResult(FProgressCancel* Progress)
 
 	float Det = ObjectToYUpGizmo.Determinant();
 
-	// Check if the transform is nearly singular
-	// this could happen if the scale on the object to world transform has a very small component.
-	if (FMath::Abs(Det) < 1.e-4)
+	// Check to see that the transform is invertible. The space deformation happens
+	// on the transformed mesh, and we'd need to be able to map the new positions back
+	// into the untransformed mesh.
+	if (FMath::Abs(Det) < KINDA_SMALL_NUMBER)
 	{
 		return;
 	}
 
 	FMatrix YUpGizmoToObject = ObjectToYUpGizmo.Inverse();
 	
-	
-	const double Curvature = GetModifierValue();
-	
-	// early out if nothing has been requested.
-	if (FMath::Abs(Curvature) < 0.001)
+	// We want to bend towards the positive Z axis, which corresponds to a progressively
+	// more clockwise rotation of points as one moves up on the Y axis. Hence the negative
+	// in the expression below.
+	const double DegreesToRadians = 0.017453292519943295769236907684886; // Pi / 180
+	const double TotalBendAngleRadians = -BendDegrees * DegreesToRadians;
+
+	// Check to see that we're actually trying to bend some amount.
+	if (FMath::Abs(TotalBendAngleRadians) < KINDA_SMALL_NUMBER)
 	{
 		return;
 	}
-
-
-
-	const double DegreesToRadians = 0.017453292519943295769236907684886; // Pi / 180
-	const double ThetaRadians = DegreesToRadians * Curvature;
-	//const FVector3d OpOrigin = ToOpSpace * AxisOriginObjectSpace;
-	
 	
 	// bounds in Op space
-	const double YMin =  -LowerBoundsInterval * AxesHalfLength;
-	const double YMax =   UpperBoundsInterval * AxesHalfLength;
+	const double YMin = LowerBoundsInterval;
+	const double YMax = UpperBoundsInterval;
+	const double BentLength = YMax - YMin;
 
-	const double Y0 = 0.;
-	const double K = ThetaRadians / AxesHalfLength;
+	bool bSharpBend = (BentLength == 0);
+
+	// Y0 is our rotation center Y, the level onto which everything is compressed before being
+	// rotated out.
+	const double Y0 = bLockBottom ? YMin : 0;
+
+	// K is the amount we rotate per unit step up Y in the bent region. Ik, its inverse,
+	// is actually our rotation center Z, because Ik is BentLength/TotalBendAngleRadians, 
+	// so it is the radius of the arc.
+	// When BentLength == 0, K doesn't actually matter- we don't use it.
+	const double K = bSharpBend ? 1.0 : TotalBendAngleRadians / BentLength;
 	const double Ik = 1.0 / K;
 
 	if (Progress && Progress->Cancelled())
@@ -77,12 +96,11 @@ void FBendMeshOp::CalculateResult(FProgressCancel* Progress)
 		return;
 	}
 
+	// Deal with the normals before we change the vertex positions (which determine the rotations)
 	if (ResultMesh->HasAttributes())
 	{
-		// Fix the normals first if they exist.
-
 		FDynamicMeshNormalOverlay* Normals = ResultMesh->Attributes()->PrimaryNormals();
-		ParallelFor(Normals->MaxElementID(), [this, Normals, &ObjectToYUpGizmo, &YUpGizmoToObject, YMin, YMax, Y0, K, Ik](int32 ElID)
+		ParallelFor(Normals->MaxElementID(), [this, Normals, &ObjectToYUpGizmo, &YUpGizmoToObject, YMin, YMax, Y0, K, Ik, bSharpBend, TotalBendAngleRadians](int32 ElID)
 		{
 			if (!Normals->IsElement(ElID))
 			{
@@ -93,10 +111,7 @@ void FBendMeshOp::CalculateResult(FProgressCancel* Progress)
 			auto VertexID = Normals->GetParentVertex(ElID);
 			const FVector3d& SrcPos = ResultMesh->GetVertex(VertexID);
 
-			FVector3f SrcNormalF = Normals->GetElement(ElID);
-			FVector3d SrcNormal;
-			SrcNormal[0] = SrcNormalF[0]; SrcNormal[1] = SrcNormalF[1]; SrcNormal[2] = SrcNormalF[2];
-
+			FVector3d SrcNormal(Normals->GetElement(ElID));
 
 			const double SrcPos4[4] = { SrcPos[0], SrcPos[1], SrcPos[2], 1.0 };
 
@@ -124,29 +139,28 @@ void FBendMeshOp::CalculateResult(FProgressCancel* Progress)
 
 			const double YHat = FMath::Clamp(GizmoPos4[1], YMin, YMax);
 
-			const double Theta = K * (YHat - Y0);
+			// In the normal case, the bend angle is proportional to the clamped value's
+			// position in the bend range. In the sharp bend case, however, where the bent
+			// range is 0-sized, there is a jump from 0 to TotalBendAngleRadians at Y0. 
+			const double Theta = bSharpBend && GizmoPos4[1] > Y0 ? TotalBendAngleRadians : K * (YHat - Y0);
 
 			const double S0 = TMathUtil<double>::Sin(Theta);
 			const double C0 = TMathUtil<double>::Cos(Theta);
 
-			const double ZP = GizmoPos4[2] - Ik;
-
-
-
-			// transform normal.  Do this before changing GizmoPos
-
 			{
+				// The below expression for the normal transformation comes from the paper, but can be derived
+				// from the inverse of the Jacobian of the transformation that we're applying. We could actually
+				// divide by rate of expansion to keep the normal the same length, but it's cleaner just to set
+				// a normalized vector at the end.
 
-				double DTheta = K;
-				if (GizmoPos4[1] < YMin || GizmoPos4[1] > YMax) DTheta = 0.;
+				double KHat = (!bSharpBend && YHat == GizmoPos4[1]) ? K : 0;
 
-				double tmp = (1. - DTheta * GizmoPos4[2]);
+				double RateOfExpansion = (1. - KHat * GizmoPos4[2]);
 				FVector3d DstNormal(0., 0., 0.);
-				DstNormal[0] = tmp * RotatedNormal[0];
-				DstNormal[1] =                          C0 * RotatedNormal[1] - S0 * tmp * RotatedNormal[2];
-				DstNormal[2] =                          S0 * RotatedNormal[1] + C0 * tmp * RotatedNormal[2];
+				DstNormal[0] = RateOfExpansion * RotatedNormal[0];
+				DstNormal[1] =                          C0 * RotatedNormal[1] - S0 * RateOfExpansion * RotatedNormal[2];
+				DstNormal[2] =                          S0 * RotatedNormal[1] + C0 * RateOfExpansion * RotatedNormal[2];
 
-			
 				// rotate back to mesh space.
 				RotatedNormal = FVector3d(0, 0, 0);
 				for (int i = 0; i < 3; ++i)
@@ -159,9 +173,7 @@ void FBendMeshOp::CalculateResult(FProgressCancel* Progress)
 
 			}
 
-			FVector3f RotatedNoramlF;
-			RotatedNoramlF[0] = RotatedNormal[0]; RotatedNoramlF[1] = RotatedNormal[1]; RotatedNoramlF[2] = RotatedNormal[2];
-			Normals->SetElement(ElID, RotatedNoramlF);
+			Normals->SetElement(ElID, FVector3f(RotatedNormal.Normalized()));
 		});
 	}
 
@@ -170,14 +182,34 @@ void FBendMeshOp::CalculateResult(FProgressCancel* Progress)
 		return;
 	}
 
-	ParallelFor(ResultMesh->MaxVertexID(), [this, &ObjectToYUpGizmo, &YUpGizmoToObject, YMin, YMax, Y0, K, Ik](int32 VertexID)
+	// Adjust the positions.
+	ParallelFor(ResultMesh->MaxVertexID(), [this, &ObjectToYUpGizmo, &YUpGizmoToObject, YMin, YMax, Y0, K, Ik, bSharpBend, TotalBendAngleRadians](int32 VertexID)
 	{
 		if (!ResultMesh->IsVertex(VertexID))
 		{
 			return;
 		}
 
-		//const FVector3d& SrcPos = OriginalPositions[VertexID];
+		// The approach we use follows the paper, and uses the variable names used there. However
+		// it may be useful to see some equivalent code that may make the operations more obvious. 
+		// This code works on the original Z-up frame and bends along the Z axis instead of Y, so
+		// you would swap Y and Z variables to make it work in the Y-up frame:
+		/*
+		double ClampedZ = FMath::Clamp(GizmoSpacePosition.Z, BentMinZ, BentMaxZ);
+		double ZDifferenceFromClampedRegion = GizmoSpacePosition.Z - ClampedZ;
+
+		double AngleToRotate = ArcAngle * (ClampedZ - RotationCenterZ) / BentLength; // assumes BentLength != 0
+
+		FVector2d YZToRotate(GizmoSpacePosition.Y, RotationCenterZ + ZDifferenceFromClampedRegion);
+
+		FVector2d RotationCenterYZ(RotationCenterY, RotationCenterZ);
+		FMatrix2d RotationMatrix = FMatrix2d::RotationRad(AngleToRotate);
+
+		FVector2d RotatedYZ = RotationMatrix * (YZToRotate - RotationCenterYZ) + RotationCenterYZ;
+
+		FVector3d NewGizmoSpacePosition = FVector3d(GizmoSpacePosition.X, RotatedYZ[0], RotatedYZ[1]);
+		*/
+
 		const FVector3d SrcPos = ResultMesh->GetVertex(VertexID);
 		const double SrcPos4[4] = { SrcPos[0], SrcPos[1], SrcPos[2], 1.0};
 
@@ -191,21 +223,20 @@ void FBendMeshOp::CalculateResult(FProgressCancel* Progress)
 			}
 		}
 
-
-		//FVector3d DstPos = ToOpSpace * (SrcPos) - OpOrigin;
-
 		const double YHat = FMath::Clamp(GizmoPos4[1], YMin, YMax);
 
-		const double Theta = K * (YHat - Y0);
+		// In the usual case, the bend angle is proportional to the clamped value's
+		// position in the bend range. In the sharp bend case, however, where the bent
+		// range is 0-sized, there is a jump from 0 to TotalBendAngleRadians at Y0. 
+		const double Theta = bSharpBend && GizmoPos4[1] > Y0 ? TotalBendAngleRadians : K * (YHat - Y0);
 
 		const double S0 = TMathUtil<double>::Sin(Theta);
 		const double C0 = TMathUtil<double>::Cos(Theta);
 
-		const double ZP = GizmoPos4[2] - Ik;
+		const double ZRelativeRotationCenter = GizmoPos4[2] - Ik;
 
-		
-		double Y = -S0 * ZP + Y0;
-		double Z =  C0 * ZP + Ik;
+		double Y = -S0 * ZRelativeRotationCenter + Y0;
+		double Z =  C0 * ZRelativeRotationCenter + Ik;
 
 		if (GizmoPos4[1] > YMax)
 		{
diff --git a/Engine/Plugins/Experimental/MeshModelingToolset/Source/ModelingOperators/Private/MeshSpaceDeformerOps/FlareMeshOp.cpp b/Engine/Plugins/Experimental/MeshModelingToolset/Source/ModelingOperators/Private/MeshSpaceDeformerOps/FlareMeshOp.cpp
--- a/Engine/Plugins/Experimental/MeshModelingToolset/Source/ModelingOperators/Private/MeshSpaceDeformerOps/FlareMeshOp.cpp
+++ b/Engine/Plugins/Experimental/MeshModelingToolset/Source/ModelingOperators/Private/MeshSpaceDeformerOps/FlareMeshOp.cpp
@@ -54,8 +54,8 @@ void FFlareMeshOp::CalculateResult(FProgressCancel* Progress)
 
 	FMatrix GizmoToObject = ObjectToGizmo.Inverse();
 
-	const double ZMin = -LowerBoundsInterval * AxesHalfLength;
-	const double ZMax =  UpperBoundsInterval * AxesHalfLength;
+	const double ZMin = LowerBoundsInterval;
+	const double ZMax =  UpperBoundsInterval;
 
 	if (ResultMesh->HasAttributes())
 	{
@@ -105,20 +105,38 @@ void FFlareMeshOp::CalculateResult(FProgressCancel* Progress)
 
 			const double T = FMath::Clamp((GizmoPos4[2] - ZMin) / (ZMax - ZMin), 0.0, 1.0);
 
-			const double R = 1. + FMath::Sin(PI * T) * (ModifierPercent / 100.f);
+			double Rx = bSmoothEnds ? 1. + (FMath::Cos(2 * PI * T - PI) + 1) * (FlarePercentX / 200.f) // Shift cos curve up 1, right PI, scale down by 2, go from 0 to 2PI
+				: 1. + FMath::Sin(PI * T) * (FlarePercentX / 100.f);
+			double Ry = bSmoothEnds ? 1. + (FMath::Cos(2 * PI * T - PI) + 1) * (FlarePercentY / 200.f) // Shift cos curve up 1, right PI, scale down by 2, go from 0 to 2PI
+				: 1. + FMath::Sin(PI * T) * (FlarePercentY / 100.f);
 			
 			// transform normal.  Do this before changing GizmoPos
-
+			// To get the normal, note that the positions are transformed like this:
+			// X = Rx * x
+			// Y = Ry * y
+			// Z = z
+			// The jacobian is:
+			// Rx  0  x*DRx
+			// 0  Ry  y*DRy
+			// 0   0   1
+			// Where DRx is dRx/dz and DRy is dRy/dz
+			// Then take the transpose of the inverse of the Jacobian and multiply by the determinant (or don't, since we don't
+			// care about the length, but it's cleaner if you do).
 			{
+				double DRx = bSmoothEnds ? -FMath::Sin(2 * PI * T - PI) * (2 * PI / (ZMax - ZMin)) * (FlarePercentX / 200.f)
+					: FMath::Cos(PI * T) * (PI / (ZMax - ZMin)) * (FlarePercentX / 100.f);
+				double DRy = bSmoothEnds ? -FMath::Sin(2 * PI * T - PI) * (2 * PI / (ZMax - ZMin)) * (FlarePercentY / 200.f)
+					: FMath::Cos(PI * T) * (PI / (ZMax - ZMin)) *  (FlarePercentY / 100.f);
 
-				double DR = FMath::Cos(PI * T) * (PI / (ZMax - ZMin)) *  (ModifierPercent / 100.f);
-				if (GizmoPos4[2] > ZMax || GizmoPos4[2] < ZMin)  DR = 0.f;
+				if (GizmoPos4[2] > ZMax || GizmoPos4[2] < ZMin)
+				{
+					DRx = DRy = 0.f;
+				}
 			
-
 				FVector3d DstNormal(0., 0., 0.);
-				DstNormal[0] = R * RotatedNormal[0];
-				DstNormal[1] = R * RotatedNormal[1];
-				DstNormal[2] = -R * DR * (RotatedNormal[0] * GizmoPos4[0] + RotatedNormal[1] * GizmoPos4[1]) + R * R * RotatedNormal[2];
+				DstNormal[0] = Ry * RotatedNormal[0];
+				DstNormal[1] = Rx * RotatedNormal[1];
+				DstNormal[2] = -Ry * DRx * GizmoPos4[0] * RotatedNormal[0] - Rx * DRy * GizmoPos4[1] * RotatedNormal[1] + Rx * Ry * RotatedNormal[2];
 
 				// rotate back to mesh space.
 				RotatedNormal = FVector3d(0, 0, 0);
@@ -132,9 +150,7 @@ void FFlareMeshOp::CalculateResult(FProgressCancel* Progress)
 
 			}
 
-			FVector3f RotatedNoramlF;
-			RotatedNoramlF[0] = RotatedNormal[0]; RotatedNoramlF[1] = RotatedNormal[1]; RotatedNoramlF[2] = RotatedNormal[2];
-			Normals->SetElement(ElID, RotatedNoramlF);
+			Normals->SetElement(ElID, FVector3f(RotatedNormal.Normalized()));
 		});
 	}
 
@@ -168,11 +184,14 @@ void FFlareMeshOp::CalculateResult(FProgressCancel* Progress)
 		// Parameterize curve between ZMin and ZMax to go between 0, 1
 		double T = FMath::Clamp( (GizmoPos4[2]- ZMin) / (ZMax - ZMin), 0.0, 1.0);
 
-		double R = 1. + FMath::Sin(PI * T) * (ModifierPercent / 100.f);
+		double Rx = bSmoothEnds ? 1. + (FMath::Cos(2 * PI * T - PI) + 1) * (FlarePercentX / 200.f) // Shift cos curve up 1, right PI, scale down by 2, go from 0 to 2PI
+			: 1. + FMath::Sin(PI * T) * (FlarePercentX / 100.f);
+		double Ry = bSmoothEnds ? 1. + (FMath::Cos(2 * PI * T - PI) + 1) * (FlarePercentY / 200.f) // Shift cos curve up 1, right PI, scale down by 2, go from 0 to 2PI
+			: 1. + FMath::Sin(PI * T) * (FlarePercentY / 100.f);
 
 		// 2d scale x,y values.
-		GizmoPos4[0] *= R;
-		GizmoPos4[1] *= R;
+		GizmoPos4[0] *= Rx;
+		GizmoPos4[1] *= Ry;
 
 
 		double DstPos4[4] = { 0., 0., 0., 0. };
diff --git a/Engine/Plugins/Experimental/MeshModelingToolset/Source/ModelingOperators/Private/MeshSpaceDeformerOps/TwistMeshOp.cpp b/Engine/Plugins/Experimental/MeshModelingToolset/Source/ModelingOperators/Private/MeshSpaceDeformerOps/TwistMeshOp.cpp
--- a/Engine/Plugins/Experimental/MeshModelingToolset/Source/ModelingOperators/Private/MeshSpaceDeformerOps/TwistMeshOp.cpp
+++ b/Engine/Plugins/Experimental/MeshModelingToolset/Source/ModelingOperators/Private/MeshSpaceDeformerOps/TwistMeshOp.cpp
@@ -32,12 +32,12 @@ void FTwistMeshOp::CalculateResult(FProgressCancel* Progress)
 	
 
 
-	const double ZMin = -LowerBoundsInterval * AxesHalfLength;
-	const double ZMax =  UpperBoundsInterval * AxesHalfLength;
+	const double ZMin = LowerBoundsInterval;
+	const double ZMax =  UpperBoundsInterval;
 
 
 	const double DegreesToRadians = 0.017453292519943295769236907684886; // Pi / 180
-	const double ThetaRadians = DegreesToRadians * GetModifierValue();
+	const double ThetaRadians = DegreesToRadians * TwistDegrees;
 
 	if (ResultMesh->HasAttributes())
 	{
@@ -86,7 +86,8 @@ void FTwistMeshOp::CalculateResult(FProgressCancel* Progress)
 			}
 
 
-			double T = FMath::Clamp((GizmoPos4[2] - ZMin) / (ZMax - ZMin), 0.0, 1.0) - 0.5;
+			double T = FMath::Clamp((GizmoPos4[2] - ZMin) / (ZMax - ZMin), 0.0, 1.0) 
+				- (bLockBottom ? 0 : 0.5);
 			double Theta = ThetaRadians * T;
 
 
@@ -119,9 +120,7 @@ void FTwistMeshOp::CalculateResult(FProgressCancel* Progress)
 				
 			}
 
-			FVector3f RotatedNoramlF;
-			RotatedNoramlF[0] = RotatedNormal[0]; RotatedNoramlF[1] = RotatedNormal[1]; RotatedNoramlF[2] = RotatedNormal[2];
-			Normals->SetElement(ElID, RotatedNoramlF);
+			Normals->SetElement(ElID, FVector3f(RotatedNormal.Normalized()));
 		});
 	}
 
@@ -151,7 +150,8 @@ void FTwistMeshOp::CalculateResult(FProgressCancel* Progress)
 			}
 		}
 
-		double T = FMath::Clamp( (GizmoPos4[2] - ZMin) / ( ZMax - ZMin ), 0.0, 1.0) - 0.5;
+		double T = FMath::Clamp( (GizmoPos4[2] - ZMin) / ( ZMax - ZMin ), 0.0, 1.0) 
+			- (bLockBottom ? 0 : 0.5);
 		double Theta = ThetaRadians * T;
 
 		
diff --git a/Engine/Plugins/Experimental/MeshModelingToolset/Source/ModelingOperators/Public/SpaceDeformerOps/BendMeshOp.h b/Engine/Plugins/Experimental/MeshModelingToolset/Source/ModelingOperators/Public/SpaceDeformerOps/BendMeshOp.h
--- a/Engine/Plugins/Experimental/MeshModelingToolset/Source/ModelingOperators/Public/SpaceDeformerOps/BendMeshOp.h
+++ b/Engine/Plugins/Experimental/MeshModelingToolset/Source/ModelingOperators/Public/SpaceDeformerOps/BendMeshOp.h
@@ -9,9 +9,10 @@
 class MODELINGOPERATORS_API FBendMeshOp : public FMeshSpaceDeformerOp
 {
 public:
-	virtual ~FBendMeshOp() {}
-	FBendMeshOp() : FMeshSpaceDeformerOp(0.0, 180.1) {};
 	virtual void CalculateResult(FProgressCancel* Progress) override;
 
+	double BendDegrees = 90;
+	bool bLockBottom = false;
+
 protected:
 };
\ No newline at end of file
diff --git a/Engine/Plugins/Experimental/MeshModelingToolset/Source/ModelingOperators/Public/SpaceDeformerOps/FlareMeshOp.h b/Engine/Plugins/Experimental/MeshModelingToolset/Source/ModelingOperators/Public/SpaceDeformerOps/FlareMeshOp.h
--- a/Engine/Plugins/Experimental/MeshModelingToolset/Source/ModelingOperators/Public/SpaceDeformerOps/FlareMeshOp.h
+++ b/Engine/Plugins/Experimental/MeshModelingToolset/Source/ModelingOperators/Public/SpaceDeformerOps/FlareMeshOp.h
@@ -10,11 +10,19 @@
 class MODELINGOPERATORS_API FFlareMeshOp : public FMeshSpaceDeformerOp
 {
 public:
-	virtual ~FFlareMeshOp() {}
-	FFlareMeshOp() : FMeshSpaceDeformerOp(-2.0, 2.0) {};
 	virtual void CalculateResult(FProgressCancel* Progress) override;
 
-
+	/** 0% does nothing, 100% moves 2x away from Z axis at extremal point, -100% squishes down to Z axis at extremal point. */
+	double FlarePercentX = 100;
+	double FlarePercentY = 100;
+
+	/** 
+	 * Changes the type of flaring. When false, the flaring is the curve of sin(y) from 0 to Pi,
+	 * which makes the ends of the flaring sharp because the derivative is not 0. When true, the 
+	 * flaring is the curve of cos(y) + 1 from -Pi to Pi, which makes the ends smooth out back
+	 * into the shape.
+	 */
+	bool bSmoothEnds = false;
 protected:
 
 };
\ No newline at end of file
diff --git a/Engine/Plugins/Experimental/MeshModelingToolset/Source/ModelingOperators/Public/SpaceDeformerOps/MeshSpaceDeformerOp.h b/Engine/Plugins/Experimental/MeshModelingToolset/Source/ModelingOperators/Public/SpaceDeformerOps/MeshSpaceDeformerOp.h
--- a/Engine/Plugins/Experimental/MeshModelingToolset/Source/ModelingOperators/Public/SpaceDeformerOps/MeshSpaceDeformerOp.h
+++ b/Engine/Plugins/Experimental/MeshModelingToolset/Source/ModelingOperators/Public/SpaceDeformerOps/MeshSpaceDeformerOp.h
@@ -9,27 +9,16 @@
 class MODELINGOPERATORS_API FMeshSpaceDeformerOp : public FDynamicMeshOperator
 {
 public:
-	FMeshSpaceDeformerOp(double Min, double Max) : ModifierMin{ Min }, ModifierMax{ Max }{};
-	virtual ~FMeshSpaceDeformerOp() {}
 
 	// Inputs
 	TSharedPtr<const FDynamicMesh3> OriginalMesh;
 	FFrame3d GizmoFrame;
 	void SetTransform(const FTransform& Transform);
 
-	//Converts the given modifier percent and returns the value between -1 and 1
-	virtual double GetModifierValue() { return FMath::Lerp(0.01 * ModifierPercent, ModifierMin, ModifierMax); }
-
-	// half the major axis of the bounding box for the source geometry
-	double AxesHalfLength;
-
-	//Percent input from the tool (between -100 and 100), as each operator will have a min/max value i.e. the angle of curvature, twist or scale. This is the percent used to interpolate between the min/max value
-	double ModifierPercent;
-
-	//The interval usually from [-1,0] upon which this operator will work. This corresponds to the region of space which this operator affects.
+	// The lower bound of the region of space that the operator affects, relative to the gizmo position.
 	double LowerBoundsInterval;
 
-	//The interval usually from [0,1] upon which this operator will work. This corresponds to the region of space which this operator affects.
+	// The upper bound of the region of space that the operator affects, relative to the gizmo position.
 	double UpperBoundsInterval;
 
 	/**
@@ -43,10 +32,4 @@ public:
 protected:
 	// transform, including translation, to gizmo space
 	FMatrix ObjectToGizmo;
-
-	//The minimum value provided by the constructor so that each operator can set its own interpolation min and max
-	double ModifierMin{};
-
-	//The minimum value provided by the constructor so that each operator can set its own interpolation min and max
-	double ModifierMax{};
 };
\ No newline at end of file
diff --git a/Engine/Plugins/Experimental/MeshModelingToolset/Source/ModelingOperators/Public/SpaceDeformerOps/TwistMeshOp.h b/Engine/Plugins/Experimental/MeshModelingToolset/Source/ModelingOperators/Public/SpaceDeformerOps/TwistMeshOp.h
--- a/Engine/Plugins/Experimental/MeshModelingToolset/Source/ModelingOperators/Public/SpaceDeformerOps/TwistMeshOp.h
+++ b/Engine/Plugins/Experimental/MeshModelingToolset/Source/ModelingOperators/Public/SpaceDeformerOps/TwistMeshOp.h
@@ -10,12 +10,10 @@
 class MODELINGOPERATORS_API FTwistMeshOp : public FMeshSpaceDeformerOp
 {
 public:
-	virtual ~FTwistMeshOp() {}
-	FTwistMeshOp() : FMeshSpaceDeformerOp(0.0, 180.0) {}
-
 	virtual void CalculateResult(FProgressCancel* Progress) override;
 
-
+	double TwistDegrees = 180;
+	bool bLockBottom = false;
 protected:
 
 };
\ No newline at end of file
