MeshModelingTools: Fix CurveControlPointsMechanic issuing transactions in the wrong order on point deletion.

#rnx

diff --git a/Engine/Plugins/Experimental/MeshModelingToolset/Source/ModelingComponents/Private/Mechanics/CurveControlPointsMechanic.cpp b/Engine/Plugins/Experimental/MeshModelingToolset/Source/ModelingComponents/Private/Mechanics/CurveControlPointsMechanic.cpp
--- a/Engine/Plugins/Experimental/MeshModelingToolset/Source/ModelingComponents/Private/Mechanics/CurveControlPointsMechanic.cpp
+++ b/Engine/Plugins/Experimental/MeshModelingToolset/Source/ModelingComponents/Private/Mechanics/CurveControlPointsMechanic.cpp
@@ -1017,28 +1017,25 @@ void UCurveControlPointsMechanic::DeleteSelectedPoints()
 	ParentTool->GetToolManager()->BeginUndoTransaction(PointDeletionTransactionText);
 
 	// There are minor inefficiencies in the way we delete multiple points since we sometimes do edge updates
-	// for edges that get deleted later in the loop, and we upate the map inside ControlPoints gets updated each
+	// for edges that get deleted later in the loop, and we upate the map inside ControlPoints each
 	// time, but avoiding these would make the code more cumbersome.
 
-	// For the purposes of undo/redo, it is more convenient to clear the selection before deleting the points, so that
-	// on undo, the points get added back before being reselected.
+	// We are deselecting, deleting, and (potentially) reverting back to interactive initialization. 
+	// When undoing all of this, we want to add the points first, potentially switch modes second 
+	// (now that we have enough points), and reselect third (now that the points exist). So, we want
+	// to perform the operations in reverse order (deselect, switch modes, delete).
+
+	// Deselect
 	TArray<int32> PointsToDelete = SelectedPointIDs;
 	ClearSelection();
 
-	for (int32 PointID : PointsToDelete)
-	{
-		ParentTool->GetToolManager()->EmitObjectChange(this, MakeUnique<FCurveControlPointsMechanicInsertionChange>(
-			ControlPoints.GetSequencePosition(PointID), PointID, ControlPoints.GetPointCoordinates(PointID),
-			false, CurrentChangeStamp), PointDeletionTransactionText);
-		DeletePoint(PointID);
-	}
-
-	// If we deleted too many points, we may need to switch modes.
-	bool bChangedMode = false;
+	// Revert back to interactive intialization if needed.
 	if (bAutoRevertToInteractiveInitialization)
 	{
-		if ((bIsLoop && ControlPoints.Num() < MinPointsForLoop)
-			|| (!bIsLoop && ControlPoints.Num() < MinPointsForNonLoop))
+		int32 NumPointsRemaining = ControlPoints.Num() - PointsToDelete.Num();
+
+		if ((bIsLoop && NumPointsRemaining < MinPointsForLoop)
+			|| (!bIsLoop && NumPointsRemaining < MinPointsForNonLoop))
 		{
 			// We emit first just out of convenience of not having to temporarily store bIsLoop, 
 			// which gets reset during the SetInteractiveIntialization(true) call.
@@ -1046,19 +1043,22 @@ void UCurveControlPointsMechanic::DeleteSelectedPoints()
 				false, bIsLoop, CurrentChangeStamp), InitializationCompletedTransactionText);
 
 			SetInteractiveInitialization(true);
-			bChangedMode = true;
+			OnModeChanged.Broadcast();
 		}
 	}
 
-	ParentTool->GetToolManager()->EndUndoTransaction();
-
-	// We broadcast the mode change first because clients are likely
-	// to respond differently to point changes in different modes.
-	if (bChangedMode)
+	// Delete
+	for (int32 PointID : PointsToDelete)
 	{
-		OnModeChanged.Broadcast();
+		ParentTool->GetToolManager()->EmitObjectChange(this, MakeUnique<FCurveControlPointsMechanicInsertionChange>(
+			ControlPoints.GetSequencePosition(PointID), PointID, ControlPoints.GetPointCoordinates(PointID),
+			false, CurrentChangeStamp), PointDeletionTransactionText);
+		DeletePoint(PointID);
 	}
+
 	OnPointsChanged.Broadcast();
+
+	ParentTool->GetToolManager()->EndUndoTransaction();
 }
 
 int32 UCurveControlPointsMechanic::DeletePoint(int32 PointID)
