MeshModelingTools: Made the ConstructionPlaneMechanic align its frame more usefully on Ctrl clicks.

#rnx

diff --git a/Engine/Plugins/Experimental/MeshModelingToolset/Source/ModelingComponents/Private/Mechanics/ConstructionPlaneMechanic.cpp b/Engine/Plugins/Experimental/MeshModelingToolset/Source/ModelingComponents/Private/Mechanics/ConstructionPlaneMechanic.cpp
--- a/Engine/Plugins/Experimental/MeshModelingToolset/Source/ModelingComponents/Private/Mechanics/ConstructionPlaneMechanic.cpp
+++ b/Engine/Plugins/Experimental/MeshModelingToolset/Source/ModelingComponents/Private/Mechanics/ConstructionPlaneMechanic.cpp
@@ -81,9 +81,22 @@ void UConstructionPlaneMechanic::TransformChanged(UTransformProxy* Proxy, FTrans
 void UConstructionPlaneMechanic::SetDrawPlaneFromWorldPos(const FVector3d& Position, const FVector3d& Normal, bool bIgnoreNormal)
 {
 	Plane.Origin = Position;
-	if (bIgnoreNormal == false)
+	if (!bIgnoreNormal)
 	{
-		Plane.AlignAxis(2, Normal);
+		// The normal will be our frame Z, but the other axes are unconstrained. If the normal
+		// is aligned with world Z, then the entire frame might as well be aligned with world.
+		if (1 - Normal.Dot(FVector3d::UnitZ()) < KINDA_SMALL_NUMBER)
+		{
+			Plane.Rotation = FQuaterniond::Identity();
+		}
+		else
+		{
+			// Otherwise, let's place one of the other axes into the XY plane so that the frame is more
+			// useful for translation. We somewhat arbitrarily choose Y for this. 
+			FVector3d FrameY = Normal.Cross(FVector3d::UnitZ()).Normalized(); // orthogonal to world Z and frame Z 
+			FVector3d FrameX = FrameY.Cross(Normal); // safe to not normalize because already orthogonal
+			Plane = FFrame3d(Position, FrameX, FrameY, Normal);
+		}
 	}
 	PlaneTransformGizmo->SetActiveTarget(PlaneTransformProxy, GetParentTool()->GetToolManager());
 	PlaneTransformGizmo->SetNewGizmoTransform(Plane.ToFTransform());
