ModelingComponents: remove unnecessary log message, add parallel compute of update buffer sets
#rnx

diff --git a/Engine/Plugins/Experimental/MeshModelingToolset/Source/ModelingComponents/Private/SimpleDynamicMeshComponent.cpp b/Engine/Plugins/Experimental/MeshModelingToolset/Source/ModelingComponents/Private/SimpleDynamicMeshComponent.cpp
--- a/Engine/Plugins/Experimental/MeshModelingToolset/Source/ModelingComponents/Private/SimpleDynamicMeshComponent.cpp
+++ b/Engine/Plugins/Experimental/MeshModelingToolset/Source/ModelingComponents/Private/SimpleDynamicMeshComponent.cpp
@@ -666,13 +666,27 @@ TFuture<bool> USimpleDynamicMeshComponent::FastNotifyTriangleVerticesUpdated_Try
 		TFuture<void> ComputeSets = Async(SimpleDynamicMeshComponentAsyncExecTarget, [this, &UpdateSetsOut, &Triangles]()
 		{
 			TRACE_CPUPROFILER_EVENT_SCOPE(SimpleDynamicMeshComponent_FastVertexUpdatePrecomp_FindSets);
+			int32 NumBuffers = Decomposition->Num();
+			TArray<std::atomic<bool>> BufferFlags;
+			BufferFlags.SetNum(NumBuffers);
+			for (int32 k = 0; k < NumBuffers; ++k)
+			{
+				BufferFlags[k] = false;
+			}
+			ParallelFor(Triangles.Num(), [&](int32 k)
+			{
+				int32 SetID = Decomposition->GetGroupForTriangle(Triangles[k]);
+				BufferFlags[SetID] = true;
+			});
 			UpdateSetsOut.Reset();
-			for (int32 tid : Triangles)
+			for (int32 k = 0; k < NumBuffers; ++k)
 			{
-				int32 SetID = Decomposition->GetGroupForTriangle(tid);
-				UpdateSetsOut.AddUnique(SetID);
+				if (BufferFlags[k])
+				{
+					UpdateSetsOut.Add(k);
+				}
 			}
-			UE_LOG(LogTemp, Warning, TEXT("FOUND %d dirty sets of %d"), UpdateSetsOut.Num(), Decomposition->Num());
+
 		});
 
 		ComputeSets.Wait();
