ModelingTools: restore skip of stamp processing in every second stamp in DynamicMeshSculptTool, it does in fact improve editor responsiveness, in very low fps situations
#rnx

diff --git a/Engine/Plugins/Experimental/MeshModelingToolset/Source/MeshModelingTools/Private/DynamicMeshSculptTool.cpp b/Engine/Plugins/Experimental/MeshModelingToolset/Source/MeshModelingTools/Private/DynamicMeshSculptTool.cpp
--- a/Engine/Plugins/Experimental/MeshModelingToolset/Source/MeshModelingTools/Private/DynamicMeshSculptTool.cpp
+++ b/Engine/Plugins/Experimental/MeshModelingToolset/Source/MeshModelingTools/Private/DynamicMeshSculptTool.cpp
@@ -1478,12 +1478,14 @@ void UDynamicMeshSculptTool::OnTick(float DeltaTime)
 
 	ActivePressure = GetCurrentDevicePressure();
 
-	// Allow a tick to pass between application of brush stamps. Bizarrely this improves responsiveness in the Editor...
-	//static int TICK_SKIP_HACK = 0;
-	//if (TICK_SKIP_HACK++ % 2 == 0)
-	//{
-	//	return;
-	//}
+	// Allow a tick to pass between application of brush stamps. If we do not do this then on large stamps
+	// that take a significant fraction of a second to compute, frames will be skipped and the editor will appear
+	// frozen, but when the user releases the mouse, sculpting has clearly happened
+	static int TICK_SKIP_HACK = 0;
+	if (TICK_SKIP_HACK++ % 2 == 0)
+	{
+		return;
+	}
 
 	ShowWireframeWatcher.CheckAndUpdate();
 	MaterialModeWatcher.CheckAndUpdate();
