ModelingTools:
Add support for brush alphas to VertexSculptTool. Replace cycle counter instrumentation with insights tags. Perf optimizations for normals computation, additional async refactoring.
MeshSculptToolBase now allows subclasses to override final commit behavior, added support for variable stamp spacing and flow rate, expose various other internal functions to subclasses. Update various brush implementations for new stamp system.
Add support for custom brush stamp Alpha function to FSculptBrushStamp
Add FMeshSculptBrushOp::::SupportsVariableSpacing(), brush stamp impls use this to indicate that they support custom Alphas. Implement support in various Sculpt brushes.
Add various functions to recompute normals for a triangle region to UE::SculptUtil
FMeshVertexChangeBuilder now supports only updating vertex color.
#rnx

diff --git a/Engine/Plugins/Experimental/MeshModelingToolset/Source/MeshModelingTools/Private/MeshAttributePaintTool.cpp b/Engine/Plugins/Experimental/MeshModelingToolset/Source/MeshModelingTools/Private/MeshAttributePaintTool.cpp
--- a/Engine/Plugins/Experimental/MeshModelingToolset/Source/MeshModelingTools/Private/MeshAttributePaintTool.cpp
+++ b/Engine/Plugins/Experimental/MeshModelingToolset/Source/MeshModelingTools/Private/MeshAttributePaintTool.cpp
@@ -5,6 +5,8 @@
 #include "ToolBuilderUtil.h"
 #include "Drawing/MeshDebugDrawing.h"
 #include "ToolSetupUtil.h"
+#include "Selections/MeshConnectedComponents.h"
+
 
 #include "MeshDescription.h"
 
@@ -144,6 +146,11 @@ void UMeshAttributePaintTool::SetAssetAPI(IToolsContextAssetAPI* AssetAPIIn)
 
 void UMeshAttributePaintTool::Setup()
 {
+	// want this before brush size/etc
+	BrushActionProps = NewObject<UMeshAttributePaintBrushOperationProperties>(this);
+	BrushActionProps->RestoreProperties(this);
+	AddToolPropertySource(BrushActionProps);
+
 	UDynamicMeshBrushTool::Setup();
 
 	// hide strength and falloff
@@ -453,16 +460,51 @@ void UMeshAttributePaintTool::ApplyStamp(const FBrushStampData& Stamp)
 		return;
 	}
 
+	FAttributeData& AttribData = Attributes[CurrentAttributeIndex];
+
+	FStampActionData ActionData;
+	CalculateVertexROI(Stamp, ActionData.ROIVertices);
+
+	if (BrushActionProps->BrushAction == EBrushActionMode::FloodFill)
+	{
+		ApplyStamp_FloodFill(Stamp, ActionData);
+	}
+	else
+	{
+		ApplyStamp_Paint(Stamp, ActionData);
+	}
+
+
+	// track changes
+	if (ActiveChangeBuilder)
+	{
+		ActiveChangeBuilder->UpdateValues(ActionData.ROIVertices, ActionData.ROIBefore, ActionData.ROIAfter);
+	}
+
+	// update values and colors
+	PreviewMesh->DeferredEditMesh([&](FDynamicMesh3& Mesh)
+	{
+		int32 NumVertices = ActionData.ROIVertices.Num();
+		for (int32 k = 0; k < NumVertices; ++k)
+		{
+			int32 vid = ActionData.ROIVertices[k];
+			AttribData.CurrentValues[vid] = ActionData.ROIAfter[k];
+			FVector3f NewColor = ColorMapper->ToColor(ActionData.ROIAfter[k]);
+			Mesh.SetVertexColor(vid, NewColor);
+		}
+	}, false);
+	PreviewMesh->NotifyDeferredEditCompleted(UPreviewMesh::ERenderUpdateMode::FastUpdate, EMeshRenderAttributeFlags::VertexColors, false);
+}
+
+
+void UMeshAttributePaintTool::ApplyStamp_Paint(const FBrushStampData& Stamp, FStampActionData& ActionData)
+{
 	FTransform3d Transform(ComponentTarget->GetWorldTransform());
 	FVector3d StampPosLocal = Transform.InverseTransformPosition(Stamp.WorldPosition);
 
-	TArray<int32> ROIVertices;
-	CalculateVertexROI(Stamp, ROIVertices);
-	int NumVertices = ROIVertices.Num();
-
-	TArray<float> ROIBefore, ROIAfter;
-	ROIBefore.SetNum(NumVertices);
-	ROIAfter.SetNum(NumVertices);
+	int32 NumVertices = ActionData.ROIVertices.Num();
+	ActionData.ROIBefore.SetNum(NumVertices);
+	ActionData.ROIAfter.SetNum(NumVertices);
 
 	FAttributeData& AttribData = Attributes[CurrentAttributeIndex];
 
@@ -473,7 +515,7 @@ void UMeshAttributePaintTool::ApplyStamp(const FBrushStampData& Stamp)
 
 		for (int32 k = 0; k < NumVertices; ++k)
 		{
-			int32 vid = ROIVertices[k];
+			int32 vid = ActionData.ROIVertices[k];
 			FVector3d Position = CurrentMesh->GetVertex(vid);
 			float ValueSum = 0, WeightSum = 0;
 			for (int32 NbrVID : CurrentMesh->VtxVerticesItr(vid))
@@ -488,8 +530,8 @@ void UMeshAttributePaintTool::ApplyStamp(const FBrushStampData& Stamp)
 			float Falloff = (float)CalculateBrushFalloff(Position.Distance(StampPosLocal));
 			float NewValue = FMathf::Lerp(AttribData.CurrentValues[vid], ValueSum, SmoothSpeed*Falloff);
 
-			ROIBefore[k] = AttribData.CurrentValues[vid];
-			ROIAfter[k] = CurrentValueRange.Clamp(NewValue);
+			ActionData.ROIBefore[k] = AttribData.CurrentValues[vid];
+			ActionData.ROIAfter[k] = CurrentValueRange.Clamp(NewValue);
 		}
 	}
 	else
@@ -499,37 +541,73 @@ void UMeshAttributePaintTool::ApplyStamp(const FBrushStampData& Stamp)
 		float UseStrength = Sign * BrushProperties->BrushStrength * CurrentValueRange.Length();
 		for (int32 k = 0; k < NumVertices; ++k)
 		{
-			int32 vid = ROIVertices[k];
+			int32 vid = ActionData.ROIVertices[k];
 			FVector3d Position = CurrentMesh->GetVertex(vid);
 			float Falloff = (float)CalculateBrushFalloff(Position.Distance(StampPosLocal));
-			ROIBefore[k] = AttribData.CurrentValues[vid];
-			ROIAfter[k] = CurrentValueRange.Clamp(ROIBefore[k] + UseStrength * Falloff);
+			ActionData.ROIBefore[k] = AttribData.CurrentValues[vid];
+			ActionData.ROIAfter[k] = CurrentValueRange.Clamp(ActionData.ROIBefore[k] + UseStrength*Falloff);
 		}
 	}
+}
 
-	// track changes
-	if (ActiveChangeBuilder)
+
+void UMeshAttributePaintTool::ApplyStamp_FloodFill(const FBrushStampData& Stamp, FStampActionData& ActionData)
+{
+	FAttributeData& AttribData = Attributes[CurrentAttributeIndex];
+	const FDynamicMesh3* CurrentMesh = PreviewMesh->GetMesh();
+
+	// convert to connected triangle set
+	TSet<int32> RemainingTriangles;
+	for (int32 vid : ActionData.ROIVertices)
 	{
-		ActiveChangeBuilder->UpdateValues(ROIVertices, ROIBefore, ROIAfter);
+		CurrentMesh->EnumerateVertexTriangles(vid, [&](int32 tid) { RemainingTriangles.Add(tid); });
 	}
 
-	// update values and colors
-	PreviewMesh->DeferredEditMesh([&](FDynamicMesh3& Mesh)
+	float SetValue = BrushProperties->BrushStrength * CurrentValueRange.Length();
+	if (bInRemoveStroke)
 	{
-		for (int32 k = 0; k < NumVertices; ++k)
+		SetValue = CurrentValueRange.Min;
+	}
+
+	ActionData.ROIVertices.Reset();
+	TArray<int32> InputTriROI, OutputTriROI, QueueTempBuffer;
+	TSet<int32> DoneTempBuffer, DoneVertices;
+	while (RemainingTriangles.Num() > 0)
+	{
+		OutputTriROI.Reset();
+		QueueTempBuffer.Reset();
+		DoneTempBuffer.Reset();
+		InputTriROI.Reset();
+		for (int32 tid : RemainingTriangles)
 		{
-			int32 vid = ROIVertices[k];
-			AttribData.CurrentValues[vid] = ROIAfter[k];
-			FVector3f NewColor = ColorMapper->ToColor(ROIAfter[k]);
-			Mesh.SetVertexColor(vid, NewColor);
+			InputTriROI.Add(tid);		// stupid way to get first set element
+			break;
 		}
-	}, false);
-	PreviewMesh->NotifyDeferredEditCompleted(UPreviewMesh::ERenderUpdateMode::FastUpdate, EMeshRenderAttributeFlags::VertexColors, false);
+		FMeshConnectedComponents::GrowToConnectedTriangles(CurrentMesh, InputTriROI, OutputTriROI, &QueueTempBuffer, &DoneTempBuffer);
+		for (int32 tid : OutputTriROI)
+		{
+			RemainingTriangles.Remove(tid);
+			FIndex3i TriVertices = CurrentMesh->GetTriangle(tid);
+			for (int32 j = 0; j < 3; ++j)
+			{
+				if (DoneVertices.Contains(TriVertices[j]) == false)
+				{
+					ActionData.ROIVertices.Add(TriVertices[j]);
+					ActionData.ROIBefore.Add(AttribData.CurrentValues[TriVertices[j]]);
+					ActionData.ROIAfter.Add(SetValue);
+					DoneVertices.Add(TriVertices[j]);
+				}
+			}
+		}
+	}
+	
+
 }
 
 
 
 
+
 void UMeshAttributePaintTool::ApplyAction(EMeshAttributePaintToolActions ActionType)
 {
 	//switch (ActionType)
@@ -553,6 +631,7 @@ void UMeshAttributePaintTool::UpdateSelectedAttribute(int32 NewSelectedIndex)
 void UMeshAttributePaintTool::OnShutdown(EToolShutdownType ShutdownType)
 {
 	BrushProperties->SaveProperties(this);
+	BrushActionProps->SaveProperties(this);
 
 	StoreCurrentAttribute();
 
diff --git a/Engine/Plugins/Experimental/MeshModelingToolset/Source/MeshModelingTools/Private/MeshGroupPaintTool.cpp b/Engine/Plugins/Experimental/MeshModelingToolset/Source/MeshModelingTools/Private/MeshGroupPaintTool.cpp
--- a/Engine/Plugins/Experimental/MeshModelingToolset/Source/MeshModelingTools/Private/MeshGroupPaintTool.cpp
+++ b/Engine/Plugins/Experimental/MeshModelingToolset/Source/MeshModelingTools/Private/MeshGroupPaintTool.cpp
@@ -1054,17 +1054,20 @@ void UMeshGroupPaintTool::OnTick(float DeltaTime)
 
 	if (IsInBrushSubMode())
 	{
-		if (IsStampPending())
+		if (InStroke())
 		{
 			SCOPE_CYCLE_COUNTER(GroupPaintTool_Tick_ApplyStampBlock);
 
-			ApplyStrokeFlowInTick();
-
 			// update brush position
 			if (UpdateStampPosition(GetPendingStampRayWorld()) == false)
 			{
 				return;
 			}
+			UpdateStampPendingState();
+			if (IsStampPending() == false)
+			{
+				return;
+			}
 
 			// update sculpt ROI
 			UpdateROI(CurrentStamp);
diff --git a/Engine/Plugins/Experimental/MeshModelingToolset/Source/MeshModelingTools/Private/MeshVertexSculptTool.cpp b/Engine/Plugins/Experimental/MeshModelingToolset/Source/MeshModelingTools/Private/MeshVertexSculptTool.cpp
--- a/Engine/Plugins/Experimental/MeshModelingToolset/Source/MeshModelingTools/Private/MeshVertexSculptTool.cpp
+++ b/Engine/Plugins/Experimental/MeshModelingToolset/Source/MeshModelingTools/Private/MeshVertexSculptTool.cpp
@@ -10,6 +10,7 @@
 #include "MeshNormals.h"
 #include "MeshIndexUtil.h"
 #include "Util/BufferUtil.h"
+#include "AssetUtils/Texture2DUtil.h"
 
 #include "Changes/MeshVertexChange.h"
 
@@ -79,7 +80,7 @@ void UMeshVertexSculptTool::Setup()
 	} 
 	else
 	{
-		Octree.RootDimension = Bounds.MaxDim();
+		Octree.RootDimension = Bounds.MaxDim() / 2.0;
 		Octree.SetMaxTreeDepth(8);
 	}
 	Octree.Initialize(Mesh);
@@ -104,7 +105,7 @@ void UMeshVertexSculptTool::Setup()
 
 	// initialize other properties
 	SculptProperties = NewObject<UVertexBrushSculptProperties>(this);
-
+	
 	// init state flags flags
 	ActiveVertexChange = nullptr;
 
@@ -114,9 +115,13 @@ void UMeshVertexSculptTool::Setup()
 	AddToolPropertySource(UMeshSculptToolBase::BrushProperties);
 	UMeshSculptToolBase::BrushProperties->bShowPerBrushProps = false;
 	UMeshSculptToolBase::BrushProperties->bShowFalloff = false;
+	SculptProperties->RestoreProperties(this);
 	AddToolPropertySource(SculptProperties);
 	CalculateBrushRadius();
-	SculptProperties->RestoreProperties(this);
+
+	AlphaProperties = NewObject<UVertexBrushAlphaProperties>(this);
+	AlphaProperties->RestoreProperties(this);
+	AddToolPropertySource(AlphaProperties);
 
 	this->BaseMeshQueryFunc = [&](int32 VertexID, const FVector3d& Position, double MaxDist, FVector3d& PosOut, FVector3d& NormalOut)
 	{
@@ -205,11 +210,18 @@ void UMeshVertexSculptTool::Setup()
 	SculptProperties->WatchProperty( SculptProperties->PrimaryFalloffType,
 		[this](EMeshSculptFalloffType NewType) { SetPrimaryFalloffType(NewType); });
 
+	SculptProperties->WatchProperty(AlphaProperties->Alpha,
+		[this](UTexture2D* NewAlpha) { UpdateBrushAlpha(AlphaProperties->Alpha); });
+
 	// must call before updating brush type so that we register all brush properties?
 	UMeshSculptToolBase::OnCompleteSetup();
 
 	UpdateBrushType(SculptProperties->PrimaryBrushType);
+	SetPrimaryFalloffType(SculptProperties->PrimaryFalloffType);
+	UpdateBrushAlpha(AlphaProperties->Alpha);
 	SetActiveSecondaryBrushType((int32)EMeshVertexSculptBrushType::Smooth);
+
+	StampRandomStream = FRandomStream(31337);
 }
 
 void UMeshVertexSculptTool::Shutdown(EToolShutdownType ShutdownType)
@@ -220,6 +232,7 @@ void UMeshVertexSculptTool::Shutdown(EToolShutdownType ShutdownType)
 	}
 
 	SculptProperties->SaveProperties(this);
+	AlphaProperties->SaveProperties(this);
 
 	// this call will commit result, unregister and destroy DynamicMeshComponent
 	UMeshSculptToolBase::Shutdown(ShutdownType);
@@ -241,7 +254,7 @@ void UMeshVertexSculptTool::OnBeginStroke(const FRay& WorldRay)
 		SculptProperties->PrimaryBrushType == EMeshVertexSculptBrushType::PlaneViewAligned)
 	{
 		UpdateROI(GetBrushFrameLocal().Origin);
-		UpdateStrokeReferencePlaneForROI(GetBrushFrameLocal(), TriangleROI,
+		UpdateStrokeReferencePlaneForROI(GetBrushFrameLocal(), TriangleROIArray,
 			(SculptProperties->PrimaryBrushType == EMeshVertexSculptBrushType::PlaneViewAligned));
 	}
 	else if (SculptProperties->PrimaryBrushType == EMeshVertexSculptBrushType::FixedPlane)
@@ -288,46 +301,124 @@ void UMeshVertexSculptTool::OnEndStroke()
 
 void UMeshVertexSculptTool::UpdateROI(const FVector3d& BrushPos)
 {
-	SCOPE_CYCLE_COUNTER(VtxSculptTool_UpdateROI);
+	TRACE_CPUPROFILER_EVENT_SCOPE(VtxSculptTool_UpdateROI);
 
 	float RadiusSqr = GetCurrentBrushRadius() * GetCurrentBrushRadius();
 	FAxisAlignedBox3d BrushBox(
 		BrushPos - GetCurrentBrushRadius() * FVector3d::One(),
 		BrushPos + GetCurrentBrushRadius() * FVector3d::One());
 
-	VertexSetBuffer.Reset();
-	TriangleROI.Reset();
+	// do a parallel range quer
+	RangeQueryTriBuffer.Reset();
 	FDynamicMesh3* Mesh = GetSculptMesh();
-	Octree.RangeQuery(BrushBox,
-		[&](int TriIdx) {
-		FIndex3i TriV = Mesh->GetTriangle(TriIdx);
-		int Count = 0;
-		for (int j = 0; j < 3; ++j)
+	{
+		TRACE_CPUPROFILER_EVENT_SCOPE(VtxSculptTool_UpdateROI_RangeQuery);
+		Octree.ParallelRangeQuery(BrushBox, RangeQueryTriBuffer);
+	}
+
+#if 1
+	// in this path we use more memory but this lets us do more in parallel
+
+	// Construct array of inside/outside flags for each triangle's vertices. If no
+	// vertices are inside, clear the triangle ID from the range query buffer.
+	// This can be done in parallel and it's cheaper to do repeated distance computations
+	// than to try to do it inside the ROI building below (todo: profile this some more?)
+	TriangleROIInBuf.SetNum(RangeQueryTriBuffer.Num(), false);
+	{
+		TRACE_CPUPROFILER_EVENT_SCOPE(DynamicMeshSculptTool_UpdateROI_TriVerts);
+		ParallelFor(RangeQueryTriBuffer.Num(), [&](int k)
 		{
-			if (VertexSetBuffer.Contains(TriV[j]))
+			const FIndex3i& TriV = Mesh->GetTriangleRef(RangeQueryTriBuffer[k]);
+			TriangleROIInBuf[k].A = (BrushPos.DistanceSquared(Mesh->GetVertexRef(TriV.A)) < RadiusSqr) ? 1 : 0;
+			TriangleROIInBuf[k].B = (BrushPos.DistanceSquared(Mesh->GetVertexRef(TriV.B)) < RadiusSqr) ? 1 : 0;
+			TriangleROIInBuf[k].C = (BrushPos.DistanceSquared(Mesh->GetVertexRef(TriV.C)) < RadiusSqr) ? 1 : 0;
+			if (TriangleROIInBuf[k].A + TriangleROIInBuf[k].B + TriangleROIInBuf[k].C == 0)
 			{
-				Count++;
+				RangeQueryTriBuffer[k] = -1;
 			}
-			else if (BrushPos.DistanceSquared(Mesh->GetVertex(TriV[j])) < RadiusSqr)
+		});
+	}
+
+	// collect set of vertices inside brush sphere, from that box
+	{
+		TRACE_CPUPROFILER_EVENT_SCOPE(DynamicMeshSculptTool_UpdateROI_3Collect);
+		VertexROIBuilder.Initialize(Mesh->MaxVertexID());
+		TriangleROIBuilder.Initialize(Mesh->MaxTriangleID());
+		int32 N = RangeQueryTriBuffer.Num();
+		for ( int32 k = 0; k < N; ++k )
+		{
+			int32 tid = RangeQueryTriBuffer[k];
+			if (tid == -1) continue;		// triangle was deleted in previous step
+			const FIndex3i& TriV = Mesh->GetTriangleRef(RangeQueryTriBuffer[k]);
+			const FIndex3i& Inside = TriangleROIInBuf[k];
+			int InsideCount = 0;
+			for (int j = 0; j < 3; ++j)
+			{
+				if (Inside[j])
+				{
+					VertexROIBuilder.Add(TriV[j]);
+					InsideCount++;
+				}
+			}
+			if (InsideCount > 0)
 			{
-				VertexSetBuffer.Add(TriV[j]);
-				Count++;
+				TriangleROIBuilder.Add(tid);
 			}
 		}
-		if (Count > 0)
+		VertexROIBuilder.SwapValuesWith(VertexROI);
+		TriangleROIBuilder.SwapValuesWith(TriangleROIArray);
+	}
+
+#else
+	// In this path we combine everything into one loop. Does fewer distance checks
+	// but nothing can be done in parallel (would change if ROIBuilders had atomic-try-add)
+
+	// collect set of vertices and triangles inside brush sphere, from range query result
+	{
+		TRACE_CPUPROFILER_EVENT_SCOPE(DynamicMeshSculptTool_UpdateROI_2Collect);
+		VertexROIBuilder.Initialize(Mesh->MaxVertexID());
+		TriangleROIBuilder.Initialize(Mesh->MaxTriangleID());
+		for (int32 TriIdx : RangeQueryTriBuffer)
 		{
-			TriangleROI.Add(TriIdx);
+			FIndex3i TriV = Mesh->GetTriangle(TriIdx);
+			int InsideCount = 0;
+			for (int j = 0; j < 3; ++j)
+			{
+				if (VertexROIBuilder.Contains(TriV[j]))
+				{
+					InsideCount++;
+				} 
+				else if (BrushPos.DistanceSquared(Mesh->GetVertexRef(TriV[j])) < RadiusSqr)
+				{
+					VertexROIBuilder.Add(TriV[j]);
+					InsideCount++;
+				}
+			}
+			if (InsideCount > 0)
+			{
+				TriangleROIBuilder.Add(tid);
+			}
 		}
-	});
-
-	VertexROI.SetNum(0, false);
-	BufferUtil::AppendElements(VertexROI, VertexSetBuffer);
+		VertexROIBuilder.SwapValuesWith(VertexROI);
+		TriangleROIBuilder.SwapValuesWith(TriangleROIArray);
+	}
+#endif
 
-	ROIPositionBuffer.SetNum(VertexROI.Num(), false);
+	{
+		TRACE_CPUPROFILER_EVENT_SCOPE(DynamicMeshSculptTool_UpdateROI_4ROI);
+		ROIPositionBuffer.SetNum(VertexROI.Num(), false);
+		ROIPrevPositionBuffer.SetNum(VertexROI.Num(), false);
+		ParallelFor(VertexROI.Num(), [&](int i)
+		{
+			ROIPrevPositionBuffer[i] = Mesh->GetVertexRef(VertexROI[i]);
+		});
+	}
 }
 
 bool UMeshVertexSculptTool::UpdateStampPosition(const FRay& WorldRay)
 {
+	TRACE_CPUPROFILER_EVENT_SCOPE(VtxSculptTool_UpdateStampPosition);
+
 	CalculateBrushRadius();
 
 	TUniquePtr<FMeshSculptBrushOp>& UseBrushOp = GetActiveBrushOp();
@@ -352,11 +443,25 @@ bool UMeshVertexSculptTool::UpdateStampPosition(const FRay& WorldRay)
 	}
 
 	CurrentStamp = LastStamp;
-	CurrentStamp.DeltaTime = FMathd::Min((FDateTime::Now() - LastStamp.TimeStamp).GetTotalSeconds(), 1.0);
+	//CurrentStamp.DeltaTime = FMathd::Min((FDateTime::Now() - LastStamp.TimeStamp).GetTotalSeconds(), 1.0);
+	CurrentStamp.DeltaTime = 0.03;		// 30 fps - using actual time is no good now that we support variable stamps!
 	CurrentStamp.WorldFrame = GetBrushFrameWorld();
 	CurrentStamp.LocalFrame = GetBrushFrameLocal();
 	CurrentStamp.Power = GetActivePressure() * GetCurrentBrushStrength();
 
+	if (bHaveBrushAlpha && (AlphaProperties->RotationAngle != 0 || AlphaProperties->bRandomize))
+	{
+		float UseAngle = AlphaProperties->RotationAngle;
+		if (AlphaProperties->bRandomize)
+		{
+			UseAngle += (StampRandomStream.GetFraction() - 0.5f) * 2.0f * AlphaProperties->RandomRange;
+		}
+
+		// possibly should be done in base brush...
+		CurrentStamp.WorldFrame.Rotate(FQuaterniond(CurrentStamp.WorldFrame.Z(), UseAngle, true));
+		CurrentStamp.LocalFrame.Rotate(FQuaterniond(CurrentStamp.LocalFrame.Z(), UseAngle, true));
+	}
+
 	CurrentStamp.PrevLocalFrame = LastStamp.LocalFrame;
 	CurrentStamp.PrevWorldFrame = LastStamp.WorldFrame;
 
@@ -370,43 +475,71 @@ bool UMeshVertexSculptTool::UpdateStampPosition(const FRay& WorldRay)
 }
 
 
-void UMeshVertexSculptTool::ApplyStamp()
+TFuture<void> UMeshVertexSculptTool::ApplyStamp()
 {
-	SCOPE_CYCLE_COUNTER(VtxSculptToolApplyStamp);
+	TRACE_CPUPROFILER_EVENT_SCOPE(VtxSculptTool_ApplyStamp);
 
 	TUniquePtr<FMeshSculptBrushOp>& UseBrushOp = GetActiveBrushOp();
+
+	// compute region plane if necessary. This may currently be expensive?
 	if (UseBrushOp->WantsStampRegionPlane())
 	{
-		CurrentStamp.RegionPlane = ComputeStampRegionPlane(CurrentStamp.LocalFrame, TriangleROI, true, false, false);
+		CurrentStamp.RegionPlane = ComputeStampRegionPlane(CurrentStamp.LocalFrame, TriangleROIArray, true, false, false);
 	}
 
-	FDynamicMesh3* Mesh = GetSculptMesh();
-	UseBrushOp->ApplyStamp(Mesh, CurrentStamp, VertexROI, ROIPositionBuffer);
-
-	SyncMeshWithPositionBuffer(Mesh);
-
-	LastStamp = CurrentStamp;
-	LastStamp.TimeStamp = FDateTime::Now();
-}
-
-
+	// set up alpha function if we have one
+	if (bHaveBrushAlpha)
+	{
+		CurrentStamp.StampAlphaFunc = [this](const FSculptBrushStamp& Stamp, const FVector3d& Position)
+		{
+			return this->SampleBrushAlpha(Stamp, Position);
+		};
+	}
 
+	// apply the stamp, which computes new positions
+	FDynamicMesh3* Mesh = GetSculptMesh();
+	{
+		TRACE_CPUPROFILER_EVENT_SCOPE(VtxSculptTool_ApplyStamp_Apply);
+		UseBrushOp->ApplyStamp(Mesh, CurrentStamp, VertexROI, ROIPositionBuffer);
+	}
 
-void UMeshVertexSculptTool::SyncMeshWithPositionBuffer(FDynamicMesh3* Mesh)
-{
-	const int32 NumV = ROIPositionBuffer.Num();
-	checkSlow(VertexROI.Num() <= NumV);
+	// can discard alpha now
+	CurrentStamp.StampAlphaFunc = nullptr;
 
-	// change update could be async here if we collected array of <idx,orig,new> and dispatched independenlty
-	for ( int32 k = 0; k < NumV; ++k)
+	// once stamp is applied, we can start updating vertex change, which can happen async as we saved all necessary info
+	TFuture<void> SaveVertexFuture;
+	if (ActiveVertexChange != nullptr)
 	{
-		int VertIdx = VertexROI[k];
-		const FVector3d& NewPos = ROIPositionBuffer[k];
-		FVector3d OrigPos = Mesh->GetVertex(VertIdx);
-		Mesh->SetVertex(VertIdx, NewPos);
+		SaveVertexFuture = Async(VertexSculptToolAsyncExecTarget, [&]()
+		{
+			TRACE_CPUPROFILER_EVENT_SCOPE(VtxSculptTool_SyncMeshWithPositionBuffer_UpdateChange);
+			const int32 NumV = ROIPositionBuffer.Num();
+			for (int k = 0; k < NumV; ++k)
+			{
+				int VertIdx = VertexROI[k];
+				ActiveVertexChange->UpdateVertex(VertIdx, ROIPrevPositionBuffer[k], ROIPositionBuffer[k]);
+			}
+		});
+	}
 
-		ActiveVertexChange->UpdateVertex(VertIdx, OrigPos, NewPos);
+	// now actually update the mesh, which happens 
+	{
+		TRACE_CPUPROFILER_EVENT_SCOPE(VtxSculptTool_ApplyStamp_Sync);
+		const int32 NumV = ROIPositionBuffer.Num();
+		ParallelFor(NumV, [&](int32 k)
+		{
+			int VertIdx = VertexROI[k];
+			const FVector3d& NewPos = ROIPositionBuffer[k];
+			Mesh->SetVertex_NoTimeStampUpdate(VertIdx, NewPos);
+		});
+		Mesh->IncrementTimeStamps(1, true, false);
 	}
+
+	LastStamp = CurrentStamp;
+	LastStamp.TimeStamp = FDateTime::Now();
+
+	// let caller wait for this to finish
+	return SaveVertexFuture;
 }
 
 
@@ -414,8 +547,12 @@ void UMeshVertexSculptTool::SyncMeshWithPositionBuffer(FDynamicMesh3* Mesh)
 
 
 
+
 int32 UMeshVertexSculptTool::FindHitSculptMeshTriangle(const FRay3d& LocalRay)
 {
+	// need this to finish before we can touch Octree
+	WaitForPendingStampUpdate();
+
 	if (GetBrushCanHitBackFaces())
 	{
 		return Octree.FindNearestHitObject(LocalRay);
@@ -515,7 +652,7 @@ void UMeshVertexSculptTool::OnTick(float DeltaTime)
 {
 	UMeshSculptToolBase::OnTick(DeltaTime);
 
-	SCOPE_CYCLE_COUNTER(VtxSculptToolTick);
+	TRACE_CPUPROFILER_EVENT_SCOPE(VtxSculptTool_Tick);
 
 	// process the undo update
 	if (bUndoUpdatePending)
@@ -540,61 +677,111 @@ void UMeshVertexSculptTool::OnTick(float DeltaTime)
 		bTargetDirty = false;
 	}
 
-
-	if (IsStampPending())
+	if (InStroke())
 	{
 		//UE_LOG(LogTemp, Warning, TEXT("dt is %.3f, tick fps %.2f - roi size %d/%d"), DeltaTime, 1.0 / DeltaTime, VertexROI.Num(), TriangleROI.Num());
-		SCOPE_CYCLE_COUNTER(VtxSculptTool_Tick_ApplyStampBlock);
-
-		ApplyStrokeFlowInTick();
+		TRACE_CPUPROFILER_EVENT_SCOPE(VtxSculptTool_Tick_StrokeUpdate);
+		FDynamicMesh3* Mesh = GetSculptMesh();
 
 		// update brush position
 		if (UpdateStampPosition(GetPendingStampRayWorld()) == false)
 		{
 			return;
 		}
+		UpdateStampPendingState();
+		if (IsStampPending() == false)
+		{
+			return;
+		}
+
+		// need to make sure previous stamp finished
+		WaitForPendingStampUpdate();
 
 		// update sculpt ROI
 		UpdateROI(CurrentStamp.LocalFrame.Origin);
 
-		// append updated ROI to modified region (async)
+		// Append updated ROI to modified region (async). For some reason this is very expensive,
+		// maybe because of TSet? but we have a lot of time to do it.
 		TFuture<void> AccumulateROI = Async(VertexSculptToolAsyncExecTarget, [&]()
 		{
-			AccumulatedTriangleROI.Append(TriangleROI);
+			TRACE_CPUPROFILER_EVENT_SCOPE(VtxSculptTool_Tick_AccumROI);
+			for (int32 tid : TriangleROIArray)
+			{
+				AccumulatedTriangleROI.Add(tid);
+			}
+		});
+
+		// Start precomputing the normals ROI. This is currently the most expensive single thing we do next
+		// to Octree re-insertion, despite it being almost trivial. Why?!?
+		bool bUsingOverlayNormals = false;
+		TFuture<void> NormalsROI = Async(VertexSculptToolAsyncExecTarget, [&]()
+		{
+			TRACE_CPUPROFILER_EVENT_SCOPE(VtxSculptTool_Tick_NormalsROI);
+
+
+			//TODO WHY DOES NORMALS FLAGS NOT WORK?? 
+
+			//UE::SculptUtil::PrecalculateNormalsROI(Mesh, TriangleROIArray, NormalsROIBuilder, bUsingOverlayNormals, false);
+			UE::SculptUtil::PrecalculateNormalsROI(Mesh, TriangleROIArray, NormalsFlags, bUsingOverlayNormals, false);
 		});
 
-		// apply the stamp
-		ApplyStamp();
+		// NOTE: you might try to speculatively do the octree remove here, to save doing it later on Reinsert().
+		// This will not improve things, as Reinsert() checks if it needs to actually re-insert, which avoids many
+		// removes, and does much of the work of Remove anyway.
+
+		// Apply the stamp. This will return a future that is updating the vertex-change record, 
+		// which can run until the end of the frame, as it is using cached information
+		TFuture<void> UpdateChangeFuture = ApplyStamp();
 
 		// begin octree rebuild calculation
-		TFuture<void> OctreeRebuild = Async(VertexSculptToolAsyncExecTarget, [&]()
+		StampUpdateOctreeFuture = Async(VertexSculptToolAsyncExecTarget, [&]()
 		{
-			SCOPE_CYCLE_COUNTER(VtxSculptTool_Tick_ApplyStamp_Insert);
-			Octree.ReinsertTriangles(TriangleROI);
+			TRACE_CPUPROFILER_EVENT_SCOPE(VtxSculptTool_Tick_OctreeReinsert);
+			static TArray<uint32> TempBuffer;
+			static TArray<bool> TempFlagBuffer;
+			Octree.ReinsertTrianglesParallel(TriangleROIArray, TempBuffer, TempFlagBuffer);
 		});
+		bStampUpdatePending = true;
+		//TFuture<void> OctreeRebuild = Async(VertexSculptToolAsyncExecTarget, [&]()
+		//{
+		//	TRACE_CPUPROFILER_EVENT_SCOPE(VtxSculptTool_Tick_OctreeReinsert);
+		//	Octree.ReinsertTriangles(TriangleROIArray);
+		//});
+
+		// TODO: first step of RecalculateROINormals() is to convert TriangleROI into vertex or element ROI.
+		// We can do this while we are computing stamp!
+
+		// precompute dynamic mesh update info
+		TArray<int32> RenderUpdateSets; FAxisAlignedBox3d RenderUpdateBounds;
+		TFuture<bool> RenderUpdatePrecompute = DynamicMeshComponent->FastNotifyTriangleVerticesUpdated_TryPrecompute(
+				TriangleROIArray, RenderUpdateSets, RenderUpdateBounds);
 
 		// recalculate normals. This has to complete before we can update component
 		// (in fact we could do it per-chunk...)
-		FDynamicMesh3* Mesh = GetSculptMesh();
 		{
-			SCOPE_CYCLE_COUNTER(VtxSculptTool_Tick_NormalsBlock);
-			UE::SculptUtil::RecalculateROINormals(Mesh, TriangleROI, VertexSetBuffer, NormalsBuffer);
+			TRACE_CPUPROFILER_EVENT_SCOPE(VtxSculptTool_Tick_RecalcNormals);
+			NormalsROI.Wait();
+			UE::SculptUtil::RecalculateROINormals(Mesh, NormalsFlags, bUsingOverlayNormals);
+			//UE::SculptUtil::RecalculateROINormals(Mesh, NormalsROIBuilder.Indices(), bUsingOverlayNormals);
 		}
 
 		{
-			SCOPE_CYCLE_COUNTER(VtxSculptTool_Tick_UpdateMeshBlock);
-			DynamicMeshComponent->FastNotifyTriangleVerticesUpdated(TriangleROI,
-				EMeshRenderAttributeFlags::Positions | EMeshRenderAttributeFlags::VertexNormals);
+			TRACE_CPUPROFILER_EVENT_SCOPE(VtxSculptTool_Tick_UpdateMesh);
+			RenderUpdatePrecompute.Wait();
+			DynamicMeshComponent->FastNotifyTriangleVerticesUpdated_ApplyPrecompute(TriangleROIArray,
+				EMeshRenderAttributeFlags::Positions | EMeshRenderAttributeFlags::VertexNormals,
+				RenderUpdatePrecompute, RenderUpdateSets, RenderUpdateBounds);
+
 			GetToolManager()->PostInvalidation();
 		}
 
 		// we don't really need to wait for these to happen to end Tick()...
+		UpdateChangeFuture.Wait();
 		AccumulateROI.Wait();
-		OctreeRebuild.Wait();
 	} 
 	else if (bTargetDirty)
 	{
-		SCOPE_CYCLE_COUNTER(VtxSculptTool_Tick_UpdateTargetBlock);
+		TRACE_CPUPROFILER_EVENT_SCOPE(VtxSculptTool_Tick_UpdateTarget);
 		check(InStroke() == false);
 
 		// this spawns futures that we could allow to run while other things happen...
@@ -608,6 +795,16 @@ void UMeshVertexSculptTool::OnTick(float DeltaTime)
 
 
 
+void UMeshVertexSculptTool::WaitForPendingStampUpdate()
+{
+	if (bStampUpdatePending)
+	{
+		StampUpdateOctreeFuture.Wait();
+		bStampUpdatePending = true;
+	}
+}
+
+
 
 void UMeshVertexSculptTool::UpdateBaseMesh(const TSet<int32>* TriangleSet)
 {
@@ -623,6 +820,7 @@ void UMeshVertexSculptTool::UpdateBaseMesh(const TSet<int32>* TriangleSet)
 	const FDynamicMesh3* SculptMesh = GetSculptMesh();
 	if ( ! TriangleSet )
 	{
+		TRACE_CPUPROFILER_EVENT_SCOPE(VtxSculptTool_Target_FullUpdate);
 		BaseMesh.Copy(*SculptMesh, false, false, false, false);
 		BaseMesh.EnableVertexNormals(FVector3f::UnitZ());
 		FMeshNormals::QuickComputeVertexNormals(BaseMesh);
@@ -641,12 +839,14 @@ void UMeshVertexSculptTool::UpdateBaseMesh(const TSet<int32>* TriangleSet)
 			BaseMesh.SetVertex(Tri.C, SculptMesh->GetVertex(Tri.C));
 			BaseMeshIndexBuffer.Add(tid);
 		}
-		auto UpdateBaseNormals = Async(EAsyncExecution::Thread, [&]()
+		auto UpdateBaseNormals = Async(VertexSculptToolAsyncExecTarget, [&]()
 		{
+			TRACE_CPUPROFILER_EVENT_SCOPE(VtxSculptTool_Target_UpdateBaseNormals);
 			FMeshNormals::QuickComputeVertexNormalsForTriangles(BaseMesh, BaseMeshIndexBuffer);
 		});
-		auto ReinsertTriangles = Async(EAsyncExecution::Thread, [&]()
+		auto ReinsertTriangles = Async(VertexSculptToolAsyncExecTarget, [&]()
 		{
+			TRACE_CPUPROFILER_EVENT_SCOPE(VtxSculptTool_Target_Reinsert);
 			BaseMeshSpatial.ReinsertTriangles(*TriangleSet);
 		});
 		UpdateBaseNormals.Wait();
@@ -677,8 +877,48 @@ void UMeshVertexSculptTool::DecreaseBrushSpeedAction()
 	//SculptProperties->PrimaryBrushSpeed = FMath::Clamp(SculptProperties->PrimaryBrushSpeed - 0.05f, 0.0f, 1.0f);
 }
 
+void UMeshVertexSculptTool::UpdateBrushAlpha(UTexture2D* NewAlpha)
+{
+	if (this->BrushAlpha != NewAlpha)
+	{
+		this->BrushAlpha = NewAlpha;
+		if (this->BrushAlpha != nullptr)
+		{
+			TImageBuilder<FVector4f> AlphaValues;
+			FImageDimensions AlphaDimensions;
+
+			bool bReadOK = UE::AssetUtils::ReadTexture(this->BrushAlpha, AlphaDimensions, AlphaValues, true);
+			if (bReadOK)
+			{
+				BrushAlphaValues = MoveTemp(AlphaValues);
+				BrushAlphaDimensions = AlphaDimensions;
+				bHaveBrushAlpha = true;
+				return;
+			}
+		}
+		bHaveBrushAlpha = false;
+		BrushAlphaValues = TImageBuilder<FVector4f>();
+		BrushAlphaDimensions = FImageDimensions();
+	}
+}
 
 
+double UMeshVertexSculptTool::SampleBrushAlpha(const FSculptBrushStamp& Stamp, const FVector3d& Position) const
+{
+	if (! bHaveBrushAlpha) return 1.0;
+
+	static const FVector4f InvalidValue(0, 0, 0, 0);
+
+	FVector2d AlphaUV = Stamp.LocalFrame.ToPlaneUV(Position, 2);
+	double u = AlphaUV.X / Stamp.Radius;
+	u = 1.0 - (u + 1.0) / 2.0;
+	double v = AlphaUV.Y / Stamp.Radius;
+	v = 1.0 - (v + 1.0) / 2.0;
+	if (u < 0 || u > 1) return 0.0;
+	if (v < 0 || v > 1) return 0.0;
+	FVector4f AlphaValue = BrushAlphaValues.BilinearSampleUV<float>(FVector2d(u, v), InvalidValue);
+	return FMathd::Clamp(AlphaValue.X, 0.0, 1.0);
+}
 
 
 
@@ -722,6 +962,9 @@ void UMeshVertexSculptTool::WaitForPendingUndoRedo()
 
 void UMeshVertexSculptTool::OnDynamicMeshComponentChanged(USimpleDynamicMeshComponent* Component, const FMeshVertexChange* Change, bool bRevert)
 {
+	// have to wait for any outstanding stamp update to finish...
+	WaitForPendingStampUpdate();
+
 	// update octree
 	FDynamicMesh3* Mesh = GetSculptMesh();
 
@@ -741,21 +984,21 @@ void UMeshVertexSculptTool::OnDynamicMeshComponentChanged(USimpleDynamicMeshComp
 	}
 
 	// start the normal recomputation
-	UndoNormalsFuture = Async(EAsyncExecution::Thread, [&]()
+	UndoNormalsFuture = Async(VertexSculptToolAsyncExecTarget, [&]()
 	{
-		UE::SculptUtil::RecalculateROINormals(Mesh, AccumulatedTriangleROI, VertexSetBuffer, NormalsBuffer);
+		UE::SculptUtil::RecalculateROINormals(Mesh, AccumulatedTriangleROI, NormalsROIBuilder);
 		return true;
 	});
 
 	// start the octree update
-	UndoUpdateOctreeFuture = Async(EAsyncExecution::Thread, [&]()
+	UndoUpdateOctreeFuture = Async(VertexSculptToolAsyncExecTarget, [&]()
 	{
 		Octree.ReinsertTriangles(AccumulatedTriangleROI);
 		return true;
 	});
 
 	// start the base mesh update
-	UndoUpdateBaseMeshFuture = Async(EAsyncExecution::Thread, [&]()
+	UndoUpdateBaseMeshFuture = Async(VertexSculptToolAsyncExecTarget, [&]()
 	{
 		UpdateBaseMesh(&AccumulatedTriangleROI);
 		return true;
@@ -785,6 +1028,18 @@ void UMeshVertexSculptTool::UpdateBrushType(EMeshVertexSculptBrushType BrushType
 		SetToolPropertySourceEnabled(GizmoProperties, true);
 	}
 
+
+	bool bEnableAlpha = false;
+	switch (BrushType)
+	{
+	case EMeshVertexSculptBrushType::Offset:
+	case EMeshVertexSculptBrushType::SculptView:
+	case EMeshVertexSculptBrushType::SculptMax:
+		bEnableAlpha = true; break;
+	}
+	SetToolPropertySourceEnabled(AlphaProperties, bEnableAlpha);
+
+
 	GetToolManager()->DisplayMessage(Builder.ToText(), EToolMessageLevel::UserNotification);
 }
 
diff --git a/Engine/Plugins/Experimental/MeshModelingToolset/Source/MeshModelingTools/Private/Sculpting/MeshSculptToolBase.cpp b/Engine/Plugins/Experimental/MeshModelingToolset/Source/MeshModelingTools/Private/Sculpting/MeshSculptToolBase.cpp
--- a/Engine/Plugins/Experimental/MeshModelingToolset/Source/MeshModelingTools/Private/Sculpting/MeshSculptToolBase.cpp
+++ b/Engine/Plugins/Experimental/MeshModelingToolset/Source/MeshModelingTools/Private/Sculpting/MeshSculptToolBase.cpp
@@ -137,11 +137,7 @@ void UMeshSculptToolBase::Shutdown(EToolShutdownType ShutdownType)
 
 			// this block bakes the modified DynamicMeshComponent back into the StaticMeshComponent inside an undo transaction
 			GetToolManager()->BeginUndoTransaction(LOCTEXT("SculptMeshToolTransactionName", "Sculpt Mesh"));
-			ComponentTarget->CommitMesh([=](const FPrimitiveComponentTarget::FCommitParams& CommitParams)
-			{
-				FConversionToMeshDescriptionOptions ConversionOptions;
-				DynamicMeshComponent->Bake(CommitParams.MeshDescription, false, ConversionOptions);
-			});
+			CommitResult(DynamicMeshComponent);
 			GetToolManager()->EndUndoTransaction();
 		}
 
@@ -154,6 +150,16 @@ void UMeshSculptToolBase::Shutdown(EToolShutdownType ShutdownType)
 
 
 
+void UMeshSculptToolBase::CommitResult(UBaseDynamicMeshComponent* Component)
+{
+	ComponentTarget->CommitMesh([=](const FPrimitiveComponentTarget::FCommitParams& CommitParams)
+	{
+		FConversionToMeshDescriptionOptions ConversionOptions;
+		Component->Bake(CommitParams.MeshDescription, false, ConversionOptions);
+	});
+}
+
+
 void UMeshSculptToolBase::OnTick(float DeltaTime)
 {
 	GizmoPositionWatcher.CheckAndUpdate();
@@ -165,6 +171,10 @@ void UMeshSculptToolBase::OnTick(float DeltaTime)
 	{
 		SaveActiveStrokeModifiers();
 	}
+	else
+	{
+		AccumulateStrokeTime(DeltaTime);
+	}
 
 	// update cached falloff
 	CurrentBrushFalloff = 0.5;
@@ -175,6 +185,11 @@ void UMeshSculptToolBase::OnTick(float DeltaTime)
 
 	UpdateHoverStamp(GetBrushFrameWorld());
 
+	// update brush position here?
+	BrushIndicator->Update((float)GetCurrentBrushRadius(),
+		(FVector)HoverStamp.WorldFrame.Origin, (FVector)HoverStamp.WorldFrame.Z(),
+		1.0f - (float)GetCurrentBrushFalloff());
+
 	UpdateWorkPlane();
 }
 
@@ -187,9 +202,6 @@ void UMeshSculptToolBase::Render(IToolsContextRenderAPI* RenderAPI)
 
 	FViewCameraState RenderCameraState = RenderAPI->GetCameraState();
 
-	BrushIndicator->Update((float)GetCurrentBrushRadius(),
-		(FVector)HoverStamp.WorldFrame.Origin, (FVector)HoverStamp.WorldFrame.Z(),
-		1.0f - (float)GetCurrentBrushFalloff());
 	if (BrushIndicatorMaterial)
 	{
 		double FixedDimScale = ToolSceneQueriesUtil::CalculateDimensionFromVisualAngleD(RenderCameraState, HoverStamp.WorldFrame.Origin, 1.5f);
@@ -440,6 +452,7 @@ void UMeshSculptToolBase::OnBeginDrag(const FRay& WorldRay)
 	if (HitTest(WorldRay, OutHit))
 	{
 		bInStroke = true;
+		ResetStrokeTime();
 
 		UpdateBrushTargetPlaneFromHit(WorldRay, OutHit);
 
@@ -459,7 +472,6 @@ void UMeshSculptToolBase::OnUpdateDrag(const FRay& WorldRay)
 	if (InStroke())
 	{
 		PendingStampRay = WorldRay;
-		bIsStampPending = true;
 	}
 }
 
@@ -467,7 +479,7 @@ void UMeshSculptToolBase::OnEndDrag(const FRay& Ray)
 {
 	bInStroke = false;
 
-	// cancel these! otherwise change record could become invalid
+	// cancel any outstanding stamps! otherwise change record could become invalid
 	bIsStampPending = false;
 
 	OnEndStroke();
@@ -488,10 +500,29 @@ FRay3d UMeshSculptToolBase::GetLocalRay(const FRay& WorldRay) const
 
 void UMeshSculptToolBase::UpdateBrushFrameWorld(const FVector3d& NewPosition, const FVector3d& NewNormal)
 {
+	FFrame3d PrevBrushFrameWorld = LastBrushFrameWorld;
+
+	bool bTriedFrameRepair = false;
+retry_frame_update:
 	FFrame3d NewFrame = LastBrushFrameWorld;
 	NewFrame.Origin = NewPosition;
 	NewFrame.AlignAxis(2, NewNormal);
-	NewFrame.ConstrainedAlignPerpAxes();
+	FVector3d CameraUp = CameraState.Up();
+
+	if (FMathd::Abs(CameraUp.Dot(NewNormal)) < 0.98)
+	{
+		NewFrame.ConstrainedAlignAxis(1, CameraUp, NewFrame.Z());
+	}
+
+	if ( (NewFrame.Rotation.Length() - 1.0) > 0.1 )		// try to recover from normalization failure
+	{
+		LastBrushFrameWorld = FFrame3d(LastBrushFrameWorld.Origin);
+		if (bTriedFrameRepair == false)
+		{
+			bTriedFrameRepair = true;
+			goto retry_frame_update;
+		}
+	}
 
 	if (InStroke() && BrushProperties->Lazyness > 0)
 	{
@@ -502,9 +533,10 @@ void UMeshSculptToolBase::UpdateBrushFrameWorld(const FVector3d& NewPosition, co
 	else
 	{
 		LastBrushFrameWorld = NewFrame;
-
 	}
 
+	ActiveStrokePathArcLen += LastBrushFrameWorld.Origin.Distance(PrevBrushFrameWorld.Origin);
+
 	LastBrushFrameLocal = LastBrushFrameWorld;
 	LastBrushFrameLocal.Transform(CurTargetTransform.Inverse());
 }
@@ -606,14 +638,78 @@ void UMeshSculptToolBase::UpdateHoverStamp(const FFrame3d& StampFrame)
 	HoverStamp.WorldFrame = StampFrame;
 }
 
-void UMeshSculptToolBase::ApplyStrokeFlowInTick()
+void UMeshSculptToolBase::UpdateStampPendingState()
 {
-	bIsStampPending = InStroke();
+	if (InStroke() == false) return;
+
+	bool bFlowStampPending = false;
+	if (BrushProperties->FlowRate >= 1.0)
+	{
+		bFlowStampPending = true;
+	}
+	else if (BrushProperties->FlowRate == 0.0)
+	{
+		bFlowStampPending = (LastFlowTimeStamp++ == 0);
+	}
+	else
+	{
+		double dt = (1.0 - BrushProperties->FlowRate);
+		int FlowTimestamp = (int)(ActiveStrokeTime / dt);
+		if (FlowTimestamp > LastFlowTimeStamp)
+		{
+			LastFlowTimeStamp = FlowTimestamp;
+			bFlowStampPending = true;
+		}
+	}
+
+	// if brush does not support variable spacing, we only use flow
+	if (GetActiveBrushOp()->SupportsVariableSpacing() == false)
+	{
+		bIsStampPending = bFlowStampPending;
+		return;
+	}
+
+	bool bSpacingStampPending = false;
+	if (BrushProperties->Spacing == 0)
+	{
+		bSpacingStampPending = true;
+	}
+	else
+	{
+		double ArcSpacing = BrushProperties->Spacing * (2 * GetCurrentBrushRadius());
+		int SpacingTimestamp = 1 + (int)(ActiveStrokePathArcLen / ArcSpacing);
+		if (SpacingTimestamp > LastSpacingTimestamp)
+		{
+			LastSpacingTimestamp = SpacingTimestamp;
+			bSpacingStampPending = true;
+		}
+	}
+
+	bIsStampPending = bFlowStampPending && bSpacingStampPending;
 }
 
 
+void UMeshSculptToolBase::ResetStrokeTime()
+{
+	ActiveStrokeTime = 0.0;
+	LastFlowTimeStamp = 0;
+	ActiveStrokePathArcLen = 0;
+	LastSpacingTimestamp = 0;
+}
+
+void UMeshSculptToolBase::AccumulateStrokeTime(float DeltaTime)
+{
+	ActiveStrokeTime += DeltaTime;
+}
+
 
 FFrame3d UMeshSculptToolBase::ComputeStampRegionPlane(const FFrame3d& StampFrame, const TSet<int32>& StampTriangles, bool bIgnoreDepth, bool bViewAligned, bool bInvDistFalloff)
+{
+	check(false);
+	return FFrame3d();
+}
+
+FFrame3d UMeshSculptToolBase::ComputeStampRegionPlane(const FFrame3d& StampFrame, const TArray<int32>& StampTriangles, bool bIgnoreDepth, bool bViewAligned, bool bInvDistFalloff)
 {
 	const FDynamicMesh3* Mesh = GetSculptMesh();
 	double FalloffRadius = GetCurrentBrushRadius();
@@ -677,7 +773,7 @@ FFrame3d UMeshSculptToolBase::ComputeStampRegionPlane(const FFrame3d& StampFrame
 
 
 
-void UMeshSculptToolBase::UpdateStrokeReferencePlaneForROI(const FFrame3d& StampFrame, const TSet<int32>& TriangleROI, bool bViewAligned)
+void UMeshSculptToolBase::UpdateStrokeReferencePlaneForROI(const FFrame3d& StampFrame, const TArray<int32>& TriangleROI, bool bViewAligned)
 {
 	StrokePlane = ComputeStampRegionPlane(GetBrushFrameLocal(), TriangleROI, false, bViewAligned);
 }
@@ -713,6 +809,7 @@ void UMeshSculptToolBase::CalculateBrushRadius()
 }
 
 
+
 double UMeshSculptToolBase::GetCurrentBrushStrength()
 {
 	TUniquePtr<FMeshSculptBrushOp>& BrushOp = GetActiveBrushOp();
@@ -735,6 +832,8 @@ double UMeshSculptToolBase::GetCurrentBrushDepth()
 
 
 
+
+
 void UMeshSculptToolBase::IncreaseBrushRadiusAction()
 {
 	BrushProperties->BrushSize = FMath::Clamp(BrushProperties->BrushSize + 0.025f, 0.0f, 1.0f);
diff --git a/Engine/Plugins/Experimental/MeshModelingToolset/Source/MeshModelingTools/Private/Sculpting/MeshSculptUtil.cpp b/Engine/Plugins/Experimental/MeshModelingToolset/Source/MeshModelingTools/Private/Sculpting/MeshSculptUtil.cpp
--- a/Engine/Plugins/Experimental/MeshModelingToolset/Source/MeshModelingTools/Private/Sculpting/MeshSculptUtil.cpp
+++ b/Engine/Plugins/Experimental/MeshModelingToolset/Source/MeshModelingTools/Private/Sculpting/MeshSculptUtil.cpp
@@ -1,6 +1,8 @@
 // Copyright Epic Games, Inc. All Rights Reserved.
 
 #include "MeshSculptUtil.h"
+#include "Async/ParallelFor.h"
+#include "Async/Async.h"
 
 void UE::SculptUtil::RecalculateNormals_Overlay(
 	FDynamicMesh3* Mesh, 
@@ -82,4 +84,338 @@ void UE::SculptUtil::RecalculateROINormals(
 	{
 		RecalculateNormals_PerVertex(Mesh, TriangleROI, VertexSetBuffer, NormalsBuffer);
 	}
-}
\ No newline at end of file
+}
+
+
+
+
+
+
+
+
+
+
+void UE::SculptUtil::RecalculateNormals_Overlay(
+	FDynamicMesh3* Mesh, 
+	const TSet<int32>& ModifiedTris, 
+	FUniqueIndexSet& VertexSetTemp)
+{
+	FDynamicMeshNormalOverlay* Normals = Mesh->HasAttributes() ? Mesh->Attributes()->PrimaryNormals() : nullptr;
+	check(Normals != nullptr);
+
+	VertexSetTemp.Initialize(Normals->MaxElementID());
+	for (int32 TriangleID : ModifiedTris)
+	{
+		FIndex3i TriElems = Normals->GetTriangle(TriangleID);
+		VertexSetTemp.Add(TriElems.A);
+		VertexSetTemp.Add(TriElems.B);
+		VertexSetTemp.Add(TriElems.C);
+	}
+
+	const TArray<int32>& ElemIndices = VertexSetTemp.Indices();
+	ParallelFor(ElemIndices.Num(), [&](int32 k) {
+		int32 elemid = ElemIndices[k];
+		FVector3d NewNormal = FMeshNormals::ComputeOverlayNormal(*Mesh, Normals, elemid);
+		Normals->SetElement(elemid, (FVector3f)NewNormal);
+	});
+}
+
+
+
+void UE::SculptUtil::RecalculateNormals_PerVertex(
+	FDynamicMesh3* Mesh, 
+	const TSet<int32>& ModifiedTris, 
+	FUniqueIndexSet& VertexSetTemp)
+{
+	VertexSetTemp.Initialize(Mesh->MaxVertexID());
+	for (int32 TriangleID : ModifiedTris)
+	{
+		const FIndex3i& Triangle = Mesh->GetTriangleRef(TriangleID);
+		VertexSetTemp.Add(Triangle.A);
+		VertexSetTemp.Add(Triangle.B);
+		VertexSetTemp.Add(Triangle.C);
+	}
+
+	const TArray<int32>& ElemIndices = VertexSetTemp.Indices();
+	ParallelFor(ElemIndices.Num(), [&](int32 k) {
+		int32 vid = ElemIndices[k];
+		FVector3d NewNormal = FMeshNormals::ComputeVertexNormal(*Mesh, vid);
+		Mesh->SetVertexNormal(vid, (FVector3f)NewNormal);
+	});
+}
+
+
+
+void UE::SculptUtil::RecalculateROINormals(
+	FDynamicMesh3* Mesh, 
+	const TSet<int32>& TriangleROI, 
+	FUniqueIndexSet& VertexSetTemp,
+	bool bForceVertex)
+{
+	if (Mesh->HasAttributes() && Mesh->Attributes()->PrimaryNormals() && bForceVertex == false)
+	{
+		RecalculateNormals_Overlay(Mesh, TriangleROI, VertexSetTemp);
+	}
+	else
+	{
+		RecalculateNormals_PerVertex(Mesh, TriangleROI, VertexSetTemp);
+	}
+}
+
+
+
+
+static FAxisAlignedBox3d ParallelComputeROIBounds(const FDynamicMesh3& Mesh, const TArray<int32>& Triangles)
+{
+	FAxisAlignedBox3d FinalBounds = FAxisAlignedBox3d::Empty();
+	FCriticalSection FinalBoundsLock;
+	int32 N = Triangles.Num();
+	constexpr int32 BlockSize = 4096;
+	int32 Blocks = (N / BlockSize) + 1;
+	ParallelFor(Blocks, [&](int bi)
+	{
+		FAxisAlignedBox3d BlockBounds = FAxisAlignedBox3d::Empty();
+		for (int32 k = 0; k < BlockSize; ++k)
+		{
+			int32 i = bi * BlockSize + k;
+			if (i < N)
+			{
+				int32 tid = Triangles[i];
+				const FIndex3i& TriV = Mesh.GetTriangleRef(tid);
+				BlockBounds.Contain(Mesh.GetVertexRef(TriV.A));
+				BlockBounds.Contain(Mesh.GetVertexRef(TriV.B));
+				BlockBounds.Contain(Mesh.GetVertexRef(TriV.C));
+			}
+		}
+		FinalBoundsLock.Lock();
+		FinalBounds.Contain(BlockBounds);
+		FinalBoundsLock.Unlock();
+	});
+	return FinalBounds;
+}
+
+
+
+namespace
+{
+	// probably should be something defined for the whole tool framework...
+#if WITH_EDITOR
+	static EAsyncExecution MeshSculptUtilAsyncExecTarget = EAsyncExecution::LargeThreadPool;
+#else
+	static EAsyncExecution MeshSculptUtilAsyncExecTarget = EAsyncExecution::ThreadPool;
+#endif
+}
+
+
+
+void UE::SculptUtil::PrecalculateNormalsROI(
+	const FDynamicMesh3* Mesh,
+	const TArray<int32>& TriangleROI,
+	FUniqueIndexSet& IndexSetTemp,
+	bool& bIsOverlayElements,
+	bool bForceVertex)
+{
+	const FDynamicMeshNormalOverlay* Normals = Mesh->HasAttributes() ? Mesh->Attributes()->PrimaryNormals() : nullptr;
+	if (Normals != nullptr && bForceVertex == false)
+	{
+		bIsOverlayElements = true;
+
+		static TArray<int32> List; 
+		static TArray<std::atomic<bool>> Flags;
+		if (Flags.Num() != Normals->MaxElementID())
+		{
+			Flags.SetNum(Normals->MaxElementID());
+		}
+		{
+			TRACE_CPUPROFILER_EVENT_SCOPE(PrecalculateNormalsROI_InitClear);
+			ParallelFor(Normals->MaxElementID(), [&](int32 i) { Flags[i] = false; });
+		}
+
+		{
+			TRACE_CPUPROFILER_EVENT_SCOPE(PrecalculateNormalsROI_FIND);
+			List.Reset();
+			FCriticalSection ListLock;
+			int32 N = TriangleROI.Num();
+			ParallelFor(TriangleROI.Num(), [&](int i) {
+				FIndex3i TriElems = Normals->GetTriangle(TriangleROI[i]);
+				Flags[TriElems.A] = true;
+				Flags[TriElems.B] = true;
+				Flags[TriElems.C] = true;
+				//bool PrevA = false, PrevB = false, PrevC = false;
+				//bool SetA = Flags[TriElems.A].compare_exchange_strong(PrevA, true);
+				//bool SetB = Flags[TriElems.B].compare_exchange_strong(PrevB, true);
+				//bool SetC = Flags[TriElems.C].compare_exchange_strong(PrevC, true);
+				//if ( SetA || SetB || SetC )
+				//{
+				//	ListLock.Lock();
+				//	if (SetA) List.Add(TriElems.A);
+				//	if (SetB) List.Add(TriElems.B);
+				//	if (SetC) List.Add(TriElems.C);
+				//	ListLock.Unlock();
+				//}
+			});
+		}
+
+		{
+			TRACE_CPUPROFILER_EVENT_SCOPE(PrecalculateNormalsROI_NORMALS);
+			ParallelFor(Normals->MaxElementID(), [&](int32 elemid) {
+				if (Flags[elemid] == true)
+				{
+					FVector3d NewNormal = FMeshNormals::ComputeOverlayNormal(*Mesh, Normals, elemid);
+					//Normals->SetElement(elemid, (FVector3f)NewNormal);
+				}
+			});
+		}
+		//{
+		//	TRACE_CPUPROFILER_EVENT_SCOPE(PrecalculateNormalsROI_Clear);
+		//	for (int32 k : List)
+		//	{
+		//		Flags[k].store(false);
+		//	}
+		//}
+
+		//{
+		//	TRACE_CPUPROFILER_EVENT_SCOPE(PrecalculateNormalsROI_InsertTest);
+		//	static TArray<int32> ListX;
+		//	ListX.Reset();
+		//	for (int32 k : List)
+		//	{
+		//		ListX.Add(k);
+		//	}
+		//}
+
+
+		{
+			TRACE_CPUPROFILER_EVENT_SCOPE(PrecalculateNormalsROI_BitArraySet);
+			IndexSetTemp.Initialize(Normals->MaxElementID());
+			for (int32 TriangleID : TriangleROI)
+			{
+				FIndex3i TriElems = Normals->GetTriangle(TriangleID);
+				IndexSetTemp.Add(TriElems.A);
+				IndexSetTemp.Add(TriElems.B);
+				IndexSetTemp.Add(TriElems.C);
+			}
+		}
+
+		UE_LOG(LogTemp, Warning, TEXT("LIST %d  INDEXSET %d"), List.Num(), IndexSetTemp.Indices().Num());
+	}
+	else
+	{
+		bIsOverlayElements = false;
+
+		IndexSetTemp.Initialize(Mesh->MaxVertexID());
+		for (int32 TriangleID : TriangleROI)
+		{
+			const FIndex3i& Triangle = Mesh->GetTriangleRef(TriangleID);
+			IndexSetTemp.Add(Triangle.A);
+			IndexSetTemp.Add(Triangle.B);
+			IndexSetTemp.Add(Triangle.C);
+		}
+	}
+}
+
+
+
+void UE::SculptUtil::RecalculateROINormals(FDynamicMesh3* Mesh, const TArray<int32>& Indices, bool bIsOverlayElements)
+{
+	if (bIsOverlayElements)
+	{
+		FDynamicMeshNormalOverlay* Normals = Mesh->HasAttributes() ? Mesh->Attributes()->PrimaryNormals() : nullptr;
+		ParallelFor(Indices.Num(), [&](int32 k) {
+			int32 elemid = Indices[k];
+			FVector3d NewNormal = FMeshNormals::ComputeOverlayNormal(*Mesh, Normals, elemid);
+			Normals->SetElement(elemid, (FVector3f)NewNormal);
+		});
+	}
+	else
+	{
+		ParallelFor(Indices.Num(), [&](int32 k) {
+			int32 vid = Indices[k];
+			FVector3d NewNormal = FMeshNormals::ComputeVertexNormal(*Mesh, vid);
+			Mesh->SetVertexNormal(vid, (FVector3f)NewNormal);
+		});
+	}
+}
+
+
+
+
+void UE::SculptUtil::PrecalculateNormalsROI(const FDynamicMesh3* Mesh, const TArray<int32>& TriangleROI,
+	TArray<std::atomic<bool>>& ROIFlags, bool& bIsOverlayElements, bool bForceVertex)
+{
+	const FDynamicMeshNormalOverlay* Normals = Mesh->HasAttributes() ? Mesh->Attributes()->PrimaryNormals() : nullptr;
+	if (Normals != nullptr && bForceVertex == false)
+	{
+		bIsOverlayElements = true;
+		
+		if (ROIFlags.Num() != Normals->MaxElementID())
+		{
+			ROIFlags.SetNum(Normals->MaxElementID());
+			ParallelFor(Normals->MaxElementID(), [&](int i) { ROIFlags[i] = false; });
+		}
+
+		{
+			TRACE_CPUPROFILER_EVENT_SCOPE(PrecalculateNormalsROI_FIND);
+			ParallelFor(TriangleROI.Num(), [&](int i) {
+				FIndex3i TriElems = Normals->GetTriangle(TriangleROI[i]);
+				ROIFlags[TriElems.A] = true;
+				ROIFlags[TriElems.B] = true;
+				ROIFlags[TriElems.C] = true;
+			});
+		}
+	}
+	else
+	{
+		bIsOverlayElements = false;
+
+		check(false);
+	}
+}
+
+void UE::SculptUtil::RecalculateROINormals(FDynamicMesh3* Mesh, TArray<std::atomic<bool>>& ROIFlags, bool bIsOverlayElements)
+{
+	if (bIsOverlayElements)
+	{
+		FDynamicMeshNormalOverlay* Normals = Mesh->HasAttributes() ? Mesh->Attributes()->PrimaryNormals() : nullptr;
+		ParallelFor(ROIFlags.Num(), [&](int32 elemid) 
+		{
+			if (ROIFlags[elemid])
+			{
+				FVector3d NewNormal = FMeshNormals::ComputeOverlayNormal(*Mesh, Normals, elemid);
+				Normals->SetElement(elemid, (FVector3f)NewNormal);
+				ROIFlags[elemid] = false;
+			}
+
+			//bool bFalse = false;
+			//bool bWasTrue = ROIFlags[elemid].compare_exchange_strong(bFalse, true);
+			//if (bWasTrue)
+			//{
+			//	FVector3d NewNormal = FMeshNormals::ComputeOverlayNormal(*Mesh, Normals, elemid);
+			//	Normals->SetElement(elemid, (FVector3f)NewNormal);
+			//}
+		});
+	}
+	else
+	{
+		ParallelFor(ROIFlags.Num(), [&](int32 vid)
+		{
+			if (ROIFlags[vid])
+			{
+				FVector3d NewNormal = FMeshNormals::ComputeVertexNormal(*Mesh, vid);
+				Mesh->SetVertexNormal(vid, (FVector3f)NewNormal);
+				ROIFlags[vid] = false;
+			}
+
+			//bool bFalse = false;
+			//bool bWasTrue = ROIFlags[vid].compare_exchange_strong(bFalse, true);
+			//if (bWasTrue)
+			//{
+			//	FVector3d NewNormal = FMeshNormals::ComputeVertexNormal(*Mesh, vid);
+			//	Mesh->SetVertexNormal(vid, (FVector3f)NewNormal);
+			//}
+		});
+
+	}
+}
+
+
diff --git a/Engine/Plugins/Experimental/MeshModelingToolset/Source/MeshModelingTools/Private/Sculpting/MeshSculptUtil.h b/Engine/Plugins/Experimental/MeshModelingToolset/Source/MeshModelingTools/Private/Sculpting/MeshSculptUtil.h
--- a/Engine/Plugins/Experimental/MeshModelingToolset/Source/MeshModelingTools/Private/Sculpting/MeshSculptUtil.h
+++ b/Engine/Plugins/Experimental/MeshModelingToolset/Source/MeshModelingTools/Private/Sculpting/MeshSculptUtil.h
@@ -6,21 +6,82 @@
 #include "DynamicMeshAABBTree3.h"
 #include "DynamicMeshOctree3.h"
 #include "MeshNormals.h"
+#include "Util/UniqueIndexSet.h"
 
 
 namespace UE 
 { 
 namespace SculptUtil 
 {
+	/** 
+	 * Recompute overlay normals for the overlay elements belonging to all the ModifiedTris. 
+	 * ElementSetBuffer and NormalsBuffer will be populated with all the normal overlay element IDs (you provide to allow for re-use of allocated memory)
+	 */
+	void RecalculateNormals_Overlay(FDynamicMesh3* Mesh, const TSet<int32>& ModifiedTris, TSet<int32>& ElementSetBuffer, TArray<int32>& NormalsBuffer);
 
-	void RecalculateNormals_Overlay(FDynamicMesh3* Mesh, const TSet<int32>& ModifiedTris, TSet<int32>& VertexSetBuffer, TArray<int32>& NormalsBuffer);
+	/**
+	 * Recompute mesh vertex normals for the vertices belonging to all the ModifiedTris.
+	 * VertexSetBuffer and NormalsBuffer will be populated with all the vertex IDs (you provide to allow for re-use of allocated memory)
+	 */
+	void RecalculateNormals_PerVertex(FDynamicMesh3* Mesh, const TSet<int32>& ModifiedTris, TSet<int32>& VertexSetBuffer, TArray<int32>& NormalsBuffer);
 
+	/**
+	 * Recompute normals for either the normals overlay or vertex normals of the mesh that belong to all the ModifiedTris
+	 * TempSetBuffer and NormalsBuffer will be populated with either the recomputed overlay element IDs or vertex IDs
+	 */
+	void RecalculateROINormals(FDynamicMesh3* Mesh, const TSet<int32>& TriangleROI, TSet<int32>& TempSetBuffer, TArray<int32>& NormalsBuffer, bool bForceVertex = false);
 
-	void RecalculateNormals_PerVertex(FDynamicMesh3* Mesh, const TSet<int32>& ModifiedTris, TSet<int32>& VertexSetBuffer, TArray<int32>& NormalsBuffer);
 
 
 
-	void RecalculateROINormals(FDynamicMesh3* Mesh, const TSet<int32>& TriangleROI, TSet<int32>& VertexSetBuffer, TArray<int32>& NormalsBuffer, bool bForceVertex = false);
+	/**
+	 * Recompute overlay normals for the overlay elements belonging to all the ModifiedTris.
+	 * ElementSetTemp will be populated with all the normal overlay element IDs (you provide to allow for re-use of allocated memory)
+	 */
+	void RecalculateNormals_Overlay(FDynamicMesh3* Mesh, const TSet<int32>& ModifiedTris, FUniqueIndexSet& ElementSetTemp);
+
+	/**
+	 * Recompute mesh vertex normals for the vertices belonging to all the ModifiedTris.
+	 * VertexSetTemp will be populated with all the vertex IDs (you provide to allow for re-use of allocated memory)
+	 */
+	void RecalculateNormals_PerVertex(FDynamicMesh3* Mesh, const TSet<int32>& ModifiedTris, FUniqueIndexSet& VertexSetTemp);
+
+	/**
+	 * Recompute normals for either the normals overlay or vertex normals of the mesh that belong to all the ModifiedTris
+	 * IndexSetTemp will be populated with either the recomputed overlay element IDs or vertex IDs
+	 */
+	void RecalculateROINormals(FDynamicMesh3* Mesh, const TSet<int32>& TriangleROI, FUniqueIndexSet& IndexSetTemp, bool bForceVertex = false);
+
+
+
+	/**
+	 * Precalculate IndexSetTemp for the modified TriangleROI, the indices will either be overlay normal element IDs, or vertex IDs
+	 * @param bIsOverlayElementsOut will be returned as true for overlay normals, false for vertex normals
+	 */
+	void PrecalculateNormalsROI(const FDynamicMesh3* Mesh, const TArray<int32>& TriangleROI, FUniqueIndexSet& IndexSetTemp, bool& bIsOverlayElementsOut, bool bForceVertex = false);
+
+	/**
+	 * Recalculate either overlay normals or vertex normals for the given Indices
+	 * @param bIsOverlayElements if true, Indices is interpreted as overlay element IDs, otherwise as vertex IDs
+	 */
+	void RecalculateROINormals(FDynamicMesh3* Mesh, const TArray<int32>& Indices, bool bIsOverlayElements);
+
+
+
+	/**
+	 * Precalculate ROIFlags for the modified TriangleROI, the indices will either be overlay normal element IDs, or vertex IDs.
+	 * ROIFlags will be resized to either the max overlay element ID, or the max vertex ID, and indices in the ROI will be set to true
+	 * @param bIsOverlayElementsOut will be returned as true for overlay normals, false for vertex normals
+	 */
+	void PrecalculateNormalsROI(const FDynamicMesh3* Mesh, const TArray<int32>& TriangleROI, TArray<std::atomic<bool>>& ROIFlags, bool& bIsOverlayElementsOut, bool bForceVertex = false);
+
+
+	/**
+	 * Recalculate either overlay normals or vertex normals for the given indices that are true in ROIFlags
+	 * @param bIsOverlayElements if true, ROIFlags is interpreted as overlay element ID flags, otherwise as vertex IDs
+	 */
+	void RecalculateROINormals(FDynamicMesh3* Mesh, TArray<std::atomic<bool>>& ROIFlags, bool bIsOverlayElements);
+
 
 
 /* end namespace UE::SculptUtil */  } }
\ No newline at end of file
diff --git a/Engine/Plugins/Experimental/MeshModelingToolset/Source/MeshModelingTools/Public/MeshAttributePaintTool.h b/Engine/Plugins/Experimental/MeshModelingToolset/Source/MeshModelingTools/Public/MeshAttributePaintTool.h
--- a/Engine/Plugins/Experimental/MeshModelingToolset/Source/MeshModelingTools/Public/MeshAttributePaintTool.h
+++ b/Engine/Plugins/Experimental/MeshModelingToolset/Source/MeshModelingTools/Public/MeshAttributePaintTool.h
@@ -90,6 +90,30 @@ public:
 
 
 
+UENUM()
+enum class EBrushActionMode
+{
+	Paint,
+	FloodFill
+};
+
+
+/**
+ * Selected-Attribute settings Attribute Paint Tool
+ */
+UCLASS()
+class MESHMODELINGTOOLS_API UMeshAttributePaintBrushOperationProperties : public UInteractiveToolPropertySet
+{
+	GENERATED_BODY()
+public:
+	UPROPERTY(EditAnywhere, Category = Attribute)
+	EBrushActionMode BrushAction = EBrushActionMode::Paint;
+};
+
+
+
+
+
 
 /**
  * Selected-Attribute settings Attribute Paint Tool
@@ -185,10 +209,25 @@ public:
 protected:
 	virtual void ApplyStamp(const FBrushStampData& Stamp);
 
+
+	struct FStampActionData
+	{
+		TArray<int32> ROIVertices;
+		TArray<float> ROIBefore;
+		TArray<float> ROIAfter;
+	};
+
+
+	virtual void ApplyStamp_Paint(const FBrushStampData& Stamp, FStampActionData& ActionData);
+	virtual void ApplyStamp_FloodFill(const FBrushStampData& Stamp, FStampActionData& ActionData);
+
 	virtual void OnShutdown(EToolShutdownType ShutdownType) override;
 	
 
 protected:
+	UPROPERTY()
+	UMeshAttributePaintBrushOperationProperties* BrushActionProps;
+
 	UPROPERTY()
 	UMeshAttributePaintToolProperties* AttribProps;
 
diff --git a/Engine/Plugins/Experimental/MeshModelingToolset/Source/MeshModelingTools/Public/MeshVertexSculptTool.h b/Engine/Plugins/Experimental/MeshModelingToolset/Source/MeshModelingTools/Public/MeshVertexSculptTool.h
--- a/Engine/Plugins/Experimental/MeshModelingToolset/Source/MeshModelingTools/Public/MeshVertexSculptTool.h
+++ b/Engine/Plugins/Experimental/MeshModelingToolset/Source/MeshModelingTools/Public/MeshVertexSculptTool.h
@@ -13,25 +13,13 @@
 #include "TransformTypes.h"
 #include "Sculpting/MeshSculptToolBase.h"
 #include "Sculpting/MeshBrushOpBase.h"
+#include "Image/ImageBuilder.h"
+#include "Util/UniqueIndexSet.h"
 #include "MeshVertexSculptTool.generated.h"
 
 class UTransformGizmo;
 class UTransformProxy;
 class UMaterialInstanceDynamic;
-
-DECLARE_STATS_GROUP(TEXT("VertexSculptTool"), STATGROUP_VtxSculptTool, STATCAT_Advanced);
-DECLARE_CYCLE_STAT(TEXT("VertexSculptTool_UpdateROI"), VtxSculptTool_UpdateROI, STATGROUP_VtxSculptTool);
-DECLARE_CYCLE_STAT(TEXT("VertexSculptTool_ApplyStamp"), VtxSculptToolApplyStamp, STATGROUP_VtxSculptTool );
-DECLARE_CYCLE_STAT(TEXT("VertexSculptTool_Tick"), VtxSculptToolTick, STATGROUP_VtxSculptTool);
-DECLARE_CYCLE_STAT(TEXT("VertexSculptTool_Tick_ApplyStampBlock"), VtxSculptTool_Tick_ApplyStampBlock, STATGROUP_VtxSculptTool);
-DECLARE_CYCLE_STAT(TEXT("VertexSculptTool_Tick_ApplyStamp_Remove"), VtxSculptTool_Tick_ApplyStamp_Remove, STATGROUP_VtxSculptTool);
-DECLARE_CYCLE_STAT(TEXT("VertexSculptTool_Tick_ApplyStamp_Insert"), VtxSculptTool_Tick_ApplyStamp_Insert, STATGROUP_VtxSculptTool);
-DECLARE_CYCLE_STAT(TEXT("VertexSculptTool_Tick_NormalsBlock"), VtxSculptTool_Tick_NormalsBlock, STATGROUP_VtxSculptTool);
-DECLARE_CYCLE_STAT(TEXT("VertexSculptTool_Tick_UpdateMeshBlock"), VtxSculptTool_Tick_UpdateMeshBlock, STATGROUP_VtxSculptTool);
-DECLARE_CYCLE_STAT(TEXT("VertexSculptTool_Tick_UpdateTargetBlock"), VtxSculptTool_Tick_UpdateTargetBlock, STATGROUP_VtxSculptTool);
-DECLARE_CYCLE_STAT(TEXT("VertexSculptTool_Normals_Collect"), VtxSculptTool_Normals_Collect, STATGROUP_VtxSculptTool);
-DECLARE_CYCLE_STAT(TEXT("VertexSculptTool_Normals_Compute"), VtxSculptTool_Normals_Compute, STATGROUP_VtxSculptTool);
-
 class FMeshVertexChangeBuilder;
 class UPreviewMesh;
 
@@ -114,7 +102,6 @@ enum class EMeshVertexSculptBrushType : uint8
 
 
 
-
 UCLASS()
 class MESHMODELINGTOOLS_API UVertexBrushSculptProperties : public UInteractiveToolPropertySet
 {
@@ -141,6 +128,34 @@ public:
 
 
 
+/**
+ * Tool Properties for a brush alpha mask
+ */
+UCLASS()
+class MESHMODELINGTOOLS_API UVertexBrushAlphaProperties : public UInteractiveToolPropertySet
+{
+	GENERATED_BODY()
+
+public:
+	/** Alpha mask applied to brush stamp. Red channel is used. */
+	UPROPERTY(EditAnywhere, Category = Alpha, meta = (DisplayName = "Alpha Mask"))
+	UTexture2D* Alpha = nullptr;
+
+	/** Alpha is rotated by this angle, inside the brush stamp frame (vertically aligned) */
+	UPROPERTY(EditAnywhere, Category = Alpha, meta = (DisplayName = "Rotation", UIMin = "-180.0", UIMax = "180.0", ClampMin = "-360.0", ClampMax = "360.0"))
+	float RotationAngle = 0.0;
+
+	/** If true, a random angle in +/- RandomRange is added to Rotation angle for each stamp */
+	UPROPERTY(EditAnywhere, Category = Alpha, AdvancedDisplay)
+	bool bRandomize = false;
+
+	/** Bounds of random generation (positive and negative) for randomized stamps */
+	UPROPERTY(EditAnywhere, Category = Alpha, AdvancedDisplay, meta = (UIMin = "0.0", UIMax = "180.0"))
+	float RandomRange = 180.0;
+};
+
+
+
 
 
 /**
@@ -170,10 +185,18 @@ public:
 	UPROPERTY()
 	UVertexBrushSculptProperties* SculptProperties;
 
+	UPROPERTY()
+	UVertexBrushAlphaProperties* AlphaProperties;
+
+	UPROPERTY()
+	UTexture2D* BrushAlpha;
+
 public:
 	virtual void IncreaseBrushSpeedAction() override;
 	virtual void DecreaseBrushSpeedAction() override;
 
+	virtual void UpdateBrushAlpha(UTexture2D* NewAlpha);
+
 protected:
 	// UMeshSculptToolBase API
 	virtual UBaseDynamicMeshComponent* GetSculptMeshComponent() { return DynamicMeshComponent; }
@@ -205,17 +228,31 @@ protected:
 	TArray<int> NormalsBuffer;
 	void WaitForPendingUndoRedo();
 
+	TFuture<void> StampUpdateOctreeFuture;
+	bool bStampUpdatePending = false;
+	void WaitForPendingStampUpdate();
+
+	TArray<int> RangeQueryTriBuffer;
+	FUniqueIndexSet VertexROIBuilder;
+	FUniqueIndexSet TriangleROIBuilder;
+	TArray<FIndex3i> TriangleROIInBuf;
 	TArray<int> VertexROI;
-	TSet<int> VertexSetBuffer;
-	TSet<int> TriangleROI;
+	TArray<int> TriangleROIArray;
 	void UpdateROI(const FVector3d& BrushPos);
 
+	FUniqueIndexSet NormalsROIBuilder;
+	TArray<std::atomic<bool>> NormalsFlags;		// set of per-vertex or per-element-id flags that indicate
+												// whether normal needs recompute. Fast to do it this way
+												// than to use a TSet or UniqueIndexSet...
+
 	bool bTargetDirty;
 
 	EMeshVertexSculptBrushType PendingStampType = EMeshVertexSculptBrushType::Smooth;
 
 	bool UpdateStampPosition(const FRay& WorldRay);
-	void ApplyStamp();
+	TFuture<void> ApplyStamp();
+
+	FRandomStream StampRandomStream;
 
 	FDynamicMesh3 BaseMesh;
 	FDynamicMeshOctree3 BaseMeshSpatial;
@@ -231,8 +268,13 @@ protected:
 
 	double SculptMaxFixedHeight = -1.0;
 
+	bool bHaveBrushAlpha = false;
+	TImageBuilder<FVector4f> BrushAlphaValues;
+	FImageDimensions BrushAlphaDimensions;
+	double SampleBrushAlpha(const FSculptBrushStamp& Stamp, const FVector3d& Position) const;
+
 	TArray<FVector3d> ROIPositionBuffer;
-	void SyncMeshWithPositionBuffer(FDynamicMesh3* Mesh);
+	TArray<FVector3d> ROIPrevPositionBuffer;
 
 	FMeshVertexChangeBuilder* ActiveVertexChange = nullptr;
 	void BeginChange();
diff --git a/Engine/Plugins/Experimental/MeshModelingToolset/Source/MeshModelingTools/Public/Sculpting/MeshBrushOpBase.h b/Engine/Plugins/Experimental/MeshModelingToolset/Source/MeshModelingTools/Public/Sculpting/MeshBrushOpBase.h
--- a/Engine/Plugins/Experimental/MeshModelingToolset/Source/MeshModelingTools/Public/Sculpting/MeshBrushOpBase.h
+++ b/Engine/Plugins/Experimental/MeshModelingToolset/Source/MeshModelingTools/Public/Sculpting/MeshBrushOpBase.h
@@ -46,6 +46,10 @@ struct MESHMODELINGTOOLS_API FSculptBrushStamp
 	// only initialized if current op requires it
 	FFrame3d RegionPlane;
 
+	// stamp alpha
+	TFunction<double(const FSculptBrushStamp& Stamp, const FVector3d& Position)> StampAlphaFunc;
+	bool HasAlpha() const { return !!StampAlphaFunc; }
+
 	FSculptBrushStamp()
 	{
 		TimeStamp = FDateTime::Now();
@@ -114,7 +118,7 @@ public:
 
 	virtual void BeginStroke(const FDynamicMesh3* Mesh, const FSculptBrushStamp& Stamp, const TArray<int32>& InitialVertices) {}
 	virtual void EndStroke(const FDynamicMesh3* Mesh, const FSculptBrushStamp& Stamp, const TArray<int32>& FinalVertices) {}
-	virtual void ApplyStamp(const FDynamicMesh3* Mesh, const FSculptBrushStamp& Stamp, const TArray<int32>& Vertices, TArray<FVector3d>& NewPositionsOut) = 0;
+	virtual void ApplyStamp(const FDynamicMesh3* Mesh, const FSculptBrushStamp& Stamp, const TArray<int32>& Vertices, TArray<FVector3d>& NewValuesOut) = 0;
 
 
 
@@ -142,6 +146,11 @@ public:
 	{
 		return false;
 	}
+
+	virtual bool SupportsVariableSpacing() const
+	{
+		return false;
+	}
 };
 
 
diff --git a/Engine/Plugins/Experimental/MeshModelingToolset/Source/MeshModelingTools/Public/Sculpting/MeshSculptBrushOps.h b/Engine/Plugins/Experimental/MeshModelingToolset/Source/MeshModelingTools/Public/Sculpting/MeshSculptBrushOps.h
--- a/Engine/Plugins/Experimental/MeshModelingToolset/Source/MeshModelingTools/Public/Sculpting/MeshSculptBrushOps.h
+++ b/Engine/Plugins/Experimental/MeshModelingToolset/Source/MeshModelingTools/Public/Sculpting/MeshSculptBrushOps.h
@@ -45,12 +45,18 @@ public:
 	{
 		return ESculptBrushOpTargetType::TargetMesh;
 	}
+	virtual bool SupportsVariableSpacing() const override
+	{
+		return true;
+	}
 
 	virtual void ApplyStamp(const FDynamicMesh3* Mesh, const FSculptBrushStamp& Stamp, const TArray<int32>& Vertices, TArray<FVector3d>& NewPositionsOut) override
 	{
 		double UsePower = Stamp.Direction * Stamp.Power * Stamp.Radius * Stamp.DeltaTime * BrushSpeedTuning;
 		double MaxOffset = Stamp.Radius;
 
+		bool bHaveAlpha = Stamp.HasAlpha();
+
 		ParallelFor(Vertices.Num(), [&](int32 k)
 		{
 			int32 VertIdx = Vertices[k];
@@ -64,8 +70,10 @@ public:
 			}
 			else
 			{
+				double Alpha = (bHaveAlpha) ? Stamp.StampAlphaFunc(Stamp, OrigPos) : 1.0;
+
 				FVector3d MoveVec = UsePower * BaseNormal;
-				double Falloff = GetFalloff().Evaluate(Stamp, OrigPos);
+				double Falloff = GetFalloff().Evaluate(Stamp, OrigPos) * Alpha;
 				FVector3d NewPos = OrigPos + Falloff * MoveVec;
 				NewPositionsOut[k] = NewPos;
 			}
@@ -115,7 +123,11 @@ public:
 		return ESculptBrushOpTargetType::TargetMesh;
 	}
 
-	virtual bool GetAlignStampToView() const
+	virtual bool GetAlignStampToView() const override
+	{
+		return true;
+	}
+	virtual bool SupportsVariableSpacing() const override
 	{
 		return true;
 	}
@@ -127,6 +139,8 @@ public:
 		double UsePower = Stamp.Direction * Stamp.Power * Stamp.Radius * Stamp.DeltaTime * BrushSpeedTuning;
 		double MaxOffset = Stamp.Radius;
 
+		bool bHaveAlpha = Stamp.HasAlpha();
+
 		ParallelFor(Vertices.Num(), [&](int32 k)
 		{
 			int32 VertIdx = Vertices[k];
@@ -140,8 +154,10 @@ public:
 			}
 			else
 			{
+				double Alpha = (bHaveAlpha) ? Stamp.StampAlphaFunc(Stamp, OrigPos) : 1.0;
+
 				FVector3d MoveVec = UsePower * StampNormal;
-				double Falloff = GetFalloff().Evaluate(Stamp, OrigPos);
+				double Falloff = GetFalloff().Evaluate(Stamp, OrigPos) * Alpha;
 				FVector3d NewPos = OrigPos + Falloff * MoveVec;
 				NewPositionsOut[k] = NewPos;
 			}
@@ -203,6 +219,10 @@ public:
 	{
 		return ESculptBrushOpTargetType::TargetMesh;
 	}
+	virtual bool SupportsVariableSpacing() const override
+	{
+		return true;
+	}
 
 	virtual void ApplyStamp(const FDynamicMesh3* Mesh, const FSculptBrushStamp& Stamp, const TArray<int32>& Vertices, TArray<FVector3d>& NewPositionsOut) override
 	{
@@ -211,6 +231,8 @@ public:
 		USculptMaxBrushOpProps* Props = GetPropertySetAs<USculptMaxBrushOpProps>();
 		double MaxOffset = (Props->bUseFixedHeight) ? Props->FixedHeight : (Props->MaxHeight * Stamp.Radius);
 
+		bool bHaveAlpha = Stamp.HasAlpha();
+
 		ParallelFor(Vertices.Num(), [&](int32 k)
 		{
 			int32 VertIdx = Vertices[k];
@@ -224,8 +246,10 @@ public:
 			}
 			else
 			{
+				double Alpha = (bHaveAlpha) ? Stamp.StampAlphaFunc(Stamp, OrigPos) : 1.0;
+
 				FVector3d MoveVec = UsePower * BaseNormal;
-				double Falloff = GetFalloff().Evaluate(Stamp, OrigPos);
+				double Falloff = GetFalloff().Evaluate(Stamp, OrigPos) * Alpha;
 				FVector3d NewPos = OrigPos + Falloff * MoveVec;
 
 				FVector3d DeltaPos = NewPos - BasePos;
diff --git a/Engine/Plugins/Experimental/MeshModelingToolset/Source/MeshModelingTools/Public/Sculpting/MeshSculptToolBase.h b/Engine/Plugins/Experimental/MeshModelingToolset/Source/MeshModelingTools/Public/Sculpting/MeshSculptToolBase.h
--- a/Engine/Plugins/Experimental/MeshModelingToolset/Source/MeshModelingTools/Public/Sculpting/MeshSculptToolBase.h
+++ b/Engine/Plugins/Experimental/MeshModelingToolset/Source/MeshModelingTools/Public/Sculpting/MeshSculptToolBase.h
@@ -57,8 +57,16 @@ public:
 	UPROPERTY(EditAnywhere, Category = Brush, meta = (DisplayPriority = 6))
 	bool bHitBackFaces = true;
 
+	/** Brush stamps are applied at this time interval. 0 for a single stamp, 1 for continuous stamps, 0.5 is a stamp every half-second */
+	UPROPERTY(EditAnywhere, Category = Brush, meta = (UIMin = "0.0", UIMax = "1.0", ClampMin = "0.0", ClampMax = "1.0", DisplayPriority = 7, HideEditConditionToggle, EditConditionHides, EditCondition = "bShowFlowRate"))
+	float FlowRate = 1.0f;
 
-	UPROPERTY(EditAnywhere, Category = Brush, meta = (UIMin = "0.0", UIMax = "1.0", ClampMin = "0.0", ClampMax = "1.0", DisplayPriority = 7, HideEditConditionToggle, EditConditionHides, EditCondition = "bShowLazyness"))
+	/** Minimum world-space spacing between stamps, defined along the brush path. Zero spacing means continuous stamps. */
+	UPROPERTY(EditAnywhere, Category = Brush, meta = (UIMin = "0.0", UIMax = "4.0", ClampMin = "0.0", ClampMax = "1000.0", DisplayPriority = 8, HideEditConditionToggle, EditConditionHides, EditCondition = "bShowSpacing"))
+	float Spacing = 0.0f;
+
+	/** Lazy brush smooths out the brush path by averaging the cursor positions */
+	UPROPERTY(EditAnywhere, Category = Brush, meta = (UIMin = "0.0", UIMax = "1.0", ClampMin = "0.0", ClampMax = "1.0", DisplayPriority = 9, HideEditConditionToggle, EditConditionHides, EditCondition = "bShowLazyness"))
 	float Lazyness = 0;
 
 	/**  */
@@ -68,6 +76,14 @@ public:
 	/**  */
 	UPROPERTY(meta = (TransientToolProperty))
 	bool bShowLazyness = true;
+
+	/**  */
+	UPROPERTY(meta = (TransientToolProperty))
+	bool bShowFlowRate = true;
+
+	UPROPERTY(meta = (TransientToolProperty))
+	bool bShowSpacing = true;
+
 };
 
 
@@ -215,6 +231,12 @@ protected:
 	void InitializeSculptMeshComponent(UBaseDynamicMeshComponent* Component);
 
 
+	/**
+	 * Subclass can override this to change what results are written.
+	 * Default is to apply a default vertex positions update to the target object.
+	 */
+	virtual void CommitResult(UBaseDynamicMeshComponent* Component);
+
 
 	//
 	// Brush Types
@@ -236,7 +258,7 @@ protected:
 	virtual void SaveAllBrushTypeProperties(UInteractiveTool* SaveFromTool);
 	virtual void RestoreAllBrushTypeProperties(UInteractiveTool* RestoreToTool);
 
-private:
+protected:
 	TUniquePtr<FMeshSculptBrushOp> PrimaryBrushOp;
 	UMeshSculptBrushOpProps* PrimaryVisiblePropSet = nullptr;		// BrushOpPropSets prevents GC of this
 
@@ -244,7 +266,7 @@ private:
 	UMeshSculptBrushOpProps* SecondaryVisiblePropSet = nullptr;
 
 	bool bBrushOpPropsVisible = true;
-protected:
+
 	void SetActivePrimaryBrushType(int32 Identifier);
 	void SetActiveSecondaryBrushType(int32 Identifier);
 	TUniquePtr<FMeshSculptBrushOp>& GetActiveBrushOp();
@@ -253,9 +275,9 @@ protected:
 	//
 	// Falloff types
 	//
-private:
-	TSharedPtr<FMeshSculptFallofFunc> PrimaryFalloff;
 protected:
+	TSharedPtr<FMeshSculptFallofFunc> PrimaryFalloff;
+
 	void SetPrimaryFalloffType(EMeshSculptFalloffType Falloff);
 
 
@@ -298,11 +320,11 @@ public:
 	//
 	// Brush/Stroke stuff
 	//
-private:
+protected:
 	FFrame3d LastBrushFrameWorld;
 	FFrame3d LastBrushFrameLocal;
 	int32 LastBrushTriangleID;
-protected:
+
 	const FFrame3d& GetBrushFrameWorld() const { return LastBrushFrameWorld; }
 	const FFrame3d& GetBrushFrameLocal() const { return LastBrushFrameLocal; }
 	int32 GetBrushTriangleID() const { return LastBrushTriangleID; }
@@ -325,59 +347,64 @@ protected:
 	//
 	// Brush Target Plane is plane that some brushes move on
 	//
-private:
-	FFrame3d ActiveBrushTargetPlaneWorld;
 protected:
-	void UpdateBrushTargetPlaneFromHit(const FRay& WorldRay, const FHitResult& Hit);
+	FFrame3d ActiveBrushTargetPlaneWorld;
+	virtual void UpdateBrushTargetPlaneFromHit(const FRay& WorldRay, const FHitResult& Hit);
 
 
 	//
 	// Stroke Modifiers
 	//
-private:
+protected:
 	bool bInStroke = false;
 	bool bSmoothing = false;
 	bool bInvert = false;
-protected:
-	void SaveActiveStrokeModifiers();
-	bool InStroke() const { return bInStroke; }
-	bool GetInSmoothingStroke() const { return bSmoothing; }
-	bool GetInInvertStroke() const { return bInvert; }
-
-	void ApplyStrokeFlowInTick();
+	virtual void SaveActiveStrokeModifiers();
+	virtual bool InStroke() const { return bInStroke; }
+	virtual bool GetInSmoothingStroke() const { return bSmoothing; }
+	virtual bool GetInInvertStroke() const { return bInvert; }
+
+	// when in a stroke, this function determines when a new stamp should be emitted, based on spacing and flow rate settings
+	virtual void UpdateStampPendingState();
+
+	// for tracking stroke time and length, to apply spacing and flow rate settings
+	double ActiveStrokeTime = 0.0;
+	double ActiveStrokePathArcLen = 0.0;
+	int LastFlowTimeStamp = 0;
+	int LastSpacingTimestamp = 0;
+	virtual void ResetStrokeTime();
+	virtual void AccumulateStrokeTime(float DeltaTime);
 
 	//
 	// Stamps
 	//
-private:
+protected:
 	bool bIsStampPending = false;
 	FRay PendingStampRay;
-protected:
 	FSculptBrushStamp HoverStamp;
 	FSculptBrushStamp CurrentStamp;
 	FSculptBrushStamp LastStamp;
-	void UpdateHoverStamp(const FFrame3d& StampFrame);
-	bool IsStampPending() const { return bIsStampPending; }
-	const FRay& GetPendingStampRayWorld() const { return PendingStampRay;  }
+	virtual void UpdateHoverStamp(const FFrame3d& StampFrame);
+	virtual bool IsStampPending() const { return bIsStampPending; }
+	virtual const FRay& GetPendingStampRayWorld() const { return PendingStampRay;  }
 
 
 	//
 	// Stamp ROI Plane is a plane used by some brush ops
 	//
-private:
-	FFrame3d StampRegionPlane;
 protected:
+	FFrame3d StampRegionPlane;
+	virtual FFrame3d ComputeStampRegionPlane(const FFrame3d& StampFrame, const TArray<int32>& StampTriangles, bool bIgnoreDepth, bool bViewAligned, bool bInvDistFalloff = true);
 	virtual FFrame3d ComputeStampRegionPlane(const FFrame3d& StampFrame, const TSet<int32>& StampTriangles, bool bIgnoreDepth, bool bViewAligned, bool bInvDistFalloff = true);
 
 
 	// Stroke plane is a plane used by some brush ops
 	//
-private:
-	FFrame3d StrokePlane;
 protected:
+	FFrame3d StrokePlane;
 	virtual const FFrame3d& GetCurrentStrokeReferencePlane() const { return StrokePlane; }
 
-	virtual void UpdateStrokeReferencePlaneForROI(const FFrame3d& StampFrame, const TSet<int32>& TriangleROI, bool bViewAligned);
+	virtual void UpdateStrokeReferencePlaneForROI(const FFrame3d& StampFrame, const TArray<int32>& TriangleROI, bool bViewAligned);
 	virtual void UpdateStrokeReferencePlaneFromWorkPlane();
 
 
@@ -450,12 +477,12 @@ protected:
 	virtual void UpdateWorkPlane();
 	virtual bool ShowWorkPlane() const { return false; };
 
-private:
+protected:
 	TValueWatcher<FVector> GizmoPositionWatcher;
 	TValueWatcher<FQuat> GizmoRotationWatcher;
 
-	void UpdateGizmoFromProperties();
-	void PlaneTransformChanged(UTransformProxy* Proxy, FTransform Transform);
+	virtual void UpdateGizmoFromProperties();
+	virtual void PlaneTransformChanged(UTransformProxy* Proxy, FTransform Transform);
 
 	enum class EPendingWorkPlaneUpdate
 	{
@@ -465,9 +492,9 @@ private:
 		MoveToHitPositionViewAligned
 	};
 	EPendingWorkPlaneUpdate PendingWorkPlaneUpdate;
-	void SetFixedSculptPlaneFromWorldPos(const FVector& Position, const FVector& Normal, EPendingWorkPlaneUpdate UpdateType);
-	void UpdateFixedSculptPlanePosition(const FVector& Position);
-	void UpdateFixedSculptPlaneRotation(const FQuat& Rotation);
-	void UpdateFixedPlaneGizmoVisibility(bool bVisible);
+	virtual void SetFixedSculptPlaneFromWorldPos(const FVector& Position, const FVector& Normal, EPendingWorkPlaneUpdate UpdateType);
+	virtual void UpdateFixedSculptPlanePosition(const FVector& Position);
+	virtual void UpdateFixedSculptPlaneRotation(const FQuat& Rotation);
+	virtual void UpdateFixedPlaneGizmoVisibility(bool bVisible);
 
 };
\ No newline at end of file
diff --git a/Engine/Plugins/Experimental/MeshModelingToolset/Source/ModelingComponents/Private/Changes/MeshVertexChange.cpp b/Engine/Plugins/Experimental/MeshModelingToolset/Source/ModelingComponents/Private/Changes/MeshVertexChange.cpp
--- a/Engine/Plugins/Experimental/MeshModelingToolset/Source/ModelingComponents/Private/Changes/MeshVertexChange.cpp
+++ b/Engine/Plugins/Experimental/MeshModelingToolset/Source/ModelingComponents/Private/Changes/MeshVertexChange.cpp
@@ -62,6 +62,25 @@ void FMeshVertexChangeBuilder::UpdateVertex(int VertexID, const FVector3d& OldPo
 	}
 }
 
+
+void FMeshVertexChangeBuilder::UpdateVertexColor(int32 VertexID, const FVector3f& OldColor, const FVector3f& NewColor)
+{
+	const int* FoundIndex = SavedVertices.Find(VertexID);
+	if (FoundIndex == nullptr)
+	{
+		int NewIndex = Change->Vertices.Num();
+		SavedVertices.Add(VertexID, NewIndex);
+		Change->Vertices.Add(VertexID);
+		Change->OldColors.Add(OldColor);
+		Change->NewColors.Add(NewColor);
+	}
+	else
+	{
+		Change->NewColors[*FoundIndex] = NewColor;
+	}
+}
+
+
 void FMeshVertexChangeBuilder::UpdateVertexFinal(int VertexID, const FVector3d& NewPosition)
 {
 	check(SavedVertices.Contains(VertexID));
diff --git a/Engine/Plugins/Experimental/MeshModelingToolset/Source/ModelingComponents/Public/Changes/MeshVertexChange.h b/Engine/Plugins/Experimental/MeshModelingToolset/Source/ModelingComponents/Public/Changes/MeshVertexChange.h
--- a/Engine/Plugins/Experimental/MeshModelingToolset/Source/ModelingComponents/Public/Changes/MeshVertexChange.h
+++ b/Engine/Plugins/Experimental/MeshModelingToolset/Source/ModelingComponents/Public/Changes/MeshVertexChange.h
@@ -88,6 +88,9 @@ public:
 	// currently only used in vertex sculpt tool. cannot be used if bSaveColors = true
 	void UpdateVertex(int32 VertexID, const FVector3d& OldPosition, const FVector3d& NewPosition);
 
+	// currently only used in element paint tool. Can only be used if bSaveColors=true and bSavePositions=false
+	void UpdateVertexColor(int32 VertexID, const FVector3f& OldColor, const FVector3f& NewColor);
+
 protected:
 	void UpdateVertexFinal(int32 VertexID, const FVector3d& NewPosition);
 
