ModelingTools:
Add support for submitting multiple UV channels to render buffers in UBaseDynamicMeshSceneProxy and USimpleDynamicMeshSceneProxy (currently not supported if there is a mesh decomposition, and not supported in fast UV updates).
Add support for visualizing different UV channels in UExistingMeshMaterialProperties.
Add UMeshUVChannelProperties property set for selecting from available UV channels.
Add support for arbitrary UV channels in ParameterizeMeshTool and UVLayoutTool/Op.
#rnx

diff --git a/Engine/Plugins/Experimental/MeshModelingToolset/Source/MeshModelingTools/Private/Properties/MeshMaterialProperties.cpp b/Engine/Plugins/Experimental/MeshModelingToolset/Source/MeshModelingTools/Private/Properties/MeshMaterialProperties.cpp
--- a/Engine/Plugins/Experimental/MeshModelingToolset/Source/MeshModelingTools/Private/Properties/MeshMaterialProperties.cpp
+++ b/Engine/Plugins/Experimental/MeshModelingToolset/Source/MeshModelingTools/Private/Properties/MeshMaterialProperties.cpp
@@ -32,6 +32,7 @@ void UExistingMeshMaterialProperties::Setup()
 		if (CheckerMaterial != nullptr)
 		{
 			CheckerMaterial->SetScalarParameterValue("Density", CheckerDensity);
+			CheckerMaterial->SetScalarParameterValue("UVChannel", (float)UVChannel);
 		}
 	}
 }
@@ -41,6 +42,7 @@ void UExistingMeshMaterialProperties::UpdateMaterials()
 	if (CheckerMaterial != nullptr)
 	{
 		CheckerMaterial->SetScalarParameterValue("Density", CheckerDensity);
+		CheckerMaterial->SetScalarParameterValue("UVChannel", (float)UVChannel);
 	}
 }
 
diff --git a/Engine/Plugins/Experimental/MeshModelingToolset/Source/MeshModelingTools/Private/Properties/MeshUVChannelProperties.cpp b/Engine/Plugins/Experimental/MeshModelingToolset/Source/MeshModelingTools/Private/Properties/MeshUVChannelProperties.cpp
new file mode 100644
--- /dev/null
+++ b/Engine/Plugins/Experimental/MeshModelingToolset/Source/MeshModelingTools/Private/Properties/MeshUVChannelProperties.cpp
@@ -0,0 +1,71 @@
+// Copyright Epic Games, Inc. All Rights Reserved.
+
+#include "Properties/MeshUVChannelProperties.h"
+
+#include "DynamicMesh3.h"
+#include "DynamicMeshAttributeSet.h"
+
+#include "MeshDescription.h"
+#include "StaticMeshAttributes.h"
+#include "StaticMeshOperations.h"
+
+
+
+void UMeshUVChannelProperties::Initialize(int32 NumUVChannels, bool bInitializeSelection)
+{
+	UVChannelNamesList.Reset();
+	for (int32 k = 0; k < NumUVChannels; ++k)
+	{
+		UVChannelNamesList.Add(FString::Printf(TEXT("UV%d"), k));
+	}
+	if (bInitializeSelection)
+	{
+		UVChannel = (NumUVChannels > 0) ? UVChannelNamesList[0] : TEXT("");
+	}
+}
+
+
+const TArray<FString>& UMeshUVChannelProperties::GetUVChannelNamesFunc()
+{
+	return UVChannelNamesList;
+}
+
+
+void UMeshUVChannelProperties::Initialize(const FMeshDescription* MeshDescription, bool bInitializeSelection)
+{
+	TVertexInstanceAttributesConstRef<FVector2D> InstanceUVs =
+		MeshDescription->VertexInstanceAttributes().GetAttributesRef<FVector2D>(MeshAttribute::VertexInstance::TextureCoordinate);
+	Initialize(InstanceUVs.GetNumChannels(), bInitializeSelection);
+}
+
+void UMeshUVChannelProperties::Initialize(const FDynamicMesh3* Mesh, bool bInitializeSelection)
+{
+	int32 NumUVChannels = Mesh->HasAttributes() ? Mesh->Attributes()->NumUVLayers() : 0;
+	Initialize(NumUVChannels, bInitializeSelection);
+}
+
+
+
+bool UMeshUVChannelProperties::ValidateSelection(bool bUpdateIfInvalid)
+{
+	int32 FoundIndex = UVChannelNamesList.IndexOfByKey(UVChannel);
+	if (FoundIndex == INDEX_NONE)
+	{
+		if (bUpdateIfInvalid)
+		{
+			UVChannel = (UVChannelNamesList.Num() > 0) ? UVChannelNamesList[0] : TEXT("");
+		}
+		return false;
+	}
+	return true;
+}
+
+int32 UMeshUVChannelProperties::GetSelectedChannelIndex(bool bForceToZeroOnFailure)
+{
+	int32 FoundIndex = UVChannelNamesList.IndexOfByKey(UVChannel);
+	if (FoundIndex == INDEX_NONE)
+	{
+		return (bForceToZeroOnFailure ) ? 0 : -1;
+	}
+	return FoundIndex;
+}
\ No newline at end of file
diff --git a/Engine/Plugins/Experimental/MeshModelingToolset/Source/MeshModelingTools/Public/Properties/MeshMaterialProperties.h b/Engine/Plugins/Experimental/MeshModelingToolset/Source/MeshModelingTools/Public/Properties/MeshMaterialProperties.h
--- a/Engine/Plugins/Experimental/MeshModelingToolset/Source/MeshModelingTools/Public/Properties/MeshMaterialProperties.h
+++ b/Engine/Plugins/Experimental/MeshModelingToolset/Source/MeshModelingTools/Public/Properties/MeshMaterialProperties.h
@@ -71,9 +71,12 @@ public:
 	UPROPERTY(EditAnywhere, Category = MaterialPreview)
 	ESetMeshMaterialMode MaterialMode = ESetMeshMaterialMode::KeepOriginal;
 
-	UPROPERTY(EditAnywhere, Category = MaterialPreview, meta = (UIMin = "1.0", UIMax = "40.0", ClampMin = "0.01", ClampMax = "1000.0", EditCondition = "MaterialMode == ESetMeshMaterialMode::Checkerboard"))
+	UPROPERTY(EditAnywhere, Category = MaterialPreview, meta = (UIMin = "1.0", UIMax = "40.0", ClampMin = "0.01", ClampMax = "1000.0", EditConditionHides, EditCondition = "MaterialMode == ESetMeshMaterialMode::Checkerboard"))
 	float CheckerDensity = 20.0f;
 
+	UPROPERTY(EditAnywhere, Category = MaterialPreview, meta = (DisplayName = "Show UV Channel", EditConditionHides, EditCondition = "MaterialMode == ESetMeshMaterialMode::Checkerboard"))
+	int UVChannel = 0;
+
 	UPROPERTY(EditAnywhere, Category = MaterialPreview, meta = (EditCondition = "MaterialMode == ESetMeshMaterialMode::Override"))
 	UMaterialInterface* OverrideMaterial = nullptr;
 
diff --git a/Engine/Plugins/Experimental/MeshModelingToolset/Source/MeshModelingTools/Public/Properties/MeshUVChannelProperties.h b/Engine/Plugins/Experimental/MeshModelingToolset/Source/MeshModelingTools/Public/Properties/MeshUVChannelProperties.h
new file mode 100644
--- /dev/null
+++ b/Engine/Plugins/Experimental/MeshModelingToolset/Source/MeshModelingTools/Public/Properties/MeshUVChannelProperties.h
@@ -0,0 +1,47 @@
+// Copyright Epic Games, Inc. All Rights Reserved.
+
+#pragma once
+
+#include "CoreMinimal.h"
+#include "InteractiveTool.h"
+
+#include "MeshUVChannelProperties.generated.h"
+
+struct FMeshDescription;
+class FDynamicMesh3;
+
+
+UCLASS()
+class MESHMODELINGTOOLS_API UMeshUVChannelProperties : public UInteractiveToolPropertySet
+{
+	GENERATED_BODY()
+public:
+
+	UPROPERTY(EditAnywhere, Category = UVChannel, meta = (DisplayName = "UV Channel", GetOptions = GetUVChannelNamesFunc))
+	FString UVChannel;
+
+	UFUNCTION()
+	const TArray<FString>& GetUVChannelNamesFunc();
+
+	UPROPERTY(meta = (TransientToolProperty))
+	TArray<FString> UVChannelNamesList;
+
+
+public:
+	void Initialize(int32 NumUVChannels, bool bInitializeSelection = true);
+	void Initialize(const FMeshDescription* MeshDescription, bool bInitializeSelection = true);
+	void Initialize(const FDynamicMesh3* Mesh, bool bInitializeSelection = true);
+
+	/**
+	 * Verify that the UV channel selection is valid
+	 * @param bUpdateIfInvalid if selection is not valid, reset UVChannel to UV0 or empty if no UV channels exist
+	 * @return true if selection in UVChannel is an entry in UVChannelNamesList.
+	 */
+	bool ValidateSelection(bool bUpdateIfInvalid);
+
+	/**
+	 * @param bForceToZeroOnFailure if true, then instead of returning -1 we return 0 so calling code can fallback to default UV paths
+	 * @return selected UV channel index, or -1 if invalid selection, or 0 if invalid selection and bool bForceToZeroOnFailure = true
+	 */
+	int32 GetSelectedChannelIndex(bool bForceToZeroOnFailure = false);
+};
diff --git a/Engine/Plugins/Experimental/MeshModelingToolset/Source/MeshModelingToolsEditorOnly/Private/ParameterizeMeshTool.cpp b/Engine/Plugins/Experimental/MeshModelingToolset/Source/MeshModelingToolsEditorOnly/Private/ParameterizeMeshTool.cpp
--- a/Engine/Plugins/Experimental/MeshModelingToolset/Source/MeshModelingToolsEditorOnly/Private/ParameterizeMeshTool.cpp
+++ b/Engine/Plugins/Experimental/MeshModelingToolset/Source/MeshModelingToolsEditorOnly/Private/ParameterizeMeshTool.cpp
@@ -89,6 +89,11 @@ void UParameterizeMeshTool::Setup()
 
 	Preview->PreviewMesh->InitializeMesh(ComponentTarget->GetMesh());
 
+	Preview->OnMeshUpdated.AddLambda([this](UMeshOpPreviewWithBackgroundCompute* Op)
+	{
+		MaterialSettings->UpdateMaterials();
+	});
+
 	// Initialize the preview mesh with a copy of the source mesh.
 	bool bHasGroups = false;
 	{
@@ -114,12 +119,23 @@ void UParameterizeMeshTool::Setup()
 	}
 
 	// initialize our properties
+
+	UVChannelProperties = NewObject<UMeshUVChannelProperties>(this);
+	UVChannelProperties->RestoreProperties(this);
+	UVChannelProperties->Initialize(ComponentTarget->GetMesh(), false);
+	UVChannelProperties->ValidateSelection(true);
+	UVChannelProperties->WatchProperty(UVChannelProperties->UVChannel, [this](const FString& NewValue) 
+	{
+		MaterialSettings->UVChannel = UVChannelProperties->GetSelectedChannelIndex(true);
+	});
+
+	AddToolPropertySource(UVChannelProperties);
+
 	Settings = NewObject<UParameterizeMeshToolProperties>(this);
 	Settings->RestoreProperties(this);
 	Settings->bIsGlobalMode = bDoAutomaticGlobalUnwrap;
 	AddToolPropertySource(Settings);
 
-
 	MaterialSettings = NewObject<UExistingMeshMaterialProperties>(this);
 	MaterialSettings->RestoreProperties(this);
 	AddToolPropertySource(MaterialSettings);
@@ -148,23 +164,28 @@ void UParameterizeMeshTool::Setup()
 
 void UParameterizeMeshTool::OnPropertyModified(UObject* PropertySet, FProperty* Property)
 {
-	if (PropertySet == MaterialSettings)
-	{
-		MaterialSettings->UpdateMaterials();
-		Preview->OverrideMaterial = MaterialSettings->GetActiveOverrideMaterial();
-	}
-	if (PropertySet == Settings)
+	bool bForceMaterialUpdate = false;
+	if (PropertySet == Settings || PropertySet == UVChannelProperties)
 	{
 		// One of the UV generation properties must have changed.  Dirty the result to force a recompute
 		Preview->InvalidateResult();
+		bForceMaterialUpdate = true;
+	}
+
+	if (PropertySet == MaterialSettings || bForceMaterialUpdate)
+	{
+		MaterialSettings->UpdateMaterials();
+		Preview->OverrideMaterial = MaterialSettings->GetActiveOverrideMaterial();
 	}
 }
 
 
 void UParameterizeMeshTool::Shutdown(EToolShutdownType ShutdownType)
 {
+	UVChannelProperties->SaveProperties(this);
 	Settings->SaveProperties(this);
 	MaterialSettings->SaveProperties(this);
+
 	FDynamicMeshOpResult Result = Preview->Shutdown();
 	if (ShutdownType == EToolShutdownType::Accept)
 	{
@@ -202,6 +223,7 @@ TUniquePtr<FDynamicMeshOperator> UParameterizeMeshTool::MakeNewOperator()
 	ParameterizeMeshOp->Stretch   = Settings->ChartStretch;
 	ParameterizeMeshOp->NumCharts = 0;
 	ParameterizeMeshOp->InputMesh = InputMesh;
+	ParameterizeMeshOp->UVLayer = UVChannelProperties->GetSelectedChannelIndex(true);
 	
 	if (bDoAutomaticGlobalUnwrap)
 	{
diff --git a/Engine/Plugins/Experimental/MeshModelingToolset/Source/MeshModelingToolsEditorOnly/Private/UVLayoutTool.cpp b/Engine/Plugins/Experimental/MeshModelingToolset/Source/MeshModelingToolsEditorOnly/Private/UVLayoutTool.cpp
--- a/Engine/Plugins/Experimental/MeshModelingToolset/Source/MeshModelingToolsEditorOnly/Private/UVLayoutTool.cpp
+++ b/Engine/Plugins/Experimental/MeshModelingToolset/Source/MeshModelingToolsEditorOnly/Private/UVLayoutTool.cpp
@@ -86,9 +86,22 @@ void UUVLayoutTool::Setup()
 		ComponentTarget->SetOwnerVisibility(false);
 	}
 
-	BasicProperties = NewObject<UUVLayoutToolProperties>(this, TEXT("UV Projection Settings"));
-	BasicProperties->RestoreProperties(this);
+	// if we only have one object, add ability to set UV channel
+	if (ComponentTargets.Num() == 1)
+	{
+		UVChannelProperties = NewObject<UMeshUVChannelProperties>(this);
+		UVChannelProperties->RestoreProperties(this);
+		UVChannelProperties->Initialize(ComponentTargets[0]->GetMesh(), false);
+		UVChannelProperties->ValidateSelection(true);
+		AddToolPropertySource(UVChannelProperties);
+		UVChannelProperties->WatchProperty(UVChannelProperties->UVChannel, [this](const FString& NewValue)
+		{
+			MaterialSettings->UVChannel = UVChannelProperties->GetSelectedChannelIndex(true);
+		});
+	}
 
+	BasicProperties = NewObject<UUVLayoutToolProperties>(this);
+	BasicProperties->RestoreProperties(this);
 	AddToolPropertySource(BasicProperties);
 
 	MaterialSettings = NewObject<UExistingMeshMaterialProperties>(this);
@@ -115,7 +128,7 @@ void UUVLayoutTool::Setup()
 
 	UpdateVisualization();
 
-	GetToolManager()->DisplayMessage(LOCTEXT("OnStartUVLayoutTool", "Transform/Rotate/Scale existing UV Islands/Shells/Charts using various strategies"),
+	GetToolManager()->DisplayMessage(LOCTEXT("OnStartUVLayoutTool", "Transform/Rotate/Scale existing UV Charts using various strategies"),
 		EToolMessageLevel::UserNotification);
 }
 
@@ -220,6 +233,8 @@ TUniquePtr<FDynamicMeshOperator> UUVLayoutOperatorFactory::MakeNewOperator()
 		break;
 	}
 
+
+	Op->UVLayerIndex = Tool->GetSelectedUVChannel();
 	//Op->bSeparateUVIslands = Tool->BasicProperties->bSeparateUVIslands;
 	Op->TextureResolution = Tool->BasicProperties->TextureResolution;
 	Op->bAllowFlips = Tool->BasicProperties->bAllowFlips;
@@ -231,6 +246,11 @@ TUniquePtr<FDynamicMeshOperator> UUVLayoutOperatorFactory::MakeNewOperator()
 }
 
 
+int32 UUVLayoutTool::GetSelectedUVChannel() const
+{
+	return UVChannelProperties->GetSelectedChannelIndex(true);
+}
+
 
 void UUVLayoutTool::Render(IToolsContextRenderAPI* RenderAPI)
 {
@@ -262,7 +282,7 @@ void UUVLayoutTool::OnTick(float DeltaTime)
 
 void UUVLayoutTool::OnPropertyModified(UObject* PropertySet, FProperty* Property)
 {
-	if (PropertySet == BasicProperties)
+	if (PropertySet == BasicProperties || PropertySet == UVChannelProperties)
 	{
 		UpdateNumPreviews();
 		for (UMeshOpPreviewWithBackgroundCompute* Preview : Previews)
@@ -287,7 +307,7 @@ void UUVLayoutTool::OnPreviewMeshUpdated(UMeshOpPreviewWithBackgroundCompute* Co
 		{
 			return;
 		}
-		UVLayoutView->UpdateUVMesh(&ResultMesh);
+		UVLayoutView->UpdateUVMesh(&ResultMesh, GetSelectedUVChannel());
 	}
 
 }
diff --git a/Engine/Plugins/Experimental/MeshModelingToolset/Source/MeshModelingToolsEditorOnly/Public/ParameterizeMeshTool.h b/Engine/Plugins/Experimental/MeshModelingToolset/Source/MeshModelingToolsEditorOnly/Public/ParameterizeMeshTool.h
--- a/Engine/Plugins/Experimental/MeshModelingToolset/Source/MeshModelingToolsEditorOnly/Public/ParameterizeMeshTool.h
+++ b/Engine/Plugins/Experimental/MeshModelingToolset/Source/MeshModelingToolsEditorOnly/Public/ParameterizeMeshTool.h
@@ -9,6 +9,7 @@
 #include "DynamicMesh3.h"
 #include "MeshOpPreviewHelpers.h"
 #include "Properties/MeshMaterialProperties.h"
+#include "Properties/MeshUVChannelProperties.h"
 
 #include "ParameterizeMeshTool.generated.h"
 
@@ -164,6 +165,9 @@ public:
 	virtual TUniquePtr<FDynamicMeshOperator> MakeNewOperator() override;
 
 protected:
+	UPROPERTY()
+	UMeshUVChannelProperties* UVChannelProperties = nullptr;
+
 	UPROPERTY()
 	UParameterizeMeshToolProperties* Settings = nullptr;
 
diff --git a/Engine/Plugins/Experimental/MeshModelingToolset/Source/MeshModelingToolsEditorOnly/Public/UVLayoutTool.h b/Engine/Plugins/Experimental/MeshModelingToolset/Source/MeshModelingToolsEditorOnly/Public/UVLayoutTool.h
--- a/Engine/Plugins/Experimental/MeshModelingToolset/Source/MeshModelingToolsEditorOnly/Public/UVLayoutTool.h
+++ b/Engine/Plugins/Experimental/MeshModelingToolset/Source/MeshModelingToolsEditorOnly/Public/UVLayoutTool.h
@@ -11,6 +11,7 @@
 #include "DynamicMesh3.h"
 #include "BaseTools/SingleClickTool.h"
 #include "Properties/MeshMaterialProperties.h"
+#include "Properties/MeshUVChannelProperties.h"
 #include "Drawing/UVLayoutPreview.h"
 
 #include "UVLayoutTool.generated.h"
@@ -136,10 +137,18 @@ public:
 
 	virtual void OnPropertyModified(UObject* PropertySet, FProperty* Property) override;
 
+
+public:
+	int32 GetSelectedUVChannel() const;
+
+
 protected:
 
 	UPROPERTY()
-	UUVLayoutToolProperties* BasicProperties;
+	UMeshUVChannelProperties* UVChannelProperties = nullptr;
+
+	UPROPERTY()
+	UUVLayoutToolProperties* BasicProperties = nullptr;
 
 	UPROPERTY()
 	UExistingMeshMaterialProperties* MaterialSettings = nullptr;
@@ -150,8 +159,8 @@ protected:
 protected:
 	TArray<TSharedPtr<FDynamicMesh3>> OriginalDynamicMeshes;
 
-	UWorld* TargetWorld;
-	IAssetGenerationAPI* AssetAPI;
+	UWorld* TargetWorld = nullptr;
+	IAssetGenerationAPI* AssetAPI = nullptr;
 
 	FViewCameraState CameraState;
 
diff --git a/Engine/Plugins/Experimental/MeshModelingToolset/Source/ModelingComponents/Private/BaseDynamicMeshSceneProxy.h b/Engine/Plugins/Experimental/MeshModelingToolset/Source/ModelingComponents/Private/BaseDynamicMeshSceneProxy.h
--- a/Engine/Plugins/Experimental/MeshModelingToolset/Source/ModelingComponents/Private/BaseDynamicMeshSceneProxy.h
+++ b/Engine/Plugins/Experimental/MeshModelingToolset/Source/ModelingComponents/Private/BaseDynamicMeshSceneProxy.h
@@ -401,13 +401,33 @@ public:
 		FMeshRenderBufferSet* RenderBuffers,
 		const FDynamicMesh3* Mesh,
 		int NumTriangles, TriangleEnumerable Enumerable,
-		FDynamicMeshUVOverlay* UVOverlay,
-		FDynamicMeshNormalOverlay* NormalOverlay,
+		const FDynamicMeshUVOverlay* UVOverlay,
+		const FDynamicMeshNormalOverlay* NormalOverlay,
 		TFunction<void(int, int, int, FVector3f&, FVector3f&)> TangentsFunc = nullptr,
 		bool bTrackTriangles = false)
 	{
-		//SCOPE_CYCLE_COUNTER(STAT_SculptToolOctree_InitializeBufferFromOverlay);
+		TArray<const FDynamicMeshUVOverlay*> UVOverlays;
+		UVOverlays.Add(UVOverlay);
+		InitializeBuffersFromOverlays(RenderBuffers, Mesh, NumTriangles, Enumerable,
+			UVOverlays, NormalOverlay, TangentsFunc, bTrackTriangles);
+	}
+
 
+
+	/**
+	 * Initialize rendering buffers from given attribute overlays.
+	 * Creates three vertices per triangle, IE no shared vertices in buffers.
+	 */
+	template<typename TriangleEnumerable>
+	void InitializeBuffersFromOverlays(
+		FMeshRenderBufferSet* RenderBuffers,
+		const FDynamicMesh3* Mesh,
+		int NumTriangles, TriangleEnumerable Enumerable,
+		const TArray<const FDynamicMeshUVOverlay*>& UVOverlays,
+		const FDynamicMeshNormalOverlay* NormalOverlay,
+		TFunction<void(int, int, int, FVector3f&, FVector3f&)> TangentsFunc = nullptr,
+		bool bTrackTriangles = false)
+	{
 		RenderBuffers->TriangleCount = NumTriangles;
 		if (NumTriangles == 0)
 		{
@@ -417,7 +437,9 @@ public:
 		bool bHaveColors = Mesh->HasVertexColors() && (bIgnoreVertexColors == false);
 
 		int NumVertices = NumTriangles * 3;
-		int NumTexCoords = 1;		// no! zero!
+		int NumTexCoords = UVOverlays.Num();
+		TArray<FIndex3i, TFixedAllocator<MAX_STATIC_TEXCOORDS>> UVTriangles;
+		UVTriangles.SetNum(NumTexCoords);
 
 		{
 			RenderBuffers->PositionVertexBuffer.Init(NumVertices);
@@ -438,7 +460,10 @@ public:
 		for (int TriangleID : Enumerable)
 		{
 			FIndex3i Tri = Mesh->GetTriangle(TriangleID);
-			FIndex3i TriUV = (UVOverlay != nullptr) ? UVOverlay->GetTriangle(TriangleID) : FIndex3i::Zero();
+			for (int32 k = 0; k < NumTexCoords; ++k)
+			{
+				UVTriangles[k] = (UVOverlays[k] != nullptr) ? UVOverlays[k]->GetTriangle(TriangleID) : FIndex3i::Invalid();
+			}
 			FIndex3i TriNormal = (NormalOverlay != nullptr) ? NormalOverlay->GetTriangle(TriangleID) : FIndex3i::Zero();
 
 			FColor TriColor = ConstantVertexColor;
@@ -466,9 +491,12 @@ public:
 				}
 				RenderBuffers->StaticMeshVertexBuffer.SetVertexTangents(VertIdx, (FVector)TangentX, (FVector)TangentY, (FVector)Normal);
 
-				FVector2f UV = (UVOverlay != nullptr && TriUV[j] != FDynamicMesh3::InvalidID) ?
-					UVOverlay->GetElement(TriUV[j]) : FVector2f::Zero();
-				RenderBuffers->StaticMeshVertexBuffer.SetVertexUV(VertIdx, 0, (FVector2D)UV);
+				for (int32 k = 0; k < NumTexCoords; ++k)
+				{
+					FVector2f UV = (UVTriangles[k][j] != FDynamicMesh3::InvalidID) ?
+						UVOverlays[k]->GetElement(UVTriangles[k][j]) : FVector2f::Zero();
+					RenderBuffers->StaticMeshVertexBuffer.SetVertexUV(VertIdx, k, (FVector2D)UV);
+				}
 
 				RenderBuffers->ColorVertexBuffer.VertexColor(VertIdx) = (bHaveColors) ?
 					(FColor)Mesh->GetVertexColor(Tri[j]) : TriColor;
@@ -492,6 +520,8 @@ public:
 	}
 
 
+
+
 	/**
 	 * Filter the triangles in a FMeshRenderBufferSet into the SecondaryIndexBuffer.
 	 * Requires that RenderBuffers->Triangles has been initialized.
diff --git a/Engine/Plugins/Experimental/MeshModelingToolset/Source/ModelingComponents/Private/SimpleDynamicMeshSceneProxy.h b/Engine/Plugins/Experimental/MeshModelingToolset/Source/ModelingComponents/Private/SimpleDynamicMeshSceneProxy.h
--- a/Engine/Plugins/Experimental/MeshModelingToolset/Source/ModelingComponents/Private/SimpleDynamicMeshSceneProxy.h
+++ b/Engine/Plugins/Experimental/MeshModelingToolset/Source/ModelingComponents/Private/SimpleDynamicMeshSceneProxy.h
@@ -152,15 +152,20 @@ public:
 	 */
 	void InitializeSingleBufferSet(FMeshRenderBufferSet* RenderBuffers)
 	{
-		FDynamicMesh3* Mesh = ParentComponent->GetRenderMesh();
+		const FDynamicMesh3* Mesh = ParentComponent->GetRenderMesh();
 
 		// find suitable overlays
-		FDynamicMeshUVOverlay* UVOverlay = nullptr;
-		FDynamicMeshNormalOverlay* NormalOverlay = nullptr;
+		TArray<const FDynamicMeshUVOverlay*> UVOverlays;
+		const FDynamicMeshNormalOverlay* NormalOverlay = nullptr;
 		if (Mesh->HasAttributes())
 		{
-			UVOverlay = Mesh->Attributes()->PrimaryUV();
-			NormalOverlay = Mesh->Attributes()->PrimaryNormals();
+			const FDynamicMeshAttributeSet* Attributes = Mesh->Attributes();
+			NormalOverlay = Attributes->PrimaryNormals();
+			UVOverlays.SetNum(Attributes->NumUVLayers());
+			for (int32 k = 0; k < UVOverlays.Num(); ++k)
+			{
+				UVOverlays[k] = Attributes->GetUVLayer(k);
+			}
 		}
 
 		TFunction<void(int, int, int, FVector3f&, FVector3f&)> TangentsFunc = nullptr;
@@ -175,7 +180,7 @@ public:
 
 		InitializeBuffersFromOverlays(RenderBuffers, Mesh,
 			Mesh->TriangleCount(), Mesh->TriangleIndicesItr(),
-			UVOverlay, NormalOverlay, TangentsFunc);
+			UVOverlays, NormalOverlay, TangentsFunc);
 
 		ENQUEUE_RENDER_COMMAND(FSimpleDynamicMeshSceneProxyInitializeSingle)(
 			[RenderBuffers](FRHICommandListImmediate& RHICmdList)
@@ -191,13 +196,21 @@ public:
 	 */
 	void InitializeByMaterial(TArray<FMeshRenderBufferSet*>& BufferSets)
 	{
-		FDynamicMesh3* Mesh = ParentComponent->GetRenderMesh();
+		const FDynamicMesh3* Mesh = ParentComponent->GetRenderMesh();
 		check(Mesh->HasAttributes() && Mesh->Attributes()->HasMaterialID());
 
+		const FDynamicMeshAttributeSet* Attributes = Mesh->Attributes();
+
 		// find suitable overlays
-		FDynamicMeshUVOverlay* UVOverlay = Mesh->Attributes()->PrimaryUV();
-		FDynamicMeshNormalOverlay* NormalOverlay = Mesh->Attributes()->PrimaryNormals();
-		FDynamicMeshMaterialAttribute* MaterialID = Mesh->Attributes()->GetMaterialID();
+		const FDynamicMeshMaterialAttribute* MaterialID = Attributes->GetMaterialID();
+		const FDynamicMeshNormalOverlay* NormalOverlay = Mesh->Attributes()->PrimaryNormals();
+
+		TArray<const FDynamicMeshUVOverlay*> UVOverlays;
+		UVOverlays.SetNum(Attributes->NumUVLayers());
+		for (int32 k = 0; k < UVOverlays.Num(); ++k)
+		{
+			UVOverlays[k] = Attributes->GetUVLayer(k);
+		}
 
 		TFunction<void(int, int, int, FVector3f&, FVector3f&)> TangentsFunc = nullptr;
 		const FMeshTangentsf* Tangents = ParentComponent->GetTangents();
@@ -257,7 +270,7 @@ public:
 
 				InitializeBuffersFromOverlays(RenderBuffers, Mesh,
 					Triangles.Num(), Triangles,
-					UVOverlay, NormalOverlay, TangentsFunc);
+					UVOverlays, NormalOverlay, TangentsFunc);
 
 				RenderBuffers->Triangles = Triangles;
 
diff --git a/Engine/Plugins/Experimental/MeshModelingToolset/Source/ModelingOperatorsEditorOnly/Private/ParameterizationOps/UVLayoutOp.cpp b/Engine/Plugins/Experimental/MeshModelingToolset/Source/ModelingOperatorsEditorOnly/Private/ParameterizationOps/UVLayoutOp.cpp
--- a/Engine/Plugins/Experimental/MeshModelingToolset/Source/ModelingOperatorsEditorOnly/Private/ParameterizationOps/UVLayoutOp.cpp
+++ b/Engine/Plugins/Experimental/MeshModelingToolset/Source/ModelingOperatorsEditorOnly/Private/ParameterizationOps/UVLayoutOp.cpp
@@ -34,8 +34,8 @@ void FUVLayoutOp::CalculateResult(FProgressCancel* Progress)
 	}
 
 
-	int UVLayerInput = 0, UVLayerOutput = 0;
-	FDynamicMeshUVOverlay* UVLayer = ResultMesh->Attributes()->GetUVLayer(UVLayerInput);
+	int UVLayerInput = UVLayerIndex, UVLayerOutput = UVLayerIndex;
+	FDynamicMeshUVOverlay* UseUVLayer = ResultMesh->Attributes()->GetUVLayer(UVLayerInput);
 	
 
 	bool bWillRepackIslands = (UVLayoutMode != EUVLayoutOpLayoutModes::TransformOnly);
@@ -43,7 +43,7 @@ void FUVLayoutOp::CalculateResult(FProgressCancel* Progress)
 	// split bowties so that we can process islands independently
 	if (bWillRepackIslands || bAlwaysSplitBowties)
 	{
-		UVLayer->SplitBowties();
+		UseUVLayer->SplitBowties();
 	}
 
 	if (Progress && Progress->Cancelled())
@@ -51,7 +51,7 @@ void FUVLayoutOp::CalculateResult(FProgressCancel* Progress)
 		return;
 	}
 
-	FDynamicMeshUVPacker Packer(UVLayer);
+	FDynamicMeshUVPacker Packer(UseUVLayer);
 	Packer.TextureResolution = this->TextureResolution;
 	Packer.GutterSize = this->GutterSize;
 	Packer.bAllowFlips = this->bAllowFlips;
@@ -80,11 +80,11 @@ void FUVLayoutOp::CalculateResult(FProgressCancel* Progress)
 
 	if (UVScaleFactor != 1.0 || UVTranslation != FVector2f::Zero() )
 	{
-		for (int ElementID : UVLayer->ElementIndicesItr())
+		for (int ElementID : UseUVLayer->ElementIndicesItr())
 		{
-			FVector2f UV = UVLayer->GetElement(ElementID);
+			FVector2f UV = UseUVLayer->GetElement(ElementID);
 			UV = (UV * UVScaleFactor) + UVTranslation;
-			UVLayer->SetElement(ElementID, UV);
+			UseUVLayer->SetElement(ElementID, UV);
 		}
 	}
 
diff --git a/Engine/Plugins/Experimental/MeshModelingToolset/Source/ModelingOperatorsEditorOnly/Public/ParameterizationOps/UVLayoutOp.h b/Engine/Plugins/Experimental/MeshModelingToolset/Source/ModelingOperatorsEditorOnly/Public/ParameterizationOps/UVLayoutOp.h
--- a/Engine/Plugins/Experimental/MeshModelingToolset/Source/ModelingOperatorsEditorOnly/Public/ParameterizationOps/UVLayoutOp.h
+++ b/Engine/Plugins/Experimental/MeshModelingToolset/Source/ModelingOperatorsEditorOnly/Public/ParameterizationOps/UVLayoutOp.h
@@ -25,6 +25,7 @@ public:
 
 	EUVLayoutOpLayoutModes UVLayoutMode = EUVLayoutOpLayoutModes::RepackToUnitRect;
 
+	int UVLayerIndex = 0;
 	int TextureResolution = 128;
 	bool bAllowFlips = false;
 	bool bAlwaysSplitBowties = true;
