ModelingTools: Add support for managing UV channels and Lightmap UVs to AttributeEditorTool. Add Optimize For Editing button.
#rnx

diff --git a/Engine/Plugins/Experimental/MeshModelingToolset/Source/MeshModelingToolsEditorOnly/Private/AttributeEditorTool.cpp b/Engine/Plugins/Experimental/MeshModelingToolset/Source/MeshModelingToolsEditorOnly/Private/AttributeEditorTool.cpp
--- a/Engine/Plugins/Experimental/MeshModelingToolset/Source/MeshModelingToolsEditorOnly/Private/AttributeEditorTool.cpp
+++ b/Engine/Plugins/Experimental/MeshModelingToolset/Source/MeshModelingToolsEditorOnly/Private/AttributeEditorTool.cpp
@@ -7,11 +7,16 @@
 #include "ToolSetupUtil.h"
 #include "MathUtil.h"
 #include "MeshDescriptionToDynamicMesh.h"   // for FMeshDescriptionToDynamicMesh::IsReservedAttributeName
+#include "AssetUtils/MeshDescriptionUtil.h"
 
 #include "MeshDescription.h"
 #include "StaticMeshAttributes.h"
 #include "StaticMeshOperations.h"
 
+// for lightmap access
+#include "Engine/StaticMesh.h"
+#include "Components/StaticMeshComponent.h"
+
 #include "Components/PrimitiveComponent.h"
 
 
@@ -71,6 +76,11 @@ void UAttributeEditorActionPropertySet::PostAction(EAttributeEditorToolActions A
 
 
 
+TArray<FString> UAttributeEditorUVActions::GetUVLayerNamesFunc()
+{
+	return UVLayerNamesList;
+}
+
 
 TArray<FString> UAttributeEditorModifyAttributeActions::GetAttributeNamesFunc()
 {
@@ -93,26 +103,24 @@ void UAttributeEditorTool::Setup()
 {
 	UInteractiveTool::Setup();
 
-	int NumUVLayers = 1;
-	for (TUniquePtr<FPrimitiveComponentTarget>& ComponentTarget : ComponentTargets)
-	{
-		FMeshDescription* MeshDescription = ComponentTarget->GetMesh();
-		TVertexInstanceAttributesRef<FVector2D> InstanceUVs =
-			MeshDescription->VertexInstanceAttributes().GetAttributesRef<FVector2D>(MeshAttribute::VertexInstance::TextureCoordinate);
-		NumUVLayers = FMath::Max(NumUVLayers, InstanceUVs.GetNumChannels());
-	}
+	OptimizeActions = NewObject<UAttributeEditorOptimizeActions>(this);
+	OptimizeActions->Initialize(this);
+	AddToolPropertySource(OptimizeActions);
 
 	NormalsActions = NewObject<UAttributeEditorNormalsActions>(this);
 	NormalsActions->Initialize(this);
 	AddToolPropertySource(NormalsActions);
 
-	UVActions = NewObject<UAttributeEditorUVActions>(this);
-	UVActions->Initialize(this);
-	AddToolPropertySource(UVActions);
-	UVActions->NumUVLayers = NumUVLayers;
-
 	if (ComponentTargets.Num() == 1)
 	{
+		UVActions = NewObject<UAttributeEditorUVActions>(this);
+		UVActions->Initialize(this);
+		AddToolPropertySource(UVActions);
+
+		LightmapUVActions = NewObject<UAttributeEditorLightmapUVActions>(this);
+		LightmapUVActions->Initialize(this);
+		AddToolPropertySource(LightmapUVActions);
+
 		NewAttributeProps = NewObject<UAttributeEditorNewAttributeActions>(this);
 		NewAttributeProps->Initialize(this);
 		AddToolPropertySource(NewAttributeProps);
@@ -272,6 +280,29 @@ void UAttributeEditorTool::InitializeAttributeLists()
 {
 	FMeshDescription* Mesh = ComponentTargets[0]->GetMesh();
 
+
+	TVertexInstanceAttributesRef<FVector2D> InstanceUVs =
+		Mesh->VertexInstanceAttributes().GetAttributesRef<FVector2D>(MeshAttribute::VertexInstance::TextureCoordinate);
+
+	UVActions->UVLayerNamesList.Reset();
+	for ( int32 k = 0; k < InstanceUVs.GetNumChannels(); ++k )
+	{
+		UVActions->UVLayerNamesList.Add(FString::Printf(TEXT("UV%d"), k));
+	}
+	UVActions->UVLayer = UVActions->UVLayerNamesList[0];
+
+	UStaticMeshComponent* StaticMeshComponent = Cast<UStaticMeshComponent>(ComponentTargets[0]->GetOwnerComponent());
+	if (StaticMeshComponent && StaticMeshComponent->GetStaticMesh() != nullptr)
+	{
+		const UStaticMesh* StaticMesh = StaticMeshComponent->GetStaticMesh();
+		const FMeshBuildSettings& BuildSettings = StaticMesh->GetSourceModel(0).BuildSettings;
+		LightmapUVActions->bGenerateLightmapUVs = BuildSettings.bGenerateLightmapUVs;
+		LightmapUVActions->SourceUVIndex = BuildSettings.SrcLightmapIndex;
+		LightmapUVActions->DestinationUVIndex = BuildSettings.DstLightmapIndex;
+
+		bHaveAutoGeneratedLightmapUVSet = (LightmapUVActions->DestinationUVIndex >= InstanceUVs.GetNumChannels());
+	}
+
 	ExtractAttribList(Mesh, Mesh->VertexAttributes(), EAttributeEditorElementType::Vertex, this->VertexAttributes, AttributeProps->VertexAttributes);
 	ExtractAttribList(Mesh, Mesh->VertexInstanceAttributes(), EAttributeEditorElementType::VertexInstance, this->InstanceAttributes, AttributeProps->InstanceAttributes);
 	ExtractAttribList(Mesh, Mesh->TriangleAttributes(), EAttributeEditorElementType::Triangle, this->TriangleAttributes, AttributeProps->TriangleAttributes);
@@ -326,14 +357,23 @@ void UAttributeEditorTool::OnTick(float DeltaTime)
 {
 	switch (PendingAction)
 	{
+	case EAttributeEditorToolActions::OptimizeForEditing:
+		OptimizeForEditing();
+		break;
 	case EAttributeEditorToolActions::ClearNormals:
 		ClearNormals();
 		break;
-	case EAttributeEditorToolActions::ClearSelectedUVs:
-		ClearUVs(true);
-		break;
 	case EAttributeEditorToolActions::ClearAllUVs:
-		ClearUVs(false);
+		ClearUVs();
+		break;
+	case EAttributeEditorToolActions::AddUVSet:
+		AddUVSet();
+		break;
+	case EAttributeEditorToolActions::DeleteSelectedUVSet:
+		DeleteSelectedUVSet();
+		break;
+	case EAttributeEditorToolActions::DuplicateSelectedUVSet:
+		DuplicateSelectedUVSet();
 		break;
 	case EAttributeEditorToolActions::AddAttribute:
 		AddNewAttribute();
@@ -347,6 +387,16 @@ void UAttributeEditorTool::OnTick(float DeltaTime)
 	case EAttributeEditorToolActions::DeleteAttribute:
 		DeleteAttribute();
 		break;
+
+	case EAttributeEditorToolActions::EnableLightmapUVs:
+		SetLightmapUVsEnabled(true);
+		break;
+	case EAttributeEditorToolActions::DisableLightmapUVs:
+		SetLightmapUVsEnabled(false);
+		break;
+	case EAttributeEditorToolActions::ResetLightmapUVChannels:
+		ResetLightmapUVsChannels();
+		break;
 	}
 	PendingAction = EAttributeEditorToolActions::NoAction;
 
@@ -358,6 +408,63 @@ void UAttributeEditorTool::OnTick(float DeltaTime)
 
 
 
+
+void UAttributeEditorTool::OptimizeForEditing()
+{
+	GetToolManager()->BeginUndoTransaction(LOCTEXT("OptimizeForEditing", "Optimize For Editing"));
+
+	for (TUniquePtr<FPrimitiveComponentTarget>& Target : ComponentTargets)
+	{
+		UStaticMeshComponent* StaticMeshComponent = Cast<UStaticMeshComponent>(Target->GetOwnerComponent());
+		if (StaticMeshComponent == nullptr || StaticMeshComponent->GetStaticMesh() == nullptr)
+		{
+			continue;
+		}
+
+		UStaticMesh* StaticMesh = StaticMeshComponent->GetStaticMesh();
+		StaticMesh->Modify();
+		StaticMesh->SetNumSourceModels(1);		// discard extra source models
+		FMeshBuildSettings& BuildSettings = StaticMesh->GetSourceModel(0).BuildSettings;
+		BuildSettings.bGenerateLightmapUVs = false;
+		BuildSettings.bBuildReversedIndexBuffer = false;
+		BuildSettings.bBuildAdjacencyBuffer = false;
+		BuildSettings.bRemoveDegenerates = false;
+
+		// dramatically reduce distance field resolution to speed up editing
+		BuildSettings.DistanceFieldResolutionScale = 0.01;
+
+		// this will call StaticMesh->PostEditChange()...
+		Target->CommitMesh([&](const FPrimitiveComponentTarget::FCommitParams& CommitParams)
+		{
+			TEdgeAttributesRef<bool> EdgeHardnesses = CommitParams.MeshDescription->EdgeAttributes().GetAttributesRef<bool>(MeshAttribute::Edge::IsHard);
+			if (EdgeHardnesses.IsValid())
+			{
+				for (FEdgeID ElID : CommitParams.MeshDescription->Edges().GetElementIDs())
+				{
+					EdgeHardnesses[ElID] = false;
+				}
+			}
+
+			// force computation of normals/tangents if they are auto-generated
+			if (BuildSettings.bRecomputeNormals || BuildSettings.bRecomputeTangents)
+			{
+				UE::MeshDescription::InitializeAutoGeneratedAttributes(*CommitParams.MeshDescription, StaticMeshComponent, 0);
+			}
+
+			// now clear these build settings
+			BuildSettings.bUseMikkTSpace = false;
+			BuildSettings.bRecomputeNormals = false;
+			BuildSettings.bRecomputeTangents = false;
+		});
+
+	}
+	GetToolManager()->EndUndoTransaction();
+
+	// update attrib lists
+	bAttributeListsValid = false;
+}
+
+
 void UAttributeEditorTool::ClearNormals()
 {
 	GetToolManager()->BeginUndoTransaction(LOCTEXT("ClearNormalsTransactionMessage", "Clear Normals"));
@@ -384,45 +491,159 @@ void UAttributeEditorTool::ClearNormals()
 
 
 
-void UAttributeEditorTool::ClearUVs(bool bSelectedOnly)
+void UAttributeEditorTool::ClearUVs()
 {
 	GetToolManager()->BeginUndoTransaction(LOCTEXT("ClearUVsTransactionMessage", "Clear Selected UVs"));
 	for (TUniquePtr<FPrimitiveComponentTarget>& Target : ComponentTargets)
 	{
 		Target->CommitMesh([&](const FPrimitiveComponentTarget::FCommitParams& CommitParams)
 		{
-			bool RemoveLayers[8] =
-			{
-				UVActions->bClearUVLayer0,
-				UVActions->bClearUVLayer1,
-				UVActions->bClearUVLayer2,
-				UVActions->bClearUVLayer3,
-				UVActions->bClearUVLayer4,
-				UVActions->bClearUVLayer5,
-				UVActions->bClearUVLayer6,
-				UVActions->bClearUVLayer7
-			};
-
 			TVertexInstanceAttributesRef<FVector2D> InstanceUVs =
 				CommitParams.MeshDescription->VertexInstanceAttributes().GetAttributesRef<FVector2D>(MeshAttribute::VertexInstance::TextureCoordinate);
+			int32 NumChannels = InstanceUVs.GetNumChannels();
 			const FVertexInstanceArray& Instances = CommitParams.MeshDescription->VertexInstances();
-			for (int LayerIndex = 7; LayerIndex >= 0; LayerIndex--)
+			for (int LayerIndex = NumChannels-1; LayerIndex >= 0; LayerIndex--)
 			{
-				bool bRemove = (bSelectedOnly == false) || RemoveLayers[LayerIndex];
-				if (bRemove && LayerIndex < InstanceUVs.GetNumChannels())
+				if (!FStaticMeshOperations::RemoveUVChannel(*CommitParams.MeshDescription, LayerIndex))
 				{
-					if (!FStaticMeshOperations::RemoveUVChannel(*CommitParams.MeshDescription, LayerIndex))
+					for (const FVertexInstanceID& ElID : Instances.GetElementIDs())
 					{
-						for (const FVertexInstanceID& ElID : Instances.GetElementIDs())
-						{
-							InstanceUVs.Set(ElID, LayerIndex, FVector2D::ZeroVector);
-						}
+						InstanceUVs.Set(ElID, LayerIndex, FVector2D::ZeroVector);
 					}
 				}
 			}
+
+			if (bHaveAutoGeneratedLightmapUVSet)
+			{
+				UpdateAutoGeneratedLightmapUVChannel(Target, InstanceUVs.GetNumChannels());
+			}
 		});
 	}
 	GetToolManager()->EndUndoTransaction();
+
+	// update attrib lists
+	bAttributeListsValid = false;
+}
+
+
+
+
+void UAttributeEditorTool::DeleteSelectedUVSet()
+{
+	int32 DeleteIndex = UVActions->UVLayerNamesList.IndexOfByKey(UVActions->UVLayer);
+	if (DeleteIndex == INDEX_NONE)
+	{
+		GetToolManager()->DisplayMessage(LOCTEXT("CannotFindUVSet", "Selected UV Set Not Found"), EToolMessageLevel::UserWarning);
+		return;
+	}
+	if (DeleteIndex == 0 && UVActions->UVLayerNamesList.Num() == 1)
+	{
+		GetToolManager()->DisplayMessage(LOCTEXT("CannotDeleteLastUVSet", "Cannot Delete Last UV Set. UVs will be cleared to Zero."), EToolMessageLevel::UserWarning);
+	}
+
+	GetToolManager()->BeginUndoTransaction(LOCTEXT("ClearUVsTransactionMessage", "Clear Selected UVs"));
+	for (TUniquePtr<FPrimitiveComponentTarget>& Target : ComponentTargets)
+	{
+		Target->CommitMesh([&](const FPrimitiveComponentTarget::FCommitParams& CommitParams)
+		{
+			TVertexInstanceAttributesRef<FVector2D> InstanceUVs =
+				CommitParams.MeshDescription->VertexInstanceAttributes().GetAttributesRef<FVector2D>(MeshAttribute::VertexInstance::TextureCoordinate);
+			if (!FStaticMeshOperations::RemoveUVChannel(*CommitParams.MeshDescription, DeleteIndex))
+			{
+				const FVertexInstanceArray& Instances = CommitParams.MeshDescription->VertexInstances();
+				for (const FVertexInstanceID& InstanceID : Instances.GetElementIDs())
+				{
+					InstanceUVs.Set(InstanceID, DeleteIndex, FVector2D::ZeroVector);
+				}
+			}
+
+			if (bHaveAutoGeneratedLightmapUVSet)
+			{
+				UpdateAutoGeneratedLightmapUVChannel(Target, InstanceUVs.GetNumChannels());
+			}
+		});
+	}
+	GetToolManager()->EndUndoTransaction();
+
+	bAttributeListsValid = false;
+}
+
+
+void UAttributeEditorTool::AddUVSet()
+{
+	GetToolManager()->BeginUndoTransaction(LOCTEXT("AddUVSetMessage", "Add UV Set"));
+	for (TUniquePtr<FPrimitiveComponentTarget>& Target : ComponentTargets)
+	{
+		Target->CommitMesh([&](const FPrimitiveComponentTarget::FCommitParams& CommitParams)
+		{
+			TVertexInstanceAttributesRef<FVector2D> InstanceUVs =
+				CommitParams.MeshDescription->VertexInstanceAttributes().GetAttributesRef<FVector2D>(MeshAttribute::VertexInstance::TextureCoordinate);
+			int32 NewChannelIndex = InstanceUVs.GetNumChannels();
+			if (!FStaticMeshOperations::AddUVChannel(*CommitParams.MeshDescription))
+			{
+				GetToolManager()->DisplayMessage(LOCTEXT("FailedToAddUVSet", "Adding UV Set Failed"), EToolMessageLevel::UserWarning);
+			}
+			else
+			{
+				GetToolManager()->DisplayMessage(FText::Format(LOCTEXT("AddedNewUVSet", "Added UV{0}"), FText::FromString(FString::FromInt(NewChannelIndex))), EToolMessageLevel::UserWarning);
+
+				if (bHaveAutoGeneratedLightmapUVSet)
+				{
+					UpdateAutoGeneratedLightmapUVChannel(Target, InstanceUVs.GetNumChannels());
+				}
+			}
+		});
+	}
+	GetToolManager()->EndUndoTransaction();
+
+	bAttributeListsValid = false;
+}
+
+
+
+void UAttributeEditorTool::DuplicateSelectedUVSet()
+{
+	int32 SourceIndex = UVActions->UVLayerNamesList.IndexOfByKey(UVActions->UVLayer);
+	if (SourceIndex == INDEX_NONE)
+	{
+		GetToolManager()->DisplayMessage(LOCTEXT("CannotFindUVSet", "Selected UV Set Not Found"), EToolMessageLevel::UserWarning);
+		return;
+	}
+
+	GetToolManager()->BeginUndoTransaction(LOCTEXT("DuplicateUVSetMessage", "Duplicate UV Set"));
+	for (TUniquePtr<FPrimitiveComponentTarget>& Target : ComponentTargets)
+	{
+		Target->CommitMesh([&](const FPrimitiveComponentTarget::FCommitParams& CommitParams)
+		{
+			TVertexInstanceAttributesRef<FVector2D> InstanceUVs =
+				CommitParams.MeshDescription->VertexInstanceAttributes().GetAttributesRef<FVector2D>(MeshAttribute::VertexInstance::TextureCoordinate);
+			int32 NewChannelIndex = InstanceUVs.GetNumChannels();
+			if (!FStaticMeshOperations::AddUVChannel(*CommitParams.MeshDescription))
+			{
+				GetToolManager()->DisplayMessage(LOCTEXT("FailedToAddUVSet", "Adding UV Set Failed"), EToolMessageLevel::UserWarning);
+			}
+			else
+			{
+				const FVertexInstanceArray& Instances = CommitParams.MeshDescription->VertexInstances();
+				for (const FVertexInstanceID& InstanceID : Instances.GetElementIDs())
+				{
+					FVector2D SourceUV = InstanceUVs.Get(InstanceID, SourceIndex);
+					InstanceUVs.Set(InstanceID, NewChannelIndex, SourceUV);
+				}
+
+				if (bHaveAutoGeneratedLightmapUVSet)
+				{
+					UpdateAutoGeneratedLightmapUVChannel(Target, InstanceUVs.GetNumChannels());
+				}
+
+				GetToolManager()->DisplayMessage(FText::Format(LOCTEXT("Copied UV Set", "Copied UV{0} to UV{1}"), 
+					FText::FromString(FString::FromInt(SourceIndex)), FText::FromString(FString::FromInt(NewChannelIndex))), EToolMessageLevel::UserWarning);
+			}
+		});
+	}
+	GetToolManager()->EndUndoTransaction();
+
+	bAttributeListsValid = false;
 }
 
 
@@ -530,6 +751,81 @@ void UAttributeEditorTool::DeleteAttribute()
 
 
 
+void UAttributeEditorTool::SetLightmapUVsEnabled(bool bEnabled)
+{
+	if (bEnabled)
+	{
+		GetToolManager()->BeginUndoTransaction(LOCTEXT("EnableLightmapVUs", "Enable Lightmap UVs"));
+	}
+	else
+	{
+		GetToolManager()->BeginUndoTransaction(LOCTEXT("DisableLightmapUVs", "Disable Lightmap UVs"));
+	}
+	for (TUniquePtr<FPrimitiveComponentTarget>& Target : ComponentTargets)
+	{
+		UStaticMeshComponent* StaticMeshComponent = Cast<UStaticMeshComponent>(Target->GetOwnerComponent());
+		if (StaticMeshComponent && StaticMeshComponent->GetStaticMesh() != nullptr)
+		{
+			UStaticMesh* StaticMesh = StaticMeshComponent->GetStaticMesh();
+			StaticMesh->Modify();
+			FMeshBuildSettings& BuildSettings = StaticMesh->GetSourceModel(0).BuildSettings;
+			BuildSettings.bGenerateLightmapUVs = bEnabled;
+
+			StaticMesh->PostEditChange();
+		}
+	}
+	GetToolManager()->EndUndoTransaction();
+
+	// update attrib lists
+	bAttributeListsValid = false;
+}
+
+
+
+void UAttributeEditorTool::ResetLightmapUVsChannels()
+{
+	GetToolManager()->BeginUndoTransaction(LOCTEXT("ResetLightmapUVs", "Reset Lightmap UVs"));
+	for (TUniquePtr<FPrimitiveComponentTarget>& Target : ComponentTargets)
+	{
+		TVertexInstanceAttributesRef<FVector2D> InstanceUVs = 
+			Target->GetMesh()->VertexInstanceAttributes().GetAttributesRef<FVector2D>(MeshAttribute::VertexInstance::TextureCoordinate);
+		int32 SetChannel = FMath::Max(InstanceUVs.GetNumChannels(), 1);
+
+		UStaticMeshComponent* StaticMeshComponent = Cast<UStaticMeshComponent>(Target->GetOwnerComponent());
+		if (StaticMeshComponent && StaticMeshComponent->GetStaticMesh() != nullptr)
+		{
+			UStaticMesh* StaticMesh = StaticMeshComponent->GetStaticMesh();
+			StaticMesh->Modify();
+			FMeshBuildSettings& BuildSettings = StaticMesh->GetSourceModel(0).BuildSettings;
+			BuildSettings.SrcLightmapIndex = 0;
+			BuildSettings.DstLightmapIndex = SetChannel;
+			StaticMesh->PostEditChange();
+		}
+	}
+
+	GetToolManager()->EndUndoTransaction();
+
+	// update attrib lists
+	bAttributeListsValid = false;
+}
+
+
+
+void UAttributeEditorTool::UpdateAutoGeneratedLightmapUVChannel(TUniquePtr<FPrimitiveComponentTarget>& Target, int32 NewMaxUVChannels)
+{
+	UStaticMeshComponent* StaticMeshComponent = Cast<UStaticMeshComponent>(Target->GetOwnerComponent());
+	if (StaticMeshComponent && StaticMeshComponent->GetStaticMesh() != nullptr)
+	{
+		UStaticMesh* StaticMesh = StaticMeshComponent->GetStaticMesh();
+		StaticMesh->Modify();
+
+		FMeshBuildSettings& BuildSettings = StaticMesh->GetSourceModel(0).BuildSettings;
+		BuildSettings.DstLightmapIndex = NewMaxUVChannels;
+	}
+}
+
+
+
 
 
 
diff --git a/Engine/Plugins/Experimental/MeshModelingToolset/Source/MeshModelingToolsEditorOnly/Public/AttributeEditorTool.h b/Engine/Plugins/Experimental/MeshModelingToolset/Source/MeshModelingToolsEditorOnly/Public/AttributeEditorTool.h
--- a/Engine/Plugins/Experimental/MeshModelingToolset/Source/MeshModelingToolsEditorOnly/Public/AttributeEditorTool.h
+++ b/Engine/Plugins/Experimental/MeshModelingToolset/Source/MeshModelingToolsEditorOnly/Public/AttributeEditorTool.h
@@ -94,15 +94,25 @@ UENUM()
 enum class EAttributeEditorToolActions
 {
 	NoAction,
+
+	OptimizeForEditing,
+
 	ClearNormals,
-	ClearSelectedUVs,
 	ClearAllUVs,
+	ClearSelectedUVs,
+	AddUVSet,
+	DeleteSelectedUVSet,
+	DuplicateSelectedUVSet,
 	AddAttribute,
 	AddWeightMapLayer,
 	AddPolyGroupLayer,
 	DeleteAttribute,
 	ClearAttribute,
-	CopyAttributeFromTo
+	CopyAttributeFromTo,
+
+	EnableLightmapUVs,
+	DisableLightmapUVs,
+	ResetLightmapUVChannels
 };
 
 
@@ -121,6 +131,21 @@ public:
 };
 
 
+UCLASS()
+class MESHMODELINGTOOLSEDITORONLY_API UAttributeEditorOptimizeActions : public UAttributeEditorActionPropertySet
+{
+	GENERATED_BODY()
+public:
+	/** Optimize this Asset for Mesh Editing by disabling expensive Static Mesh Rebuild Settings and reducing Distance Field resolution */
+	UFUNCTION(CallInEditor, Category = Optimize, meta = (DisplayPriority = 1))
+	void OptimizeForEditing()
+	{
+		PostAction(EAttributeEditorToolActions::OptimizeForEditing);
+	}
+};
+
+
+
 UCLASS()
 class MESHMODELINGTOOLSEDITORONLY_API UAttributeEditorNormalsActions : public UAttributeEditorActionPropertySet
 {
@@ -142,51 +167,85 @@ class MESHMODELINGTOOLSEDITORONLY_API UAttributeEditorUVActions : public UAttrib
 	GENERATED_BODY()
 public:
 
-	/** Clear data from UV layer */
-	UPROPERTY(EditAnywhere, Category = UVs, meta = (DisplayName = "Layer 0", HideEditConditionToggle, EditConditionHides, EditCondition = "NumUVLayers > 0"))
-	bool bClearUVLayer0;
-	/** Clear data from UV layer */
-	UPROPERTY(EditAnywhere, Category = UVs, meta = (DisplayName = "Layer 1", HideEditConditionToggle, EditConditionHides, EditCondition = "NumUVLayers > 1"))
-	bool bClearUVLayer1;
-	/** Clear data from UV layer */
-	UPROPERTY(EditAnywhere, Category = UVs, meta = (DisplayName = "Layer 2", HideEditConditionToggle, EditConditionHides, EditCondition = "NumUVLayers > 2"))
-	bool bClearUVLayer2;
-	/** Clear data from UV layer */
-	UPROPERTY(EditAnywhere, Category = UVs, meta = (DisplayName = "Layer 3", HideEditConditionToggle, EditConditionHides, EditCondition = "NumUVLayers > 3"))
-	bool bClearUVLayer3;
-	/** Clear data from UV layer */
-	UPROPERTY(EditAnywhere, Category = UVs, meta = (DisplayName = "Layer 4", HideEditConditionToggle, EditConditionHides, EditCondition = "NumUVLayers > 4"))
-	bool bClearUVLayer4;
-	/** Clear data from UV layer */
-	UPROPERTY(EditAnywhere, Category = UVs, meta = (DisplayName = "Layer 5", HideEditConditionToggle, EditConditionHides, EditCondition = "NumUVLayers > 5"))
-	bool bClearUVLayer5;
-	/** Clear data from UV layer */
-	UPROPERTY(EditAnywhere, Category = UVs, meta = (DisplayName = "Layer 6", HideEditConditionToggle, EditConditionHides, EditCondition = "NumUVLayers > 6"))
-	bool bClearUVLayer6;
-	/** Clear data from UV layer */
-	UPROPERTY(EditAnywhere, Category = UVs, meta = (DisplayName = "Layer 7", HideEditConditionToggle, EditConditionHides, EditCondition = "NumUVLayers > 7"))
-	bool bClearUVLayer7;
+	UPROPERTY(EditAnywhere, Category = UVs, meta = (GetOptions = GetUVLayerNamesFunc))
+	FString UVLayer;
+
+	UFUNCTION()
+	TArray<FString> GetUVLayerNamesFunc();
 
 	UPROPERTY()
-	int NumUVLayers = 0;
+	TArray<FString> UVLayerNamesList;
 
-	/** Clear the selected UV layers, setting all UV values to (0,0) */
+	/** Clear all UV layers, setting all UV values to (0,0) */
 	UFUNCTION(CallInEditor, Category = UVs, meta = (DisplayPriority = 1))
-	void ClearSelectedUVSets()
+	void ClearAll()
 	{
-		PostAction(EAttributeEditorToolActions::ClearSelectedUVs);
+		PostAction(EAttributeEditorToolActions::ClearAllUVs);
 	}
 
-	/** Clear all UV layers, setting all UV values to (0,0) */
+	/** Clear the selected UV layers, setting all UV values to (0,0) */
 	UFUNCTION(CallInEditor, Category = UVs, meta = (DisplayPriority = 2))
-	void ClearAllUVSets()
+	void AddNew()
 	{
-		PostAction(EAttributeEditorToolActions::ClearAllUVs);
+		PostAction(EAttributeEditorToolActions::AddUVSet);
+	}
+
+	/** Clear the selected UV layers, setting all UV values to (0,0) */
+	UFUNCTION(CallInEditor, Category = UVs, meta = (DisplayPriority = 3))
+	void DeleteSelected()
+	{
+		PostAction(EAttributeEditorToolActions::DeleteSelectedUVSet);
+	}
+
+	/** Clear the selected UV layers, setting all UV values to (0,0) */
+	UFUNCTION(CallInEditor, Category = UVs, meta = (DisplayPriority = 4))
+	void DuplicateSelected()
+	{
+		PostAction(EAttributeEditorToolActions::DuplicateSelectedUVSet);
 	}
 };
 
 
 
+
+UCLASS()
+class MESHMODELINGTOOLSEDITORONLY_API UAttributeEditorLightmapUVActions : public UAttributeEditorActionPropertySet
+{
+	GENERATED_BODY()
+public:
+	/** Whether or not Lightmap UVs are enabled in the Static Mesh Build Settings. Use the Static Mesh Editor to change this value. */
+	UPROPERTY(VisibleAnywhere, Category = LightmapUVs)
+	bool bGenerateLightmapUVs;
+
+	/** Source UV channel used to compute Lightmap UVs. Use the Static Mesh Editor to change this value. */
+	UPROPERTY(VisibleAnywhere, Category = LightmapUVs, meta = (DisplayName = "Source Channel"))
+	int32 SourceUVIndex;
+
+	/** Lightmap UVs are stored in this UV Channel. Use the Static Mesh Editor to change this value. */
+	UPROPERTY(VisibleAnywhere, Category = LightmapUVs, meta = (DisplayName = "Dest Channel"))
+	int32 DestinationUVIndex;
+
+	UFUNCTION(CallInEditor, Category = LightmapUVs, meta = (DisplayPriority = 1))
+	void Enable()
+	{
+		PostAction(EAttributeEditorToolActions::EnableLightmapUVs);
+	}
+
+	UFUNCTION(CallInEditor, Category = LightmapUVs, meta = (DisplayPriority = 2))
+	void Disable()
+	{
+		PostAction(EAttributeEditorToolActions::DisableLightmapUVs);
+	}
+
+	/** Reset Lightmap UV channels to Source Channel UV0 and Destination as UVMax+1 */
+	UFUNCTION(CallInEditor, Category = LightmapUVs, meta = (DisplayPriority = 3))
+	void Reset()
+	{
+		PostAction(EAttributeEditorToolActions::ResetLightmapUVChannels);
+	}
+};
+
+
 UCLASS()
 class MESHMODELINGTOOLSEDITORONLY_API UAttributeEditorNewAttributeActions : public UAttributeEditorActionPropertySet
 {
@@ -305,12 +364,18 @@ public:
 
 protected:
 
+	UPROPERTY()
+	UAttributeEditorOptimizeActions* OptimizeActions;
+
 	UPROPERTY()
 	UAttributeEditorNormalsActions* NormalsActions;
 
 	UPROPERTY()
 	UAttributeEditorUVActions* UVActions;
 
+	UPROPERTY()
+	UAttributeEditorLightmapUVActions* LightmapUVActions;
+
 	UPROPERTY()
 	UAttributeEditorAttribProperties* AttributeProps;
 
@@ -336,18 +401,25 @@ protected:
 
 
 	bool bAttributeListsValid = false;
+	bool bHaveAutoGeneratedLightmapUVSet = false;
 	void InitializeAttributeLists();
 
-
+	void UpdateAutoGeneratedLightmapUVChannel(TUniquePtr<FPrimitiveComponentTarget>& Target, int32 NewMaxUVChannels);
 
 	EAttributeEditorToolActions PendingAction = EAttributeEditorToolActions::NoAction;
+	void OptimizeForEditing();
 	void ClearNormals();
-	void ClearUVs(bool bSelectedOnly);
+	void ClearUVs();
+	void DeleteSelectedUVSet();
+	void DuplicateSelectedUVSet();
+	void AddUVSet();
 	void AddNewAttribute();
 	void AddNewWeightMap();
 	void AddNewGroupsLayer();
 	void DeleteAttribute();
 	void ClearAttribute();
+	void SetLightmapUVsEnabled(bool bEnabled);
+	void ResetLightmapUVsChannels();
 
 	void AddNewAttribute(EAttributeEditorElementType ElemType, EAttributeEditorAttribType AttribType, FName AttributeName);
 };
