ModelingTools: handle some async race conditions in VertexSculptTool that could cause crash on very fast repeated undo/redo. Clean up some lambda captures.
#rnx

diff --git a/Engine/Plugins/Experimental/MeshModelingToolset/Source/MeshModelingTools/Private/MeshVertexSculptTool.cpp b/Engine/Plugins/Experimental/MeshModelingToolset/Source/MeshModelingTools/Private/MeshVertexSculptTool.cpp
--- a/Engine/Plugins/Experimental/MeshModelingToolset/Source/MeshModelingTools/Private/MeshVertexSculptTool.cpp
+++ b/Engine/Plugins/Experimental/MeshModelingToolset/Source/MeshModelingTools/Private/MeshVertexSculptTool.cpp
@@ -248,6 +248,8 @@ void UMeshVertexSculptTool::OnPropertyModified(UObject* PropertySet, FProperty*
 
 void UMeshVertexSculptTool::OnBeginStroke(const FRay& WorldRay)
 {
+	WaitForPendingUndoRedo();		// cannot start stroke if there is an outstanding undo/redo update
+
 	UpdateBrushPosition(WorldRay);
 
 	if (SculptProperties->PrimaryBrushType == EMeshVertexSculptBrushType::Plane ||
@@ -510,7 +512,7 @@ TFuture<void> UMeshVertexSculptTool::ApplyStamp()
 	TFuture<void> SaveVertexFuture;
 	if (ActiveVertexChange != nullptr)
 	{
-		SaveVertexFuture = Async(VertexSculptToolAsyncExecTarget, [&]()
+		SaveVertexFuture = Async(VertexSculptToolAsyncExecTarget, [this]()
 		{
 			TRACE_CPUPROFILER_EVENT_SCOPE(VtxSculptTool_SyncMeshWithPositionBuffer_UpdateChange);
 			const int32 NumV = ROIPositionBuffer.Num();
@@ -702,7 +704,7 @@ void UMeshVertexSculptTool::OnTick(float DeltaTime)
 
 		// Append updated ROI to modified region (async). For some reason this is very expensive,
 		// maybe because of TSet? but we have a lot of time to do it.
-		TFuture<void> AccumulateROI = Async(VertexSculptToolAsyncExecTarget, [&]()
+		TFuture<void> AccumulateROI = Async(VertexSculptToolAsyncExecTarget, [this]()
 		{
 			TRACE_CPUPROFILER_EVENT_SCOPE(VtxSculptTool_Tick_AccumROI);
 			for (int32 tid : TriangleROIArray)
@@ -713,16 +715,13 @@ void UMeshVertexSculptTool::OnTick(float DeltaTime)
 
 		// Start precomputing the normals ROI. This is currently the most expensive single thing we do next
 		// to Octree re-insertion, despite it being almost trivial. Why?!?
-		bool bUsingOverlayNormals = false;
-		TFuture<void> NormalsROI = Async(VertexSculptToolAsyncExecTarget, [&]()
+		bool bUsingOverlayNormalsOut = false;
+		TFuture<void> NormalsROI = Async(VertexSculptToolAsyncExecTarget, [Mesh, &bUsingOverlayNormalsOut, this]()
 		{
 			TRACE_CPUPROFILER_EVENT_SCOPE(VtxSculptTool_Tick_NormalsROI);
 
-
-			//TODO WHY DOES NORMALS FLAGS NOT WORK?? 
-
 			//UE::SculptUtil::PrecalculateNormalsROI(Mesh, TriangleROIArray, NormalsROIBuilder, bUsingOverlayNormals, false);
-			UE::SculptUtil::PrecalculateNormalsROI(Mesh, TriangleROIArray, NormalsFlags, bUsingOverlayNormals, false);
+			UE::SculptUtil::PrecalculateNormalsROI(Mesh, TriangleROIArray, NormalsFlags, bUsingOverlayNormalsOut, false);
 		});
 
 		// NOTE: you might try to speculatively do the octree remove here, to save doing it later on Reinsert().
@@ -734,12 +733,10 @@ void UMeshVertexSculptTool::OnTick(float DeltaTime)
 		TFuture<void> UpdateChangeFuture = ApplyStamp();
 
 		// begin octree rebuild calculation
-		StampUpdateOctreeFuture = Async(VertexSculptToolAsyncExecTarget, [&]()
+		StampUpdateOctreeFuture = Async(VertexSculptToolAsyncExecTarget, [this]()
 		{
 			TRACE_CPUPROFILER_EVENT_SCOPE(VtxSculptTool_Tick_OctreeReinsert);
-			static TArray<uint32> TempBuffer;
-			static TArray<bool> TempFlagBuffer;
-			Octree.ReinsertTrianglesParallel(TriangleROIArray, TempBuffer, TempFlagBuffer);
+			Octree.ReinsertTrianglesParallel(TriangleROIArray, OctreeUpdateTempBuffer, OctreeUpdateTempFlagBuffer);
 		});
 		bStampUpdatePending = true;
 		//TFuture<void> OctreeRebuild = Async(VertexSculptToolAsyncExecTarget, [&]()
@@ -761,7 +758,7 @@ void UMeshVertexSculptTool::OnTick(float DeltaTime)
 		{
 			TRACE_CPUPROFILER_EVENT_SCOPE(VtxSculptTool_Tick_RecalcNormals);
 			NormalsROI.Wait();
-			UE::SculptUtil::RecalculateROINormals(Mesh, NormalsFlags, bUsingOverlayNormals);
+			UE::SculptUtil::RecalculateROINormals(Mesh, NormalsFlags, bUsingOverlayNormalsOut);
 			//UE::SculptUtil::RecalculateROINormals(Mesh, NormalsROIBuilder.Indices(), bUsingOverlayNormals);
 		}
 
@@ -839,12 +836,12 @@ void UMeshVertexSculptTool::UpdateBaseMesh(const TSet<int32>* TriangleSet)
 			BaseMesh.SetVertex(Tri.C, SculptMesh->GetVertex(Tri.C));
 			BaseMeshIndexBuffer.Add(tid);
 		}
-		auto UpdateBaseNormals = Async(VertexSculptToolAsyncExecTarget, [&]()
+		auto UpdateBaseNormals = Async(VertexSculptToolAsyncExecTarget, [this]()
 		{
 			TRACE_CPUPROFILER_EVENT_SCOPE(VtxSculptTool_Target_UpdateBaseNormals);
 			FMeshNormals::QuickComputeVertexNormalsForTriangles(BaseMesh, BaseMeshIndexBuffer);
 		});
-		auto ReinsertTriangles = Async(VertexSculptToolAsyncExecTarget, [&]()
+		auto ReinsertTriangles = Async(VertexSculptToolAsyncExecTarget, [TriangleSet, this]()
 		{
 			TRACE_CPUPROFILER_EVENT_SCOPE(VtxSculptTool_Target_Reinsert);
 			BaseMeshSpatial.ReinsertTriangles(*TriangleSet);
@@ -964,41 +961,31 @@ void UMeshVertexSculptTool::OnDynamicMeshComponentChanged(USimpleDynamicMeshComp
 {
 	// have to wait for any outstanding stamp update to finish...
 	WaitForPendingStampUpdate();
+	// wait for previous Undo to finish (possibly never hit because the change records do it?)
+	WaitForPendingUndoRedo();
 
-	// update octree
 	FDynamicMesh3* Mesh = GetSculptMesh();
 
-	// make sure any previous async computations are done, and update the undo ROI
-	if (bUndoUpdatePending)
-	{
-		// we should never hit this anymore, because of pre-change calling WaitForPendingUndoRedo()
-		WaitForPendingUndoRedo();
-
-		// this is not right because now we are going to do extra recomputation, but it's very messy otherwise...
-		MeshIndexUtil::VertexToTriangleOneRing(Mesh, Change->Vertices, AccumulatedTriangleROI);
-	}
-	else
-	{
-		AccumulatedTriangleROI.Reset();
-		MeshIndexUtil::VertexToTriangleOneRing(Mesh, Change->Vertices, AccumulatedTriangleROI);
-	}
+	// figure out the set of modified triangles
+	AccumulatedTriangleROI.Reset();
+	MeshIndexUtil::VertexToTriangleOneRing(Mesh, Change->Vertices, AccumulatedTriangleROI);
 
 	// start the normal recomputation
-	UndoNormalsFuture = Async(VertexSculptToolAsyncExecTarget, [&]()
+	UndoNormalsFuture = Async(VertexSculptToolAsyncExecTarget, [this, Mesh]()
 	{
 		UE::SculptUtil::RecalculateROINormals(Mesh, AccumulatedTriangleROI, NormalsROIBuilder);
 		return true;
 	});
 
 	// start the octree update
-	UndoUpdateOctreeFuture = Async(VertexSculptToolAsyncExecTarget, [&]()
+	UndoUpdateOctreeFuture = Async(VertexSculptToolAsyncExecTarget, [this, Mesh]()
 	{
 		Octree.ReinsertTriangles(AccumulatedTriangleROI);
 		return true;
 	});
 
 	// start the base mesh update
-	UndoUpdateBaseMeshFuture = Async(VertexSculptToolAsyncExecTarget, [&]()
+	UndoUpdateBaseMeshFuture = Async(VertexSculptToolAsyncExecTarget, [this, Mesh]()
 	{
 		UpdateBaseMesh(&AccumulatedTriangleROI);
 		return true;
diff --git a/Engine/Plugins/Experimental/MeshModelingToolset/Source/MeshModelingTools/Public/MeshVertexSculptTool.h b/Engine/Plugins/Experimental/MeshModelingToolset/Source/MeshModelingTools/Public/MeshVertexSculptTool.h
--- a/Engine/Plugins/Experimental/MeshModelingToolset/Source/MeshModelingTools/Public/MeshVertexSculptTool.h
+++ b/Engine/Plugins/Experimental/MeshModelingToolset/Source/MeshModelingTools/Public/MeshVertexSculptTool.h
@@ -228,6 +228,8 @@ protected:
 	TArray<int> NormalsBuffer;
 	void WaitForPendingUndoRedo();
 
+	TArray<uint32> OctreeUpdateTempBuffer;
+	TArray<bool> OctreeUpdateTempFlagBuffer;
 	TFuture<void> StampUpdateOctreeFuture;
 	bool bStampUpdatePending = false;
 	void WaitForPendingStampUpdate();
diff --git a/Engine/Plugins/Experimental/MeshModelingToolset/Source/ModelingComponents/Private/SimpleDynamicMeshComponent.cpp b/Engine/Plugins/Experimental/MeshModelingToolset/Source/ModelingComponents/Private/SimpleDynamicMeshComponent.cpp
--- a/Engine/Plugins/Experimental/MeshModelingToolset/Source/ModelingComponents/Private/SimpleDynamicMeshComponent.cpp
+++ b/Engine/Plugins/Experimental/MeshModelingToolset/Source/ModelingComponents/Private/SimpleDynamicMeshComponent.cpp
@@ -329,7 +329,7 @@ void USimpleDynamicMeshComponent::FastNotifyPositionsUpdated(bool bNormals, bool
 	{
 		// calculate bounds while we are updating vertices
 		TFuture<void> UpdateBoundsCalc;
-		UpdateBoundsCalc = Async(SimpleDynamicMeshComponentAsyncExecTarget, [&]()
+		UpdateBoundsCalc = Async(SimpleDynamicMeshComponentAsyncExecTarget, [this]()
 		{
 			TRACE_CPUPROFILER_EVENT_SCOPE(SimpleDynamicMeshComponent_FastPositionsUpdate_AsyncBoundsUpdate);
 			LocalBounds = Mesh->GetBounds(true);
@@ -391,7 +391,7 @@ void USimpleDynamicMeshComponent::FastNotifyVertexAttributesUpdated(EMeshRenderA
 		TFuture<void> UpdateBoundsCalc;
 		if (bPositions)
 		{
-			UpdateBoundsCalc = Async(SimpleDynamicMeshComponentAsyncExecTarget, [&]()
+			UpdateBoundsCalc = Async(SimpleDynamicMeshComponentAsyncExecTarget, [this]()
 			{
 				TRACE_CPUPROFILER_EVENT_SCOPE(SimpleDynamicMeshComponent_FastVertexAttribUpdate_AsyncBoundsUpdate);
 				LocalBounds = Mesh->GetBounds(true);
@@ -491,7 +491,7 @@ void USimpleDynamicMeshComponent::FastNotifyTriangleVerticesUpdated(const TArray
 		TFuture<void> UpdateBoundsCalc;
 		if (bPositions)
 		{
-			UpdateBoundsCalc = Async(SimpleDynamicMeshComponentAsyncExecTarget, [&]()
+			UpdateBoundsCalc = Async(SimpleDynamicMeshComponentAsyncExecTarget, [this]()
 			{
 				TRACE_CPUPROFILER_EVENT_SCOPE(SimpleDynamicMeshComponent_FastVertexUpdate_AsyncBoundsUpdate);
 				LocalBounds = Mesh->GetBounds(true);
@@ -570,7 +570,7 @@ void USimpleDynamicMeshComponent::FastNotifyTriangleVerticesUpdated(const TSet<i
 		TFuture<void> UpdateBoundsCalc;
 		if (bPositions)
 		{
-			UpdateBoundsCalc = Async(SimpleDynamicMeshComponentAsyncExecTarget, [&]()
+			UpdateBoundsCalc = Async(SimpleDynamicMeshComponentAsyncExecTarget, [this]()
 			{
 				TRACE_CPUPROFILER_EVENT_SCOPE(SimpleDynamicMeshComponent_FastVertexUpdate_AsyncBoundsUpdate);
 				LocalBounds = Mesh->GetBounds(true);
@@ -652,10 +652,10 @@ TFuture<bool> USimpleDynamicMeshComponent::FastNotifyTriangleVerticesUpdated_Try
 	if ((!!RenderMeshPostProcessor) || (GetCurrentSceneProxy() == nullptr) || (!Decomposition))
 	{
 		// is there a simpler way to do this? cannot seem to just make a TFuture<bool>...
-		return Async(SimpleDynamicMeshComponentAsyncExecTarget, [&]() { return false; });
+		return Async(SimpleDynamicMeshComponentAsyncExecTarget, []() { return false; });
 	}
 
-	return Async(SimpleDynamicMeshComponentAsyncExecTarget, [&]() 
+	return Async(SimpleDynamicMeshComponentAsyncExecTarget, [this, &Triangles, &UpdateSetsOut, &BoundsOut]()
 	{
 		TFuture<void> ComputeBounds = Async(SimpleDynamicMeshComponentAsyncExecTarget, [this, &BoundsOut, &Triangles]()
 		{
