ModelingComponents: always create at least one texcoord in render buffers, even if there are no UV overlays, otherwise checks are hit in rendering setup
#rnx

diff --git a/Engine/Plugins/Experimental/MeshModelingToolset/Source/ModelingComponents/Private/BaseDynamicMeshSceneProxy.h b/Engine/Plugins/Experimental/MeshModelingToolset/Source/ModelingComponents/Private/BaseDynamicMeshSceneProxy.h
--- a/Engine/Plugins/Experimental/MeshModelingToolset/Source/ModelingComponents/Private/BaseDynamicMeshSceneProxy.h
+++ b/Engine/Plugins/Experimental/MeshModelingToolset/Source/ModelingComponents/Private/BaseDynamicMeshSceneProxy.h
@@ -437,13 +437,14 @@ public:
 		bool bHaveColors = Mesh->HasVertexColors() && (bIgnoreVertexColors == false);
 
 		int NumVertices = NumTriangles * 3;
-		int NumTexCoords = UVOverlays.Num();
+		int NumUVOverlays = UVOverlays.Num();
+		int NumTexCoords = FMath::Max(1, NumUVOverlays);		// must have at least one tex coord
 		TArray<FIndex3i, TFixedAllocator<MAX_STATIC_TEXCOORDS>> UVTriangles;
 		UVTriangles.SetNum(NumTexCoords);
 
 		{
 			RenderBuffers->PositionVertexBuffer.Init(NumVertices);
-			RenderBuffers->StaticMeshVertexBuffer.Init(NumVertices, NumTexCoords);
+			RenderBuffers->StaticMeshVertexBuffer.Init(NumVertices, NumTexCoords );
 			RenderBuffers->ColorVertexBuffer.Init(NumVertices);
 			RenderBuffers->IndexBuffer.Indices.AddUninitialized(NumTriangles * 3);
 		}
@@ -462,7 +463,7 @@ public:
 			FIndex3i Tri = Mesh->GetTriangle(TriangleID);
 			for (int32 k = 0; k < NumTexCoords; ++k)
 			{
-				UVTriangles[k] = (UVOverlays[k] != nullptr) ? UVOverlays[k]->GetTriangle(TriangleID) : FIndex3i::Invalid();
+				UVTriangles[k] = (k < NumUVOverlays && UVOverlays[k] != nullptr) ? UVOverlays[k]->GetTriangle(TriangleID) : FIndex3i::Invalid();
 			}
 			FIndex3i TriNormal = (NormalOverlay != nullptr) ? NormalOverlay->GetTriangle(TriangleID) : FIndex3i::Zero();
 
