reset + hide transform of object that will be replaced for create-from-selected type tools (mesh boolean, self union, vox merge/blend/wrap, etc), for more predictable behavior on tool accept

#rnx

diff --git a/Engine/Plugins/Experimental/MeshModelingToolset/Source/MeshModelingTools/Private/CSGMeshesTool.cpp b/Engine/Plugins/Experimental/MeshModelingToolset/Source/MeshModelingTools/Private/CSGMeshesTool.cpp
--- a/Engine/Plugins/Experimental/MeshModelingToolset/Source/MeshModelingTools/Private/CSGMeshesTool.cpp
+++ b/Engine/Plugins/Experimental/MeshModelingToolset/Source/MeshModelingTools/Private/CSGMeshesTool.cpp
@@ -193,9 +193,7 @@ TUniquePtr<FDynamicMeshOperator> UCSGMeshesTool::MakeNewOperator()
 	for (int Idx = 0; Idx < 2; Idx++)
 	{
 		BooleanOp->Meshes[Idx] = OriginalDynamicMeshes[Idx];
-		FTransform UseTransform = TransformProxies[Idx]->GetTransform();
-		UseTransform.MultiplyScale3D(TransformInitialScales[Idx]);
-		BooleanOp->Transforms[Idx] = UseTransform;
+		BooleanOp->Transforms[Idx] = TransformProxies[Idx]->GetTransform();
 	}
 
 	return BooleanOp;
diff --git a/Engine/Plugins/Experimental/MeshModelingToolset/Source/MeshModelingTools/Private/SelfUnionMeshesTool.cpp b/Engine/Plugins/Experimental/MeshModelingToolset/Source/MeshModelingTools/Private/SelfUnionMeshesTool.cpp
--- a/Engine/Plugins/Experimental/MeshModelingToolset/Source/MeshModelingTools/Private/SelfUnionMeshesTool.cpp
+++ b/Engine/Plugins/Experimental/MeshModelingToolset/Source/MeshModelingTools/Private/SelfUnionMeshesTool.cpp
@@ -109,9 +109,7 @@ void USelfUnionMeshesTool::ConvertInputsAndSetPreviewMaterials(bool bSetPreviewM
 			MaterialIDs->SetValue(TID, MaterialRemap[ComponentIdx][MaterialIDs->GetValue(TID)]);
 		}
 		// TODO: center the meshes
-		FTransform UseTransform = TransformProxies[ComponentIdx]->GetTransform();
-		UseTransform.MultiplyScale3D(TransformInitialScales[ComponentIdx]);
-		FTransform3d WorldTransform = (FTransform3d)UseTransform;
+		FTransform3d WorldTransform = (FTransform3d)TransformProxies[ComponentIdx]->GetTransform();
 		if (WorldTransform.GetDeterminant() < 0)
 		{
 			ComponentMesh.ReverseOrientation(false);
diff --git a/Engine/Plugins/Experimental/MeshModelingToolset/Source/MeshModelingTools/Private/VoxelBlendMeshesTool.cpp b/Engine/Plugins/Experimental/MeshModelingToolset/Source/MeshModelingTools/Private/VoxelBlendMeshesTool.cpp
--- a/Engine/Plugins/Experimental/MeshModelingToolset/Source/MeshModelingTools/Private/VoxelBlendMeshesTool.cpp
+++ b/Engine/Plugins/Experimental/MeshModelingToolset/Source/MeshModelingTools/Private/VoxelBlendMeshesTool.cpp
@@ -62,9 +62,7 @@ TUniquePtr<FDynamicMeshOperator> UVoxelBlendMeshesTool::MakeNewOperator()
 	for (int Idx = 0; Idx < ComponentTargets.Num(); Idx++)
 	{
 		Op->Meshes[Idx] = OriginalDynamicMeshes[Idx];
-		FTransform UseTransform = TransformProxies[Idx]->GetTransform();
-		UseTransform.MultiplyScale3D(TransformInitialScales[Idx]);
-		Op->Transforms[Idx] = UseTransform;
+		Op->Transforms[Idx] = TransformProxies[Idx]->GetTransform();
 	}
 
 	Op->BlendFalloff = BlendProperties->BlendFalloff;
diff --git a/Engine/Plugins/Experimental/MeshModelingToolset/Source/MeshModelingTools/Private/VoxelMorphologyMeshesTool.cpp b/Engine/Plugins/Experimental/MeshModelingToolset/Source/MeshModelingTools/Private/VoxelMorphologyMeshesTool.cpp
--- a/Engine/Plugins/Experimental/MeshModelingToolset/Source/MeshModelingTools/Private/VoxelMorphologyMeshesTool.cpp
+++ b/Engine/Plugins/Experimental/MeshModelingToolset/Source/MeshModelingTools/Private/VoxelMorphologyMeshesTool.cpp
@@ -62,9 +62,7 @@ TUniquePtr<FDynamicMeshOperator> UVoxelMorphologyMeshesTool::MakeNewOperator()
 	for (int Idx = 0; Idx < ComponentTargets.Num(); Idx++)
 	{
 		Op->Meshes[Idx] = OriginalDynamicMeshes[Idx];
-		FTransform UseTransform = TransformProxies[Idx]->GetTransform();
-		UseTransform.MultiplyScale3D(TransformInitialScales[Idx]);
-		Op->Transforms[Idx] = UseTransform;
+		Op->Transforms[Idx] = TransformProxies[Idx]->GetTransform();
 	}
 
 	VoxProperties->SetPropertiesOnOp(*Op);
diff --git a/Engine/Plugins/Experimental/MeshModelingToolset/Source/MeshModelingTools/Private/VoxelSolidifyMeshesTool.cpp b/Engine/Plugins/Experimental/MeshModelingToolset/Source/MeshModelingTools/Private/VoxelSolidifyMeshesTool.cpp
--- a/Engine/Plugins/Experimental/MeshModelingToolset/Source/MeshModelingTools/Private/VoxelSolidifyMeshesTool.cpp
+++ b/Engine/Plugins/Experimental/MeshModelingToolset/Source/MeshModelingTools/Private/VoxelSolidifyMeshesTool.cpp
@@ -57,9 +57,7 @@ TUniquePtr<FDynamicMeshOperator> UVoxelSolidifyMeshesTool::MakeNewOperator()
 	for (int Idx = 0; Idx < ComponentTargets.Num(); Idx++)
 	{
 		Op->Meshes[Idx] = OriginalDynamicMeshes[Idx];
-		FTransform UseTransform = TransformProxies[Idx]->GetTransform();
-		UseTransform.MultiplyScale3D(TransformInitialScales[Idx]);
-		Op->Transforms[Idx] = UseTransform;
+		Op->Transforms[Idx] = TransformProxies[Idx]->GetTransform();
 	}
 
 	Op->bSolidAtBoundaries = SolidifyProperties->bSolidAtBoundaries;
diff --git a/Engine/Plugins/Experimental/MeshModelingToolset/Source/ModelingComponents/Private/BaseTools/BaseCreateFromSelectedTool.cpp b/Engine/Plugins/Experimental/MeshModelingToolset/Source/ModelingComponents/Private/BaseTools/BaseCreateFromSelectedTool.cpp
--- a/Engine/Plugins/Experimental/MeshModelingToolset/Source/ModelingComponents/Private/BaseTools/BaseCreateFromSelectedTool.cpp
+++ b/Engine/Plugins/Experimental/MeshModelingToolset/Source/ModelingComponents/Private/BaseTools/BaseCreateFromSelectedTool.cpp
@@ -100,11 +100,18 @@ void UBaseCreateFromSelectedTool::Setup()
 		if (NewType == EBaseCreateFromSelectedTargetType::NewAsset)
 		{
 			HandleSourcesProperties->OutputAsset = TEXT("");
+			HideGizmoIndex = -1;
+			UpdateGizmoVisibility();
 		}
 		else
 		{
 			int32 Index = (HandleSourcesProperties->WriteOutputTo == EBaseCreateFromSelectedTargetType::FirstInputAsset) ? 0 : ComponentTargets.Num() - 1;
 			HandleSourcesProperties->OutputAsset = AssetGenerationUtil::GetComponentAssetBaseName(ComponentTargets[Index]->GetOwnerComponent(), false);
+			HideGizmoIndex = Index;
+			// Reset the hidden gizmo to its initial position
+			FTransform ComponentTransform = ComponentTargets[Index]->GetWorldTransform();
+			TransformGizmos[Index]->SetNewGizmoTransform(ComponentTransform, true);
+			UpdateGizmoVisibility();
 		}
 	});
 
@@ -138,27 +145,25 @@ void UBaseCreateFromSelectedTool::OnTick(float DeltaTime)
 
 void UBaseCreateFromSelectedTool::UpdateGizmoVisibility()
 {
-	for (UTransformGizmo* Gizmo : TransformGizmos)
+	for (int32 GizmoIndex = 0; GizmoIndex < TransformGizmos.Num(); GizmoIndex++)
 	{
-		Gizmo->SetVisibility(TransformProperties->bShowTransformUI);
+		UTransformGizmo* Gizmo = TransformGizmos[GizmoIndex];
+		Gizmo->SetVisibility(TransformProperties->bShowTransformUI && GizmoIndex != HideGizmoIndex);
 	}
 }
 
 
+
 void UBaseCreateFromSelectedTool::SetTransformGizmos()
 {
 	UInteractiveGizmoManager* GizmoManager = GetToolManager()->GetPairedGizmoManager();
 
 	for (int ComponentIdx = 0; ComponentIdx < ComponentTargets.Num(); ComponentIdx++)
 	{
-		TUniquePtr<FPrimitiveComponentTarget>& Target = ComponentTargets[ComponentIdx];
 		UTransformProxy* Proxy = TransformProxies.Add_GetRef(NewObject<UTransformProxy>(this));
 		UTransformGizmo* Gizmo = TransformGizmos.Add_GetRef(GizmoManager->Create3AxisTransformGizmo(this));
+		Proxy->SetTransform(ComponentTargets[ComponentIdx]->GetWorldTransform());
 		Gizmo->SetActiveTarget(Proxy, GetToolManager());
-		FTransform InitialTransform = Target->GetWorldTransform();
-		TransformInitialScales.Add(InitialTransform.GetScale3D());
-		InitialTransform.SetScale3D(FVector::OneVector);
-		Gizmo->ReinitializeGizmoTransform(InitialTransform);
 		Proxy->OnTransformChanged.AddUObject(this, &UBaseCreateFromSelectedTool::TransformChanged);
 	}
 	UpdateGizmoVisibility();
@@ -230,35 +235,20 @@ void UBaseCreateFromSelectedTool::UpdateAsset(const FDynamicMeshOpResult& Result
 	check(Result.Mesh.Get() != nullptr);
 
 	FTransform3d TargetToWorld = (FTransform3d)UpdateTarget->GetWorldTransform();
-	FTransform3d WorldToTarget = TargetToWorld.Inverse();
 
-	FTransform3d NewTransform;
-	if (ComponentTargets.Num() == 1) // in the single-selection case, shove the result back into the original component space
-	{
-		MeshTransforms::ApplyTransform(*Result.Mesh, WorldToTarget);
-		UpdateTarget->CommitMesh([&](const FPrimitiveComponentTarget::FCommitParams& CommitParams)
-		{
-			FDynamicMeshToMeshDescription Converter;
-			Converter.Convert(Result.Mesh.Get(), *CommitParams.MeshDescription);
-		});
-	}
-	else // in the multi-selection case, center the pivot for the combined result
-	{
-		FTransform3d ResultTransform = Result.Transform;
-		MeshTransforms::ApplyTransform(*Result.Mesh, ResultTransform);
-		MeshTransforms::ApplyTransform(*Result.Mesh, WorldToTarget);
+	FTransform3d ResultTransform = Result.Transform;
+	MeshTransforms::ApplyTransform(*Result.Mesh, ResultTransform);
+	MeshTransforms::ApplyTransformInverse(*Result.Mesh, TargetToWorld);
 
-		UpdateTarget->CommitMesh([&](const FPrimitiveComponentTarget::FCommitParams& CommitParams)
-		{
-			FDynamicMeshToMeshDescription Converter;
-			Converter.Convert(Result.Mesh.Get(), *CommitParams.MeshDescription);
-		});
-	}
+	UpdateTarget->CommitMesh([&](const FPrimitiveComponentTarget::FCommitParams& CommitParams)
+	{
+		FDynamicMeshToMeshDescription Converter;
+		Converter.Convert(Result.Mesh.Get(), *CommitParams.MeshDescription);
+	});
 
 	FComponentMaterialSet MaterialSet;
 	MaterialSet.Materials = GetOutputMaterials();
 	UpdateTarget->CommitMaterialSetUpdate(MaterialSet, true);
-
 }
 
 
diff --git a/Engine/Plugins/Experimental/MeshModelingToolset/Source/ModelingComponents/Public/BaseTools/BaseCreateFromSelectedTool.h b/Engine/Plugins/Experimental/MeshModelingToolset/Source/ModelingComponents/Public/BaseTools/BaseCreateFromSelectedTool.h
--- a/Engine/Plugins/Experimental/MeshModelingToolset/Source/ModelingComponents/Public/BaseTools/BaseCreateFromSelectedTool.h
+++ b/Engine/Plugins/Experimental/MeshModelingToolset/Source/ModelingComponents/Public/BaseTools/BaseCreateFromSelectedTool.h
@@ -198,12 +198,10 @@ protected:
 	UPROPERTY()
 	TArray<UTransformGizmo*> TransformGizmos;
 
-	UPROPERTY()
-	TArray<FVector> TransformInitialScales;
-
 
 	UWorld* TargetWorld = nullptr;
 	IAssetGenerationAPI* AssetAPI = nullptr;
 
+	int32 HideGizmoIndex = -1; // Override UI setting and always hide this gizmo
 };
 
diff --git a/Engine/Source/Runtime/Experimental/InteractiveToolsFramework/Private/BaseGizmos/TransformGizmo.cpp b/Engine/Source/Runtime/Experimental/InteractiveToolsFramework/Private/BaseGizmos/TransformGizmo.cpp
--- a/Engine/Source/Runtime/Experimental/InteractiveToolsFramework/Private/BaseGizmos/TransformGizmo.cpp
+++ b/Engine/Source/Runtime/Experimental/InteractiveToolsFramework/Private/BaseGizmos/TransformGizmo.cpp
@@ -577,7 +577,7 @@ void UTransformGizmo::SetActiveTarget(UTransformProxy* Target, IToolContextTrans
 	}
 }
 
-void UTransformGizmo::ReinitializeGizmoTransform(const FTransform& NewTransform)
+void UTransformGizmo::ReinitializeGizmoTransform(const FTransform& NewTransform, bool bKeepGizmoUnscaled)
 {
 	// To update the gizmo location without triggering any callbacks, we temporarily
 	// store a copy of the callback list, detach them, reposition, and then reattach
@@ -585,7 +585,12 @@ void UTransformGizmo::ReinitializeGizmoTransform(const FTransform& NewTransform)
 	USceneComponent* GizmoComponent = GizmoActor->GetRootComponent();
 	auto temp = GizmoComponent->TransformUpdated;
 	GizmoComponent->TransformUpdated.Clear();
-	GizmoComponent->SetWorldTransform(NewTransform);
+	FTransform GizmoTransform = NewTransform;
+	if (bKeepGizmoUnscaled)
+	{
+		GizmoTransform.SetScale3D(FVector(1, 1, 1));
+	}
+	GizmoComponent->SetWorldTransform(GizmoTransform);
 	GizmoComponent->TransformUpdated = temp;
 
 	// The underlying proxy has an existing way to reinitialize its transform without callbacks.
@@ -595,21 +600,25 @@ void UTransformGizmo::ReinitializeGizmoTransform(const FTransform& NewTransform)
 	ActiveTarget->bSetPivotMode = bSavedSetPivotMode;
 }
 
-void UTransformGizmo::SetNewGizmoTransform(const FTransform& NewTransform)
+void UTransformGizmo::SetNewGizmoTransform(const FTransform& NewTransform, bool bKeepGizmoUnscaled)
 {
 	check(ActiveTarget != nullptr);
 
 	StateTarget->BeginUpdate();
 
 	USceneComponent* GizmoComponent = GizmoActor->GetRootComponent();
-	GizmoComponent->SetWorldTransform(NewTransform);
+	FTransform GizmoTransform = NewTransform;
+	if (bKeepGizmoUnscaled)
+	{
+		GizmoTransform.SetScale3D(FVector(1, 1, 1));
+	}
+	GizmoComponent->SetWorldTransform(GizmoTransform);
 	ActiveTarget->SetTransform(NewTransform);
 
 	StateTarget->EndUpdate();
 }
 
 
-// TODO: This should either be named to "SetScale" or removed, since it can be done with ReinitializeGizmoTransform
 void UTransformGizmo::SetNewChildScale(const FVector& NewChildScale)
 {
 	FTransform NewTransform = ActiveTarget->GetTransform();
diff --git a/Engine/Source/Runtime/Experimental/InteractiveToolsFramework/Public/BaseGizmos/TransformGizmo.h b/Engine/Source/Runtime/Experimental/InteractiveToolsFramework/Public/BaseGizmos/TransformGizmo.h
--- a/Engine/Source/Runtime/Experimental/InteractiveToolsFramework/Public/BaseGizmos/TransformGizmo.h
+++ b/Engine/Source/Runtime/Experimental/InteractiveToolsFramework/Public/BaseGizmos/TransformGizmo.h
@@ -309,14 +309,16 @@ public:
 	 * Repositions the gizmo without issuing undo/redo changes, triggering callbacks, 
 	 * or moving any components. Useful for resetting the gizmo to a new location without
 	 * it being viewed as a gizmo manipulation.
+	 * @param bKeepGizmoUnscaled If true, the scale component of NewTransform is passed through to the target but gizmo scale is set to 1
 	 */
-	void ReinitializeGizmoTransform(const FTransform& NewTransform);
+	void ReinitializeGizmoTransform(const FTransform& NewTransform, bool bKeepGizmoUnscaled = true);
 
 	/**
 	 * Set a new position for the Gizmo. This is done via the same mechanisms as the sub-gizmos,
 	 * so it generates the same Change/Modify() events, and hence works with Undo/Redo
+	 * @param bKeepGizmoUnscaled If true, the scale component of NewTransform is passed through to the target but gizmo scale is set to 1
 	 */
-	virtual void SetNewGizmoTransform(const FTransform& NewTransform);
+	virtual void SetNewGizmoTransform(const FTransform& NewTransform, bool bKeepGizmoUnscaled = true);
 
 	/**
 	 * Explicitly set the child scale. Mainly useful to "reset" the child scale to (1,1,1) when re-using Gizmo across multiple transform actions.
