ModelingTools: make seam edges rendering toggeable in Simplify and Remesh tools, disable by default as it is slow on large meshes
#rnx

diff --git a/Engine/Plugins/Experimental/MeshModelingToolset/Source/MeshModelingTools/Private/RemeshMeshTool.cpp b/Engine/Plugins/Experimental/MeshModelingToolset/Source/MeshModelingTools/Private/RemeshMeshTool.cpp
--- a/Engine/Plugins/Experimental/MeshModelingToolset/Source/MeshModelingTools/Private/RemeshMeshTool.cpp
+++ b/Engine/Plugins/Experimental/MeshModelingToolset/Source/MeshModelingTools/Private/RemeshMeshTool.cpp
@@ -224,20 +224,23 @@ void URemeshMeshTool::Render(IToolsContextRenderAPI* RenderAPI)
 	FPrimitiveDrawInterface* PDI = RenderAPI->GetPrimitiveDrawInterface();
 	FTransform Transform = ComponentTargets[0]->GetWorldTransform();
 
-	FColor LineColor(255, 0, 0);
-	const FDynamicMesh3* TargetMesh = Preview->PreviewMesh->GetPreviewDynamicMesh();
-	if (TargetMesh && TargetMesh->HasAttributes())
+	if (BasicProperties->bShowConstraintEdges)
 	{
-		float PDIScale = RenderAPI->GetCameraState().GetPDIScalingFactor();
-		const FDynamicMeshUVOverlay* UVOverlay = TargetMesh->Attributes()->PrimaryUV();
-		for (int eid : TargetMesh->EdgeIndicesItr())
+		FColor LineColor(255, 0, 0);
+		const FDynamicMesh3* TargetMesh = Preview->PreviewMesh->GetPreviewDynamicMesh();
+		if (TargetMesh && TargetMesh->HasAttributes())
 		{
-			if (UVOverlay->IsSeamEdge(eid))
+			float PDIScale = RenderAPI->GetCameraState().GetPDIScalingFactor();
+			const FDynamicMeshUVOverlay* UVOverlay = TargetMesh->Attributes()->PrimaryUV();
+			for (int eid : TargetMesh->EdgeIndicesItr())
 			{
-				FVector3d A, B;
-				TargetMesh->GetEdgeV(eid, A, B);
-				PDI->DrawLine(Transform.TransformPosition((FVector)A), Transform.TransformPosition((FVector)B),
-					LineColor, 0, 2.0*PDIScale, 1.0f, true);
+				if (UVOverlay->IsSeamEdge(eid))
+				{
+					FVector3d A, B;
+					TargetMesh->GetEdgeV(eid, A, B);
+					PDI->DrawLine(Transform.TransformPosition((FVector)A), Transform.TransformPosition((FVector)B),
+						LineColor, 0, 2.0 * PDIScale, 1.0f, true);
+				}
 			}
 		}
 	}
diff --git a/Engine/Plugins/Experimental/MeshModelingToolset/Source/MeshModelingTools/Public/RemeshMeshTool.h b/Engine/Plugins/Experimental/MeshModelingToolset/Source/MeshModelingTools/Public/RemeshMeshTool.h
--- a/Engine/Plugins/Experimental/MeshModelingToolset/Source/MeshModelingTools/Public/RemeshMeshTool.h
+++ b/Engine/Plugins/Experimental/MeshModelingToolset/Source/MeshModelingTools/Public/RemeshMeshTool.h
@@ -66,6 +66,11 @@ public:
 	UPROPERTY(EditAnywhere, Category = Display)
 	bool bShowGroupColors = false;
 
+	/** Display constrained edges */
+	UPROPERTY(EditAnywhere, Category = Display)
+	bool bShowConstraintEdges = false;
+
+
 	/** Remeshing type */
 	UPROPERTY(EditAnywhere, Category = Remeshing, AdvancedDisplay)
 	ERemeshType RemeshType;
diff --git a/Engine/Plugins/Experimental/MeshModelingToolset/Source/MeshModelingToolsEditorOnly/Private/SimplifyMeshTool.cpp b/Engine/Plugins/Experimental/MeshModelingToolset/Source/MeshModelingToolsEditorOnly/Private/SimplifyMeshTool.cpp
--- a/Engine/Plugins/Experimental/MeshModelingToolset/Source/MeshModelingToolsEditorOnly/Private/SimplifyMeshTool.cpp
+++ b/Engine/Plugins/Experimental/MeshModelingToolset/Source/MeshModelingToolsEditorOnly/Private/SimplifyMeshTool.cpp
@@ -227,21 +227,24 @@ void USimplifyMeshTool::Render(IToolsContextRenderAPI* RenderAPI)
 	FPrimitiveDrawInterface* PDI = RenderAPI->GetPrimitiveDrawInterface();
 	FTransform Transform = ComponentTarget->GetWorldTransform(); //Actor->GetTransform();
 
-	FColor LineColor(255, 0, 0);
-	const FDynamicMesh3* TargetMesh = Preview->PreviewMesh->GetPreviewDynamicMesh();
-	if (TargetMesh->HasAttributes())
+	if (SimplifyProperties->bShowConstraintEdges)
 	{
-		float PDIScale = RenderAPI->GetCameraState().GetPDIScalingFactor();
-
-		const FDynamicMeshUVOverlay* UVOverlay = TargetMesh->Attributes()->PrimaryUV();
-		for (int eid : TargetMesh->EdgeIndicesItr())
+		FColor LineColor(255, 0, 0);
+		const FDynamicMesh3* TargetMesh = Preview->PreviewMesh->GetPreviewDynamicMesh();
+		if (TargetMesh->HasAttributes())
 		{
-			if (UVOverlay->IsSeamEdge(eid))
+			float PDIScale = RenderAPI->GetCameraState().GetPDIScalingFactor();
+
+			const FDynamicMeshUVOverlay* UVOverlay = TargetMesh->Attributes()->PrimaryUV();
+			for (int eid : TargetMesh->EdgeIndicesItr())
 			{
-				FVector3d A, B;
-				TargetMesh->GetEdgeV(eid, A, B);
-				PDI->DrawLine(Transform.TransformPosition((FVector)A), Transform.TransformPosition((FVector)B),
-					LineColor, 0, 2.0*PDIScale, 1.0f, true);
+				if (UVOverlay->IsSeamEdge(eid))
+				{
+					FVector3d A, B;
+					TargetMesh->GetEdgeV(eid, A, B);
+					PDI->DrawLine(Transform.TransformPosition((FVector)A), Transform.TransformPosition((FVector)B),
+						LineColor, 0, 2.0 * PDIScale, 1.0f, true);
+				}
 			}
 		}
 	}
diff --git a/Engine/Plugins/Experimental/MeshModelingToolset/Source/MeshModelingToolsEditorOnly/Public/SimplifyMeshTool.h b/Engine/Plugins/Experimental/MeshModelingToolset/Source/MeshModelingToolsEditorOnly/Public/SimplifyMeshTool.h
--- a/Engine/Plugins/Experimental/MeshModelingToolset/Source/MeshModelingToolsEditorOnly/Public/SimplifyMeshTool.h
+++ b/Engine/Plugins/Experimental/MeshModelingToolset/Source/MeshModelingToolsEditorOnly/Public/SimplifyMeshTool.h
@@ -77,6 +77,10 @@ public:
 	UPROPERTY(EditAnywhere, Category = Display)
 	bool bShowGroupColors = false;
 
+	/** Display constrained edges */
+	UPROPERTY(EditAnywhere, Category = Display)
+	bool bShowConstraintEdges = false;
+
 	/** Enable projection back to input mesh */
 	UPROPERTY(EditAnywhere, Category = Options, AdvancedDisplay)
 	bool bReproject;
