fix ensure() being hit in voxwrap when run on an empty mesh

#rnx

diff --git a/Engine/Plugins/Experimental/MeshModelingToolset/Source/ModelingOperators/Private/CompositionOps/VoxelSolidifyMeshesOp.cpp b/Engine/Plugins/Experimental/MeshModelingToolset/Source/ModelingOperators/Private/CompositionOps/VoxelSolidifyMeshesOp.cpp
--- a/Engine/Plugins/Experimental/MeshModelingToolset/Source/ModelingOperators/Private/CompositionOps/VoxelSolidifyMeshesOp.cpp
+++ b/Engine/Plugins/Experimental/MeshModelingToolset/Source/ModelingOperators/Private/CompositionOps/VoxelSolidifyMeshesOp.cpp
@@ -55,6 +55,11 @@ void FVoxelSolidifyMeshesOp::CalculateResult(FProgressCancel* Progress)
 		}
 	}
 
+	if (CombinedMesh.TriangleCount() == 0)
+	{
+		return;
+	}
+
 	if (bApplyThickenShells)
 	{
 		// thickness should be at least a cell wide so we don't end up deleting a bunch of the input surface
