protect against creating empty meshes, and add user messaging to some tools where it is likely to happen
#rnx

diff --git a/Engine/Plugins/Experimental/MeshModelingToolset/Source/MeshModelingTools/Private/DrawAndRevolveTool.cpp b/Engine/Plugins/Experimental/MeshModelingToolset/Source/MeshModelingTools/Private/DrawAndRevolveTool.cpp
--- a/Engine/Plugins/Experimental/MeshModelingToolset/Source/MeshModelingTools/Private/DrawAndRevolveTool.cpp
+++ b/Engine/Plugins/Experimental/MeshModelingToolset/Source/MeshModelingTools/Private/DrawAndRevolveTool.cpp
@@ -100,7 +100,7 @@ void UDrawAndRevolveTool::OnBackspacePress()
 
 bool UDrawAndRevolveTool::CanAccept() const
 {
-	return Preview != nullptr && Preview->HaveValidResult();
+	return Preview != nullptr && Preview->HaveValidNonEmptyResult();
 }
 
 void UDrawAndRevolveTool::Setup()
@@ -252,6 +252,13 @@ void UDrawAndRevolveTool::StartPreview()
 		ToolSetupUtil::GetDefaultWorkingMaterial(GetToolManager()));
 	Preview->PreviewMesh->EnableWireframe(MaterialProperties->bWireframe);
 
+	Preview->OnMeshUpdated.AddLambda(
+		[this](const UMeshOpPreviewWithBackgroundCompute* UpdatedPreview)
+		{
+			UpdateAcceptWarnings(UpdatedPreview->HaveEmptyResult() ? EAcceptWarning::EmptyForbidden : EAcceptWarning::NoWarning);
+		}
+	);
+
 	Preview->SetVisibility(true);
 	Preview->InvalidateResult();
 }
diff --git a/Engine/Plugins/Experimental/MeshModelingToolset/Source/MeshModelingTools/Private/RemoveOccludedTrianglesTool.cpp b/Engine/Plugins/Experimental/MeshModelingToolset/Source/MeshModelingTools/Private/RemoveOccludedTrianglesTool.cpp
--- a/Engine/Plugins/Experimental/MeshModelingToolset/Source/MeshModelingTools/Private/RemoveOccludedTrianglesTool.cpp
+++ b/Engine/Plugins/Experimental/MeshModelingToolset/Source/MeshModelingTools/Private/RemoveOccludedTrianglesTool.cpp
@@ -18,6 +18,8 @@
 
 #include "AssetGenerationUtil.h"
 
+#include "Misc/MessageDialog.h"
+
 
 #if WITH_EDITOR
 #include "Misc/ScopedSlowTask.h"
@@ -490,11 +492,45 @@ void URemoveOccludedTrianglesTool::GenerateAsset(const TArray<FDynamicMeshOpResu
 	GetToolManager()->BeginUndoTransaction(LOCTEXT("RemoveOccludedTrianglesToolTransactionName", "Remove Occluded Triangles"));
 
 	check(Results.Num() == Previews.Num());
+
+	// check if we entirely remove away any meshes
+	bool bWantDestroy = false;
+	for (int32 PreviewIdx = 0; PreviewIdx < Previews.Num(); PreviewIdx++)
+	{
+		bWantDestroy = bWantDestroy || (Results[PreviewIdx].Mesh.Get()->TriangleCount() == 0);
+	}
+	// if so ask user what to do
+	if (bWantDestroy)
+	{
+		FText Title = LOCTEXT("RemoveOccludedDestroyTitle", "Delete mesh components?");
+		EAppReturnType::Type Ret = FMessageDialog::Open(EAppMsgType::YesNo,
+			LOCTEXT("RemoveOccludedDestroyQuestion", "Jacketing has removed all triangles from some meshes.  Actually destroy these mesh components?"), &Title);
+		if (Ret == EAppReturnType::No)
+		{
+			bWantDestroy = false;
+		}
+	}
 	
 	for (int32 PreviewIdx = 0; PreviewIdx < Previews.Num(); PreviewIdx++)
 	{
 		check(Results[PreviewIdx].Mesh.Get() != nullptr);
 		int ComponentIdx = PreviewToTargetIdx[PreviewIdx];
+
+		if (Results[PreviewIdx].Mesh.Get()->TriangleCount() == 0)
+		{
+			if (bWantDestroy)
+			{
+				for (int TargetIdx = 0; TargetIdx < ComponentTargets.Num(); TargetIdx++)
+				{
+					if (TargetToPreviewIdx[TargetIdx] == PreviewIdx)
+					{
+						ComponentTargets[TargetIdx]->GetOwnerComponent()->DestroyComponent();
+					}
+				}
+			}
+			continue;
+		}
+
 		ComponentTargets[ComponentIdx]->CommitMesh([&Results, &PreviewIdx, this](const FPrimitiveComponentTarget::FCommitParams& CommitParams)
 		{
 			FDynamicMeshToMeshDescription Converter;
diff --git a/Engine/Plugins/Experimental/MeshModelingToolset/Source/MeshModelingTools/Private/RevolveBoundaryTool.cpp b/Engine/Plugins/Experimental/MeshModelingToolset/Source/MeshModelingTools/Private/RevolveBoundaryTool.cpp
--- a/Engine/Plugins/Experimental/MeshModelingToolset/Source/MeshModelingTools/Private/RevolveBoundaryTool.cpp
+++ b/Engine/Plugins/Experimental/MeshModelingToolset/Source/MeshModelingTools/Private/RevolveBoundaryTool.cpp
@@ -216,7 +216,7 @@ void URevolveBoundaryTool::OnClicked(const FInputDeviceRay& ClickPos)
 
 bool URevolveBoundaryTool::CanAccept() const
 {
-	return Preview != nullptr && Preview->HaveValidResult();
+	return Preview != nullptr && Preview->HaveValidNonEmptyResult();
 }
 
 /** 
@@ -245,6 +245,13 @@ void URevolveBoundaryTool::StartPreview()
 		ToolSetupUtil::GetDefaultWorkingMaterial(GetToolManager()));
 	Preview->PreviewMesh->EnableWireframe(MaterialProperties->bWireframe);
 
+	Preview->OnMeshUpdated.AddLambda(
+		[this](const UMeshOpPreviewWithBackgroundCompute* UpdatedPreview)
+		{
+			UpdateAcceptWarnings(UpdatedPreview->HaveEmptyResult() ? EAcceptWarning::EmptyForbidden : EAcceptWarning::NoWarning);
+		}
+	);
+
 	Preview->SetVisibility(true);
 	Preview->InvalidateResult();
 }
diff --git a/Engine/Plugins/Experimental/MeshModelingToolset/Source/MeshModelingToolsEditorOnly/Private/MergeMeshesTool.cpp b/Engine/Plugins/Experimental/MeshModelingToolset/Source/MeshModelingToolsEditorOnly/Private/MergeMeshesTool.cpp
--- a/Engine/Plugins/Experimental/MeshModelingToolset/Source/MeshModelingToolsEditorOnly/Private/MergeMeshesTool.cpp
+++ b/Engine/Plugins/Experimental/MeshModelingToolset/Source/MeshModelingToolsEditorOnly/Private/MergeMeshesTool.cpp
@@ -103,6 +103,7 @@ void UMergeMeshesTool::Setup()
 	Preview->Setup(this->TargetWorld, this);
 	Preview->OnMeshUpdated.AddLambda([this](UMeshOpPreviewWithBackgroundCompute* Compute) {
 		MeshStatisticsProperties->Update(*Compute->PreviewMesh->GetPreviewDynamicMesh());
+		UpdateAcceptWarnings(Compute->HaveEmptyResult() ? EAcceptWarning::EmptyForbidden : EAcceptWarning::NoWarning);
 	});
 
 	CreateLowQualityPreview(); // update the preview with a low-quality result
@@ -159,7 +160,7 @@ void UMergeMeshesTool::OnTick(float DeltaTime)
 
 bool UMergeMeshesTool::CanAccept() const
 {
-	return Super::CanAccept() && Preview->HaveValidResult();
+	return Super::CanAccept() && Preview->HaveValidNonEmptyResult();
 }
 
 void UMergeMeshesTool::OnPropertyModified(UObject* PropertySet, FProperty* Property)
diff --git a/Engine/Plugins/Experimental/MeshModelingToolset/Source/MeshModelingToolsEditorOnly/Private/PolygonOnMeshTool.cpp b/Engine/Plugins/Experimental/MeshModelingToolset/Source/MeshModelingToolsEditorOnly/Private/PolygonOnMeshTool.cpp
--- a/Engine/Plugins/Experimental/MeshModelingToolset/Source/MeshModelingToolsEditorOnly/Private/PolygonOnMeshTool.cpp
+++ b/Engine/Plugins/Experimental/MeshModelingToolset/Source/MeshModelingToolsEditorOnly/Private/PolygonOnMeshTool.cpp
@@ -114,10 +114,11 @@ void UPolygonOnMeshTool::Setup()
 		}
 	);
 	Preview->OnMeshUpdated.AddLambda(
-		[this](const UMeshOpPreviewWithBackgroundCompute*)
+		[this](const UMeshOpPreviewWithBackgroundCompute* UpdatedPreview)
 		{
 			GetToolManager()->PostInvalidation();
 			UpdateVisualization();
+			UpdateAcceptWarnings(UpdatedPreview->HaveEmptyResult() ? EAcceptWarning::EmptyForbidden : EAcceptWarning::NoWarning);
 		}
 	);
 
@@ -397,7 +398,7 @@ void UPolygonOnMeshTool::CompleteDrawPolygon()
 
 bool UPolygonOnMeshTool::CanAccept() const
 {
-	return Super::CanAccept() && Preview != nullptr && Preview->HaveValidResult();
+	return Super::CanAccept() && Preview != nullptr && Preview->HaveValidNonEmptyResult();
 }
 
 
diff --git a/Engine/Plugins/Experimental/MeshModelingToolset/Source/MeshModelingToolsEditorOnly/Private/VoxelCSGMeshesTool.cpp b/Engine/Plugins/Experimental/MeshModelingToolset/Source/MeshModelingToolsEditorOnly/Private/VoxelCSGMeshesTool.cpp
--- a/Engine/Plugins/Experimental/MeshModelingToolset/Source/MeshModelingToolsEditorOnly/Private/VoxelCSGMeshesTool.cpp
+++ b/Engine/Plugins/Experimental/MeshModelingToolset/Source/MeshModelingToolsEditorOnly/Private/VoxelCSGMeshesTool.cpp
@@ -99,6 +99,7 @@ void UVoxelCSGMeshesTool::Setup()
 	Preview->Setup(this->TargetWorld, this);
 	Preview->OnMeshUpdated.AddLambda([this](UMeshOpPreviewWithBackgroundCompute* Compute) {
 		MeshStatisticsProperties->Update(*Compute->PreviewMesh->GetPreviewDynamicMesh());
+		UpdateAcceptWarnings(Compute->HaveEmptyResult() ? EAcceptWarning::EmptyForbidden : EAcceptWarning::NoWarning);
 	});
 
 	CreateLowQualityPreview();
@@ -158,7 +159,7 @@ void UVoxelCSGMeshesTool::Render(IToolsContextRenderAPI* RenderAPI)
 
 bool UVoxelCSGMeshesTool::CanAccept() const
 {
-	return Super::CanAccept() && Preview->HaveValidResult();
+	return Super::CanAccept() && Preview->HaveValidNonEmptyResult();
 }
 
 void UVoxelCSGMeshesTool::OnPropertyModified(UObject* PropertySet, FProperty* Property)
diff --git a/Engine/Plugins/Experimental/MeshModelingToolset/Source/ModelingComponents/Private/BaseTools/BaseCreateFromSelectedTool.cpp b/Engine/Plugins/Experimental/MeshModelingToolset/Source/ModelingComponents/Private/BaseTools/BaseCreateFromSelectedTool.cpp
--- a/Engine/Plugins/Experimental/MeshModelingToolset/Source/ModelingComponents/Private/BaseTools/BaseCreateFromSelectedTool.cpp
+++ b/Engine/Plugins/Experimental/MeshModelingToolset/Source/ModelingComponents/Private/BaseTools/BaseCreateFromSelectedTool.cpp
@@ -88,6 +88,12 @@ void UBaseCreateFromSelectedTool::Setup()
 	Preview->Setup(this->TargetWorld, this);
 
 	SetPreviewCallbacks();
+	Preview->OnMeshUpdated.AddLambda(
+		[this](const UMeshOpPreviewWithBackgroundCompute* UpdatedPreview)
+		{
+			UpdateAcceptWarnings(UpdatedPreview->HaveEmptyResult() ? EAcceptWarning::EmptyForbidden : EAcceptWarning::NoWarning);
+		}
+	);
 
 	SetTransformGizmos();
 
@@ -345,7 +351,7 @@ void UBaseCreateFromSelectedTool::Shutdown(EToolShutdownType ShutdownType)
 
 bool UBaseCreateFromSelectedTool::CanAccept() const
 {
-	return Super::CanAccept() && Preview->HaveValidResult();
+	return Super::CanAccept() && Preview->HaveValidNonEmptyResult();
 }
 
 
diff --git a/Engine/Plugins/Experimental/MeshModelingToolset/Source/ModelingComponents/Private/BaseTools/BaseMeshProcessingTool.cpp b/Engine/Plugins/Experimental/MeshModelingToolset/Source/ModelingComponents/Private/BaseTools/BaseMeshProcessingTool.cpp
--- a/Engine/Plugins/Experimental/MeshModelingToolset/Source/ModelingComponents/Private/BaseTools/BaseMeshProcessingTool.cpp
+++ b/Engine/Plugins/Experimental/MeshModelingToolset/Source/ModelingComponents/Private/BaseTools/BaseMeshProcessingTool.cpp
@@ -292,7 +292,7 @@ bool UBaseMeshProcessingTool::HasAccept() const
 
 bool UBaseMeshProcessingTool::CanAccept() const
 {
-	return Super::CanAccept() && bResultValid;
+	return Super::CanAccept() && bResultValid && Preview->HaveValidNonEmptyResult();
 }
 
 
diff --git a/Engine/Plugins/Experimental/MeshModelingToolset/Source/ModelingComponents/Public/MeshOpPreviewHelpers.h b/Engine/Plugins/Experimental/MeshModelingToolset/Source/ModelingComponents/Public/MeshOpPreviewHelpers.h
--- a/Engine/Plugins/Experimental/MeshModelingToolset/Source/ModelingComponents/Public/MeshOpPreviewHelpers.h
+++ b/Engine/Plugins/Experimental/MeshModelingToolset/Source/ModelingComponents/Public/MeshOpPreviewHelpers.h
@@ -94,6 +94,17 @@ public:
 	bool HaveValidResult() const { return bResultValid; }
 
 
+	/**
+	 * @return true if current PreviewMesh result is valid (no update actively being computed) and that mesh has at least one triangle
+	 */
+	bool HaveValidNonEmptyResult() const { return bResultValid && PreviewMesh->GetMesh() && PreviewMesh->GetMesh()->TriangleCount() > 0; }
+
+	/**
+	 * @return true if current PreviewMesh result is valid (no update actively being computed) but that mesh has no triangles
+	 */
+	bool HaveEmptyResult() const { return bResultValid && PreviewMesh->GetMesh() && PreviewMesh->GetMesh()->TriangleCount() == 0; }
+
+
 	/**
 	 * Read back a copy of current preview mesh.
 	 * @param bOnlyIfValid if true, then only create mesh copy if HaveValidResult() == true
diff --git a/Engine/Source/Runtime/Experimental/InteractiveToolsFramework/Private/InteractiveTool.cpp b/Engine/Source/Runtime/Experimental/InteractiveToolsFramework/Private/InteractiveTool.cpp
--- a/Engine/Source/Runtime/Experimental/InteractiveToolsFramework/Private/InteractiveTool.cpp
+++ b/Engine/Source/Runtime/Experimental/InteractiveToolsFramework/Private/InteractiveTool.cpp
@@ -232,6 +232,24 @@ bool UInteractiveTool::CanAccept() const
 }
 
 
+void UInteractiveTool::UpdateAcceptWarnings(EAcceptWarning Warning)
+{
+	switch (Warning)
+	{
+	case EAcceptWarning::NoWarning:
+		// clear warning
+		GetToolManager()->DisplayMessage(FText(), EToolMessageLevel::UserWarning);
+		break;
+	case EAcceptWarning::EmptyForbidden:
+		GetToolManager()->DisplayMessage(LOCTEXT("CannotCreateEmptyMesh", "WARNING: Tool doesn't allow creation of an empty mesh."),
+			EToolMessageLevel::UserWarning);
+		break;
+	default:
+		check(false);
+	}
+}
+
+
 void UInteractiveTool::Tick(float DeltaTime)
 {
 	for (auto& Object : ToolPropertyObjects)
diff --git a/Engine/Source/Runtime/Experimental/InteractiveToolsFramework/Public/InteractiveTool.h b/Engine/Source/Runtime/Experimental/InteractiveToolsFramework/Public/InteractiveTool.h
--- a/Engine/Source/Runtime/Experimental/InteractiveToolsFramework/Public/InteractiveTool.h
+++ b/Engine/Source/Runtime/Experimental/InteractiveToolsFramework/Public/InteractiveTool.h
@@ -410,6 +410,22 @@ protected:
 	virtual bool SetToolPropertySourceEnabled(UInteractiveToolPropertySet* PropertySet, bool bEnabled);
 
 
+	enum EAcceptWarning
+	{
+		NoWarning,
+		EmptyForbidden
+	};
+
+	/**
+	 * Helper function to update a standard warning when we need to explain why "Accept" is disabled
+	 * Note that calling this base implementation will clear any unrelated warnings.
+	 * Note that this function is not automatically called by the base class.
+	 *
+	 * @param Warning Reason that tool cannot be accepted (or EAcceptWarning::NoWarning if no warning need be displayed)
+	 */
+	virtual void UpdateAcceptWarnings(EAcceptWarning Warning);
+
+
 	//
 	// Action support/system
 	//
