ModelingTools UX: change tool names so that label now visible in the main viewport corresponds to icon name. Update accept/cancel/complete tooltips and button styles.
#rnx

diff --git a/Engine/Plugins/Experimental/MeshLODToolset/Source/MeshLODToolset/Private/Tools/GenerateStaticMeshLODAssetTool.cpp b/Engine/Plugins/Experimental/MeshLODToolset/Source/MeshLODToolset/Private/Tools/GenerateStaticMeshLODAssetTool.cpp
--- a/Engine/Plugins/Experimental/MeshLODToolset/Source/MeshLODToolset/Private/Tools/GenerateStaticMeshLODAssetTool.cpp
+++ b/Engine/Plugins/Experimental/MeshLODToolset/Source/MeshLODToolset/Private/Tools/GenerateStaticMeshLODAssetTool.cpp
@@ -196,8 +196,18 @@ void UGenerateStaticMeshLODAssetTool::Setup()
 
 	UInteractiveTool::Setup();
 
+	BasicProperties = NewObject<UGenerateStaticMeshLODAssetToolProperties>(this);
+	AddToolPropertySource(BasicProperties);
+	BasicProperties->RestoreProperties(this);
+
+	BasicProperties->OutputName = AssetGenerationUtil::GetComponentAssetBaseName(ComponentTargets[0]->GetOwnerComponent());
+
+	BasicProperties->GeneratedSuffix = TEXT("_AutoLOD");
+
+
+	SetToolDisplayName(LOCTEXT("ToolName", "Generate LOD"));
 	GetToolManager()->DisplayMessage(
-		LOCTEXT("OnStartStaticMeshLODAssetTool", "This tool creates a new LOD asset"),
+		LOCTEXT("OnStartStaticMeshLODAssetTool", "Create a new LOD asset"),
 		EToolMessageLevel::UserNotification);
 
 	GenerateProcess = NewObject<UGenerateStaticMeshLODProcess>(this);
diff --git a/Engine/Plugins/Experimental/MeshModelingToolset/Source/MeshModelingTools/Private/AddPrimitiveTool.cpp b/Engine/Plugins/Experimental/MeshModelingToolset/Source/MeshModelingTools/Private/AddPrimitiveTool.cpp
--- a/Engine/Plugins/Experimental/MeshModelingToolset/Source/MeshModelingTools/Private/AddPrimitiveTool.cpp
+++ b/Engine/Plugins/Experimental/MeshModelingToolset/Source/MeshModelingTools/Private/AddPrimitiveTool.cpp
@@ -368,6 +368,14 @@ void UAddPrimitiveTool::OnClicked(const FInputDeviceRay& DeviceClickPos)
 #endif
 }
 
+
+UAddBoxPrimitiveTool::UAddBoxPrimitiveTool(const FObjectInitializer& ObjectInitializer)
+	: Super(ObjectInitializer.SetDefaultSubobjectClass<UProceduralBoxToolProperties>(TEXT("ShapeSettings")))
+{
+	AssetName = TEXT("Box");
+	SetToolDisplayName(LOCTEXT("BoxToolName", "Create Boxes"));
+}
+
 void UAddBoxPrimitiveTool::GenerateMesh(FDynamicMesh3* OutMesh) const
 {
 	FGridBoxMeshGenerator BoxGen;
@@ -384,6 +392,15 @@ void UAddBoxPrimitiveTool::GenerateMesh(FDynamicMesh3* OutMesh) const
 	OutMesh->Copy(&BoxGen);
 }
 
+
+
+UAddRectanglePrimitiveTool::UAddRectanglePrimitiveTool(const FObjectInitializer& ObjectInitializer)
+	: Super(ObjectInitializer.SetDefaultSubobjectClass<UProceduralRectangleToolProperties>(TEXT("ShapeSettings")))
+{
+	AssetName = TEXT("Rectangle");
+	SetToolDisplayName(LOCTEXT("RectToolName", "Create Rectangles"));
+}
+
 void UAddRectanglePrimitiveTool::GenerateMesh(FDynamicMesh3* OutMesh) const
 {
 	FRectangleMeshGenerator RectGen;
@@ -400,6 +417,14 @@ void UAddRectanglePrimitiveTool::GenerateMesh(FDynamicMesh3* OutMesh) const
 	OutMesh->Copy(&RectGen);
 }
 
+
+UAddRoundedRectanglePrimitiveTool::UAddRoundedRectanglePrimitiveTool(const FObjectInitializer& ObjectInitializer)
+	: Super(ObjectInitializer.SetDefaultSubobjectClass<UProceduralRoundedRectangleToolProperties>(TEXT("ShapeSettings")))
+{
+	AssetName = TEXT("RoundedRectangle");
+	SetToolDisplayName(LOCTEXT("RoudRectToolName", "Create RoundRects"));
+}
+
 void UAddRoundedRectanglePrimitiveTool::GenerateMesh(FDynamicMesh3* OutMesh) const
 {
 	FRoundedRectangleMeshGenerator RectGen;
@@ -418,6 +443,15 @@ void UAddRoundedRectanglePrimitiveTool::GenerateMesh(FDynamicMesh3* OutMesh) con
 	OutMesh->Copy(&RectGen);
 }
 
+
+
+UAddDiscPrimitiveTool::UAddDiscPrimitiveTool(const FObjectInitializer& ObjectInitializer)
+	: Super(ObjectInitializer.SetDefaultSubobjectClass<UProceduralDiscToolProperties>(TEXT("ShapeSettings")))
+{
+	AssetName = TEXT("Disc");
+	SetToolDisplayName(LOCTEXT("DiscToolName", "Create Discs"));
+}
+
 void UAddDiscPrimitiveTool::GenerateMesh(FDynamicMesh3* OutMesh) const
 {
 	FDiscMeshGenerator Gen;
@@ -433,6 +467,15 @@ void UAddDiscPrimitiveTool::GenerateMesh(FDynamicMesh3* OutMesh) const
 	OutMesh->Copy(&Gen);
 }
 
+
+
+UAddPuncturedDiscPrimitiveTool::UAddPuncturedDiscPrimitiveTool(const FObjectInitializer& ObjectInitializer)
+	: Super(ObjectInitializer.SetDefaultSubobjectClass<UProceduralPuncturedDiscToolProperties>(TEXT("ShapeSettings")))
+{
+	AssetName = TEXT("PuncturedDisc");
+	SetToolDisplayName(LOCTEXT("CircleToolName", "Create Circles"));
+}
+
 void UAddPuncturedDiscPrimitiveTool::GenerateMesh(FDynamicMesh3* OutMesh) const
 {
 	FPuncturedDiscMeshGenerator Gen;
@@ -449,6 +492,14 @@ void UAddPuncturedDiscPrimitiveTool::GenerateMesh(FDynamicMesh3* OutMesh) const
 	OutMesh->Copy(&Gen);
 }
 
+
+UAddTorusPrimitiveTool::UAddTorusPrimitiveTool(const FObjectInitializer& ObjectInitializer)
+	: Super(ObjectInitializer.SetDefaultSubobjectClass<UProceduralTorusToolProperties>(TEXT("ShapeSettings")))
+{
+	AssetName = TEXT("Torus");
+	SetToolDisplayName(LOCTEXT("TorusToolName", "Create Torii"));
+}
+
 void UAddTorusPrimitiveTool::GenerateMesh(FDynamicMesh3* OutMesh) const
 {
 	FGeneralizedCylinderGenerator Gen;
@@ -470,6 +521,15 @@ void UAddTorusPrimitiveTool::GenerateMesh(FDynamicMesh3* OutMesh) const
 	OutMesh->Copy(&Gen);
 }
 
+
+
+UAddCylinderPrimitiveTool::UAddCylinderPrimitiveTool(const FObjectInitializer& ObjectInitializer)
+	: Super(ObjectInitializer.SetDefaultSubobjectClass<UProceduralCylinderToolProperties>(TEXT("ShapeSettings")))
+{
+	AssetName = TEXT("Cylinder");
+	SetToolDisplayName(LOCTEXT("CylinderToolName", "Create Cylinders"));
+}
+
 void UAddCylinderPrimitiveTool::GenerateMesh(FDynamicMesh3* OutMesh) const
 {
 	FCylinderGenerator CylGen;
@@ -487,6 +547,15 @@ void UAddCylinderPrimitiveTool::GenerateMesh(FDynamicMesh3* OutMesh) const
 	OutMesh->Copy(&CylGen);
 }
 
+
+
+UAddConePrimitiveTool::UAddConePrimitiveTool(const FObjectInitializer& ObjectInitializer)
+	: Super(ObjectInitializer.SetDefaultSubobjectClass<UProceduralConeToolProperties>(TEXT("ShapeSettings")))
+{
+	AssetName = TEXT("Cone");
+	SetToolDisplayName(LOCTEXT("ConeToolName", "Create Cones"));
+}
+
 void UAddConePrimitiveTool::GenerateMesh(FDynamicMesh3* OutMesh) const
 {
 	// Unreal's standard cone is just a cylinder with a very small top
@@ -506,6 +575,14 @@ void UAddConePrimitiveTool::GenerateMesh(FDynamicMesh3* OutMesh) const
 	OutMesh->Copy(&CylGen);
 }
 
+
+UAddArrowPrimitiveTool::UAddArrowPrimitiveTool(const FObjectInitializer& ObjectInitializer)
+	: Super(ObjectInitializer.SetDefaultSubobjectClass<UProceduralArrowToolProperties>(TEXT("ShapeSettings")))
+{
+	AssetName = TEXT("Arrow");
+	SetToolDisplayName(LOCTEXT("ArrowToolName", "Create Arrows"));
+}
+
 void UAddArrowPrimitiveTool::GenerateMesh(FDynamicMesh3* OutMesh) const
 {
 	FArrowGenerator ArrowGen;
@@ -526,6 +603,15 @@ void UAddArrowPrimitiveTool::GenerateMesh(FDynamicMesh3* OutMesh) const
 	OutMesh->Copy(&ArrowGen);
 }
 
+
+
+UAddSpherePrimitiveTool::UAddSpherePrimitiveTool(const FObjectInitializer& ObjectInitializer)
+	: Super(ObjectInitializer.SetDefaultSubobjectClass<UProceduralSphereToolProperties>(TEXT("ShapeSettings")))
+{
+	AssetName = TEXT("ASphere");
+	SetToolDisplayName(LOCTEXT("SphereToolName", "Create Spheres"));
+}
+
 void UAddSpherePrimitiveTool::GenerateMesh(FDynamicMesh3* OutMesh) const
 {
 	FSphereGenerator SphereGen;
@@ -541,6 +627,15 @@ void UAddSpherePrimitiveTool::GenerateMesh(FDynamicMesh3* OutMesh) const
 	OutMesh->Copy(&SphereGen);
 }
 
+
+
+UAddSphericalBoxPrimitiveTool::UAddSphericalBoxPrimitiveTool(const FObjectInitializer& ObjectInitializer)
+	: Super(ObjectInitializer.SetDefaultSubobjectClass<UProceduralSphericalBoxToolProperties>(TEXT("ShapeSettings")))
+{
+	AssetName = TEXT("BSphere");
+	SetToolDisplayName(LOCTEXT("BoxSphereToolName", "Create BoxSpheres"));
+}
+
 void UAddSphericalBoxPrimitiveTool::GenerateMesh(FDynamicMesh3* OutMesh) const
 {
 	FBoxSphereGenerator SphereGen;
diff --git a/Engine/Plugins/Experimental/MeshModelingToolset/Source/MeshModelingTools/Private/AlignObjectsTool.cpp b/Engine/Plugins/Experimental/MeshModelingToolset/Source/MeshModelingTools/Private/AlignObjectsTool.cpp
--- a/Engine/Plugins/Experimental/MeshModelingToolset/Source/MeshModelingTools/Private/AlignObjectsTool.cpp
+++ b/Engine/Plugins/Experimental/MeshModelingToolset/Source/MeshModelingTools/Private/AlignObjectsTool.cpp
@@ -79,6 +79,7 @@ void UAlignObjectsTool::Setup()
 
 	Precompute();
 
+	SetToolDisplayName(LOCTEXT("ToolName", "Align"));
 	GetToolManager()->DisplayMessage(
 		LOCTEXT("OnStartTool", "This Tool aligns the Pivots or Bounding Boxes of the input Objects."),
 		EToolMessageLevel::UserNotification);
diff --git a/Engine/Plugins/Experimental/MeshModelingToolset/Source/MeshModelingTools/Private/BakeMeshAttributeMapsTool.cpp b/Engine/Plugins/Experimental/MeshModelingToolset/Source/MeshModelingTools/Private/BakeMeshAttributeMapsTool.cpp
--- a/Engine/Plugins/Experimental/MeshModelingToolset/Source/MeshModelingTools/Private/BakeMeshAttributeMapsTool.cpp
+++ b/Engine/Plugins/Experimental/MeshModelingToolset/Source/MeshModelingTools/Private/BakeMeshAttributeMapsTool.cpp
@@ -202,8 +202,9 @@ void UBakeMeshAttributeMapsTool::Setup()
 	bResultValid = false;
 	bDetailMeshValid = false;
 
+	SetToolDisplayName(LOCTEXT("ToolName", "Bake Textures"));
 	GetToolManager()->DisplayMessage(
-		LOCTEXT("OnStartTool", "Bake Maps. Select Bake Mesh (LowPoly) first, then Detail Mesh second. Texture Assets will be created on Accept. "),
+		LOCTEXT("OnStartTool", "Bake Maps. Select Bake Mesh (LowPoly) first, then (optionally) Detail Mesh second. Texture Assets will be created on Accept. "),
 		EToolMessageLevel::UserNotification);
 }
 
diff --git a/Engine/Plugins/Experimental/MeshModelingToolset/Source/MeshModelingTools/Private/BakeTransformTool.cpp b/Engine/Plugins/Experimental/MeshModelingToolset/Source/MeshModelingTools/Private/BakeTransformTool.cpp
--- a/Engine/Plugins/Experimental/MeshModelingToolset/Source/MeshModelingTools/Private/BakeTransformTool.cpp
+++ b/Engine/Plugins/Experimental/MeshModelingToolset/Source/MeshModelingTools/Private/BakeTransformTool.cpp
@@ -109,6 +109,7 @@ void UBakeTransformTool::Setup()
 	
 	GetToolManager()->DisplayMessage(AllTheWarnings, EToolMessageLevel::UserWarning);
 
+	SetToolDisplayName(LOCTEXT("ToolName", "Bake Transform"));
 	GetToolManager()->DisplayMessage(
 		LOCTEXT("OnStartTool", "This Tool applies the current Rotation and/or Scaling of the object's Transform to the underlying mesh Asset."),
 		EToolMessageLevel::UserNotification);
diff --git a/Engine/Plugins/Experimental/MeshModelingToolset/Source/MeshModelingTools/Private/CSGMeshesTool.cpp b/Engine/Plugins/Experimental/MeshModelingToolset/Source/MeshModelingTools/Private/CSGMeshesTool.cpp
--- a/Engine/Plugins/Experimental/MeshModelingToolset/Source/MeshModelingTools/Private/CSGMeshesTool.cpp
+++ b/Engine/Plugins/Experimental/MeshModelingToolset/Source/MeshModelingTools/Private/CSGMeshesTool.cpp
@@ -28,7 +28,7 @@ void UCSGMeshesTool::SetupProperties()
 		TrimProperties->RestoreProperties(this);
 		AddToolPropertySource(TrimProperties);
 
-		SetToolDisplayName(LOCTEXT("TrimMeshesToolName", "Trim Meshes Tool"));
+		SetToolDisplayName(LOCTEXT("TrimMeshesToolName", "Trim"));
 		GetToolManager()->DisplayMessage(
 			LOCTEXT("OnStartTrimTool", "Trim one mesh with another. Use the transform gizmos to tweak the positions of the input objects (can help to resolve errors/failures)"),
 			EToolMessageLevel::UserNotification);
@@ -39,7 +39,7 @@ void UCSGMeshesTool::SetupProperties()
 		CSGProperties->RestoreProperties(this);
 		AddToolPropertySource(CSGProperties);
 
-		SetToolDisplayName(LOCTEXT("CSGMeshesToolName", "Mesh Boolean Tool"));
+		SetToolDisplayName(LOCTEXT("CSGMeshesToolName", "Boolean"));
 		GetToolManager()->DisplayMessage(
 			LOCTEXT("OnStartTool", "Compute CSG Booleans on the input meshes. Use the transform gizmos to tweak the positions of the input objects (can help to resolve errors/failures)"),
 			EToolMessageLevel::UserNotification);
diff --git a/Engine/Plugins/Experimental/MeshModelingToolset/Source/MeshModelingTools/Private/CombineMeshesTool.cpp b/Engine/Plugins/Experimental/MeshModelingToolset/Source/MeshModelingTools/Private/CombineMeshesTool.cpp
--- a/Engine/Plugins/Experimental/MeshModelingToolset/Source/MeshModelingTools/Private/CombineMeshesTool.cpp
+++ b/Engine/Plugins/Experimental/MeshModelingToolset/Source/MeshModelingTools/Private/CombineMeshesTool.cpp
@@ -107,12 +107,12 @@ void UCombineMeshesTool::Setup()
 
 	if (bDuplicateMode)
 	{
-		SetToolDisplayName(LOCTEXT("DuplicateMeshesToolName", "Duplicate Mesh Tool"));
+		SetToolDisplayName(LOCTEXT("DuplicateMeshesToolName", "Duplicate"));
 		BasicProperties->OutputName = AssetGenerationUtil::GetComponentAssetBaseName(ComponentTargets[0]->GetOwnerComponent());
 	}
 	else
 	{
-		SetToolDisplayName(LOCTEXT("CombineMeshesToolName", "Combine Meshes Tool"));
+		SetToolDisplayName(LOCTEXT("CombineMeshesToolName", "Append"));
 		BasicProperties->OutputName = FString("Combined");
 	}
 
diff --git a/Engine/Plugins/Experimental/MeshModelingToolset/Source/MeshModelingTools/Private/ConvertToPolygonsTool.cpp b/Engine/Plugins/Experimental/MeshModelingToolset/Source/MeshModelingTools/Private/ConvertToPolygonsTool.cpp
--- a/Engine/Plugins/Experimental/MeshModelingToolset/Source/MeshModelingTools/Private/ConvertToPolygonsTool.cpp
+++ b/Engine/Plugins/Experimental/MeshModelingToolset/Source/MeshModelingTools/Private/ConvertToPolygonsTool.cpp
@@ -45,7 +45,7 @@ UInteractiveTool* UConvertToPolygonsToolBuilder::BuildTool(const FToolBuilderSta
  */
 UConvertToPolygonsTool::UConvertToPolygonsTool()
 {
-	SetToolDisplayName(LOCTEXT("ConvertToPolygonsToolName", "Find PolyGroups Tool"));
+	SetToolDisplayName(LOCTEXT("ConvertToPolygonsToolName", "Generate PolyGroups"));
 }
 
 void UConvertToPolygonsTool::Setup()
diff --git a/Engine/Plugins/Experimental/MeshModelingToolset/Source/MeshModelingTools/Private/DeformMeshPolygonsTool.cpp b/Engine/Plugins/Experimental/MeshModelingToolset/Source/MeshModelingTools/Private/DeformMeshPolygonsTool.cpp
--- a/Engine/Plugins/Experimental/MeshModelingToolset/Source/MeshModelingTools/Private/DeformMeshPolygonsTool.cpp
+++ b/Engine/Plugins/Experimental/MeshModelingToolset/Source/MeshModelingTools/Private/DeformMeshPolygonsTool.cpp
@@ -838,7 +838,7 @@ inline void FGroupTopologyLaplacianDeformer::Shutdown()
 
 UDeformMeshPolygonsTool::UDeformMeshPolygonsTool()
 {
-	SetToolDisplayName(LOCTEXT("DeformPolygroupsToolName", "Deform Polygroups"));
+	SetToolDisplayName(LOCTEXT("DeformPolygroupsToolName", "PolyGroup Deform"));
 }
 
 void UDeformMeshPolygonsTool::Setup()
diff --git a/Engine/Plugins/Experimental/MeshModelingToolset/Source/MeshModelingTools/Private/DisplaceMeshTool.cpp b/Engine/Plugins/Experimental/MeshModelingToolset/Source/MeshModelingTools/Private/DisplaceMeshTool.cpp
--- a/Engine/Plugins/Experimental/MeshModelingToolset/Source/MeshModelingTools/Private/DisplaceMeshTool.cpp
+++ b/Engine/Plugins/Experimental/MeshModelingToolset/Source/MeshModelingTools/Private/DisplaceMeshTool.cpp
@@ -875,6 +875,7 @@ void UDisplaceMeshTool::Setup()
 
 	StartComputation();
 
+	SetToolDisplayName(LOCTEXT("ToolName", "Displace"));
 	GetToolManager()->DisplayMessage(
 		LOCTEXT("OnStartDisplaceMesh", "Subdivide and Displace the input mesh using different noise functions and maps"),
 		EToolMessageLevel::UserNotification);
diff --git a/Engine/Plugins/Experimental/MeshModelingToolset/Source/MeshModelingTools/Private/DrawAndRevolveTool.cpp b/Engine/Plugins/Experimental/MeshModelingToolset/Source/MeshModelingTools/Private/DrawAndRevolveTool.cpp
--- a/Engine/Plugins/Experimental/MeshModelingToolset/Source/MeshModelingTools/Private/DrawAndRevolveTool.cpp
+++ b/Engine/Plugins/Experimental/MeshModelingToolset/Source/MeshModelingTools/Private/DrawAndRevolveTool.cpp
@@ -107,6 +107,7 @@ void UDrawAndRevolveTool::Setup()
 {
 	UInteractiveTool::Setup();
 
+	SetToolDisplayName(LOCTEXT("ToolName", "Revolve PolyPath"));
 	GetToolManager()->DisplayMessage(InitializationModeMessage, EToolMessageLevel::UserNotification);
 
 	Settings = NewObject<URevolveToolProperties>(this, TEXT("Revolve Tool Settings"));
diff --git a/Engine/Plugins/Experimental/MeshModelingToolset/Source/MeshModelingTools/Private/DrawPolyPathTool.cpp b/Engine/Plugins/Experimental/MeshModelingToolset/Source/MeshModelingTools/Private/DrawPolyPathTool.cpp
--- a/Engine/Plugins/Experimental/MeshModelingToolset/Source/MeshModelingTools/Private/DrawPolyPathTool.cpp
+++ b/Engine/Plugins/Experimental/MeshModelingToolset/Source/MeshModelingTools/Private/DrawPolyPathTool.cpp
@@ -217,6 +217,7 @@ void UDrawPolyPathTool::Setup()
 	// begin path draw
 	InitializeNewSurfacePath();
 
+	SetToolDisplayName(LOCTEXT("ToolName", "Extrude PolyPath"));
 	ShowStartupMessage();
 }
 
diff --git a/Engine/Plugins/Experimental/MeshModelingToolset/Source/MeshModelingTools/Private/DrawPolygonTool.cpp b/Engine/Plugins/Experimental/MeshModelingToolset/Source/MeshModelingTools/Private/DrawPolygonTool.cpp
--- a/Engine/Plugins/Experimental/MeshModelingToolset/Source/MeshModelingTools/Private/DrawPolygonTool.cpp
+++ b/Engine/Plugins/Experimental/MeshModelingToolset/Source/MeshModelingTools/Private/DrawPolygonTool.cpp
@@ -73,6 +73,7 @@ UDrawPolygonTool::UDrawPolygonTool()
 	DrawPlaneOrigin = FVector3d::Zero();
 	DrawPlaneOrientation = FQuaterniond::Identity();
 	bInInteractiveExtrude = false;
+	SetToolDisplayName(LOCTEXT("ToolName", "Extrude Polygon"));
 }
 
 void UDrawPolygonTool::SetWorld(UWorld* World)
diff --git a/Engine/Plugins/Experimental/MeshModelingToolset/Source/MeshModelingTools/Private/DynamicMeshSculptTool.cpp b/Engine/Plugins/Experimental/MeshModelingToolset/Source/MeshModelingTools/Private/DynamicMeshSculptTool.cpp
--- a/Engine/Plugins/Experimental/MeshModelingToolset/Source/MeshModelingTools/Private/DynamicMeshSculptTool.cpp
+++ b/Engine/Plugins/Experimental/MeshModelingToolset/Source/MeshModelingTools/Private/DynamicMeshSculptTool.cpp
@@ -94,6 +94,8 @@ void UDynamicMeshSculptTool::Setup()
 {
 	UMeshSurfacePointTool::Setup();
 
+	SetToolDisplayName(LOCTEXT("ToolName", "DynaSculpt"));
+
 	// create dynamic mesh component to use for live preview
 	DynamicMeshComponent = NewObject<UOctreeDynamicMeshComponent>(ComponentTarget->GetOwnerActor(), "DynamicMeshSculptToolMesh");
 	DynamicMeshComponent->SetupAttachment(ComponentTarget->GetOwnerActor()->GetRootComponent());
diff --git a/Engine/Plugins/Experimental/MeshModelingToolset/Source/MeshModelingTools/Private/EdgeLoopInsertionTool.cpp b/Engine/Plugins/Experimental/MeshModelingToolset/Source/MeshModelingTools/Private/EdgeLoopInsertionTool.cpp
--- a/Engine/Plugins/Experimental/MeshModelingToolset/Source/MeshModelingTools/Private/EdgeLoopInsertionTool.cpp
+++ b/Engine/Plugins/Experimental/MeshModelingToolset/Source/MeshModelingTools/Private/EdgeLoopInsertionTool.cpp
@@ -103,6 +103,7 @@ void UEdgeLoopInsertionTool::Setup()
 		return;
 	}
 
+	SetToolDisplayName(LOCTEXT("ToolName", "Insert PolyLoop"));
 	GetToolManager()->DisplayMessage(
 		LOCTEXT("EdgeLoopInsertionToolDescription", "Click an edge to insert an edge loop passing across that edge. Edge loops follow a sequence of quad-like polygroups."),
 		EToolMessageLevel::UserNotification);
diff --git a/Engine/Plugins/Experimental/MeshModelingToolset/Source/MeshModelingTools/Private/EditMeshPolygonsTool.cpp b/Engine/Plugins/Experimental/MeshModelingToolset/Source/MeshModelingTools/Private/EditMeshPolygonsTool.cpp
--- a/Engine/Plugins/Experimental/MeshModelingToolset/Source/MeshModelingTools/Private/EditMeshPolygonsTool.cpp
+++ b/Engine/Plugins/Experimental/MeshModelingToolset/Source/MeshModelingTools/Private/EditMeshPolygonsTool.cpp
@@ -145,7 +145,7 @@ void UEditMeshPolygonsToolActionPropertySet::PostAction(EEditMeshPolygonsToolAct
 
 UEditMeshPolygonsTool::UEditMeshPolygonsTool()
 {
-	SetToolDisplayName(LOCTEXT("EditMeshPolygonsToolName", "Edit PolyGroups Tool"));
+	SetToolDisplayName(LOCTEXT("EditMeshPolygonsToolName", "PolyGroup Edit"));
 }
 
 void UEditMeshPolygonsTool::EnableTriangleMode()
@@ -338,7 +338,7 @@ void UEditMeshPolygonsTool::Setup()
 
 	if (bTriangleMode)
 	{
-		SetToolDisplayName(LOCTEXT("EditMeshTrianglesToolName", "Edit Triangles Tool"));
+		SetToolDisplayName(LOCTEXT("EditMeshTrianglesToolName", "Triangle Edit"));
 		GetToolManager()->DisplayMessage(
 			LOCTEXT("OnStartEditMeshPolygonsTool_TriangleMode", "Select Triangles to edit mesh. Q to toggle Gizmo Orientation Lock."),
 			EToolMessageLevel::UserNotification);
diff --git a/Engine/Plugins/Experimental/MeshModelingToolset/Source/MeshModelingTools/Private/EditNormalsTool.cpp b/Engine/Plugins/Experimental/MeshModelingToolset/Source/MeshModelingTools/Private/EditNormalsTool.cpp
--- a/Engine/Plugins/Experimental/MeshModelingToolset/Source/MeshModelingTools/Private/EditNormalsTool.cpp
+++ b/Engine/Plugins/Experimental/MeshModelingToolset/Source/MeshModelingTools/Private/EditNormalsTool.cpp
@@ -118,6 +118,7 @@ void UEditNormalsTool::Setup()
 		Preview->InvalidateResult();
 	}
 
+	SetToolDisplayName(LOCTEXT("ToolName", "Edit Normals"));
 	GetToolManager()->DisplayMessage(
 		LOCTEXT("OnStartTool", "Configure or Recalculate Normals on a Mesh (disables autogenerated Normals)"),
 		EToolMessageLevel::UserNotification);
diff --git a/Engine/Plugins/Experimental/MeshModelingToolset/Source/MeshModelingTools/Private/EditPivotTool.cpp b/Engine/Plugins/Experimental/MeshModelingToolset/Source/MeshModelingTools/Private/EditPivotTool.cpp
--- a/Engine/Plugins/Experimental/MeshModelingToolset/Source/MeshModelingTools/Private/EditPivotTool.cpp
+++ b/Engine/Plugins/Experimental/MeshModelingToolset/Source/MeshModelingTools/Private/EditPivotTool.cpp
@@ -116,6 +116,7 @@ void UEditPivotTool::Setup()
 
 	GetToolManager()->DisplayMessage(AllTheWarnings, EToolMessageLevel::UserWarning);
 
+	SetToolDisplayName(LOCTEXT("ToolName", "Edit Pivot"));
 	GetToolManager()->DisplayMessage(
 		LOCTEXT("OnStartTool", "This Tool edits the Pivot (Origin) of the input objects. Use the Gizmo or enable Snap Dragging and click-drag on the surface to reposition."),
 		EToolMessageLevel::UserNotification);
diff --git a/Engine/Plugins/Experimental/MeshModelingToolset/Source/MeshModelingTools/Private/EditUVIslandsTool.cpp b/Engine/Plugins/Experimental/MeshModelingToolset/Source/MeshModelingTools/Private/EditUVIslandsTool.cpp
--- a/Engine/Plugins/Experimental/MeshModelingToolset/Source/MeshModelingTools/Private/EditUVIslandsTool.cpp
+++ b/Engine/Plugins/Experimental/MeshModelingToolset/Source/MeshModelingTools/Private/EditUVIslandsTool.cpp
@@ -114,6 +114,7 @@ void UEditUVIslandsTool::Setup()
 	});
 	OnMaterialSettingsChanged();
 
+	SetToolDisplayName(LOCTEXT("ToolName", "UV Transform"));
 	GetToolManager()->DisplayMessage(LOCTEXT("UEditUVIslandsToolStartupMessage", "Click on a UV Island to select it, and then use the Gizmo to translate/rotate/scale the UVs"), EToolMessageLevel::UserNotification);
 }
 
diff --git a/Engine/Plugins/Experimental/MeshModelingToolset/Source/MeshModelingTools/Private/GroupEdgeInsertionTool.cpp b/Engine/Plugins/Experimental/MeshModelingToolset/Source/MeshModelingTools/Private/GroupEdgeInsertionTool.cpp
--- a/Engine/Plugins/Experimental/MeshModelingToolset/Source/MeshModelingTools/Private/GroupEdgeInsertionTool.cpp
+++ b/Engine/Plugins/Experimental/MeshModelingToolset/Source/MeshModelingTools/Private/GroupEdgeInsertionTool.cpp
@@ -84,6 +84,7 @@ void UGroupEdgeInsertionTool::Setup()
 		return;
 	}
 
+	SetToolDisplayName(LOCTEXT("ToolName", "Insert PolyEdge"));
 	GetToolManager()->DisplayMessage(
 		LOCTEXT("GroupEdgeInsertionToolDescription", "Click two points on the boundary of a face to insert a new edge between the points and split the face."),
 		EToolMessageLevel::UserNotification);
diff --git a/Engine/Plugins/Experimental/MeshModelingToolset/Source/MeshModelingTools/Private/HoleFillTool.cpp b/Engine/Plugins/Experimental/MeshModelingToolset/Source/MeshModelingTools/Private/HoleFillTool.cpp
--- a/Engine/Plugins/Experimental/MeshModelingToolset/Source/MeshModelingTools/Private/HoleFillTool.cpp
+++ b/Engine/Plugins/Experimental/MeshModelingToolset/Source/MeshModelingTools/Private/HoleFillTool.cpp
@@ -226,6 +226,7 @@ void UHoleFillTool::Setup()
 		ComponentTarget->SetOwnerVisibility(false);
 	}
 
+	SetToolDisplayName(LOCTEXT("ToolName", "Fill Holes"));
 	GetToolManager()->DisplayMessage(
 		LOCTEXT("HoleFillToolDescription", "Fill Holes in the selected Mesh by adding triangles. Click on individual holes to fill them, or use the Select All button to fill all holes."),
 		EToolMessageLevel::UserNotification);
diff --git a/Engine/Plugins/Experimental/MeshModelingToolset/Source/MeshModelingTools/Private/LatticeDeformerTool.cpp b/Engine/Plugins/Experimental/MeshModelingToolset/Source/MeshModelingTools/Private/LatticeDeformerTool.cpp
--- a/Engine/Plugins/Experimental/MeshModelingToolset/Source/MeshModelingTools/Private/LatticeDeformerTool.cpp
+++ b/Engine/Plugins/Experimental/MeshModelingToolset/Source/MeshModelingTools/Private/LatticeDeformerTool.cpp
@@ -98,6 +98,7 @@ void ULatticeDeformerTool::Setup()
 {
 	UInteractiveTool::Setup();
 
+	SetToolDisplayName(LOCTEXT("ToolName", "Grid Deform"));
 	GetToolManager()->DisplayMessage(LOCTEXT("LatticeDeformerToolMessage", 
 		"Drag the lattice control points to deform the mesh"), EToolMessageLevel::UserNotification);
 
diff --git a/Engine/Plugins/Experimental/MeshModelingToolset/Source/MeshModelingTools/Private/MeshAttributePaintTool.cpp b/Engine/Plugins/Experimental/MeshModelingToolset/Source/MeshModelingTools/Private/MeshAttributePaintTool.cpp
--- a/Engine/Plugins/Experimental/MeshModelingToolset/Source/MeshModelingTools/Private/MeshAttributePaintTool.cpp
+++ b/Engine/Plugins/Experimental/MeshModelingToolset/Source/MeshModelingTools/Private/MeshAttributePaintTool.cpp
@@ -187,6 +187,7 @@ void UMeshAttributePaintTool::Setup()
 
 	RecalculateBrushRadius();
 
+	SetToolDisplayName(LOCTEXT("ToolName", "Paint WeightMaps"));
 	GetToolManager()->DisplayMessage(
 		LOCTEXT("OnStartAttribPaint", "Paint per-vertex attribute maps. Ctrl to Erase/Subtract, Shift to Smooth"),
 		EToolMessageLevel::UserNotification);
diff --git a/Engine/Plugins/Experimental/MeshModelingToolset/Source/MeshModelingTools/Private/MeshGroupPaintTool.cpp b/Engine/Plugins/Experimental/MeshModelingToolset/Source/MeshModelingTools/Private/MeshGroupPaintTool.cpp
--- a/Engine/Plugins/Experimental/MeshModelingToolset/Source/MeshModelingTools/Private/MeshGroupPaintTool.cpp
+++ b/Engine/Plugins/Experimental/MeshModelingToolset/Source/MeshModelingTools/Private/MeshGroupPaintTool.cpp
@@ -81,6 +81,8 @@ void UMeshGroupPaintTool::Setup()
 {
 	UMeshSculptToolBase::Setup();
 
+	SetToolDisplayName(LOCTEXT("ToolName", "Paint PolyGroups"));
+
 	// create dynamic mesh component to use for live preview
 	DynamicMeshComponent = NewObject<USimpleDynamicMeshComponent>(ComponentTarget->GetOwnerActor());
 	InitializeSculptMeshComponent(DynamicMeshComponent);
diff --git a/Engine/Plugins/Experimental/MeshModelingToolset/Source/MeshModelingTools/Private/MeshInspectorTool.cpp b/Engine/Plugins/Experimental/MeshModelingToolset/Source/MeshModelingTools/Private/MeshInspectorTool.cpp
--- a/Engine/Plugins/Experimental/MeshModelingToolset/Source/MeshModelingTools/Private/MeshInspectorTool.cpp
+++ b/Engine/Plugins/Experimental/MeshModelingToolset/Source/MeshModelingTools/Private/MeshInspectorTool.cpp
@@ -108,6 +108,7 @@ void UMeshInspectorTool::Setup()
 
 	UpdateVisualization();
 
+	SetToolDisplayName(LOCTEXT("ToolName", "Mesh Inspector"));
 	GetToolManager()->DisplayMessage(
 		LOCTEXT("OnStartTool", "Visualize various geometric properties and attributes of the Mesh"),
 		EToolMessageLevel::UserNotification);
diff --git a/Engine/Plugins/Experimental/MeshModelingToolset/Source/MeshModelingTools/Private/MeshSelectionTool.cpp b/Engine/Plugins/Experimental/MeshModelingToolset/Source/MeshModelingTools/Private/MeshSelectionTool.cpp
--- a/Engine/Plugins/Experimental/MeshModelingToolset/Source/MeshModelingTools/Private/MeshSelectionTool.cpp
+++ b/Engine/Plugins/Experimental/MeshModelingToolset/Source/MeshModelingTools/Private/MeshSelectionTool.cpp
@@ -137,6 +137,7 @@ void UMeshSelectionTool::Setup()
 	RecalculateBrushRadius();
 	UpdateVisualization(true);
 
+	SetToolDisplayName(LOCTEXT("ToolName", "Select Triangles"));
 	GetToolManager()->DisplayMessage(
 		LOCTEXT("OnStartMeshSelectionTool", "This Tool allows you to modify the mesh based on a triangle selection. [Q] cyles through Selection Mode. [A] cycles through Face Color modes. [ and ] change brush size, < and > grow/shrink selection."),
 		EToolMessageLevel::UserNotification);
diff --git a/Engine/Plugins/Experimental/MeshModelingToolset/Source/MeshModelingTools/Private/MeshSpaceDeformerTool.cpp b/Engine/Plugins/Experimental/MeshModelingToolset/Source/MeshModelingTools/Private/MeshSpaceDeformerTool.cpp
--- a/Engine/Plugins/Experimental/MeshModelingToolset/Source/MeshModelingTools/Private/MeshSpaceDeformerTool.cpp
+++ b/Engine/Plugins/Experimental/MeshModelingToolset/Source/MeshModelingTools/Private/MeshSpaceDeformerTool.cpp
@@ -319,6 +319,7 @@ void UMeshSpaceDeformerTool::Setup()
 	DragAlignmentMechanic->AddToGizmo(TransformGizmo, nullptr, &ComponentsToInclude);
 	DragAlignmentMechanic->AddToGizmo(IntervalGizmo, nullptr, &ComponentsToInclude);
 
+	SetToolDisplayName(LOCTEXT("ToolName", "Space Warp"));
 	GetToolManager()->DisplayMessage(
 		LOCTEXT("MeshSpaceDeformerToolDescription", "Deform the vertices of the selected Mesh using various spatial deformations. Use the in-viewport Gizmo to control the extents/strength of the deformation."),
 		EToolMessageLevel::UserNotification);
diff --git a/Engine/Plugins/Experimental/MeshModelingToolset/Source/MeshModelingTools/Private/MeshVertexSculptTool.cpp b/Engine/Plugins/Experimental/MeshModelingToolset/Source/MeshModelingTools/Private/MeshVertexSculptTool.cpp
--- a/Engine/Plugins/Experimental/MeshModelingToolset/Source/MeshModelingTools/Private/MeshVertexSculptTool.cpp
+++ b/Engine/Plugins/Experimental/MeshModelingToolset/Source/MeshModelingTools/Private/MeshVertexSculptTool.cpp
@@ -55,6 +55,8 @@ void UMeshVertexSculptTool::Setup()
 {
 	UMeshSculptToolBase::Setup();
 
+	SetToolDisplayName(LOCTEXT("ToolName", "Sculpt"));
+
 	// create dynamic mesh component to use for live preview
 	DynamicMeshComponent = NewObject<USimpleDynamicMeshComponent>(ComponentTarget->GetOwnerActor());
 	InitializeSculptMeshComponent(DynamicMeshComponent);
diff --git a/Engine/Plugins/Experimental/MeshModelingToolset/Source/MeshModelingTools/Private/MirrorTool.cpp b/Engine/Plugins/Experimental/MeshModelingToolset/Source/MeshModelingTools/Private/MirrorTool.cpp
--- a/Engine/Plugins/Experimental/MeshModelingToolset/Source/MeshModelingTools/Private/MirrorTool.cpp
+++ b/Engine/Plugins/Experimental/MeshModelingToolset/Source/MeshModelingTools/Private/MirrorTool.cpp
@@ -185,6 +185,7 @@ void UMirrorTool::Setup()
 {
 	UInteractiveTool::Setup();
 
+	SetToolDisplayName(LOCTEXT("ToolName", "Mirror"));
 	GetToolManager()->DisplayMessage(
 		LOCTEXT("OnStartMirrorTool", "Mirror one or more meshes across a plane. Grid snapping behavior is swapped while the shift key is down. The plane can be set by using the preset buttons, moving the gizmo, or ctrl+clicking on a spot on the original mesh."),
 		EToolMessageLevel::UserNotification);
diff --git a/Engine/Plugins/Experimental/MeshModelingToolset/Source/MeshModelingTools/Private/OffsetMeshTool.cpp b/Engine/Plugins/Experimental/MeshModelingToolset/Source/MeshModelingTools/Private/OffsetMeshTool.cpp
--- a/Engine/Plugins/Experimental/MeshModelingToolset/Source/MeshModelingTools/Private/OffsetMeshTool.cpp
+++ b/Engine/Plugins/Experimental/MeshModelingToolset/Source/MeshModelingTools/Private/OffsetMeshTool.cpp
@@ -16,6 +16,10 @@
  * Tool
  */
 
+UOffsetMeshTool::UOffsetMeshTool()
+{
+	SetToolDisplayName(LOCTEXT("ToolName", "Offset"));
+}
 
 void UOffsetMeshTool::InitializeProperties()
 {
@@ -50,7 +54,7 @@ void UOffsetMeshTool::OnShutdown(EToolShutdownType ShutdownType)
 
 FText UOffsetMeshTool::GetToolMessageString() const
 {
-	return LOCTEXT("StartOffsetToolMessage", "Offsets the mesh vertex positions using various Offseting methods.");
+	return LOCTEXT("StartOffsetToolMessage", "Offset the mesh vertex positions using various Offseting methods.");
 }
 
 FText UOffsetMeshTool::GetAcceptTransactionName() const
diff --git a/Engine/Plugins/Experimental/MeshModelingToolset/Source/MeshModelingTools/Private/Physics/ExtractCollisionGeometryTool.cpp b/Engine/Plugins/Experimental/MeshModelingToolset/Source/MeshModelingTools/Private/Physics/ExtractCollisionGeometryTool.cpp
--- a/Engine/Plugins/Experimental/MeshModelingToolset/Source/MeshModelingTools/Private/Physics/ExtractCollisionGeometryTool.cpp
+++ b/Engine/Plugins/Experimental/MeshModelingToolset/Source/MeshModelingTools/Private/Physics/ExtractCollisionGeometryTool.cpp
@@ -103,6 +103,7 @@ void UExtractCollisionGeometryTool::Setup()
 	}
 
 
+	SetToolDisplayName(LOCTEXT("ToolName", "Collision To Mesh"));
 	GetToolManager()->DisplayMessage(
 		LOCTEXT("OnStartTool", "Convert Collision geometry to Static Mesh"),
 		EToolMessageLevel::UserNotification);
diff --git a/Engine/Plugins/Experimental/MeshModelingToolset/Source/MeshModelingTools/Private/Physics/PhysicsInspectorTool.cpp b/Engine/Plugins/Experimental/MeshModelingToolset/Source/MeshModelingTools/Private/Physics/PhysicsInspectorTool.cpp
--- a/Engine/Plugins/Experimental/MeshModelingToolset/Source/MeshModelingTools/Private/Physics/PhysicsInspectorTool.cpp
+++ b/Engine/Plugins/Experimental/MeshModelingToolset/Source/MeshModelingTools/Private/Physics/PhysicsInspectorTool.cpp
@@ -90,7 +90,7 @@ void UPhysicsInspectorTool::Setup()
 		}
 	}
 
-
+	SetToolDisplayName(LOCTEXT("ToolName", "Physics Inspector"));
 	GetToolManager()->DisplayMessage(
 		LOCTEXT("OnStartTool", "Inspect Physics data for the seleced Static Meshes"),
 		EToolMessageLevel::UserNotification);
diff --git a/Engine/Plugins/Experimental/MeshModelingToolset/Source/MeshModelingTools/Private/Physics/SetCollisionGeometryTool.cpp b/Engine/Plugins/Experimental/MeshModelingToolset/Source/MeshModelingTools/Private/Physics/SetCollisionGeometryTool.cpp
--- a/Engine/Plugins/Experimental/MeshModelingToolset/Source/MeshModelingTools/Private/Physics/SetCollisionGeometryTool.cpp
+++ b/Engine/Plugins/Experimental/MeshModelingToolset/Source/MeshModelingTools/Private/Physics/SetCollisionGeometryTool.cpp
@@ -122,6 +122,7 @@ void USetCollisionGeometryTool::Setup()
 	CollisionProps = NewObject<UPhysicsObjectToolPropertySet>(this);
 	AddToolPropertySource(CollisionProps);
 
+	SetToolDisplayName(LOCTEXT("ToolName", "Mesh To Collision"));
 	GetToolManager()->DisplayMessage(
 		LOCTEXT("OnStartTool", "Initialize Simple Collision geometry for a Mesh from one or more input Meshes (including itself)."),
 		EToolMessageLevel::UserNotification);
diff --git a/Engine/Plugins/Experimental/MeshModelingToolset/Source/MeshModelingTools/Private/PlaneCutTool.cpp b/Engine/Plugins/Experimental/MeshModelingToolset/Source/MeshModelingTools/Private/PlaneCutTool.cpp
--- a/Engine/Plugins/Experimental/MeshModelingToolset/Source/MeshModelingTools/Private/PlaneCutTool.cpp
+++ b/Engine/Plugins/Experimental/MeshModelingToolset/Source/MeshModelingTools/Private/PlaneCutTool.cpp
@@ -215,6 +215,7 @@ void UPlaneCutTool::Setup()
 		Preview->InvalidateResult();
 	}
 
+	SetToolDisplayName(LOCTEXT("ToolName", "Plane Cut"));
 	GetToolManager()->DisplayMessage(
 		LOCTEXT("OnStartPlaneCutTool", "Press 'A' or use the Cut button to cut the mesh without leaving the tool.  When grid snapping is enabled, you can toggle snapping with the shift key."),
 		EToolMessageLevel::UserNotification);
diff --git a/Engine/Plugins/Experimental/MeshModelingToolset/Source/MeshModelingTools/Private/ProjectToTargetTool.cpp b/Engine/Plugins/Experimental/MeshModelingToolset/Source/MeshModelingTools/Private/ProjectToTargetTool.cpp
--- a/Engine/Plugins/Experimental/MeshModelingToolset/Source/MeshModelingTools/Private/ProjectToTargetTool.cpp
+++ b/Engine/Plugins/Experimental/MeshModelingToolset/Source/MeshModelingTools/Private/ProjectToTargetTool.cpp
@@ -48,7 +48,7 @@ void UProjectToTargetTool::Setup()
 	// Now setup parent RemeshMeshTool class
 	URemeshMeshTool::Setup();
 
-	SetToolDisplayName(LOCTEXT("ProjectToTargetToolName", "Remesh To Target Tool"));
+	SetToolDisplayName(LOCTEXT("ProjectToTargetToolName", "Project To Target"));
 	GetToolManager()->DisplayMessage(
 		LOCTEXT("ProjectToTargetToolDescription", "Incrementally deform the first selected mesh towards the second, while applying Remeshing. This can be used to improve the accuracy of shrink-wrapping strategies."),
 		EToolMessageLevel::UserNotification);
diff --git a/Engine/Plugins/Experimental/MeshModelingToolset/Source/MeshModelingTools/Private/RemeshMeshTool.cpp b/Engine/Plugins/Experimental/MeshModelingToolset/Source/MeshModelingTools/Private/RemeshMeshTool.cpp
--- a/Engine/Plugins/Experimental/MeshModelingToolset/Source/MeshModelingTools/Private/RemeshMeshTool.cpp
+++ b/Engine/Plugins/Experimental/MeshModelingToolset/Source/MeshModelingTools/Private/RemeshMeshTool.cpp
@@ -151,6 +151,7 @@ void URemeshMeshTool::Setup()
 
 	Preview->InvalidateResult();
 
+	SetToolDisplayName(LOCTEXT("ToolName", "Remesh"));
 	GetToolManager()->DisplayMessage(
 		LOCTEXT("OnStartTool", "Retriangulate the selected Mesh. Use the Boundary Constraints to preserve mesh borders. Enable Discard Attributes to ignore UV/Normal Seams. "),
 		EToolMessageLevel::UserNotification);
diff --git a/Engine/Plugins/Experimental/MeshModelingToolset/Source/MeshModelingTools/Private/RemoveOccludedTrianglesTool.cpp b/Engine/Plugins/Experimental/MeshModelingToolset/Source/MeshModelingTools/Private/RemoveOccludedTrianglesTool.cpp
--- a/Engine/Plugins/Experimental/MeshModelingToolset/Source/MeshModelingTools/Private/RemoveOccludedTrianglesTool.cpp
+++ b/Engine/Plugins/Experimental/MeshModelingToolset/Source/MeshModelingTools/Private/RemoveOccludedTrianglesTool.cpp
@@ -83,7 +83,7 @@ URemoveOccludedTrianglesAdvancedProperties::URemoveOccludedTrianglesAdvancedProp
 
 URemoveOccludedTrianglesTool::URemoveOccludedTrianglesTool()
 {
-	SetToolDisplayName(LOCTEXT("ProjectToTargetToolName", "Jacketing Tool"));
+	SetToolDisplayName(LOCTEXT("ProjectToTargetToolName", "Jacket"));
 }
 
 void URemoveOccludedTrianglesTool::SetWorld(UWorld* World)
diff --git a/Engine/Plugins/Experimental/MeshModelingToolset/Source/MeshModelingTools/Private/RevolveBoundaryTool.cpp b/Engine/Plugins/Experimental/MeshModelingToolset/Source/MeshModelingTools/Private/RevolveBoundaryTool.cpp
--- a/Engine/Plugins/Experimental/MeshModelingToolset/Source/MeshModelingTools/Private/RevolveBoundaryTool.cpp
+++ b/Engine/Plugins/Experimental/MeshModelingToolset/Source/MeshModelingTools/Private/RevolveBoundaryTool.cpp
@@ -122,8 +122,9 @@ void URevolveBoundaryTool::Setup()
 
 	ComponentTarget->SetOwnerVisibility(Settings->bDisplayOriginalMesh);
 
+	SetToolDisplayName(LOCTEXT("ToolName", "Revolve Boundary"));
 	GetToolManager()->DisplayMessage(
-		LOCTEXT("OnStartRevolveBoundaryTool", "This tool revolves the mesh boundary around the axis to create a new mesh. Ctrl+click will reposition the revolution axis, potentially aligning it with an edge."),
+		LOCTEXT("OnStartRevolveBoundaryTool", "Revolve an open mesh boundary loop around an axis to create a new mesh. Ctrl+click will reposition the revolution axis, potentially aligning it with an edge."),
 		EToolMessageLevel::UserNotification);
 	if (Topology->Edges.Num() == 1)
 	{
diff --git a/Engine/Plugins/Experimental/MeshModelingToolset/Source/MeshModelingTools/Private/SeamSculptTool.cpp b/Engine/Plugins/Experimental/MeshModelingToolset/Source/MeshModelingTools/Private/SeamSculptTool.cpp
--- a/Engine/Plugins/Experimental/MeshModelingToolset/Source/MeshModelingTools/Private/SeamSculptTool.cpp
+++ b/Engine/Plugins/Experimental/MeshModelingToolset/Source/MeshModelingTools/Private/SeamSculptTool.cpp
@@ -49,7 +49,7 @@ void USeamSculptTool::SetWorld(UWorld* World)
 
 USeamSculptTool::USeamSculptTool()
 {
-	SetToolDisplayName(LOCTEXT("SeamSculptToolName", "UV Seam Tool"));
+	SetToolDisplayName(LOCTEXT("SeamSculptToolName", "Add UV Seams"));
 }
 
 void USeamSculptTool::Setup()
diff --git a/Engine/Plugins/Experimental/MeshModelingToolset/Source/MeshModelingTools/Private/SelfUnionMeshesTool.cpp b/Engine/Plugins/Experimental/MeshModelingToolset/Source/MeshModelingTools/Private/SelfUnionMeshesTool.cpp
--- a/Engine/Plugins/Experimental/MeshModelingToolset/Source/MeshModelingTools/Private/SelfUnionMeshesTool.cpp
+++ b/Engine/Plugins/Experimental/MeshModelingToolset/Source/MeshModelingTools/Private/SelfUnionMeshesTool.cpp
@@ -26,6 +26,7 @@ void USelfUnionMeshesTool::SetupProperties()
 	Properties->RestoreProperties(this);
 	AddToolPropertySource(Properties);
 
+	SetToolDisplayName(LOCTEXT("ToolName", "Merge"));
 	GetToolManager()->DisplayMessage(
 		LOCTEXT("OnStartTool", "Compute a Self-Union of the input meshes, to resolve self-intersections. Use the transform gizmos to tweak the positions of the input objects (can help to resolve errors/failures)"),
 		EToolMessageLevel::UserNotification);
diff --git a/Engine/Plugins/Experimental/MeshModelingToolset/Source/MeshModelingTools/Private/SmoothMeshTool.cpp b/Engine/Plugins/Experimental/MeshModelingToolset/Source/MeshModelingTools/Private/SmoothMeshTool.cpp
--- a/Engine/Plugins/Experimental/MeshModelingToolset/Source/MeshModelingTools/Private/SmoothMeshTool.cpp
+++ b/Engine/Plugins/Experimental/MeshModelingToolset/Source/MeshModelingTools/Private/SmoothMeshTool.cpp
@@ -18,6 +18,11 @@
  * Tool
  */
 
+USmoothMeshTool::USmoothMeshTool()
+{
+	SetToolDisplayName(LOCTEXT("ToolName", "Smooth"));
+
+}
 
 void USmoothMeshTool::InitializeProperties()
 {
@@ -51,7 +56,7 @@ void USmoothMeshTool::OnShutdown(EToolShutdownType ShutdownType)
 
 FText USmoothMeshTool::GetToolMessageString() const
 {
-	return LOCTEXT("StartSmoothToolMessage", "Smooths the mesh vertex positions using various smoothing methods.");
+	return LOCTEXT("StartSmoothToolMessage", "Smooth the mesh vertex positions using various smoothing methods.");
 }
 
 FText USmoothMeshTool::GetAcceptTransactionName() const
diff --git a/Engine/Plugins/Experimental/MeshModelingToolset/Source/MeshModelingTools/Private/TransformMeshesTool.cpp b/Engine/Plugins/Experimental/MeshModelingToolset/Source/MeshModelingTools/Private/TransformMeshesTool.cpp
--- a/Engine/Plugins/Experimental/MeshModelingToolset/Source/MeshModelingTools/Private/TransformMeshesTool.cpp
+++ b/Engine/Plugins/Experimental/MeshModelingToolset/Source/MeshModelingTools/Private/TransformMeshesTool.cpp
@@ -86,6 +86,7 @@ void UTransformMeshesTool::Setup()
 
 	UpdateTransformMode(TransformProps->TransformMode);
 
+	SetToolDisplayName(LOCTEXT("ToolName", "Transform"));
 	GetToolManager()->DisplayMessage(
 		LOCTEXT("OnStartTransformMeshesTool", "Translate/Rotate/Scale the selected objects. [A] cycles through Transform modes. Reposition the Gizmo using Set Pivot Mode [S]. Toggle Snap Drag Mode [D] to move the object by click-draging. [W] and [E] cycle through SnapDrag Source and Rotation types."),
 		EToolMessageLevel::UserNotification);
diff --git a/Engine/Plugins/Experimental/MeshModelingToolset/Source/MeshModelingTools/Private/UVProjectionTool.cpp b/Engine/Plugins/Experimental/MeshModelingToolset/Source/MeshModelingTools/Private/UVProjectionTool.cpp
--- a/Engine/Plugins/Experimental/MeshModelingToolset/Source/MeshModelingTools/Private/UVProjectionTool.cpp
+++ b/Engine/Plugins/Experimental/MeshModelingToolset/Source/MeshModelingTools/Private/UVProjectionTool.cpp
@@ -128,6 +128,7 @@ void UUVProjectionTool::Setup()
 	}
 	UpdateVisualization();
 
+	SetToolDisplayName(LOCTEXT("ToolName", "UV Projection"));
 	GetToolManager()->DisplayMessage(
 		LOCTEXT("UVProjectionToolDescription", "Generate UVs for a Mesh by projecting onto simple geometric shapes."),
 		EToolMessageLevel::UserNotification);
diff --git a/Engine/Plugins/Experimental/MeshModelingToolset/Source/MeshModelingTools/Private/VoxelBlendMeshesTool.cpp b/Engine/Plugins/Experimental/MeshModelingToolset/Source/MeshModelingTools/Private/VoxelBlendMeshesTool.cpp
--- a/Engine/Plugins/Experimental/MeshModelingToolset/Source/MeshModelingTools/Private/VoxelBlendMeshesTool.cpp
+++ b/Engine/Plugins/Experimental/MeshModelingToolset/Source/MeshModelingTools/Private/VoxelBlendMeshesTool.cpp
@@ -39,7 +39,7 @@ void UVoxelBlendMeshesTool::SetupProperties()
 	BlendProperties->RestoreProperties(this);
 	AddToolPropertySource(BlendProperties);
 
-	SetToolDisplayName(LOCTEXT("VoxelBlendMeshesToolName", "Blend Meshes Tool"));
+	SetToolDisplayName(LOCTEXT("VoxelBlendMeshesToolName", "Voxel Blend"));
 	GetToolManager()->DisplayMessage(
 		LOCTEXT("VoxelBlendMeshesToolDescription", "Compute a volumetric Blend of the input meshes, controlled by the Blend Power/Falloff. UVs, sharp edges, and small/thin features will be lost. Increase Voxel Count to enhance accuracy."),
 		EToolMessageLevel::UserNotification);
diff --git a/Engine/Plugins/Experimental/MeshModelingToolset/Source/MeshModelingTools/Private/VoxelMorphologyMeshesTool.cpp b/Engine/Plugins/Experimental/MeshModelingToolset/Source/MeshModelingTools/Private/VoxelMorphologyMeshesTool.cpp
--- a/Engine/Plugins/Experimental/MeshModelingToolset/Source/MeshModelingTools/Private/VoxelMorphologyMeshesTool.cpp
+++ b/Engine/Plugins/Experimental/MeshModelingToolset/Source/MeshModelingTools/Private/VoxelMorphologyMeshesTool.cpp
@@ -37,7 +37,7 @@ void UVoxelMorphologyMeshesTool::SetupProperties()
 	MorphologyProperties->RestoreProperties(this);
 	AddToolPropertySource(MorphologyProperties);
 
-	SetToolDisplayName(LOCTEXT("VoxelMorphologyMeshesToolName", "Mesh Morphology Tool"));
+	SetToolDisplayName(LOCTEXT("VoxelMorphologyMeshesToolName", "Voxel Morphology"));
 	GetToolManager()->DisplayMessage(
 		LOCTEXT("OnStartTool", "Apply Morphological operations to the input meshes to create a new Mesh, using voxelization techniques. UVs, sharp edges, and small/thin features will be lost. Increase Voxel Count to enhance accuracy."),
 		EToolMessageLevel::UserNotification);
diff --git a/Engine/Plugins/Experimental/MeshModelingToolset/Source/MeshModelingTools/Private/VoxelSolidifyMeshesTool.cpp b/Engine/Plugins/Experimental/MeshModelingToolset/Source/MeshModelingTools/Private/VoxelSolidifyMeshesTool.cpp
--- a/Engine/Plugins/Experimental/MeshModelingToolset/Source/MeshModelingTools/Private/VoxelSolidifyMeshesTool.cpp
+++ b/Engine/Plugins/Experimental/MeshModelingToolset/Source/MeshModelingTools/Private/VoxelSolidifyMeshesTool.cpp
@@ -34,7 +34,7 @@ void UVoxelSolidifyMeshesTool::SetupProperties()
 	SolidifyProperties->RestoreProperties(this);
 	AddToolPropertySource(SolidifyProperties);
 
-	SetToolDisplayName(LOCTEXT("VoxelSolidifyMeshesToolName", "Wrap Meshes Tool"));
+	SetToolDisplayName(LOCTEXT("VoxelSolidifyMeshesToolName", "Voxel Wrap"));
 	GetToolManager()->DisplayMessage(
 		LOCTEXT("VoxelSolidifyMeshesToolDescription", "Create a new closed/solid shell mesh that wraps the input meshes. Holes will automatically be filled, controlled by the Winding Threshold. UVs, sharp edges, and small/thin features will be lost. Increase Voxel Count to enhance accuracy."),
 		EToolMessageLevel::UserNotification);
diff --git a/Engine/Plugins/Experimental/MeshModelingToolset/Source/MeshModelingTools/Private/WeldMeshEdgesTool.cpp b/Engine/Plugins/Experimental/MeshModelingToolset/Source/MeshModelingTools/Private/WeldMeshEdgesTool.cpp
--- a/Engine/Plugins/Experimental/MeshModelingToolset/Source/MeshModelingTools/Private/WeldMeshEdgesTool.cpp
+++ b/Engine/Plugins/Experimental/MeshModelingToolset/Source/MeshModelingTools/Private/WeldMeshEdgesTool.cpp
@@ -81,6 +81,7 @@ void UWeldMeshEdgesTool::Setup()
 
 	bResultValid = false;
 
+	SetToolDisplayName(LOCTEXT("ToolName", "Weld Edges"));
 	GetToolManager()->DisplayMessage(
 		LOCTEXT("WeldMeshEdgesToolDescription", "Weld overlapping/identical border edges of the selected Mesh, by merging the vertices."),
 		EToolMessageLevel::UserNotification);
diff --git a/Engine/Plugins/Experimental/MeshModelingToolset/Source/MeshModelingTools/Public/AddPrimitiveTool.h b/Engine/Plugins/Experimental/MeshModelingToolset/Source/MeshModelingTools/Public/AddPrimitiveTool.h
--- a/Engine/Plugins/Experimental/MeshModelingToolset/Source/MeshModelingTools/Public/AddPrimitiveTool.h
+++ b/Engine/Plugins/Experimental/MeshModelingToolset/Source/MeshModelingTools/Public/AddPrimitiveTool.h
@@ -91,7 +91,7 @@ public:
 
 	/** If the shape settings haven't changed, create instances of the last created asset rather than creating a whole new asset.  If false, all created actors will have separate underlying mesh assets. */
 	UPROPERTY(EditAnywhere, Category = AssetSettings)
-	bool bInstanceIfPossible = true;
+	bool bInstanceIfPossible = false;
 
 	/** How should Polygroups be assigned to triangles of Primitive */
 	UPROPERTY(EditAnywhere, Category = ShapeSettings, meta = (ProceduralShapeSetting))
@@ -445,16 +445,13 @@ protected:
 	void UpdatePreviewMesh();
 };
 
+
 UCLASS()
 class UAddBoxPrimitiveTool : public UAddPrimitiveTool
 {
 	GENERATED_BODY()
 public:
-	UAddBoxPrimitiveTool(const FObjectInitializer& ObjectInitializer)
-	: Super(ObjectInitializer.SetDefaultSubobjectClass<UProceduralBoxToolProperties>(TEXT("ShapeSettings")))
-	{
-		AssetName = TEXT("Box");
-	}
+	UAddBoxPrimitiveTool(const FObjectInitializer& ObjectInitializer);
 protected:
 	void GenerateMesh(FDynamicMesh3* OutMesh) const override;
 };
@@ -464,11 +461,7 @@ class UAddCylinderPrimitiveTool : public UAddPrimitiveTool
 {
 	GENERATED_BODY()
 public:
-	UAddCylinderPrimitiveTool(const FObjectInitializer& ObjectInitializer)
-	: Super(ObjectInitializer.SetDefaultSubobjectClass<UProceduralCylinderToolProperties>(TEXT("ShapeSettings")))
-	{
-		AssetName = TEXT("Cylinder");
-	}
+	UAddCylinderPrimitiveTool(const FObjectInitializer& ObjectInitializer);
 protected:
 	void GenerateMesh(FDynamicMesh3* OutMesh) const override;
 };
@@ -478,11 +471,7 @@ class UAddConePrimitiveTool : public UAddPrimitiveTool
 {
 	GENERATED_BODY()
 public:
-	UAddConePrimitiveTool(const FObjectInitializer& ObjectInitializer)
-	: Super(ObjectInitializer.SetDefaultSubobjectClass<UProceduralConeToolProperties>(TEXT("ShapeSettings")))
-	{
-		AssetName = TEXT("Cone");
-	}
+	UAddConePrimitiveTool(const FObjectInitializer& ObjectInitializer);
 protected:
 	void GenerateMesh(FDynamicMesh3* OutMesh) const override;
 };
@@ -492,11 +481,7 @@ class UAddRectanglePrimitiveTool : public UAddPrimitiveTool
 {
 	GENERATED_BODY()
 public:
-	UAddRectanglePrimitiveTool(const FObjectInitializer& ObjectInitializer)
-	: Super(ObjectInitializer.SetDefaultSubobjectClass<UProceduralRectangleToolProperties>(TEXT("ShapeSettings")))
-	{
-		AssetName = TEXT("Rectangle");
-	}
+	UAddRectanglePrimitiveTool(const FObjectInitializer& ObjectInitializer);
 protected:
 	void GenerateMesh(FDynamicMesh3* OutMesh) const override;
 };
@@ -506,11 +491,7 @@ class UAddRoundedRectanglePrimitiveTool : public UAddPrimitiveTool
 {
 	GENERATED_BODY()
 public:
-	UAddRoundedRectanglePrimitiveTool(const FObjectInitializer& ObjectInitializer)
-	: Super(ObjectInitializer.SetDefaultSubobjectClass<UProceduralRoundedRectangleToolProperties>(TEXT("ShapeSettings")))
-	{
-		AssetName = TEXT("RoundedRectangle");
-	}
+	UAddRoundedRectanglePrimitiveTool(const FObjectInitializer& ObjectInitializer);
 protected:
 	void GenerateMesh(FDynamicMesh3* OutMesh) const override;
 };
@@ -520,11 +501,7 @@ class UAddDiscPrimitiveTool : public UAddPrimitiveTool
 {
 	GENERATED_BODY()
 public:
-	UAddDiscPrimitiveTool(const FObjectInitializer& ObjectInitializer)
-	: Super(ObjectInitializer.SetDefaultSubobjectClass<UProceduralDiscToolProperties>(TEXT("ShapeSettings")))
-	{
-		AssetName = TEXT("Disc");
-	}
+	UAddDiscPrimitiveTool(const FObjectInitializer& ObjectInitializer);
 protected:
 	void GenerateMesh(FDynamicMesh3* OutMesh) const override;
 };
@@ -534,11 +511,7 @@ class UAddPuncturedDiscPrimitiveTool : public UAddPrimitiveTool
 {
 	GENERATED_BODY()
 public:
-	UAddPuncturedDiscPrimitiveTool(const FObjectInitializer& ObjectInitializer)
-	: Super(ObjectInitializer.SetDefaultSubobjectClass<UProceduralPuncturedDiscToolProperties>(TEXT("ShapeSettings")))
-	{
-		AssetName = TEXT("PuncturedDisc");
-	}
+	UAddPuncturedDiscPrimitiveTool(const FObjectInitializer& ObjectInitializer);
 protected:
 	void GenerateMesh(FDynamicMesh3* OutMesh) const override;
 };
@@ -548,11 +521,7 @@ class UAddTorusPrimitiveTool : public UAddPrimitiveTool
 {
 	GENERATED_BODY()
 public:
-	UAddTorusPrimitiveTool(const FObjectInitializer& ObjectInitializer)
-	: Super(ObjectInitializer.SetDefaultSubobjectClass<UProceduralTorusToolProperties>(TEXT("ShapeSettings")))
-	{
-		AssetName = TEXT("Torus");
-	}
+	UAddTorusPrimitiveTool(const FObjectInitializer& ObjectInitializer);
 protected:
 	void GenerateMesh(FDynamicMesh3* OutMesh) const override;
 };
@@ -562,11 +531,7 @@ class UAddArrowPrimitiveTool : public UAddPrimitiveTool
 {
 	GENERATED_BODY()
 public:
-	UAddArrowPrimitiveTool(const FObjectInitializer& ObjectInitializer)
-	: Super(ObjectInitializer.SetDefaultSubobjectClass<UProceduralArrowToolProperties>(TEXT("ShapeSettings")))
-	{
-		AssetName = TEXT("Arrow");
-	}
+	UAddArrowPrimitiveTool(const FObjectInitializer& ObjectInitializer);
 protected:
 	void GenerateMesh(FDynamicMesh3* OutMesh) const override;
 };
@@ -576,11 +541,7 @@ class UAddSpherePrimitiveTool : public UAddPrimitiveTool
 {
 	GENERATED_BODY()
 public:
-	UAddSpherePrimitiveTool(const FObjectInitializer& ObjectInitializer)
-	: Super(ObjectInitializer.SetDefaultSubobjectClass<UProceduralSphereToolProperties>(TEXT("ShapeSettings")))
-	{
-		AssetName = TEXT("Sphere Type 2");
-	}
+	UAddSpherePrimitiveTool(const FObjectInitializer& ObjectInitializer);
 protected:
 	void GenerateMesh(FDynamicMesh3* OutMesh) const override;
 };
@@ -590,11 +551,8 @@ class UAddSphericalBoxPrimitiveTool : public UAddPrimitiveTool
 {
 	GENERATED_BODY()
 public:
-	UAddSphericalBoxPrimitiveTool(const FObjectInitializer& ObjectInitializer)
-	: Super(ObjectInitializer.SetDefaultSubobjectClass<UProceduralSphericalBoxToolProperties>(TEXT("ShapeSettings")))
-	{
-		AssetName = TEXT("Spherical Type 1");
-	}
+	UAddSphericalBoxPrimitiveTool(const FObjectInitializer& ObjectInitializer);
 protected:
 	void GenerateMesh(FDynamicMesh3* OutMesh) const override;
 };
+
diff --git a/Engine/Plugins/Experimental/MeshModelingToolset/Source/MeshModelingTools/Public/OffsetMeshTool.h b/Engine/Plugins/Experimental/MeshModelingToolset/Source/MeshModelingTools/Public/OffsetMeshTool.h
--- a/Engine/Plugins/Experimental/MeshModelingToolset/Source/MeshModelingTools/Public/OffsetMeshTool.h
+++ b/Engine/Plugins/Experimental/MeshModelingToolset/Source/MeshModelingTools/Public/OffsetMeshTool.h
@@ -109,7 +109,7 @@ class MESHMODELINGTOOLS_API UOffsetMeshTool : public UBaseMeshProcessingTool
 	GENERATED_BODY()
 
 public:
-	UOffsetMeshTool() = default;
+	UOffsetMeshTool();
 
 	virtual void InitializeProperties() override;
 	virtual void OnShutdown(EToolShutdownType ShutdownType) override;
diff --git a/Engine/Plugins/Experimental/MeshModelingToolset/Source/MeshModelingTools/Public/SmoothMeshTool.h b/Engine/Plugins/Experimental/MeshModelingToolset/Source/MeshModelingTools/Public/SmoothMeshTool.h
--- a/Engine/Plugins/Experimental/MeshModelingToolset/Source/MeshModelingTools/Public/SmoothMeshTool.h
+++ b/Engine/Plugins/Experimental/MeshModelingToolset/Source/MeshModelingTools/Public/SmoothMeshTool.h
@@ -130,7 +130,7 @@ class MESHMODELINGTOOLS_API USmoothMeshTool : public UBaseMeshProcessingTool
 	GENERATED_BODY()
 
 public:
-	USmoothMeshTool() = default;
+	USmoothMeshTool();
 
 	virtual void InitializeProperties() override;
 	virtual void OnShutdown(EToolShutdownType ShutdownType) override;
diff --git a/Engine/Plugins/Experimental/MeshModelingToolset/Source/MeshModelingToolsEditorOnly/Private/AttributeEditorTool.cpp b/Engine/Plugins/Experimental/MeshModelingToolset/Source/MeshModelingToolsEditorOnly/Private/AttributeEditorTool.cpp
--- a/Engine/Plugins/Experimental/MeshModelingToolset/Source/MeshModelingToolsEditorOnly/Private/AttributeEditorTool.cpp
+++ b/Engine/Plugins/Experimental/MeshModelingToolset/Source/MeshModelingToolsEditorOnly/Private/AttributeEditorTool.cpp
@@ -141,6 +141,7 @@ void UAttributeEditorTool::Setup()
 		InitializeAttributeLists();
 	}
 
+	SetToolDisplayName(LOCTEXT("ToolName", "Edit Attributes"));
 	GetToolManager()->DisplayMessage(
 		LOCTEXT("OnStartAttribEditor", "Inspect and Modify Attributes of a StaticMesh Asset"),
 		EToolMessageLevel::UserNotification);
diff --git a/Engine/Plugins/Experimental/MeshModelingToolset/Source/MeshModelingToolsEditorOnly/Private/BspConversionTool.cpp b/Engine/Plugins/Experimental/MeshModelingToolset/Source/MeshModelingToolsEditorOnly/Private/BspConversionTool.cpp
--- a/Engine/Plugins/Experimental/MeshModelingToolset/Source/MeshModelingToolsEditorOnly/Private/BspConversionTool.cpp
+++ b/Engine/Plugins/Experimental/MeshModelingToolset/Source/MeshModelingToolsEditorOnly/Private/BspConversionTool.cpp
@@ -103,6 +103,7 @@ void UBspConversionTool::Setup()
 	ToolActions->Initialize(this);
 	AddToolPropertySource(ToolActions);
 
+	SetToolDisplayName(LOCTEXT("ToolName", "Convert BSP"));
 	// Give a description to put in the side panel
 	GetToolManager()->DisplayMessage(
 		LOCTEXT("OnStartTool", "Convert geometry brushes (also known as BSP brushes) into a single static mesh. \"Convert then Combine\" first converts the individual brushes and performs boolean operations on the resulting meshes (this requires individual brushes to have manifold geometry. \"Combine then Convert\" is the old functionality, combining brushes, then converting them."),
diff --git a/Engine/Plugins/Experimental/MeshModelingToolset/Source/MeshModelingToolsEditorOnly/Private/EditMeshMaterialsTool.cpp b/Engine/Plugins/Experimental/MeshModelingToolset/Source/MeshModelingToolsEditorOnly/Private/EditMeshMaterialsTool.cpp
--- a/Engine/Plugins/Experimental/MeshModelingToolset/Source/MeshModelingToolsEditorOnly/Private/EditMeshMaterialsTool.cpp
+++ b/Engine/Plugins/Experimental/MeshModelingToolset/Source/MeshModelingToolsEditorOnly/Private/EditMeshMaterialsTool.cpp
@@ -47,6 +47,8 @@ void UEditMeshMaterialsTool::Setup()
 {
 	UMeshSelectionTool::Setup();
 
+	SetToolDisplayName(LOCTEXT("ToolName", "Edit Materials"));
+
 	PreviewMesh->ClearOverrideRenderMaterial();
 
 	FComponentMaterialSet AssetMaterials;
diff --git a/Engine/Plugins/Experimental/MeshModelingToolset/Source/MeshModelingToolsEditorOnly/Private/MergeMeshesTool.cpp b/Engine/Plugins/Experimental/MeshModelingToolset/Source/MeshModelingToolsEditorOnly/Private/MergeMeshesTool.cpp
--- a/Engine/Plugins/Experimental/MeshModelingToolset/Source/MeshModelingToolsEditorOnly/Private/MergeMeshesTool.cpp
+++ b/Engine/Plugins/Experimental/MeshModelingToolset/Source/MeshModelingToolsEditorOnly/Private/MergeMeshesTool.cpp
@@ -116,8 +116,9 @@ void UMergeMeshesTool::Setup()
 	Preview->InvalidateResult();    // start compute
 
 
+	SetToolDisplayName(LOCTEXT("ToolName", "Voxel Merge"));
 	GetToolManager()->DisplayMessage(
-		LOCTEXT("OnStartTool", "This Tool combines the input meshes into closed solids using voxelization techniques. UVs, sharp edges, and small/thin features will be lost. Increase Voxel Count to enhance accuracy."),
+		LOCTEXT("OnStartTool", "Combine the input meshes into closed solids using voxelization techniques. UVs, sharp edges, and small/thin features will be lost. Increase Voxel Count to enhance accuracy."),
 		EToolMessageLevel::UserNotification);
 }
 
diff --git a/Engine/Plugins/Experimental/MeshModelingToolset/Source/MeshModelingToolsEditorOnly/Private/MeshTangentsTool.cpp b/Engine/Plugins/Experimental/MeshModelingToolset/Source/MeshModelingToolsEditorOnly/Private/MeshTangentsTool.cpp
--- a/Engine/Plugins/Experimental/MeshModelingToolset/Source/MeshModelingToolsEditorOnly/Private/MeshTangentsTool.cpp
+++ b/Engine/Plugins/Experimental/MeshModelingToolset/Source/MeshModelingToolsEditorOnly/Private/MeshTangentsTool.cpp
@@ -103,6 +103,7 @@ void UMeshTangentsTool::Setup()
 	Compute->OnResultUpdated.AddLambda( [this](const TUniquePtr<FMeshTangentsd>& NewResult) { OnTangentsUpdated(NewResult); } );
 	Compute->InvalidateResult();
 
+	SetToolDisplayName(LOCTEXT("ToolName", "Edit Tangents"));
 	GetToolManager()->DisplayMessage(
 		LOCTEXT("OnStartTool", "Configure or Recalculate Tangents on a Static Mesh Asset (disables autogenerated Tangents and Normals)"),
 		EToolMessageLevel::UserNotification);
diff --git a/Engine/Plugins/Experimental/MeshModelingToolset/Source/MeshModelingToolsEditorOnly/Private/ParameterizeMeshTool.cpp b/Engine/Plugins/Experimental/MeshModelingToolset/Source/MeshModelingToolsEditorOnly/Private/ParameterizeMeshTool.cpp
--- a/Engine/Plugins/Experimental/MeshModelingToolset/Source/MeshModelingToolsEditorOnly/Private/ParameterizeMeshTool.cpp
+++ b/Engine/Plugins/Experimental/MeshModelingToolset/Source/MeshModelingToolsEditorOnly/Private/ParameterizeMeshTool.cpp
@@ -150,12 +150,14 @@ void UParameterizeMeshTool::Setup()
 
 	if (bDoAutomaticGlobalUnwrap)
 	{
+		SetToolDisplayName(LOCTEXT("ToolNameGlobal", "AutoUV"));
 		GetToolManager()->DisplayMessage(
 			LOCTEXT("OnStartTool_Global", "Automatically partition the selected Mesh into UV islands, flatten, and pack into a single UV chart"),
 			EToolMessageLevel::UserNotification);
 	}
 	else
 	{
+		SetToolDisplayName(LOCTEXT("ToolNameLocal", "UV Unwrap"));
 		GetToolManager()->DisplayMessage(
 			LOCTEXT("OnStartTool_Regions", "Generate UVs for polygroups or existing UV charts of the Mesh using various strategies. Does not calculate layout/packing."),
 			EToolMessageLevel::UserNotification);
diff --git a/Engine/Plugins/Experimental/MeshModelingToolset/Source/MeshModelingToolsEditorOnly/Private/PolygonOnMeshTool.cpp b/Engine/Plugins/Experimental/MeshModelingToolset/Source/MeshModelingToolsEditorOnly/Private/PolygonOnMeshTool.cpp
--- a/Engine/Plugins/Experimental/MeshModelingToolset/Source/MeshModelingToolsEditorOnly/Private/PolygonOnMeshTool.cpp
+++ b/Engine/Plugins/Experimental/MeshModelingToolset/Source/MeshModelingToolsEditorOnly/Private/PolygonOnMeshTool.cpp
@@ -63,7 +63,7 @@ void UPolygonOnMeshToolActionPropertySet::PostAction(EPolygonOnMeshToolActions A
 
 UPolygonOnMeshTool::UPolygonOnMeshTool()
 {
-	SetToolDisplayName(LOCTEXT("PolygonOnMeshToolName", "Polygon Cut Tool"));
+	SetToolDisplayName(LOCTEXT("PolygonOnMeshToolName", "PolyCut"));
 }
 
 void UPolygonOnMeshTool::SetWorld(UWorld* World)
diff --git a/Engine/Plugins/Experimental/MeshModelingToolset/Source/MeshModelingToolsEditorOnly/Private/SimplifyMeshTool.cpp b/Engine/Plugins/Experimental/MeshModelingToolset/Source/MeshModelingToolsEditorOnly/Private/SimplifyMeshTool.cpp
--- a/Engine/Plugins/Experimental/MeshModelingToolset/Source/MeshModelingToolsEditorOnly/Private/SimplifyMeshTool.cpp
+++ b/Engine/Plugins/Experimental/MeshModelingToolset/Source/MeshModelingToolsEditorOnly/Private/SimplifyMeshTool.cpp
@@ -163,6 +163,7 @@ void USimplifyMeshTool::Setup()
 	UpdateVisualization();
 	Preview->InvalidateResult();
 
+	SetToolDisplayName(LOCTEXT("ToolName", "Simplify"));
 	GetToolManager()->DisplayMessage(
 		LOCTEXT("OnStartTool", "Reduce the number of triangles in the selected Mesh using various strategies."),
 		EToolMessageLevel::UserNotification);
diff --git a/Engine/Plugins/Experimental/MeshModelingToolset/Source/MeshModelingToolsEditorOnly/Private/SubdividePolyTool.cpp b/Engine/Plugins/Experimental/MeshModelingToolset/Source/MeshModelingToolsEditorOnly/Private/SubdividePolyTool.cpp
--- a/Engine/Plugins/Experimental/MeshModelingToolset/Source/MeshModelingToolsEditorOnly/Private/SubdividePolyTool.cpp
+++ b/Engine/Plugins/Experimental/MeshModelingToolset/Source/MeshModelingToolsEditorOnly/Private/SubdividePolyTool.cpp
@@ -160,6 +160,7 @@ void USubdividePolyTool::CapSubdivisionLevel(ESubdivisionScheme Scheme, int Desi
 void USubdividePolyTool::Setup()
 {
 	UInteractiveTool::Setup();
+	SetToolDisplayName(LOCTEXT("ToolName", "Subdivide"));
 
 	if (!ComponentTarget)
 	{
diff --git a/Engine/Plugins/Experimental/MeshModelingToolset/Source/MeshModelingToolsEditorOnly/Private/UVLayoutTool.cpp b/Engine/Plugins/Experimental/MeshModelingToolset/Source/MeshModelingToolsEditorOnly/Private/UVLayoutTool.cpp
--- a/Engine/Plugins/Experimental/MeshModelingToolset/Source/MeshModelingToolsEditorOnly/Private/UVLayoutTool.cpp
+++ b/Engine/Plugins/Experimental/MeshModelingToolset/Source/MeshModelingToolsEditorOnly/Private/UVLayoutTool.cpp
@@ -128,6 +128,7 @@ void UUVLayoutTool::Setup()
 
 	UpdateVisualization();
 
+	SetToolDisplayName(LOCTEXT("ToolName", "UV Layout"));
 	GetToolManager()->DisplayMessage(LOCTEXT("OnStartUVLayoutTool", "Transform/Rotate/Scale existing UV Charts using various strategies"),
 		EToolMessageLevel::UserNotification);
 }
diff --git a/Engine/Plugins/Experimental/MeshModelingToolset/Source/MeshModelingToolsEditorOnly/Private/VoxelCSGMeshesTool.cpp b/Engine/Plugins/Experimental/MeshModelingToolset/Source/MeshModelingToolsEditorOnly/Private/VoxelCSGMeshesTool.cpp
--- a/Engine/Plugins/Experimental/MeshModelingToolset/Source/MeshModelingToolsEditorOnly/Private/VoxelCSGMeshesTool.cpp
+++ b/Engine/Plugins/Experimental/MeshModelingToolset/Source/MeshModelingToolsEditorOnly/Private/VoxelCSGMeshesTool.cpp
@@ -111,8 +111,9 @@ void UVoxelCSGMeshesTool::Setup()
 	
 	Preview->InvalidateResult();    // start compute
 
+	SetToolDisplayName(LOCTEXT("ToolName", "Voxel Boolean"));
 	GetToolManager()->DisplayMessage(
-		LOCTEXT("OnStartTool", "This Tool computes a CSG Boolean of the input meshes using voxelization techniques. UVs, sharp edges, and small/thin features will be lost. Increase Voxel Count to enhance accuracy."),
+		LOCTEXT("OnStartTool", "Compute a CSG Boolean of the input meshes using voxelization techniques. UVs, sharp edges, and small/thin features will be lost. Increase Voxel Count to enhance accuracy."),
 		EToolMessageLevel::UserNotification);
 }
 
diff --git a/Engine/Plugins/Experimental/ModelingToolsEditorMode/Source/ModelingToolsEditorMode/Private/ModelingToolsEditorModeToolkit.cpp b/Engine/Plugins/Experimental/ModelingToolsEditorMode/Source/ModelingToolsEditorMode/Private/ModelingToolsEditorModeToolkit.cpp
--- a/Engine/Plugins/Experimental/ModelingToolsEditorMode/Source/ModelingToolsEditorMode/Private/ModelingToolsEditorModeToolkit.cpp
+++ b/Engine/Plugins/Experimental/ModelingToolsEditorMode/Source/ModelingToolsEditorMode/Private/ModelingToolsEditorModeToolkit.cpp
@@ -146,6 +146,7 @@ void FModelingToolsEditorModeToolkit::Init(const TSharedPtr<IToolkitHost>& InitT
 				.ButtonStyle(FAppStyle::Get(), "PrimaryButton")
 				.TextStyle( FAppStyle::Get(), "DialogButtonText" )
 				.Text(LOCTEXT("OverlayAccept", "Accept"))
+				.ToolTipText(LOCTEXT("OverlayAcceptTooltip", "Accept/Commit the results of the active Tool [Enter]"))
 				.HAlign(HAlign_Center)
 				.OnClicked_Lambda([this]() { GetScriptableEditorMode()->GetInteractiveToolsContext()->EndTool(EToolShutdownType::Accept); return FReply::Handled(); })
 				.IsEnabled_Lambda([this]() { return GetScriptableEditorMode()->GetInteractiveToolsContext()->CanAcceptActiveTool(); })
@@ -160,6 +161,7 @@ void FModelingToolsEditorModeToolkit::Init(const TSharedPtr<IToolkitHost>& InitT
 				.ButtonStyle(FAppStyle::Get(), "Button")
 				.TextStyle( FAppStyle::Get(), "DialogButtonText" )
 				.Text(LOCTEXT("OverlayCancel", "Cancel"))
+				.ToolTipText(LOCTEXT("OverlayCancelTooltip", "Cancel the active Tool [Esc]"))
 				.HAlign(HAlign_Center)
 				.OnClicked_Lambda([this]() { GetScriptableEditorMode()->GetInteractiveToolsContext()->EndTool(EToolShutdownType::Cancel); return FReply::Handled(); })
 				.IsEnabled_Lambda([this]() { return GetScriptableEditorMode()->GetInteractiveToolsContext()->CanCancelActiveTool(); })
@@ -171,9 +173,10 @@ void FModelingToolsEditorModeToolkit::Init(const TSharedPtr<IToolkitHost>& InitT
 			.Padding(FMargin(2.0, 0.f, 0.f, 0.f))
 			[
 				SNew(SButton)
-				.ButtonStyle(FAppStyle::Get(), "Button")
+				.ButtonStyle(FAppStyle::Get(), "PrimaryButton")
 				.TextStyle( FAppStyle::Get(), "DialogButtonText" )
 				.Text(LOCTEXT("OverlayComplete", "Complete"))
+				.ToolTipText(LOCTEXT("OverlayCompleteTooltip", "Exit the active Tool [Enter]"))
 				.HAlign(HAlign_Center)
 				.OnClicked_Lambda([this]() { GetScriptableEditorMode()->GetInteractiveToolsContext()->EndTool(EToolShutdownType::Completed); return FReply::Handled(); })
 				.IsEnabled_Lambda([this]() { return GetScriptableEditorMode()->GetInteractiveToolsContext()->CanCompleteActiveTool(); })
diff --git a/Engine/Plugins/Experimental/ModelingToolsEditorMode/Source/ModelingToolsEditorMode/Private/ModelingToolsManagerActions.cpp b/Engine/Plugins/Experimental/ModelingToolsEditorMode/Source/ModelingToolsEditorMode/Private/ModelingToolsManagerActions.cpp
--- a/Engine/Plugins/Experimental/ModelingToolsEditorMode/Source/ModelingToolsEditorMode/Private/ModelingToolsManagerActions.cpp
+++ b/Engine/Plugins/Experimental/ModelingToolsEditorMode/Source/ModelingToolsEditorMode/Private/ModelingToolsManagerActions.cpp
@@ -35,7 +35,7 @@ void FModelingToolsManagerCommands::RegisterCommands()
 	UI_COMMAND(BeginAddSpherePrimitiveTool, "SphrA", "Create new Sphere StaticMesh Assets with Lat/Long Parameterization", EUserInterfaceActionType::ToggleButton, FInputChord());
 
 	UI_COMMAND(BeginAddPatchTool, "Patch", "Add Patch", EUserInterfaceActionType::ToggleButton, FInputChord());
-	UI_COMMAND(BeginRevolveBoundaryTool, "CutRev", "Cut Mesh and Revolve Cut Section Curve into New Asset", EUserInterfaceActionType::ToggleButton, FInputChord());
+	UI_COMMAND(BeginRevolveBoundaryTool, "BdryRev", "Revolve Mesh Boundary Loop into New Asset", EUserInterfaceActionType::ToggleButton, FInputChord());
 	UI_COMMAND(BeginDrawPolygonTool, "PolyExt", "Draw/Extrude Polygons to create new StaticMesh Assets", EUserInterfaceActionType::ToggleButton, FInputChord());
 	UI_COMMAND(BeginDrawPolyPathTool, "PathExt", "Draw/Extrude PolyPaths to create new StaticMesh Assets", EUserInterfaceActionType::ToggleButton, FInputChord());
 	UI_COMMAND(BeginDrawAndRevolveTool, "PathRev", "Draw and Revolve a PolyPath to create a new StaticMesh Assets", EUserInterfaceActionType::ToggleButton, FInputChord());
