MeshModelingTools: Minor SpaceDeform UI changes- changed back to using lower bound and upper bound as names, made lock bottom default to false.

#rnx

diff --git a/Engine/Plugins/Experimental/MeshModelingToolset/Source/MeshModelingTools/Private/MeshSpaceDeformerTool.cpp b/Engine/Plugins/Experimental/MeshModelingToolset/Source/MeshModelingTools/Private/MeshSpaceDeformerTool.cpp
--- a/Engine/Plugins/Experimental/MeshModelingToolset/Source/MeshModelingTools/Private/MeshSpaceDeformerTool.cpp
+++ b/Engine/Plugins/Experimental/MeshModelingToolset/Source/MeshModelingTools/Private/MeshSpaceDeformerTool.cpp
@@ -244,7 +244,7 @@ void UMeshSpaceDeformerTool::Setup()
 
 	// Sync the properties panel to the interval handles.
 	Settings->UpperBoundsInterval = UpIntervalSource->Value;
-	Settings->AbsLowerBoundsInterval = FMath::Abs(DownIntervalSource->Value);
+	Settings->LowerBoundsInterval = DownIntervalSource->Value;
 
 	// Wire up callbacks to update result mesh and the properties panel when these parameters are changed (by gizmo manipulation in viewport).  Note this is just a one-way
 	// coupling (Sources to Properties). The OnPropertyModified() method provides the Properties to Souces coupling 
@@ -256,7 +256,7 @@ void UMeshSpaceDeformerTool::Setup()
 
 	DownIntervalSource->OnParameterChanged.AddLambda([this](IGizmoFloatParameterSource* ParamSource, FGizmoFloatParameterChange Change)->void
 	{
-		Settings->AbsLowerBoundsInterval = FMath::Abs(Change.CurrentValue);
+		Settings->LowerBoundsInterval = Change.CurrentValue;
 		UpdatePreview();
 	});
 
@@ -375,7 +375,7 @@ void  UMeshSpaceDeformerTool::TransformProxyChanged(UTransformProxy* Proxy, FTra
 void  UMeshSpaceDeformerTool::OnPropertyModified(UObject* PropertySet, FProperty* Property)
 {
 	UpIntervalSource->Value = Settings->UpperBoundsInterval;
-	DownIntervalSource->Value = -Settings->AbsLowerBoundsInterval;
+	DownIntervalSource->Value = Settings->LowerBoundsInterval;
 	ForwardIntervalSource->Value = GetModifierGizmoValue();
 
 	UpdatePreview();
@@ -392,8 +392,8 @@ void UMeshSpaceDeformerTool::UpdateOpParameters(FMeshSpaceDeformerOp& MeshSpaceD
 	MeshSpaceDeformerOp.GizmoFrame = GizmoFrame;
 
 	// set the bound range
-	MeshSpaceDeformerOp.UpperBoundsInterval =  Settings->UpperBoundsInterval;
-	MeshSpaceDeformerOp.LowerBoundsInterval = -Settings->AbsLowerBoundsInterval;
+	MeshSpaceDeformerOp.UpperBoundsInterval = Settings->UpperBoundsInterval;
+	MeshSpaceDeformerOp.LowerBoundsInterval = Settings->LowerBoundsInterval;
 }
 
 void UMeshSpaceDeformerTool::RequestAction(EMeshSpaceDeformerToolAction ActionType)
@@ -453,18 +453,18 @@ void UMeshSpaceDeformerTool::UpdatePreview()
 
 		VisualizationPoints.SetNumUninitialized(NUM_RENDER_POINTS);
 
-		double BentLength = Settings->UpperBoundsInterval + Settings->AbsLowerBoundsInterval;
+		double BentLength = Settings->UpperBoundsInterval - Settings->LowerBoundsInterval;
 		double ArcAngle = Settings->BendDegrees * PI / 180;
 		double ArcRadius = BentLength / ArcAngle;
 
-		double RotationCenterZ = Settings->bLockBottom ? -Settings->AbsLowerBoundsInterval : 0;
+		double RotationCenterZ = Settings->bLockBottom ? Settings->LowerBoundsInterval : 0;
 		FVector2d RotationCenterYZ(ArcRadius, RotationCenterZ);
 
 		double PointSpacing = BentLength / (NUM_RENDER_POINTS - 1);
 
 		for (int32 i = 0; i < NUM_RENDER_POINTS; ++i)
 		{
-			double OriginalZ = -Settings->AbsLowerBoundsInterval + PointSpacing * i;
+			double OriginalZ = Settings->LowerBoundsInterval + PointSpacing * i;
 			FVector2d YZToRotate(0, RotationCenterZ);
 
 			// The negative here is because we are rotating clockwise in the direction of the positive Y axis
diff --git a/Engine/Plugins/Experimental/MeshModelingToolset/Source/MeshModelingTools/Public/MeshSpaceDeformerTool.h b/Engine/Plugins/Experimental/MeshModelingToolset/Source/MeshModelingTools/Public/MeshSpaceDeformerTool.h
--- a/Engine/Plugins/Experimental/MeshModelingToolset/Source/MeshModelingTools/Public/MeshSpaceDeformerTool.h
+++ b/Engine/Plugins/Experimental/MeshModelingToolset/Source/MeshModelingTools/Public/MeshSpaceDeformerTool.h
@@ -74,15 +74,13 @@ public:
 	UPROPERTY(EditAnywhere, Category = Options, meta = (DisplayName = "Operation Type"))
 	ENonlinearOperationType SelectedOperationType = ENonlinearOperationType::Bend;
 
-	/** The upper bound to the region of space which the operation will affect. Measured upward along the gizmo Z axis from the gizmo center. */
-	UPROPERTY(EditAnywhere, Category = Options, meta = (DisplayName = "Up Extent", UIMin = "0.0", ClampMin = "0.0"))
-	float UpperBoundsInterval = 10.0;
+	/** The upper bound to the region of space which the operation will affect. Measured along the gizmo Z axis from the gizmo center. */
+	UPROPERTY(EditAnywhere, Category = Options, meta = (DisplayName = "Upper Bound", UIMin = "0.0", ClampMin = "0.0"))
+	float UpperBoundsInterval = 100.0;
 
-	//~ Note that the value here is the negative of what we use elsewhere, because otherwise the user has to reverse
-	//~ their scrubbing direction compared to UpperBoundsInterval for increasing/decreasing the interval.
-	/** The lower bound to the region of space which the operation will affect. Measured downward along the gizmo Z axis from the gizmo center. */
-	UPROPERTY(EditAnywhere, Category = Options, meta = (DisplayName = "Down Extent", UIMin = "0", ClampMin = "0"))
-	float AbsLowerBoundsInterval = 10.0;
+	/** The lower bound to the region of space which the operation will affect. Measured along the gizmo Z axis from the gizmo center. */
+	UPROPERTY(EditAnywhere, Category = Options, meta = (DisplayName = "Lower Bound", UIMax = "0", ClampMax = "0"))
+	float LowerBoundsInterval = -100.0;
 
 	//~ The two degrees properties (bend vs twist) are separate because they have different clamp values. 
 	//~ Bending a "negative" amount probably won't do what the user wants, whereas twisting in the opposite
@@ -129,7 +127,7 @@ public:
 	 */
 	UPROPERTY(EditAnywhere, Category = Options, meta = (
 		EditCondition = "SelectedOperationType == ENonlinearOperationType::Bend || SelectedOperationType == ENonlinearOperationType::Twist", EditConditionHides))
-	bool bLockBottom = true;
+	bool bLockBottom = false;
 
 	/**
 	 * When true, changes the flare/pinch function to have C1 continuity at the upper and lower bounds.
