ModelingTools: avoid crashes in Hair tools when cards editor is run on a non-cards mesh
#rnx

diff --git a/Engine/Plugins/Experimental/MeshModelingToolset/Source/MeshModelingToolsEditorOnly/Private/Hair/GroomQueryUtil.cpp b/Engine/Plugins/Experimental/MeshModelingToolset/Source/MeshModelingToolsEditorOnly/Private/Hair/GroomQueryUtil.cpp
--- a/Engine/Plugins/Experimental/MeshModelingToolset/Source/MeshModelingToolsEditorOnly/Private/Hair/GroomQueryUtil.cpp
+++ b/Engine/Plugins/Experimental/MeshModelingToolset/Source/MeshModelingToolsEditorOnly/Private/Hair/GroomQueryUtil.cpp
@@ -237,11 +237,13 @@ void UE::GroomQueries::ExtractCardQuads(
 	for (int32 ci = 0; ci < NumCards; ++ci)
 	{
 		FMeshCardStrip& CardStrip = CardsInOut.CardStrips[ci];
+		bool bIsValidCard = true;
 
 		TArray<int32> RemainingTris(CardStrip.Triangles);
 		while (RemainingTris.Num() > 1)
 		{
 			bool bIsLastQuad = (RemainingTris.Num() < 4);
+			int32 LastRemainingTrisCount = RemainingTris.Num();
 
 			for (int32 tid : RemainingTris)
 			{
@@ -340,6 +342,18 @@ void UE::GroomQueries::ExtractCardQuads(
 					break;
 				}
 			}
+
+			// if we did not remove any triangles this loop, this is not a valid card
+			if (LastRemainingTrisCount == RemainingTris.Num())
+			{
+				bIsValidCard = false;
+				break;
+			}
+		}
+
+		if (!bIsValidCard)
+		{
+			continue;
 		}
 
 		// Output of code above has each quad oriented properly, but the vertex cycle is not consistent.
