ModelingComponents: Add FMeshRenderDecomposition::IsGroup() to check if group is initialized, use in ::BuildAssociations() to prevent access to un-initialized group pointers
#rnx

diff --git a/Engine/Plugins/Experimental/MeshModelingToolset/Source/ModelingComponents/Private/Drawing/MeshRenderDecomposition.cpp b/Engine/Plugins/Experimental/MeshModelingToolset/Source/ModelingComponents/Private/Drawing/MeshRenderDecomposition.cpp
--- a/Engine/Plugins/Experimental/MeshModelingToolset/Source/ModelingComponents/Private/Drawing/MeshRenderDecomposition.cpp
+++ b/Engine/Plugins/Experimental/MeshModelingToolset/Source/ModelingComponents/Private/Drawing/MeshRenderDecomposition.cpp
@@ -14,10 +14,13 @@ void FMeshRenderDecomposition::BuildAssociations(const FDynamicMesh3* Mesh)
 
 	ParallelFor(Groups.Num(), [&](int32 GroupIndex) 
 	{
-		const FGroup& Group = GetGroup(GroupIndex);
-		for (int32 tid : Group.Triangles)
+		if (IsGroup(GroupIndex))
 		{
-			TriangleToGroupMap[tid] = GroupIndex;
+			const FGroup& Group = GetGroup(GroupIndex);
+			for (int32 tid : Group.Triangles)
+			{
+				TriangleToGroupMap[tid] = GroupIndex;
+			}
 		}
 	});
 }
diff --git a/Engine/Plugins/Experimental/MeshModelingToolset/Source/ModelingComponents/Public/Drawing/MeshRenderDecomposition.h b/Engine/Plugins/Experimental/MeshModelingToolset/Source/ModelingComponents/Public/Drawing/MeshRenderDecomposition.h
--- a/Engine/Plugins/Experimental/MeshModelingToolset/Source/ModelingComponents/Public/Drawing/MeshRenderDecomposition.h
+++ b/Engine/Plugins/Experimental/MeshModelingToolset/Source/ModelingComponents/Public/Drawing/MeshRenderDecomposition.h
@@ -57,6 +57,11 @@ public:
 		return Groups.Num();
 	}
 
+	bool IsGroup(int32 Index) const
+	{
+		return Groups[Index] != nullptr;
+	}
+
 	FGroup& GetGroup(int32 Index)
 	{
 		return *Groups[Index];
