change all the modeling / geometry tools shared pointers to thread safe mode, because they're often used in background threads

#rnx

diff --git a/Engine/Plugins/Experimental/Enterprise/DataprepGeometryOperations/Source/Private/DataprepGeometryOperations.cpp b/Engine/Plugins/Experimental/Enterprise/DataprepGeometryOperations/Source/Private/DataprepGeometryOperations.cpp
--- a/Engine/Plugins/Experimental/Enterprise/DataprepGeometryOperations/Source/Private/DataprepGeometryOperations.cpp
+++ b/Engine/Plugins/Experimental/Enterprise/DataprepGeometryOperations/Source/Private/DataprepGeometryOperations.cpp
@@ -79,7 +79,7 @@ void UDataprepRemeshOperation::OnExecution_Implementation(const FDataprepContext
 		
 		FMeshDescription* MeshDescription = StaticMesh->GetMeshDescription(0);
 
-		TSharedPtr<FDynamicMesh3> OriginalMesh = MakeShared<FDynamicMesh3>();
+		TSharedPtr<FDynamicMesh3, ESPMode::ThreadSafe> OriginalMesh = MakeShared<FDynamicMesh3, ESPMode::ThreadSafe>();
 		FMeshDescriptionToDynamicMesh MeshDescriptionToDynamicMesh;
 		MeshDescriptionToDynamicMesh.Convert(MeshDescription, *OriginalMesh);
 
@@ -100,7 +100,7 @@ void UDataprepRemeshOperation::OnExecution_Implementation(const FDataprepContext
 		Op.SmoothingStrength = SmoothingStrength;
 		Op.SmoothingType = ERemeshSmoothingType::MeanValue;
 
-		TSharedPtr<FDynamicMeshAABBTree3> OriginalMeshSpatial = MakeShared<FDynamicMeshAABBTree3>(OriginalMesh.Get(), true);
+		TSharedPtr<FDynamicMeshAABBTree3, ESPMode::ThreadSafe> OriginalMeshSpatial = MakeShared<FDynamicMeshAABBTree3, ESPMode::ThreadSafe>(OriginalMesh.Get(), true);
 
 		double InitialMeshArea = 0;
 		for (int tid : OriginalMesh->TriangleIndicesItr())
@@ -162,13 +162,13 @@ void UDataprepSimplifyMeshOperation::OnExecution_Implementation(const FDataprepC
 		
 		FMeshDescription* MeshDescription = StaticMesh->GetMeshDescription(0);
 
-		TSharedPtr<FDynamicMesh3> OriginalMesh = MakeShared<FDynamicMesh3>();
+		TSharedPtr<FDynamicMesh3, ESPMode::ThreadSafe> OriginalMesh = MakeShared<FDynamicMesh3, ESPMode::ThreadSafe>();
 		FMeshDescriptionToDynamicMesh MeshDescriptionToDynamicMesh;
 		MeshDescriptionToDynamicMesh.Convert(MeshDescription, *OriginalMesh);
 
-		TSharedPtr<FDynamicMeshAABBTree3> OriginalMeshSpatial = MakeShared<FDynamicMeshAABBTree3>(OriginalMesh.Get(), true);
+		TSharedPtr<FDynamicMeshAABBTree3, ESPMode::ThreadSafe> OriginalMeshSpatial = MakeShared<FDynamicMeshAABBTree3, ESPMode::ThreadSafe>(OriginalMesh.Get(), true);
 
-		TSharedPtr<FMeshDescription> OriginalMeshDescription = MakeShared<FMeshDescription>(*StaticMesh->GetMeshDescription(0));
+		TSharedPtr<FMeshDescription, ESPMode::ThreadSafe> OriginalMeshDescription = MakeShared<FMeshDescription, ESPMode::ThreadSafe>(*StaticMesh->GetMeshDescription(0));
 
 		FSimplifyMeshOp Op;
 
@@ -400,7 +400,7 @@ void UDataprepWeldEdgesOperation::OnExecution_Implementation(const FDataprepCont
 		
 		FMeshDescription* MeshDescription = StaticMesh->GetMeshDescription(0);
 
-		TSharedPtr<FDynamicMesh3> TargetMesh = MakeShared<FDynamicMesh3>();
+		TSharedPtr<FDynamicMesh3, ESPMode::ThreadSafe> TargetMesh = MakeShared<FDynamicMesh3, ESPMode::ThreadSafe>();
 		FMeshDescriptionToDynamicMesh MeshDescriptionToDynamicMesh;
 		MeshDescriptionToDynamicMesh.Convert(MeshDescription, *TargetMesh);
 
@@ -471,7 +471,7 @@ void UDataprepPlaneCutOperation::OnExecution_Implementation(const FDataprepConte
 
 		FMeshDescription* MeshDescription = StaticMesh->GetMeshDescription(0);
 
-		TSharedPtr<FDynamicMesh3> OriginalMesh = MakeShared<FDynamicMesh3>();
+		TSharedPtr<FDynamicMesh3, ESPMode::ThreadSafe> OriginalMesh = MakeShared<FDynamicMesh3, ESPMode::ThreadSafe>();
 		FMeshDescriptionToDynamicMesh MeshDescriptionToDynamicMesh;
 		MeshDescriptionToDynamicMesh.Convert(MeshDescription, *OriginalMesh);
 
@@ -605,7 +605,7 @@ void UDataprepPlaneCutOperation::OnExecution_Implementation(const FDataprepConte
 
 TUniquePtr<FDynamicMeshOperator> UDataprepPlaneCutOperation::MakeNewOperator(
 	const UStaticMeshComponent* InStaticMeshComponent, 
-	TSharedPtr<FDynamicMesh3> InOriginalMesh, 
+	TSharedPtr<FDynamicMesh3, ESPMode::ThreadSafe> InOriginalMesh,
 	float InMeshUVScaleFactor)
 {
 	TUniquePtr<FPlaneCutOp> CutOp = MakeUnique<FPlaneCutOp>();
diff --git a/Engine/Plugins/Experimental/Enterprise/DataprepGeometryOperations/Source/Private/DataprepGeometryOperations.h b/Engine/Plugins/Experimental/Enterprise/DataprepGeometryOperations/Source/Private/DataprepGeometryOperations.h
--- a/Engine/Plugins/Experimental/Enterprise/DataprepGeometryOperations/Source/Private/DataprepGeometryOperations.h
+++ b/Engine/Plugins/Experimental/Enterprise/DataprepGeometryOperations/Source/Private/DataprepGeometryOperations.h
@@ -242,6 +242,6 @@ private:
 	// Borrowed from UPlaneCutOperatorFactory::MakeNewOperator
 	TUniquePtr<FDynamicMeshOperator> MakeNewOperator(
 		const UStaticMeshComponent* InStaticMeshComponent, 
-		TSharedPtr<FDynamicMesh3> InOriginalMesh,
+		TSharedPtr<FDynamicMesh3, ESPMode::ThreadSafe> InOriginalMesh,
 		float InMeshUVScaleFactor);
 };
diff --git a/Engine/Plugins/Experimental/GeometryFlow/Source/GeometryFlowMeshProcessingEditor/Private/MeshProcessingNodes/MeshAutoGenerateUVsNode.cpp b/Engine/Plugins/Experimental/GeometryFlow/Source/GeometryFlowMeshProcessingEditor/Private/MeshProcessingNodes/MeshAutoGenerateUVsNode.cpp
--- a/Engine/Plugins/Experimental/GeometryFlow/Source/GeometryFlowMeshProcessingEditor/Private/MeshProcessingNodes/MeshAutoGenerateUVsNode.cpp
+++ b/Engine/Plugins/Experimental/GeometryFlow/Source/GeometryFlowMeshProcessingEditor/Private/MeshProcessingNodes/MeshAutoGenerateUVsNode.cpp
@@ -11,7 +11,7 @@ using namespace UE::GeometryFlow;
 void FMeshAutoGenerateUVsNode::GenerateUVs(const FDynamicMesh3& MeshIn, const FMeshAutoGenerateUVsSettings& Settings, FDynamicMesh3& MeshOut)
 {
 	// this is horrible - have to copy input mesh so that we can pass a TSharedPtr
-	TSharedPtr<FDynamicMesh3> InputMesh = MakeShared<FDynamicMesh3>(MeshIn);
+	TSharedPtr<FDynamicMesh3, ESPMode::ThreadSafe> InputMesh = MakeShared<FDynamicMesh3, ESPMode::ThreadSafe>(MeshIn);
 
 	FParameterizeMeshOp ParameterizeMeshOp;
 	ParameterizeMeshOp.Stretch = Settings.Stretch;
diff --git a/Engine/Plugins/Experimental/MeshModelingToolset/Source/MeshModelingTools/Private/BakeMeshAttributeMapsTool.cpp b/Engine/Plugins/Experimental/MeshModelingToolset/Source/MeshModelingTools/Private/BakeMeshAttributeMapsTool.cpp
--- a/Engine/Plugins/Experimental/MeshModelingToolset/Source/MeshModelingTools/Private/BakeMeshAttributeMapsTool.cpp
+++ b/Engine/Plugins/Experimental/MeshModelingToolset/Source/MeshModelingTools/Private/BakeMeshAttributeMapsTool.cpp
@@ -92,7 +92,7 @@ void UBakeMeshAttributeMapsTool::Setup()
 	UInteractiveTool::Setup();
 
 	// copy input MeshDescription and make sure it has initialized normals/tangents
-	BaseMeshDescription = MakeShared<FMeshDescription>(*ComponentTargets[0]->GetMesh());
+	BaseMeshDescription = MakeShared<FMeshDescription, ESPMode::ThreadSafe>(*ComponentTargets[0]->GetMesh());
 	UE::MeshDescription::InitializeAutoGeneratedAttributes(*BaseMeshDescription, ComponentTargets[0]->GetOwnerComponent(), 0);
 
 	// create dynamic mesh component to use for live preview
@@ -114,7 +114,7 @@ void UBakeMeshAttributeMapsTool::Setup()
 	
 	BaseMesh.Copy(*DynamicMeshComponent->GetMesh());
 	BaseSpatial.SetMesh(&BaseMesh, true);
-	BaseMeshTangents = MakeShared<FMeshTangentsd>(&BaseMesh);
+	BaseMeshTangents = MakeShared<FMeshTangentsd, ESPMode::ThreadSafe>(&BaseMesh);
 	BaseMeshTangents->CopyTriVertexTangents(*DynamicMeshComponent->GetTangents());
 
 	UMaterial* Material = LoadObject<UMaterial>(nullptr, TEXT("/MeshModelingToolset/Materials/BakePreviewMaterial"));
@@ -318,7 +318,7 @@ void UBakeMeshAttributeMapsTool::Render(IToolsContextRenderAPI* RenderAPI)
 
 void UBakeMeshAttributeMapsTool::UpdateDetailMesh()
 {
-	DetailMesh = MakeShared<FDynamicMesh3>();
+	DetailMesh = MakeShared<FDynamicMesh3, ESPMode::ThreadSafe>();
 	FMeshDescriptionToDynamicMesh Converter;
 	TUniquePtr<FPrimitiveComponentTarget>& UseDetailTarget = (bIsBakeToSelf) ? ComponentTargets[0] : ComponentTargets[1];
 	Converter.Convert(UseDetailTarget->GetMesh(), *DetailMesh);
@@ -331,7 +331,7 @@ void UBakeMeshAttributeMapsTool::UpdateDetailMesh()
 		MeshTransforms::ApplyTransform(*DetailMesh, WorldToBase.Inverse());
 	}
 	
-	DetailSpatial = MakeShared<FDynamicMeshAABBTree3>();
+	DetailSpatial = MakeShared<FDynamicMeshAABBTree3, ESPMode::ThreadSafe>();
 	DetailSpatial->SetMesh(DetailMesh.Get(), true);
 
 	bResultValid = false;
diff --git a/Engine/Plugins/Experimental/MeshModelingToolset/Source/MeshModelingTools/Private/CSGMeshesTool.cpp b/Engine/Plugins/Experimental/MeshModelingToolset/Source/MeshModelingTools/Private/CSGMeshesTool.cpp
--- a/Engine/Plugins/Experimental/MeshModelingToolset/Source/MeshModelingTools/Private/CSGMeshesTool.cpp
+++ b/Engine/Plugins/Experimental/MeshModelingToolset/Source/MeshModelingTools/Private/CSGMeshesTool.cpp
@@ -105,7 +105,7 @@ void UCSGMeshesTool::ConvertInputsAndSetPreviewMaterials(bool bSetPreviewMesh)
 
 	for (int ComponentIdx = 0; ComponentIdx < ComponentTargets.Num(); ComponentIdx++)
 	{
-		OriginalDynamicMeshes[ComponentIdx] = MakeShared<FDynamicMesh3>();
+		OriginalDynamicMeshes[ComponentIdx] = MakeShared<FDynamicMesh3, ESPMode::ThreadSafe>();
 		FMeshDescriptionToDynamicMesh Converter;
 		Converter.Convert(ComponentTargets[ComponentIdx]->GetMesh(), *OriginalDynamicMeshes[ComponentIdx]);
 
diff --git a/Engine/Plugins/Experimental/MeshModelingToolset/Source/MeshModelingTools/Private/DisplaceMeshTool.cpp b/Engine/Plugins/Experimental/MeshModelingToolset/Source/MeshModelingTools/Private/DisplaceMeshTool.cpp
--- a/Engine/Plugins/Experimental/MeshModelingToolset/Source/MeshModelingTools/Private/DisplaceMeshTool.cpp
+++ b/Engine/Plugins/Experimental/MeshModelingToolset/Source/MeshModelingTools/Private/DisplaceMeshTool.cpp
@@ -259,13 +259,13 @@ namespace DisplaceMeshToolLocals{
 	class FSubdivideMeshOp : public FDynamicMeshOperator
 	{
 	public:
-		FSubdivideMeshOp(const FDynamicMesh3& SourceMesh, int SubdivisionsCountIn, TSharedPtr<FIndexedWeightMap> WeightMap);
+		FSubdivideMeshOp(const FDynamicMesh3& SourceMesh, int SubdivisionsCountIn, TSharedPtr<FIndexedWeightMap, ESPMode::ThreadSafe> WeightMap);
 		void CalculateResult(FProgressCancel* Progress) final;
 	private:
 		int SubdivisionsCount;
 	};
 
-	FSubdivideMeshOp::FSubdivideMeshOp(const FDynamicMesh3& SourceMesh, int SubdivisionsCountIn, TSharedPtr<FIndexedWeightMap> WeightMap)
+	FSubdivideMeshOp::FSubdivideMeshOp(const FDynamicMesh3& SourceMesh, int SubdivisionsCountIn, TSharedPtr<FIndexedWeightMap, ESPMode::ThreadSafe> WeightMap)
 		: SubdivisionsCount(SubdivisionsCountIn)
 	{
 		ResultMesh->Copy(SourceMesh);
@@ -308,14 +308,14 @@ namespace DisplaceMeshToolLocals{
 	public:
 		FSubdivideMeshOpFactory(FDynamicMesh3& SourceMeshIn,
 			int SubdivisionsCountIn,
-			TSharedPtr<FIndexedWeightMap> WeightMapIn)
+			TSharedPtr<FIndexedWeightMap, ESPMode::ThreadSafe> WeightMapIn)
 			: SourceMesh(SourceMeshIn), SubdivisionsCount(SubdivisionsCountIn), WeightMap(WeightMapIn)
 		{
 		}
 		void SetSubdivisionsCount(int SubdivisionsCountIn);
 		int  GetSubdivisionsCount();
 
-		void SetWeightMap(TSharedPtr<FIndexedWeightMap> WeightMapIn);
+		void SetWeightMap(TSharedPtr<FIndexedWeightMap, ESPMode::ThreadSafe> WeightMapIn);
 
 		TUniquePtr<FDynamicMeshOperator> MakeNewOperator() final
 		{
@@ -324,7 +324,7 @@ namespace DisplaceMeshToolLocals{
 	private:
 		const FDynamicMesh3& SourceMesh;
 		int SubdivisionsCount;
-		TSharedPtr<FIndexedWeightMap> WeightMap;
+		TSharedPtr<FIndexedWeightMap, ESPMode::ThreadSafe> WeightMap;
 	};
 
 	void FSubdivideMeshOpFactory::SetSubdivisionsCount(int SubdivisionsCountIn)
@@ -337,7 +337,7 @@ namespace DisplaceMeshToolLocals{
 		return SubdivisionsCount;
 	}
 
-	void FSubdivideMeshOpFactory::SetWeightMap(TSharedPtr<FIndexedWeightMap> WeightMapIn)
+	void FSubdivideMeshOpFactory::SetWeightMap(TSharedPtr<FIndexedWeightMap, ESPMode::ThreadSafe> WeightMapIn)
 	{
 		WeightMap = WeightMapIn;
 	}
@@ -364,22 +364,22 @@ namespace DisplaceMeshToolLocals{
 		FVector2f UVOffset = FVector2f(0, 0);
 		// This gets used by worker threads, so do not try to change an existing curve- make
 		// a new one each time.
-		TSharedPtr<FRichCurve> AdjustmentCurve;
+		TSharedPtr<FRichCurve, ESPMode::ThreadSafe> AdjustmentCurve;
 
-		TSharedPtr<FIndexedWeightMap> WeightMap;
+		TSharedPtr<FIndexedWeightMap, ESPMode::ThreadSafe> WeightMap;
 		TFunction<float(const FVector3d&, const FIndexedWeightMap)> WeightMapQueryFunc;
 	};
 
 	class FDisplaceMeshOp : public FDynamicMeshOperator
 	{
 	public:
-		FDisplaceMeshOp(TSharedPtr<FDynamicMesh3> SourceMeshIn, 
+		FDisplaceMeshOp(TSharedPtr<FDynamicMesh3, ESPMode::ThreadSafe> SourceMeshIn,
 						const DisplaceMeshParameters& DisplaceParametersIn,
 						EDisplaceMeshToolDisplaceType DisplacementTypeIn);
 		void CalculateResult(FProgressCancel* Progress) final;
 
 	private:
-		TSharedPtr<FDynamicMesh3> SourceMesh;
+		TSharedPtr<FDynamicMesh3, ESPMode::ThreadSafe> SourceMesh;
 		DisplaceMeshParameters Parameters;
 		EDisplaceMeshToolDisplaceType DisplacementType;
 		TArray<FVector3d> SourcePositions;
@@ -387,7 +387,7 @@ namespace DisplaceMeshToolLocals{
 		TArray<FVector3d> DisplacedPositions;
 	};
 
-	FDisplaceMeshOp::FDisplaceMeshOp(TSharedPtr<FDynamicMesh3> SourceMeshIn, 
+	FDisplaceMeshOp::FDisplaceMeshOp(TSharedPtr<FDynamicMesh3, ESPMode::ThreadSafe> SourceMeshIn,
 									 const DisplaceMeshParameters& DisplaceParametersIn,
 									 EDisplaceMeshToolDisplaceType DisplacementTypeIn)
 		: SourceMesh(MoveTemp(SourceMeshIn)), 
@@ -532,7 +532,7 @@ namespace DisplaceMeshToolLocals{
 	class FDisplaceMeshOpFactory : public IDynamicMeshOperatorFactory
 	{
 	public:
-		FDisplaceMeshOpFactory(TSharedPtr<FDynamicMesh3>& SourceMeshIn,
+		FDisplaceMeshOpFactory(TSharedPtr<FDynamicMesh3, ESPMode::ThreadSafe>& SourceMeshIn,
 			const DisplaceMeshParameters& DisplaceParametersIn,
 			EDisplaceMeshToolDisplaceType DisplacementTypeIn )
 			: SourceMesh(SourceMeshIn)
@@ -572,7 +572,7 @@ namespace DisplaceMeshToolLocals{
 		void SetFilterDirection(const FVector& Direction);
 		void SetFilterFalloffWidth(float FalloffWidth);
 		void SetPerlinNoiseLayerProperties(const TArray<FPerlinLayerProperties>& PerlinLayerProperties);
-		void SetWeightMap(TSharedPtr<FIndexedWeightMap> WeightMap);
+		void SetWeightMap(TSharedPtr<FIndexedWeightMap, ESPMode::ThreadSafe> WeightMap);
 		void SetRecalculateNormals(bool bRecalculateNormals);
 
 		TUniquePtr<FDynamicMeshOperator> MakeNewOperator() final
@@ -585,7 +585,7 @@ namespace DisplaceMeshToolLocals{
 		DisplaceMeshParameters Parameters;
 		EDisplaceMeshToolDisplaceType DisplacementType;
 
-		TSharedPtr<FDynamicMesh3>& SourceMesh;
+		TSharedPtr<FDynamicMesh3, ESPMode::ThreadSafe>& SourceMesh;
 	};
 
 	void FDisplaceMeshOpFactory::SetIntensity(float IntensityIn)
@@ -622,7 +622,7 @@ namespace DisplaceMeshToolLocals{
 
 	void FDisplaceMeshOpFactory::SetAdjustmentCurve(UCurveFloat* CurveFloat)
 	{
-		Parameters.AdjustmentCurve = CurveFloat ? TSharedPtr<FRichCurve>(
+		Parameters.AdjustmentCurve = CurveFloat ? TSharedPtr<FRichCurve, ESPMode::ThreadSafe>(
 			static_cast<FRichCurve*>(CurveFloat->FloatCurve.Duplicate()))
 			: nullptr;
 	}
@@ -706,7 +706,7 @@ namespace DisplaceMeshToolLocals{
 		Parameters.PerlinLayerProperties = LayerProperties;
 	}
 
-	void FDisplaceMeshOpFactory::SetWeightMap(TSharedPtr<FIndexedWeightMap> WeightMap)
+	void FDisplaceMeshOpFactory::SetWeightMap(TSharedPtr<FIndexedWeightMap, ESPMode::ThreadSafe> WeightMap)
 	{
 		Parameters.WeightMap = WeightMap;
 	}
@@ -1127,7 +1127,7 @@ void UDisplaceMeshTool::AdvanceComputation()
 {
 	if (SubdivideTask && SubdivideTask->IsDone())
 	{
-		SubdividedMesh = TSharedPtr<FDynamicMesh3>(SubdivideTask->GetTask().ExtractOperator()->ExtractResult().Release());
+		SubdividedMesh = TSharedPtr<FDynamicMesh3, ESPMode::ThreadSafe>(SubdivideTask->GetTask().ExtractOperator()->ExtractResult().Release());
 		delete SubdivideTask;
 		SubdivideTask = nullptr;
 	}
@@ -1159,7 +1159,7 @@ void UDisplaceMeshTool::UpdateActiveWeightMap()
 	}
 	else
 	{
-		TSharedPtr<FIndexedWeightMap> NewWeightMap = MakeShared<FIndexedWeightMap>();
+		TSharedPtr<FIndexedWeightMap, ESPMode::ThreadSafe> NewWeightMap = MakeShared<FIndexedWeightMap, ESPMode::ThreadSafe>();
 		UE::WeightMaps::GetVertexWeightMap(ComponentTarget->GetMesh(), CommonProperties->WeightMap, *NewWeightMap, 1.0f);
 		if (CommonProperties->bInvertWeightMap)
 		{
diff --git a/Engine/Plugins/Experimental/MeshModelingToolset/Source/MeshModelingTools/Private/DynamicMeshSculptTool.cpp b/Engine/Plugins/Experimental/MeshModelingToolset/Source/MeshModelingTools/Private/DynamicMeshSculptTool.cpp
--- a/Engine/Plugins/Experimental/MeshModelingToolset/Source/MeshModelingTools/Private/DynamicMeshSculptTool.cpp
+++ b/Engine/Plugins/Experimental/MeshModelingToolset/Source/MeshModelingTools/Private/DynamicMeshSculptTool.cpp
@@ -2619,8 +2619,8 @@ void UDynamicMeshSculptTool::UpdateFixedPlaneGizmoVisibility(bool bVisible)
 
 void UDynamicMeshSculptTool::DiscardAttributes()
 {
-	TSharedPtr<FDynamicMesh3> BeforeMesh = MakeShared<FDynamicMesh3>(*DynamicMeshComponent->GetMesh());
-	TSharedPtr<FDynamicMesh3> AfterMesh = MakeShared<FDynamicMesh3>(*BeforeMesh);
+	TSharedPtr<FDynamicMesh3, ESPMode::ThreadSafe> BeforeMesh = MakeShared<FDynamicMesh3, ESPMode::ThreadSafe>(*DynamicMeshComponent->GetMesh());
+	TSharedPtr<FDynamicMesh3, ESPMode::ThreadSafe> AfterMesh = MakeShared<FDynamicMesh3, ESPMode::ThreadSafe>(*BeforeMesh);
 	AfterMesh->DiscardAttributes();
 	FMeshNormals::QuickComputeVertexNormals(*AfterMesh);
 
diff --git a/Engine/Plugins/Experimental/MeshModelingToolset/Source/MeshModelingTools/Private/EditNormalsTool.cpp b/Engine/Plugins/Experimental/MeshModelingToolset/Source/MeshModelingTools/Private/EditNormalsTool.cpp
--- a/Engine/Plugins/Experimental/MeshModelingToolset/Source/MeshModelingTools/Private/EditNormalsTool.cpp
+++ b/Engine/Plugins/Experimental/MeshModelingToolset/Source/MeshModelingTools/Private/EditNormalsTool.cpp
@@ -146,7 +146,7 @@ void UEditNormalsTool::UpdateNumPreviews()
 			UEditNormalsOperatorFactory *OpFactory = NewObject<UEditNormalsOperatorFactory>();
 			OpFactory->Tool = this;
 			OpFactory->ComponentIndex = PreviewIdx;
-			OriginalDynamicMeshes[PreviewIdx] = MakeShared<FDynamicMesh3>();
+			OriginalDynamicMeshes[PreviewIdx] = MakeShared<FDynamicMesh3, ESPMode::ThreadSafe>();
 			FMeshDescriptionToDynamicMesh Converter;
 			Converter.Convert(ComponentTargets[PreviewIdx]->GetMesh(), *OriginalDynamicMeshes[PreviewIdx]);
 
diff --git a/Engine/Plugins/Experimental/MeshModelingToolset/Source/MeshModelingTools/Private/GroupEdgeInsertionTool.cpp b/Engine/Plugins/Experimental/MeshModelingToolset/Source/MeshModelingTools/Private/GroupEdgeInsertionTool.cpp
--- a/Engine/Plugins/Experimental/MeshModelingToolset/Source/MeshModelingTools/Private/GroupEdgeInsertionTool.cpp
+++ b/Engine/Plugins/Experimental/MeshModelingToolset/Source/MeshModelingTools/Private/GroupEdgeInsertionTool.cpp
@@ -90,10 +90,10 @@ void UGroupEdgeInsertionTool::Setup()
 		EToolMessageLevel::UserNotification);
 
 	// Initialize the mesh that we'll be operating on
-	CurrentMesh = MakeShared<FDynamicMesh3>();
+	CurrentMesh = MakeShared<FDynamicMesh3, ESPMode::ThreadSafe>();
 	FMeshDescriptionToDynamicMesh Converter;
 	Converter.Convert(ComponentTarget->GetMesh(), *CurrentMesh);
-	CurrentTopology = MakeShared<FGroupTopology>(CurrentMesh.Get(), true);
+	CurrentTopology = MakeShared<FGroupTopology, ESPMode::ThreadSafe>(CurrentMesh.Get(), true);
 	MeshSpatial.SetMesh(CurrentMesh.Get(), true);
 
 	// Set up properties
diff --git a/Engine/Plugins/Experimental/MeshModelingToolset/Source/MeshModelingTools/Private/HoleFillTool.cpp b/Engine/Plugins/Experimental/MeshModelingToolset/Source/MeshModelingTools/Private/HoleFillTool.cpp
--- a/Engine/Plugins/Experimental/MeshModelingToolset/Source/MeshModelingTools/Private/HoleFillTool.cpp
+++ b/Engine/Plugins/Experimental/MeshModelingToolset/Source/MeshModelingTools/Private/HoleFillTool.cpp
@@ -128,7 +128,7 @@ void UHoleFillTool::Setup()
 	}
 
 	// create mesh to operate on
-	OriginalMesh = MakeShared<FDynamicMesh3>();
+	OriginalMesh = MakeShared<FDynamicMesh3, ESPMode::ThreadSafe>();
 	FMeshDescriptionToDynamicMesh Converter;
 	Converter.Convert(ComponentTarget->GetMesh(), *OriginalMesh);
 
diff --git a/Engine/Plugins/Experimental/MeshModelingToolset/Source/MeshModelingTools/Private/LatticeDeformerTool.cpp b/Engine/Plugins/Experimental/MeshModelingToolset/Source/MeshModelingTools/Private/LatticeDeformerTool.cpp
--- a/Engine/Plugins/Experimental/MeshModelingToolset/Source/MeshModelingTools/Private/LatticeDeformerTool.cpp
+++ b/Engine/Plugins/Experimental/MeshModelingToolset/Source/MeshModelingTools/Private/LatticeDeformerTool.cpp
@@ -81,7 +81,7 @@ bool ULatticeDeformerTool::CanAccept() const
 
 void ULatticeDeformerTool::InitializeLattice(TArray<FVector3d>& OutLatticePoints, TArray<FVector2i>& OutLatticeEdges)
 {
-	Lattice = MakeShared<FFFDLattice>(GetLatticeResolution(), *OriginalMesh, Settings->Padding);
+	Lattice = MakeShared<FFFDLattice, ESPMode::ThreadSafe>(GetLatticeResolution(), *OriginalMesh, Settings->Padding);
 
 	Lattice->GenerateInitialLatticePositions(OutLatticePoints);
 
@@ -102,7 +102,7 @@ void ULatticeDeformerTool::Setup()
 	GetToolManager()->DisplayMessage(LOCTEXT("LatticeDeformerToolMessage", 
 		"Drag the lattice control points to deform the mesh"), EToolMessageLevel::UserNotification);
 
-	OriginalMesh = MakeShared<FDynamicMesh3>();
+	OriginalMesh = MakeShared<FDynamicMesh3, ESPMode::ThreadSafe>();
 	FMeshDescriptionToDynamicMesh Converter;
 	Converter.Convert(ComponentTarget->GetMesh(), *OriginalMesh);
 
diff --git a/Engine/Plugins/Experimental/MeshModelingToolset/Source/MeshModelingTools/Private/MeshBoundaryToolBase.cpp b/Engine/Plugins/Experimental/MeshModelingToolset/Source/MeshModelingTools/Private/MeshBoundaryToolBase.cpp
--- a/Engine/Plugins/Experimental/MeshModelingToolset/Source/MeshModelingTools/Private/MeshBoundaryToolBase.cpp
+++ b/Engine/Plugins/Experimental/MeshModelingToolset/Source/MeshModelingTools/Private/MeshBoundaryToolBase.cpp
@@ -21,7 +21,7 @@ void UMeshBoundaryToolBase::Setup()
 	}
 
 	// create mesh to operate on
-	OriginalMesh = MakeShared<FDynamicMesh3>();
+	OriginalMesh = MakeShared<FDynamicMesh3, ESPMode::ThreadSafe>();
 	FMeshDescriptionToDynamicMesh Converter;
 	Converter.Convert(ComponentTarget->GetMesh(), *OriginalMesh);
 
diff --git a/Engine/Plugins/Experimental/MeshModelingToolset/Source/MeshModelingTools/Private/MeshSpaceDeformerTool.cpp b/Engine/Plugins/Experimental/MeshModelingToolset/Source/MeshModelingTools/Private/MeshSpaceDeformerTool.cpp
--- a/Engine/Plugins/Experimental/MeshModelingToolset/Source/MeshModelingTools/Private/MeshSpaceDeformerTool.cpp
+++ b/Engine/Plugins/Experimental/MeshModelingToolset/Source/MeshModelingTools/Private/MeshSpaceDeformerTool.cpp
@@ -152,7 +152,7 @@ void UMeshSpaceDeformerTool::Setup()
 
 	// populate the OriginalDynamicMesh with a conversion of the input mesh.
 	{
-		OriginalDynamicMesh = MakeShared< FDynamicMesh3 >();
+		OriginalDynamicMesh = MakeShared<FDynamicMesh3, ESPMode::ThreadSafe>();
 		FMeshDescriptionToDynamicMesh Converter;
 		Converter.Convert(ComponentTarget->GetMesh(), *OriginalDynamicMesh);
 	}
diff --git a/Engine/Plugins/Experimental/MeshModelingToolset/Source/MeshModelingTools/Private/MirrorTool.cpp b/Engine/Plugins/Experimental/MeshModelingToolset/Source/MeshModelingTools/Private/MirrorTool.cpp
--- a/Engine/Plugins/Experimental/MeshModelingToolset/Source/MeshModelingTools/Private/MirrorTool.cpp
+++ b/Engine/Plugins/Experimental/MeshModelingToolset/Source/MeshModelingTools/Private/MirrorTool.cpp
@@ -207,7 +207,7 @@ void UMirrorTool::Setup()
 		TUniquePtr<FPrimitiveComponentTarget>& ComponentTarget = ComponentTargets[i];
 
 		// Convert into dynamic mesh
-		TSharedPtr<FDynamicMesh3> DynamicMesh = MakeShared<FDynamicMesh3>();
+		TSharedPtr<FDynamicMesh3, ESPMode::ThreadSafe> DynamicMesh = MakeShared<FDynamicMesh3, ESPMode::ThreadSafe>();
 		FMeshDescriptionToDynamicMesh Converter;
 		Converter.Convert(ComponentTarget->GetMesh(), *DynamicMesh);
 
diff --git a/Engine/Plugins/Experimental/MeshModelingToolset/Source/MeshModelingTools/Private/Physics/SetCollisionGeometryTool.cpp b/Engine/Plugins/Experimental/MeshModelingToolset/Source/MeshModelingTools/Private/Physics/SetCollisionGeometryTool.cpp
--- a/Engine/Plugins/Experimental/MeshModelingToolset/Source/MeshModelingTools/Private/Physics/SetCollisionGeometryTool.cpp
+++ b/Engine/Plugins/Experimental/MeshModelingToolset/Source/MeshModelingTools/Private/Physics/SetCollisionGeometryTool.cpp
@@ -83,7 +83,7 @@ void USetCollisionGeometryTool::Setup()
 	PreviewGeom->CreateInWorld(CollisionTarget->GetOwnerActor()->GetWorld(), PreviewTransform);
 
 	// initialize initial collision object
-	InitialCollision = MakeShared<FPhysicsDataCollection>();
+	InitialCollision = MakeShared<FPhysicsDataCollection, ESPMode::ThreadSafe>();
 	InitialCollision->InitializeFromComponent(CollisionTarget->GetOwnerComponent(), true);
 	InitialCollision->ExternalScale3D = TargetScale3D;
 
@@ -252,14 +252,14 @@ void USetCollisionGeometryTool::UpdateGeneratedCollision()
 	ECollisionGeometryType ComputeType = Settings->GeometryType;
 
 
-	TSharedPtr<FPhysicsDataCollection> NewCollision = MakeShared<FPhysicsDataCollection>();
+	TSharedPtr<FPhysicsDataCollection, ESPMode::ThreadSafe> NewCollision = MakeShared<FPhysicsDataCollection, ESPMode::ThreadSafe>();
 	NewCollision->InitializeFromExisting(*InitialCollision);
 	if (Settings->bAppendToExisting || ComputeType == ECollisionGeometryType::KeepExisting)
 	{
 		NewCollision->CopyGeometryFromExisting(*InitialCollision);
 	}
 
-	TSharedPtr<FMeshSimpleShapeApproximation> UseShapeGenerator = GetApproximator(Settings->InputMode);
+	TSharedPtr<FMeshSimpleShapeApproximation, ESPMode::ThreadSafe> UseShapeGenerator = GetApproximator(Settings->InputMode);
 
 	UseShapeGenerator->bDetectSpheres = Settings->bDetectSpheres;
 	UseShapeGenerator->bDetectBoxes = Settings->bDetectBoxes;
@@ -337,8 +337,8 @@ void USetCollisionGeometryTool::UpdateGeneratedCollision()
 
 
 void USetCollisionGeometryTool::InitializeDerivedMeshSet(
-	const TArray<TSharedPtr<FDynamicMesh3>>& FromInputMeshes,
-	TArray<TSharedPtr<FDynamicMesh3>>& ToMeshes,
+	const TArray<TSharedPtr<FDynamicMesh3, ESPMode::ThreadSafe>>& FromInputMeshes,
+	TArray<TSharedPtr<FDynamicMesh3, ESPMode::ThreadSafe>>& ToMeshes,
 	TFunctionRef<bool(const FDynamicMesh3* Mesh, int32 Tri0, int32 Tri1)> TrisConnectedPredicate)
 {
 	// find connected-components on input meshes, under given connectivity predicate
@@ -386,17 +386,17 @@ void USetCollisionGeometryTool::InitializeDerivedMeshSet(
 		const FSubmeshSource& Source = AllSubmeshes[k];
 		const FMeshConnectedComponents::FComponent& Component = ComponentSets[Source.ComponentIdx.A]->GetComponent(Source.ComponentIdx.B);
 		FDynamicSubmesh3 Submesh(Source.SourceMesh, Component.Indices, (int32)EMeshComponents::None, false);
-		ToMeshes[k] = MakeShared<FDynamicMesh3>( MoveTemp(Submesh.GetSubmesh()) );
+		ToMeshes[k] = MakeShared<FDynamicMesh3, ESPMode::ThreadSafe>( MoveTemp(Submesh.GetSubmesh()) );
 	});
 }
 
 
 template<typename T>
-TArray<const T*> MakeRawPointerList(const TArray<TSharedPtr<T>>& InputList)
+TArray<const T*> MakeRawPointerList(const TArray<TSharedPtr<T, ESPMode::ThreadSafe>>& InputList)
 {
 	TArray<const T*> Result;
 	Result.Reserve(InputList.Num());
-	for (const TSharedPtr<T>& Ptr : InputList)
+	for (const TSharedPtr<T, ESPMode::ThreadSafe>& Ptr : InputList)
 	{
 		Result.Add(Ptr.Get());
 	}
@@ -426,9 +426,9 @@ void USetCollisionGeometryTool::PrecomputeInputMeshes()
 			MeshTransforms::ApplyTransform(SourceMesh, TargetTransformInv);
 		}
 		SourceMesh.DiscardAttributes();
-		InputMeshes[k] = MakeShared<FDynamicMesh3>(MoveTemp(SourceMesh));
+		InputMeshes[k] = MakeShared<FDynamicMesh3, ESPMode::ThreadSafe>(MoveTemp(SourceMesh));
 	});
-	InputMeshesApproximator = MakeShared<FMeshSimpleShapeApproximation>();
+	InputMeshesApproximator = MakeShared<FMeshSimpleShapeApproximation, ESPMode::ThreadSafe>();
 	InputMeshesApproximator->InitializeSourceMeshes(MakeRawPointerList<FDynamicMesh3>(InputMeshes));
 
 
@@ -438,33 +438,33 @@ void USetCollisionGeometryTool::PrecomputeInputMeshes()
 	CombinedMesh.EnableTriangleGroups();
 	FDynamicMeshEditor Appender(&CombinedMesh);
 	FMeshIndexMappings TmpMappings;
-	for (const TSharedPtr<FDynamicMesh3>& InputMesh : InputMeshes)
+	for (const TSharedPtr<FDynamicMesh3, ESPMode::ThreadSafe>& InputMesh : InputMeshes)
 	{
 		TmpMappings.Reset();
 		Appender.AppendMesh(InputMesh.Get(), TmpMappings);
 	}
-	CombinedInputMeshes.Add( MakeShared<FDynamicMesh3>(MoveTemp(CombinedMesh)) );
-	CombinedInputMeshesApproximator = MakeShared<FMeshSimpleShapeApproximation>();
+	CombinedInputMeshes.Add( MakeShared<FDynamicMesh3, ESPMode::ThreadSafe>(MoveTemp(CombinedMesh)) );
+	CombinedInputMeshesApproximator = MakeShared<FMeshSimpleShapeApproximation, ESPMode::ThreadSafe>();
 	CombinedInputMeshesApproximator->InitializeSourceMeshes(MakeRawPointerList<FDynamicMesh3>(CombinedInputMeshes));
 
 	// build separated input meshes
 	SeparatedInputMeshes.Reset();
 	InitializeDerivedMeshSet(InputMeshes, SeparatedInputMeshes, 
 		[&](const FDynamicMesh3* Mesh, int32 Tri0, int32 Tri1) { return true; });
-	SeparatedMeshesApproximator = MakeShared<FMeshSimpleShapeApproximation>();
+	SeparatedMeshesApproximator = MakeShared<FMeshSimpleShapeApproximation, ESPMode::ThreadSafe>();
 	SeparatedMeshesApproximator->InitializeSourceMeshes(MakeRawPointerList<FDynamicMesh3>(SeparatedInputMeshes));
 
 	// build per-group input meshes
 	PerGroupInputMeshes.Reset();
 	InitializeDerivedMeshSet(InputMeshes, PerGroupInputMeshes,
 		[&](const FDynamicMesh3* Mesh, int32 Tri0, int32 Tri1) { return Mesh->GetTriangleGroup(Tri0) == Mesh->GetTriangleGroup(Tri1); });
-	PerGroupMeshesApproximator = MakeShared<FMeshSimpleShapeApproximation>();
+	PerGroupMeshesApproximator = MakeShared<FMeshSimpleShapeApproximation, ESPMode::ThreadSafe>();
 	PerGroupMeshesApproximator->InitializeSourceMeshes(MakeRawPointerList<FDynamicMesh3>(PerGroupInputMeshes));
 
 }
 
 
-TSharedPtr<FMeshSimpleShapeApproximation>& USetCollisionGeometryTool::GetApproximator(ESetCollisionGeometryInputMode MeshSetMode)
+TSharedPtr<FMeshSimpleShapeApproximation, ESPMode::ThreadSafe>& USetCollisionGeometryTool::GetApproximator(ESetCollisionGeometryInputMode MeshSetMode)
 {
 	if (MeshSetMode == ESetCollisionGeometryInputMode::CombineAll)
 	{
diff --git a/Engine/Plugins/Experimental/MeshModelingToolset/Source/MeshModelingTools/Private/PlaneCutTool.cpp b/Engine/Plugins/Experimental/MeshModelingToolset/Source/MeshModelingTools/Private/PlaneCutTool.cpp
--- a/Engine/Plugins/Experimental/MeshModelingToolset/Source/MeshModelingTools/Private/PlaneCutTool.cpp
+++ b/Engine/Plugins/Experimental/MeshModelingToolset/Source/MeshModelingTools/Private/PlaneCutTool.cpp
@@ -148,7 +148,7 @@ void UPlaneCutTool::Setup()
 		MeshUVScaleFactor.Add(1.0 / OriginalDynamicMesh->GetBounds().MaxDim());
 
 		// Set callbacks so previews are invalidated on undo/redo changing the meshes
-		Target->SetMesh(TSharedPtr<const FDynamicMesh3>(OriginalDynamicMesh));
+		Target->SetMesh(TSharedPtr<const FDynamicMesh3, ESPMode::ThreadSafe>(OriginalDynamicMesh));
 		Target->OnMeshChanged.AddLambda([this, Idx]() { Previews[Idx]->InvalidateResult(); });
 	}
 
@@ -283,7 +283,7 @@ void UPlaneCutTool::Cut()
 		UMeshOpPreviewWithBackgroundCompute* Preview = Previews[Idx];
 		TUniquePtr<FDynamicMesh3> ResultMesh = Preview->PreviewMesh->ExtractPreviewMesh();
 		ChangeSeq->AppendChange(MeshesToCut[Idx], MeshesToCut[Idx]->ReplaceMesh(
-			TSharedPtr<const FDynamicMesh3>(ResultMesh.Release()))
+			TSharedPtr<const FDynamicMesh3, ESPMode::ThreadSafe>(ResultMesh.Release()))
 		);
 	}
 
diff --git a/Engine/Plugins/Experimental/MeshModelingToolset/Source/MeshModelingTools/Private/RemeshMeshTool.cpp b/Engine/Plugins/Experimental/MeshModelingToolset/Source/MeshModelingTools/Private/RemeshMeshTool.cpp
--- a/Engine/Plugins/Experimental/MeshModelingToolset/Source/MeshModelingTools/Private/RemeshMeshTool.cpp
+++ b/Engine/Plugins/Experimental/MeshModelingToolset/Source/MeshModelingTools/Private/RemeshMeshTool.cpp
@@ -121,7 +121,7 @@ void URemeshMeshTool::Setup()
 	BasicProperties->WatchProperty(BasicProperties->bShowWireframe,
 								   [this](bool bNewValue) { UpdateVisualization();});
 
-	OriginalMesh = MakeShared<FDynamicMesh3>();
+	OriginalMesh = MakeShared<FDynamicMesh3, ESPMode::ThreadSafe>();
 	FMeshDescriptionToDynamicMesh Converter;
 	Converter.Convert(ComponentTarget->GetMesh(), *OriginalMesh);
 
@@ -129,7 +129,7 @@ void URemeshMeshTool::Setup()
 	Preview->PreviewMesh->SetTangentsMode(EDynamicMeshTangentCalcType::AutoCalculated);
 	Preview->PreviewMesh->UpdatePreview(OriginalMesh.Get());
 
-	OriginalMeshSpatial = MakeShared<FDynamicMeshAABBTree3>(OriginalMesh.Get(), true);
+	OriginalMeshSpatial = MakeShared<FDynamicMeshAABBTree3, ESPMode::ThreadSafe>(OriginalMesh.Get(), true);
 
 	// calculate initial mesh area (no utility fn yet)
 	// TODO: will need to change to account for component transform's Scale3D
diff --git a/Engine/Plugins/Experimental/MeshModelingToolset/Source/MeshModelingTools/Private/RemoveOccludedTrianglesTool.cpp b/Engine/Plugins/Experimental/MeshModelingToolset/Source/MeshModelingTools/Private/RemoveOccludedTrianglesTool.cpp
--- a/Engine/Plugins/Experimental/MeshModelingToolset/Source/MeshModelingTools/Private/RemoveOccludedTrianglesTool.cpp
+++ b/Engine/Plugins/Experimental/MeshModelingToolset/Source/MeshModelingTools/Private/RemoveOccludedTrianglesTool.cpp
@@ -199,7 +199,7 @@ void URemoveOccludedTrianglesTool::SetupPreviews()
 	int32 NumTargets = ComponentTargets.Num();
 	int32 NumPreviews = PreviewToTargetIdx.Num();
 
-	CombinedMeshTrees = MakeShared<IndexMeshWithAcceleration>();
+	CombinedMeshTrees = MakeShared<IndexMeshWithAcceleration, ESPMode::ThreadSafe>();
 
 #if WITH_EDITOR
 	static const FText SlowTaskText = LOCTEXT("RemoveOccludedTrianglesInit", "Building mesh occlusion data...");
@@ -257,7 +257,7 @@ void URemoveOccludedTrianglesTool::SetupPreviews()
 			URemoveOccludedTrianglesOperatorFactory *OpFactory = NewObject<URemoveOccludedTrianglesOperatorFactory>();
 			OpFactory->Tool = this;
 			OpFactory->PreviewIdx = PreviewIdx;
-			OriginalDynamicMeshes[PreviewIdx] = MakeShared<FDynamicMesh3>();
+			OriginalDynamicMeshes[PreviewIdx] = MakeShared<FDynamicMesh3, ESPMode::ThreadSafe>();
 			FMeshDescriptionToDynamicMesh Converter;
 			Converter.Convert(ComponentTargets[TargetIdx]->GetMesh(), *OriginalDynamicMeshes[PreviewIdx]);
 
diff --git a/Engine/Plugins/Experimental/MeshModelingToolset/Source/MeshModelingTools/Private/SeamSculptTool.cpp b/Engine/Plugins/Experimental/MeshModelingToolset/Source/MeshModelingTools/Private/SeamSculptTool.cpp
--- a/Engine/Plugins/Experimental/MeshModelingToolset/Source/MeshModelingTools/Private/SeamSculptTool.cpp
+++ b/Engine/Plugins/Experimental/MeshModelingToolset/Source/MeshModelingTools/Private/SeamSculptTool.cpp
@@ -61,7 +61,7 @@ void USeamSculptTool::Setup()
 	BrushProperties->RestoreProperties(this);
 
 	MeshTransform = FTransform3d(ComponentTarget->GetWorldTransform());
-	InputMesh = MakeShared<FDynamicMesh3>(*PreviewMesh->GetMesh());
+	InputMesh = MakeShared<FDynamicMesh3, ESPMode::ThreadSafe>(*PreviewMesh->GetMesh());
 	FMeshNormals::QuickComputeVertexNormals(*InputMesh);
 	NormalOffset = InputMesh->GetCachedBounds().MinDim() * 0.001;
 
diff --git a/Engine/Plugins/Experimental/MeshModelingToolset/Source/MeshModelingTools/Private/SelfUnionMeshesTool.cpp b/Engine/Plugins/Experimental/MeshModelingToolset/Source/MeshModelingTools/Private/SelfUnionMeshesTool.cpp
--- a/Engine/Plugins/Experimental/MeshModelingToolset/Source/MeshModelingTools/Private/SelfUnionMeshesTool.cpp
+++ b/Engine/Plugins/Experimental/MeshModelingToolset/Source/MeshModelingTools/Private/SelfUnionMeshesTool.cpp
@@ -89,7 +89,7 @@ void USelfUnionMeshesTool::ConvertInputsAndSetPreviewMaterials(bool bSetPreviewM
 		}
 	}
 
-	CombinedSourceMeshes = MakeShared<FDynamicMesh3>();
+	CombinedSourceMeshes = MakeShared<FDynamicMesh3, ESPMode::ThreadSafe>();
 	CombinedSourceMeshes->EnableAttributes();
 	CombinedSourceMeshes->EnableTriangleGroups(0);
 	CombinedSourceMeshes->Attributes()->EnableMaterialID();
diff --git a/Engine/Plugins/Experimental/MeshModelingToolset/Source/MeshModelingTools/Private/Tests/DisplaceMeshTool_Tests.inl b/Engine/Plugins/Experimental/MeshModelingToolset/Source/MeshModelingTools/Private/Tests/DisplaceMeshTool_Tests.inl
--- a/Engine/Plugins/Experimental/MeshModelingToolset/Source/MeshModelingTools/Private/Tests/DisplaceMeshTool_Tests.inl
+++ b/Engine/Plugins/Experimental/MeshModelingToolset/Source/MeshModelingTools/Private/Tests/DisplaceMeshTool_Tests.inl
@@ -13,9 +13,9 @@ namespace DisplaceMeshToolTestsLocals
 {
 	using namespace DisplaceMeshToolLocals;
 
-	TSharedPtr<FDynamicMesh3> CreateSingleTriangleMesh()
+	TSharedPtr<FDynamicMesh3, ESPMode::ThreadSafe> CreateSingleTriangleMesh()
 	{
-		TSharedPtr<FDynamicMesh3> Mesh = MakeShareable<FDynamicMesh3>(new FDynamicMesh3);
+		TSharedPtr<FDynamicMesh3, ESPMode::ThreadSafe> Mesh = MakeShareable<FDynamicMesh3>(new FDynamicMesh3);
 		int VertexIndexA = Mesh->AppendVertex(FVector3d{ 0,0,0 });
 		int VertexIndexB = Mesh->AppendVertex(FVector3d{ 0,1,0 });
 		int VertexIndexC = Mesh->AppendVertex(FVector3d{ 1,0,0 });
@@ -24,7 +24,7 @@ namespace DisplaceMeshToolTestsLocals
 		return Mesh;
 	}
 
-	TSharedPtr<FDynamicMesh3> CreateTestBoxMesh()
+	TSharedPtr<FDynamicMesh3, ESPMode::ThreadSafe> CreateTestBoxMesh()
 	{
 		FGridBoxMeshGenerator BoxGenerator;
 		BoxGenerator.Box = FOrientedBox3d(FVector3d::Zero(), FVector3d(1.0, 1.0, 1.0));
@@ -32,7 +32,7 @@ namespace DisplaceMeshToolTestsLocals
 		BoxGenerator.EdgeVertices = FIndex3i(EdgeNum, EdgeNum, EdgeNum);
 		BoxGenerator.Generate();
 
-		TSharedPtr<FDynamicMesh3> Mesh = MakeShareable<FDynamicMesh3>(new FDynamicMesh3(&BoxGenerator));
+		TSharedPtr<FDynamicMesh3, ESPMode::ThreadSafe> Mesh = MakeShareable<FDynamicMesh3>(new FDynamicMesh3(&BoxGenerator));
 		return Mesh;
 	}
 }
@@ -48,7 +48,7 @@ IMPLEMENT_SIMPLE_AUTOMATION_TEST(UDisplaceMeshTestConstant,
 	using namespace DisplaceMeshToolTestsLocals;
 
 	// ---- Setup
-	TSharedPtr<FDynamicMesh3> InputMesh = CreateSingleTriangleMesh();
+	TSharedPtr<FDynamicMesh3, ESPMode::ThreadSafe> InputMesh = CreateSingleTriangleMesh();
 
 	// Requirements:
 	UTEST_TRUE("Input mesh has 3 vertices", InputMesh->VertexCount() == 3);
@@ -95,7 +95,7 @@ IMPLEMENT_SIMPLE_AUTOMATION_TEST(UDisplaceMeshTestRandom,
 	using namespace DisplaceMeshToolTestsLocals;
 
 	// ---- Setup
-	TSharedPtr<FDynamicMesh3> InputMesh = CreateSingleTriangleMesh();
+	TSharedPtr<FDynamicMesh3, ESPMode::ThreadSafe> InputMesh = CreateSingleTriangleMesh();
 
 	EDisplaceMeshToolDisplaceType DisplaceType = EDisplaceMeshToolDisplaceType::RandomNoise;
 	DisplaceMeshParameters Params;
@@ -131,7 +131,7 @@ IMPLEMENT_SIMPLE_AUTOMATION_TEST(UDisplaceMeshTestPerlin,
 	using namespace DisplaceMeshToolTestsLocals;
 
 	// ---- Setup
-	TSharedPtr<FDynamicMesh3> InputMesh = CreateSingleTriangleMesh();
+	TSharedPtr<FDynamicMesh3, ESPMode::ThreadSafe> InputMesh = CreateSingleTriangleMesh();
 
 	EDisplaceMeshToolDisplaceType DisplaceType = EDisplaceMeshToolDisplaceType::PerlinNoise;
 	DisplaceMeshParameters Params;
@@ -168,7 +168,7 @@ IMPLEMENT_SIMPLE_AUTOMATION_TEST(UDisplaceMeshTestSine,
 	using namespace DisplaceMeshToolTestsLocals;
 
 	// ---- Setup
-	TSharedPtr<FDynamicMesh3> InputMesh = CreateSingleTriangleMesh();
+	TSharedPtr<FDynamicMesh3, ESPMode::ThreadSafe> InputMesh = CreateSingleTriangleMesh();
 
 	EDisplaceMeshToolDisplaceType DisplaceType = EDisplaceMeshToolDisplaceType::SineWave;
 	DisplaceMeshParameters Params;
@@ -213,7 +213,7 @@ IMPLEMENT_SIMPLE_AUTOMATION_TEST(UDisplaceMeshTestSineDirection,
 	using namespace DisplaceMeshToolTestsLocals;
 
 	// ---- Setup
-	TSharedPtr<FDynamicMesh3> InputMesh = CreateSingleTriangleMesh();
+	TSharedPtr<FDynamicMesh3, ESPMode::ThreadSafe> InputMesh = CreateSingleTriangleMesh();
 
 	EDisplaceMeshToolDisplaceType DisplaceType = EDisplaceMeshToolDisplaceType::SineWave;
 	DisplaceMeshParameters Params;
@@ -266,7 +266,7 @@ IMPLEMENT_SIMPLE_AUTOMATION_TEST(UDisplaceMeshTestDirectionalFilter,
 
 	// ---- Setup
 	// Create a box mesh, random noise, and a directional filter pointing in the +Y direction
-	TSharedPtr<FDynamicMesh3> InputMesh = CreateTestBoxMesh();
+	TSharedPtr<FDynamicMesh3, ESPMode::ThreadSafe> InputMesh = CreateTestBoxMesh();
 
 	FMeshNormals InputMeshNormals(InputMesh.Get());
 	InputMeshNormals.ComputeVertexNormals();
diff --git a/Engine/Plugins/Experimental/MeshModelingToolset/Source/MeshModelingTools/Private/UVProjectionTool.cpp b/Engine/Plugins/Experimental/MeshModelingToolset/Source/MeshModelingTools/Private/UVProjectionTool.cpp
--- a/Engine/Plugins/Experimental/MeshModelingToolset/Source/MeshModelingTools/Private/UVProjectionTool.cpp
+++ b/Engine/Plugins/Experimental/MeshModelingToolset/Source/MeshModelingTools/Private/UVProjectionTool.cpp
@@ -159,7 +159,7 @@ void UUVProjectionTool::UpdateNumPreviews()
 			UUVProjectionOperatorFactory *OpFactory = NewObject<UUVProjectionOperatorFactory>();
 			OpFactory->Tool = this;
 			OpFactory->ComponentIndex = PreviewIdx;
-			OriginalDynamicMeshes[PreviewIdx] = MakeShared<FDynamicMesh3>();
+			OriginalDynamicMeshes[PreviewIdx] = MakeShared<FDynamicMesh3, ESPMode::ThreadSafe>();
 			FMeshDescriptionToDynamicMesh Converter;
 			Converter.Convert(ComponentTargets[PreviewIdx]->GetMesh(), *OriginalDynamicMeshes[PreviewIdx]);
 
diff --git a/Engine/Plugins/Experimental/MeshModelingToolset/Source/MeshModelingTools/Public/BakeMeshAttributeMapsTool.h b/Engine/Plugins/Experimental/MeshModelingToolset/Source/MeshModelingTools/Public/BakeMeshAttributeMapsTool.h
--- a/Engine/Plugins/Experimental/MeshModelingToolset/Source/MeshModelingTools/Public/BakeMeshAttributeMapsTool.h
+++ b/Engine/Plugins/Experimental/MeshModelingToolset/Source/MeshModelingTools/Public/BakeMeshAttributeMapsTool.h
@@ -300,15 +300,15 @@ protected:
 	UPROPERTY()
 	UMaterialInstanceDynamic* PreviewMaterial;
 
-	TSharedPtr<FMeshDescription> BaseMeshDescription;
-	TSharedPtr<TMeshTangents<double>> BaseMeshTangents;
+	TSharedPtr<FMeshDescription, ESPMode::ThreadSafe> BaseMeshDescription;
+	TSharedPtr<TMeshTangents<double>, ESPMode::ThreadSafe> BaseMeshTangents;
 	FDynamicMesh3 BaseMesh;
 	FDynamicMeshAABBTree3 BaseSpatial;
 
 	bool bIsBakeToSelf = false;
 
-	TSharedPtr<FDynamicMesh3> DetailMesh;
-	TSharedPtr<FDynamicMeshAABBTree3> DetailSpatial;
+	TSharedPtr<FDynamicMesh3, ESPMode::ThreadSafe> DetailMesh;
+	TSharedPtr<FDynamicMeshAABBTree3, ESPMode::ThreadSafe> DetailSpatial;
 	int32 DetailMeshTimestamp = 0;
 	void UpdateDetailMesh();
 	bool bDetailMeshValid = false;
diff --git a/Engine/Plugins/Experimental/MeshModelingToolset/Source/MeshModelingTools/Public/CSGMeshesTool.h b/Engine/Plugins/Experimental/MeshModelingToolset/Source/MeshModelingTools/Public/CSGMeshesTool.h
--- a/Engine/Plugins/Experimental/MeshModelingToolset/Source/MeshModelingTools/Public/CSGMeshesTool.h
+++ b/Engine/Plugins/Experimental/MeshModelingToolset/Source/MeshModelingTools/Public/CSGMeshesTool.h
@@ -105,7 +105,7 @@ protected:
 	UPROPERTY()
 	UTrimMeshesToolProperties* TrimProperties;
 
-	TArray<TSharedPtr<FDynamicMesh3>> OriginalDynamicMeshes;
+	TArray<TSharedPtr<FDynamicMesh3, ESPMode::ThreadSafe>> OriginalDynamicMeshes;
 
 	UPROPERTY()
 	ULineSetComponent* DrawnLineSet;
diff --git a/Engine/Plugins/Experimental/MeshModelingToolset/Source/MeshModelingTools/Public/DisplaceMeshTool.h b/Engine/Plugins/Experimental/MeshModelingToolset/Source/MeshModelingTools/Public/DisplaceMeshTool.h
--- a/Engine/Plugins/Experimental/MeshModelingToolset/Source/MeshModelingTools/Public/DisplaceMeshTool.h
+++ b/Engine/Plugins/Experimental/MeshModelingToolset/Source/MeshModelingTools/Public/DisplaceMeshTool.h
@@ -285,11 +285,11 @@ private:
 
 	FDynamicMesh3 OriginalMesh;
 	FDynamicMeshAABBTree3 OriginalMeshSpatial;
-	TSharedPtr<FIndexedWeightMap> ActiveWeightMap;
+	TSharedPtr<FIndexedWeightMap, ESPMode::ThreadSafe> ActiveWeightMap;
 	void UpdateActiveWeightMap();
 	float WeightMapQuery(const FVector3d& Position, const FIndexedWeightMap& WeightMap) const;
 
-	TSharedPtr<FDynamicMesh3> SubdividedMesh = nullptr;
+	TSharedPtr<FDynamicMesh3, ESPMode::ThreadSafe> SubdividedMesh = nullptr;
 	USimpleDynamicMeshComponent* DynamicMeshComponent = nullptr;
 
 	TUniquePtr<IDynamicMeshOperatorFactory> Subdivider = nullptr;
diff --git a/Engine/Plugins/Experimental/MeshModelingToolset/Source/MeshModelingTools/Public/EditNormalsTool.h b/Engine/Plugins/Experimental/MeshModelingToolset/Source/MeshModelingTools/Public/EditNormalsTool.h
--- a/Engine/Plugins/Experimental/MeshModelingToolset/Source/MeshModelingTools/Public/EditNormalsTool.h
+++ b/Engine/Plugins/Experimental/MeshModelingToolset/Source/MeshModelingTools/Public/EditNormalsTool.h
@@ -168,7 +168,7 @@ protected:
 
 
 protected:
-	TArray<TSharedPtr<FDynamicMesh3>> OriginalDynamicMeshes;
+	TArray<TSharedPtr<FDynamicMesh3, ESPMode::ThreadSafe>> OriginalDynamicMeshes;
 
 	UWorld* TargetWorld;
 	IAssetGenerationAPI* AssetAPI;
diff --git a/Engine/Plugins/Experimental/MeshModelingToolset/Source/MeshModelingTools/Public/GroupEdgeInsertionTool.h b/Engine/Plugins/Experimental/MeshModelingToolset/Source/MeshModelingTools/Public/GroupEdgeInsertionTool.h
--- a/Engine/Plugins/Experimental/MeshModelingToolset/Source/MeshModelingTools/Public/GroupEdgeInsertionTool.h
+++ b/Engine/Plugins/Experimental/MeshModelingToolset/Source/MeshModelingTools/Public/GroupEdgeInsertionTool.h
@@ -129,8 +129,8 @@ protected:
 	UWorld* TargetWorld;
 	IToolsContextAssetAPI* AssetAPI;
 
-	TSharedPtr<FDynamicMesh3> CurrentMesh;
-	TSharedPtr<FGroupTopology> CurrentTopology;
+	TSharedPtr<FDynamicMesh3, ESPMode::ThreadSafe> CurrentMesh;
+	TSharedPtr<FGroupTopology, ESPMode::ThreadSafe> CurrentTopology;
 	FDynamicMeshAABBTree3 MeshSpatial;
 	FGroupTopologySelector TopologySelector;
 
@@ -161,7 +161,7 @@ protected:
 	EToolState ToolState = EToolState::GettingStart;
 	bool bShowingBaseMesh = false;
 	bool bLastComputeSucceeded = false;
-	TSharedPtr<FGroupTopology> LatestOpTopologyResult;
+	TSharedPtr<FGroupTopology, ESPMode::ThreadSafe> LatestOpTopologyResult;
 
 	int32 CurrentChangeStamp = 0;
 
diff --git a/Engine/Plugins/Experimental/MeshModelingToolset/Source/MeshModelingTools/Public/HoleFillTool.h b/Engine/Plugins/Experimental/MeshModelingToolset/Source/MeshModelingTools/Public/HoleFillTool.h
--- a/Engine/Plugins/Experimental/MeshModelingToolset/Source/MeshModelingTools/Public/HoleFillTool.h
+++ b/Engine/Plugins/Experimental/MeshModelingToolset/Source/MeshModelingTools/Public/HoleFillTool.h
@@ -255,7 +255,7 @@ protected:
 	UPolygonSelectionMechanic* SelectionMechanic = nullptr;
 
 	// Input mesh. Ownership shared with Op.
-	TSharedPtr<FDynamicMesh3> OriginalMesh;
+	TSharedPtr<FDynamicMesh3, ESPMode::ThreadSafe> OriginalMesh;
 
 	// UV Scale factor is cached based on the bounding box of the mesh before any fills are performed
 	float MeshUVScaleFactor = 0.0f;
diff --git a/Engine/Plugins/Experimental/MeshModelingToolset/Source/MeshModelingTools/Public/LatticeDeformerTool.h b/Engine/Plugins/Experimental/MeshModelingToolset/Source/MeshModelingTools/Public/LatticeDeformerTool.h
--- a/Engine/Plugins/Experimental/MeshModelingToolset/Source/MeshModelingTools/Public/LatticeDeformerTool.h
+++ b/Engine/Plugins/Experimental/MeshModelingToolset/Source/MeshModelingTools/Public/LatticeDeformerTool.h
@@ -122,9 +122,9 @@ protected:
 	IToolsContextAssetAPI* AssetAPI;
 
 	// Input mesh
-	TSharedPtr<FDynamicMesh3> OriginalMesh;
+	TSharedPtr<FDynamicMesh3, ESPMode::ThreadSafe> OriginalMesh;
 
-	TSharedPtr<FFFDLattice> Lattice;
+	TSharedPtr<FFFDLattice, ESPMode::ThreadSafe> Lattice;
 
 	UPROPERTY()
 	ULatticeControlPointsMechanic* ControlPointsMechanic = nullptr;
diff --git a/Engine/Plugins/Experimental/MeshModelingToolset/Source/MeshModelingTools/Public/MeshBoundaryToolBase.h b/Engine/Plugins/Experimental/MeshModelingToolset/Source/MeshModelingTools/Public/MeshBoundaryToolBase.h
--- a/Engine/Plugins/Experimental/MeshModelingToolset/Source/MeshModelingTools/Public/MeshBoundaryToolBase.h
+++ b/Engine/Plugins/Experimental/MeshModelingToolset/Source/MeshModelingTools/Public/MeshBoundaryToolBase.h
@@ -42,7 +42,7 @@ protected:
 
 	UWorld* TargetWorld;
 
-	TSharedPtr<FDynamicMesh3> OriginalMesh;
+	TSharedPtr<FDynamicMesh3, ESPMode::ThreadSafe> OriginalMesh;
 
 	// Used for hit querying
 	FDynamicMeshAABBTree3 MeshSpatial;
diff --git a/Engine/Plugins/Experimental/MeshModelingToolset/Source/MeshModelingTools/Public/MeshSpaceDeformerTool.h b/Engine/Plugins/Experimental/MeshModelingToolset/Source/MeshModelingTools/Public/MeshSpaceDeformerTool.h
--- a/Engine/Plugins/Experimental/MeshModelingToolset/Source/MeshModelingTools/Public/MeshSpaceDeformerTool.h
+++ b/Engine/Plugins/Experimental/MeshModelingToolset/Source/MeshModelingTools/Public/MeshSpaceDeformerTool.h
@@ -241,7 +241,7 @@ protected:
 
 protected:	
 
-	TSharedPtr<FDynamicMesh3> OriginalDynamicMesh;
+	TSharedPtr<FDynamicMesh3, ESPMode::ThreadSafe> OriginalDynamicMesh;
 
 	UPROPERTY()
 	UPreviewMesh* OriginalMeshPreview;
diff --git a/Engine/Plugins/Experimental/MeshModelingToolset/Source/MeshModelingTools/Public/Physics/SetCollisionGeometryTool.h b/Engine/Plugins/Experimental/MeshModelingToolset/Source/MeshModelingTools/Public/Physics/SetCollisionGeometryTool.h
--- a/Engine/Plugins/Experimental/MeshModelingToolset/Source/MeshModelingTools/Public/Physics/SetCollisionGeometryTool.h
+++ b/Engine/Plugins/Experimental/MeshModelingToolset/Source/MeshModelingTools/Public/Physics/SetCollisionGeometryTool.h
@@ -191,29 +191,29 @@ protected:
 	};
 	
 	bool bInputMeshesValid = false;
-	TArray<TSharedPtr<FDynamicMesh3>> InputMeshes;
-	TArray<TSharedPtr<FDynamicMesh3>> CombinedInputMeshes;
-	TArray<TSharedPtr<FDynamicMesh3>> SeparatedInputMeshes;
-	TArray<TSharedPtr<FDynamicMesh3>> PerGroupInputMeshes;
+	TArray<TSharedPtr<FDynamicMesh3, ESPMode::ThreadSafe>> InputMeshes;
+	TArray<TSharedPtr<FDynamicMesh3, ESPMode::ThreadSafe>> CombinedInputMeshes;
+	TArray<TSharedPtr<FDynamicMesh3, ESPMode::ThreadSafe>> SeparatedInputMeshes;
+	TArray<TSharedPtr<FDynamicMesh3, ESPMode::ThreadSafe>> PerGroupInputMeshes;
 
-	TSharedPtr<FMeshSimpleShapeApproximation> InputMeshesApproximator;
-	TSharedPtr<FMeshSimpleShapeApproximation> CombinedInputMeshesApproximator;
-	TSharedPtr<FMeshSimpleShapeApproximation> SeparatedMeshesApproximator;
-	TSharedPtr<FMeshSimpleShapeApproximation> PerGroupMeshesApproximator;
+	TSharedPtr<FMeshSimpleShapeApproximation, ESPMode::ThreadSafe> InputMeshesApproximator;
+	TSharedPtr<FMeshSimpleShapeApproximation, ESPMode::ThreadSafe> CombinedInputMeshesApproximator;
+	TSharedPtr<FMeshSimpleShapeApproximation, ESPMode::ThreadSafe> SeparatedMeshesApproximator;
+	TSharedPtr<FMeshSimpleShapeApproximation, ESPMode::ThreadSafe> PerGroupMeshesApproximator;
 
 	void PrecomputeInputMeshes();
 	void InitializeDerivedMeshSet(
-		const TArray<TSharedPtr<FDynamicMesh3>>& FromInputMeshes, 
-		TArray<TSharedPtr<FDynamicMesh3>>& ToMeshes,
+		const TArray<TSharedPtr<FDynamicMesh3, ESPMode::ThreadSafe>>& FromInputMeshes,
+		TArray<TSharedPtr<FDynamicMesh3, ESPMode::ThreadSafe>>& ToMeshes,
 		TFunctionRef<bool(const FDynamicMesh3*, int32, int32)> TrisConnectedPredicate);
-	TSharedPtr<FMeshSimpleShapeApproximation>& GetApproximator(ESetCollisionGeometryInputMode MeshSetMode);
+	TSharedPtr<FMeshSimpleShapeApproximation, ESPMode::ThreadSafe>& GetApproximator(ESetCollisionGeometryInputMode MeshSetMode);
 
 	FTransform OrigTargetTransform;
 	FVector TargetScale3D;
 
 	bool bResultValid = false;
-	TSharedPtr<FPhysicsDataCollection> InitialCollision;
-	TSharedPtr<FPhysicsDataCollection> GeneratedCollision;
+	TSharedPtr<FPhysicsDataCollection, ESPMode::ThreadSafe> InitialCollision;
+	TSharedPtr<FPhysicsDataCollection, ESPMode::ThreadSafe> GeneratedCollision;
 
 	void UpdateGeneratedCollision();
 //	TSharedPtr<FPhysicsDataCollection> GenerateCollision_MinVolume();
diff --git a/Engine/Plugins/Experimental/MeshModelingToolset/Source/MeshModelingTools/Public/RemeshMeshTool.h b/Engine/Plugins/Experimental/MeshModelingToolset/Source/MeshModelingTools/Public/RemeshMeshTool.h
--- a/Engine/Plugins/Experimental/MeshModelingToolset/Source/MeshModelingTools/Public/RemeshMeshTool.h
+++ b/Engine/Plugins/Experimental/MeshModelingToolset/Source/MeshModelingTools/Public/RemeshMeshTool.h
@@ -150,8 +150,8 @@ protected:
 	UWorld* TargetWorld;
 	IToolsContextAssetAPI* AssetAPI;
 
-	TSharedPtr<FDynamicMesh3> OriginalMesh;
-	TSharedPtr<FDynamicMeshAABBTree3> OriginalMeshSpatial;
+	TSharedPtr<FDynamicMesh3, ESPMode::ThreadSafe> OriginalMesh;
+	TSharedPtr<FDynamicMeshAABBTree3, ESPMode::ThreadSafe> OriginalMeshSpatial;
 	double InitialMeshArea;
 
 	double CalculateTargetEdgeLength(int TargetTriCount);
diff --git a/Engine/Plugins/Experimental/MeshModelingToolset/Source/MeshModelingTools/Public/RemoveOccludedTrianglesTool.h b/Engine/Plugins/Experimental/MeshModelingToolset/Source/MeshModelingTools/Public/RemoveOccludedTrianglesTool.h
--- a/Engine/Plugins/Experimental/MeshModelingToolset/Source/MeshModelingTools/Public/RemoveOccludedTrianglesTool.h
+++ b/Engine/Plugins/Experimental/MeshModelingToolset/Source/MeshModelingTools/Public/RemoveOccludedTrianglesTool.h
@@ -201,10 +201,10 @@ protected:
 
 
 protected:
-	TArray<TSharedPtr<FDynamicMesh3>> OriginalDynamicMeshes;
+	TArray<TSharedPtr<FDynamicMesh3, ESPMode::ThreadSafe>> OriginalDynamicMeshes;
 
 	// pre-transformed combined all-mesh data (with mesh geometry data in just raw index buffers as that's all we need)
-	TSharedPtr<IndexMeshWithAcceleration> CombinedMeshTrees;
+	TSharedPtr<IndexMeshWithAcceleration, ESPMode::ThreadSafe> CombinedMeshTrees;
 
 	TArray<TArray<int32>> PreviewToCopyIdx;
 	TArray<int32> PreviewToTargetIdx;
diff --git a/Engine/Plugins/Experimental/MeshModelingToolset/Source/MeshModelingTools/Public/SeamSculptTool.h b/Engine/Plugins/Experimental/MeshModelingToolset/Source/MeshModelingTools/Public/SeamSculptTool.h
--- a/Engine/Plugins/Experimental/MeshModelingToolset/Source/MeshModelingTools/Public/SeamSculptTool.h
+++ b/Engine/Plugins/Experimental/MeshModelingToolset/Source/MeshModelingTools/Public/SeamSculptTool.h
@@ -84,7 +84,7 @@ protected:
 	UPROPERTY()
 	UPreviewGeometry* PreviewGeom;
 
-	TSharedPtr<FDynamicMesh3> InputMesh;
+	TSharedPtr<FDynamicMesh3, ESPMode::ThreadSafe> InputMesh;
 	FTransform3d MeshTransform;
 	double NormalOffset = 0;
 
diff --git a/Engine/Plugins/Experimental/MeshModelingToolset/Source/MeshModelingTools/Public/SelfUnionMeshesTool.h b/Engine/Plugins/Experimental/MeshModelingToolset/Source/MeshModelingTools/Public/SelfUnionMeshesTool.h
--- a/Engine/Plugins/Experimental/MeshModelingToolset/Source/MeshModelingTools/Public/SelfUnionMeshesTool.h
+++ b/Engine/Plugins/Experimental/MeshModelingToolset/Source/MeshModelingTools/Public/SelfUnionMeshesTool.h
@@ -93,7 +93,7 @@ protected:
 	UPROPERTY()
 	ULineSetComponent* DrawnLineSet;
 
-	TSharedPtr<FDynamicMesh3> CombinedSourceMeshes;
+	TSharedPtr<FDynamicMesh3, ESPMode::ThreadSafe> CombinedSourceMeshes;
 
 	// for visualization of any errors in the currently-previewed merge operation
 	TArray<int> CreatedBoundaryEdges;
diff --git a/Engine/Plugins/Experimental/MeshModelingToolset/Source/MeshModelingTools/Public/UVProjectionTool.h b/Engine/Plugins/Experimental/MeshModelingToolset/Source/MeshModelingTools/Public/UVProjectionTool.h
--- a/Engine/Plugins/Experimental/MeshModelingToolset/Source/MeshModelingTools/Public/UVProjectionTool.h
+++ b/Engine/Plugins/Experimental/MeshModelingToolset/Source/MeshModelingTools/Public/UVProjectionTool.h
@@ -177,7 +177,7 @@ protected:
 	TArray<UTransformProxy*> TransformProxies;
 
 protected:
-	TArray<TSharedPtr<FDynamicMesh3>> OriginalDynamicMeshes;
+	TArray<TSharedPtr<FDynamicMesh3, ESPMode::ThreadSafe>> OriginalDynamicMeshes;
 	TArray<FDynamicMesh3> ReferencePrimitives;
 
 	UWorld* TargetWorld;
diff --git a/Engine/Plugins/Experimental/MeshModelingToolset/Source/MeshModelingToolsEditorOnly/Private/GenerateLODMeshesTool.cpp b/Engine/Plugins/Experimental/MeshModelingToolset/Source/MeshModelingToolsEditorOnly/Private/GenerateLODMeshesTool.cpp
--- a/Engine/Plugins/Experimental/MeshModelingToolset/Source/MeshModelingToolsEditorOnly/Private/GenerateLODMeshesTool.cpp
+++ b/Engine/Plugins/Experimental/MeshModelingToolset/Source/MeshModelingToolsEditorOnly/Private/GenerateLODMeshesTool.cpp
@@ -116,13 +116,13 @@ void UGenerateLODMeshesTool::Setup()
 #else
 		auto EnterProgressFrame = [](int Progress) {};
 #endif
-		OriginalMeshDescription = MakeShared<FMeshDescription>(*ComponentTarget->GetMesh());
+		OriginalMeshDescription = MakeShared<FMeshDescription, ESPMode::ThreadSafe>(*ComponentTarget->GetMesh());
 		EnterProgressFrame(1);
-		OriginalMesh = MakeShared<FDynamicMesh3>();
+		OriginalMesh = MakeShared<FDynamicMesh3, ESPMode::ThreadSafe>();
 		FMeshDescriptionToDynamicMesh Converter;
 		Converter.Convert(ComponentTarget->GetMesh(), *OriginalMesh);
 		EnterProgressFrame(2);
-		OriginalMeshSpatial = MakeShared<FDynamicMeshAABBTree3>(OriginalMesh.Get(), true);
+		OriginalMeshSpatial = MakeShared<FDynamicMeshAABBTree3, ESPMode::ThreadSafe>(OriginalMesh.Get(), true);
 	}
 
 	WorldBounds = FAxisAlignedBox3d(ComponentTarget->GetOwnerActor()->GetComponentsBoundingBox());
diff --git a/Engine/Plugins/Experimental/MeshModelingToolset/Source/MeshModelingToolsEditorOnly/Private/Hair/GroomToMeshTool.cpp b/Engine/Plugins/Experimental/MeshModelingToolset/Source/MeshModelingToolsEditorOnly/Private/Hair/GroomToMeshTool.cpp
--- a/Engine/Plugins/Experimental/MeshModelingToolset/Source/MeshModelingToolsEditorOnly/Private/Hair/GroomToMeshTool.cpp
+++ b/Engine/Plugins/Experimental/MeshModelingToolset/Source/MeshModelingToolsEditorOnly/Private/Hair/GroomToMeshTool.cpp
@@ -237,10 +237,10 @@ public:
 
 void UGroomToMeshTool::RecalculateMesh()
 {
-	TSharedPtr<FDynamicMesh3> CurrentResult = UpdateVoxelization();
+	TSharedPtr<FDynamicMesh3, ESPMode::ThreadSafe> CurrentResult = UpdateVoxelization();
 
 	// helper to only use NewMesh if it was valid, otherwise ignore it
-	auto SelectValidMesh = [](TSharedPtr<FDynamicMesh3> CurrentMesh, TSharedPtr<FDynamicMesh3> NewMesh)
+	auto SelectValidMesh = [](TSharedPtr<FDynamicMesh3, ESPMode::ThreadSafe> CurrentMesh, TSharedPtr<FDynamicMesh3, ESPMode::ThreadSafe> NewMesh)
 	{
 		if (NewMesh->TriangleCount() > 0)
 		{
@@ -393,7 +393,7 @@ static void ProcessHairCurves(AGroomActor* GroomActor,
 
 
 
-TSharedPtr<FDynamicMesh3> UGroomToMeshTool::UpdateVoxelization()
+TSharedPtr<FDynamicMesh3, ESPMode::ThreadSafe> UGroomToMeshTool::UpdateVoxelization()
 {
 	FVoxelizeSettings NewSettings;
 	NewSettings.VoxelCount = Settings->VoxelCount;
@@ -498,7 +498,7 @@ TSharedPtr<FDynamicMesh3> UGroomToMeshTool::UpdateVoxelization()
 	// clear implicit function...
 	MarchingCubes.Implicit = nullptr;
 
-	CurrentVoxelizeResult = MakeShared<FDynamicMesh3>(&MarchingCubes);
+	CurrentVoxelizeResult = MakeShared<FDynamicMesh3, ESPMode::ThreadSafe>(&MarchingCubes);
 	CachedVoxelizeSettings = NewSettings;
 	return CurrentVoxelizeResult;
 }
@@ -506,7 +506,7 @@ TSharedPtr<FDynamicMesh3> UGroomToMeshTool::UpdateVoxelization()
 
 
 
-TSharedPtr<FDynamicMesh3> UGroomToMeshTool::UpdateMorphology(TSharedPtr<FDynamicMesh3> InputMesh)
+TSharedPtr<FDynamicMesh3, ESPMode::ThreadSafe> UGroomToMeshTool::UpdateMorphology(TSharedPtr<FDynamicMesh3, ESPMode::ThreadSafe> InputMesh)
 {
 	FMorphologySettings NewSettings;
 	NewSettings.InputMesh = InputMesh;
@@ -518,7 +518,7 @@ TSharedPtr<FDynamicMesh3> UGroomToMeshTool::UpdateMorphology(TSharedPtr<FDynamic
 		return CachedMorphologyResult;
 	}
 
-	CachedMorphologyResult = MakeShared<FDynamicMesh3>(*InputMesh);
+	CachedMorphologyResult = MakeShared<FDynamicMesh3, ESPMode::ThreadSafe>(*InputMesh);
 
 	if (NewSettings.CloseDist > 0 && CachedMorphologyResult->TriangleCount() > 0)
 	{
@@ -529,10 +529,10 @@ TSharedPtr<FDynamicMesh3> UGroomToMeshTool::UpdateMorphology(TSharedPtr<FDynamic
 		ImplicitMorphology.SourceSpatial = &Spatial;
 		ImplicitMorphology.SetCellSizesAndDistance(CachedMorphologyResult->GetCachedBounds(), NewSettings.CloseDist, NewSettings.VoxelCount, NewSettings.VoxelCount);
 
-		CachedMorphologyResult = MakeShared<FDynamicMesh3>(&ImplicitMorphology.Generate());
+		CachedMorphologyResult = MakeShared<FDynamicMesh3, ESPMode::ThreadSafe>(&ImplicitMorphology.Generate());
 		if (CachedMorphologyResult->TriangleCount() == 0)
 		{
-			CachedPostprocessResult = MakeShared<FDynamicMesh3>(*InputMesh);
+			CachedPostprocessResult = MakeShared<FDynamicMesh3, ESPMode::ThreadSafe>(*InputMesh);
 		}
 	}
 
@@ -546,10 +546,10 @@ TSharedPtr<FDynamicMesh3> UGroomToMeshTool::UpdateMorphology(TSharedPtr<FDynamic
 		ImplicitMorphology.SourceSpatial = &Spatial;
 		ImplicitMorphology.SetCellSizesAndDistance(CachedMorphologyResult->GetCachedBounds(), NewSettings.OpenDist, NewSettings.VoxelCount, NewSettings.VoxelCount);
 
-		CachedMorphologyResult = MakeShared<FDynamicMesh3>(&ImplicitMorphology.Generate());
+		CachedMorphologyResult = MakeShared<FDynamicMesh3, ESPMode::ThreadSafe>(&ImplicitMorphology.Generate());
 		if (CachedMorphologyResult->TriangleCount() == 0)
 		{
-			CachedPostprocessResult = MakeShared<FDynamicMesh3>(*InputMesh);
+			CachedPostprocessResult = MakeShared<FDynamicMesh3, ESPMode::ThreadSafe>(*InputMesh);
 		}
 	}
 
@@ -794,7 +794,7 @@ public:
 
 
 
-TSharedPtr<FDynamicMesh3> UGroomToMeshTool::UpdateClipMesh(TSharedPtr<FDynamicMesh3> InputMesh)
+TSharedPtr<FDynamicMesh3, ESPMode::ThreadSafe> UGroomToMeshTool::UpdateClipMesh(TSharedPtr<FDynamicMesh3, ESPMode::ThreadSafe> InputMesh)
 {
 	FClipMeshSettings NewSettings;
 	NewSettings.InputMesh = InputMesh;
@@ -804,7 +804,7 @@ TSharedPtr<FDynamicMesh3> UGroomToMeshTool::UpdateClipMesh(TSharedPtr<FDynamicMe
 		return CachedClipMeshResult;
 	}
 
-	CachedClipMeshResult = MakeShared<FDynamicMesh3>(*InputMesh);
+	CachedClipMeshResult = MakeShared<FDynamicMesh3, ESPMode::ThreadSafe>(*InputMesh);
 	FDynamicMesh3* EditMesh = CachedClipMeshResult.Get();
 
 	// do a quick tiny-edges collapse pass
@@ -866,7 +866,7 @@ TSharedPtr<FDynamicMesh3> UGroomToMeshTool::UpdateClipMesh(TSharedPtr<FDynamicMe
 
 
 
-TSharedPtr<FDynamicMesh3> UGroomToMeshTool::UpdateSmoothing(TSharedPtr<FDynamicMesh3> InputMesh)
+TSharedPtr<FDynamicMesh3, ESPMode::ThreadSafe> UGroomToMeshTool::UpdateSmoothing(TSharedPtr<FDynamicMesh3, ESPMode::ThreadSafe> InputMesh)
 {
 	FSmoothingSettings NewSettings;
 	NewSettings.InputMesh = InputMesh;
@@ -877,7 +877,7 @@ TSharedPtr<FDynamicMesh3> UGroomToMeshTool::UpdateSmoothing(TSharedPtr<FDynamicM
 		return CachedSmoothResult;
 	}
 
-	CachedSmoothResult = MakeShared<FDynamicMesh3>(*InputMesh);
+	CachedSmoothResult = MakeShared<FDynamicMesh3, ESPMode::ThreadSafe>(*InputMesh);
 
 	if (NewSettings.Smoothness > 0)
 	{
@@ -925,7 +925,7 @@ TSharedPtr<FDynamicMesh3> UGroomToMeshTool::UpdateSmoothing(TSharedPtr<FDynamicM
 		MeshTransforms::Scale(*NewMesh, FVector3d(SrcScale, SrcScale, SrcScale), FVector3d::Zero());
 		MeshTransforms::Translate(*NewMesh, SrcTranslate);
 
-		CachedSmoothResult = TSharedPtr<FDynamicMesh3>(NewMesh.Release());
+		CachedSmoothResult = TSharedPtr<FDynamicMesh3, ESPMode::ThreadSafe>(NewMesh.Release());
 	}
 
 
@@ -937,7 +937,7 @@ TSharedPtr<FDynamicMesh3> UGroomToMeshTool::UpdateSmoothing(TSharedPtr<FDynamicM
 
 
 
-TSharedPtr<FDynamicMesh3> UGroomToMeshTool::UpdateSimplification(TSharedPtr<FDynamicMesh3> InputMesh)
+TSharedPtr<FDynamicMesh3, ESPMode::ThreadSafe> UGroomToMeshTool::UpdateSimplification(TSharedPtr<FDynamicMesh3, ESPMode::ThreadSafe> InputMesh)
 {
 	FSimplifySettings NewSettings;
 	NewSettings.InputMesh = InputMesh;
@@ -947,7 +947,7 @@ TSharedPtr<FDynamicMesh3> UGroomToMeshTool::UpdateSimplification(TSharedPtr<FDyn
 		return CachedSimplifyResult;
 	}
 
-	CachedSimplifyResult = MakeShared<FDynamicMesh3>(*InputMesh);
+	CachedSimplifyResult = MakeShared<FDynamicMesh3, ESPMode::ThreadSafe>(*InputMesh);
 
 	FQEMSimplification Simplifier(CachedSimplifyResult.Get());
 	Simplifier.SimplifyToVertexCount(NewSettings.TargetCount);
@@ -959,7 +959,7 @@ TSharedPtr<FDynamicMesh3> UGroomToMeshTool::UpdateSimplification(TSharedPtr<FDyn
 
 
 
-TSharedPtr<FDynamicMesh3> UGroomToMeshTool::UpdatePostprocessing(TSharedPtr<FDynamicMesh3> InputMesh)
+TSharedPtr<FDynamicMesh3, ESPMode::ThreadSafe> UGroomToMeshTool::UpdatePostprocessing(TSharedPtr<FDynamicMesh3, ESPMode::ThreadSafe> InputMesh)
 {
 	FPostprocessSettings NewSettings;
 	NewSettings.InputMesh = InputMesh;
@@ -969,7 +969,7 @@ TSharedPtr<FDynamicMesh3> UGroomToMeshTool::UpdatePostprocessing(TSharedPtr<FDyn
 		return CachedPostprocessResult;
 	}
 
-	CachedPostprocessResult = MakeShared<FDynamicMesh3>(*InputMesh);
+	CachedPostprocessResult = MakeShared<FDynamicMesh3, ESPMode::ThreadSafe>(*InputMesh);
 
 	FMeshNormals::InitializeMeshToPerTriangleNormals(CachedPostprocessResult.Get());
 
@@ -1009,9 +1009,9 @@ void PartitionIndices(EnumerableType Enumerable,
 
 
 
-TSharedPtr<FDynamicMesh3> UGroomToMeshTool::UpdateUVs(TSharedPtr<FDynamicMesh3> InputMesh, EGroomToMeshUVMode UVMode)
+TSharedPtr<FDynamicMesh3, ESPMode::ThreadSafe> UGroomToMeshTool::UpdateUVs(TSharedPtr<FDynamicMesh3, ESPMode::ThreadSafe> InputMesh, EGroomToMeshUVMode UVMode)
 {
-	TSharedPtr<FDynamicMesh3> Result;
+	TSharedPtr<FDynamicMesh3, ESPMode::ThreadSafe> Result;
 
 	if (InputMesh->IsClosed())
 	{
@@ -1094,9 +1094,9 @@ static void GenerateUVTiling(int32 NumTiles, TArray<FAxisAlignedBox2d>& BoxesOut
 
 
 
-TSharedPtr<FDynamicMesh3> UGroomToMeshTool::UpdateUVs_MinimalConformal(TSharedPtr<FDynamicMesh3> InputMesh)
+TSharedPtr<FDynamicMesh3, ESPMode::ThreadSafe> UGroomToMeshTool::UpdateUVs_MinimalConformal(TSharedPtr<FDynamicMesh3, ESPMode::ThreadSafe> InputMesh)
 {
-	TSharedPtr<FDynamicMesh3> Result = MakeShared<FDynamicMesh3>(*InputMesh);
+	TSharedPtr<FDynamicMesh3, ESPMode::ThreadSafe> Result = MakeShared<FDynamicMesh3, ESPMode::ThreadSafe>(*InputMesh);
 	FDynamicMesh3* EditMesh = Result.Get();
 
 	FDynamicMeshUVEditor UVEditor(EditMesh, 0, true);
@@ -1131,14 +1131,14 @@ TSharedPtr<FDynamicMesh3> UGroomToMeshTool::UpdateUVs_MinimalConformal(TSharedPt
 
 
 
-TSharedPtr<FDynamicMesh3> UGroomToMeshTool::UpdateUVs_ExpMapPlaneSplits(TSharedPtr<FDynamicMesh3> InputMesh, bool bRecalcAsConformal)
+TSharedPtr<FDynamicMesh3, ESPMode::ThreadSafe> UGroomToMeshTool::UpdateUVs_ExpMapPlaneSplits(TSharedPtr<FDynamicMesh3, ESPMode::ThreadSafe> InputMesh, bool bRecalcAsConformal)
 {
 	int32 SplitAxisIndex1 = 0;
 	int32 SplitAxisIndex2 = 1;
 	double MiddleFraction1 = 0.6;
 	double MiddleFraction2 = 0.6;
 
-	TSharedPtr<FDynamicMesh3> Result = MakeShared<FDynamicMesh3>(*InputMesh);
+	TSharedPtr<FDynamicMesh3, ESPMode::ThreadSafe> Result = MakeShared<FDynamicMesh3, ESPMode::ThreadSafe>(*InputMesh);
 	FDynamicMesh3* EditMesh = Result.Get();
 	
 	FDynamicMeshAABBTree3 BVTree(EditMesh, true);
@@ -1294,7 +1294,7 @@ void UGroomToMeshTool::UpdateLineSet()
 
 
 
-void UGroomToMeshTool::UpdatePreview(TSharedPtr<FDynamicMesh3> ResultMesh)
+void UGroomToMeshTool::UpdatePreview(TSharedPtr<FDynamicMesh3, ESPMode::ThreadSafe> ResultMesh)
 {
 	CurrentMesh = FDynamicMesh3(*ResultMesh);
 
diff --git a/Engine/Plugins/Experimental/MeshModelingToolset/Source/MeshModelingToolsEditorOnly/Private/MeshTangentsTool.cpp b/Engine/Plugins/Experimental/MeshModelingToolset/Source/MeshModelingToolsEditorOnly/Private/MeshTangentsTool.cpp
--- a/Engine/Plugins/Experimental/MeshModelingToolset/Source/MeshModelingToolsEditorOnly/Private/MeshTangentsTool.cpp
+++ b/Engine/Plugins/Experimental/MeshModelingToolset/Source/MeshModelingToolsEditorOnly/Private/MeshTangentsTool.cpp
@@ -60,7 +60,7 @@ void UMeshTangentsTool::Setup()
 	UInteractiveTool::Setup();
 
 	// copy input MeshDescription and make sure it has initialized normals/tangents
-	InputMeshDescription = MakeShared<FMeshDescription>(*ComponentTarget->GetMesh());
+	InputMeshDescription = MakeShared<FMeshDescription, ESPMode::ThreadSafe>(*ComponentTarget->GetMesh());
 	UE::MeshDescription::InitializeAutoGeneratedAttributes(*InputMeshDescription, ComponentTarget->GetOwnerComponent(), 0);
 
 	// hide input StaticMeshComponent
@@ -80,8 +80,8 @@ void UMeshTangentsTool::Setup()
 	PreviewMesh->InitializeMesh(InputMeshDescription.Get());
 
 	// make a copy of initialized mesh and tangents
-	InputMesh = MakeShared<FDynamicMesh3>(*PreviewMesh->GetMesh());
-	InitialTangents = MakeShared<FMeshTangentsf>(InputMesh.Get());
+	InputMesh = MakeShared<FDynamicMesh3, ESPMode::ThreadSafe>(*PreviewMesh->GetMesh());
+	InitialTangents = MakeShared<FMeshTangentsf, ESPMode::ThreadSafe>(InputMesh.Get());
 	InitialTangents->CopyTriVertexTangents(*PreviewMesh->GetTangents());
 
 	// initialize our properties
diff --git a/Engine/Plugins/Experimental/MeshModelingToolset/Source/MeshModelingToolsEditorOnly/Private/ParameterizeMeshTool.cpp b/Engine/Plugins/Experimental/MeshModelingToolset/Source/MeshModelingToolsEditorOnly/Private/ParameterizeMeshTool.cpp
--- a/Engine/Plugins/Experimental/MeshModelingToolset/Source/MeshModelingToolsEditorOnly/Private/ParameterizeMeshTool.cpp
+++ b/Engine/Plugins/Experimental/MeshModelingToolset/Source/MeshModelingToolsEditorOnly/Private/ParameterizeMeshTool.cpp
@@ -97,7 +97,7 @@ void UParameterizeMeshTool::Setup()
 	// Initialize the preview mesh with a copy of the source mesh.
 	bool bHasGroups = false;
 	{
-		InputMesh = MakeShared<FDynamicMesh3>( * Preview->PreviewMesh->GetMesh() );
+		InputMesh = MakeShared<FDynamicMesh3, ESPMode::ThreadSafe>( * Preview->PreviewMesh->GetMesh() );
 		FMeshDescriptionToDynamicMesh Converter;
 		bHasGroups = FaceGroupUtil::HasMultipleGroups(*InputMesh);
 
diff --git a/Engine/Plugins/Experimental/MeshModelingToolset/Source/MeshModelingToolsEditorOnly/Private/PolygonOnMeshTool.cpp b/Engine/Plugins/Experimental/MeshModelingToolset/Source/MeshModelingToolsEditorOnly/Private/PolygonOnMeshTool.cpp
--- a/Engine/Plugins/Experimental/MeshModelingToolset/Source/MeshModelingToolsEditorOnly/Private/PolygonOnMeshTool.cpp
+++ b/Engine/Plugins/Experimental/MeshModelingToolset/Source/MeshModelingToolsEditorOnly/Private/PolygonOnMeshTool.cpp
@@ -135,7 +135,7 @@ void UPolygonOnMeshTool::Setup()
 	PlaneMechanic->SetPlaneCtrlClickBehaviorTarget->InvisibleComponentsToHitTest.Add(ComponentTarget->GetOwnerComponent());
 
 	// Convert input mesh description to dynamic mesh
-	OriginalDynamicMesh = MakeShared<FDynamicMesh3>();
+	OriginalDynamicMesh = MakeShared<FDynamicMesh3, ESPMode::ThreadSafe>();
 	FMeshDescriptionToDynamicMesh Converter;
 	Converter.Convert(ComponentTarget->GetMesh(), *OriginalDynamicMesh);
 	// TODO: consider adding an AABB tree construction here?  tradeoff vs doing a raycast against full every time a param change happens ...
diff --git a/Engine/Plugins/Experimental/MeshModelingToolset/Source/MeshModelingToolsEditorOnly/Private/SimplifyMeshTool.cpp b/Engine/Plugins/Experimental/MeshModelingToolset/Source/MeshModelingToolsEditorOnly/Private/SimplifyMeshTool.cpp
--- a/Engine/Plugins/Experimental/MeshModelingToolset/Source/MeshModelingToolsEditorOnly/Private/SimplifyMeshTool.cpp
+++ b/Engine/Plugins/Experimental/MeshModelingToolset/Source/MeshModelingToolsEditorOnly/Private/SimplifyMeshTool.cpp
@@ -118,7 +118,7 @@ void USimplifyMeshTool::Setup()
 #else
 		auto EnterProgressFrame = [](int Progress) {};
 #endif
-		OriginalMeshDescription = MakeShared<FMeshDescription>(*ComponentTarget->GetMesh());
+		OriginalMeshDescription = MakeShared<FMeshDescription, ESPMode::ThreadSafe>(*ComponentTarget->GetMesh());
 		// make sure the mesh description has tangents. 
 		{
 			// if the source was imported with auto-generate tangents, the mesh description tangents won't have been populated
@@ -129,13 +129,13 @@ void USimplifyMeshTool::Setup()
 			UE::MeshDescription::InitializeAutoGeneratedAttributes(*OriginalMeshDescription, &BuildSettings);
 		}
 		EnterProgressFrame(1);
-		OriginalMesh = MakeShared<FDynamicMesh3>();
+		OriginalMesh = MakeShared<FDynamicMesh3, ESPMode::ThreadSafe>();
 		FMeshDescriptionToDynamicMesh Converter; 
 		// convert with tangent overlay
 		const bool bUseTangents = true;
 		Converter.Convert(ComponentTarget->GetMesh(), *OriginalMesh, bUseTangents);
 		EnterProgressFrame(2);
-		OriginalMeshSpatial = MakeShared<FDynamicMeshAABBTree3>(OriginalMesh.Get(), true);
+		OriginalMeshSpatial = MakeShared<FDynamicMeshAABBTree3, ESPMode::ThreadSafe>(OriginalMesh.Get(), true);
 	}
 
 	Preview->PreviewMesh->SetTransform(ComponentTarget->GetWorldTransform());
diff --git a/Engine/Plugins/Experimental/MeshModelingToolset/Source/MeshModelingToolsEditorOnly/Private/SubdividePolyTool.cpp b/Engine/Plugins/Experimental/MeshModelingToolset/Source/MeshModelingToolsEditorOnly/Private/SubdividePolyTool.cpp
--- a/Engine/Plugins/Experimental/MeshModelingToolset/Source/MeshModelingToolsEditorOnly/Private/SubdividePolyTool.cpp
+++ b/Engine/Plugins/Experimental/MeshModelingToolset/Source/MeshModelingToolsEditorOnly/Private/SubdividePolyTool.cpp
@@ -168,7 +168,7 @@ void USubdividePolyTool::Setup()
 	}
 
 	bool bWantVertexNormals = false;
-	OriginalMesh = MakeShared<FDynamicMesh3>(bWantVertexNormals, false, false, false);
+	OriginalMesh = MakeShared<FDynamicMesh3, ESPMode::ThreadSafe>(bWantVertexNormals, false, false, false);
 	FMeshDescriptionToDynamicMesh Converter;
 	Converter.Convert(ComponentTarget->GetMesh(), *OriginalMesh);
 
diff --git a/Engine/Plugins/Experimental/MeshModelingToolset/Source/MeshModelingToolsEditorOnly/Private/UVLayoutTool.cpp b/Engine/Plugins/Experimental/MeshModelingToolset/Source/MeshModelingToolsEditorOnly/Private/UVLayoutTool.cpp
--- a/Engine/Plugins/Experimental/MeshModelingToolset/Source/MeshModelingToolsEditorOnly/Private/UVLayoutTool.cpp
+++ b/Engine/Plugins/Experimental/MeshModelingToolset/Source/MeshModelingToolsEditorOnly/Private/UVLayoutTool.cpp
@@ -155,7 +155,7 @@ void UUVLayoutTool::UpdateNumPreviews()
 			UUVLayoutOperatorFactory *OpFactory = NewObject<UUVLayoutOperatorFactory>();
 			OpFactory->Tool = this;
 			OpFactory->ComponentIndex = PreviewIdx;
-			OriginalDynamicMeshes[PreviewIdx] = MakeShared<FDynamicMesh3>();
+			OriginalDynamicMeshes[PreviewIdx] = MakeShared<FDynamicMesh3, ESPMode::ThreadSafe>();
 			FMeshDescriptionToDynamicMesh Converter;
 			Converter.Convert(ComponentTargets[PreviewIdx]->GetMesh(), *OriginalDynamicMeshes[PreviewIdx]);
 
diff --git a/Engine/Plugins/Experimental/MeshModelingToolset/Source/MeshModelingToolsEditorOnly/Public/GenerateLODMeshesTool.h b/Engine/Plugins/Experimental/MeshModelingToolset/Source/MeshModelingToolsEditorOnly/Public/GenerateLODMeshesTool.h
--- a/Engine/Plugins/Experimental/MeshModelingToolset/Source/MeshModelingToolsEditorOnly/Public/GenerateLODMeshesTool.h
+++ b/Engine/Plugins/Experimental/MeshModelingToolset/Source/MeshModelingToolsEditorOnly/Public/GenerateLODMeshesTool.h
@@ -189,10 +189,10 @@ private:
 	UWorld* TargetWorld;
 	IAssetGenerationAPI* AssetAPI;
 
-	TSharedPtr<FMeshDescription> OriginalMeshDescription;
+	TSharedPtr<FMeshDescription, ESPMode::ThreadSafe> OriginalMeshDescription;
 	// Dynamic Mesh versions precomputed in Setup (rather than recomputed for every simplify op)
-	TSharedPtr<FDynamicMesh3> OriginalMesh;
-	TSharedPtr<FDynamicMeshAABBTree3> OriginalMeshSpatial;
+	TSharedPtr<FDynamicMesh3, ESPMode::ThreadSafe> OriginalMesh;
+	TSharedPtr<FDynamicMeshAABBTree3, ESPMode::ThreadSafe> OriginalMeshSpatial;
 
 	TArray<FLODLevelGenerateSettings> CachedLODLevels;
 
diff --git a/Engine/Plugins/Experimental/MeshModelingToolset/Source/MeshModelingToolsEditorOnly/Public/Hair/GroomToMeshTool.h b/Engine/Plugins/Experimental/MeshModelingToolset/Source/MeshModelingToolsEditorOnly/Public/Hair/GroomToMeshTool.h
--- a/Engine/Plugins/Experimental/MeshModelingToolset/Source/MeshModelingToolsEditorOnly/Public/Hair/GroomToMeshTool.h
+++ b/Engine/Plugins/Experimental/MeshModelingToolset/Source/MeshModelingToolsEditorOnly/Public/Hair/GroomToMeshTool.h
@@ -183,14 +183,14 @@ protected:
 		}
 	};
 	FVoxelizeSettings CachedVoxelizeSettings;
-	TSharedPtr<FDynamicMesh3> CurrentVoxelizeResult;
-	TSharedPtr<FDynamicMesh3> UpdateVoxelization();
+	TSharedPtr<FDynamicMesh3, ESPMode::ThreadSafe> CurrentVoxelizeResult;
+	TSharedPtr<FDynamicMesh3, ESPMode::ThreadSafe> UpdateVoxelization();
 
 
 
 	struct FMorphologySettings
 	{
-		TSharedPtr<FDynamicMesh3> InputMesh;
+		TSharedPtr<FDynamicMesh3, ESPMode::ThreadSafe> InputMesh;
 		int32 VoxelCount;
 		float CloseDist;
 		float OpenDist;
@@ -201,13 +201,13 @@ protected:
 		}
 	};
 	FMorphologySettings CachedMorphologySettings;
-	TSharedPtr<FDynamicMesh3> CachedMorphologyResult;
-	TSharedPtr<FDynamicMesh3> UpdateMorphology(TSharedPtr<FDynamicMesh3> InputMesh);
+	TSharedPtr<FDynamicMesh3, ESPMode::ThreadSafe> CachedMorphologyResult;
+	TSharedPtr<FDynamicMesh3, ESPMode::ThreadSafe> UpdateMorphology(TSharedPtr<FDynamicMesh3, ESPMode::ThreadSafe> InputMesh);
 
 
 	struct FClipMeshSettings
 	{
-		TSharedPtr<FDynamicMesh3> InputMesh;
+		TSharedPtr<FDynamicMesh3, ESPMode::ThreadSafe> InputMesh;
 		AActor* ClipSource;
 
 		bool operator==(const FClipMeshSettings& Other) const
@@ -216,13 +216,13 @@ protected:
 		}
 	};
 	FClipMeshSettings CachedClipMeshSettings;
-	TSharedPtr<FDynamicMesh3> CachedClipMeshResult;
-	TSharedPtr<FDynamicMesh3> UpdateClipMesh(TSharedPtr<FDynamicMesh3> InputMesh);
+	TSharedPtr<FDynamicMesh3, ESPMode::ThreadSafe> CachedClipMeshResult;
+	TSharedPtr<FDynamicMesh3, ESPMode::ThreadSafe> UpdateClipMesh(TSharedPtr<FDynamicMesh3, ESPMode::ThreadSafe> InputMesh);
 
 
 	struct FSmoothingSettings
 	{
-		TSharedPtr<FDynamicMesh3> InputMesh;
+		TSharedPtr<FDynamicMesh3, ESPMode::ThreadSafe> InputMesh;
 		float Smoothness;
 		float VolumeCorrection;
 
@@ -232,13 +232,13 @@ protected:
 		}
 	};
 	FSmoothingSettings CachedSmoothSettings;
-	TSharedPtr<FDynamicMesh3> CachedSmoothResult;
-	TSharedPtr<FDynamicMesh3> UpdateSmoothing(TSharedPtr<FDynamicMesh3> InputMesh);
+	TSharedPtr<FDynamicMesh3, ESPMode::ThreadSafe> CachedSmoothResult;
+	TSharedPtr<FDynamicMesh3, ESPMode::ThreadSafe> UpdateSmoothing(TSharedPtr<FDynamicMesh3, ESPMode::ThreadSafe> InputMesh);
 
 
 	struct FSimplifySettings
 	{
-		TSharedPtr<FDynamicMesh3> InputMesh;
+		TSharedPtr<FDynamicMesh3, ESPMode::ThreadSafe> InputMesh;
 		int32 TargetCount;
 
 		bool operator==(const FSimplifySettings& Other) const
@@ -247,13 +247,13 @@ protected:
 		}
 	};
 	FSimplifySettings CachedSimplifySettings;
-	TSharedPtr<FDynamicMesh3> CachedSimplifyResult;
-	TSharedPtr<FDynamicMesh3> UpdateSimplification(TSharedPtr<FDynamicMesh3> InputMesh);
+	TSharedPtr<FDynamicMesh3, ESPMode::ThreadSafe> CachedSimplifyResult;
+	TSharedPtr<FDynamicMesh3, ESPMode::ThreadSafe> UpdateSimplification(TSharedPtr<FDynamicMesh3, ESPMode::ThreadSafe> InputMesh);
 
 
 	struct FPostprocessSettings
 	{
-		TSharedPtr<FDynamicMesh3> InputMesh;
+		TSharedPtr<FDynamicMesh3, ESPMode::ThreadSafe> InputMesh;
 		EGroomToMeshUVMode UVGenMode;
 
 		bool operator==(const FPostprocessSettings& Other) const
@@ -262,14 +262,14 @@ protected:
 		}
 	};
 	FPostprocessSettings CachedPostprocessSettings;
-	TSharedPtr<FDynamicMesh3> CachedPostprocessResult;
-	TSharedPtr<FDynamicMesh3> UpdatePostprocessing(TSharedPtr<FDynamicMesh3> InputMesh);
+	TSharedPtr<FDynamicMesh3, ESPMode::ThreadSafe> CachedPostprocessResult;
+	TSharedPtr<FDynamicMesh3, ESPMode::ThreadSafe> UpdatePostprocessing(TSharedPtr<FDynamicMesh3, ESPMode::ThreadSafe> InputMesh);
 
 
-	TSharedPtr<FDynamicMesh3> UpdateUVs(TSharedPtr<FDynamicMesh3> InputMesh, EGroomToMeshUVMode UVMode);
+	TSharedPtr<FDynamicMesh3, ESPMode::ThreadSafe> UpdateUVs(TSharedPtr<FDynamicMesh3, ESPMode::ThreadSafe> InputMesh, EGroomToMeshUVMode UVMode);
 
-	TSharedPtr<FDynamicMesh3> UpdateUVs_ExpMapPlaneSplits(TSharedPtr<FDynamicMesh3> InputMesh, bool bRecalcAsConformal);
-	TSharedPtr<FDynamicMesh3> UpdateUVs_MinimalConformal(TSharedPtr<FDynamicMesh3> InputMesh);
+	TSharedPtr<FDynamicMesh3, ESPMode::ThreadSafe> UpdateUVs_ExpMapPlaneSplits(TSharedPtr<FDynamicMesh3, ESPMode::ThreadSafe> InputMesh, bool bRecalcAsConformal);
+	TSharedPtr<FDynamicMesh3, ESPMode::ThreadSafe> UpdateUVs_MinimalConformal(TSharedPtr<FDynamicMesh3, ESPMode::ThreadSafe> InputMesh);
 
-	void UpdatePreview(TSharedPtr<FDynamicMesh3> ResultMesh);
+	void UpdatePreview(TSharedPtr<FDynamicMesh3, ESPMode::ThreadSafe> ResultMesh);
 };
diff --git a/Engine/Plugins/Experimental/MeshModelingToolset/Source/MeshModelingToolsEditorOnly/Public/MeshTangentsTool.h b/Engine/Plugins/Experimental/MeshModelingToolset/Source/MeshModelingToolsEditorOnly/Public/MeshTangentsTool.h
--- a/Engine/Plugins/Experimental/MeshModelingToolset/Source/MeshModelingToolsEditorOnly/Public/MeshTangentsTool.h
+++ b/Engine/Plugins/Experimental/MeshModelingToolset/Source/MeshModelingToolsEditorOnly/Public/MeshTangentsTool.h
@@ -119,9 +119,9 @@ protected:
 protected:
 	UWorld* TargetWorld;
 
-	TSharedPtr<FMeshDescription> InputMeshDescription;
-	TSharedPtr<FMeshTangentsf> InitialTangents;
-	TSharedPtr<FDynamicMesh3> InputMesh;
+	TSharedPtr<FMeshDescription, ESPMode::ThreadSafe> InputMeshDescription;
+	TSharedPtr<FMeshTangentsf, ESPMode::ThreadSafe> InitialTangents;
+	TSharedPtr<FDynamicMesh3, ESPMode::ThreadSafe> InputMesh;
 
 	bool bThicknessDirty = false;
 	bool bLengthDirty = false;
diff --git a/Engine/Plugins/Experimental/MeshModelingToolset/Source/MeshModelingToolsEditorOnly/Public/ParameterizeMeshTool.h b/Engine/Plugins/Experimental/MeshModelingToolset/Source/MeshModelingToolsEditorOnly/Public/ParameterizeMeshTool.h
--- a/Engine/Plugins/Experimental/MeshModelingToolset/Source/MeshModelingToolsEditorOnly/Public/ParameterizeMeshTool.h
+++ b/Engine/Plugins/Experimental/MeshModelingToolset/Source/MeshModelingToolsEditorOnly/Public/ParameterizeMeshTool.h
@@ -194,6 +194,6 @@ protected:
 	UWorld* TargetWorld;
 	IAssetGenerationAPI* AssetAPI;
 
-	TSharedPtr<FDynamicMesh3> InputMesh;
+	TSharedPtr<FDynamicMesh3, ESPMode::ThreadSafe> InputMesh;
 
 };
diff --git a/Engine/Plugins/Experimental/MeshModelingToolset/Source/MeshModelingToolsEditorOnly/Public/PolygonOnMeshTool.h b/Engine/Plugins/Experimental/MeshModelingToolset/Source/MeshModelingToolsEditorOnly/Public/PolygonOnMeshTool.h
--- a/Engine/Plugins/Experimental/MeshModelingToolset/Source/MeshModelingToolsEditorOnly/Public/PolygonOnMeshTool.h
+++ b/Engine/Plugins/Experimental/MeshModelingToolset/Source/MeshModelingToolsEditorOnly/Public/PolygonOnMeshTool.h
@@ -215,7 +215,7 @@ protected:
 	FTransform3d WorldTransform;
 	FViewCameraState CameraState;
 
-	TSharedPtr<FDynamicMesh3> OriginalDynamicMesh;
+	TSharedPtr<FDynamicMesh3, ESPMode::ThreadSafe> OriginalDynamicMesh;
 
 	UPROPERTY()
 	UConstructionPlaneMechanic* PlaneMechanic = nullptr;
diff --git a/Engine/Plugins/Experimental/MeshModelingToolset/Source/MeshModelingToolsEditorOnly/Public/SimplifyMeshTool.h b/Engine/Plugins/Experimental/MeshModelingToolset/Source/MeshModelingToolsEditorOnly/Public/SimplifyMeshTool.h
--- a/Engine/Plugins/Experimental/MeshModelingToolset/Source/MeshModelingToolsEditorOnly/Public/SimplifyMeshTool.h
+++ b/Engine/Plugins/Experimental/MeshModelingToolset/Source/MeshModelingToolsEditorOnly/Public/SimplifyMeshTool.h
@@ -129,10 +129,10 @@ private:
 	UWorld* TargetWorld;
 	IToolsContextAssetAPI* AssetAPI;
 
-	TSharedPtr<FMeshDescription> OriginalMeshDescription;
+	TSharedPtr<FMeshDescription, ESPMode::ThreadSafe> OriginalMeshDescription;
 	// Dynamic Mesh versions precomputed in Setup (rather than recomputed for every simplify op)
-	TSharedPtr<FDynamicMesh3> OriginalMesh;
-	TSharedPtr<FDynamicMeshAABBTree3> OriginalMeshSpatial;
+	TSharedPtr<FDynamicMesh3, ESPMode::ThreadSafe> OriginalMesh;
+	TSharedPtr<FDynamicMeshAABBTree3, ESPMode::ThreadSafe> OriginalMeshSpatial;
 
 	void GenerateAsset(const FDynamicMeshOpResult& Result);
 	void UpdateVisualization();
diff --git a/Engine/Plugins/Experimental/MeshModelingToolset/Source/MeshModelingToolsEditorOnly/Public/SubdividePolyTool.h b/Engine/Plugins/Experimental/MeshModelingToolset/Source/MeshModelingToolsEditorOnly/Public/SubdividePolyTool.h
--- a/Engine/Plugins/Experimental/MeshModelingToolset/Source/MeshModelingToolsEditorOnly/Public/SubdividePolyTool.h
+++ b/Engine/Plugins/Experimental/MeshModelingToolset/Source/MeshModelingToolsEditorOnly/Public/SubdividePolyTool.h
@@ -99,7 +99,7 @@ protected:
 	USubdividePolyToolProperties* Properties = nullptr;
 
 	// Input mesh
-	TSharedPtr<FDynamicMesh3> OriginalMesh;
+	TSharedPtr<FDynamicMesh3, ESPMode::ThreadSafe> OriginalMesh;
 
 	UPROPERTY()
 	UPreviewGeometry* PreviewGeometry = nullptr;
diff --git a/Engine/Plugins/Experimental/MeshModelingToolset/Source/MeshModelingToolsEditorOnly/Public/UVLayoutTool.h b/Engine/Plugins/Experimental/MeshModelingToolset/Source/MeshModelingToolsEditorOnly/Public/UVLayoutTool.h
--- a/Engine/Plugins/Experimental/MeshModelingToolset/Source/MeshModelingToolsEditorOnly/Public/UVLayoutTool.h
+++ b/Engine/Plugins/Experimental/MeshModelingToolset/Source/MeshModelingToolsEditorOnly/Public/UVLayoutTool.h
@@ -157,7 +157,7 @@ protected:
 	TArray<UMeshOpPreviewWithBackgroundCompute*> Previews;
 
 protected:
-	TArray<TSharedPtr<FDynamicMesh3>> OriginalDynamicMeshes;
+	TArray<TSharedPtr<FDynamicMesh3, ESPMode::ThreadSafe>> OriginalDynamicMeshes;
 
 	UWorld* TargetWorld = nullptr;
 	IAssetGenerationAPI* AssetAPI = nullptr;
diff --git a/Engine/Plugins/Experimental/MeshModelingToolset/Source/ModelingComponents/Private/BaseTools/BaseVoxelTool.cpp b/Engine/Plugins/Experimental/MeshModelingToolset/Source/ModelingComponents/Private/BaseTools/BaseVoxelTool.cpp
--- a/Engine/Plugins/Experimental/MeshModelingToolset/Source/ModelingComponents/Private/BaseTools/BaseVoxelTool.cpp
+++ b/Engine/Plugins/Experimental/MeshModelingToolset/Source/ModelingComponents/Private/BaseTools/BaseVoxelTool.cpp
@@ -43,7 +43,7 @@ void UBaseVoxelTool::ConvertInputsAndSetPreviewMaterials(bool bSetPreviewMesh)
 
 	for (int ComponentIdx = 0; ComponentIdx < ComponentTargets.Num(); ComponentIdx++)
 	{
-		OriginalDynamicMeshes[ComponentIdx] = MakeShared<FDynamicMesh3>();
+		OriginalDynamicMeshes[ComponentIdx] = MakeShared<FDynamicMesh3, ESPMode::ThreadSafe>();
 		FMeshDescriptionToDynamicMesh Converter;
 		Converter.Convert(ComponentTargets[ComponentIdx]->GetMesh(), *OriginalDynamicMeshes[ComponentIdx]);
 	}
diff --git a/Engine/Plugins/Experimental/MeshModelingToolset/Source/ModelingComponents/Private/Changes/DynamicMeshChangeTarget.cpp b/Engine/Plugins/Experimental/MeshModelingToolset/Source/ModelingComponents/Private/Changes/DynamicMeshChangeTarget.cpp
--- a/Engine/Plugins/Experimental/MeshModelingToolset/Source/ModelingComponents/Private/Changes/DynamicMeshChangeTarget.cpp
+++ b/Engine/Plugins/Experimental/MeshModelingToolset/Source/ModelingComponents/Private/Changes/DynamicMeshChangeTarget.cpp
@@ -10,7 +10,7 @@ void UDynamicMeshReplacementChangeTarget::ApplyChange(const FMeshReplacementChan
 	OnMeshChanged.Broadcast();
 }
 
-TUniquePtr<FMeshReplacementChange> UDynamicMeshReplacementChangeTarget::ReplaceMesh(const TSharedPtr<const FDynamicMesh3>& UpdateMesh)
+TUniquePtr<FMeshReplacementChange> UDynamicMeshReplacementChangeTarget::ReplaceMesh(const TSharedPtr<const FDynamicMesh3, ESPMode::ThreadSafe>& UpdateMesh)
 {
 	TUniquePtr<FMeshReplacementChange> Change = MakeUnique<FMeshReplacementChange>(Mesh, UpdateMesh);
 	Mesh = UpdateMesh;
diff --git a/Engine/Plugins/Experimental/MeshModelingToolset/Source/ModelingComponents/Private/Changes/MeshReplacementChange.cpp b/Engine/Plugins/Experimental/MeshModelingToolset/Source/ModelingComponents/Private/Changes/MeshReplacementChange.cpp
--- a/Engine/Plugins/Experimental/MeshModelingToolset/Source/ModelingComponents/Private/Changes/MeshReplacementChange.cpp
+++ b/Engine/Plugins/Experimental/MeshModelingToolset/Source/ModelingComponents/Private/Changes/MeshReplacementChange.cpp
@@ -7,7 +7,7 @@ FMeshReplacementChange::FMeshReplacementChange()
 {
 }
 
-FMeshReplacementChange::FMeshReplacementChange(TSharedPtr<const FDynamicMesh3> BeforeIn, TSharedPtr<const FDynamicMesh3> AfterIn)
+FMeshReplacementChange::FMeshReplacementChange(TSharedPtr<const FDynamicMesh3, ESPMode::ThreadSafe> BeforeIn, TSharedPtr<const FDynamicMesh3, ESPMode::ThreadSafe> AfterIn)
 {
 	Before = BeforeIn;
 	After = AfterIn;
diff --git a/Engine/Plugins/Experimental/MeshModelingToolset/Source/ModelingComponents/Public/BaseTools/BaseVoxelTool.h b/Engine/Plugins/Experimental/MeshModelingToolset/Source/ModelingComponents/Public/BaseTools/BaseVoxelTool.h
--- a/Engine/Plugins/Experimental/MeshModelingToolset/Source/ModelingComponents/Public/BaseTools/BaseVoxelTool.h
+++ b/Engine/Plugins/Experimental/MeshModelingToolset/Source/ModelingComponents/Public/BaseTools/BaseVoxelTool.h
@@ -37,6 +37,6 @@ protected:
 	UPROPERTY()
 	UVoxelProperties* VoxProperties;
 
-	TArray<TSharedPtr<FDynamicMesh3>> OriginalDynamicMeshes;
+	TArray<TSharedPtr<FDynamicMesh3, ESPMode::ThreadSafe>> OriginalDynamicMeshes;
 };
 
diff --git a/Engine/Plugins/Experimental/MeshModelingToolset/Source/ModelingComponents/Public/Changes/DynamicMeshChangeTarget.h b/Engine/Plugins/Experimental/MeshModelingToolset/Source/ModelingComponents/Public/Changes/DynamicMeshChangeTarget.h
--- a/Engine/Plugins/Experimental/MeshModelingToolset/Source/ModelingComponents/Public/Changes/DynamicMeshChangeTarget.h
+++ b/Engine/Plugins/Experimental/MeshModelingToolset/Source/ModelingComponents/Public/Changes/DynamicMeshChangeTarget.h
@@ -25,11 +25,11 @@ class MODELINGCOMPONENTS_API UDynamicMeshReplacementChangeTarget : public UObjec
 	GENERATED_BODY()
 
 protected:
-	TSharedPtr<const FDynamicMesh3> Mesh;
+	TSharedPtr<const FDynamicMesh3, ESPMode::ThreadSafe> Mesh;
 
 public:
 	UDynamicMeshReplacementChangeTarget() = default;
-	UDynamicMeshReplacementChangeTarget(TSharedPtr<const FDynamicMesh3> Mesh) : Mesh(Mesh)
+	UDynamicMeshReplacementChangeTarget(TSharedPtr<const FDynamicMesh3, ESPMode::ThreadSafe> Mesh) : Mesh(Mesh)
 	{}
 	virtual ~UDynamicMeshReplacementChangeTarget() = default;
 
@@ -38,12 +38,12 @@ public:
 	*/
 	virtual void ApplyChange(const FMeshReplacementChange* Change, bool bRevert) override;
 
-	virtual void SetMesh(TSharedPtr<const FDynamicMesh3> MeshIn)
+	virtual void SetMesh(TSharedPtr<const FDynamicMesh3, ESPMode::ThreadSafe> MeshIn)
 	{
 		Mesh = MeshIn;
 	}
 
-	const TSharedPtr<const FDynamicMesh3>& GetMesh() const
+	const TSharedPtr<const FDynamicMesh3, ESPMode::ThreadSafe>& GetMesh() const
 	{
 		return Mesh;
 	}
@@ -53,7 +53,7 @@ public:
 	 * Does not emit the change.
 	 * Does not fire OnMeshChanged.
 	 */
-	TUniquePtr<FMeshReplacementChange> ReplaceMesh(const TSharedPtr<const FDynamicMesh3>& UpdateMesh);
+	TUniquePtr<FMeshReplacementChange> ReplaceMesh(const TSharedPtr<const FDynamicMesh3, ESPMode::ThreadSafe>& UpdateMesh);
 
 	/**
 	* This delegate fires when a FCommandChange is applied to this component, so that
diff --git a/Engine/Plugins/Experimental/MeshModelingToolset/Source/ModelingComponents/Public/Changes/MeshReplacementChange.h b/Engine/Plugins/Experimental/MeshModelingToolset/Source/ModelingComponents/Public/Changes/MeshReplacementChange.h
--- a/Engine/Plugins/Experimental/MeshModelingToolset/Source/ModelingComponents/Public/Changes/MeshReplacementChange.h
+++ b/Engine/Plugins/Experimental/MeshModelingToolset/Source/ModelingComponents/Public/Changes/MeshReplacementChange.h
@@ -19,13 +19,13 @@ class FDynamicMesh3;
  */
 class MODELINGCOMPONENTS_API FMeshReplacementChange : public FToolCommandChange
 {
-	TSharedPtr<const FDynamicMesh3> Before, After;
+	TSharedPtr<const FDynamicMesh3, ESPMode::ThreadSafe> Before, After;
 
 public:
 	FMeshReplacementChange();
-	FMeshReplacementChange(TSharedPtr<const FDynamicMesh3> Before, TSharedPtr<const FDynamicMesh3> After);
+	FMeshReplacementChange(TSharedPtr<const FDynamicMesh3, ESPMode::ThreadSafe> Before, TSharedPtr<const FDynamicMesh3, ESPMode::ThreadSafe> After);
 
-	const TSharedPtr<const FDynamicMesh3>& GetMesh(bool bRevert) const
+	const TSharedPtr<const FDynamicMesh3, ESPMode::ThreadSafe>& GetMesh(bool bRevert) const
 	{
 		return bRevert ? Before : After;
 	}
diff --git a/Engine/Plugins/Experimental/MeshModelingToolset/Source/ModelingOperators/Private/BaseOps/VoxelBaseOp.cpp b/Engine/Plugins/Experimental/MeshModelingToolset/Source/ModelingOperators/Private/BaseOps/VoxelBaseOp.cpp
--- a/Engine/Plugins/Experimental/MeshModelingToolset/Source/ModelingOperators/Private/BaseOps/VoxelBaseOp.cpp
+++ b/Engine/Plugins/Experimental/MeshModelingToolset/Source/ModelingOperators/Private/BaseOps/VoxelBaseOp.cpp
@@ -46,7 +46,7 @@ void FVoxelBaseOp::PostProcessResult(FProgressCancel* Progress, double MeshCellS
 
 	if (bFixNormals)
 	{
-		TSharedPtr<FDynamicMesh3> OpResultMesh(ExtractResult().Release()); // moved the unique pointer
+		TSharedPtr<FDynamicMesh3, ESPMode::ThreadSafe> OpResultMesh(ExtractResult().Release()); // moved the unique pointer
 		OpResultMesh->EnableAttributes();
 
 		// Recompute the normals
diff --git a/Engine/Plugins/Experimental/MeshModelingToolset/Source/ModelingOperators/Private/CuttingOps/GroupEdgeInsertionOp.cpp b/Engine/Plugins/Experimental/MeshModelingToolset/Source/ModelingOperators/Private/CuttingOps/GroupEdgeInsertionOp.cpp
--- a/Engine/Plugins/Experimental/MeshModelingToolset/Source/ModelingOperators/Private/CuttingOps/GroupEdgeInsertionOp.cpp
+++ b/Engine/Plugins/Experimental/MeshModelingToolset/Source/ModelingOperators/Private/CuttingOps/GroupEdgeInsertionOp.cpp
@@ -30,7 +30,7 @@ void FGroupEdgeInsertionOp::CalculateResult(FProgressCancel* Progress)
 	}
 
 	ResultMesh->Copy(*OriginalMesh, true, true, true, true);
-	ResultTopology = MakeShared<FGroupTopology>();
+	ResultTopology = MakeShared<FGroupTopology, ESPMode::ThreadSafe>();
 	*ResultTopology = *OriginalTopology;
 	ResultTopology->RetargetOnClonedMesh(ResultMesh.Get());
 
diff --git a/Engine/Plugins/Experimental/MeshModelingToolset/Source/ModelingOperators/Private/DeformationOps/LatticeDeformerOp.cpp b/Engine/Plugins/Experimental/MeshModelingToolset/Source/ModelingOperators/Private/DeformationOps/LatticeDeformerOp.cpp
--- a/Engine/Plugins/Experimental/MeshModelingToolset/Source/ModelingOperators/Private/DeformationOps/LatticeDeformerOp.cpp
+++ b/Engine/Plugins/Experimental/MeshModelingToolset/Source/ModelingOperators/Private/DeformationOps/LatticeDeformerOp.cpp
@@ -30,8 +30,8 @@ void FLatticeDeformerOp::CalculateResult(FProgressCancel* Progress)
 	}
 }
 
-FLatticeDeformerOp::FLatticeDeformerOp(TSharedPtr<FDynamicMesh3> InOriginalMesh,
-									   TSharedPtr<FFFDLattice> InLattice,
+FLatticeDeformerOp::FLatticeDeformerOp(TSharedPtr<FDynamicMesh3, ESPMode::ThreadSafe> InOriginalMesh,
+									   TSharedPtr<FFFDLattice, ESPMode::ThreadSafe> InLattice,
 									   const TArray<FVector3d>& InLatticeControlPoints,
 									   ELatticeInterpolation InInterpolationType) :
 	Lattice(InLattice),
diff --git a/Engine/Plugins/Experimental/MeshModelingToolset/Source/ModelingOperators/Public/CleaningOps/EditNormalsOp.h b/Engine/Plugins/Experimental/MeshModelingToolset/Source/ModelingOperators/Public/CleaningOps/EditNormalsOp.h
--- a/Engine/Plugins/Experimental/MeshModelingToolset/Source/ModelingOperators/Public/CleaningOps/EditNormalsOp.h
+++ b/Engine/Plugins/Experimental/MeshModelingToolset/Source/ModelingOperators/Public/CleaningOps/EditNormalsOp.h
@@ -40,7 +40,7 @@ public:
 	virtual ~FEditNormalsOp() {}
 
 	// inputs
-	TSharedPtr<FDynamicMesh3> OriginalMesh;
+	TSharedPtr<FDynamicMesh3, ESPMode::ThreadSafe> OriginalMesh;
 
 	bool bFixInconsistentNormals;
 	bool bInvertNormals;
diff --git a/Engine/Plugins/Experimental/MeshModelingToolset/Source/ModelingOperators/Public/CleaningOps/HoleFillOp.h b/Engine/Plugins/Experimental/MeshModelingToolset/Source/ModelingOperators/Public/CleaningOps/HoleFillOp.h
--- a/Engine/Plugins/Experimental/MeshModelingToolset/Source/ModelingOperators/Public/CleaningOps/HoleFillOp.h
+++ b/Engine/Plugins/Experimental/MeshModelingToolset/Source/ModelingOperators/Public/CleaningOps/HoleFillOp.h
@@ -39,7 +39,7 @@ public:
 
 	// inputs
 	
-	TSharedPtr<FDynamicMesh3> OriginalMesh;		// Ownership shared with Tool
+	TSharedPtr<FDynamicMesh3, ESPMode::ThreadSafe> OriginalMesh;		// Ownership shared with Tool
 
 	EHoleFillOpFillType FillType = EHoleFillOpFillType::Minimal;
 	double MeshUVScaleFactor = 1.0;
diff --git a/Engine/Plugins/Experimental/MeshModelingToolset/Source/ModelingOperators/Public/CleaningOps/RemeshMeshOp.h b/Engine/Plugins/Experimental/MeshModelingToolset/Source/ModelingOperators/Public/CleaningOps/RemeshMeshOp.h
--- a/Engine/Plugins/Experimental/MeshModelingToolset/Source/ModelingOperators/Public/CleaningOps/RemeshMeshOp.h
+++ b/Engine/Plugins/Experimental/MeshModelingToolset/Source/ModelingOperators/Public/CleaningOps/RemeshMeshOp.h
@@ -51,8 +51,8 @@ public:
 	virtual ~FRemeshMeshOp() {}
 
 	// inputs
-	TSharedPtr<FDynamicMesh3> OriginalMesh;
-	TSharedPtr<FDynamicMeshAABBTree3> OriginalMeshSpatial;
+	TSharedPtr<FDynamicMesh3, ESPMode::ThreadSafe> OriginalMesh;
+	TSharedPtr<FDynamicMeshAABBTree3, ESPMode::ThreadSafe> OriginalMeshSpatial;
 
 	ERemeshType RemeshType = ERemeshType::Standard;
 
diff --git a/Engine/Plugins/Experimental/MeshModelingToolset/Source/ModelingOperators/Public/CleaningOps/RemoveOccludedTrianglesOp.h b/Engine/Plugins/Experimental/MeshModelingToolset/Source/ModelingOperators/Public/CleaningOps/RemoveOccludedTrianglesOp.h
--- a/Engine/Plugins/Experimental/MeshModelingToolset/Source/ModelingOperators/Public/CleaningOps/RemoveOccludedTrianglesOp.h
+++ b/Engine/Plugins/Experimental/MeshModelingToolset/Source/ModelingOperators/Public/CleaningOps/RemoveOccludedTrianglesOp.h
@@ -39,8 +39,8 @@ public:
 	virtual ~FRemoveOccludedTrianglesOp() {}
 
 	// inputs
-	TSharedPtr<FDynamicMesh3> OriginalMesh;
-	TSharedPtr<IndexMeshWithAcceleration> CombinedMeshTrees;
+	TSharedPtr<FDynamicMesh3, ESPMode::ThreadSafe> OriginalMesh;
+	TSharedPtr<IndexMeshWithAcceleration, ESPMode::ThreadSafe> CombinedMeshTrees;
 
 	TArray<FTransform3d> MeshTransforms;
 
diff --git a/Engine/Plugins/Experimental/MeshModelingToolset/Source/ModelingOperators/Public/CompositionOps/BooleanMeshesOp.h b/Engine/Plugins/Experimental/MeshModelingToolset/Source/ModelingOperators/Public/CompositionOps/BooleanMeshesOp.h
--- a/Engine/Plugins/Experimental/MeshModelingToolset/Source/ModelingOperators/Public/CompositionOps/BooleanMeshesOp.h
+++ b/Engine/Plugins/Experimental/MeshModelingToolset/Source/ModelingOperators/Public/CompositionOps/BooleanMeshesOp.h
@@ -56,7 +56,7 @@ public:
 	ETrimOperation TrimOperation;
 	ETrimSide TrimSide;
 	bool bTrimMode = false; // if true, do a trim operation instead of a boolean
-	TArray<TSharedPtr<const FDynamicMesh3>> Meshes;
+	TArray<TSharedPtr<const FDynamicMesh3, ESPMode::ThreadSafe>> Meshes;
 	TArray<FTransform> Transforms; // 1:1 with Meshes
 	bool bAttemptFixHoles = false;
 
diff --git a/Engine/Plugins/Experimental/MeshModelingToolset/Source/ModelingOperators/Public/CompositionOps/MirrorOp.h b/Engine/Plugins/Experimental/MeshModelingToolset/Source/ModelingOperators/Public/CompositionOps/MirrorOp.h
--- a/Engine/Plugins/Experimental/MeshModelingToolset/Source/ModelingOperators/Public/CompositionOps/MirrorOp.h
+++ b/Engine/Plugins/Experimental/MeshModelingToolset/Source/ModelingOperators/Public/CompositionOps/MirrorOp.h
@@ -12,7 +12,7 @@ public:
 	virtual ~FMirrorOp() {}
 
 	// Inputs
-	TSharedPtr<const FDynamicMesh3> OriginalMesh;
+	TSharedPtr<const FDynamicMesh3, ESPMode::ThreadSafe> OriginalMesh;
 	FVector3d LocalPlaneOrigin;
 	FVector3d LocalPlaneNormal;
 
diff --git a/Engine/Plugins/Experimental/MeshModelingToolset/Source/ModelingOperators/Public/CompositionOps/SelfUnionMeshesOp.h b/Engine/Plugins/Experimental/MeshModelingToolset/Source/ModelingOperators/Public/CompositionOps/SelfUnionMeshesOp.h
--- a/Engine/Plugins/Experimental/MeshModelingToolset/Source/ModelingOperators/Public/CompositionOps/SelfUnionMeshesOp.h
+++ b/Engine/Plugins/Experimental/MeshModelingToolset/Source/ModelingOperators/Public/CompositionOps/SelfUnionMeshesOp.h
@@ -13,7 +13,7 @@ public:
 	virtual ~FSelfUnionMeshesOp() {}
 
 	// inputs
-	TSharedPtr<const FDynamicMesh3> CombinedMesh;
+	TSharedPtr<const FDynamicMesh3, ESPMode::ThreadSafe> CombinedMesh;
 	bool bAttemptFixHoles = false;
 	double WindingNumberThreshold = .5;
 	bool bTrimFlaps = false;
diff --git a/Engine/Plugins/Experimental/MeshModelingToolset/Source/ModelingOperators/Public/CompositionOps/VoxelBlendMeshesOp.h b/Engine/Plugins/Experimental/MeshModelingToolset/Source/ModelingOperators/Public/CompositionOps/VoxelBlendMeshesOp.h
--- a/Engine/Plugins/Experimental/MeshModelingToolset/Source/ModelingOperators/Public/CompositionOps/VoxelBlendMeshesOp.h
+++ b/Engine/Plugins/Experimental/MeshModelingToolset/Source/ModelingOperators/Public/CompositionOps/VoxelBlendMeshesOp.h
@@ -16,7 +16,7 @@ public:
 	virtual ~FVoxelBlendMeshesOp() {}
 
 	// inputs
-	TArray<TSharedPtr<const FDynamicMesh3>> Meshes;
+	TArray<TSharedPtr<const FDynamicMesh3, ESPMode::ThreadSafe>> Meshes;
 	TArray<FTransform> Transforms; // 1:1 with Meshes
 
 	double BlendFalloff = 10;
diff --git a/Engine/Plugins/Experimental/MeshModelingToolset/Source/ModelingOperators/Public/CompositionOps/VoxelMorphologyMeshesOp.h b/Engine/Plugins/Experimental/MeshModelingToolset/Source/ModelingOperators/Public/CompositionOps/VoxelMorphologyMeshesOp.h
--- a/Engine/Plugins/Experimental/MeshModelingToolset/Source/ModelingOperators/Public/CompositionOps/VoxelMorphologyMeshesOp.h
+++ b/Engine/Plugins/Experimental/MeshModelingToolset/Source/ModelingOperators/Public/CompositionOps/VoxelMorphologyMeshesOp.h
@@ -35,7 +35,7 @@ public:
 	virtual ~FVoxelMorphologyMeshesOp() {}
 
 	// inputs
-	TArray<TSharedPtr<const FDynamicMesh3>> Meshes;
+	TArray<TSharedPtr<const FDynamicMesh3, ESPMode::ThreadSafe>> Meshes;
 	TArray<FTransform> Transforms; // 1:1 with Meshes
 
 	double Distance = 1.0;
diff --git a/Engine/Plugins/Experimental/MeshModelingToolset/Source/ModelingOperators/Public/CompositionOps/VoxelSolidifyMeshesOp.h b/Engine/Plugins/Experimental/MeshModelingToolset/Source/ModelingOperators/Public/CompositionOps/VoxelSolidifyMeshesOp.h
--- a/Engine/Plugins/Experimental/MeshModelingToolset/Source/ModelingOperators/Public/CompositionOps/VoxelSolidifyMeshesOp.h
+++ b/Engine/Plugins/Experimental/MeshModelingToolset/Source/ModelingOperators/Public/CompositionOps/VoxelSolidifyMeshesOp.h
@@ -16,7 +16,7 @@ public:
 	virtual ~FVoxelSolidifyMeshesOp() {}
 
 	// inputs
-	TArray<TSharedPtr<const FDynamicMesh3>> Meshes;
+	TArray<TSharedPtr<const FDynamicMesh3, ESPMode::ThreadSafe>> Meshes;
 	TArray<FTransform> Transforms; // 1:1 with Meshes
 
 	double WindingThreshold = .5;
diff --git a/Engine/Plugins/Experimental/MeshModelingToolset/Source/ModelingOperators/Public/CuttingOps/GroupEdgeInsertionOp.h b/Engine/Plugins/Experimental/MeshModelingToolset/Source/ModelingOperators/Public/CuttingOps/GroupEdgeInsertionOp.h
--- a/Engine/Plugins/Experimental/MeshModelingToolset/Source/ModelingOperators/Public/CuttingOps/GroupEdgeInsertionOp.h
+++ b/Engine/Plugins/Experimental/MeshModelingToolset/Source/ModelingOperators/Public/CuttingOps/GroupEdgeInsertionOp.h
@@ -15,8 +15,8 @@ public:
 	virtual ~FGroupEdgeInsertionOp() {}
 
 	// Inputs:
-	TSharedPtr<const FDynamicMesh3> OriginalMesh;
-	TSharedPtr<const FGroupTopology> OriginalTopology;
+	TSharedPtr<const FDynamicMesh3, ESPMode::ThreadSafe> OriginalMesh;
+	TSharedPtr<const FGroupTopology, ESPMode::ThreadSafe> OriginalTopology;
 	FGroupEdgeInserter::EInsertionMode Mode;
 	double VertexTolerance = KINDA_SMALL_NUMBER * 10;
 
@@ -34,7 +34,7 @@ public:
 	// Edge ID's in the ResultMesh that together make up the new group edge.
 	TSet<int32> Eids; 
 
-	TSharedPtr<FGroupTopology> ResultTopology;
+	TSharedPtr<FGroupTopology, ESPMode::ThreadSafe> ResultTopology;
 	bool bSucceeded = false;
 
 	void SetTransform(const FTransform& Transform);
diff --git a/Engine/Plugins/Experimental/MeshModelingToolset/Source/ModelingOperators/Public/CuttingOps/PlaneCutOp.h b/Engine/Plugins/Experimental/MeshModelingToolset/Source/ModelingOperators/Public/CuttingOps/PlaneCutOp.h
--- a/Engine/Plugins/Experimental/MeshModelingToolset/Source/ModelingOperators/Public/CuttingOps/PlaneCutOp.h
+++ b/Engine/Plugins/Experimental/MeshModelingToolset/Source/ModelingOperators/Public/CuttingOps/PlaneCutOp.h
@@ -22,7 +22,7 @@ public:
 	float CutPlaneLocalThickness = 0; // plane thickness in the local space of the mesh
 	float UVScaleFactor;
 	static const FName ObjectIndexAttribute;
-	TSharedPtr<const FDynamicMesh3> OriginalMesh;
+	TSharedPtr<const FDynamicMesh3, ESPMode::ThreadSafe> OriginalMesh;
 
 	void SetTransform(const FTransform& Transform);
 
diff --git a/Engine/Plugins/Experimental/MeshModelingToolset/Source/ModelingOperators/Public/DeformationOps/LatticeDeformerOp.h b/Engine/Plugins/Experimental/MeshModelingToolset/Source/ModelingOperators/Public/DeformationOps/LatticeDeformerOp.h
--- a/Engine/Plugins/Experimental/MeshModelingToolset/Source/ModelingOperators/Public/DeformationOps/LatticeDeformerOp.h
+++ b/Engine/Plugins/Experimental/MeshModelingToolset/Source/ModelingOperators/Public/DeformationOps/LatticeDeformerOp.h
@@ -16,8 +16,8 @@ class MODELINGOPERATORS_API FLatticeDeformerOp : public FDynamicMeshOperator
 {
 public:
 
-	FLatticeDeformerOp(TSharedPtr<FDynamicMesh3> InOriginalMesh,
-					   TSharedPtr<FFFDLattice> InLattice,
+	FLatticeDeformerOp(TSharedPtr<FDynamicMesh3, ESPMode::ThreadSafe> InOriginalMesh,
+					   TSharedPtr<FFFDLattice, ESPMode::ThreadSafe> InLattice,
 					   const TArray<FVector3d>& InLatticeControlPoints,
 					   ELatticeInterpolation InInterpolationType);
 
@@ -27,8 +27,8 @@ public:
 protected:
 
 	// Inputs
-	const TSharedPtr<const FFFDLattice> Lattice;
-	const TSharedPtr<const FDynamicMesh3> OriginalMesh;
+	const TSharedPtr<const FFFDLattice, ESPMode::ThreadSafe> Lattice;
+	const TSharedPtr<const FDynamicMesh3, ESPMode::ThreadSafe> OriginalMesh;
 	const TArray<FVector3d> LatticeControlPoints;
 	ELatticeInterpolation InterpolationType;
 };
diff --git a/Engine/Plugins/Experimental/MeshModelingToolset/Source/ModelingOperators/Public/ParameterizationOps/UVProjectionOp.h b/Engine/Plugins/Experimental/MeshModelingToolset/Source/ModelingOperators/Public/ParameterizationOps/UVProjectionOp.h
--- a/Engine/Plugins/Experimental/MeshModelingToolset/Source/ModelingOperators/Public/ParameterizationOps/UVProjectionOp.h
+++ b/Engine/Plugins/Experimental/MeshModelingToolset/Source/ModelingOperators/Public/ParameterizationOps/UVProjectionOp.h
@@ -25,7 +25,7 @@ public:
 	virtual ~FUVProjectionOp() {}
 
 	// inputs
-	TSharedPtr<FDynamicMesh3> OriginalMesh;
+	TSharedPtr<FDynamicMesh3, ESPMode::ThreadSafe> OriginalMesh;
 
 	EUVProjectionMethod ProjectionMethod;
 	float CylinderProjectToTopOrBottomAngleThreshold;
diff --git a/Engine/Plugins/Experimental/MeshModelingToolset/Source/ModelingOperators/Public/SpaceDeformerOps/MeshSpaceDeformerOp.h b/Engine/Plugins/Experimental/MeshModelingToolset/Source/ModelingOperators/Public/SpaceDeformerOps/MeshSpaceDeformerOp.h
--- a/Engine/Plugins/Experimental/MeshModelingToolset/Source/ModelingOperators/Public/SpaceDeformerOps/MeshSpaceDeformerOp.h
+++ b/Engine/Plugins/Experimental/MeshModelingToolset/Source/ModelingOperators/Public/SpaceDeformerOps/MeshSpaceDeformerOp.h
@@ -11,7 +11,7 @@ class MODELINGOPERATORS_API FMeshSpaceDeformerOp : public FDynamicMeshOperator
 public:
 
 	// Inputs
-	TSharedPtr<const FDynamicMesh3> OriginalMesh;
+	TSharedPtr<const FDynamicMesh3, ESPMode::ThreadSafe> OriginalMesh;
 	FFrame3d GizmoFrame;
 	void SetTransform(const FTransform& Transform);
 
diff --git a/Engine/Plugins/Experimental/MeshModelingToolset/Source/ModelingOperatorsEditorOnly/Private/CompositionOps/VoxelBooleanMeshesOp.cpp b/Engine/Plugins/Experimental/MeshModelingToolset/Source/ModelingOperatorsEditorOnly/Private/CompositionOps/VoxelBooleanMeshesOp.cpp
--- a/Engine/Plugins/Experimental/MeshModelingToolset/Source/ModelingOperatorsEditorOnly/Private/CompositionOps/VoxelBooleanMeshesOp.cpp
+++ b/Engine/Plugins/Experimental/MeshModelingToolset/Source/ModelingOperatorsEditorOnly/Private/CompositionOps/VoxelBooleanMeshesOp.cpp
@@ -140,7 +140,7 @@ void FVoxelBooleanMeshesOp::CalculateResult(FProgressCancel* Progress)
 	if (bFixNormals)
 	{
 
-		TSharedPtr<FDynamicMesh3> OpResultMesh(ExtractResult().Release()); // moved the unique pointer
+		TSharedPtr<FDynamicMesh3, ESPMode::ThreadSafe> OpResultMesh(ExtractResult().Release()); // moved the unique pointer
 
 		// Recompute the normals
 		FEditNormalsOp EditNormalsOp;
diff --git a/Engine/Plugins/Experimental/MeshModelingToolset/Source/ModelingOperatorsEditorOnly/Private/CompositionOps/VoxelMergeMeshesOp.cpp b/Engine/Plugins/Experimental/MeshModelingToolset/Source/ModelingOperatorsEditorOnly/Private/CompositionOps/VoxelMergeMeshesOp.cpp
--- a/Engine/Plugins/Experimental/MeshModelingToolset/Source/ModelingOperatorsEditorOnly/Private/CompositionOps/VoxelMergeMeshesOp.cpp
+++ b/Engine/Plugins/Experimental/MeshModelingToolset/Source/ModelingOperatorsEditorOnly/Private/CompositionOps/VoxelMergeMeshesOp.cpp
@@ -100,7 +100,7 @@ void FVoxelMergeMeshesOp::CalculateResult(FProgressCancel* Progress)
 	if (bFixNormals)
 	{
 
-		TSharedPtr<FDynamicMesh3> OpResultMesh(ExtractResult().Release()); // moved the unique pointer
+		TSharedPtr<FDynamicMesh3, ESPMode::ThreadSafe> OpResultMesh(ExtractResult().Release()); // moved the unique pointer
 
 		// Recompute the normals
 		FEditNormalsOp EditNormalsOp;
diff --git a/Engine/Plugins/Experimental/MeshModelingToolset/Source/ModelingOperatorsEditorOnly/Public/CleaningOps/SimplifyMeshOp.h b/Engine/Plugins/Experimental/MeshModelingToolset/Source/ModelingOperatorsEditorOnly/Public/CleaningOps/SimplifyMeshOp.h
--- a/Engine/Plugins/Experimental/MeshModelingToolset/Source/ModelingOperatorsEditorOnly/Public/CleaningOps/SimplifyMeshOp.h
+++ b/Engine/Plugins/Experimental/MeshModelingToolset/Source/ModelingOperatorsEditorOnly/Public/CleaningOps/SimplifyMeshOp.h
@@ -73,10 +73,10 @@ public:
 	float MinimalPlanarAngleThresh = 0.01f;
 
 	// stored for the UE4 Standard path
-	TSharedPtr<FMeshDescription> OriginalMeshDescription;
+	TSharedPtr<FMeshDescription, ESPMode::ThreadSafe> OriginalMeshDescription;
 	// stored for the GeometryProcessing custom simplifier paths (currently precomputed once in tool setup)
-	TSharedPtr<FDynamicMesh3> OriginalMesh;
-	TSharedPtr<FDynamicMeshAABBTree3> OriginalMeshSpatial;
+	TSharedPtr<FDynamicMesh3, ESPMode::ThreadSafe> OriginalMesh;
+	TSharedPtr<FDynamicMeshAABBTree3, ESPMode::ThreadSafe> OriginalMeshSpatial;
 
 	class IMeshReduction* MeshReduction;
 
diff --git a/Engine/Plugins/Experimental/MeshModelingToolset/Source/ModelingOperatorsEditorOnly/Public/CuttingOps/EmbedPolygonsOp.h b/Engine/Plugins/Experimental/MeshModelingToolset/Source/ModelingOperatorsEditorOnly/Public/CuttingOps/EmbedPolygonsOp.h
--- a/Engine/Plugins/Experimental/MeshModelingToolset/Source/ModelingOperatorsEditorOnly/Public/CuttingOps/EmbedPolygonsOp.h
+++ b/Engine/Plugins/Experimental/MeshModelingToolset/Source/ModelingOperatorsEditorOnly/Public/CuttingOps/EmbedPolygonsOp.h
@@ -46,7 +46,7 @@ public:
 
 	//float ExtrudeDistance; // TODO if we support extrude
 
-	TSharedPtr<FDynamicMesh3> OriginalMesh;
+	TSharedPtr<FDynamicMesh3, ESPMode::ThreadSafe> OriginalMesh;
 
 	//
 	// FDynamicMeshOperator implementation
diff --git a/Engine/Plugins/Experimental/MeshModelingToolset/Source/ModelingOperatorsEditorOnly/Public/ParameterizationOps/CalculateTangentsOp.h b/Engine/Plugins/Experimental/MeshModelingToolset/Source/ModelingOperatorsEditorOnly/Public/ParameterizationOps/CalculateTangentsOp.h
--- a/Engine/Plugins/Experimental/MeshModelingToolset/Source/ModelingOperatorsEditorOnly/Public/ParameterizationOps/CalculateTangentsOp.h
+++ b/Engine/Plugins/Experimental/MeshModelingToolset/Source/ModelingOperatorsEditorOnly/Public/ParameterizationOps/CalculateTangentsOp.h
@@ -31,8 +31,8 @@ public:
 	virtual ~FCalculateTangentsOp() {}
 
 	// inputs
-	TSharedPtr<FDynamicMesh3> SourceMesh;
-	TSharedPtr<FMeshTangentsf> SourceTangents;
+	TSharedPtr<FDynamicMesh3, ESPMode::ThreadSafe> SourceMesh;
+	TSharedPtr<FMeshTangentsf, ESPMode::ThreadSafe> SourceTangents;
 
 	// parameters
 	EMeshTangentsType CalculationMethod;
diff --git a/Engine/Plugins/Experimental/MeshModelingToolset/Source/ModelingOperatorsEditorOnly/Public/ParameterizationOps/ParameterizeMeshOp.h b/Engine/Plugins/Experimental/MeshModelingToolset/Source/ModelingOperatorsEditorOnly/Public/ParameterizationOps/ParameterizeMeshOp.h
--- a/Engine/Plugins/Experimental/MeshModelingToolset/Source/ModelingOperatorsEditorOnly/Public/ParameterizationOps/ParameterizeMeshOp.h
+++ b/Engine/Plugins/Experimental/MeshModelingToolset/Source/ModelingOperatorsEditorOnly/Public/ParameterizationOps/ParameterizeMeshOp.h
@@ -35,7 +35,7 @@ public:
 	// 
 
 	// source mesh
-	TSharedPtr<FDynamicMesh3> InputMesh;
+	TSharedPtr<FDynamicMesh3, ESPMode::ThreadSafe> InputMesh;
 		
 	// UV generation parameters
 	float Stretch;
diff --git a/Engine/Plugins/Experimental/MeshModelingToolset/Source/ModelingOperatorsEditorOnly/Public/ParameterizationOps/UVLayoutOp.h b/Engine/Plugins/Experimental/MeshModelingToolset/Source/ModelingOperatorsEditorOnly/Public/ParameterizationOps/UVLayoutOp.h
--- a/Engine/Plugins/Experimental/MeshModelingToolset/Source/ModelingOperatorsEditorOnly/Public/ParameterizationOps/UVLayoutOp.h
+++ b/Engine/Plugins/Experimental/MeshModelingToolset/Source/ModelingOperatorsEditorOnly/Public/ParameterizationOps/UVLayoutOp.h
@@ -21,7 +21,7 @@ public:
 	virtual ~FUVLayoutOp() {}
 
 	// inputs
-	TSharedPtr<FDynamicMesh3> OriginalMesh;
+	TSharedPtr<FDynamicMesh3, ESPMode::ThreadSafe> OriginalMesh;
 
 	EUVLayoutOpLayoutModes UVLayoutMode = EUVLayoutOpLayoutModes::RepackToUnitRect;
 
