ModelingComponents: Disable 4.26 hotfix code that is resulting in duplicate convex hulls
#rnx

diff --git a/Engine/Plugins/Experimental/MeshModelingToolset/Source/ModelingComponents/Private/Physics/PhysicsDataCollection.cpp b/Engine/Plugins/Experimental/MeshModelingToolset/Source/ModelingComponents/Private/Physics/PhysicsDataCollection.cpp
--- a/Engine/Plugins/Experimental/MeshModelingToolset/Source/ModelingComponents/Private/Physics/PhysicsDataCollection.cpp
+++ b/Engine/Plugins/Experimental/MeshModelingToolset/Source/ModelingComponents/Private/Physics/PhysicsDataCollection.cpp
@@ -94,7 +94,7 @@ void FPhysicsDataCollection::CopyGeometryToAggregate()
 		FKConvexElem Element;
 		UE::Geometry::GetFKElement(ConvexGeom.Mesh, Element);
 
-#if WITH_PHYSX
+#if !WITH_CHAOS
 		// Chaos will compute the IndexData itself on the call to FKConvexElem::UpdateElemBox() in ::GetFKElement() above.
 		// PhysX will not, so initialize that data with the mesh triangles.
 		// (This code should go into ::GetFKElement but cannot because it needs to be added in a hotfix)
