Mesh Modeling Tools:  Polygroup generation, add option to use connected mesh regions to define polygroups.   Also, allow the UI to hide options when they don't apply to the selection method.

diff --git a/Engine/Plugins/Experimental/MeshModelingToolset/Source/MeshModelingTools/Private/ConvertToPolygonsTool.cpp b/Engine/Plugins/Experimental/MeshModelingToolset/Source/MeshModelingTools/Private/ConvertToPolygonsTool.cpp
--- a/Engine/Plugins/Experimental/MeshModelingToolset/Source/MeshModelingTools/Private/ConvertToPolygonsTool.cpp
+++ b/Engine/Plugins/Experimental/MeshModelingToolset/Source/MeshModelingTools/Private/ConvertToPolygonsTool.cpp
@@ -12,6 +12,7 @@
 #include "DynamicMeshToMeshDescription.h"
 #include "Util/ColorConstants.h"
 #include "ToolSetupUtil.h"
+#include "Polygroups/PolygroupUtil.h"
 
 #include "SceneManagement.h" // for FPrimitiveDrawInterface
 
@@ -158,18 +159,26 @@ void UConvertToPolygonsTool::Render(IToolsContextRenderAPI* RenderAPI)
 void UConvertToPolygonsTool::UpdatePolygons()
 {
 	Polygons = FFindPolygonsAlgorithm(&SearchMesh);
-	if (Settings->ConversionMode == EConvertToPolygonsMode::FromUVISlands)
+	switch (Settings->ConversionMode)
 	{
-		Polygons.FindPolygonsFromUVIslands();
-	}
-	else if (Settings->ConversionMode == EConvertToPolygonsMode::FaceNormalDeviation)
-	{
-		double DotTolerance = 1.0 - FMathd::Cos(Settings->AngleTolerance * FMathd::DegToRad);
-		Polygons.FindPolygonsFromFaceNormals(DotTolerance);
-	}
-	else
-	{
-		check(false);
+		case EConvertToPolygonsMode::FromUVIslands:
+		{
+			Polygons.FindPolygonsFromUVIslands();
+			break;
+		}
+		case EConvertToPolygonsMode::FromConnectedTris:
+		{
+			Polygons.FindPolygonsFromConnectedTris();
+			break;
+		}
+		case EConvertToPolygonsMode::FaceNormalDeviation:
+		{
+			double DotTolerance = 1.0 - FMathd::Cos(Settings->AngleTolerance * FMathd::DegToRad);
+			Polygons.FindPolygonsFromFaceNormals(DotTolerance);
+			break;
+		}
+		default:
+			check(0);
 	}
 	
 	Polygons.FindPolygonEdges();
@@ -178,7 +187,7 @@ void UConvertToPolygonsTool::UpdatePolygons()
 		FText::Format(LOCTEXT("UpdatePolygonsMessage", "Found {0} Polygroups in {1} Triangles"),
 			FText::AsNumber(Polygons.FoundPolygons.Num()), FText::AsNumber(SearchMesh.TriangleCount())), EToolMessageLevel::Internal);
 
-	if (Settings->bCalculateNormals)
+	if (Settings->bCalculateNormals && Settings->ConversionMode == EConvertToPolygonsMode::FaceNormalDeviation)
 	{
 		if (SearchMesh.HasAttributes() == false)
 		{
@@ -212,7 +221,7 @@ void UConvertToPolygonsTool::UpdatePolygons()
 void UConvertToPolygonsTool::ConvertToPolygons(FMeshDescription* MeshIn)
 {
 	// restore input normals
-	if (Settings->bCalculateNormals == false)
+	if (Settings->bCalculateNormals == false || Settings->ConversionMode != EConvertToPolygonsMode::FaceNormalDeviation)
 	{
 		SearchMesh.Attributes()->PrimaryNormals()->Copy(InitialNormals);
 	}
diff --git a/Engine/Plugins/Experimental/MeshModelingToolset/Source/MeshModelingTools/Private/FindPolygonsAlgorithm.cpp b/Engine/Plugins/Experimental/MeshModelingToolset/Source/MeshModelingTools/Private/FindPolygonsAlgorithm.cpp
--- a/Engine/Plugins/Experimental/MeshModelingToolset/Source/MeshModelingTools/Private/FindPolygonsAlgorithm.cpp
+++ b/Engine/Plugins/Experimental/MeshModelingToolset/Source/MeshModelingTools/Private/FindPolygonsAlgorithm.cpp
@@ -3,6 +3,7 @@
 #include "FindPolygonsAlgorithm.h"
 #include "MeshNormals.h"
 #include "Selections/MeshConnectedComponents.h"
+#include "Util/IndexUtil.h"
 
 
 
@@ -15,11 +16,10 @@ FFindPolygonsAlgorithm::FFindPolygonsAlgorithm(FDynamicMesh3* MeshIn)
 
 bool FFindPolygonsAlgorithm::FindPolygonsFromUVIslands()
 {
-	PolygonGroupIDs.SetNum(Mesh->MaxTriangleID());
 	const FDynamicMeshUVOverlay* UV = Mesh->Attributes()->GetUVLayer(0);
 
 	FMeshConnectedComponents Components(Mesh);
-	Components.FindConnectedTriangles([&](int32 TriIdx0, int32 TriIdx1)
+	Components.FindConnectedTriangles([&UV](int32 TriIdx0, int32 TriIdx1)
 	{
 		return UV->AreTrianglesConnected(TriIdx0, TriIdx1);
 	});
@@ -37,6 +37,30 @@ bool FFindPolygonsAlgorithm::FindPolygonsFromUVIslands()
 
 
 
+bool FFindPolygonsAlgorithm::FindPolygonsFromConnectedTris()
+{
+
+	FMeshConnectedComponents Components(Mesh);
+	Components.FindConnectedTriangles([this](int32 TriIdx0, int32 TriIdx1)
+	{
+		FIndex3i NbrTris = Mesh->GetTriNeighbourTris(TriIdx0);
+		int NbrIndex = IndexUtil::FindTriIndex(TriIdx1, NbrTris);
+		return (NbrIndex != IndexConstants::InvalidID);
+	});
+
+	int32 NumComponents = Components.Num();
+	for (int32 ci = 0; ci < NumComponents; ++ci)
+	{
+		FoundPolygons.Add(Components.GetComponent(ci).Indices);
+	}
+
+	SetGroupsFromPolygons();
+
+	return (FoundPolygons.Num() > 0);
+}
+
+
+
 bool FFindPolygonsAlgorithm::FindPolygonsFromFaceNormals(double DotTolerance)
 {
 	DotTolerance = 1.0 - DotTolerance;
@@ -47,7 +71,6 @@ bool FFindPolygonsAlgorithm::FindPolygonsFromFaceNormals(double DotTolerance)
 
 	TArray<bool> DoneTriangle;
 	DoneTriangle.SetNum(Mesh->MaxTriangleID());
-	PolygonGroupIDs.SetNum(Mesh->MaxTriangleID());
 
 	TArray<int> Stack;
 
@@ -72,8 +95,7 @@ bool FFindPolygonsAlgorithm::FindPolygonsFromFaceNormals(double DotTolerance)
 			for (int j = 0; j < 3; ++j)
 			{
 				if (NbrTris[j] >= 0
-					&& DoneTriangle[NbrTris[j]] == false
-					&& PolygonGroupIDs[CurTri] == PolygonGroupIDs[NbrTris[j]])
+					&& DoneTriangle[NbrTris[j]] == false)
 				{
 					double Dot = Normals[CurTri].Dot(Normals[NbrTris[j]]);
 					if (Dot > DotTolerance)
diff --git a/Engine/Plugins/Experimental/MeshModelingToolset/Source/MeshModelingTools/Public/ConvertToPolygonsTool.h b/Engine/Plugins/Experimental/MeshModelingToolset/Source/MeshModelingTools/Public/ConvertToPolygonsTool.h
--- a/Engine/Plugins/Experimental/MeshModelingToolset/Source/MeshModelingTools/Public/ConvertToPolygonsTool.h
+++ b/Engine/Plugins/Experimental/MeshModelingToolset/Source/MeshModelingTools/Public/ConvertToPolygonsTool.h
@@ -34,9 +34,11 @@ UENUM()
 enum class EConvertToPolygonsMode
 {
 	/** Convert based on Angle Tolerance between Face Normals */
-	FaceNormalDeviation,
+	FaceNormalDeviation UMETA(DisplayName = "Face Normal Deviation"),
 	/** Create PolyGroups based on UV Islands */
-	FromUVISlands
+	FromUVIslands  UMETA(DisplayName = "From UV Islands"),
+	/** Create Polygroups based on Connected Triangles */
+	FromConnectedTris UMETA(DisplayName = "From Connected Tris")
 };
 
 
@@ -52,11 +54,11 @@ public:
 	EConvertToPolygonsMode ConversionMode = EConvertToPolygonsMode::FaceNormalDeviation;
 
 	/** Tolerance for planarity */
-	UPROPERTY(EditAnywhere, Category = PolyGroups, meta = (UIMin = "0.001", UIMax = "20.0", ClampMin = "0.0", ClampMax = "90.0", EditCondition = "ConversionMode == EConvertToPolygonsMode::FaceNormalDeviation"))
+	UPROPERTY(EditAnywhere, Category = PolyGroups, meta = (UIMin = "0.001", UIMax = "20.0", ClampMin = "0.0", ClampMax = "90.0", EditCondition = "ConversionMode == EConvertToPolygonsMode::FaceNormalDeviation", EditConditionHides))
 	float AngleTolerance = 0.1f;
 
 	/** If true, normals are recomputed per-group, with hard edges at group boundaries */
-	UPROPERTY(EditAnywhere, Category = PolyGroups)
+	UPROPERTY(EditAnywhere, Category = PolyGroups, meta = (EditCondition = "ConversionMode == EConvertToPolygonsMode::FaceNormalDeviation", EditConditionHides))
 	bool bCalculateNormals = true;
 	
 	/** Display each group with a different auto-generated color */
diff --git a/Engine/Plugins/Experimental/MeshModelingToolset/Source/MeshModelingTools/Public/FindPolygonsAlgorithm.h b/Engine/Plugins/Experimental/MeshModelingToolset/Source/MeshModelingTools/Public/FindPolygonsAlgorithm.h
--- a/Engine/Plugins/Experimental/MeshModelingToolset/Source/MeshModelingTools/Public/FindPolygonsAlgorithm.h
+++ b/Engine/Plugins/Experimental/MeshModelingToolset/Source/MeshModelingTools/Public/FindPolygonsAlgorithm.h
@@ -21,13 +21,12 @@ public:
 
 	bool FindPolygonsFromFaceNormals(double DotTolerance = 0.0001);
 	bool FindPolygonsFromUVIslands();
+	bool FindPolygonsFromConnectedTris();
 
 	bool FindPolygonEdges();
 
 
 protected:
 
-	TArray<int> PolygonGroupIDs;
-
 	void SetGroupsFromPolygons();
 };
\ No newline at end of file
