MeshModelingTools: Replace MultiSelectionTool PrimitiveComponentTarget usage with ToolTargets (pt.1)

diff --git a/Engine/Plugins/Experimental/MeshModelingToolset/Source/MeshModelingTools/Private/AlignObjectsTool.cpp b/Engine/Plugins/Experimental/MeshModelingToolset/Source/MeshModelingTools/Private/AlignObjectsTool.cpp
--- a/Engine/Plugins/Experimental/MeshModelingToolset/Source/MeshModelingTools/Private/AlignObjectsTool.cpp
+++ b/Engine/Plugins/Experimental/MeshModelingToolset/Source/MeshModelingTools/Private/AlignObjectsTool.cpp
@@ -7,9 +7,10 @@
 #include "ToolSetupUtil.h"
 #include "DynamicMesh3.h"
 #include "MeshAdapterTransforms.h"
-#include "MeshDescriptionAdapter.h"
 #include "MeshTransforms.h"
 #include "BaseBehaviors/ClickDragBehavior.h"
+#include "TargetInterfaces/PrimitiveComponentBackedTarget.h"
+#include "ToolTargetManager.h"
 
 #include "Components/PrimitiveComponent.h"
 #include "CollisionQueryParams.h"
@@ -21,30 +22,25 @@
  * ToolBuilder
  */
 
+const FToolTargetTypeRequirements& UAlignObjectsToolBuilder::GetTargetRequirements() const
+{
+	static FToolTargetTypeRequirements TypeRequirements({
+		UPrimitiveComponentBackedTarget::StaticClass()
+		});
+	return TypeRequirements;
+}
 
 bool UAlignObjectsToolBuilder::CanBuildTool(const FToolBuilderState& SceneState) const
 {
-	return ToolBuilderUtil::CountComponents(SceneState, CanMakeComponentTarget) >= 2;
+	return SceneState.TargetManager->CountSelectedAndTargetable(SceneState, GetTargetRequirements()) >= 2;
 }
 
 UInteractiveTool* UAlignObjectsToolBuilder::BuildTool(const FToolBuilderState& SceneState) const
 {
 	UAlignObjectsTool* NewTool = NewObject<UAlignObjectsTool>(SceneState.ToolManager);
 
-	TArray<UActorComponent*> Components = ToolBuilderUtil::FindAllComponents(SceneState, CanMakeComponentTarget);
-	check(Components.Num() > 0);
-
-	TArray<TUniquePtr<FPrimitiveComponentTarget>> ComponentTargets;
-	for (UActorComponent* ActorComponent : Components)
-	{
-		auto* MeshComponent = Cast<UPrimitiveComponent>(ActorComponent);
-		if (MeshComponent)
-		{
-			ComponentTargets.Add(MakeComponentTarget(MeshComponent));
-		}
-	}
-
-	NewTool->SetSelection(MoveTemp(ComponentTargets));
+	TArray<TObjectPtr<UToolTarget>> Targets = SceneState.TargetManager->BuildAllSelectedTargetable(SceneState, GetTargetRequirements());
+	NewTool->SetTargets(MoveTemp(Targets));
 	NewTool->SetWorld(SceneState.World, SceneState.GizmoManager);
 
 	return NewTool;
@@ -147,8 +143,11 @@ void UAlignObjectsTool::Precompute()
 	CombinedBounds = FAxisAlignedBox3d::Empty();
 	AveragePivot = FVector3d::Zero();
 
-	for (TUniquePtr<FPrimitiveComponentTarget>& Component : ComponentTargets)
+	for (TObjectPtr<UToolTarget>& Target : Targets)
 	{
+		IPrimitiveComponentBackedTarget* Component = Cast<IPrimitiveComponentBackedTarget>(Target);
+		check(Component);
+
 		FAlignInfo AlignInfo;
 		AlignInfo.Component = Component->GetOwnerComponent();
 		AlignInfo.SavedTransform = Component->GetWorldTransform();
diff --git a/Engine/Plugins/Experimental/MeshModelingToolset/Source/MeshModelingTools/Private/BakeMeshAttributeMapsTool.cpp b/Engine/Plugins/Experimental/MeshModelingToolset/Source/MeshModelingTools/Private/BakeMeshAttributeMapsTool.cpp
--- a/Engine/Plugins/Experimental/MeshModelingToolset/Source/MeshModelingTools/Private/BakeMeshAttributeMapsTool.cpp
+++ b/Engine/Plugins/Experimental/MeshModelingToolset/Source/MeshModelingTools/Private/BakeMeshAttributeMapsTool.cpp
@@ -26,6 +26,12 @@
 #include "AssetUtils/MeshDescriptionUtil.h"
 #include "AssetGenerationUtil.h"
 
+#include "TargetInterfaces/MaterialProvider.h"
+#include "TargetInterfaces/MeshDescriptionProvider.h"
+#include "TargetInterfaces/PrimitiveComponentBackedTarget.h"
+#include "TargetInterfaces/StaticMeshBackedTarget.h"
+#include "ToolTargetManager.h"
+
 // required to pass UStaticMesh asset so we can save at same location
 #include "Engine/Classes/Components/StaticMeshComponent.h"
 #include "Engine/Classes/Engine/StaticMesh.h"
@@ -37,28 +43,31 @@
  * ToolBuilder
  */
 
+const FToolTargetTypeRequirements& UBakeMeshAttributeMapsToolBuilder::GetTargetRequirements() const
+{
+	static FToolTargetTypeRequirements TypeRequirements({
+		UMeshDescriptionProvider::StaticClass(),
+		UPrimitiveComponentBackedTarget::StaticClass(),
+		UStaticMeshBackedTarget::StaticClass(),			// currently only supports StaticMesh targets
+		UMaterialProvider::StaticClass()
+		});
+	return TypeRequirements;
+}
 
 bool UBakeMeshAttributeMapsToolBuilder::CanBuildTool(const FToolBuilderState& SceneState) const
 {
-	// NOTE: currently can only bake for UStaticMeshComponent
-	int32 NumTargets = ToolBuilderUtil::CountComponents(SceneState, CanMakeComponentTarget);
-	return (NumTargets == 1 || NumTargets == 2)
-		&& (ToolBuilderUtil::CountComponents(SceneState, [&](UActorComponent* Comp) { return Cast<UStaticMeshComponent>(Comp) != nullptr; }) == NumTargets);
+	int32 NumTargets = SceneState.TargetManager->CountSelectedAndTargetable(SceneState, GetTargetRequirements());
+	return (NumTargets == 1 || NumTargets == 2);
 }
 
 UInteractiveTool* UBakeMeshAttributeMapsToolBuilder::BuildTool(const FToolBuilderState& SceneState) const
 {
 	UBakeMeshAttributeMapsTool* NewTool = NewObject<UBakeMeshAttributeMapsTool>(SceneState.ToolManager);
-	NewTool->SetAssetAPI(AssetAPI);
 
-	TArray<UActorComponent*> Components = ToolBuilderUtil::FindAllComponents(SceneState, CanMakeComponentTarget);
-	TArray<TUniquePtr<FPrimitiveComponentTarget>> MeshComponents;
-	for (int32 k = 0; k < Components.Num(); ++k)
-	{
-		MeshComponents.Add(MakeComponentTarget(Cast<UPrimitiveComponent>(Components[k])));
-	}
+	TArray<TObjectPtr<UToolTarget>> Targets = SceneState.TargetManager->BuildAllSelectedTargetable(SceneState, GetTargetRequirements());
+	NewTool->SetTargets(MoveTemp(Targets));
+	NewTool->SetAssetAPI(AssetAPI);
 
-	NewTool->SetSelection(MoveTemp(MeshComponents));
 	return NewTool;
 }
 
@@ -91,19 +100,23 @@ void UBakeMeshAttributeMapsTool::Setup()
 {
 	UInteractiveTool::Setup();
 
+	IPrimitiveComponentBackedTarget* TargetComponent = TargetComponentInterface(0);
+	IMeshDescriptionProvider* MeshProvider = TargetMeshProviderInterface(0);
+	IMaterialProvider* MaterialProvider = TargetMaterialInterface(0);
+
 	// copy input MeshDescription and make sure it has initialized normals/tangents
-	BaseMeshDescription = MakeShared<FMeshDescription, ESPMode::ThreadSafe>(*ComponentTargets[0]->GetMesh());
-	UE::MeshDescription::InitializeAutoGeneratedAttributes(*BaseMeshDescription, ComponentTargets[0]->GetOwnerComponent(), 0);
+	BaseMeshDescription = MakeShared<FMeshDescription, ESPMode::ThreadSafe>(*MeshProvider->GetMeshDescription());
+	UE::MeshDescription::InitializeAutoGeneratedAttributes(*BaseMeshDescription, TargetComponent->GetOwnerComponent(), 0);
 
 	// create dynamic mesh component to use for live preview
-	DynamicMeshComponent = NewObject<USimpleDynamicMeshComponent>(ComponentTargets[0]->GetOwnerActor(), "DynamicMesh");
-	DynamicMeshComponent->SetupAttachment(ComponentTargets[0]->GetOwnerActor()->GetRootComponent());
+	DynamicMeshComponent = NewObject<USimpleDynamicMeshComponent>(TargetComponent->GetOwnerActor(), "DynamicMesh");
+	DynamicMeshComponent->SetupAttachment(TargetComponent->GetOwnerActor()->GetRootComponent());
 	DynamicMeshComponent->RegisterComponent();
-	DynamicMeshComponent->SetWorldTransform(ComponentTargets[0]->GetWorldTransform());
+	DynamicMeshComponent->SetWorldTransform(TargetComponent->GetWorldTransform());
 
 	// transfer materials
 	FComponentMaterialSet MaterialSet;
-	ComponentTargets[0]->GetMaterialSet(MaterialSet);
+	MaterialProvider->GetMaterialSet(MaterialSet, false);
 	for (int k = 0; k < MaterialSet.Materials.Num(); ++k)
 	{
 		DynamicMeshComponent->SetMaterial(k, MaterialSet.Materials[k]);
@@ -125,10 +138,10 @@ void UBakeMeshAttributeMapsTool::Setup()
 		DynamicMeshComponent->SetOverrideRenderMaterial(PreviewMaterial);
 	}
 
-	bIsBakeToSelf = (ComponentTargets.Num() == 1);
+	bIsBakeToSelf = (Targets.Num() == 1);
 
 	// hide input StaticMeshComponent
-	ComponentTargets[0]->SetOwnerVisibility(false);
+	TargetComponent->SetOwnerVisibility(false);
 
 
 	Settings = NewObject<UBakeMeshAttributeMapsToolProperties>(this);
@@ -219,6 +232,8 @@ bool UBakeMeshAttributeMapsTool::CanAccept() const
 
 void UBakeMeshAttributeMapsTool::Shutdown(EToolShutdownType ShutdownType)
 {
+	IPrimitiveComponentBackedTarget* TargetComponent = TargetComponentInterface(0);
+
 	Settings->SaveProperties(this);
 	OcclusionMapProps->SaveProperties(this);
 	NormalMapProps->SaveProperties(this);
@@ -228,14 +243,14 @@ void UBakeMeshAttributeMapsTool::Shutdown(EToolShutdownType ShutdownType)
 
 	if (DynamicMeshComponent != nullptr)
 	{
-		ComponentTargets[0]->SetOwnerVisibility(true);
+		TargetComponent->SetOwnerVisibility(true);
 
 		if (ShutdownType == EToolShutdownType::Accept)
 		{
-			UStaticMeshComponent* StaticMeshComponent = CastChecked<UStaticMeshComponent>(ComponentTargets[0]->GetOwnerComponent());
+			UStaticMeshComponent* StaticMeshComponent = CastChecked<UStaticMeshComponent>(TargetComponent->GetOwnerComponent());
 			UStaticMesh* StaticMeshAsset = StaticMeshComponent->GetStaticMesh();
 			check(StaticMeshAsset);
-			FString BaseName = ComponentTargets[0]->GetOwnerActor()->GetName();
+			FString BaseName = TargetComponent->GetOwnerActor()->GetName();
 
 			if (AssetAPI != nullptr)
 			{
@@ -304,13 +319,15 @@ void UBakeMeshAttributeMapsTool::Render(IToolsContextRenderAPI* RenderAPI)
 {
 	UpdateResult();
 
+	IPrimitiveComponentBackedTarget* TargetComponent = TargetComponentInterface(0);
+
 	float GrayLevel = VisualizationProps->BaseGrayLevel;
 	PreviewMaterial->SetVectorParameterValue(TEXT("BaseColor"), FVector(GrayLevel, GrayLevel, GrayLevel) );
 	float AOWeight = VisualizationProps->OcclusionMultiplier;
 	PreviewMaterial->SetScalarParameterValue(TEXT("AOWeight"), AOWeight );
 
 	FPrimitiveDrawInterface* PDI = RenderAPI->GetPrimitiveDrawInterface();
-	FTransform Transform = ComponentTargets[0]->GetWorldTransform();
+	FTransform Transform = TargetComponent->GetWorldTransform();
 }
 
 
@@ -318,16 +335,18 @@ void UBakeMeshAttributeMapsTool::Render(IToolsContextRenderAPI* RenderAPI)
 
 void UBakeMeshAttributeMapsTool::UpdateDetailMesh()
 {
+	IPrimitiveComponentBackedTarget* TargetComponent = TargetComponentInterface(0);
+	IPrimitiveComponentBackedTarget* DetailComponent = TargetComponentInterface(bIsBakeToSelf ? 0 : 1);
+	IMeshDescriptionProvider* DetailMeshProvider = TargetMeshProviderInterface(bIsBakeToSelf ? 0 : 1);
+
 	DetailMesh = MakeShared<FDynamicMesh3, ESPMode::ThreadSafe>();
 	FMeshDescriptionToDynamicMesh Converter;
-	TUniquePtr<FPrimitiveComponentTarget>& UseDetailTarget = (bIsBakeToSelf) ? ComponentTargets[0] : ComponentTargets[1];
-	Converter.Convert(UseDetailTarget->GetMesh(), *DetailMesh);
-
+	Converter.Convert(DetailMeshProvider->GetMeshDescription(), *DetailMesh);
 	if (Settings->bUseWorldSpace && bIsBakeToSelf == false)
 	{
-		FTransform3d DetailToWorld(UseDetailTarget->GetWorldTransform());
+		FTransform3d DetailToWorld(DetailComponent->GetWorldTransform());
 		MeshTransforms::ApplyTransform(*DetailMesh, DetailToWorld);
-		FTransform3d WorldToBase(ComponentTargets[0]->GetWorldTransform());
+		FTransform3d WorldToBase(TargetComponent->GetWorldTransform());
 		MeshTransforms::ApplyTransform(*DetailMesh, WorldToBase.Inverse());
 	}
 	
diff --git a/Engine/Plugins/Experimental/MeshModelingToolset/Source/MeshModelingTools/Private/BakeTransformTool.cpp b/Engine/Plugins/Experimental/MeshModelingToolset/Source/MeshModelingTools/Private/BakeTransformTool.cpp
--- a/Engine/Plugins/Experimental/MeshModelingToolset/Source/MeshModelingTools/Private/BakeTransformTool.cpp
+++ b/Engine/Plugins/Experimental/MeshModelingToolset/Source/MeshModelingTools/Private/BakeTransformTool.cpp
@@ -18,6 +18,10 @@
 
 #include "AssetGenerationUtil.h"
 
+#include "TargetInterfaces/MeshDescriptionCommitter.h"
+#include "TargetInterfaces/PrimitiveComponentBackedTarget.h"
+#include "ToolTargetManager.h"
+
 
 #define LOCTEXT_NAMESPACE "UBakeTransformTool"
 
@@ -27,29 +31,26 @@
  */
 
 
+const FToolTargetTypeRequirements& UBakeTransformToolBuilder::GetTargetRequirements() const
+{
+	static FToolTargetTypeRequirements TypeRequirements({
+		UMeshDescriptionCommitter::StaticClass(),
+		UPrimitiveComponentBackedTarget::StaticClass()
+		});
+	return TypeRequirements;
+}
+
 bool UBakeTransformToolBuilder::CanBuildTool(const FToolBuilderState& SceneState) const
 {
-	return ToolBuilderUtil::CountComponents(SceneState, CanMakeComponentTarget) > 0;
+	return SceneState.TargetManager->CountSelectedAndTargetable(SceneState, GetTargetRequirements()) > 0;
 }
 
 UInteractiveTool* UBakeTransformToolBuilder::BuildTool(const FToolBuilderState& SceneState) const
 {
 	UBakeTransformTool* NewTool = NewObject<UBakeTransformTool>(SceneState.ToolManager);
 
-	TArray<UActorComponent*> Components = ToolBuilderUtil::FindAllComponents(SceneState, CanMakeComponentTarget);
-	check(Components.Num() > 0);
-
-	TArray<TUniquePtr<FPrimitiveComponentTarget>> ComponentTargets;
-	for (UActorComponent* ActorComponent : Components)
-	{
-		auto* MeshComponent = Cast<UPrimitiveComponent>(ActorComponent);
-		if (MeshComponent)
-		{
-			ComponentTargets.Add(MakeComponentTarget(MeshComponent));
-		}
-	}
-
-	NewTool->SetSelection(MoveTemp(ComponentTargets));
+	TArray<TObjectPtr<UToolTarget>> Targets = SceneState.TargetManager->BuildAllSelectedTargetable(SceneState, GetTargetRequirements());
+	NewTool->SetTargets(MoveTemp(Targets));
 	NewTool->SetWorld(SceneState.World);
 	NewTool->SetAssetAPI(AssetAPI);
 
@@ -88,7 +89,7 @@ void UBakeTransformTool::Setup()
 	FText AllTheWarnings = LOCTEXT("BakeTransformWarning", "WARNING: This Tool will Modify the selected StaticMesh Assets! If you do not wish to modify the original Assets, please make copies in the Content Browser first!");
 
 	// detect and warn about any meshes in selection that correspond to same source data
-	bool bSharesSources = GetMapToFirstComponentsSharingSourceData(MapToFirstOccurrences);
+	bool bSharesSources = GetMapToSharedSourceData(MapToFirstOccurrences);
 	if (bSharesSources)
 	{
 		AllTheWarnings = FText::Format(FTextFormat::FromString("{0}\n\n{1}"), AllTheWarnings, LOCTEXT("BakeTransformSharedAssetsWarning", "WARNING: Multiple meshes in your selection use the same source asset!  This is not supported -- each asset can only have one baked transform."));
@@ -97,7 +98,8 @@ void UBakeTransformTool::Setup()
 	bool bHasZeroScales = false;
 	for (int32 ComponentIdx = 0; ComponentIdx < ComponentTargets.Num(); ComponentIdx++)
 	{
-		if (ComponentTargets[ComponentIdx]->GetWorldTransform().GetScale3D().GetAbsMin() < KINDA_SMALL_NUMBER)
+		IPrimitiveComponentBackedTarget* TargetComponent = TargetComponentInterface(ComponentIdx);
+		if (TargetComponent->GetWorldTransform().GetScale3D().GetAbsMin() < KINDA_SMALL_NUMBER)
 		{
 			bHasZeroScales = true;
 		}
@@ -138,10 +140,12 @@ void UBakeTransformTool::UpdateAssets()
 	GetToolManager()->BeginUndoTransaction(LOCTEXT("BakeTransformToolTransactionName", "Bake Transforms"));
 
 	TArray<FTransform3d> BakedTransforms;
-	for (int32 ComponentIdx = 0; ComponentIdx < ComponentTargets.Num(); ComponentIdx++)
+	for (int32 ComponentIdx = 0; ComponentIdx < Targets.Num(); ComponentIdx++)
 	{
-		
-		FTransform3d ComponentToWorld(ComponentTargets[ComponentIdx]->GetWorldTransform());
+		IPrimitiveComponentBackedTarget* TargetComponent = TargetComponentInterface(ComponentIdx);
+		IMeshDescriptionCommitter* TargetMeshCommitter = TargetMeshCommitterInterface(ComponentIdx);
+
+		FTransform3d ComponentToWorld(TargetComponent->GetWorldTransform());
 		FTransform3d ToBakePart = FTransform3d::Identity();
 		FTransform3d NewWorldPart = ComponentToWorld;
 
@@ -208,14 +212,14 @@ void UBakeTransformTool::UpdateAssets()
 				check(false); // must explicitly handle all cases
 			}
 
-			ComponentTargets[ComponentIdx]->CommitMesh([this, &ToBakePart, &NewWorldPart](const FPrimitiveComponentTarget::FCommitParams& CommitParams)
+			TargetMeshCommitter->CommitMeshDescription([this, &ToBakePart, &NewWorldPart](const IMeshDescriptionCommitter::FCommitterParams& CommitParams)
 			{
-				FMeshDescriptionEditableTriangleMeshAdapter EditableMeshDescAdapter(CommitParams.MeshDescription);
+				FMeshDescriptionEditableTriangleMeshAdapter EditableMeshDescAdapter(CommitParams.MeshDescriptionOut);
 
 				// do this part within the commit because we have the MeshDescription already computed
 				if (BasicProperties->bRecenterPivot)
 				{
-					FBox BBox = CommitParams.MeshDescription->ComputeBoundingBox();
+					FBox BBox = CommitParams.MeshDescriptionOut->ComputeBoundingBox();
 					FVector3d Center(BBox.GetCenter());
 					FFrame3d LocalFrame(Center);
 					ToBakePart.SetTranslation(ToBakePart.GetTranslation() - Center);
@@ -227,17 +231,17 @@ void UBakeTransformTool::UpdateAssets()
 				FVector3d ScaleVec = ToBakePart.GetScale();
 				if (ScaleVec.X * ScaleVec.Y * ScaleVec.Z < 0)
 				{
-					CommitParams.MeshDescription->ReverseAllPolygonFacing();
+					CommitParams.MeshDescriptionOut->ReverseAllPolygonFacing();
 				}
 			});
 
 			BakedTransforms.Add(ToBakePart);
 		}
 
-		UPrimitiveComponent* Component = ComponentTargets[ComponentIdx]->GetOwnerComponent();
+		UPrimitiveComponent* Component = TargetComponent->GetOwnerComponent();
 		Component->Modify();
 		Component->SetWorldTransform((FTransform)NewWorldPart);
-		ComponentTargets[ComponentIdx]->GetOwnerActor()->MarkComponentsRenderStateDirty();
+		TargetComponent->GetOwnerActor()->MarkComponentsRenderStateDirty();
 	}
 
 	GetToolManager()->EndUndoTransaction();
diff --git a/Engine/Plugins/Experimental/MeshModelingToolset/Source/MeshModelingTools/Private/CombineMeshesTool.cpp b/Engine/Plugins/Experimental/MeshModelingToolset/Source/MeshModelingTools/Private/CombineMeshesTool.cpp
--- a/Engine/Plugins/Experimental/MeshModelingToolset/Source/MeshModelingTools/Private/CombineMeshesTool.cpp
+++ b/Engine/Plugins/Experimental/MeshModelingToolset/Source/MeshModelingTools/Private/CombineMeshesTool.cpp
@@ -21,6 +21,12 @@
 #include "Components/StaticMeshComponent.h"
 #include "Engine/StaticMesh.h"
 
+#include "TargetInterfaces/MaterialProvider.h"
+#include "TargetInterfaces/MeshDescriptionCommitter.h"
+#include "TargetInterfaces/MeshDescriptionProvider.h"
+#include "TargetInterfaces/PrimitiveComponentBackedTarget.h"
+#include "ToolTargetManager.h"
+
 #if WITH_EDITOR
 #include "Misc/ScopedSlowTask.h"
 #endif
@@ -34,31 +40,30 @@
  */
 
 
+const FToolTargetTypeRequirements& UCombineMeshesToolBuilder::GetTargetRequirements() const
+{
+	static FToolTargetTypeRequirements TypeRequirements({
+		UMeshDescriptionCommitter::StaticClass(),
+		UMeshDescriptionProvider::StaticClass(),
+		UPrimitiveComponentBackedTarget::StaticClass(),
+		UMaterialProvider::StaticClass()
+		});
+	return TypeRequirements;
+}
+
 bool UCombineMeshesToolBuilder::CanBuildTool(const FToolBuilderState& SceneState) const
 {
 	return (bIsDuplicateTool) ?
-		  (AssetAPI != nullptr && ToolBuilderUtil::CountComponents(SceneState, CanMakeComponentTarget) == 1)
-		: (AssetAPI != nullptr && ToolBuilderUtil::CountComponents(SceneState, CanMakeComponentTarget) > 1);
+		  (AssetAPI != nullptr && SceneState.TargetManager->CountSelectedAndTargetable(SceneState, GetTargetRequirements()) == 1)
+		: (AssetAPI != nullptr && SceneState.TargetManager->CountSelectedAndTargetable(SceneState, GetTargetRequirements()) > 1);
 }
 
 UInteractiveTool* UCombineMeshesToolBuilder::BuildTool(const FToolBuilderState& SceneState) const
 {
 	UCombineMeshesTool* NewTool = NewObject<UCombineMeshesTool>(SceneState.ToolManager);
 
-	TArray<UActorComponent*> Components = ToolBuilderUtil::FindAllComponents(SceneState, CanMakeComponentTarget);
-	check(Components.Num() > 0);
-
-	TArray<TUniquePtr<FPrimitiveComponentTarget>> ComponentTargets;
-	for (UActorComponent* ActorComponent : Components)
-	{
-		auto* MeshComponent = Cast<UPrimitiveComponent>(ActorComponent);
-		if (MeshComponent)
-		{
-			ComponentTargets.Add(MakeComponentTarget(MeshComponent));
-		}
-	}
-
-	NewTool->SetSelection(MoveTemp(ComponentTargets));
+	TArray<TObjectPtr<UToolTarget>> Targets = SceneState.TargetManager->BuildAllSelectedTargetable(SceneState, GetTargetRequirements());
+	NewTool->SetTargets(MoveTemp(Targets));
 	NewTool->SetWorld(SceneState.World);
 	NewTool->SetAssetAPI(AssetAPI);
 	NewTool->SetDuplicateMode(bIsDuplicateTool);
@@ -100,15 +105,15 @@ void UCombineMeshesTool::Setup()
 		}
 		else
 		{
-			int32 Index = (BasicProperties->WriteOutputTo == ECombineTargetType::FirstInputAsset) ? 0 : ComponentTargets.Num() - 1;
-			BasicProperties->OutputAsset = AssetGenerationUtil::GetComponentAssetBaseName(ComponentTargets[Index]->GetOwnerComponent(), false);
+			int32 Index = (BasicProperties->WriteOutputTo == ECombineTargetType::FirstInputAsset) ? 0 : Targets.Num() - 1;
+			BasicProperties->OutputAsset = AssetGenerationUtil::GetComponentAssetBaseName(TargetComponentInterface(Index)->GetOwnerComponent(), false);
 		}
 	});
 
 	if (bDuplicateMode)
 	{
 		SetToolDisplayName(LOCTEXT("DuplicateMeshesToolName", "Duplicate"));
-		BasicProperties->OutputName = AssetGenerationUtil::GetComponentAssetBaseName(ComponentTargets[0]->GetOwnerComponent());
+		BasicProperties->OutputName = AssetGenerationUtil::GetComponentAssetBaseName(TargetComponentInterface(0)->GetOwnerComponent());
 	}
 	else
 	{
@@ -175,21 +180,21 @@ void UCombineMeshesTool::CreateNewAsset()
 
 #if WITH_EDITOR
 	FBox Box(ForceInit);
-	for (int32 ComponentIdx = 0; ComponentIdx < ComponentTargets.Num(); ComponentIdx++)
+	for (int32 ComponentIdx = 0; ComponentIdx < Targets.Num(); ComponentIdx++)
 	{
-		Box += ComponentTargets[ComponentIdx]->GetOwnerComponent()->Bounds.GetBox();
+		Box += TargetComponentInterface(ComponentIdx)->GetOwnerComponent()->Bounds.GetBox();
 	}
 
 	bool bMergeSameMaterials = true;
 	TArray<UMaterialInterface*> AllMaterials;
 	TMap<UMaterialInterface*, int> KnownMaterials;
 	TArray<int> CombinedMatToOutMatIdx;
-	for (int32 ComponentIdx = 0; ComponentIdx < ComponentTargets.Num(); ComponentIdx++)
+	for (int32 ComponentIdx = 0; ComponentIdx < Targets.Num(); ComponentIdx++)
 	{
-		TUniquePtr<FPrimitiveComponentTarget>& ComponentTarget = ComponentTargets[ComponentIdx];
-		for (int MaterialIdx = 0, NumMaterials = ComponentTarget->GetNumMaterials(); MaterialIdx < NumMaterials; MaterialIdx++)
+		IMaterialProvider* ComponentMaterial = TargetMaterialInterface(ComponentIdx);
+		for (int MaterialIdx = 0, NumMaterials = ComponentMaterial->GetNumMaterials(); MaterialIdx < NumMaterials; MaterialIdx++)
 		{
-			UMaterialInterface* Mat = ComponentTarget->GetMaterial(MaterialIdx);
+			UMaterialInterface* Mat = ComponentMaterial->GetMaterial(MaterialIdx);
 			int32 OutMatIdx = -1;
 			if (!bMergeSameMaterials || !KnownMaterials.Contains(Mat))
 			{
@@ -216,23 +221,23 @@ void UCombineMeshesTool::CreateNewAsset()
 	FTransform ToAccum(-Box.GetCenter());
 
 	{
-		FScopedSlowTask SlowTask(ComponentTargets.Num()+1, 
+		FScopedSlowTask SlowTask(Targets.Num()+1, 
 			bDuplicateMode ? 
 			LOCTEXT("DuplicateMeshBuild", "Building duplicate mesh ...") :
 			LOCTEXT("CombineMeshesBuild", "Building combined mesh ..."));
 		SlowTask.MakeDialog();
 
 		int MatIndexBase = 0;
-		for (int32 ComponentIdx = 0; ComponentIdx < ComponentTargets.Num(); ComponentIdx++)
+		for (int32 ComponentIdx = 0; ComponentIdx < Targets.Num(); ComponentIdx++)
 		{
 			SlowTask.EnterProgressFrame(1);
-			TUniquePtr<FPrimitiveComponentTarget>& ComponentTarget = ComponentTargets[ComponentIdx];
+			IMeshDescriptionProvider* TargetMeshProvider = TargetMeshProviderInterface(ComponentIdx);
 
 			FMeshDescriptionToDynamicMesh Converter;
 			FDynamicMesh3 ComponentDMesh;
-			Converter.Convert(ComponentTarget->GetMesh(), ComponentDMesh);
+			Converter.Convert(TargetMeshProvider->GetMeshDescription(), ComponentDMesh);
 
-			FTransform3d XF = (FTransform3d)(ComponentTarget->GetWorldTransform() * ToAccum);
+			FTransform3d XF = (FTransform3d)(TargetComponentInterface(ComponentIdx)->GetWorldTransform() * ToAccum);
 			if (XF.GetDeterminant() < 0)
 			{
 				ComponentDMesh.ReverseOrientation(false);
@@ -259,7 +264,7 @@ void UCombineMeshesTool::CreateNewAsset()
 			}
 			
 
-			MatIndexBase += ComponentTarget->GetNumMaterials();
+			MatIndexBase += TargetMaterialInterface(ComponentIdx)->GetNumMaterials();
 		}
 
 		SlowTask.EnterProgressFrame(1);
@@ -267,8 +272,8 @@ void UCombineMeshesTool::CreateNewAsset()
 		if (bDuplicateMode)
 		{
 			// TODO: will need to refactor this when we support duplicating multiple
-			check(ComponentTargets.Num() == 1);
-			AccumToWorld = ComponentTargets[0]->GetWorldTransform();
+			check(Targets.Num() == 1);
+			AccumToWorld = TargetComponentInterface(0)->GetWorldTransform();
 		}
 
 		// max len explicitly enforced here, would ideally notify user
@@ -299,9 +304,9 @@ void UCombineMeshesTool::CreateNewAsset()
 
 	
 	TArray<AActor*> Actors;
-	for (auto& ComponentTarget : ComponentTargets)
+	for (int32 Idx = 0; Idx < Targets.Num(); Idx++)
 	{
-		Actors.Add(ComponentTarget->GetOwnerActor());
+		Actors.Add(TargetComponentInterface(Idx)->GetOwnerActor());
 	}
 	HandleSourceProperties->ApplyMethod(Actors, GetToolManager());
 
@@ -334,12 +339,12 @@ void UCombineMeshesTool::UpdateExistingAsset()
 	TArray<UMaterialInterface*> AllMaterials;
 	TMap<UMaterialInterface*, int> KnownMaterials;
 	TArray<int> CombinedMatToOutMatIdx;
-	for (int32 ComponentIdx = 0; ComponentIdx < ComponentTargets.Num(); ComponentIdx++)
+	for (int32 ComponentIdx = 0; ComponentIdx < Targets.Num(); ComponentIdx++)
 	{
-		TUniquePtr<FPrimitiveComponentTarget>& ComponentTarget = ComponentTargets[ComponentIdx];
-		for (int MaterialIdx = 0, NumMaterials = ComponentTarget->GetNumMaterials(); MaterialIdx < NumMaterials; MaterialIdx++)
+		IMaterialProvider* TargetMaterial = TargetMaterialInterface(ComponentIdx);
+		for (int MaterialIdx = 0, NumMaterials = TargetMaterial->GetNumMaterials(); MaterialIdx < NumMaterials; MaterialIdx++)
 		{
-			UMaterialInterface* Mat = ComponentTarget->GetMaterial(MaterialIdx);
+			UMaterialInterface* Mat = TargetMaterial->GetMaterial(MaterialIdx);
 			int32 OutMatIdx = -1;
 			if (!bMergeSameMaterials || !KnownMaterials.Contains(Mat))
 			{
@@ -363,29 +368,30 @@ void UCombineMeshesTool::UpdateExistingAsset()
 	AccumulateDMesh.EnableAttributes();
 	AccumulateDMesh.Attributes()->EnableMaterialID();
 
-	int32 SkipIndex = (BasicProperties->WriteOutputTo == ECombineTargetType::FirstInputAsset) ? 0 : (ComponentTargets.Num() - 1);
-	TUniquePtr<FPrimitiveComponentTarget>& UpdateTarget = ComponentTargets[SkipIndex];
+	int32 SkipIndex = (BasicProperties->WriteOutputTo == ECombineTargetType::FirstInputAsset) ? 0 : (Targets.Num() - 1);
+	IPrimitiveComponentBackedTarget* UpdateTarget = TargetComponentInterface(SkipIndex);
+	IMeshDescriptionCommitter* UpdateTargetCommitter = TargetMeshCommitterInterface(SkipIndex);
+	IMaterialProvider* UpdateTargetMaterial = TargetMaterialInterface(SkipIndex);
 	SkipActor = UpdateTarget->GetOwnerActor();
 
 	FTransform3d TargetToWorld = (FTransform3d)UpdateTarget->GetWorldTransform();
 	FTransform3d WorldToTarget = TargetToWorld.Inverse();
 
 	{
-		FScopedSlowTask SlowTask(ComponentTargets.Num()+1, 
+		FScopedSlowTask SlowTask(Targets.Num()+1, 
 			bDuplicateMode ? 
 			LOCTEXT("DuplicateMeshBuild", "Building duplicate mesh ...") :
 			LOCTEXT("CombineMeshesBuild", "Building combined mesh ..."));
 		SlowTask.MakeDialog();
 
 		int MatIndexBase = 0;
-		for (int32 ComponentIdx = 0; ComponentIdx < ComponentTargets.Num(); ComponentIdx++)
+		for (int32 ComponentIdx = 0; ComponentIdx < Targets.Num(); ComponentIdx++)
 		{
 			SlowTask.EnterProgressFrame(1);
-			TUniquePtr<FPrimitiveComponentTarget>& ComponentTarget = ComponentTargets[ComponentIdx];
 
 			FMeshDescriptionToDynamicMesh Converter;
 			FDynamicMesh3 ComponentDMesh;
-			Converter.Convert(ComponentTarget->GetMesh(), ComponentDMesh);
+			Converter.Convert(TargetMeshProviderInterface(ComponentIdx)->GetMeshDescription(), ComponentDMesh);
 
 			// update material IDs to account for combined material set
 			FDynamicMeshMaterialAttribute* MatAttrib = ComponentDMesh.Attributes()->GetMaterialID();
@@ -393,12 +399,12 @@ void UCombineMeshesTool::UpdateExistingAsset()
 			{
 				MatAttrib->SetValue(TID, CombinedMatToOutMatIdx[MatIndexBase + MatAttrib->GetValue(TID)]);
 			}
-			MatIndexBase += ComponentTarget->GetNumMaterials();
+			MatIndexBase += TargetMaterialInterface(ComponentIdx)->GetNumMaterials();
 
 
 			if (ComponentIdx != SkipIndex)
 			{
-				FTransform3d ComponentToWorld = (FTransform3d)ComponentTarget->GetWorldTransform();
+				FTransform3d ComponentToWorld = (FTransform3d)TargetComponentInterface(ComponentIdx)->GetWorldTransform();
 				MeshTransforms::ApplyTransform(ComponentDMesh, ComponentToWorld);
 				if (ComponentToWorld.GetDeterminant() < 0)
 				{
@@ -418,15 +424,15 @@ void UCombineMeshesTool::UpdateExistingAsset()
 
 		SlowTask.EnterProgressFrame(1);
 
-		UpdateTarget->CommitMesh([&](const FPrimitiveComponentTarget::FCommitParams& CommitParams)
+		UpdateTargetCommitter->CommitMeshDescription([&](const IMeshDescriptionCommitter::FCommitterParams& CommitParams)
 		{
 			FDynamicMeshToMeshDescription Converter;
-			Converter.Convert(&AccumulateDMesh, *CommitParams.MeshDescription);
+			Converter.Convert(&AccumulateDMesh, *CommitParams.MeshDescriptionOut);
 		});
 
 		FComponentMaterialSet MaterialSet;
 		MaterialSet.Materials = AllMaterials;
-		UpdateTarget->CommitMaterialSetUpdate(MaterialSet, true);
+		UpdateTargetMaterial->CommitMaterialSetUpdate(MaterialSet, true);
 
 		// select the new actor
 		ToolSelectionUtil::SetNewActorSelection(GetToolManager(), SkipActor);
@@ -435,9 +441,9 @@ void UCombineMeshesTool::UpdateExistingAsset()
 
 	
 	TArray<AActor*> Actors;
-	for (auto& ComponentTarget : ComponentTargets)
+	for (int Idx = 0; Idx < Targets.Num(); Idx++)
 	{
-		AActor* Actor = ComponentTarget->GetOwnerActor();
+		AActor* Actor = TargetComponentInterface(Idx)->GetOwnerActor();
 		if (Actor != SkipActor)
 		{
 			Actors.Add(Actor);
diff --git a/Engine/Plugins/Experimental/MeshModelingToolset/Source/MeshModelingTools/Public/AlignObjectsTool.h b/Engine/Plugins/Experimental/MeshModelingToolset/Source/MeshModelingTools/Public/AlignObjectsTool.h
--- a/Engine/Plugins/Experimental/MeshModelingToolset/Source/MeshModelingTools/Public/AlignObjectsTool.h
+++ b/Engine/Plugins/Experimental/MeshModelingToolset/Source/MeshModelingTools/Public/AlignObjectsTool.h
@@ -24,6 +24,9 @@ class MESHMODELINGTOOLS_API UAlignObjectsToolBuilder : public UInteractiveToolBu
 public:
 	virtual bool CanBuildTool(const FToolBuilderState& SceneState) const override;
 	virtual UInteractiveTool* BuildTool(const FToolBuilderState& SceneState) const override;
+
+protected:
+	virtual const FToolTargetTypeRequirements& GetTargetRequirements() const override;
 };
 
 
diff --git a/Engine/Plugins/Experimental/MeshModelingToolset/Source/MeshModelingTools/Public/BakeMeshAttributeMapsTool.h b/Engine/Plugins/Experimental/MeshModelingToolset/Source/MeshModelingTools/Public/BakeMeshAttributeMapsTool.h
--- a/Engine/Plugins/Experimental/MeshModelingToolset/Source/MeshModelingTools/Public/BakeMeshAttributeMapsTool.h
+++ b/Engine/Plugins/Experimental/MeshModelingToolset/Source/MeshModelingTools/Public/BakeMeshAttributeMapsTool.h
@@ -20,6 +20,9 @@ class UTexture2D;
 class IAssetGenerationAPI;
 template<typename RealType> class TMeshTangents;
 class FMeshImageBakingCache;
+class IPrimitiveComponentBackedTarget;
+class IMeshDescriptionProvider;
+class IMaterialProvider;
 
 /**
  *
@@ -34,6 +37,9 @@ public:
 
 	virtual bool CanBuildTool(const FToolBuilderState& SceneState) const override;
 	virtual UInteractiveTool* BuildTool(const FToolBuilderState& SceneState) const override;
+
+protected:
+	virtual const FToolTargetTypeRequirements& GetTargetRequirements() const override;
 };
 
 
diff --git a/Engine/Plugins/Experimental/MeshModelingToolset/Source/MeshModelingTools/Public/BakeTransformTool.h b/Engine/Plugins/Experimental/MeshModelingToolset/Source/MeshModelingTools/Public/BakeTransformTool.h
--- a/Engine/Plugins/Experimental/MeshModelingToolset/Source/MeshModelingTools/Public/BakeTransformTool.h
+++ b/Engine/Plugins/Experimental/MeshModelingToolset/Source/MeshModelingTools/Public/BakeTransformTool.h
@@ -33,6 +33,9 @@ public:
 
 	virtual bool CanBuildTool(const FToolBuilderState& SceneState) const override;
 	virtual UInteractiveTool* BuildTool(const FToolBuilderState& SceneState) const override;
+
+protected:
+	virtual const FToolTargetTypeRequirements& GetTargetRequirements() const override;
 };
 
 
diff --git a/Engine/Plugins/Experimental/MeshModelingToolset/Source/MeshModelingTools/Public/CombineMeshesTool.h b/Engine/Plugins/Experimental/MeshModelingToolset/Source/MeshModelingTools/Public/CombineMeshesTool.h
--- a/Engine/Plugins/Experimental/MeshModelingToolset/Source/MeshModelingTools/Public/CombineMeshesTool.h
+++ b/Engine/Plugins/Experimental/MeshModelingToolset/Source/MeshModelingTools/Public/CombineMeshesTool.h
@@ -18,6 +18,10 @@ struct FMeshDescription;
 class USimpleDynamicMeshComponent;
 class IAssetGenerationAPI;
 
+class IPrimitiveComponentBackedTarget;
+class IMeshDescriptionCommitter;
+class IMeshDescriptionProvider;
+class IMaterialProvider;
 
 
 
@@ -35,6 +39,9 @@ public:
 
 	virtual bool CanBuildTool(const FToolBuilderState& SceneState) const override;
 	virtual UInteractiveTool* BuildTool(const FToolBuilderState& SceneState) const override;
+
+protected:
+	virtual const FToolTargetTypeRequirements& GetTargetRequirements() const override;
 };
 
 
diff --git a/Engine/Plugins/Experimental/MeshModelingToolset/Source/MeshModelingToolsEditorOnly/Private/MeshToVolumeTool.cpp b/Engine/Plugins/Experimental/MeshModelingToolset/Source/MeshModelingToolsEditorOnly/Private/MeshToVolumeTool.cpp
--- a/Engine/Plugins/Experimental/MeshModelingToolset/Source/MeshModelingToolsEditorOnly/Private/MeshToVolumeTool.cpp
+++ b/Engine/Plugins/Experimental/MeshModelingToolset/Source/MeshModelingToolsEditorOnly/Private/MeshToVolumeTool.cpp
@@ -32,11 +32,11 @@
 
 const FToolTargetTypeRequirements& UMeshToVolumeToolBuilder::GetTargetRequirements() const
 {
-	static FToolTargetTypeRequirements TypeRequirements(nullptr, TArray<const UClass*>({
+	static FToolTargetTypeRequirements TypeRequirements({
 		UMeshDescriptionProvider::StaticClass(),
 		UPrimitiveComponentBackedTarget::StaticClass(),
 		UMaterialProvider::StaticClass()
-		}));
+		});
 	return TypeRequirements;
 }
 
diff --git a/Engine/Source/Editor/Experimental/EditorInteractiveToolsFramework/Private/ToolTargets/StaticMeshComponentToolTarget.cpp b/Engine/Source/Editor/Experimental/EditorInteractiveToolsFramework/Private/ToolTargets/StaticMeshComponentToolTarget.cpp
--- a/Engine/Source/Editor/Experimental/EditorInteractiveToolsFramework/Private/ToolTargets/StaticMeshComponentToolTarget.cpp
+++ b/Engine/Source/Editor/Experimental/EditorInteractiveToolsFramework/Private/ToolTargets/StaticMeshComponentToolTarget.cpp
@@ -159,6 +159,11 @@ void UStaticMeshComponentToolTarget::CommitMeshDescription(const FCommitter& Com
 	Component->RecreatePhysicsState();
 }
 
+UStaticMesh* UStaticMeshComponentToolTarget::GetStaticMesh() const
+{
+	return IsValid() ? Cast<UStaticMeshComponent>(Component)->GetStaticMesh() : nullptr;
+}
+
 
 // Factory
 
diff --git a/Engine/Source/Editor/Experimental/EditorInteractiveToolsFramework/Public/ToolTargets/StaticMeshComponentToolTarget.h b/Engine/Source/Editor/Experimental/EditorInteractiveToolsFramework/Public/ToolTargets/StaticMeshComponentToolTarget.h
--- a/Engine/Source/Editor/Experimental/EditorInteractiveToolsFramework/Public/ToolTargets/StaticMeshComponentToolTarget.h
+++ b/Engine/Source/Editor/Experimental/EditorInteractiveToolsFramework/Public/ToolTargets/StaticMeshComponentToolTarget.h
@@ -7,17 +7,20 @@
 #include "TargetInterfaces/MaterialProvider.h"
 #include "TargetInterfaces/MeshDescriptionCommitter.h"
 #include "TargetInterfaces/MeshDescriptionProvider.h"
+#include "TargetInterfaces/StaticMeshBackedTarget.h"
 #include "ToolTargets/PrimitiveComponentToolTarget.h"
 
 #include "StaticMeshComponentToolTarget.generated.h"
 
+class UStaticMesh;
+
 /**
  * A tool target backed by a static mesh component that can provide and take a mesh
  * description.
  */
 UCLASS(Transient)
 class EDITORINTERACTIVETOOLSFRAMEWORK_API UStaticMeshComponentToolTarget : public UPrimitiveComponentToolTarget,
-	public IMeshDescriptionCommitter, public IMeshDescriptionProvider, public IMaterialProvider
+	public IMeshDescriptionCommitter, public IMeshDescriptionProvider, public IMaterialProvider, public IStaticMeshBackedTarget
 {
 	GENERATED_BODY()
 
@@ -35,6 +38,9 @@ public:
 	void GetMaterialSet(FComponentMaterialSet& MaterialSetOut, bool bPreferAssetMaterials) const override;
 	virtual bool CommitMaterialSetUpdate(const FComponentMaterialSet& MaterialSet, bool bApplyToAsset) override;	
 
+	// IStaticMeshBackedTarget
+	UStaticMesh* GetStaticMesh() const override;
+
 	// Rest provided by parent class
 
 protected:
diff --git a/Engine/Source/Runtime/Experimental/InteractiveToolsFramework/Private/MultiSelectionTool.cpp b/Engine/Source/Runtime/Experimental/InteractiveToolsFramework/Private/MultiSelectionTool.cpp
new file mode 100644
--- /dev/null
+++ b/Engine/Source/Runtime/Experimental/InteractiveToolsFramework/Private/MultiSelectionTool.cpp
@@ -0,0 +1,74 @@
+// Copyright Epic Games, Inc. All Rights Reserved.
+#include "MultiSelectionTool.h"
+#include "TargetInterfaces/MaterialProvider.h"
+#include "TargetInterfaces/MeshDescriptionCommitter.h"
+#include "TargetInterfaces/MeshDescriptionProvider.h"
+#include "TargetInterfaces/PrimitiveComponentBackedTarget.h"
+#include "TargetInterfaces/AssetBackedTarget.h"
+
+bool UMultiSelectionTool::GetMapToSharedSourceData(TArray<int32>& MapToFirstOccurrences)
+{
+	bool bSharesSources = false;
+	MapToFirstOccurrences.SetNumUninitialized(Targets.Num());
+	for (int32 ComponentIdx = 0; ComponentIdx < Targets.Num(); ComponentIdx++)
+	{
+		MapToFirstOccurrences[ComponentIdx] = -1;
+	}
+	for (int32 ComponentIdx = 0; ComponentIdx < Targets.Num(); ComponentIdx++)
+	{
+		if (MapToFirstOccurrences[ComponentIdx] >= 0) // already mapped
+		{
+			continue;
+		}
+
+		MapToFirstOccurrences[ComponentIdx] = ComponentIdx;
+
+		IAssetBackedTarget* Target = Cast<IAssetBackedTarget>(Targets[ComponentIdx]);
+		if (!Target)
+		{
+			continue;
+		}
+		for (int32 VsIdx = ComponentIdx + 1; VsIdx < Targets.Num(); VsIdx++)
+		{
+			IAssetBackedTarget* OtherTarget = Cast<IAssetBackedTarget>(Targets[VsIdx]);
+			if (OtherTarget && OtherTarget->GetSourceData() == Target->GetSourceData())
+			{
+				bSharesSources = true;
+				MapToFirstOccurrences[VsIdx] = ComponentIdx;
+			}
+		}
+	}
+	return bSharesSources;
+}
+
+IPrimitiveComponentBackedTarget* UMultiSelectionTool::TargetComponentInterface(int32 ComponentIdx) const
+{
+	check(ComponentIdx >= 0 && ComponentIdx < Targets.Num());
+	IPrimitiveComponentBackedTarget* Component = Cast<IPrimitiveComponentBackedTarget>(Targets[ComponentIdx]);
+	check(Component);
+	return Component;
+}
+
+IMeshDescriptionCommitter* UMultiSelectionTool::TargetMeshCommitterInterface(int32 ComponentIdx) const
+{
+	check(ComponentIdx >= 0 && ComponentIdx < Targets.Num());
+	IMeshDescriptionCommitter* TargetMeshCommitter = Cast<IMeshDescriptionCommitter>(Targets[ComponentIdx]);
+	check(TargetMeshCommitter);
+	return TargetMeshCommitter;
+}
+
+IMeshDescriptionProvider* UMultiSelectionTool::TargetMeshProviderInterface(int32 ComponentIdx) const
+{
+	check(ComponentIdx >= 0 && ComponentIdx < Targets.Num());
+	IMeshDescriptionProvider* TargetMeshProvider = Cast<IMeshDescriptionProvider>(Targets[ComponentIdx]);
+	check(TargetMeshProvider);
+	return TargetMeshProvider;
+}
+
+IMaterialProvider* UMultiSelectionTool::TargetMaterialInterface(int32 ComponentIdx) const
+{
+	check(ComponentIdx >= 0 && ComponentIdx < Targets.Num());
+	IMaterialProvider* TargetMaterial = Cast<IMaterialProvider>(Targets[ComponentIdx]);
+	check(TargetMaterial);
+	return TargetMaterial;
+}
\ No newline at end of file
diff --git a/Engine/Source/Runtime/Experimental/InteractiveToolsFramework/Private/ToolTargetManager.cpp b/Engine/Source/Runtime/Experimental/InteractiveToolsFramework/Private/ToolTargetManager.cpp
--- a/Engine/Source/Runtime/Experimental/InteractiveToolsFramework/Private/ToolTargetManager.cpp
+++ b/Engine/Source/Runtime/Experimental/InteractiveToolsFramework/Private/ToolTargetManager.cpp
@@ -66,3 +66,19 @@ UToolTarget* UToolTargetManager::BuildFirstSelectedTargetable(const FToolBuilder
 		}),
 		TargetType);
 }
+
+TArray<TObjectPtr<UToolTarget>> UToolTargetManager::BuildAllSelectedTargetable(const FToolBuilderState& SceneState,
+	const FToolTargetTypeRequirements& TargetType)
+{
+	TArray<UActorComponent*> Components = ToolBuilderUtil::FindAllComponents(SceneState, [&](UActorComponent* Object)
+	{
+		return CanBuildTarget(Object, TargetType);
+	});
+	TArray<TObjectPtr<UToolTarget>> Targets;
+	Targets.Reserve(Components.Num());
+	for (UActorComponent* Component : Components)
+	{
+		Targets.Add(BuildTarget(Component, TargetType));
+	}
+	return Targets;
+}
diff --git a/Engine/Source/Runtime/Experimental/InteractiveToolsFramework/Private/ToolTargets/ToolTarget.cpp b/Engine/Source/Runtime/Experimental/InteractiveToolsFramework/Private/ToolTargets/ToolTarget.cpp
--- a/Engine/Source/Runtime/Experimental/InteractiveToolsFramework/Private/ToolTargets/ToolTarget.cpp
+++ b/Engine/Source/Runtime/Experimental/InteractiveToolsFramework/Private/ToolTargets/ToolTarget.cpp
@@ -4,21 +4,15 @@
 
 bool FToolTargetTypeRequirements::AreSatisfiedBy(UClass* Class) const
 {
-	// Required type must either be null (no requirement) or our type, or our parent type.
-	if (!BaseType || BaseType == Class ||
-		Class->IsChildOf(BaseType))
+	// we have to support all the required interfaces
+	for (const UClass* Interface : Interfaces)
 	{
-		// we have to support all the required interfaces
-		for (const UClass* Interface : Interfaces)
+		if (!Class->ImplementsInterface(Interface))
 		{
-			if (!Class->ImplementsInterface(Interface))
-			{
-				return false;
-			}
+			return false;
 		}
-		return true;
 	}
-	return false;
+	return true;
 }
 
 bool FToolTargetTypeRequirements::AreSatisfiedBy(UToolTarget* ToolTarget) const
diff --git a/Engine/Source/Runtime/Experimental/InteractiveToolsFramework/Public/MultiSelectionTool.h b/Engine/Source/Runtime/Experimental/InteractiveToolsFramework/Public/MultiSelectionTool.h
--- a/Engine/Source/Runtime/Experimental/InteractiveToolsFramework/Public/MultiSelectionTool.h
+++ b/Engine/Source/Runtime/Experimental/InteractiveToolsFramework/Public/MultiSelectionTool.h
@@ -4,23 +4,43 @@
 
 #include "InteractiveTool.h"
 #include "ComponentSourceInterfaces.h"
+#include "ToolTargets/ToolTarget.h"
+
 #include "MultiSelectionTool.generated.h"
 
+class IPrimitiveComponentBackedTarget;
+class IMeshDescriptionCommitter;
+class IMeshDescriptionProvider;
+class IMaterialProvider;
+
 UCLASS(Transient)
 class INTERACTIVETOOLSFRAMEWORK_API UMultiSelectionTool : public UInteractiveTool
 {
 GENERATED_BODY()
 public:
+	// @deprecated Use SetTarget instead, and don't use FPrimitiveComponentTarget.
 	void SetSelection(TArray<TUniquePtr<FPrimitiveComponentTarget>> ComponentTargetsIn)
     {
 		ComponentTargets = MoveTemp(ComponentTargetsIn);
 	}
 
+	void SetTargets(TArray<TObjectPtr<UToolTarget>> TargetsIn)
+	{
+		Targets = MoveTemp(TargetsIn);
+	}
+
 	/**
 	 * @return true if all ComponentTargets of this tool are still valid
 	 */
 	virtual bool AreAllTargetsValid() const
 	{
+		for (const TObjectPtr<UToolTarget>& Target : Targets)
+		{
+			if (Target->IsValid() == false)
+			{
+				return false;
+			}
+		}
 		for (const TUniquePtr<FPrimitiveComponentTarget>& Target : ComponentTargets)
 		{
 			if (Target->IsValid() == false)
@@ -41,6 +61,9 @@ public:
 protected:
 	TArray<TUniquePtr<FPrimitiveComponentTarget>> ComponentTargets{};
 
+	UPROPERTY()
+	TArray<TObjectPtr<UToolTarget>> Targets{};
+
 	/**
 	 * Helper to find which component targets share source data
 	 *
@@ -75,4 +98,31 @@ protected:
 		}
 		return bSharesSources;
 	}
+
+	/**
+	 * Helper to find which targets share source data
+	 *
+	 * @return Array of indices, 1:1 with Targets, indicating the first index where a component target sharing the same source data appeared.
+	 */
+	bool GetMapToSharedSourceData(TArray<int32>& MapToFirstOccurrences);
+
+	/**
+	 * Helper to cast a Target into the IPrimitiveComponentBackedTarget interface.
+	 */
+	IPrimitiveComponentBackedTarget* TargetComponentInterface(int32 ComponentIdx) const;
+
+	/**
+	 * Helper to cast a Target into the IMeshDescriptionCommitter interface.
+	 */
+	IMeshDescriptionCommitter* TargetMeshCommitterInterface(int32 ComponentIdx) const;
+
+	/**
+	 * Helper to cast a Target into the IMeshDescriptionProvider interface.
+	 */
+	IMeshDescriptionProvider* TargetMeshProviderInterface(int32 ComponentIdx) const;
+
+	/**
+	 * Helper to cast a Target into the IMaterialProvider interface.
+	 */
+	IMaterialProvider* TargetMaterialInterface(int32 ComponentIdx) const;
 };
diff --git a/Engine/Source/Runtime/Experimental/InteractiveToolsFramework/Public/TargetInterfaces/AssetBackedTarget.h b/Engine/Source/Runtime/Experimental/InteractiveToolsFramework/Public/TargetInterfaces/AssetBackedTarget.h
new file mode 100644
--- /dev/null
+++ b/Engine/Source/Runtime/Experimental/InteractiveToolsFramework/Public/TargetInterfaces/AssetBackedTarget.h
@@ -0,0 +1,26 @@
+// Copyright Epic Games, Inc. All Rights Reserved.
+
+#pragma once
+
+#include "CoreMinimal.h"
+#include "Engine/EngineTypes.h"
+#include "UObject/Interface.h"
+
+#include "AssetBackedTarget.generated.h"
+
+UINTERFACE()
+class INTERACTIVETOOLSFRAMEWORK_API UAssetBackedTarget : public UInterface
+{
+	GENERATED_BODY()
+};
+
+class INTERACTIVETOOLSFRAMEWORK_API IAssetBackedTarget
+{
+	GENERATED_BODY()
+
+public:
+	/**
+	 * @return the underlying source asset for this Target.
+	 */
+	virtual UObject* GetSourceData() const = 0;
+};
\ No newline at end of file
diff --git a/Engine/Source/Runtime/Experimental/InteractiveToolsFramework/Public/TargetInterfaces/StaticMeshBackedTarget.h b/Engine/Source/Runtime/Experimental/InteractiveToolsFramework/Public/TargetInterfaces/StaticMeshBackedTarget.h
new file mode 100644
--- /dev/null
+++ b/Engine/Source/Runtime/Experimental/InteractiveToolsFramework/Public/TargetInterfaces/StaticMeshBackedTarget.h
@@ -0,0 +1,35 @@
+// Copyright Epic Games, Inc. All Rights Reserved.
+
+#pragma once
+
+#include "CoreMinimal.h"
+#include "Engine/EngineTypes.h"
+#include "UObject/Interface.h"
+#include "TargetInterfaces/AssetBackedTarget.h"
+
+#include "StaticMeshBackedTarget.generated.h"
+
+UINTERFACE()
+class INTERACTIVETOOLSFRAMEWORK_API UStaticMeshBackedTarget : public UAssetBackedTarget
+{
+	GENERATED_BODY()
+};
+
+class INTERACTIVETOOLSFRAMEWORK_API IStaticMeshBackedTarget : public IAssetBackedTarget
+{
+	GENERATED_BODY()
+
+public:
+	/**
+	 * @return the underlying source asset for this Target.
+	 */
+	virtual UObject* GetSourceData() const override
+	{
+		return GetStaticMesh();
+	}
+
+	/**
+	 * @return the underlying source StaticMesh asset for this Target
+	 */
+	virtual UStaticMesh* GetStaticMesh() const = 0;
+};
\ No newline at end of file
diff --git a/Engine/Source/Runtime/Experimental/InteractiveToolsFramework/Public/ToolTargetManager.h b/Engine/Source/Runtime/Experimental/InteractiveToolsFramework/Public/ToolTargetManager.h
--- a/Engine/Source/Runtime/Experimental/InteractiveToolsFramework/Public/ToolTargetManager.h
+++ b/Engine/Source/Runtime/Experimental/InteractiveToolsFramework/Public/ToolTargetManager.h
@@ -68,6 +68,13 @@ public:
 	virtual UToolTarget* BuildFirstSelectedTargetable(const FToolBuilderState& SceneState,
 		const FToolTargetTypeRequirements& TargetRequirements);
 
+	/**
+	 * Looks through the current selected components and actors and builds all targets that
+	 * satisfy the requirements.
+	 */
+	virtual TArray<TObjectPtr<UToolTarget>> BuildAllSelectedTargetable(const FToolBuilderState& SceneState,
+		const FToolTargetTypeRequirements& TargetRequirements);
+
 protected:
 	friend class UInteractiveToolsContext;		// to call Initialize/Shutdown
 
diff --git a/Engine/Source/Runtime/Experimental/InteractiveToolsFramework/Public/ToolTargets/ToolTarget.h b/Engine/Source/Runtime/Experimental/InteractiveToolsFramework/Public/ToolTargets/ToolTarget.h
--- a/Engine/Source/Runtime/Experimental/InteractiveToolsFramework/Public/ToolTargets/ToolTarget.h
+++ b/Engine/Source/Runtime/Experimental/InteractiveToolsFramework/Public/ToolTargets/ToolTarget.h
@@ -33,27 +33,19 @@ class INTERACTIVETOOLSFRAMEWORK_API FToolTargetTypeRequirements
 {
 public:
 
-	const UClass* BaseType = nullptr; //TODO: Do we even need this? Seems like we shouldn't use it.
 	TArray<const UClass*, TInlineAllocator<2>> Interfaces;
 
 	FToolTargetTypeRequirements()
 	{
 	}
 
-	explicit FToolTargetTypeRequirements(const UClass* BaseTypeIn)
+	explicit FToolTargetTypeRequirements(const UClass* Interface0)
 	{
-		BaseType = BaseTypeIn;
-	}
-
-	explicit FToolTargetTypeRequirements(const UClass* BaseTypeIn, const UClass* Interface0)
-	{
-		BaseType = BaseTypeIn;
 		Interfaces.Add(Interface0);
 	}
 
-	explicit FToolTargetTypeRequirements(const UClass* BaseTypeIn, const TArray<const UClass*>& InterfacesIn)
+	explicit FToolTargetTypeRequirements(const TArray<const UClass*>& InterfacesIn)
 	{
-		BaseType = BaseTypeIn;
 		Interfaces = InterfacesIn;
 	}
 
