Fix ClangEditor build.

diff --git a/Engine/Plugins/Experimental/MeshModelingToolset/Source/MeshModelingTools/Private/AlignObjectsTool.cpp b/Engine/Plugins/Experimental/MeshModelingToolset/Source/MeshModelingTools/Private/AlignObjectsTool.cpp
--- a/Engine/Plugins/Experimental/MeshModelingToolset/Source/MeshModelingTools/Private/AlignObjectsTool.cpp
+++ b/Engine/Plugins/Experimental/MeshModelingToolset/Source/MeshModelingTools/Private/AlignObjectsTool.cpp
@@ -24,9 +24,9 @@
 
 const FToolTargetTypeRequirements& UAlignObjectsToolBuilder::GetTargetRequirements() const
 {
-	static FToolTargetTypeRequirements TypeRequirements({
+	static FToolTargetTypeRequirements TypeRequirements(
 		UPrimitiveComponentBackedTarget::StaticClass()
-		});
+		);
 	return TypeRequirements;
 }
 
