Fix crash when loading virtual texture to use as alpha in sculpt tool.

#rnx

diff --git a/Engine/Plugins/Experimental/MeshModelingToolset/Source/MeshModelingTools/Private/MeshVertexSculptTool.cpp b/Engine/Plugins/Experimental/MeshModelingToolset/Source/MeshModelingTools/Private/MeshVertexSculptTool.cpp
--- a/Engine/Plugins/Experimental/MeshModelingToolset/Source/MeshModelingTools/Private/MeshVertexSculptTool.cpp
+++ b/Engine/Plugins/Experimental/MeshModelingToolset/Source/MeshModelingTools/Private/MeshVertexSculptTool.cpp
@@ -888,7 +888,8 @@ void UMeshVertexSculptTool::UpdateBrushAlpha(UTexture2D* NewAlpha)
 			TImageBuilder<FVector4f> AlphaValues;
 			FImageDimensions AlphaDimensions;
 
-			bool bReadOK = UE::AssetUtils::ReadTexture(this->BrushAlpha, AlphaDimensions, AlphaValues, true);
+			constexpr bool bPreferPlatformData = false;
+			bool bReadOK = UE::AssetUtils::ReadTexture(this->BrushAlpha, AlphaDimensions, AlphaValues, bPreferPlatformData);
 			if (bReadOK)
 			{
 				BrushAlphaValues = MoveTemp(AlphaValues);
diff --git a/Engine/Plugins/Experimental/MeshModelingToolset/Source/ModelingComponents/Private/AssetUtils/Texture2DUtil.cpp b/Engine/Plugins/Experimental/MeshModelingToolset/Source/ModelingComponents/Private/AssetUtils/Texture2DUtil.cpp
--- a/Engine/Plugins/Experimental/MeshModelingToolset/Source/ModelingComponents/Private/AssetUtils/Texture2DUtil.cpp
+++ b/Engine/Plugins/Experimental/MeshModelingToolset/Source/ModelingComponents/Private/AssetUtils/Texture2DUtil.cpp
@@ -139,7 +139,11 @@ bool UE::AssetUtils::ReadTexture(
 	if (ensure(TextureMap) == false) return false;
 
 #if WITH_EDITOR
-	if (TextureMap->Source.IsValid() && bPreferPlatformData == false)
+	bool bHasVTData = TextureMap->GetPlatformData()->VTData != nullptr;
+	bool bHasMips = TextureMap->GetPlatformData()->Mips.Num() != 0;
+	ensure(bHasVTData != bHasMips);		// should be one or the other
+
+	if (TextureMap->Source.IsValid() && (bPreferPlatformData == false || bHasMips == false))
 	{
 		return ReadTexture_SourceData(TextureMap, Dimensions, DestImage);
 	}
