ModelingTools: remove unecessary explicit names passed to NewObject to avoid future problems

diff --git a/Engine/Plugins/Experimental/MeshModelingToolset/Source/MeshModelingTools/Private/AddPatchTool.cpp b/Engine/Plugins/Experimental/MeshModelingToolset/Source/MeshModelingTools/Private/AddPatchTool.cpp
--- a/Engine/Plugins/Experimental/MeshModelingToolset/Source/MeshModelingTools/Private/AddPatchTool.cpp
+++ b/Engine/Plugins/Experimental/MeshModelingToolset/Source/MeshModelingTools/Private/AddPatchTool.cpp
@@ -76,7 +76,7 @@ void UAddPatchTool::Setup()
 	MaterialProperties->RestoreProperties(this);
 
 	// create preview mesh object
-	PreviewMesh = NewObject<UPreviewMesh>(this, TEXT("PreviewMesh"));
+	PreviewMesh = NewObject<UPreviewMesh>(this);
 	PreviewMesh->CreateInWorld(TargetWorld, FTransform::Identity);
 	PreviewMesh->SetVisible(false);
 	PreviewMesh->SetMaterial(MaterialProperties->Material.Get());
diff --git a/Engine/Plugins/Experimental/MeshModelingToolset/Source/MeshModelingTools/Private/AddPrimitiveTool.cpp b/Engine/Plugins/Experimental/MeshModelingToolset/Source/MeshModelingTools/Private/AddPrimitiveTool.cpp
--- a/Engine/Plugins/Experimental/MeshModelingToolset/Source/MeshModelingTools/Private/AddPrimitiveTool.cpp
+++ b/Engine/Plugins/Experimental/MeshModelingToolset/Source/MeshModelingTools/Private/AddPrimitiveTool.cpp
@@ -142,7 +142,7 @@ void UAddPrimitiveTool::Setup()
 	MaterialProperties->RestoreProperties(this);
 
 	// create preview mesh object
-	PreviewMesh = NewObject<UPreviewMesh>(this, TEXT("PreviewMesh"));
+	PreviewMesh = NewObject<UPreviewMesh>(this);
 	PreviewMesh->CreateInWorld(TargetWorld, FTransform::Identity);
 	PreviewMesh->SetVisible(false);
 	PreviewMesh->SetMaterial(MaterialProperties->Material.Get());
