MeshModelingTools:  Fixed snap draggin option not working in edit pivot tool. Fixed undo/redo not working correctly for normal gizmos in "Set Pivot" mode, which caused undo to break in Edit Pivot tool. Added alignment mechanic to Edit Pivot tool.

#rnx

diff --git a/Engine/Plugins/Experimental/MeshModelingToolset/Source/MeshModelingTools/Private/EditPivotTool.cpp b/Engine/Plugins/Experimental/MeshModelingToolset/Source/MeshModelingTools/Private/EditPivotTool.cpp
--- a/Engine/Plugins/Experimental/MeshModelingToolset/Source/MeshModelingTools/Private/EditPivotTool.cpp
+++ b/Engine/Plugins/Experimental/MeshModelingToolset/Source/MeshModelingTools/Private/EditPivotTool.cpp
@@ -7,6 +7,7 @@
 #include "ToolSetupUtil.h"
 #include "DynamicMesh3.h"
 #include "DynamicMeshToMeshDescription.h"
+#include "Mechanics/DragAlignmentMechanic.h"
 #include "MeshDescriptionToDynamicMesh.h"
 #include "MeshAdapterTransforms.h"
 #include "MeshDescriptionAdapter.h"
@@ -104,6 +105,10 @@ void UEditPivotTool::Setup()
 	SetActiveGizmos_Single(false);
 	UpdateSetPivotModes(true);
 
+	DragAlignmentMechanic = NewObject<UDragAlignmentMechanic>(this);
+	DragAlignmentMechanic->Setup(this);
+	DragAlignmentMechanic->AddToGizmo(ActiveGizmos[0].TransformGizmo);
+
 	Precompute();
 
 	FText AllTheWarnings = LOCTEXT("EditPivotWarning", "WARNING: This Tool will Modify the selected StaticMesh Assets! If you do not wish to modify the original Assets, please make copies in the Content Browser first!");
@@ -119,7 +124,7 @@ void UEditPivotTool::Setup()
 
 	SetToolDisplayName(LOCTEXT("ToolName", "Edit Pivot"));
 	GetToolManager()->DisplayMessage(
-		LOCTEXT("OnStartTool", "This Tool edits the Pivot (Origin) of the input objects. Use the Gizmo or enable Snap Dragging and click-drag on the surface to reposition."),
+		LOCTEXT("OnStartTool", "This tool edits the Pivot (Origin) of the input assets. Hold Ctrl while using the gizmo to align to scene. Enable Snap Dragging and click+drag to place gizmo directly into clicked position."),
 		EToolMessageLevel::UserNotification);
 }
 
@@ -127,6 +132,8 @@ void UEditPivotTool::Setup()
 
 void UEditPivotTool::Shutdown(EToolShutdownType ShutdownType)
 {
+	DragAlignmentMechanic->Shutdown();
+
 	FFrame3d CurPivotFrame(ActiveGizmos[0].TransformProxy->GetTransform());
 
 	GizmoManager->DestroyAllGizmosByOwner(this);
@@ -204,6 +211,7 @@ void UEditPivotTool::OnTick(float DeltaTime)
 
 void UEditPivotTool::Render(IToolsContextRenderAPI* RenderAPI)
 {
+	DragAlignmentMechanic->Render(RenderAPI);
 }
 
 void UEditPivotTool::OnPropertyModified(UObject* PropertySet, FProperty* Property)
@@ -318,24 +326,14 @@ FInputRayHit UEditPivotTool::CanBeginClickDragSequence(const FInputDeviceRay& Pr
 		return FInputRayHit();
 	}
 
-	ActiveSnapDragIndex = -1;
-
-	float MinHitDistance = TNumericLimits<float>::Max();
-	FVector HitNormal;
+	FHitResult Result;
+	bool bWorldHit = ToolSceneQueriesUtil::FindNearestVisibleObjectHit(TargetWorld, Result, PressPos.WorldRay);
 
-	for ( int k = 0; k < ComponentTargets.Num(); ++k )
+	if (!bWorldHit)
 	{
-		TUniquePtr<FPrimitiveComponentTarget>& Target = ComponentTargets[k];
-
-		FHitResult WorldHit;
-		if (Target->HitTest(PressPos.WorldRay, WorldHit))
-		{
-			MinHitDistance = FMath::Min(MinHitDistance, WorldHit.Distance);
-			HitNormal = WorldHit.Normal;
-			ActiveSnapDragIndex = k;
-		}
+		return FInputRayHit();
 	}
-	return (MinHitDistance < TNumericLimits<float>::Max()) ? FInputRayHit(MinHitDistance, HitNormal) : FInputRayHit();
+	return FInputRayHit(Result.Distance, Result.ImpactNormal);
 }
 
 void UEditPivotTool::OnClickPress(const FInputDeviceRay& PressPos)
@@ -348,35 +346,11 @@ void UEditPivotTool::OnClickPress(const FInputDeviceRay& PressPos)
 	FEditPivotTarget& ActiveTarget = ActiveGizmos[0];
 	USceneComponent* GizmoComponent = ActiveTarget.TransformGizmo->GetGizmoActor()->GetRootComponent();
 	StartDragTransform = GizmoComponent->GetComponentToWorld();
-
-	//if (TransformProps->SnapDragSource == EEditPivotSnapDragSource::ClickPoint)
-	//{
-		StartDragFrameWorld = FFrame3d(PressPos.WorldRay.PointAt(HitPos.HitDepth), HitPos.HitNormal);
-	//}
-	//else
-	//{
-	//	StartDragFrameWorld = FFrame3d(StartDragTransform);
-	//}
-
 }
 
 
 void UEditPivotTool::OnClickDrag(const FInputDeviceRay& DragPos)
 {
-	//bool bApplyToPivot = true;
-	//if (bApplyToPivot == false)
-	//{
-	//	int IgnoreIndex = -1;
-	//	for (int k = 0; k < ComponentTargets.Num(); ++k)
-	//	{
-	//		if (IgnoreIndex == -1 || k == IgnoreIndex)
-	//		{
-	//			CollisionParams.AddIgnoredComponent(ComponentTargets[k]->GetOwnerComponent());
-	//		}
-	//	}
-	//}
-
-
 	bool bRotate = (TransformProps->RotationMode != EEditPivotSnapDragRotationMode::Ignore);
 	float NormalSign = (TransformProps->RotationMode == EEditPivotSnapDragRotationMode::AlignFlipped) ? -1.0f : 1.0f;
 
@@ -399,10 +373,7 @@ void UEditPivotTool::OnClickDrag(const FInputDeviceRay& DragPos)
 	NewTransform.SetTranslation(HitPos);
 
 	FEditPivotTarget& ActiveTarget = ActiveGizmos[0];
-	USceneComponent* GizmoComponent = ActiveTarget.TransformGizmo->GetGizmoActor()->GetRootComponent();
-	GizmoComponent->SetWorldTransform(NewTransform);
-
-
+	ActiveTarget.TransformGizmo->SetNewGizmoTransform(NewTransform);
 }
 
 
@@ -422,8 +393,6 @@ void UEditPivotTool::OnTerminateDragSequence()
 		LOCTEXT("TransformToolTransformTxnName", "SnapDrag"));
 
 	GetToolManager()->EndUndoTransaction();
-
-	ActiveSnapDragIndex = -1;
 }
 
 
diff --git a/Engine/Plugins/Experimental/MeshModelingToolset/Source/MeshModelingTools/Public/EditPivotTool.h b/Engine/Plugins/Experimental/MeshModelingToolset/Source/MeshModelingTools/Public/EditPivotTool.h
--- a/Engine/Plugins/Experimental/MeshModelingToolset/Source/MeshModelingTools/Public/EditPivotTool.h
+++ b/Engine/Plugins/Experimental/MeshModelingToolset/Source/MeshModelingTools/Public/EditPivotTool.h
@@ -11,6 +11,7 @@
 #include "BoxTypes.h"
 #include "EditPivotTool.generated.h"
 
+class UDragAlignmentMechanic;
 class UBaseAxisTranslationGizmo;
 class UAxisAngleGizmo;
 class UTransformGizmo;
@@ -65,7 +66,7 @@ public:
 
 	/** When Snap-Dragging, align source and target normals */
 	UPROPERTY(EditAnywhere, Category = Options, meta = (EditCondition = "bEnableSnapDragging == true"))
-	EEditPivotSnapDragRotationMode RotationMode = EEditPivotSnapDragRotationMode::AlignFlipped;
+	EEditPivotSnapDragRotationMode RotationMode = EEditPivotSnapDragRotationMode::Align;
 };
 
 
@@ -197,13 +198,14 @@ protected:
 	UPROPERTY()
 	TArray<FEditPivotTarget> ActiveGizmos;
 
+	UPROPERTY()
+	UDragAlignmentMechanic* DragAlignmentMechanic = nullptr;
+
 	void UpdateSetPivotModes(bool bEnableSetPivot);
 	void SetActiveGizmos_Single(bool bLocalRotations);
 	void ResetActiveGizmos();
 
-	FFrame3d StartDragFrameWorld;
 	FTransform StartDragTransform;
-	int ActiveSnapDragIndex = -1;
 
 	EEditPivotToolActions PendingAction;
 	virtual void ApplyAction(EEditPivotToolActions ActionType);
diff --git a/Engine/Source/Runtime/Experimental/InteractiveToolsFramework/Private/BaseGizmos/RepositionableTransformGizmo.cpp b/Engine/Source/Runtime/Experimental/InteractiveToolsFramework/Private/BaseGizmos/RepositionableTransformGizmo.cpp
--- a/Engine/Source/Runtime/Experimental/InteractiveToolsFramework/Private/BaseGizmos/RepositionableTransformGizmo.cpp
+++ b/Engine/Source/Runtime/Experimental/InteractiveToolsFramework/Private/BaseGizmos/RepositionableTransformGizmo.cpp
@@ -47,7 +47,7 @@ void URepositionableTransformGizmo::SetActiveTarget(UTransformProxy* Target, ITo
 
 	// Create the alternate transform source.
 	UGizmoTransformProxyTransformSource* ProxyTransformSource = UGizmoTransformProxyTransformSource::Construct(ActiveTarget, this);
-	ProxyTransformSource->bSetPivotMode = true;
+	ProxyTransformSource->bOverrideSetPivotMode = true;
 	USceneComponent* GizmoComponent = GizmoActor->GetRootComponent();
 	UGizmoScaledAndUnscaledTransformSources* TransformSource = UGizmoScaledAndUnscaledTransformSources::Construct(
 		ProxyTransformSource, GizmoComponent, this);
@@ -58,7 +58,7 @@ void URepositionableTransformGizmo::SetActiveTarget(UTransformProxy* Target, ITo
 		Cast<IToolContextTransactionProvider>(StateTarget->TransactionManager.GetObject()), 
 		this);
 	TUniquePtr<FTransformProxyChangeSource> ProxyChangeSource = MakeUnique<FTransformProxyChangeSource>(Target);
-	ProxyChangeSource->bSetPivotMode = true;
+	ProxyChangeSource->bOverrideSetPivotMode = true;
 	RepositionStateTarget->DependentChangeSources.Add(MoveTemp(ProxyChangeSource));
 
 	// Add on the extra gizmos to existing components. We use the base class addition functions and
diff --git a/Engine/Source/Runtime/Experimental/InteractiveToolsFramework/Private/BaseGizmos/TransformProxy.cpp b/Engine/Source/Runtime/Experimental/InteractiveToolsFramework/Private/BaseGizmos/TransformProxy.cpp
--- a/Engine/Source/Runtime/Experimental/InteractiveToolsFramework/Private/BaseGizmos/TransformProxy.cpp
+++ b/Engine/Source/Runtime/Experimental/InteractiveToolsFramework/Private/BaseGizmos/TransformProxy.cpp
@@ -188,9 +188,9 @@ void FTransformProxyChangeSource::BeginChange()
 	{
 		ActiveChange = MakeUnique<FTransformProxyChange>();
 		ActiveChange->From = Proxy->GetTransform();
-		ActiveChange->bSetPivotMode = bSetPivotMode;
+		ActiveChange->bSetPivotMode = bOverrideSetPivotMode ? true : Proxy->bSetPivotMode;
 
-		if (bSetPivotMode)
+		if (ActiveChange->bSetPivotMode)
 		{
 			Proxy->BeginPivotEditSequence();
 		}
@@ -205,7 +205,7 @@ TUniquePtr<FToolCommandChange> FTransformProxyChangeSource::EndChange()
 {
 	if (Proxy.IsValid())
 	{
-		if (bSetPivotMode)
+		if (ActiveChange->bSetPivotMode)
 		{
 			Proxy->EndPivotEditSequence();
 		}
diff --git a/Engine/Source/Runtime/Experimental/InteractiveToolsFramework/Private/BaseGizmos/TransformSources.cpp b/Engine/Source/Runtime/Experimental/InteractiveToolsFramework/Private/BaseGizmos/TransformSources.cpp
--- a/Engine/Source/Runtime/Experimental/InteractiveToolsFramework/Private/BaseGizmos/TransformSources.cpp
+++ b/Engine/Source/Runtime/Experimental/InteractiveToolsFramework/Private/BaseGizmos/TransformSources.cpp
@@ -57,7 +57,7 @@ FTransform UGizmoTransformProxyTransformSource::GetTransform() const
 
 void UGizmoTransformProxyTransformSource::SetTransform(const FTransform& NewTransform)
 {
-	if (bSetPivotMode)
+	if (bOverrideSetPivotMode)
 	{
 		bool bProxySetPivotOriginal = Proxy->bSetPivotMode;
 		Proxy->bSetPivotMode = true;
@@ -69,7 +69,14 @@ void UGizmoTransformProxyTransformSource::SetTransform(const FTransform& NewTran
 	else
 	{
 		Proxy->SetTransform(NewTransform);
-		OnTransformChanged.Broadcast(this);
+		if (Proxy->bSetPivotMode)
+		{
+			OnPivotChanged.Broadcast(this);
+		}
+		else
+		{
+			OnTransformChanged.Broadcast(this);
+		}
 	}
 }
 
diff --git a/Engine/Source/Runtime/Experimental/InteractiveToolsFramework/Public/BaseGizmos/TransformProxy.h b/Engine/Source/Runtime/Experimental/InteractiveToolsFramework/Public/BaseGizmos/TransformProxy.h
--- a/Engine/Source/Runtime/Experimental/InteractiveToolsFramework/Public/BaseGizmos/TransformProxy.h
+++ b/Engine/Source/Runtime/Experimental/InteractiveToolsFramework/Public/BaseGizmos/TransformProxy.h
@@ -189,5 +189,11 @@ public:
 	virtual UObject* GetChangeTarget() override;
 	virtual FText GetChangeDescription() override;
 
-	bool bSetPivotMode = false;
+	/**
+	 * If true, the emitted changes will always have bSetPivotMode set to true, regardless of
+	 * the current proxy settings. This is meant to accompany a UGizmoTransformProxyTransformSource 
+	 * that has bOverrideSetPivotMode set to true, used for gizmos that reposition a proxy that
+	 * otherwise behaves normally.
+	 */
+	bool bOverrideSetPivotMode = false;
 };
\ No newline at end of file
diff --git a/Engine/Source/Runtime/Experimental/InteractiveToolsFramework/Public/BaseGizmos/TransformSources.h b/Engine/Source/Runtime/Experimental/InteractiveToolsFramework/Public/BaseGizmos/TransformSources.h
--- a/Engine/Source/Runtime/Experimental/InteractiveToolsFramework/Public/BaseGizmos/TransformSources.h
+++ b/Engine/Source/Runtime/Experimental/InteractiveToolsFramework/Public/BaseGizmos/TransformSources.h
@@ -172,10 +172,11 @@ public:
 	TObjectPtr<UTransformProxy> Proxy;
 
 	/**
-	 * The underlying proxy is modified with its SetPivotMode flag set to this. Thus, the transform source
-	 * can be used for either proxy repositioning or actual use.
+	 * If true, the underlying proxy is modified with its SetPivotMode flag temporarily set to true.
+	 * This allows the transform source to be used for proxy repositioning on a proxy that otherwise
+	 * operates normally.
 	 */
-	bool bSetPivotMode = false;
+	bool bOverrideSetPivotMode = false;
 
 	DECLARE_MULTICAST_DELEGATE_OneParam(FOnGizmoTransformSourcePivotChanged, IGizmoTransformSource*);
 	FOnGizmoTransformSourcePivotChanged OnPivotChanged;
