MeshModelingTools: Added to tool descriptions. Fixed snap drag mode in Transform tool. Added missing shutdown in Transform tool.

#rnx

diff --git a/Engine/Plugins/Experimental/MeshModelingToolset/Source/MeshModelingTools/Private/DrawAndRevolveTool.cpp b/Engine/Plugins/Experimental/MeshModelingToolset/Source/MeshModelingTools/Private/DrawAndRevolveTool.cpp
--- a/Engine/Plugins/Experimental/MeshModelingToolset/Source/MeshModelingTools/Private/DrawAndRevolveTool.cpp
+++ b/Engine/Plugins/Experimental/MeshModelingToolset/Source/MeshModelingTools/Private/DrawAndRevolveTool.cpp
@@ -19,7 +19,7 @@
 #define LOCTEXT_NAMESPACE "UDrawAndRevolveTool"
 
 
-const FText InitializationModeMessage = LOCTEXT("CurveInitialization", "Draw and a profile curve and it will be revolved around the purple draw plane axis. Ctrl+click repositions draw plane and axis. The curve is ended by clicking the end again or connecting to its start. Holding shift toggles snapping to be opposite the EnableSnapping setting.");
+const FText InitializationModeMessage = LOCTEXT("CurveInitialization", "Draw a profile curve and revolve it around purple axis. Ctrl+click repositions draw plane, Ctrl+Shift+click repositions without aligning axis. End curve by clicking the end again or connecting to start. Backspace deletes points, shift inverts snapping behavior.");
 const FText EditModeMessage = LOCTEXT("CurveEditing", "Click points to select them, Shift+click to add/remove points to selection. Ctrl+click a segment to add a point, or select an endpoint and Ctrl+click somewhere on the plane to add to the ends. Backspace deletes selected points. Holding Shift toggles snapping to be opposite the EnableSnapping setting.");
 
 
diff --git a/Engine/Plugins/Experimental/MeshModelingToolset/Source/MeshModelingTools/Private/EditMeshPolygonsTool.cpp b/Engine/Plugins/Experimental/MeshModelingToolset/Source/MeshModelingTools/Private/EditMeshPolygonsTool.cpp
--- a/Engine/Plugins/Experimental/MeshModelingToolset/Source/MeshModelingTools/Private/EditMeshPolygonsTool.cpp
+++ b/Engine/Plugins/Experimental/MeshModelingToolset/Source/MeshModelingTools/Private/EditMeshPolygonsTool.cpp
@@ -340,13 +340,13 @@ void UEditMeshPolygonsTool::Setup()
 	{
 		SetToolDisplayName(LOCTEXT("EditMeshTrianglesToolName", "Triangle Edit"));
 		GetToolManager()->DisplayMessage(
-			LOCTEXT("OnStartEditMeshPolygonsTool_TriangleMode", "Select Triangles to edit mesh. Q to toggle Gizmo Orientation Lock."),
+			LOCTEXT("OnStartEditMeshPolygonsTool_TriangleMode", "Select triangles to edit mesh. Use middle mouse on gizmo to reposition it. Hold Ctrl while translating or (in local mode) rotating to align to scene. Shift and Ctrl change marquee select behavior. Q toggles Gizmo Orientation Lock."),
 			EToolMessageLevel::UserNotification);
 	}
 	else
 	{
 		GetToolManager()->DisplayMessage(
-			LOCTEXT("OnStartEditMeshPolygonsTool", "Select PolyGroups to edit mesh. Q to toggle Gizmo Orientation Lock."),
+			LOCTEXT("OnStartEditMeshPolygonsTool", "Select PolyGroups to edit mesh. Use middle mouse on gizmo to reposition it. Hold Ctrl while translating or (in local mode) rotating to align to scene. Shift and Ctrl change marquee select behavior. Q toggles Gizmo Orientation Lock."),
 			EToolMessageLevel::UserNotification);
 	}
 
diff --git a/Engine/Plugins/Experimental/MeshModelingToolset/Source/MeshModelingTools/Private/MeshSpaceDeformerTool.cpp b/Engine/Plugins/Experimental/MeshModelingToolset/Source/MeshModelingTools/Private/MeshSpaceDeformerTool.cpp
--- a/Engine/Plugins/Experimental/MeshModelingToolset/Source/MeshModelingTools/Private/MeshSpaceDeformerTool.cpp
+++ b/Engine/Plugins/Experimental/MeshModelingToolset/Source/MeshModelingTools/Private/MeshSpaceDeformerTool.cpp
@@ -321,7 +321,7 @@ void UMeshSpaceDeformerTool::Setup()
 
 	SetToolDisplayName(LOCTEXT("ToolName", "Space Warp"));
 	GetToolManager()->DisplayMessage(
-		LOCTEXT("MeshSpaceDeformerToolDescription", "Deform the vertices of the selected Mesh using various spatial deformations. Use the in-viewport Gizmo to control the extents/strength of the deformation."),
+		LOCTEXT("MeshSpaceDeformerToolDescription", "Deform the vertices of the selected Mesh using various spatial deformations. Use the in-viewport Gizmo to control the extents/strength of the deformation. Hold Ctrl while translating/rotating gizmo to align to world."),
 		EToolMessageLevel::UserNotification);
 }
 
diff --git a/Engine/Plugins/Experimental/MeshModelingToolset/Source/MeshModelingTools/Private/TransformMeshesTool.cpp b/Engine/Plugins/Experimental/MeshModelingToolset/Source/MeshModelingTools/Private/TransformMeshesTool.cpp
--- a/Engine/Plugins/Experimental/MeshModelingToolset/Source/MeshModelingTools/Private/TransformMeshesTool.cpp
+++ b/Engine/Plugins/Experimental/MeshModelingToolset/Source/MeshModelingTools/Private/TransformMeshesTool.cpp
@@ -88,7 +88,7 @@ void UTransformMeshesTool::Setup()
 
 	SetToolDisplayName(LOCTEXT("ToolName", "Transform"));
 	GetToolManager()->DisplayMessage(
-		LOCTEXT("OnStartTransformMeshesTool", "Translate/Rotate/Scale the selected objects. [A] cycles through Transform modes. Reposition the Gizmo using Set Pivot Mode [S]. Toggle Snap Drag Mode [D] to move the object by click-draging. [W] and [E] cycle through SnapDrag Source and Rotation types."),
+		LOCTEXT("OnStartTransformMeshesTool", "Translate/Rotate/Scale the selected objects. Use middle mouse button on gizmo to reposition it. Hold Ctrl while translating or (in local mode) rotating to align. [A] cycles through Transform modes. [S] toggles Set Pivot Mode. [D] Toggles Snap Drag Mode. [W] and [E] cycle through Snap Drag Source and Rotation types."),
 		EToolMessageLevel::UserNotification);
 }
 
@@ -100,6 +100,8 @@ void UTransformMeshesTool::Setup()
 
 void UTransformMeshesTool::Shutdown(EToolShutdownType ShutdownType)
 {
+	DragAlignmentMechanic->Shutdown();
+
 	GizmoManager->DestroyAllGizmosByOwner(this);
 }
 
@@ -370,8 +372,7 @@ void UTransformMeshesTool::OnClickDrag(const FInputDeviceRay& DragPos)
 		FTransformMeshesTarget& ActiveTarget =
 			(TransformProps->TransformMode == ETransformMeshesTransformMode::PerObjectGizmo) ?
 			ActiveGizmos[ActiveSnapDragIndex] : ActiveGizmos[0];
-		USceneComponent* GizmoComponent = ActiveTarget.TransformGizmo->GetGizmoActor()->GetRootComponent();
-		GizmoComponent->SetWorldTransform(NewTransform);
+		ActiveTarget.TransformGizmo->SetNewGizmoTransform(NewTransform);
 	}
 	else
 	{
@@ -400,8 +401,7 @@ void UTransformMeshesTool::OnClickDrag(const FInputDeviceRay& DragPos)
 		FTransformMeshesTarget& ActiveTarget =
 			(TransformProps->TransformMode == ETransformMeshesTransformMode::PerObjectGizmo) ?
 			ActiveGizmos[ActiveSnapDragIndex] : ActiveGizmos[0];
-		USceneComponent* GizmoComponent = ActiveTarget.TransformGizmo->GetGizmoActor()->GetRootComponent();
-		GizmoComponent->SetWorldTransform(NewTransform);
+		ActiveTarget.TransformGizmo->SetNewGizmoTransform(NewTransform);
 	}
 
 }
