ModelingTools: Update BakeTransformToolto properly update existing Simple Collision geometry. Update CombineMeshesTool to properly transfer/combine/append Simple Collision geometry and settings.

#rnx

diff --git a/Engine/Plugins/Experimental/MeshModelingToolset/Source/MeshModelingTools/Private/BakeTransformTool.cpp b/Engine/Plugins/Experimental/MeshModelingToolset/Source/MeshModelingTools/Private/BakeTransformTool.cpp
--- a/Engine/Plugins/Experimental/MeshModelingToolset/Source/MeshModelingTools/Private/BakeTransformTool.cpp
+++ b/Engine/Plugins/Experimental/MeshModelingToolset/Source/MeshModelingTools/Private/BakeTransformTool.cpp
@@ -17,8 +17,10 @@
 #include "DynamicMeshToMeshDescription.h"
 
 #include "AssetGenerationUtil.h"
+#include "Physics/ComponentCollisionUtil.h"
 
 #include "TargetInterfaces/MeshDescriptionCommitter.h"
+#include "TargetInterfaces/MeshDescriptionProvider.h"
 #include "TargetInterfaces/PrimitiveComponentBackedTarget.h"
 #include "ToolTargetManager.h"
 
@@ -35,6 +37,7 @@ const FToolTargetTypeRequirements& UBakeTransformToolBuilder::GetTargetRequireme
 {
 	static FToolTargetTypeRequirements TypeRequirements({
 		UMeshDescriptionCommitter::StaticClass(),
+		UMeshDescriptionProvider::StaticClass(),
 		UPrimitiveComponentBackedTarget::StaticClass()
 		});
 	return TypeRequirements;
@@ -137,6 +140,14 @@ void UBakeTransformTool::SetAssetAPI(IToolsContextAssetAPI* AssetAPIIn)
 
 void UBakeTransformTool::UpdateAssets()
 {
+	// Make sure mesh descriptions are deserialized before we open transaction.
+	// This is to avoid potential stability issues related to creation/load of
+	// mesh descriptions inside a transaction.
+	for (int32 ComponentIdx = 0; ComponentIdx < Targets.Num(); ComponentIdx++)
+	{
+		TargetMeshProviderInterface(ComponentIdx)->GetMeshDescription();
+	}
+
 	GetToolManager()->BeginUndoTransaction(LOCTEXT("BakeTransformToolTransactionName", "Bake Transforms"));
 
 	TArray<FTransform3d> BakedTransforms;
@@ -144,6 +155,8 @@ void UBakeTransformTool::UpdateAssets()
 	{
 		IPrimitiveComponentBackedTarget* TargetComponent = TargetComponentInterface(ComponentIdx);
 		IMeshDescriptionCommitter* TargetMeshCommitter = TargetMeshCommitterInterface(ComponentIdx);
+		UPrimitiveComponent* Component = TargetComponent->GetOwnerComponent();
+		Component->Modify();
 
 		FTransform3d ComponentToWorld(TargetComponent->GetWorldTransform());
 		FTransform3d ToBakePart = FTransform3d::Identity();
@@ -235,11 +248,12 @@ void UBakeTransformTool::UpdateAssets()
 				}
 			});
 
+			// try to transform simple collision
+			UE::Geometry::TransformSimpleCollision(Component, ToBakePart);
+
 			BakedTransforms.Add(ToBakePart);
 		}
 
-		UPrimitiveComponent* Component = TargetComponent->GetOwnerComponent();
-		Component->Modify();
 		Component->SetWorldTransform((FTransform)NewWorldPart);
 		TargetComponent->GetOwnerActor()->MarkComponentsRenderStateDirty();
 	}
diff --git a/Engine/Plugins/Experimental/MeshModelingToolset/Source/MeshModelingTools/Private/CombineMeshesTool.cpp b/Engine/Plugins/Experimental/MeshModelingToolset/Source/MeshModelingTools/Private/CombineMeshesTool.cpp
--- a/Engine/Plugins/Experimental/MeshModelingToolset/Source/MeshModelingTools/Private/CombineMeshesTool.cpp
+++ b/Engine/Plugins/Experimental/MeshModelingToolset/Source/MeshModelingTools/Private/CombineMeshesTool.cpp
@@ -17,6 +17,8 @@
 
 #include "AssetGenerationUtil.h"
 #include "Selection/ToolSelectionUtil.h"
+#include "Physics/ComponentCollisionUtil.h"
+#include "ShapeApproximation/SimpleShapeSet3.h"
 
 #include "Components/StaticMeshComponent.h"
 #include "Engine/StaticMesh.h"
@@ -170,6 +172,15 @@ void UCombineMeshesTool::SetAssetAPI(IAssetGenerationAPI* AssetAPIIn)
 
 void UCombineMeshesTool::CreateNewAsset()
 {
+	// Make sure mesh descriptions are deserialized before we open transaction.
+	// This is to avoid potential stability issues related to creation/load of
+	// mesh descriptions inside a transaction.
+	TArray<const FMeshDescription*> MeshDescriptions;
+	for (int32 ComponentIdx = 0; ComponentIdx < Targets.Num(); ComponentIdx++)
+	{
+		MeshDescriptions.Add(TargetMeshProviderInterface(ComponentIdx)->GetMeshDescription());
+	}
+
 	GetToolManager()->BeginUndoTransaction(
 		bDuplicateMode ? 
 		LOCTEXT("DuplicateMeshToolTransactionName", "Duplicate Mesh") :
@@ -220,6 +231,9 @@ void UCombineMeshesTool::CreateNewAsset()
 	FTransform AccumToWorld(Box.GetCenter());
 	FTransform ToAccum(-Box.GetCenter());
 
+	FSimpleShapeSet3d SimpleCollision;
+	UE::Geometry::FComponentCollisionSettings CollisionSettings;
+
 	{
 		FScopedSlowTask SlowTask(Targets.Num()+1, 
 			bDuplicateMode ? 
@@ -231,11 +245,13 @@ void UCombineMeshesTool::CreateNewAsset()
 		for (int32 ComponentIdx = 0; ComponentIdx < Targets.Num(); ComponentIdx++)
 		{
 			SlowTask.EnterProgressFrame(1);
+			IPrimitiveComponentBackedTarget* TargetComponent = TargetComponentInterface(ComponentIdx);
 			IMeshDescriptionProvider* TargetMeshProvider = TargetMeshProviderInterface(ComponentIdx);
 
 			FMeshDescriptionToDynamicMesh Converter;
 			FDynamicMesh3 ComponentDMesh;
-			Converter.Convert(TargetMeshProvider->GetMeshDescription(), ComponentDMesh);
+			const FMeshDescription* MeshDescription = MeshDescriptions[ComponentIdx];
+			Converter.Convert(MeshDescription, ComponentDMesh);
 
 			FTransform3d XF = (FTransform3d)(TargetComponentInterface(ComponentIdx)->GetWorldTransform() * ToAccum);
 			if (XF.GetDeterminant() < 0)
@@ -255,14 +271,16 @@ void UCombineMeshesTool::CreateNewAsset()
 			if (bDuplicateMode) // no transform if duplicating
 			{
 				Editor.AppendMesh(&ComponentDMesh, IndexMapping);
+				CollisionSettings = UE::Geometry::GetCollisionSettings(TargetComponent->GetOwnerComponent());
+				UE::Geometry::AppendSimpleCollision(TargetComponent->GetOwnerComponent(), &SimpleCollision, FTransform3d::Identity());
 			}
 			else
 			{
 				Editor.AppendMesh(&ComponentDMesh, IndexMapping,
 					[&XF](int Unused, const FVector3d P) { return XF.TransformPosition(P); },
 					[&XF](int Unused, const FVector3d N) { return XF.TransformNormal(N); });
+				UE::Geometry::AppendSimpleCollision(TargetComponent->GetOwnerComponent(), &SimpleCollision, XF);
 			}
-			
 
 			MatIndexBase += TargetMaterialInterface(ComponentIdx)->GetNumMaterials();
 		}
@@ -296,6 +314,12 @@ void UCombineMeshesTool::CreateNewAsset()
 				NewMesh->SetMaterial(MatIdx, AllMaterials[MatIdx]);
 			}
 
+			// if any inputs have Simple Collision geometry we will forward it to new mesh.
+			if (SimpleCollision.TotalElementsNum() > 0)
+			{
+				UE::Geometry::SetSimpleCollision(NewMeshComponent, &SimpleCollision, CollisionSettings);
+			}
+
 			// select the new actor
 			ToolSelectionUtil::SetNewActorSelection(GetToolManager(), NewActor);
 		}
@@ -325,6 +349,12 @@ void UCombineMeshesTool::CreateNewAsset()
 
 void UCombineMeshesTool::UpdateExistingAsset()
 {
+	TArray<const FMeshDescription*> MeshDescriptions;
+	for (int32 ComponentIdx = 0; ComponentIdx < ComponentTargets.Num(); ComponentIdx++)
+	{
+		MeshDescriptions.Add(TargetMeshProviderInterface(ComponentIdx)->GetMeshDescription());
+	}
+
 	check(!bDuplicateMode);
 	GetToolManager()->BeginUndoTransaction(LOCTEXT("CombineMeshesToolTransactionName", "Combine Meshes"));
 
@@ -377,6 +407,11 @@ void UCombineMeshesTool::UpdateExistingAsset()
 	FTransform3d TargetToWorld = (FTransform3d)UpdateTarget->GetWorldTransform();
 	FTransform3d WorldToTarget = TargetToWorld.Inverse();
 
+	FSimpleShapeSet3d SimpleCollision;
+	UE::Geometry::FComponentCollisionSettings CollisionSettings = UE::Geometry::GetCollisionSettings(UpdateTarget->GetOwnerComponent());
+	TArray<FTransform3d> Transforms;
+	Transforms.SetNum(2);
+
 	{
 		FScopedSlowTask SlowTask(Targets.Num()+1, 
 			bDuplicateMode ? 
@@ -389,9 +424,11 @@ void UCombineMeshesTool::UpdateExistingAsset()
 		{
 			SlowTask.EnterProgressFrame(1);
 
+			IPrimitiveComponentBackedTarget* TargetComponent = TargetComponentInterface(ComponentIdx);
+
 			FMeshDescriptionToDynamicMesh Converter;
 			FDynamicMesh3 ComponentDMesh;
-			Converter.Convert(TargetMeshProviderInterface(ComponentIdx)->GetMeshDescription(), ComponentDMesh);
+			Converter.Convert(MeshDescriptions[ComponentIdx], ComponentDMesh);
 
 			// update material IDs to account for combined material set
 			FDynamicMeshMaterialAttribute* MatAttrib = ComponentDMesh.Attributes()->GetMaterialID();
@@ -404,7 +441,7 @@ void UCombineMeshesTool::UpdateExistingAsset()
 
 			if (ComponentIdx != SkipIndex)
 			{
-				FTransform3d ComponentToWorld = (FTransform3d)TargetComponentInterface(ComponentIdx)->GetWorldTransform();
+				FTransform3d ComponentToWorld = (FTransform3d)TargetComponent->GetWorldTransform();
 				MeshTransforms::ApplyTransform(ComponentDMesh, ComponentToWorld);
 				if (ComponentToWorld.GetDeterminant() < 0)
 				{
@@ -415,6 +452,13 @@ void UCombineMeshesTool::UpdateExistingAsset()
 				{
 					ComponentDMesh.ReverseOrientation(true);
 				}
+				Transforms[0] = ComponentToWorld;
+				Transforms[1] = WorldToTarget;
+				UE::Geometry::AppendSimpleCollision(TargetComponent->GetOwnerComponent(), &SimpleCollision, Transforms);
+			}
+			else
+			{
+				UE::Geometry::AppendSimpleCollision(TargetComponent->GetOwnerComponent(), &SimpleCollision, FTransform3d::Identity());
 			}
 
 			FDynamicMeshEditor Editor(&AccumulateDMesh);
@@ -430,6 +474,8 @@ void UCombineMeshesTool::UpdateExistingAsset()
 			Converter.Convert(&AccumulateDMesh, *CommitParams.MeshDescriptionOut);
 		});
 
+		UE::Geometry::SetSimpleCollision(UpdateTarget->GetOwnerComponent(), &SimpleCollision, CollisionSettings);
+
 		FComponentMaterialSet MaterialSet;
 		MaterialSet.Materials = AllMaterials;
 		UpdateTargetMaterial->CommitMaterialSetUpdate(MaterialSet, true);
