ModelingTools: Merge Rounded Rect into Rect tool. Merge Punctured Disc tool into Disc tool. Merge Spherical Box tool into Sphere tool.

#rnx

diff --git a/Engine/Plugins/Experimental/MeshModelingToolset/Source/MeshModelingTools/Private/AddPrimitiveTool.cpp b/Engine/Plugins/Experimental/MeshModelingToolset/Source/MeshModelingTools/Private/AddPrimitiveTool.cpp
--- a/Engine/Plugins/Experimental/MeshModelingToolset/Source/MeshModelingTools/Private/AddPrimitiveTool.cpp
+++ b/Engine/Plugins/Experimental/MeshModelingToolset/Source/MeshModelingTools/Private/AddPrimitiveTool.cpp
@@ -59,21 +59,12 @@ UInteractiveTool* UAddPrimitiveToolBuilder::BuildTool(const FToolBuilderState& S
 	case EMakeMeshShapeType::Rectangle:
 		NewTool = NewObject<UAddRectanglePrimitiveTool>(SceneState.ToolManager);
 		break;
-	case EMakeMeshShapeType::RoundedRectangle:
-		NewTool = NewObject<UAddRoundedRectanglePrimitiveTool>(SceneState.ToolManager);
-		break;
 	case EMakeMeshShapeType::Disc:
 		NewTool = NewObject<UAddDiscPrimitiveTool>(SceneState.ToolManager);
 		break;
-	case EMakeMeshShapeType::PuncturedDisc:
-		NewTool = NewObject<UAddPuncturedDiscPrimitiveTool>(SceneState.ToolManager);
-		break;
 	case EMakeMeshShapeType::Torus:
 		NewTool = NewObject<UAddTorusPrimitiveTool>(SceneState.ToolManager);
 		break;
-	case EMakeMeshShapeType::SphericalBox:
-		NewTool = NewObject<UAddSphericalBoxPrimitiveTool>(SceneState.ToolManager);
-		break;
 	case EMakeMeshShapeType::Sphere:
 		NewTool = NewObject<UAddSpherePrimitiveTool>(SceneState.ToolManager);
 		break;
@@ -407,48 +398,39 @@ UAddRectanglePrimitiveTool::UAddRectanglePrimitiveTool(const FObjectInitializer&
 
 void UAddRectanglePrimitiveTool::GenerateMesh(FDynamicMesh3* OutMesh) const
 {
-	FRectangleMeshGenerator RectGen;
 	auto* RectangleSettings = Cast<UProceduralRectangleToolProperties>(ShapeSettings);
-	RectGen.Width = RectangleSettings->Depth;
-	RectGen.Height = RectangleSettings->Width;
-	RectGen.WidthVertexCount = RectangleSettings->DepthSubdivisions + 1;
-	RectGen.HeightVertexCount = RectangleSettings->WidthSubdivisions + 1;
-	if (ShapeSettings->PolygroupMode != EMakeMeshPolygroupMode::PerQuad)
-	{
-		RectGen.bSinglePolygroup = true;
+	switch (RectangleSettings->RectType)
+	{
+	case EProceduralRectType::Rectangle:
+	{
+		FRectangleMeshGenerator RectGen;
+		RectGen.Width = RectangleSettings->Depth;
+		RectGen.Height = RectangleSettings->Width;
+		RectGen.WidthVertexCount = RectangleSettings->DepthSubdivisions + 1;
+		RectGen.HeightVertexCount = RectangleSettings->WidthSubdivisions + 1;
+		RectGen.bSinglePolygroup = (ShapeSettings->PolygroupMode != EMakeMeshPolygroupMode::PerQuad);
+		RectGen.Generate();
+		OutMesh->Copy(&RectGen);
+		break;
+	}
+	case EProceduralRectType::RoundedRectangle:
+	{
+		FRoundedRectangleMeshGenerator RectGen;
+		RectGen.Width = RectangleSettings->Depth;
+		RectGen.Height = RectangleSettings->Width;
+		RectGen.WidthVertexCount = RectangleSettings->DepthSubdivisions + 1;
+		RectGen.HeightVertexCount = RectangleSettings->WidthSubdivisions + 1;
+		RectGen.bSinglePolygroup = (ShapeSettings->PolygroupMode != EMakeMeshPolygroupMode::PerQuad);
+		RectGen.Radius = RectangleSettings->CornerRadius;
+		RectGen.AngleSamples = RectangleSettings->CornerSlices - 1;
+		RectGen.Generate();
+		OutMesh->Copy(&RectGen);
+		break;
 	}
-	RectGen.Generate();
-	OutMesh->Copy(&RectGen);
-}
-
-
-UAddRoundedRectanglePrimitiveTool::UAddRoundedRectanglePrimitiveTool(const FObjectInitializer& ObjectInitializer)
-	: Super(ObjectInitializer.SetDefaultSubobjectClass<UProceduralRoundedRectangleToolProperties>(TEXT("ShapeSettings")))
-{
-	AssetName = TEXT("RoundedRectangle");
-	SetToolDisplayName(LOCTEXT("RoudRectToolName", "Create RoundRects"));
-}
-
-void UAddRoundedRectanglePrimitiveTool::GenerateMesh(FDynamicMesh3* OutMesh) const
-{
-	FRoundedRectangleMeshGenerator RectGen;
-	auto* RectangleSettings = Cast<UProceduralRoundedRectangleToolProperties>(ShapeSettings);
-	RectGen.Width = RectangleSettings->Depth;
-	RectGen.Height = RectangleSettings->Width;
-	RectGen.WidthVertexCount = RectangleSettings->DepthSubdivisions + 1;
-	RectGen.HeightVertexCount = RectangleSettings->WidthSubdivisions + 1;
-	if (ShapeSettings->PolygroupMode != EMakeMeshPolygroupMode::PerQuad)
-	{
-		RectGen.bSinglePolygroup = true;
 	}
-	RectGen.Radius = RectangleSettings->CornerRadius;
-	RectGen.AngleSamples = RectangleSettings->CornerSlices - 1;
-	RectGen.Generate();
-	OutMesh->Copy(&RectGen);
 }
 
 
-
 UAddDiscPrimitiveTool::UAddDiscPrimitiveTool(const FObjectInitializer& ObjectInitializer)
 	: Super(ObjectInitializer.SetDefaultSubobjectClass<UProceduralDiscToolProperties>(TEXT("ShapeSettings")))
 {
@@ -458,42 +440,33 @@ UAddDiscPrimitiveTool::UAddDiscPrimitiveTool(const FObjectInitializer& ObjectIni
 
 void UAddDiscPrimitiveTool::GenerateMesh(FDynamicMesh3* OutMesh) const
 {
-	FDiscMeshGenerator Gen;
 	auto* DiscSettings = Cast<UProceduralDiscToolProperties>(ShapeSettings);
-	Gen.Radius = DiscSettings->Radius;
-	Gen.AngleSamples = DiscSettings->RadialSlices;
-	Gen.RadialSamples = DiscSettings->RadialSubdivisions;
-	if (ShapeSettings->PolygroupMode != EMakeMeshPolygroupMode::PerQuad)
+	switch (DiscSettings->DiscType)
+	{
+	case EProceduralDiscType::Disc:
 	{
-		Gen.bSinglePolygroup = true;
+		FDiscMeshGenerator Gen;
+		Gen.Radius = DiscSettings->Radius;
+		Gen.AngleSamples = DiscSettings->RadialSlices;
+		Gen.RadialSamples = DiscSettings->RadialSubdivisions;
+		Gen.bSinglePolygroup = (ShapeSettings->PolygroupMode != EMakeMeshPolygroupMode::PerQuad);
+		Gen.Generate();
+		OutMesh->Copy(&Gen);
+		break;
 	}
-	Gen.Generate();
-	OutMesh->Copy(&Gen);
-}
-
-
-
-UAddPuncturedDiscPrimitiveTool::UAddPuncturedDiscPrimitiveTool(const FObjectInitializer& ObjectInitializer)
-	: Super(ObjectInitializer.SetDefaultSubobjectClass<UProceduralPuncturedDiscToolProperties>(TEXT("ShapeSettings")))
-{
-	AssetName = TEXT("PuncturedDisc");
-	SetToolDisplayName(LOCTEXT("CircleToolName", "Create Circles"));
-}
-
-void UAddPuncturedDiscPrimitiveTool::GenerateMesh(FDynamicMesh3* OutMesh) const
-{
-	FPuncturedDiscMeshGenerator Gen;
-	auto* DiscSettings = Cast<UProceduralPuncturedDiscToolProperties>(ShapeSettings);
-	Gen.Radius = DiscSettings->Radius;
-	Gen.HoleRadius = FMath::Min(DiscSettings->HoleRadius, Gen.Radius * .999f); // hole cannot be bigger than outer radius
-	Gen.AngleSamples = DiscSettings->RadialSlices;
-	Gen.RadialSamples = DiscSettings->RadialSubdivisions;
-	if (ShapeSettings->PolygroupMode != EMakeMeshPolygroupMode::PerQuad)
+	case EProceduralDiscType::PuncturedDisc:
 	{
-		Gen.bSinglePolygroup = true;
+		FPuncturedDiscMeshGenerator Gen;
+		Gen.Radius = DiscSettings->Radius;
+		Gen.HoleRadius = FMath::Min(DiscSettings->HoleRadius, Gen.Radius * .999f); // hole cannot be bigger than outer radius
+		Gen.AngleSamples = DiscSettings->RadialSlices;
+		Gen.RadialSamples = DiscSettings->RadialSubdivisions;
+		Gen.bSinglePolygroup = (ShapeSettings->PolygroupMode != EMakeMeshPolygroupMode::PerQuad);
+		Gen.Generate();
+		OutMesh->Copy(&Gen);
+		break;
+	}
 	}
-	Gen.Generate();
-	OutMesh->Copy(&Gen);
 }
 
 
@@ -618,47 +591,40 @@ UAddSpherePrimitiveTool::UAddSpherePrimitiveTool(const FObjectInitializer& Objec
 
 void UAddSpherePrimitiveTool::GenerateMesh(FDynamicMesh3* OutMesh) const
 {
-	FSphereGenerator SphereGen;
 	auto* SphereSettings = Cast<UProceduralSphereToolProperties>(ShapeSettings);
-	SphereGen.Radius = SphereSettings->Radius;
-	SphereGen.NumTheta = SphereSettings->LongitudeSlices + 1;
-	SphereGen.NumPhi = SphereSettings->LatitudeSlices + 1;
-	if (ShapeSettings->PolygroupMode == EMakeMeshPolygroupMode::PerQuad)
+	switch (SphereSettings->SphereType)
 	{
-		SphereGen.bPolygroupPerQuad = true;
+	case EProceduralSphereType::LatLong:
+	{
+		FSphereGenerator SphereGen;
+		SphereGen.Radius = SphereSettings->Radius;
+		SphereGen.NumTheta = SphereSettings->LongitudeSlices + 1;
+		SphereGen.NumPhi = SphereSettings->LatitudeSlices + 1;
+		SphereGen.bPolygroupPerQuad = (ShapeSettings->PolygroupMode == EMakeMeshPolygroupMode::PerQuad);
+		SphereGen.Generate();
+		OutMesh->Copy(&SphereGen);
+		break;
+	}
+	case EProceduralSphereType::Box:
+	{
+		FBoxSphereGenerator SphereGen;
+		SphereGen.Radius = SphereSettings->Radius;
+		SphereGen.Box = FOrientedBox3d(FVector3d::Zero(),
+			0.5 * FVector3d(SphereSettings->Subdivisions + 1,
+				SphereSettings->Subdivisions + 1,
+				SphereSettings->Subdivisions + 1));
+		int EdgeNum = SphereSettings->Subdivisions + 1;
+		SphereGen.EdgeVertices = FIndex3i(EdgeNum, EdgeNum, EdgeNum);
+		SphereGen.bPolygroupPerQuad = (ShapeSettings->PolygroupMode == EMakeMeshPolygroupMode::PerQuad);
+		SphereGen.Generate();
+		OutMesh->Copy(&SphereGen);
+		break;
+	}
 	}
-	SphereGen.Generate();
-	OutMesh->Copy(&SphereGen);
 }
 
 
 
-UAddSphericalBoxPrimitiveTool::UAddSphericalBoxPrimitiveTool(const FObjectInitializer& ObjectInitializer)
-	: Super(ObjectInitializer.SetDefaultSubobjectClass<UProceduralSphericalBoxToolProperties>(TEXT("ShapeSettings")))
-{
-	AssetName = TEXT("BSphere");
-	SetToolDisplayName(LOCTEXT("BoxSphereToolName", "Create BoxSpheres"));
-}
-
-void UAddSphericalBoxPrimitiveTool::GenerateMesh(FDynamicMesh3* OutMesh) const
-{
-	FBoxSphereGenerator SphereGen;
-	auto* SphereSettings = Cast<UProceduralSphericalBoxToolProperties>(ShapeSettings);
-	SphereGen.Radius = SphereSettings->Radius;
-	SphereGen.Box = FOrientedBox3d(FVector3d::Zero(),
-								   0.5*FVector3d(SphereSettings->Subdivisions + 1,
-												 SphereSettings->Subdivisions + 1,
-												 SphereSettings->Subdivisions + 1));
-	int EdgeNum = SphereSettings->Subdivisions + 1;
-	SphereGen.EdgeVertices = FIndex3i(EdgeNum, EdgeNum, EdgeNum);
-	if (ShapeSettings->PolygroupMode == EMakeMeshPolygroupMode::PerQuad)
-	{
-		SphereGen.bPolygroupPerQuad = true;
-	}
-	SphereGen.Generate();
-	OutMesh->Copy(&SphereGen);
-}
-
 UAddStairsPrimitiveTool::UAddStairsPrimitiveTool(const FObjectInitializer& ObjectInitializer)
 	: Super(ObjectInitializer.SetDefaultSubobjectClass<UProceduralStairsToolProperties>(TEXT("ShapeSettings")))
 {
diff --git a/Engine/Plugins/Experimental/MeshModelingToolset/Source/MeshModelingTools/Public/AddPrimitiveTool.h b/Engine/Plugins/Experimental/MeshModelingToolset/Source/MeshModelingTools/Public/AddPrimitiveTool.h
--- a/Engine/Plugins/Experimental/MeshModelingToolset/Source/MeshModelingTools/Public/AddPrimitiveTool.h
+++ b/Engine/Plugins/Experimental/MeshModelingToolset/Source/MeshModelingTools/Public/AddPrimitiveTool.h
@@ -32,11 +32,8 @@ public:
 		Cone,
 		Arrow,
 		Rectangle,
-		RoundedRectangle,
 		Disc,
-		PuncturedDisc,
 		Torus,
-		SphericalBox,
 		Sphere,
 		Stairs
 	};
@@ -130,11 +127,11 @@ public:
 };
 
 UCLASS()
-class MESHMODELINGTOOLS_API UProceduralRectangleToolProperties : public UProceduralShapeToolProperties
+class MESHMODELINGTOOLS_API UProceduralBoxToolProperties : public UProceduralShapeToolProperties
 {
 	GENERATED_BODY()
 public:
-	UProceduralRectangleToolProperties() = default;
+	UProceduralBoxToolProperties() = default;
 
 	/** Width of Shape */
 	UPROPERTY(EditAnywhere, Category = "ShapeSettings|Dimensions", meta = (DisplayName = "Width", UIMin = "1.0", UIMax = "1000.0", ClampMin = "0.0001", ClampMax = "1000000.0", ProceduralShapeSetting))
@@ -144,50 +141,86 @@ public:
 	UPROPERTY(EditAnywhere, Category = "ShapeSettings|Dimensions", meta = (DisplayName = "Depth", UIMin = "1.0", UIMax = "1000.0", ClampMin = "0.0001", ClampMax = "1000000.0", ProceduralShapeSetting))
 	float Depth = 100.f;
 
+	/** Height of Shape */
+	UPROPERTY(EditAnywhere, Category = "ShapeSettings|Dimensions", meta = (DisplayName = "Height", UIMin = "1.0", UIMax = "1000.0", ClampMin = "0.0001", ClampMax = "1000000.0", ProceduralShapeSetting))
+	float Height = 100.f;
+
 	/** Number of Subdivisions Along the Width */
 	UPROPERTY(EditAnywhere, Category = "ShapeSettings|Subdivisions", meta = (DisplayName = "Width", UIMin = "1", UIMax = "100", ClampMin = "1", ClampMax = "500", ProceduralShapeSetting))
 	int WidthSubdivisions = 1;
 
-	/** Number of Subdivisions Along the Depth*/
+	/** Number of Subdivisions Along the Depth */
 	UPROPERTY(EditAnywhere, Category = "ShapeSettings|Subdivisions", meta = (DisplayName = "Depth", UIMin = "1", UIMax = "100", ClampMin = "1", ClampMax = "500", ProceduralShapeSetting))
 	int DepthSubdivisions = 1;
-};
-
-UCLASS()
-class MESHMODELINGTOOLS_API UProceduralBoxToolProperties : public UProceduralRectangleToolProperties
-{
-	GENERATED_BODY()
-public:
-	UProceduralBoxToolProperties() = default;
-
-	/** Height of Shape */
-	UPROPERTY(EditAnywhere, Category = "ShapeSettings|Dimensions", meta = (DisplayName = "Height", UIMin = "1.0", UIMax = "1000.0", ClampMin = "0.0001", ClampMax = "1000000.0", ProceduralShapeSetting))
-	float Height = 100.f;
 
 	/** Number of Subdivisions Along the Height */
 	UPROPERTY(EditAnywhere, Category = "ShapeSettings|Subdivisions", meta = (DisplayName = "Height", UIMin = "0", UIMax = "100", ClampMin = "0", ClampMax = "500", ProceduralShapeSetting))
 	int HeightSubdivisions = 1;
 };
 
+UENUM()
+enum class EProceduralRectType
+{
+	/** Create a Rectangle */
+	Rectangle UMETA(DisplayName = "Rectangle"),
+	/** Create a Rounded Rectangle */
+	RoundedRectangle UMETA(DisplayName = "Rounded Rectangle")
+};
+
 UCLASS()
-class MESHMODELINGTOOLS_API UProceduralRoundedRectangleToolProperties : public UProceduralRectangleToolProperties
+class MESHMODELINGTOOLS_API UProceduralRectangleToolProperties : public UProceduralShapeToolProperties
 {
 	GENERATED_BODY()
 public:
+	UProceduralRectangleToolProperties() = default;
+
+	/** Type of rectangle to create */
+	UPROPERTY(EditAnywhere, Category = "ShapeSettings")
+	EProceduralRectType RectType = EProceduralRectType::Rectangle;
+
+	/** Width of Shape */
+	UPROPERTY(EditAnywhere, Category = "ShapeSettings|Dimensions", meta = (DisplayName = "Width", UIMin = "1.0", UIMax = "1000.0", ClampMin = "0.0001", ClampMax = "1000000.0", ProceduralShapeSetting))
+	float Width = 100.f;
+
+	/** Depth of Shape */
+	UPROPERTY(EditAnywhere, Category = "ShapeSettings|Dimensions", meta = (DisplayName = "Depth", UIMin = "1.0", UIMax = "1000.0", ClampMin = "0.0001", ClampMax = "1000000.0", ProceduralShapeSetting))
+	float Depth = 100.f;
+
+	/** Number of Subdivisions Along the Width */
+	UPROPERTY(EditAnywhere, Category = "ShapeSettings|Subdivisions", meta = (DisplayName = "Width", UIMin = "1", UIMax = "100", ClampMin = "1", ClampMax = "500", ProceduralShapeSetting))
+	int WidthSubdivisions = 1;
+
+	/** Number of Subdivisions Along the Depth */
+	UPROPERTY(EditAnywhere, Category = "ShapeSettings|Subdivisions", meta = (DisplayName = "Depth", UIMin = "1", UIMax = "100", ClampMin = "1", ClampMax = "500", ProceduralShapeSetting))
+	int DepthSubdivisions = 1;
+
 	/** Radius of Rounded Corners */
-	UPROPERTY(EditAnywhere, Category = "ShapeSettings|Dimensions", meta = (DisplayName = "Corner Radius", UIMin = "1.0", UIMax = "1000.0", ClampMin = "0.0001", ClampMax = "1000000.0", ProceduralShapeSetting))
+	UPROPERTY(EditAnywhere, Category = "ShapeSettings|Dimensions", meta = (DisplayName = "Corner Radius", UIMin = "1.0", UIMax = "1000.0", ClampMin = "0.0001", ClampMax = "1000000.0", ProceduralShapeSetting, EditCondition = "RectType == EProceduralRectType::RoundedRectangle", EditConditionHides))
 	float CornerRadius = 25.f;
 
-	/** Number of Angular Slices in Each Corner*/
-	UPROPERTY(EditAnywhere, Category = "ShapeSettings|Slices", meta = (DisplayName = "Corner", UIMin = "3", UIMax = "128", ClampMin = "3", ClampMax = "500", ProceduralShapeSetting))
+	/** Number of Angular Slices in Each Rounded Corner */
+	UPROPERTY(EditAnywhere, Category = "ShapeSettings|Subdivisions", meta = (DisplayName = "Corner", UIMin = "3", UIMax = "128", ClampMin = "3", ClampMax = "500", ProceduralShapeSetting, EditCondition = "RectType == EProceduralRectType::RoundedRectangle", EditConditionHides))
 	int CornerSlices = 16;
 };
 
+UENUM()
+enum class EProceduralDiscType
+{
+	/** Create a solid Disc */
+	Disc UMETA(DisplayName = "Disc"),
+	/** Create a Disc with a hole */
+	PuncturedDisc UMETA(DisplayName = "Punctured Disc")
+};
+
 UCLASS()
 class MESHMODELINGTOOLS_API UProceduralDiscToolProperties : public UProceduralShapeToolProperties
 {
 	GENERATED_BODY()
 public:
+	/** Type of disc to create */
+	UPROPERTY(EditAnywhere, Category = "ShapeSettings")
+	EProceduralDiscType DiscType = EProceduralDiscType::Disc;
+
 	/** Radius of Disc */
 	UPROPERTY(EditAnywhere, Category = "ShapeSettings|Dimensions", meta = (DisplayName = "Radius", UIMin = "1.0", UIMax = "1000.0", ClampMin = "0.0001", ClampMax = "1000000.0", ProceduralShapeSetting))
 	float Radius = 50.f;
@@ -199,15 +232,9 @@ public:
 	/** Number of Radial Subdivisions in the Disc */
 	UPROPERTY(EditAnywhere, Category = "ShapeSettings|Subdivisions", meta = (DisplayName = "Radial", UIMin = "1", UIMax = "100", ClampMin = "1", ClampMax = "500", ProceduralShapeSetting))
 	int RadialSubdivisions = 1;
-};
 
-UCLASS()
-class MESHMODELINGTOOLS_API UProceduralPuncturedDiscToolProperties : public UProceduralDiscToolProperties
-{
-	GENERATED_BODY()
-public:
 	/** Radius of the Disc's Hole */
-	UPROPERTY(EditAnywhere, Category = "ShapeSettings|Dimensions", meta = (DisplayName = "Hole Radius", UIMin = "1.0", UIMax = "1000.0", ClampMin = "0.0001", ClampMax = "1000000.0", ProceduralShapeSetting))
+	UPROPERTY(EditAnywhere, Category = "ShapeSettings|Dimensions", meta = (DisplayName = "Hole Radius", UIMin = "1.0", UIMax = "1000.0", ClampMin = "0.0001", ClampMax = "1000000.0", ProceduralShapeSetting, EditCondition = "DiscType == EProceduralDiscType::PuncturedDisc", EditConditionHides))
 	float HoleRadius = 25.f;
 };
 
@@ -251,7 +278,7 @@ public:
 	UPROPERTY(EditAnywhere, Category = "ShapeSettings|Slices", meta = (DisplayName = "Radial", UIMin = "3", UIMax = "128", ClampMin = "3", ClampMax = "500", ProceduralShapeSetting))
 	int RadialSlices = 16;
 
-	/** Number of Vertical Subdivisions Along the Height of the Cylidner */
+	/** Number of Vertical Subdivisions Along the Height of the Cylinder */
 	UPROPERTY(EditAnywhere, Category = "ShapeSettings|Subdivisions", meta = (DisplayName = "Height", UIMin = "1", UIMax = "100", ClampMin = "1", ClampMax = "500", ProceduralShapeSetting))
 	int HeightSubdivisions = 1;
 };
@@ -308,35 +335,38 @@ public:
 	int TotalSubdivisions = 1;
 };
 
+UENUM()
+enum class EProceduralSphereType
+{
+	/** Create a Sphere with Lat Long parameterization */
+	LatLong UMETA(DisplayName = "Lat Long"),
+	/** Create a Sphere with Box parameterization */
+	Box UMETA(DisplayName = "Box")
+};
+
 UCLASS()
 class MESHMODELINGTOOLS_API UProceduralSphereToolProperties : public UProceduralShapeToolProperties
 {
 	GENERATED_BODY()
 public:
+	/** Type of Sphere to create */
+	UPROPERTY(EditAnywhere, Category = "ShapeSettings")
+	EProceduralSphereType SphereType = EProceduralSphereType::LatLong;
+
 	/** Radius of the Sphere */
 	UPROPERTY(EditAnywhere, Category = "ShapeSettings|Dimensions", meta = (DisplayName = "Radius", UIMin = "1.0", UIMax = "1000.0", ClampMin = "0.0001", ClampMax = "1000000.0", ProceduralShapeSetting))
 	float Radius = 50.f;
 
 	/** Number of Latitudinal Slices of the Sphere */
-	UPROPERTY(EditAnywhere, Category = "ShapeSettings|Slices", meta = (DisplayName = "Latitude Slices", UIMin = "3", UIMax = "100", ClampMin = "4", ClampMax = "500", ProceduralShapeSetting))
+	UPROPERTY(EditAnywhere, Category = "ShapeSettings|Subdivisions", meta = (DisplayName = "Latitude Slices", UIMin = "3", UIMax = "100", ClampMin = "4", ClampMax = "500", ProceduralShapeSetting, EditCondition = "SphereType == EProceduralSphereType::LatLong", EditConditionHides))
 	int LatitudeSlices = 16;
 
 	/** Number of Longitudinal Slices around the Sphere */
-	UPROPERTY(EditAnywhere, Category = "ShapeSettings|Slices", meta = (DisplayName = "Longitude Slices", UIMin = "3", UIMax = "100", ClampMin = "4", ClampMax = "500", ProceduralShapeSetting))
+	UPROPERTY(EditAnywhere, Category = "ShapeSettings|Subdivisions", meta = (DisplayName = "Longitude Slices", UIMin = "3", UIMax = "100", ClampMin = "4", ClampMax = "500", ProceduralShapeSetting, EditCondition = "SphereType == EProceduralSphereType::LatLong", EditConditionHides))
 	int LongitudeSlices = 16;
-};
-
-UCLASS()
-class MESHMODELINGTOOLS_API UProceduralSphericalBoxToolProperties : public UProceduralShapeToolProperties
-{
-	GENERATED_BODY()
-public:
-	/** Radius of the Sphere */
-	UPROPERTY(EditAnywhere, Category = "ShapeSettings|Dimensions", meta = (DisplayName = "Radius", UIMin = "1.0", UIMax = "1000.0", ClampMin = "0.0001", ClampMax = "1000000.0", ProceduralShapeSetting))
-	float Radius = 50.f;
 
 	/** Number of Subdivisions of each Side of the Sphere */
-	UPROPERTY(EditAnywhere, Category = "ShapeSettings|Subdivisions", meta = (DisplayName = "Side", UIMin = "1", UIMax = "100", ClampMin = "1", ClampMax = "500", ProceduralShapeSetting))
+	UPROPERTY(EditAnywhere, Category = "ShapeSettings|Subdivisions", meta = (DisplayName = "Side", UIMin = "1", UIMax = "100", ClampMin = "1", ClampMax = "500", ProceduralShapeSetting, EditCondition = "SphereType == EProceduralSphereType::Box", EditConditionHides))
 	int Subdivisions = 16;
 };
 
@@ -531,16 +561,6 @@ protected:
 	void GenerateMesh(FDynamicMesh3* OutMesh) const override;
 };
 
-UCLASS()
-class UAddRoundedRectanglePrimitiveTool : public UAddPrimitiveTool
-{
-	GENERATED_BODY()
-public:
-	UAddRoundedRectanglePrimitiveTool(const FObjectInitializer& ObjectInitializer);
-protected:
-	void GenerateMesh(FDynamicMesh3* OutMesh) const override;
-};
-
 UCLASS()
 class UAddDiscPrimitiveTool : public UAddPrimitiveTool
 {
@@ -551,16 +571,6 @@ protected:
 	void GenerateMesh(FDynamicMesh3* OutMesh) const override;
 };
 
-UCLASS()
-class UAddPuncturedDiscPrimitiveTool : public UAddPrimitiveTool
-{
-	GENERATED_BODY()
-public:
-	UAddPuncturedDiscPrimitiveTool(const FObjectInitializer& ObjectInitializer);
-protected:
-	void GenerateMesh(FDynamicMesh3* OutMesh) const override;
-};
-
 UCLASS()
 class UAddTorusPrimitiveTool : public UAddPrimitiveTool
 {
@@ -591,16 +601,6 @@ protected:
 	void GenerateMesh(FDynamicMesh3* OutMesh) const override;
 };
 
-UCLASS()
-class UAddSphericalBoxPrimitiveTool : public UAddPrimitiveTool
-{
-	GENERATED_BODY()
-public:
-	UAddSphericalBoxPrimitiveTool(const FObjectInitializer& ObjectInitializer);
-protected:
-	void GenerateMesh(FDynamicMesh3* OutMesh) const override;
-};
-
 UCLASS()
 class UAddStairsPrimitiveTool : public UAddPrimitiveTool
 {
diff --git a/Engine/Plugins/Experimental/ModelingToolsEditorMode/Source/ModelingToolsEditorMode/Private/ModelingToolsEditorMode.cpp b/Engine/Plugins/Experimental/ModelingToolsEditorMode/Source/ModelingToolsEditorMode/Private/ModelingToolsEditorMode.cpp
--- a/Engine/Plugins/Experimental/ModelingToolsEditorMode/Source/ModelingToolsEditorMode/Private/ModelingToolsEditorMode.cpp
+++ b/Engine/Plugins/Experimental/ModelingToolsEditorMode/Source/ModelingToolsEditorMode/Private/ModelingToolsEditorMode.cpp
@@ -329,24 +329,15 @@ void UModelingToolsEditorMode::Enter()
 	RegisterPrimitiveToolFunc(ToolManagerCommands.BeginAddRectanglePrimitiveTool,
 							  TEXT("BeginAddRectanglePrimitiveTool"),
 							  UAddPrimitiveToolBuilder::EMakeMeshShapeType::Rectangle);
-	RegisterPrimitiveToolFunc(ToolManagerCommands.BeginAddRoundedRectanglePrimitiveTool,
-							  TEXT("BeginAddRoundedRectanglePrimitiveTool"),
-							  UAddPrimitiveToolBuilder::EMakeMeshShapeType::RoundedRectangle);
 	RegisterPrimitiveToolFunc(ToolManagerCommands.BeginAddDiscPrimitiveTool,
 							  TEXT("BeginAddDiscPrimitiveTool"),
 							  UAddPrimitiveToolBuilder::EMakeMeshShapeType::Disc);
-	RegisterPrimitiveToolFunc(ToolManagerCommands.BeginAddPuncturedDiscPrimitiveTool,
-							  TEXT("BeginAddPuncturedDiscPrimitiveTool"),
-							  UAddPrimitiveToolBuilder::EMakeMeshShapeType::PuncturedDisc);
 	RegisterPrimitiveToolFunc(ToolManagerCommands.BeginAddTorusPrimitiveTool,
 							  TEXT("BeginAddTorusPrimitiveTool"),
 							  UAddPrimitiveToolBuilder::EMakeMeshShapeType::Torus);
 	RegisterPrimitiveToolFunc(ToolManagerCommands.BeginAddSpherePrimitiveTool,
 							  TEXT("BeginAddSpherePrimitiveTool"),
 							  UAddPrimitiveToolBuilder::EMakeMeshShapeType::Sphere);
-	RegisterPrimitiveToolFunc(ToolManagerCommands.BeginAddSphericalBoxPrimitiveTool,
-							  TEXT("BeginAddSphericalBoxPrimitiveTool"),
-							  UAddPrimitiveToolBuilder::EMakeMeshShapeType::SphericalBox);
 	RegisterPrimitiveToolFunc(ToolManagerCommands.BeginAddStairsPrimitiveTool,
 							  TEXT("BeginAddStairsPrimitiveTool"),
 							  UAddPrimitiveToolBuilder::EMakeMeshShapeType::Stairs);
diff --git a/Engine/Plugins/Experimental/ModelingToolsEditorMode/Source/ModelingToolsEditorMode/Private/ModelingToolsEditorModeStyle.cpp b/Engine/Plugins/Experimental/ModelingToolsEditorMode/Source/ModelingToolsEditorMode/Private/ModelingToolsEditorModeStyle.cpp
--- a/Engine/Plugins/Experimental/ModelingToolsEditorMode/Source/ModelingToolsEditorMode/Private/ModelingToolsEditorModeStyle.cpp
+++ b/Engine/Plugins/Experimental/ModelingToolsEditorMode/Source/ModelingToolsEditorMode/Private/ModelingToolsEditorModeStyle.cpp
@@ -113,18 +113,12 @@ void FModelingToolsEditorModeStyle::Initialize()
 		StyleSet->Set("ModelingToolsManagerCommands.BeginAddArrowPrimitiveTool.Small", 		new IMAGE_PLUGIN_BRUSH("Icons/ModelingArrow_x40",		Icon40x40));
 		StyleSet->Set("ModelingToolsManagerCommands.BeginAddRectanglePrimitiveTool", 			new IMAGE_PLUGIN_BRUSH("Icons/ModelingRectangle_x20", 		Icon20x20));
 		StyleSet->Set("ModelingToolsManagerCommands.BeginAddRectanglePrimitiveTool.Small", 		new IMAGE_PLUGIN_BRUSH("Icons/ModelingRectangle_x40",		Icon40x40));
-		StyleSet->Set("ModelingToolsManagerCommands.BeginAddRoundedRectanglePrimitiveTool", 			new IMAGE_PLUGIN_BRUSH("Icons/ModelingRoundedRectangle_x20", 		Icon20x20));
-		StyleSet->Set("ModelingToolsManagerCommands.BeginAddRoundedRectanglePrimitiveTool.Small", 		new IMAGE_PLUGIN_BRUSH("Icons/ModelingRoundedRectangle_x40",		Icon40x40));
 		StyleSet->Set("ModelingToolsManagerCommands.BeginAddDiscPrimitiveTool", 			new IMAGE_PLUGIN_BRUSH("Icons/ModelingDisc_x20", 		Icon20x20));
 		StyleSet->Set("ModelingToolsManagerCommands.BeginAddDiscPrimitiveTool.Small", 		new IMAGE_PLUGIN_BRUSH("Icons/ModelingDisc_x40",		Icon40x40));
-		StyleSet->Set("ModelingToolsManagerCommands.BeginAddPuncturedDiscPrimitiveTool", 			new IMAGE_PLUGIN_BRUSH("Icons/ModelingPuncturedDisc_x20", 		Icon20x20));
-		StyleSet->Set("ModelingToolsManagerCommands.BeginAddPuncturedDiscPrimitiveTool.Small", 		new IMAGE_PLUGIN_BRUSH("Icons/ModelingPuncturedDisc_x40",		Icon40x40));
 		StyleSet->Set("ModelingToolsManagerCommands.BeginAddTorusPrimitiveTool", 			new IMAGE_PLUGIN_BRUSH("Icons/ModelingTorus_x20", 		Icon20x20));
 		StyleSet->Set("ModelingToolsManagerCommands.BeginAddTorusPrimitiveTool.Small", 		new IMAGE_PLUGIN_BRUSH("Icons/ModelingTorus_x40",		Icon40x40));
 		StyleSet->Set("ModelingToolsManagerCommands.BeginAddSpherePrimitiveTool", 			new IMAGE_PLUGIN_BRUSH("Icons/ModelingSphere_x20", 		Icon20x20));
 		StyleSet->Set("ModelingToolsManagerCommands.BeginAddSpherePrimitiveTool.Small", 		new IMAGE_PLUGIN_BRUSH("Icons/ModelingSphere_x40",		Icon40x40));
-		StyleSet->Set("ModelingToolsManagerCommands.BeginAddSphericalBoxPrimitiveTool", 			new IMAGE_PLUGIN_BRUSH("Icons/ModelingSphericalBox_x20", 		Icon20x20));
-		StyleSet->Set("ModelingToolsManagerCommands.BeginAddSphericalBoxPrimitiveTool.Small", 		new IMAGE_PLUGIN_BRUSH("Icons/ModelingSphericalBox_x40",		Icon40x40));
 
 		StyleSet->Set("ModelingToolsManagerCommands.BeginDrawPolygonTool", 				new IMAGE_PLUGIN_BRUSH("Icons/DrawPolygon_40x",		Icon20x20));
 		StyleSet->Set("ModelingToolsManagerCommands.BeginDrawPolygonTool.Small", 		new IMAGE_PLUGIN_BRUSH("Icons/DrawPolygon_40x", 	Icon20x20));
diff --git a/Engine/Plugins/Experimental/ModelingToolsEditorMode/Source/ModelingToolsEditorMode/Private/ModelingToolsEditorModeToolkit.cpp b/Engine/Plugins/Experimental/ModelingToolsEditorMode/Source/ModelingToolsEditorMode/Private/ModelingToolsEditorModeToolkit.cpp
--- a/Engine/Plugins/Experimental/ModelingToolsEditorMode/Source/ModelingToolsEditorMode/Private/ModelingToolsEditorModeToolkit.cpp
+++ b/Engine/Plugins/Experimental/ModelingToolsEditorMode/Source/ModelingToolsEditorMode/Private/ModelingToolsEditorModeToolkit.cpp
@@ -460,15 +460,12 @@ void FModelingToolsEditorModeToolkit::BuildToolPalette_Experimental(FName Palett
 	{
 		ToolbarBuilder.AddToolBarButton(Commands.BeginAddBoxPrimitiveTool);
 		ToolbarBuilder.AddToolBarButton(Commands.BeginAddSpherePrimitiveTool);
-		ToolbarBuilder.AddToolBarButton(Commands.BeginAddSphericalBoxPrimitiveTool);
 		ToolbarBuilder.AddToolBarButton(Commands.BeginAddCylinderPrimitiveTool);
 		ToolbarBuilder.AddToolBarButton(Commands.BeginAddConePrimitiveTool);
 		ToolbarBuilder.AddToolBarButton(Commands.BeginAddTorusPrimitiveTool);
 		ToolbarBuilder.AddToolBarButton(Commands.BeginAddArrowPrimitiveTool);
 		ToolbarBuilder.AddToolBarButton(Commands.BeginAddRectanglePrimitiveTool);
-		ToolbarBuilder.AddToolBarButton(Commands.BeginAddRoundedRectanglePrimitiveTool);
 		ToolbarBuilder.AddToolBarButton(Commands.BeginAddDiscPrimitiveTool);
-		ToolbarBuilder.AddToolBarButton(Commands.BeginAddPuncturedDiscPrimitiveTool);
 		ToolbarBuilder.AddToolBarButton(Commands.BeginAddStairsPrimitiveTool);
 	}
 	else if (PaletteIndex == CreateTabName)
diff --git a/Engine/Plugins/Experimental/ModelingToolsEditorMode/Source/ModelingToolsEditorMode/Private/ModelingToolsManagerActions.cpp b/Engine/Plugins/Experimental/ModelingToolsEditorMode/Source/ModelingToolsEditorMode/Private/ModelingToolsManagerActions.cpp
--- a/Engine/Plugins/Experimental/ModelingToolsEditorMode/Source/ModelingToolsEditorMode/Private/ModelingToolsManagerActions.cpp
+++ b/Engine/Plugins/Experimental/ModelingToolsEditorMode/Source/ModelingToolsEditorMode/Private/ModelingToolsManagerActions.cpp
@@ -27,12 +27,9 @@ void FModelingToolsManagerCommands::RegisterCommands()
 	UI_COMMAND(BeginAddConePrimitiveTool, "Cone", "Create new Cone StaticMesh Assets", EUserInterfaceActionType::ToggleButton, FInputChord());
 	UI_COMMAND(BeginAddArrowPrimitiveTool, "Arrow", "Create new Arrow StaticMesh Assets", EUserInterfaceActionType::ToggleButton, FInputChord());
 	UI_COMMAND(BeginAddRectanglePrimitiveTool, "Rect", "Create new Rectangle StaticMesh Assets", EUserInterfaceActionType::ToggleButton, FInputChord());
-	UI_COMMAND(BeginAddRoundedRectanglePrimitiveTool, "RndRct", "Create new Rounded Rectangle StaticMesh Assets", EUserInterfaceActionType::ToggleButton, FInputChord());
 	UI_COMMAND(BeginAddDiscPrimitiveTool, "Disc", "Create new Disc StaticMesh Assets", EUserInterfaceActionType::ToggleButton, FInputChord());
-	UI_COMMAND(BeginAddPuncturedDiscPrimitiveTool, "Circle", "Create new Punctured Disc StaticMesh Assets", EUserInterfaceActionType::ToggleButton, FInputChord());
 	UI_COMMAND(BeginAddTorusPrimitiveTool, "Torus", "Create new Torus StaticMesh Assets", EUserInterfaceActionType::ToggleButton, FInputChord());
-	UI_COMMAND(BeginAddSphericalBoxPrimitiveTool, "SphrB", "Create new Sphere StaticMesh Assets with Box Parameterization", EUserInterfaceActionType::ToggleButton, FInputChord());
-	UI_COMMAND(BeginAddSpherePrimitiveTool, "SphrA", "Create new Sphere StaticMesh Assets with Lat/Long Parameterization", EUserInterfaceActionType::ToggleButton, FInputChord());
+	UI_COMMAND(BeginAddSpherePrimitiveTool, "Sphere", "Create new Sphere StaticMesh Assets", EUserInterfaceActionType::ToggleButton, FInputChord());
 	UI_COMMAND(BeginAddStairsPrimitiveTool, "Stair", "Create new Stairs StaticMesh Assets", EUserInterfaceActionType::ToggleButton, FInputChord());
 
 	UI_COMMAND(BeginAddPatchTool, "Patch", "Add Patch", EUserInterfaceActionType::ToggleButton, FInputChord());
diff --git a/Engine/Plugins/Experimental/ModelingToolsEditorMode/Source/ModelingToolsEditorMode/Public/ModelingToolsManagerActions.h b/Engine/Plugins/Experimental/ModelingToolsEditorMode/Source/ModelingToolsEditorMode/Public/ModelingToolsManagerActions.h
--- a/Engine/Plugins/Experimental/ModelingToolsEditorMode/Source/ModelingToolsEditorMode/Public/ModelingToolsManagerActions.h
+++ b/Engine/Plugins/Experimental/ModelingToolsEditorMode/Source/ModelingToolsEditorMode/Public/ModelingToolsManagerActions.h
@@ -20,11 +20,8 @@ public:
 	TSharedPtr<FUICommandInfo> BeginAddConePrimitiveTool;
 	TSharedPtr<FUICommandInfo> BeginAddArrowPrimitiveTool;
 	TSharedPtr<FUICommandInfo> BeginAddRectanglePrimitiveTool;
-	TSharedPtr<FUICommandInfo> BeginAddRoundedRectanglePrimitiveTool;
 	TSharedPtr<FUICommandInfo> BeginAddDiscPrimitiveTool;
-	TSharedPtr<FUICommandInfo> BeginAddPuncturedDiscPrimitiveTool;
 	TSharedPtr<FUICommandInfo> BeginAddTorusPrimitiveTool;
-	TSharedPtr<FUICommandInfo> BeginAddSphericalBoxPrimitiveTool;
 	TSharedPtr<FUICommandInfo> BeginAddSpherePrimitiveTool;
 	TSharedPtr<FUICommandInfo> BeginAddStairsPrimitiveTool;
 
