ModelingTools: Set Sphere generator default to Box parameterization. Remove Add Primitive property sub-categories.

#rnx

diff --git a/Engine/Plugins/Experimental/MeshModelingToolset/Source/MeshModelingTools/Public/AddPrimitiveTool.h b/Engine/Plugins/Experimental/MeshModelingToolset/Source/MeshModelingTools/Public/AddPrimitiveTool.h
--- a/Engine/Plugins/Experimental/MeshModelingToolset/Source/MeshModelingTools/Public/AddPrimitiveTool.h
+++ b/Engine/Plugins/Experimental/MeshModelingToolset/Source/MeshModelingTools/Public/AddPrimitiveTool.h
@@ -134,27 +134,27 @@ public:
 	UProceduralBoxToolProperties() = default;
 
 	/** Width of Shape */
-	UPROPERTY(EditAnywhere, Category = "ShapeSettings|Dimensions", meta = (DisplayName = "Width", UIMin = "1.0", UIMax = "1000.0", ClampMin = "0.0001", ClampMax = "1000000.0", ProceduralShapeSetting))
+	UPROPERTY(EditAnywhere, Category = "ShapeSettings", meta = (DisplayName = "Width", UIMin = "1.0", UIMax = "1000.0", ClampMin = "0.0001", ClampMax = "1000000.0", ProceduralShapeSetting))
 	float Width = 100.f;
 
 	/** Depth of Shape */
-	UPROPERTY(EditAnywhere, Category = "ShapeSettings|Dimensions", meta = (DisplayName = "Depth", UIMin = "1.0", UIMax = "1000.0", ClampMin = "0.0001", ClampMax = "1000000.0", ProceduralShapeSetting))
+	UPROPERTY(EditAnywhere, Category = "ShapeSettings", meta = (DisplayName = "Depth", UIMin = "1.0", UIMax = "1000.0", ClampMin = "0.0001", ClampMax = "1000000.0", ProceduralShapeSetting))
 	float Depth = 100.f;
 
 	/** Height of Shape */
-	UPROPERTY(EditAnywhere, Category = "ShapeSettings|Dimensions", meta = (DisplayName = "Height", UIMin = "1.0", UIMax = "1000.0", ClampMin = "0.0001", ClampMax = "1000000.0", ProceduralShapeSetting))
+	UPROPERTY(EditAnywhere, Category = "ShapeSettings", meta = (DisplayName = "Height", UIMin = "1.0", UIMax = "1000.0", ClampMin = "0.0001", ClampMax = "1000000.0", ProceduralShapeSetting))
 	float Height = 100.f;
 
 	/** Number of Subdivisions Along the Width */
-	UPROPERTY(EditAnywhere, Category = "ShapeSettings|Subdivisions", meta = (DisplayName = "Width", UIMin = "1", UIMax = "100", ClampMin = "1", ClampMax = "500", ProceduralShapeSetting))
+	UPROPERTY(EditAnywhere, Category = "ShapeSettings", meta = (DisplayName = "Width Subdivisions", UIMin = "1", UIMax = "100", ClampMin = "1", ClampMax = "500", ProceduralShapeSetting))
 	int WidthSubdivisions = 1;
 
 	/** Number of Subdivisions Along the Depth */
-	UPROPERTY(EditAnywhere, Category = "ShapeSettings|Subdivisions", meta = (DisplayName = "Depth", UIMin = "1", UIMax = "100", ClampMin = "1", ClampMax = "500", ProceduralShapeSetting))
+	UPROPERTY(EditAnywhere, Category = "ShapeSettings", meta = (DisplayName = "Depth Subdivisions", UIMin = "1", UIMax = "100", ClampMin = "1", ClampMax = "500", ProceduralShapeSetting))
 	int DepthSubdivisions = 1;
 
 	/** Number of Subdivisions Along the Height */
-	UPROPERTY(EditAnywhere, Category = "ShapeSettings|Subdivisions", meta = (DisplayName = "Height", UIMin = "0", UIMax = "100", ClampMin = "0", ClampMax = "500", ProceduralShapeSetting))
+	UPROPERTY(EditAnywhere, Category = "ShapeSettings", meta = (DisplayName = "Height Subdivisions", UIMin = "0", UIMax = "100", ClampMin = "0", ClampMax = "500", ProceduralShapeSetting))
 	int HeightSubdivisions = 1;
 };
 
@@ -179,27 +179,27 @@ public:
 	EProceduralRectType RectType = EProceduralRectType::Rectangle;
 
 	/** Width of Shape */
-	UPROPERTY(EditAnywhere, Category = "ShapeSettings|Dimensions", meta = (DisplayName = "Width", UIMin = "1.0", UIMax = "1000.0", ClampMin = "0.0001", ClampMax = "1000000.0", ProceduralShapeSetting))
+	UPROPERTY(EditAnywhere, Category = "ShapeSettings", meta = (DisplayName = "Width", UIMin = "1.0", UIMax = "1000.0", ClampMin = "0.0001", ClampMax = "1000000.0", ProceduralShapeSetting))
 	float Width = 100.f;
 
 	/** Depth of Shape */
-	UPROPERTY(EditAnywhere, Category = "ShapeSettings|Dimensions", meta = (DisplayName = "Depth", UIMin = "1.0", UIMax = "1000.0", ClampMin = "0.0001", ClampMax = "1000000.0", ProceduralShapeSetting))
+	UPROPERTY(EditAnywhere, Category = "ShapeSettings", meta = (DisplayName = "Depth", UIMin = "1.0", UIMax = "1000.0", ClampMin = "0.0001", ClampMax = "1000000.0", ProceduralShapeSetting))
 	float Depth = 100.f;
 
 	/** Number of Subdivisions Along the Width */
-	UPROPERTY(EditAnywhere, Category = "ShapeSettings|Subdivisions", meta = (DisplayName = "Width", UIMin = "1", UIMax = "100", ClampMin = "1", ClampMax = "500", ProceduralShapeSetting))
+	UPROPERTY(EditAnywhere, Category = "ShapeSettings", meta = (DisplayName = "Width Subdivisions", UIMin = "1", UIMax = "100", ClampMin = "1", ClampMax = "500", ProceduralShapeSetting))
 	int WidthSubdivisions = 1;
 
 	/** Number of Subdivisions Along the Depth */
-	UPROPERTY(EditAnywhere, Category = "ShapeSettings|Subdivisions", meta = (DisplayName = "Depth", UIMin = "1", UIMax = "100", ClampMin = "1", ClampMax = "500", ProceduralShapeSetting))
+	UPROPERTY(EditAnywhere, Category = "ShapeSettings", meta = (DisplayName = "Depth Subdivisions", UIMin = "1", UIMax = "100", ClampMin = "1", ClampMax = "500", ProceduralShapeSetting))
 	int DepthSubdivisions = 1;
 
 	/** Radius of Rounded Corners */
-	UPROPERTY(EditAnywhere, Category = "ShapeSettings|Dimensions", meta = (DisplayName = "Corner Radius", UIMin = "1.0", UIMax = "1000.0", ClampMin = "0.0001", ClampMax = "1000000.0", ProceduralShapeSetting, EditCondition = "RectType == EProceduralRectType::RoundedRectangle", EditConditionHides))
+	UPROPERTY(EditAnywhere, Category = "ShapeSettings", meta = (DisplayName = "Corner Radius", UIMin = "1.0", UIMax = "1000.0", ClampMin = "0.0001", ClampMax = "1000000.0", ProceduralShapeSetting, EditCondition = "RectType == EProceduralRectType::RoundedRectangle", EditConditionHides))
 	float CornerRadius = 25.f;
 
 	/** Number of Angular Slices in Each Rounded Corner */
-	UPROPERTY(EditAnywhere, Category = "ShapeSettings|Subdivisions", meta = (DisplayName = "Corner", UIMin = "3", UIMax = "128", ClampMin = "3", ClampMax = "500", ProceduralShapeSetting, EditCondition = "RectType == EProceduralRectType::RoundedRectangle", EditConditionHides))
+	UPROPERTY(EditAnywhere, Category = "ShapeSettings", meta = (DisplayName = "Corner Slices", UIMin = "3", UIMax = "128", ClampMin = "3", ClampMax = "500", ProceduralShapeSetting, EditCondition = "RectType == EProceduralRectType::RoundedRectangle", EditConditionHides))
 	int CornerSlices = 16;
 };
 
@@ -222,19 +222,19 @@ public:
 	EProceduralDiscType DiscType = EProceduralDiscType::Disc;
 
 	/** Radius of Disc */
-	UPROPERTY(EditAnywhere, Category = "ShapeSettings|Dimensions", meta = (DisplayName = "Radius", UIMin = "1.0", UIMax = "1000.0", ClampMin = "0.0001", ClampMax = "1000000.0", ProceduralShapeSetting))
+	UPROPERTY(EditAnywhere, Category = "ShapeSettings", meta = (DisplayName = "Radius", UIMin = "1.0", UIMax = "1000.0", ClampMin = "0.0001", ClampMax = "1000000.0", ProceduralShapeSetting))
 	float Radius = 50.f;
 
 	/** Number of Angular Slices around the Disc */
-	UPROPERTY(EditAnywhere, Category = "ShapeSettings|Slices", meta = (DisplayName = "Radial", UIMin = "3", UIMax = "128", ClampMin = "3", ClampMax = "500", ProceduralShapeSetting))
+	UPROPERTY(EditAnywhere, Category = "ShapeSettings", meta = (DisplayName = "Radial Slices", UIMin = "3", UIMax = "128", ClampMin = "3", ClampMax = "500", ProceduralShapeSetting))
 	int RadialSlices = 16;
 
 	/** Number of Radial Subdivisions in the Disc */
-	UPROPERTY(EditAnywhere, Category = "ShapeSettings|Subdivisions", meta = (DisplayName = "Radial", UIMin = "1", UIMax = "100", ClampMin = "1", ClampMax = "500", ProceduralShapeSetting))
+	UPROPERTY(EditAnywhere, Category = "ShapeSettings", meta = (DisplayName = "Radial Subdivisions", UIMin = "1", UIMax = "100", ClampMin = "1", ClampMax = "500", ProceduralShapeSetting))
 	int RadialSubdivisions = 1;
 
 	/** Radius of the Disc's Hole */
-	UPROPERTY(EditAnywhere, Category = "ShapeSettings|Dimensions", meta = (DisplayName = "Hole Radius", UIMin = "1.0", UIMax = "1000.0", ClampMin = "0.0001", ClampMax = "1000000.0", ProceduralShapeSetting, EditCondition = "DiscType == EProceduralDiscType::PuncturedDisc", EditConditionHides))
+	UPROPERTY(EditAnywhere, Category = "ShapeSettings", meta = (DisplayName = "Hole Radius", UIMin = "1.0", UIMax = "1000.0", ClampMin = "0.0001", ClampMax = "1000000.0", ProceduralShapeSetting, EditCondition = "DiscType == EProceduralDiscType::PuncturedDisc", EditConditionHides))
 	float HoleRadius = 25.f;
 };
 
@@ -244,19 +244,19 @@ class MESHMODELINGTOOLS_API UProceduralTorusToolProperties : public UProceduralS
 	GENERATED_BODY()
 public:
 	/** Radius from the Torus Center to the Center of the Torus Tube */
-	UPROPERTY(EditAnywhere, Category = "ShapeSettings|Dimensions", meta = (DisplayName = "Major Radius", UIMin = "1.0", UIMax = "1000.0", ClampMin = "0.0001", ClampMax = "1000000.0", ProceduralShapeSetting))
+	UPROPERTY(EditAnywhere, Category = "ShapeSettings", meta = (DisplayName = "Major Radius", UIMin = "1.0", UIMax = "1000.0", ClampMin = "0.0001", ClampMax = "1000000.0", ProceduralShapeSetting))
 	float MajorRadius = 50.f;
 
 	/** Radius of the Torus Tube */
-	UPROPERTY(EditAnywhere, Category = "ShapeSettings|Dimensions", meta = (DisplayName = "Minor Radius", UIMin = "1.0", UIMax = "1000.0", ClampMin = "0.0001", ClampMax = "1000000.0", ProceduralShapeSetting))
+	UPROPERTY(EditAnywhere, Category = "ShapeSettings", meta = (DisplayName = "Minor Radius", UIMin = "1.0", UIMax = "1000.0", ClampMin = "0.0001", ClampMax = "1000000.0", ProceduralShapeSetting))
 	float MinorRadius = 25.f;
 
 	/** Number of Angular Slices Along the Torus Tube */
-	UPROPERTY(EditAnywhere, Category = "ShapeSettings|Slices", meta = (DisplayName = "Major", UIMin = "3", UIMax = "128", ClampMin = "3", ClampMax = "500", ProceduralShapeSetting))
+	UPROPERTY(EditAnywhere, Category = "ShapeSettings", meta = (DisplayName = "Major Slices", UIMin = "3", UIMax = "128", ClampMin = "3", ClampMax = "500", ProceduralShapeSetting))
 	int TubeSlices = 16;
 
 	/** Number of Angular Slices Around the Tube of the Torus */
-	UPROPERTY(EditAnywhere, Category = "ShapeSettings|Slices", meta = (DisplayName = "Minor", UIMin = "3", UIMax = "128", ClampMin = "3", ClampMax = "500", ProceduralShapeSetting))
+	UPROPERTY(EditAnywhere, Category = "ShapeSettings", meta = (DisplayName = "Minor Slices", UIMin = "3", UIMax = "128", ClampMin = "3", ClampMax = "500", ProceduralShapeSetting))
 	int CrossSectionSlices = 16;
 };
 
@@ -267,19 +267,19 @@ class MESHMODELINGTOOLS_API UProceduralCylinderToolProperties : public UProcedur
 public:
 
 	/** Radius of The Cylinder */
-	UPROPERTY(EditAnywhere, Category = "ShapeSettings|Dimensions", meta = (DisplayName = "Radius", UIMin = "1.0", UIMax = "1000.0", ClampMin = "0.0001", ClampMax = "1000000.0", ProceduralShapeSetting))
+	UPROPERTY(EditAnywhere, Category = "ShapeSettings", meta = (DisplayName = "Radius", UIMin = "1.0", UIMax = "1000.0", ClampMin = "0.0001", ClampMax = "1000000.0", ProceduralShapeSetting))
 	float Radius = 50.f;
 
 	/** Height of Cylinder */
-	UPROPERTY(EditAnywhere, Category = "ShapeSettings|Dimensions", meta = (DisplayName = "Height", UIMin = "1.0", UIMax = "1000.0", ClampMin = "0.0001", ClampMax = "1000000.0", ProceduralShapeSetting))
+	UPROPERTY(EditAnywhere, Category = "ShapeSettings", meta = (DisplayName = "Height", UIMin = "1.0", UIMax = "1000.0", ClampMin = "0.0001", ClampMax = "1000000.0", ProceduralShapeSetting))
 	float Height = 200.f;
 
 	/** Number of Slices on the Cylinder Caps */
-	UPROPERTY(EditAnywhere, Category = "ShapeSettings|Slices", meta = (DisplayName = "Radial", UIMin = "3", UIMax = "128", ClampMin = "3", ClampMax = "500", ProceduralShapeSetting))
+	UPROPERTY(EditAnywhere, Category = "ShapeSettings", meta = (DisplayName = "Radial Slices", UIMin = "3", UIMax = "128", ClampMin = "3", ClampMax = "500", ProceduralShapeSetting))
 	int RadialSlices = 16;
 
 	/** Number of Vertical Subdivisions Along the Height of the Cylinder */
-	UPROPERTY(EditAnywhere, Category = "ShapeSettings|Subdivisions", meta = (DisplayName = "Height", UIMin = "1", UIMax = "100", ClampMin = "1", ClampMax = "500", ProceduralShapeSetting))
+	UPROPERTY(EditAnywhere, Category = "ShapeSettings", meta = (DisplayName = "Height Subdivisions", UIMin = "1", UIMax = "100", ClampMin = "1", ClampMax = "500", ProceduralShapeSetting))
 	int HeightSubdivisions = 1;
 };
 
@@ -289,19 +289,19 @@ class MESHMODELINGTOOLS_API UProceduralConeToolProperties : public UProceduralSh
 	GENERATED_BODY()
 public:
 	/** Radius of the Cone */
-	UPROPERTY(EditAnywhere, Category = "ShapeSettings|Dimensions", meta = (DisplayName = "Radius", UIMin = "1.0", UIMax = "1000.0", ClampMin = "0.0001", ClampMax = "1000000.0", ProceduralShapeSetting))
+	UPROPERTY(EditAnywhere, Category = "ShapeSettings", meta = (DisplayName = "Radius", UIMin = "1.0", UIMax = "1000.0", ClampMin = "0.0001", ClampMax = "1000000.0", ProceduralShapeSetting))
 	float Radius = 50.f;
 
 	/** Height of Cone */
-	UPROPERTY(EditAnywhere, Category = "ShapeSettings|Dimensions", meta = (DisplayName = "Height", UIMin = "1.0", UIMax = "1000.0", ClampMin = "0.0001", ClampMax = "1000000.0", ProceduralShapeSetting))
+	UPROPERTY(EditAnywhere, Category = "ShapeSettings", meta = (DisplayName = "Height", UIMin = "1.0", UIMax = "1000.0", ClampMin = "0.0001", ClampMax = "1000000.0", ProceduralShapeSetting))
 	float Height = 200.f;
 
 	/** Number of Slices on the Cone Base */
-	UPROPERTY(EditAnywhere, Category = "ShapeSettings|Slices", meta = (DisplayName = "Radial", UIMin = "3", UIMax = "128", ClampMin = "3", ClampMax = "500", ProceduralShapeSetting))
+	UPROPERTY(EditAnywhere, Category = "ShapeSettings", meta = (DisplayName = "Radial Slices", UIMin = "3", UIMax = "128", ClampMin = "3", ClampMax = "500", ProceduralShapeSetting))
 	int RadialSlices = 16;
 
 	/** Number of Vertical Subdivisions Along the Hight of the Cone */
-	UPROPERTY(EditAnywhere, Category = "ShapeSettings|Subdivisions", meta = (DisplayName = "Height", UIMin = "1", UIMax = "100", ClampMin = "1", ClampMax = "500", ProceduralShapeSetting))
+	UPROPERTY(EditAnywhere, Category = "ShapeSettings", meta = (DisplayName = "Height Subdivisions", UIMin = "1", UIMax = "100", ClampMin = "1", ClampMax = "500", ProceduralShapeSetting))
 	int HeightSubdivisions = 1;
 };
 
@@ -311,27 +311,27 @@ class MESHMODELINGTOOLS_API UProceduralArrowToolProperties : public UProceduralS
 	GENERATED_BODY()
 public:
 	/** Radius of the Arrow Shaft */
-	UPROPERTY(EditAnywhere, Category = "ShapeSettings|Dimensions", meta = (DisplayName = "Shaft Radius", UIMin = "1.0", UIMax = "1000.0", ClampMin = "0.0001", ClampMax = "1000000.0", ProceduralShapeSetting))
+	UPROPERTY(EditAnywhere, Category = "ShapeSettings", meta = (DisplayName = "Shaft Radius", UIMin = "1.0", UIMax = "1000.0", ClampMin = "0.0001", ClampMax = "1000000.0", ProceduralShapeSetting))
 	float ShaftRadius = 20.f;
 
 	/** Height of Arrow Shaft */
-	UPROPERTY(EditAnywhere, Category = "ShapeSettings|Dimensions", meta = (DisplayName = "Shaft Height", UIMin = "1.0", UIMax = "1000.0", ClampMin = "0.0001", ClampMax = "1000000.0", ProceduralShapeSetting))
+	UPROPERTY(EditAnywhere, Category = "ShapeSettings", meta = (DisplayName = "Shaft Height", UIMin = "1.0", UIMax = "1000.0", ClampMin = "0.0001", ClampMax = "1000000.0", ProceduralShapeSetting))
 	float ShaftHeight = 200.f;
 
 	/** Radius of the Arrow Head */
-	UPROPERTY(EditAnywhere, Category = "ShapeSettings|Dimensions", meta = (DisplayName = "Head Radius", UIMin = "1.0", UIMax = "1000.0", ClampMin = "0.0001", ClampMax = "1000000.0", ProceduralShapeSetting))
+	UPROPERTY(EditAnywhere, Category = "ShapeSettings", meta = (DisplayName = "Head Radius", UIMin = "1.0", UIMax = "1000.0", ClampMin = "0.0001", ClampMax = "1000000.0", ProceduralShapeSetting))
 	float HeadRadius = 60.f;
 
 	/** Height of Arrow's Head */
-	UPROPERTY(EditAnywhere, Category = "ShapeSettings|Dimensions", meta = (DisplayName = "Head Height", UIMin = "1.0", UIMax = "1000.0", ClampMin = "0.0001", ClampMax = "1000000.0", ProceduralShapeSetting))
+	UPROPERTY(EditAnywhere, Category = "ShapeSettings", meta = (DisplayName = "Head Height", UIMin = "1.0", UIMax = "1000.0", ClampMin = "0.0001", ClampMax = "1000000.0", ProceduralShapeSetting))
 	float HeadHeight = 120.f;
 
 	/** Number of Angular Slices Around the Arrow */
-	UPROPERTY(EditAnywhere, Category = "ShapeSettings|Slices", meta = (DisplayName = "Radial Slices", UIMin = "3", UIMax = "100", ClampMin = "3", ClampMax = "500", ProceduralShapeSetting))
+	UPROPERTY(EditAnywhere, Category = "ShapeSettings", meta = (DisplayName = "Radial Slices", UIMin = "3", UIMax = "100", ClampMin = "3", ClampMax = "500", ProceduralShapeSetting))
 	int RadialSlices = 16;
 
 	/** Number of Vertical Subdivisions Along in the Arrow */
-	UPROPERTY(EditAnywhere, Category = "ShapeSettings|Subdivisions", meta = (DisplayName = "Total", UIMin = "1", UIMax = "100", ClampMin = "1", ClampMax = "500", ProceduralShapeSetting))
+	UPROPERTY(EditAnywhere, Category = "ShapeSettings", meta = (DisplayName = "Total Subdivisions", UIMin = "1", UIMax = "100", ClampMin = "1", ClampMax = "500", ProceduralShapeSetting))
 	int TotalSubdivisions = 1;
 };
 
@@ -351,22 +351,22 @@ class MESHMODELINGTOOLS_API UProceduralSphereToolProperties : public UProcedural
 public:
 	/** Type of Sphere to create */
 	UPROPERTY(EditAnywhere, Category = "ShapeSettings")
-	EProceduralSphereType SphereType = EProceduralSphereType::LatLong;
+	EProceduralSphereType SphereType = EProceduralSphereType::Box;
 
 	/** Radius of the Sphere */
-	UPROPERTY(EditAnywhere, Category = "ShapeSettings|Dimensions", meta = (DisplayName = "Radius", UIMin = "1.0", UIMax = "1000.0", ClampMin = "0.0001", ClampMax = "1000000.0", ProceduralShapeSetting))
+	UPROPERTY(EditAnywhere, Category = "ShapeSettings", meta = (DisplayName = "Radius", UIMin = "1.0", UIMax = "1000.0", ClampMin = "0.0001", ClampMax = "1000000.0", ProceduralShapeSetting))
 	float Radius = 50.f;
 
 	/** Number of Latitudinal Slices of the Sphere */
-	UPROPERTY(EditAnywhere, Category = "ShapeSettings|Subdivisions", meta = (DisplayName = "Latitude Slices", UIMin = "3", UIMax = "100", ClampMin = "4", ClampMax = "500", ProceduralShapeSetting, EditCondition = "SphereType == EProceduralSphereType::LatLong", EditConditionHides))
+	UPROPERTY(EditAnywhere, Category = "ShapeSettings", meta = (DisplayName = "Latitude Slices", UIMin = "3", UIMax = "100", ClampMin = "4", ClampMax = "500", ProceduralShapeSetting, EditCondition = "SphereType == EProceduralSphereType::LatLong", EditConditionHides))
 	int LatitudeSlices = 16;
 
 	/** Number of Longitudinal Slices around the Sphere */
-	UPROPERTY(EditAnywhere, Category = "ShapeSettings|Subdivisions", meta = (DisplayName = "Longitude Slices", UIMin = "3", UIMax = "100", ClampMin = "4", ClampMax = "500", ProceduralShapeSetting, EditCondition = "SphereType == EProceduralSphereType::LatLong", EditConditionHides))
+	UPROPERTY(EditAnywhere, Category = "ShapeSettings", meta = (DisplayName = "Longitude Slices", UIMin = "3", UIMax = "100", ClampMin = "4", ClampMax = "500", ProceduralShapeSetting, EditCondition = "SphereType == EProceduralSphereType::LatLong", EditConditionHides))
 	int LongitudeSlices = 16;
 
 	/** Number of Subdivisions of each Side of the Sphere */
-	UPROPERTY(EditAnywhere, Category = "ShapeSettings|Subdivisions", meta = (DisplayName = "Side", UIMin = "1", UIMax = "100", ClampMin = "1", ClampMax = "500", ProceduralShapeSetting, EditCondition = "SphereType == EProceduralSphereType::Box", EditConditionHides))
+	UPROPERTY(EditAnywhere, Category = "ShapeSettings", meta = (DisplayName = "Side Subdivisions", UIMin = "1", UIMax = "100", ClampMin = "1", ClampMax = "500", ProceduralShapeSetting, EditCondition = "SphereType == EProceduralSphereType::Box", EditConditionHides))
 	int Subdivisions = 16;
 };
 
@@ -389,27 +389,27 @@ public:
 		EProceduralStairsType StairsType = EProceduralStairsType::Linear;
 
 	/** Number of Steps of Shape */
-	UPROPERTY(EditAnywhere, Category = "ShapeSettings|Dimensions", meta = (DisplayName = "Number of Steps", UIMin = "1", UIMax = "100", ClampMin = "1", ClampMax = "1000", ProceduralShapeSetting))
+	UPROPERTY(EditAnywhere, Category = "ShapeSettings", meta = (DisplayName = "Number of Steps", UIMin = "1", UIMax = "100", ClampMin = "1", ClampMax = "1000", ProceduralShapeSetting))
 	int NumSteps = 8;
 
 	/** Width of each step */
-	UPROPERTY(EditAnywhere, Category = "ShapeSettings|Dimensions", meta = (DisplayName = "Step Width", UIMin = "1.0", UIMax = "1000.0", ClampMin = "0.0001", ClampMax = "1000000.0", ProceduralShapeSetting))
+	UPROPERTY(EditAnywhere, Category = "ShapeSettings", meta = (DisplayName = "Step Width", UIMin = "1.0", UIMax = "1000.0", ClampMin = "0.0001", ClampMax = "1000000.0", ProceduralShapeSetting))
 	float StepWidth = 150.0f;
 
 	/** Height of each step */
-	UPROPERTY(EditAnywhere, Category = "ShapeSettings|Dimensions", meta = (DisplayName = "Step Height", UIMin = "1.0", UIMax = "1000.0", ClampMin = "0.0001", ClampMax = "1000000.0", ProceduralShapeSetting))
+	UPROPERTY(EditAnywhere, Category = "ShapeSettings", meta = (DisplayName = "Step Height", UIMin = "1.0", UIMax = "1000.0", ClampMin = "0.0001", ClampMax = "1000000.0", ProceduralShapeSetting))
 	float StepHeight = 20.0f;
 
 	/** Depth of each step for linear stair shapes */
-	UPROPERTY(EditAnywhere, Category = "ShapeSettings|Dimensions", meta = (DisplayName = "Step Depth", UIMin = "1.0", UIMax = "1000.0", ClampMin = "0.0001", ClampMax = "1000000.0", ProceduralShapeSetting, EditCondition = "StairsType == EProceduralStairsType::Linear", EditConditionHides))
+	UPROPERTY(EditAnywhere, Category = "ShapeSettings", meta = (DisplayName = "Step Depth", UIMin = "1.0", UIMax = "1000.0", ClampMin = "0.0001", ClampMax = "1000000.0", ProceduralShapeSetting, EditCondition = "StairsType == EProceduralStairsType::Linear", EditConditionHides))
 	float StepDepth = 30.0f;
 
-	/** Curve angle for curved/spiral stair shapes */
-	UPROPERTY(EditAnywhere, Category = "ShapeSettings|Dimensions", meta = (DisplayName = "Curve Angle", UIMin = "-360.0", UIMax = "360.0", ClampMin = "-360.0", ClampMax = "360.0", ProceduralShapeSetting, EditCondition = "StairsType == EProceduralStairsType::Curved", EditConditionHides))
+	/** Angular length of curved stair shapes */
+	UPROPERTY(EditAnywhere, Category = "ShapeSettings", meta = (DisplayName = "Curve Angle", UIMin = "-360.0", UIMax = "360.0", ClampMin = "-360.0", ClampMax = "360.0", ProceduralShapeSetting, EditCondition = "StairsType == EProceduralStairsType::Curved", EditConditionHides))
 	float CurveAngle = 90.0f;
 
 	/** Inner radius for curved/spiral stair shapes */
-	UPROPERTY(EditAnywhere, Category = "ShapeSettings|Dimensions", meta = (DisplayName = "Inner Radius", UIMin = "1.0", UIMax = "1000.0", ClampMin = "0.0001", ClampMax = "1000000.0", ProceduralShapeSetting, EditCondition = "StairsType == EProceduralStairsType::Curved", EditConditionHides))
+	UPROPERTY(EditAnywhere, Category = "ShapeSettings", meta = (DisplayName = "Inner Radius", UIMin = "1.0", UIMax = "1000.0", ClampMin = "0.0001", ClampMax = "1000000.0", ProceduralShapeSetting, EditCondition = "StairsType == EProceduralStairsType::Curved", EditConditionHides))
 	float InnerRadius = 150.0f;
 };
 
