Mesh Modeling Tools:  Make Gen Groups tool use async framework.

diff --git a/Engine/Plugins/Experimental/MeshModelingToolset/Source/MeshModelingTools/Private/ConvertToPolygonsTool.cpp b/Engine/Plugins/Experimental/MeshModelingToolset/Source/MeshModelingTools/Private/ConvertToPolygonsTool.cpp
--- a/Engine/Plugins/Experimental/MeshModelingToolset/Source/MeshModelingTools/Private/ConvertToPolygonsTool.cpp
+++ b/Engine/Plugins/Experimental/MeshModelingToolset/Source/MeshModelingTools/Private/ConvertToPolygonsTool.cpp
@@ -13,8 +13,12 @@
 #include "Util/ColorConstants.h"
 #include "ToolSetupUtil.h"
 #include "Polygroups/PolygroupUtil.h"
+#include "PreviewMesh.h"
 
 #include "SceneManagement.h" // for FPrimitiveDrawInterface
+#include "ModelingOperators.h"
+#include "MeshOpPreviewHelpers.h"
+#include "ModelingOperators.h"
 
 #define LOCTEXT_NAMESPACE "UConvertToPolygonsTool"
 
@@ -37,10 +41,111 @@ UInteractiveTool* UConvertToPolygonsToolBuilder::BuildTool(const FToolBuilderSta
 	UPrimitiveComponent* MeshComponent = Cast<UPrimitiveComponent>(ActorComponent);
 	check(MeshComponent != nullptr);
 	NewTool->SetSelection(MakeComponentTarget(MeshComponent));
+	NewTool->SetWorld(SceneState.World);
 
 	return NewTool;
 }
 
+class FConvertToPolygonsOp : public  FDynamicMeshOperator
+{
+public:
+	FConvertToPolygonsOp() : FDynamicMeshOperator() {}
+
+	virtual void CalculateResult(FProgressCancel* Progress) override
+	{
+		if (Progress && Progress->Cancelled() || !OriginalMesh)
+		{
+			return;
+		}
+
+		ResultMesh->Copy(*OriginalMesh, true, true, true, true);
+	
+
+		if (Progress && Progress->Cancelled())
+		{
+			return;
+		}
+
+		Polygons = FFindPolygonsAlgorithm(ResultMesh.Get());
+
+		switch (ConversionMode)
+		{
+		case EConvertToPolygonsMode::FromUVIslands:
+		{
+			Polygons.FindPolygonsFromUVIslands();
+			break;
+		}
+		case EConvertToPolygonsMode::FromConnectedTris:
+		{
+			Polygons.FindPolygonsFromConnectedTris();
+			break;
+		}
+		case EConvertToPolygonsMode::FaceNormalDeviation:
+		{
+			double DotTolerance = 1.0 - FMathd::Cos(AngleTolerance * FMathd::DegToRad);
+			Polygons.FindPolygonsFromFaceNormals(DotTolerance);
+			break;
+		}
+		default:
+			check(0);
+		}
+
+		Polygons.FindPolygonEdges();
+
+		if (bCalculateNormals && ConversionMode == EConvertToPolygonsMode::FaceNormalDeviation)
+		{
+			if (ResultMesh->HasAttributes() == false)
+			{
+				ResultMesh->EnableAttributes();
+			}
+
+			FDynamicMeshNormalOverlay* NormalOverlay = ResultMesh->Attributes()->PrimaryNormals();
+			NormalOverlay->ClearElements();
+
+			FDynamicMeshEditor Editor(ResultMesh.Get());
+			for (const TArray<int>& Polygon : Polygons.FoundPolygons)
+			{
+				FVector3f Normal = (FVector3f)ResultMesh->GetTriNormal(Polygon[0]);
+				Editor.SetTriangleNormals(Polygon, Normal);
+			}
+
+			FMeshNormals Normals(ResultMesh.Get());
+			Normals.RecomputeOverlayNormals(ResultMesh->Attributes()->PrimaryNormals());
+			Normals.CopyToOverlay(NormalOverlay, false);
+		}
+
+	}
+
+	void SetTransform(const FTransform& Transform)
+	{
+		ResultTransform = (FTransform3d)Transform;
+	}
+
+	
+	FFindPolygonsAlgorithm Polygons;
+	TSharedPtr<FDynamicMesh3, ESPMode::ThreadSafe> OriginalMesh;
+
+	// parameters set by the tool
+	EConvertToPolygonsMode ConversionMode;
+	double AngleTolerance;
+	bool bCalculateNormals;
+};
+
+TUniquePtr<FDynamicMeshOperator> UConvertToPolygonsOperatorFactory::MakeNewOperator()
+{
+	// backpointer used to populate parameters.
+	check(ConvertToPolygonsTool);
+
+	// Create the actual operator type based on the requested operation
+	TUniquePtr<FConvertToPolygonsOp>  MeshOp = MakeUnique<FConvertToPolygonsOp>();
+
+	// Operator runs on another thread - copy data over that it needs.
+	ConvertToPolygonsTool->UpdateOpParameters(*MeshOp);
+
+	// give the operator
+	return MeshOp;
+}
+
 /*
  * Tool
  */
@@ -49,196 +154,183 @@ UConvertToPolygonsTool::UConvertToPolygonsTool()
 	SetToolDisplayName(LOCTEXT("ConvertToPolygonsToolName", "Generate PolyGroups"));
 }
 
+void UConvertToPolygonsTool::SetWorld(UWorld* World)
+{
+	this->TargetWorld = World;
+}
+
+bool UConvertToPolygonsTool::CanAccept() const
+{
+	return Super::CanAccept() && (PreviewWithBackgroundCompute == nullptr || PreviewWithBackgroundCompute->HaveValidResult());
+}
+
 void UConvertToPolygonsTool::Setup()
 {
 	UInteractiveTool::Setup();
 
-	FMeshDescription* MeshDescription = nullptr;
-	MeshDescription = ComponentTarget->GetMesh();
-
-	FMeshDescriptionToDynamicMesh Converter;
-	Converter.Convert(MeshDescription, SearchMesh);
+	FComponentMaterialSet MaterialSet;
+	ComponentTarget->GetMaterialSet(MaterialSet);
 
-	if (SearchMesh.HasAttributes())
+	// populate the OriginalDynamicMesh with a conversion of the input mesh.
 	{
-		InitialNormals.Copy( *SearchMesh.Attributes()->PrimaryNormals() );
+		OriginalDynamicMesh = MakeShared<FDynamicMesh3, ESPMode::ThreadSafe>();
+		FMeshDescriptionToDynamicMesh Converter;
+		Converter.Convert(ComponentTarget->GetMesh(), *OriginalDynamicMesh);
 	}
 
 	Settings = NewObject<UConvertToPolygonsToolProperties>(this);
 	Settings->RestoreProperties(this);
 	AddToolPropertySource(Settings);
+	FTransform MeshTransform = ComponentTarget->GetWorldTransform();
+	// hide existing mesh
+	ComponentTarget->SetOwnerVisibility(false);
+	// Set up the preview object
+	{
+		// create the operator factory
+		UConvertToPolygonsOperatorFactory* ConvertToPolygonsOperatorFactory = NewObject<UConvertToPolygonsOperatorFactory>(this);
+		ConvertToPolygonsOperatorFactory->ConvertToPolygonsTool = this; // set the back pointer
+
+
+		PreviewWithBackgroundCompute = NewObject<UMeshOpPreviewWithBackgroundCompute>(ConvertToPolygonsOperatorFactory, "Preview");
+		PreviewWithBackgroundCompute->Setup(this->TargetWorld, ConvertToPolygonsOperatorFactory);
+		PreviewWithBackgroundCompute->SetIsMeshTopologyConstant(true, EMeshRenderAttributeFlags::Positions | EMeshRenderAttributeFlags::VertexNormals);
+
+		// Give the preview something to display
+		PreviewWithBackgroundCompute->PreviewMesh->SetTransform(ComponentTarget->GetWorldTransform());
+		PreviewWithBackgroundCompute->PreviewMesh->SetTangentsMode(EDynamicMeshTangentCalcType::AutoCalculated);
+		PreviewWithBackgroundCompute->PreviewMesh->UpdatePreview(OriginalDynamicMesh.Get());
+		
+		PreviewWithBackgroundCompute->ConfigureMaterials(MaterialSet.Materials,
+			ToolSetupUtil::GetDefaultWorkingMaterial(GetToolManager())
+		);
+
+		// show the preview mesh
+		PreviewWithBackgroundCompute->SetVisibility(true);
+
+		// something to capture the polygons from the async task when it is done
+		PreviewWithBackgroundCompute->OnOpCompleted.AddLambda([this](const FDynamicMeshOperator* MeshOp) 
+		{ 
+			const FConvertToPolygonsOp*  ConvertToPolygonsOp = static_cast<const FConvertToPolygonsOp*>(MeshOp);
+				
+			// edges used for tool ::Render() method 
+			this->PolygonEdges = ConvertToPolygonsOp->Polygons.PolygonEdges;
+				
+			// we have new triangle groups to color
+			UpdateVisualization();
+		});
 
-	// create preview mesh object
-	PreviewMesh = NewObject<UPreviewMesh>(this, TEXT("PreviewMesh"));
-	PreviewMesh->CreateInWorld(ComponentTarget->GetOwnerActor()->GetWorld(), FTransform::Identity);
-	PreviewMesh->SetVisible(false);
-	PreviewMesh->SetTransform(ComponentTarget->GetWorldTransform());
-	PreviewMesh->SetTangentsMode(EDynamicMeshTangentCalcType::AutoCalculated);
-
-	FComponentMaterialSet MaterialSet;
-	ComponentTarget->GetMaterialSet(MaterialSet);
-	PreviewMesh->SetMaterials(MaterialSet.Materials);
+		// updates the triangle color visualization
+		UpdateVisualization();
+		// start the compute
+		PreviewWithBackgroundCompute->InvalidateResult();
+		
+	}
+	
 
 	Settings->WatchProperty(Settings->ConversionMode,
 							[this](EConvertToPolygonsMode NewMode)
-							{ bPolygonsValid = false; });
+							{ OnSettingsModified(); });
 	Settings->WatchProperty(Settings->bShowGroupColors,
-							[this](bool bNewValue) { UpdateVisualization(); });
-	if (Settings->bShowGroupColors)
-	{
-		UpdateVisualization();
-	}
+							[this](bool bNewValue) 
+							{ UpdateVisualization(); });
+	Settings->WatchProperty(Settings->AngleTolerance,
+							[this](float value)
+							{ OnSettingsModified(); });
+	
 
 	GetToolManager()->DisplayMessage(
 		LOCTEXT("OnStartTool", "Cluster triangles of the Mesh into PolyGroups using various strategies"),
 		EToolMessageLevel::UserNotification);
+}
 
-	UpdatePolygons();
+void UConvertToPolygonsTool::UpdateOpParameters(FConvertToPolygonsOp& ConvertToPolygonsOp) const
+{
+	ConvertToPolygonsOp.bCalculateNormals = Settings->bCalculateNormals;
+	ConvertToPolygonsOp.ConversionMode    = Settings->ConversionMode;
+	ConvertToPolygonsOp.AngleTolerance    = Settings->AngleTolerance;
+	ConvertToPolygonsOp.OriginalMesh = OriginalDynamicMesh;
+	
+	FTransform LocalToWorld = ComponentTarget->GetWorldTransform();
+	ConvertToPolygonsOp.SetTransform(LocalToWorld);
 }
 
+void UConvertToPolygonsTool::GenerateAsset(const FDynamicMeshOpResult& Result)
+{
+	GetToolManager()->BeginUndoTransaction(LOCTEXT("ConvertToPolygonsToolTransactionName", "Find Polygroups"));
+
+	FDynamicMesh3* DynamicMeshResult = Result.Mesh.Get();
+	check(DynamicMeshResult != nullptr);
+
+	ComponentTarget->CommitMesh([DynamicMeshResult](const FPrimitiveComponentTarget::FCommitParams& CommitParams)
+	{
+		FDynamicMeshToMeshDescription Converter;
+		Converter.Convert(DynamicMeshResult, *CommitParams.MeshDescription);
+	});
+
+	GetToolManager()->EndUndoTransaction();
+}
 
 void UConvertToPolygonsTool::Shutdown(EToolShutdownType ShutdownType)
 {
 	Settings->SaveProperties(this);
-
-	PreviewMesh->SetVisible(false);
-	PreviewMesh->Disconnect();
-	PreviewMesh = nullptr;
-
 	ComponentTarget->SetOwnerVisibility(true);
 
-	if (ShutdownType == EToolShutdownType::Accept)
+	if (PreviewWithBackgroundCompute)
 	{
-		GetToolManager()->BeginUndoTransaction(LOCTEXT("ConvertToPolygonsToolTransactionName", "Find Polygroups"));
-		ComponentTarget->CommitMesh([=](const FPrimitiveComponentTarget::FCommitParams& CommitParams)
+		FDynamicMeshOpResult Result = PreviewWithBackgroundCompute->Shutdown();
+		if (ShutdownType == EToolShutdownType::Accept)
 		{
-			ConvertToPolygons(CommitParams.MeshDescription);
-		});
-		GetToolManager()->EndUndoTransaction();
+			GenerateAsset(Result);
+		}
 	}
 }
 
-
-void UConvertToPolygonsTool::OnPropertyModified(UObject* PropertySet, FProperty* Property)
+void UConvertToPolygonsTool::OnSettingsModified()
 {
-	bPolygonsValid = false;
-	GetToolManager()->PostInvalidation();
+	PreviewWithBackgroundCompute->InvalidateResult();
 }
 
 
 void UConvertToPolygonsTool::OnTick(float DeltaTime)
 {
-	if (bPolygonsValid == false)
-	{
-		UpdatePolygons();
-	}
+	PreviewWithBackgroundCompute->Tick(DeltaTime);
 }
 
 void UConvertToPolygonsTool::Render(IToolsContextRenderAPI* RenderAPI)
 {
 	FColor LineColor(255, 0, 0);
 
-	if (true)
-	{
-		FPrimitiveDrawInterface* PDI = RenderAPI->GetPrimitiveDrawInterface();
-		float PDIScale = RenderAPI->GetCameraState().GetPDIScalingFactor();
-		FTransform Transform = ComponentTarget->GetWorldTransform(); //Actor->GetTransform();
-
-		TArray<int>& Edges = Polygons.PolygonEdges;
-		for (int eid : Edges)
-		{
-			FVector3d A, B;
-			Polygons.Mesh->GetEdgeV(eid, A, B);
-			PDI->DrawLine(Transform.TransformPosition((FVector)A), Transform.TransformPosition((FVector)B),
-				LineColor, 0, 2.0*PDIScale, 1.0f, true);
-		}
-	}
-}
-
-
-void UConvertToPolygonsTool::UpdatePolygons()
-{
-	Polygons = FFindPolygonsAlgorithm(&SearchMesh);
-	switch (Settings->ConversionMode)
-	{
-		case EConvertToPolygonsMode::FromUVIslands:
-		{
-			Polygons.FindPolygonsFromUVIslands();
-			break;
-		}
-		case EConvertToPolygonsMode::FromConnectedTris:
-		{
-			Polygons.FindPolygonsFromConnectedTris();
-			break;
-		}
-		case EConvertToPolygonsMode::FaceNormalDeviation:
-		{
-			double DotTolerance = 1.0 - FMathd::Cos(Settings->AngleTolerance * FMathd::DegToRad);
-			Polygons.FindPolygonsFromFaceNormals(DotTolerance);
-			break;
-		}
-		default:
-			check(0);
-	}
-	
-	Polygons.FindPolygonEdges();
-
-	GetToolManager()->DisplayMessage(
-		FText::Format(LOCTEXT("UpdatePolygonsMessage", "Found {0} Polygroups in {1} Triangles"),
-			FText::AsNumber(Polygons.FoundPolygons.Num()), FText::AsNumber(SearchMesh.TriangleCount())), EToolMessageLevel::Internal);
+	FPrimitiveDrawInterface* PDI = RenderAPI->GetPrimitiveDrawInterface();
+	float PDIScale = RenderAPI->GetCameraState().GetPDIScalingFactor();
+	FTransform Transform = ComponentTarget->GetWorldTransform(); //Actor->GetTransform();
 
-	if (Settings->bCalculateNormals && Settings->ConversionMode == EConvertToPolygonsMode::FaceNormalDeviation)
+	for (int eid : PolygonEdges)
 	{
-		if (SearchMesh.HasAttributes() == false)
-		{
-			SearchMesh.EnableAttributes();
-		}
-
-		FDynamicMeshNormalOverlay* NormalOverlay = SearchMesh.Attributes()->PrimaryNormals();
-		NormalOverlay->ClearElements();
-
-		FDynamicMeshEditor Editor(&SearchMesh);
-		for (const TArray<int>& Polygon : Polygons.FoundPolygons)
-		{
-			FVector3f Normal = (FVector3f)SearchMesh.GetTriNormal(Polygon[0]);
-			Editor.SetTriangleNormals(Polygon, Normal);
-		}
-
-		FMeshNormals Normals(&SearchMesh);
-		Normals.RecomputeOverlayNormals(SearchMesh.Attributes()->PrimaryNormals());
-		Normals.CopyToOverlay(NormalOverlay, false);
+		FVector3d A, B;
+		OriginalDynamicMesh->GetEdgeV(eid, A, B);
+		PDI->DrawLine(Transform.TransformPosition((FVector)A), Transform.TransformPosition((FVector)B),
+			LineColor, 0, 2.0*PDIScale, 1.0f, true);
 	}
-
-	PreviewMesh->UpdatePreview(&SearchMesh);
-	PreviewMesh->SetVisible(true);
-	ComponentTarget->SetOwnerVisibility(false);
-
-	bPolygonsValid = true;
 }
 
 
-
-void UConvertToPolygonsTool::ConvertToPolygons(FMeshDescription* MeshIn)
+void UConvertToPolygonsTool::UpdateVisualization()
 {
-	// restore input normals
-	if (Settings->bCalculateNormals == false || Settings->ConversionMode != EConvertToPolygonsMode::FaceNormalDeviation)
+	
+	if (!PreviewWithBackgroundCompute)
 	{
-		SearchMesh.Attributes()->PrimaryNormals()->Copy(InitialNormals);
+		return;
 	}
 
-	FDynamicMeshToMeshDescription Converter;
-	Converter.ConversionOptions.bSetPolyGroups = true;
-	Converter.Convert(&SearchMesh, *MeshIn);
-}
-
-
-
-void UConvertToPolygonsTool::UpdateVisualization()
-{
+	FComponentMaterialSet MaterialSet;
 	if (Settings->bShowGroupColors)
 	{
-		PreviewMesh->SetOverrideRenderMaterial(ToolSetupUtil::GetSelectionMaterial(GetToolManager()));
-		PreviewMesh->SetTriangleColorFunction([this](const FDynamicMesh3* Mesh, int TriangleID)
+		int32 NumMaterials = ComponentTarget->GetNumMaterials();
+		for (int32 i = 0; i < NumMaterials; ++i)
+		{ 
+			MaterialSet.Materials.Add(ToolSetupUtil::GetSelectionMaterial(GetToolManager()));
+		}
+		PreviewWithBackgroundCompute->PreviewMesh->SetTriangleColorFunction([this](const FDynamicMesh3* Mesh, int TriangleID)
 		{
 			return LinearColors::SelectFColor(Mesh->GetTriangleGroup(TriangleID));
 		}, 
@@ -246,9 +338,11 @@ void UConvertToPolygonsTool::UpdateVisualization()
 	}
 	else
 	{
-		PreviewMesh->ClearOverrideRenderMaterial();
-		PreviewMesh->ClearTriangleColorFunction(UPreviewMesh::ERenderUpdateMode::FastUpdate);
+		ComponentTarget->GetMaterialSet(MaterialSet);
+		PreviewWithBackgroundCompute->PreviewMesh->ClearTriangleColorFunction(UPreviewMesh::ERenderUpdateMode::FastUpdate);
 	}
+	PreviewWithBackgroundCompute->ConfigureMaterials(MaterialSet.Materials,
+		ToolSetupUtil::GetDefaultWorkingMaterial(GetToolManager()));
 }
 
 
diff --git a/Engine/Plugins/Experimental/MeshModelingToolset/Source/MeshModelingTools/Public/ConvertToPolygonsTool.h b/Engine/Plugins/Experimental/MeshModelingToolset/Source/MeshModelingTools/Public/ConvertToPolygonsTool.h
--- a/Engine/Plugins/Experimental/MeshModelingToolset/Source/MeshModelingTools/Public/ConvertToPolygonsTool.h
+++ b/Engine/Plugins/Experimental/MeshModelingToolset/Source/MeshModelingTools/Public/ConvertToPolygonsTool.h
@@ -9,10 +9,13 @@
 #include "DynamicMesh3.h"
 #include "FindPolygonsAlgorithm.h"
 #include "PreviewMesh.h"
+#include "ModelingOperators.h"
+#include "MeshOpPreviewHelpers.h"
 #include "ConvertToPolygonsTool.generated.h"
 
 // predeclaration
-struct FMeshDescription;
+class UConvertToPolygonsTool;
+class FConvertToPolygonsOp;
 
 /**
  *
@@ -66,6 +69,19 @@ public:
 	bool bShowGroupColors = true;
 };
 
+UCLASS()
+class MESHMODELINGTOOLS_API UConvertToPolygonsOperatorFactory : public UObject, public IDynamicMeshOperatorFactory
+{
+	GENERATED_BODY()
+
+public:
+	// IDynamicMeshOperatorFactory API
+	virtual TUniquePtr<FDynamicMeshOperator> MakeNewOperator() override;
+
+	UPROPERTY()
+	UConvertToPolygonsTool* ConvertToPolygonsTool;  // back pointer
+};
+
 /**
  *
  */
@@ -80,31 +96,39 @@ public:
 	virtual void Setup() override;
 	virtual void Shutdown(EToolShutdownType ShutdownType) override;
 
+	virtual void SetWorld(UWorld* World);
+
 	virtual void Render(IToolsContextRenderAPI* RenderAPI) override;
 	virtual void OnTick(float DeltaTime) override;
 
 	virtual bool HasCancel() const override { return true; }
 	virtual bool HasAccept() const override { return true; }
+	virtual bool CanAccept() const override;
 
-	virtual void OnPropertyModified(UObject* PropertySet, FProperty* Property) override;
+	// update parameters in ConvertToPolygonsOp based on current Settings
+	void UpdateOpParameters(FConvertToPolygonsOp& ConvertToPolygonsOp) const;
+
+protected:
+	
+	void OnSettingsModified();
 
 protected:
 	UPROPERTY()
 	UConvertToPolygonsToolProperties* Settings;
 
 	UPROPERTY()
-	UPreviewMesh* PreviewMesh;
+	UMeshOpPreviewWithBackgroundCompute* PreviewWithBackgroundCompute = nullptr;
 
 protected:
-	FDynamicMesh3 SearchMesh;
-	FDynamicMeshNormalOverlay InitialNormals;
 
-	FFindPolygonsAlgorithm Polygons;
-	bool bPolygonsValid = false;
-	void UpdatePolygons();
+	UWorld* TargetWorld;
 
-	void UpdateVisualization();
+	TSharedPtr<FDynamicMesh3, ESPMode::ThreadSafe> OriginalDynamicMesh;
 
-	void ConvertToPolygons(FMeshDescription* MeshIn);
+	// for visualization
+	TArray<int> PolygonEdges;
+	
+	void UpdateVisualization();
+	void GenerateAsset(const FDynamicMeshOpResult& Result);
 
 };
