compute vertex normals when needed for simplify

#rnx

diff --git a/Engine/Plugins/Experimental/MeshModelingToolset/Source/ModelingOperatorsEditorOnly/Private/CleaningOps/SimplifyMeshOp.cpp b/Engine/Plugins/Experimental/MeshModelingToolset/Source/ModelingOperatorsEditorOnly/Private/CleaningOps/SimplifyMeshOp.cpp
--- a/Engine/Plugins/Experimental/MeshModelingToolset/Source/ModelingOperatorsEditorOnly/Private/CleaningOps/SimplifyMeshOp.cpp
+++ b/Engine/Plugins/Experimental/MeshModelingToolset/Source/ModelingOperatorsEditorOnly/Private/CleaningOps/SimplifyMeshOp.cpp
@@ -140,6 +140,10 @@ void FSimplifyMeshOp::CalculateResult(FProgressCancel* Progress)
 	{
 		bool bUseQuadricMemory = false;
 		ResultMesh->Copy(*OriginalMesh, true, true, true, !bDiscardAttributes);
+		if (!ResultMesh->HasAttributes() && !ResultMesh->HasVertexNormals())
+		{
+			FMeshNormals::QuickComputeVertexNormals(*ResultMesh, false);
+		}
 		ComputeSimplify<FAttrMeshSimplification>(TargetMesh, bReproject, OriginalTriCount, *OriginalMesh, *OriginalMeshSpatial,
 													MeshBoundaryConstraint,
 													GroupBoundaryConstraint,
@@ -152,6 +156,10 @@ void FSimplifyMeshOp::CalculateResult(FProgressCancel* Progress)
 	{
 		bool bUseQuadricMemory = false;
 		ResultMesh->Copy(*OriginalMesh, true, true, true, !bDiscardAttributes);
+		if (!ResultMesh->HasAttributes() && !ResultMesh->HasVertexNormals())
+		{
+			FMeshNormals::QuickComputeVertexNormals(*ResultMesh, false);
+		}
 		ComputeSimplify<FAttrMeshSimplification>(TargetMesh, bReproject, OriginalTriCount, *OriginalMesh, *OriginalMeshSpatial,
 			MeshBoundaryConstraint,
 			GroupBoundaryConstraint,
