ModelingMode: add ability to restrict VertexSculpt brush to a single mesh connected-component, or a single polygroup. Expose backface-hits option in GroupPaintTool.
#rnx

diff --git a/Engine/Plugins/Experimental/MeshModelingToolset/Source/MeshModelingTools/Private/MeshGroupPaintTool.cpp b/Engine/Plugins/Experimental/MeshModelingToolset/Source/MeshModelingTools/Private/MeshGroupPaintTool.cpp
--- a/Engine/Plugins/Experimental/MeshModelingToolset/Source/MeshModelingTools/Private/MeshGroupPaintTool.cpp
+++ b/Engine/Plugins/Experimental/MeshModelingToolset/Source/MeshModelingTools/Private/MeshGroupPaintTool.cpp
@@ -157,8 +157,11 @@ void UMeshGroupPaintTool::Setup()
 		[this](EMeshGroupPaintInteractionType NewType) { UpdateSubToolType(NewType); });
 	FilterProperties->WatchProperty(FilterProperties->BrushSize,
 		[this](float NewSize) { UMeshSculptToolBase::BrushProperties->BrushSize = NewSize; });
+	FilterProperties->WatchProperty(FilterProperties->bHitBackFaces,
+		[this](bool bNewValue) { UMeshSculptToolBase::BrushProperties->bHitBackFaces = bNewValue; });
 	FilterProperties->RestoreProperties(this);
 	FilterProperties->BrushSize = UMeshSculptToolBase::BrushProperties->BrushSize;
+	FilterProperties->bHitBackFaces = UMeshSculptToolBase::BrushProperties->bHitBackFaces;
 	AddToolPropertySource(FilterProperties);
 
 	InitializeIndicator();
@@ -1490,6 +1493,8 @@ void UMeshGroupPaintTool::UpdateActiveGroupLayer()
 
 void UMeshGroupPaintTool::UpdateSubToolType(EMeshGroupPaintInteractionType NewType)
 {
+	// Currenly we mirror base-brush properties in UGroupPaintBrushFilterProperties, so we never
+	// want to show both
 	//bool bSculptPropsVisible = (NewType == EMeshGroupPaintInteractionType::Brush);
 	//SetToolPropertySourceEnabled(UMeshSculptToolBase::BrushProperties, bSculptPropsVisible);
 	SetToolPropertySourceEnabled(UMeshSculptToolBase::BrushProperties, false);
diff --git a/Engine/Plugins/Experimental/MeshModelingToolset/Source/MeshModelingTools/Private/MeshVertexSculptTool.cpp b/Engine/Plugins/Experimental/MeshModelingToolset/Source/MeshModelingTools/Private/MeshVertexSculptTool.cpp
--- a/Engine/Plugins/Experimental/MeshModelingToolset/Source/MeshModelingTools/Private/MeshVertexSculptTool.cpp
+++ b/Engine/Plugins/Experimental/MeshModelingToolset/Source/MeshModelingTools/Private/MeshVertexSculptTool.cpp
@@ -5,6 +5,7 @@
 #include "InteractiveGizmoManager.h"
 #include "Async/ParallelFor.h"
 #include "Async/Async.h"
+#include "Selections/MeshConnectedComponents.h"
 
 #include "MeshWeights.h"
 #include "MeshNormals.h"
@@ -72,35 +73,75 @@ void UMeshVertexSculptTool::Setup()
 	DynamicMeshComponent->bInvalidateProxyOnChange = false;
 	OnDynamicMeshComponentChangedHandle = DynamicMeshComponent->OnMeshVerticesChanged.AddUObject(this, &UMeshVertexSculptTool::OnDynamicMeshComponentChanged);
 
+	FDynamicMesh3* SculptMesh = GetSculptMesh();
+	FAxisAlignedBox3d Bounds = SculptMesh->GetBounds(true);
+
 	// initialize dynamic octree
-	FDynamicMesh3* Mesh = GetSculptMesh();
-	FAxisAlignedBox3d Bounds = Mesh->GetCachedBounds();
-	if (Mesh->TriangleCount() > 100000)
+	TFuture<void> InitializeOctree = Async(VertexSculptToolAsyncExecTarget, [SculptMesh, Bounds, this]()
 	{
-		Octree.RootDimension = Bounds.MaxDim() / 10.0;
-		Octree.SetMaxTreeDepth(4);
-	} 
-	else
+		if (SculptMesh->TriangleCount() > 100000)
+		{
+			Octree.RootDimension = Bounds.MaxDim() / 10.0;
+			Octree.SetMaxTreeDepth(4);
+		}
+		else
+		{
+			Octree.RootDimension = Bounds.MaxDim() / 2.0;
+			Octree.SetMaxTreeDepth(8);
+		}
+		Octree.Initialize(SculptMesh);
+		//Octree.CheckValidity(EValidityCheckFailMode::Check, true, true);
+		//FDynamicMeshOctree3::FStatistics Stats;
+		//Octree.ComputeStatistics(Stats);
+		//UE_LOG(LogTemp, Warning, TEXT("Octree Stats: %s"), *Stats.ToString());
+	});
+
+	// find mesh connected-component index for each triangle
+	TFuture<void> InitializeComponents = Async(VertexSculptToolAsyncExecTarget, [SculptMesh, this]()
 	{
-		Octree.RootDimension = Bounds.MaxDim() / 2.0;
-		Octree.SetMaxTreeDepth(8);
-	}
-	Octree.Initialize(Mesh);
-	//Octree.CheckValidity(EValidityCheckFailMode::Check, true, true);
-	//FDynamicMeshOctree3::FStatistics Stats;
-	//Octree.ComputeStatistics(Stats);
-	//UE_LOG(LogTemp, Warning, TEXT("Octree Stats: %s"), *Stats.ToString());
+		TriangleComponentIDs.SetNum(SculptMesh->MaxTriangleID());
+		FMeshConnectedComponents Components(SculptMesh);
+		Components.FindConnectedTriangles();
+		int32 ComponentIdx = 1;
+		for (const FMeshConnectedComponents::FComponent& Component : Components)
+		{
+			for (int32 TriIdx : Component.Indices)
+			{
+				TriangleComponentIDs[TriIdx] = ComponentIdx;
+			}
+			ComponentIdx++;
+		}
+	});
 
-	// initialize render decomposition
-	TUniquePtr<FMeshRenderDecomposition> Decomp = MakeUnique<FMeshRenderDecomposition>();
-	FMeshRenderDecomposition::BuildChunkedDecomposition(Mesh, &MaterialSet, *Decomp);
-	Decomp->BuildAssociations(Mesh);
-	//UE_LOG(LogTemp, Warning, TEXT("Decomposition has %d groups"), Decomp->Num());
-	DynamicMeshComponent->SetExternalDecomposition(MoveTemp(Decomp));
+	// currently only supporting default polygroup set
+	TFuture<void> InitializeGroups = Async(VertexSculptToolAsyncExecTarget, [SculptMesh, this]()
+	{
+		ActiveGroupSet = MakeUnique<UE::Geometry::FPolygroupSet>(SculptMesh);
+	});
 
 	// initialize target mesh
-	UpdateBaseMesh(nullptr);
-	bTargetDirty = false;
+	TFuture<void> InitializeBaseMesh = Async(VertexSculptToolAsyncExecTarget, [this]()
+	{
+		UpdateBaseMesh(nullptr);
+		bTargetDirty = false;
+	});
+
+	// initialize render decomposition
+	TFuture<void> InitializeRenderDecomp = Async(VertexSculptToolAsyncExecTarget, [SculptMesh, &MaterialSet, this]()
+	{
+		TUniquePtr<FMeshRenderDecomposition> Decomp = MakeUnique<FMeshRenderDecomposition>();
+		FMeshRenderDecomposition::BuildChunkedDecomposition(SculptMesh, &MaterialSet, *Decomp);
+		Decomp->BuildAssociations(SculptMesh);
+		//UE_LOG(LogTemp, Warning, TEXT("Decomposition has %d groups"), Decomp->Num());
+		DynamicMeshComponent->SetExternalDecomposition(MoveTemp(Decomp));
+	});
+
+	// Wait for above precomputations to finish before continuing
+	InitializeOctree.Wait();
+	InitializeComponents.Wait();
+	InitializeGroups.Wait();
+	InitializeBaseMesh.Wait();
+	InitializeRenderDecomp.Wait();
 
 	// initialize brush radius range interval, brush properties
 	UMeshSculptToolBase::InitializeBrushSizeRange(Bounds);
@@ -276,6 +317,9 @@ void UMeshVertexSculptTool::OnBeginStroke(const FRay& WorldRay)
 	LastStamp.Power = GetActivePressure() * GetCurrentBrushStrength();
 	LastStamp.TimeStamp = FDateTime::Now();
 
+	InitialStrokeTriangleID = -1;
+	InitialStrokeTriangleID = GetBrushTriangleID();
+
 	FSculptBrushOptions SculptOptions;
 	//SculptOptions.bPreserveUVFlow = false; // SculptProperties->bPreserveUVFlow;
 	SculptOptions.ConstantReferencePlane = GetCurrentStrokeReferencePlane();
@@ -320,6 +364,18 @@ void UMeshVertexSculptTool::UpdateROI(const FVector3d& BrushPos)
 		Octree.ParallelRangeQuery(BrushBox, RangeQueryTriBuffer);
 	}
 
+	int32 ActiveComponentID = -1;
+	int32 ActiveGroupID = -1;
+	if (SculptProperties->BrushFilter == EMeshVertexSculptBrushFilterType::Component)
+	{
+		ActiveComponentID = (InitialStrokeTriangleID >= 0 && InitialStrokeTriangleID <= TriangleComponentIDs.Num()) ?
+			TriangleComponentIDs[InitialStrokeTriangleID] : -1;
+	}
+	else if (SculptProperties->BrushFilter == EMeshVertexSculptBrushFilterType::PolyGroup)
+	{
+		ActiveGroupID = Mesh->IsTriangle(InitialStrokeTriangleID) ? ActiveGroupSet->GetGroup(InitialStrokeTriangleID) : -1;
+	}
+
 #if 1
 	// in this path we use more memory but this lets us do more in parallel
 
@@ -332,7 +388,25 @@ void UMeshVertexSculptTool::UpdateROI(const FVector3d& BrushPos)
 		TRACE_CPUPROFILER_EVENT_SCOPE(DynamicMeshSculptTool_UpdateROI_TriVerts);
 		ParallelFor(RangeQueryTriBuffer.Num(), [&](int k)
 		{
-			const FIndex3i& TriV = Mesh->GetTriangleRef(RangeQueryTriBuffer[k]);
+			// check various triangle ROI filters
+			int32 tid = RangeQueryTriBuffer[k];
+			bool bDiscardTriangle = false;
+			if (ActiveComponentID >= 0 && TriangleComponentIDs[tid] != ActiveComponentID)
+			{
+				bDiscardTriangle = true;
+			}
+			if (ActiveGroupID >= 0 && ActiveGroupSet->GetGroup(tid) != ActiveGroupID)
+			{
+				bDiscardTriangle = true;
+			}
+			if (bDiscardTriangle)
+			{
+				TriangleROIInBuf[k].A = TriangleROIInBuf[k].B = TriangleROIInBuf[k].C = 0;
+				RangeQueryTriBuffer[k] = -1;
+				return;
+			}
+
+			const FIndex3i& TriV = Mesh->GetTriangleRef(tid);
 			TriangleROIInBuf[k].A = (BrushPos.DistanceSquared(Mesh->GetVertexRef(TriV.A)) < RadiusSqr) ? 1 : 0;
 			TriangleROIInBuf[k].B = (BrushPos.DistanceSquared(Mesh->GetVertexRef(TriV.B)) < RadiusSqr) ? 1 : 0;
 			TriangleROIInBuf[k].C = (BrushPos.DistanceSquared(Mesh->GetVertexRef(TriV.C)) < RadiusSqr) ? 1 : 0;
@@ -343,7 +417,7 @@ void UMeshVertexSculptTool::UpdateROI(const FVector3d& BrushPos)
 		});
 	}
 
-	// collect set of vertices inside brush sphere, from that box
+	// Build up vertex and triangle ROIs from the remaining range-query triangles.
 	{
 		TRACE_CPUPROFILER_EVENT_SCOPE(DynamicMeshSculptTool_UpdateROI_3Collect);
 		VertexROIBuilder.Initialize(Mesh->MaxVertexID());
@@ -377,6 +451,9 @@ void UMeshVertexSculptTool::UpdateROI(const FVector3d& BrushPos)
 	// In this path we combine everything into one loop. Does fewer distance checks
 	// but nothing can be done in parallel (would change if ROIBuilders had atomic-try-add)
 
+	// TODO would need to support these, this branch is likely dead though
+	ensure(SculptProperties->BrushFilter == EMeshVertexSculptBrushFilterType::None);
+
 	// collect set of vertices and triangles inside brush sphere, from range query result
 	{
 		TRACE_CPUPROFILER_EVENT_SCOPE(DynamicMeshSculptTool_UpdateROI_2Collect);
diff --git a/Engine/Plugins/Experimental/MeshModelingToolset/Source/MeshModelingTools/Public/MeshGroupPaintTool.h b/Engine/Plugins/Experimental/MeshModelingToolset/Source/MeshModelingTools/Public/MeshGroupPaintTool.h
--- a/Engine/Plugins/Experimental/MeshModelingToolset/Source/MeshModelingTools/Public/MeshGroupPaintTool.h
+++ b/Engine/Plugins/Experimental/MeshModelingToolset/Source/MeshModelingTools/Public/MeshGroupPaintTool.h
@@ -133,6 +133,11 @@ public:
 		HideEditConditionToggle, EditConditionHides, EditCondition = "SubToolType == EMeshGroupPaintInteractionType::Brush"))
 	EMeshGroupPaintBrushAreaType BrushAreaMode = EMeshGroupPaintBrushAreaType::Connected;
 
+	/** Allow the Brush to hit the back-side of the mesh */
+	UPROPERTY(EditAnywhere, Category = ActionType, meta = (DisplayName = "Hit Back Faces",
+		HideEditConditionToggle, EditConditionHides, EditCondition = "SubToolType == EMeshGroupPaintInteractionType::Brush"))
+	bool bHitBackFaces = true;
+
 	/** The group that will be assigned to triangles */
 	UPROPERTY(EditAnywhere, Category = ActionType)
 	int32 SetGroup = 1;
diff --git a/Engine/Plugins/Experimental/MeshModelingToolset/Source/MeshModelingTools/Public/MeshVertexSculptTool.h b/Engine/Plugins/Experimental/MeshModelingToolset/Source/MeshModelingTools/Public/MeshVertexSculptTool.h
--- a/Engine/Plugins/Experimental/MeshModelingToolset/Source/MeshModelingTools/Public/MeshVertexSculptTool.h
+++ b/Engine/Plugins/Experimental/MeshModelingToolset/Source/MeshModelingTools/Public/MeshVertexSculptTool.h
@@ -15,6 +15,7 @@
 #include "Sculpting/MeshBrushOpBase.h"
 #include "Image/ImageBuilder.h"
 #include "Util/UniqueIndexSet.h"
+#include "Polygroups/PolygroupSet.h"
 #include "MeshVertexSculptTool.generated.h"
 
 class UTransformGizmo;
@@ -100,6 +101,18 @@ enum class EMeshVertexSculptBrushType : uint8
 
 
 
+/** Brush Triangle Filter Type */
+UENUM()
+enum class EMeshVertexSculptBrushFilterType : uint8
+{
+	/** Do not filter brush area */
+	None,
+	/** Only apply brush to triangles in the same connected mesh component/island */
+	Component,
+	/** Only apply brush to triangles with the same PolyGroup */
+	PolyGroup
+};
+
 
 
 UCLASS()
@@ -112,10 +125,14 @@ public:
 	UPROPERTY(EditAnywhere, Category = Sculpting, meta = (DisplayName = "Brush Type"))
 	EMeshVertexSculptBrushType PrimaryBrushType = EMeshVertexSculptBrushType::Offset;
 
-	/** Primary Brush Falloff */
+	/** Primary Brush Falloff Type, multiplied by Alpha Mask where applicable */
 	UPROPERTY(EditAnywhere, Category = Sculpting, meta = (DisplayName = "Falloff Shape"))
 	EMeshSculptFalloffType PrimaryFalloffType = EMeshSculptFalloffType::Smooth;
 
+	/** Filter applied to Stamp Region Triangles, based on first Stroke Stamp */
+	UPROPERTY(EditAnywhere, Category = Sculpting, meta = (DisplayName = "Region Filter"))
+	EMeshVertexSculptBrushFilterType BrushFilter = EMeshVertexSculptBrushFilterType::None;
+
 	/** When Freeze Target is toggled on, the Brush Target Surface will be Frozen in its current state, until toggled off. Brush strokes will be applied relative to the Target Surface, for applicable Brushes */
 	UPROPERTY(EditAnywhere, Category = Sculpting, meta = (EditCondition = "PrimaryBrushType == EMeshVertexSculptBrushType::Offset || PrimaryBrushType == EMeshVertexSculptBrushType::SculptMax || PrimaryBrushType == EMeshVertexSculptBrushType::SculptView || PrimaryBrushType == EMeshVertexSculptBrushType::Pinch || PrimaryBrushType == EMeshVertexSculptBrushType::Resample" ))
 	bool bFreezeTarget = false;
@@ -220,6 +237,11 @@ protected:
 
 	void UpdateBrushType(EMeshVertexSculptBrushType BrushType);
 
+	TUniquePtr<UE::Geometry::FPolygroupSet> ActiveGroupSet;
+	TArray<int32> TriangleComponentIDs;
+
+	int32 InitialStrokeTriangleID = -1;
+
 	TSet<int32> AccumulatedTriangleROI;
 	bool bUndoUpdatePending = false;
 	TFuture<bool> UndoNormalsFuture;
