try removing Win64 platform whitelist from modeling toolset modules that shouldn't need it
+ fix minor issues uncovered by doing so

#rnx

diff --git a/Engine/Plugins/Experimental/MeshModelingToolset/MeshModelingToolset.uplugin b/Engine/Plugins/Experimental/MeshModelingToolset/MeshModelingToolset.uplugin
--- a/Engine/Plugins/Experimental/MeshModelingToolset/MeshModelingToolset.uplugin
+++ b/Engine/Plugins/Experimental/MeshModelingToolset/MeshModelingToolset.uplugin
@@ -18,14 +18,12 @@
 		{
 			"Name": "ModelingComponents",
 			"Type": "Runtime",
-			"LoadingPhase": "Default",
-			"WhitelistPlatforms": [ "Win64" ]
+			"LoadingPhase": "Default"
 		},
 		{
 			"Name": "ModelingOperators",
 			"Type": "Runtime",
-			"LoadingPhase": "Default",
-			"WhitelistPlatforms": [ "Win64" ]
+			"LoadingPhase": "Default"
 		},
 		{
 			"Name": "ModelingOperatorsEditorOnly",
@@ -36,8 +34,7 @@
 		{
 			"Name": "MeshModelingTools",
 			"Type": "Runtime",
-			"LoadingPhase": "Default",
-			"WhitelistPlatforms": [ "Win64" ]
+			"LoadingPhase": "Default"
 		},
 		{
 			"Name": "MeshModelingToolsEditorOnly",
diff --git a/Engine/Plugins/Experimental/MeshModelingToolset/Source/MeshModelingTools/Private/ConvertToPolygonsTool.cpp b/Engine/Plugins/Experimental/MeshModelingToolset/Source/MeshModelingTools/Private/ConvertToPolygonsTool.cpp
--- a/Engine/Plugins/Experimental/MeshModelingToolset/Source/MeshModelingTools/Private/ConvertToPolygonsTool.cpp
+++ b/Engine/Plugins/Experimental/MeshModelingToolset/Source/MeshModelingTools/Private/ConvertToPolygonsTool.cpp
@@ -53,7 +53,7 @@ public:
 
 	virtual void CalculateResult(FProgressCancel* Progress) override
 	{
-		if (Progress && Progress->Cancelled() || !OriginalMesh)
+		if ((Progress && Progress->Cancelled()) || !OriginalMesh)
 		{
 			return;
 		}
diff --git a/Engine/Plugins/Experimental/MeshModelingToolset/Source/ModelingComponents/Private/OctreeDynamicMeshComponent.cpp b/Engine/Plugins/Experimental/MeshModelingToolset/Source/ModelingComponents/Private/OctreeDynamicMeshComponent.cpp
--- a/Engine/Plugins/Experimental/MeshModelingToolset/Source/ModelingComponents/Private/OctreeDynamicMeshComponent.cpp
+++ b/Engine/Plugins/Experimental/MeshModelingToolset/Source/ModelingComponents/Private/OctreeDynamicMeshComponent.cpp
@@ -271,7 +271,7 @@ FPrimitiveSceneProxy* UOctreeDynamicMeshComponent::CreateSceneProxy()
 		if (TriangleColorFunc != nullptr)
 		{
 			NewProxy->bUsePerTriangleColor = true;
-			NewProxy->PerTriangleColorFunc = [this](const FDynamicMesh3* Mesh, int TriangleID) { return GetTriangleColor(TriangleID); };
+			NewProxy->PerTriangleColorFunc = [this](const FDynamicMesh3* MeshIn, int TriangleID) { return GetTriangleColor(TriangleID); };
 		}
 
 
diff --git a/Engine/Plugins/Experimental/MeshModelingToolset/Source/ModelingComponents/Public/AssetGenerationUtil.h b/Engine/Plugins/Experimental/MeshModelingToolset/Source/ModelingComponents/Public/AssetGenerationUtil.h
--- a/Engine/Plugins/Experimental/MeshModelingToolset/Source/ModelingComponents/Public/AssetGenerationUtil.h
+++ b/Engine/Plugins/Experimental/MeshModelingToolset/Source/ModelingComponents/Public/AssetGenerationUtil.h
@@ -14,7 +14,7 @@ class AActor;
 class UPrimitiveComponent;
 class UTexture2D;
 
-MODELINGCOMPONENTS_API class IAssetGenerationAPI
+class MODELINGCOMPONENTS_API IAssetGenerationAPI
 {
 public:
 	virtual ~IAssetGenerationAPI(){}
diff --git a/Engine/Plugins/Experimental/MeshModelingToolset/Source/ModelingComponents/Public/AssetUtils/MeshDescriptionUtil.h b/Engine/Plugins/Experimental/MeshModelingToolset/Source/ModelingComponents/Public/AssetUtils/MeshDescriptionUtil.h
--- a/Engine/Plugins/Experimental/MeshModelingToolset/Source/ModelingComponents/Public/AssetUtils/MeshDescriptionUtil.h
+++ b/Engine/Plugins/Experimental/MeshModelingToolset/Source/ModelingComponents/Public/AssetUtils/MeshDescriptionUtil.h
@@ -46,7 +46,7 @@ namespace UE
 		//
 
 		/** How to update a boolean build setting */
-		MODELINGCOMPONENTS_API enum class EBuildSettingBoolChange
+		enum class EBuildSettingBoolChange
 		{
 			Disable,
 			Enable,
@@ -54,7 +54,7 @@ namespace UE
 		};
 
 		/** Set of changes to apply to StaticMesh Build Settings */
-		MODELINGCOMPONENTS_API struct FStaticMeshBuildSettingChange
+		struct FStaticMeshBuildSettingChange
 		{
 			EBuildSettingBoolChange AutoGeneratedNormals = EBuildSettingBoolChange::NoChange;
 			EBuildSettingBoolChange AutoGeneratedTangents = EBuildSettingBoolChange::NoChange;
diff --git a/Engine/Plugins/Experimental/MeshModelingToolset/Source/ModelingComponents/Public/Physics/ComponentCollisionUtil.h b/Engine/Plugins/Experimental/MeshModelingToolset/Source/ModelingComponents/Public/Physics/ComponentCollisionUtil.h
--- a/Engine/Plugins/Experimental/MeshModelingToolset/Source/ModelingComponents/Public/Physics/ComponentCollisionUtil.h
+++ b/Engine/Plugins/Experimental/MeshModelingToolset/Source/ModelingComponents/Public/Physics/ComponentCollisionUtil.h
@@ -19,7 +19,7 @@ namespace Geometry
 /**
  * Component/BodySetup collision settings (eg StaticMeshComponent) we might need to pass through the functions below
  */
-MODELINGCOMPONENTS_API struct FComponentCollisionSettings
+struct FComponentCollisionSettings
 {
 	int32 CollisionTypeFlag = 0;		// this is ECollisionTraceFlag
 	bool bIsGeneratedCollision = true;
diff --git a/Engine/Plugins/Experimental/MeshModelingToolset/Source/ModelingOperators/Private/MeshSpaceDeformerOps/MeshSpaceDeformerOp.cpp b/Engine/Plugins/Experimental/MeshModelingToolset/Source/ModelingOperators/Private/MeshSpaceDeformerOps/MeshSpaceDeformerOp.cpp
--- a/Engine/Plugins/Experimental/MeshModelingToolset/Source/ModelingOperators/Private/MeshSpaceDeformerOps/MeshSpaceDeformerOp.cpp
+++ b/Engine/Plugins/Experimental/MeshModelingToolset/Source/ModelingOperators/Private/MeshSpaceDeformerOps/MeshSpaceDeformerOp.cpp
@@ -10,7 +10,7 @@ void FMeshSpaceDeformerOp::SetTransform(const FTransform& Transform)
 
 void FMeshSpaceDeformerOp::CalculateResult(FProgressCancel* Progress)
 {
-	if (Progress && Progress->Cancelled() || !OriginalMesh)
+	if ((Progress && Progress->Cancelled()) || !OriginalMesh)
 	{
 		return;
 	}
