fixes for ci errors from non-windows compile of modeling toolset modules

#rnx

diff --git a/Engine/Plugins/Experimental/MeshModelingToolset/Source/MeshModelingTools/Private/FindPolygonsAlgorithm.cpp b/Engine/Plugins/Experimental/MeshModelingToolset/Source/MeshModelingTools/Private/FindPolygonsAlgorithm.cpp
--- a/Engine/Plugins/Experimental/MeshModelingToolset/Source/MeshModelingTools/Private/FindPolygonsAlgorithm.cpp
+++ b/Engine/Plugins/Experimental/MeshModelingToolset/Source/MeshModelingTools/Private/FindPolygonsAlgorithm.cpp
@@ -141,10 +141,10 @@ void FFindPolygonsAlgorithm::SetGroupsFromPolygons()
 
 		// find a normal if the average failed
 		AccumNormal.Normalize();
-		int pi = 0;
-		while (AccumNormal.Length() < 0.9 && pi < NumTriangles)
+		int SubIdx = 0;
+		while (AccumNormal.Length() < 0.9 && SubIdx < NumTriangles)
 		{
-			AccumNormal = Mesh->GetTriNormal(Polygon[pi++]);
+			AccumNormal = Mesh->GetTriNormal(Polygon[SubIdx++]);
 		}
 		if (AccumNormal.Length() < 0.9)
 		{
diff --git a/Engine/Plugins/Experimental/MeshModelingToolset/Source/ModelingOperators/Private/MeshSpaceDeformerOps/BendMeshOp.cpp b/Engine/Plugins/Experimental/MeshModelingToolset/Source/ModelingOperators/Private/MeshSpaceDeformerOps/BendMeshOp.cpp
--- a/Engine/Plugins/Experimental/MeshModelingToolset/Source/ModelingOperators/Private/MeshSpaceDeformerOps/BendMeshOp.cpp
+++ b/Engine/Plugins/Experimental/MeshModelingToolset/Source/ModelingOperators/Private/MeshSpaceDeformerOps/BendMeshOp.cpp
@@ -1,5 +1,5 @@
 // Copyright Epic Games, Inc. All Rights Reserved.
-#include "SpaceDeformerOps\BendMeshOp.h"
+#include "SpaceDeformerOps/BendMeshOp.h"
 
 #include "Async/ParallelFor.h"
 #include "DynamicMeshAttributeSet.h"
diff --git a/Engine/Plugins/Experimental/MeshModelingToolset/Source/ModelingOperators/Private/MeshSpaceDeformerOps/FlareMeshOp.cpp b/Engine/Plugins/Experimental/MeshModelingToolset/Source/ModelingOperators/Private/MeshSpaceDeformerOps/FlareMeshOp.cpp
--- a/Engine/Plugins/Experimental/MeshModelingToolset/Source/ModelingOperators/Private/MeshSpaceDeformerOps/FlareMeshOp.cpp
+++ b/Engine/Plugins/Experimental/MeshModelingToolset/Source/ModelingOperators/Private/MeshSpaceDeformerOps/FlareMeshOp.cpp
@@ -1,5 +1,5 @@
 // Copyright Epic Games, Inc. All Rights Reserved.
-#include "SpaceDeformerOps\FlareMeshOp.h"
+#include "SpaceDeformerOps/FlareMeshOp.h"
 
 #include "Async/ParallelFor.h"
 #include "DynamicMeshAttributeSet.h"
diff --git a/Engine/Plugins/Experimental/MeshModelingToolset/Source/ModelingOperators/Private/MeshSpaceDeformerOps/MeshSpaceDeformerOp.cpp b/Engine/Plugins/Experimental/MeshModelingToolset/Source/ModelingOperators/Private/MeshSpaceDeformerOps/MeshSpaceDeformerOp.cpp
--- a/Engine/Plugins/Experimental/MeshModelingToolset/Source/ModelingOperators/Private/MeshSpaceDeformerOps/MeshSpaceDeformerOp.cpp
+++ b/Engine/Plugins/Experimental/MeshModelingToolset/Source/ModelingOperators/Private/MeshSpaceDeformerOps/MeshSpaceDeformerOp.cpp
@@ -1,6 +1,6 @@
 // Copyright Epic Games, Inc. All Rights Reserved.
 
-#include "SpaceDeformerOps\MeshSpaceDeformerOp.h"
+#include "SpaceDeformerOps/MeshSpaceDeformerOp.h"
 #include "DynamicMesh3.h"
 
 void FMeshSpaceDeformerOp::SetTransform(const FTransform& Transform)
diff --git a/Engine/Plugins/Experimental/MeshModelingToolset/Source/ModelingOperators/Private/MeshSpaceDeformerOps/TwistMeshOp.cpp b/Engine/Plugins/Experimental/MeshModelingToolset/Source/ModelingOperators/Private/MeshSpaceDeformerOps/TwistMeshOp.cpp
--- a/Engine/Plugins/Experimental/MeshModelingToolset/Source/ModelingOperators/Private/MeshSpaceDeformerOps/TwistMeshOp.cpp
+++ b/Engine/Plugins/Experimental/MeshModelingToolset/Source/ModelingOperators/Private/MeshSpaceDeformerOps/TwistMeshOp.cpp
@@ -1,6 +1,6 @@
 // Copyright Epic Games, Inc. All Rights Reserved.
 
-#include "SpaceDeformerOps\TwistMeshOp.h"
+#include "SpaceDeformerOps/TwistMeshOp.h"
 
 #include "Async/ParallelFor.h"
 #include "DynamicMeshAttributeSet.h"
