fix crash when modeling preview-with-background-compute has a null preview (encountered when accepting a boolean result with a weird material)

#rnx

diff --git a/Engine/Plugins/Experimental/MeshModelingToolset/Source/ModelingComponents/Public/MeshOpPreviewHelpers.h b/Engine/Plugins/Experimental/MeshModelingToolset/Source/ModelingComponents/Public/MeshOpPreviewHelpers.h
--- a/Engine/Plugins/Experimental/MeshModelingToolset/Source/ModelingComponents/Public/MeshOpPreviewHelpers.h
+++ b/Engine/Plugins/Experimental/MeshModelingToolset/Source/ModelingComponents/Public/MeshOpPreviewHelpers.h
@@ -97,12 +97,12 @@ public:
 	/**
 	 * @return true if current PreviewMesh result is valid (no update actively being computed) and that mesh has at least one triangle
 	 */
-	bool HaveValidNonEmptyResult() const { return bResultValid && PreviewMesh->GetMesh() && PreviewMesh->GetMesh()->TriangleCount() > 0; }
+	bool HaveValidNonEmptyResult() const { return bResultValid && PreviewMesh && PreviewMesh->GetMesh() && PreviewMesh->GetMesh()->TriangleCount() > 0; }
 
 	/**
 	 * @return true if current PreviewMesh result is valid (no update actively being computed) but that mesh has no triangles
 	 */
-	bool HaveEmptyResult() const { return bResultValid && PreviewMesh->GetMesh() && PreviewMesh->GetMesh()->TriangleCount() == 0; }
+	bool HaveEmptyResult() const { return bResultValid && PreviewMesh && PreviewMesh->GetMesh() && PreviewMesh->GetMesh()->TriangleCount() == 0; }
 
 
 	/**
