MeshModelingTools: Fix PolyEdit edge ring selection crash on open meshes.

#rnx

diff --git a/Engine/Plugins/Experimental/MeshModelingToolset/Source/ModelingComponents/Private/Selection/GroupTopologySelector.cpp b/Engine/Plugins/Experimental/MeshModelingToolset/Source/ModelingComponents/Private/Selection/GroupTopologySelector.cpp
--- a/Engine/Plugins/Experimental/MeshModelingToolset/Source/ModelingComponents/Private/Selection/GroupTopologySelector.cpp
+++ b/Engine/Plugins/Experimental/MeshModelingToolset/Source/ModelingComponents/Private/Selection/GroupTopologySelector.cpp
@@ -785,7 +785,7 @@ void AddNewEdgeRingEdges(const FGroupTopology& Topology, int32 StartEdgeID, int3
 	int32 CurrentForwardGroupID = ForwardGroupID;
 	while (true)
 	{
-		if (ForwardGroupID == IndexConstants::InvalidID)
+		if (CurrentForwardGroupID == IndexConstants::InvalidID)
 		{
 			break; // Reached a boundary
 		}
