GeometryProcessing: move all public types in GeometryProcessing plugin into UE::Geometry:: namespace, as well as some code in MeshModelingToolset. Fix fallout in Water, Paint Mode, ChaosEditor, DataPr

diff --git a/Engine/Plugins/Experimental/MeshModelingToolset/Source/MeshModelingTools/Public/BakeMeshAttributeMapsTool.h b/Engine/Plugins/Experimental/MeshModelingToolset/Source/MeshModelingTools/Public/BakeMeshAttributeMapsTool.h
--- a/Engine/Plugins/Experimental/MeshModelingToolset/Source/MeshModelingTools/Public/BakeMeshAttributeMapsTool.h
+++ b/Engine/Plugins/Experimental/MeshModelingToolset/Source/MeshModelingTools/Public/BakeMeshAttributeMapsTool.h
@@ -18,8 +18,9 @@ class USimpleDynamicMeshComponent;
 class UMaterialInstanceDynamic;
 class UTexture2D;
 class IAssetGenerationAPI;
-template<typename RealType> class TMeshTangents;
-class FMeshImageBakingCache;
+PREDECLARE_GEOMETRY(template<typename RealType> class TMeshTangents);
+PREDECLARE_GEOMETRY(class FMeshImageBakingCache);
+using UE::Geometry::FImageDimensions;
 class IPrimitiveComponentBackedTarget;
 class IMeshDescriptionProvider;
 class IMaterialProvider;
@@ -307,14 +308,14 @@ protected:
 	UMaterialInstanceDynamic* PreviewMaterial;
 
 	TSharedPtr<FMeshDescription, ESPMode::ThreadSafe> BaseMeshDescription;
-	TSharedPtr<TMeshTangents<double>, ESPMode::ThreadSafe> BaseMeshTangents;
-	FDynamicMesh3 BaseMesh;
-	FDynamicMeshAABBTree3 BaseSpatial;
+	TSharedPtr<UE::Geometry::TMeshTangents<double>, ESPMode::ThreadSafe> BaseMeshTangents;
+	UE::Geometry::FDynamicMesh3 BaseMesh;
+	UE::Geometry::FDynamicMeshAABBTree3 BaseSpatial;
 
 	bool bIsBakeToSelf = false;
 
-	TSharedPtr<FDynamicMesh3, ESPMode::ThreadSafe> DetailMesh;
-	TSharedPtr<FDynamicMeshAABBTree3, ESPMode::ThreadSafe> DetailSpatial;
+	TSharedPtr<UE::Geometry::FDynamicMesh3, ESPMode::ThreadSafe> DetailMesh;
+	TSharedPtr<UE::Geometry::FDynamicMeshAABBTree3, ESPMode::ThreadSafe> DetailSpatial;
 	int32 DetailMeshTimestamp = 0;
 	void UpdateDetailMesh();
 	bool bDetailMeshValid = false;
@@ -325,7 +326,7 @@ protected:
 	void UpdateOnModeChange();
 	void UpdateVisualization();
 
-	TPimplPtr<FMeshImageBakingCache> BakeCache;
+	TPimplPtr<UE::Geometry::FMeshImageBakingCache> BakeCache;
 	struct FBakeCacheSettings
 	{
 		FImageDimensions Dimensions;
diff --git a/Engine/Plugins/Experimental/MeshModelingToolset/Source/MeshModelingTools/Public/CSGMeshesTool.h b/Engine/Plugins/Experimental/MeshModelingToolset/Source/MeshModelingTools/Public/CSGMeshesTool.h
--- a/Engine/Plugins/Experimental/MeshModelingToolset/Source/MeshModelingTools/Public/CSGMeshesTool.h
+++ b/Engine/Plugins/Experimental/MeshModelingToolset/Source/MeshModelingTools/Public/CSGMeshesTool.h
@@ -13,7 +13,7 @@
 #include "CSGMeshesTool.generated.h"
 
 // predeclarations
-class FDynamicMesh3;
+PREDECLARE_USE_GEOMETRY_CLASS(FDynamicMesh3);
 
 
 /**
@@ -93,7 +93,7 @@ protected:
 	virtual FText GetActionName() const;
 
 	// IDynamicMeshOperatorFactory API
-	virtual TUniquePtr<FDynamicMeshOperator> MakeNewOperator() override;
+	virtual TUniquePtr<UE::Geometry::FDynamicMeshOperator> MakeNewOperator() override;
 
 protected:
 
diff --git a/Engine/Plugins/Experimental/MeshModelingToolset/Source/MeshModelingTools/Public/ConversionUtils/VolumeToDynamicMesh.h b/Engine/Plugins/Experimental/MeshModelingToolset/Source/MeshModelingTools/Public/ConversionUtils/VolumeToDynamicMesh.h
--- a/Engine/Plugins/Experimental/MeshModelingToolset/Source/MeshModelingTools/Public/ConversionUtils/VolumeToDynamicMesh.h
+++ b/Engine/Plugins/Experimental/MeshModelingToolset/Source/MeshModelingTools/Public/ConversionUtils/VolumeToDynamicMesh.h
@@ -7,7 +7,7 @@
 #include "VectorTypes.h"
 
 class AVolume;
-class FDynamicMesh3;
+PREDECLARE_GEOMETRY(class FDynamicMesh3);
 
 namespace UE {
 namespace Conversion {
@@ -24,6 +24,6 @@ struct FVolumeToMeshOptions
 /**
  * Converts a volume to a dynamic mesh. Does not initialize normals and does not delete the volume.
  */
-void MESHMODELINGTOOLS_API VolumeToDynamicMesh(AVolume* Volume, FDynamicMesh3& Mesh, const FVolumeToMeshOptions& Options);
+void MESHMODELINGTOOLS_API VolumeToDynamicMesh(AVolume* Volume, UE::Geometry::FDynamicMesh3& Mesh, const FVolumeToMeshOptions& Options);
 
 }}//end namespace UE::Conversion
\ No newline at end of file
diff --git a/Engine/Plugins/Experimental/MeshModelingToolset/Source/MeshModelingTools/Public/ConvertToPolygonsTool.h b/Engine/Plugins/Experimental/MeshModelingToolset/Source/MeshModelingTools/Public/ConvertToPolygonsTool.h
--- a/Engine/Plugins/Experimental/MeshModelingToolset/Source/MeshModelingTools/Public/ConvertToPolygonsTool.h
+++ b/Engine/Plugins/Experimental/MeshModelingToolset/Source/MeshModelingTools/Public/ConvertToPolygonsTool.h
@@ -70,13 +70,13 @@ public:
 };
 
 UCLASS()
-class MESHMODELINGTOOLS_API UConvertToPolygonsOperatorFactory : public UObject, public IDynamicMeshOperatorFactory
+class MESHMODELINGTOOLS_API UConvertToPolygonsOperatorFactory : public UObject, public UE::Geometry::IDynamicMeshOperatorFactory
 {
 	GENERATED_BODY()
 
 public:
 	// IDynamicMeshOperatorFactory API
-	virtual TUniquePtr<FDynamicMeshOperator> MakeNewOperator() override;
+	virtual TUniquePtr<UE::Geometry::FDynamicMeshOperator> MakeNewOperator() override;
 
 	UPROPERTY()
 	UConvertToPolygonsTool* ConvertToPolygonsTool;  // back pointer
@@ -123,7 +123,7 @@ protected:
 
 	UWorld* TargetWorld;
 
-	TSharedPtr<FDynamicMesh3, ESPMode::ThreadSafe> OriginalDynamicMesh;
+	TSharedPtr<UE::Geometry::FDynamicMesh3, ESPMode::ThreadSafe> OriginalDynamicMesh;
 
 	// for visualization
 	TArray<int> PolygonEdges;
diff --git a/Engine/Plugins/Experimental/MeshModelingToolset/Source/MeshModelingTools/Public/DeformMeshPolygonsTool.h b/Engine/Plugins/Experimental/MeshModelingToolset/Source/MeshModelingTools/Public/DeformMeshPolygonsTool.h
--- a/Engine/Plugins/Experimental/MeshModelingToolset/Source/MeshModelingTools/Public/DeformMeshPolygonsTool.h
+++ b/Engine/Plugins/Experimental/MeshModelingToolset/Source/MeshModelingTools/Public/DeformMeshPolygonsTool.h
@@ -88,6 +88,7 @@ UCLASS()
 class MESHMODELINGTOOLS_API UDeformMeshPolygonsTransformProperties : public UInteractiveToolPropertySet
 {
 	GENERATED_BODY()
+	using FVector3d = UE::Geometry::FVector3d;
 
 public:
 	UDeformMeshPolygonsTransformProperties();
@@ -148,6 +149,7 @@ UCLASS()
 class MESHMODELINGTOOLS_API UDeformMeshPolygonsTool : public UMeshSurfacePointTool
 {
 	GENERATED_BODY()
+	using FVector3d = UE::Geometry::FVector3d;
 
 public:
 	UDeformMeshPolygonsTool();
@@ -204,8 +206,8 @@ protected:
 	FVector LastBrushPosLocal;
 	FVector StartBrushPosLocal;
 
-	FFrame3d ActiveSurfaceFrame;
-	FQuickTransformer* GetActiveQuickTransformer();
+	UE::Geometry::FFrame3d ActiveSurfaceFrame;
+	UE::Geometry::FQuickTransformer* GetActiveQuickTransformer();
 	void UpdateActiveSurfaceFrame(FGroupTopologySelection& Selection);
 	void UpdateQuickTransformer();
 
@@ -214,12 +216,12 @@ protected:
 	void ComputeUpdate();
 
 	FVector3d LastMoveDelta;
-	FQuickAxisTranslater QuickAxisTranslater;
+	UE::Geometry::FQuickAxisTranslater QuickAxisTranslater;
 	void ComputeUpdate_Translate();
 
-	FQuickAxisRotator QuickAxisRotator;
+	UE::Geometry::FQuickAxisRotator QuickAxisRotator;
 	FVector3d RotationStartPointWorld;
-	FFrame3d RotationStartFrame;
+	UE::Geometry::FFrame3d RotationStartFrame;
 	void ComputeUpdate_Rotate();
 
 	FGroupTopology Topology;
@@ -235,15 +237,15 @@ protected:
 	FGroupTopologySelection HilightSelection;
 	FToolDataVisualizer HilightRenderer;
 
-	FDynamicMeshAABBTree3 MeshSpatial;
-	FDynamicMeshAABBTree3& GetSpatial();
+	UE::Geometry::FDynamicMeshAABBTree3 MeshSpatial;
+	UE::Geometry::FDynamicMeshAABBTree3& GetSpatial();
 
 	FMeshVertexChangeBuilder* ActiveVertexChange;
 
 	EGroupTopologyDeformationStrategy DeformationStrategy;
 
 	// The two deformer type options.
-	FGroupTopologyDeformer LinearDeformer;
+	UE::Geometry::FGroupTopologyDeformer LinearDeformer;
 	TPimplPtr<FGroupTopologyLaplacianDeformer> LaplacianDeformer;
 
 
diff --git a/Engine/Plugins/Experimental/MeshModelingToolset/Source/MeshModelingTools/Public/DisplaceMeshTool.h b/Engine/Plugins/Experimental/MeshModelingToolset/Source/MeshModelingTools/Public/DisplaceMeshTool.h
--- a/Engine/Plugins/Experimental/MeshModelingToolset/Source/MeshModelingTools/Public/DisplaceMeshTool.h
+++ b/Engine/Plugins/Experimental/MeshModelingToolset/Source/MeshModelingTools/Public/DisplaceMeshTool.h
@@ -240,7 +240,7 @@ UCLASS()
 class MESHMODELINGTOOLS_API UDisplaceMeshTool : public USingleSelectionTool
 {
 	GENERATED_BODY()
-
+	using FVector3d = UE::Geometry::FVector3d;
 public:
 	void Setup() final;
 	void Shutdown(EToolShutdownType ShutdownType) final;
@@ -278,20 +278,20 @@ private:
 	void AdvanceComputation();
 	void ValidateSubdivisions();
 
-	FAsyncTaskExecuterWithAbort<TModelingOpTask<FDynamicMeshOperator>>* SubdivideTask = nullptr;
+	UE::Geometry::FAsyncTaskExecuterWithAbort<UE::Geometry::TModelingOpTask<UE::Geometry::FDynamicMeshOperator>>* SubdivideTask = nullptr;
 	bool bNeedsSubdivided = true;
-	FAsyncTaskExecuterWithAbort<TModelingOpTask<FDynamicMeshOperator>>* DisplaceTask = nullptr;
+	UE::Geometry::FAsyncTaskExecuterWithAbort<UE::Geometry::TModelingOpTask<UE::Geometry::FDynamicMeshOperator>>* DisplaceTask = nullptr;
 	bool bNeedsDisplaced = true;
 
-	FDynamicMesh3 OriginalMesh;
-	FDynamicMeshAABBTree3 OriginalMeshSpatial;
-	TSharedPtr<FIndexedWeightMap, ESPMode::ThreadSafe> ActiveWeightMap;
+	UE::Geometry::FDynamicMesh3 OriginalMesh;
+	UE::Geometry::FDynamicMeshAABBTree3 OriginalMeshSpatial;
+	TSharedPtr<UE::Geometry::FIndexedWeightMap, ESPMode::ThreadSafe> ActiveWeightMap;
 	void UpdateActiveWeightMap();
-	float WeightMapQuery(const FVector3d& Position, const FIndexedWeightMap& WeightMap) const;
+	float WeightMapQuery(const FVector3d& Position, const UE::Geometry::FIndexedWeightMap& WeightMap) const;
 
-	TSharedPtr<FDynamicMesh3, ESPMode::ThreadSafe> SubdividedMesh = nullptr;
+	TSharedPtr<UE::Geometry::FDynamicMesh3, ESPMode::ThreadSafe> SubdividedMesh = nullptr;
 	USimpleDynamicMeshComponent* DynamicMeshComponent = nullptr;
 
-	TUniquePtr<IDynamicMeshOperatorFactory> Subdivider = nullptr;
-	TUniquePtr<IDynamicMeshOperatorFactory> Displacer = nullptr;
+	TUniquePtr<UE::Geometry::IDynamicMeshOperatorFactory> Subdivider = nullptr;
+	TUniquePtr<UE::Geometry::IDynamicMeshOperatorFactory> Displacer = nullptr;
 };
diff --git a/Engine/Plugins/Experimental/MeshModelingToolset/Source/MeshModelingTools/Public/DrawAndRevolveTool.h b/Engine/Plugins/Experimental/MeshModelingToolset/Source/MeshModelingTools/Public/DrawAndRevolveTool.h
--- a/Engine/Plugins/Experimental/MeshModelingToolset/Source/MeshModelingTools/Public/DrawAndRevolveTool.h
+++ b/Engine/Plugins/Experimental/MeshModelingToolset/Source/MeshModelingTools/Public/DrawAndRevolveTool.h
@@ -18,7 +18,6 @@
 class UCollectSurfacePathMechanic;
 class UConstructionPlaneMechanic;
 class UCurveControlPointsMechanic;
-class FCurveSweepOp;
 class IAssetGenerationAPI;
 
 UCLASS()
@@ -70,13 +69,13 @@ public:
 };
 
 UCLASS()
-class MESHMODELINGTOOLS_API URevolveOperatorFactory : public UObject, public IDynamicMeshOperatorFactory
+class MESHMODELINGTOOLS_API URevolveOperatorFactory : public UObject, public UE::Geometry::IDynamicMeshOperatorFactory
 {
 	GENERATED_BODY()
 
 public:
 	// IDynamicMeshOperatorFactory API
-	virtual TUniquePtr<FDynamicMeshOperator> MakeNewOperator() override;
+	virtual TUniquePtr<UE::Geometry::FDynamicMeshOperator> MakeNewOperator() override;
 
 	UPROPERTY()
 	UDrawAndRevolveTool* RevolveTool;
@@ -117,8 +116,8 @@ protected:
 
 	// This information is replicated in the user-editable transform in the settings and in the PlaneMechanic
 	// plane, but the tool turned out to be much easier to write and edit with this decoupling.
-	FVector3d RevolutionAxisOrigin;
-	FVector3d RevolutionAxisDirection;
+	UE::Geometry::FVector3d RevolutionAxisOrigin;
+	UE::Geometry::FVector3d RevolutionAxisDirection;
 
 	bool bProfileCurveComplete = false;
 
diff --git a/Engine/Plugins/Experimental/MeshModelingToolset/Source/MeshModelingTools/Public/DrawPolyPathTool.h b/Engine/Plugins/Experimental/MeshModelingToolset/Source/MeshModelingTools/Public/DrawPolyPathTool.h
--- a/Engine/Plugins/Experimental/MeshModelingToolset/Source/MeshModelingTools/Public/DrawPolyPathTool.h
+++ b/Engine/Plugins/Experimental/MeshModelingToolset/Source/MeshModelingTools/Public/DrawPolyPathTool.h
@@ -130,7 +130,6 @@ UCLASS()
 class MESHMODELINGTOOLS_API UDrawPolyPathTool : public UInteractiveTool, public IClickBehaviorTarget, public IHoverBehaviorTarget
 {
 	GENERATED_BODY()
-
 public:
 	virtual void RegisterActions(FInteractiveToolActionSet& ActionSet) override;
 
@@ -174,7 +173,7 @@ protected:
 
 protected:
 	// camera state at last render
-	FTransform3d WorldTransform;
+	UE::Geometry::FTransform3d WorldTransform;
 	FViewCameraState CameraState;
 
 	// drawing plane and gizmo
@@ -182,17 +181,17 @@ protected:
 	UPROPERTY()
 	UConstructionPlaneMechanic* PlaneMechanic = nullptr;
 
-	FFrame3d DrawPlaneWorld;
+	UE::Geometry::FFrame3d DrawPlaneWorld;
 
 	bool CanUpdateDrawPlane() const;
 
 	// UV Scale factor to apply to texturing on any new geometry (e.g. new faces added by extrude)
 	float UVScaleFactor = 1.0f;
 
-	TArray<FFrame3d> CurPathPoints;
+	TArray<UE::Geometry::FFrame3d> CurPathPoints;
 	TArray<double> OffsetScaleFactors;
 	TArray<double> ArcLengths;
-	TArray<FVector3d> CurPolyLine;
+	TArray<UE::Geometry::FVector3d> CurPolyLine;
 	double CurPathLength;
 	double CurOffsetDistance;
 	double CurHeight;
@@ -221,8 +220,8 @@ protected:
 	void UpdateExtrudePreview();
 	void InitializePreviewMesh();
 	void ClearPreview();
-	void GeneratePathMesh(FDynamicMesh3& Mesh);
-	void GenerateExtrudeMesh(FDynamicMesh3& PathMesh);
+	void GeneratePathMesh(UE::Geometry::FDynamicMesh3& Mesh);
+	void GenerateExtrudeMesh(UE::Geometry::FDynamicMesh3& PathMesh);
 	void GenerateRampMesh(FDynamicMesh3& PathMesh);
 	void EmitNewObject(EDrawPolyPathOutputMode OutputMode);
 
diff --git a/Engine/Plugins/Experimental/MeshModelingToolset/Source/MeshModelingTools/Public/DrawPolygonTool.h b/Engine/Plugins/Experimental/MeshModelingToolset/Source/MeshModelingTools/Public/DrawPolygonTool.h
--- a/Engine/Plugins/Experimental/MeshModelingToolset/Source/MeshModelingTools/Public/DrawPolygonTool.h
+++ b/Engine/Plugins/Experimental/MeshModelingToolset/Source/MeshModelingTools/Public/DrawPolygonTool.h
@@ -19,6 +19,7 @@
 class UTransformGizmo;
 class UTransformProxy;
 class IAssetGenerationAPI;
+PREDECLARE_USE_GEOMETRY_CLASS(FDynamicMesh3);
 
 /**
  *
@@ -166,7 +167,8 @@ UCLASS()
 class MESHMODELINGTOOLS_API UDrawPolygonTool : public UInteractiveTool, public IClickSequenceBehaviorTarget
 {
 	GENERATED_BODY()
-
+	using FVector3d = UE::Geometry::FVector3d;
+	using FVector2d = UE::Geometry::FVector2d;
 public:
 	UDrawPolygonTool();
 
@@ -233,7 +235,7 @@ protected:
 	FVector3d DrawPlaneOrigin;
 
 	/** Orientation of plane we will draw polygon on */
-	FQuaterniond DrawPlaneOrientation;
+	UE::Geometry::FQuaterniond DrawPlaneOrientation;
 	
 	/** Vertices of current preview polygon */
 	TArray<FVector3d> PolygonVertices;
@@ -293,7 +295,7 @@ protected:
 	FVector3d SurfaceOffsetPoint;
 
 	bool bIgnoreSnappingToggle = false;		// toggled by hotkey (shift)
-	FPointPlanarSnapSolver SnapEngine;
+	UE::Geometry::FPointPlanarSnapSolver SnapEngine;
 	ToolSceneQueriesUtil::FSnapGeometry LastSnapGeometry;
 	FVector3d LastGridSnapPoint;
 
@@ -312,10 +314,10 @@ protected:
 	UPlaneDistanceFromHitMechanic* HeightMechanic;
 
 
-	FFrame3f HitPosFrameWorld;
+	UE::Geometry::FFrame3f HitPosFrameWorld;
 
 	/** Generate extruded meshes.  Returns true on success. */
-	bool GeneratePolygonMesh(const TArray<FVector3d>& Polygon, const TArray<TArray<FVector3d>>& PolygonHoles, FDynamicMesh3* ResultMeshOut, FFrame3d& WorldFrameOut, bool bIncludePreviewVtx, double ExtrudeDistance, bool bExtrudeSymmetric);
+	bool GeneratePolygonMesh(const TArray<FVector3d>& Polygon, const TArray<TArray<FVector3d>>& PolygonHoles, FDynamicMesh3* ResultMeshOut, UE::Geometry::FFrame3d& WorldFrameOut, bool bIncludePreviewVtx, double ExtrudeDistance, bool bExtrudeSymmetric);
 
 
 	// user feedback messages
@@ -336,7 +338,7 @@ protected:
 class MESHMODELINGTOOLS_API FDrawPolygonStateChange : public FToolCommandChange
 {
 public:
-	using Points = TArray<FVector3d>;
+	using Points = TArray<UE::Geometry::FVector3d>;
 	bool bHaveDoneUndo = false;
 	int32 CurveTimestamp = 0;
 	const Points FixedVertexPoints;
diff --git a/Engine/Plugins/Experimental/MeshModelingToolset/Source/MeshModelingTools/Public/DynamicMeshBrushTool.h b/Engine/Plugins/Experimental/MeshModelingToolset/Source/MeshModelingTools/Public/DynamicMeshBrushTool.h
--- a/Engine/Plugins/Experimental/MeshModelingToolset/Source/MeshModelingTools/Public/DynamicMeshBrushTool.h
+++ b/Engine/Plugins/Experimental/MeshModelingToolset/Source/MeshModelingTools/Public/DynamicMeshBrushTool.h
@@ -11,6 +11,7 @@
 #include "PreviewMesh.h"
 #include "DynamicMeshBrushTool.generated.h"
 
+PREDECLARE_USE_GEOMETRY_CLASS(FDynamicMesh3);
 
 /**
  * UDynamicMeshBrushTool is a base class that specializes UBaseBrushTool
@@ -47,7 +48,7 @@ protected:
 	virtual void OnBaseMeshComponentChanged() {}	
 	FDelegateHandle OnBaseMeshComponentChangedHandle;
 
-	FAxisAlignedBox3d InputMeshBoundsLocal;
+	UE::Geometry::FAxisAlignedBox3d InputMeshBoundsLocal;
 
 	//
 	// UBaseBrushTool private interface
diff --git a/Engine/Plugins/Experimental/MeshModelingToolset/Source/MeshModelingTools/Public/DynamicMeshSculptTool.h b/Engine/Plugins/Experimental/MeshModelingToolset/Source/MeshModelingTools/Public/DynamicMeshSculptTool.h
--- a/Engine/Plugins/Experimental/MeshModelingToolset/Source/MeshModelingTools/Public/DynamicMeshSculptTool.h
+++ b/Engine/Plugins/Experimental/MeshModelingToolset/Source/MeshModelingTools/Public/DynamicMeshSculptTool.h
@@ -28,11 +28,13 @@ class UTransformProxy;
 class UMaterialInstanceDynamic;
 
 class FMeshVertexChangeBuilder;
-class FDynamicMeshChangeTracker;
 class UPreviewMesh;
-class FSubRegionRemesher;
+PREDECLARE_GEOMETRY(class FDynamicMeshChangeTracker);
+PREDECLARE_GEOMETRY(class FSubRegionRemesher);
+PREDECLARE_USE_GEOMETRY_CLASS(FDynamicMesh3);
 class FPersistentStampRemesher;
 
+
 /** Mesh Sculpting Brush Types */
 UENUM()
 enum class EDynamicMeshSculptBrushType : uint8
@@ -224,6 +226,9 @@ class MESHMODELINGTOOLS_API UDynamicMeshSculptTool : public UMeshSurfacePointToo
 	GENERATED_BODY()
 
 public:
+	using FFrame3d = UE::Geometry::FFrame3d;
+	using FRay3d = UE::Geometry::FRay3d;
+
 	UDynamicMeshSculptTool();
 
 	virtual void SetWorld(UWorld* World);
@@ -315,8 +320,8 @@ private:
 	UPROPERTY()
 	UMaterialInstanceDynamic* ActiveOverrideMaterial;
 
-	FTransform3d InitialTargetTransform;
-	FTransform3d CurTargetTransform;
+	UE::Geometry::FTransform3d InitialTargetTransform;
+	UE::Geometry::FTransform3d CurTargetTransform;
 
 	// realtime visualization
 	void OnDynamicMeshComponentChanged();
@@ -337,15 +342,15 @@ private:
 	void UpdateBrushType(EDynamicMeshSculptBrushType BrushType);
 	void UpdateGizmoFromProperties();
 
-	FInterval1d BrushRelativeSizeRange;
+	UE::Geometry::FInterval1d BrushRelativeSizeRange;
 	double CurrentBrushRadius;
 	void CalculateBrushRadius();
 
 	bool bEnableRemeshing;
 	double InitialEdgeLength;
 	void ScheduleRemeshPass();
-	void ConfigureRemesher(FSubRegionRemesher& Remesher);
-	void InitializeRemesherROI(FSubRegionRemesher& Remesher);
+	void ConfigureRemesher(UE::Geometry::FSubRegionRemesher& Remesher);
+	void InitializeRemesherROI(UE::Geometry::FSubRegionRemesher& Remesher);
 
 	TSharedPtr<FPersistentStampRemesher> ActiveRemesher;
 	void InitializeActiveRemesher();
@@ -355,7 +360,7 @@ private:
 
 	bool bInDrag;
 
-	FFrame3d ActiveDragPlane;
+	UE::Geometry::FFrame3d ActiveDragPlane;
 	FVector3d LastHitPosWorld;
 	FVector3d BrushStartCenterWorld;
 	FVector3d BrushStartNormalWorld;
@@ -366,9 +371,9 @@ private:
 	int32 LastBrushTriangleID = -1;
 
 	TArray<int> UpdateROITriBuffer;
-	FUniqueIndexSet VertexROIBuilder;
+	UE::Geometry::FUniqueIndexSet VertexROIBuilder;
 	TArray<int> VertexROI;
-	FUniqueIndexSet TriangleROIBuilder;
+	UE::Geometry::FUniqueIndexSet TriangleROIBuilder;
 	TSet<int> TriangleROI;
 	//TSet<int> TriangleROI;
 	void UpdateROI(const FVector3d& BrushPos);
@@ -393,8 +398,8 @@ private:
 	void ApplyStamp(const FRay& WorldRay);
 
 	FDynamicMesh3 BrushTargetMesh;
-	FDynamicMeshAABBTree3 BrushTargetMeshSpatial;
-	FMeshNormals BrushTargetNormals;
+	UE::Geometry::FDynamicMeshAABBTree3 BrushTargetMeshSpatial;
+	UE::Geometry::FMeshNormals BrushTargetNormals;
 	bool bCachedFreezeTarget = false;
 	void UpdateTarget();
 	bool GetTargetMeshNearest(const FVector3d& Position, double SearchRadius, FVector3d& TargetPosOut, FVector3d& TargetNormalOut);
@@ -429,8 +434,8 @@ private:
 	TArray<FVector3d> ROIPositionBuffer;
 	void SyncMeshWithPositionBuffer(FDynamicMesh3* Mesh);
 
-	FFrame3d ActiveFixedBrushPlane;
-	FFrame3d ComputeROIBrushPlane(const FVector3d& BrushCenter, bool bIgnoreDepth, bool bViewAligned);
+	UE::Geometry::FFrame3d ActiveFixedBrushPlane;
+	UE::Geometry::FFrame3d ComputeROIBrushPlane(const FVector3d& BrushCenter, bool bIgnoreDepth, bool bViewAligned);
 
 	TArray<int> NormalsBuffer;
 	TArray<bool> NormalsVertexFlags;
@@ -446,7 +451,7 @@ private:
 	void RemeshROIPass();
 
 	FMeshVertexChangeBuilder* ActiveVertexChange = nullptr;
-	FDynamicMeshChangeTracker* ActiveMeshChange = nullptr;
+	UE::Geometry::FDynamicMeshChangeTracker* ActiveMeshChange = nullptr;
 	void BeginChange(bool bIsVertexChange);
 	void EndChange();
 	void SaveActiveROI();
diff --git a/Engine/Plugins/Experimental/MeshModelingToolset/Source/MeshModelingTools/Public/EdgeLoopInsertionTool.h b/Engine/Plugins/Experimental/MeshModelingToolset/Source/MeshModelingTools/Public/EdgeLoopInsertionTool.h
--- a/Engine/Plugins/Experimental/MeshModelingToolset/Source/MeshModelingTools/Public/EdgeLoopInsertionTool.h
+++ b/Engine/Plugins/Experimental/MeshModelingToolset/Source/MeshModelingTools/Public/EdgeLoopInsertionTool.h
@@ -101,13 +101,13 @@ public:
 };
 
 UCLASS()
-class MESHMODELINGTOOLS_API UEdgeLoopInsertionOperatorFactory : public UObject, public IDynamicMeshOperatorFactory
+class MESHMODELINGTOOLS_API UEdgeLoopInsertionOperatorFactory : public UObject, public UE::Geometry::IDynamicMeshOperatorFactory
 {
 	GENERATED_BODY()
 
 public:
 	// IDynamicMeshOperatorFactory API
-	virtual TUniquePtr<FDynamicMeshOperator> MakeNewOperator() override;
+	virtual TUniquePtr<UE::Geometry::FDynamicMeshOperator> MakeNewOperator() override;
 
 	UPROPERTY()
 	UEdgeLoopInsertionTool* Tool;
@@ -118,7 +118,7 @@ UCLASS()
 class MESHMODELINGTOOLS_API UEdgeLoopInsertionTool : public USingleSelectionTool, public IHoverBehaviorTarget, public IClickBehaviorTarget
 {
 	GENERATED_BODY()
-
+	using FVector3d = UE::Geometry::FVector3d;
 public:
 
 	friend class UEdgeLoopInsertionOperatorFactory;
@@ -158,7 +158,7 @@ protected:
 
 	TSharedPtr<FDynamicMesh3> CurrentMesh;
 	TSharedPtr<FGroupTopology> CurrentTopology;
-	FDynamicMeshAABBTree3 MeshSpatial;
+	UE::Geometry::FDynamicMeshAABBTree3 MeshSpatial;
 	FGroupTopologySelector TopologySelector;
 
 	TArray<TPair<FVector3d, FVector3d>> PreviewEdges;
diff --git a/Engine/Plugins/Experimental/MeshModelingToolset/Source/MeshModelingTools/Public/EditMeshPolygonsTool.h b/Engine/Plugins/Experimental/MeshModelingToolset/Source/MeshModelingTools/Public/EditMeshPolygonsTool.h
--- a/Engine/Plugins/Experimental/MeshModelingToolset/Source/MeshModelingTools/Public/EditMeshPolygonsTool.h
+++ b/Engine/Plugins/Experimental/MeshModelingToolset/Source/MeshModelingTools/Public/EditMeshPolygonsTool.h
@@ -476,7 +476,10 @@ UCLASS()
 class MESHMODELINGTOOLS_API UEditMeshPolygonsTool : public UMeshSurfacePointTool, public IClickBehaviorTarget
 {
 	GENERATED_BODY()
-
+	using FVector3d = UE::Geometry::FVector3d;
+	using FTransform3d = UE::Geometry::FTransform3d;
+	using FFrame3d = UE::Geometry::FFrame3d;
+	using FAxisAlignedBox3d = UE::Geometry::FAxisAlignedBox3d;
 public:
 	UEditMeshPolygonsTool();
 
@@ -595,7 +598,7 @@ protected:
 
 
 	// camera state at last render
-	FTransform3d WorldTransform;
+	UE::Geometry::FTransform3d WorldTransform;
 	FViewCameraState CameraState;
 
 	// True for the duration of UI click+drag
@@ -613,8 +616,8 @@ protected:
 	TUniquePtr<FGroupTopology> Topology;
 	void PrecomputeTopology();
 
-	FDynamicMeshAABBTree3 MeshSpatial;
-	FDynamicMeshAABBTree3& GetSpatial();
+	UE::Geometry::FDynamicMeshAABBTree3 MeshSpatial;
+	UE::Geometry::FDynamicMeshAABBTree3& GetSpatial();
 	bool bSpatialDirty;
 
 	// UV Scale factor to apply to texturing on any new geometry (e.g. new faces added by extrude)
@@ -705,7 +708,7 @@ protected:
 	// data for current drag
 	//
 
-	FGroupTopologyDeformer LinearDeformer;
+	UE::Geometry::FGroupTopologyDeformer LinearDeformer;
 	void UpdateDeformerFromSelection(const FGroupTopologySelection& Selection);
 
 
diff --git a/Engine/Plugins/Experimental/MeshModelingToolset/Source/MeshModelingTools/Public/EditNormalsTool.h b/Engine/Plugins/Experimental/MeshModelingToolset/Source/MeshModelingTools/Public/EditNormalsTool.h
--- a/Engine/Plugins/Experimental/MeshModelingToolset/Source/MeshModelingTools/Public/EditNormalsTool.h
+++ b/Engine/Plugins/Experimental/MeshModelingToolset/Source/MeshModelingTools/Public/EditNormalsTool.h
@@ -106,13 +106,13 @@ public:
  *  stores a pointer to the tool and enough info to know which specific operator it should spawn
  */
 UCLASS()
-class MESHMODELINGTOOLS_API UEditNormalsOperatorFactory : public UObject, public IDynamicMeshOperatorFactory
+class MESHMODELINGTOOLS_API UEditNormalsOperatorFactory : public UObject, public UE::Geometry::IDynamicMeshOperatorFactory
 {
 	GENERATED_BODY()
 
 public:
 	// IDynamicMeshOperatorFactory API
-	virtual TUniquePtr<FDynamicMeshOperator> MakeNewOperator() override;
+	virtual TUniquePtr<UE::Geometry::FDynamicMeshOperator> MakeNewOperator() override;
 
 	UPROPERTY()
 	UEditNormalsTool *Tool;
@@ -168,7 +168,7 @@ protected:
 
 
 protected:
-	TArray<TSharedPtr<FDynamicMesh3, ESPMode::ThreadSafe>> OriginalDynamicMeshes;
+	TArray<TSharedPtr<UE::Geometry::FDynamicMesh3, ESPMode::ThreadSafe>> OriginalDynamicMeshes;
 
 	UWorld* TargetWorld;
 	IAssetGenerationAPI* AssetAPI;
diff --git a/Engine/Plugins/Experimental/MeshModelingToolset/Source/MeshModelingTools/Public/EditPivotTool.h b/Engine/Plugins/Experimental/MeshModelingToolset/Source/MeshModelingTools/Public/EditPivotTool.h
--- a/Engine/Plugins/Experimental/MeshModelingToolset/Source/MeshModelingTools/Public/EditPivotTool.h
+++ b/Engine/Plugins/Experimental/MeshModelingToolset/Source/MeshModelingTools/Public/EditPivotTool.h
@@ -190,9 +190,9 @@ protected:
 
 	TArray<int> MapToFirstOccurrences;
 
-	FTransform3d Transform;
-	FAxisAlignedBox3d ObjectBounds;
-	FAxisAlignedBox3d WorldBounds;
+	UE::Geometry::FTransform3d Transform;
+	UE::Geometry::FAxisAlignedBox3d ObjectBounds;
+	UE::Geometry::FAxisAlignedBox3d WorldBounds;
 	void Precompute();
 
 	UPROPERTY()
@@ -211,5 +211,5 @@ protected:
 	virtual void ApplyAction(EEditPivotToolActions ActionType);
 	virtual void SetPivotToBoxPoint(EEditPivotToolActions ActionPoint);
 
-	void UpdateAssets(const FFrame3d& NewPivotWorldFrame);
+	void UpdateAssets(const UE::Geometry::FFrame3d& NewPivotWorldFrame);
 };
diff --git a/Engine/Plugins/Experimental/MeshModelingToolset/Source/MeshModelingTools/Public/EditUVIslandsTool.h b/Engine/Plugins/Experimental/MeshModelingToolset/Source/MeshModelingTools/Public/EditUVIslandsTool.h
--- a/Engine/Plugins/Experimental/MeshModelingToolset/Source/MeshModelingTools/Public/EditUVIslandsTool.h
+++ b/Engine/Plugins/Experimental/MeshModelingToolset/Source/MeshModelingTools/Public/EditUVIslandsTool.h
@@ -17,7 +17,7 @@
 #include "EditUVIslandsTool.generated.h"
 
 class FMeshVertexChangeBuilder;
-
+using UE::Geometry::FDynamicMeshUVOverlay;
 
 /**
  * ToolBuilder
@@ -40,7 +40,7 @@ public:
 	const FDynamicMeshUVOverlay* UVOverlay;
 
 	FUVGroupTopology() {}
-	FUVGroupTopology(const FDynamicMesh3* Mesh, uint32 UVLayerIndex, bool bAutoBuild = false);
+	FUVGroupTopology(const UE::Geometry::FDynamicMesh3* Mesh, uint32 UVLayerIndex, bool bAutoBuild = false);
 
 	void CalculateIslandGroups();
 
@@ -49,7 +49,7 @@ public:
 		return TriIslandGroups[TriangleID];
 	}
 
-	FFrame3d GetIslandFrame(int32 GroupID, FDynamicMeshAABBTree3& AABBTree);
+	UE::Geometry::FFrame3d GetIslandFrame(int32 GroupID, UE::Geometry::FDynamicMeshAABBTree3& AABBTree);
 };
 
 
@@ -116,15 +116,15 @@ protected:
 	FDelegateHandle OnDynamicMeshComponentChangedHandle;
 
 	// camera state at last render
-	FTransform3d WorldTransform;
+	UE::Geometry::FTransform3d WorldTransform;
 	FViewCameraState CameraState;
 
 	// True for the duration of UI click+drag
 	bool bInDrag;
 
 	double UVTranslateScale;
-	FFrame3d InitialGizmoFrame;
-	FVector3d InitialGizmoScale;
+	UE::Geometry::FFrame3d InitialGizmoFrame;
+	UE::Geometry::FVector3d InitialGizmoScale;
 	void ComputeUpdate_Gizmo();
 
 	FUVGroupTopology Topology;
@@ -140,12 +140,12 @@ protected:
 	//
 	struct FEditIsland
 	{
-		FFrame3d LocalFrame;
+		UE::Geometry::FFrame3d LocalFrame;
 		TArray<int32> Triangles;
 		TArray<int32> UVs;
-		FAxisAlignedBox2d UVBounds;
-		FVector2d UVOrigin;
-		TArray<FVector2f> InitialPositions;
+		UE::Geometry::FAxisAlignedBox2d UVBounds;
+		UE::Geometry::FVector2d UVOrigin;
+		TArray<UE::Geometry::FVector2f> InitialPositions;
 	};
 	TArray<FEditIsland> ActiveIslands;
 	void UpdateUVTransformFromSelection(const FGroupTopologySelection& Selection);
diff --git a/Engine/Plugins/Experimental/MeshModelingToolset/Source/MeshModelingTools/Public/FindPolygonsAlgorithm.h b/Engine/Plugins/Experimental/MeshModelingToolset/Source/MeshModelingTools/Public/FindPolygonsAlgorithm.h
--- a/Engine/Plugins/Experimental/MeshModelingToolset/Source/MeshModelingTools/Public/FindPolygonsAlgorithm.h
+++ b/Engine/Plugins/Experimental/MeshModelingToolset/Source/MeshModelingTools/Public/FindPolygonsAlgorithm.h
@@ -5,8 +5,11 @@
 #include "CoreMinimal.h"
 #include "DynamicMesh3.h"
 
+using UE::Geometry::FDynamicMesh3;
+
 class FFindPolygonsAlgorithm
 {
+	using FVector3d = UE::Geometry::FVector3d;
 public:
 
 	FFindPolygonsAlgorithm() {}
diff --git a/Engine/Plugins/Experimental/MeshModelingToolset/Source/MeshModelingTools/Public/GroupEdgeInsertionTool.h b/Engine/Plugins/Experimental/MeshModelingToolset/Source/MeshModelingTools/Public/GroupEdgeInsertionTool.h
--- a/Engine/Plugins/Experimental/MeshModelingToolset/Source/MeshModelingTools/Public/GroupEdgeInsertionTool.h
+++ b/Engine/Plugins/Experimental/MeshModelingToolset/Source/MeshModelingTools/Public/GroupEdgeInsertionTool.h
@@ -19,6 +19,9 @@
 
 class UDynamicMeshReplacementChangeTarget;
 
+using UE::Geometry::FGroupEdgeInserter;
+using UE::Geometry::FDynamicMesh3;
+
 UCLASS()
 class MESHMODELINGTOOLS_API UGroupEdgeInsertionToolBuilder : public UInteractiveToolBuilder
 {
@@ -61,13 +64,13 @@ public:
 };
 
 UCLASS()
-class MESHMODELINGTOOLS_API UGroupEdgeInsertionOperatorFactory : public UObject, public IDynamicMeshOperatorFactory
+class MESHMODELINGTOOLS_API UGroupEdgeInsertionOperatorFactory : public UObject, public UE::Geometry::IDynamicMeshOperatorFactory
 {
 	GENERATED_BODY()
 
 public:
 	// IDynamicMeshOperatorFactory API
-	virtual TUniquePtr<FDynamicMeshOperator> MakeNewOperator() override;
+	virtual TUniquePtr<UE::Geometry::FDynamicMeshOperator> MakeNewOperator() override;
 
 	UPROPERTY()
 	UGroupEdgeInsertionTool* Tool;
@@ -86,6 +89,8 @@ class MESHMODELINGTOOLS_API UGroupEdgeInsertionTool : public USingleSelectionToo
 		WaitingForInsertComplete
 	};
 
+	using FRay3d = UE::Geometry::FRay3d;
+	using FVector3d = UE::Geometry::FVector3d;
 public:
 
 	friend class UGroupEdgeInsertionOperatorFactory;
@@ -131,7 +136,7 @@ protected:
 
 	TSharedPtr<FDynamicMesh3, ESPMode::ThreadSafe> CurrentMesh;
 	TSharedPtr<FGroupTopology, ESPMode::ThreadSafe> CurrentTopology;
-	FDynamicMeshAABBTree3 MeshSpatial;
+	UE::Geometry::FDynamicMeshAABBTree3 MeshSpatial;
 	FGroupTopologySelector TopologySelector;
 
 	TArray<TPair<FVector3d, FVector3d>> PreviewEdges;
diff --git a/Engine/Plugins/Experimental/MeshModelingToolset/Source/MeshModelingTools/Public/HoleFillTool.h b/Engine/Plugins/Experimental/MeshModelingToolset/Source/MeshModelingTools/Public/HoleFillTool.h
--- a/Engine/Plugins/Experimental/MeshModelingToolset/Source/MeshModelingTools/Public/HoleFillTool.h
+++ b/Engine/Plugins/Experimental/MeshModelingToolset/Source/MeshModelingTools/Public/HoleFillTool.h
@@ -83,7 +83,7 @@ public:
 	USmoothHoleFillProperties()
 	{
 		// Create a default FSmoothFillOptions and populate this class with its defaults.
-		FSmoothFillOptions DefaultOptionsObject;
+		UE::Geometry::FSmoothFillOptions DefaultOptionsObject;
 		bConstrainToHoleInterior = DefaultOptionsObject.bConstrainToHoleInterior;
 		RemeshingExteriorRegionWidth = DefaultOptionsObject.RemeshingExteriorRegionWidth;
 		SmoothingExteriorRegionWidth = DefaultOptionsObject.SmoothingExteriorRegionWidth;
@@ -93,9 +93,9 @@ public:
 		bProjectDuringRemesh = DefaultOptionsObject.bProjectDuringRemesh;
 	}
 
-	FSmoothFillOptions ToSmoothFillOptions() const
+	UE::Geometry::FSmoothFillOptions ToSmoothFillOptions() const
 	{
-		FSmoothFillOptions Options;
+		UE::Geometry::FSmoothFillOptions Options;
 		Options.bConstrainToHoleInterior = bConstrainToHoleInterior;
 		Options.RemeshingExteriorRegionWidth = RemeshingExteriorRegionWidth;
 		Options.SmoothingExteriorRegionWidth = SmoothingExteriorRegionWidth;
@@ -184,7 +184,7 @@ public:
 
 	void Initialize(const UHoleFillTool& HoleFillTool);
 
-	void Update(const UHoleFillTool& HoleFillTool, const FHoleFillOp& HoleFillOp);
+	void Update(const UHoleFillTool& HoleFillTool, const UE::Geometry::FHoleFillOp& HoleFillOp);
 };
 
 
@@ -193,13 +193,13 @@ public:
  */
 
 UCLASS()
-class MESHMODELINGTOOLS_API UHoleFillOperatorFactory : public UObject, public IDynamicMeshOperatorFactory
+class MESHMODELINGTOOLS_API UHoleFillOperatorFactory : public UObject, public UE::Geometry::IDynamicMeshOperatorFactory
 {
 	GENERATED_BODY()
 
 public:
 
-	TUniquePtr<FDynamicMeshOperator> MakeNewOperator() override;
+	TUniquePtr<UE::Geometry::FDynamicMeshOperator> MakeNewOperator() override;
 
 	UPROPERTY()
 	UHoleFillTool* FillTool;
@@ -255,7 +255,7 @@ protected:
 	UPolygonSelectionMechanic* SelectionMechanic = nullptr;
 
 	// Input mesh. Ownership shared with Op.
-	TSharedPtr<FDynamicMesh3, ESPMode::ThreadSafe> OriginalMesh;
+	TSharedPtr<UE::Geometry::FDynamicMesh3, ESPMode::ThreadSafe> OriginalMesh;
 
 	// UV Scale factor is cached based on the bounding box of the mesh before any fills are performed
 	float MeshUVScaleFactor = 0.0f;
@@ -264,7 +264,7 @@ protected:
 	UWorld* TargetWorld = nullptr;
 
 	// Used for hit querying
-	FDynamicMeshAABBTree3 MeshSpatial;
+	UE::Geometry::FDynamicMeshAABBTree3 MeshSpatial;
 
 	TSet<int32> NewTriangleIDs;
 
@@ -285,10 +285,10 @@ protected:
 	// Object used to get boundary loop information. All triangles return the same GroupID, so the only boundaries that 
 	// are returned are the the actual mesh boundaries.
 	// TODO: It seems like overkill to use a FGroupTopology subclass when we don't actually care about groups.
-	class FBasicTopology : public FGroupTopology
+	class FBasicTopology : public UE::Geometry::FGroupTopology
 	{
 	public:
-		FBasicTopology(const FDynamicMesh3* Mesh, bool bAutoBuild) :
+		FBasicTopology(const UE::Geometry::FDynamicMesh3* Mesh, bool bAutoBuild) :
 			FGroupTopology(Mesh, bAutoBuild)
 		{}
 
@@ -312,6 +312,6 @@ protected:
 
 	// Populate an array of Edge Loops to be processed by an FHoleFillOp. Returns the edge loops currently selected
 	// by this tool.
-	void GetLoopsToFill(TArray<FEdgeLoop>& OutLoops) const;
+	void GetLoopsToFill(TArray<UE::Geometry::FEdgeLoop>& OutLoops) const;
 
 };
diff --git a/Engine/Plugins/Experimental/MeshModelingToolset/Source/MeshModelingTools/Public/LatticeDeformerTool.h b/Engine/Plugins/Experimental/MeshModelingToolset/Source/MeshModelingTools/Public/LatticeDeformerTool.h
--- a/Engine/Plugins/Experimental/MeshModelingToolset/Source/MeshModelingTools/Public/LatticeDeformerTool.h
+++ b/Engine/Plugins/Experimental/MeshModelingToolset/Source/MeshModelingTools/Public/LatticeDeformerTool.h
@@ -82,13 +82,13 @@ public:
 
 
 UCLASS()
-class MESHMODELINGTOOLS_API ULatticeDeformerOperatorFactory : public UObject, public IDynamicMeshOperatorFactory
+class MESHMODELINGTOOLS_API ULatticeDeformerOperatorFactory : public UObject, public UE::Geometry::IDynamicMeshOperatorFactory
 {
 	GENERATED_BODY()
 
 public:
 	// IDynamicMeshOperatorFactory API
-	virtual TUniquePtr<FDynamicMeshOperator> MakeNewOperator() override;
+	virtual TUniquePtr<UE::Geometry::FDynamicMeshOperator> MakeNewOperator() override;
 
 	UPROPERTY()
 	ULatticeDeformerTool* LatticeDeformerTool;
@@ -118,7 +118,7 @@ public:
 	virtual void OnTick(float DeltaTime) override;
 	virtual void Render(IToolsContextRenderAPI* RenderAPI) override;
 
-	FVector3i GetLatticeResolution() const;
+	UE::Geometry::FVector3i GetLatticeResolution() const;
 
 protected:
 
@@ -126,9 +126,9 @@ protected:
 	IToolsContextAssetAPI* AssetAPI;
 
 	// Input mesh
-	TSharedPtr<FDynamicMesh3, ESPMode::ThreadSafe> OriginalMesh;
+	TSharedPtr<UE::Geometry::FDynamicMesh3, ESPMode::ThreadSafe> OriginalMesh;
 
-	TSharedPtr<FFFDLattice, ESPMode::ThreadSafe> Lattice;
+	TSharedPtr<UE::Geometry::FFFDLattice, ESPMode::ThreadSafe> Lattice;
 
 	UPROPERTY()
 	ULatticeControlPointsMechanic* ControlPointsMechanic = nullptr;
@@ -145,7 +145,7 @@ protected:
 	bool bShouldRebuild = false;
 
 	// Create and store an FFFDLattice. Pass out the lattice's positions and edges.
-	void InitializeLattice(TArray<FVector3d>& OutLatticePoints, TArray<FVector2i>& OutLatticeEdges);
+	void InitializeLattice(TArray<UE::Geometry::FVector3d>& OutLatticePoints, TArray<UE::Geometry::FVector2i>& OutLatticeEdges);
 
 	void StartPreview();
 
diff --git a/Engine/Plugins/Experimental/MeshModelingToolset/Source/MeshModelingTools/Public/MeshAttributePaintTool.h b/Engine/Plugins/Experimental/MeshModelingToolset/Source/MeshModelingTools/Public/MeshAttributePaintTool.h
--- a/Engine/Plugins/Experimental/MeshModelingToolset/Source/MeshModelingTools/Public/MeshAttributePaintTool.h
+++ b/Engine/Plugins/Experimental/MeshModelingToolset/Source/MeshModelingTools/Public/MeshAttributePaintTool.h
@@ -47,7 +47,7 @@ public:
 	virtual int ElementNum() const = 0;
 	virtual float GetValue(int32 Index) const = 0;
 	virtual void SetValue(int32 Index, float Value) = 0;
-	virtual FInterval1f GetValueRange() = 0;
+	virtual UE::Geometry::FInterval1f GetValueRange() = 0;
 };
 
 
@@ -249,7 +249,7 @@ protected:
 	TUniquePtr<FMeshDescription> EditedMesh;
 
 	double CalculateBrushFalloff(double Distance);
-	TDynamicVerticesOctree3<FDynamicMesh3> VerticesOctree;
+	UE::Geometry::TDynamicVerticesOctree3<FDynamicMesh3> VerticesOctree;
 	TArray<int> PreviewBrushROI;
 	void CalculateVertexROI(const FBrushStampData& Stamp, TArray<int>& VertexROI);
 
@@ -266,7 +266,7 @@ protected:
 	TArray<FAttributeData> Attributes;
 	int32 AttributeBufferCount;
 	int32 CurrentAttributeIndex;
-	FInterval1f CurrentValueRange;
+	UE::Geometry::FInterval1f CurrentValueRange;
 
 	// actions
 
diff --git a/Engine/Plugins/Experimental/MeshModelingToolset/Source/MeshModelingTools/Public/MeshBoundaryToolBase.h b/Engine/Plugins/Experimental/MeshModelingToolset/Source/MeshModelingTools/Public/MeshBoundaryToolBase.h
--- a/Engine/Plugins/Experimental/MeshModelingToolset/Source/MeshModelingTools/Public/MeshBoundaryToolBase.h
+++ b/Engine/Plugins/Experimental/MeshModelingToolset/Source/MeshModelingTools/Public/MeshBoundaryToolBase.h
@@ -12,7 +12,7 @@
 
 #include "MeshBoundaryToolBase.generated.h"
 
-class FDynamicMesh3;
+PREDECLARE_USE_GEOMETRY_CLASS(FDynamicMesh3);
 class UPolygonSelectionMechanic;
 class USingleClickInputBehavior;
 
@@ -45,7 +45,7 @@ protected:
 	TSharedPtr<FDynamicMesh3, ESPMode::ThreadSafe> OriginalMesh;
 
 	// Used for hit querying
-	FDynamicMeshAABBTree3 MeshSpatial;
+	UE::Geometry::FDynamicMeshAABBTree3 MeshSpatial;
 
 	UPROPERTY()
 	UPolygonSelectionMechanic* SelectionMechanic = nullptr;
@@ -53,7 +53,7 @@ protected:
 
 	// A variant of group topology that considers all triangles one group, so that group edges are boundary
 	// edges in the mesh.
-	class FBasicTopology : public FGroupTopology
+	class FBasicTopology : public UE::Geometry::FGroupTopology
 	{
 	public:
 		FBasicTopology(const FDynamicMesh3* Mesh, bool bAutoBuild) :
diff --git a/Engine/Plugins/Experimental/MeshModelingToolset/Source/MeshModelingTools/Public/MeshGroupPaintTool.h b/Engine/Plugins/Experimental/MeshModelingToolset/Source/MeshModelingTools/Public/MeshGroupPaintTool.h
--- a/Engine/Plugins/Experimental/MeshModelingToolset/Source/MeshModelingTools/Public/MeshGroupPaintTool.h
+++ b/Engine/Plugins/Experimental/MeshModelingToolset/Source/MeshModelingTools/Public/MeshGroupPaintTool.h
@@ -417,7 +417,7 @@ protected:
 	bool UpdateStampPosition(const FRay& WorldRay);
 	void ApplyStamp();
 
-	FDynamicMeshOctree3 Octree;
+	UE::Geometry::FDynamicMeshOctree3 Octree;
 
 	bool UpdateBrushPosition(const FRay& WorldRay);
 
diff --git a/Engine/Plugins/Experimental/MeshModelingToolset/Source/MeshModelingTools/Public/MeshSelectionTool.h b/Engine/Plugins/Experimental/MeshModelingToolset/Source/MeshModelingTools/Public/MeshSelectionTool.h
--- a/Engine/Plugins/Experimental/MeshModelingToolset/Source/MeshModelingTools/Public/MeshSelectionTool.h
+++ b/Engine/Plugins/Experimental/MeshModelingToolset/Source/MeshModelingTools/Public/MeshSelectionTool.h
@@ -326,9 +326,9 @@ protected:
 	IAssetGenerationAPI* AssetAPI;
 
 	// note: ideally this octree would be part of PreviewMesh!
-	TUniquePtr<FDynamicMeshOctree3> Octree;
+	TUniquePtr<UE::Geometry::FDynamicMeshOctree3> Octree;
 	bool bOctreeValid = false;
-	TUniquePtr<FDynamicMeshOctree3>& GetOctree();
+	TUniquePtr<UE::Geometry::FDynamicMeshOctree3>& GetOctree();
 
 	EMeshSelectionElementType SelectionType = EMeshSelectionElementType::Face;
 
diff --git a/Engine/Plugins/Experimental/MeshModelingToolset/Source/MeshModelingTools/Public/MeshSpaceDeformerTool.h b/Engine/Plugins/Experimental/MeshModelingToolset/Source/MeshModelingTools/Public/MeshSpaceDeformerTool.h
--- a/Engine/Plugins/Experimental/MeshModelingToolset/Source/MeshModelingTools/Public/MeshSpaceDeformerTool.h
+++ b/Engine/Plugins/Experimental/MeshModelingToolset/Source/MeshModelingTools/Public/MeshSpaceDeformerTool.h
@@ -176,13 +176,13 @@ public:
 
 
 UCLASS()
-class MESHMODELINGTOOLS_API USpaceDeformerOperatorFactory : public UObject, public IDynamicMeshOperatorFactory
+class MESHMODELINGTOOLS_API USpaceDeformerOperatorFactory : public UObject, public UE::Geometry::IDynamicMeshOperatorFactory
 {
 	GENERATED_BODY()
 
 public:
 	// IDynamicMeshOperatorFactory API
-	virtual TUniquePtr<FDynamicMeshOperator> MakeNewOperator() override;
+	virtual TUniquePtr<UE::Geometry::FDynamicMeshOperator> MakeNewOperator() override;
 
 	UPROPERTY()
 	UMeshSpaceDeformerTool* SpaceDeformerTool;  // back pointer
@@ -196,7 +196,6 @@ UCLASS()
 class MESHMODELINGTOOLS_API UMeshSpaceDeformerTool : public USingleSelectionTool
 {
 	GENERATED_BODY()
-
 public:
 	UMeshSpaceDeformerTool();
 
@@ -218,7 +217,7 @@ public:
 	virtual void OnPropertyModified(UObject* PropertySet, FProperty* Property) override;
 
 	// sync the parameters owned by the MeshSpaceDeformerOp 
-	void UpdateOpParameters(FMeshSpaceDeformerOp& MeshSpaceDeformerOp) const;
+	void UpdateOpParameters(UE::Geometry::FMeshSpaceDeformerOp& MeshSpaceDeformerOp) const;
 
 protected:
 
@@ -241,7 +240,7 @@ protected:
 
 protected:	
 
-	TSharedPtr<FDynamicMesh3, ESPMode::ThreadSafe> OriginalDynamicMesh;
+	TSharedPtr<UE::Geometry::FDynamicMesh3, ESPMode::ThreadSafe> OriginalDynamicMesh;
 
 	UPROPERTY()
 	UPreviewMesh* OriginalMeshPreview;
@@ -274,9 +273,9 @@ protected:
 
 	// Button click support
 	EMeshSpaceDeformerToolAction PendingAction;
-	FVector3d MeshCenter;
+	UE::Geometry::FVector3d MeshCenter;
 
-	FFrame3d GizmoFrame;
+	UE::Geometry::FFrame3d GizmoFrame;
 
 	// The length of the third interval gizmo (which sets the intensity of the deformation)
 	// when the modifier is set to some reasonable maximum.
@@ -284,7 +283,7 @@ protected:
 
 	TPimplPtr<FSelectClickedAction> SetPointInWorldConnector;
 
-	TArray<FVector3d> VisualizationPoints;
+	TArray<UE::Geometry::FVector3d> VisualizationPoints;
 	FToolDataVisualizer VisualizationRenderer;
 
 	void TransformProxyChanged(UTransformProxy* Proxy, FTransform Transform);
diff --git a/Engine/Plugins/Experimental/MeshModelingToolset/Source/MeshModelingTools/Public/MeshVertexSculptTool.h b/Engine/Plugins/Experimental/MeshModelingToolset/Source/MeshModelingTools/Public/MeshVertexSculptTool.h
--- a/Engine/Plugins/Experimental/MeshModelingToolset/Source/MeshModelingTools/Public/MeshVertexSculptTool.h
+++ b/Engine/Plugins/Experimental/MeshModelingToolset/Source/MeshModelingTools/Public/MeshVertexSculptTool.h
@@ -182,7 +182,7 @@ UCLASS()
 class MESHMODELINGTOOLS_API UMeshVertexSculptTool : public UMeshSculptToolBase
 {
 	GENERATED_BODY()
-
+	using FVector3d = UE::Geometry::FVector3d;
 public:
 	virtual void Setup() override;
 	virtual void Shutdown(EToolShutdownType ShutdownType) override;
@@ -257,14 +257,14 @@ protected:
 	void WaitForPendingStampUpdate();
 
 	TArray<int> RangeQueryTriBuffer;
-	FUniqueIndexSet VertexROIBuilder;
-	FUniqueIndexSet TriangleROIBuilder;
-	TArray<FIndex3i> TriangleROIInBuf;
+	UE::Geometry::FUniqueIndexSet VertexROIBuilder;
+	UE::Geometry::FUniqueIndexSet TriangleROIBuilder;
+	TArray<UE::Geometry::FIndex3i> TriangleROIInBuf;
 	TArray<int> VertexROI;
 	TArray<int> TriangleROIArray;
 	void UpdateROI(const FVector3d& BrushPos);
 
-	FUniqueIndexSet NormalsROIBuilder;
+	UE::Geometry::FUniqueIndexSet NormalsROIBuilder;
 	TArray<std::atomic<bool>> NormalsFlags;		// set of per-vertex or per-element-id flags that indicate
 												// whether normal needs recompute. Fast to do it this way
 												// than to use a TSet or UniqueIndexSet...
@@ -279,22 +279,22 @@ protected:
 	FRandomStream StampRandomStream;
 
 	FDynamicMesh3 BaseMesh;
-	FDynamicMeshOctree3 BaseMeshSpatial;
+	UE::Geometry::FDynamicMeshOctree3 BaseMeshSpatial;
 	TArray<int32> BaseMeshIndexBuffer;
 	bool bCachedFreezeTarget = false;
 	void UpdateBaseMesh(const TSet<int32>* TriangleROI = nullptr);
 	bool GetBaseMeshNearest(int32 VertexID, const FVector3d& Position, double SearchRadius, FVector3d& TargetPosOut, FVector3d& TargetNormalOut);
 	TFunction<bool(int32, const FVector3d&, double MaxDist, FVector3d&, FVector3d&)> BaseMeshQueryFunc;
 
-	FDynamicMeshOctree3 Octree;
+	UE::Geometry::FDynamicMeshOctree3 Octree;
 
 	bool UpdateBrushPosition(const FRay& WorldRay);
 
 	double SculptMaxFixedHeight = -1.0;
 
 	bool bHaveBrushAlpha = false;
-	TImageBuilder<FVector4f> BrushAlphaValues;
-	FImageDimensions BrushAlphaDimensions;
+	UE::Geometry::TImageBuilder<UE::Geometry::FVector4f> BrushAlphaValues;
+	UE::Geometry::FImageDimensions BrushAlphaDimensions;
 	double SampleBrushAlpha(const FSculptBrushStamp& Stamp, const FVector3d& Position) const;
 
 	TArray<FVector3d> ROIPositionBuffer;
diff --git a/Engine/Plugins/Experimental/MeshModelingToolset/Source/MeshModelingTools/Public/MirrorTool.h b/Engine/Plugins/Experimental/MeshModelingToolset/Source/MeshModelingTools/Public/MirrorTool.h
--- a/Engine/Plugins/Experimental/MeshModelingToolset/Source/MeshModelingTools/Public/MirrorTool.h
+++ b/Engine/Plugins/Experimental/MeshModelingToolset/Source/MeshModelingTools/Public/MirrorTool.h
@@ -110,13 +110,13 @@ public:
 
 
 UCLASS()
-class MESHMODELINGTOOLS_API UMirrorOperatorFactory : public UObject, public IDynamicMeshOperatorFactory
+class MESHMODELINGTOOLS_API UMirrorOperatorFactory : public UObject, public UE::Geometry::IDynamicMeshOperatorFactory
 {
 	GENERATED_BODY()
 
 public:
 	// IDynamicMeshOperatorFactory API
-	virtual TUniquePtr<FDynamicMeshOperator> MakeNewOperator() override;
+	virtual TUniquePtr<UE::Geometry::FDynamicMeshOperator> MakeNewOperator() override;
 
 	UPROPERTY()
 	UMirrorTool* MirrorTool;
@@ -185,7 +185,7 @@ UCLASS()
 class MESHMODELINGTOOLS_API UMirrorTool : public UMultiSelectionTool, public IModifierToggleBehaviorTarget
 {
 	GENERATED_BODY()
-
+	using FVector3d = UE::Geometry::FVector3d;
 public:
 
 	friend UMirrorOperatorFactory;
diff --git a/Engine/Plugins/Experimental/MeshModelingToolset/Source/MeshModelingTools/Public/OffsetMeshTool.h b/Engine/Plugins/Experimental/MeshModelingToolset/Source/MeshModelingTools/Public/OffsetMeshTool.h
--- a/Engine/Plugins/Experimental/MeshModelingToolset/Source/MeshModelingTools/Public/OffsetMeshTool.h
+++ b/Engine/Plugins/Experimental/MeshModelingToolset/Source/MeshModelingTools/Public/OffsetMeshTool.h
@@ -114,7 +114,7 @@ public:
 	virtual void InitializeProperties() override;
 	virtual void OnShutdown(EToolShutdownType ShutdownType) override;
 
-	virtual TUniquePtr<FDynamicMeshOperator> MakeNewOperator() override;
+	virtual TUniquePtr<UE::Geometry::FDynamicMeshOperator> MakeNewOperator() override;
 
 	virtual bool RequiresInitialVtxNormals() const { return true; }
 	virtual bool RequiresInitialBoundaryLoops() const { return true; }
diff --git a/Engine/Plugins/Experimental/MeshModelingToolset/Source/MeshModelingTools/Public/Physics/ExtractCollisionGeometryTool.h b/Engine/Plugins/Experimental/MeshModelingToolset/Source/MeshModelingTools/Public/Physics/ExtractCollisionGeometryTool.h
--- a/Engine/Plugins/Experimental/MeshModelingToolset/Source/MeshModelingTools/Public/Physics/ExtractCollisionGeometryTool.h
+++ b/Engine/Plugins/Experimental/MeshModelingToolset/Source/MeshModelingTools/Public/Physics/ExtractCollisionGeometryTool.h
@@ -69,7 +69,7 @@ protected:
 	UWorld* TargetWorld = nullptr;
 	IAssetGenerationAPI* AssetAPI = nullptr;
 
-	FDynamicMesh3 CurrentMesh;
+	UE::Geometry::FDynamicMesh3 CurrentMesh;
 	bool bResultValid = false;
 	void RecalculateMesh();
 
diff --git a/Engine/Plugins/Experimental/MeshModelingToolset/Source/MeshModelingTools/Public/Physics/SetCollisionGeometryTool.h b/Engine/Plugins/Experimental/MeshModelingToolset/Source/MeshModelingTools/Public/Physics/SetCollisionGeometryTool.h
--- a/Engine/Plugins/Experimental/MeshModelingToolset/Source/MeshModelingTools/Public/Physics/SetCollisionGeometryTool.h
+++ b/Engine/Plugins/Experimental/MeshModelingToolset/Source/MeshModelingTools/Public/Physics/SetCollisionGeometryTool.h
@@ -13,8 +13,8 @@
 #include "SetCollisionGeometryTool.generated.h"
 
 class UPreviewGeometry;
-class FMeshSimpleShapeApproximation;
-
+PREDECLARE_GEOMETRY(class FMeshSimpleShapeApproximation)
+PREDECLARE_USE_GEOMETRY_CLASS(FDynamicMesh3);
 
 UCLASS()
 class MESHMODELINGTOOLS_API USetCollisionGeometryToolBuilder : public UInteractiveToolBuilder
@@ -185,9 +185,9 @@ protected:
 
 		EDetectedCollisionGeometry DetectedType = EDetectedCollisionGeometry::None;
 
-		FSphere3d DetectedSphere;
-		FOrientedBox3d DetectedBox;
-		FCapsule3d DetectedCapsule;
+		UE::Geometry::FSphere3d DetectedSphere;
+		UE::Geometry::FOrientedBox3d DetectedBox;
+		UE::Geometry::FCapsule3d DetectedCapsule;
 	};
 	
 	bool bInputMeshesValid = false;
@@ -196,17 +196,17 @@ protected:
 	TArray<TSharedPtr<FDynamicMesh3, ESPMode::ThreadSafe>> SeparatedInputMeshes;
 	TArray<TSharedPtr<FDynamicMesh3, ESPMode::ThreadSafe>> PerGroupInputMeshes;
 
-	TSharedPtr<FMeshSimpleShapeApproximation, ESPMode::ThreadSafe> InputMeshesApproximator;
-	TSharedPtr<FMeshSimpleShapeApproximation, ESPMode::ThreadSafe> CombinedInputMeshesApproximator;
-	TSharedPtr<FMeshSimpleShapeApproximation, ESPMode::ThreadSafe> SeparatedMeshesApproximator;
-	TSharedPtr<FMeshSimpleShapeApproximation, ESPMode::ThreadSafe> PerGroupMeshesApproximator;
+	TSharedPtr<UE::Geometry::FMeshSimpleShapeApproximation, ESPMode::ThreadSafe> InputMeshesApproximator;
+	TSharedPtr<UE::Geometry::FMeshSimpleShapeApproximation, ESPMode::ThreadSafe> CombinedInputMeshesApproximator;
+	TSharedPtr<UE::Geometry::FMeshSimpleShapeApproximation, ESPMode::ThreadSafe> SeparatedMeshesApproximator;
+	TSharedPtr<UE::Geometry::FMeshSimpleShapeApproximation, ESPMode::ThreadSafe> PerGroupMeshesApproximator;
 
 	void PrecomputeInputMeshes();
 	void InitializeDerivedMeshSet(
 		const TArray<TSharedPtr<FDynamicMesh3, ESPMode::ThreadSafe>>& FromInputMeshes,
 		TArray<TSharedPtr<FDynamicMesh3, ESPMode::ThreadSafe>>& ToMeshes,
 		TFunctionRef<bool(const FDynamicMesh3*, int32, int32)> TrisConnectedPredicate);
-	TSharedPtr<FMeshSimpleShapeApproximation, ESPMode::ThreadSafe>& GetApproximator(ESetCollisionGeometryInputMode MeshSetMode);
+	TSharedPtr<UE::Geometry::FMeshSimpleShapeApproximation, ESPMode::ThreadSafe>& GetApproximator(ESetCollisionGeometryInputMode MeshSetMode);
 
 	FTransform OrigTargetTransform;
 	FVector TargetScale3D;
diff --git a/Engine/Plugins/Experimental/MeshModelingToolset/Source/MeshModelingTools/Public/PlaneCutTool.h b/Engine/Plugins/Experimental/MeshModelingToolset/Source/MeshModelingTools/Public/PlaneCutTool.h
--- a/Engine/Plugins/Experimental/MeshModelingToolset/Source/MeshModelingTools/Public/PlaneCutTool.h
+++ b/Engine/Plugins/Experimental/MeshModelingToolset/Source/MeshModelingTools/Public/PlaneCutTool.h
@@ -97,13 +97,13 @@ public:
 
 
 UCLASS()
-class MESHMODELINGTOOLS_API UPlaneCutOperatorFactory : public UObject, public IDynamicMeshOperatorFactory
+class MESHMODELINGTOOLS_API UPlaneCutOperatorFactory : public UObject, public UE::Geometry::IDynamicMeshOperatorFactory
 {
 	GENERATED_BODY()
 
 public:
 	// IDynamicMeshOperatorFactory API
-	virtual TUniquePtr<FDynamicMeshOperator> MakeNewOperator() override;
+	virtual TUniquePtr<UE::Geometry::FDynamicMeshOperator> MakeNewOperator() override;
 
 	UPROPERTY()
 	UPlaneCutTool *CutTool;
diff --git a/Engine/Plugins/Experimental/MeshModelingToolset/Source/MeshModelingTools/Public/PositionPlaneGizmo.h b/Engine/Plugins/Experimental/MeshModelingToolset/Source/MeshModelingTools/Public/PositionPlaneGizmo.h
--- a/Engine/Plugins/Experimental/MeshModelingToolset/Source/MeshModelingTools/Public/PositionPlaneGizmo.h
+++ b/Engine/Plugins/Experimental/MeshModelingToolset/Source/MeshModelingTools/Public/PositionPlaneGizmo.h
@@ -61,7 +61,7 @@ public:
 	virtual void OnEndDrag(const FRay& Ray);
 
 
-	TFunction<void(const FFrame3d&)> OnPositionUpdatedFunc = nullptr;
+	TFunction<void(const UE::Geometry::FFrame3d&)> OnPositionUpdatedFunc = nullptr;
 
 
 	virtual void ExternalUpdatePosition(const FVector& Position, const FQuat& Orientation, bool bPostUpdate);
@@ -81,7 +81,7 @@ public:
 	UPreviewMesh* MakeSphereMesh();
 
 	bool bInTransformDrag;
-	FQuickAxisTranslater QuickTransformer;
+	UE::Geometry::FQuickAxisTranslater QuickTransformer;
 
 	void PostUpdatedPosition();
 
diff --git a/Engine/Plugins/Experimental/MeshModelingToolset/Source/MeshModelingTools/Public/ProjectToTargetTool.h b/Engine/Plugins/Experimental/MeshModelingToolset/Source/MeshModelingTools/Public/ProjectToTargetTool.h
--- a/Engine/Plugins/Experimental/MeshModelingToolset/Source/MeshModelingTools/Public/ProjectToTargetTool.h
+++ b/Engine/Plugins/Experimental/MeshModelingToolset/Source/MeshModelingTools/Public/ProjectToTargetTool.h
@@ -59,11 +59,11 @@ public:
 
 	virtual void Setup() override;
 
-	virtual TUniquePtr<FDynamicMeshOperator> MakeNewOperator() override;
+	virtual TUniquePtr<UE::Geometry::FDynamicMeshOperator> MakeNewOperator() override;
 
 private:
 
-	TUniquePtr<FDynamicMesh3> ProjectionTarget;
-	TUniquePtr<FDynamicMeshAABBTree3> ProjectionTargetSpatial;
+	TUniquePtr<UE::Geometry::FDynamicMesh3> ProjectionTarget;
+	TUniquePtr<UE::Geometry::FDynamicMeshAABBTree3> ProjectionTargetSpatial;
 
 };
diff --git a/Engine/Plugins/Experimental/MeshModelingToolset/Source/MeshModelingTools/Public/Properties/MeshAnalysisProperties.h b/Engine/Plugins/Experimental/MeshModelingToolset/Source/MeshModelingTools/Public/Properties/MeshAnalysisProperties.h
--- a/Engine/Plugins/Experimental/MeshModelingToolset/Source/MeshModelingTools/Public/Properties/MeshAnalysisProperties.h
+++ b/Engine/Plugins/Experimental/MeshModelingToolset/Source/MeshModelingTools/Public/Properties/MeshAnalysisProperties.h
@@ -3,14 +3,13 @@
 #pragma once
 
 #include "CoreMinimal.h"
-
 #include "InteractiveToolBuilder.h"
-
+#include "GeometryBase.h"
 #include "MeshAnalysisProperties.generated.h"
 
 
 // predeclarations
-class FDynamicMesh3;
+PREDECLARE_USE_GEOMETRY_CLASS(FDynamicMesh3);
 
 
 
diff --git a/Engine/Plugins/Experimental/MeshModelingToolset/Source/MeshModelingTools/Public/Properties/MeshMaterialProperties.h b/Engine/Plugins/Experimental/MeshModelingToolset/Source/MeshModelingTools/Public/Properties/MeshMaterialProperties.h
--- a/Engine/Plugins/Experimental/MeshModelingToolset/Source/MeshModelingTools/Public/Properties/MeshMaterialProperties.h
+++ b/Engine/Plugins/Experimental/MeshModelingToolset/Source/MeshModelingTools/Public/Properties/MeshMaterialProperties.h
@@ -5,12 +5,10 @@
 #include "CoreMinimal.h"
 #include "Engine/Classes/Engine/Texture2D.h"
 #include "InteractiveToolBuilder.h"
-
 #include "MeshMaterialProperties.generated.h"
 
 
 // predeclarations
-class FDynamicMesh3;
 class UMaterialInterface;
 class UMaterialInstanceDynamic;
 
diff --git a/Engine/Plugins/Experimental/MeshModelingToolset/Source/MeshModelingTools/Public/Properties/MeshStatisticsProperties.h b/Engine/Plugins/Experimental/MeshModelingToolset/Source/MeshModelingTools/Public/Properties/MeshStatisticsProperties.h
--- a/Engine/Plugins/Experimental/MeshModelingToolset/Source/MeshModelingTools/Public/Properties/MeshStatisticsProperties.h
+++ b/Engine/Plugins/Experimental/MeshModelingToolset/Source/MeshModelingTools/Public/Properties/MeshStatisticsProperties.h
@@ -3,15 +3,13 @@
 #pragma once
 
 #include "CoreMinimal.h"
-
 #include "InteractiveToolBuilder.h"
-
+#include "GeometryBase.h"
 #include "MeshStatisticsProperties.generated.h"
 
 
 // predeclarations
-class FDynamicMesh3;
-
+PREDECLARE_USE_GEOMETRY_CLASS(FDynamicMesh3);
 
 
 UCLASS()
diff --git a/Engine/Plugins/Experimental/MeshModelingToolset/Source/MeshModelingTools/Public/Properties/MeshUVChannelProperties.h b/Engine/Plugins/Experimental/MeshModelingToolset/Source/MeshModelingTools/Public/Properties/MeshUVChannelProperties.h
--- a/Engine/Plugins/Experimental/MeshModelingToolset/Source/MeshModelingTools/Public/Properties/MeshUVChannelProperties.h
+++ b/Engine/Plugins/Experimental/MeshModelingToolset/Source/MeshModelingTools/Public/Properties/MeshUVChannelProperties.h
@@ -4,11 +4,11 @@
 
 #include "CoreMinimal.h"
 #include "InteractiveTool.h"
-
+#include "GeometryBase.h"
 #include "MeshUVChannelProperties.generated.h"
 
 struct FMeshDescription;
-class FDynamicMesh3;
+PREDECLARE_USE_GEOMETRY_CLASS(FDynamicMesh3);
 
 
 UCLASS()
diff --git a/Engine/Plugins/Experimental/MeshModelingToolset/Source/MeshModelingTools/Public/Properties/RemeshProperties.h b/Engine/Plugins/Experimental/MeshModelingToolset/Source/MeshModelingTools/Public/Properties/RemeshProperties.h
--- a/Engine/Plugins/Experimental/MeshModelingToolset/Source/MeshModelingTools/Public/Properties/RemeshProperties.h
+++ b/Engine/Plugins/Experimental/MeshModelingToolset/Source/MeshModelingTools/Public/Properties/RemeshProperties.h
@@ -7,6 +7,8 @@
 
 #include "RemeshProperties.generated.h"
 
+using UE::Geometry::EEdgeRefineFlags;
+
 /** Mesh Boundary Constraint Types */
 UENUM()
 enum class EMeshBoundaryConstraint : uint8
diff --git a/Engine/Plugins/Experimental/MeshModelingToolset/Source/MeshModelingTools/Public/Properties/RevolveProperties.h b/Engine/Plugins/Experimental/MeshModelingToolset/Source/MeshModelingTools/Public/Properties/RevolveProperties.h
--- a/Engine/Plugins/Experimental/MeshModelingToolset/Source/MeshModelingTools/Public/Properties/RevolveProperties.h
+++ b/Engine/Plugins/Experimental/MeshModelingToolset/Source/MeshModelingTools/Public/Properties/RevolveProperties.h
@@ -10,7 +10,7 @@
 #include "RevolveProperties.generated.h"
 
 class UNewMeshMaterialProperties;
-class FCurveSweepOp;
+PREDECLARE_GEOMETRY(class FCurveSweepOp);
 
 UENUM()
 enum class ERevolvePropertiesCapFillMode : uint8
@@ -59,7 +59,7 @@ UCLASS()
 class MESHMODELINGTOOLS_API URevolveProperties : public UInteractiveToolPropertySet
 {
 	GENERATED_BODY()
-
+	using FVector3d = UE::Geometry::FVector3d;
 public:
 
 	URevolveProperties()
@@ -197,7 +197,7 @@ public:
 	 */
 	void ApplyToCurveSweepOp(const UNewMeshMaterialProperties& MaterialProperties,
 		const FVector3d& RevolutionAxisOrigin, const FVector3d& RevolutionAxisDirection,
-		FCurveSweepOp& CurveSweepOpOut) const;
+		UE::Geometry::FCurveSweepOp& CurveSweepOpOut) const;
 
 	protected:
 		int32 ClampedRevolutionDegreesWatcherIndex;
diff --git a/Engine/Plugins/Experimental/MeshModelingToolset/Source/MeshModelingTools/Public/RemeshMeshTool.h b/Engine/Plugins/Experimental/MeshModelingToolset/Source/MeshModelingTools/Public/RemeshMeshTool.h
--- a/Engine/Plugins/Experimental/MeshModelingToolset/Source/MeshModelingTools/Public/RemeshMeshTool.h
+++ b/Engine/Plugins/Experimental/MeshModelingToolset/Source/MeshModelingTools/Public/RemeshMeshTool.h
@@ -110,7 +110,7 @@ public:
  * that subclasses of this class can work with multiple meshes (see, for example, UProjectToTargetTool.)
  */
 UCLASS()
-class MESHMODELINGTOOLS_API URemeshMeshTool : public UMultiSelectionTool, public IDynamicMeshOperatorFactory
+class MESHMODELINGTOOLS_API URemeshMeshTool : public UMultiSelectionTool, public UE::Geometry::IDynamicMeshOperatorFactory
 {
 	GENERATED_BODY()
 
@@ -134,7 +134,7 @@ public:
 	virtual void OnPropertyModified(UObject* PropertySet, FProperty* Property) override;
 
 	// IDynamicMeshOperatorFactory API
-	virtual TUniquePtr<FDynamicMeshOperator> MakeNewOperator() override;
+	virtual TUniquePtr<UE::Geometry::FDynamicMeshOperator> MakeNewOperator() override;
 
 	UPROPERTY()
 	URemeshMeshToolProperties* BasicProperties;
@@ -150,8 +150,8 @@ protected:
 	UWorld* TargetWorld;
 	IToolsContextAssetAPI* AssetAPI;
 
-	TSharedPtr<FDynamicMesh3, ESPMode::ThreadSafe> OriginalMesh;
-	TSharedPtr<FDynamicMeshAABBTree3, ESPMode::ThreadSafe> OriginalMeshSpatial;
+	TSharedPtr<UE::Geometry::FDynamicMesh3, ESPMode::ThreadSafe> OriginalMesh;
+	TSharedPtr<UE::Geometry::FDynamicMeshAABBTree3, ESPMode::ThreadSafe> OriginalMeshSpatial;
 	double InitialMeshArea;
 
 	double CalculateTargetEdgeLength(int TargetTriCount);
diff --git a/Engine/Plugins/Experimental/MeshModelingToolset/Source/MeshModelingTools/Public/RemoveOccludedTrianglesTool.h b/Engine/Plugins/Experimental/MeshModelingToolset/Source/MeshModelingTools/Public/RemoveOccludedTrianglesTool.h
--- a/Engine/Plugins/Experimental/MeshModelingToolset/Source/MeshModelingTools/Public/RemoveOccludedTrianglesTool.h
+++ b/Engine/Plugins/Experimental/MeshModelingToolset/Source/MeshModelingTools/Public/RemoveOccludedTrianglesTool.h
@@ -16,6 +16,7 @@
 #include "RemoveOccludedTrianglesTool.generated.h"
 
 class IAssetGenerationAPI;
+PREDECLARE_USE_GEOMETRY_CLASS(FDynamicMesh3);
 
 /**
  *
@@ -133,13 +134,13 @@ public:
  *  stores a pointer to the tool and enough info to know which specific operator it should spawn
  */
 UCLASS()
-class MESHMODELINGTOOLS_API URemoveOccludedTrianglesOperatorFactory : public UObject, public IDynamicMeshOperatorFactory
+class MESHMODELINGTOOLS_API URemoveOccludedTrianglesOperatorFactory : public UObject, public UE::Geometry::IDynamicMeshOperatorFactory
 {
 	GENERATED_BODY()
 
 public:
 	// IDynamicMeshOperatorFactory API
-	virtual TUniquePtr<FDynamicMeshOperator> MakeNewOperator() override;
+	virtual TUniquePtr<UE::Geometry::FDynamicMeshOperator> MakeNewOperator() override;
 
 	UPROPERTY()
 	URemoveOccludedTrianglesTool *Tool;
@@ -204,9 +205,9 @@ protected:
 	TArray<TSharedPtr<FDynamicMesh3, ESPMode::ThreadSafe>> OriginalDynamicMeshes;
 
 	// AABB trees and winding trees for every mesh target, with repeated instances as pointers to the same data
-	TArray<TSharedPtr<FDynamicMeshAABBTree3, ESPMode::ThreadSafe>> OccluderTrees;
-	TArray<TSharedPtr<TFastWindingTree<FDynamicMesh3>, ESPMode::ThreadSafe>> OccluderWindings;
-	TArray<FTransform3d> OccluderTransforms;
+	TArray<TSharedPtr<UE::Geometry::FDynamicMeshAABBTree3, ESPMode::ThreadSafe>> OccluderTrees;
+	TArray<TSharedPtr<UE::Geometry::TFastWindingTree<FDynamicMesh3>, ESPMode::ThreadSafe>> OccluderWindings;
+	TArray<UE::Geometry::FTransform3d> OccluderTransforms;
 
 	TArray<TArray<int32>> PreviewToCopyIdx;
 	TArray<int32> PreviewToTargetIdx;
diff --git a/Engine/Plugins/Experimental/MeshModelingToolset/Source/MeshModelingTools/Public/RevolveBoundaryTool.h b/Engine/Plugins/Experimental/MeshModelingToolset/Source/MeshModelingTools/Public/RevolveBoundaryTool.h
--- a/Engine/Plugins/Experimental/MeshModelingToolset/Source/MeshModelingTools/Public/RevolveBoundaryTool.h
+++ b/Engine/Plugins/Experimental/MeshModelingToolset/Source/MeshModelingTools/Public/RevolveBoundaryTool.h
@@ -36,13 +36,13 @@ public:
 };
 
 UCLASS()
-class MESHMODELINGTOOLS_API URevolveBoundaryOperatorFactory : public UObject, public IDynamicMeshOperatorFactory
+class MESHMODELINGTOOLS_API URevolveBoundaryOperatorFactory : public UObject, public UE::Geometry::IDynamicMeshOperatorFactory
 {
 	GENERATED_BODY()
 
 public:
 
-	TUniquePtr<FDynamicMeshOperator> MakeNewOperator() override;
+	TUniquePtr<UE::Geometry::FDynamicMeshOperator> MakeNewOperator() override;
 
 	UPROPERTY()
 	URevolveBoundaryTool* RevolveBoundaryTool;
@@ -84,6 +84,7 @@ UCLASS()
 class MESHMODELINGTOOLS_API URevolveBoundaryTool : public UMeshBoundaryToolBase, public IClickBehaviorTarget
 {
 	GENERATED_BODY()
+	using FVector3d = UE::Geometry::FVector3d;
 
 public:
 	virtual void SetAssetAPI(IAssetGenerationAPI* NewAssetApi) { AssetAPI = NewAssetApi; }
diff --git a/Engine/Plugins/Experimental/MeshModelingToolset/Source/MeshModelingTools/Public/Sculpting/MeshBrushOpBase.h b/Engine/Plugins/Experimental/MeshModelingToolset/Source/MeshModelingTools/Public/Sculpting/MeshBrushOpBase.h
--- a/Engine/Plugins/Experimental/MeshModelingToolset/Source/MeshModelingTools/Public/Sculpting/MeshBrushOpBase.h
+++ b/Engine/Plugins/Experimental/MeshModelingToolset/Source/MeshModelingTools/Public/Sculpting/MeshBrushOpBase.h
@@ -8,7 +8,10 @@
 #include "InteractiveTool.h"
 #include "MeshBrushOpBase.generated.h"
 
-class FDynamicMesh3;
+PREDECLARE_USE_GEOMETRY_CLASS(FDynamicMesh3);
+using UE::Geometry::FFrame3d;
+using UE::Geometry::FMatrix3d;
+using UE::Geometry::FVector3d;
 
 enum class ESculptBrushOpTargetType : uint8
 {
diff --git a/Engine/Plugins/Experimental/MeshModelingToolset/Source/MeshModelingTools/Public/Sculpting/MeshInflateBrushOps.h b/Engine/Plugins/Experimental/MeshModelingToolset/Source/MeshModelingTools/Public/Sculpting/MeshInflateBrushOps.h
--- a/Engine/Plugins/Experimental/MeshModelingToolset/Source/MeshModelingTools/Public/Sculpting/MeshInflateBrushOps.h
+++ b/Engine/Plugins/Experimental/MeshModelingToolset/Source/MeshModelingTools/Public/Sculpting/MeshInflateBrushOps.h
@@ -48,7 +48,7 @@ public:
 		{
 			int32 VertIdx = Vertices[k];
 			FVector3d OrigPos = Mesh->GetVertex(VertIdx);
-			FVector3d Normal = FMeshNormals::ComputeVertexNormal(*Mesh, VertIdx);
+			FVector3d Normal = UE::Geometry::FMeshNormals::ComputeVertexNormal(*Mesh, VertIdx);
 			FVector3d MoveVec = UsePower * Normal;
 
 			double Falloff = GetFalloff().Evaluate(Stamp, OrigPos);
diff --git a/Engine/Plugins/Experimental/MeshModelingToolset/Source/MeshModelingTools/Public/Sculpting/MeshPinchBrushOps.h b/Engine/Plugins/Experimental/MeshModelingToolset/Source/MeshModelingTools/Public/Sculpting/MeshPinchBrushOps.h
--- a/Engine/Plugins/Experimental/MeshModelingToolset/Source/MeshModelingTools/Public/Sculpting/MeshPinchBrushOps.h
+++ b/Engine/Plugins/Experimental/MeshModelingToolset/Source/MeshModelingTools/Public/Sculpting/MeshPinchBrushOps.h
@@ -62,7 +62,7 @@ public:
 		FVector3d MotionVec = NewSmoothBrushPosLocal - LastSmoothBrushPosLocal;
 		bool bHaveMotion = (MotionVec.Length() > FMathd::ZeroTolerance);
 		MotionVec.Normalize();
-		FLine3d MoveLine(LastSmoothBrushPosLocal, MotionVec);
+		UE::Geometry::FLine3d MoveLine(LastSmoothBrushPosLocal, MotionVec);
 
 		FVector3d DepthPosLocal = NewSmoothBrushPosLocal - (Stamp.Depth * Stamp.Radius * NewSmoothBrushNormal);
 		double UseSpeed = Stamp.Direction * Stamp.Radius * Stamp.Power * Stamp.DeltaTime * BrushSpeedTuning;
diff --git a/Engine/Plugins/Experimental/MeshModelingToolset/Source/MeshModelingTools/Public/Sculpting/MeshSculptToolBase.h b/Engine/Plugins/Experimental/MeshModelingToolset/Source/MeshModelingTools/Public/Sculpting/MeshSculptToolBase.h
--- a/Engine/Plugins/Experimental/MeshModelingToolset/Source/MeshModelingTools/Public/Sculpting/MeshSculptToolBase.h
+++ b/Engine/Plugins/Experimental/MeshModelingToolset/Source/MeshModelingTools/Public/Sculpting/MeshSculptToolBase.h
@@ -159,7 +159,6 @@ public:
 
 
 
-
 /**
  * Base Tool for mesh sculpting tools, provides some shared functionality
  */
@@ -167,7 +166,16 @@ UCLASS()
 class MESHMODELINGTOOLS_API UMeshSculptToolBase : public UMeshSurfacePointTool
 {
 	GENERATED_BODY()
+protected:
+	using FVector3d = UE::Geometry::FVector3d;
+	using FVector3f = UE::Geometry::FVector3f;
+	using FVector2d = UE::Geometry::FVector2d;
+	using FVector2f = UE::Geometry::FVector2f;
+	using FTransform3d = UE::Geometry::FTransform3d;
+	using FFrame3d = UE::Geometry::FFrame3d;
+	using FRay3d = UE::Geometry::FRay3d;
 public:
+
 	virtual void RegisterActions(FInteractiveToolActionSet& ActionSet) override;
 
 	virtual void SetWorld(UWorld* World);
@@ -206,9 +214,9 @@ protected:
 	FViewCameraState CameraState;
 
 	/** Initial transformation on target mesh */
-	FTransform3d InitialTargetTransform;
+	UE::Geometry::FTransform3d InitialTargetTransform;
 	/** Active transformation on target mesh, includes baked scale */
-	FTransform3d CurTargetTransform;
+	UE::Geometry::FTransform3d CurTargetTransform;
 
 	FRay3d GetLocalRay(const FRay& WorldRay) const;
 
@@ -287,9 +295,9 @@ protected:
 	// Brush Size
 	//
 protected:
-	FInterval1d BrushRelativeSizeRange;
+	UE::Geometry::FInterval1d BrushRelativeSizeRange;
 	double CurrentBrushRadius = 1.0;
-	virtual void InitializeBrushSizeRange(const FAxisAlignedBox3d& TargetBounds);
+	virtual void InitializeBrushSizeRange(const UE::Geometry::FAxisAlignedBox3d& TargetBounds);
 	virtual void CalculateBrushRadius();
 	virtual double GetCurrentBrushRadius() const { return CurrentBrushRadius; }
 
diff --git a/Engine/Plugins/Experimental/MeshModelingToolset/Source/MeshModelingTools/Public/Sculpting/MeshSmoothingBrushOps.h b/Engine/Plugins/Experimental/MeshModelingToolset/Source/MeshModelingTools/Public/Sculpting/MeshSmoothingBrushOps.h
--- a/Engine/Plugins/Experimental/MeshModelingToolset/Source/MeshModelingTools/Public/Sculpting/MeshSmoothingBrushOps.h
+++ b/Engine/Plugins/Experimental/MeshModelingToolset/Source/MeshModelingTools/Public/Sculpting/MeshSmoothingBrushOps.h
@@ -88,11 +88,11 @@ public:
 			FVector3d SmoothedPos = OrigPos;
 			if (bPreserveUVFlow)
 			{
-				SmoothedPos = FMeshWeights::CotanCentroidSafe(*Mesh, VertIdx, 10.0);
+				SmoothedPos = UE::Geometry::FMeshWeights::CotanCentroidSafe(*Mesh, VertIdx, 10.0);
 			}
 			else
 			{
-				SmoothedPos = FMeshWeights::UniformCentroid(*Mesh, VertIdx);
+				SmoothedPos = UE::Geometry::FMeshWeights::UniformCentroid(*Mesh, VertIdx);
 			}
 
 			FVector3d NewPos = FVector3d::Lerp(OrigPos, SmoothedPos, Falloff * Stamp.Power);
@@ -146,18 +146,18 @@ public:
 			int32 VertIdx = Vertices[k];
 
 			FVector3d OrigPos = Mesh->GetVertex(VertIdx);
-			FVector3d Normal = FMeshNormals::ComputeVertexNormal(*Mesh, VertIdx);
+			FVector3d Normal = UE::Geometry::FMeshNormals::ComputeVertexNormal(*Mesh, VertIdx);
 
 			double Falloff = GetFalloff().Evaluate(Stamp, OrigPos);
 
 			FVector3d SmoothedPos = OrigPos;
 			if (bPreserveUVFlow)
 			{
-				SmoothedPos = FMeshWeights::CotanCentroidSafe(*Mesh, VertIdx, 10.0);
+				SmoothedPos = UE::Geometry::FMeshWeights::CotanCentroidSafe(*Mesh, VertIdx, 10.0);
 			}
 			else
 			{
-				SmoothedPos = FMeshWeights::UniformCentroid(*Mesh, VertIdx);
+				SmoothedPos = UE::Geometry::FMeshWeights::UniformCentroid(*Mesh, VertIdx);
 			}
 
 			FVector3d NewPos = FVector3d::Lerp(OrigPos, SmoothedPos, Falloff * Stamp.Power);
@@ -226,7 +226,7 @@ public:
 		{
 			int32 VertIdx = Vertices[k];
 			FVector3d OrigPos = Mesh->GetVertex(VertIdx);
-			FVector3d Normal = FMeshNormals::ComputeVertexNormal(*Mesh, VertIdx);
+			FVector3d Normal = UE::Geometry::FMeshNormals::ComputeVertexNormal(*Mesh, VertIdx);
 			FVector3d PlanePos = FlattenPlane.ToPlane(OrigPos, 2);
 			FVector3d MoveDelta = PlanePos - OrigPos;
 
diff --git a/Engine/Plugins/Experimental/MeshModelingToolset/Source/MeshModelingTools/Public/SeamSculptTool.h b/Engine/Plugins/Experimental/MeshModelingToolset/Source/MeshModelingTools/Public/SeamSculptTool.h
--- a/Engine/Plugins/Experimental/MeshModelingToolset/Source/MeshModelingTools/Public/SeamSculptTool.h
+++ b/Engine/Plugins/Experimental/MeshModelingToolset/Source/MeshModelingTools/Public/SeamSculptTool.h
@@ -8,7 +8,7 @@
 
 class UExistingMeshMaterialProperties;
 class UPreviewGeometry;
-class FDynamicMesh3;
+PREDECLARE_USE_GEOMETRY_CLASS(FDynamicMesh3);
 
 /**
  *
@@ -50,6 +50,7 @@ UCLASS(Transient)
 class MESHMODELINGTOOLS_API USeamSculptTool : public UDynamicMeshBrushTool
 {
 	GENERATED_BODY()
+	using FVector3d = UE::Geometry::FVector3d;
 public:
 	USeamSculptTool();
 
@@ -85,7 +86,7 @@ protected:
 	UPreviewGeometry* PreviewGeom;
 
 	TSharedPtr<FDynamicMesh3, ESPMode::ThreadSafe> InputMesh;
-	FTransform3d MeshTransform;
+	UE::Geometry::FTransform3d MeshTransform;
 	double NormalOffset = 0;
 
 	void InitPreviewGeometry();
diff --git a/Engine/Plugins/Experimental/MeshModelingToolset/Source/MeshModelingTools/Public/SelfUnionMeshesTool.h b/Engine/Plugins/Experimental/MeshModelingToolset/Source/MeshModelingTools/Public/SelfUnionMeshesTool.h
--- a/Engine/Plugins/Experimental/MeshModelingToolset/Source/MeshModelingTools/Public/SelfUnionMeshesTool.h
+++ b/Engine/Plugins/Experimental/MeshModelingToolset/Source/MeshModelingTools/Public/SelfUnionMeshesTool.h
@@ -12,8 +12,7 @@
 #include "SelfUnionMeshesTool.generated.h"
 
 // predeclarations
-class FDynamicMesh3;
-
+PREDECLARE_USE_GEOMETRY_CLASS(FDynamicMesh3);
 
 
 /**
@@ -79,7 +78,7 @@ protected:
 	virtual FText GetActionName() const;
 
 	// IDynamicMeshOperatorFactory API
-	virtual TUniquePtr<FDynamicMeshOperator> MakeNewOperator() override;
+	virtual TUniquePtr<UE::Geometry::FDynamicMeshOperator> MakeNewOperator() override;
 
 protected:
 
diff --git a/Engine/Plugins/Experimental/MeshModelingToolset/Source/MeshModelingTools/Public/SmoothMeshTool.h b/Engine/Plugins/Experimental/MeshModelingToolset/Source/MeshModelingTools/Public/SmoothMeshTool.h
--- a/Engine/Plugins/Experimental/MeshModelingToolset/Source/MeshModelingTools/Public/SmoothMeshTool.h
+++ b/Engine/Plugins/Experimental/MeshModelingToolset/Source/MeshModelingTools/Public/SmoothMeshTool.h
@@ -135,7 +135,7 @@ public:
 	virtual void InitializeProperties() override;
 	virtual void OnShutdown(EToolShutdownType ShutdownType) override;
 
-	virtual TUniquePtr<FDynamicMeshOperator> MakeNewOperator() override;
+	virtual TUniquePtr<UE::Geometry::FDynamicMeshOperator> MakeNewOperator() override;
 
 	virtual bool RequiresInitialVtxNormals() const { return true; }
 	virtual bool HasMeshTopologyChanged() const override;
diff --git a/Engine/Plugins/Experimental/MeshModelingToolset/Source/MeshModelingTools/Public/TransformMeshesTool.h b/Engine/Plugins/Experimental/MeshModelingToolset/Source/MeshModelingTools/Public/TransformMeshesTool.h
--- a/Engine/Plugins/Experimental/MeshModelingToolset/Source/MeshModelingTools/Public/TransformMeshesTool.h
+++ b/Engine/Plugins/Experimental/MeshModelingToolset/Source/MeshModelingTools/Public/TransformMeshesTool.h
@@ -187,7 +187,7 @@ protected:
 	void ResetActiveGizmos();
 
 
-	FFrame3d StartDragFrameWorld;
+	UE::Geometry::FFrame3d StartDragFrameWorld;
 	FTransform StartDragTransform;
 	int ActiveSnapDragIndex = -1;
 
diff --git a/Engine/Plugins/Experimental/MeshModelingToolset/Source/MeshModelingTools/Public/UVProjectionTool.h b/Engine/Plugins/Experimental/MeshModelingToolset/Source/MeshModelingTools/Public/UVProjectionTool.h
--- a/Engine/Plugins/Experimental/MeshModelingToolset/Source/MeshModelingTools/Public/UVProjectionTool.h
+++ b/Engine/Plugins/Experimental/MeshModelingToolset/Source/MeshModelingTools/Public/UVProjectionTool.h
@@ -22,7 +22,7 @@ class USimpleDynamicMeshComponent;
 class UTransformGizmo;
 class UTransformProxy;
 class IAssetGenerationAPI;
-
+PREDECLARE_USE_GEOMETRY_CLASS(FDynamicMesh3);
 
 /**
  *
@@ -102,13 +102,13 @@ public:
  *  stores a pointer to the tool and enough info to know which specific operator it should spawn
  */
 UCLASS()
-class MESHMODELINGTOOLS_API UUVProjectionOperatorFactory : public UObject, public IDynamicMeshOperatorFactory
+class MESHMODELINGTOOLS_API UUVProjectionOperatorFactory : public UObject, public UE::Geometry::IDynamicMeshOperatorFactory
 {
 	GENERATED_BODY()
 
 public:
 	// IDynamicMeshOperatorFactory API
-	virtual TUniquePtr<FDynamicMeshOperator> MakeNewOperator() override;
+	virtual TUniquePtr<UE::Geometry::FDynamicMeshOperator> MakeNewOperator() override;
 
 	UPROPERTY()
 	UUVProjectionTool *Tool;
diff --git a/Engine/Plugins/Experimental/MeshModelingToolset/Source/MeshModelingTools/Public/VolumeToMeshTool.h b/Engine/Plugins/Experimental/MeshModelingToolset/Source/MeshModelingTools/Public/VolumeToMeshTool.h
--- a/Engine/Plugins/Experimental/MeshModelingToolset/Source/MeshModelingTools/Public/VolumeToMeshTool.h
+++ b/Engine/Plugins/Experimental/MeshModelingToolset/Source/MeshModelingTools/Public/VolumeToMeshTool.h
@@ -104,7 +104,7 @@ protected:
 	UWorld* TargetWorld = nullptr;
 	IAssetGenerationAPI* AssetAPI = nullptr;
 
-	FDynamicMesh3 CurrentMesh;
+	UE::Geometry::FDynamicMesh3 CurrentMesh;
 
 	void RecalculateMesh();
 
diff --git a/Engine/Plugins/Experimental/MeshModelingToolset/Source/MeshModelingTools/Public/WeldMeshEdgesTool.h b/Engine/Plugins/Experimental/MeshModelingToolset/Source/MeshModelingTools/Public/WeldMeshEdgesTool.h
--- a/Engine/Plugins/Experimental/MeshModelingToolset/Source/MeshModelingTools/Public/WeldMeshEdgesTool.h
+++ b/Engine/Plugins/Experimental/MeshModelingToolset/Source/MeshModelingTools/Public/WeldMeshEdgesTool.h
@@ -64,7 +64,7 @@ protected:
 
 protected:
 	USimpleDynamicMeshComponent* DynamicMeshComponent;
-	FDynamicMesh3 OriginalMesh;
+	UE::Geometry::FDynamicMesh3 OriginalMesh;
 	bool bResultValid;
 	void UpdateResult();
 };
diff --git a/Engine/Plugins/Experimental/MeshModelingToolset/Source/MeshModelingToolsEditorOnly/Private/EditMeshMaterialsTool.cpp b/Engine/Plugins/Experimental/MeshModelingToolset/Source/MeshModelingToolsEditorOnly/Private/EditMeshMaterialsTool.cpp
--- a/Engine/Plugins/Experimental/MeshModelingToolset/Source/MeshModelingToolsEditorOnly/Private/EditMeshMaterialsTool.cpp
+++ b/Engine/Plugins/Experimental/MeshModelingToolset/Source/MeshModelingToolsEditorOnly/Private/EditMeshMaterialsTool.cpp
@@ -15,9 +15,10 @@
 #include "AssetGenerationUtil.h"
 #include "ToolSetupUtil.h"
 
-#define LOCTEXT_NAMESPACE "UEditMeshMaterialsTool"
-
+#include "ExplicitUseGeometryMathTypes.h"		// using UE::Geometry::(math types)
+using namespace UE::Geometry;
 
+#define LOCTEXT_NAMESPACE "UEditMeshMaterialsTool"
 
 void UEditMeshMaterialsEditActions::PostMaterialAction(EEditMeshMaterialsToolActions Action)
 {
diff --git a/Engine/Plugins/Experimental/MeshModelingToolset/Source/MeshModelingToolsEditorOnly/Private/GenerateLODMeshesTool.cpp b/Engine/Plugins/Experimental/MeshModelingToolset/Source/MeshModelingToolsEditorOnly/Private/GenerateLODMeshesTool.cpp
--- a/Engine/Plugins/Experimental/MeshModelingToolset/Source/MeshModelingToolsEditorOnly/Private/GenerateLODMeshesTool.cpp
+++ b/Engine/Plugins/Experimental/MeshModelingToolset/Source/MeshModelingToolsEditorOnly/Private/GenerateLODMeshesTool.cpp
@@ -28,6 +28,9 @@
 #include "Misc/ScopedSlowTask.h"
 #endif
 
+#include "ExplicitUseGeometryMathTypes.h"		// using UE::Geometry::(math types)
+using namespace UE::Geometry;
+
 #define LOCTEXT_NAMESPACE "UGenerateLODMeshesTool"
 
 /*
@@ -276,9 +279,9 @@ void UGenerateLODMeshesTool::OnPreviewUpdated(UMeshOpPreviewWithBackgroundComput
 }
 
 
-TUniquePtr<FDynamicMeshOperator> FGenerateLODOperatorFactory::MakeNewOperator()
+TUniquePtr<UE::Geometry::FDynamicMeshOperator> FGenerateLODOperatorFactory::MakeNewOperator()
 {
-	TUniquePtr<FSimplifyMeshOp> Op = MakeUnique<FSimplifyMeshOp>();
+	TUniquePtr<UE::Geometry::FSimplifyMeshOp> Op = MakeUnique<UE::Geometry::FSimplifyMeshOp>();
 
 	UGenerateLODMeshesToolProperties* SimplifyProperties = ParentTool->SimplifyProperties;
 
diff --git a/Engine/Plugins/Experimental/MeshModelingToolset/Source/MeshModelingToolsEditorOnly/Private/MeshToVolumeTool.cpp b/Engine/Plugins/Experimental/MeshModelingToolset/Source/MeshModelingToolsEditorOnly/Private/MeshToVolumeTool.cpp
--- a/Engine/Plugins/Experimental/MeshModelingToolset/Source/MeshModelingToolsEditorOnly/Private/MeshToVolumeTool.cpp
+++ b/Engine/Plugins/Experimental/MeshModelingToolset/Source/MeshModelingToolsEditorOnly/Private/MeshToVolumeTool.cpp
@@ -23,8 +23,10 @@
 #include "Model.h"
 #include "BSPOps.h"		// in UnrealEd
 
-#define LOCTEXT_NAMESPACE "UMeshToVolumeTool"
+#include "ExplicitUseGeometryMathTypes.h"		// using UE::Geometry::(math types)
+using namespace UE::Geometry;
 
+#define LOCTEXT_NAMESPACE "UMeshToVolumeTool"
 
 /*
  * ToolBuilder
diff --git a/Engine/Plugins/Experimental/MeshModelingToolset/Source/MeshModelingToolsEditorOnly/Private/ParameterizeMeshTool.cpp b/Engine/Plugins/Experimental/MeshModelingToolset/Source/MeshModelingToolsEditorOnly/Private/ParameterizeMeshTool.cpp
--- a/Engine/Plugins/Experimental/MeshModelingToolset/Source/MeshModelingToolsEditorOnly/Private/ParameterizeMeshTool.cpp
+++ b/Engine/Plugins/Experimental/MeshModelingToolset/Source/MeshModelingToolsEditorOnly/Private/ParameterizeMeshTool.cpp
@@ -16,8 +16,10 @@
 #include "ParameterizationOps/ParameterizeMeshOp.h"
 #include "ToolSetupUtil.h"
 
-#define LOCTEXT_NAMESPACE "UParameterizeMeshTool"
+#include "ExplicitUseGeometryMathTypes.h"		// using UE::Geometry::(math types)
+using namespace UE::Geometry;
 
+#define LOCTEXT_NAMESPACE "UParameterizeMeshTool"
 
 DEFINE_LOG_CATEGORY_STATIC(LogParameterizeMeshTool, Log, All);
 
diff --git a/Engine/Plugins/Experimental/MeshModelingToolset/Source/MeshModelingToolsEditorOnly/Private/UVLayoutTool.cpp b/Engine/Plugins/Experimental/MeshModelingToolset/Source/MeshModelingToolsEditorOnly/Private/UVLayoutTool.cpp
--- a/Engine/Plugins/Experimental/MeshModelingToolset/Source/MeshModelingToolsEditorOnly/Private/UVLayoutTool.cpp
+++ b/Engine/Plugins/Experimental/MeshModelingToolset/Source/MeshModelingToolsEditorOnly/Private/UVLayoutTool.cpp
@@ -18,8 +18,10 @@
 
 #include "AssetGenerationUtil.h"
 
-#define LOCTEXT_NAMESPACE "UUVLayoutTool"
+#include "ExplicitUseGeometryMathTypes.h"		// using UE::Geometry::(math types)
+using namespace UE::Geometry;
 
+#define LOCTEXT_NAMESPACE "UUVLayoutTool"
 
 /*
  * ToolBuilder
diff --git a/Engine/Plugins/Experimental/MeshModelingToolset/Source/MeshModelingToolsEditorOnly/Public/BspConversionTool.h b/Engine/Plugins/Experimental/MeshModelingToolset/Source/MeshModelingToolsEditorOnly/Public/BspConversionTool.h
--- a/Engine/Plugins/Experimental/MeshModelingToolset/Source/MeshModelingToolsEditorOnly/Public/BspConversionTool.h
+++ b/Engine/Plugins/Experimental/MeshModelingToolset/Source/MeshModelingToolsEditorOnly/Public/BspConversionTool.h
@@ -174,7 +174,7 @@ protected:
 	TArray<ABrush*> BrushesToDelete;
 	ABrush* BrushForPivot = nullptr;
 
-	TMap<ABrush*, TPair<TSharedPtr<const FDynamicMesh3>, 
+	TMap<ABrush*, TPair<TSharedPtr<const UE::Geometry::FDynamicMesh3>, 
 		TSharedPtr<const TArray<UMaterialInterface*>>>> CachedBrushes;
 
 	bool IsValidConversionTarget(const ABrush* Brush) const;
diff --git a/Engine/Plugins/Experimental/MeshModelingToolset/Source/MeshModelingToolsEditorOnly/Public/ConversionUtils/DynamicMeshToVolume.h b/Engine/Plugins/Experimental/MeshModelingToolset/Source/MeshModelingToolsEditorOnly/Public/ConversionUtils/DynamicMeshToVolume.h
--- a/Engine/Plugins/Experimental/MeshModelingToolset/Source/MeshModelingToolsEditorOnly/Public/ConversionUtils/DynamicMeshToVolume.h
+++ b/Engine/Plugins/Experimental/MeshModelingToolset/Source/MeshModelingToolsEditorOnly/Public/ConversionUtils/DynamicMeshToVolume.h
@@ -12,11 +12,13 @@
 // one file).
 
 class AVolume;
-class FDynamicMesh3;
+PREDECLARE_USE_GEOMETRY_CLASS(FDynamicMesh3);
 
 namespace UE {
 namespace Conversion { 
 
+using namespace UE::Geometry;
+
 struct FDynamicMeshFace
 {
 	FFrame3d Plane;
diff --git a/Engine/Plugins/Experimental/MeshModelingToolset/Source/MeshModelingToolsEditorOnly/Public/GenerateLODMeshesTool.h b/Engine/Plugins/Experimental/MeshModelingToolset/Source/MeshModelingToolsEditorOnly/Public/GenerateLODMeshesTool.h
--- a/Engine/Plugins/Experimental/MeshModelingToolset/Source/MeshModelingToolsEditorOnly/Public/GenerateLODMeshesTool.h
+++ b/Engine/Plugins/Experimental/MeshModelingToolset/Source/MeshModelingToolsEditorOnly/Public/GenerateLODMeshesTool.h
@@ -136,7 +136,7 @@ public:
 
 
 
-class FGenerateLODOperatorFactory : public IDynamicMeshOperatorFactory
+class FGenerateLODOperatorFactory : public UE::Geometry::IDynamicMeshOperatorFactory
 {
 public:
 	TWeakObjectPtr<UGenerateLODMeshesTool> ParentTool;
@@ -144,7 +144,7 @@ public:
 	FTransform UseTransform;
 
 	// IDynamicMeshOperatorFactory API
-	virtual TUniquePtr<FDynamicMeshOperator> MakeNewOperator() override;
+	virtual TUniquePtr<UE::Geometry::FDynamicMeshOperator> MakeNewOperator() override;
 };
 
 
@@ -191,12 +191,12 @@ private:
 
 	TSharedPtr<FMeshDescription, ESPMode::ThreadSafe> OriginalMeshDescription;
 	// Dynamic Mesh versions precomputed in Setup (rather than recomputed for every simplify op)
-	TSharedPtr<FDynamicMesh3, ESPMode::ThreadSafe> OriginalMesh;
-	TSharedPtr<FDynamicMeshAABBTree3, ESPMode::ThreadSafe> OriginalMeshSpatial;
+	TSharedPtr<UE::Geometry::FDynamicMesh3, ESPMode::ThreadSafe> OriginalMesh;
+	TSharedPtr<UE::Geometry::FDynamicMeshAABBTree3, ESPMode::ThreadSafe> OriginalMeshSpatial;
 
 	TArray<FLODLevelGenerateSettings> CachedLODLevels;
 
-	FAxisAlignedBox3d WorldBounds;
+	UE::Geometry::FAxisAlignedBox3d WorldBounds;
 
 	void GenerateAssets();
 	void UpdateVisualization();
diff --git a/Engine/Plugins/Experimental/MeshModelingToolset/Source/MeshModelingToolsEditorOnly/Public/MergeMeshesTool.h b/Engine/Plugins/Experimental/MeshModelingToolset/Source/MeshModelingToolsEditorOnly/Public/MergeMeshesTool.h
--- a/Engine/Plugins/Experimental/MeshModelingToolset/Source/MeshModelingToolsEditorOnly/Public/MergeMeshesTool.h
+++ b/Engine/Plugins/Experimental/MeshModelingToolset/Source/MeshModelingToolsEditorOnly/Public/MergeMeshesTool.h
@@ -66,7 +66,7 @@ public:
  *
  */
 UCLASS()
-class MESHMODELINGTOOLSEDITORONLY_API UMergeMeshesTool : public UMultiSelectionTool, public IDynamicMeshOperatorFactory
+class MESHMODELINGTOOLSEDITORONLY_API UMergeMeshesTool : public UMultiSelectionTool, public UE::Geometry::IDynamicMeshOperatorFactory
 {
 	GENERATED_BODY()
 
@@ -88,7 +88,7 @@ public:
 	virtual void OnPropertyModified(UObject* PropertySet, FProperty* Property) override;
 
 	// IDynamicMeshOperatorFactory API
-	virtual TUniquePtr<FDynamicMeshOperator> MakeNewOperator() override;
+	virtual TUniquePtr<UE::Geometry::FDynamicMeshOperator> MakeNewOperator() override;
 
 protected:
 	UPROPERTY()
diff --git a/Engine/Plugins/Experimental/MeshModelingToolset/Source/MeshModelingToolsEditorOnly/Public/MeshTangentsTool.h b/Engine/Plugins/Experimental/MeshModelingToolset/Source/MeshModelingToolsEditorOnly/Public/MeshTangentsTool.h
--- a/Engine/Plugins/Experimental/MeshModelingToolset/Source/MeshModelingToolsEditorOnly/Public/MeshTangentsTool.h
+++ b/Engine/Plugins/Experimental/MeshModelingToolset/Source/MeshModelingToolsEditorOnly/Public/MeshTangentsTool.h
@@ -78,10 +78,10 @@ public:
  * Simple Mesh Simplifying Tool
  */
 UCLASS()
-class MESHMODELINGTOOLSEDITORONLY_API UMeshTangentsTool : public USingleSelectionTool, public IGenericDataOperatorFactory<FMeshTangentsd>
+class MESHMODELINGTOOLSEDITORONLY_API UMeshTangentsTool : public USingleSelectionTool, public UE::Geometry::IGenericDataOperatorFactory<UE::Geometry::FMeshTangentsd>
 {
 	GENERATED_BODY()
-
+	using FVector3d = UE::Geometry::FVector3d;
 public:
 	UMeshTangentsTool();
 
@@ -98,7 +98,7 @@ public:
 	virtual bool CanAccept() const override;
 
 	// IGenericDataOperatorFactory API
-	virtual TUniquePtr<TGenericDataOperator<FMeshTangentsd>> MakeNewOperator() override;
+	virtual TUniquePtr<UE::Geometry::TGenericDataOperator<UE::Geometry::FMeshTangentsd>> MakeNewOperator() override;
 
 protected:
 	UPROPERTY()
@@ -114,20 +114,20 @@ protected:
 	UPROPERTY()
 	UPreviewGeometry* PreviewGeometry = nullptr;
 
-	TUniquePtr<TGenericDataBackgroundCompute<FMeshTangentsd>> Compute = nullptr;
+	TUniquePtr<TGenericDataBackgroundCompute<UE::Geometry::FMeshTangentsd>> Compute = nullptr;
 
 protected:
 	UWorld* TargetWorld;
 
 	TSharedPtr<FMeshDescription, ESPMode::ThreadSafe> InputMeshDescription;
-	TSharedPtr<FMeshTangentsf, ESPMode::ThreadSafe> InitialTangents;
-	TSharedPtr<FDynamicMesh3, ESPMode::ThreadSafe> InputMesh;
+	TSharedPtr<UE::Geometry::FMeshTangentsf, ESPMode::ThreadSafe> InitialTangents;
+	TSharedPtr<UE::Geometry::FDynamicMesh3, ESPMode::ThreadSafe> InputMesh;
 
 	bool bThicknessDirty = false;
 	bool bLengthDirty = false;
 	bool bVisibilityChanged = false;
 
-	void OnTangentsUpdated(const TUniquePtr<FMeshTangentsd>& NewResult);
+	void OnTangentsUpdated(const TUniquePtr<UE::Geometry::FMeshTangentsd>& NewResult);
 	void UpdateVisualization(bool bThicknessChanged, bool bLengthChanged);
 
 	struct FMikktDeviation
diff --git a/Engine/Plugins/Experimental/MeshModelingToolset/Source/MeshModelingToolsEditorOnly/Public/MeshToVolumeTool.h b/Engine/Plugins/Experimental/MeshModelingToolset/Source/MeshModelingToolsEditorOnly/Public/MeshToVolumeTool.h
--- a/Engine/Plugins/Experimental/MeshModelingToolset/Source/MeshModelingToolsEditorOnly/Public/MeshToVolumeTool.h
+++ b/Engine/Plugins/Experimental/MeshModelingToolset/Source/MeshModelingToolsEditorOnly/Public/MeshToVolumeTool.h
@@ -109,7 +109,7 @@ protected:
 protected:
 	IAssetGenerationAPI* AssetAPI = nullptr;
 
-	FDynamicMesh3 InputMesh;
+	UE::Geometry::FDynamicMesh3 InputMesh;
 
 	void RecalculateVolume();
 	void UpdateLineSet();
diff --git a/Engine/Plugins/Experimental/MeshModelingToolset/Source/MeshModelingToolsEditorOnly/Public/ParameterizeMeshTool.h b/Engine/Plugins/Experimental/MeshModelingToolset/Source/MeshModelingToolsEditorOnly/Public/ParameterizeMeshTool.h
--- a/Engine/Plugins/Experimental/MeshModelingToolset/Source/MeshModelingToolsEditorOnly/Public/ParameterizeMeshTool.h
+++ b/Engine/Plugins/Experimental/MeshModelingToolset/Source/MeshModelingToolsEditorOnly/Public/ParameterizeMeshTool.h
@@ -149,7 +149,7 @@ public:
  * Simple Mesh Simplifying Tool
  */
 UCLASS()
-class MESHMODELINGTOOLSEDITORONLY_API UParameterizeMeshTool : public USingleSelectionTool, public IDynamicMeshOperatorFactory
+class MESHMODELINGTOOLSEDITORONLY_API UParameterizeMeshTool : public USingleSelectionTool, public UE::Geometry::IDynamicMeshOperatorFactory
 {
 	GENERATED_BODY()
 
@@ -173,7 +173,7 @@ public:
 	virtual void OnPropertyModified(UObject* PropertySet, FProperty* Property) override;
 
 	// IDynamicMeshOperatorFactory API
-	virtual TUniquePtr<FDynamicMeshOperator> MakeNewOperator() override;
+	virtual TUniquePtr<UE::Geometry::FDynamicMeshOperator> MakeNewOperator() override;
 
 protected:
 	UPROPERTY()
@@ -205,6 +205,6 @@ protected:
 	UWorld* TargetWorld;
 	IAssetGenerationAPI* AssetAPI;
 
-	TSharedPtr<FDynamicMesh3, ESPMode::ThreadSafe> InputMesh;
+	TSharedPtr<UE::Geometry::FDynamicMesh3, ESPMode::ThreadSafe> InputMesh;
 
 };
diff --git a/Engine/Plugins/Experimental/MeshModelingToolset/Source/MeshModelingToolsEditorOnly/Public/PolygonOnMeshTool.h b/Engine/Plugins/Experimental/MeshModelingToolset/Source/MeshModelingToolsEditorOnly/Public/PolygonOnMeshTool.h
--- a/Engine/Plugins/Experimental/MeshModelingToolset/Source/MeshModelingToolsEditorOnly/Public/PolygonOnMeshTool.h
+++ b/Engine/Plugins/Experimental/MeshModelingToolset/Source/MeshModelingToolsEditorOnly/Public/PolygonOnMeshTool.h
@@ -144,7 +144,7 @@ public:
  * Simple Mesh Plane Cutting Tool
  */
 UCLASS()
-class MESHMODELINGTOOLSEDITORONLY_API UPolygonOnMeshTool : public USingleSelectionTool, public IDynamicMeshOperatorFactory, public IClickBehaviorTarget, public IHoverBehaviorTarget
+class MESHMODELINGTOOLSEDITORONLY_API UPolygonOnMeshTool : public USingleSelectionTool, public UE::Geometry::IDynamicMeshOperatorFactory, public IClickBehaviorTarget, public IHoverBehaviorTarget
 {
 	GENERATED_BODY()
 
@@ -167,7 +167,7 @@ public:
 	virtual void OnPropertyModified(UObject* PropertySet, FProperty* Property) override;
 
 	// IDynamicMeshOperatorFactory API
-	virtual TUniquePtr<FDynamicMeshOperator> MakeNewOperator() override;
+	virtual TUniquePtr<UE::Geometry::FDynamicMeshOperator> MakeNewOperator() override;
 
 
 	virtual void RequestAction(EPolygonOnMeshToolActions ActionType);
@@ -212,10 +212,10 @@ protected:
 protected:
 	UWorld* TargetWorld;
 
-	FTransform3d WorldTransform;
+	UE::Geometry::FTransform3d WorldTransform;
 	FViewCameraState CameraState;
 
-	TSharedPtr<FDynamicMesh3, ESPMode::ThreadSafe> OriginalDynamicMesh;
+	TSharedPtr<UE::Geometry::FDynamicMesh3, ESPMode::ThreadSafe> OriginalDynamicMesh;
 
 	UPROPERTY()
 	UConstructionPlaneMechanic* PlaneMechanic = nullptr;
@@ -225,10 +225,10 @@ protected:
 
 	EPolygonOnMeshToolActions PendingAction = EPolygonOnMeshToolActions::NoAction;
 
-	FFrame3d DrawPlaneWorld;
+	UE::Geometry::FFrame3d DrawPlaneWorld;
 
-	FPolygon2d LastDrawnPolygon;
-	FPolygon2d ActivePolygon;
+	UE::Geometry::FPolygon2d LastDrawnPolygon;
+	UE::Geometry::FPolygon2d ActivePolygon;
 	void UpdatePolygonType();
 
 	void SetupPreview();
diff --git a/Engine/Plugins/Experimental/MeshModelingToolset/Source/MeshModelingToolsEditorOnly/Public/ShapeSprayTool.h b/Engine/Plugins/Experimental/MeshModelingToolset/Source/MeshModelingToolsEditorOnly/Public/ShapeSprayTool.h
--- a/Engine/Plugins/Experimental/MeshModelingToolset/Source/MeshModelingToolsEditorOnly/Public/ShapeSprayTool.h
+++ b/Engine/Plugins/Experimental/MeshModelingToolset/Source/MeshModelingToolsEditorOnly/Public/ShapeSprayTool.h
@@ -102,9 +102,9 @@ protected:
 	UWorld* TargetWorld;
 	IAssetGenerationAPI* AssetAPI;
 
-	FDynamicMesh3 ShapeMesh;
+	UE::Geometry::FDynamicMesh3 ShapeMesh;
 	void UpdateShapeMesh();
-	void SplatShape(const FFrame3d& LocalFrame, double Scale, FDynamicMesh3* TargetMesh);
+	void SplatShape(const UE::Geometry::FFrame3d& LocalFrame, double Scale, UE::Geometry::FDynamicMesh3* TargetMesh);
 	TArray<int> VertexMap;
 
 	FRandomStream Random;
diff --git a/Engine/Plugins/Experimental/MeshModelingToolset/Source/MeshModelingToolsEditorOnly/Public/SimplifyMeshTool.h b/Engine/Plugins/Experimental/MeshModelingToolset/Source/MeshModelingToolsEditorOnly/Public/SimplifyMeshTool.h
--- a/Engine/Plugins/Experimental/MeshModelingToolset/Source/MeshModelingToolsEditorOnly/Public/SimplifyMeshTool.h
+++ b/Engine/Plugins/Experimental/MeshModelingToolset/Source/MeshModelingToolsEditorOnly/Public/SimplifyMeshTool.h
@@ -93,7 +93,7 @@ public:
  * Simple Mesh Simplifying Tool
  */
 UCLASS()
-class MESHMODELINGTOOLSEDITORONLY_API USimplifyMeshTool : public USingleSelectionTool, public IDynamicMeshOperatorFactory
+class MESHMODELINGTOOLSEDITORONLY_API USimplifyMeshTool : public USingleSelectionTool, public UE::Geometry::IDynamicMeshOperatorFactory
 {
 	GENERATED_BODY()
 
@@ -114,7 +114,7 @@ public:
 	virtual void OnPropertyModified(UObject* PropertySet, FProperty* Property) override;
 
 	// IDynamicMeshOperatorFactory API
-	virtual TUniquePtr<FDynamicMeshOperator> MakeNewOperator() override;
+	virtual TUniquePtr<UE::Geometry::FDynamicMeshOperator> MakeNewOperator() override;
 
 private:
 	UPROPERTY()
@@ -131,8 +131,8 @@ private:
 
 	TSharedPtr<FMeshDescription, ESPMode::ThreadSafe> OriginalMeshDescription;
 	// Dynamic Mesh versions precomputed in Setup (rather than recomputed for every simplify op)
-	TSharedPtr<FDynamicMesh3, ESPMode::ThreadSafe> OriginalMesh;
-	TSharedPtr<FDynamicMeshAABBTree3, ESPMode::ThreadSafe> OriginalMeshSpatial;
+	TSharedPtr<UE::Geometry::FDynamicMesh3, ESPMode::ThreadSafe> OriginalMesh;
+	TSharedPtr<UE::Geometry::FDynamicMeshAABBTree3, ESPMode::ThreadSafe> OriginalMeshSpatial;
 
 	void GenerateAsset(const FDynamicMeshOpResult& Result);
 	void UpdateVisualization();
diff --git a/Engine/Plugins/Experimental/MeshModelingToolset/Source/MeshModelingToolsEditorOnly/Public/SubdividePoly.h b/Engine/Plugins/Experimental/MeshModelingToolset/Source/MeshModelingToolsEditorOnly/Public/SubdividePoly.h
--- a/Engine/Plugins/Experimental/MeshModelingToolset/Source/MeshModelingToolsEditorOnly/Public/SubdividePoly.h
+++ b/Engine/Plugins/Experimental/MeshModelingToolset/Source/MeshModelingToolsEditorOnly/Public/SubdividePoly.h
@@ -6,8 +6,8 @@
 #include "ModelingOperators.h"
 #include "SubdividePoly.generated.h"
 
-class FGroupTopology;
-class FDynamicMesh3;
+PREDECLARE_USE_GEOMETRY_CLASS(FGroupTopology);
+PREDECLARE_USE_GEOMETRY_CLASS(FDynamicMesh3);
 class FProgressCancel;
 
 /**
diff --git a/Engine/Plugins/Experimental/MeshModelingToolset/Source/MeshModelingToolsEditorOnly/Public/SubdividePolyTool.h b/Engine/Plugins/Experimental/MeshModelingToolset/Source/MeshModelingToolsEditorOnly/Public/SubdividePolyTool.h
--- a/Engine/Plugins/Experimental/MeshModelingToolset/Source/MeshModelingToolsEditorOnly/Public/SubdividePolyTool.h
+++ b/Engine/Plugins/Experimental/MeshModelingToolset/Source/MeshModelingToolsEditorOnly/Public/SubdividePolyTool.h
@@ -99,7 +99,7 @@ protected:
 	USubdividePolyToolProperties* Properties = nullptr;
 
 	// Input mesh
-	TSharedPtr<FDynamicMesh3, ESPMode::ThreadSafe> OriginalMesh;
+	TSharedPtr<UE::Geometry::FDynamicMesh3, ESPMode::ThreadSafe> OriginalMesh;
 
 	UPROPERTY()
 	UPreviewGeometry* PreviewGeometry = nullptr;
diff --git a/Engine/Plugins/Experimental/MeshModelingToolset/Source/MeshModelingToolsEditorOnly/Public/UVLayoutTool.h b/Engine/Plugins/Experimental/MeshModelingToolset/Source/MeshModelingToolsEditorOnly/Public/UVLayoutTool.h
--- a/Engine/Plugins/Experimental/MeshModelingToolset/Source/MeshModelingToolsEditorOnly/Public/UVLayoutTool.h
+++ b/Engine/Plugins/Experimental/MeshModelingToolset/Source/MeshModelingToolsEditorOnly/Public/UVLayoutTool.h
@@ -93,13 +93,13 @@ public:
  *  stores a pointer to the tool and enough info to know which specific operator it should spawn
  */
 UCLASS()
-class MESHMODELINGTOOLSEDITORONLY_API UUVLayoutOperatorFactory : public UObject, public IDynamicMeshOperatorFactory
+class MESHMODELINGTOOLSEDITORONLY_API UUVLayoutOperatorFactory : public UObject, public UE::Geometry::IDynamicMeshOperatorFactory
 {
 	GENERATED_BODY()
 
 public:
 	// IDynamicMeshOperatorFactory API
-	virtual TUniquePtr<FDynamicMeshOperator> MakeNewOperator() override;
+	virtual TUniquePtr<UE::Geometry::FDynamicMeshOperator> MakeNewOperator() override;
 
 	UPROPERTY()
 	UUVLayoutTool *Tool;
@@ -157,7 +157,7 @@ protected:
 	TArray<UMeshOpPreviewWithBackgroundCompute*> Previews;
 
 protected:
-	TArray<TSharedPtr<FDynamicMesh3, ESPMode::ThreadSafe>> OriginalDynamicMeshes;
+	TArray<TSharedPtr<UE::Geometry::FDynamicMesh3, ESPMode::ThreadSafe>> OriginalDynamicMeshes;
 
 	UWorld* TargetWorld = nullptr;
 	IAssetGenerationAPI* AssetAPI = nullptr;
