GeometryProcessing: move all public types in GeometryProcessing plugin into UE::Geometry:: namespace, as well as some code in MeshModelingToolset. Fix fallout in Water, Paint Mode, ChaosEditor, DataPr

diff --git a/Engine/Plugins/Experimental/MeshModelingToolset/Source/MeshModelingToolsEditorOnly/Public/VoxelCSGMeshesTool.h b/Engine/Plugins/Experimental/MeshModelingToolset/Source/MeshModelingToolsEditorOnly/Public/VoxelCSGMeshesTool.h
--- a/Engine/Plugins/Experimental/MeshModelingToolset/Source/MeshModelingToolsEditorOnly/Public/VoxelCSGMeshesTool.h
+++ b/Engine/Plugins/Experimental/MeshModelingToolset/Source/MeshModelingToolsEditorOnly/Public/VoxelCSGMeshesTool.h
@@ -92,7 +92,7 @@ public:
  *
  */
 UCLASS()
-class MESHMODELINGTOOLSEDITORONLY_API UVoxelCSGMeshesTool : public UMultiSelectionTool, public IDynamicMeshOperatorFactory
+class MESHMODELINGTOOLSEDITORONLY_API UVoxelCSGMeshesTool : public UMultiSelectionTool, public UE::Geometry::IDynamicMeshOperatorFactory
 {
 	GENERATED_BODY()
 
@@ -115,7 +115,7 @@ public:
 	virtual void OnPropertyModified(UObject* PropertySet, FProperty* Property) override;
 
 	// IDynamicMeshOperatorFactory API
-	virtual TUniquePtr<FDynamicMeshOperator> MakeNewOperator() override;
+	virtual TUniquePtr<UE::Geometry::FDynamicMeshOperator> MakeNewOperator() override;
 
 
 protected:
diff --git a/Engine/Plugins/Experimental/MeshModelingToolset/Source/ModelingComponents/Private/AssetGenerationUtil.cpp b/Engine/Plugins/Experimental/MeshModelingToolset/Source/ModelingComponents/Private/AssetGenerationUtil.cpp
--- a/Engine/Plugins/Experimental/MeshModelingToolset/Source/ModelingComponents/Private/AssetGenerationUtil.cpp
+++ b/Engine/Plugins/Experimental/MeshModelingToolset/Source/ModelingComponents/Private/AssetGenerationUtil.cpp
@@ -16,7 +16,7 @@ AActor* AssetGenerationUtil::GenerateStaticMeshActor(
 	IAssetGenerationAPI* AssetAPI,
 	UWorld* TargetWorld,
 	const FDynamicMesh3* Mesh,
-	const FTransform3d& Transform,
+	const UE::Geometry::FTransform3d& Transform,
 	FString ObjectName,
 	UMaterialInterface* Material
 )
@@ -29,7 +29,7 @@ AActor* AssetGenerationUtil::GenerateStaticMeshActor(
 	IAssetGenerationAPI* AssetAPI,
 	UWorld* TargetWorld,
 	const FDynamicMesh3* Mesh,
-	const FTransform3d& Transform,
+	const UE::Geometry::FTransform3d& Transform,
 	FString ObjectName,
 	const TArrayView<UMaterialInterface*>& Materials
 )
diff --git a/Engine/Plugins/Experimental/MeshModelingToolset/Source/ModelingComponents/Private/AssetUtils/Texture2DUtil.cpp b/Engine/Plugins/Experimental/MeshModelingToolset/Source/ModelingComponents/Private/AssetUtils/Texture2DUtil.cpp
--- a/Engine/Plugins/Experimental/MeshModelingToolset/Source/ModelingComponents/Private/AssetUtils/Texture2DUtil.cpp
+++ b/Engine/Plugins/Experimental/MeshModelingToolset/Source/ModelingComponents/Private/AssetUtils/Texture2DUtil.cpp
@@ -2,10 +2,8 @@
 
 #include "AssetUtils/Texture2DUtil.h"
 
-
-
-
-
+#include "ExplicitUseGeometryMathTypes.h"		// using UE::Geometry::(math types)
+using namespace UE::Geometry;
 
 static bool ReadTexture_PlatformData(
 	UTexture2D* TextureMap,
diff --git a/Engine/Plugins/Experimental/MeshModelingToolset/Source/ModelingComponents/Private/BaseDynamicMeshSceneProxy.h b/Engine/Plugins/Experimental/MeshModelingToolset/Source/ModelingComponents/Private/BaseDynamicMeshSceneProxy.h
--- a/Engine/Plugins/Experimental/MeshModelingToolset/Source/ModelingComponents/Private/BaseDynamicMeshSceneProxy.h
+++ b/Engine/Plugins/Experimental/MeshModelingToolset/Source/ModelingComponents/Private/BaseDynamicMeshSceneProxy.h
@@ -9,6 +9,11 @@
 #include "PrimitiveViewRelevance.h"
 #include "BaseDynamicMeshComponent.h"
 
+using UE::Geometry::FDynamicMesh3;
+using UE::Geometry::FDynamicMeshAttributeSet;
+using UE::Geometry::FDynamicMeshUVOverlay;
+using UE::Geometry::FDynamicMeshNormalOverlay;
+using UE::Geometry::FDynamicMeshMaterialAttribute;
 
 /**
  * FMeshRenderBufferSet stores a set of RenderBuffers for a mesh
@@ -279,6 +284,11 @@ protected:
  */
 class FBaseDynamicMeshSceneProxy : public FPrimitiveSceneProxy
 {
+	using FVector3d = UE::Geometry::FVector3d;
+	using FVector3f = UE::Geometry::FVector3f;
+	using FVector2f = UE::Geometry::FVector2f;
+	using FIndex2i = UE::Geometry::FIndex2i;
+	using FIndex3i = UE::Geometry::FIndex3i;
 public:
 	UBaseDynamicMeshComponent* ParentBaseComponent;
 
diff --git a/Engine/Plugins/Experimental/MeshModelingToolset/Source/ModelingComponents/Private/Drawing/MeshWireframeComponent.cpp b/Engine/Plugins/Experimental/MeshModelingToolset/Source/ModelingComponents/Private/Drawing/MeshWireframeComponent.cpp
--- a/Engine/Plugins/Experimental/MeshModelingToolset/Source/ModelingComponents/Private/Drawing/MeshWireframeComponent.cpp
+++ b/Engine/Plugins/Experimental/MeshModelingToolset/Source/ModelingComponents/Private/Drawing/MeshWireframeComponent.cpp
@@ -13,6 +13,8 @@
 #include "IndexTypes.h"
 #include "Async/ParallelFor.h"
 
+using UE::Geometry::FIndex4i;
+
 struct FWireframeLinesMeshBatchData
 {
 	FWireframeLinesMeshBatchData()
diff --git a/Engine/Plugins/Experimental/MeshModelingToolset/Source/ModelingComponents/Private/Drawing/TriangleSetComponent.cpp b/Engine/Plugins/Experimental/MeshModelingToolset/Source/ModelingComponents/Private/Drawing/TriangleSetComponent.cpp
--- a/Engine/Plugins/Experimental/MeshModelingToolset/Source/ModelingComponents/Private/Drawing/TriangleSetComponent.cpp
+++ b/Engine/Plugins/Experimental/MeshModelingToolset/Source/ModelingComponents/Private/Drawing/TriangleSetComponent.cpp
@@ -285,7 +285,7 @@ int32 UTriangleSetComponent::AddTriangle(const FVector& A, const FVector& B, con
 	return AddTriangle(NewTriangle);
 }
 
-FIndex2i UTriangleSetComponent::AddQuad(const FVector& A, const FVector& B, const FVector& C, const FVector& D, const FVector& Normal, const FColor& Color, UMaterialInterface* Material)
+UE::Geometry::FIndex2i UTriangleSetComponent::AddQuad(const FVector& A, const FVector& B, const FVector& C, const FVector& D, const FVector& Normal, const FColor& Color, UMaterialInterface* Material)
 {
 	FRenderableTriangle NewTriangle0;
 	NewTriangle0.Material = Material;
@@ -300,7 +300,7 @@ FIndex2i UTriangleSetComponent::AddQuad(const FVector& A, const FVector& B, cons
 
 	int32 Index0 = AddTriangle(NewTriangle0);
 	int32 Index1 = AddTriangle(NewTriangle1);
-	return FIndex2i(Index0, Index1);
+	return UE::Geometry::FIndex2i(Index0, Index1);
 }
 
 
diff --git a/Engine/Plugins/Experimental/MeshModelingToolset/Source/ModelingComponents/Private/Mechanics/PolyLassoMarqueeMechanic.cpp b/Engine/Plugins/Experimental/MeshModelingToolset/Source/ModelingComponents/Private/Mechanics/PolyLassoMarqueeMechanic.cpp
--- a/Engine/Plugins/Experimental/MeshModelingToolset/Source/ModelingComponents/Private/Mechanics/PolyLassoMarqueeMechanic.cpp
+++ b/Engine/Plugins/Experimental/MeshModelingToolset/Source/ModelingComponents/Private/Mechanics/PolyLassoMarqueeMechanic.cpp
@@ -8,6 +8,10 @@
 #include "ToolSceneQueriesUtil.h"
 #include "TransformTypes.h"
 
+#include "ExplicitUseGeometryMathTypes.h"		// using UE::Geometry::(math types)
+using namespace UE::Geometry;
+
+
 #define LOCTEXT_NAMESPACE "UPolyLassoMarqueeMechanic"
 
 
diff --git a/Engine/Plugins/Experimental/MeshModelingToolset/Source/ModelingComponents/Private/Mechanics/SpatialCurveDistanceMechanic.cpp b/Engine/Plugins/Experimental/MeshModelingToolset/Source/ModelingComponents/Private/Mechanics/SpatialCurveDistanceMechanic.cpp
--- a/Engine/Plugins/Experimental/MeshModelingToolset/Source/ModelingComponents/Private/Mechanics/SpatialCurveDistanceMechanic.cpp
+++ b/Engine/Plugins/Experimental/MeshModelingToolset/Source/ModelingComponents/Private/Mechanics/SpatialCurveDistanceMechanic.cpp
@@ -7,9 +7,8 @@
 #include "MeshTransforms.h"
 #include "Distance/DistRay3Segment3.h"
 
-
-
-
+#include "ExplicitUseGeometryMathTypes.h"		// using UE::Geometry::(math types)
+using namespace UE::Geometry;
 
 void USpatialCurveDistanceMechanic::Setup(UInteractiveTool* ParentToolIn)
 {
diff --git a/Engine/Plugins/Experimental/MeshModelingToolset/Source/ModelingComponents/Private/OctreeDynamicMeshSceneProxy.h b/Engine/Plugins/Experimental/MeshModelingToolset/Source/ModelingComponents/Private/OctreeDynamicMeshSceneProxy.h
--- a/Engine/Plugins/Experimental/MeshModelingToolset/Source/ModelingComponents/Private/OctreeDynamicMeshSceneProxy.h
+++ b/Engine/Plugins/Experimental/MeshModelingToolset/Source/ModelingComponents/Private/OctreeDynamicMeshSceneProxy.h
@@ -29,6 +29,11 @@ DECLARE_CYCLE_STAT(TEXT("SculptToolOctree_BufferUpload"), STAT_SculptToolOctree_
  */
 class FOctreeDynamicMeshSceneProxy final : public FBaseDynamicMeshSceneProxy
 {
+	using FVector3d = UE::Geometry::FVector3d;
+	using FVector3f = UE::Geometry::FVector3f;
+	using FVector2f = UE::Geometry::FVector2f;
+	using FIndex2i = UE::Geometry::FIndex2i;
+	using FIndex3i = UE::Geometry::FIndex3i;
 private:
 	FMaterialRelevance MaterialRelevance;
 
@@ -96,7 +101,7 @@ public:
 
 
 
-	void InitializeFromDecomposition(const FArrayIndexSetsDecomposition& Decomposition)
+	void InitializeFromDecomposition(const UE::Geometry::FArrayIndexSetsDecomposition& Decomposition)
 	{
 		check(RenderBufferSets.Num() == 0);
 
@@ -134,7 +139,7 @@ public:
 
 
 
-	void UpdateFromDecomposition(const FArrayIndexSetsDecomposition& Decomposition, const TArray<int32>& SetsToUpdate )
+	void UpdateFromDecomposition(const UE::Geometry::FArrayIndexSetsDecomposition& Decomposition, const TArray<int32>& SetsToUpdate )
 	{
 		// CAN WE REUSE EXISTING BUFFER SETS??
 		//   - could have timestamp for each decomposition set array...if tris don't change we only have to update vertices
diff --git a/Engine/Plugins/Experimental/MeshModelingToolset/Source/ModelingComponents/Private/Physics/ComponentCollisionUtil.cpp b/Engine/Plugins/Experimental/MeshModelingToolset/Source/ModelingComponents/Private/Physics/ComponentCollisionUtil.cpp
--- a/Engine/Plugins/Experimental/MeshModelingToolset/Source/ModelingComponents/Private/Physics/ComponentCollisionUtil.cpp
+++ b/Engine/Plugins/Experimental/MeshModelingToolset/Source/ModelingComponents/Private/Physics/ComponentCollisionUtil.cpp
@@ -17,8 +17,9 @@
 #include "Engine/Classes/PhysicsEngine/BodySetup.h"
 #include "PhysicsEngine/AggregateGeom.h"
 
+using namespace UE::Geometry;
 
-UE::Geometry::FComponentCollisionSettings UE::Geometry::GetCollisionSettings(const UPrimitiveComponent* Component)
+FComponentCollisionSettings UE::Geometry::GetCollisionSettings(const UPrimitiveComponent* Component)
 {
 	FComponentCollisionSettings Settings;
 
diff --git a/Engine/Plugins/Experimental/MeshModelingToolset/Source/ModelingComponents/Private/SimpleDynamicMeshComponent.cpp b/Engine/Plugins/Experimental/MeshModelingToolset/Source/ModelingComponents/Private/SimpleDynamicMeshComponent.cpp
--- a/Engine/Plugins/Experimental/MeshModelingToolset/Source/ModelingComponents/Private/SimpleDynamicMeshComponent.cpp
+++ b/Engine/Plugins/Experimental/MeshModelingToolset/Source/ModelingComponents/Private/SimpleDynamicMeshComponent.cpp
@@ -28,11 +28,11 @@
 #include "DynamicMeshChangeTracker.h"
 #include "MeshTransforms.h"
 
-
-
 // default proxy for this component
 #include "SimpleDynamicMeshSceneProxy.h"
 
+#include "ExplicitUseGeometryMathTypes.h"		// using UE::Geometry::(math types)
+using namespace UE::Geometry;
 
 namespace
 {
diff --git a/Engine/Plugins/Experimental/MeshModelingToolset/Source/ModelingComponents/Private/SimpleDynamicMeshSceneProxy.h b/Engine/Plugins/Experimental/MeshModelingToolset/Source/ModelingComponents/Private/SimpleDynamicMeshSceneProxy.h
--- a/Engine/Plugins/Experimental/MeshModelingToolset/Source/ModelingComponents/Private/SimpleDynamicMeshSceneProxy.h
+++ b/Engine/Plugins/Experimental/MeshModelingToolset/Source/ModelingComponents/Private/SimpleDynamicMeshSceneProxy.h
@@ -6,7 +6,6 @@
 #include "BaseDynamicMeshSceneProxy.h"
 #include "Async/ParallelFor.h"
 
-
 /**
  * Scene Proxy for USimpleDynamicMeshComponent.
  * 
@@ -17,6 +16,11 @@
  */
 class FSimpleDynamicMeshSceneProxy final : public FBaseDynamicMeshSceneProxy
 {
+	using FVector3d = UE::Geometry::FVector3d;
+	using FVector3f = UE::Geometry::FVector3f;
+	using FVector2f = UE::Geometry::FVector2f;
+	using FIndex2i = UE::Geometry::FIndex2i;
+	using FIndex3i = UE::Geometry::FIndex3i;
 private:
 	FMaterialRelevance MaterialRelevance;
 
@@ -109,7 +113,7 @@ public:
 		FDynamicMeshNormalOverlay* NormalOverlay = Mesh->Attributes()->PrimaryNormals();
 
 		TFunction<void(int, int, int, FVector3f&, FVector3f&)> TangentsFunc = nullptr;
-		const FMeshTangentsf* Tangents = ParentComponent->GetTangents();
+		const UE::Geometry::FMeshTangentsf* Tangents = ParentComponent->GetTangents();
 		if (Tangents != nullptr)
 		{
 			TangentsFunc = [Tangents](int VertexID, int TriangleID, int TriVtxIdx, FVector3f& TangentX, FVector3f& TangentY)
@@ -169,7 +173,7 @@ public:
 		}
 
 		TFunction<void(int, int, int, FVector3f&, FVector3f&)> TangentsFunc = nullptr;
-		const FMeshTangentsf* Tangents = ParentComponent->GetTangents();
+		const UE::Geometry::FMeshTangentsf* Tangents = ParentComponent->GetTangents();
 		if (Tangents != nullptr)
 		{
 			TangentsFunc = [Tangents](int VertexID, int TriangleID, int TriVtxIdx, FVector3f& TangentX, FVector3f& TangentY)
@@ -213,7 +217,7 @@ public:
 		}
 
 		TFunction<void(int, int, int, FVector3f&, FVector3f&)> TangentsFunc = nullptr;
-		const FMeshTangentsf* Tangents = ParentComponent->GetTangents();
+		const UE::Geometry::FMeshTangentsf* Tangents = ParentComponent->GetTangents();
 		if (Tangents != nullptr)
 		{
 			TangentsFunc = [Tangents](int VertexID, int TriangleID, int TriVtxIdx, FVector3f& TangentX, FVector3f& TangentY)
@@ -367,7 +371,7 @@ public:
 			check(Mesh->HasAttributes());
 			NormalOverlay = Mesh->Attributes()->PrimaryNormals();
 
-			const FMeshTangentsf* Tangents = ParentComponent->GetTangents();
+			const UE::Geometry::FMeshTangentsf* Tangents = ParentComponent->GetTangents();
 			if (Tangents != nullptr)
 			{
 				TangentsFunc = [Tangents](int VertexID, int TriangleID, int TriVtxIdx, const FVector3f& Normal, FVector3f& TangentX, FVector3f& TangentY)
@@ -468,7 +472,7 @@ public:
 			check(Mesh->HasAttributes());
 			NormalOverlay = Mesh->Attributes()->PrimaryNormals();
 
-			const FMeshTangentsf* Tangents = ParentComponent->GetTangents();
+			const UE::Geometry::FMeshTangentsf* Tangents = ParentComponent->GetTangents();
 			if (Tangents != nullptr)
 			{
 				TangentsFunc = [Tangents](int VertexID, int TriangleID, int TriVtxIdx, const FVector3f& Normal, FVector3f& TangentX, FVector3f& TangentY)
diff --git a/Engine/Plugins/Experimental/MeshModelingToolset/Source/ModelingComponents/Public/AssetGenerationUtil.h b/Engine/Plugins/Experimental/MeshModelingToolset/Source/ModelingComponents/Public/AssetGenerationUtil.h
--- a/Engine/Plugins/Experimental/MeshModelingToolset/Source/ModelingComponents/Public/AssetGenerationUtil.h
+++ b/Engine/Plugins/Experimental/MeshModelingToolset/Source/ModelingComponents/Public/AssetGenerationUtil.h
@@ -9,7 +9,7 @@
 class UMaterialInterface;
 class UInteractiveToolManager;
 struct FMeshDescription;
-class FDynamicMesh3;
+PREDECLARE_USE_GEOMETRY_CLASS(FDynamicMesh3);
 class AActor;
 class UPrimitiveComponent;
 class UTexture2D;
@@ -71,7 +71,7 @@ namespace AssetGenerationUtil
 		IAssetGenerationAPI* AssetAPI,
 		UWorld* TargetWorld,
 		const FDynamicMesh3* Mesh,
-		const FTransform3d& Transform,
+		const UE::Geometry::FTransform3d& Transform,
 		FString ObjectName,
 		UMaterialInterface* Material = nullptr
 	);
@@ -91,7 +91,7 @@ namespace AssetGenerationUtil
 		IAssetGenerationAPI* AssetAPI,
 		UWorld* TargetWorld,
 		const FDynamicMesh3* Mesh,
-		const FTransform3d& Transform,
+		const UE::Geometry::FTransform3d& Transform,
 		FString ObjectName,
 		const TArrayView<UMaterialInterface*>& Materials
 	);
diff --git a/Engine/Plugins/Experimental/MeshModelingToolset/Source/ModelingComponents/Public/BaseDynamicMeshComponent.h b/Engine/Plugins/Experimental/MeshModelingToolset/Source/ModelingComponents/Public/BaseDynamicMeshComponent.h
--- a/Engine/Plugins/Experimental/MeshModelingToolset/Source/ModelingComponents/Public/BaseDynamicMeshComponent.h
+++ b/Engine/Plugins/Experimental/MeshModelingToolset/Source/ModelingComponents/Public/BaseDynamicMeshComponent.h
@@ -18,7 +18,7 @@
 struct FMeshDescription;
 class FMeshVertexChange;
 class FMeshChange;
-
+using UE::Geometry::FDynamicMesh3;
 
 
 /**
@@ -131,7 +131,7 @@ public:
 	//
 public:
 
-	virtual void ApplyTransform(const FTransform3d& Transform, bool bInvert)
+	virtual void ApplyTransform(const UE::Geometry::FTransform3d& Transform, bool bInvert)
 	{
 		unimplemented();
 	}
diff --git a/Engine/Plugins/Experimental/MeshModelingToolset/Source/ModelingComponents/Public/BaseTools/BaseCreateFromSelectedTool.h b/Engine/Plugins/Experimental/MeshModelingToolset/Source/ModelingComponents/Public/BaseTools/BaseCreateFromSelectedTool.h
--- a/Engine/Plugins/Experimental/MeshModelingToolset/Source/ModelingComponents/Public/BaseTools/BaseCreateFromSelectedTool.h
+++ b/Engine/Plugins/Experimental/MeshModelingToolset/Source/ModelingComponents/Public/BaseTools/BaseCreateFromSelectedTool.h
@@ -94,10 +94,17 @@ public:
  * that provides support for common functionality in tools that create a new mesh from a selection of one or more existing meshes
  */
 UCLASS()
-class MODELINGCOMPONENTS_API UBaseCreateFromSelectedTool : public UMultiSelectionTool, public IDynamicMeshOperatorFactory
+class MODELINGCOMPONENTS_API UBaseCreateFromSelectedTool : public UMultiSelectionTool, public UE::Geometry::IDynamicMeshOperatorFactory
 {
 	GENERATED_BODY()
-
+protected:
+	using FVector3d = UE::Geometry::FVector3d;
+	using FVector3f = UE::Geometry::FVector3f;
+	using FVector2d = UE::Geometry::FVector2d;
+	using FVector2f = UE::Geometry::FVector2f;
+	using FTransform3d = UE::Geometry::FTransform3d;
+	using FFrame3d = UE::Geometry::FFrame3d;
+	using FRay3d = UE::Geometry::FRay3d;
 public:
 	UBaseCreateFromSelectedTool() = default;
 
@@ -159,10 +166,10 @@ protected:
 	/**
 	 * IDynamicMeshOperatorFactory implementation that subclass must override and implement
 	 */
-	virtual TUniquePtr<FDynamicMeshOperator> MakeNewOperator() override
+	virtual TUniquePtr<UE::Geometry::FDynamicMeshOperator> MakeNewOperator() override
 	{
 		check(false);
-		return TUniquePtr<FDynamicMeshOperator>();
+		return TUniquePtr<UE::Geometry::FDynamicMeshOperator>();
 	}
 
 protected:
diff --git a/Engine/Plugins/Experimental/MeshModelingToolset/Source/ModelingComponents/Public/BaseTools/BaseMeshProcessingTool.h b/Engine/Plugins/Experimental/MeshModelingToolset/Source/ModelingComponents/Public/BaseTools/BaseMeshProcessingTool.h
--- a/Engine/Plugins/Experimental/MeshModelingToolset/Source/ModelingComponents/Public/BaseTools/BaseMeshProcessingTool.h
+++ b/Engine/Plugins/Experimental/MeshModelingToolset/Source/ModelingComponents/Public/BaseTools/BaseMeshProcessingTool.h
@@ -14,11 +14,12 @@
 
 // predeclarations
 struct FMeshDescription;
-class FMeshNormals;
-class FMeshBoundaryLoops;
 class USimpleDynamicMeshComponent;
 class UPreviewMesh;
-
+PREDECLARE_USE_GEOMETRY_CLASS(FMeshBoundaryLoops);
+PREDECLARE_USE_GEOMETRY_CLASS(FMeshNormals);
+using UE::Geometry::FDynamicMesh3;
+using UE::Geometry::FIndexedWeightMap1f;
 /**
  * ToolBuilder for UBaseMeshProcessingTool
  */
@@ -74,10 +75,17 @@ public:
  *
  */
 UCLASS()
-class MODELINGCOMPONENTS_API UBaseMeshProcessingTool : public USingleSelectionTool, public IDynamicMeshOperatorFactory
+class MODELINGCOMPONENTS_API UBaseMeshProcessingTool : public USingleSelectionTool, public UE::Geometry::IDynamicMeshOperatorFactory
 {
 	GENERATED_BODY()
-
+protected:
+	using FVector3d = UE::Geometry::FVector3d;
+	using FVector3f = UE::Geometry::FVector3f;
+	using FVector2d = UE::Geometry::FVector2d;
+	using FVector2f = UE::Geometry::FVector2f;
+	using FTransform3d = UE::Geometry::FTransform3d;
+	using FFrame3d = UE::Geometry::FFrame3d;
+	using FRay3d = UE::Geometry::FRay3d;
 public:
 	UBaseMeshProcessingTool() = default;
 
@@ -107,10 +115,10 @@ protected:
 	/**
 	 * IDynamicMeshOperatorFactory implementation that subclass must override and implement
 	 */
-	virtual TUniquePtr<FDynamicMeshOperator> MakeNewOperator() override
+	virtual TUniquePtr<UE::Geometry::FDynamicMeshOperator> MakeNewOperator() override
 	{
 		check(false);
-		return TUniquePtr<FDynamicMeshOperator>();
+		return TUniquePtr<UE::Geometry::FDynamicMeshOperator>();
 	}
 
 
diff --git a/Engine/Plugins/Experimental/MeshModelingToolset/Source/ModelingComponents/Public/BaseTools/BaseVoxelTool.h b/Engine/Plugins/Experimental/MeshModelingToolset/Source/ModelingComponents/Public/BaseTools/BaseVoxelTool.h
--- a/Engine/Plugins/Experimental/MeshModelingToolset/Source/ModelingComponents/Public/BaseTools/BaseVoxelTool.h
+++ b/Engine/Plugins/Experimental/MeshModelingToolset/Source/ModelingComponents/Public/BaseTools/BaseVoxelTool.h
@@ -9,7 +9,7 @@
 
 #include "BaseVoxelTool.generated.h"
 
-class FDynamicMesh3;
+PREDECLARE_USE_GEOMETRY_CLASS(FDynamicMesh3);
 
 /**
  * Base for Voxel tools
diff --git a/Engine/Plugins/Experimental/MeshModelingToolset/Source/ModelingComponents/Public/Changes/DynamicMeshChangeTarget.h b/Engine/Plugins/Experimental/MeshModelingToolset/Source/ModelingComponents/Public/Changes/DynamicMeshChangeTarget.h
--- a/Engine/Plugins/Experimental/MeshModelingToolset/Source/ModelingComponents/Public/Changes/DynamicMeshChangeTarget.h
+++ b/Engine/Plugins/Experimental/MeshModelingToolset/Source/ModelingComponents/Public/Changes/DynamicMeshChangeTarget.h
@@ -13,7 +13,7 @@
 #include "DynamicMeshChangeTarget.generated.h"
 
 
-class FDynamicMesh3;
+PREDECLARE_USE_GEOMETRY_CLASS(FDynamicMesh3);
 
 /**
  * Bare bones wrapper of FDynamicMesh3 that supports MeshReplacementChange-based updates.  Shares the mesh ptr with the MeshReplacementChange objects, so it must not be changed directly
diff --git a/Engine/Plugins/Experimental/MeshModelingToolset/Source/ModelingComponents/Public/Changes/MeshChange.h b/Engine/Plugins/Experimental/MeshModelingToolset/Source/ModelingComponents/Public/Changes/MeshChange.h
--- a/Engine/Plugins/Experimental/MeshModelingToolset/Source/ModelingComponents/Public/Changes/MeshChange.h
+++ b/Engine/Plugins/Experimental/MeshModelingToolset/Source/ModelingComponents/Public/Changes/MeshChange.h
@@ -22,9 +22,9 @@ class MODELINGCOMPONENTS_API FMeshChange : public FToolCommandChange
 {
 public:
 	FMeshChange();
-	FMeshChange(TUniquePtr<FDynamicMeshChange> DynamicMeshChangeIn);
+	FMeshChange(TUniquePtr<UE::Geometry::FDynamicMeshChange> DynamicMeshChangeIn);
 
-	TUniquePtr<FDynamicMeshChange> DynamicMeshChange;
+	TUniquePtr<UE::Geometry::FDynamicMeshChange> DynamicMeshChange;
 
 	/** This function is called on Apply and Revert (last argument is true on Apply)*/
 	TFunction<void(FMeshChange*, UObject*, bool)> OnChangeAppliedFunc;
@@ -45,7 +45,7 @@ public:
 	 * This allows FMeshChange subclasses to customize the change behavior if necessary.
 	 * The default behavior just forwards the call to DynamicMeshChange->Apply(Mesh, bRevert).
 	 */
-	virtual void ApplyChangeToMesh(FDynamicMesh3* Mesh, bool bRevert) const;
+	virtual void ApplyChangeToMesh(UE::Geometry::FDynamicMesh3* Mesh, bool bRevert) const;
 };
 
 
diff --git a/Engine/Plugins/Experimental/MeshModelingToolset/Source/ModelingComponents/Public/Changes/MeshPolygroupChange.h b/Engine/Plugins/Experimental/MeshModelingToolset/Source/ModelingComponents/Public/Changes/MeshPolygroupChange.h
--- a/Engine/Plugins/Experimental/MeshModelingToolset/Source/ModelingComponents/Public/Changes/MeshPolygroupChange.h
+++ b/Engine/Plugins/Experimental/MeshModelingToolset/Source/ModelingComponents/Public/Changes/MeshPolygroupChange.h
@@ -7,8 +7,10 @@
 #include "InteractiveToolChange.h"
 #include "VectorTypes.h"
 #include "Changes/MeshChange.h"
+#include "GeometryBase.h"
 #include "Polygroups/PolygroupSet.h"
 
+PREDECLARE_USE_GEOMETRY_CLASS(FDynamicMesh3);
 
 /**
  * FDynamicMeshGroupEdit stores a modification of polygroup IDs on a set of triangles
diff --git a/Engine/Plugins/Experimental/MeshModelingToolset/Source/ModelingComponents/Public/Changes/MeshReplacementChange.h b/Engine/Plugins/Experimental/MeshModelingToolset/Source/ModelingComponents/Public/Changes/MeshReplacementChange.h
--- a/Engine/Plugins/Experimental/MeshModelingToolset/Source/ModelingComponents/Public/Changes/MeshReplacementChange.h
+++ b/Engine/Plugins/Experimental/MeshModelingToolset/Source/ModelingComponents/Public/Changes/MeshReplacementChange.h
@@ -5,12 +5,11 @@
 #include "CoreMinimal.h"
 #include "UObject/Interface.h"
 #include "InteractiveToolChange.h"
+#include "GeometryBase.h"
 #include "MeshReplacementChange.generated.h"
 
 
-class FDynamicMesh3;
-
-
+PREDECLARE_USE_GEOMETRY_CLASS(FDynamicMesh3);
 
 
 /**
diff --git a/Engine/Plugins/Experimental/MeshModelingToolset/Source/ModelingComponents/Public/Changes/MeshVertexChange.h b/Engine/Plugins/Experimental/MeshModelingToolset/Source/ModelingComponents/Public/Changes/MeshVertexChange.h
--- a/Engine/Plugins/Experimental/MeshModelingToolset/Source/ModelingComponents/Public/Changes/MeshVertexChange.h
+++ b/Engine/Plugins/Experimental/MeshModelingToolset/Source/ModelingComponents/Public/Changes/MeshVertexChange.h
@@ -8,7 +8,7 @@
 #include "VectorTypes.h"
 #include "MeshVertexChange.generated.h"
 
-class FDynamicMesh3;
+PREDECLARE_USE_GEOMETRY_CLASS(FDynamicMesh3);
 
 /**
  * FMeshVertexChange represents an reversible change to a set of vertex positions.
@@ -19,6 +19,8 @@ class FDynamicMesh3;
  */
 class MODELINGCOMPONENTS_API FMeshVertexChange : public FToolCommandChange
 {
+	using FVector3d = UE::Geometry::FVector3d;
+	using FVector3f = UE::Geometry::FVector3f;
 public:
 	bool bHaveVertexPositions = true;
 	bool bHaveVertexColors = false;
@@ -59,6 +61,8 @@ ENUM_CLASS_FLAGS(EMeshVertexChangeComponents);
  */
 class MODELINGCOMPONENTS_API FMeshVertexChangeBuilder
 {
+	using FVector3d = UE::Geometry::FVector3d;
+	using FVector3f = UE::Geometry::FVector3f;
 public:
 	TUniquePtr<FMeshVertexChange> Change;
 	TMap<int32, int32> SavedVertices;
diff --git a/Engine/Plugins/Experimental/MeshModelingToolset/Source/ModelingComponents/Public/Drawing/MeshDebugDrawing.h b/Engine/Plugins/Experimental/MeshModelingToolset/Source/ModelingComponents/Public/Drawing/MeshDebugDrawing.h
--- a/Engine/Plugins/Experimental/MeshModelingToolset/Source/ModelingComponents/Public/Drawing/MeshDebugDrawing.h
+++ b/Engine/Plugins/Experimental/MeshModelingToolset/Source/ModelingComponents/Public/Drawing/MeshDebugDrawing.h
@@ -8,12 +8,15 @@
 #include "ToolContextInterfaces.h"  // for FViewCameraState
 
 class FPrimitiveDrawInterface;
+using UE::Geometry::FDynamicMesh3;
+using UE::Geometry::FDynamicMeshNormalOverlay;
 
 /**
  * drawing utility functions useful for debugging. These are generally not performant.
  */
 namespace MeshDebugDraw
 {
+	using namespace UE::Geometry;
 
 	/**
 	 * Draw normals of mesh overlay as lines
diff --git a/Engine/Plugins/Experimental/MeshModelingToolset/Source/ModelingComponents/Public/Drawing/MeshElementsVisualizer.h b/Engine/Plugins/Experimental/MeshModelingToolset/Source/ModelingComponents/Public/Drawing/MeshElementsVisualizer.h
--- a/Engine/Plugins/Experimental/MeshModelingToolset/Source/ModelingComponents/Public/Drawing/MeshElementsVisualizer.h
+++ b/Engine/Plugins/Experimental/MeshModelingToolset/Source/ModelingComponents/Public/Drawing/MeshElementsVisualizer.h
@@ -6,9 +6,10 @@
 #include "Drawing/PreviewGeometryActor.h"
 #include "InteractiveTool.h"
 #include "MeshWireframeComponent.h"
+#include "GeometryBase.h"
 #include "MeshElementsVisualizer.generated.h"
 
-class FDynamicMesh3;
+PREDECLARE_GEOMETRY(class FDynamicMesh3);
 
 /**
  * Visualization settings UMeshElementsVisualizer
@@ -86,7 +87,7 @@ public:
 	 * will access the mesh when creating their Scene Proxies/etc. The callback can return null
 	 * if the mesh has become invalid/etc
 	 */
-	void SetMeshAccessFunction(TUniqueFunction<const FDynamicMesh3* (void)>&& MeshAccessFunction);
+	void SetMeshAccessFunction(TUniqueFunction<const UE::Geometry::FDynamicMesh3* (void)>&& MeshAccessFunction);
 
 	/**
 	 * Call if mesh provided by MeshAccessFunction has been modified, will cause a full recomputation
diff --git a/Engine/Plugins/Experimental/MeshModelingToolset/Source/ModelingComponents/Public/Drawing/MeshRenderDecomposition.h b/Engine/Plugins/Experimental/MeshModelingToolset/Source/ModelingComponents/Public/Drawing/MeshRenderDecomposition.h
--- a/Engine/Plugins/Experimental/MeshModelingToolset/Source/ModelingComponents/Public/Drawing/MeshRenderDecomposition.h
+++ b/Engine/Plugins/Experimental/MeshModelingToolset/Source/ModelingComponents/Public/Drawing/MeshRenderDecomposition.h
@@ -3,8 +3,9 @@
 #pragma once
 
 #include "CoreMinimal.h"
+#include "GeometryBase.h"
 
-class FDynamicMesh3;
+PREDECLARE_USE_GEOMETRY_CLASS(FDynamicMesh3);
 class UMaterialInterface;
 struct FComponentMaterialSet;
 
diff --git a/Engine/Plugins/Experimental/MeshModelingToolset/Source/ModelingComponents/Public/Drawing/PolyEditPreviewMesh.h b/Engine/Plugins/Experimental/MeshModelingToolset/Source/ModelingComponents/Public/Drawing/PolyEditPreviewMesh.h
--- a/Engine/Plugins/Experimental/MeshModelingToolset/Source/ModelingComponents/Public/Drawing/PolyEditPreviewMesh.h
+++ b/Engine/Plugins/Experimental/MeshModelingToolset/Source/ModelingComponents/Public/Drawing/PolyEditPreviewMesh.h
@@ -18,6 +18,9 @@ UCLASS(Transient)
 class MODELINGCOMPONENTS_API UPolyEditPreviewMesh : public UPreviewMesh
 {
 	GENERATED_BODY()
+	using FVector3d = UE::Geometry::FVector3d;
+	using FVector3f = UE::Geometry::FVector3f;
+	using FTransform3d = UE::Geometry::FTransform3d;	
 public:
 
 	//
@@ -63,12 +66,12 @@ public:
 	const FDynamicMesh3& GetInitialPatchMesh() const { return InitialEditPatch; }
 
 	// must be non-const because query functions are non-const
-	FDynamicMeshAABBTree3& GetInitialPatchMeshSpatial() { return InitialEditPatchBVTree; }
+	UE::Geometry::FDynamicMeshAABBTree3& GetInitialPatchMeshSpatial() { return InitialEditPatchBVTree; }
 
 protected:
-	TUniquePtr<FDynamicSubmesh3> ActiveSubmesh;
+	TUniquePtr<UE::Geometry::FDynamicSubmesh3> ActiveSubmesh;
 	FDynamicMesh3 InitialEditPatch;
-	FDynamicMeshAABBTree3 InitialEditPatchBVTree;
+	UE::Geometry::FDynamicMeshAABBTree3 InitialEditPatchBVTree;
 
 	TArray<int32> EditVertices;
 	TArray<FVector3d> InitialPositions;
@@ -78,6 +81,6 @@ protected:
 
 	FVector3d InputDirection;
 
-	FTransform3d MeshTransform;
+	UE::Geometry::FTransform3d MeshTransform;
 	bool bHaveMeshTransform;
 };
\ No newline at end of file
diff --git a/Engine/Plugins/Experimental/MeshModelingToolset/Source/ModelingComponents/Public/Drawing/TriangleSetComponent.h b/Engine/Plugins/Experimental/MeshModelingToolset/Source/ModelingComponents/Public/Drawing/TriangleSetComponent.h
--- a/Engine/Plugins/Experimental/MeshModelingToolset/Source/ModelingComponents/Public/Drawing/TriangleSetComponent.h
+++ b/Engine/Plugins/Experimental/MeshModelingToolset/Source/ModelingComponents/Public/Drawing/TriangleSetComponent.h
@@ -119,7 +119,7 @@ public:
 	 * Add a Quad (two triangles) with the given vertices, normal, Color, and Material
 	 * @return ID of the two triangles created
 	 */
-	FIndex2i AddQuad(const FVector& A, const FVector& B, const FVector& C, const FVector& D, const FVector& Normal, const FColor& Color, UMaterialInterface* Material);
+	UE::Geometry::FIndex2i AddQuad(const FVector& A, const FVector& B, const FVector& C, const FVector& D, const FVector& Normal, const FColor& Color, UMaterialInterface* Material);
 
 private:
 
diff --git a/Engine/Plugins/Experimental/MeshModelingToolset/Source/ModelingComponents/Public/Drawing/UVLayoutPreview.h b/Engine/Plugins/Experimental/MeshModelingToolset/Source/ModelingComponents/Public/Drawing/UVLayoutPreview.h
--- a/Engine/Plugins/Experimental/MeshModelingToolset/Source/ModelingComponents/Public/Drawing/UVLayoutPreview.h
+++ b/Engine/Plugins/Experimental/MeshModelingToolset/Source/ModelingComponents/Public/Drawing/UVLayoutPreview.h
@@ -8,9 +8,7 @@
 #include "InteractiveTool.h"
 #include "UVLayoutPreview.generated.h"
 
-class FDynamicMesh3;
-
-
+PREDECLARE_USE_GEOMETRY_CLASS(FDynamicMesh3);
 
 /**
  * Where should in-viewport UVLayoutPreview be shown, relative to target object
@@ -149,10 +147,10 @@ protected:
 
 	FComponentMaterialSet SourceMaterials;
 
-	FFrame3d SourceObjectFrame;
-	FAxisAlignedBox3d SourceObjectWorldBounds;
+	UE::Geometry::FFrame3d SourceObjectFrame;
+	UE::Geometry::FAxisAlignedBox3d SourceObjectWorldBounds;
 
-	FFrame3d CurrentWorldFrame;
+	UE::Geometry::FFrame3d CurrentWorldFrame;
 
 	FViewCameraState CameraState;
 
diff --git a/Engine/Plugins/Experimental/MeshModelingToolset/Source/ModelingComponents/Public/Mechanics/CollectSurfacePathMechanic.h b/Engine/Plugins/Experimental/MeshModelingToolset/Source/ModelingComponents/Public/Mechanics/CollectSurfacePathMechanic.h
--- a/Engine/Plugins/Experimental/MeshModelingToolset/Source/ModelingComponents/Public/Mechanics/CollectSurfacePathMechanic.h
+++ b/Engine/Plugins/Experimental/MeshModelingToolset/Source/ModelingComponents/Public/Mechanics/CollectSurfacePathMechanic.h
@@ -9,6 +9,7 @@
 #include "ToolDataVisualizer.h"
 #include "CollectSurfacePathMechanic.generated.h"
 
+using UE::Geometry::FDynamicMesh3;
 
 enum class ECollectSurfacePathDoneMode
 {
@@ -28,6 +29,9 @@ class MODELINGCOMPONENTS_API UCollectSurfacePathMechanic : public UInteractionMe
 {
 	GENERATED_BODY()
 public:
+	using FFrame3d = UE::Geometry::FFrame3d;
+	using FRay3d = UE::Geometry::FRay3d;
+	using FVector3d = UE::Geometry::FVector3d;
 
 	TUniqueFunction<bool()> IsDoneFunc = nullptr;
 
@@ -85,7 +89,7 @@ public:
 
 protected:
 	FDynamicMesh3 TargetSurface;
-	FDynamicMeshAABBTree3 TargetSurfaceAABB;
+	UE::Geometry::FDynamicMeshAABBTree3 TargetSurfaceAABB;
 
 	FFrame3d TargetPlane;
 	bool bHaveTargetPlane;
diff --git a/Engine/Plugins/Experimental/MeshModelingToolset/Source/ModelingComponents/Public/Mechanics/ConstructionPlaneMechanic.h b/Engine/Plugins/Experimental/MeshModelingToolset/Source/ModelingComponents/Public/Mechanics/ConstructionPlaneMechanic.h
--- a/Engine/Plugins/Experimental/MeshModelingToolset/Source/ModelingComponents/Public/Mechanics/ConstructionPlaneMechanic.h
+++ b/Engine/Plugins/Experimental/MeshModelingToolset/Source/ModelingComponents/Public/Mechanics/ConstructionPlaneMechanic.h
@@ -32,7 +32,7 @@ public:
 
 	bool bShowGrid = true;
 
-	FFrame3d Plane;
+	UE::Geometry::FFrame3d Plane;
 
 	DECLARE_MULTICAST_DELEGATE(OnConstructionPlaneChangedEvent);
 	OnConstructionPlaneChangedEvent OnPlaneChanged;
@@ -45,16 +45,16 @@ public:
 	virtual void Render(IToolsContextRenderAPI* RenderAPI) override;
 	virtual void Tick(float DeltaTime) override;
 
-	void Initialize(UWorld* TargetWorld, const FFrame3d& InitialPlane);
+	void Initialize(UWorld* TargetWorld, const UE::Geometry::FFrame3d& InitialPlane);
 
-	void SetDrawPlaneFromWorldPos(const FVector3d& Position, const FVector3d& Normal, bool bIgnoreNormal);
+	void SetDrawPlaneFromWorldPos(const UE::Geometry::FVector3d& Position, const UE::Geometry::FVector3d& Normal, bool bIgnoreNormal);
 
 	/** 
 	 * Sets the plane without broadcasting OnPlaneChanged. Useful when the user of the tool wants to change
 	 * the plane through some other means. Better than setting the Plane field directly because this function
 	 * properly deals with the gizmo.
 	 */
-	void SetPlaneWithoutBroadcast(const FFrame3d& Plane);
+	void SetPlaneWithoutBroadcast(const UE::Geometry::FFrame3d& Plane);
 
 	void SetEnableGridSnaping(bool bEnable);
 
diff --git a/Engine/Plugins/Experimental/MeshModelingToolset/Source/ModelingComponents/Public/Mechanics/CurveControlPointsMechanic.h b/Engine/Plugins/Experimental/MeshModelingToolset/Source/ModelingComponents/Public/Mechanics/CurveControlPointsMechanic.h
--- a/Engine/Plugins/Experimental/MeshModelingToolset/Source/ModelingComponents/Public/Mechanics/CurveControlPointsMechanic.h
+++ b/Engine/Plugins/Experimental/MeshModelingToolset/Source/ModelingComponents/Public/Mechanics/CurveControlPointsMechanic.h
@@ -45,6 +45,7 @@ class MODELINGCOMPONENTS_API UCurveControlPointsMechanic : public UInteractionMe
 	GENERATED_BODY()
 
 protected:
+	using FVector3d = UE::Geometry::FVector3d;
 
 	// We want some way to store the control point sequence that lets us easily associate points with their renderable and hit-testable
 	// representations, since we need to alter all of these together as points get moved or added. We use FOrderedPoints for this, until
@@ -246,7 +247,7 @@ public:
 	}
 
 	/** Sets the plane on which new points are added on the ends and in which the points are moved. */
-	void SetPlane(const FFrame3d& DrawPlaneIn);
+	void SetPlane(const UE::Geometry::FFrame3d& DrawPlaneIn);
 	// TODO: It is simple to allow the points to be moved arbitrarily, not just inside the plane, if we ever
 	// want to use the mechanic somewhere where that is desirable. However, we'd need to do a little more work
 	// to allow new points to be added in arbitrary locations.
@@ -255,7 +256,7 @@ public:
 
 	// Adds additional line to snap points to. Useful, for instance, if the curve is a revolution profile curve
 	// and needs to be able to snap to the revolution axis.
-	void AddSnapLine(int32 LineID, const FLine3d& Line);
+	void AddSnapLine(int32 LineID, const UE::Geometry::FLine3d& Line);
 
 	void RemoveSnapLine(int32 LineID);
 
@@ -320,7 +321,7 @@ protected:
 	bool bAutoRevertToInteractiveInitialization = false;
 
 	bool bSnappingEnabled = true;
-	FPointPlanarSnapSolver SnapEngine;
+	UE::Geometry::FPointPlanarSnapSolver SnapEngine;
 
 	// Used for snapping to the start/end of the curve to get out of initialization mode
 	int32 FirstPointSnapID;
@@ -334,7 +335,7 @@ protected:
 	int32 LineSnapPriority;
 	
 	// Used for spatial queries
-	FGeometrySet3 GeometrySet;
+	UE::Geometry::FGeometrySet3 GeometrySet;
 
 	/** Used for displaying points/segments */
 	UPROPERTY()
@@ -365,7 +366,7 @@ protected:
 	FColor SnapLineColor;
 
 	// Used for adding new points on the ends and for limiting point movement
-	FFrame3d DrawPlane;
+	UE::Geometry::FFrame3d DrawPlane;
 
 	// Support for Shift and Ctrl toggles
 	bool bAddToSelectionToggle = false;
@@ -456,6 +457,7 @@ protected:
 
 class MODELINGCOMPONENTS_API FCurveControlPointsMechanicInsertionChange : public FToolCommandChange
 {
+	using FVector3d = UE::Geometry::FVector3d;
 public:
 	FCurveControlPointsMechanicInsertionChange(int32 SequencePositionIn, int32 PointID, 
 		const FVector3d& CoordinatesIn, bool AddedIn, int32 ChangeStampIn);
@@ -497,6 +499,7 @@ protected:
 
 class MODELINGCOMPONENTS_API FCurveControlPointsMechanicMovementChange : public FToolCommandChange
 {
+	using FVector3d = UE::Geometry::FVector3d;
 public:
 	FCurveControlPointsMechanicMovementChange(int32 PointIDIn, const FVector3d& OriginalPositionIn, 
 		const FVector3d& NewPositionIn, int32 ChangeStampIn);
diff --git a/Engine/Plugins/Experimental/MeshModelingToolset/Source/ModelingComponents/Public/Mechanics/DragAlignmentMechanic.h b/Engine/Plugins/Experimental/MeshModelingToolset/Source/ModelingComponents/Public/Mechanics/DragAlignmentMechanic.h
--- a/Engine/Plugins/Experimental/MeshModelingToolset/Source/ModelingComponents/Public/Mechanics/DragAlignmentMechanic.h
+++ b/Engine/Plugins/Experimental/MeshModelingToolset/Source/ModelingComponents/Public/Mechanics/DragAlignmentMechanic.h
@@ -14,7 +14,7 @@
 
 #include "DragAlignmentMechanic.generated.h"
 
-class FGroupTopologyDeformer;
+PREDECLARE_USE_GEOMETRY_CLASS(FGroupTopologyDeformer);
 class UPrimitiveComponent;
 class UIntervalGizmo;
 class UTransformGizmo;
@@ -53,8 +53,8 @@ public:
 	 * This works with only one deformed mesh, so it probably wouldn't be used with multiple gizmos.
 	 */
 	void InitializeDeformedMeshRayCast(
-		TFunction<FDynamicMeshAABBTree3* ()> GetSpatialIn,
-		const FTransform3d& TargetTransform, const FGroupTopologyDeformer* LinearDeformer);
+		TFunction<UE::Geometry::FDynamicMeshAABBTree3* ()> GetSpatialIn,
+		const UE::Geometry::FTransform3d& TargetTransform, const FGroupTopologyDeformer* LinearDeformer);
 
 	virtual void Setup(UInteractiveTool* ParentToolIn) override;
 	virtual void Shutdown() override;
@@ -90,7 +90,7 @@ protected:
 	bool bPreviewEndpointsValid = false;
 	bool bWaitingOnProjectedResult = false;
 	FHitResult LastHitResult;
-	FVector3d LastProjectedResult;
+	UE::Geometry::FVector3d LastProjectedResult;
 
 	bool bAlignmentToggle = false;
 	int32 AlignmentModifierID = 1;
diff --git a/Engine/Plugins/Experimental/MeshModelingToolset/Source/ModelingComponents/Public/Mechanics/LatticeControlPointsMechanic.h b/Engine/Plugins/Experimental/MeshModelingToolset/Source/ModelingComponents/Public/Mechanics/LatticeControlPointsMechanic.h
--- a/Engine/Plugins/Experimental/MeshModelingToolset/Source/ModelingComponents/Public/Mechanics/LatticeControlPointsMechanic.h
+++ b/Engine/Plugins/Experimental/MeshModelingToolset/Source/ModelingComponents/Public/Mechanics/LatticeControlPointsMechanic.h
@@ -31,6 +31,9 @@ class MODELINGCOMPONENTS_API ULatticeControlPointsMechanic :
 	public UInteractionMechanic, public IClickBehaviorTarget, public IHoverBehaviorTarget
 {
 	GENERATED_BODY()
+	using FVector3d = UE::Geometry::FVector3d;
+	using FVector2i = UE::Geometry::FVector2i;
+	using FTransform3d = UE::Geometry::FTransform3d;
 
 public:
 
@@ -79,7 +82,7 @@ protected:
 	FTransform3d LocalToWorldTransform;
 
 	// Used for spatial queries
-	FGeometrySet3 GeometrySet;
+	UE::Geometry::FGeometrySet3 GeometrySet;
 	FViewCameraState CachedCameraState;
 
 	/** Used for displaying points/segments */
@@ -208,6 +211,7 @@ protected:
 // Control points have moved
 class MODELINGCOMPONENTS_API FLatticeControlPointsMechanicMovementChange : public FToolCommandChange
 {
+	using FVector3d = UE::Geometry::FVector3d;
 public:
 
 	FLatticeControlPointsMechanicMovementChange(
diff --git a/Engine/Plugins/Experimental/MeshModelingToolset/Source/ModelingComponents/Public/Mechanics/PlaneDistanceFromHitMechanic.h b/Engine/Plugins/Experimental/MeshModelingToolset/Source/ModelingComponents/Public/Mechanics/PlaneDistanceFromHitMechanic.h
--- a/Engine/Plugins/Experimental/MeshModelingToolset/Source/ModelingComponents/Public/Mechanics/PlaneDistanceFromHitMechanic.h
+++ b/Engine/Plugins/Experimental/MeshModelingToolset/Source/ModelingComponents/Public/Mechanics/PlaneDistanceFromHitMechanic.h
@@ -18,7 +18,8 @@ UCLASS()
 class MODELINGCOMPONENTS_API UPlaneDistanceFromHitMechanic : public UInteractionMechanic
 {
 	GENERATED_BODY()
-
+	using FVector3d = UE::Geometry::FVector3d;
+	using FFrame3d = UE::Geometry::FFrame3d;
 public:
 	/** If this function is set, we will check it for a ray intersection if the target mesh is not hit */
 	TUniqueFunction<bool(const FRay&, FHitResult&)> WorldHitQueryFunc = nullptr;
@@ -36,7 +37,7 @@ public:
 	FHitResult LastActiveWorldHit;
 
 	/** World frame at last hit point */
-	FFrame3d CurrentHitPosFrameWorld;
+	UE::Geometry::FFrame3d CurrentHitPosFrameWorld;
 
 	/** If true, then if we don't find any intersection point, just use nearest point on plane normal to ray */
 	bool bFallbackToLineAxisPoint = true;
@@ -53,7 +54,7 @@ public:
 	 * If bMeshInWorldCoords is true, then HitTargetMesh is in world coords.
 	 * Otherwise we assume it is in local coords of PlaneFrameWorld
 	 */
-	virtual void Initialize(FDynamicMesh3&& HitTargetMesh, const FFrame3d& PlaneFrameWorld, bool bMeshInWorldCoords);
+	virtual void Initialize(UE::Geometry::FDynamicMesh3&& HitTargetMesh, const FFrame3d& PlaneFrameWorld, bool bMeshInWorldCoords);
 
 	/**
 	 * Update the current distance/height based on the input world ray
@@ -61,8 +62,8 @@ public:
 	virtual void UpdateCurrentDistance(const FRay& WorldRay);
 
 protected:
-	FDynamicMesh3 PreviewHeightTarget;
-	FDynamicMeshAABBTree3 PreviewHeightTargetAABB;
+	UE::Geometry::FDynamicMesh3 PreviewHeightTarget;
+	UE::Geometry::FDynamicMeshAABBTree3 PreviewHeightTargetAABB;
 	FFrame3d PreviewHeightFrame;
 
 };
diff --git a/Engine/Plugins/Experimental/MeshModelingToolset/Source/ModelingComponents/Public/Mechanics/SpaceCurveDeformationMechanic.h b/Engine/Plugins/Experimental/MeshModelingToolset/Source/ModelingComponents/Public/Mechanics/SpaceCurveDeformationMechanic.h
--- a/Engine/Plugins/Experimental/MeshModelingToolset/Source/ModelingComponents/Public/Mechanics/SpaceCurveDeformationMechanic.h
+++ b/Engine/Plugins/Experimental/MeshModelingToolset/Source/ModelingComponents/Public/Mechanics/SpaceCurveDeformationMechanic.h
@@ -31,7 +31,7 @@ class MODELINGCOMPONENTS_API FSpaceCurveSource
 {
 public:
 	TUniqueFunction<int32()> GetPointCount;
-	TUniqueFunction<FFrame3d(int32)> GetPoint;
+	TUniqueFunction<UE::Geometry::FFrame3d(int32)> GetPoint;
 	TUniqueFunction<bool()> IsLoop;
 };
 
@@ -68,6 +68,7 @@ UCLASS()
 class MODELINGCOMPONENTS_API USpaceCurveDeformationMechanicPropertySet : public UInteractiveToolPropertySet
 {
 	GENERATED_BODY()
+	using FVector3d = UE::Geometry::FVector3d;
 
 public:
 	UPROPERTY(EditAnywhere, Category = Options)
@@ -94,7 +95,7 @@ UCLASS()
 class MODELINGCOMPONENTS_API USpaceCurveDeformationMechanic : public UInteractionMechanic, public IClickBehaviorTarget, public IHoverBehaviorTarget
 {
 	GENERATED_BODY()
-
+	using FVector3d = UE::Geometry::FVector3d;
 public:
 
 	// Behaviors used for moving points around and hovering them
@@ -119,7 +120,7 @@ public:
 	void SelectionFill();
 	void SelectionClear();
 
-	void GetCurrentCurvePoints(TArray<FFrame3d>& PointsOut) { PointsOut = CurvePoints; }
+	void GetCurrentCurvePoints(TArray<UE::Geometry::FFrame3d>& PointsOut) { PointsOut = CurvePoints; }
 
 
 	// Some other standard functions
@@ -152,10 +153,10 @@ protected:
 	TSharedPtr<FSpaceCurveSource> CurveSource;
 	TSharedPtr<FSpaceCurveSource> EmptyCurveSource;
 
-	TArray<FFrame3d> CurvePoints;
+	TArray<UE::Geometry::FFrame3d> CurvePoints;
 
 	// Used for spatial queries
-	FGeometrySet3 GeometrySet;
+	UE::Geometry::FGeometrySet3 GeometrySet;
 
 	bool bSpatialValid = false;
 	void UpdateSpatial();
@@ -214,9 +215,9 @@ protected:
 	// Support for selection
 	TArray<int32> SelectedPointIDs;
 	// We need the selected point start positions so we can move multiple points appropriately.
-	TArray<FFrame3d> CurveStartPositions;
+	TArray<UE::Geometry::FFrame3d> CurveStartPositions;
 	// The starting point of the gizmo is needed to determine the offset by which to move the points.
-	FFrame3d GizmoStartPosition;
+	UE::Geometry::FFrame3d GizmoStartPosition;
 
 	// These issue undo/redo change objects, and must therefore not be called in undo/redo code.
 	void ChangeSelection(int32 NewPointID, bool AddToSelection);
@@ -227,7 +228,7 @@ protected:
 	bool DeselectPoint(int32 PointID);
 	void UpdateSelection(const TArray<int32>& NewSelection);
 	void UpdateGizmoLocation();
-	void UpdateCurve(const TArray<FFrame3d>& NewPositions);
+	void UpdateCurve(const TArray<UE::Geometry::FFrame3d>& NewPositions);
 
 	// Used for expiring undo/redo changes, which compare this to their stored value and expire themselves if they do not match.
 	int32 CurrentChangeStamp = 0;
@@ -255,13 +256,13 @@ protected:
 class MODELINGCOMPONENTS_API FSpaceCurveDeformationMechanicMovementChange : public FToolCommandChange
 {
 public:
-	FSpaceCurveDeformationMechanicMovementChange(const TArray<FFrame3d>& FromPositions, const TArray<FFrame3d>& ToPositions);
+	FSpaceCurveDeformationMechanicMovementChange(const TArray<UE::Geometry::FFrame3d>& FromPositions, const TArray<UE::Geometry::FFrame3d>& ToPositions);
 
 	virtual void Apply(UObject* Object) override;
 	virtual void Revert(UObject* Object) override;
 	virtual FString ToString() const override;
 
 protected:
-	TArray<FFrame3d> From;
-	TArray<FFrame3d> To;
+	TArray<UE::Geometry::FFrame3d> From;
+	TArray<UE::Geometry::FFrame3d> To;
 };
diff --git a/Engine/Plugins/Experimental/MeshModelingToolset/Source/ModelingComponents/Public/Mechanics/SpatialCurveDistanceMechanic.h b/Engine/Plugins/Experimental/MeshModelingToolset/Source/ModelingComponents/Public/Mechanics/SpatialCurveDistanceMechanic.h
--- a/Engine/Plugins/Experimental/MeshModelingToolset/Source/ModelingComponents/Public/Mechanics/SpatialCurveDistanceMechanic.h
+++ b/Engine/Plugins/Experimental/MeshModelingToolset/Source/ModelingComponents/Public/Mechanics/SpatialCurveDistanceMechanic.h
@@ -15,7 +15,8 @@ UCLASS()
 class MODELINGCOMPONENTS_API USpatialCurveDistanceMechanic : public UInteractionMechanic
 {
 	GENERATED_BODY()
-
+	using FVector3d = UE::Geometry::FVector3d;
+	using FTransform3d = UE::Geometry::FTransform3d;
 public:
 	/** If this function is set, the hit point will be passed in to this function for snapping. Return false to indicate no snapping occurred. */
 	TUniqueFunction<bool(const FVector3d&, FVector3d&)> WorldPointSnapFunc = nullptr;
@@ -43,7 +44,7 @@ public:
 
 protected:
 	TArray<FVector3d> Curve;
-	FDynamicMesh3 TargetHitMesh;
-	FDynamicMeshAABBTree3 TargetHitMeshAABB;
+	UE::Geometry::FDynamicMesh3 TargetHitMesh;
+	UE::Geometry::FDynamicMeshAABBTree3 TargetHitMeshAABB;
 	FTransform3d Transform;
 };
diff --git a/Engine/Plugins/Experimental/MeshModelingToolset/Source/ModelingComponents/Public/MeshOpPreviewHelpers.h b/Engine/Plugins/Experimental/MeshModelingToolset/Source/ModelingComponents/Public/MeshOpPreviewHelpers.h
--- a/Engine/Plugins/Experimental/MeshModelingToolset/Source/ModelingComponents/Public/MeshOpPreviewHelpers.h
+++ b/Engine/Plugins/Experimental/MeshModelingToolset/Source/ModelingComponents/Public/MeshOpPreviewHelpers.h
@@ -15,7 +15,7 @@
  * FBackgroundDynamicMeshComputeSource is an instantiation of the TBackgroundModelingComputeSource
  * template for FDynamicMeshOperator / IDynamicMeshOperatorFactory
  */
-using FBackgroundDynamicMeshComputeSource = TBackgroundModelingComputeSource<FDynamicMeshOperator, IDynamicMeshOperatorFactory>;
+using FBackgroundDynamicMeshComputeSource = UE::Geometry::TBackgroundModelingComputeSource<UE::Geometry::FDynamicMeshOperator, UE::Geometry::IDynamicMeshOperatorFactory>;
 
 
 /**
@@ -23,8 +23,8 @@ using FBackgroundDynamicMeshComputeSource = TBackgroundModelingComputeSource<FDy
  */
 struct MODELINGCOMPONENTS_API FDynamicMeshOpResult
 {
-	TUniquePtr<FDynamicMesh3> Mesh;
-	FTransform3d Transform;
+	TUniquePtr<UE::Geometry::FDynamicMesh3> Mesh;
+	UE::Geometry::FTransform3d Transform;
 };
 
 
@@ -59,7 +59,7 @@ public:
 	 * @param InWorld the Preview mesh actor will be created in this UWorld
 	 * @param OpGenerator This factory is called to create new MeshOperators on-demand
 	 */
-	void Setup(UWorld* InWorld, IDynamicMeshOperatorFactory* OpGenerator);
+	void Setup(UWorld* InWorld, UE::Geometry::IDynamicMeshOperatorFactory* OpGenerator);
 
 	/**
 	 * Terminate any active computation and return the current Preview Mesh/Transform
@@ -164,7 +164,7 @@ public:
 	/** This delegate is broadcast whenever the embedded preview mesh is updated */
 	FOnMeshUpdated OnMeshUpdated;
 
-	DECLARE_MULTICAST_DELEGATE_OneParam(FOnOpCompleted, const FDynamicMeshOperator*);
+	DECLARE_MULTICAST_DELEGATE_OneParam(FOnOpCompleted, const UE::Geometry::FDynamicMeshOperator*);
 	FOnOpCompleted OnOpCompleted;
 
 public:
@@ -229,9 +229,9 @@ template<typename ResultDataType>
 class TGenericDataBackgroundCompute
 {
 public:
-	using OperatorType = TGenericDataOperator<ResultDataType>;
-	using FactoryType = IGenericDataOperatorFactory<ResultDataType>;
-	using ComputeSourceType = TBackgroundModelingComputeSource<OperatorType, FactoryType>;
+	using OperatorType = UE::Geometry::TGenericDataOperator<ResultDataType>;
+	using FactoryType = UE::Geometry::IGenericDataOperatorFactory<ResultDataType>;
+	using ComputeSourceType = UE::Geometry::TBackgroundModelingComputeSource<OperatorType, FactoryType>;
 
 	//
 	// required calls to setup/update/shutdown this object
@@ -327,8 +327,8 @@ protected:
 		check(BackgroundCompute);
 		if (BackgroundCompute)
 		{
-			EBackgroundComputeTaskStatus Status = BackgroundCompute->CheckStatus();
-			if (Status == EBackgroundComputeTaskStatus::NewResultAvailable)
+			UE::Geometry::EBackgroundComputeTaskStatus Status = BackgroundCompute->CheckStatus();
+			if (Status == UE::Geometry::EBackgroundComputeTaskStatus::NewResultAvailable)
 			{
 				TUniquePtr<OperatorType> ResultOp = BackgroundCompute->ExtractResult();
 				OnOpCompleted.Broadcast(ResultOp.Get());
diff --git a/Engine/Plugins/Experimental/MeshModelingToolset/Source/ModelingComponents/Public/OctreeDynamicMeshComponent.h b/Engine/Plugins/Experimental/MeshModelingToolset/Source/ModelingComponents/Public/OctreeDynamicMeshComponent.h
--- a/Engine/Plugins/Experimental/MeshModelingToolset/Source/ModelingComponents/Public/OctreeDynamicMeshComponent.h
+++ b/Engine/Plugins/Experimental/MeshModelingToolset/Source/ModelingComponents/Public/OctreeDynamicMeshComponent.h
@@ -54,7 +54,7 @@ public:
 	virtual const FDynamicMesh3* GetMesh() const override { return Mesh.Get(); }
 
 
-	FDynamicMeshOctree3* GetOctree() { return Octree.Get(); }
+	UE::Geometry::FDynamicMeshOctree3* GetOctree() { return Octree.Get(); }
 
 	/**
 	 * @return the current internal mesh, which is replaced with an empty mesh
@@ -82,7 +82,7 @@ public:
 	 * Apply transform to internal mesh. Invalidates RenderProxy.
 	 * @param bInvert if true, inverse tranform is applied instead of forward transform
 	 */
-	virtual void ApplyTransform(const FTransform3d& Transform, bool bInvert) override;
+	virtual void ApplyTransform(const UE::Geometry::FTransform3d& Transform, bool bInvert) override;
 
 	//
 	// change tracking/etc
@@ -157,12 +157,12 @@ private:
 	//~ Begin USceneComponent Interface.
 
 	TUniquePtr<FDynamicMesh3> Mesh;
-	TUniquePtr<FDynamicMeshOctree3> Octree;
-	TUniquePtr<FDynamicMeshOctree3::FTreeCutSet> OctreeCut;
-	FArrayIndexSetsDecomposition TriangleDecomposition;
+	TUniquePtr<UE::Geometry::FDynamicMeshOctree3> Octree;
+	TUniquePtr<UE::Geometry::FDynamicMeshOctree3::FTreeCutSet> OctreeCut;
+	UE::Geometry::FArrayIndexSetsDecomposition TriangleDecomposition;
 	struct FCutCellIndexSet
 	{
-		FDynamicMeshOctree3::FCellReference CellRef;
+		UE::Geometry::FDynamicMeshOctree3::FCellReference CellRef;
 		int32 DecompSetID;
 	};
 	TArray<FCutCellIndexSet> CutCellSetMap;
diff --git a/Engine/Plugins/Experimental/MeshModelingToolset/Source/ModelingComponents/Public/Physics/PhysicsDataCollection.h b/Engine/Plugins/Experimental/MeshModelingToolset/Source/ModelingComponents/Public/Physics/PhysicsDataCollection.h
--- a/Engine/Plugins/Experimental/MeshModelingToolset/Source/ModelingComponents/Public/Physics/PhysicsDataCollection.h
+++ b/Engine/Plugins/Experimental/MeshModelingToolset/Source/ModelingComponents/Public/Physics/PhysicsDataCollection.h
@@ -10,7 +10,6 @@
 class UActorComponent;
 class UBodySetup;
 
-
 /**
  * FPhysicsDataCollection holds onto physics-system data that is needed for various interactive tools and algorithms.
  * Currently this is split into two parts, pointers to the owning physics state, and collision geometry data.
@@ -43,7 +42,7 @@ public:
 	/**
 	 * Stores representation of Collision geometry
 	 */
-	FSimpleShapeSet3d Geometry;
+	UE::Geometry::FSimpleShapeSet3d Geometry;
 
 	/**
 	 * Collision geometry in Physics-system format. Not necessarily in sync with Geometry member.
diff --git a/Engine/Plugins/Experimental/MeshModelingToolset/Source/ModelingComponents/Public/PreviewMesh.h b/Engine/Plugins/Experimental/MeshModelingToolset/Source/ModelingComponents/Public/PreviewMesh.h
--- a/Engine/Plugins/Experimental/MeshModelingToolset/Source/ModelingComponents/Public/PreviewMesh.h
+++ b/Engine/Plugins/Experimental/MeshModelingToolset/Source/ModelingComponents/Public/PreviewMesh.h
@@ -10,11 +10,11 @@
 #include "Changes/MeshChange.h"
 #include "PreviewMesh.generated.h"
 
-class FDynamicMesh3;
+PREDECLARE_USE_GEOMETRY_CLASS(FDynamicMesh3);
 struct FMeshDescription;
 
 // predeclare tangents template
-template<typename RealType> class TMeshTangents;
+PREDECLARE_GEOMETRY(template<typename RealType> class TMeshTangents);
 
 
 /**
@@ -57,7 +57,8 @@ UCLASS(Transient)
 class MODELINGCOMPONENTS_API UPreviewMesh : public UObject, public IMeshVertexCommandChangeTarget, public IMeshCommandChangeTarget, public IMeshReplacementCommandChangeTarget
 {
 	GENERATED_BODY()
-
+	using FRay3d = UE::Geometry::FRay3d;
+	using FVector3d = UE::Geometry::FVector3d;
 public:
 	UPreviewMesh();
 	virtual ~UPreviewMesh();
@@ -176,13 +177,13 @@ public:
 	/**
 	 * @return a MeshTangents data structure for the underlying component, if available, otherwise nullptr
 	 */
-	const TMeshTangents<float>* GetTangents() const;
+	const UE::Geometry::TMeshTangents<float>* GetTangents() const;
 
 	/**
 	 * Update tangents on internal Mesh Component
 	 */
 	template<typename RealType>
-	void UpdateTangents(const TMeshTangents<RealType>* ExternalTangents, bool bFastUpdateIfPossible);
+	void UpdateTangents(const UE::Geometry::TMeshTangents<RealType>* ExternalTangents, bool bFastUpdateIfPossible);
 
 
 
@@ -312,7 +313,7 @@ public:
 	 * @return point to the current AABBTree used for preview spatial mesh, or nullptr if not available
 	 * @warning this has to return non-const because of current FDynamicMeshAABBTree3 API, but you should not modify this!
 	 */
-	FDynamicMeshAABBTree3* GetSpatial();
+	UE::Geometry::FDynamicMeshAABBTree3* GetSpatial();
 
 	/**
 	 * @return the current preview FDynamicMesh, and replace with a new empty mesh
@@ -386,7 +387,7 @@ public:
 	 * EditFunc is required to notify the given FDynamicMeshChangeTracker about all mesh changes
 	 * @return FMeshChange extracted from FDynamicMeshChangeTracker that represents mesh edit
 	 */
-	TUniquePtr<FMeshChange> TrackedEditMesh(TFunctionRef<void(FDynamicMesh3&, FDynamicMeshChangeTracker&)> EditFunc);
+	TUniquePtr<FMeshChange> TrackedEditMesh(TFunctionRef<void(FDynamicMesh3&, UE::Geometry::FDynamicMeshChangeTracker&)> EditFunc);
 
 	/**
 	 * Apply/Revert a vertex deformation change to the internal mesh (implements IMeshVertexCommandChangeTarget)
@@ -444,7 +445,7 @@ protected:
 	USimpleDynamicMeshComponent* DynamicMeshComponent = nullptr;
 
 	/** Spatial data structure that is initialized if bBuildSpatialDataStructure = true when UpdatePreview() is called */
-	FDynamicMeshAABBTree3 MeshAABBTree;
+	UE::Geometry::FDynamicMeshAABBTree3 MeshAABBTree;
 
 	/** If true, mesh will be chunked into multiple render buffers inside the DynamicMeshComponent */
 	bool bDecompositionEnabled = false;
@@ -456,7 +457,7 @@ protected:
 
 
 template<typename RealType>
-void UPreviewMesh::UpdateTangents(const TMeshTangents<RealType>* ExternalTangents, bool bFastUpdateIfPossible)
+void UPreviewMesh::UpdateTangents(const UE::Geometry::TMeshTangents<RealType>* ExternalTangents, bool bFastUpdateIfPossible)
 {
 	if (ensure(DynamicMeshComponent))
 	{
diff --git a/Engine/Plugins/Experimental/MeshModelingToolset/Source/ModelingComponents/Public/PropertySets/VoxelProperties.h b/Engine/Plugins/Experimental/MeshModelingToolset/Source/ModelingComponents/Public/PropertySets/VoxelProperties.h
--- a/Engine/Plugins/Experimental/MeshModelingToolset/Source/ModelingComponents/Public/PropertySets/VoxelProperties.h
+++ b/Engine/Plugins/Experimental/MeshModelingToolset/Source/ModelingComponents/Public/PropertySets/VoxelProperties.h
@@ -35,7 +35,7 @@ public:
 	UPROPERTY(EditAnywhere, Category = VoxelSettings, meta = (UIMin = "0.0", UIMax = "100", ClampMin = "0.0", ClampMax = "1000"))
 	double CubeRootMinComponentVolume = 0.0;
 
-	void SetPropertiesOnOp(FVoxelBaseOp& Op)
+	void SetPropertiesOnOp(UE::Geometry::FVoxelBaseOp& Op)
 	{
 		Op.InputVoxelCount = VoxelCount;
 		Op.OutputVoxelCount = VoxelCount;
diff --git a/Engine/Plugins/Experimental/MeshModelingToolset/Source/ModelingComponents/Public/Selection/GroupTopologySelector.h b/Engine/Plugins/Experimental/MeshModelingToolset/Source/ModelingComponents/Public/Selection/GroupTopologySelector.h
--- a/Engine/Plugins/Experimental/MeshModelingToolset/Source/ModelingComponents/Public/Selection/GroupTopologySelector.h
+++ b/Engine/Plugins/Experimental/MeshModelingToolset/Source/ModelingComponents/Public/Selection/GroupTopologySelector.h
@@ -11,6 +11,12 @@ struct FCameraRectangle;
 class FToolDataVisualizer;
 struct FViewCameraState;
 
+using UE::Geometry::FGroupTopologySelection;
+using UE::Geometry::FGroupTopology;
+using UE::Geometry::FDynamicMeshAABBTree3;
+using UE::Geometry::FDynamicMesh3;
+using UE::Geometry::FGeometrySet3;
+
 /**
  * FGroupTopologySelector implements selection behavior for a FGroupTopology mesh.
  * Groups, Group Edges, and Corners can be selected dependings on the settings passed in.
@@ -29,6 +35,9 @@ struct FViewCameraState;
  */
 class MODELINGCOMPONENTS_API FGroupTopologySelector
 {
+	using FRay3d = UE::Geometry::FRay3d;
+	using FVector3d = UE::Geometry::FVector3d;
+	using FTransform3d = UE::Geometry::FTransform3d;
 public:
 
 	//
diff --git a/Engine/Plugins/Experimental/MeshModelingToolset/Source/ModelingComponents/Public/Selection/GroupTopologyStorableSelection.h b/Engine/Plugins/Experimental/MeshModelingToolset/Source/ModelingComponents/Public/Selection/GroupTopologyStorableSelection.h
--- a/Engine/Plugins/Experimental/MeshModelingToolset/Source/ModelingComponents/Public/Selection/GroupTopologyStorableSelection.h
+++ b/Engine/Plugins/Experimental/MeshModelingToolset/Source/ModelingComponents/Public/Selection/GroupTopologyStorableSelection.h
@@ -5,16 +5,17 @@
 #include "CoreMinimal.h"
 
 #include "Components/PrimitiveComponent.h"
+#include "MathUtil.h"
 #include "IndexTypes.h"
 #include "InteractiveToolStorableSelection.h"
 
 #include "GroupTopologyStorableSelection.generated.h"
 
-struct FCompactMaps;
-struct FGroupTopologySelection;
-class FGroupTopology;
 class UPrimitiveComponent;
 class UInteractiveTool;
+PREDECLARE_USE_GEOMETRY_STRUCT(FGroupTopologySelection);
+PREDECLARE_USE_GEOMETRY_CLASS(FGroupTopology);
+PREDECLARE_USE_GEOMETRY_STRUCT(FCompactMaps);
 
 /**
  * Class used to represent a group topology selection independently of a FGroupTopology. Relies
@@ -29,6 +30,7 @@ UCLASS()
 class MODELINGCOMPONENTS_API UGroupTopologyStorableSelection : public UInteractiveToolStorableSelection
 {
 	GENERATED_BODY()
+	using FIndex2i = UE::Geometry::FIndex2i;
 
 public:
 	enum class ETopologyType
diff --git a/Engine/Plugins/Experimental/MeshModelingToolset/Source/ModelingComponents/Public/Selection/PolygonSelectionMechanic.h b/Engine/Plugins/Experimental/MeshModelingToolset/Source/ModelingComponents/Public/Selection/PolygonSelectionMechanic.h
--- a/Engine/Plugins/Experimental/MeshModelingToolset/Source/ModelingComponents/Public/Selection/PolygonSelectionMechanic.h
+++ b/Engine/Plugins/Experimental/MeshModelingToolset/Source/ModelingComponents/Public/Selection/PolygonSelectionMechanic.h
@@ -15,13 +15,15 @@
 #include "InteractionMechanic.h"
 #include "PolygonSelectionMechanic.generated.h"
 
-struct FCompactMaps;
 class FPolygonSelectionMechanicSelectionChange;
 class UGroupTopologyStorableSelection;
 class UMouseHoverBehavior;
 class URectangleMarqueeMechanic;
 class USingleClickInputBehavior;
 
+using UE::Geometry::FDynamicMeshAABBTree3;
+PREDECLARE_USE_GEOMETRY_STRUCT(FCompactMaps);
+
 UCLASS()
 class MODELINGCOMPONENTS_API UPolygonSelectionMechanicProperties : public UInteractiveToolPropertySet
 {
@@ -79,6 +81,9 @@ UCLASS()
 class MODELINGCOMPONENTS_API UPolygonSelectionMechanic : public UInteractionMechanic, public IClickBehaviorTarget, public IHoverBehaviorTarget
 {
 	GENERATED_BODY()
+	using FFrame3d = UE::Geometry::FFrame3d;
+	using FVector3d = UE::Geometry::FVector3d;
+	using FTransform3d = UE::Geometry::FTransform3d;	
 public:
 
 	virtual ~UPolygonSelectionMechanic();
diff --git a/Engine/Plugins/Experimental/MeshModelingToolset/Source/ModelingComponents/Public/SimpleDynamicMeshComponent.h b/Engine/Plugins/Experimental/MeshModelingToolset/Source/ModelingComponents/Public/SimpleDynamicMeshComponent.h
--- a/Engine/Plugins/Experimental/MeshModelingToolset/Source/ModelingComponents/Public/SimpleDynamicMeshComponent.h
+++ b/Engine/Plugins/Experimental/MeshModelingToolset/Source/ModelingComponents/Public/SimpleDynamicMeshComponent.h
@@ -88,13 +88,13 @@ public:
 	 * Copy externally-calculated tangents into the internal tangets buffer.
 	 * @param bFastUpdateIfPossible if true, will try to do a fast normals/tangets update of the SceneProxy, instead of full invalidatiohn
 	 */
-	void UpdateTangents(const FMeshTangentsf* ExternalTangents, bool bFastUpdateIfPossible);
+	void UpdateTangents(const UE::Geometry::FMeshTangentsf* ExternalTangents, bool bFastUpdateIfPossible);
 
 	/**
 	 * Copy externally-calculated tangents into the internal tangets buffer.
 	 * @param bFastUpdateIfPossible if true, will try to do a fast normals/tangets update of the SceneProxy, instead of full invalidatiohn
 	 */
-	void UpdateTangents(const FMeshTangentsd* ExternalTangents, bool bFastUpdateIfPossible);
+	void UpdateTangents(const UE::Geometry::FMeshTangentsd* ExternalTangents, bool bFastUpdateIfPossible);
 
 
 	/**
@@ -102,7 +102,7 @@ public:
 	 * @warning calling this with TangentsType = AutoCalculated will result in possibly-expensive Tangents calculation
 	 * @warning this is only currently safe to call on the Game Thread!!
 	 */
-	const FMeshTangentsf* GetTangents();
+	const UE::Geometry::FMeshTangentsf* GetTangents();
 
 
 
@@ -128,7 +128,7 @@ public:
 	 * Apply transform to internal mesh. Updates Octree and RenderProxy if available.
 	 * @param bInvert if true, inverse tranform is applied instead of forward transform
 	 */
-	virtual void ApplyTransform(const FTransform3d& Transform, bool bInvert) override;
+	virtual void ApplyTransform(const UE::Geometry::FTransform3d& Transform, bool bInvert) override;
 
 
 	//
@@ -195,7 +195,7 @@ public:
 	 * buffers that will be modified, as well as the bounds of the modified region. These are both computed in parallel.
 	 * @return a future that will (eventually) return true if the precompute is OK, and (immediately) false if it is not
 	 */
-	TFuture<bool> FastNotifyTriangleVerticesUpdated_TryPrecompute(const TArray<int32>& Triangles, TArray<int32>& UpdateSetsOut, FAxisAlignedBox3d& BoundsOut);
+	TFuture<bool> FastNotifyTriangleVerticesUpdated_TryPrecompute(const TArray<int32>& Triangles, TArray<int32>& UpdateSetsOut, UE::Geometry::FAxisAlignedBox3d& BoundsOut);
 
 
 	/**
@@ -206,7 +206,7 @@ public:
 	 * If the Precompute future returns false, then we forward the call to FastNotifyTriangleVerticesUpdated(), which will do more work.
 	 */
 	void FastNotifyTriangleVerticesUpdated_ApplyPrecompute(const TArray<int32>& Triangles, EMeshRenderAttributeFlags UpdatedAttributes,
-		TFuture<bool>& Precompute, const TArray<int32>& UpdateSets, const FAxisAlignedBox3d& UpdateSetBounds);
+		TFuture<bool>& Precompute, const TArray<int32>& UpdateSets, const UE::Geometry::FAxisAlignedBox3d& UpdateSetBounds);
 
 
 
@@ -335,10 +335,10 @@ private:
 	void InitializeNewMesh();
 
 	// local-space bounding of Mesh
-	FAxisAlignedBox3d LocalBounds;
+	UE::Geometry::FAxisAlignedBox3d LocalBounds;
 
 	bool bTangentsValid = false;
-	FMeshTangentsf Tangents;
+	UE::Geometry::FMeshTangentsf Tangents;
 	
 	FColor GetTriangleColor(const FDynamicMesh3* Mesh, int TriangleID);
 
diff --git a/Engine/Plugins/Experimental/MeshModelingToolset/Source/ModelingComponents/Public/ToolSceneQueriesUtil.h b/Engine/Plugins/Experimental/MeshModelingToolset/Source/ModelingComponents/Public/ToolSceneQueriesUtil.h
--- a/Engine/Plugins/Experimental/MeshModelingToolset/Source/ModelingComponents/Public/ToolSceneQueriesUtil.h
+++ b/Engine/Plugins/Experimental/MeshModelingToolset/Source/ModelingComponents/Public/ToolSceneQueriesUtil.h
@@ -13,6 +13,9 @@
  */
 namespace ToolSceneQueriesUtil
 {
+	using namespace UE::Geometry;
+	using FVector3d = UE::Geometry::FVector3d;
+
 	/**
 	 * @return global visual angle snap threshold (default is 1 degree)
 	 */
diff --git a/Engine/Plugins/Experimental/MeshModelingToolset/Source/ModelingComponents/Public/Transforms/MultiTransformer.h b/Engine/Plugins/Experimental/MeshModelingToolset/Source/ModelingComponents/Public/Transforms/MultiTransformer.h
--- a/Engine/Plugins/Experimental/MeshModelingToolset/Source/ModelingComponents/Public/Transforms/MultiTransformer.h
+++ b/Engine/Plugins/Experimental/MeshModelingToolset/Source/ModelingComponents/Public/Transforms/MultiTransformer.h
@@ -30,6 +30,8 @@ UCLASS()
 class MODELINGCOMPONENTS_API UMultiTransformer : public UObject
 {
 	GENERATED_BODY()
+	using FFrame3d = UE::Geometry::FFrame3d;
+	using FVector3d = UE::Geometry::FVector3d;
 public:
 	
 	virtual void Setup(UInteractiveGizmoManager* GizmoManager, IToolContextTransactionProvider* TransactionProviderIn);
diff --git a/Engine/Plugins/Experimental/MeshModelingToolset/Source/ModelingOperators/Private/CleaningOps/EditNormalsOp.cpp b/Engine/Plugins/Experimental/MeshModelingToolset/Source/ModelingOperators/Private/CleaningOps/EditNormalsOp.cpp
--- a/Engine/Plugins/Experimental/MeshModelingToolset/Source/ModelingOperators/Private/CleaningOps/EditNormalsOp.cpp
+++ b/Engine/Plugins/Experimental/MeshModelingToolset/Source/ModelingOperators/Private/CleaningOps/EditNormalsOp.cpp
@@ -8,11 +8,13 @@
 #include "MeshSimplification.h"
 #include "MeshConstraintsUtil.h"
 
-
 #include "MeshNormals.h"
 #include "Operations/RepairOrientation.h"
 #include "DynamicMeshAABBTree3.h"
 
+#include "ExplicitUseGeometryMathTypes.h"		// using UE::Geometry::(math types)
+using namespace UE::Geometry;
+
 void FEditNormalsOp::SetTransform(const FTransform& Transform) {
 	ResultTransform = (FTransform3d)Transform;
 }
diff --git a/Engine/Plugins/Experimental/MeshModelingToolset/Source/ModelingOperators/Private/CompositionOps/BooleanMeshesOp.cpp b/Engine/Plugins/Experimental/MeshModelingToolset/Source/ModelingOperators/Private/CompositionOps/BooleanMeshesOp.cpp
--- a/Engine/Plugins/Experimental/MeshModelingToolset/Source/ModelingOperators/Private/CompositionOps/BooleanMeshesOp.cpp
+++ b/Engine/Plugins/Experimental/MeshModelingToolset/Source/ModelingOperators/Private/CompositionOps/BooleanMeshesOp.cpp
@@ -3,12 +3,13 @@
 #include "CompositionOps/BooleanMeshesOp.h"
 
 #include "Operations/MeshBoolean.h"
-
 #include "MeshSimplification.h"
-
 #include "MeshBoundaryLoops.h"
 #include "Operations/MinimalHoleFiller.h"
 
+#include "ExplicitUseGeometryMathTypes.h"		// using UE::Geometry::(math types)
+using namespace UE::Geometry;
+
 void FBooleanMeshesOp::SetTransform(const FTransform& Transform) {
 	ResultTransform = (FTransform3d)Transform;
 }
diff --git a/Engine/Plugins/Experimental/MeshModelingToolset/Source/ModelingOperators/Private/CompositionOps/SelfUnionMeshesOp.cpp b/Engine/Plugins/Experimental/MeshModelingToolset/Source/ModelingOperators/Private/CompositionOps/SelfUnionMeshesOp.cpp
--- a/Engine/Plugins/Experimental/MeshModelingToolset/Source/ModelingOperators/Private/CompositionOps/SelfUnionMeshesOp.cpp
+++ b/Engine/Plugins/Experimental/MeshModelingToolset/Source/ModelingOperators/Private/CompositionOps/SelfUnionMeshesOp.cpp
@@ -3,12 +3,13 @@
 #include "CompositionOps/SelfUnionMeshesOp.h"
 
 #include "Operations/MeshSelfUnion.h"
-
 #include "MeshSimplification.h"
-
 #include "MeshBoundaryLoops.h"
 #include "Operations/MinimalHoleFiller.h"
 
+#include "ExplicitUseGeometryMathTypes.h"		// using UE::Geometry::(math types)
+using namespace UE::Geometry;
+
 void FSelfUnionMeshesOp::SetTransform(const FTransform& Transform) {
 	ResultTransform = (FTransform3d)Transform;
 }
diff --git a/Engine/Plugins/Experimental/MeshModelingToolset/Source/ModelingOperators/Public/BaseOps/SimpleMeshProcessingBaseOp.h b/Engine/Plugins/Experimental/MeshModelingToolset/Source/ModelingOperators/Public/BaseOps/SimpleMeshProcessingBaseOp.h
--- a/Engine/Plugins/Experimental/MeshModelingToolset/Source/ModelingOperators/Public/BaseOps/SimpleMeshProcessingBaseOp.h
+++ b/Engine/Plugins/Experimental/MeshModelingToolset/Source/ModelingOperators/Public/BaseOps/SimpleMeshProcessingBaseOp.h
@@ -7,9 +7,12 @@
 #include "ModelingOperators.h"
 #include "VectorTypes.h"
 
-class FDynamicMesh3;
-class FMeshNormals;
+namespace UE
+{
+namespace Geometry
+{
 
+class FDynamicMesh3;
 
 class MODELINGOPERATORS_API FMeshBoundaryCache
 {
@@ -45,4 +48,8 @@ public:
 
 protected:
 	TArray<FVector3d> PositionBuffer;
-};
\ No newline at end of file
+};
+
+
+} // end namespace UE::Geometry
+} // end namespace UE
diff --git a/Engine/Plugins/Experimental/MeshModelingToolset/Source/ModelingOperators/Public/BaseOps/VoxelBaseOp.h b/Engine/Plugins/Experimental/MeshModelingToolset/Source/ModelingOperators/Public/BaseOps/VoxelBaseOp.h
--- a/Engine/Plugins/Experimental/MeshModelingToolset/Source/ModelingOperators/Public/BaseOps/VoxelBaseOp.h
+++ b/Engine/Plugins/Experimental/MeshModelingToolset/Source/ModelingOperators/Public/BaseOps/VoxelBaseOp.h
@@ -6,7 +6,10 @@
 #include "Util/ProgressCancel.h"
 #include "ModelingOperators.h"
 
-
+namespace UE
+{
+namespace Geometry
+{
 
 class MODELINGOPERATORS_API FVoxelBaseOp : public FDynamicMeshOperator
 {
@@ -28,3 +31,5 @@ public:
 };
 
 
+} // end namespace UE::Geometry
+} // end namespace UE
diff --git a/Engine/Plugins/Experimental/MeshModelingToolset/Source/ModelingOperators/Public/CleaningOps/HoleFillOp.h b/Engine/Plugins/Experimental/MeshModelingToolset/Source/ModelingOperators/Public/CleaningOps/HoleFillOp.h
--- a/Engine/Plugins/Experimental/MeshModelingToolset/Source/ModelingOperators/Public/CleaningOps/HoleFillOp.h
+++ b/Engine/Plugins/Experimental/MeshModelingToolset/Source/ModelingOperators/Public/CleaningOps/HoleFillOp.h
@@ -28,6 +28,11 @@ enum class EHoleFillOpFillType : uint8
 };
 
 
+namespace UE
+{
+namespace Geometry
+{
+
 struct FFillOptions
 {
 	bool bRemoveIsolatedTriangles = false;
@@ -43,10 +48,10 @@ public:
 
 	EHoleFillOpFillType FillType = EHoleFillOpFillType::Minimal;
 	double MeshUVScaleFactor = 1.0;
-	TArray<FEdgeLoop> Loops;
+	TArray<UE::Geometry::FEdgeLoop> Loops;
 
 	FFillOptions FillOptions;
-	FSmoothFillOptions SmoothFillOptions;
+	UE::Geometry::FSmoothFillOptions SmoothFillOptions;
 
 	// output
 	TArray<int32> NewTriangles;
@@ -57,3 +62,5 @@ public:
 
 };
 
+} // end namespace UE::Geometry
+} // end namespace UE
diff --git a/Engine/Plugins/Experimental/MeshModelingToolset/Source/ModelingOperators/Public/CleaningOps/RemeshMeshOp.h b/Engine/Plugins/Experimental/MeshModelingToolset/Source/ModelingOperators/Public/CleaningOps/RemeshMeshOp.h
--- a/Engine/Plugins/Experimental/MeshModelingToolset/Source/ModelingOperators/Public/CleaningOps/RemeshMeshOp.h
+++ b/Engine/Plugins/Experimental/MeshModelingToolset/Source/ModelingOperators/Public/CleaningOps/RemeshMeshOp.h
@@ -3,19 +3,14 @@
 #pragma once
 
 #include "CoreMinimal.h"
+#include "DynamicMesh3.h"
+#include "DynamicMeshAABBTree3.h"
 #include "MeshConstraints.h"
 #include "Util/ProgressCancel.h"
 #include "ModelingOperators.h"
 
 #include "RemeshMeshOp.generated.h"
 
-class FDynamicMesh3;
-class FRemesher;
-
-template <class TriangleMeshType>
-class TMeshAABBTree3;
-typedef TMeshAABBTree3<FDynamicMesh3> FDynamicMeshAABBTree3;
-
 /** Remeshing modes */
 UENUM()
 enum class ERemeshType : uint8
@@ -45,6 +40,13 @@ enum class ERemeshSmoothingType : uint8
 	MeanValue = 2 UMETA(DisplayName = "Mixed")
 };
 
+namespace UE
+{
+namespace Geometry
+{
+
+class FRemesher;
+
 class MODELINGOPERATORS_API FRemeshMeshOp : public FDynamicMeshOperator
 {
 public:
@@ -80,4 +82,7 @@ protected:
 
 };
 
+} // end namespace UE::Geometry
+} // end namespace UE
+
 
diff --git a/Engine/Plugins/Experimental/MeshModelingToolset/Source/ModelingOperators/Public/CleaningOps/RemoveOccludedTrianglesOp.h b/Engine/Plugins/Experimental/MeshModelingToolset/Source/ModelingOperators/Public/CleaningOps/RemoveOccludedTrianglesOp.h
--- a/Engine/Plugins/Experimental/MeshModelingToolset/Source/ModelingOperators/Public/CleaningOps/RemoveOccludedTrianglesOp.h
+++ b/Engine/Plugins/Experimental/MeshModelingToolset/Source/ModelingOperators/Public/CleaningOps/RemoveOccludedTrianglesOp.h
@@ -10,7 +10,10 @@
 #include "MeshAdapter.h"
 #include "Operations/RemoveOccludedTriangles.h"
 
-
+namespace UE
+{
+namespace Geometry
+{
 
 class MODELINGOPERATORS_API FRemoveOccludedTrianglesOp : public FDynamicMeshOperator
 {
@@ -20,14 +23,14 @@ public:
 	// inputs
 	TSharedPtr<FDynamicMesh3, ESPMode::ThreadSafe> OriginalMesh;
 	
-	TArray<TSharedPtr<FDynamicMeshAABBTree3, ESPMode::ThreadSafe>> OccluderTrees;
-	TArray<TSharedPtr<TFastWindingTree<FDynamicMesh3>, ESPMode::ThreadSafe>> OccluderWindings;
+	TArray<TSharedPtr<UE::Geometry::FDynamicMeshAABBTree3, ESPMode::ThreadSafe>> OccluderTrees;
+	TArray<TSharedPtr<UE::Geometry::TFastWindingTree<FDynamicMesh3>, ESPMode::ThreadSafe>> OccluderWindings;
 	TArray<FTransform3d> OccluderTransforms;
 
 	TArray<FTransform3d> MeshTransforms;
 
-	EOcclusionTriangleSampling TriangleSamplingMethod =
-		EOcclusionTriangleSampling::Centroids;
+	UE::Geometry::EOcclusionTriangleSampling TriangleSamplingMethod =
+		UE::Geometry::EOcclusionTriangleSampling::Centroids;
 
 	// we nudge points out by this amount to try to counteract numerical issues
 	double NormalOffset = FMathd::ZeroTolerance;
@@ -56,8 +59,8 @@ public:
 	// if true, the ActiveGroupLayer is the name of the "default" layer, so we'll use the built-in group IDs
 	bool bActiveGroupLayerIsDefault = true;
 
-	EOcclusionCalculationMode InsideMode =
-		EOcclusionCalculationMode::FastWindingNumber;
+	UE::Geometry::EOcclusionCalculationMode InsideMode =
+		UE::Geometry::EOcclusionCalculationMode::FastWindingNumber;
 
 
 	// outputs, used when bSetTriangleGroupInsteadOfRemoving is true
@@ -75,3 +78,5 @@ public:
 };
 
 
+} // end namespace UE::Geometry
+} // end namespace UE
diff --git a/Engine/Plugins/Experimental/MeshModelingToolset/Source/ModelingOperators/Public/CompositionOps/CurveSweepOp.h b/Engine/Plugins/Experimental/MeshModelingToolset/Source/ModelingOperators/Public/CompositionOps/CurveSweepOp.h
--- a/Engine/Plugins/Experimental/MeshModelingToolset/Source/ModelingOperators/Public/CompositionOps/CurveSweepOp.h
+++ b/Engine/Plugins/Experimental/MeshModelingToolset/Source/ModelingOperators/Public/CompositionOps/CurveSweepOp.h
@@ -7,8 +7,10 @@
 #include "ModelingOperators.h"
 #include "Util/ProgressCancel.h"
 
-class URevolveProperties;
-class UNewMeshMaterialProperties;
+namespace UE
+{
+namespace Geometry
+{
 
 /*
  * Operation for sweeping a profile curve along a sweep curve to create a mesh.
@@ -76,4 +78,7 @@ public:
 	// 
 
 	virtual void CalculateResult(FProgressCancel* Progress) override;
-};
\ No newline at end of file
+};
+
+} // end namespace UE::Geometry
+} // end namespace UE
diff --git a/Engine/Plugins/Experimental/MeshModelingToolset/Source/ModelingOperators/Public/CompositionOps/VoxelBlendMeshesOp.h b/Engine/Plugins/Experimental/MeshModelingToolset/Source/ModelingOperators/Public/CompositionOps/VoxelBlendMeshesOp.h
--- a/Engine/Plugins/Experimental/MeshModelingToolset/Source/ModelingOperators/Public/CompositionOps/VoxelBlendMeshesOp.h
+++ b/Engine/Plugins/Experimental/MeshModelingToolset/Source/ModelingOperators/Public/CompositionOps/VoxelBlendMeshesOp.h
@@ -8,7 +8,10 @@
 
 #include "BaseOps/VoxelBaseOp.h"
 
-
+namespace UE
+{
+namespace Geometry
+{
 
 class MODELINGOPERATORS_API FVoxelBlendMeshesOp : public FVoxelBaseOp
 {
@@ -39,3 +42,5 @@ public:
 };
 
 
+} // end namespace UE::Geometry
+} // end namespace UE
diff --git a/Engine/Plugins/Experimental/MeshModelingToolset/Source/ModelingOperators/Public/DeformationOps/LatticeDeformerOp.h b/Engine/Plugins/Experimental/MeshModelingToolset/Source/ModelingOperators/Public/DeformationOps/LatticeDeformerOp.h
--- a/Engine/Plugins/Experimental/MeshModelingToolset/Source/ModelingOperators/Public/DeformationOps/LatticeDeformerOp.h
+++ b/Engine/Plugins/Experimental/MeshModelingToolset/Source/ModelingOperators/Public/DeformationOps/LatticeDeformerOp.h
@@ -12,14 +12,19 @@
 #include "Implicit/GridInterpolant.h"
 #include "Operations/FFDLattice.h"
 
+namespace UE
+{
+namespace Geometry
+{
+
 class MODELINGOPERATORS_API FLatticeDeformerOp : public FDynamicMeshOperator
 {
 public:
 
 	FLatticeDeformerOp(TSharedPtr<FDynamicMesh3, ESPMode::ThreadSafe> InOriginalMesh,
-					   TSharedPtr<FFFDLattice, ESPMode::ThreadSafe> InLattice,
+					   TSharedPtr<UE::Geometry::FFFDLattice, ESPMode::ThreadSafe> InLattice,
 					   const TArray<FVector3d>& InLatticeControlPoints,
-					   ELatticeInterpolation InInterpolationType,
+					   UE::Geometry::ELatticeInterpolation InInterpolationType,
 					   bool bInDeformNormals);
 
 	// FDynamicMeshOperator implementation
@@ -28,10 +33,13 @@ public:
 protected:
 
 	// Inputs
-	const TSharedPtr<const FFFDLattice, ESPMode::ThreadSafe> Lattice;
+	const TSharedPtr<const UE::Geometry::FFFDLattice, ESPMode::ThreadSafe> Lattice;
 	const TSharedPtr<const FDynamicMesh3, ESPMode::ThreadSafe> OriginalMesh;
 	const TArray<FVector3d> LatticeControlPoints;
-	ELatticeInterpolation InterpolationType;
+	UE::Geometry::ELatticeInterpolation InterpolationType;
 	bool bDeformNormals = false;
 };
 
+} // end namespace UE::Geometry
+} // end namespace UE
+
diff --git a/Engine/Plugins/Experimental/MeshModelingToolset/Source/ModelingOperators/Public/DeformationOps/MeshOffsetOps.h b/Engine/Plugins/Experimental/MeshModelingToolset/Source/ModelingOperators/Public/DeformationOps/MeshOffsetOps.h
--- a/Engine/Plugins/Experimental/MeshModelingToolset/Source/ModelingOperators/Public/DeformationOps/MeshOffsetOps.h
+++ b/Engine/Plugins/Experimental/MeshModelingToolset/Source/ModelingOperators/Public/DeformationOps/MeshOffsetOps.h
@@ -11,6 +11,12 @@
 #include "Solvers/ConstrainedMeshDeformer.h"
 
 
+
+namespace UE
+{
+namespace Geometry
+{
+
 /**
  * Base class for Mesh Offset operations (ie displacement along some kind of normal)
  * Optionally duplicates input mesh and joins the two together, either by welding along
@@ -18,6 +24,7 @@
  */
 class FMeshOffsetBaseOp : public FSimpleMeshProcessingBaseOp
 {
+
 public:
 	FMeshOffsetBaseOp(const FDynamicMesh3* Mesh) : FSimpleMeshProcessingBaseOp(Mesh) {}
 
@@ -33,12 +40,12 @@ public:
 	TSharedPtr<FMeshNormals> BaseNormals;
 	FMeshBoundaryCache BoundaryCache;
 
-	TSharedPtr<FMeshBoundaryLoops> BoundaryLoops;
+	TSharedPtr<UE::Geometry::FMeshBoundaryLoops> BoundaryLoops;
 	bool bFixedBoundary = false;
 
 	bool bCreateShell = false;
 
-	TSharedPtr<FIndexedWeightMap1f> WeightMap;
+	TSharedPtr<UE::Geometry::FIndexedWeightMap1f> WeightMap;
 	bool bUseWeightMap = false;
 
 	/** @return signed offset distance at vertex, transformed by weight map iv applicable */
@@ -74,14 +81,14 @@ public:
 				ResultMesh->ReverseOrientation();
 			}
 
-			FDynamicMeshEditor Editor(ResultMesh.Get());
-			FMeshIndexMappings MeshMap;
+			UE::Geometry::FDynamicMeshEditor Editor(ResultMesh.Get());
+			UE::Geometry::FMeshIndexMappings MeshMap;
 			Editor.AppendMesh(&InnerMesh, MeshMap);
 
 			bool bWeldStitch = bFixedBoundary;
 
 			// join the boundary loops by weld or stitch
-			for (FEdgeLoop& BaseLoop : BoundaryLoops->Loops)
+			for (UE::Geometry::FEdgeLoop& BaseLoop : BoundaryLoops->Loops)
 			{
 				int32 LoopCount = BaseLoop.GetVertexCount();
 				TArray<int32> OffsetLoop;
@@ -100,7 +107,7 @@ public:
 				}
 				else
 				{
-					FJoinMeshLoops Join(ResultMesh.Get(), LoopA, LoopB);
+					UE::Geometry::FJoinMeshLoops Join(ResultMesh.Get(), LoopA, LoopB);
 					Join.Apply();
 				}
 			}
@@ -343,3 +350,5 @@ public:
 	}
 };
 
+} // end namespace UE::Geometry
+} // end namespace UE
diff --git a/Engine/Plugins/Experimental/MeshModelingToolset/Source/ModelingOperators/Public/ModelingTaskTypes.h b/Engine/Plugins/Experimental/MeshModelingToolset/Source/ModelingOperators/Public/ModelingTaskTypes.h
--- a/Engine/Plugins/Experimental/MeshModelingToolset/Source/ModelingOperators/Public/ModelingTaskTypes.h
+++ b/Engine/Plugins/Experimental/MeshModelingToolset/Source/ModelingOperators/Public/ModelingTaskTypes.h
@@ -7,9 +7,10 @@
 #include "Util/ProgressCancel.h"
 
 
-
-
-
+namespace UE
+{
+namespace Geometry
+{
 
 
 /**
@@ -142,3 +143,7 @@ public:
 };
 
 
+} // end namespace UE::Geometry
+} // end namespace UE
+
+
diff --git a/Engine/Plugins/Experimental/MeshModelingToolset/Source/ModelingOperators/Public/ParameterizationOps/UVProjectionOp.h b/Engine/Plugins/Experimental/MeshModelingToolset/Source/ModelingOperators/Public/ParameterizationOps/UVProjectionOp.h
--- a/Engine/Plugins/Experimental/MeshModelingToolset/Source/ModelingOperators/Public/ParameterizationOps/UVProjectionOp.h
+++ b/Engine/Plugins/Experimental/MeshModelingToolset/Source/ModelingOperators/Public/ParameterizationOps/UVProjectionOp.h
@@ -8,7 +8,6 @@
 
 #include "UVProjectionOp.generated.h"
 
-
 UENUM()
 enum class EUVProjectionMethod : uint8
 {
@@ -17,7 +16,10 @@ enum class EUVProjectionMethod : uint8
 	Plane
 };
 
-
+namespace UE
+{
+namespace Geometry
+{
 
 class MODELINGOPERATORS_API FUVProjectionOp : public FDynamicMeshOperator
 {
@@ -42,4 +44,6 @@ public:
 	virtual void CalculateResult(FProgressCancel* Progress) override;
 };
 
+} // end namespace UE::Geometry
+} // end namespace UE
 
diff --git a/Engine/Plugins/Experimental/MeshModelingToolset/Source/ModelingOperators/Public/SmoothingOps/SmoothingOpBase.h b/Engine/Plugins/Experimental/MeshModelingToolset/Source/ModelingOperators/Public/SmoothingOps/SmoothingOpBase.h
--- a/Engine/Plugins/Experimental/MeshModelingToolset/Source/ModelingOperators/Public/SmoothingOps/SmoothingOpBase.h
+++ b/Engine/Plugins/Experimental/MeshModelingToolset/Source/ModelingOperators/Public/SmoothingOps/SmoothingOpBase.h
@@ -8,6 +8,12 @@
 #include "VectorTypes.h"
 #include "WeightMapTypes.h"
 
+
+namespace UE
+{
+namespace Geometry
+{
+
 class FDynamicMesh3;
 class FMeshNormals;
 
@@ -42,12 +48,12 @@ public:
 		double WeightClamp = FMathf::MaxReal;
 
 		// mesh normals calculated for input mesh
-		TSharedPtr<FMeshNormals> BaseNormals;
+		TSharedPtr<UE::Geometry::FMeshNormals> BaseNormals;
 
 		// offset used by some smoothers
 		double NormalOffset = 0.0;
 
-		TSharedPtr<FIndexedWeightMap1f> WeightMap;
+		TSharedPtr<UE::Geometry::FIndexedWeightMap1f> WeightMap;
 		bool bUseWeightMap = false;
 		float WeightMapMinMultiplier = 0.0;
 	};
@@ -71,4 +77,7 @@ protected:
 
 	TArray<FVector3d> PositionBuffer;
 
-};
\ No newline at end of file
+};
+
+} // end namespace UE::Geometry
+} // end namespace UE
diff --git a/Engine/Plugins/Experimental/MeshModelingToolset/Source/ModelingOperators/Public/SpaceDeformerOps/FlareMeshOp.h b/Engine/Plugins/Experimental/MeshModelingToolset/Source/ModelingOperators/Public/SpaceDeformerOps/FlareMeshOp.h
--- a/Engine/Plugins/Experimental/MeshModelingToolset/Source/ModelingOperators/Public/SpaceDeformerOps/FlareMeshOp.h
+++ b/Engine/Plugins/Experimental/MeshModelingToolset/Source/ModelingOperators/Public/SpaceDeformerOps/FlareMeshOp.h
@@ -5,7 +5,10 @@
 #include "CoreMinimal.h"
 #include "MeshSpaceDeformerOp.h"
 
-
+namespace UE
+{
+namespace Geometry
+{
 
 class MODELINGOPERATORS_API FFlareMeshOp : public FMeshSpaceDeformerOp
 {
@@ -25,4 +28,7 @@ public:
 	bool bSmoothEnds = false;
 protected:
 
-};
\ No newline at end of file
+};
+
+} // end namespace UE::Geometry
+} // end namespace UE
diff --git a/Engine/Plugins/Experimental/MeshModelingToolset/Source/ModelingOperators/Public/SpaceDeformerOps/TwistMeshOp.h b/Engine/Plugins/Experimental/MeshModelingToolset/Source/ModelingOperators/Public/SpaceDeformerOps/TwistMeshOp.h
--- a/Engine/Plugins/Experimental/MeshModelingToolset/Source/ModelingOperators/Public/SpaceDeformerOps/TwistMeshOp.h
+++ b/Engine/Plugins/Experimental/MeshModelingToolset/Source/ModelingOperators/Public/SpaceDeformerOps/TwistMeshOp.h
@@ -5,7 +5,10 @@
 #include "CoreMinimal.h"
 #include "MeshSpaceDeformerOp.h"
 
-
+namespace UE
+{
+namespace Geometry
+{
 
 class MODELINGOPERATORS_API FTwistMeshOp : public FMeshSpaceDeformerOp
 {
@@ -16,4 +19,7 @@ public:
 	bool bLockBottom = false;
 protected:
 
-};
\ No newline at end of file
+};
+
+} // end namespace UE::Geometry
+} // end namespace UE
diff --git a/Engine/Plugins/Experimental/MeshModelingToolset/Source/ModelingOperatorsEditorOnly/Public/CleaningOps/SimplifyMeshOp.h b/Engine/Plugins/Experimental/MeshModelingToolset/Source/ModelingOperatorsEditorOnly/Public/CleaningOps/SimplifyMeshOp.h
--- a/Engine/Plugins/Experimental/MeshModelingToolset/Source/ModelingOperatorsEditorOnly/Public/CleaningOps/SimplifyMeshOp.h
+++ b/Engine/Plugins/Experimental/MeshModelingToolset/Source/ModelingOperatorsEditorOnly/Public/CleaningOps/SimplifyMeshOp.h
@@ -5,16 +5,14 @@
 #include "CoreMinimal.h"
 #include "MeshConstraints.h"
 #include "ModelingOperators.h"
+#include "DynamicMesh3.h"
+#include "DynamicMeshAABBTree3.h"
 #include "Util/ProgressCancel.h"
 
 #include "SimplifyMeshOp.generated.h"
 
-class FDynamicMesh3;
 struct FMeshDescription;
-
-template <class TriangleMeshType>
-class TMeshAABBTree3;
-typedef TMeshAABBTree3<FDynamicMesh3> FDynamicMeshAABBTree3;
+class IMeshReduction;
 
 UENUM()
 enum class ESimplifyTargetType : uint8
@@ -55,6 +53,12 @@ enum class ESimplifyType : uint8
 
 };
 
+namespace UE
+{
+namespace Geometry
+{
+
+
 class MODELINGOPERATORSEDITORONLY_API FSimplifyMeshOp : public FDynamicMeshOperator
 {
 public:
@@ -68,7 +72,7 @@ public:
 	int TargetPercentage, TargetCount;
 	float TargetEdgeLength;
 	bool bDiscardAttributes, bReproject, bPreventNormalFlips, bPreserveSharpEdges, bAllowSeamCollapse;
-	EEdgeRefineFlags MeshBoundaryConstraint, GroupBoundaryConstraint, MaterialBoundaryConstraint;
+	UE::Geometry::EEdgeRefineFlags MeshBoundaryConstraint, GroupBoundaryConstraint, MaterialBoundaryConstraint;
 	/** Angle threshold in degrees used for testing if two triangles should be considered coplanar, or two lines collinear */
 	float MinimalPlanarAngleThresh = 0.01f;
 
@@ -78,7 +82,7 @@ public:
 	TSharedPtr<FDynamicMesh3, ESPMode::ThreadSafe> OriginalMesh;
 	TSharedPtr<FDynamicMeshAABBTree3, ESPMode::ThreadSafe> OriginalMeshSpatial;
 
-	class IMeshReduction* MeshReduction;
+	IMeshReduction* MeshReduction;
 
 
 	// set ability on protected transform.
@@ -93,3 +97,6 @@ public:
 
 	virtual void CalculateResult(FProgressCancel* Progress) override;
 };
+
+} // end namespace UE::Geometry
+} // end namespace UE
diff --git a/Engine/Plugins/Experimental/MeshModelingToolset/Source/ModelingOperatorsEditorOnly/Public/CuttingOps/EmbedPolygonsOp.h b/Engine/Plugins/Experimental/MeshModelingToolset/Source/ModelingOperatorsEditorOnly/Public/CuttingOps/EmbedPolygonsOp.h
--- a/Engine/Plugins/Experimental/MeshModelingToolset/Source/ModelingOperatorsEditorOnly/Public/CuttingOps/EmbedPolygonsOp.h
+++ b/Engine/Plugins/Experimental/MeshModelingToolset/Source/ModelingOperatorsEditorOnly/Public/CuttingOps/EmbedPolygonsOp.h
@@ -6,7 +6,7 @@
 #include "Util/ProgressCancel.h"
 #include "ModelingOperators.h"
 #include "ProxyLODVolume.h"
-
+#include "DynamicMesh3.h"
 #include "Polygon2.h"
 
 #include "EmbedPolygonsOp.generated.h"
@@ -21,6 +21,11 @@ enum class EEmbeddedPolygonOpMethod : uint8
 	CutThrough
 };
 
+namespace UE
+{
+namespace Geometry
+{
+
 
 class MODELINGOPERATORSEDITORONLY_API FEmbedPolygonsOp : public FDynamicMeshOperator
 {
@@ -66,4 +71,6 @@ private:
 	void BooleanPath(FProgressCancel* Progress);
 };
 
+} // end namespace UE::Geometry
+} // end namespace UE
 
diff --git a/Engine/Plugins/Experimental/MeshModelingToolset/Source/ModelingOperatorsEditorOnly/Public/ParameterizationOps/CalculateTangentsOp.h b/Engine/Plugins/Experimental/MeshModelingToolset/Source/ModelingOperatorsEditorOnly/Public/ParameterizationOps/CalculateTangentsOp.h
--- a/Engine/Plugins/Experimental/MeshModelingToolset/Source/ModelingOperatorsEditorOnly/Public/ParameterizationOps/CalculateTangentsOp.h
+++ b/Engine/Plugins/Experimental/MeshModelingToolset/Source/ModelingOperatorsEditorOnly/Public/ParameterizationOps/CalculateTangentsOp.h
@@ -23,16 +23,20 @@ enum class EMeshTangentsType : uint8
 	CopyExisting = 3
 };
 
+namespace UE
+{
+namespace Geometry
+{
 
 
-class MODELINGOPERATORSEDITORONLY_API FCalculateTangentsOp : public TGenericDataOperator<FMeshTangentsd>
+class MODELINGOPERATORSEDITORONLY_API FCalculateTangentsOp : public TGenericDataOperator<UE::Geometry::FMeshTangentsd>
 {
 public:
 	virtual ~FCalculateTangentsOp() {}
 
 	// inputs
 	TSharedPtr<FDynamicMesh3, ESPMode::ThreadSafe> SourceMesh;
-	TSharedPtr<FMeshTangentsf, ESPMode::ThreadSafe> SourceTangents;
+	TSharedPtr<UE::Geometry::FMeshTangentsf, ESPMode::ThreadSafe> SourceTangents;
 
 	// parameters
 	EMeshTangentsType CalculationMethod;
@@ -47,9 +51,11 @@ public:
 	virtual void CalculateResult(FProgressCancel* Progress) override;
 
 
-	virtual void CalculateStandard(FProgressCancel* Progress, TUniquePtr<FMeshTangentsd>& Tangents);
-	virtual void CalculateMikkT(FProgressCancel* Progress, TUniquePtr<FMeshTangentsd>& Tangents);
-	virtual void CopyFromSource(FProgressCancel* Progress, TUniquePtr<FMeshTangentsd>& Tangents);
+	virtual void CalculateStandard(FProgressCancel* Progress, TUniquePtr<UE::Geometry::FMeshTangentsd>& Tangents);
+	virtual void CalculateMikkT(FProgressCancel* Progress, TUniquePtr<UE::Geometry::FMeshTangentsd>& Tangents);
+	virtual void CopyFromSource(FProgressCancel* Progress, TUniquePtr<UE::Geometry::FMeshTangentsd>& Tangents);
 };
 
+} // end namespace UE::Geometry
+} // end namespace UE
 
diff --git a/Engine/Plugins/Experimental/PlanarCutPlugin/Source/PlanarCut/Private/FractureAutoUV.cpp b/Engine/Plugins/Experimental/PlanarCutPlugin/Source/PlanarCut/Private/FractureAutoUV.cpp
--- a/Engine/Plugins/Experimental/PlanarCutPlugin/Source/PlanarCut/Private/FractureAutoUV.cpp
+++ b/Engine/Plugins/Experimental/PlanarCutPlugin/Source/PlanarCut/Private/FractureAutoUV.cpp
@@ -14,6 +14,9 @@
 #include "Misc/ScopedSlowTask.h"
 #endif
 
+#include "ExplicitUseGeometryMathTypes.h"		// using UE::Geometry::(math types)
+using namespace UE::Geometry;
+
 namespace UE { namespace PlanarCut {
 
 /**
@@ -58,7 +61,7 @@ int32 SetActiveTriangles(FGeometryCollection* Collection, TArray<bool>& ActiveTr
 
 // Adapter to use geometry collections in the UVPacker, and other GeometricObjects generic mesh processing code
 // Note: Adapts the whole geometry collection as a single mesh, with triangles filtered by an ActiveTriangles mask
-struct FGeomMesh : public FUVPacker::IUVMeshView
+struct FGeomMesh : public UE::Geometry::FUVPacker::IUVMeshView
 {
 	FGeometryCollection* Collection;
 	const TArrayView<bool> ActiveTriangles;
@@ -244,7 +247,7 @@ bool UVLayout(
 	TArray<TArray<int32>> UVIslands;
 	UE::UVPacking::CreateUVIslandsFromMeshTopology<FGeomMesh>(UVMesh, UVIslands);
 
-	FUVPacker Packer;
+	UE::Geometry::FUVPacker Packer;
 	Packer.bAllowFlips = false;
 	// StandardPack doesn't support the Packer.GutterSize feature, and always tries to leave a 1 texel gutter.
 	// To approximate a larger gutter, we tell it to consider a smaller output resolution --
diff --git a/Engine/Plugins/Experimental/PlanarCutPlugin/Source/PlanarCut/Private/PlanarCut.cpp b/Engine/Plugins/Experimental/PlanarCutPlugin/Source/PlanarCut/Private/PlanarCut.cpp
--- a/Engine/Plugins/Experimental/PlanarCutPlugin/Source/PlanarCut/Private/PlanarCut.cpp
+++ b/Engine/Plugins/Experimental/PlanarCutPlugin/Source/PlanarCut/Private/PlanarCut.cpp
@@ -39,9 +39,8 @@
 #include "Misc/ScopedSlowTask.h"
 #endif
 
-
-
-
+#include "ExplicitUseGeometryMathTypes.h"		// using UE::Geometry::(math types)
+using namespace UE::Geometry;
 
 
 
diff --git a/Engine/Plugins/Experimental/PlanarCutPlugin/Source/PlanarCut/Public/FractureAutoUV.h b/Engine/Plugins/Experimental/PlanarCutPlugin/Source/PlanarCut/Public/FractureAutoUV.h
--- a/Engine/Plugins/Experimental/PlanarCutPlugin/Source/PlanarCut/Public/FractureAutoUV.h
+++ b/Engine/Plugins/Experimental/PlanarCutPlugin/Source/PlanarCut/Public/FractureAutoUV.h
@@ -41,7 +41,7 @@ void PLANARCUT_API TextureInternalSurfaces(
 	FGeometryCollection& Collection,
 	double MaxDistance,
 	int32 GutterSize,
-	TImageBuilder<FVector3f>& TextureOut,
+	UE::Geometry::TImageBuilder<UE::Geometry::FVector3f>& TextureOut,
 	bool bOnlyOddMaterials = true,
 	TArrayView<int32> WhichMaterials = TArrayView<int32>()
 );
diff --git a/Engine/Plugins/Experimental/Text3D/Source/Text3D/Private/Arrangement.h b/Engine/Plugins/Experimental/Text3D/Source/Text3D/Private/Arrangement.h
--- a/Engine/Plugins/Experimental/Text3D/Source/Text3D/Private/Arrangement.h
+++ b/Engine/Plugins/Experimental/Text3D/Source/Text3D/Private/Arrangement.h
@@ -10,8 +10,14 @@
  */
 struct FArrangement final
 {
-    FDynamicGraph2d Graph;
-    TPointHashGrid2d<int> PointHash;
+	using FVector2d = UE::Geometry::FVector2d;
+	using FIndex2i = UE::Geometry::FIndex2i;
+	using FSegment2f = UE::Geometry::FSegment2f;
+	using FSegment2d = UE::Geometry::FSegment2d;
+	using FAxisAlignedBox2f = UE::Geometry::FAxisAlignedBox2f;
+	
+    UE::Geometry::FDynamicGraph2d Graph;
+    UE::Geometry::TPointHashGrid2d<int> PointHash;
     TMap<int, bool> Directions; //bool - from A to B
 	const double VertexSnapTol = 0.001;
 
@@ -31,13 +37,13 @@ protected:
         int EID;
         int SideX;
         int SideY;
-        FIntrSegment2Segment2d Intr;
+        UE::Geometry::FIntrSegment2Segment2d Intr;
     };
 
 
     int insert_point(const FVector2d &P, double Tol = 0);
     bool insert_segment(FVector2d A, FVector2d B, double Tol = 0);
-    FIndex2i split_segment_at_t(int EID, double T, double Tol);
+    UE::Geometry::FIndex2i split_segment_at_t(int EID, double T, double Tol);
     int find_existing_vertex(FVector2d Pt);
     int find_nearest_vertex(FVector2d Pt, double SearchRadius, int IgnoreVID = -1);
     bool find_intersecting_edges(FVector2d A, FVector2d B, TArray<FIntersection>& Hits, double Tol = 0);
diff --git a/Engine/Plugins/Experimental/Text3D/Source/Text3D/Private/ContourNode.h b/Engine/Plugins/Experimental/Text3D/Source/Text3D/Private/ContourNode.h
--- a/Engine/Plugins/Experimental/Text3D/Source/Text3D/Private/ContourNode.h
+++ b/Engine/Plugins/Experimental/Text3D/Source/Text3D/Private/ContourNode.h
@@ -5,6 +5,8 @@
 #include "Templates/SharedPointer.h"
 #include "Polygon2.h"
 
+using UE::Geometry::FPolygon2f;
+
 struct FContourNode;
 using TSharedContourNode = TSharedPtr<FContourNode>;
 
diff --git a/Engine/Plugins/Experimental/Text3D/Source/Text3D/Private/GlyphLoader.h b/Engine/Plugins/Experimental/Text3D/Source/Text3D/Private/GlyphLoader.h
--- a/Engine/Plugins/Experimental/Text3D/Source/Text3D/Private/GlyphLoader.h
+++ b/Engine/Plugins/Experimental/Text3D/Source/Text3D/Private/GlyphLoader.h
@@ -10,9 +10,11 @@
 #include "Containers/List.h"
 #include "Polygon2.h"
 
+using UE::Geometry::FPolygon2f;
 
 class FGlyphLoader final
 {
+	using FVector2f = UE::Geometry::FVector2f;
 public:
 	FGlyphLoader(const FT_GlyphSlot Glyph);
 
diff --git a/Engine/Plugins/Experimental/Text3D/Source/Text3D/Private/MeshCreator.cpp b/Engine/Plugins/Experimental/Text3D/Source/Text3D/Private/MeshCreator.cpp
--- a/Engine/Plugins/Experimental/Text3D/Source/Text3D/Private/MeshCreator.cpp
+++ b/Engine/Plugins/Experimental/Text3D/Source/Text3D/Private/MeshCreator.cpp
@@ -9,6 +9,7 @@
 
 #include "ConstrainedDelaunay2.h"
 
+using namespace UE::Geometry;
 
 constexpr float FMeshCreator::OutlineExpand;
 
@@ -358,8 +359,8 @@ void FMeshCreator::TriangulateAndConvert(const TSharedContourNode& Node, int32&
 	if (!Node->bClockwise)
 	{
 		int32 VertexCount = 0;
-		FConstrainedDelaunay2f Triangulation;
-		Triangulation.FillRule = FConstrainedDelaunay2f::EFillRule::Positive;
+		UE::Geometry::FConstrainedDelaunay2f Triangulation;
+		Triangulation.FillRule = UE::Geometry::FConstrainedDelaunay2f::EFillRule::Positive;
 
 		const TSharedPtr<FContourList> ContoursLocal = Contours;
 		const TSharedRef<FData> DataLocal = Data;
@@ -369,7 +370,7 @@ void FMeshCreator::TriangulateAndConvert(const TSharedContourNode& Node, int32&
 			FContour& Contour = ContoursLocal->Add();
 			const FPolygon2f& Polygon = *NodeIn->Contour;
 
-			for (const FVector2f& Vertex : Polygon.GetVertices())
+			for (const UE::Geometry::FVector2f& Vertex : Polygon.GetVertices())
 			{
 				// Add point to contour in old format
 				const FPartPtr Point = MakeShared<FPart>();
