add ghostly previews of cutting surfaces for mesh boolean subtraction and for mesh trimming
also only show gizmo for subtracting/trimming mesh

#rnx

diff --git a/Engine/Plugins/Experimental/MeshModelingToolset/Source/MeshModelingTools/Private/CSGMeshesTool.cpp b/Engine/Plugins/Experimental/MeshModelingToolset/Source/MeshModelingTools/Private/CSGMeshesTool.cpp
--- a/Engine/Plugins/Experimental/MeshModelingToolset/Source/MeshModelingTools/Private/CSGMeshesTool.cpp
+++ b/Engine/Plugins/Experimental/MeshModelingToolset/Source/MeshModelingTools/Private/CSGMeshesTool.cpp
@@ -30,6 +30,16 @@ void UCSGMeshesTool::SetupProperties()
 		TrimProperties->RestoreProperties(this);
 		AddToolPropertySource(TrimProperties);
 
+		TrimProperties->WatchProperty(TrimProperties->WhichMesh, [this](ETrimOperation)
+		{
+			UpdateGizmoVisibility();
+			UpdatePreviewsVisibility();
+		});
+		TrimProperties->WatchProperty(TrimProperties->bShowTrimmingMesh, [this](bool)
+		{
+			UpdatePreviewsVisibility();
+		});
+
 		SetToolDisplayName(LOCTEXT("TrimMeshesToolName", "Trim"));
 		GetToolManager()->DisplayMessage(
 			LOCTEXT("OnStartTrimTool", "Trim one mesh with another. Use the transform gizmos to tweak the positions of the input objects (can help to resolve errors/failures)"),
@@ -41,6 +51,16 @@ void UCSGMeshesTool::SetupProperties()
 		CSGProperties->RestoreProperties(this);
 		AddToolPropertySource(CSGProperties);
 
+		CSGProperties->WatchProperty(CSGProperties->Operation, [this](ECSGOperation)
+		{
+			UpdateGizmoVisibility();
+			UpdatePreviewsVisibility();
+		});
+		CSGProperties->WatchProperty(CSGProperties->bShowSubtractedMesh, [this](bool)
+		{
+			UpdatePreviewsVisibility();
+		});
+
 		SetToolDisplayName(LOCTEXT("CSGMeshesToolName", "Boolean"));
 		GetToolManager()->DisplayMessage(
 			LOCTEXT("OnStartTool", "Compute CSG Booleans on the input meshes. Use the transform gizmos to tweak the positions of the input objects (can help to resolve errors/failures)"),
@@ -48,6 +68,55 @@ void UCSGMeshesTool::SetupProperties()
 	}
 }
 
+void UCSGMeshesTool::UpdatePreviewsVisibility()
+{
+	int32 ShowPreviewIdx = -1;
+	if (bTrimMode && TrimProperties->bShowTrimmingMesh)
+	{
+		ShowPreviewIdx = TrimProperties->WhichMesh == ETrimOperation::TrimA ? OriginalMeshPreviews.Num() - 1 : 0;
+	}
+	else if (!bTrimMode && CSGProperties->bShowSubtractedMesh)
+	{
+		if (CSGProperties->Operation == ECSGOperation::DifferenceAB)
+		{
+			ShowPreviewIdx = OriginalMeshPreviews.Num() - 1;
+		}
+		else if (CSGProperties->Operation == ECSGOperation::DifferenceBA)
+		{
+			ShowPreviewIdx = 0;
+		}
+	}
+	for (int32 MeshIdx = 0; MeshIdx < OriginalMeshPreviews.Num(); MeshIdx++)
+	{
+		OriginalMeshPreviews[MeshIdx]->SetVisible(ShowPreviewIdx == MeshIdx);
+	}
+}
+
+int32 UCSGMeshesTool::GetHiddenGizmoIndex() const
+{
+	int32 ParentHiddenIndex = Super::GetHiddenGizmoIndex();
+	if (ParentHiddenIndex != -1)
+	{
+		return ParentHiddenIndex;
+	}
+	if (bTrimMode)
+	{
+		return TrimProperties->WhichMesh == ETrimOperation::TrimA ? 0 : 1;
+	}
+	else if (CSGProperties->Operation == ECSGOperation::DifferenceAB)
+	{
+		return 0;
+	}
+	else if (CSGProperties->Operation == ECSGOperation::DifferenceBA)
+	{
+		return 1;
+	}
+	else
+	{
+		return -1;
+	}
+}
+
 void UCSGMeshesTool::SaveProperties()
 {
 	Super::SaveProperties();
@@ -119,6 +188,13 @@ void UCSGMeshesTool::ConvertInputsAndSetPreviewMaterials(bool bSetPreviewMesh)
 		{
 			MaterialIDs->SetValue(TID, MaterialRemap[ComponentIdx][MaterialIDs->GetValue(TID)]);
 		}
+
+		UPreviewMesh* OriginalMeshPreview = OriginalMeshPreviews.Add_GetRef(NewObject<UPreviewMesh>());
+		OriginalMeshPreview->CreateInWorld(TargetWorld, ComponentTargets[ComponentIdx]->GetWorldTransform());
+		OriginalMeshPreview->UpdatePreview(OriginalDynamicMeshes[ComponentIdx].Get());
+		OriginalMeshPreview->SetMaterial(0, ToolSetupUtil::GetSimpleCustomMaterial(GetToolManager(), FLinearColor::White, 0.05));
+		OriginalMeshPreview->SetVisible(false);
+		TransformProxies[ComponentIdx]->AddComponent(OriginalMeshPreview->GetRootComponent());
 	}
 	Preview->ConfigureMaterials(AllMaterialSet.Materials, ToolSetupUtil::GetDefaultWorkingMaterial(GetToolManager()));
 }
@@ -254,7 +330,17 @@ FText UCSGMeshesTool::GetActionName() const
 
 
 
+void UCSGMeshesTool::Shutdown(EToolShutdownType ShutdownType)
+{
+	Super::Shutdown(ShutdownType);
 
+	for (UPreviewMesh* MeshPreview : OriginalMeshPreviews)
+	{
+		MeshPreview->SetVisible(false);
+		MeshPreview->Disconnect();
+		MeshPreview = nullptr;
+	}
+}
 
 
 
diff --git a/Engine/Plugins/Experimental/MeshModelingToolset/Source/MeshModelingTools/Public/CSGMeshesTool.h b/Engine/Plugins/Experimental/MeshModelingToolset/Source/MeshModelingTools/Public/CSGMeshesTool.h
--- a/Engine/Plugins/Experimental/MeshModelingToolset/Source/MeshModelingTools/Public/CSGMeshesTool.h
+++ b/Engine/Plugins/Experimental/MeshModelingToolset/Source/MeshModelingTools/Public/CSGMeshesTool.h
@@ -40,6 +40,10 @@ public:
 	UPROPERTY(EditAnywhere, Category = Options)
 	bool bCollapseExtraEdges = true;
 
+	/** Whether to show a translucent version of the subtracted mesh, to help visualize what is being removed */
+	UPROPERTY(EditAnywhere, Category = Options, meta = (EditCondition = "Operation == ECSGOperation::DifferenceAB || Operation == ECSGOperation::DifferenceBA", EditConditionHides))
+	bool bShowSubtractedMesh = true;
+
 	/** If true, only the first mesh will keep its materials assignments; all other triangles will be assigned material 0 */
 	UPROPERTY(EditAnywhere, Category = Materials)
 	bool bOnlyUseFirstMeshMaterials = false;
@@ -61,6 +65,10 @@ public:
 	/** Whether to remove the surface inside or outside of the trimming geometry */
 	UPROPERTY(EditAnywhere, Category = Options)
 	ETrimSide TrimSide = ETrimSide::RemoveInside;
+
+	/** Whether to show a translucent version of the trimming mesh, to help visualize what is being cut */
+	UPROPERTY(EditAnywhere, Category = Options)
+	bool bShowTrimmingMesh = true;
 };
 
 
@@ -81,6 +89,8 @@ public:
 
 protected:
 
+	virtual void Shutdown(EToolShutdownType ShutdownType) override;
+
 	virtual void OnPropertyModified(UObject* PropertySet, FProperty* Property) override;
 
 	virtual void ConvertInputsAndSetPreviewMaterials(bool bSetPreviewMesh = true) override;
@@ -107,6 +117,9 @@ protected:
 
 	TArray<TSharedPtr<FDynamicMesh3, ESPMode::ThreadSafe>> OriginalDynamicMeshes;
 
+	UPROPERTY()
+	TArray<UPreviewMesh*> OriginalMeshPreviews;
+
 	UPROPERTY()
 	ULineSetComponent* DrawnLineSet;
 
@@ -114,6 +127,11 @@ protected:
 	TArray<int> CreatedBoundaryEdges;
 
 	bool bTrimMode = false;
+
+	virtual int32 GetHiddenGizmoIndex() const;
+
+	// Update visibility of ghostly preview meshes (used to show trimming or subtracting surface)
+	void UpdatePreviewsVisibility();
 };
 
 
diff --git a/Engine/Plugins/Experimental/MeshModelingToolset/Source/ModelingComponents/Private/BaseTools/BaseCreateFromSelectedTool.cpp b/Engine/Plugins/Experimental/MeshModelingToolset/Source/ModelingComponents/Private/BaseTools/BaseCreateFromSelectedTool.cpp
--- a/Engine/Plugins/Experimental/MeshModelingToolset/Source/ModelingComponents/Private/BaseTools/BaseCreateFromSelectedTool.cpp
+++ b/Engine/Plugins/Experimental/MeshModelingToolset/Source/ModelingComponents/Private/BaseTools/BaseCreateFromSelectedTool.cpp
@@ -106,14 +106,13 @@ void UBaseCreateFromSelectedTool::Setup()
 		if (NewType == EBaseCreateFromSelectedTargetType::NewAsset)
 		{
 			HandleSourcesProperties->OutputAsset = TEXT("");
-			HideGizmoIndex = -1;
 			UpdateGizmoVisibility();
 		}
 		else
 		{
 			int32 Index = (HandleSourcesProperties->WriteOutputTo == EBaseCreateFromSelectedTargetType::FirstInputAsset) ? 0 : ComponentTargets.Num() - 1;
 			HandleSourcesProperties->OutputAsset = AssetGenerationUtil::GetComponentAssetBaseName(ComponentTargets[Index]->GetOwnerComponent(), false);
-			HideGizmoIndex = Index;
+
 			// Reset the hidden gizmo to its initial position
 			FTransform ComponentTransform = ComponentTargets[Index]->GetWorldTransform();
 			TransformGizmos[Index]->SetNewGizmoTransform(ComponentTransform, true);
@@ -125,6 +124,19 @@ void UBaseCreateFromSelectedTool::Setup()
 }
 
 
+int32 UBaseCreateFromSelectedTool::GetHiddenGizmoIndex() const
+{
+	if (HandleSourcesProperties->WriteOutputTo == EBaseCreateFromSelectedTargetType::NewAsset)
+	{
+		return -1;
+	}
+	else
+	{
+		return (HandleSourcesProperties->WriteOutputTo == EBaseCreateFromSelectedTargetType::FirstInputAsset) ? 0 : ComponentTargets.Num() - 1;
+	}
+}
+
+
 void UBaseCreateFromSelectedTool::SetWorld(UWorld* World)
 {
 	this->TargetWorld = World;
@@ -154,7 +166,7 @@ void UBaseCreateFromSelectedTool::UpdateGizmoVisibility()
 	for (int32 GizmoIndex = 0; GizmoIndex < TransformGizmos.Num(); GizmoIndex++)
 	{
 		UTransformGizmo* Gizmo = TransformGizmos[GizmoIndex];
-		Gizmo->SetVisibility(TransformProperties->bShowTransformUI && GizmoIndex != HideGizmoIndex);
+		Gizmo->SetVisibility(TransformProperties->bShowTransformUI && GizmoIndex != GetHiddenGizmoIndex());
 	}
 }
 
diff --git a/Engine/Plugins/Experimental/MeshModelingToolset/Source/ModelingComponents/Public/BaseTools/BaseCreateFromSelectedTool.h b/Engine/Plugins/Experimental/MeshModelingToolset/Source/ModelingComponents/Public/BaseTools/BaseCreateFromSelectedTool.h
--- a/Engine/Plugins/Experimental/MeshModelingToolset/Source/ModelingComponents/Public/BaseTools/BaseCreateFromSelectedTool.h
+++ b/Engine/Plugins/Experimental/MeshModelingToolset/Source/ModelingComponents/Public/BaseTools/BaseCreateFromSelectedTool.h
@@ -188,6 +188,9 @@ protected:
 	// Helper to generate assets when a result is accepted; typically does not need to be overloaded
 	virtual void UpdateAsset(const FDynamicMeshOpResult& Result, TUniquePtr<FPrimitiveComponentTarget>& Target);
 
+	// Which of the transform gizmos to hide, or -1 if all gizmos can be shown
+	virtual int32 GetHiddenGizmoIndex() const;
+
 protected:
 
 	UPROPERTY()
@@ -208,7 +211,5 @@ protected:
 
 	UWorld* TargetWorld = nullptr;
 	IAssetGenerationAPI* AssetAPI = nullptr;
-
-	int32 HideGizmoIndex = -1; // Override UI setting and always hide this gizmo
 };
 
