add Keep First and Keep Last Source options to the modeling tool OnAcceptProperties, so users can keep the "operator" mesh(es) around when doing mesh booleans and similar

#rnx

diff --git a/Engine/Plugins/Experimental/MeshModelingToolset/Source/ModelingComponents/Private/PropertySets/OnAcceptProperties.cpp b/Engine/Plugins/Experimental/MeshModelingToolset/Source/ModelingComponents/Private/PropertySets/OnAcceptProperties.cpp
--- a/Engine/Plugins/Experimental/MeshModelingToolset/Source/ModelingComponents/Private/PropertySets/OnAcceptProperties.cpp
+++ b/Engine/Plugins/Experimental/MeshModelingToolset/Source/ModelingComponents/Private/PropertySets/OnAcceptProperties.cpp
@@ -12,9 +12,17 @@
 void UOnAcceptHandleSourcesProperties::ApplyMethod(const TArray<AActor*>& Actors, UInteractiveToolManager* ToolManager)
 {
 	// Hide or destroy the sources
-	if (OnToolAccept != EHandleSourcesMethod::KeepSources)
+	bool bKeepSources = OnToolAccept == EHandleSourcesMethod::KeepSources;
+	if (Actors.Num() == 1 && (OnToolAccept == EHandleSourcesMethod::KeepFirstSource || OnToolAccept == EHandleSourcesMethod::KeepLastSource))
 	{
-		bool bDelete = OnToolAccept == EHandleSourcesMethod::DeleteSources;
+		// if there's only one actor, keeping any source == keeping all sources
+		bKeepSources = true;
+	}
+	if (!bKeepSources)
+	{
+		bool bDelete = OnToolAccept == EHandleSourcesMethod::DeleteSources
+					|| OnToolAccept == EHandleSourcesMethod::KeepFirstSource
+					|| OnToolAccept == EHandleSourcesMethod::KeepLastSource;
 		if (bDelete)
 		{
 			ToolManager->BeginUndoTransaction(LOCTEXT("RemoveSources", "Remove Sources"));
@@ -26,8 +34,23 @@ void UOnAcceptHandleSourcesProperties::ApplyMethod(const TArray<AActor*>& Actors
 #endif
 		}
 
-		for (AActor* Actor : Actors)
+		int32 SkipIdx = -1;
+		if (OnToolAccept == EHandleSourcesMethod::KeepFirstSource)
 		{
+			SkipIdx = 0;
+		}
+		else if (OnToolAccept == EHandleSourcesMethod::KeepLastSource)
+		{
+			SkipIdx = Actors.Num() - 1;
+		}
+		for (int32 ActorIdx = 0; ActorIdx < Actors.Num(); ActorIdx++)
+		{
+			if (ActorIdx == SkipIdx)
+			{
+				continue;
+			}
+
+			AActor* Actor = Actors[ActorIdx];
 			if (bDelete)
 			{
 				Actor->Destroy();
diff --git a/Engine/Plugins/Experimental/MeshModelingToolset/Source/ModelingComponents/Public/PropertySets/OnAcceptProperties.h b/Engine/Plugins/Experimental/MeshModelingToolset/Source/ModelingComponents/Public/PropertySets/OnAcceptProperties.h
--- a/Engine/Plugins/Experimental/MeshModelingToolset/Source/ModelingComponents/Public/PropertySets/OnAcceptProperties.h
+++ b/Engine/Plugins/Experimental/MeshModelingToolset/Source/ModelingComponents/Public/PropertySets/OnAcceptProperties.h
@@ -20,7 +20,13 @@ enum class EHandleSourcesMethod : uint8
 
 	HideSources = 1		UMETA(DisplayName = "Hide Sources"),
 
-	KeepSources = 2		UMETA(DisplayName = "Keep Sources")
+	KeepSources = 2		UMETA(DisplayName = "Keep Sources"),
+
+	/** Delete all but the first source */
+	KeepFirstSource = 3	UMETA(DisplayName = "Keep First Source"),
+
+	/** Delete all but the last source */
+	KeepLastSource = 4	UMETA(DisplayName = "Keep Last Source")
 };
 
 
