MeshModeling Tools:  Make UV updates robust to the case when the src MeshDescription has duplicate tris.

diff --git a/Engine/Plugins/Experimental/MeshModelingToolset/Source/MeshModelingTools/Private/UVProjectionTool.cpp b/Engine/Plugins/Experimental/MeshModelingToolset/Source/MeshModelingTools/Private/UVProjectionTool.cpp
--- a/Engine/Plugins/Experimental/MeshModelingToolset/Source/MeshModelingTools/Private/UVProjectionTool.cpp
+++ b/Engine/Plugins/Experimental/MeshModelingToolset/Source/MeshModelingTools/Private/UVProjectionTool.cpp
@@ -355,8 +355,22 @@ void UUVProjectionTool::GenerateAsset(const TArray<FDynamicMeshOpResult>& Result
 		check(Results[ComponentIdx].Mesh.Get() != nullptr);
 		ComponentTargets[ComponentIdx]->CommitMesh([&Results, &ComponentIdx, this](const FPrimitiveComponentTarget::FCommitParams& CommitParams)
 		{
-			FDynamicMeshToMeshDescription Converter;
-			Converter.UpdateAttributes(Results[ComponentIdx].Mesh.Get(), *CommitParams.MeshDescription, false, false, true/*update uvs*/);
+			FDynamicMesh3* DynamicMesh = Results[ComponentIdx].Mesh.Get();
+			FMeshDescription* MeshDescription = CommitParams.MeshDescription;
+
+			bool bVerticesOnly = false;
+			bool bAttributesOnly = true;
+			if (FDynamicMeshToMeshDescription::HaveMatchingElementCounts(DynamicMesh, MeshDescription, bVerticesOnly, bAttributesOnly))
+			{
+				FDynamicMeshToMeshDescription Converter;
+				Converter.UpdateAttributes(DynamicMesh, *MeshDescription, false, false, true/*update uvs*/);
+			}
+			else
+			{
+				// must have been duplicate tris in the mesh description; we can't count on 1-to-1 mapping of TriangleIDs.  Just convert 
+				FDynamicMeshToMeshDescription Converter;
+				Converter.Convert(DynamicMesh, *MeshDescription);
+			}
 		});
 	}
 
diff --git a/Engine/Plugins/Experimental/MeshModelingToolset/Source/MeshModelingToolsEditorOnly/Private/ParameterizeMeshTool.cpp b/Engine/Plugins/Experimental/MeshModelingToolset/Source/MeshModelingToolsEditorOnly/Private/ParameterizeMeshTool.cpp
--- a/Engine/Plugins/Experimental/MeshModelingToolset/Source/MeshModelingToolsEditorOnly/Private/ParameterizeMeshTool.cpp
+++ b/Engine/Plugins/Experimental/MeshModelingToolset/Source/MeshModelingToolsEditorOnly/Private/ParameterizeMeshTool.cpp
@@ -199,8 +199,21 @@ void UParameterizeMeshTool::Shutdown(EToolShutdownType ShutdownType)
 
 		ComponentTarget->CommitMesh([DynamicMeshResult](const FPrimitiveComponentTarget::FCommitParams& CommitParams)
 		{
-			FDynamicMeshToMeshDescription Converter;
-			Converter.UpdateAttributes(DynamicMeshResult, *CommitParams.MeshDescription, false, false, true/*update uvs*/);
+			FMeshDescription* MeshDescription = CommitParams.MeshDescription;
+
+			bool bVerticesOnly = false;
+			bool bAttributesOnly = true;
+			if (FDynamicMeshToMeshDescription::HaveMatchingElementCounts(DynamicMeshResult, MeshDescription, bVerticesOnly, bAttributesOnly))
+			{
+				FDynamicMeshToMeshDescription Converter;
+				Converter.UpdateAttributes(DynamicMeshResult, *MeshDescription, false, false, true/*update uvs*/);
+			}
+			else
+			{
+				// must have been duplicate tris in the mesh description; we can't count on 1-to-1 mapping of TriangleIDs.  Just convert 
+				FDynamicMeshToMeshDescription Converter;
+				Converter.Convert(DynamicMeshResult, *MeshDescription);
+			}
 		});
 		GetToolManager()->EndUndoTransaction();
 	}
diff --git a/Engine/Plugins/Experimental/MeshModelingToolset/Source/MeshModelingToolsEditorOnly/Private/UVLayoutTool.cpp b/Engine/Plugins/Experimental/MeshModelingToolset/Source/MeshModelingToolsEditorOnly/Private/UVLayoutTool.cpp
--- a/Engine/Plugins/Experimental/MeshModelingToolset/Source/MeshModelingToolsEditorOnly/Private/UVLayoutTool.cpp
+++ b/Engine/Plugins/Experimental/MeshModelingToolset/Source/MeshModelingToolsEditorOnly/Private/UVLayoutTool.cpp
@@ -356,8 +356,22 @@ void UUVLayoutTool::GenerateAsset(const TArray<FDynamicMeshOpResult>& Results)
 		check(Results[ComponentIdx].Mesh.Get() != nullptr);
 		ComponentTargets[ComponentIdx]->CommitMesh([&Results, &ComponentIdx, this](const FPrimitiveComponentTarget::FCommitParams& CommitParams)
 		{
-			FDynamicMeshToMeshDescription Converter;
-			Converter.UpdateAttributes(Results[ComponentIdx].Mesh.Get(), *CommitParams.MeshDescription, false, false, true/*update uvs*/);
+			FDynamicMesh3* DynamicMesh = Results[ComponentIdx].Mesh.Get();
+			FMeshDescription* MeshDescription = CommitParams.MeshDescription;
+			
+			bool bVerticesOnly = false;
+			bool bAttributesOnly = true;
+			if (FDynamicMeshToMeshDescription::HaveMatchingElementCounts(DynamicMesh, MeshDescription, bVerticesOnly, bAttributesOnly))
+			{
+				FDynamicMeshToMeshDescription Converter;
+				Converter.UpdateAttributes(DynamicMesh, *MeshDescription, false, false, true/*update uvs*/);
+			}
+			else
+			{
+				// must have been duplicate tris in the mesh description; we can't count on 1-to-1 mapping of TriangleIDs.  Just convert
+				FDynamicMeshToMeshDescription Converter;
+				Converter.Convert(DynamicMesh, *MeshDescription);
+			}
 		});
 	}
 
