ModelingTools: Replace MultiSelectionTool PrimitiveComponentTarget with ToolTargets (pt.2)

#rnx

diff --git a/Engine/Plugins/Experimental/MeshLODToolset/Source/MeshLODToolset/Private/Tools/GenerateStaticMeshLODAssetTool.cpp b/Engine/Plugins/Experimental/MeshLODToolset/Source/MeshLODToolset/Private/Tools/GenerateStaticMeshLODAssetTool.cpp
--- a/Engine/Plugins/Experimental/MeshLODToolset/Source/MeshLODToolset/Private/Tools/GenerateStaticMeshLODAssetTool.cpp
+++ b/Engine/Plugins/Experimental/MeshLODToolset/Source/MeshLODToolset/Private/Tools/GenerateStaticMeshLODAssetTool.cpp
@@ -43,6 +43,9 @@
 #include "Polygroups/PolygroupSet.h"
 #include "Polygroups/PolygroupUtil.h"
 
+#include "TargetInterfaces/PrimitiveComponentBackedTarget.h"
+#include "TargetInterfaces/StaticMeshBackedTarget.h"
+#include "ToolTargetManager.h"
 
 #if WITH_EDITOR
 #include "Misc/ScopedSlowTask.h"
@@ -147,30 +150,27 @@ namespace GenerateStaticMeshLODAssetLocals
  */
 
 
+const FToolTargetTypeRequirements& UGenerateStaticMeshLODAssetToolBuilder::GetTargetRequirements() const
+{
+	static FToolTargetTypeRequirements TypeRequirements({
+		UPrimitiveComponentBackedTarget::StaticClass(),
+		UStaticMeshBackedTarget::StaticClass()
+		});
+	return TypeRequirements;
+}
+
 bool UGenerateStaticMeshLODAssetToolBuilder::CanBuildTool(const FToolBuilderState& SceneState) const
 {
 	// hack to make multi-tool look like single-tool
-	return (AssetAPI != nullptr && ToolBuilderUtil::CountComponents(SceneState, CanMakeComponentTarget) == 1);
+	return (AssetAPI != nullptr && SceneState.TargetManager->CountSelectedAndTargetable(SceneState, GetTargetRequirements()) == 1);
 }
 
 UInteractiveTool* UGenerateStaticMeshLODAssetToolBuilder::BuildTool(const FToolBuilderState& SceneState) const
 {
 	UGenerateStaticMeshLODAssetTool* NewTool = NewObject<UGenerateStaticMeshLODAssetTool>(SceneState.ToolManager);
 
-	TArray<UActorComponent*> Components = ToolBuilderUtil::FindAllComponents(SceneState, CanMakeComponentTarget);
-	check(Components.Num() > 0);
-
-	TArray<TUniquePtr<FPrimitiveComponentTarget>> ComponentTargets;
-	for (UActorComponent* ActorComponent : Components)
-	{
-		auto* MeshComponent = Cast<UPrimitiveComponent>(ActorComponent);
-		if (MeshComponent)
-		{
-			ComponentTargets.Add(MakeComponentTarget(MeshComponent));
-		}
-	}
-
-	NewTool->SetSelection(MoveTemp(ComponentTargets));
+	TArray<TObjectPtr<UToolTarget>> Targets = SceneState.TargetManager->BuildAllSelectedTargetable(SceneState, GetTargetRequirements());
+	NewTool->SetTargets(MoveTemp(Targets));
 	NewTool->SetWorld(SceneState.World);
 	NewTool->SetAssetAPI(AssetAPI);
 
@@ -204,7 +204,7 @@ void UGenerateStaticMeshLODAssetTool::Setup()
 
 	GenerateProcess = NewObject<UGenerateStaticMeshLODProcess>(this);
 
-	TUniquePtr<FPrimitiveComponentTarget>& SourceComponent = ComponentTargets[0];
+	IPrimitiveComponentBackedTarget* SourceComponent = TargetComponentInterface(0);
 	UStaticMeshComponent* StaticMeshComponent = CastChecked<UStaticMeshComponent>(SourceComponent->GetOwnerComponent());
 	UStaticMesh* StaticMesh = StaticMeshComponent->GetStaticMesh();
 	if (StaticMesh)
@@ -215,7 +215,7 @@ void UGenerateStaticMeshLODAssetTool::Setup()
 	BasicProperties = NewObject<UGenerateStaticMeshLODAssetToolProperties>(this);
 	AddToolPropertySource(BasicProperties);
 	BasicProperties->RestoreProperties(this);
-	BasicProperties->OutputName = AssetGenerationUtil::GetComponentAssetBaseName(ComponentTargets[0]->GetOwnerComponent());
+	BasicProperties->OutputName = AssetGenerationUtil::GetComponentAssetBaseName(TargetComponentInterface(0)->GetOwnerComponent());
 	BasicProperties->GeneratedSuffix = TEXT("_AutoLOD");
 	BasicProperties->GeneratorSettings = GenerateProcess->GetCurrentSettings();
 	BasicProperties->WatchProperty(BasicProperties->GeneratorSettings.FilterGroupLayer, [this](FName) { OnSettingsModified(); });
diff --git a/Engine/Plugins/Experimental/MeshLODToolset/Source/MeshLODToolset/Private/Tools/LODManagerTool.cpp b/Engine/Plugins/Experimental/MeshLODToolset/Source/MeshLODToolset/Private/Tools/LODManagerTool.cpp
--- a/Engine/Plugins/Experimental/MeshLODToolset/Source/MeshLODToolset/Private/Tools/LODManagerTool.cpp
+++ b/Engine/Plugins/Experimental/MeshLODToolset/Source/MeshLODToolset/Private/Tools/LODManagerTool.cpp
@@ -20,6 +20,11 @@
 
 #include "Components/PrimitiveComponent.h"
 
+#include "TargetInterfaces/MaterialProvider.h"
+#include "TargetInterfaces/PrimitiveComponentBackedTarget.h"
+#include "TargetInterfaces/StaticMeshBackedTarget.h"
+#include "ToolTargetManager.h"
+
 #include "ExplicitUseGeometryMathTypes.h"		// using UE::Geometry::(math types)
 using namespace UE::Geometry;
 
@@ -30,32 +35,27 @@ using namespace UE::Geometry;
  */
 
 
+const FToolTargetTypeRequirements& ULODManagerToolBuilder::GetTargetRequirements() const
+{
+	static FToolTargetTypeRequirements TypeRequirements({
+		UMaterialProvider::StaticClass(),
+		UPrimitiveComponentBackedTarget::StaticClass(),
+		UStaticMeshBackedTarget::StaticClass()
+		});
+	return TypeRequirements;
+}
+
 bool ULODManagerToolBuilder::CanBuildTool(const FToolBuilderState& SceneState) const
 {
-	int32 NumStaticMeshes = ToolBuilderUtil::CountComponents(SceneState, [&](UActorComponent* Comp) { return Cast<UStaticMeshComponent>(Comp) != nullptr; });
-	int32 NumComponentTargets = ToolBuilderUtil::CountComponents(SceneState, CanMakeComponentTarget);
-	return (NumStaticMeshes > 0 && NumStaticMeshes == NumComponentTargets);
+	return SceneState.TargetManager->CountSelectedAndTargetable(SceneState, GetTargetRequirements()) > 0;
 }
 
 UInteractiveTool* ULODManagerToolBuilder::BuildTool(const FToolBuilderState& SceneState) const
 {
 	ULODManagerTool* NewTool = NewObject<ULODManagerTool>(SceneState.ToolManager);
 
-	TArray<UActorComponent*> ValidComponents = ToolBuilderUtil::FindAllComponents(SceneState,
-		[&](UActorComponent* Comp) { return Cast<UStaticMeshComponent>(Comp) != nullptr; });
-	check(ValidComponents.Num() > 0);
-
-	TArray<TUniquePtr<FPrimitiveComponentTarget>> ComponentTargets;
-	for (UActorComponent* ActorComponent : ValidComponents)
-	{
-		UStaticMeshComponent* MeshComponent = Cast<UStaticMeshComponent>(ActorComponent);
-		if ( MeshComponent )
-		{
-			ComponentTargets.Add(MakeComponentTarget(MeshComponent));
-		}
-	}
-
-	NewTool->SetSelection(MoveTemp(ComponentTargets));
+	TArray<TObjectPtr<UToolTarget>> Targets = SceneState.TargetManager->BuildAllSelectedTargetable(SceneState, GetTargetRequirements());
+	NewTool->SetTargets(MoveTemp(Targets));
 	NewTool->SetWorld(SceneState.World);
 
 	return NewTool;
@@ -96,7 +96,7 @@ void ULODManagerTool::Setup()
 	MaterialActions->Initialize(this);
 	AddToolPropertySource(MaterialActions);
 
-	if (ComponentTargets.Num() == 1)
+	if (Targets.Num() == 1)
 	{
 		LODInfoProperties = NewObject<ULODManagerLODProperties>(this);
 		AddToolPropertySource(LODInfoProperties);
@@ -107,14 +107,14 @@ void ULODManagerTool::Setup()
 		LODPreviewProperties->WatchProperty(LODPreviewProperties->VisibleLOD, [this](FString NewLOD) { bPreviewLODValid = false; });
 
 		LODPreview = NewObject<UPreviewMesh>(this);
-		LODPreview->CreateInWorld(TargetWorld, ComponentTargets[0]->GetWorldTransform());
+		LODPreview->CreateInWorld(TargetWorld, TargetComponentInterface(0)->GetWorldTransform());
 		LODPreview->SetVisible(false);
 
 		LODPreviewLines = NewObject<UPreviewGeometry>(this);
-		LODPreviewLines->CreateInWorld(TargetWorld, ComponentTargets[0]->GetWorldTransform());
+		LODPreviewLines->CreateInWorld(TargetWorld, TargetComponentInterface(0)->GetWorldTransform());
 
 		FComponentMaterialSet MaterialSet;
-		ComponentTargets[0]->GetMaterialSet(MaterialSet);
+		TargetMaterialInterface(0)->GetMaterialSet(MaterialSet);
 		LODPreview->SetMaterials(MaterialSet.Materials);
 
 		bLODInfoValid = false;
@@ -136,9 +136,9 @@ void ULODManagerTool::Shutdown(EToolShutdownType ShutdownType)
 		LODPreview->Disconnect();
 		LODPreviewLines->Disconnect();
 	}
-	for (TUniquePtr<FPrimitiveComponentTarget>& Target : ComponentTargets)
+	for (int32 ComponentIdx = 0; ComponentIdx < Targets.Num(); ComponentIdx++)
 	{
-		Target->SetOwnerVisibility(true);
+		TargetComponentInterface(ComponentIdx)->SetOwnerVisibility(true);
 	}
 }
 
@@ -187,11 +187,11 @@ void ULODManagerTool::OnTick(float DeltaTime)
 
 UStaticMesh* ULODManagerTool::GetSingleStaticMesh()
 {
-	if (ComponentTargets.Num() > 1)
+	if (Targets.Num() > 1)
 	{
 		return nullptr;
 	}
-	TUniquePtr<FPrimitiveComponentTarget>& Target = ComponentTargets[0];
+	IPrimitiveComponentBackedTarget* Target = TargetComponentInterface(0);
 	UStaticMeshComponent* StaticMeshComponent = Cast<UStaticMeshComponent>(Target->GetOwnerComponent());
 	if (StaticMeshComponent == nullptr || StaticMeshComponent->GetStaticMesh() == nullptr)
 	{
@@ -377,7 +377,7 @@ void ULODManagerTool::UpdatePreviewLOD()
 	{
 		// badness
 		LODPreview->SetVisible(false);
-		ComponentTargets[0]->SetOwnerVisibility(true);
+		TargetComponentInterface(0)->SetOwnerVisibility(true);
 		return;
 	}
 
@@ -390,7 +390,7 @@ void ULODManagerTool::UpdatePreviewLOD()
 	
 	if (Found)
 	{
-		ComponentTargets[0]->SetOwnerVisibility(false);
+		TargetComponentInterface(0)->SetOwnerVisibility(false);
 		LODPreview->ReplaceMesh( (*Found)->Mesh );
 		LODPreview->SetVisible(true);
 		UpdatePreviewLines(**Found);
@@ -400,7 +400,7 @@ void ULODManagerTool::UpdatePreviewLOD()
 	{
 		LODPreview->SetVisible(false);
 		LODPreviewLines->SetAllVisible(false);
-		ComponentTargets[0]->SetOwnerVisibility(true);
+		TargetComponentInterface(0)->SetOwnerVisibility(true);
 	}
 }
 
@@ -461,9 +461,10 @@ bool ULODManagerTool::CacheLODMesh(const FString& Name, FLODName LODName)
 void ULODManagerTool::DeleteHiResSourceModel()
 {
 	GetToolManager()->BeginUndoTransaction(LOCTEXT("DeleteHiResSourceModel", "Delete HiRes Source"));
-	for (TUniquePtr<FPrimitiveComponentTarget>& Target : ComponentTargets)
+	for (int32 ComponentIdx = 0; ComponentIdx < Targets.Num(); ComponentIdx++)
 	{
-		UStaticMeshComponent* StaticMeshComponent = Cast<UStaticMeshComponent>(Target->GetOwnerComponent());
+		IPrimitiveComponentBackedTarget* TargetComponent = TargetComponentInterface(ComponentIdx);
+		UStaticMeshComponent* StaticMeshComponent = Cast<UStaticMeshComponent>(TargetComponent->GetOwnerComponent());
 		if (StaticMeshComponent == nullptr || StaticMeshComponent->GetStaticMesh() == nullptr)
 		{
 			continue;
@@ -493,9 +494,10 @@ void ULODManagerTool::DeleteHiResSourceModel()
 void ULODManagerTool::MoveHiResToLOD0()
 {
 	GetToolManager()->BeginUndoTransaction(LOCTEXT("RestoreLOD0", "Move HiRes to LOD0"));
-	for (TUniquePtr<FPrimitiveComponentTarget>& Target : ComponentTargets)
+	for (int32 ComponentIdx = 0; ComponentIdx < Targets.Num(); ComponentIdx++)
 	{
-		UStaticMeshComponent* StaticMeshComponent = Cast<UStaticMeshComponent>(Target->GetOwnerComponent());
+		IPrimitiveComponentBackedTarget* TargetComponent = TargetComponentInterface(ComponentIdx);
+		UStaticMeshComponent* StaticMeshComponent = Cast<UStaticMeshComponent>(TargetComponent->GetOwnerComponent());
 		if (StaticMeshComponent == nullptr || StaticMeshComponent->GetStaticMesh() == nullptr)
 		{
 			continue;
@@ -534,9 +536,10 @@ void ULODManagerTool::MoveHiResToLOD0()
 void ULODManagerTool::RemoveUnreferencedMaterials()
 {
 	GetToolManager()->BeginUndoTransaction(LOCTEXT("RemoveUnreferencedMaterials", "Remove Unreferenced Materials"));
-	for (TUniquePtr<FPrimitiveComponentTarget>& Target : ComponentTargets)
+	for (int32 ComponentIdx = 0; ComponentIdx < Targets.Num(); ComponentIdx++)
 	{
-		UStaticMeshComponent* StaticMeshComponent = Cast<UStaticMeshComponent>(Target->GetOwnerComponent());
+		IPrimitiveComponentBackedTarget* TargetComponent = TargetComponentInterface(ComponentIdx);
+		UStaticMeshComponent* StaticMeshComponent = Cast<UStaticMeshComponent>(TargetComponent->GetOwnerComponent());
 		if (StaticMeshComponent == nullptr || StaticMeshComponent->GetStaticMesh() == nullptr)
 		{
 			continue;
diff --git a/Engine/Plugins/Experimental/MeshLODToolset/Source/MeshLODToolset/Public/Tools/GenerateStaticMeshLODAssetTool.h b/Engine/Plugins/Experimental/MeshLODToolset/Source/MeshLODToolset/Public/Tools/GenerateStaticMeshLODAssetTool.h
--- a/Engine/Plugins/Experimental/MeshLODToolset/Source/MeshLODToolset/Public/Tools/GenerateStaticMeshLODAssetTool.h
+++ b/Engine/Plugins/Experimental/MeshLODToolset/Source/MeshLODToolset/Public/Tools/GenerateStaticMeshLODAssetTool.h
@@ -47,6 +47,9 @@ public:
 
 	virtual bool CanBuildTool(const FToolBuilderState& SceneState) const override;
 	virtual UInteractiveTool* BuildTool(const FToolBuilderState& SceneState) const override;
+
+protected:
+	virtual const FToolTargetTypeRequirements& GetTargetRequirements() const override;
 };
 
 
diff --git a/Engine/Plugins/Experimental/MeshLODToolset/Source/MeshLODToolset/Public/Tools/LODManagerTool.h b/Engine/Plugins/Experimental/MeshLODToolset/Source/MeshLODToolset/Public/Tools/LODManagerTool.h
--- a/Engine/Plugins/Experimental/MeshLODToolset/Source/MeshLODToolset/Public/Tools/LODManagerTool.h
+++ b/Engine/Plugins/Experimental/MeshLODToolset/Source/MeshLODToolset/Public/Tools/LODManagerTool.h
@@ -27,6 +27,9 @@ public:
 
 	virtual bool CanBuildTool(const FToolBuilderState& SceneState) const override;
 	virtual UInteractiveTool* BuildTool(const FToolBuilderState& SceneState) const override;
+
+protected:
+	virtual const FToolTargetTypeRequirements& GetTargetRequirements() const override;
 };
 
 
diff --git a/Engine/Plugins/Experimental/MeshModelingToolset/Source/MeshModelingTools/Private/BakeMeshAttributeMapsTool.cpp b/Engine/Plugins/Experimental/MeshModelingToolset/Source/MeshModelingTools/Private/BakeMeshAttributeMapsTool.cpp
--- a/Engine/Plugins/Experimental/MeshModelingToolset/Source/MeshModelingTools/Private/BakeMeshAttributeMapsTool.cpp
+++ b/Engine/Plugins/Experimental/MeshModelingToolset/Source/MeshModelingTools/Private/BakeMeshAttributeMapsTool.cpp
@@ -118,7 +118,7 @@ void UBakeMeshAttributeMapsTool::Setup()
 
 	// transfer materials
 	FComponentMaterialSet MaterialSet;
-	MaterialProvider->GetMaterialSet(MaterialSet, false);
+	MaterialProvider->GetMaterialSet(MaterialSet);
 	for (int k = 0; k < MaterialSet.Materials.Num(); ++k)
 	{
 		DynamicMeshComponent->SetMaterial(k, MaterialSet.Materials[k]);
diff --git a/Engine/Plugins/Experimental/MeshModelingToolset/Source/MeshModelingTools/Private/BakeTransformTool.cpp b/Engine/Plugins/Experimental/MeshModelingToolset/Source/MeshModelingTools/Private/BakeTransformTool.cpp
--- a/Engine/Plugins/Experimental/MeshModelingToolset/Source/MeshModelingTools/Private/BakeTransformTool.cpp
+++ b/Engine/Plugins/Experimental/MeshModelingToolset/Source/MeshModelingTools/Private/BakeTransformTool.cpp
@@ -22,6 +22,7 @@
 #include "TargetInterfaces/MeshDescriptionCommitter.h"
 #include "TargetInterfaces/MeshDescriptionProvider.h"
 #include "TargetInterfaces/PrimitiveComponentBackedTarget.h"
+#include "TargetInterfaces/AssetBackedTarget.h"
 #include "ToolTargetManager.h"
 
 #include "ExplicitUseGeometryMathTypes.h"		// using UE::Geometry::(math types)
@@ -39,7 +40,8 @@ const FToolTargetTypeRequirements& UBakeTransformToolBuilder::GetTargetRequireme
 	static FToolTargetTypeRequirements TypeRequirements({
 		UMeshDescriptionCommitter::StaticClass(),
 		UMeshDescriptionProvider::StaticClass(),
-		UPrimitiveComponentBackedTarget::StaticClass()
+		UPrimitiveComponentBackedTarget::StaticClass(),
+		UAssetBackedTarget::StaticClass()
 		});
 	return TypeRequirements;
 }
@@ -100,7 +102,7 @@ void UBakeTransformTool::Setup()
 	}
 
 	bool bHasZeroScales = false;
-	for (int32 ComponentIdx = 0; ComponentIdx < ComponentTargets.Num(); ComponentIdx++)
+	for (int32 ComponentIdx = 0; ComponentIdx < Targets.Num(); ComponentIdx++)
 	{
 		IPrimitiveComponentBackedTarget* TargetComponent = TargetComponentInterface(ComponentIdx);
 		if (TargetComponent->GetWorldTransform().GetScale3D().GetAbsMin() < KINDA_SMALL_NUMBER)
diff --git a/Engine/Plugins/Experimental/MeshModelingToolset/Source/MeshModelingTools/Private/CSGMeshesTool.cpp b/Engine/Plugins/Experimental/MeshModelingToolset/Source/MeshModelingTools/Private/CSGMeshesTool.cpp
--- a/Engine/Plugins/Experimental/MeshModelingToolset/Source/MeshModelingTools/Private/CSGMeshesTool.cpp
+++ b/Engine/Plugins/Experimental/MeshModelingToolset/Source/MeshModelingTools/Private/CSGMeshesTool.cpp
@@ -133,17 +133,17 @@ void UCSGMeshesTool::SaveProperties()
 
 void UCSGMeshesTool::ConvertInputsAndSetPreviewMaterials(bool bSetPreviewMesh)
 {
-	OriginalDynamicMeshes.SetNum(ComponentTargets.Num());
+	OriginalDynamicMeshes.SetNum(Targets.Num());
 	FComponentMaterialSet AllMaterialSet;
 	TMap<UMaterialInterface*, int> KnownMaterials;
-	TArray<TArray<int>> MaterialRemap; MaterialRemap.SetNum(ComponentTargets.Num());
+	TArray<TArray<int>> MaterialRemap; MaterialRemap.SetNum(Targets.Num());
 
 	if (bTrimMode || !CSGProperties->bOnlyUseFirstMeshMaterials)
 	{
-		for (int ComponentIdx = 0; ComponentIdx < ComponentTargets.Num(); ComponentIdx++)
+		for (int ComponentIdx = 0; ComponentIdx < Targets.Num(); ComponentIdx++)
 		{
 			FComponentMaterialSet ComponentMaterialSet;
-			ComponentTargets[ComponentIdx]->GetMaterialSet(ComponentMaterialSet);
+			TargetMaterialInterface(ComponentIdx)->GetMaterialSet(ComponentMaterialSet);
 			for (UMaterialInterface* Mat : ComponentMaterialSet.Materials)
 			{
 				int* FoundMatIdx = KnownMaterials.Find(Mat);
@@ -163,22 +163,22 @@ void UCSGMeshesTool::ConvertInputsAndSetPreviewMaterials(bool bSetPreviewMesh)
 	}
 	else
 	{
-		ComponentTargets[0]->GetMaterialSet(AllMaterialSet);
+		TargetMaterialInterface(0)->GetMaterialSet(AllMaterialSet);
 		for (int MatIdx = 0; MatIdx < AllMaterialSet.Materials.Num(); MatIdx++)
 		{
 			MaterialRemap[0].Add(MatIdx);
 		}
-		for (int ComponentIdx = 1; ComponentIdx < ComponentTargets.Num(); ComponentIdx++)
+		for (int ComponentIdx = 1; ComponentIdx < Targets.Num(); ComponentIdx++)
 		{
-			MaterialRemap[ComponentIdx].Init(0, ComponentTargets[ComponentIdx]->GetNumMaterials());
+			MaterialRemap[ComponentIdx].Init(0, TargetMaterialInterface(ComponentIdx)->GetNumMaterials());
 		}
 	}
 
-	for (int ComponentIdx = 0; ComponentIdx < ComponentTargets.Num(); ComponentIdx++)
+	for (int ComponentIdx = 0; ComponentIdx < Targets.Num(); ComponentIdx++)
 	{
 		OriginalDynamicMeshes[ComponentIdx] = MakeShared<FDynamicMesh3, ESPMode::ThreadSafe>();
 		FMeshDescriptionToDynamicMesh Converter;
-		Converter.Convert(ComponentTargets[ComponentIdx]->GetMesh(), *OriginalDynamicMeshes[ComponentIdx]);
+		Converter.Convert(TargetMeshProviderInterface(ComponentIdx)->GetMeshDescription(), *OriginalDynamicMeshes[ComponentIdx]);
 
 		// ensure materials and attributes are always enabled
 		OriginalDynamicMeshes[ComponentIdx]->EnableAttributes();
@@ -190,7 +190,7 @@ void UCSGMeshesTool::ConvertInputsAndSetPreviewMaterials(bool bSetPreviewMesh)
 		}
 
 		UPreviewMesh* OriginalMeshPreview = OriginalMeshPreviews.Add_GetRef(NewObject<UPreviewMesh>());
-		OriginalMeshPreview->CreateInWorld(TargetWorld, ComponentTargets[ComponentIdx]->GetWorldTransform());
+		OriginalMeshPreview->CreateInWorld(TargetWorld, TargetComponentInterface(ComponentIdx)->GetWorldTransform());
 		OriginalMeshPreview->UpdatePreview(OriginalDynamicMeshes[ComponentIdx].Get());
 		OriginalMeshPreview->SetMaterial(0, ToolSetupUtil::GetSimpleCustomMaterial(GetToolManager(), FLinearColor::White, 0.05));
 		OriginalMeshPreview->SetVisible(false);
@@ -265,7 +265,7 @@ TUniquePtr<FDynamicMeshOperator> UCSGMeshesTool::MakeNewOperator()
 	}
 
 	check(OriginalDynamicMeshes.Num() == 2);
-	check(ComponentTargets.Num() == 2);
+	check(Targets.Num() == 2);
 	BooleanOp->Transforms.SetNum(2);
 	BooleanOp->Meshes.SetNum(2);
 	for (int Idx = 0; Idx < 2; Idx++)
diff --git a/Engine/Plugins/Experimental/MeshModelingToolset/Source/MeshModelingTools/Private/CombineMeshesTool.cpp b/Engine/Plugins/Experimental/MeshModelingToolset/Source/MeshModelingTools/Private/CombineMeshesTool.cpp
--- a/Engine/Plugins/Experimental/MeshModelingToolset/Source/MeshModelingTools/Private/CombineMeshesTool.cpp
+++ b/Engine/Plugins/Experimental/MeshModelingToolset/Source/MeshModelingTools/Private/CombineMeshesTool.cpp
@@ -351,7 +351,7 @@ void UCombineMeshesTool::CreateNewAsset()
 void UCombineMeshesTool::UpdateExistingAsset()
 {
 	TArray<const FMeshDescription*> MeshDescriptions;
-	for (int32 ComponentIdx = 0; ComponentIdx < ComponentTargets.Num(); ComponentIdx++)
+	for (int32 ComponentIdx = 0; ComponentIdx < Targets.Num(); ComponentIdx++)
 	{
 		MeshDescriptions.Add(TargetMeshProviderInterface(ComponentIdx)->GetMeshDescription());
 	}
diff --git a/Engine/Plugins/Experimental/MeshModelingToolset/Source/MeshModelingTools/Private/EditNormalsTool.cpp b/Engine/Plugins/Experimental/MeshModelingToolset/Source/MeshModelingTools/Private/EditNormalsTool.cpp
--- a/Engine/Plugins/Experimental/MeshModelingToolset/Source/MeshModelingTools/Private/EditNormalsTool.cpp
+++ b/Engine/Plugins/Experimental/MeshModelingToolset/Source/MeshModelingTools/Private/EditNormalsTool.cpp
@@ -20,6 +20,11 @@
 #include "Engine/Classes/Engine/StaticMesh.h"
 #include "Engine/Classes/Components/StaticMeshComponent.h"
 
+#include "TargetInterfaces/MeshDescriptionCommitter.h"
+#include "TargetInterfaces/MeshDescriptionProvider.h"
+#include "TargetInterfaces/PrimitiveComponentBackedTarget.h"
+#include "ToolTargetManager.h"
+
 #include "ExplicitUseGeometryMathTypes.h"		// using UE::Geometry::(math types)
 using namespace UE::Geometry;
 
@@ -30,29 +35,27 @@ using namespace UE::Geometry;
  */
 
 
+const FToolTargetTypeRequirements& UEditNormalsToolBuilder::GetTargetRequirements() const
+{
+	static FToolTargetTypeRequirements TypeRequirements({
+		UMeshDescriptionCommitter::StaticClass(),
+		UMeshDescriptionProvider::StaticClass(),
+		UPrimitiveComponentBackedTarget::StaticClass()
+		});
+	return TypeRequirements;
+}
+
 bool UEditNormalsToolBuilder::CanBuildTool(const FToolBuilderState& SceneState) const
 {
-	return ToolBuilderUtil::CountComponents(SceneState, CanMakeComponentTarget) > 0;
+	return SceneState.TargetManager->CountSelectedAndTargetable(SceneState, GetTargetRequirements()) > 0;
 }
 
 UInteractiveTool* UEditNormalsToolBuilder::BuildTool(const FToolBuilderState& SceneState) const
 {
 	UEditNormalsTool* NewTool = NewObject<UEditNormalsTool>(SceneState.ToolManager);
 
-	TArray<UActorComponent*> Components = ToolBuilderUtil::FindAllComponents(SceneState, CanMakeComponentTarget);
-	check(Components.Num() > 0);
-
-	TArray<TUniquePtr<FPrimitiveComponentTarget>> ComponentTargets;
-	for (UActorComponent* ActorComponent : Components)
-	{
-		auto* MeshComponent = Cast<UPrimitiveComponent>(ActorComponent);
-		if ( MeshComponent )
-		{
-			ComponentTargets.Add(MakeComponentTarget(MeshComponent));
-		}
-	}
-
-	NewTool->SetSelection(MoveTemp(ComponentTargets));
+	TArray<TObjectPtr<UToolTarget>> Targets = SceneState.TargetManager->BuildAllSelectedTargetable(SceneState, GetTargetRequirements());
+	NewTool->SetTargets(MoveTemp(Targets));
 	NewTool->SetWorld(SceneState.World);
 	NewTool->SetAssetAPI(AssetAPI);
 
@@ -96,9 +99,10 @@ void UEditNormalsTool::Setup()
 	UInteractiveTool::Setup();
 
 	// hide input StaticMeshComponent
-	for (TUniquePtr<FPrimitiveComponentTarget>& ComponentTarget : ComponentTargets)
+	for (int32 ComponentIdx = 0, NumTargets = Targets.Num(); ComponentIdx < NumTargets; ComponentIdx++)
 	{
-		ComponentTarget->SetOwnerVisibility(false);
+		IPrimitiveComponentBackedTarget* TargetComponent = TargetComponentInterface(ComponentIdx);
+		TargetComponent->SetOwnerVisibility(false);
 	}
 
 	BasicProperties = NewObject<UEditNormalsToolProperties>(this, TEXT("Mesh Normals Settings"));
@@ -129,7 +133,7 @@ void UEditNormalsTool::Setup()
 void UEditNormalsTool::UpdateNumPreviews()
 {
 	int32 CurrentNumPreview = Previews.Num();
-	int32 TargetNumPreview = ComponentTargets.Num();
+	int32 TargetNumPreview = Targets.Num();
 	if (TargetNumPreview < CurrentNumPreview)
 	{
 		for (int32 PreviewIdx = CurrentNumPreview - 1; PreviewIdx >= TargetNumPreview; PreviewIdx--)
@@ -149,14 +153,14 @@ void UEditNormalsTool::UpdateNumPreviews()
 			OpFactory->ComponentIndex = PreviewIdx;
 			OriginalDynamicMeshes[PreviewIdx] = MakeShared<FDynamicMesh3, ESPMode::ThreadSafe>();
 			FMeshDescriptionToDynamicMesh Converter;
-			Converter.Convert(ComponentTargets[PreviewIdx]->GetMesh(), *OriginalDynamicMeshes[PreviewIdx]);
+			Converter.Convert(TargetMeshProviderInterface(PreviewIdx)->GetMeshDescription(), *OriginalDynamicMeshes[PreviewIdx]);
 
 			UMeshOpPreviewWithBackgroundCompute* Preview = Previews.Add_GetRef(NewObject<UMeshOpPreviewWithBackgroundCompute>(OpFactory, "Preview"));
 			Preview->Setup(this->TargetWorld, OpFactory);
 			Preview->PreviewMesh->SetTangentsMode(EDynamicMeshTangentCalcType::AutoCalculated);
 
 			FComponentMaterialSet MaterialSet;
-			ComponentTargets[PreviewIdx]->GetMaterialSet(MaterialSet);
+			TargetMaterialInterface(PreviewIdx)->GetMaterialSet(MaterialSet);
 			Preview->ConfigureMaterials(MaterialSet.Materials,
 				ToolSetupUtil::GetDefaultWorkingMaterial(GetToolManager())
 			);
@@ -173,9 +177,10 @@ void UEditNormalsTool::Shutdown(EToolShutdownType ShutdownType)
 	AdvancedProperties->SaveProperties(this);
 
 	// Restore (unhide) the source meshes
-	for (TUniquePtr<FPrimitiveComponentTarget>& ComponentTarget : ComponentTargets)
+	for (int32 ComponentIdx = 0, NumTargets = Targets.Num(); ComponentIdx < NumTargets; ComponentIdx++)
 	{
-		ComponentTarget->SetOwnerVisibility(true);
+		IPrimitiveComponentBackedTarget* TargetComponent = TargetComponentInterface(ComponentIdx);
+		TargetComponent->SetOwnerVisibility(true);
 	}
 
 	TArray<FDynamicMeshOpResult> Results;
@@ -205,7 +210,7 @@ TUniquePtr<FDynamicMeshOperator> UEditNormalsOperatorFactory::MakeNewOperator()
 	NormalsOp->NormalCalculationMethod = Tool->BasicProperties->NormalCalculationMethod;
 	NormalsOp->NormalSplitThreshold = Tool->BasicProperties->SharpEdgeAngleThreshold;
 
-	FTransform LocalToWorld = Tool->ComponentTargets[ComponentIndex]->GetWorldTransform();
+	FTransform LocalToWorld = Tool->TargetComponentInterface(ComponentIndex)->GetWorldTransform();
 	NormalsOp->OriginalMesh = Tool->OriginalDynamicMeshes[ComponentIndex];
 	
 	NormalsOp->SetTransform(LocalToWorld);
@@ -273,14 +278,12 @@ void UEditNormalsTool::GenerateAsset(const TArray<FDynamicMeshOpResult>& Results
 {
 	GetToolManager()->BeginUndoTransaction(LOCTEXT("EditNormalsToolTransactionName", "Edit Normals Tool"));
 
-	check(Results.Num() == ComponentTargets.Num());
+	check(Results.Num() == Targets.Num());
 	
-	for (int32 ComponentIdx = 0; ComponentIdx < ComponentTargets.Num(); ComponentIdx++)
+	for (int32 ComponentIdx = 0; ComponentIdx < Targets.Num(); ComponentIdx++)
 	{
-		TUniquePtr<FPrimitiveComponentTarget>& ComponentTarget = ComponentTargets[ComponentIdx];
-
 		// disable auto-generated normals StaticMesh build setting
-		UStaticMeshComponent* StaticMeshComponent = Cast<UStaticMeshComponent>(ComponentTarget->GetOwnerComponent());
+		UStaticMeshComponent* StaticMeshComponent = Cast<UStaticMeshComponent>(TargetComponentInterface(ComponentIdx)->GetOwnerComponent());
 		if (StaticMeshComponent != nullptr)
 		{
 			UStaticMesh* StaticMesh = StaticMeshComponent->GetStaticMesh();
@@ -294,14 +297,14 @@ void UEditNormalsTool::GenerateAsset(const TArray<FDynamicMeshOpResult>& Results
 		}
 
 		check(Results[ComponentIdx].Mesh.Get() != nullptr);
-		ComponentTarget->CommitMesh([&Results, &ComponentIdx, this](const FPrimitiveComponentTarget::FCommitParams& CommitParams)
+		TargetMeshCommitterInterface(ComponentIdx)->CommitMeshDescription([&Results, &ComponentIdx, this](const IMeshDescriptionCommitter::FCommitterParams& CommitParams)
 		{
 			FDynamicMeshToMeshDescription Converter;
 			
-			if (BasicProperties->WillTopologyChange() || !FDynamicMeshToMeshDescription::HaveMatchingElementCounts(Results[ComponentIdx].Mesh.Get(), CommitParams.MeshDescription, false, true))
+			if (BasicProperties->WillTopologyChange() || !FDynamicMeshToMeshDescription::HaveMatchingElementCounts(Results[ComponentIdx].Mesh.Get(), CommitParams.MeshDescriptionOut, false, true))
 			{
 				// full conversion if normal topology changed or faces were inverted
-				Converter.Convert(Results[ComponentIdx].Mesh.Get(), *CommitParams.MeshDescription);
+				Converter.Convert(Results[ComponentIdx].Mesh.Get(), *CommitParams.MeshDescriptionOut);
 			}
 			else
 			{
@@ -309,7 +312,7 @@ void UEditNormalsTool::GenerateAsset(const TArray<FDynamicMeshOpResult>& Results
 				const bool bUpdateNormals = true;
 				const bool bCopyOverlayTangents = false;
 				const bool bCopyOverlayUVs = false;
-				Converter.UpdateAttributes(Results[ComponentIdx].Mesh.Get(), *CommitParams.MeshDescription, bUpdateNormals, bCopyOverlayTangents, bCopyOverlayUVs);
+				Converter.UpdateAttributes(Results[ComponentIdx].Mesh.Get(), *CommitParams.MeshDescriptionOut, bUpdateNormals, bCopyOverlayTangents, bCopyOverlayUVs);
 			}
 		});
 	}
diff --git a/Engine/Plugins/Experimental/MeshModelingToolset/Source/MeshModelingTools/Private/EditPivotTool.cpp b/Engine/Plugins/Experimental/MeshModelingToolset/Source/MeshModelingTools/Private/EditPivotTool.cpp
--- a/Engine/Plugins/Experimental/MeshModelingToolset/Source/MeshModelingTools/Private/EditPivotTool.cpp
+++ b/Engine/Plugins/Experimental/MeshModelingToolset/Source/MeshModelingTools/Private/EditPivotTool.cpp
@@ -22,6 +22,12 @@
 #include "Components/PrimitiveComponent.h"
 #include "Engine/World.h"
 
+#include "TargetInterfaces/MeshDescriptionCommitter.h"
+#include "TargetInterfaces/MeshDescriptionProvider.h"
+#include "TargetInterfaces/PrimitiveComponentBackedTarget.h"
+#include "TargetInterfaces/AssetBackedTarget.h"
+#include "ToolTargetManager.h"
+
 #include "ExplicitUseGeometryMathTypes.h"		// using UE::Geometry::(math types)
 using namespace UE::Geometry;
 
@@ -32,29 +38,28 @@ using namespace UE::Geometry;
  */
 
 
+const FToolTargetTypeRequirements& UEditPivotToolBuilder::GetTargetRequirements() const
+{
+	static FToolTargetTypeRequirements TypeRequirements({
+		UMeshDescriptionCommitter::StaticClass(),
+		UMeshDescriptionProvider::StaticClass(),
+		UPrimitiveComponentBackedTarget::StaticClass(),
+		UAssetBackedTarget::StaticClass()
+	});
+	return TypeRequirements;
+}
+
 bool UEditPivotToolBuilder::CanBuildTool(const FToolBuilderState& SceneState) const
 {
-	return ToolBuilderUtil::CountComponents(SceneState, CanMakeComponentTarget) >= 1;
+	return SceneState.TargetManager->CountSelectedAndTargetable(SceneState, GetTargetRequirements()) >= 1;
 }
 
 UInteractiveTool* UEditPivotToolBuilder::BuildTool(const FToolBuilderState& SceneState) const
 {
 	UEditPivotTool* NewTool = NewObject<UEditPivotTool>(SceneState.ToolManager);
 
-	TArray<UActorComponent*> Components = ToolBuilderUtil::FindAllComponents(SceneState, CanMakeComponentTarget);
-	check(Components.Num() > 0);
-
-	TArray<TUniquePtr<FPrimitiveComponentTarget>> ComponentTargets;
-	for (UActorComponent* ActorComponent : Components)
-	{
-		auto* MeshComponent = Cast<UPrimitiveComponent>(ActorComponent);
-		if ( MeshComponent )
-		{
-			ComponentTargets.Add(MakeComponentTarget(MeshComponent));
-		}
-	}
-
-	NewTool->SetSelection(MoveTemp(ComponentTargets));
+	TArray<TObjectPtr<UToolTarget>> Targets = SceneState.TargetManager->BuildAllSelectedTargetable(SceneState, GetTargetRequirements());
+	NewTool->SetTargets(MoveTemp(Targets));
 	NewTool->SetWorld(SceneState.World, SceneState.GizmoManager);
 
 	return NewTool;
@@ -117,7 +122,7 @@ void UEditPivotTool::Setup()
 	FText AllTheWarnings = LOCTEXT("EditPivotWarning", "WARNING: This Tool will Modify the selected StaticMesh Assets! If you do not wish to modify the original Assets, please make copies in the Content Browser first!");
 
 	// detect and warn about any meshes in selection that correspond to same source data
-	bool bSharesSources = GetMapToFirstComponentsSharingSourceData(MapToFirstOccurrences);
+	bool bSharesSources = GetMapToSharedSourceData(MapToFirstOccurrences);
 	if (bSharesSources)
 	{
 		AllTheWarnings = FText::Format(FTextFormat::FromString("{0}\n\n{1}"), AllTheWarnings, LOCTEXT("EditPivotSharedAssetsWarning", "WARNING: Multiple meshes in your selection use the same source asset!  This is not supported -- each asset can only have one baked pivot."));
@@ -165,12 +170,12 @@ void UEditPivotTool::Precompute()
 	ObjectBounds = FAxisAlignedBox3d::Empty();
 	WorldBounds = FAxisAlignedBox3d::Empty();
 
-	int NumComponents = ComponentTargets.Num();
+	int NumComponents = Targets.Num();
 	if (NumComponents == 1)
 	{
-		Transform = FTransform3d(ComponentTargets[0]->GetWorldTransform());
+		Transform = FTransform3d(TargetComponentInterface(0)->GetWorldTransform());
 
-		FMeshDescription* Mesh = ComponentTargets[0]->GetMesh();
+		FMeshDescription* Mesh = TargetMeshProviderInterface(0)->GetMeshDescription();
 		VertexIteration(Mesh, [&](int32 VertexID, const FVector& Position) {
 			ObjectBounds.Contain(Position);
 			WorldBounds.Contain(Transform.TransformPosition(Position));
@@ -179,11 +184,12 @@ void UEditPivotTool::Precompute()
 	else
 	{
 		Transform = FTransform3d::Identity();
-		for (int k = 0; k < ComponentTargets.Num(); ++k)
+		for (int k = 0; k < NumComponents; ++k)
 		{
-			TUniquePtr<FPrimitiveComponentTarget>& Target = ComponentTargets[k];
-			FTransform3d CurTransform(Target->GetWorldTransform());
-			FMeshDescription* Mesh = Target->GetMesh();
+			IPrimitiveComponentBackedTarget* TargetComponent = TargetComponentInterface(k);
+			IMeshDescriptionProvider* TargetMeshProvider = TargetMeshProviderInterface(k);
+			FTransform3d CurTransform(TargetComponent->GetWorldTransform());
+			FMeshDescription* Mesh = TargetMeshProvider->GetMeshDescription();
 			VertexIteration(Mesh, [&](int32 VertexID, const FVector& Position) {
 				ObjectBounds.Contain(CurTransform.TransformPosition(Position));
 				WorldBounds.Contain(CurTransform.TransformPosition(Position));
@@ -297,9 +303,9 @@ void UEditPivotTool::SetActiveGizmos_Single(bool bLocalRotations)
 	Transformable.TransformProxy = NewObject<UTransformProxy>(this);
 	Transformable.TransformProxy->bRotatePerObject = bLocalRotations;
 
-	for (TUniquePtr<FPrimitiveComponentTarget>& Target : ComponentTargets)
+	for (int32 ComponentIdx = 0; ComponentIdx < Targets.Num(); ComponentIdx++)
 	{
-		Transformable.TransformProxy->AddComponent(Target->GetOwnerComponent());
+		Transformable.TransformProxy->AddComponent(TargetComponentInterface(ComponentIdx)->GetOwnerComponent());
 	}
 	Transformable.TransformGizmo = GizmoManager->CreateCustomTransformGizmo(
 		ETransformGizmoSubElements::StandardTranslateRotate, this
@@ -408,25 +414,26 @@ void UEditPivotTool::UpdateAssets(const FFrame3d& NewPivotWorldFrame)
 	FTransform NewWorldTransform = NewPivotWorldFrame.ToFTransform();
 	FTransform NewWorldInverse = NewWorldTransform.Inverse();
 	TArray<FTransform> OriginalTransforms;
-	for (int32 ComponentIdx = 0; ComponentIdx < ComponentTargets.Num(); ComponentIdx++)
+	for (int32 ComponentIdx = 0; ComponentIdx < Targets.Num(); ComponentIdx++)
 	{
-		OriginalTransforms.Add(ComponentTargets[ComponentIdx]->GetWorldTransform());
+		OriginalTransforms.Add(TargetComponentInterface(ComponentIdx)->GetWorldTransform());
 	}
-	for (int32 ComponentIdx = 0; ComponentIdx < ComponentTargets.Num(); ComponentIdx++)
+	for (int32 ComponentIdx = 0; ComponentIdx < Targets.Num(); ComponentIdx++)
 	{
+		IPrimitiveComponentBackedTarget* TargetComponent = TargetComponentInterface(ComponentIdx);
 		if (MapToFirstOccurrences[ComponentIdx] == ComponentIdx)
 		{
 			FTransform3d ToBake(OriginalTransforms[ComponentIdx] * NewWorldInverse);
 
-			UPrimitiveComponent* Component = ComponentTargets[ComponentIdx]->GetOwnerComponent();
+			UPrimitiveComponent* Component = TargetComponent->GetOwnerComponent();
 			Component->Modify();
 
 			// transform simple collision geometry
 			UE::Geometry::TransformSimpleCollision(Component, ToBake);
 
-			ComponentTargets[ComponentIdx]->CommitMesh([&ToBake](const FPrimitiveComponentTarget::FCommitParams& CommitParams)
+			TargetMeshCommitterInterface(ComponentIdx)->CommitMeshDescription([&ToBake](const IMeshDescriptionCommitter::FCommitterParams& CommitParams)
 			{
-				FMeshDescriptionEditableTriangleMeshAdapter EditableMeshDescAdapter(CommitParams.MeshDescription);
+				FMeshDescriptionEditableTriangleMeshAdapter EditableMeshDescAdapter(CommitParams.MeshDescriptionOut);
 				MeshAdapterTransforms::ApplyTransform(EditableMeshDescAdapter, ToBake);
 			});
 
@@ -434,22 +441,22 @@ void UEditPivotTool::UpdateAssets(const FFrame3d& NewPivotWorldFrame)
 		}
 		else
 		{
-			UPrimitiveComponent* Component = ComponentTargets[ComponentIdx]->GetOwnerComponent();
+			UPrimitiveComponent* Component = TargetComponent->GetOwnerComponent();
 			Component->Modify();
 			// try to invert baked transform
 			FTransform Baked = OriginalTransforms[MapToFirstOccurrences[ComponentIdx]] * NewWorldInverse;
 			Component->SetWorldTransform(Baked.Inverse() * OriginalTransforms[ComponentIdx]);
 		}
-		ComponentTargets[ComponentIdx]->GetOwnerActor()->MarkComponentsRenderStateDirty();
+		TargetComponent->GetOwnerActor()->MarkComponentsRenderStateDirty();
 	}
 
 	// hack to ensure user sees the updated pivot immediately: request re-select of the original selection
 	FSelectedOjectsChangeList NewSelection;
 	NewSelection.ModificationType = ESelectedObjectsModificationType::Replace;
-	for (int OrigMeshIdx = 0; OrigMeshIdx < ComponentTargets.Num(); OrigMeshIdx++)
+	for (int OrigMeshIdx = 0; OrigMeshIdx < Targets.Num(); OrigMeshIdx++)
 	{
-		TUniquePtr<FPrimitiveComponentTarget>& ComponentTarget = ComponentTargets[OrigMeshIdx];
-		NewSelection.Actors.Add(ComponentTarget->GetOwnerActor());
+		IPrimitiveComponentBackedTarget* TargetComponent = TargetComponentInterface(OrigMeshIdx);
+		NewSelection.Actors.Add(TargetComponent->GetOwnerActor());
 	}
 	GetToolManager()->RequestSelectionChange(NewSelection);
 
diff --git a/Engine/Plugins/Experimental/MeshModelingToolset/Source/MeshModelingTools/Private/MirrorTool.cpp b/Engine/Plugins/Experimental/MeshModelingToolset/Source/MeshModelingTools/Private/MirrorTool.cpp
--- a/Engine/Plugins/Experimental/MeshModelingToolset/Source/MeshModelingTools/Private/MirrorTool.cpp
+++ b/Engine/Plugins/Experimental/MeshModelingToolset/Source/MeshModelingTools/Private/MirrorTool.cpp
@@ -14,40 +14,44 @@
 #include "ToolBuilderUtil.h"
 #include "ToolSetupUtil.h"
 
+#include "TargetInterfaces/MaterialProvider.h"
+#include "TargetInterfaces/MeshDescriptionCommitter.h"
+#include "TargetInterfaces/MeshDescriptionProvider.h"
+#include "TargetInterfaces/PrimitiveComponentBackedTarget.h"
+#include "TargetInterfaces/AssetBackedTarget.h"
+#include "ToolTargetManager.h"
+
 #include "ExplicitUseGeometryMathTypes.h"		// using UE::Geometry::(math types)
 using namespace UE::Geometry;
 
 #define LOCTEXT_NAMESPACE "UMirrorTool"
 
-// Local function forward declarations
-void CheckAndDisplayWarnings(const TArray<TUniquePtr<FPrimitiveComponentTarget>>& ComponentTargets, UInteractiveToolManager& ToolsManager);
-
 
 //Tool builder functions
 
+const FToolTargetTypeRequirements& UMirrorToolBuilder::GetTargetRequirements() const
+{
+	static FToolTargetTypeRequirements TypeRequirements({
+		UMaterialProvider::StaticClass(),
+		UMeshDescriptionCommitter::StaticClass(),
+		UMeshDescriptionProvider::StaticClass(),
+		UPrimitiveComponentBackedTarget::StaticClass(),
+		UAssetBackedTarget::StaticClass()
+		});
+	return TypeRequirements;
+}
+
 bool UMirrorToolBuilder::CanBuildTool(const FToolBuilderState& SceneState) const
 {
-	return AssetAPI != nullptr && ToolBuilderUtil::CountComponents(SceneState, CanMakeComponentTarget) > 0;
+	return AssetAPI != nullptr && SceneState.TargetManager->CountSelectedAndTargetable(SceneState, GetTargetRequirements()) > 0;
 }
 
 UInteractiveTool* UMirrorToolBuilder::BuildTool(const FToolBuilderState& SceneState) const
 {
 	UMirrorTool* NewTool = NewObject<UMirrorTool>(SceneState.ToolManager);
 
-	TArray<UActorComponent*> Components = ToolBuilderUtil::FindAllComponents(SceneState, CanMakeComponentTarget);
-	check(Components.Num() > 0);
-
-	TArray<TUniquePtr<FPrimitiveComponentTarget>> ComponentTargets;
-	for (UActorComponent* ActorComponent : Components)
-	{
-		auto* MeshComponent = Cast<UPrimitiveComponent>(ActorComponent);
-		if (MeshComponent)
-		{
-			ComponentTargets.Add(MakeComponentTarget(MeshComponent));
-		}
-	}
-
-	NewTool->SetSelection(MoveTemp(ComponentTargets));
+	TArray<TObjectPtr<UToolTarget>> Targets = SceneState.TargetManager->BuildAllSelectedTargetable(SceneState, GetTargetRequirements());
+	NewTool->SetTargets(MoveTemp(Targets));
 	NewTool->SetWorld(SceneState.World);
 	NewTool->SetAssetAPI(AssetAPI);
 
@@ -68,7 +72,7 @@ TUniquePtr<FDynamicMeshOperator> UMirrorOperatorFactory::MakeNewOperator()
 	MirrorOp->bWeldAlongPlane = MirrorTool->Settings->bWeldVerticesOnMirrorPlane;
 	MirrorOp->bAllowBowtieVertexCreation = MirrorTool->Settings->bAllowBowtieVertexCreation;
 
-	FTransform LocalToWorld = MirrorTool->ComponentTargets[ComponentIndex]->GetWorldTransform();
+	FTransform LocalToWorld = MirrorTool->TargetComponentInterface(ComponentIndex)->GetWorldTransform();
 	MirrorOp->SetTransform(LocalToWorld);
 
 	// We also need WorldToLocal. Threshold the LocalToWorld scaling transform so we can get the inverse.
@@ -138,9 +142,9 @@ void UMirrorTool::OnPropertyModified(UObject* PropertySet, FProperty* Property)
 	// Editing the "show preview" option changes whether we need to be displaying the preview or the original mesh.
 	if (Property && (Property->GetFName() == GET_MEMBER_NAME_CHECKED(UMirrorToolProperties, bShowPreview)))
 	{
-		for (TUniquePtr<FPrimitiveComponentTarget>& ComponentTarget : ComponentTargets)
+		for (int32 ComponentIdx = 0; ComponentIdx < Targets.Num(); ComponentIdx++)
 		{
-			ComponentTarget->SetOwnerVisibility(!Settings->bShowPreview);
+			TargetComponentInterface(ComponentIdx)->SetOwnerVisibility(!Settings->bShowPreview);
 		}
 		for (UMeshOpPreviewWithBackgroundCompute* Preview : Previews)
 		{
@@ -202,17 +206,18 @@ void UMirrorTool::Setup()
 	ToolActions->Initialize(this);
 	AddToolPropertySource(ToolActions);
 
-	CheckAndDisplayWarnings(ComponentTargets, *GetToolManager());
+	CheckAndDisplayWarnings();
 
 	// Fill in the MeshesToMirror array with suitably converted meshes.
-	for (int i = 0; i < ComponentTargets.Num(); i++)
+	for (int i = 0; i < Targets.Num(); i++)
 	{
-		TUniquePtr<FPrimitiveComponentTarget>& ComponentTarget = ComponentTargets[i];
+		IPrimitiveComponentBackedTarget* TargetComponent = TargetComponentInterface(i);
+		IMeshDescriptionProvider* TargetMeshProvider = TargetMeshProviderInterface(i);
 
 		// Convert into dynamic mesh
 		TSharedPtr<FDynamicMesh3, ESPMode::ThreadSafe> DynamicMesh = MakeShared<FDynamicMesh3, ESPMode::ThreadSafe>();
 		FMeshDescriptionToDynamicMesh Converter;
-		Converter.Convert(ComponentTarget->GetMesh(), *DynamicMesh);
+		Converter.Convert(TargetMeshProvider->GetMeshDescription(), *DynamicMesh);
 
 		// Wrap the dynamic mesh in a replacement change target
 		UDynamicMeshReplacementChangeTarget* WrappedTarget = MeshesToMirror.Add_GetRef(NewObject<UDynamicMeshReplacementChangeTarget>());
@@ -223,9 +228,9 @@ void UMirrorTool::Setup()
 	}
 
 	// Set the visibility of the StaticMeshComponents depending on whether we are showing them or the preview.
-	for (TUniquePtr<FPrimitiveComponentTarget>& ComponentTarget : ComponentTargets)
+	for (int32 ComponentIdx = 0; ComponentIdx < Targets.Num(); ComponentIdx++)
 	{
-		ComponentTarget->SetOwnerVisibility(!Settings->bShowPreview);
+		TargetComponentInterface(ComponentIdx)->SetOwnerVisibility(!Settings->bShowPreview);
 	}
 
 	// Initialize the PreviewMesh and BackgroundCompute objects
@@ -233,10 +238,10 @@ void UMirrorTool::Setup()
 
 	// Update the bounding box of the meshes.
 	CombinedBounds.Init();
-	for (TUniquePtr<FPrimitiveComponentTarget>& ComponentTarget : ComponentTargets)
+	for (int32 ComponentIdx = 0; ComponentIdx < Targets.Num(); ComponentIdx++)
 	{
 		FVector ComponentOrigin, ComponentExtents;
-		ComponentTarget->GetOwnerActor()->GetActorBounds(false, ComponentOrigin, ComponentExtents);
+		TargetComponentInterface(ComponentIdx)->GetOwnerActor()->GetActorBounds(false, ComponentOrigin, ComponentExtents);
 		CombinedBounds += FBox::BuildAABB(ComponentOrigin, ComponentExtents);
 	}
 
@@ -272,9 +277,10 @@ void UMirrorTool::Setup()
 	};
 	// Also include the original components in the ctrl+click hit testing even though we made them 
 	// invisible, since we want to be able to reposition the plane onto the original mesh.
-	for (const TUniquePtr<FPrimitiveComponentTarget>& Target : ComponentTargets)
+	for (int32 ComponentIdx = 0; ComponentIdx < Targets.Num(); ComponentIdx++)
 	{
-		PlaneMechanic->SetPlaneCtrlClickBehaviorTarget->InvisibleComponentsToHitTest.Add(Target->GetOwnerComponent());
+		IPrimitiveComponentBackedTarget* TargetComponent = TargetComponentInterface(ComponentIdx);
+		PlaneMechanic->SetPlaneCtrlClickBehaviorTarget->InvisibleComponentsToHitTest.Add(TargetComponent->GetOwnerComponent());
 	}
 
 
@@ -307,38 +313,25 @@ void UMirrorTool::SetupPreviews()
 		Preview->PreviewMesh->SetTangentsMode(EDynamicMeshTangentCalcType::AutoCalculated);
 
 		FComponentMaterialSet MaterialSet;
-		ComponentTargets[PreviewIndex]->GetMaterialSet(MaterialSet);
+		TargetMaterialInterface(PreviewIndex)->GetMaterialSet(MaterialSet);
 		Preview->ConfigureMaterials(MaterialSet.Materials, ToolSetupUtil::GetDefaultWorkingMaterial(GetToolManager()));
 
 		// Set initial preview to unprocessed mesh, so that things don't disappear initially
 		Preview->PreviewMesh->UpdatePreview(MeshesToMirror[PreviewIndex]->GetMesh().Get());
-		Preview->PreviewMesh->SetTransform(ComponentTargets[PreviewIndex]->GetWorldTransform());
+		Preview->PreviewMesh->SetTransform(TargetComponentInterface(PreviewIndex)->GetWorldTransform());
 		Preview->SetVisibility(Settings->bShowPreview);
 	}
 }
 
-void CheckAndDisplayWarnings(const TArray<TUniquePtr<FPrimitiveComponentTarget>>& ComponentTargets, 
-	UInteractiveToolManager& ToolsManager)
+void UMirrorTool::CheckAndDisplayWarnings()
 {
 	// We can have more than one warning, which makes this a bit more work.
 	FText SameSourceWarning;
 	FText NonUniformScaleWarning;
 
 	// See if any of the selected components have the same source.
-	bool bAnyHaveSameSource = false;
-	for (int32 FirstComponentIndex = 0; !bAnyHaveSameSource && FirstComponentIndex < ComponentTargets.Num(); FirstComponentIndex++)
-	{
-		FPrimitiveComponentTarget* ComponentTarget = ComponentTargets[FirstComponentIndex].Get();
-		for (int32 SecondComponentIndex = FirstComponentIndex + 1; SecondComponentIndex < ComponentTargets.Num(); SecondComponentIndex++)
-		{
-			if (ComponentTarget->HasSameSourceData(*ComponentTargets[SecondComponentIndex]))
-			{
-				bAnyHaveSameSource = true;
-				break;
-			}
-		}
-	}
-
+	TArray<int32> MapToFirstOccurrences;
+	bool bAnyHaveSameSource = GetMapToSharedSourceData(MapToFirstOccurrences);
 	if (bAnyHaveSameSource)
 	{
 		SameSourceWarning = LOCTEXT("MirrorMultipleAssetsWithSameSource", "WARNING: Multiple meshes in your selection use the same source asset! Only the \"Create New Assets\" save mode is supported.");
@@ -348,13 +341,13 @@ void CheckAndDisplayWarnings(const TArray<TUniquePtr<FPrimitiveComponentTarget>>
 	}
 
 	// See if any of the selected components have a nonuniform scaling transform.
-	FPrimitiveComponentTarget* NonUniformScalingTarget = nullptr;
-	for (int32 i = 0; i < ComponentTargets.Num(); ++i)
+	IPrimitiveComponentBackedTarget* NonUniformScalingTarget = nullptr;
+	for (int32 i = 0; i < Targets.Num(); ++i)
 	{
-		FVector Scaling = ComponentTargets[i]->GetWorldTransform().GetScale3D();
+		FVector Scaling = TargetComponentInterface(i)->GetWorldTransform().GetScale3D();
 		if (Scaling.X != Scaling.Y || Scaling.Y != Scaling.Z)
 		{
-			NonUniformScalingTarget = ComponentTargets[i].Get();
+			NonUniformScalingTarget = TargetComponentInterface(i);
 			break;
 		}
 	}
@@ -369,16 +362,16 @@ void CheckAndDisplayWarnings(const TArray<TUniquePtr<FPrimitiveComponentTarget>>
 	if (bAnyHaveSameSource && NonUniformScalingTarget)
 	{
 		// Concatenates the two warnings with an extra line in between.
-		ToolsManager.DisplayMessage(FText::Format(LOCTEXT("CombinedWarnings", "{0}\n\n{1}"),
+		GetToolManager()->DisplayMessage(FText::Format(LOCTEXT("CombinedWarnings", "{0}\n\n{1}"),
 			SameSourceWarning, NonUniformScaleWarning), EToolMessageLevel::UserWarning);
 	}
 	else if (bAnyHaveSameSource)
 	{
-		ToolsManager.DisplayMessage(SameSourceWarning, EToolMessageLevel::UserWarning);
+		GetToolManager()->DisplayMessage(SameSourceWarning, EToolMessageLevel::UserWarning);
 	}
 	else if (NonUniformScalingTarget)
 	{
-		ToolsManager.DisplayMessage(NonUniformScaleWarning, EToolMessageLevel::UserWarning);
+		GetToolManager()->DisplayMessage(NonUniformScaleWarning, EToolMessageLevel::UserWarning);
 	}
 }
 
@@ -389,9 +382,9 @@ void UMirrorTool::Shutdown(EToolShutdownType ShutdownType)
 	PlaneMechanic->Shutdown();
 
 	// Restore (unhide) the source meshes
-	for (TUniquePtr<FPrimitiveComponentTarget>& ComponentTarget : ComponentTargets)
+	for (int32 ComponentIdx = 0; ComponentIdx < Targets.Num(); ComponentIdx++)
 	{
-		ComponentTarget->SetOwnerVisibility(true);
+		TargetComponentInterface(ComponentIdx)->SetOwnerVisibility(true);
 	}
 
 	// Swap in results, if appropriate
@@ -456,6 +449,7 @@ void UMirrorTool::GenerateAsset(const TArray<FDynamicMeshOpResult>& Results)
 	NewSelection.ModificationType = ESelectedObjectsModificationType::Replace;
 	for (int OrigMeshIdx = 0; OrigMeshIdx < NumSourceMeshes; OrigMeshIdx++)
 	{
+		IPrimitiveComponentBackedTarget* TargetComponent = TargetComponentInterface(OrigMeshIdx);
 		FDynamicMesh3* Mesh = Results[OrigMeshIdx].Mesh.Get();
 		check(Mesh != nullptr);
 
@@ -463,28 +457,28 @@ void UMirrorTool::GenerateAsset(const TArray<FDynamicMeshOpResult>& Results)
 		{
 			if (bWantToDestroy)
 			{
-				ComponentTargets[OrigMeshIdx]->GetOwnerComponent()->DestroyComponent();
+				TargetComponent->GetOwnerComponent()->DestroyComponent();
 			}
 			continue;
 		}
 		else if (Settings->SaveMode == EMirrorSaveMode::UpdateAssets)
 		{
-			NewSelection.Actors.Add(ComponentTargets[OrigMeshIdx]->GetOwnerActor());
+			NewSelection.Actors.Add(TargetComponent->GetOwnerActor());
 
-			ComponentTargets[OrigMeshIdx]->CommitMesh([&Mesh](const FPrimitiveComponentTarget::FCommitParams& CommitParams)
+			TargetMeshCommitterInterface(OrigMeshIdx)->CommitMeshDescription([&Mesh](const IMeshDescriptionCommitter::FCommitterParams& CommitParams)
 				{
 					FDynamicMeshToMeshDescription Converter;
-					Converter.Convert(Mesh, *CommitParams.MeshDescription);
+					Converter.Convert(Mesh, *CommitParams.MeshDescriptionOut);
 				});
 		}
 		else
 		{
 			// Build array of materials from the original.
 			TArray<UMaterialInterface*> Materials;
-			TUniquePtr<FPrimitiveComponentTarget>& ComponentTarget = ComponentTargets[OrigMeshIdx];
-			for (int MaterialIdx = 0, NumMaterials = ComponentTarget->GetNumMaterials(); MaterialIdx < NumMaterials; MaterialIdx++)
+			IMaterialProvider* TargetMaterial = TargetMaterialInterface(OrigMeshIdx);
+			for (int MaterialIdx = 0, NumMaterials = TargetMaterial->GetNumMaterials(); MaterialIdx < NumMaterials; MaterialIdx++)
 			{
-				Materials.Add(ComponentTarget->GetMaterial(MaterialIdx));
+				Materials.Add(TargetMaterial->GetMaterial(MaterialIdx));
 			}
 
 			// Create the new actor
@@ -496,7 +490,7 @@ void UMirrorTool::GenerateAsset(const TArray<FDynamicMeshOpResult>& Results)
 			}
 
 			// Remove the original actor
-			ComponentTarget->GetOwnerComponent()->DestroyComponent();
+			TargetComponent->GetOwnerComponent()->DestroyComponent();
 		}
 	}
 
diff --git a/Engine/Plugins/Experimental/MeshModelingToolset/Source/MeshModelingTools/Private/Physics/PhysicsInspectorTool.cpp b/Engine/Plugins/Experimental/MeshModelingToolset/Source/MeshModelingTools/Private/Physics/PhysicsInspectorTool.cpp
--- a/Engine/Plugins/Experimental/MeshModelingToolset/Source/MeshModelingTools/Private/Physics/PhysicsInspectorTool.cpp
+++ b/Engine/Plugins/Experimental/MeshModelingToolset/Source/MeshModelingTools/Private/Physics/PhysicsInspectorTool.cpp
@@ -15,16 +15,28 @@
 #include "Engine/Classes/PhysicsEngine/BodySetup.h"
 #include "Engine/Classes/PhysicsEngine/AggregateGeom.h"
 
+#include "TargetInterfaces/PrimitiveComponentBackedTarget.h"
+#include "TargetInterfaces/StaticMeshBackedTarget.h"
+#include "ToolTargetManager.h"
+
 #include "ExplicitUseGeometryMathTypes.h"		// using UE::Geometry::(math types)
 using namespace UE::Geometry;
 
 #define LOCTEXT_NAMESPACE "UPhysicsInspectorTool"
 
+
+const FToolTargetTypeRequirements& UPhysicsInspectorToolBuilder::GetTargetRequirements() const
+{
+	static FToolTargetTypeRequirements TypeRequirements({
+		UPrimitiveComponentBackedTarget::StaticClass(),
+		UStaticMeshBackedTarget::StaticClass()
+		});
+	return TypeRequirements;
+}
+
 bool UPhysicsInspectorToolBuilder::CanBuildTool(const FToolBuilderState& SceneState) const
 {
-	int32 NumStaticMeshes = ToolBuilderUtil::CountComponents(SceneState, [&](UActorComponent* Comp) { return Cast<UStaticMeshComponent>(Comp) != nullptr; });
-	int32 NumComponentTargets = ToolBuilderUtil::CountComponents(SceneState, CanMakeComponentTarget);
-	return (NumStaticMeshes > 0 && NumStaticMeshes == NumComponentTargets);
+	return SceneState.TargetManager->CountSelectedAndTargetable(SceneState, GetTargetRequirements()) > 0;
 }
 
 
@@ -32,21 +44,9 @@ UInteractiveTool* UPhysicsInspectorToolBuilder::BuildTool(const FToolBuilderStat
 {
 	UPhysicsInspectorTool* NewTool = NewObject<UPhysicsInspectorTool>(SceneState.ToolManager);
 
-	TArray<UActorComponent*> ValidComponents = ToolBuilderUtil::FindAllComponents(SceneState,
-		[&](UActorComponent* Comp) { return Cast<UStaticMeshComponent>(Comp) != nullptr; });
-	check(ValidComponents.Num() > 0);
-
-	TArray<TUniquePtr<FPrimitiveComponentTarget>> ComponentTargets;
-	for (UActorComponent* ActorComponent : ValidComponents)
-	{
-		UStaticMeshComponent* MeshComponent = Cast<UStaticMeshComponent>(ActorComponent);
-		if (MeshComponent)
-		{
-			ComponentTargets.Add(MakeComponentTarget(MeshComponent));
-		}
-	}
+	TArray<TObjectPtr<UToolTarget>> Targets = SceneState.TargetManager->BuildAllSelectedTargetable(SceneState, GetTargetRequirements());
+	NewTool->SetTargets(MoveTemp(Targets));
 
-	NewTool->SetSelection(MoveTemp(ComponentTargets));
 	return NewTool;
 }
 
@@ -62,9 +62,10 @@ void UPhysicsInspectorTool::Setup()
 	VizSettings->WatchProperty(VizSettings->Color, [this](FColor NewValue) { bVisualizationDirty = true; });
 	VizSettings->WatchProperty(VizSettings->bShowHidden, [this](bool bNewValue) { bVisualizationDirty = true; });
 
-	for (TUniquePtr<FPrimitiveComponentTarget>& ComponentTarget : ComponentTargets)
+	for (int32 ComponentIdx = 0; ComponentIdx < Targets.Num(); ComponentIdx++)
 	{
-		const UStaticMeshComponent* Component = CastChecked<UStaticMeshComponent>(ComponentTarget->GetOwnerComponent());
+		IPrimitiveComponentBackedTarget* TargetComponent = TargetComponentInterface(ComponentIdx);
+		const UStaticMeshComponent* Component = CastChecked<UStaticMeshComponent>(TargetComponent->GetOwnerComponent());
 		const UStaticMesh* StaticMesh = Component->GetStaticMesh();
 		if (ensure(StaticMesh && StaticMesh->GetBodySetup()))
 		{
@@ -76,10 +77,10 @@ void UPhysicsInspectorTool::Setup()
 			PhysicsInfos.Add(PhysicsData);
 
 			UPreviewGeometry* PreviewGeom = NewObject<UPreviewGeometry>(this);
-			FTransform TargetTransform = ComponentTarget->GetWorldTransform();
+			FTransform TargetTransform = TargetComponent->GetWorldTransform();
 			PhysicsData->ExternalScale3D = TargetTransform.GetScale3D();
 			TargetTransform.SetScale3D(FVector::OneVector);
-			PreviewGeom->CreateInWorld(ComponentTarget->GetOwnerActor()->GetWorld(), TargetTransform);
+			PreviewGeom->CreateInWorld(TargetComponent->GetOwnerActor()->GetWorld(), TargetTransform);
 			PreviewElements.Add(PreviewGeom);
 
 			InitializeGeometry(*PhysicsData, PreviewGeom);
diff --git a/Engine/Plugins/Experimental/MeshModelingToolset/Source/MeshModelingTools/Private/Physics/SetCollisionGeometryTool.cpp b/Engine/Plugins/Experimental/MeshModelingToolset/Source/MeshModelingTools/Private/Physics/SetCollisionGeometryTool.cpp
--- a/Engine/Plugins/Experimental/MeshModelingToolset/Source/MeshModelingTools/Private/Physics/SetCollisionGeometryTool.cpp
+++ b/Engine/Plugins/Experimental/MeshModelingToolset/Source/MeshModelingTools/Private/Physics/SetCollisionGeometryTool.cpp
@@ -24,32 +24,38 @@
 
 #include "Async/ParallelFor.h"
 
+#include "TargetInterfaces/MeshDescriptionProvider.h"
+#include "TargetInterfaces/PrimitiveComponentBackedTarget.h"
+#include "TargetInterfaces/StaticMeshBackedTarget.h"
+#include "ToolTargetManager.h"
+
 #include "ExplicitUseGeometryMathTypes.h"		// using UE::Geometry::(math types)
 using namespace UE::Geometry;
 
 #define LOCTEXT_NAMESPACE "USetCollisionGeometryTool"
 
+
+const FToolTargetTypeRequirements& USetCollisionGeometryToolBuilder::GetTargetRequirements() const
+{
+	static FToolTargetTypeRequirements TypeRequirements({
+		UMeshDescriptionProvider::StaticClass(),
+		UPrimitiveComponentBackedTarget::StaticClass()
+		});
+	return TypeRequirements;
+}
+
 bool USetCollisionGeometryToolBuilder::CanBuildTool(const FToolBuilderState& SceneState) const
 {
-	TArray<UActorComponent*> Components = ToolBuilderUtil::FindAllComponents(SceneState, CanMakeComponentTarget);
-	return (Components.Num() > 0 && Cast<UStaticMeshComponent>(Components[Components.Num() - 1]) != nullptr);
+	TArray<TObjectPtr<UToolTarget>> Targets = SceneState.TargetManager->BuildAllSelectedTargetable(SceneState, GetTargetRequirements());
+	return (Targets.Num() > 0 && Cast<IStaticMeshBackedTarget>(Targets[Targets.Num() - 1]) != nullptr);
 }
 
 
 UInteractiveTool* USetCollisionGeometryToolBuilder::BuildTool(const FToolBuilderState& SceneState) const
 {
 	USetCollisionGeometryTool* NewTool = NewObject<USetCollisionGeometryTool>(SceneState.ToolManager);
-	TArray<UActorComponent*> Components = ToolBuilderUtil::FindAllComponents(SceneState, CanMakeComponentTarget);
-	TArray<TUniquePtr<FPrimitiveComponentTarget>> ComponentTargets;
-	for (UActorComponent* ActorComponent : Components)
-	{
-		auto* MeshComponent = Cast<UPrimitiveComponent>(ActorComponent);
-		if (MeshComponent)
-		{
-			ComponentTargets.Add(MakeComponentTarget(MeshComponent));
-		}
-	}
-	NewTool->SetSelection(MoveTemp(ComponentTargets));
+	TArray<TObjectPtr<UToolTarget>> Targets = SceneState.TargetManager->BuildAllSelectedTargetable(SceneState, GetTargetRequirements());
+	NewTool->SetTargets(MoveTemp(Targets));
 	return NewTool;
 }
 
@@ -58,24 +64,24 @@ void USetCollisionGeometryTool::Setup()
 {
 	UInteractiveTool::Setup();
 
-	// if we have one selection, use it as the source, otherwise use all but hte last selected mesh
-	bSourcesHidden = (ComponentTargets.Num() > 1);
-	if (ComponentTargets.Num() == 1)
+	// if we have one selection, use it as the source, otherwise use all but the last selected mesh
+	bSourcesHidden = (Targets.Num() > 1);
+	if (Targets.Num() == 1)
 	{
 		SourceObjectIndices.Add(0);
 	}
 	else
 	{
-		for (int32 k = 0; k < ComponentTargets.Num() -1; ++k)
+		for (int32 k = 0; k < Targets.Num() -1; ++k)
 		{
 			SourceObjectIndices.Add(k);
-			ComponentTargets[k]->SetOwnerVisibility(false);
+			TargetComponentInterface(k)->SetOwnerVisibility(false);
 		}
 	}
 
 	bInputMeshesValid = true;
 
-	TUniquePtr<FPrimitiveComponentTarget>& CollisionTarget = ComponentTargets[ComponentTargets.Num() - 1];
+	IPrimitiveComponentBackedTarget* CollisionTarget = TargetComponentInterface(Targets.Num() - 1);
 	PreviewGeom = NewObject<UPreviewGeometry>(this);
 	FTransform PreviewTransform = CollisionTarget->GetWorldTransform();
 	OrigTargetTransform = PreviewTransform;
@@ -141,7 +147,7 @@ void USetCollisionGeometryTool::Shutdown(EToolShutdownType ShutdownType)
 	{
 		for (int32 k : SourceObjectIndices)
 		{
-			ComponentTargets[k]->SetOwnerVisibility(true);
+			TargetComponentInterface(k)->SetOwnerVisibility(true);
 		}
 	}
 
@@ -154,7 +160,7 @@ void USetCollisionGeometryTool::Shutdown(EToolShutdownType ShutdownType)
 
 		// code below derived from FStaticMeshEditor::DuplicateSelectedPrims(), FStaticMeshEditor::OnCollisionSphere(), and GeomFitUtils.cpp::GenerateSphylAsSimpleCollision()
 
-		TUniquePtr<FPrimitiveComponentTarget>& CollisionTarget = ComponentTargets[ComponentTargets.Num() - 1];
+		IPrimitiveComponentBackedTarget* CollisionTarget = TargetComponentInterface(Targets.Num() - 1);
 		UStaticMeshComponent* StaticMeshComponent = CastChecked<UStaticMeshComponent>(CollisionTarget->GetOwnerComponent());
 		UStaticMesh* StaticMesh = StaticMeshComponent->GetStaticMesh();
 		UBodySetup* BodySetup = StaticMesh->GetBodySetup();
@@ -407,7 +413,7 @@ TArray<const T*> MakeRawPointerList(const TArray<TSharedPtr<T, ESPMode::ThreadSa
 
 void USetCollisionGeometryTool::PrecomputeInputMeshes()
 {
-	TUniquePtr<FPrimitiveComponentTarget>& CollisionTarget = ComponentTargets[ComponentTargets.Num()-1];
+	IPrimitiveComponentBackedTarget* CollisionTarget = TargetComponentInterface(Targets.Num()-1);
 	FTransform3d TargetTransform(CollisionTarget->GetWorldTransform());
 	FTransform3d TargetTransformInv = TargetTransform.Inverse();
 
@@ -419,10 +425,10 @@ void USetCollisionGeometryTool::PrecomputeInputMeshes()
 		Converter.bCalculateMaps = false;
 		Converter.bDisableAttributes = true;
 		FDynamicMesh3 SourceMesh;
-		Converter.Convert(ComponentTargets[k]->GetMesh(), SourceMesh);
+		Converter.Convert(TargetMeshProviderInterface(k)->GetMeshDescription(), SourceMesh);
 		if (Settings->bUseWorldSpace)
 		{
-			FTransform3d ToWorld(ComponentTargets[k]->GetWorldTransform());
+			FTransform3d ToWorld(TargetComponentInterface(k)->GetWorldTransform());
 			MeshTransforms::ApplyTransform(SourceMesh, ToWorld);
 			MeshTransforms::ApplyTransform(SourceMesh, TargetTransformInv);
 		}
diff --git a/Engine/Plugins/Experimental/MeshModelingToolset/Source/MeshModelingTools/Private/PlaneCutTool.cpp b/Engine/Plugins/Experimental/MeshModelingToolset/Source/MeshModelingTools/Private/PlaneCutTool.cpp
--- a/Engine/Plugins/Experimental/MeshModelingToolset/Source/MeshModelingTools/Private/PlaneCutTool.cpp
+++ b/Engine/Plugins/Experimental/MeshModelingToolset/Source/MeshModelingTools/Private/PlaneCutTool.cpp
@@ -30,6 +30,13 @@
 
 #include "Misc/MessageDialog.h"
 
+#include "TargetInterfaces/MaterialProvider.h"
+#include "TargetInterfaces/MeshDescriptionCommitter.h"
+#include "TargetInterfaces/MeshDescriptionProvider.h"
+#include "TargetInterfaces/PrimitiveComponentBackedTarget.h"
+#include "TargetInterfaces/AssetBackedTarget.h"
+#include "ToolTargetManager.h"
+
 #include "ExplicitUseGeometryMathTypes.h"		// using UE::Geometry::(math types)
 using namespace UE::Geometry;
 
@@ -39,30 +46,29 @@ using namespace UE::Geometry;
  * ToolBuilder
  */
 
+const FToolTargetTypeRequirements& UPlaneCutToolBuilder::GetTargetRequirements() const
+{
+	static FToolTargetTypeRequirements TypeRequirements({
+		UMaterialProvider::StaticClass(),
+		UMeshDescriptionCommitter::StaticClass(),
+		UMeshDescriptionProvider::StaticClass(),
+		UPrimitiveComponentBackedTarget::StaticClass(),
+		UAssetBackedTarget::StaticClass()
+		});
+	return TypeRequirements;
+}
 
 bool UPlaneCutToolBuilder::CanBuildTool(const FToolBuilderState& SceneState) const
 {
-	return AssetAPI != nullptr && ToolBuilderUtil::CountComponents(SceneState, CanMakeComponentTarget) > 0;
+	return AssetAPI != nullptr && SceneState.TargetManager->CountSelectedAndTargetable(SceneState, GetTargetRequirements()) > 0;
 }
 
 UInteractiveTool* UPlaneCutToolBuilder::BuildTool(const FToolBuilderState& SceneState) const
 {
 	UPlaneCutTool* NewTool = NewObject<UPlaneCutTool>(SceneState.ToolManager);
 
-	TArray<UActorComponent*> Components = ToolBuilderUtil::FindAllComponents(SceneState, CanMakeComponentTarget);
-	check(Components.Num() > 0);
-
-	TArray<TUniquePtr<FPrimitiveComponentTarget>> ComponentTargets;
-	for (UActorComponent* ActorComponent : Components)
-	{
-		auto* MeshComponent = Cast<UPrimitiveComponent>(ActorComponent);
-		if ( MeshComponent )
-		{
-			ComponentTargets.Add(MakeComponentTarget(MeshComponent));
-		}
-	}
-
-	NewTool->SetSelection(MoveTemp(ComponentTargets));
+	TArray<TObjectPtr<UToolTarget>> Targets = SceneState.TargetManager->BuildAllSelectedTargetable(SceneState, GetTargetRequirements());
+	NewTool->SetTargets(MoveTemp(Targets));
 	NewTool->SetWorld(SceneState.World);
 	NewTool->SetAssetAPI(AssetAPI);
 
@@ -106,26 +112,15 @@ void UPlaneCutTool::Setup()
 	AddInputBehavior(SnapToggleBehavior);
 
 	// hide input StaticMeshComponents
-	for (TUniquePtr<FPrimitiveComponentTarget>& ComponentTarget : ComponentTargets)
+	for (int32 ComponentIdx = 0; ComponentIdx < Targets.Num(); ComponentIdx++)
 	{
-		ComponentTarget->SetOwnerVisibility(false);
+		TargetComponentInterface(ComponentIdx)->SetOwnerVisibility(false);
 	}
 
-	bool bAnyHasSameSource = false;
-	for (int32 ComponentIdx = 0; !bAnyHasSameSource && ComponentIdx < ComponentTargets.Num(); ComponentIdx++)
-	{
-		FPrimitiveComponentTarget* ComponentTarget = ComponentTargets[ComponentIdx].Get();
-		for (int32 VsIdx = ComponentIdx + 1; VsIdx < ComponentTargets.Num(); VsIdx++)
-		{
-			if (ComponentTarget->HasSameSourceData(*ComponentTargets[VsIdx]))
-			{
-				bAnyHasSameSource = true;
-				break;
-			}
-		}
-	}
+	TArray<int32> MapToFirstOccurrences;
+	bool bAnyHaveSameSource = GetMapToSharedSourceData(MapToFirstOccurrences);
 
-	if (bAnyHasSameSource)
+	if (bAnyHaveSameSource)
 	{
 		GetToolManager()->DisplayMessage(
 			LOCTEXT("PlaneCutMultipleAssetWithSameSource", "WARNING: Multiple meshes in your selection use the same source asset!  Plane cuts apply to the source asset, and this tool will not duplicate assets for you, so the tool typically cannot give a correct result in this case.  Please consider exiting the tool and duplicating the source assets."),
@@ -133,12 +128,12 @@ void UPlaneCutTool::Setup()
 	}
 
 	// Convert input mesh descriptions to dynamic mesh
-	for (int Idx = 0; Idx < ComponentTargets.Num(); Idx++)
+	for (int Idx = 0; Idx < Targets.Num(); Idx++)
 	{
-		TUniquePtr<FPrimitiveComponentTarget>& ComponentTarget = ComponentTargets[Idx];
+		IMeshDescriptionProvider* TargetMesh = TargetMeshProviderInterface(Idx);
 		FDynamicMesh3* OriginalDynamicMesh = new FDynamicMesh3;
 		FMeshDescriptionToDynamicMesh Converter;
-		Converter.Convert(ComponentTarget->GetMesh(), *OriginalDynamicMesh);
+		Converter.Convert(TargetMesh->GetMeshDescription(), *OriginalDynamicMesh);
 		OriginalDynamicMesh->EnableAttributes();
 		TDynamicMeshScalarTriangleAttribute<int>* SubObjectIDs = new TDynamicMeshScalarTriangleAttribute<int>(OriginalDynamicMesh);
 		SubObjectIDs->Initialize(0);
@@ -160,9 +155,9 @@ void UPlaneCutTool::Setup()
 
 	// Include the original components even though we made them invisible, since we want
 	// to be able to reposition the plane onto the original mesh.
-	for (const TUniquePtr<FPrimitiveComponentTarget>& Target : ComponentTargets)
+	for (int Idx = 0; Idx < Targets.Num(); Idx++)
 	{
-		SetPlaneAction->InvisibleComponentsToHitTest.Add(Target->GetOwnerComponent());
+		SetPlaneAction->InvisibleComponentsToHitTest.Add(TargetComponentInterface(Idx)->GetOwnerComponent());
 	}
 
 	SetPlaneAction->OnClickedPositionFunc = [this](const FHitResult& Hit)
@@ -202,10 +197,10 @@ void UPlaneCutTool::Setup()
 
 	// set initial cut plane (also attaches gizmo/proxy)
 	FBox CombinedBounds; CombinedBounds.Init();
-	for (TUniquePtr<FPrimitiveComponentTarget>& ComponentTarget : ComponentTargets)
+	for (int Idx = 0; Idx < Targets.Num(); Idx++)
 	{
 		FVector ComponentOrigin, ComponentExtents;
-		ComponentTarget->GetOwnerActor()->GetActorBounds(false, ComponentOrigin, ComponentExtents);
+		TargetComponentInterface(Idx)->GetOwnerActor()->GetActorBounds(false, ComponentOrigin, ComponentExtents);
 		CombinedBounds += FBox::BuildAABB(ComponentOrigin, ComponentExtents);
 	}
 	SetCutPlaneFromWorldPos(CombinedBounds.GetCenter(), FVector::UpVector, true);
@@ -254,14 +249,14 @@ void UPlaneCutTool::SetupPreviews()
 		Preview->PreviewMesh->SetTangentsMode(EDynamicMeshTangentCalcType::AutoCalculated);
 
 		FComponentMaterialSet MaterialSet;
-		ComponentTargets[PreviewIdx]->GetMaterialSet(MaterialSet);
+		TargetMaterialInterface(PreviewIdx)->GetMaterialSet(MaterialSet);
 		Preview->ConfigureMaterials(MaterialSet.Materials,
 			ToolSetupUtil::GetDefaultWorkingMaterial(GetToolManager())
 		);
 
 		// set initial preview to un-processed mesh, so stuff doesn't just disappear if the first cut takes a while
 		Preview->PreviewMesh->UpdatePreview(MeshesToCut[PreviewIdx]->GetMesh().Get());
-		Preview->PreviewMesh->SetTransform(ComponentTargets[PreviewIdx]->GetWorldTransform());
+		Preview->PreviewMesh->SetTransform(TargetComponentInterface(PreviewIdx)->GetWorldTransform());
 		Preview->SetVisibility(BasicProperties->bShowPreview);
 	}
 }
@@ -305,9 +300,9 @@ void UPlaneCutTool::Shutdown(EToolShutdownType ShutdownType)
 	AcceptProperties->SaveProperties(this);
 
 	// Restore (unhide) the source meshes
-	for (TUniquePtr<FPrimitiveComponentTarget>& ComponentTarget : ComponentTargets)
+	for (int Idx = 0; Idx < Targets.Num(); Idx++)
 	{
-		ComponentTarget->SetOwnerVisibility(true);
+		TargetComponentInterface(Idx)->SetOwnerVisibility(true);
 	}
 
 	TArray<FDynamicMeshOpResult> Results;
@@ -339,7 +334,7 @@ TUniquePtr<FDynamicMeshOperator> UPlaneCutOperatorFactory::MakeNewOperator()
 	CutOp->bFillCutHole = CutTool->BasicProperties->bFillCutHole;
 	CutOp->bFillSpans = CutTool->BasicProperties->bFillSpans;
 
-	FTransform LocalToWorld = CutTool->ComponentTargets[ComponentIndex]->GetWorldTransform();
+	FTransform LocalToWorld = CutTool->TargetComponentInterface(ComponentIndex)->GetWorldTransform();
 	CutOp->SetTransform(LocalToWorld);
 	// for all plane computation, change LocalToWorld to not have any zero scale dims
 	FVector LocalToWorldScale = LocalToWorld.GetScale3D();
@@ -418,9 +413,9 @@ void UPlaneCutTool::OnPropertyModified(UObject* PropertySet, FProperty* Property
 {
 	if (Property && (Property->GetFName() == GET_MEMBER_NAME_CHECKED(UPlaneCutToolProperties, bShowPreview)))
 	{
-		for (TUniquePtr<FPrimitiveComponentTarget>& ComponentTarget : ComponentTargets)
+		for (int Idx = 0; Idx < Targets.Num(); Idx++)
 		{
-			ComponentTarget->SetOwnerVisibility(!BasicProperties->bShowPreview);
+			TargetComponentInterface(Idx)->SetOwnerVisibility(!BasicProperties->bShowPreview);
 		}
 		for (UMeshOpPreviewWithBackgroundCompute* Preview : Previews)
 		{
@@ -508,8 +503,8 @@ void UPlaneCutTool::GenerateAsset(const TArray<FDynamicMeshOpResult>& Results)
 	NewSelection.ModificationType = ESelectedObjectsModificationType::Replace;
 	for (int OrigMeshIdx = 0; OrigMeshIdx < NumSourceMeshes; OrigMeshIdx++)
 	{
-		TUniquePtr<FPrimitiveComponentTarget>& ComponentTarget = ComponentTargets[OrigMeshIdx];
-		NewSelection.Actors.Add(ComponentTarget->GetOwnerActor());
+		IPrimitiveComponentBackedTarget* TargetComponent = TargetComponentInterface(OrigMeshIdx);
+		NewSelection.Actors.Add(TargetComponent->GetOwnerActor());
 	}
 
 	// check if we entirely cut away any meshes
@@ -540,7 +535,7 @@ void UPlaneCutTool::GenerateAsset(const TArray<FDynamicMeshOpResult>& Results)
 		{
 			if (bWantDestroy)
 			{
-				ComponentTargets[OrigMeshIdx]->GetOwnerComponent()->DestroyComponent();
+				TargetComponentInterface(OrigMeshIdx)->GetOwnerComponent()->DestroyComponent();
 			}
 			continue;
 		}
@@ -568,10 +563,10 @@ void UPlaneCutTool::GenerateAsset(const TArray<FDynamicMeshOpResult>& Results)
 			}
 		}
 
-		ComponentTargets[OrigMeshIdx]->CommitMesh([&UseMesh](const FPrimitiveComponentTarget::FCommitParams& CommitParams)
+		TargetMeshCommitterInterface(OrigMeshIdx)->CommitMeshDescription([&UseMesh](const IMeshDescriptionCommitter::FCommitterParams& CommitParams)
 		{
 			FDynamicMeshToMeshDescription Converter;
-			Converter.Convert(UseMesh, *CommitParams.MeshDescription);
+			Converter.Convert(UseMesh, *CommitParams.MeshDescriptionOut);
 		});
 	}
 
@@ -587,10 +582,10 @@ void UPlaneCutTool::GenerateAsset(const TArray<FDynamicMeshOpResult>& Results)
 
 			// build array of materials from the original
 			TArray<UMaterialInterface*> Materials;
-			TUniquePtr<FPrimitiveComponentTarget>& ComponentTarget = ComponentTargets[OrigMeshIdx];
-			for (int MaterialIdx = 0, NumMaterials = ComponentTarget->GetNumMaterials(); MaterialIdx < NumMaterials; MaterialIdx++)
+			IMaterialProvider* TargetMaterial = TargetMaterialInterface(OrigMeshIdx);
+			for (int MaterialIdx = 0, NumMaterials = TargetMaterial->GetNumMaterials(); MaterialIdx < NumMaterials; MaterialIdx++)
 			{
-				Materials.Add(ComponentTarget->GetMaterial(MaterialIdx));
+				Materials.Add(TargetMaterial->GetMaterial(MaterialIdx));
 			}
 
 			// add all the additional meshes
diff --git a/Engine/Plugins/Experimental/MeshModelingToolset/Source/MeshModelingTools/Private/ProjectToTargetTool.cpp b/Engine/Plugins/Experimental/MeshModelingToolset/Source/MeshModelingTools/Private/ProjectToTargetTool.cpp
--- a/Engine/Plugins/Experimental/MeshModelingToolset/Source/MeshModelingTools/Private/ProjectToTargetTool.cpp
+++ b/Engine/Plugins/Experimental/MeshModelingToolset/Source/MeshModelingTools/Private/ProjectToTargetTool.cpp
@@ -4,34 +4,39 @@
 #include "MeshDescriptionToDynamicMesh.h"
 #include "InteractiveToolManager.h"
 
+#include "TargetInterfaces/MaterialProvider.h"
+#include "TargetInterfaces/MeshDescriptionCommitter.h"
+#include "TargetInterfaces/MeshDescriptionProvider.h"
+#include "TargetInterfaces/PrimitiveComponentBackedTarget.h"
+#include "ToolTargetManager.h"
+
 #include "ExplicitUseGeometryMathTypes.h"		// using UE::Geometry::(math types)
 using namespace UE::Geometry;
 
 #define LOCTEXT_NAMESPACE "UProjectToTargetTool"
 
+const FToolTargetTypeRequirements& UProjectToTargetToolBuilder::GetTargetRequirements() const
+{
+	static FToolTargetTypeRequirements TypeRequirements({
+		UMeshDescriptionCommitter::StaticClass(),
+		UMeshDescriptionProvider::StaticClass(),
+		UPrimitiveComponentBackedTarget::StaticClass(),
+		UMaterialProvider::StaticClass()
+		});
+	return TypeRequirements;
+}
+
 bool UProjectToTargetToolBuilder::CanBuildTool(const FToolBuilderState& SceneState) const
 {
-	return ToolBuilderUtil::CountComponents(SceneState, CanMakeComponentTarget) == 2;
+	return SceneState.TargetManager->CountSelectedAndTargetable(SceneState, GetTargetRequirements()) == 2;
 }
 
 UInteractiveTool* UProjectToTargetToolBuilder::BuildTool(const FToolBuilderState& SceneState) const
 {
 	UProjectToTargetTool* NewTool = NewObject<UProjectToTargetTool>(SceneState.ToolManager);
 
-	TArray<UActorComponent*> Components = ToolBuilderUtil::FindAllComponents(SceneState, CanMakeComponentTarget);
-	check(Components.Num() == 2);
-
-	TArray<TUniquePtr<FPrimitiveComponentTarget>> ComponentTargets;
-	for (UActorComponent* ActorComponent : Components)
-	{
-		auto* MeshComponent = Cast<UPrimitiveComponent>(ActorComponent);
-		if (MeshComponent)
-		{
-			ComponentTargets.Add(MakeComponentTarget(MeshComponent));
-		}
-	}
-
-	NewTool->SetSelection(MoveTemp(ComponentTargets));
+	TArray<TObjectPtr<UToolTarget>> Targets = SceneState.TargetManager->BuildAllSelectedTargetable(SceneState, GetTargetRequirements());
+	NewTool->SetTargets(MoveTemp(Targets));
 	NewTool->SetWorld(SceneState.World);
 	NewTool->SetAssetAPI(AssetAPI);
 
@@ -42,10 +47,9 @@ void UProjectToTargetTool::Setup()
 {
 	// ProjectionTarget and ProjectionTargetSpatial are setup before calling the parent class's Setup
 	FMeshDescriptionToDynamicMesh ProjectionConverter;
-	check(ComponentTargets.Num() == 2);
-	FPrimitiveComponentTarget* TargetComponentTarget = ComponentTargets[1].Get();
+	check(Targets.Num() == 2);
 	ProjectionTarget = MakeUnique<FDynamicMesh3>();
-	ProjectionConverter.Convert(TargetComponentTarget->GetMesh(), *ProjectionTarget);
+	ProjectionConverter.Convert(TargetMeshProviderInterface(1)->GetMeshDescription(), *ProjectionTarget);
 	ProjectionTargetSpatial = MakeUnique<FDynamicMeshAABBTree3>(ProjectionTarget.Get(), true);
 
 	// Now setup parent RemeshMeshTool class
diff --git a/Engine/Plugins/Experimental/MeshModelingToolset/Source/MeshModelingTools/Private/RemeshMeshTool.cpp b/Engine/Plugins/Experimental/MeshModelingToolset/Source/MeshModelingTools/Private/RemeshMeshTool.cpp
--- a/Engine/Plugins/Experimental/MeshModelingToolset/Source/MeshModelingTools/Private/RemeshMeshTool.cpp
+++ b/Engine/Plugins/Experimental/MeshModelingToolset/Source/MeshModelingTools/Private/RemeshMeshTool.cpp
@@ -17,6 +17,12 @@
 
 #include "SceneManagement.h" // for FPrimitiveDrawInterface
 
+#include "TargetInterfaces/MaterialProvider.h"
+#include "TargetInterfaces/MeshDescriptionCommitter.h"
+#include "TargetInterfaces/MeshDescriptionProvider.h"
+#include "TargetInterfaces/PrimitiveComponentBackedTarget.h"
+#include "ToolTargetManager.h"
+
 #include "ExplicitUseGeometryMathTypes.h"		// using UE::Geometry::(math types)
 using namespace UE::Geometry;
 
@@ -25,25 +31,28 @@ using namespace UE::Geometry;
 /*
  * ToolBuilder
  */
+const FToolTargetTypeRequirements& URemeshMeshToolBuilder::GetTargetRequirements() const
+{
+	static FToolTargetTypeRequirements TypeRequirements({
+		UMaterialProvider::StaticClass(),
+		UMeshDescriptionCommitter::StaticClass(),
+		UMeshDescriptionProvider::StaticClass(),
+		UPrimitiveComponentBackedTarget::StaticClass()
+		});
+	return TypeRequirements;
+}
+
 bool URemeshMeshToolBuilder::CanBuildTool(const FToolBuilderState& SceneState) const
 {
-	return ToolBuilderUtil::CountComponents(SceneState, CanMakeComponentTarget) == 1;
+	return SceneState.TargetManager->CountSelectedAndTargetable(SceneState, GetTargetRequirements()) == 1;
 }
 
 UInteractiveTool* URemeshMeshToolBuilder::BuildTool(const FToolBuilderState& SceneState) const
 {
 	URemeshMeshTool* NewTool = NewObject<URemeshMeshTool>(SceneState.ToolManager);
 
-	TArray<UActorComponent*> Components = ToolBuilderUtil::FindAllComponents(SceneState, CanMakeComponentTarget);
-	check(Components.Num() == 1);
-
-	auto* MeshComponent = Cast<UPrimitiveComponent>(Components[0]);
-	check(MeshComponent != nullptr);
-
-	TArray<TUniquePtr<FPrimitiveComponentTarget>> ComponentTargets;
-	ComponentTargets.Add(MakeComponentTarget(MeshComponent));
-
-	NewTool->SetSelection(MoveTemp(ComponentTargets));
+	TArray<TObjectPtr<UToolTarget>> Targets = SceneState.TargetManager->BuildAllSelectedTargetable(SceneState, GetTargetRequirements());
+	NewTool->SetTargets(MoveTemp(Targets));
 	NewTool->SetWorld(SceneState.World);
 	NewTool->SetAssetAPI(AssetAPI);
 
@@ -99,12 +108,14 @@ void URemeshMeshTool::Setup()
 	BasicProperties->RestoreProperties(this);
 	MeshStatisticsProperties = NewObject<UMeshStatisticsProperties>(this);
 
-	check(ComponentTargets.Num() > 0);
-	check(ComponentTargets[0]);
-	FPrimitiveComponentTarget* ComponentTarget = ComponentTargets[0].Get();
+	check(Targets.Num() > 0);
+	check(Targets[0]);
+	IPrimitiveComponentBackedTarget* TargetComponent = TargetComponentInterface(0);
+	IMeshDescriptionProvider* TargetMeshProvider = TargetMeshProviderInterface(0);
+	IMaterialProvider* TargetMaterial = TargetMaterialInterface(0);
 
 	// hide component and create + show preview
-	ComponentTarget->SetOwnerVisibility(false);
+	TargetComponent->SetOwnerVisibility(false);
 	Preview = NewObject<UMeshOpPreviewWithBackgroundCompute>(this, "Preview");
 	Preview->Setup(this->TargetWorld, this);
 	Preview->OnMeshUpdated.AddLambda([this](UMeshOpPreviewWithBackgroundCompute* Compute)
@@ -112,7 +123,7 @@ void URemeshMeshTool::Setup()
 		MeshStatisticsProperties->Update(*Compute->PreviewMesh->GetPreviewDynamicMesh());
 	});
 	FComponentMaterialSet MaterialSet;
-	ComponentTarget->GetMaterialSet(MaterialSet);
+	TargetMaterial->GetMaterialSet(MaterialSet);
 	Preview->ConfigureMaterials( MaterialSet.Materials,
 								 ToolSetupUtil::GetDefaultWorkingMaterial(GetToolManager())
 	);
@@ -125,9 +136,9 @@ void URemeshMeshTool::Setup()
 
 	OriginalMesh = MakeShared<FDynamicMesh3, ESPMode::ThreadSafe>();
 	FMeshDescriptionToDynamicMesh Converter;
-	Converter.Convert(ComponentTarget->GetMesh(), *OriginalMesh);
+	Converter.Convert(TargetMeshProvider->GetMeshDescription(), *OriginalMesh);
 
-	Preview->PreviewMesh->SetTransform(ComponentTarget->GetWorldTransform());
+	Preview->PreviewMesh->SetTransform(TargetComponent->GetWorldTransform());
 	Preview->PreviewMesh->SetTangentsMode(EDynamicMeshTangentCalcType::AutoCalculated);
 	Preview->PreviewMesh->UpdatePreview(OriginalMesh.Get());
 
@@ -162,9 +173,9 @@ void URemeshMeshTool::Setup()
 void URemeshMeshTool::Shutdown(EToolShutdownType ShutdownType)
 {
 	BasicProperties->SaveProperties(this);
-	for (auto& ComponentTarget : ComponentTargets)
+	for (int ComponentIdx = 0; ComponentIdx < Targets.Num(); ComponentIdx++)
 	{
-		ComponentTarget->SetOwnerVisibility(true);
+		TargetComponentInterface(ComponentIdx)->SetOwnerVisibility(true);
 	}
 	FDynamicMeshOpResult Result = Preview->Shutdown();
 	if (ShutdownType == EToolShutdownType::Accept)
@@ -209,8 +220,8 @@ TUniquePtr<FDynamicMeshOperator> URemeshMeshTool::MakeNewOperator()
 	Op->SmoothingStrength = BasicProperties->SmoothingStrength;
 	Op->SmoothingType = BasicProperties->SmoothingType;
 
-	check(ComponentTargets.Num() > 0);
-	FTransform LocalToWorld = ComponentTargets[0]->GetWorldTransform();
+	check(Targets.Num() > 0);
+	FTransform LocalToWorld = TargetComponentInterface(0)->GetWorldTransform();
 	Op->SetTransform(LocalToWorld);
 
 	Op->OriginalMesh = OriginalMesh;
@@ -225,7 +236,7 @@ TUniquePtr<FDynamicMeshOperator> URemeshMeshTool::MakeNewOperator()
 void URemeshMeshTool::Render(IToolsContextRenderAPI* RenderAPI)
 {
 	FPrimitiveDrawInterface* PDI = RenderAPI->GetPrimitiveDrawInterface();
-	FTransform Transform = ComponentTargets[0]->GetWorldTransform();
+	FTransform Transform = TargetComponentInterface(0)->GetWorldTransform();
 
 	if (BasicProperties->bShowConstraintEdges)
 	{
@@ -280,7 +291,7 @@ void URemeshMeshTool::UpdateVisualization()
 	}
 	else
 	{
-		ComponentTargets[0]->GetMaterialSet(MaterialSet);
+		TargetMaterialInterface(0)->GetMaterialSet(MaterialSet);
 		Preview->PreviewMesh->ClearTriangleColorFunction(UPreviewMesh::ERenderUpdateMode::FastUpdate);
 	}
 	Preview->ConfigureMaterials(MaterialSet.Materials,
@@ -304,12 +315,12 @@ void URemeshMeshTool::GenerateAsset(const FDynamicMeshOpResult& Result)
 	GetToolManager()->BeginUndoTransaction(LOCTEXT("RemeshMeshToolTransactionName", "Remesh Mesh"));
 
 	check(Result.Mesh.Get() != nullptr);
-	ComponentTargets[0]->CommitMesh([&Result](const FPrimitiveComponentTarget::FCommitParams& CommitParams)
+	TargetMeshCommitterInterface(0)->CommitMeshDescription([&Result](const IMeshDescriptionCommitter::FCommitterParams& CommitParams)
 	{
 		FDynamicMeshToMeshDescription Converter;
 
 		// full conversion if normal topology changed or faces were inverted
-		Converter.Convert(Result.Mesh.Get(), *CommitParams.MeshDescription);
+		Converter.Convert(Result.Mesh.Get(), *CommitParams.MeshDescriptionOut);
 	});
 
 	GetToolManager()->EndUndoTransaction();
diff --git a/Engine/Plugins/Experimental/MeshModelingToolset/Source/MeshModelingTools/Private/RemoveOccludedTrianglesTool.cpp b/Engine/Plugins/Experimental/MeshModelingToolset/Source/MeshModelingTools/Private/RemoveOccludedTrianglesTool.cpp
--- a/Engine/Plugins/Experimental/MeshModelingToolset/Source/MeshModelingTools/Private/RemoveOccludedTrianglesTool.cpp
+++ b/Engine/Plugins/Experimental/MeshModelingToolset/Source/MeshModelingTools/Private/RemoveOccludedTrianglesTool.cpp
@@ -25,6 +25,13 @@
 #include "Misc/ScopedSlowTask.h"
 #endif
 
+#include "TargetInterfaces/MaterialProvider.h"
+#include "TargetInterfaces/MeshDescriptionCommitter.h"
+#include "TargetInterfaces/MeshDescriptionProvider.h"
+#include "TargetInterfaces/PrimitiveComponentBackedTarget.h"
+#include "TargetInterfaces/AssetBackedTarget.h"
+#include "ToolTargetManager.h"
+
 #include "ExplicitUseGeometryMathTypes.h"		// using UE::Geometry::(math types)
 using namespace UE::Geometry;
 
@@ -35,29 +42,29 @@ using namespace UE::Geometry;
  */
 
 
+const FToolTargetTypeRequirements& URemoveOccludedTrianglesToolBuilder::GetTargetRequirements() const
+{
+	static FToolTargetTypeRequirements TypeRequirements({
+		UMaterialProvider::StaticClass(),
+		UMeshDescriptionCommitter::StaticClass(),
+		UMeshDescriptionProvider::StaticClass(),
+		UPrimitiveComponentBackedTarget::StaticClass(),
+		UAssetBackedTarget::StaticClass()
+		});
+	return TypeRequirements;
+}
+
 bool URemoveOccludedTrianglesToolBuilder::CanBuildTool(const FToolBuilderState& SceneState) const
 {
-	return ToolBuilderUtil::CountComponents(SceneState, CanMakeComponentTarget) > 0;
+	return SceneState.TargetManager->CountSelectedAndTargetable(SceneState, GetTargetRequirements()) > 0;
 }
 
 UInteractiveTool* URemoveOccludedTrianglesToolBuilder::BuildTool(const FToolBuilderState& SceneState) const
 {
 	URemoveOccludedTrianglesTool* NewTool = NewObject<URemoveOccludedTrianglesTool>(SceneState.ToolManager);
 
-	TArray<UActorComponent*> Components = ToolBuilderUtil::FindAllComponents(SceneState, CanMakeComponentTarget);
-	check(Components.Num() > 0);
-
-	TArray<TUniquePtr<FPrimitiveComponentTarget>> ComponentTargets;
-	for (UActorComponent* ActorComponent : Components)
-	{
-		auto* MeshComponent = Cast<UPrimitiveComponent>(ActorComponent);
-		if ( MeshComponent )
-		{
-			ComponentTargets.Add(MakeComponentTarget(MeshComponent));
-		}
-	}
-
-	NewTool->SetSelection(MoveTemp(ComponentTargets));
+	TArray<TObjectPtr<UToolTarget>> Targets = SceneState.TargetManager->BuildAllSelectedTargetable(SceneState, GetTargetRequirements());
+	NewTool->SetTargets(MoveTemp(Targets));
 	NewTool->SetWorld(SceneState.World);
 	NewTool->SetAssetAPI(AssetAPI);
 
@@ -96,39 +103,28 @@ void URemoveOccludedTrianglesTool::Setup()
 	UInteractiveTool::Setup();
 
 	// hide input StaticMeshComponent
-	for (TUniquePtr<FPrimitiveComponentTarget>& ComponentTarget : ComponentTargets)
+	for (int Idx = 0; Idx < Targets.Num(); Idx++)
 	{
-		ComponentTarget->SetOwnerVisibility(false);
+		TargetComponentInterface(Idx)->SetOwnerVisibility(false);
 	}
 
+
 	// find components with the same source asset
-	TargetToPreviewIdx.SetNum(ComponentTargets.Num());
+	TArray<int32> MapToFirstOccurrences;
+	bool bAnyHasSameSource = GetMapToSharedSourceData(MapToFirstOccurrences);
+	TargetToPreviewIdx.SetNum(Targets.Num());
 	PreviewToTargetIdx.Reset();
-	for (int ComponentIdx = 0; ComponentIdx < TargetToPreviewIdx.Num(); ComponentIdx++)
-	{
-		TargetToPreviewIdx[ComponentIdx] = -1;
-	}
-	bool bAnyHasSameSource = false;
-	for (int32 ComponentIdx = 0; ComponentIdx < ComponentTargets.Num(); ComponentIdx++)
+	for (int32 ComponentIdx = 0; ComponentIdx < MapToFirstOccurrences.Num(); ComponentIdx++)
 	{
-		if (TargetToPreviewIdx[ComponentIdx] >= 0) // already mapped
+		if (MapToFirstOccurrences[ComponentIdx] == ComponentIdx)
 		{
-			continue;
+			int32 NumPreviews = PreviewToTargetIdx.Num();
+			PreviewToTargetIdx.Add(ComponentIdx);
+			TargetToPreviewIdx[ComponentIdx] = NumPreviews;
 		}
-
-		int32 NumPreviews = PreviewToTargetIdx.Num();
-		PreviewToTargetIdx.Add(ComponentIdx);
-		TargetToPreviewIdx[ComponentIdx] = NumPreviews;
-
-		FPrimitiveComponentTarget* ComponentTarget = ComponentTargets[ComponentIdx].Get();
-		for (int32 VsIdx = ComponentIdx + 1; VsIdx < ComponentTargets.Num(); VsIdx++)
+		else
 		{
-			if (ComponentTarget->HasSameSourceData(*ComponentTargets[VsIdx]))
-			{
-				bAnyHasSameSource = true;
-
-				TargetToPreviewIdx[VsIdx] = NumPreviews;
-			}
+			TargetToPreviewIdx[ComponentIdx] = TargetToPreviewIdx[MapToFirstOccurrences[ComponentIdx]];
 		}
 	}
 
@@ -196,7 +192,7 @@ void URemoveOccludedTrianglesTool::MakePolygroupLayerProperties()
 
 void URemoveOccludedTrianglesTool::SetupPreviews()
 {	
-	int32 NumTargets = ComponentTargets.Num();
+	int32 NumTargets = Targets.Num();
 	int32 NumPreviews = PreviewToTargetIdx.Num();
 
 #if WITH_EDITOR
@@ -261,25 +257,25 @@ void URemoveOccludedTrianglesTool::SetupPreviews()
 			OpFactory->PreviewIdx = PreviewIdx;
 			OriginalDynamicMeshes[PreviewIdx] = MakeShared<FDynamicMesh3, ESPMode::ThreadSafe>();
 			FMeshDescriptionToDynamicMesh Converter;
-			Converter.Convert(ComponentTargets[TargetIdx]->GetMesh(), *OriginalDynamicMeshes[PreviewIdx]);
+			Converter.Convert(TargetMeshProviderInterface(TargetIdx)->GetMeshDescription(), *OriginalDynamicMeshes[PreviewIdx]);
 
 			UMeshOpPreviewWithBackgroundCompute* Preview = Previews.Add_GetRef(NewObject<UMeshOpPreviewWithBackgroundCompute>(OpFactory, "Preview"));
 			Preview->Setup(this->TargetWorld, OpFactory);
 			Preview->PreviewMesh->SetTangentsMode(EDynamicMeshTangentCalcType::AutoCalculated);
 
 			FComponentMaterialSet MaterialSet;
-			ComponentTargets[TargetIdx]->GetMaterialSet(MaterialSet);
+			TargetMaterialInterface(TargetIdx)->GetMaterialSet(MaterialSet);
 			Preview->ConfigureMaterials(MaterialSet.Materials,
 				ToolSetupUtil::GetDefaultWorkingMaterial(GetToolManager())
 			);
 
-			Preview->PreviewMesh->SetTransform(ComponentTargets[TargetIdx]->GetWorldTransform());
+			Preview->PreviewMesh->SetTransform(TargetComponentInterface(TargetIdx)->GetWorldTransform());
 			Preview->PreviewMesh->UpdatePreview(OriginalDynamicMeshes[PreviewIdx].Get());
 			Preview->SetVisibility(true);
 
 			OccluderTrees[TargetIdx] = MakeShared<FDynamicMeshAABBTree3, ESPMode::ThreadSafe>(OriginalDynamicMeshes[PreviewIdx].Get());
 			OccluderWindings[TargetIdx] = MakeShared<TFastWindingTree<FDynamicMesh3>, ESPMode::ThreadSafe>(OccluderTrees[TargetIdx].Get());
-			OccluderTransforms[TargetIdx] = (FTransform3d)ComponentTargets[TargetIdx]->GetWorldTransform();
+			OccluderTransforms[TargetIdx] = (FTransform3d)TargetComponentInterface(TargetIdx)->GetWorldTransform();
 
 			// configure secondary render material
 			Preview->SecondaryMaterial = OccludedMaterial;
@@ -299,21 +295,21 @@ void URemoveOccludedTrianglesTool::SetupPreviews()
 			// already did the conversion for a full UMeshOpPreviewWithBackgroundCompute -- just make a light version of that and hook it up to copy the other's work
 			int CopyIdx = PreviewCopies.Num();
 			UPreviewMesh* PreviewMesh = PreviewCopies.Add_GetRef(NewObject<UPreviewMesh>(this));
-			PreviewMesh->CreateInWorld(this->TargetWorld, ComponentTargets[TargetIdx]->GetWorldTransform());
+			PreviewMesh->CreateInWorld(this->TargetWorld, TargetComponentInterface(TargetIdx)->GetWorldTransform());
 
 			PreviewToCopyIdx[PreviewIdx].Add(CopyIdx);
 
 			PreviewMesh->UpdatePreview(OriginalDynamicMeshes[PreviewIdx].Get());
 
 			FComponentMaterialSet MaterialSet;
-			ComponentTargets[TargetIdx]->GetMaterialSet(MaterialSet);
+			TargetMaterialInterface(TargetIdx)->GetMaterialSet(MaterialSet);
 			PreviewMesh->SetMaterials(MaterialSet.Materials);
 			
 			PreviewMesh->SetVisible(true);
 
 			OccluderTrees[TargetIdx] = OccluderTrees[PreviewToTargetIdx[PreviewIdx]];
 			OccluderWindings[TargetIdx] = OccluderWindings[PreviewToTargetIdx[PreviewIdx]];
-			OccluderTransforms[TargetIdx] = (FTransform3d)ComponentTargets[TargetIdx]->GetWorldTransform();
+			OccluderTransforms[TargetIdx] = (FTransform3d)TargetComponentInterface(TargetIdx)->GetWorldTransform();
 
 			PreviewMesh->SetSecondaryRenderMaterial(OccludedMaterial);
 			PreviewMesh->EnableSecondaryTriangleBuffers(MoveTemp(IsOccludedGroupFn));
@@ -336,9 +332,9 @@ void URemoveOccludedTrianglesTool::Shutdown(EToolShutdownType ShutdownType)
 	}
 
 	// Restore (unhide) the source meshes
-	for (TUniquePtr<FPrimitiveComponentTarget>& ComponentTarget : ComponentTargets)
+	for (int ComponentIdx = 0; ComponentIdx < Targets.Num(); ComponentIdx++)
 	{
-		ComponentTarget->SetOwnerVisibility(true);
+		TargetComponentInterface(ComponentIdx)->SetOwnerVisibility(true);
 	}
 
 	// clear all the preview copies
@@ -403,7 +399,7 @@ TUniquePtr<FDynamicMeshOperator> URemoveOccludedTrianglesOperatorFactory::MakeNe
 	Op->WindingIsoValue = Tool->BasicProperties->WindingIsoValue;
 
 	int ComponentIndex = Tool->PreviewToTargetIdx[PreviewIdx];
-	FTransform LocalToWorld = Tool->ComponentTargets[ComponentIndex]->GetWorldTransform();
+	FTransform LocalToWorld = Tool->TargetComponentInterface(ComponentIndex)->GetWorldTransform();
 	Op->OriginalMesh = Tool->OriginalDynamicMeshes[PreviewIdx];
 
 	if (Tool->BasicProperties->bOnlySelfOcclude)
@@ -536,21 +532,21 @@ void URemoveOccludedTrianglesTool::GenerateAsset(const TArray<FDynamicMeshOpResu
 		{
 			if (bWantDestroy)
 			{
-				for (int TargetIdx = 0; TargetIdx < ComponentTargets.Num(); TargetIdx++)
+				for (int TargetIdx = 0; TargetIdx < Targets.Num(); TargetIdx++)
 				{
 					if (TargetToPreviewIdx[TargetIdx] == PreviewIdx)
 					{
-						ComponentTargets[TargetIdx]->GetOwnerComponent()->DestroyComponent();
+						TargetComponentInterface(TargetIdx)->GetOwnerComponent()->DestroyComponent();
 					}
 				}
 			}
 			continue;
 		}
 
-		ComponentTargets[ComponentIdx]->CommitMesh([&Results, &PreviewIdx, this](const FPrimitiveComponentTarget::FCommitParams& CommitParams)
+		TargetMeshCommitterInterface(ComponentIdx)->CommitMeshDescription([&Results, &PreviewIdx, this](const IMeshDescriptionCommitter::FCommitterParams& CommitParams)
 		{
 			FDynamicMeshToMeshDescription Converter;
-			Converter.Convert(Results[PreviewIdx].Mesh.Get(), *CommitParams.MeshDescription);
+			Converter.Convert(Results[PreviewIdx].Mesh.Get(), *CommitParams.MeshDescriptionOut);
 		});
 	}
 
diff --git a/Engine/Plugins/Experimental/MeshModelingToolset/Source/MeshModelingTools/Private/SelfUnionMeshesTool.cpp b/Engine/Plugins/Experimental/MeshModelingToolset/Source/MeshModelingTools/Private/SelfUnionMeshesTool.cpp
--- a/Engine/Plugins/Experimental/MeshModelingToolset/Source/MeshModelingTools/Private/SelfUnionMeshesTool.cpp
+++ b/Engine/Plugins/Experimental/MeshModelingToolset/Source/MeshModelingTools/Private/SelfUnionMeshesTool.cpp
@@ -52,14 +52,14 @@ void USelfUnionMeshesTool::ConvertInputsAndSetPreviewMaterials(bool bSetPreviewM
 {
 	FComponentMaterialSet AllMaterialSet;
 	TMap<UMaterialInterface*, int> KnownMaterials;
-	TArray<TArray<int>> MaterialRemap; MaterialRemap.SetNum(ComponentTargets.Num());
+	TArray<TArray<int>> MaterialRemap; MaterialRemap.SetNum(Targets.Num());
 
 	if (!Properties->bOnlyUseFirstMeshMaterials)
 	{
-		for (int ComponentIdx = 0; ComponentIdx < ComponentTargets.Num(); ComponentIdx++)
+		for (int ComponentIdx = 0; ComponentIdx < Targets.Num(); ComponentIdx++)
 		{
 			FComponentMaterialSet ComponentMaterialSet;
-			ComponentTargets[ComponentIdx]->GetMaterialSet(ComponentMaterialSet);
+			TargetMaterialInterface(ComponentIdx)->GetMaterialSet(ComponentMaterialSet);
 			for (UMaterialInterface* Mat : ComponentMaterialSet.Materials)
 			{
 				int* FoundMatIdx = KnownMaterials.Find(Mat);
@@ -79,14 +79,14 @@ void USelfUnionMeshesTool::ConvertInputsAndSetPreviewMaterials(bool bSetPreviewM
 	}
 	else
 	{
-		ComponentTargets[0]->GetMaterialSet(AllMaterialSet);
+		TargetMaterialInterface(0)->GetMaterialSet(AllMaterialSet);
 		for (int MatIdx = 0; MatIdx < AllMaterialSet.Materials.Num(); MatIdx++)
 		{
 			MaterialRemap[0].Add(MatIdx);
 		}
-		for (int ComponentIdx = 1; ComponentIdx < ComponentTargets.Num(); ComponentIdx++)
+		for (int ComponentIdx = 1; ComponentIdx < Targets.Num(); ComponentIdx++)
 		{
-			MaterialRemap[ComponentIdx].Init(0, ComponentTargets[ComponentIdx]->GetNumMaterials());
+			MaterialRemap[ComponentIdx].Init(0, TargetMaterialInterface(ComponentIdx)->GetNumMaterials());
 		}
 	}
 
@@ -96,11 +96,11 @@ void USelfUnionMeshesTool::ConvertInputsAndSetPreviewMaterials(bool bSetPreviewM
 	CombinedSourceMeshes->Attributes()->EnableMaterialID();
 	FDynamicMeshEditor AppendEditor(CombinedSourceMeshes.Get());
 
-	for (int ComponentIdx = 0; ComponentIdx < ComponentTargets.Num(); ComponentIdx++)
+	for (int ComponentIdx = 0; ComponentIdx < Targets.Num(); ComponentIdx++)
 	{
 		FDynamicMesh3 ComponentMesh;
 		FMeshDescriptionToDynamicMesh Converter;
-		Converter.Convert(ComponentTargets[ComponentIdx]->GetMesh(), ComponentMesh);
+		Converter.Convert(TargetMeshProviderInterface(ComponentIdx)->GetMeshDescription(), ComponentMesh);
 
 		// ensure materials and attributes are always enabled
 		ComponentMesh.EnableAttributes();
diff --git a/Engine/Plugins/Experimental/MeshModelingToolset/Source/MeshModelingTools/Private/TransformMeshesTool.cpp b/Engine/Plugins/Experimental/MeshModelingToolset/Source/MeshModelingTools/Private/TransformMeshesTool.cpp
--- a/Engine/Plugins/Experimental/MeshModelingToolset/Source/MeshModelingTools/Private/TransformMeshesTool.cpp
+++ b/Engine/Plugins/Experimental/MeshModelingToolset/Source/MeshModelingTools/Private/TransformMeshesTool.cpp
@@ -16,6 +16,9 @@
 #include "Components/PrimitiveComponent.h"
 #include "Engine/World.h"
 
+#include "TargetInterfaces/PrimitiveComponentBackedTarget.h"
+#include "ToolTargetManager.h"
+
 #include "ExplicitUseGeometryMathTypes.h"		// using UE::Geometry::(math types)
 using namespace UE::Geometry;
 
@@ -27,29 +30,25 @@ using namespace UE::Geometry;
  */
 
 
+const FToolTargetTypeRequirements& UTransformMeshesToolBuilder::GetTargetRequirements() const
+{
+	static FToolTargetTypeRequirements TypeRequirements(
+		UPrimitiveComponentBackedTarget::StaticClass()
+		);
+	return TypeRequirements;
+}
+
 bool UTransformMeshesToolBuilder::CanBuildTool(const FToolBuilderState& SceneState) const
 {
-	return ToolBuilderUtil::CountComponents(SceneState, CanMakeComponentTarget) >= 1;
+	return SceneState.TargetManager->CountSelectedAndTargetable(SceneState, GetTargetRequirements()) >= 1;
 }
 
 UInteractiveTool* UTransformMeshesToolBuilder::BuildTool(const FToolBuilderState& SceneState) const
 {
 	UTransformMeshesTool* NewTool = NewObject<UTransformMeshesTool>(SceneState.ToolManager);
 
-	TArray<UActorComponent*> Components = ToolBuilderUtil::FindAllComponents(SceneState, CanMakeComponentTarget);
-	check(Components.Num() > 0);
-
-	TArray<TUniquePtr<FPrimitiveComponentTarget>> ComponentTargets;
-	for (UActorComponent* ActorComponent : Components)
-	{
-		auto* MeshComponent = Cast<UPrimitiveComponent>(ActorComponent);
-		if ( MeshComponent )
-		{
-			ComponentTargets.Add(MakeComponentTarget(MeshComponent));
-		}
-	}
-
-	NewTool->SetSelection(MoveTemp(ComponentTargets));
+	TArray<TObjectPtr<UToolTarget>> Targets = SceneState.TargetManager->BuildAllSelectedTargetable(SceneState, GetTargetRequirements());
+	NewTool->SetTargets(MoveTemp(Targets));
 	NewTool->SetWorld(SceneState.World, SceneState.GizmoManager);
 
 	return NewTool;
@@ -229,14 +228,15 @@ void UTransformMeshesTool::SetActiveGizmos_Single(bool bLocalRotations)
 
 	TArray<const UPrimitiveComponent*> ComponentsToIgnoreInAlignment;
 
-	for (TUniquePtr<FPrimitiveComponentTarget>& Target : ComponentTargets)
+	for (int32 ComponentIdx = 0; ComponentIdx < Targets.Num(); ComponentIdx++)
 	{
-		Transformable.TransformProxy->AddComponent(Target->GetOwnerComponent());
-		ComponentsToIgnoreInAlignment.Add(Target->GetOwnerComponent());
+		IPrimitiveComponentBackedTarget* TargetComponent = TargetComponentInterface(ComponentIdx);
+		Transformable.TransformProxy->AddComponent(TargetComponent->GetOwnerComponent());
+		ComponentsToIgnoreInAlignment.Add(TargetComponent->GetOwnerComponent());
 	}
 
 	// leave out nonuniform scale if we have multiple objects in non-local mode
-	bool bCanNonUniformScale = ComponentTargets.Num() == 1 || bLocalRotations;
+	bool bCanNonUniformScale = Targets.Num() == 1 || bLocalRotations;
 	ETransformGizmoSubElements GizmoElements = (bCanNonUniformScale) ?
 		ETransformGizmoSubElements::FullTranslateRotateScale : ETransformGizmoSubElements::TranslateRotateUniformScale;
 	Transformable.TransformGizmo = GizmoManager->CreateCustomRepositionableTransformGizmo(GizmoElements, this);
@@ -252,14 +252,16 @@ void UTransformMeshesTool::SetActiveGizmos_PerObject()
 	check(ActiveGizmos.Num() == 0);
 
 	TArray<const UPrimitiveComponent*> ComponentsToIgnoreInAlignment;
-	for (TUniquePtr<FPrimitiveComponentTarget>& Target : ComponentTargets)
+	for (int32 ComponentIdx = 0; ComponentIdx < Targets.Num(); ComponentIdx++)
 	{
+		IPrimitiveComponentBackedTarget* TargetComponent = TargetComponentInterface(ComponentIdx);
+
 		ComponentsToIgnoreInAlignment.Reset();
 
 		FTransformMeshesTarget Transformable;
 		Transformable.TransformProxy = NewObject<UTransformProxy>(this);
-		Transformable.TransformProxy->AddComponent(Target->GetOwnerComponent());
-		ComponentsToIgnoreInAlignment.Add(Target->GetOwnerComponent());
+		Transformable.TransformProxy->AddComponent(TargetComponent->GetOwnerComponent());
+		ComponentsToIgnoreInAlignment.Add(TargetComponent->GetOwnerComponent());
 
 		ETransformGizmoSubElements GizmoElements = ETransformGizmoSubElements::FullTranslateRotateScale;
 		Transformable.TransformGizmo = GizmoManager->CreateCustomRepositionableTransformGizmo(GizmoElements, this);
@@ -292,12 +294,12 @@ FInputRayHit UTransformMeshesTool::CanBeginClickDragSequence(const FInputDeviceR
 	float MinHitDistance = TNumericLimits<float>::Max();
 	FVector HitNormal;
 
-	for ( int k = 0; k < ComponentTargets.Num(); ++k )
+	for ( int k = 0; k < Targets.Num(); ++k )
 	{
-		TUniquePtr<FPrimitiveComponentTarget>& Target = ComponentTargets[k];
+		IPrimitiveComponentBackedTarget* Target = TargetComponentInterface(k);
 
 		FHitResult WorldHit;
-		if (Target->HitTest(PressPos.WorldRay, WorldHit))
+		if (Target->HitTestComponent(PressPos.WorldRay, WorldHit))
 		{
 			MinHitDistance = FMath::Min(MinHitDistance, WorldHit.Distance);
 			HitNormal = WorldHit.Normal;
@@ -341,11 +343,11 @@ void UTransformMeshesTool::OnClickDrag(const FInputDeviceRay& DragPos)
 	{
 		int IgnoreIndex = (TransformProps->TransformMode == ETransformMeshesTransformMode::PerObjectGizmo) ?
 			ActiveSnapDragIndex : -1;
-		for (int k = 0; k < ComponentTargets.Num(); ++k)
+		for (int k = 0; k < Targets.Num(); ++k)
 		{
 			if (IgnoreIndex == -1 || k == IgnoreIndex)
 			{
-				IgnoreComponents.Add(ComponentTargets[k]->GetOwnerComponent());
+				IgnoreComponents.Add(TargetComponentInterface(k)->GetOwnerComponent());
 			}
 		}
 	}
diff --git a/Engine/Plugins/Experimental/MeshModelingToolset/Source/MeshModelingTools/Private/UVProjectionTool.cpp b/Engine/Plugins/Experimental/MeshModelingToolset/Source/MeshModelingTools/Private/UVProjectionTool.cpp
--- a/Engine/Plugins/Experimental/MeshModelingToolset/Source/MeshModelingTools/Private/UVProjectionTool.cpp
+++ b/Engine/Plugins/Experimental/MeshModelingToolset/Source/MeshModelingTools/Private/UVProjectionTool.cpp
@@ -21,6 +21,12 @@
 #include "BaseGizmos/GizmoComponents.h"
 #include "BaseGizmos/TransformGizmo.h"
 
+#include "TargetInterfaces/MaterialProvider.h"
+#include "TargetInterfaces/MeshDescriptionCommitter.h"
+#include "TargetInterfaces/MeshDescriptionProvider.h"
+#include "TargetInterfaces/PrimitiveComponentBackedTarget.h"
+#include "ToolTargetManager.h"
+
 #include "ExplicitUseGeometryMathTypes.h"		// using UE::Geometry::(math types)
 using namespace UE::Geometry;
 
@@ -31,31 +37,30 @@ using namespace UE::Geometry;
  */
 
 
+const FToolTargetTypeRequirements& UUVProjectionToolBuilder::GetTargetRequirements() const
+{
+	static FToolTargetTypeRequirements TypeRequirements({
+		UMaterialProvider::StaticClass(),
+		UMeshDescriptionCommitter::StaticClass(),
+		UMeshDescriptionProvider::StaticClass(),
+		UPrimitiveComponentBackedTarget::StaticClass()
+		});
+	return TypeRequirements;
+}
+
 bool UUVProjectionToolBuilder::CanBuildTool(const FToolBuilderState& SceneState) const
 {
 	// note most of the code is written to support working on any number > 0, but that seems maybe confusing UI-wise and is not fully tested, so I've limited it to acting on one for now
 	// TODO: if enable tool working on multiple components, figure out what to do if we have multiple component targets that point to the same underlying mesh data?
-	return ToolBuilderUtil::CountComponents(SceneState, CanMakeComponentTarget) == 1;
+	return SceneState.TargetManager->CountSelectedAndTargetable(SceneState, GetTargetRequirements()) == 1;
 }
 
 UInteractiveTool* UUVProjectionToolBuilder::BuildTool(const FToolBuilderState& SceneState) const
 {
 	UUVProjectionTool* NewTool = NewObject<UUVProjectionTool>(SceneState.ToolManager);
 
-	TArray<UActorComponent*> Components = ToolBuilderUtil::FindAllComponents(SceneState, CanMakeComponentTarget);
-	check(Components.Num() > 0);
-
-	TArray<TUniquePtr<FPrimitiveComponentTarget>> ComponentTargets;
-	for (UActorComponent* ActorComponent : Components)
-	{
-		auto* MeshComponent = Cast<UPrimitiveComponent>(ActorComponent);
-		if ( MeshComponent )
-		{
-			ComponentTargets.Add(MakeComponentTarget(MeshComponent));
-		}
-	}
-
-	NewTool->SetSelection(MoveTemp(ComponentTargets));
+	TArray<TObjectPtr<UToolTarget>> Targets = SceneState.TargetManager->BuildAllSelectedTargetable(SceneState, GetTargetRequirements());
+	NewTool->SetTargets(MoveTemp(Targets));
 	NewTool->SetWorld(SceneState.World, SceneState.GizmoManager);
 	NewTool->SetAssetAPI(AssetAPI);
 
@@ -97,9 +102,9 @@ void UUVProjectionTool::Setup()
 	UInteractiveTool::Setup();
 
 	// hide input StaticMeshComponent
-	for (TUniquePtr<FPrimitiveComponentTarget>& ComponentTarget : ComponentTargets)
+	for (int32 ComponentIdx = 0; ComponentIdx < Targets.Num(); ComponentIdx++)
 	{
-		ComponentTarget->SetOwnerVisibility(false);
+		TargetComponentInterface(ComponentIdx)->SetOwnerVisibility(false);
 	}
 
 	BasicProperties = NewObject<UUVProjectionToolProperties>(this, TEXT("UV Projection Settings"));
@@ -139,7 +144,7 @@ void UUVProjectionTool::Setup()
 void UUVProjectionTool::UpdateNumPreviews()
 {
 	int32 CurrentNumPreview = Previews.Num();
-	int32 TargetNumPreview = ComponentTargets.Num();
+	int32 TargetNumPreview = Targets.Num();
 	if (TargetNumPreview < CurrentNumPreview)
 	{
 		for (int32 PreviewIdx = CurrentNumPreview - 1; PreviewIdx >= TargetNumPreview; PreviewIdx--)
@@ -162,10 +167,10 @@ void UUVProjectionTool::UpdateNumPreviews()
 			OpFactory->ComponentIndex = PreviewIdx;
 			OriginalDynamicMeshes[PreviewIdx] = MakeShared<FDynamicMesh3, ESPMode::ThreadSafe>();
 			FMeshDescriptionToDynamicMesh Converter;
-			Converter.Convert(ComponentTargets[PreviewIdx]->GetMesh(), *OriginalDynamicMeshes[PreviewIdx]);
+			Converter.Convert(TargetMeshProviderInterface(PreviewIdx)->GetMeshDescription(), *OriginalDynamicMeshes[PreviewIdx]);
 
 			FVector Center, Extents;
-			FBoxSphereBounds Bounds = ComponentTargets[PreviewIdx]->GetOwnerComponent()->CalcLocalBounds();
+			FBoxSphereBounds Bounds = TargetComponentInterface(PreviewIdx)->GetOwnerComponent()->CalcLocalBounds();
 			
 			FTransform LocalXF(Bounds.Origin);
 			LocalXF.SetScale3D(Bounds.BoxExtent);
@@ -175,17 +180,17 @@ void UUVProjectionTool::UpdateNumPreviews()
 			Preview->PreviewMesh->SetTangentsMode(EDynamicMeshTangentCalcType::AutoCalculated);
 
 			FComponentMaterialSet MaterialSet;
-			ComponentTargets[PreviewIdx]->GetMaterialSet(MaterialSet);
+			TargetMaterialInterface(PreviewIdx)->GetMaterialSet(MaterialSet);
 			Preview->ConfigureMaterials(MaterialSet.Materials,
 				ToolSetupUtil::GetDefaultWorkingMaterial(GetToolManager())
 			);
 			Preview->PreviewMesh->UpdatePreview(OriginalDynamicMeshes[PreviewIdx].Get());
-			Preview->PreviewMesh->SetTransform(ComponentTargets[PreviewIdx]->GetWorldTransform());
+			Preview->PreviewMesh->SetTransform(TargetComponentInterface(PreviewIdx)->GetWorldTransform());
 
 			Preview->SetVisibility(true);
 
 			UTransformProxy* TransformProxy = TransformProxies.Add_GetRef(NewObject<UTransformProxy>(this));
-			TransformProxy->SetTransform(LocalXF * ComponentTargets[PreviewIdx]->GetWorldTransform());
+			TransformProxy->SetTransform(LocalXF * TargetComponentInterface(PreviewIdx)->GetWorldTransform());
 			TransformProxy->OnTransformChanged.AddUObject(this, &UUVProjectionTool::TransformChanged);
 
 			UTransformGizmo* TransformGizmo = TransformGizmos.Add_GetRef(GizmoManager->Create3AxisTransformGizmo(this));
@@ -203,9 +208,9 @@ void UUVProjectionTool::Shutdown(EToolShutdownType ShutdownType)
 	MaterialSettings->SaveProperties(this);
 
 	// Restore (unhide) the source meshes
-	for (TUniquePtr<FPrimitiveComponentTarget>& ComponentTarget : ComponentTargets)
+	for (int32 ComponentIdx = 0; ComponentIdx < Targets.Num(); ComponentIdx++)
 	{
-		ComponentTarget->SetOwnerVisibility(true);
+		TargetComponentInterface(ComponentIdx)->SetOwnerVisibility(true);
 	}
 
 	TArray<FDynamicMeshOpResult> Results;
@@ -241,7 +246,7 @@ TUniquePtr<FDynamicMeshOperator> UUVProjectionOperatorFactory::MakeNewOperator()
 	Op->UVOffset = (FVector2f)Tool->BasicProperties->UVOffset;
 	Op->bWorldSpaceUVScale = Tool->BasicProperties->bWorldSpaceUVScale;
 
-	FTransform LocalToWorld = Tool->ComponentTargets[ComponentIndex]->GetWorldTransform();
+	FTransform LocalToWorld = Tool->TargetComponentInterface(ComponentIndex)->GetWorldTransform();
 	Op->OriginalMesh = Tool->OriginalDynamicMeshes[ComponentIndex];
 	
 	Op->SetTransform(LocalToWorld);
@@ -348,15 +353,15 @@ void UUVProjectionTool::GenerateAsset(const TArray<FDynamicMeshOpResult>& Result
 {
 	GetToolManager()->BeginUndoTransaction(LOCTEXT("UVProjectionToolTransactionName", "UV Projection Tool"));
 
-	check(Results.Num() == ComponentTargets.Num());
+	check(Results.Num() == Targets.Num());
 	
-	for (int32 ComponentIdx = 0; ComponentIdx < ComponentTargets.Num(); ComponentIdx++)
+	for (int32 ComponentIdx = 0; ComponentIdx < Targets.Num(); ComponentIdx++)
 	{
 		check(Results[ComponentIdx].Mesh.Get() != nullptr);
-		ComponentTargets[ComponentIdx]->CommitMesh([&Results, &ComponentIdx, this](const FPrimitiveComponentTarget::FCommitParams& CommitParams)
+		TargetMeshCommitterInterface(ComponentIdx)->CommitMeshDescription([&Results, &ComponentIdx, this](const IMeshDescriptionCommitter::FCommitterParams& CommitParams)
 		{
 			FDynamicMesh3* DynamicMesh = Results[ComponentIdx].Mesh.Get();
-			FMeshDescription* MeshDescription = CommitParams.MeshDescription;
+			FMeshDescription* MeshDescription = CommitParams.MeshDescriptionOut;
 
 			bool bVerticesOnly = false;
 			bool bAttributesOnly = true;
diff --git a/Engine/Plugins/Experimental/MeshModelingToolset/Source/MeshModelingTools/Private/VoxelBlendMeshesTool.cpp b/Engine/Plugins/Experimental/MeshModelingToolset/Source/MeshModelingTools/Private/VoxelBlendMeshesTool.cpp
--- a/Engine/Plugins/Experimental/MeshModelingToolset/Source/MeshModelingTools/Private/VoxelBlendMeshesTool.cpp
+++ b/Engine/Plugins/Experimental/MeshModelingToolset/Source/MeshModelingTools/Private/VoxelBlendMeshesTool.cpp
@@ -58,9 +58,9 @@ TUniquePtr<FDynamicMeshOperator> UVoxelBlendMeshesTool::MakeNewOperator()
 {
 	TUniquePtr<FVoxelBlendMeshesOp> Op = MakeUnique<FVoxelBlendMeshesOp>();
 
-	Op->Transforms.SetNum(ComponentTargets.Num());
-	Op->Meshes.SetNum(ComponentTargets.Num());
-	for (int Idx = 0; Idx < ComponentTargets.Num(); Idx++)
+	Op->Transforms.SetNum(Targets.Num());
+	Op->Meshes.SetNum(Targets.Num());
+	for (int Idx = 0; Idx < Targets.Num(); Idx++)
 	{
 		Op->Meshes[Idx] = OriginalDynamicMeshes[Idx];
 		Op->Transforms[Idx] = TransformProxies[Idx]->GetTransform();
diff --git a/Engine/Plugins/Experimental/MeshModelingToolset/Source/MeshModelingTools/Private/VoxelMorphologyMeshesTool.cpp b/Engine/Plugins/Experimental/MeshModelingToolset/Source/MeshModelingTools/Private/VoxelMorphologyMeshesTool.cpp
--- a/Engine/Plugins/Experimental/MeshModelingToolset/Source/MeshModelingTools/Private/VoxelMorphologyMeshesTool.cpp
+++ b/Engine/Plugins/Experimental/MeshModelingToolset/Source/MeshModelingTools/Private/VoxelMorphologyMeshesTool.cpp
@@ -58,9 +58,9 @@ TUniquePtr<FDynamicMeshOperator> UVoxelMorphologyMeshesTool::MakeNewOperator()
 {
 	TUniquePtr<FVoxelMorphologyMeshesOp> Op = MakeUnique<FVoxelMorphologyMeshesOp>();
 
-	Op->Transforms.SetNum(ComponentTargets.Num());
-	Op->Meshes.SetNum(ComponentTargets.Num());
-	for (int Idx = 0; Idx < ComponentTargets.Num(); Idx++)
+	Op->Transforms.SetNum(Targets.Num());
+	Op->Meshes.SetNum(Targets.Num());
+	for (int Idx = 0; Idx < Targets.Num(); Idx++)
 	{
 		Op->Meshes[Idx] = OriginalDynamicMeshes[Idx];
 		Op->Transforms[Idx] = TransformProxies[Idx]->GetTransform();
diff --git a/Engine/Plugins/Experimental/MeshModelingToolset/Source/MeshModelingTools/Private/VoxelSolidifyMeshesTool.cpp b/Engine/Plugins/Experimental/MeshModelingToolset/Source/MeshModelingTools/Private/VoxelSolidifyMeshesTool.cpp
--- a/Engine/Plugins/Experimental/MeshModelingToolset/Source/MeshModelingTools/Private/VoxelSolidifyMeshesTool.cpp
+++ b/Engine/Plugins/Experimental/MeshModelingToolset/Source/MeshModelingTools/Private/VoxelSolidifyMeshesTool.cpp
@@ -53,9 +53,9 @@ TUniquePtr<FDynamicMeshOperator> UVoxelSolidifyMeshesTool::MakeNewOperator()
 {
 	TUniquePtr<FVoxelSolidifyMeshesOp> Op = MakeUnique<FVoxelSolidifyMeshesOp>();
 
-	Op->Transforms.SetNum(ComponentTargets.Num());
-	Op->Meshes.SetNum(ComponentTargets.Num());
-	for (int Idx = 0; Idx < ComponentTargets.Num(); Idx++)
+	Op->Transforms.SetNum(Targets.Num());
+	Op->Meshes.SetNum(Targets.Num());
+	for (int Idx = 0; Idx < Targets.Num(); Idx++)
 	{
 		Op->Meshes[Idx] = OriginalDynamicMeshes[Idx];
 		Op->Transforms[Idx] = TransformProxies[Idx]->GetTransform();
diff --git a/Engine/Plugins/Experimental/MeshModelingToolset/Source/MeshModelingTools/Public/EditNormalsTool.h b/Engine/Plugins/Experimental/MeshModelingToolset/Source/MeshModelingTools/Public/EditNormalsTool.h
--- a/Engine/Plugins/Experimental/MeshModelingToolset/Source/MeshModelingTools/Public/EditNormalsTool.h
+++ b/Engine/Plugins/Experimental/MeshModelingToolset/Source/MeshModelingTools/Public/EditNormalsTool.h
@@ -34,6 +34,9 @@ public:
 
 	virtual bool CanBuildTool(const FToolBuilderState& SceneState) const override;
 	virtual UInteractiveTool* BuildTool(const FToolBuilderState& SceneState) const override;
+
+protected:
+	virtual const FToolTargetTypeRequirements& GetTargetRequirements() const override;
 };
 
 
diff --git a/Engine/Plugins/Experimental/MeshModelingToolset/Source/MeshModelingTools/Public/EditPivotTool.h b/Engine/Plugins/Experimental/MeshModelingToolset/Source/MeshModelingTools/Public/EditPivotTool.h
--- a/Engine/Plugins/Experimental/MeshModelingToolset/Source/MeshModelingTools/Public/EditPivotTool.h
+++ b/Engine/Plugins/Experimental/MeshModelingToolset/Source/MeshModelingTools/Public/EditPivotTool.h
@@ -28,6 +28,9 @@ class MESHMODELINGTOOLS_API UEditPivotToolBuilder : public UInteractiveToolBuild
 public:
 	virtual bool CanBuildTool(const FToolBuilderState& SceneState) const override;
 	virtual UInteractiveTool* BuildTool(const FToolBuilderState& SceneState) const override;
+
+protected:
+	virtual const FToolTargetTypeRequirements& GetTargetRequirements() const override;
 };
 
 
diff --git a/Engine/Plugins/Experimental/MeshModelingToolset/Source/MeshModelingTools/Public/MirrorTool.h b/Engine/Plugins/Experimental/MeshModelingToolset/Source/MeshModelingTools/Public/MirrorTool.h
--- a/Engine/Plugins/Experimental/MeshModelingToolset/Source/MeshModelingTools/Public/MirrorTool.h
+++ b/Engine/Plugins/Experimental/MeshModelingToolset/Source/MeshModelingTools/Public/MirrorTool.h
@@ -29,6 +29,9 @@ public:
 
 	virtual bool CanBuildTool(const FToolBuilderState& SceneState) const override;
 	virtual UInteractiveTool* BuildTool(const FToolBuilderState& SceneState) const override;
+
+protected:
+	virtual const FToolTargetTypeRequirements& GetTargetRequirements() const override;
 };
 
 UENUM()
@@ -244,4 +247,7 @@ protected:
 
 	void SetupPreviews();
 	void GenerateAsset(const TArray<FDynamicMeshOpResult>& Results);
+
+private:
+	void CheckAndDisplayWarnings();
 };
\ No newline at end of file
diff --git a/Engine/Plugins/Experimental/MeshModelingToolset/Source/MeshModelingTools/Public/Physics/PhysicsInspectorTool.h b/Engine/Plugins/Experimental/MeshModelingToolset/Source/MeshModelingTools/Public/Physics/PhysicsInspectorTool.h
--- a/Engine/Plugins/Experimental/MeshModelingToolset/Source/MeshModelingTools/Public/Physics/PhysicsInspectorTool.h
+++ b/Engine/Plugins/Experimental/MeshModelingToolset/Source/MeshModelingTools/Public/Physics/PhysicsInspectorTool.h
@@ -18,6 +18,9 @@ class MESHMODELINGTOOLS_API UPhysicsInspectorToolBuilder : public UInteractiveTo
 public:
 	virtual bool CanBuildTool(const FToolBuilderState& SceneState) const override;
 	virtual UInteractiveTool* BuildTool(const FToolBuilderState& SceneState) const override;
+
+protected:
+	virtual const FToolTargetTypeRequirements& GetTargetRequirements() const override;
 };
 
 
diff --git a/Engine/Plugins/Experimental/MeshModelingToolset/Source/MeshModelingTools/Public/Physics/SetCollisionGeometryTool.h b/Engine/Plugins/Experimental/MeshModelingToolset/Source/MeshModelingTools/Public/Physics/SetCollisionGeometryTool.h
--- a/Engine/Plugins/Experimental/MeshModelingToolset/Source/MeshModelingTools/Public/Physics/SetCollisionGeometryTool.h
+++ b/Engine/Plugins/Experimental/MeshModelingToolset/Source/MeshModelingTools/Public/Physics/SetCollisionGeometryTool.h
@@ -24,6 +24,9 @@ class MESHMODELINGTOOLS_API USetCollisionGeometryToolBuilder : public UInteracti
 public:
 	virtual bool CanBuildTool(const FToolBuilderState& SceneState) const override;
 	virtual UInteractiveTool* BuildTool(const FToolBuilderState& SceneState) const override;
+
+protected:
+	virtual const FToolTargetTypeRequirements& GetTargetRequirements() const override;
 };
 
 
diff --git a/Engine/Plugins/Experimental/MeshModelingToolset/Source/MeshModelingTools/Public/PlaneCutTool.h b/Engine/Plugins/Experimental/MeshModelingToolset/Source/MeshModelingTools/Public/PlaneCutTool.h
--- a/Engine/Plugins/Experimental/MeshModelingToolset/Source/MeshModelingTools/Public/PlaneCutTool.h
+++ b/Engine/Plugins/Experimental/MeshModelingToolset/Source/MeshModelingTools/Public/PlaneCutTool.h
@@ -35,6 +35,9 @@ public:
 
 	virtual bool CanBuildTool(const FToolBuilderState& SceneState) const override;
 	virtual UInteractiveTool* BuildTool(const FToolBuilderState& SceneState) const override;
+
+protected:
+	virtual const FToolTargetTypeRequirements& GetTargetRequirements() const override;
 };
 
 
diff --git a/Engine/Plugins/Experimental/MeshModelingToolset/Source/MeshModelingTools/Public/ProjectToTargetTool.h b/Engine/Plugins/Experimental/MeshModelingToolset/Source/MeshModelingTools/Public/ProjectToTargetTool.h
--- a/Engine/Plugins/Experimental/MeshModelingToolset/Source/MeshModelingTools/Public/ProjectToTargetTool.h
+++ b/Engine/Plugins/Experimental/MeshModelingToolset/Source/MeshModelingTools/Public/ProjectToTargetTool.h
@@ -20,6 +20,9 @@ public:
 	virtual bool CanBuildTool(const FToolBuilderState& SceneState) const override;
 
 	virtual UInteractiveTool* BuildTool(const FToolBuilderState& SceneState) const override;
+
+protected:
+	virtual const FToolTargetTypeRequirements& GetTargetRequirements() const override;
 };
 
 
diff --git a/Engine/Plugins/Experimental/MeshModelingToolset/Source/MeshModelingTools/Public/RemeshMeshTool.h b/Engine/Plugins/Experimental/MeshModelingToolset/Source/MeshModelingTools/Public/RemeshMeshTool.h
--- a/Engine/Plugins/Experimental/MeshModelingToolset/Source/MeshModelingTools/Public/RemeshMeshTool.h
+++ b/Engine/Plugins/Experimental/MeshModelingToolset/Source/MeshModelingTools/Public/RemeshMeshTool.h
@@ -33,6 +33,9 @@ public:
 	virtual bool CanBuildTool(const FToolBuilderState& SceneState) const override;
 
 	virtual UInteractiveTool* BuildTool(const FToolBuilderState& SceneState) const override;
+
+protected:
+	virtual const FToolTargetTypeRequirements& GetTargetRequirements() const override;
 };
 
 /**
diff --git a/Engine/Plugins/Experimental/MeshModelingToolset/Source/MeshModelingTools/Public/RemoveOccludedTrianglesTool.h b/Engine/Plugins/Experimental/MeshModelingToolset/Source/MeshModelingTools/Public/RemoveOccludedTrianglesTool.h
--- a/Engine/Plugins/Experimental/MeshModelingToolset/Source/MeshModelingTools/Public/RemoveOccludedTrianglesTool.h
+++ b/Engine/Plugins/Experimental/MeshModelingToolset/Source/MeshModelingTools/Public/RemoveOccludedTrianglesTool.h
@@ -31,6 +31,9 @@ public:
 
 	virtual bool CanBuildTool(const FToolBuilderState& SceneState) const override;
 	virtual UInteractiveTool* BuildTool(const FToolBuilderState& SceneState) const override;
+
+protected:
+	virtual const FToolTargetTypeRequirements& GetTargetRequirements() const override;
 };
 
 
diff --git a/Engine/Plugins/Experimental/MeshModelingToolset/Source/MeshModelingTools/Public/TransformMeshesTool.h b/Engine/Plugins/Experimental/MeshModelingToolset/Source/MeshModelingTools/Public/TransformMeshesTool.h
--- a/Engine/Plugins/Experimental/MeshModelingToolset/Source/MeshModelingTools/Public/TransformMeshesTool.h
+++ b/Engine/Plugins/Experimental/MeshModelingToolset/Source/MeshModelingTools/Public/TransformMeshesTool.h
@@ -27,6 +27,9 @@ class MESHMODELINGTOOLS_API UTransformMeshesToolBuilder : public UInteractiveToo
 public:
 	virtual bool CanBuildTool(const FToolBuilderState& SceneState) const override;
 	virtual UInteractiveTool* BuildTool(const FToolBuilderState& SceneState) const override;
+
+protected:
+	virtual const FToolTargetTypeRequirements& GetTargetRequirements() const override;
 };
 
 
diff --git a/Engine/Plugins/Experimental/MeshModelingToolset/Source/MeshModelingTools/Public/UVProjectionTool.h b/Engine/Plugins/Experimental/MeshModelingToolset/Source/MeshModelingTools/Public/UVProjectionTool.h
--- a/Engine/Plugins/Experimental/MeshModelingToolset/Source/MeshModelingTools/Public/UVProjectionTool.h
+++ b/Engine/Plugins/Experimental/MeshModelingToolset/Source/MeshModelingTools/Public/UVProjectionTool.h
@@ -37,6 +37,9 @@ public:
 
 	virtual bool CanBuildTool(const FToolBuilderState& SceneState) const override;
 	virtual UInteractiveTool* BuildTool(const FToolBuilderState& SceneState) const override;
+
+protected:
+	virtual const FToolTargetTypeRequirements& GetTargetRequirements() const override;
 };
 
 
diff --git a/Engine/Plugins/Experimental/MeshModelingToolset/Source/MeshModelingToolsEditorOnly/Private/AttributeEditorTool.cpp b/Engine/Plugins/Experimental/MeshModelingToolset/Source/MeshModelingToolsEditorOnly/Private/AttributeEditorTool.cpp
--- a/Engine/Plugins/Experimental/MeshModelingToolset/Source/MeshModelingToolsEditorOnly/Private/AttributeEditorTool.cpp
+++ b/Engine/Plugins/Experimental/MeshModelingToolset/Source/MeshModelingToolsEditorOnly/Private/AttributeEditorTool.cpp
@@ -19,6 +19,11 @@
 
 #include "Components/PrimitiveComponent.h"
 
+#include "TargetInterfaces/MeshDescriptionCommitter.h"
+#include "TargetInterfaces/MeshDescriptionProvider.h"
+#include "TargetInterfaces/PrimitiveComponentBackedTarget.h"
+#include "ToolTargetManager.h"
+
 #include "ExplicitUseGeometryMathTypes.h"		// using UE::Geometry::(math types)
 using namespace UE::Geometry;
 
@@ -29,29 +34,27 @@ using namespace UE::Geometry;
  */
 
 
+const FToolTargetTypeRequirements& UAttributeEditorToolBuilder::GetTargetRequirements() const
+{
+	static FToolTargetTypeRequirements TypeRequirements({
+		UMeshDescriptionCommitter::StaticClass(),
+		UMeshDescriptionProvider::StaticClass(),
+		UPrimitiveComponentBackedTarget::StaticClass()
+		});
+	return TypeRequirements;
+}
+
 bool UAttributeEditorToolBuilder::CanBuildTool(const FToolBuilderState& SceneState) const
 {
-	return ToolBuilderUtil::CountComponents(SceneState, CanMakeComponentTarget) > 0;
+	return SceneState.TargetManager->CountSelectedAndTargetable(SceneState, GetTargetRequirements()) > 0;
 }
 
 UInteractiveTool* UAttributeEditorToolBuilder::BuildTool(const FToolBuilderState& SceneState) const
 {
 	UAttributeEditorTool* NewTool = NewObject<UAttributeEditorTool>(SceneState.ToolManager);
 
-	TArray<UActorComponent*> Components = ToolBuilderUtil::FindAllComponents(SceneState, CanMakeComponentTarget);
-	check(Components.Num() > 0);
-
-	TArray<TUniquePtr<FPrimitiveComponentTarget>> ComponentTargets;
-	for (UActorComponent* ActorComponent : Components)
-	{
-		auto* MeshComponent = Cast<UPrimitiveComponent>(ActorComponent);
-		if ( MeshComponent )
-		{
-			ComponentTargets.Add(MakeComponentTarget(MeshComponent));
-		}
-	}
-
-	NewTool->SetSelection(MoveTemp(ComponentTargets));
+	TArray<TObjectPtr<UToolTarget>> Targets = SceneState.TargetManager->BuildAllSelectedTargetable(SceneState, GetTargetRequirements());
+	NewTool->SetTargets(MoveTemp(Targets));
 	NewTool->SetWorld(SceneState.World);
 
 	return NewTool;
@@ -112,7 +115,7 @@ void UAttributeEditorTool::Setup()
 	NormalsActions->Initialize(this);
 	AddToolPropertySource(NormalsActions);
 
-	if (ComponentTargets.Num() == 1)
+	if (Targets.Num() == 1)
 	{
 		UVActions = NewObject<UAttributeEditorUVActions>(this);
 		UVActions->Initialize(this);
@@ -280,7 +283,7 @@ static bool RemoveAttribute(FMeshDescription* Mesh, EAttributeEditorElementType
 
 void UAttributeEditorTool::InitializeAttributeLists()
 {
-	FMeshDescription* Mesh = ComponentTargets[0]->GetMesh();
+	FMeshDescription* Mesh = TargetMeshProviderInterface(0)->GetMeshDescription();
 
 
 	TVertexInstanceAttributesRef<FVector2D> InstanceUVs =
@@ -293,7 +296,7 @@ void UAttributeEditorTool::InitializeAttributeLists()
 	}
 	UVActions->UVLayer = UVActions->UVLayerNamesList[0];
 
-	UStaticMeshComponent* StaticMeshComponent = Cast<UStaticMeshComponent>(ComponentTargets[0]->GetOwnerComponent());
+	UStaticMeshComponent* StaticMeshComponent = Cast<UStaticMeshComponent>(TargetComponentInterface(0)->GetOwnerComponent());
 	if (StaticMeshComponent && StaticMeshComponent->GetStaticMesh() != nullptr)
 	{
 		const UStaticMesh* StaticMesh = StaticMeshComponent->GetStaticMesh();
@@ -402,7 +405,7 @@ void UAttributeEditorTool::OnTick(float DeltaTime)
 	}
 	PendingAction = EAttributeEditorToolActions::NoAction;
 
-	if (bAttributeListsValid == false && ComponentTargets.Num() == 1)
+	if (bAttributeListsValid == false && Targets.Num() == 1)
 	{
 		InitializeAttributeLists();
 	}
@@ -415,9 +418,9 @@ void UAttributeEditorTool::OptimizeForEditing()
 {
 	GetToolManager()->BeginUndoTransaction(LOCTEXT("OptimizeForEditing", "Optimize For Editing"));
 
-	for (TUniquePtr<FPrimitiveComponentTarget>& Target : ComponentTargets)
+	for (int32 ComponentIdx = 0; ComponentIdx < Targets.Num(); ComponentIdx++)
 	{
-		UStaticMeshComponent* StaticMeshComponent = Cast<UStaticMeshComponent>(Target->GetOwnerComponent());
+		UStaticMeshComponent* StaticMeshComponent = Cast<UStaticMeshComponent>(TargetComponentInterface(ComponentIdx)->GetOwnerComponent());
 		if (StaticMeshComponent == nullptr || StaticMeshComponent->GetStaticMesh() == nullptr)
 		{
 			continue;
@@ -435,12 +438,12 @@ void UAttributeEditorTool::OptimizeForEditing()
 		BuildSettings.DistanceFieldResolutionScale = 0.01;
 
 		// this will call StaticMesh->PostEditChange()...
-		Target->CommitMesh([&](const FPrimitiveComponentTarget::FCommitParams& CommitParams)
+		TargetMeshCommitterInterface(ComponentIdx)->CommitMeshDescription([&](const IMeshDescriptionCommitter::FCommitterParams& CommitParams)
 		{
-			TEdgeAttributesRef<bool> EdgeHardnesses = CommitParams.MeshDescription->EdgeAttributes().GetAttributesRef<bool>(MeshAttribute::Edge::IsHard);
+			TEdgeAttributesRef<bool> EdgeHardnesses = CommitParams.MeshDescriptionOut->EdgeAttributes().GetAttributesRef<bool>(MeshAttribute::Edge::IsHard);
 			if (EdgeHardnesses.IsValid())
 			{
-				for (FEdgeID ElID : CommitParams.MeshDescription->Edges().GetElementIDs())
+				for (FEdgeID ElID : CommitParams.MeshDescriptionOut->Edges().GetElementIDs())
 				{
 					EdgeHardnesses[ElID] = false;
 				}
@@ -449,7 +452,7 @@ void UAttributeEditorTool::OptimizeForEditing()
 			// force computation of normals/tangents if they are auto-generated
 			if (BuildSettings.bRecomputeNormals || BuildSettings.bRecomputeTangents)
 			{
-				UE::MeshDescription::InitializeAutoGeneratedAttributes(*CommitParams.MeshDescription, StaticMeshComponent, 0);
+				UE::MeshDescription::InitializeAutoGeneratedAttributes(*CommitParams.MeshDescriptionOut, StaticMeshComponent, 0);
 			}
 
 			// now clear these build settings
@@ -470,20 +473,20 @@ void UAttributeEditorTool::ClearNormals()
 {
 	GetToolManager()->BeginUndoTransaction(LOCTEXT("ClearNormalsTransactionMessage", "Clear Normals"));
 
-	for (TUniquePtr<FPrimitiveComponentTarget>& Target : ComponentTargets)
+	for (int32 ComponentIdx = 0; ComponentIdx < Targets.Num(); ComponentIdx++)
 	{
-		Target->CommitMesh([&](const FPrimitiveComponentTarget::FCommitParams& CommitParams)
+		TargetMeshCommitterInterface(ComponentIdx)->CommitMeshDescription([&](const IMeshDescriptionCommitter::FCommitterParams& CommitParams)
 		{
-			TEdgeAttributesRef<bool> EdgeHardnesses = CommitParams.MeshDescription->EdgeAttributes().GetAttributesRef<bool>(MeshAttribute::Edge::IsHard);
+			TEdgeAttributesRef<bool> EdgeHardnesses = CommitParams.MeshDescriptionOut->EdgeAttributes().GetAttributesRef<bool>(MeshAttribute::Edge::IsHard);
 			if (EdgeHardnesses.IsValid())
 			{
-				for (FEdgeID ElID : CommitParams.MeshDescription->Edges().GetElementIDs())
+				for (FEdgeID ElID : CommitParams.MeshDescriptionOut->Edges().GetElementIDs())
 				{
 					EdgeHardnesses[ElID] = false;
 				}
 			}
-			FStaticMeshOperations::ComputeTriangleTangentsAndNormals(*CommitParams.MeshDescription, FMathf::Epsilon);
-			FStaticMeshOperations::RecomputeNormalsAndTangentsIfNeeded(*CommitParams.MeshDescription, EComputeNTBsFlags::WeightedNTBs | EComputeNTBsFlags::Normals);
+			FStaticMeshOperations::ComputeTriangleTangentsAndNormals(*CommitParams.MeshDescriptionOut, FMathf::Epsilon);
+			FStaticMeshOperations::RecomputeNormalsAndTangentsIfNeeded(*CommitParams.MeshDescriptionOut, EComputeNTBsFlags::WeightedNTBs | EComputeNTBsFlags::Normals);
 		});
 	}
 	GetToolManager()->EndUndoTransaction();
@@ -495,17 +498,17 @@ void UAttributeEditorTool::ClearNormals()
 void UAttributeEditorTool::ClearUVs()
 {
 	GetToolManager()->BeginUndoTransaction(LOCTEXT("ClearUVsTransactionMessage", "Clear Selected UVs"));
-	for (TUniquePtr<FPrimitiveComponentTarget>& Target : ComponentTargets)
+	for (int32 ComponentIdx = 0; ComponentIdx < Targets.Num(); ComponentIdx++)
 	{
-		Target->CommitMesh([&](const FPrimitiveComponentTarget::FCommitParams& CommitParams)
+		TargetMeshCommitterInterface(ComponentIdx)->CommitMeshDescription([&](const IMeshDescriptionCommitter::FCommitterParams& CommitParams)
 		{
 			TVertexInstanceAttributesRef<FVector2D> InstanceUVs =
-				CommitParams.MeshDescription->VertexInstanceAttributes().GetAttributesRef<FVector2D>(MeshAttribute::VertexInstance::TextureCoordinate);
+				CommitParams.MeshDescriptionOut->VertexInstanceAttributes().GetAttributesRef<FVector2D>(MeshAttribute::VertexInstance::TextureCoordinate);
 			int32 NumChannels = InstanceUVs.GetNumChannels();
-			const FVertexInstanceArray& Instances = CommitParams.MeshDescription->VertexInstances();
+			const FVertexInstanceArray& Instances = CommitParams.MeshDescriptionOut->VertexInstances();
 			for (int LayerIndex = NumChannels-1; LayerIndex >= 0; LayerIndex--)
 			{
-				if (!FStaticMeshOperations::RemoveUVChannel(*CommitParams.MeshDescription, LayerIndex))
+				if (!FStaticMeshOperations::RemoveUVChannel(*CommitParams.MeshDescriptionOut, LayerIndex))
 				{
 					for (const FVertexInstanceID& ElID : Instances.GetElementIDs())
 					{
@@ -516,7 +519,7 @@ void UAttributeEditorTool::ClearUVs()
 
 			if (bHaveAutoGeneratedLightmapUVSet)
 			{
-				UpdateAutoGeneratedLightmapUVChannel(Target, InstanceUVs.GetNumChannels());
+				UpdateAutoGeneratedLightmapUVChannel(TargetComponentInterface(ComponentIdx), InstanceUVs.GetNumChannels());
 			}
 		});
 	}
@@ -543,15 +546,15 @@ void UAttributeEditorTool::DeleteSelectedUVSet()
 	}
 
 	GetToolManager()->BeginUndoTransaction(LOCTEXT("ClearUVsTransactionMessage", "Clear Selected UVs"));
-	for (TUniquePtr<FPrimitiveComponentTarget>& Target : ComponentTargets)
+	for (int32 ComponentIdx = 0; ComponentIdx < Targets.Num(); ComponentIdx++)
 	{
-		Target->CommitMesh([&](const FPrimitiveComponentTarget::FCommitParams& CommitParams)
+		TargetMeshCommitterInterface(ComponentIdx)->CommitMeshDescription([&](const IMeshDescriptionCommitter::FCommitterParams& CommitParams)
 		{
 			TVertexInstanceAttributesRef<FVector2D> InstanceUVs =
-				CommitParams.MeshDescription->VertexInstanceAttributes().GetAttributesRef<FVector2D>(MeshAttribute::VertexInstance::TextureCoordinate);
-			if (!FStaticMeshOperations::RemoveUVChannel(*CommitParams.MeshDescription, DeleteIndex))
+				CommitParams.MeshDescriptionOut->VertexInstanceAttributes().GetAttributesRef<FVector2D>(MeshAttribute::VertexInstance::TextureCoordinate);
+			if (!FStaticMeshOperations::RemoveUVChannel(*CommitParams.MeshDescriptionOut, DeleteIndex))
 			{
-				const FVertexInstanceArray& Instances = CommitParams.MeshDescription->VertexInstances();
+				const FVertexInstanceArray& Instances = CommitParams.MeshDescriptionOut->VertexInstances();
 				for (const FVertexInstanceID& InstanceID : Instances.GetElementIDs())
 				{
 					InstanceUVs.Set(InstanceID, DeleteIndex, FVector2D::ZeroVector);
@@ -560,7 +563,7 @@ void UAttributeEditorTool::DeleteSelectedUVSet()
 
 			if (bHaveAutoGeneratedLightmapUVSet)
 			{
-				UpdateAutoGeneratedLightmapUVChannel(Target, InstanceUVs.GetNumChannels());
+				UpdateAutoGeneratedLightmapUVChannel(TargetComponentInterface(ComponentIdx), InstanceUVs.GetNumChannels());
 			}
 		});
 	}
@@ -573,14 +576,14 @@ void UAttributeEditorTool::DeleteSelectedUVSet()
 void UAttributeEditorTool::AddUVSet()
 {
 	GetToolManager()->BeginUndoTransaction(LOCTEXT("AddUVSetMessage", "Add UV Set"));
-	for (TUniquePtr<FPrimitiveComponentTarget>& Target : ComponentTargets)
+	for (int32 ComponentIdx = 0; ComponentIdx < Targets.Num(); ComponentIdx++)
 	{
-		Target->CommitMesh([&](const FPrimitiveComponentTarget::FCommitParams& CommitParams)
+		TargetMeshCommitterInterface(ComponentIdx)->CommitMeshDescription([&](const IMeshDescriptionCommitter::FCommitterParams& CommitParams)
 		{
 			TVertexInstanceAttributesRef<FVector2D> InstanceUVs =
-				CommitParams.MeshDescription->VertexInstanceAttributes().GetAttributesRef<FVector2D>(MeshAttribute::VertexInstance::TextureCoordinate);
+				CommitParams.MeshDescriptionOut->VertexInstanceAttributes().GetAttributesRef<FVector2D>(MeshAttribute::VertexInstance::TextureCoordinate);
 			int32 NewChannelIndex = InstanceUVs.GetNumChannels();
-			if (!FStaticMeshOperations::AddUVChannel(*CommitParams.MeshDescription))
+			if (!FStaticMeshOperations::AddUVChannel(*CommitParams.MeshDescriptionOut))
 			{
 				GetToolManager()->DisplayMessage(LOCTEXT("FailedToAddUVSet", "Adding UV Set Failed"), EToolMessageLevel::UserWarning);
 			}
@@ -590,7 +593,7 @@ void UAttributeEditorTool::AddUVSet()
 
 				if (bHaveAutoGeneratedLightmapUVSet)
 				{
-					UpdateAutoGeneratedLightmapUVChannel(Target, InstanceUVs.GetNumChannels());
+					UpdateAutoGeneratedLightmapUVChannel(TargetComponentInterface(ComponentIdx), InstanceUVs.GetNumChannels());
 				}
 			}
 		});
@@ -612,20 +615,20 @@ void UAttributeEditorTool::DuplicateSelectedUVSet()
 	}
 
 	GetToolManager()->BeginUndoTransaction(LOCTEXT("DuplicateUVSetMessage", "Duplicate UV Set"));
-	for (TUniquePtr<FPrimitiveComponentTarget>& Target : ComponentTargets)
+	for (int32 ComponentIdx = 0; ComponentIdx < Targets.Num(); ComponentIdx++)
 	{
-		Target->CommitMesh([&](const FPrimitiveComponentTarget::FCommitParams& CommitParams)
+		TargetMeshCommitterInterface(ComponentIdx)->CommitMeshDescription([&](const IMeshDescriptionCommitter::FCommitterParams& CommitParams)
 		{
 			TVertexInstanceAttributesRef<FVector2D> InstanceUVs =
-				CommitParams.MeshDescription->VertexInstanceAttributes().GetAttributesRef<FVector2D>(MeshAttribute::VertexInstance::TextureCoordinate);
+				CommitParams.MeshDescriptionOut->VertexInstanceAttributes().GetAttributesRef<FVector2D>(MeshAttribute::VertexInstance::TextureCoordinate);
 			int32 NewChannelIndex = InstanceUVs.GetNumChannels();
-			if (!FStaticMeshOperations::AddUVChannel(*CommitParams.MeshDescription))
+			if (!FStaticMeshOperations::AddUVChannel(*CommitParams.MeshDescriptionOut))
 			{
 				GetToolManager()->DisplayMessage(LOCTEXT("FailedToAddUVSet", "Adding UV Set Failed"), EToolMessageLevel::UserWarning);
 			}
 			else
 			{
-				const FVertexInstanceArray& Instances = CommitParams.MeshDescription->VertexInstances();
+				const FVertexInstanceArray& Instances = CommitParams.MeshDescriptionOut->VertexInstances();
 				for (const FVertexInstanceID& InstanceID : Instances.GetElementIDs())
 				{
 					FVector2D SourceUV = InstanceUVs.Get(InstanceID, SourceIndex);
@@ -634,7 +637,7 @@ void UAttributeEditorTool::DuplicateSelectedUVSet()
 
 				if (bHaveAutoGeneratedLightmapUVSet)
 				{
-					UpdateAutoGeneratedLightmapUVChannel(Target, InstanceUVs.GetNumChannels());
+					UpdateAutoGeneratedLightmapUVChannel(TargetComponentInterface(ComponentIdx), InstanceUVs.GetNumChannels());
 				}
 
 				GetToolManager()->DisplayMessage(FText::Format(LOCTEXT("Copied UV Set", "Copied UV{0} to UV{1}"), 
@@ -659,7 +662,7 @@ void UAttributeEditorTool::AddNewAttribute(EAttributeEditorElementType ElemType,
 		return;
 	}
 
-	FMeshDescription* CurMesh = ComponentTargets[0]->GetMesh();
+	FMeshDescription* CurMesh = TargetMeshProviderInterface(0)->GetMeshDescription();
 	if (HasAttribute(CurMesh, ElemType, AttributeName))
 	{
 		GetToolManager()->DisplayMessage(LOCTEXT("ErrorAddingDuplicateNameMessage", "Attribute with this name already exists"), EToolMessageLevel::UserWarning);
@@ -674,9 +677,9 @@ void UAttributeEditorTool::AddNewAttribute(EAttributeEditorElementType ElemType,
 	}
 
 	GetToolManager()->BeginUndoTransaction(LOCTEXT("NewAttributeTransactionMessage", "Add Attribute"));
-	ComponentTargets[0]->CommitMesh([&NewMesh](const FPrimitiveComponentTarget::FCommitParams& CommitParams)
+	TargetMeshCommitterInterface(0)->CommitMeshDescription([&NewMesh](const IMeshDescriptionCommitter::FCommitterParams& CommitParams)
 	{
-		*CommitParams.MeshDescription = NewMesh;
+		*CommitParams.MeshDescriptionOut = NewMesh;
 	});
 	GetToolManager()->EndUndoTransaction();
 
@@ -712,7 +715,7 @@ void UAttributeEditorTool::ClearAttribute()
 
 void UAttributeEditorTool::DeleteAttribute()
 {
-	FMeshDescription* CurMesh = ComponentTargets[0]->GetMesh();
+	FMeshDescription* CurMesh = TargetMeshProviderInterface(0)->GetMeshDescription();
 	FName SelectedName(ModifyAttributeProps->Attribute);
 
 	// @todo: use FAttributesSetBase::DoesAttributeHaveAnyFlags(EMeshAttributeFlags::Mandatory) to determine this ?
@@ -741,9 +744,9 @@ void UAttributeEditorTool::DeleteAttribute()
 	}
 
 	GetToolManager()->BeginUndoTransaction(LOCTEXT("RemoveAttributeTransactionMessage", "Remove Attribute"));
-	ComponentTargets[0]->CommitMesh([&NewMesh](const FPrimitiveComponentTarget::FCommitParams& CommitParams)
+	TargetMeshCommitterInterface(0)->CommitMeshDescription([&NewMesh](const IMeshDescriptionCommitter::FCommitterParams& CommitParams)
 	{
-		*CommitParams.MeshDescription = NewMesh;
+		*CommitParams.MeshDescriptionOut = NewMesh;
 	});
 	GetToolManager()->EndUndoTransaction();
 
@@ -762,9 +765,9 @@ void UAttributeEditorTool::SetLightmapUVsEnabled(bool bEnabled)
 	{
 		GetToolManager()->BeginUndoTransaction(LOCTEXT("DisableLightmapUVs", "Disable Lightmap UVs"));
 	}
-	for (TUniquePtr<FPrimitiveComponentTarget>& Target : ComponentTargets)
+	for (int32 ComponentIdx = 0; ComponentIdx < Targets.Num(); ComponentIdx++)
 	{
-		UStaticMeshComponent* StaticMeshComponent = Cast<UStaticMeshComponent>(Target->GetOwnerComponent());
+		UStaticMeshComponent* StaticMeshComponent = Cast<UStaticMeshComponent>(TargetComponentInterface(ComponentIdx)->GetOwnerComponent());
 		if (StaticMeshComponent && StaticMeshComponent->GetStaticMesh() != nullptr)
 		{
 			UStaticMesh* StaticMesh = StaticMeshComponent->GetStaticMesh();
@@ -786,13 +789,13 @@ void UAttributeEditorTool::SetLightmapUVsEnabled(bool bEnabled)
 void UAttributeEditorTool::ResetLightmapUVsChannels()
 {
 	GetToolManager()->BeginUndoTransaction(LOCTEXT("ResetLightmapUVs", "Reset Lightmap UVs"));
-	for (TUniquePtr<FPrimitiveComponentTarget>& Target : ComponentTargets)
+	for (int32 ComponentIdx = 0; ComponentIdx < Targets.Num(); ComponentIdx++)
 	{
 		TVertexInstanceAttributesRef<FVector2D> InstanceUVs = 
-			Target->GetMesh()->VertexInstanceAttributes().GetAttributesRef<FVector2D>(MeshAttribute::VertexInstance::TextureCoordinate);
+			TargetMeshProviderInterface(ComponentIdx)->GetMeshDescription()->VertexInstanceAttributes().GetAttributesRef<FVector2D>(MeshAttribute::VertexInstance::TextureCoordinate);
 		int32 SetChannel = FMath::Max(InstanceUVs.GetNumChannels(), 1);
 
-		UStaticMeshComponent* StaticMeshComponent = Cast<UStaticMeshComponent>(Target->GetOwnerComponent());
+		UStaticMeshComponent* StaticMeshComponent = Cast<UStaticMeshComponent>(TargetComponentInterface(ComponentIdx)->GetOwnerComponent());
 		if (StaticMeshComponent && StaticMeshComponent->GetStaticMesh() != nullptr)
 		{
 			UStaticMesh* StaticMesh = StaticMeshComponent->GetStaticMesh();
@@ -812,7 +815,7 @@ void UAttributeEditorTool::ResetLightmapUVsChannels()
 
 
 
-void UAttributeEditorTool::UpdateAutoGeneratedLightmapUVChannel(TUniquePtr<FPrimitiveComponentTarget>& Target, int32 NewMaxUVChannels)
+void UAttributeEditorTool::UpdateAutoGeneratedLightmapUVChannel(IPrimitiveComponentBackedTarget* Target, int32 NewMaxUVChannels)
 {
 	UStaticMeshComponent* StaticMeshComponent = Cast<UStaticMeshComponent>(Target->GetOwnerComponent());
 	if (StaticMeshComponent && StaticMeshComponent->GetStaticMesh() != nullptr)
diff --git a/Engine/Plugins/Experimental/MeshModelingToolset/Source/MeshModelingToolsEditorOnly/Private/MergeMeshesTool.cpp b/Engine/Plugins/Experimental/MeshModelingToolset/Source/MeshModelingToolsEditorOnly/Private/MergeMeshesTool.cpp
--- a/Engine/Plugins/Experimental/MeshModelingToolset/Source/MeshModelingToolsEditorOnly/Private/MergeMeshesTool.cpp
+++ b/Engine/Plugins/Experimental/MeshModelingToolset/Source/MeshModelingToolsEditorOnly/Private/MergeMeshesTool.cpp
@@ -12,6 +12,10 @@
 
 #include "CompositionOps/VoxelMergeMeshesOp.h"
 
+#include "TargetInterfaces/MeshDescriptionProvider.h"
+#include "TargetInterfaces/PrimitiveComponentBackedTarget.h"
+#include "ToolTargetManager.h"
+
 #include "ExplicitUseGeometryMathTypes.h"		// using UE::Geometry::(math types)
 using namespace UE::Geometry;
 
@@ -22,11 +26,20 @@ using namespace UE::Geometry;
  */
 
 
+const FToolTargetTypeRequirements& UMergeMeshesToolBuilder::GetTargetRequirements() const
+{
+	static FToolTargetTypeRequirements TypeRequirements({
+		UMeshDescriptionProvider::StaticClass(),
+		UPrimitiveComponentBackedTarget::StaticClass()
+		});
+	return TypeRequirements;
+}
+
 bool UMergeMeshesToolBuilder::CanBuildTool(const FToolBuilderState& SceneState) const
 {
 	const bool bHasBuildAPI = (this->AssetAPI != nullptr);
 	const int32 MinRequiredComponents = 1;
-	const bool bHasComponents = ToolBuilderUtil::CountComponents(SceneState, CanMakeComponentTarget) >= MinRequiredComponents;
+	const bool bHasComponents = SceneState.TargetManager->CountSelectedAndTargetable(SceneState, GetTargetRequirements()) >= MinRequiredComponents;
 	return ( bHasBuildAPI && bHasComponents );
 }
 
@@ -34,20 +47,8 @@ UInteractiveTool* UMergeMeshesToolBuilder::BuildTool(const FToolBuilderState& Sc
 {
 	UMergeMeshesTool* NewTool = NewObject<UMergeMeshesTool>(SceneState.ToolManager);
 
-	TArray<UActorComponent*> Components = ToolBuilderUtil::FindAllComponents(SceneState, CanMakeComponentTarget);
-	check(Components.Num() > 0);
-
-	TArray<TUniquePtr<FPrimitiveComponentTarget>> ComponentTargets;
-	for (UActorComponent* ActorComponent : Components)
-	{
-		auto* MeshComponent = Cast<UPrimitiveComponent>(ActorComponent);
-		if ( MeshComponent )
-		{
-			ComponentTargets.Add(MakeComponentTarget(MeshComponent));
-		}
-	}
-
-	NewTool->SetSelection(MoveTemp(ComponentTargets));
+	TArray<TObjectPtr<UToolTarget>> Targets = SceneState.TargetManager->BuildAllSelectedTargetable(SceneState, GetTargetRequirements());
+	NewTool->SetTargets(MoveTemp(Targets));
 	NewTool->SetWorld(SceneState.World);
 	NewTool->SetAssetAPI(AssetAPI);
 
@@ -93,9 +94,9 @@ void UMergeMeshesTool::Setup()
 	AddToolPropertySource(HandleSourcesProperties);
 
 	// Hide the source meshes
-	for (auto& ComponentTarget : ComponentTargets)
+	for (int32 ComponentIdx = 0; ComponentIdx < Targets.Num(); ComponentIdx++)
 	{
-		ComponentTarget->SetOwnerVisibility(false);
+		TargetComponentInterface(ComponentIdx)->SetOwnerVisibility(false);
 	}
 
 	// save transformed version of input meshes (maybe this could happen in the Operator?)
@@ -133,9 +134,9 @@ void UMergeMeshesTool::Shutdown(EToolShutdownType ShutdownType)
 
 	FDynamicMeshOpResult Result = Preview->Shutdown();
 	// Restore (unhide) the source meshes
-	for (auto& ComponentTarget : ComponentTargets)
+	for (int32 ComponentIdx = 0; ComponentIdx < Targets.Num(); ComponentIdx++)
 	{
-		ComponentTarget->SetOwnerVisibility(true);
+		TargetComponentInterface(ComponentIdx)->SetOwnerVisibility(true);
 	}
 	if (ShutdownType == EToolShutdownType::Accept)
 	{
@@ -145,9 +146,9 @@ void UMergeMeshesTool::Shutdown(EToolShutdownType ShutdownType)
 		GenerateAsset(Result);
 
 		TArray<AActor*> Actors;
-		for (auto& ComponentTarget : ComponentTargets)
+		for (int32 ComponentIdx = 0; ComponentIdx < Targets.Num(); ComponentIdx++)
 		{
-			Actors.Add(ComponentTarget->GetOwnerActor());
+			Actors.Add(TargetComponentInterface(ComponentIdx)->GetOwnerActor());
 		}
 		HandleSourcesProperties->ApplyMethod(Actors, GetToolManager());
 
@@ -193,11 +194,11 @@ void UMergeMeshesTool::CacheInputMeshes()
 	InputMeshes = MakeShared<TArray<IVoxelBasedCSG::FPlacedMesh>>();
 
 	// Package the selected meshes and transforms for consumption by the CSGTool
-	for (auto& ComponentTarget : ComponentTargets)
+	for (int32 ComponentIdx = 0; ComponentIdx < Targets.Num(); ComponentIdx++)
 	{
 		IVoxelBasedCSG::FPlacedMesh PlacedMesh;
-		PlacedMesh.Mesh = ComponentTarget->GetMesh();
-		PlacedMesh.Transform = ComponentTarget->GetWorldTransform();
+		PlacedMesh.Mesh = TargetMeshProviderInterface(ComponentIdx)->GetMeshDescription();
+		PlacedMesh.Transform = TargetComponentInterface(ComponentIdx)->GetWorldTransform();
 		InputMeshes->Add(PlacedMesh);
 	}
 }
diff --git a/Engine/Plugins/Experimental/MeshModelingToolset/Source/MeshModelingToolsEditorOnly/Private/MeshToVolumeTool.cpp b/Engine/Plugins/Experimental/MeshModelingToolset/Source/MeshModelingToolsEditorOnly/Private/MeshToVolumeTool.cpp
--- a/Engine/Plugins/Experimental/MeshModelingToolset/Source/MeshModelingToolsEditorOnly/Private/MeshToVolumeTool.cpp
+++ b/Engine/Plugins/Experimental/MeshModelingToolset/Source/MeshModelingToolsEditorOnly/Private/MeshToVolumeTool.cpp
@@ -83,7 +83,7 @@ void UMeshToVolumeTool::Setup()
 
 	FComponentMaterialSet MaterialSet;
 	check(Cast<IMaterialProvider>(Target));
-	Cast<IMaterialProvider>(Target)->GetMaterialSet(MaterialSet, false);
+	Cast<IMaterialProvider>(Target)->GetMaterialSet(MaterialSet);
 	PreviewMesh->SetMaterials(MaterialSet.Materials);
 
 	PreviewMesh->SetTangentsMode(EDynamicMeshTangentCalcType::ExternallyCalculated);
diff --git a/Engine/Plugins/Experimental/MeshModelingToolset/Source/MeshModelingToolsEditorOnly/Private/UVLayoutTool.cpp b/Engine/Plugins/Experimental/MeshModelingToolset/Source/MeshModelingToolsEditorOnly/Private/UVLayoutTool.cpp
--- a/Engine/Plugins/Experimental/MeshModelingToolset/Source/MeshModelingToolsEditorOnly/Private/UVLayoutTool.cpp
+++ b/Engine/Plugins/Experimental/MeshModelingToolset/Source/MeshModelingToolsEditorOnly/Private/UVLayoutTool.cpp
@@ -18,6 +18,12 @@
 
 #include "AssetGenerationUtil.h"
 
+#include "TargetInterfaces/MaterialProvider.h"
+#include "TargetInterfaces/MeshDescriptionCommitter.h"
+#include "TargetInterfaces/MeshDescriptionProvider.h"
+#include "TargetInterfaces/PrimitiveComponentBackedTarget.h"
+#include "ToolTargetManager.h"
+
 #include "ExplicitUseGeometryMathTypes.h"		// using UE::Geometry::(math types)
 using namespace UE::Geometry;
 
@@ -28,29 +34,28 @@ using namespace UE::Geometry;
  */
 
 
+const FToolTargetTypeRequirements& UUVLayoutToolBuilder::GetTargetRequirements() const
+{
+	static FToolTargetTypeRequirements TypeRequirements({
+		UMaterialProvider::StaticClass(),
+		UMeshDescriptionCommitter::StaticClass(),
+		UMeshDescriptionProvider::StaticClass(),
+		UPrimitiveComponentBackedTarget::StaticClass()
+		});
+	return TypeRequirements;
+}
+
 bool UUVLayoutToolBuilder::CanBuildTool(const FToolBuilderState& SceneState) const
 {
-	return ToolBuilderUtil::CountComponents(SceneState, CanMakeComponentTarget) >= 1;
+	return SceneState.TargetManager->CountSelectedAndTargetable(SceneState, GetTargetRequirements()) >= 1;
 }
 
 UInteractiveTool* UUVLayoutToolBuilder::BuildTool(const FToolBuilderState& SceneState) const
 {
 	UUVLayoutTool* NewTool = NewObject<UUVLayoutTool>(SceneState.ToolManager);
 
-	TArray<UActorComponent*> Components = ToolBuilderUtil::FindAllComponents(SceneState, CanMakeComponentTarget);
-	check(Components.Num() > 0);
-
-	TArray<TUniquePtr<FPrimitiveComponentTarget>> ComponentTargets;
-	for (UActorComponent* ActorComponent : Components)
-	{
-		auto* MeshComponent = Cast<UPrimitiveComponent>(ActorComponent);
-		if ( MeshComponent )
-		{
-			ComponentTargets.Add(MakeComponentTarget(MeshComponent));
-		}
-	}
-
-	NewTool->SetSelection(MoveTemp(ComponentTargets));
+	TArray<TObjectPtr<UToolTarget>> Targets = SceneState.TargetManager->BuildAllSelectedTargetable(SceneState, GetTargetRequirements());
+	NewTool->SetTargets(MoveTemp(Targets));
 	NewTool->SetWorld(SceneState.World, SceneState.GizmoManager);
 	NewTool->SetAssetAPI(AssetAPI);
 
@@ -83,17 +88,17 @@ void UUVLayoutTool::Setup()
 	UInteractiveTool::Setup();
 
 	// hide input StaticMeshComponent
-	for (TUniquePtr<FPrimitiveComponentTarget>& ComponentTarget : ComponentTargets)
+	for (int32 ComponentIdx = 0; ComponentIdx < Targets.Num(); ComponentIdx++)
 	{
-		ComponentTarget->SetOwnerVisibility(false);
+		TargetComponentInterface(ComponentIdx)->SetOwnerVisibility(false);
 	}
 
 	// if we only have one object, add ability to set UV channel
-	if (ComponentTargets.Num() == 1)
+	if (Targets.Num() == 1)
 	{
 		UVChannelProperties = NewObject<UMeshUVChannelProperties>(this);
 		UVChannelProperties->RestoreProperties(this);
-		UVChannelProperties->Initialize(ComponentTargets[0]->GetMesh(), false);
+		UVChannelProperties->Initialize(TargetMeshProviderInterface(0)->GetMeshDescription(), false);
 		UVChannelProperties->ValidateSelection(true);
 		AddToolPropertySource(UVChannelProperties);
 		UVChannelProperties->WatchProperty(UVChannelProperties->UVChannel, [this](const FString& NewValue)
@@ -111,18 +116,18 @@ void UUVLayoutTool::Setup()
 	AddToolPropertySource(MaterialSettings);
 
 	// if we only have one object, add optional UV layout view
-	if (ComponentTargets.Num() == 1)
+	if (Targets.Num() == 1)
 	{
 		UVLayoutView = NewObject<UUVLayoutPreview>(this);
 		UVLayoutView->CreateInWorld(TargetWorld);
 
 		FComponentMaterialSet MaterialSet;
-		ComponentTargets[0]->GetMaterialSet(MaterialSet);
+		TargetMaterialInterface(0)->GetMaterialSet(MaterialSet);
 		UVLayoutView->SetSourceMaterials(MaterialSet);
 
 		UVLayoutView->SetSourceWorldPosition(
-			ComponentTargets[0]->GetOwnerActor()->GetTransform(),
-			ComponentTargets[0]->GetOwnerActor()->GetComponentsBoundingBox());
+			TargetComponentInterface(0)->GetOwnerActor()->GetTransform(),
+			TargetComponentInterface(0)->GetOwnerActor()->GetComponentsBoundingBox());
 
 		UVLayoutView->Settings->RestoreProperties(this);
 		AddToolPropertySource(UVLayoutView->Settings);
@@ -139,7 +144,7 @@ void UUVLayoutTool::Setup()
 void UUVLayoutTool::UpdateNumPreviews()
 {
 	int32 CurrentNumPreview = Previews.Num();
-	int32 TargetNumPreview = ComponentTargets.Num();
+	int32 TargetNumPreview = Targets.Num();
 	if (TargetNumPreview < CurrentNumPreview)
 	{
 		for (int32 PreviewIdx = CurrentNumPreview - 1; PreviewIdx >= TargetNumPreview; PreviewIdx--)
@@ -159,18 +164,18 @@ void UUVLayoutTool::UpdateNumPreviews()
 			OpFactory->ComponentIndex = PreviewIdx;
 			OriginalDynamicMeshes[PreviewIdx] = MakeShared<FDynamicMesh3, ESPMode::ThreadSafe>();
 			FMeshDescriptionToDynamicMesh Converter;
-			Converter.Convert(ComponentTargets[PreviewIdx]->GetMesh(), *OriginalDynamicMeshes[PreviewIdx]);
+			Converter.Convert(TargetMeshProviderInterface(PreviewIdx)->GetMeshDescription(), *OriginalDynamicMeshes[PreviewIdx]);
 
 			UMeshOpPreviewWithBackgroundCompute* Preview = Previews.Add_GetRef(NewObject<UMeshOpPreviewWithBackgroundCompute>(OpFactory, "Preview"));
 			Preview->Setup(this->TargetWorld, OpFactory);
 
 			FComponentMaterialSet MaterialSet;
-			ComponentTargets[PreviewIdx]->GetMaterialSet(MaterialSet);
+			TargetMaterialInterface(PreviewIdx)->GetMaterialSet(MaterialSet);
 			Preview->ConfigureMaterials(MaterialSet.Materials,
 				ToolSetupUtil::GetDefaultWorkingMaterial(GetToolManager())
 			);
 			Preview->PreviewMesh->UpdatePreview(OriginalDynamicMeshes[PreviewIdx].Get());
-			Preview->PreviewMesh->SetTransform(ComponentTargets[PreviewIdx]->GetWorldTransform());
+			Preview->PreviewMesh->SetTransform(TargetComponentInterface(PreviewIdx)->GetWorldTransform());
 
 			Preview->OnMeshUpdated.AddLambda([this](UMeshOpPreviewWithBackgroundCompute* Compute)
 			{
@@ -195,9 +200,9 @@ void UUVLayoutTool::Shutdown(EToolShutdownType ShutdownType)
 	MaterialSettings->SaveProperties(this);
 
 	// Restore (unhide) the source meshes
-	for (TUniquePtr<FPrimitiveComponentTarget>& ComponentTarget : ComponentTargets)
+	for (int32 ComponentIdx = 0; ComponentIdx < Targets.Num(); ComponentIdx++)
 	{
-		ComponentTarget->SetOwnerVisibility(true);
+		TargetComponentInterface(ComponentIdx)->SetOwnerVisibility(true);
 	}
 
 	TArray<FDynamicMeshOpResult> Results;
@@ -220,7 +225,7 @@ TUniquePtr<FDynamicMeshOperator> UUVLayoutOperatorFactory::MakeNewOperator()
 {
 	TUniquePtr<FUVLayoutOp> Op = MakeUnique<FUVLayoutOp>();
 
-	FTransform LocalToWorld = Tool->ComponentTargets[ComponentIndex]->GetWorldTransform();
+	FTransform LocalToWorld = Tool->TargetComponentInterface(ComponentIndex)->GetWorldTransform();
 	Op->OriginalMesh = Tool->OriginalDynamicMeshes[ComponentIndex];
 
 	switch (Tool->BasicProperties->LayoutType)
@@ -251,7 +256,7 @@ TUniquePtr<FDynamicMeshOperator> UUVLayoutOperatorFactory::MakeNewOperator()
 
 int32 UUVLayoutTool::GetSelectedUVChannel() const
 {
-	return UVChannelProperties->GetSelectedChannelIndex(true);
+	return UVChannelProperties ? UVChannelProperties->GetSelectedChannelIndex(true) : 0;
 }
 
 
@@ -349,15 +354,15 @@ void UUVLayoutTool::GenerateAsset(const TArray<FDynamicMeshOpResult>& Results)
 {
 	GetToolManager()->BeginUndoTransaction(LOCTEXT("UVLayoutToolTransactionName", "UV Layout Tool"));
 
-	check(Results.Num() == ComponentTargets.Num());
+	check(Results.Num() == Targets.Num());
 	
-	for (int32 ComponentIdx = 0; ComponentIdx < ComponentTargets.Num(); ComponentIdx++)
+	for (int32 ComponentIdx = 0; ComponentIdx < Targets.Num(); ComponentIdx++)
 	{
 		check(Results[ComponentIdx].Mesh.Get() != nullptr);
-		ComponentTargets[ComponentIdx]->CommitMesh([&Results, &ComponentIdx, this](const FPrimitiveComponentTarget::FCommitParams& CommitParams)
+		TargetMeshCommitterInterface(ComponentIdx)->CommitMeshDescription([&Results, &ComponentIdx, this](const IMeshDescriptionCommitter::FCommitterParams& CommitParams)
 		{
 			FDynamicMesh3* DynamicMesh = Results[ComponentIdx].Mesh.Get();
-			FMeshDescription* MeshDescription = CommitParams.MeshDescription;
+			FMeshDescription* MeshDescription = CommitParams.MeshDescriptionOut;
 			
 			bool bVerticesOnly = false;
 			bool bAttributesOnly = true;
diff --git a/Engine/Plugins/Experimental/MeshModelingToolset/Source/MeshModelingToolsEditorOnly/Private/VoxelCSGMeshesTool.cpp b/Engine/Plugins/Experimental/MeshModelingToolset/Source/MeshModelingToolsEditorOnly/Private/VoxelCSGMeshesTool.cpp
--- a/Engine/Plugins/Experimental/MeshModelingToolset/Source/MeshModelingToolsEditorOnly/Private/VoxelCSGMeshesTool.cpp
+++ b/Engine/Plugins/Experimental/MeshModelingToolset/Source/MeshModelingToolsEditorOnly/Private/VoxelCSGMeshesTool.cpp
@@ -12,6 +12,11 @@
 
 #include "CompositionOps/VoxelBooleanMeshesOp.h"
 
+#include "TargetInterfaces/MeshDescriptionCommitter.h"
+#include "TargetInterfaces/MeshDescriptionProvider.h"
+#include "TargetInterfaces/PrimitiveComponentBackedTarget.h"
+#include "ToolTargetManager.h"
+
 #include "ExplicitUseGeometryMathTypes.h"		// using UE::Geometry::(math types)
 using namespace UE::Geometry;
 
@@ -21,10 +26,20 @@ using namespace UE::Geometry;
 /*
  * ToolBuilder
  */
+const FToolTargetTypeRequirements& UVoxelCSGMeshesToolBuilder::GetTargetRequirements() const
+{
+	static FToolTargetTypeRequirements TypeRequirements({
+		UMeshDescriptionCommitter::StaticClass(),
+		UMeshDescriptionProvider::StaticClass(),
+		UPrimitiveComponentBackedTarget::StaticClass()
+		});
+	return TypeRequirements;
+}
+
 bool UVoxelCSGMeshesToolBuilder::CanBuildTool(const FToolBuilderState& SceneState) const
 {
 	bool bHasBuildAPI = (this->AssetAPI != nullptr);
-	bool bHasComponents = ToolBuilderUtil::CountComponents(SceneState, CanMakeComponentTarget) == 2;
+	bool bHasComponents = SceneState.TargetManager->CountSelectedAndTargetable(SceneState, GetTargetRequirements()) == 2;
 	return (bHasBuildAPI && bHasComponents);
 }
 
@@ -32,20 +47,8 @@ UInteractiveTool* UVoxelCSGMeshesToolBuilder::BuildTool(const FToolBuilderState&
 {
 	UVoxelCSGMeshesTool* NewTool = NewObject<UVoxelCSGMeshesTool>(SceneState.ToolManager);
 
-	TArray<UActorComponent*> Components = ToolBuilderUtil::FindAllComponents(SceneState, CanMakeComponentTarget);
-	check(Components.Num() > 0);
-
-	TArray<TUniquePtr<FPrimitiveComponentTarget>> ComponentTargets;
-	for (UActorComponent* ActorComponent : Components)
-	{
-		auto* MeshComponent = Cast<UPrimitiveComponent>(ActorComponent);
-		if ( MeshComponent )
-		{
-			ComponentTargets.Add(MakeComponentTarget(MeshComponent));
-		}
-	}
-
-	NewTool->SetSelection(MoveTemp(ComponentTargets));
+	TArray<TObjectPtr<UToolTarget>> Targets = SceneState.TargetManager->BuildAllSelectedTargetable(SceneState, GetTargetRequirements());
+	NewTool->SetTargets(MoveTemp(Targets));
 	NewTool->SetWorld(SceneState.World);
 	NewTool->SetAssetAPI(AssetAPI);
 
@@ -89,9 +92,9 @@ void UVoxelCSGMeshesTool::Setup()
 
 
 	// Hide the source meshes
-	for (auto& ComponentTarget : ComponentTargets)
+	for (int32 ComponentIdx = 0; ComponentIdx < Targets.Num(); ComponentIdx++)
 	{
-		ComponentTarget->SetOwnerVisibility(false);
+		TargetComponentInterface(ComponentIdx)->SetOwnerVisibility(false);
 	}
 
 	// save transformed version of input meshes (maybe this could happen in the Operator?)
@@ -128,9 +131,9 @@ void UVoxelCSGMeshesTool::Shutdown(EToolShutdownType ShutdownType)
 
 	FDynamicMeshOpResult Result = Preview->Shutdown();
 	// Restore (unhide) the source meshes
-	for (auto& ComponentTarget : ComponentTargets)
+	for (int32 ComponentIdx = 0; ComponentIdx < Targets.Num(); ComponentIdx++)
 	{
-		ComponentTarget->SetOwnerVisibility(true);
+		TargetComponentInterface(ComponentIdx)->SetOwnerVisibility(true);
 	}
 	if (ShutdownType == EToolShutdownType::Accept)
 	{
@@ -140,9 +143,9 @@ void UVoxelCSGMeshesTool::Shutdown(EToolShutdownType ShutdownType)
 		GenerateAsset(Result);
 
 		TArray<AActor*> Actors;
-		for (auto& ComponentTarget : ComponentTargets)
+		for (int32 ComponentIdx = 0; ComponentIdx < Targets.Num(); ComponentIdx++)
 		{
-			Actors.Add(ComponentTarget->GetOwnerActor());
+			Actors.Add(TargetComponentInterface(ComponentIdx)->GetOwnerActor());
 		}
 		HandleSourcesProperties->ApplyMethod(Actors, GetToolManager());
 
@@ -188,11 +191,11 @@ void UVoxelCSGMeshesTool::CacheInputMeshes()
 	InputMeshes = MakeShared<TArray<IVoxelBasedCSG::FPlacedMesh>>();
 
 	// Package the selected meshes and transforms for consumption by the CSGTool
-	for (auto& ComponentTarget : ComponentTargets)
+	for (int32 ComponentIdx = 0; ComponentIdx < Targets.Num(); ComponentIdx++)
 	{
 		IVoxelBasedCSG::FPlacedMesh PlacedMesh;
-		PlacedMesh.Mesh = ComponentTarget->GetMesh();
-		PlacedMesh.Transform = ComponentTarget->GetWorldTransform();
+		PlacedMesh.Mesh = TargetMeshProviderInterface(ComponentIdx)->GetMeshDescription();
+		PlacedMesh.Transform = TargetComponentInterface(ComponentIdx)->GetWorldTransform();
 		InputMeshes->Add(PlacedMesh);
 	}
 }
diff --git a/Engine/Plugins/Experimental/MeshModelingToolset/Source/MeshModelingToolsEditorOnly/Public/AttributeEditorTool.h b/Engine/Plugins/Experimental/MeshModelingToolset/Source/MeshModelingToolsEditorOnly/Public/AttributeEditorTool.h
--- a/Engine/Plugins/Experimental/MeshModelingToolset/Source/MeshModelingToolsEditorOnly/Public/AttributeEditorTool.h
+++ b/Engine/Plugins/Experimental/MeshModelingToolset/Source/MeshModelingToolsEditorOnly/Public/AttributeEditorTool.h
@@ -22,6 +22,9 @@ public:
 
 	virtual bool CanBuildTool(const FToolBuilderState& SceneState) const override;
 	virtual UInteractiveTool* BuildTool(const FToolBuilderState& SceneState) const override;
+
+protected:
+	virtual const FToolTargetTypeRequirements& GetTargetRequirements() const override;
 };
 
 
@@ -404,7 +407,7 @@ protected:
 	bool bHaveAutoGeneratedLightmapUVSet = false;
 	void InitializeAttributeLists();
 
-	void UpdateAutoGeneratedLightmapUVChannel(TUniquePtr<FPrimitiveComponentTarget>& Target, int32 NewMaxUVChannels);
+	void UpdateAutoGeneratedLightmapUVChannel(IPrimitiveComponentBackedTarget* Target, int32 NewMaxUVChannels);
 
 	EAttributeEditorToolActions PendingAction = EAttributeEditorToolActions::NoAction;
 	void OptimizeForEditing();
diff --git a/Engine/Plugins/Experimental/MeshModelingToolset/Source/MeshModelingToolsEditorOnly/Public/MergeMeshesTool.h b/Engine/Plugins/Experimental/MeshModelingToolset/Source/MeshModelingToolsEditorOnly/Public/MergeMeshesTool.h
--- a/Engine/Plugins/Experimental/MeshModelingToolset/Source/MeshModelingToolsEditorOnly/Public/MergeMeshesTool.h
+++ b/Engine/Plugins/Experimental/MeshModelingToolset/Source/MeshModelingToolsEditorOnly/Public/MergeMeshesTool.h
@@ -32,6 +32,9 @@ public:
 
 	virtual bool CanBuildTool(const FToolBuilderState& SceneState) const override;
 	virtual UInteractiveTool* BuildTool(const FToolBuilderState& SceneState) const override;
+
+protected:
+	virtual const FToolTargetTypeRequirements& GetTargetRequirements() const override;
 };
 
 
diff --git a/Engine/Plugins/Experimental/MeshModelingToolset/Source/MeshModelingToolsEditorOnly/Public/UVLayoutTool.h b/Engine/Plugins/Experimental/MeshModelingToolset/Source/MeshModelingToolsEditorOnly/Public/UVLayoutTool.h
--- a/Engine/Plugins/Experimental/MeshModelingToolset/Source/MeshModelingToolsEditorOnly/Public/UVLayoutTool.h
+++ b/Engine/Plugins/Experimental/MeshModelingToolset/Source/MeshModelingToolsEditorOnly/Public/UVLayoutTool.h
@@ -36,6 +36,9 @@ public:
 
 	virtual bool CanBuildTool(const FToolBuilderState& SceneState) const override;
 	virtual UInteractiveTool* BuildTool(const FToolBuilderState& SceneState) const override;
+
+protected:
+	virtual const FToolTargetTypeRequirements& GetTargetRequirements() const override;
 };
 
 
diff --git a/Engine/Plugins/Experimental/MeshModelingToolset/Source/MeshModelingToolsEditorOnly/Public/VoxelCSGMeshesTool.h b/Engine/Plugins/Experimental/MeshModelingToolset/Source/MeshModelingToolsEditorOnly/Public/VoxelCSGMeshesTool.h
--- a/Engine/Plugins/Experimental/MeshModelingToolset/Source/MeshModelingToolsEditorOnly/Public/VoxelCSGMeshesTool.h
+++ b/Engine/Plugins/Experimental/MeshModelingToolset/Source/MeshModelingToolsEditorOnly/Public/VoxelCSGMeshesTool.h
@@ -34,6 +34,9 @@ public:
 
 	virtual bool CanBuildTool(const FToolBuilderState& SceneState) const override;
 	virtual UInteractiveTool* BuildTool(const FToolBuilderState& SceneState) const override;
+
+protected:
+	virtual const FToolTargetTypeRequirements& GetTargetRequirements() const override;
 };
 
 
diff --git a/Engine/Plugins/Experimental/MeshModelingToolset/Source/ModelingComponents/Private/BaseTools/BaseCreateFromSelectedTool.cpp b/Engine/Plugins/Experimental/MeshModelingToolset/Source/ModelingComponents/Private/BaseTools/BaseCreateFromSelectedTool.cpp
--- a/Engine/Plugins/Experimental/MeshModelingToolset/Source/ModelingComponents/Private/BaseTools/BaseCreateFromSelectedTool.cpp
+++ b/Engine/Plugins/Experimental/MeshModelingToolset/Source/ModelingComponents/Private/BaseTools/BaseCreateFromSelectedTool.cpp
@@ -15,6 +15,11 @@
 #include "AssetGenerationUtil.h"
 #include "Selection/ToolSelectionUtil.h"
 
+#include "TargetInterfaces/MeshDescriptionCommitter.h"
+#include "TargetInterfaces/MeshDescriptionProvider.h"
+#include "TargetInterfaces/PrimitiveComponentBackedTarget.h"
+#include "TargetInterfaces/MaterialProvider.h"
+#include "ToolTargetManager.h"
 
 #define LOCTEXT_NAMESPACE "UBaseCreateFromSelectedTool"
 
@@ -24,10 +29,20 @@ using namespace UE::Geometry;
  * ToolBuilder
  */
 
+const FToolTargetTypeRequirements& UBaseCreateFromSelectedToolBuilder::GetTargetRequirements() const
+{
+	static FToolTargetTypeRequirements TypeRequirements({
+		UMeshDescriptionCommitter::StaticClass(),
+		UMeshDescriptionProvider::StaticClass(),
+		UPrimitiveComponentBackedTarget::StaticClass(),
+		UMaterialProvider::StaticClass()
+		});
+	return TypeRequirements;
+}
 
 bool UBaseCreateFromSelectedToolBuilder::CanBuildTool(const FToolBuilderState& SceneState) const
 {
-	int32 ComponentCount = ToolBuilderUtil::CountComponents(SceneState, CanMakeComponentTarget);
+	int32 ComponentCount = SceneState.TargetManager->CountSelectedAndTargetable(SceneState, GetTargetRequirements());
 	return AssetAPI != nullptr && ComponentCount >= MinComponentsSupported() && (!MaxComponentsSupported().IsSet() || ComponentCount <= MaxComponentsSupported().GetValue());
 }
 
@@ -36,20 +51,8 @@ UInteractiveTool* UBaseCreateFromSelectedToolBuilder::BuildTool(const FToolBuild
 {
 	UBaseCreateFromSelectedTool* NewTool = MakeNewToolInstance(SceneState.ToolManager);
 
-	TArray<UActorComponent*> Components = ToolBuilderUtil::FindAllComponents(SceneState, CanMakeComponentTarget);
-	check(Components.Num() > 0);
-
-	TArray<TUniquePtr<FPrimitiveComponentTarget>> ComponentTargets;
-	for (UActorComponent* ActorComponent : Components)
-	{
-		auto* MeshComponent = Cast<UPrimitiveComponent>(ActorComponent);
-		if (MeshComponent)
-		{
-			ComponentTargets.Add(MakeComponentTarget(MeshComponent));
-		}
-	}
-
-	NewTool->SetSelection(MoveTemp(ComponentTargets));
+	TArray<TObjectPtr<UToolTarget>> Targets = SceneState.TargetManager->BuildAllSelectedTargetable(SceneState, GetTargetRequirements());
+	NewTool->SetTargets(MoveTemp(Targets));
 	NewTool->SetWorld(SceneState.World);
 	NewTool->SetAssetAPI(AssetAPI);
 
@@ -68,9 +71,9 @@ void UBaseCreateFromSelectedTool::Setup()
 	UInteractiveTool::Setup();
 
 	// hide input StaticMeshComponents
-	for (TUniquePtr<FPrimitiveComponentTarget>& ComponentTarget : ComponentTargets)
+	for (int32 ComponentIdx = 0; ComponentIdx < Targets.Num(); ComponentIdx++)
 	{
-		ComponentTarget->SetOwnerVisibility(false);
+		TargetComponentInterface(ComponentIdx)->SetOwnerVisibility(false);
 	}
 
 	// initialize our properties
@@ -110,11 +113,11 @@ void UBaseCreateFromSelectedTool::Setup()
 		}
 		else
 		{
-			int32 Index = (HandleSourcesProperties->WriteOutputTo == EBaseCreateFromSelectedTargetType::FirstInputAsset) ? 0 : ComponentTargets.Num() - 1;
-			HandleSourcesProperties->OutputAsset = AssetGenerationUtil::GetComponentAssetBaseName(ComponentTargets[Index]->GetOwnerComponent(), false);
+			int32 Index = (HandleSourcesProperties->WriteOutputTo == EBaseCreateFromSelectedTargetType::FirstInputAsset) ? 0 : Targets.Num() - 1;
+			HandleSourcesProperties->OutputAsset = AssetGenerationUtil::GetComponentAssetBaseName(TargetComponentInterface(Index)->GetOwnerComponent(), false);
 
 			// Reset the hidden gizmo to its initial position
-			FTransform ComponentTransform = ComponentTargets[Index]->GetWorldTransform();
+			FTransform ComponentTransform = TargetComponentInterface(Index)->GetWorldTransform();
 			TransformGizmos[Index]->SetNewGizmoTransform(ComponentTransform, true);
 			UpdateGizmoVisibility();
 		}
@@ -132,7 +135,7 @@ int32 UBaseCreateFromSelectedTool::GetHiddenGizmoIndex() const
 	}
 	else
 	{
-		return (HandleSourcesProperties->WriteOutputTo == EBaseCreateFromSelectedTargetType::FirstInputAsset) ? 0 : ComponentTargets.Num() - 1;
+		return (HandleSourcesProperties->WriteOutputTo == EBaseCreateFromSelectedTargetType::FirstInputAsset) ? 0 : Targets.Num() - 1;
 	}
 }
 
@@ -176,11 +179,11 @@ void UBaseCreateFromSelectedTool::SetTransformGizmos()
 {
 	UInteractiveGizmoManager* GizmoManager = GetToolManager()->GetPairedGizmoManager();
 
-	for (int ComponentIdx = 0; ComponentIdx < ComponentTargets.Num(); ComponentIdx++)
+	for (int ComponentIdx = 0; ComponentIdx < Targets.Num(); ComponentIdx++)
 	{
 		UTransformProxy* Proxy = TransformProxies.Add_GetRef(NewObject<UTransformProxy>(this));
 		UTransformGizmo* Gizmo = TransformGizmos.Add_GetRef(GizmoManager->Create3AxisTransformGizmo(this));
-		Proxy->SetTransform(ComponentTargets[ComponentIdx]->GetWorldTransform());
+		Proxy->SetTransform(TargetComponentInterface(ComponentIdx)->GetWorldTransform());
 		Gizmo->SetActiveTarget(Proxy, GetToolManager());
 		Proxy->OnTransformChanged.AddUObject(this, &UBaseCreateFromSelectedTool::TransformChanged);
 	}
@@ -211,11 +214,11 @@ void UBaseCreateFromSelectedTool::GenerateAsset(const FDynamicMeshOpResult& Resu
 	check(Result.Mesh.Get() != nullptr);
 
 	FTransform3d NewTransform;
-	if (ComponentTargets.Num() == 1) // in the single-selection case, shove the result back into the original component space
+	if (Targets.Num() == 1) // in the single-selection case, shove the result back into the original component space
 	{
-		FTransform3d ToSourceComponentSpace = (FTransform3d)ComponentTargets[0]->GetWorldTransform().Inverse();
+		FTransform3d ToSourceComponentSpace = (FTransform3d)TargetComponentInterface(0)->GetWorldTransform().Inverse();
 		MeshTransforms::ApplyTransform(*Result.Mesh, ToSourceComponentSpace);
-		NewTransform = (FTransform3d)ComponentTargets[0]->GetWorldTransform();
+		NewTransform = (FTransform3d)TargetComponentInterface(0)->GetWorldTransform();
 	}
 	else // in the multi-selection case, center the pivot for the combined result
 	{
@@ -248,33 +251,37 @@ void UBaseCreateFromSelectedTool::GenerateAsset(const FDynamicMeshOpResult& Resu
 
 
 
-void UBaseCreateFromSelectedTool::UpdateAsset(const FDynamicMeshOpResult& Result, TUniquePtr<FPrimitiveComponentTarget>& UpdateTarget)
+void UBaseCreateFromSelectedTool::UpdateAsset(const FDynamicMeshOpResult& Result, UToolTarget* UpdateTarget)
 {
 	check(Result.Mesh.Get() != nullptr);
 
-	FTransform3d TargetToWorld = (FTransform3d)UpdateTarget->GetWorldTransform();
+	IPrimitiveComponentBackedTarget* TargetComponent = Cast<IPrimitiveComponentBackedTarget>(UpdateTarget);
+	IMeshDescriptionCommitter* TargetMeshCommitter = Cast<IMeshDescriptionCommitter>(UpdateTarget);
+	IMaterialProvider* TargetMaterial = Cast<IMaterialProvider>(UpdateTarget);
+
+	FTransform3d TargetToWorld = (FTransform3d)TargetComponent->GetWorldTransform();
 
 	FTransform3d ResultTransform = Result.Transform;
 	MeshTransforms::ApplyTransform(*Result.Mesh, ResultTransform);
 	MeshTransforms::ApplyTransformInverse(*Result.Mesh, TargetToWorld);
 
-	UpdateTarget->CommitMesh([&](const FPrimitiveComponentTarget::FCommitParams& CommitParams)
+	TargetMeshCommitter->CommitMeshDescription([&](const IMeshDescriptionCommitter::FCommitterParams& CommitParams)
 	{
 		FDynamicMeshToMeshDescription Converter;
-		Converter.Convert(Result.Mesh.Get(), *CommitParams.MeshDescription);
+		Converter.Convert(Result.Mesh.Get(), *CommitParams.MeshDescriptionOut);
 	});
 
 	FComponentMaterialSet MaterialSet;
 	MaterialSet.Materials = GetOutputMaterials();
-	UpdateTarget->CommitMaterialSetUpdate(MaterialSet, true);
+	TargetMaterial->CommitMaterialSetUpdate(MaterialSet, true);
 }
 
 
 FString UBaseCreateFromSelectedTool::PrefixWithSourceNameIfSingleSelection(const FString& AssetName) const
 {
-	if (ComponentTargets.Num() == 1)
+	if (Targets.Num() == 1)
 	{
-		FString CurName = AssetGenerationUtil::GetComponentAssetBaseName(ComponentTargets[0]->GetOwnerComponent());
+		FString CurName = AssetGenerationUtil::GetComponentAssetBaseName(TargetComponentInterface(0)->GetOwnerComponent());
 		return FString::Printf(TEXT("%s_%s"), *CurName, *AssetName);
 	}
 	else
@@ -312,9 +319,9 @@ void UBaseCreateFromSelectedTool::Shutdown(EToolShutdownType ShutdownType)
 
 	FDynamicMeshOpResult Result = Preview->Shutdown();
 	// Restore (unhide) the source meshes
-	for (auto& ComponentTarget : ComponentTargets)
+	for (int32 ComponentIdx = 0; ComponentIdx < Targets.Num(); ComponentIdx++)
 	{
-		ComponentTarget->SetOwnerVisibility(true);
+		TargetComponentInterface(ComponentIdx)->SetOwnerVisibility(true);
 	}
 
 	if (ShutdownType == EToolShutdownType::Accept)
@@ -329,20 +336,19 @@ void UBaseCreateFromSelectedTool::Shutdown(EToolShutdownType ShutdownType)
 		}
 		else
 		{
-			int32 TargetIndex = (HandleSourcesProperties->WriteOutputTo == EBaseCreateFromSelectedTargetType::FirstInputAsset) ? 0 : (ComponentTargets.Num() - 1);
-			TUniquePtr<FPrimitiveComponentTarget>& UpdateTarget = ComponentTargets[TargetIndex];
-			KeepActor = UpdateTarget->GetOwnerActor();
+			int32 TargetIndex = (HandleSourcesProperties->WriteOutputTo == EBaseCreateFromSelectedTargetType::FirstInputAsset) ? 0 : (Targets.Num() - 1);
+			KeepActor = TargetComponentInterface(TargetIndex)->GetOwnerActor();
 
-			UpdateAsset(Result, ComponentTargets[TargetIndex]);
+			UpdateAsset(Result, Targets[TargetIndex]);
 		}
 
 		TArray<AActor*> Actors;
-		for (auto& ComponentTarget : ComponentTargets)
+		for (int32 ComponentIdx = 0; ComponentIdx < Targets.Num(); ComponentIdx++)
 		{
-			AActor* Actor = ComponentTarget->GetOwnerActor();
+			AActor* Actor = TargetComponentInterface(ComponentIdx)->GetOwnerActor();
 			if (Actor != KeepActor)
 			{
-				Actors.Add(ComponentTarget->GetOwnerActor());
+				Actors.Add(Actor);
 			}
 		}
 		HandleSourcesProperties->ApplyMethod(Actors, GetToolManager());
diff --git a/Engine/Plugins/Experimental/MeshModelingToolset/Source/ModelingComponents/Private/BaseTools/BaseVoxelTool.cpp b/Engine/Plugins/Experimental/MeshModelingToolset/Source/ModelingComponents/Private/BaseTools/BaseVoxelTool.cpp
--- a/Engine/Plugins/Experimental/MeshModelingToolset/Source/ModelingComponents/Private/BaseTools/BaseVoxelTool.cpp
+++ b/Engine/Plugins/Experimental/MeshModelingToolset/Source/ModelingComponents/Private/BaseTools/BaseVoxelTool.cpp
@@ -40,13 +40,13 @@ TArray<UMaterialInterface*> UBaseVoxelTool::GetOutputMaterials() const
 
 void UBaseVoxelTool::ConvertInputsAndSetPreviewMaterials(bool bSetPreviewMesh)
 {
-	OriginalDynamicMeshes.SetNum(ComponentTargets.Num());
+	OriginalDynamicMeshes.SetNum(Targets.Num());
 
-	for (int ComponentIdx = 0; ComponentIdx < ComponentTargets.Num(); ComponentIdx++)
+	for (int ComponentIdx = 0; ComponentIdx < Targets.Num(); ComponentIdx++)
 	{
 		OriginalDynamicMeshes[ComponentIdx] = MakeShared<FDynamicMesh3, ESPMode::ThreadSafe>();
 		FMeshDescriptionToDynamicMesh Converter;
-		Converter.Convert(ComponentTargets[ComponentIdx]->GetMesh(), *OriginalDynamicMeshes[ComponentIdx]);
+		Converter.Convert(TargetMeshProviderInterface(ComponentIdx)->GetMeshDescription(), *OriginalDynamicMeshes[ComponentIdx]);
 	}
 
 	//if (bSetPreviewMesh)
diff --git a/Engine/Plugins/Experimental/MeshModelingToolset/Source/ModelingComponents/Public/BaseTools/BaseCreateFromSelectedTool.h b/Engine/Plugins/Experimental/MeshModelingToolset/Source/ModelingComponents/Public/BaseTools/BaseCreateFromSelectedTool.h
--- a/Engine/Plugins/Experimental/MeshModelingToolset/Source/ModelingComponents/Public/BaseTools/BaseCreateFromSelectedTool.h
+++ b/Engine/Plugins/Experimental/MeshModelingToolset/Source/ModelingComponents/Public/BaseTools/BaseCreateFromSelectedTool.h
@@ -33,6 +33,8 @@ public:
 	// subclass must override!
 	virtual UBaseCreateFromSelectedTool* MakeNewToolInstance(UObject* Outer) const { check(false); return nullptr; }
 
+protected:
+	virtual const FToolTargetTypeRequirements& GetTargetRequirements() const override;
 };
 
 
@@ -186,7 +188,7 @@ protected:
 	virtual void GenerateAsset(const FDynamicMeshOpResult& Result);
 
 	// Helper to generate assets when a result is accepted; typically does not need to be overloaded
-	virtual void UpdateAsset(const FDynamicMeshOpResult& Result, TUniquePtr<FPrimitiveComponentTarget>& Target);
+	virtual void UpdateAsset(const FDynamicMeshOpResult& Result, UToolTarget* Target);
 
 	// Which of the transform gizmos to hide, or -1 if all gizmos can be shown
 	virtual int32 GetHiddenGizmoIndex() const;
diff --git a/Engine/Plugins/Experimental/MotionTrailEditorMode/Source/MotionTrailEditorMode/Public/MotionTrailEditorToolset.h b/Engine/Plugins/Experimental/MotionTrailEditorMode/Source/MotionTrailEditorMode/Public/MotionTrailEditorToolset.h
--- a/Engine/Plugins/Experimental/MotionTrailEditorMode/Source/MotionTrailEditorMode/Public/MotionTrailEditorToolset.h
+++ b/Engine/Plugins/Experimental/MotionTrailEditorMode/Source/MotionTrailEditorMode/Public/MotionTrailEditorToolset.h
@@ -111,7 +111,6 @@ public:
 
 	UWorld* GetWorld() const { return TargetWorld; }
 	UInteractiveGizmoManager* GetGizmoManager() const { return GizmoManager; }
-	const TArray<TUniquePtr<FPrimitiveComponentTarget>>& GetSelection() const { return ComponentTargets; }
 
 
 	static FString TrailKeyTransformGizmoInstanceIdentifier;
diff --git a/Engine/Source/Runtime/Experimental/InteractiveToolsFramework/Private/MultiSelectionTool.cpp b/Engine/Source/Runtime/Experimental/InteractiveToolsFramework/Private/MultiSelectionTool.cpp
--- a/Engine/Source/Runtime/Experimental/InteractiveToolsFramework/Private/MultiSelectionTool.cpp
+++ b/Engine/Source/Runtime/Experimental/InteractiveToolsFramework/Private/MultiSelectionTool.cpp
@@ -1,9 +1,5 @@
 // Copyright Epic Games, Inc. All Rights Reserved.
 #include "MultiSelectionTool.h"
-#include "TargetInterfaces/MaterialProvider.h"
-#include "TargetInterfaces/MeshDescriptionCommitter.h"
-#include "TargetInterfaces/MeshDescriptionProvider.h"
-#include "TargetInterfaces/PrimitiveComponentBackedTarget.h"
 #include "TargetInterfaces/AssetBackedTarget.h"
 
 bool UMultiSelectionTool::GetMapToSharedSourceData(TArray<int32>& MapToFirstOccurrences)
@@ -41,34 +37,3 @@ bool UMultiSelectionTool::GetMapToSharedSourceData(TArray<int32>& MapToFirstOccu
 	return bSharesSources;
 }
 
-IPrimitiveComponentBackedTarget* UMultiSelectionTool::TargetComponentInterface(int32 ComponentIdx) const
-{
-	check(ComponentIdx >= 0 && ComponentIdx < Targets.Num());
-	IPrimitiveComponentBackedTarget* Component = Cast<IPrimitiveComponentBackedTarget>(Targets[ComponentIdx]);
-	check(Component);
-	return Component;
-}
-
-IMeshDescriptionCommitter* UMultiSelectionTool::TargetMeshCommitterInterface(int32 ComponentIdx) const
-{
-	check(ComponentIdx >= 0 && ComponentIdx < Targets.Num());
-	IMeshDescriptionCommitter* TargetMeshCommitter = Cast<IMeshDescriptionCommitter>(Targets[ComponentIdx]);
-	check(TargetMeshCommitter);
-	return TargetMeshCommitter;
-}
-
-IMeshDescriptionProvider* UMultiSelectionTool::TargetMeshProviderInterface(int32 ComponentIdx) const
-{
-	check(ComponentIdx >= 0 && ComponentIdx < Targets.Num());
-	IMeshDescriptionProvider* TargetMeshProvider = Cast<IMeshDescriptionProvider>(Targets[ComponentIdx]);
-	check(TargetMeshProvider);
-	return TargetMeshProvider;
-}
-
-IMaterialProvider* UMultiSelectionTool::TargetMaterialInterface(int32 ComponentIdx) const
-{
-	check(ComponentIdx >= 0 && ComponentIdx < Targets.Num());
-	IMaterialProvider* TargetMaterial = Cast<IMaterialProvider>(Targets[ComponentIdx]);
-	check(TargetMaterial);
-	return TargetMaterial;
-}
\ No newline at end of file
diff --git a/Engine/Source/Runtime/Experimental/InteractiveToolsFramework/Public/MultiSelectionTool.h b/Engine/Source/Runtime/Experimental/InteractiveToolsFramework/Public/MultiSelectionTool.h
--- a/Engine/Source/Runtime/Experimental/InteractiveToolsFramework/Public/MultiSelectionTool.h
+++ b/Engine/Source/Runtime/Experimental/InteractiveToolsFramework/Public/MultiSelectionTool.h
@@ -6,24 +6,18 @@
 #include "ComponentSourceInterfaces.h"
 #include "ToolTargets/ToolTarget.h"
 
-#include "MultiSelectionTool.generated.h"
+#include "TargetInterfaces/MaterialProvider.h"
+#include "TargetInterfaces/MeshDescriptionCommitter.h"
+#include "TargetInterfaces/MeshDescriptionProvider.h"
+#include "TargetInterfaces/PrimitiveComponentBackedTarget.h"
 
-class IPrimitiveComponentBackedTarget;
-class IMeshDescriptionCommitter;
-class IMeshDescriptionProvider;
-class IMaterialProvider;
+#include "MultiSelectionTool.generated.h"
 
 UCLASS(Transient)
 class INTERACTIVETOOLSFRAMEWORK_API UMultiSelectionTool : public UInteractiveTool
 {
 GENERATED_BODY()
 public:
-	// @deprecated Use SetTarget instead, and don't use FPrimitiveComponentTarget.
-	void SetSelection(TArray<TUniquePtr<FPrimitiveComponentTarget>> ComponentTargetsIn)
-    {
-		ComponentTargets = MoveTemp(ComponentTargetsIn);
-	}
-
 	void SetTargets(TArray<TObjectPtr<UToolTarget>> TargetsIn)
 	{
 		Targets = MoveTemp(TargetsIn);
@@ -41,13 +35,6 @@ public:
 				return false;
 			}
 		}
-		for (const TUniquePtr<FPrimitiveComponentTarget>& Target : ComponentTargets)
-		{
-			if (Target->IsValid() == false)
-			{
-				return false;
-			}
-		}
 		return true;
 	}
 
@@ -59,52 +46,27 @@ public:
 	}
 
 protected:
-	TArray<TUniquePtr<FPrimitiveComponentTarget>> ComponentTargets{};
-
 	UPROPERTY()
 	TArray<TObjectPtr<UToolTarget>> Targets{};
 
 	/**
-	 * Helper to find which component targets share source data
+	 * Helper to find which targets share source data.
+	 * Requires UAssetBackedTarget as a tool target requirement.
 	 *
-	 * @return Array of indices, 1:1 with ComponentTargets, indicating the first index where a component target sharing the same source data appeared.
+	 * @return Array of indices, 1:1 with Targets, indicating the first index where a component target sharing the same source data appeared.
 	 */
-	bool GetMapToFirstComponentsSharingSourceData(TArray<int32>& MapToFirstOccurrences)
-	{
-		bool bSharesSources = false;
-		MapToFirstOccurrences.SetNumUninitialized(ComponentTargets.Num());
-		for (int32 ComponentIdx = 0; ComponentIdx < ComponentTargets.Num(); ComponentIdx++)
-		{
-			MapToFirstOccurrences[ComponentIdx] = -1;
-		}
-		for (int32 ComponentIdx = 0; ComponentIdx < ComponentTargets.Num(); ComponentIdx++)
-		{
-			if (MapToFirstOccurrences[ComponentIdx] >= 0) // already mapped
-			{
-				continue;
-			}
-
-			MapToFirstOccurrences[ComponentIdx] = ComponentIdx;
-
-			FPrimitiveComponentTarget* ComponentTarget = ComponentTargets[ComponentIdx].Get();
-			for (int32 VsIdx = ComponentIdx + 1; VsIdx < ComponentTargets.Num(); VsIdx++)
-			{
-				if (ComponentTarget->HasSameSourceData(*ComponentTargets[VsIdx]))
-				{
-					bSharesSources = true;
-					MapToFirstOccurrences[VsIdx] = ComponentIdx;
-				}
-			}
-		}
-		return bSharesSources;
-	}
+	bool GetMapToSharedSourceData(TArray<int32>& MapToFirstOccurrences);
 
 	/**
-	 * Helper to find which targets share source data
-	 *
-	 * @return Array of indices, 1:1 with Targets, indicating the first index where a component target sharing the same source data appeared.
+	 * Template helper to retrieve an interface from a ToolTarget.
 	 */
-	bool GetMapToSharedSourceData(TArray<int32>& MapToFirstOccurrences);
+	template <class T>
+	T* TargetInterface(int32 TargetIdx) const
+	{
+		T* Interface = Cast<T>(Targets[TargetIdx]);
+		check(Interface);
+		return Interface;
+	}
 
 	/**
 	 * Helper to cast a Target into the IPrimitiveComponentBackedTarget interface.
@@ -126,3 +88,24 @@ protected:
 	 */
 	IMaterialProvider* TargetMaterialInterface(int32 ComponentIdx) const;
 };
+
+inline IPrimitiveComponentBackedTarget* UMultiSelectionTool::TargetComponentInterface(int32 TargetIdx) const
+{
+	return TargetInterface<IPrimitiveComponentBackedTarget>(TargetIdx);
+}
+
+inline IMeshDescriptionCommitter* UMultiSelectionTool::TargetMeshCommitterInterface(int32 TargetIdx) const
+{
+	return TargetInterface<IMeshDescriptionCommitter>(TargetIdx);
+}
+
+inline IMeshDescriptionProvider* UMultiSelectionTool::TargetMeshProviderInterface(int32 TargetIdx) const
+{
+	return TargetInterface<IMeshDescriptionProvider>(TargetIdx);
+}
+
+inline IMaterialProvider* UMultiSelectionTool::TargetMaterialInterface(int32 TargetIdx) const
+{
+	return TargetInterface<IMaterialProvider>(TargetIdx);
+}
+
diff --git a/Engine/Source/Runtime/Experimental/InteractiveToolsFramework/Public/TargetInterfaces/MaterialProvider.h b/Engine/Source/Runtime/Experimental/InteractiveToolsFramework/Public/TargetInterfaces/MaterialProvider.h
--- a/Engine/Source/Runtime/Experimental/InteractiveToolsFramework/Public/TargetInterfaces/MaterialProvider.h
+++ b/Engine/Source/Runtime/Experimental/InteractiveToolsFramework/Public/TargetInterfaces/MaterialProvider.h
@@ -50,7 +50,7 @@ public:
 	 * @param bAssetMaterials Prefer the underlying asset materials. This may be ignored by targets
 	 *  to which this is not relevant.
 	 */
-	virtual void GetMaterialSet(FComponentMaterialSet& MaterialSetOut, bool bPreferAssetMaterials) const = 0;
+	virtual void GetMaterialSet(FComponentMaterialSet& MaterialSetOut, bool bPreferAssetMaterials = false) const = 0;
 
 	/**
 	 * Commit an update to the material set. This may generate a transaction.
