Fix uninitialized ustruct members

#rnx

diff --git a/Engine/Plugins/Experimental/MeshModelingToolset/Source/MeshModelingTools/Public/Physics/CollisionPropertySets.h b/Engine/Plugins/Experimental/MeshModelingToolset/Source/MeshModelingTools/Public/Physics/CollisionPropertySets.h
--- a/Engine/Plugins/Experimental/MeshModelingToolset/Source/MeshModelingTools/Public/Physics/CollisionPropertySets.h
+++ b/Engine/Plugins/Experimental/MeshModelingToolset/Source/MeshModelingTools/Public/Physics/CollisionPropertySets.h
@@ -32,7 +32,7 @@ struct MESHMODELINGTOOLS_API FPhysicsSphereData
 	GENERATED_BODY()
 
 	UPROPERTY(VisibleAnywhere, Category = Sphere)
-	float Radius;
+	float Radius = 0.0f;
 
 	UPROPERTY(VisibleAnywhere, Category = Sphere)
 	FTransform Transform;
@@ -62,10 +62,10 @@ struct MESHMODELINGTOOLS_API FPhysicsCapsuleData
 	GENERATED_BODY()
 
 	UPROPERTY(VisibleAnywhere, Category = Sphere)
-	float Radius;
+	float Radius = 0.0f;
 
 	UPROPERTY(VisibleAnywhere, Category = Sphere)
-	float Length;
+	float Length = 0.0f;
 
 	UPROPERTY(VisibleAnywhere, Category = Sphere)
 	FTransform Transform;
@@ -80,10 +80,10 @@ struct MESHMODELINGTOOLS_API FPhysicsConvexData
 	GENERATED_BODY()
 
 	UPROPERTY(VisibleAnywhere, Category = Convex)
-	int32 NumVertices;
+	int32 NumVertices = 0;
 
 	UPROPERTY(VisibleAnywhere, Category = Convex)
-	int32 NumFaces;
+	int32 NumFaces = 0;
 
 	UPROPERTY(VisibleAnywhere, Category = Sphere)
 	FKShapeElem Element;
diff --git a/Engine/Plugins/Experimental/MeshModelingToolset/Source/MeshModelingToolsEditorOnly/Public/GenerateLODMeshesTool.h b/Engine/Plugins/Experimental/MeshModelingToolset/Source/MeshModelingToolsEditorOnly/Public/GenerateLODMeshesTool.h
--- a/Engine/Plugins/Experimental/MeshModelingToolset/Source/MeshModelingToolsEditorOnly/Public/GenerateLODMeshesTool.h
+++ b/Engine/Plugins/Experimental/MeshModelingToolset/Source/MeshModelingToolsEditorOnly/Public/GenerateLODMeshesTool.h
@@ -46,15 +46,15 @@ struct MESHMODELINGTOOLSEDITORONLY_API FLODLevelGenerateSettings
 
 	/** Simplification Target Type  */
 	UPROPERTY(EditAnywhere, Category = Options)
-	ESimplifyTargetType TargetMode;
+	ESimplifyTargetType TargetMode = ESimplifyTargetType::Percentage;
 
 	/** Target percentage */
 	UPROPERTY(EditAnywhere, Category = Options, meta = (UIMin = "0", UIMax = "100", EditConditionHides, EditCondition = "TargetMode == ESimplifyTargetType::Percentage"))
-	int32 TargetPercentage;
+	int32 TargetPercentage = 100;
 
 	/** Target vertex/triangle count */
 	UPROPERTY(EditAnywhere, Category = Options, meta = (UIMin = "4", UIMax = "10000", ClampMin = "1", ClampMax = "9999999999", EditConditionHides, EditCondition = "TargetMode == ESimplifyTargetType::TriangleCount || TargetMode == ESimplifyTargetType::VertexCount"))
-	int32 TargetCount;
+	int32 TargetCount = 500;
 
 	/** Target vertex/triangle count */
 	UPROPERTY(EditAnywhere, AdvancedDisplay, Category = Options)
diff --git a/Engine/Source/Runtime/Engine/Classes/Components/InstancedStaticMeshComponent.h b/Engine/Source/Runtime/Engine/Classes/Components/InstancedStaticMeshComponent.h
--- a/Engine/Source/Runtime/Engine/Classes/Components/InstancedStaticMeshComponent.h
+++ b/Engine/Source/Runtime/Engine/Classes/Components/InstancedStaticMeshComponent.h
@@ -532,5 +532,5 @@ public:
 
 	/* The cached random seed */
 	UPROPERTY()
-	int32 InstancingRandomSeed;
+	int32 InstancingRandomSeed = 0;
 };
