GeometryProcessing: remove FTransform3f/d and FPlane3f/d from ExplicitUseGeometryMathTypes.h and clean up fallout
#rnx

diff --git a/Engine/Plugins/Experimental/GeometryProcessing/Source/DynamicMesh/Private/MeshAdapterUtil.cpp b/Engine/Plugins/Experimental/GeometryProcessing/Source/DynamicMesh/Private/MeshAdapterUtil.cpp
--- a/Engine/Plugins/Experimental/GeometryProcessing/Source/DynamicMesh/Private/MeshAdapterUtil.cpp
+++ b/Engine/Plugins/Experimental/GeometryProcessing/Source/DynamicMesh/Private/MeshAdapterUtil.cpp
@@ -98,7 +98,7 @@ FPointSetAdapterd MeshAdapterUtil::MakeBoundaryEdgeMidpointsAdapter(const FDynam
 FTriangleMeshAdapterd MeshAdapterUtil::MakeTransformedDynamicMeshAdapter(const FDynamicMesh3* Mesh, FTransform Transform)
 {
 	FTriangleMeshAdapterd Adapter = MakeDynamicMeshAdapter(Mesh);
-	FTransform3d Transformd(Transform);
+	UE::Geometry::FTransform3d Transformd(Transform);
 	Adapter.GetVertex = [Mesh, Transformd](int32 Idx)
 	{
 		return FVector3d(Transformd.TransformPosition(Mesh->GetVertex(Idx)));
diff --git a/Engine/Plugins/Experimental/GeometryProcessing/Source/DynamicMesh/Private/ShapeApproximation/SimpleShapeSet3.cpp b/Engine/Plugins/Experimental/GeometryProcessing/Source/DynamicMesh/Private/ShapeApproximation/SimpleShapeSet3.cpp
--- a/Engine/Plugins/Experimental/GeometryProcessing/Source/DynamicMesh/Private/ShapeApproximation/SimpleShapeSet3.cpp
+++ b/Engine/Plugins/Experimental/GeometryProcessing/Source/DynamicMesh/Private/ShapeApproximation/SimpleShapeSet3.cpp
@@ -307,13 +307,13 @@ void FSimpleShapeSet3d::RemoveContainedGeometry()
 }
 
 
-static void TransformSphereShape(FSphereShape3d& SphereShape, const FTransform3d& Transform)
+static void TransformSphereShape(FSphereShape3d& SphereShape, const UE::Geometry::FTransform3d& Transform)
 {
 	double RadiusScale = Transform.GetScale().Length() / FVector3d::One().Length();
 	SphereShape.Sphere.Center = Transform.TransformPosition(SphereShape.Sphere.Center);
 	SphereShape.Sphere.Radius *= RadiusScale;
 }
-static void TransformBoxShape(FBoxShape3d& BoxShape, const FTransform3d& Transform)
+static void TransformBoxShape(FBoxShape3d& BoxShape, const UE::Geometry::FTransform3d& Transform)
 {
 	FVector3d CornerVec = BoxShape.Box.Frame.PointAt(BoxShape.Box.Extents) - BoxShape.Box.Frame.Origin;
 	BoxShape.Box.Frame.Transform(Transform);
@@ -322,7 +322,7 @@ static void TransformBoxShape(FBoxShape3d& BoxShape, const FTransform3d& Transfo
 	BoxShape.Box.Extents.Y = CornerVec.Dot(BoxShape.Box.AxisY());
 	BoxShape.Box.Extents.Z = CornerVec.Dot(BoxShape.Box.AxisZ());
 }
-static void TransformCapsuleShape(FCapsuleShape3d& CapsuleShape, const FTransform3d& Transform)
+static void TransformCapsuleShape(FCapsuleShape3d& CapsuleShape, const UE::Geometry::FTransform3d& Transform)
 {
 	FVector3d P0 = Transform.TransformPosition(CapsuleShape.Capsule.Segment.StartPoint());
 	FVector3d P1 = Transform.TransformPosition(CapsuleShape.Capsule.Segment.EndPoint());
@@ -342,19 +342,19 @@ static void TransformCapsuleShape(FCapsuleShape3d& CapsuleShape, const FTransfor
 
 
 
-static void TransformSphereShape(FSphereShape3d& SphereShape, const TArray<FTransform3d>& TransformSequence)
+static void TransformSphereShape(FSphereShape3d& SphereShape, const TArray<UE::Geometry::FTransform3d>& TransformSequence)
 {
-	for (const FTransform3d& XForm : TransformSequence)
+	for (const UE::Geometry::FTransform3d& XForm : TransformSequence)
 	{
 		SphereShape.Sphere.Center = XForm.TransformPosition(SphereShape.Sphere.Center);
 		double RadiusScale = XForm.GetScale().Length() / FVector3d::One().Length();
 		SphereShape.Sphere.Radius *= RadiusScale;
 	}
 }
-static void TransformBoxShape(FBoxShape3d& BoxShape, const TArray<FTransform3d>& TransformSequence)
+static void TransformBoxShape(FBoxShape3d& BoxShape, const TArray<UE::Geometry::FTransform3d>& TransformSequence)
 {
 	FVector3d CornerVec = BoxShape.Box.Frame.PointAt(BoxShape.Box.Extents) - BoxShape.Box.Frame.Origin;
-	for (const FTransform3d& XForm : TransformSequence)
+	for (const UE::Geometry::FTransform3d& XForm : TransformSequence)
 	{
 		BoxShape.Box.Frame.Transform(XForm);
 		CornerVec = XForm.TransformVector(CornerVec);
@@ -363,7 +363,7 @@ static void TransformBoxShape(FBoxShape3d& BoxShape, const TArray<FTransform3d>&
 	BoxShape.Box.Extents.Y = CornerVec.Dot(BoxShape.Box.AxisY());
 	BoxShape.Box.Extents.Z = CornerVec.Dot(BoxShape.Box.AxisZ());
 }
-static void TransformCapsuleShape(FCapsuleShape3d& CapsuleShape, const TArray<FTransform3d>& TransformSequence)
+static void TransformCapsuleShape(FCapsuleShape3d& CapsuleShape, const TArray<UE::Geometry::FTransform3d>& TransformSequence)
 {
 	FVector3d P0 = CapsuleShape.Capsule.Segment.StartPoint();
 	FVector3d P1 = CapsuleShape.Capsule.Segment.EndPoint();
@@ -373,7 +373,7 @@ static void TransformCapsuleShape(FCapsuleShape3d& CapsuleShape, const TArray<FT
 	FVector3d InitialSideVec = CapsuleFrame.PointAt(FVector3d(CurRadius, CurRadius, 0)) - CapsuleFrame.Origin;
 	FVector3d NewSideVec = InitialSideVec;
 
-	for (const FTransform3d& XForm : TransformSequence)
+	for (const UE::Geometry::FTransform3d& XForm : TransformSequence)
 	{
 		P0 = XForm.TransformPosition(P0);
 		P1 = XForm.TransformPosition(P1);
diff --git a/Engine/Plugins/Experimental/GeometryProcessing/Source/DynamicMesh/Public/MeshTransforms.h b/Engine/Plugins/Experimental/GeometryProcessing/Source/DynamicMesh/Public/MeshTransforms.h
--- a/Engine/Plugins/Experimental/GeometryProcessing/Source/DynamicMesh/Public/MeshTransforms.h
+++ b/Engine/Plugins/Experimental/GeometryProcessing/Source/DynamicMesh/Public/MeshTransforms.h
@@ -13,6 +13,7 @@
 namespace MeshTransforms
 {
 	using namespace UE::Geometry;
+	using UE::Geometry::FTransform3d;
 
 	/**
 	 * Apply Translation to vertex positions of Mesh. Does not modify any other attributes.
diff --git a/Engine/Plugins/Experimental/GeometryProcessing/Source/GeometricObjects/Private/Tests/TransformTypes.spec.cpp b/Engine/Plugins/Experimental/GeometryProcessing/Source/GeometricObjects/Private/Tests/TransformTypes.spec.cpp
--- a/Engine/Plugins/Experimental/GeometryProcessing/Source/GeometricObjects/Private/Tests/TransformTypes.spec.cpp
+++ b/Engine/Plugins/Experimental/GeometryProcessing/Source/GeometricObjects/Private/Tests/TransformTypes.spec.cpp
@@ -34,7 +34,7 @@ void FTransformTypesSpec::Define()
 				FVector Translation = MakeRandomVector();
 				FVector Scale = FVector(1.1f,1.1f,1.1f) + MakeRandomVector();		// need to stay away from zero scale or floats will explode and precision is too low
 				FTransform UETransform(Rotation, Translation, Scale);
-				FTransform3f GPTransform(UETransform);
+				UE::Geometry::FTransform3f GPTransform(UETransform);
 
 				float tol = 0.001f;		// float precision is very bad...some float comparisons will fail at 0.0001f
 				for (int j = 0; j < 1000; ++j)
diff --git a/Engine/Plugins/Experimental/GeometryProcessing/Source/GeometricObjects/Public/ExplicitUseGeometryMathTypes.h b/Engine/Plugins/Experimental/GeometryProcessing/Source/GeometricObjects/Public/ExplicitUseGeometryMathTypes.h
--- a/Engine/Plugins/Experimental/GeometryProcessing/Source/GeometricObjects/Public/ExplicitUseGeometryMathTypes.h
+++ b/Engine/Plugins/Experimental/GeometryProcessing/Source/GeometricObjects/Public/ExplicitUseGeometryMathTypes.h
@@ -80,8 +80,8 @@ using UE::Geometry::FAxisAlignedBox3i;
 using UE::Geometry::FHalfspace3d;
 using UE::Geometry::FHalfspace3f;
 
-using UE::Geometry::FPlane3d;
-using UE::Geometry::FPlane3f;
+//using UE::Geometry::FPlane3d;
+//using UE::Geometry::FPlane3f;
 
 using UE::Geometry::FLine2f;
 using UE::Geometry::FLine2d;
@@ -115,8 +115,8 @@ using UE::Geometry::FCircle2f;
 using UE::Geometry::FCircle3d;
 using UE::Geometry::FCircle3f;
 
-using UE::Geometry::FTransform3d;
-using UE::Geometry::FTransform3f;
+//using UE::Geometry::FTransform3d;
+//using UE::Geometry::FTransform3f;
 
 using UE::Geometry::FFrame3d;
 using UE::Geometry::FFrame3f;
diff --git a/Engine/Plugins/Experimental/MeshLODToolset/Source/MeshLODToolset/Private/Tools/GenerateStaticMeshLODAssetTool.cpp b/Engine/Plugins/Experimental/MeshLODToolset/Source/MeshLODToolset/Private/Tools/GenerateStaticMeshLODAssetTool.cpp
--- a/Engine/Plugins/Experimental/MeshLODToolset/Source/MeshLODToolset/Private/Tools/GenerateStaticMeshLODAssetTool.cpp
+++ b/Engine/Plugins/Experimental/MeshLODToolset/Source/MeshLODToolset/Private/Tools/GenerateStaticMeshLODAssetTool.cpp
@@ -121,7 +121,7 @@ namespace GenerateStaticMeshLODAssetLocals
 
 	public:
 
-		FGenerateStaticMeshLODAssetOperatorFactory(UGenerateStaticMeshLODAssetTool* AutoLODTool, FTransform3d ResultTransform) : 
+		FGenerateStaticMeshLODAssetOperatorFactory(UGenerateStaticMeshLODAssetTool* AutoLODTool, UE::Geometry::FTransform3d ResultTransform) :
 			AutoLODTool(AutoLODTool), 
 			ResultTransform(ResultTransform) 
 		{}
@@ -139,7 +139,7 @@ namespace GenerateStaticMeshLODAssetLocals
 		}
 
 		UGenerateStaticMeshLODAssetTool* AutoLODTool = nullptr;
-		FTransform3d ResultTransform;
+		UE::Geometry::FTransform3d ResultTransform;
 	};
 
 }
@@ -245,7 +245,7 @@ void UGenerateStaticMeshLODAssetTool::Setup()
 	FTransform PreviewTransform = SourceComponent->GetWorldTransform();
 	PreviewTransform.AddToTranslation(FVector(0, 2.5f*Bounds.BoxExtent.Y, 0));
 
-	this->OpFactory = MakeUnique<FGenerateStaticMeshLODAssetOperatorFactory>(this, (FTransform3d)PreviewTransform);
+	this->OpFactory = MakeUnique<FGenerateStaticMeshLODAssetOperatorFactory>(this, (UE::Geometry::FTransform3d)PreviewTransform);
 	PreviewWithBackgroundCompute = NewObject<UMeshOpPreviewWithBackgroundCompute>(this, "Preview");
 	PreviewWithBackgroundCompute->Setup(this->TargetWorld, this->OpFactory.Get());
 
diff --git a/Engine/Plugins/Experimental/MeshModelingToolset/Source/MeshModelingTools/Private/AddPatchTool.cpp b/Engine/Plugins/Experimental/MeshModelingToolset/Source/MeshModelingTools/Private/AddPatchTool.cpp
--- a/Engine/Plugins/Experimental/MeshModelingToolset/Source/MeshModelingTools/Private/AddPatchTool.cpp
+++ b/Engine/Plugins/Experimental/MeshModelingToolset/Source/MeshModelingTools/Private/AddPatchTool.cpp
@@ -300,7 +300,7 @@ void UAddPatchTool::OnClicked(const FInputDeviceRay& DeviceClickPos)
 {
 #if WITH_EDITOR
 	const FDynamicMesh3* CurMesh = PreviewMesh->GetPreviewDynamicMesh();
-	FTransform3d CurTransform(PreviewMesh->GetTransform());
+	UE::Geometry::FTransform3d CurTransform(PreviewMesh->GetTransform());
 	UMaterialInterface* Material = PreviewMesh->GetMaterial();
 	GetToolManager()->BeginUndoTransaction(LOCTEXT("AddPatchToolTransactionName", "Add Patch Mesh"));
 
diff --git a/Engine/Plugins/Experimental/MeshModelingToolset/Source/MeshModelingTools/Private/AddPrimitiveTool.cpp b/Engine/Plugins/Experimental/MeshModelingToolset/Source/MeshModelingTools/Private/AddPrimitiveTool.cpp
--- a/Engine/Plugins/Experimental/MeshModelingToolset/Source/MeshModelingTools/Private/AddPrimitiveTool.cpp
+++ b/Engine/Plugins/Experimental/MeshModelingToolset/Source/MeshModelingTools/Private/AddPrimitiveTool.cpp
@@ -323,7 +323,7 @@ void UAddPrimitiveTool::OnClicked(const FInputDeviceRay& DeviceClickPos)
 		GetToolManager()->BeginUndoTransaction(LOCTEXT("AddPrimitiveToolTransactionName", "Add Primitive Mesh"));
 		FActorSpawnParameters SpawnParameters;
 		SpawnParameters.Template = LastGenerated->Actor;
-		FTransform3d CurTransform(PreviewMesh->GetTransform());
+		UE::Geometry::FTransform3d CurTransform(PreviewMesh->GetTransform());
 		FRotator Rotation(0.0f, 0.0f, 0.0f);
 		AStaticMeshActor* CloneActor = TargetWorld->SpawnActor<AStaticMeshActor>(FVector::ZeroVector, Rotation, SpawnParameters);
 		// some properties must be manually set on the component because they will not persist reliably through the spawn template (especially if the actor creation was undone)
@@ -339,7 +339,7 @@ void UAddPrimitiveTool::OnClicked(const FInputDeviceRay& DeviceClickPos)
 	}
 
 	const FDynamicMesh3* CurMesh = PreviewMesh->GetPreviewDynamicMesh();
-	FTransform3d CurTransform(PreviewMesh->GetTransform());
+	UE::Geometry::FTransform3d CurTransform(PreviewMesh->GetTransform());
 
 	GetToolManager()->BeginUndoTransaction(LOCTEXT("AddPrimitiveToolTransactionName", "Add Primitive Mesh"));
 
diff --git a/Engine/Plugins/Experimental/MeshModelingToolset/Source/MeshModelingTools/Private/BakeMeshAttributeMapsTool.cpp b/Engine/Plugins/Experimental/MeshModelingToolset/Source/MeshModelingTools/Private/BakeMeshAttributeMapsTool.cpp
--- a/Engine/Plugins/Experimental/MeshModelingToolset/Source/MeshModelingTools/Private/BakeMeshAttributeMapsTool.cpp
+++ b/Engine/Plugins/Experimental/MeshModelingToolset/Source/MeshModelingTools/Private/BakeMeshAttributeMapsTool.cpp
@@ -346,9 +346,9 @@ void UBakeMeshAttributeMapsTool::UpdateDetailMesh()
 	Converter.Convert(DetailMeshProvider->GetMeshDescription(), *DetailMesh);
 	if (Settings->bUseWorldSpace && bIsBakeToSelf == false)
 	{
-		FTransform3d DetailToWorld(DetailComponent->GetWorldTransform());
+		UE::Geometry::FTransform3d DetailToWorld(DetailComponent->GetWorldTransform());
 		MeshTransforms::ApplyTransform(*DetailMesh, DetailToWorld);
-		FTransform3d WorldToBase(TargetComponent->GetWorldTransform());
+		UE::Geometry::FTransform3d WorldToBase(TargetComponent->GetWorldTransform());
 		MeshTransforms::ApplyTransform(*DetailMesh, WorldToBase.Inverse());
 	}
 	
diff --git a/Engine/Plugins/Experimental/MeshModelingToolset/Source/MeshModelingTools/Private/BakeTransformTool.cpp b/Engine/Plugins/Experimental/MeshModelingToolset/Source/MeshModelingTools/Private/BakeTransformTool.cpp
--- a/Engine/Plugins/Experimental/MeshModelingToolset/Source/MeshModelingTools/Private/BakeTransformTool.cpp
+++ b/Engine/Plugins/Experimental/MeshModelingToolset/Source/MeshModelingTools/Private/BakeTransformTool.cpp
@@ -153,7 +153,7 @@ void UBakeTransformTool::UpdateAssets()
 
 	GetToolManager()->BeginUndoTransaction(LOCTEXT("BakeTransformToolTransactionName", "Bake Transforms"));
 
-	TArray<FTransform3d> BakedTransforms;
+	TArray<UE::Geometry::FTransform3d> BakedTransforms;
 	for (int32 ComponentIdx = 0; ComponentIdx < Targets.Num(); ComponentIdx++)
 	{
 		IPrimitiveComponentBackedTarget* TargetComponent = TargetComponentInterface(ComponentIdx);
@@ -161,19 +161,19 @@ void UBakeTransformTool::UpdateAssets()
 		UPrimitiveComponent* Component = TargetComponent->GetOwnerComponent();
 		Component->Modify();
 
-		FTransform3d ComponentToWorld(TargetComponent->GetWorldTransform());
-		FTransform3d ToBakePart = FTransform3d::Identity();
-		FTransform3d NewWorldPart = ComponentToWorld;
+		UE::Geometry::FTransform3d ComponentToWorld(TargetComponent->GetWorldTransform());
+		UE::Geometry::FTransform3d ToBakePart = UE::Geometry::FTransform3d::Identity();
+		UE::Geometry::FTransform3d NewWorldPart = ComponentToWorld;
 
 		if (MapToFirstOccurrences[ComponentIdx] < ComponentIdx)
 		{
 			ToBakePart = BakedTransforms[MapToFirstOccurrences[ComponentIdx]];
 			BakedTransforms.Add(ToBakePart);
 			// try to invert baked transform
-			NewWorldPart = FTransform3d(
+			NewWorldPart = UE::Geometry::FTransform3d(
 				NewWorldPart.GetRotation() * ToBakePart.GetRotation().Inverse(),
 				NewWorldPart.GetTranslation(),
-				NewWorldPart.GetScale() * FTransform3d::GetSafeScaleReciprocal(ToBakePart.GetScale())
+				NewWorldPart.GetScale() * UE::Geometry::FTransform3d::GetSafeScaleReciprocal(ToBakePart.GetScale())
 			);
 			NewWorldPart.SetTranslation(NewWorldPart.GetTranslation() - NewWorldPart.TransformVector(ToBakePart.GetTranslation()));
 		}
diff --git a/Engine/Plugins/Experimental/MeshModelingToolset/Source/MeshModelingTools/Private/CombineMeshesTool.cpp b/Engine/Plugins/Experimental/MeshModelingToolset/Source/MeshModelingTools/Private/CombineMeshesTool.cpp
--- a/Engine/Plugins/Experimental/MeshModelingToolset/Source/MeshModelingTools/Private/CombineMeshesTool.cpp
+++ b/Engine/Plugins/Experimental/MeshModelingToolset/Source/MeshModelingTools/Private/CombineMeshesTool.cpp
@@ -254,7 +254,7 @@ void UCombineMeshesTool::CreateNewAsset()
 			const FMeshDescription* MeshDescription = MeshDescriptions[ComponentIdx];
 			Converter.Convert(MeshDescription, ComponentDMesh);
 
-			FTransform3d XF = (FTransform3d)(TargetComponentInterface(ComponentIdx)->GetWorldTransform() * ToAccum);
+			UE::Geometry::FTransform3d XF = (UE::Geometry::FTransform3d)(TargetComponentInterface(ComponentIdx)->GetWorldTransform() * ToAccum);
 			if (XF.GetDeterminant() < 0)
 			{
 				ComponentDMesh.ReverseOrientation(false);
@@ -273,7 +273,7 @@ void UCombineMeshesTool::CreateNewAsset()
 			{
 				Editor.AppendMesh(&ComponentDMesh, IndexMapping);
 				CollisionSettings = UE::Geometry::GetCollisionSettings(TargetComponent->GetOwnerComponent());
-				UE::Geometry::AppendSimpleCollision(TargetComponent->GetOwnerComponent(), &SimpleCollision, FTransform3d::Identity());
+				UE::Geometry::AppendSimpleCollision(TargetComponent->GetOwnerComponent(), &SimpleCollision, UE::Geometry::FTransform3d::Identity());
 			}
 			else
 			{
@@ -303,7 +303,7 @@ void UCombineMeshesTool::CreateNewAsset()
 		}
 		AActor* NewActor = AssetGenerationUtil::GenerateStaticMeshActor(
 			AssetAPI, TargetWorld,
-			&AccumulateDMesh, (FTransform3d)AccumToWorld, UseBaseName, AllMaterials);
+			&AccumulateDMesh, (UE::Geometry::FTransform3d)AccumToWorld, UseBaseName, AllMaterials);
 		if (NewActor != nullptr)
 		{
 			// copy the component materials onto the new static mesh asset too
@@ -405,12 +405,12 @@ void UCombineMeshesTool::UpdateExistingAsset()
 	IMaterialProvider* UpdateTargetMaterial = TargetMaterialInterface(SkipIndex);
 	SkipActor = UpdateTarget->GetOwnerActor();
 
-	FTransform3d TargetToWorld = (FTransform3d)UpdateTarget->GetWorldTransform();
-	FTransform3d WorldToTarget = TargetToWorld.Inverse();
+	UE::Geometry::FTransform3d TargetToWorld = (UE::Geometry::FTransform3d)UpdateTarget->GetWorldTransform();
+	UE::Geometry::FTransform3d WorldToTarget = TargetToWorld.Inverse();
 
 	FSimpleShapeSet3d SimpleCollision;
 	UE::Geometry::FComponentCollisionSettings CollisionSettings = UE::Geometry::GetCollisionSettings(UpdateTarget->GetOwnerComponent());
-	TArray<FTransform3d> Transforms;
+	TArray<UE::Geometry::FTransform3d> Transforms;
 	Transforms.SetNum(2);
 
 	{
@@ -442,7 +442,7 @@ void UCombineMeshesTool::UpdateExistingAsset()
 
 			if (ComponentIdx != SkipIndex)
 			{
-				FTransform3d ComponentToWorld = (FTransform3d)TargetComponent->GetWorldTransform();
+				UE::Geometry::FTransform3d ComponentToWorld = (UE::Geometry::FTransform3d)TargetComponent->GetWorldTransform();
 				MeshTransforms::ApplyTransform(ComponentDMesh, ComponentToWorld);
 				if (ComponentToWorld.GetDeterminant() < 0)
 				{
@@ -459,7 +459,7 @@ void UCombineMeshesTool::UpdateExistingAsset()
 			}
 			else
 			{
-				UE::Geometry::AppendSimpleCollision(TargetComponent->GetOwnerComponent(), &SimpleCollision, FTransform3d::Identity());
+				UE::Geometry::AppendSimpleCollision(TargetComponent->GetOwnerComponent(), &SimpleCollision, UE::Geometry::FTransform3d::Identity());
 			}
 
 			FDynamicMeshEditor Editor(&AccumulateDMesh);
diff --git a/Engine/Plugins/Experimental/MeshModelingToolset/Source/MeshModelingTools/Private/ConversionUtils/VolumeToDynamicMesh.cpp b/Engine/Plugins/Experimental/MeshModelingToolset/Source/MeshModelingTools/Private/ConversionUtils/VolumeToDynamicMesh.cpp
--- a/Engine/Plugins/Experimental/MeshModelingToolset/Source/MeshModelingTools/Private/ConversionUtils/VolumeToDynamicMesh.cpp
+++ b/Engine/Plugins/Experimental/MeshModelingToolset/Source/MeshModelingTools/Private/ConversionUtils/VolumeToDynamicMesh.cpp
@@ -29,7 +29,8 @@ void VolumeToDynamicMesh(AVolume* Volume, FDynamicMesh3& Mesh,
 	}
 
 	UModel* Model = Volume->Brush;
-	FTransform3d XForm = (Options.bInWorldSpace) ? FTransform3d(Volume->GetTransform()) : FTransform3d::Identity();
+	UE::Geometry::FTransform3d XForm = (Options.bInWorldSpace) ? 
+		UE::Geometry::FTransform3d(Volume->GetTransform()) : UE::Geometry::FTransform3d::Identity();
 
 	// Each "BspNode" is a planar polygon, triangulate each polygon and accumulate in a mesh.
 	// Note that this does not make any attempt to weld vertices/edges
diff --git a/Engine/Plugins/Experimental/MeshModelingToolset/Source/MeshModelingTools/Private/ConvertToPolygonsTool.cpp b/Engine/Plugins/Experimental/MeshModelingToolset/Source/MeshModelingTools/Private/ConvertToPolygonsTool.cpp
--- a/Engine/Plugins/Experimental/MeshModelingToolset/Source/MeshModelingTools/Private/ConvertToPolygonsTool.cpp
+++ b/Engine/Plugins/Experimental/MeshModelingToolset/Source/MeshModelingTools/Private/ConvertToPolygonsTool.cpp
@@ -120,7 +120,7 @@ public:
 
 	void SetTransform(const FTransform& Transform)
 	{
-		ResultTransform = (FTransform3d)Transform;
+		ResultTransform = (UE::Geometry::FTransform3d)Transform;
 	}
 
 	
diff --git a/Engine/Plugins/Experimental/MeshModelingToolset/Source/MeshModelingTools/Private/DeformMeshPolygonsTool.cpp b/Engine/Plugins/Experimental/MeshModelingToolset/Source/MeshModelingTools/Private/DeformMeshPolygonsTool.cpp
--- a/Engine/Plugins/Experimental/MeshModelingToolset/Source/MeshModelingTools/Private/DeformMeshPolygonsTool.cpp
+++ b/Engine/Plugins/Experimental/MeshModelingToolset/Source/MeshModelingTools/Private/DeformMeshPolygonsTool.cpp
@@ -882,7 +882,7 @@ void UDeformMeshPolygonsTool::Setup()
 	TopoSelector.Initialize(DynamicMeshComponent->GetMesh(), &Topology);
 	TopoSelector.SetSpatialSource([this]() { return &GetSpatial(); });
 	TopoSelector.PointsWithinToleranceTest = [this](const FVector3d& Position1, const FVector3d& Position2, double TolScale) {
-		FTransform3d Transform(ComponentTarget->GetWorldTransform());
+		UE::Geometry::FTransform3d Transform(ComponentTarget->GetWorldTransform());
 		return ToolSceneQueriesUtil::PointSnapQuery(CameraState, Transform.TransformPosition(Position1), Transform.TransformPosition(Position2),
 			ToolSceneQueriesUtil::GetDefaultVisualAngleSnapThreshD() * TolScale);
 	};
@@ -1034,7 +1034,7 @@ FDynamicMeshAABBTree3& UDeformMeshPolygonsTool::GetSpatial()
 
 bool UDeformMeshPolygonsTool::HitTest(const FRay& WorldRay, FHitResult& OutHit)
 {
-	FTransform3d Transform(ComponentTarget->GetWorldTransform());
+	UE::Geometry::FTransform3d Transform(ComponentTarget->GetWorldTransform());
 	FRay3d LocalRay(Transform.InverseTransformPosition(WorldRay.Origin),
 	                Transform.InverseTransformVector(WorldRay.Direction));
 	UE::Geometry::Normalize(LocalRay.Direction);
@@ -1080,7 +1080,7 @@ bool UDeformMeshPolygonsTool::HitTest(const FRay& WorldRay, FHitResult& OutHit)
 
 void UDeformMeshPolygonsTool::OnBeginDrag(const FRay& WorldRay)
 {
-	FTransform3d Transform(ComponentTarget->GetWorldTransform());
+	UE::Geometry::FTransform3d Transform(ComponentTarget->GetWorldTransform());
 	FRay3d LocalRay(Transform.InverseTransformPosition(WorldRay.Origin),
 	                Transform.InverseTransformVector(WorldRay.Direction));
 	UE::Geometry::Normalize(LocalRay.Direction);
@@ -1189,7 +1189,7 @@ void UDeformMeshPolygonsTool::OnBeginDrag(const FRay& WorldRay)
 
 void UDeformMeshPolygonsTool::UpdateActiveSurfaceFrame(FGroupTopologySelection& Selection)
 {
-	FTransform3d Transform(ComponentTarget->GetWorldTransform());
+	UE::Geometry::FTransform3d Transform(ComponentTarget->GetWorldTransform());
 
 	// update surface frame
 	ActiveSurfaceFrame.Origin = StartHitPosWorld;
@@ -1295,7 +1295,7 @@ bool UDeformMeshPolygonsTool::OnUpdateHover(const FInputDeviceRay& DevicePos)
 	//if (!bNeedEmitEndChange)
 	if (ActiveVertexChange == nullptr)
 	{
-		FTransform3d Transform(ComponentTarget->GetWorldTransform());
+		UE::Geometry::FTransform3d Transform(ComponentTarget->GetWorldTransform());
 		FRay3d LocalRay(Transform.InverseTransformPosition(DevicePos.WorldRay.Origin),
 		                Transform.InverseTransformVector(DevicePos.WorldRay.Direction));
 		UE::Geometry::Normalize(LocalRay.Direction);
@@ -1393,7 +1393,7 @@ void UDeformMeshPolygonsTool::ComputeUpdate_Rotate()
 	FGroupTopologyDeformer& SelectedDeformer = (bIsLaplacian) ? *LaplacianDeformer : LinearDeformer;
 
 	FDynamicMesh3* Mesh    = DynamicMeshComponent->GetMesh();
-	FTransform3d Transform = FTransform3d(ComponentTarget->GetWorldTransform());
+	UE::Geometry::FTransform3d Transform(ComponentTarget->GetWorldTransform());
 	FVector NewHitPosWorld = LastHitPosWorld;
 
 	FVector3d SnappedPoint;
diff --git a/Engine/Plugins/Experimental/MeshModelingToolset/Source/MeshModelingTools/Private/DrawPolyPathTool.cpp b/Engine/Plugins/Experimental/MeshModelingToolset/Source/MeshModelingTools/Private/DrawPolyPathTool.cpp
--- a/Engine/Plugins/Experimental/MeshModelingToolset/Source/MeshModelingTools/Private/DrawPolyPathTool.cpp
+++ b/Engine/Plugins/Experimental/MeshModelingToolset/Source/MeshModelingTools/Private/DrawPolyPathTool.cpp
@@ -513,7 +513,7 @@ void UDrawPolyPathTool::BeginInteractiveOffsetDistance()
 	// begin setting offset distance
 	CurveDistMechanic = NewObject<USpatialCurveDistanceMechanic>(this);
 	CurveDistMechanic->Setup(this);
-	CurveDistMechanic->InitializePolyCurve(CurPolyLine, FTransform3d::Identity());
+	CurveDistMechanic->InitializePolyCurve(CurPolyLine, UE::Geometry::FTransform3d::Identity());
 
 	InitializePreviewMesh();
 
diff --git a/Engine/Plugins/Experimental/MeshModelingToolset/Source/MeshModelingTools/Private/DrawPolygonTool.cpp b/Engine/Plugins/Experimental/MeshModelingToolset/Source/MeshModelingTools/Private/DrawPolygonTool.cpp
--- a/Engine/Plugins/Experimental/MeshModelingToolset/Source/MeshModelingTools/Private/DrawPolygonTool.cpp
+++ b/Engine/Plugins/Experimental/MeshModelingToolset/Source/MeshModelingTools/Private/DrawPolygonTool.cpp
@@ -352,7 +352,7 @@ void UDrawPolygonTool::Render(IToolsContextRenderAPI* RenderAPI)
 			FVector3d P0(DrawSnapLine.PointAt(-9999)), P1(DrawSnapLine.PointAt(9999));		// should be smarter here...
 			if (RenderCameraState.bIsOrthographic == false)
 			{
-				FPlane3d CameraPlane(RenderCameraState.Forward(), RenderCameraState.Position + 1.0*RenderCameraState.Forward());
+				UE::Geometry::FPlane3d CameraPlane(RenderCameraState.Forward(), RenderCameraState.Position + 1.0*RenderCameraState.Forward());
 				CameraPlane.ClipSegment(P0, P1);
 			}
 			PDI->DrawLine((FVector)P0, (FVector)P1, SnapLineColor, SDPG_Foreground, 0.5 * PDIScale, 0.0f, true);
diff --git a/Engine/Plugins/Experimental/MeshModelingToolset/Source/MeshModelingTools/Private/DynamicMeshSculptTool.cpp b/Engine/Plugins/Experimental/MeshModelingToolset/Source/MeshModelingTools/Private/DynamicMeshSculptTool.cpp
--- a/Engine/Plugins/Experimental/MeshModelingToolset/Source/MeshModelingTools/Private/DynamicMeshSculptTool.cpp
+++ b/Engine/Plugins/Experimental/MeshModelingToolset/Source/MeshModelingTools/Private/DynamicMeshSculptTool.cpp
@@ -108,14 +108,14 @@ void UDynamicMeshSculptTool::Setup()
 
 	// transform mesh to world space because handling scaling inside brush is a mess
 	// Note: this transform does not include translation ( so only the 3x3 transform)
-	InitialTargetTransform = FTransform3d(ComponentTarget->GetWorldTransform());
+	InitialTargetTransform = UE::Geometry::FTransform3d(ComponentTarget->GetWorldTransform());
 	// clamp scaling because if we allow zero-scale we cannot invert this transform on Accept
 	InitialTargetTransform.ClampMinimumScale(0.01);
 	FVector3d Translation = InitialTargetTransform.GetTranslation();
 	InitialTargetTransform.SetTranslation(FVector3d::Zero());
 	DynamicMeshComponent->ApplyTransform(InitialTargetTransform, false);
 	// since we moved to World coords there is not a current transform anymore.
-	CurTargetTransform = FTransform3d(Translation);
+	CurTargetTransform = UE::Geometry::FTransform3d(Translation);
 	DynamicMeshComponent->SetWorldTransform((FTransform)CurTargetTransform);
 
 	// copy material if there is one
diff --git a/Engine/Plugins/Experimental/MeshModelingToolset/Source/MeshModelingTools/Private/EdgeLoopInsertionTool.cpp b/Engine/Plugins/Experimental/MeshModelingToolset/Source/MeshModelingTools/Private/EdgeLoopInsertionTool.cpp
--- a/Engine/Plugins/Experimental/MeshModelingToolset/Source/MeshModelingTools/Private/EdgeLoopInsertionTool.cpp
+++ b/Engine/Plugins/Experimental/MeshModelingToolset/Source/MeshModelingTools/Private/EdgeLoopInsertionTool.cpp
@@ -143,7 +143,7 @@ void UEdgeLoopInsertionTool::Setup()
 	TopologySelector.Initialize(CurrentMesh.Get(), CurrentTopology.Get());
 	TopologySelector.SetSpatialSource([this]() {return &MeshSpatial; });
 	TopologySelector.PointsWithinToleranceTest = [this](const FVector3d& Position1, const FVector3d& Position2, double TolScale) {
-		FTransform3d Transform(ComponentTarget->GetWorldTransform());
+		UE::Geometry::FTransform3d Transform(ComponentTarget->GetWorldTransform());
 		return ToolSceneQueriesUtil::PointSnapQuery(CameraState, Transform.TransformPosition(Position1), Transform.TransformPosition(Position2),
 			ToolSceneQueriesUtil::GetDefaultVisualAngleSnapThreshD() * TolScale);
 	};
diff --git a/Engine/Plugins/Experimental/MeshModelingToolset/Source/MeshModelingTools/Private/EditPivotTool.cpp b/Engine/Plugins/Experimental/MeshModelingToolset/Source/MeshModelingTools/Private/EditPivotTool.cpp
--- a/Engine/Plugins/Experimental/MeshModelingToolset/Source/MeshModelingTools/Private/EditPivotTool.cpp
+++ b/Engine/Plugins/Experimental/MeshModelingToolset/Source/MeshModelingTools/Private/EditPivotTool.cpp
@@ -173,7 +173,7 @@ void UEditPivotTool::Precompute()
 	int NumComponents = Targets.Num();
 	if (NumComponents == 1)
 	{
-		Transform = FTransform3d(TargetComponentInterface(0)->GetWorldTransform());
+		Transform = UE::Geometry::FTransform3d(TargetComponentInterface(0)->GetWorldTransform());
 
 		FMeshDescription* Mesh = TargetMeshProviderInterface(0)->GetMeshDescription();
 		VertexIteration(Mesh, [&](int32 VertexID, const FVector& Position) {
@@ -183,12 +183,12 @@ void UEditPivotTool::Precompute()
 	}
 	else
 	{
-		Transform = FTransform3d::Identity();
+		Transform = UE::Geometry::FTransform3d::Identity();
 		for (int k = 0; k < NumComponents; ++k)
 		{
 			IPrimitiveComponentBackedTarget* TargetComponent = TargetComponentInterface(k);
 			IMeshDescriptionProvider* TargetMeshProvider = TargetMeshProviderInterface(k);
-			FTransform3d CurTransform(TargetComponent->GetWorldTransform());
+			UE::Geometry::FTransform3d CurTransform(TargetComponent->GetWorldTransform());
 			FMeshDescription* Mesh = TargetMeshProvider->GetMeshDescription();
 			VertexIteration(Mesh, [&](int32 VertexID, const FVector& Position) {
 				ObjectBounds.Contain(CurTransform.TransformPosition(Position));
@@ -423,7 +423,7 @@ void UEditPivotTool::UpdateAssets(const FFrame3d& NewPivotWorldFrame)
 		IPrimitiveComponentBackedTarget* TargetComponent = TargetComponentInterface(ComponentIdx);
 		if (MapToFirstOccurrences[ComponentIdx] == ComponentIdx)
 		{
-			FTransform3d ToBake(OriginalTransforms[ComponentIdx] * NewWorldInverse);
+			UE::Geometry::FTransform3d ToBake(OriginalTransforms[ComponentIdx] * NewWorldInverse);
 
 			UPrimitiveComponent* Component = TargetComponent->GetOwnerComponent();
 			Component->Modify();
diff --git a/Engine/Plugins/Experimental/MeshModelingToolset/Source/MeshModelingTools/Private/EditUVIslandsTool.cpp b/Engine/Plugins/Experimental/MeshModelingToolset/Source/MeshModelingTools/Private/EditUVIslandsTool.cpp
--- a/Engine/Plugins/Experimental/MeshModelingToolset/Source/MeshModelingTools/Private/EditUVIslandsTool.cpp
+++ b/Engine/Plugins/Experimental/MeshModelingToolset/Source/MeshModelingTools/Private/EditUVIslandsTool.cpp
@@ -46,7 +46,7 @@ void UEditUVIslandsTool::Setup()
 	DynamicMeshComponent->SetupAttachment(ComponentTarget->GetOwnerActor()->GetRootComponent());
 	DynamicMeshComponent->RegisterComponent();
 	DynamicMeshComponent->SetWorldTransform(ComponentTarget->GetWorldTransform());
-	WorldTransform = FTransform3d(DynamicMeshComponent->GetComponentTransform());
+	WorldTransform = UE::Geometry::FTransform3d(DynamicMeshComponent->GetComponentTransform());
 
 	// set materials
 	FComponentMaterialSet MaterialSet;
diff --git a/Engine/Plugins/Experimental/MeshModelingToolset/Source/MeshModelingTools/Private/GroupEdgeInsertionTool.cpp b/Engine/Plugins/Experimental/MeshModelingToolset/Source/MeshModelingTools/Private/GroupEdgeInsertionTool.cpp
--- a/Engine/Plugins/Experimental/MeshModelingToolset/Source/MeshModelingTools/Private/GroupEdgeInsertionTool.cpp
+++ b/Engine/Plugins/Experimental/MeshModelingToolset/Source/MeshModelingTools/Private/GroupEdgeInsertionTool.cpp
@@ -127,7 +127,7 @@ void UGroupEdgeInsertionTool::Setup()
 	TopologySelector.Initialize(CurrentMesh.Get(), CurrentTopology.Get());
 	TopologySelector.SetSpatialSource([this]() {return &MeshSpatial; });
 	TopologySelector.PointsWithinToleranceTest = [this](const FVector3d& Position1, const FVector3d& Position2, double TolScale) {
-		FTransform3d Transform(ComponentTarget->GetWorldTransform());
+		UE::Geometry::FTransform3d Transform(ComponentTarget->GetWorldTransform());
 		return ToolSceneQueriesUtil::PointSnapQuery(CameraState, Transform.TransformPosition(Position1), Transform.TransformPosition(Position2),
 			ToolSceneQueriesUtil::GetDefaultVisualAngleSnapThreshD() * TolScale);
 	};
diff --git a/Engine/Plugins/Experimental/MeshModelingToolset/Source/MeshModelingTools/Private/HoleFillTool.cpp b/Engine/Plugins/Experimental/MeshModelingToolset/Source/MeshModelingTools/Private/HoleFillTool.cpp
--- a/Engine/Plugins/Experimental/MeshModelingToolset/Source/MeshModelingTools/Private/HoleFillTool.cpp
+++ b/Engine/Plugins/Experimental/MeshModelingToolset/Source/MeshModelingTools/Private/HoleFillTool.cpp
@@ -102,7 +102,7 @@ TUniquePtr<FDynamicMeshOperator> UHoleFillOperatorFactory::MakeNewOperator()
 	TUniquePtr<FHoleFillOp> FillOp = MakeUnique<FHoleFillOp>();
 
 	FTransform LocalToWorld = FillTool->ComponentTarget->GetWorldTransform();
-	FillOp->SetResultTransform((FTransform3d)LocalToWorld);
+	FillOp->SetResultTransform((UE::Geometry::FTransform3d)LocalToWorld);
 	FillOp->OriginalMesh = FillTool->OriginalMesh;
 	FillOp->MeshUVScaleFactor = FillTool->MeshUVScaleFactor;
 	FillTool->GetLoopsToFill(FillOp->Loops);
diff --git a/Engine/Plugins/Experimental/MeshModelingToolset/Source/MeshModelingTools/Private/LatticeDeformerTool.cpp b/Engine/Plugins/Experimental/MeshModelingToolset/Source/MeshModelingTools/Private/LatticeDeformerTool.cpp
--- a/Engine/Plugins/Experimental/MeshModelingToolset/Source/MeshModelingTools/Private/LatticeDeformerTool.cpp
+++ b/Engine/Plugins/Experimental/MeshModelingToolset/Source/MeshModelingTools/Private/LatticeDeformerTool.cpp
@@ -89,7 +89,7 @@ void ULatticeDeformerTool::InitializeLattice(TArray<FVector3d>& OutLatticePoints
 	Lattice->GenerateInitialLatticePositions(OutLatticePoints);
 
 	// Put the lattice in world space
-	FTransform3d LocalToWorld = FTransform3d(ComponentTarget->GetWorldTransform());
+	UE::Geometry::FTransform3d LocalToWorld(ComponentTarget->GetWorldTransform());
 	Algo::ForEach(OutLatticePoints, [&LocalToWorld](FVector3d& Point) {
 		Point = LocalToWorld.TransformPosition(Point);
 	});
@@ -144,7 +144,7 @@ void ULatticeDeformerTool::Setup()
 	ControlPointsMechanic = NewObject<ULatticeControlPointsMechanic>(this);
 	ControlPointsMechanic->Setup(this);
 	ControlPointsMechanic->SetWorld(TargetWorld);
-	FTransform3d LocalToWorld = FTransform3d(ComponentTarget->GetWorldTransform());
+	UE::Geometry::FTransform3d LocalToWorld(ComponentTarget->GetWorldTransform());
 	ControlPointsMechanic->Initialize(LatticePoints, LatticeEdges, LocalToWorld);
 
 	auto OnPointsChangedLambda = [this]()
@@ -180,7 +180,7 @@ void ULatticeDeformerTool::Shutdown(EToolShutdownType ShutdownType)
 
 			// The lattice and its output mesh are in world space, so get them in local space.
 			// TODO: Would it make more sense to do all the lattice computation in local space?
-			FTransform3d LocalToWorld = FTransform3d(ComponentTarget->GetWorldTransform());
+			UE::Geometry::FTransform3d LocalToWorld(ComponentTarget->GetWorldTransform());
 			MeshTransforms::ApplyTransformInverse(*DynamicMeshResult, LocalToWorld);
 
 			ComponentTarget->CommitMesh([DynamicMeshResult](const FPrimitiveComponentTarget::FCommitParams& CommitParams)
@@ -235,7 +235,7 @@ void ULatticeDeformerTool::OnTick(float DeltaTime)
 			TArray<FVector3d> LatticePoints;
 			TArray<FVector2i> LatticeEdges;
 			InitializeLattice(LatticePoints, LatticeEdges);
-			FTransform3d LocalToWorld = FTransform3d(ComponentTarget->GetWorldTransform());
+			UE::Geometry::FTransform3d LocalToWorld(ComponentTarget->GetWorldTransform());
 			ControlPointsMechanic->Initialize(LatticePoints, LatticeEdges, LocalToWorld);
 			Preview->InvalidateResult();
 			bShouldRebuild = false;
diff --git a/Engine/Plugins/Experimental/MeshModelingToolset/Source/MeshModelingTools/Private/MeshAttributePaintTool.cpp b/Engine/Plugins/Experimental/MeshModelingToolset/Source/MeshModelingTools/Private/MeshAttributePaintTool.cpp
--- a/Engine/Plugins/Experimental/MeshModelingToolset/Source/MeshModelingTools/Private/MeshAttributePaintTool.cpp
+++ b/Engine/Plugins/Experimental/MeshModelingToolset/Source/MeshModelingTools/Private/MeshAttributePaintTool.cpp
@@ -341,7 +341,7 @@ bool UMeshAttributePaintTool::OnUpdateHover(const FInputDeviceRay& DevicePos)
 
 void UMeshAttributePaintTool::CalculateVertexROI(const FBrushStampData& Stamp, TArray<int>& VertexROI)
 {
-	FTransform3d Transform(ComponentTarget->GetWorldTransform());
+	UE::Geometry::FTransform3d Transform(ComponentTarget->GetWorldTransform());
 	FVector3d StampPosLocal = Transform.InverseTransformPosition(Stamp.WorldPosition);
 
 	float Radius = GetCurrentBrushRadiusLocal();
@@ -499,7 +499,7 @@ void UMeshAttributePaintTool::ApplyStamp(const FBrushStampData& Stamp)
 
 void UMeshAttributePaintTool::ApplyStamp_Paint(const FBrushStampData& Stamp, FStampActionData& ActionData)
 {
-	FTransform3d Transform(ComponentTarget->GetWorldTransform());
+	UE::Geometry::FTransform3d Transform(ComponentTarget->GetWorldTransform());
 	FVector3d StampPosLocal = Transform.InverseTransformPosition(Stamp.WorldPosition);
 
 	int32 NumVertices = ActionData.ROIVertices.Num();
diff --git a/Engine/Plugins/Experimental/MeshModelingToolset/Source/MeshModelingTools/Private/MeshSelectionTool.cpp b/Engine/Plugins/Experimental/MeshModelingToolset/Source/MeshModelingTools/Private/MeshSelectionTool.cpp
--- a/Engine/Plugins/Experimental/MeshModelingToolset/Source/MeshModelingTools/Private/MeshSelectionTool.cpp
+++ b/Engine/Plugins/Experimental/MeshModelingToolset/Source/MeshModelingTools/Private/MeshSelectionTool.cpp
@@ -404,7 +404,7 @@ void UMeshSelectionTool::CalculateVertexROI(const FBrushStampData& Stamp, TArray
 
 void UMeshSelectionTool::CalculateTriangleROI(const FBrushStampData& Stamp, TArray<int>& TriangleROI)
 {
-	FTransform3d Transform(ComponentTarget->GetWorldTransform());
+	UE::Geometry::FTransform3d Transform(ComponentTarget->GetWorldTransform());
 	FVector3d StampPosLocal = Transform.InverseTransformPosition(Stamp.WorldPosition);
 
 	// always select first triangle
@@ -1348,7 +1348,7 @@ void UMeshSelectionTool::SeparateSelectedTriangles()
 	Editor.AppendTriangles(SourceMesh, SelectedFaces, Mappings, EditResult);
 
 	// emit new asset
-	FTransform3d Transform(PreviewMesh->GetTransform());
+	UE::Geometry::FTransform3d Transform(PreviewMesh->GetTransform());
 	GetToolManager()->BeginUndoTransaction(LOCTEXT("MeshSelectionToolSeparate", "Separate"));
 
 	// build array of materials from the original
diff --git a/Engine/Plugins/Experimental/MeshModelingToolset/Source/MeshModelingTools/Private/MirrorTool.cpp b/Engine/Plugins/Experimental/MeshModelingToolset/Source/MeshModelingTools/Private/MirrorTool.cpp
--- a/Engine/Plugins/Experimental/MeshModelingToolset/Source/MeshModelingTools/Private/MirrorTool.cpp
+++ b/Engine/Plugins/Experimental/MeshModelingToolset/Source/MeshModelingTools/Private/MirrorTool.cpp
@@ -87,7 +87,7 @@ TUniquePtr<FDynamicMeshOperator> UMirrorOperatorFactory::MakeNewOperator()
 		}
 	}
 	LocalToWorld.SetScale3D(LocalToWorldScale);
-	FTransform3d WorldToLocal = FTransform3d(LocalToWorld).Inverse();
+	UE::Geometry::FTransform3d WorldToLocal = UE::Geometry::FTransform3d(LocalToWorld).Inverse();
 
 	// Now we can get the plane parameters in local space.
 	MirrorOp->LocalPlaneOrigin = WorldToLocal.TransformPosition(MirrorTool->MirrorPlaneOrigin);;
diff --git a/Engine/Plugins/Experimental/MeshModelingToolset/Source/MeshModelingTools/Private/Physics/CollisionGeometryVisualization.cpp b/Engine/Plugins/Experimental/MeshModelingToolset/Source/MeshModelingTools/Private/Physics/CollisionGeometryVisualization.cpp
--- a/Engine/Plugins/Experimental/MeshModelingToolset/Source/MeshModelingTools/Private/Physics/CollisionGeometryVisualization.cpp
+++ b/Engine/Plugins/Experimental/MeshModelingToolset/Source/MeshModelingTools/Private/Physics/CollisionGeometryVisualization.cpp
@@ -23,7 +23,7 @@ void UE::PhysicsTools::InitializePreviewGeometryLines(const FPhysicsDataCollecti
 		const FKSphereElem& Sphere = AggGeom.SphereElems[Index];
 		FTransform ElemTransform = Sphere.GetTransform();
 		ElemTransform.ScaleTranslation(PhysicsData.ExternalScale3D);
-		FTransform3f ElemTransformf(ElemTransform);
+		UE::Geometry::FTransform3f ElemTransformf(ElemTransform);
 		float Radius = PhysicsData.ExternalScale3D.GetAbsMin() * Sphere.Radius;
 		UE::Geometry::GenerateCircleSegments<float>(CircleSteps, Radius, FVector3f::Zero(), FVector3f::UnitX(), FVector3f::UnitY(), ElemTransformf,
 			[&](const FVector3f& A, const FVector3f& B) { LinesOut.Add(FRenderableLine((FVector)A, (FVector)B, SphereColor, LineThickness, DepthBias)); });
@@ -40,7 +40,7 @@ void UE::PhysicsTools::InitializePreviewGeometryLines(const FPhysicsDataCollecti
 		const FKBoxElem& Box = AggGeom.BoxElems[Index];
 		FTransform ElemTransform = Box.GetTransform();
 		ElemTransform.ScaleTranslation(PhysicsData.ExternalScale3D);
-		FTransform3f ElemTransformf(ElemTransform);
+		UE::Geometry::FTransform3f ElemTransformf(ElemTransform);
 		FVector3f HalfDimensions(
 			PhysicsData.ExternalScale3D.X * Box.X * 0.5f,
 			PhysicsData.ExternalScale3D.Y * Box.Y * 0.5f,
@@ -56,7 +56,7 @@ void UE::PhysicsTools::InitializePreviewGeometryLines(const FPhysicsDataCollecti
 		const FKSphylElem& Capsule = AggGeom.SphylElems[Index];
 		FTransform ElemTransform = Capsule.GetTransform();
 		ElemTransform.ScaleTranslation(PhysicsData.ExternalScale3D);
-		FTransform3f ElemTransformf(ElemTransform);
+		UE::Geometry::FTransform3f ElemTransformf(ElemTransform);
 		const float HalfLength = Capsule.GetScaledCylinderLength(PhysicsData.ExternalScale3D) * .5f;
 		const float Radius = Capsule.GetScaledRadius(PhysicsData.ExternalScale3D);
 		FVector3f Top(0, 0, HalfLength), Bottom(0, 0, -HalfLength);
diff --git a/Engine/Plugins/Experimental/MeshModelingToolset/Source/MeshModelingTools/Private/Physics/ExtractCollisionGeometryTool.cpp b/Engine/Plugins/Experimental/MeshModelingToolset/Source/MeshModelingTools/Private/Physics/ExtractCollisionGeometryTool.cpp
--- a/Engine/Plugins/Experimental/MeshModelingToolset/Source/MeshModelingTools/Private/Physics/ExtractCollisionGeometryTool.cpp
+++ b/Engine/Plugins/Experimental/MeshModelingToolset/Source/MeshModelingTools/Private/Physics/ExtractCollisionGeometryTool.cpp
@@ -114,7 +114,7 @@ void UExtractCollisionGeometryTool::Shutdown(EToolShutdownType ShutdownType)
 {
 	VizSettings->SaveProperties(this);
 
-	FTransform3d Transform(PreviewMesh->GetTransform());
+	UE::Geometry::FTransform3d Transform(PreviewMesh->GetTransform());
 
 	PreviewElements->Disconnect();
 	PreviewMesh->SetVisible(false);
@@ -236,7 +236,7 @@ void UExtractCollisionGeometryTool::RecalculateMesh()
 
 		MeshTransforms::Translate(CapsuleMesh, FVector3d(0,0,-0.5*Capsule.Length) );
 
-		FTransform3d Transform(Capsule.GetTransform());
+		UE::Geometry::FTransform3d Transform(Capsule.GetTransform());
 		MeshTransforms::ApplyTransform(CapsuleMesh, Transform);
 
 		FMeshIndexMappings Mappings;
@@ -246,7 +246,7 @@ void UExtractCollisionGeometryTool::RecalculateMesh()
 
 	for (const FKConvexElem& Convex : AggGeom.ConvexElems)
 	{
-		FTransform3d ElemTransform(Convex.GetTransform());
+		UE::Geometry::FTransform3d ElemTransform(Convex.GetTransform());
 		FDynamicMesh3 ConvexMesh(EMeshComponents::None);
 		int32 NumVertices = Convex.VertexData.Num();
 		for (int32 k = 0; k < NumVertices; ++k)
diff --git a/Engine/Plugins/Experimental/MeshModelingToolset/Source/MeshModelingTools/Private/Physics/SetCollisionGeometryTool.cpp b/Engine/Plugins/Experimental/MeshModelingToolset/Source/MeshModelingTools/Private/Physics/SetCollisionGeometryTool.cpp
--- a/Engine/Plugins/Experimental/MeshModelingToolset/Source/MeshModelingTools/Private/Physics/SetCollisionGeometryTool.cpp
+++ b/Engine/Plugins/Experimental/MeshModelingToolset/Source/MeshModelingTools/Private/Physics/SetCollisionGeometryTool.cpp
@@ -414,8 +414,8 @@ TArray<const T*> MakeRawPointerList(const TArray<TSharedPtr<T, ESPMode::ThreadSa
 void USetCollisionGeometryTool::PrecomputeInputMeshes()
 {
 	IPrimitiveComponentBackedTarget* CollisionTarget = TargetComponentInterface(Targets.Num()-1);
-	FTransform3d TargetTransform(CollisionTarget->GetWorldTransform());
-	FTransform3d TargetTransformInv = TargetTransform.Inverse();
+	UE::Geometry::FTransform3d TargetTransform(CollisionTarget->GetWorldTransform());
+	UE::Geometry::FTransform3d TargetTransformInv = TargetTransform.Inverse();
 
 	InputMeshes.Reset();
 	InputMeshes.SetNum(SourceObjectIndices.Num());
@@ -428,7 +428,7 @@ void USetCollisionGeometryTool::PrecomputeInputMeshes()
 		Converter.Convert(TargetMeshProviderInterface(k)->GetMeshDescription(), SourceMesh);
 		if (Settings->bUseWorldSpace)
 		{
-			FTransform3d ToWorld(TargetComponentInterface(k)->GetWorldTransform());
+			UE::Geometry::FTransform3d ToWorld(TargetComponentInterface(k)->GetWorldTransform());
 			MeshTransforms::ApplyTransform(SourceMesh, ToWorld);
 			MeshTransforms::ApplyTransform(SourceMesh, TargetTransformInv);
 		}
diff --git a/Engine/Plugins/Experimental/MeshModelingToolset/Source/MeshModelingTools/Private/PlaneCutTool.cpp b/Engine/Plugins/Experimental/MeshModelingToolset/Source/MeshModelingTools/Private/PlaneCutTool.cpp
--- a/Engine/Plugins/Experimental/MeshModelingToolset/Source/MeshModelingTools/Private/PlaneCutTool.cpp
+++ b/Engine/Plugins/Experimental/MeshModelingToolset/Source/MeshModelingTools/Private/PlaneCutTool.cpp
@@ -351,7 +351,7 @@ TUniquePtr<FDynamicMeshOperator> UPlaneCutOperatorFactory::MakeNewOperator()
 	FTransform WorldToLocal = LocalToWorld.Inverse();
 	FVector LocalOrigin = WorldToLocal.TransformPosition(CutTool->CutPlaneOrigin);
 	FVector WorldNormal = CutTool->CutPlaneOrientation.GetAxisZ();
-	FTransform3d W2LForNormal(WorldToLocal);
+	UE::Geometry::FTransform3d W2LForNormal(WorldToLocal);
 	FVector LocalNormal = (FVector)W2LForNormal.TransformNormal(WorldNormal);
 	FVector BackTransformed = LocalToWorld.TransformVector(LocalNormal);
 	float NormalScaleFactor = FVector::DotProduct(BackTransformed, WorldNormal);
diff --git a/Engine/Plugins/Experimental/MeshModelingToolset/Source/MeshModelingTools/Private/RemoveOccludedTrianglesTool.cpp b/Engine/Plugins/Experimental/MeshModelingToolset/Source/MeshModelingTools/Private/RemoveOccludedTrianglesTool.cpp
--- a/Engine/Plugins/Experimental/MeshModelingToolset/Source/MeshModelingTools/Private/RemoveOccludedTrianglesTool.cpp
+++ b/Engine/Plugins/Experimental/MeshModelingToolset/Source/MeshModelingTools/Private/RemoveOccludedTrianglesTool.cpp
@@ -275,7 +275,7 @@ void URemoveOccludedTrianglesTool::SetupPreviews()
 
 			OccluderTrees[TargetIdx] = MakeShared<FDynamicMeshAABBTree3, ESPMode::ThreadSafe>(OriginalDynamicMeshes[PreviewIdx].Get());
 			OccluderWindings[TargetIdx] = MakeShared<TFastWindingTree<FDynamicMesh3>, ESPMode::ThreadSafe>(OccluderTrees[TargetIdx].Get());
-			OccluderTransforms[TargetIdx] = (FTransform3d)TargetComponentInterface(TargetIdx)->GetWorldTransform();
+			OccluderTransforms[TargetIdx] = (UE::Geometry::FTransform3d)TargetComponentInterface(TargetIdx)->GetWorldTransform();
 
 			// configure secondary render material
 			Preview->SecondaryMaterial = OccludedMaterial;
@@ -309,7 +309,7 @@ void URemoveOccludedTrianglesTool::SetupPreviews()
 
 			OccluderTrees[TargetIdx] = OccluderTrees[PreviewToTargetIdx[PreviewIdx]];
 			OccluderWindings[TargetIdx] = OccluderWindings[PreviewToTargetIdx[PreviewIdx]];
-			OccluderTransforms[TargetIdx] = (FTransform3d)TargetComponentInterface(TargetIdx)->GetWorldTransform();
+			OccluderTransforms[TargetIdx] = (UE::Geometry::FTransform3d)TargetComponentInterface(TargetIdx)->GetWorldTransform();
 
 			PreviewMesh->SetSecondaryRenderMaterial(OccludedMaterial);
 			PreviewMesh->EnableSecondaryTriangleBuffers(MoveTemp(IsOccludedGroupFn));
@@ -407,7 +407,7 @@ TUniquePtr<FDynamicMeshOperator> URemoveOccludedTrianglesOperatorFactory::MakeNe
 		int32 TargetIdx = Tool->PreviewToTargetIdx[PreviewIdx];
 		Op->OccluderTrees.Add(Tool->OccluderTrees[TargetIdx]);
 		Op->OccluderWindings.Add(Tool->OccluderWindings[TargetIdx]);
-		Op->OccluderTransforms.Add(FTransform3d::Identity());
+		Op->OccluderTransforms.Add(UE::Geometry::FTransform3d::Identity());
 	}
 	else
 	{
@@ -428,10 +428,10 @@ TUniquePtr<FDynamicMeshOperator> URemoveOccludedTrianglesOperatorFactory::MakeNe
 	
 	Op->SetTransform(LocalToWorld);
 
-	Op->MeshTransforms.Add((FTransform3d)LocalToWorld);
+	Op->MeshTransforms.Add((UE::Geometry::FTransform3d)LocalToWorld);
 	for (int32 CopyIdx : Tool->PreviewToCopyIdx[PreviewIdx])
 	{
-		Op->MeshTransforms.Add((FTransform3d)Tool->PreviewCopies[CopyIdx]->GetTransform());
+		Op->MeshTransforms.Add((UE::Geometry::FTransform3d)Tool->PreviewCopies[CopyIdx]->GetTransform());
 	}
 
 	return Op;
diff --git a/Engine/Plugins/Experimental/MeshModelingToolset/Source/MeshModelingTools/Private/SeamSculptTool.cpp b/Engine/Plugins/Experimental/MeshModelingToolset/Source/MeshModelingTools/Private/SeamSculptTool.cpp
--- a/Engine/Plugins/Experimental/MeshModelingToolset/Source/MeshModelingTools/Private/SeamSculptTool.cpp
+++ b/Engine/Plugins/Experimental/MeshModelingToolset/Source/MeshModelingTools/Private/SeamSculptTool.cpp
@@ -60,7 +60,7 @@ void USeamSculptTool::Setup()
 	BrushProperties->bShowStrength = BrushProperties->bShowFalloff = false;
 	BrushProperties->RestoreProperties(this);
 
-	MeshTransform = FTransform3d(ComponentTarget->GetWorldTransform());
+	MeshTransform = UE::Geometry::FTransform3d(ComponentTarget->GetWorldTransform());
 	InputMesh = MakeShared<FDynamicMesh3, ESPMode::ThreadSafe>(*PreviewMesh->GetMesh());
 	FMeshNormals::QuickComputeVertexNormals(*InputMesh);
 	NormalOffset = InputMesh->GetCachedBounds().MinDim() * 0.001;
diff --git a/Engine/Plugins/Experimental/MeshModelingToolset/Source/MeshModelingTools/Private/VolumeToMeshTool.cpp b/Engine/Plugins/Experimental/MeshModelingToolset/Source/MeshModelingTools/Private/VolumeToMeshTool.cpp
--- a/Engine/Plugins/Experimental/MeshModelingToolset/Source/MeshModelingTools/Private/VolumeToMeshTool.cpp
+++ b/Engine/Plugins/Experimental/MeshModelingToolset/Source/MeshModelingTools/Private/VolumeToMeshTool.cpp
@@ -112,7 +112,7 @@ void UVolumeToMeshTool::Shutdown(EToolShutdownType ShutdownType)
 {
 	Settings->SaveProperties(this);
 
-	FTransform3d Transform(PreviewMesh->GetTransform());
+	UE::Geometry::FTransform3d Transform(PreviewMesh->GetTransform());
 	PreviewMesh->SetVisible(false);
 	PreviewMesh->Disconnect();
 	PreviewMesh = nullptr;
diff --git a/Engine/Plugins/Experimental/MeshModelingToolset/Source/MeshModelingToolsEditorOnly/Private/BspConversionTool.cpp b/Engine/Plugins/Experimental/MeshModelingToolset/Source/MeshModelingToolsEditorOnly/Private/BspConversionTool.cpp
--- a/Engine/Plugins/Experimental/MeshModelingToolset/Source/MeshModelingToolsEditorOnly/Private/BspConversionTool.cpp
+++ b/Engine/Plugins/Experimental/MeshModelingToolset/Source/MeshModelingToolsEditorOnly/Private/BspConversionTool.cpp
@@ -33,9 +33,9 @@ using namespace UE::Geometry;
 void ConvertBrushesToDynamicMesh(TArray<ABrush*>& BrushesToConvert, FDynamicMesh3& OutputMesh, 
 	TArray<UMaterialInterface*>& OutputMaterials);
 void ApplyStaticMeshBooleanOperation(
-	FDynamicMesh3& MeshA, const FTransform3d& TransformA, const TArray<UMaterialInterface*>& MaterialsA,
-	FDynamicMesh3& MeshB, const FTransform3d& TransformB, const TArray<UMaterialInterface*>& MaterialsB,
-	FDynamicMesh3& OutputMesh, FTransform3d& OutputTransform, TArray<UMaterialInterface*>& OutputMaterials,
+	FDynamicMesh3& MeshA, const UE::Geometry::FTransform3d& TransformA, const TArray<UMaterialInterface*>& MaterialsA,
+	FDynamicMesh3& MeshB, const UE::Geometry::FTransform3d& TransformB, const TArray<UMaterialInterface*>& MaterialsB,
+	FDynamicMesh3& OutputMesh, UE::Geometry::FTransform3d& OutputTransform, TArray<UMaterialInterface*>& OutputMaterials,
 	FMeshBoolean::EBooleanOp Operation);
 FText GetBrushGeometryErrorMessage(ABrush* Brush);
 
@@ -215,7 +215,7 @@ void UBspConversionTool::Shutdown(EToolShutdownType ShutdownType)
 			// Create the final actor
 			AActor* NewActor = AssetGenerationUtil::GenerateStaticMeshActor(
 				AssetAPI, TargetWorld,
-				PreviewMesh->GetMesh(), (FTransform3d)PreviewMesh->GetTransform(),
+				PreviewMesh->GetMesh(), (UE::Geometry::FTransform3d)PreviewMesh->GetTransform(),
 				TEXT("BspConvertedToStaticMesh"), Materials);
 
 			if (NewActor)
@@ -362,7 +362,7 @@ bool UBspConversionTool::CombineThenConvert(FText*)
 	}
 
 	MeshTransforms::Translate(OutputMesh, -NewPivot);
-	FTransform3d Transform = FTransform3d::Identity();
+	UE::Geometry::FTransform3d Transform = UE::Geometry::FTransform3d::Identity();
 	Transform.SetTranslation(NewPivot);
 
 	PreviewMesh->UpdatePreview(&OutputMesh);
@@ -392,14 +392,14 @@ bool UBspConversionTool::ConvertThenCombine(FText *ErrorMessage)
 	// We'll need temporary space to store meshes as we build them up. The input and output meshes,
 	// transforms, and material arrays will swap after we merge in each new mesh.
 	FDynamicMesh3 MeshStorage[2];
-	FTransform3d TransformStorage[2];
+	UE::Geometry::FTransform3d TransformStorage[2];
 	TArray<UMaterialInterface*> MaterialArraysStorage[2];
 
 	FDynamicMesh3* InputMesh = &MeshStorage[0];
 	FDynamicMesh3* OutputMesh = &MeshStorage[1];
 	
-	FTransform3d* InputTransform = &TransformStorage[0];
-	FTransform3d* OutputTransform = &TransformStorage[1];
+	UE::Geometry::FTransform3d* InputTransform = &TransformStorage[0];
+	UE::Geometry::FTransform3d* OutputTransform = &TransformStorage[1];
 
 	TArray<UMaterialInterface*>* InputMaterials = &MaterialArraysStorage[0];
 	TArray<UMaterialInterface*>* OutputMaterials = &MaterialArraysStorage[1];
@@ -411,7 +411,7 @@ bool UBspConversionTool::ConvertThenCombine(FText *ErrorMessage)
 	FDynamicMesh3 NextMesh;
 
 	// The transform when a brush is first converted is always the same
-	FTransform3d NextTransform = FTransform3d::Identity();
+	UE::Geometry::FTransform3d NextTransform = UE::Geometry::FTransform3d::Identity();
 
 	// We could actually point to cached materials since we don't change these, but for consistency,
 	// we'll have a local copy.
@@ -444,7 +444,7 @@ bool UBspConversionTool::ConvertThenCombine(FText *ErrorMessage)
 					MakeShared<const TArray<UMaterialInterface*>>(*OutputMaterials)));
 			}
 		}
-		*OutputTransform = FTransform3d::Identity(); // Always identity when first converted
+		*OutputTransform = UE::Geometry::FTransform3d::Identity(); // Always identity when first converted
 
 		// Check validity for the purposes of boolean operations. We'll do this after each brush conversion.
 		if (!OutputMesh->CheckValidity(FDynamicMesh3::FValidityOptions(), EValidityCheckFailMode::ReturnOnly))
@@ -584,9 +584,9 @@ void ConvertBrushesToDynamicMesh(TArray<ABrush*>& BrushesToConvert, FDynamicMesh
 }
 
 void ApplyStaticMeshBooleanOperation(
-	FDynamicMesh3& MeshA, const FTransform3d& TransformA, const TArray<UMaterialInterface*>& MaterialsA,
-	FDynamicMesh3& MeshB, const FTransform3d& TransformB, const TArray<UMaterialInterface*>& MaterialsB,
-	FDynamicMesh3& OutputMesh, FTransform3d& OutputTransform, TArray<UMaterialInterface*>& OutputMaterials,
+	FDynamicMesh3& MeshA, const UE::Geometry::FTransform3d& TransformA, const TArray<UMaterialInterface*>& MaterialsA,
+	FDynamicMesh3& MeshB, const UE::Geometry::FTransform3d& TransformB, const TArray<UMaterialInterface*>& MaterialsB,
+	FDynamicMesh3& OutputMesh, UE::Geometry::FTransform3d& OutputTransform, TArray<UMaterialInterface*>& OutputMaterials,
 	FMeshBoolean::EBooleanOp Operation)
 {
 	// These need to be enabled on both meshes to deal with materials properly. This is relevant, for
diff --git a/Engine/Plugins/Experimental/MeshModelingToolset/Source/MeshModelingToolsEditorOnly/Private/GenerateLODMeshesTool.cpp b/Engine/Plugins/Experimental/MeshModelingToolset/Source/MeshModelingToolsEditorOnly/Private/GenerateLODMeshesTool.cpp
--- a/Engine/Plugins/Experimental/MeshModelingToolset/Source/MeshModelingToolsEditorOnly/Private/GenerateLODMeshesTool.cpp
+++ b/Engine/Plugins/Experimental/MeshModelingToolset/Source/MeshModelingToolsEditorOnly/Private/GenerateLODMeshesTool.cpp
@@ -411,7 +411,7 @@ void UGenerateLODMeshesTool::GenerateAssets()
 
 		AActor* NewActor = AssetGenerationUtil::GenerateStaticMeshActor(
 			AssetAPI, TargetWorld,
-			Result.Mesh.Get(), FTransform3d(Transform), Name, MaterialSet.Materials);
+			Result.Mesh.Get(), UE::Geometry::FTransform3d(Transform), Name, MaterialSet.Materials);
 	}
 
 	GetToolManager()->EndUndoTransaction();
diff --git a/Engine/Plugins/Experimental/MeshModelingToolset/Source/MeshModelingToolsEditorOnly/Private/Hair/GroomCardsEditorTool.cpp b/Engine/Plugins/Experimental/MeshModelingToolset/Source/MeshModelingToolsEditorOnly/Private/Hair/GroomCardsEditorTool.cpp
--- a/Engine/Plugins/Experimental/MeshModelingToolset/Source/MeshModelingToolsEditorOnly/Private/Hair/GroomCardsEditorTool.cpp
+++ b/Engine/Plugins/Experimental/MeshModelingToolset/Source/MeshModelingToolsEditorOnly/Private/Hair/GroomCardsEditorTool.cpp
@@ -560,7 +560,7 @@ void UGroomCardsEditorTool::BeginCardEdit(int32 CardGroupID)
 {
 	const UE::GroomQueries::FMeshCardStrip& CardStrip = EditableCardSet->CardStrips->FindStripForGroup(CardGroupID);
 
-	FTransform3d Transform(PreviewMesh->GetTransform());
+	UE::Geometry::FTransform3d Transform(PreviewMesh->GetTransform());
 	check(FMathd::Abs(Transform.TransformVector(UE::Geometry::Normalized(FVector3d(1,1,1))).SquaredLength() - 1.0) < FMathd::ZeroTolerance);
 	FFrame3d TransformFrame(Transform.GetTranslation(), Transform.GetRotation());
 
@@ -595,7 +595,7 @@ void UGroomCardsEditorTool::BeginCardEdit(int32 CardGroupID)
 void UGroomCardsEditorTool::RestoreCardEdit(const FGroomCardEdit* RestoreEdit)
 {
 	// this should probably be cached once...
-	FTransform3d Transform(PreviewMesh->GetTransform());
+	UE::Geometry::FTransform3d Transform(PreviewMesh->GetTransform());
 	check(FMathd::Abs(Transform.TransformVector(UE::Geometry::Normalized(FVector3d(1, 1, 1))).SquaredLength() - 1.0) < FMathd::ZeroTolerance);
 	FFrame3d TransformFrame(Transform.GetTranslation(), Transform.GetRotation());
 
@@ -643,7 +643,7 @@ void UGroomCardsEditorTool::UpdateOnCurveEdit()
 		//BeginMoveChange();
 	}
 
-	FTransform3d Transform(PreviewMesh->GetTransform());
+	UE::Geometry::FTransform3d Transform(PreviewMesh->GetTransform());
 	check( FMathd::Abs(Transform.TransformVector(UE::Geometry::Normalized(FVector3d(1,1,1))).SquaredLength() - 1.0) < FMathd::ZeroTolerance );
 	FFrame3d TransformFrame(Transform.GetTranslation(), Transform.GetRotation());
 
diff --git a/Engine/Plugins/Experimental/MeshModelingToolset/Source/MeshModelingToolsEditorOnly/Private/Hair/GroomToMeshTool.cpp b/Engine/Plugins/Experimental/MeshModelingToolset/Source/MeshModelingToolsEditorOnly/Private/Hair/GroomToMeshTool.cpp
--- a/Engine/Plugins/Experimental/MeshModelingToolset/Source/MeshModelingToolsEditorOnly/Private/Hair/GroomToMeshTool.cpp
+++ b/Engine/Plugins/Experimental/MeshModelingToolset/Source/MeshModelingToolsEditorOnly/Private/Hair/GroomToMeshTool.cpp
@@ -155,7 +155,7 @@ void UGroomToMeshTool::Shutdown(EToolShutdownType ShutdownType)
 {
 	Settings->SaveProperties(this);
 
-	FTransform3d Transform(PreviewMesh->GetTransform());
+	UE::Geometry::FTransform3d Transform(PreviewMesh->GetTransform());
 	PreviewMesh->Disconnect();
 	PreviewMesh = nullptr;
 
@@ -580,7 +580,7 @@ struct FTargetSpace
 
 	AActor* TargetActor = nullptr;
 	USceneComponent* TargetComponent = nullptr;
-	FTransform3d TargetTransform = FTransform3d::Identity();
+	UE::Geometry::FTransform3d TargetTransform = UE::Geometry::FTransform3d::Identity();
 	FFrame3d TargetFrame = FFrame3d();
 
 	FTargetSpace() {}
@@ -603,7 +603,7 @@ struct FTargetSpace
 		TargetFrame = FrameIn;
 	}
 
-	FTargetSpace(const FTransform3d& TransformIn)
+	FTargetSpace(const UE::Geometry::FTransform3d& TransformIn)
 	{
 		TargetSpace = ETargetSpaceType::Transform;
 		TargetTransform = TransformIn;
@@ -626,22 +626,22 @@ static void TransformToSpace(FDynamicMesh3* Mesh, UPrimitiveComponent* SourceCom
 	// Could also add a variant there that applies multiple transforms in sequence to avoid write overhead
 
 	// transform up to world
-	FTransform3d ToWorld(SourceComponent->GetComponentTransform());
+	UE::Geometry::FTransform3d ToWorld(SourceComponent->GetComponentTransform());
 	MeshTransforms::ApplyTransform(*Mesh, ToWorld);
 
 	if (TargetSpace.TargetSpace == ETargetSpaceType::Actor)
 	{
 		check(TargetSpace.TargetActor);
-		FTransform3d ActorTransform(TargetSpace.TargetActor->GetActorTransform());
-		FTransform3d ToActorTransform = ActorTransform.Inverse();
+		UE::Geometry::FTransform3d ActorTransform(TargetSpace.TargetActor->GetActorTransform());
+		UE::Geometry::FTransform3d ToActorTransform = ActorTransform.Inverse();
 		MeshTransforms::ApplyTransform(*Mesh, ToActorTransform);
 
 	}
 	else if (TargetSpace.TargetSpace == ETargetSpaceType::Component)
 	{
 		check(TargetSpace.TargetComponent);
-		FTransform3d ComponentTransform(TargetSpace.TargetComponent->GetComponentTransform());
-		FTransform3d ToComponentTransform = ComponentTransform.Inverse();
+		UE::Geometry::FTransform3d ComponentTransform(TargetSpace.TargetComponent->GetComponentTransform());
+		UE::Geometry::FTransform3d ToComponentTransform = ComponentTransform.Inverse();
 		MeshTransforms::ApplyTransform(*Mesh, ToComponentTransform);
 	}
 	else if (TargetSpace.TargetSpace == ETargetSpaceType::Frame)
diff --git a/Engine/Plugins/Experimental/MeshModelingToolset/Source/MeshModelingToolsEditorOnly/Private/ParameterizeMeshTool.cpp b/Engine/Plugins/Experimental/MeshModelingToolset/Source/MeshModelingToolsEditorOnly/Private/ParameterizeMeshTool.cpp
--- a/Engine/Plugins/Experimental/MeshModelingToolset/Source/MeshModelingToolsEditorOnly/Private/ParameterizeMeshTool.cpp
+++ b/Engine/Plugins/Experimental/MeshModelingToolset/Source/MeshModelingToolsEditorOnly/Private/ParameterizeMeshTool.cpp
@@ -272,8 +272,7 @@ TUniquePtr<FDynamicMeshOperator> UParameterizeMeshTool::MakeNewOperator()
 
 	ParameterizeMeshOp->Method = (EParamOpBackend)(int)Settings->Method;
 
-	const FTransform XForm = ComponentTarget->GetWorldTransform();
-	FTransform3d XForm3d(XForm);
+	UE::Geometry::FTransform3d XForm3d(ComponentTarget->GetWorldTransform());
 	ParameterizeMeshOp->SetTransform(XForm3d);
 
 	return ParameterizeMeshOp;
diff --git a/Engine/Plugins/Experimental/MeshModelingToolset/Source/MeshModelingToolsEditorOnly/Private/PolygonOnMeshTool.cpp b/Engine/Plugins/Experimental/MeshModelingToolset/Source/MeshModelingToolsEditorOnly/Private/PolygonOnMeshTool.cpp
--- a/Engine/Plugins/Experimental/MeshModelingToolset/Source/MeshModelingToolsEditorOnly/Private/PolygonOnMeshTool.cpp
+++ b/Engine/Plugins/Experimental/MeshModelingToolset/Source/MeshModelingToolsEditorOnly/Private/PolygonOnMeshTool.cpp
@@ -84,7 +84,7 @@ void UPolygonOnMeshTool::Setup()
 	HoverBehavior->Initialize(this);
 	AddInputBehavior(HoverBehavior);
 
-	WorldTransform = FTransform3d(ComponentTarget->GetWorldTransform());
+	WorldTransform = UE::Geometry::FTransform3d(ComponentTarget->GetWorldTransform());
 
 	// hide input StaticMeshComponent
 	ComponentTarget->SetOwnerVisibility(false);
@@ -264,7 +264,7 @@ TUniquePtr<FDynamicMeshOperator> UPolygonOnMeshTool::MakeNewOperator()
 	EmbedOp->bAttemptFixHolesOnBoolean = BasicProperties->bAttemptFixHoles;
 
 	FFrame3d LocalFrame = DrawPlaneWorld;
-	FTransform3d ToLocal = WorldTransform.Inverse();
+	UE::Geometry::FTransform3d ToLocal = WorldTransform.Inverse();
 	LocalFrame.Transform(ToLocal);
 	EmbedOp->PolygonFrame = LocalFrame;
 	
diff --git a/Engine/Plugins/Experimental/MeshModelingToolset/Source/MeshModelingToolsEditorOnly/Private/ShapeSprayTool.cpp b/Engine/Plugins/Experimental/MeshModelingToolset/Source/MeshModelingToolsEditorOnly/Private/ShapeSprayTool.cpp
--- a/Engine/Plugins/Experimental/MeshModelingToolset/Source/MeshModelingToolsEditorOnly/Private/ShapeSprayTool.cpp
+++ b/Engine/Plugins/Experimental/MeshModelingToolset/Source/MeshModelingToolsEditorOnly/Private/ShapeSprayTool.cpp
@@ -232,7 +232,7 @@ bool UShapeSprayTool::CanAccept() const
 void UShapeSprayTool::EmitResult()
 {
 	const FDynamicMesh3* Mesh = AccumMeshComponent->GetMesh();
-	FTransform3d UseTransform = FTransform3d(ComponentTarget->GetOwnerActor()->GetTransform());
+	UE::Geometry::FTransform3d UseTransform(ComponentTarget->GetOwnerActor()->GetTransform());
 
 	GetToolManager()->BeginUndoTransaction(LOCTEXT("EmitShapeSpray", "Create ShapeSpray"));
 
diff --git a/Engine/Plugins/Experimental/MeshModelingToolset/Source/ModelingComponents/Private/Drawing/UVLayoutPreview.cpp b/Engine/Plugins/Experimental/MeshModelingToolset/Source/ModelingComponents/Private/Drawing/UVLayoutPreview.cpp
--- a/Engine/Plugins/Experimental/MeshModelingToolset/Source/ModelingComponents/Private/Drawing/UVLayoutPreview.cpp
+++ b/Engine/Plugins/Experimental/MeshModelingToolset/Source/ModelingComponents/Private/Drawing/UVLayoutPreview.cpp
@@ -217,7 +217,7 @@ void UUVLayoutPreview::RecalculatePosition()
 
 	CurrentWorldFrame = ObjFrame;
 
-	FTransform3d Transform(ObjFrame.Rotation, ObjFrame.Origin);
+	UE::Geometry::FTransform3d Transform(ObjFrame.Rotation, ObjFrame.Origin);
 	Transform.SetScale(UseScale * FVector3d::One());
 
 	SetTransform((FTransform)Transform);
diff --git a/Engine/Plugins/Experimental/MeshModelingToolset/Source/ModelingComponents/Private/Mechanics/DragAlignmentMechanic.cpp b/Engine/Plugins/Experimental/MeshModelingToolset/Source/ModelingComponents/Private/Mechanics/DragAlignmentMechanic.cpp
--- a/Engine/Plugins/Experimental/MeshModelingToolset/Source/ModelingComponents/Private/Mechanics/DragAlignmentMechanic.cpp
+++ b/Engine/Plugins/Experimental/MeshModelingToolset/Source/ModelingComponents/Private/Mechanics/DragAlignmentMechanic.cpp
@@ -178,7 +178,7 @@ void UDragAlignmentMechanic::OnGizmoTransformChanged(FTransform NewTransform)
 }
 
 void UDragAlignmentMechanic::InitializeDeformedMeshRayCast(TFunction<FDynamicMeshAABBTree3* ()> GetSpatialIn, 
-	const FTransform3d &TargetTransform, const FGroupTopologyDeformer* LinearDeformer)
+	const UE::Geometry::FTransform3d &TargetTransform, const FGroupTopologyDeformer* LinearDeformer)
 {
 	// Use the deformer to create a way to filter out triangles we don't want to be hitting (the ones we're manipulating)
 	const FDynamicMesh3* Mesh = LinearDeformer->GetMesh();
diff --git a/Engine/Plugins/Experimental/MeshModelingToolset/Source/ModelingComponents/Private/MeshOpPreviewHelpers.cpp b/Engine/Plugins/Experimental/MeshModelingToolset/Source/ModelingComponents/Private/MeshOpPreviewHelpers.cpp
--- a/Engine/Plugins/Experimental/MeshModelingToolset/Source/ModelingComponents/Private/MeshOpPreviewHelpers.cpp
+++ b/Engine/Plugins/Experimental/MeshModelingToolset/Source/ModelingComponents/Private/MeshOpPreviewHelpers.cpp
@@ -21,7 +21,7 @@ FDynamicMeshOpResult UMeshOpPreviewWithBackgroundCompute::Shutdown()
 
 	FDynamicMeshOpResult Result{};
 	Result.Mesh = PreviewMesh->ExtractPreviewMesh();
-	Result.Transform = FTransform3d(PreviewMesh->GetTransform());
+	Result.Transform = UE::Geometry::FTransform3d(PreviewMesh->GetTransform());
 
 	PreviewMesh->SetVisible(false);
 	PreviewMesh->Disconnect();
diff --git a/Engine/Plugins/Experimental/MeshModelingToolset/Source/ModelingComponents/Private/OctreeDynamicMeshComponent.cpp b/Engine/Plugins/Experimental/MeshModelingToolset/Source/ModelingComponents/Private/OctreeDynamicMeshComponent.cpp
--- a/Engine/Plugins/Experimental/MeshModelingToolset/Source/ModelingComponents/Private/OctreeDynamicMeshComponent.cpp
+++ b/Engine/Plugins/Experimental/MeshModelingToolset/Source/ModelingComponents/Private/OctreeDynamicMeshComponent.cpp
@@ -94,7 +94,7 @@ void UOctreeDynamicMeshComponent::InitializeNewMesh()
 
 
 
-void UOctreeDynamicMeshComponent::ApplyTransform(const FTransform3d& Transform, bool bInvert)
+void UOctreeDynamicMeshComponent::ApplyTransform(const UE::Geometry::FTransform3d& Transform, bool bInvert)
 {
 	if (bInvert)
 	{
diff --git a/Engine/Plugins/Experimental/MeshModelingToolset/Source/ModelingComponents/Private/PreviewMesh.cpp b/Engine/Plugins/Experimental/MeshModelingToolset/Source/ModelingComponents/Private/PreviewMesh.cpp
--- a/Engine/Plugins/Experimental/MeshModelingToolset/Source/ModelingComponents/Private/PreviewMesh.cpp
+++ b/Engine/Plugins/Experimental/MeshModelingToolset/Source/ModelingComponents/Private/PreviewMesh.cpp
@@ -316,7 +316,7 @@ bool UPreviewMesh::FindRayIntersection(const FRay3d& WorldRay, FHitResult& HitOu
 {
 	if (IsVisible() && TemporaryParentActor != nullptr && bBuildSpatialDataStructure)
 	{
-		FTransform3d Transform(TemporaryParentActor->GetActorTransform());
+		UE::Geometry::FTransform3d Transform(TemporaryParentActor->GetActorTransform());
 		FRay3d LocalRay(Transform.InverseTransformPosition(WorldRay.Origin),
 			Transform.InverseTransformVector(WorldRay.Direction));
 		UE::Geometry::Normalize(LocalRay.Direction);
@@ -343,7 +343,7 @@ bool UPreviewMesh::FindRayIntersection(const FRay3d& WorldRay, FHitResult& HitOu
 
 FVector3d UPreviewMesh::FindNearestPoint(const FVector3d& WorldPoint, bool bLinearSearch)
 {
-	FTransform3d Transform(TemporaryParentActor->GetActorTransform());
+	UE::Geometry::FTransform3d Transform(TemporaryParentActor->GetActorTransform());
 	FVector3d LocalPoint = Transform.TransformPosition(WorldPoint);
 	const FDynamicMesh3* UseMesh = GetMesh();
 	if (bLinearSearch)
diff --git a/Engine/Plugins/Experimental/MeshModelingToolset/Source/ModelingComponents/Private/Selection/GroupTopologySelector.cpp b/Engine/Plugins/Experimental/MeshModelingToolset/Source/ModelingComponents/Private/Selection/GroupTopologySelector.cpp
--- a/Engine/Plugins/Experimental/MeshModelingToolset/Source/ModelingComponents/Private/Selection/GroupTopologySelector.cpp
+++ b/Engine/Plugins/Experimental/MeshModelingToolset/Source/ModelingComponents/Private/Selection/GroupTopologySelector.cpp
@@ -438,7 +438,7 @@ bool FGroupTopologySelector::DoEdgeBasedSelection(const FSelectionSettings& Sett
 		// to do this for perspective mode, we would have picked a plane through the edge and the camera origin.
 		int32 ClosestEid = Topology->GetGroupEdgeEdges(ClosestElement->ID)[ClosestElement->PolySegmentIdx];
 		FDynamicMesh3::FEdge ClosestEdge = Mesh->GetEdge(ClosestEid);
-		FPlane3d PlaneThroughClosestEdge(Mesh->GetVertex(ClosestEdge.Vert.A), Mesh->GetVertex(ClosestEdge.Vert.B), 
+		UE::Geometry::FPlane3d PlaneThroughClosestEdge(Mesh->GetVertex(ClosestEdge.Vert.A), Mesh->GetVertex(ClosestEdge.Vert.B), 
 			Mesh->GetVertex(ClosestEdge.Vert.A) + Ray.Direction);
 
 		for (const FGeometrySet3::FNearest& Element : ElementsWithinTolerance)
@@ -473,7 +473,7 @@ bool FGroupTopologySelector::DoEdgeBasedSelection(const FSelectionSettings& Sett
 
 		// Grab a plane through the two verts that contains the ray direction (again, this is assuming that we'd
 		// only do this in ortho mode, otherwise the plane would go through the ray origin.
-		FPlane3d PlaneThroughClosestEdge(VertA, VertB, VertA + Ray.Direction);
+		UE::Geometry::FPlane3d PlaneThroughClosestEdge(VertA, VertB, VertA + Ray.Direction);
 
 		// Checking that the face is coplanar simply entails checking that the opposite vert is in the plane. However,
 		// it is possible even for the closest edge to have multiple coplanar faces if a group going towards
diff --git a/Engine/Plugins/Experimental/MeshModelingToolset/Source/ModelingComponents/Private/SimpleDynamicMeshComponent.cpp b/Engine/Plugins/Experimental/MeshModelingToolset/Source/ModelingComponents/Private/SimpleDynamicMeshComponent.cpp
--- a/Engine/Plugins/Experimental/MeshModelingToolset/Source/ModelingComponents/Private/SimpleDynamicMeshComponent.cpp
+++ b/Engine/Plugins/Experimental/MeshModelingToolset/Source/ModelingComponents/Private/SimpleDynamicMeshComponent.cpp
@@ -162,7 +162,7 @@ void USimpleDynamicMeshComponent::InitializeNewMesh()
 }
 
 
-void USimpleDynamicMeshComponent::ApplyTransform(const FTransform3d& Transform, bool bInvert)
+void USimpleDynamicMeshComponent::ApplyTransform(const UE::Geometry::FTransform3d& Transform, bool bInvert)
 {
 	if (bInvert)
 	{
diff --git a/Engine/Plugins/Experimental/PlanarCutPlugin/Source/PlanarCut/Private/PlanarCut.cpp b/Engine/Plugins/Experimental/PlanarCutPlugin/Source/PlanarCut/Private/PlanarCut.cpp
--- a/Engine/Plugins/Experimental/PlanarCutPlugin/Source/PlanarCut/Private/PlanarCut.cpp
+++ b/Engine/Plugins/Experimental/PlanarCutPlugin/Source/PlanarCut/Private/PlanarCut.cpp
@@ -379,7 +379,7 @@ struct FCellMeshes
 
 		if (Transform.IsSet())
 		{
-			MeshTransforms::ApplyTransform(SingleCutter, FTransform3d(Transform.GetValue()));
+			MeshTransforms::ApplyTransform(SingleCutter, UE::Geometry::FTransform3d(Transform.GetValue()));
 		}
 
 		// Mesh should already be augmented
@@ -1338,7 +1338,7 @@ struct FDynamicMeshCollection
 				continue;
 			}
 
-			FTransform3d CollectionToLocal = FTransform3d(GeometryCollectionAlgo::GlobalMatrix(Collection->Transform, Collection->Parent, TransformIdx) * TransformCollection);
+			UE::Geometry::FTransform3d CollectionToLocal = UE::Geometry::FTransform3d(GeometryCollectionAlgo::GlobalMatrix(Collection->Transform, Collection->Parent, TransformIdx) * TransformCollection);
 
 			int32 AddedMeshIdx = Meshes.Add(new FMeshData);
 			FMeshData& MeshData = Meshes[AddedMeshIdx];
