ModelingTools: ToolTarget refactor for Sculpt tools

#rnx

diff --git a/Engine/Plugins/Experimental/MeshModelingToolset/Source/MeshModelingTools/Private/DeformMeshPolygonsTool.cpp b/Engine/Plugins/Experimental/MeshModelingToolset/Source/MeshModelingTools/Private/DeformMeshPolygonsTool.cpp
--- a/Engine/Plugins/Experimental/MeshModelingToolset/Source/MeshModelingTools/Private/DeformMeshPolygonsTool.cpp
+++ b/Engine/Plugins/Experimental/MeshModelingToolset/Source/MeshModelingTools/Private/DeformMeshPolygonsTool.cpp
@@ -15,6 +15,11 @@
 #include "ToolBuilderUtil.h"
 #include "ToolSceneQueriesUtil.h"
 
+#include "TargetInterfaces/MaterialProvider.h"
+#include "TargetInterfaces/MeshDescriptionCommitter.h"
+#include "TargetInterfaces/MeshDescriptionProvider.h"
+#include "TargetInterfaces/PrimitiveComponentBackedTarget.h"
+
 #include "ExplicitUseGeometryMathTypes.h"		// using UE::Geometry::(math types)
 using namespace UE::Geometry;
 
@@ -849,14 +854,15 @@ void UDeformMeshPolygonsTool::Setup()
 	LaplacianDeformer = MakePimpl<FGroupTopologyLaplacianDeformer>();
 
 	// create dynamic mesh component to use for live preview
-	DynamicMeshComponent = NewObject<USimpleDynamicMeshComponent>(ComponentTarget->GetOwnerActor(), "DynamicMesh");
-	DynamicMeshComponent->SetupAttachment(ComponentTarget->GetOwnerActor()->GetRootComponent());
+	IPrimitiveComponentBackedTarget* TargetComponent = Cast<IPrimitiveComponentBackedTarget>(Target);
+	DynamicMeshComponent = NewObject<USimpleDynamicMeshComponent>(TargetComponent->GetOwnerActor(), "DynamicMesh");
+	DynamicMeshComponent->SetupAttachment(TargetComponent->GetOwnerActor()->GetRootComponent());
 	DynamicMeshComponent->RegisterComponent();
-	DynamicMeshComponent->SetWorldTransform(ComponentTarget->GetWorldTransform());
+	DynamicMeshComponent->SetWorldTransform(TargetComponent->GetWorldTransform());
 
 	// set materials
 	FComponentMaterialSet MaterialSet;
-	ComponentTarget->GetMaterialSet(MaterialSet);
+	Cast<IMaterialProvider>(Target)->GetMaterialSet(MaterialSet);
 	for (int k = 0; k < MaterialSet.Materials.Num(); ++k)
 	{
 		DynamicMeshComponent->SetMaterial(k, MaterialSet.Materials[k]);
@@ -864,7 +870,7 @@ void UDeformMeshPolygonsTool::Setup()
 
 	// dynamic mesh configuration settings
 	DynamicMeshComponent->TangentsType = EDynamicMeshTangentCalcType::AutoCalculated;
-	DynamicMeshComponent->InitializeMesh(ComponentTarget->GetMesh());
+	DynamicMeshComponent->InitializeMesh(Cast<IMeshDescriptionProvider>(Target)->GetMeshDescription());
 	OnDynamicMeshComponentChangedHandle =
 	    DynamicMeshComponent->OnMeshChanged.Add(FSimpleMulticastDelegate::FDelegate::CreateUObject(
 	        this, &UDeformMeshPolygonsTool::OnDynamicMeshComponentChanged));
@@ -881,14 +887,14 @@ void UDeformMeshPolygonsTool::Setup()
 	//initialize topology selector
 	TopoSelector.Initialize(DynamicMeshComponent->GetMesh(), &Topology);
 	TopoSelector.SetSpatialSource([this]() { return &GetSpatial(); });
-	TopoSelector.PointsWithinToleranceTest = [this](const FVector3d& Position1, const FVector3d& Position2, double TolScale) {
-		UE::Geometry::FTransform3d Transform(ComponentTarget->GetWorldTransform());
+	TopoSelector.PointsWithinToleranceTest = [this, TargetComponent](const FVector3d& Position1, const FVector3d& Position2, double TolScale) {
+		UE::Geometry::FTransform3d Transform(TargetComponent->GetWorldTransform());
 		return ToolSceneQueriesUtil::PointSnapQuery(CameraState, Transform.TransformPosition(Position1), Transform.TransformPosition(Position2),
 			ToolSceneQueriesUtil::GetDefaultVisualAngleSnapThreshD() * TolScale);
 	};
 
 	// hide input StaticMeshComponent
-	ComponentTarget->SetOwnerVisibility(false);
+	TargetComponent->SetOwnerVisibility(false);
 
 	// init state flags flags
 	bInDrag = false;
@@ -947,17 +953,17 @@ void UDeformMeshPolygonsTool::Shutdown(EToolShutdownType ShutdownType)
 	{
 		DynamicMeshComponent->OnMeshChanged.Remove(OnDynamicMeshComponentChangedHandle);
 
-		ComponentTarget->SetOwnerVisibility(true);
+		Cast<IPrimitiveComponentBackedTarget>(Target)->SetOwnerVisibility(true);
 
 		if (ShutdownType == EToolShutdownType::Accept)
 		{
 			// this block bakes the modified DynamicMeshComponent back into the StaticMeshComponent inside an undo transaction
 			GetToolManager()->BeginUndoTransaction(LOCTEXT("DeformMeshPolygonsToolTransactionName", "Deform Mesh"));
-			ComponentTarget->CommitMesh([=](const FPrimitiveComponentTarget::FCommitParams& CommitParams) {
+			Cast<IMeshDescriptionCommitter>(Target)->CommitMeshDescription([=](const IMeshDescriptionCommitter::FCommitterParams& CommitParams) {
 				FConversionToMeshDescriptionOptions ConversionOptions;
 				ConversionOptions.bSetPolyGroups =
 				    false; // don't save polygroups, as we may change these temporarily in this tool just to get a different edit effect
-				DynamicMeshComponent->Bake(CommitParams.MeshDescription, false, ConversionOptions);
+				DynamicMeshComponent->Bake(CommitParams.MeshDescriptionOut, false, ConversionOptions);
 			});
 			GetToolManager()->EndUndoTransaction();
 		}
@@ -1034,7 +1040,7 @@ FDynamicMeshAABBTree3& UDeformMeshPolygonsTool::GetSpatial()
 
 bool UDeformMeshPolygonsTool::HitTest(const FRay& WorldRay, FHitResult& OutHit)
 {
-	UE::Geometry::FTransform3d Transform(ComponentTarget->GetWorldTransform());
+	UE::Geometry::FTransform3d Transform(Cast<IPrimitiveComponentBackedTarget>(Target)->GetWorldTransform());
 	FRay3d LocalRay(Transform.InverseTransformPosition(WorldRay.Origin),
 	                Transform.InverseTransformVector(WorldRay.Direction));
 	UE::Geometry::Normalize(LocalRay.Direction);
@@ -1080,7 +1086,7 @@ bool UDeformMeshPolygonsTool::HitTest(const FRay& WorldRay, FHitResult& OutHit)
 
 void UDeformMeshPolygonsTool::OnBeginDrag(const FRay& WorldRay)
 {
-	UE::Geometry::FTransform3d Transform(ComponentTarget->GetWorldTransform());
+	UE::Geometry::FTransform3d Transform(Cast<IPrimitiveComponentBackedTarget>(Target)->GetWorldTransform());
 	FRay3d LocalRay(Transform.InverseTransformPosition(WorldRay.Origin),
 	                Transform.InverseTransformVector(WorldRay.Direction));
 	UE::Geometry::Normalize(LocalRay.Direction);
@@ -1189,7 +1195,7 @@ void UDeformMeshPolygonsTool::OnBeginDrag(const FRay& WorldRay)
 
 void UDeformMeshPolygonsTool::UpdateActiveSurfaceFrame(FGroupTopologySelection& Selection)
 {
-	UE::Geometry::FTransform3d Transform(ComponentTarget->GetWorldTransform());
+	UE::Geometry::FTransform3d Transform(Cast<IPrimitiveComponentBackedTarget>(Target)->GetWorldTransform());
 
 	// update surface frame
 	ActiveSurfaceFrame.Origin = StartHitPosWorld;
@@ -1295,7 +1301,7 @@ bool UDeformMeshPolygonsTool::OnUpdateHover(const FInputDeviceRay& DevicePos)
 	//if (!bNeedEmitEndChange)
 	if (ActiveVertexChange == nullptr)
 	{
-		UE::Geometry::FTransform3d Transform(ComponentTarget->GetWorldTransform());
+		UE::Geometry::FTransform3d Transform(Cast<IPrimitiveComponentBackedTarget>(Target)->GetWorldTransform());
 		FRay3d LocalRay(Transform.InverseTransformPosition(DevicePos.WorldRay.Origin),
 		                Transform.InverseTransformVector(DevicePos.WorldRay.Direction));
 		UE::Geometry::Normalize(LocalRay.Direction);
@@ -1393,7 +1399,7 @@ void UDeformMeshPolygonsTool::ComputeUpdate_Rotate()
 	FGroupTopologyDeformer& SelectedDeformer = (bIsLaplacian) ? *LaplacianDeformer : LinearDeformer;
 
 	FDynamicMesh3* Mesh    = DynamicMeshComponent->GetMesh();
-	UE::Geometry::FTransform3d Transform(ComponentTarget->GetWorldTransform());
+	UE::Geometry::FTransform3d Transform(Cast<IPrimitiveComponentBackedTarget>(Target)->GetWorldTransform());
 	FVector NewHitPosWorld = LastHitPosWorld;
 
 	FVector3d SnappedPoint;
@@ -1487,7 +1493,7 @@ void UDeformMeshPolygonsTool::ComputeUpdate_Translate()
 		};
 	}
 
-	FTransform Transform   = ComponentTarget->GetWorldTransform();
+	FTransform Transform   = Cast<IPrimitiveComponentBackedTarget>(Target)->GetWorldTransform();
 	FVector NewHitPosWorld = LastHitPosWorld;
 	FVector3d SnappedPoint;
 	if (QuickAxisTranslater.UpdateSnap(FRay3d(UpdateRay), SnappedPoint, PointConstraintFunc))
@@ -1566,7 +1572,7 @@ void UDeformMeshPolygonsTool::Render(IToolsContextRenderAPI* RenderAPI)
 	FDynamicMesh3* TargetMesh                    = DynamicMeshComponent->GetMesh();
 
 	PolyEdgesRenderer.BeginFrame(RenderAPI, CameraState);
-	PolyEdgesRenderer.SetTransform(ComponentTarget->GetWorldTransform());
+	PolyEdgesRenderer.SetTransform(Cast<IPrimitiveComponentBackedTarget>(Target)->GetWorldTransform());
 
 
 	for (FGroupTopology::FGroupEdge& Edge : Topology.Edges)
@@ -1583,7 +1589,7 @@ void UDeformMeshPolygonsTool::Render(IToolsContextRenderAPI* RenderAPI)
 
 
 	HilightRenderer.BeginFrame(RenderAPI, CameraState);
-	HilightRenderer.SetTransform(ComponentTarget->GetWorldTransform());
+	HilightRenderer.SetTransform(Cast<IPrimitiveComponentBackedTarget>(Target)->GetWorldTransform());
 
 #ifdef DEBUG_ROI_WEIGHTS
 	FDynamicMesh3* MeshPtr = DynamicMeshComponent->GetMesh();
diff --git a/Engine/Plugins/Experimental/MeshModelingToolset/Source/MeshModelingTools/Private/DynamicMeshBrushTool.cpp b/Engine/Plugins/Experimental/MeshModelingToolset/Source/MeshModelingTools/Private/DynamicMeshBrushTool.cpp
--- a/Engine/Plugins/Experimental/MeshModelingToolset/Source/MeshModelingTools/Private/DynamicMeshBrushTool.cpp
+++ b/Engine/Plugins/Experimental/MeshModelingToolset/Source/MeshModelingTools/Private/DynamicMeshBrushTool.cpp
@@ -4,6 +4,10 @@
 #include "InteractiveToolManager.h"
 #include "ToolBuilderUtil.h"
 
+#include "TargetInterfaces/MaterialProvider.h"
+#include "TargetInterfaces/MeshDescriptionProvider.h"
+#include "TargetInterfaces/PrimitiveComponentBackedTarget.h"
+
 #include "ExplicitUseGeometryMathTypes.h"		// using UE::Geometry::(math types)
 using namespace UE::Geometry;
 
@@ -21,28 +25,29 @@ UDynamicMeshBrushTool::UDynamicMeshBrushTool()
 
 void UDynamicMeshBrushTool::Setup()
 {
+	IPrimitiveComponentBackedTarget* TargetComponent = Cast<IPrimitiveComponentBackedTarget>(Target);
 	PreviewMesh = NewObject<UPreviewMesh>(this);
 	PreviewMesh->bBuildSpatialDataStructure = true;
-	PreviewMesh->CreateInWorld(ComponentTarget->GetOwnerActor()->GetWorld(), FTransform::Identity);
-	PreviewMesh->SetTransform(ComponentTarget->GetWorldTransform());
+	PreviewMesh->CreateInWorld(TargetComponent->GetOwnerActor()->GetWorld(), FTransform::Identity);
+	PreviewMesh->SetTransform(TargetComponent->GetWorldTransform());
 
 	FComponentMaterialSet MaterialSet;
-	ComponentTarget->GetMaterialSet(MaterialSet);
+	Cast<IMaterialProvider>(Target)->GetMaterialSet(MaterialSet);
 	PreviewMesh->SetMaterials(MaterialSet.Materials);
 
 	// initialize from LOD-0 MeshDescription
-	PreviewMesh->InitializeMesh(ComponentTarget->GetMesh());
+	PreviewMesh->InitializeMesh(Cast<IMeshDescriptionProvider>(Target)->GetMeshDescription());
 	OnBaseMeshComponentChangedHandle = PreviewMesh->GetOnMeshChanged().Add(
 		FSimpleMulticastDelegate::FDelegate::CreateUObject(this, &UDynamicMeshBrushTool::OnBaseMeshComponentChanged));
 
 	// call this here so that base tool can estimate target dimension
 	InputMeshBoundsLocal = PreviewMesh->GetPreviewDynamicMesh()->GetBounds();
-	double ScaledDim = ComponentTarget->GetWorldTransform().TransformVector(FVector::OneVector).Size();
+	double ScaledDim = TargetComponent->GetWorldTransform().TransformVector(FVector::OneVector).Size();
 	this->WorldToLocalScale = FMathd::Sqrt3 / FMathd::Max(FMathf::ZeroTolerance, ScaledDim);
 	UBaseBrushTool::Setup();
 
 	// hide input StaticMeshComponent
-	ComponentTarget->SetOwnerVisibility(false);
+	TargetComponent->SetOwnerVisibility(false);
 }
 
 
@@ -57,7 +62,7 @@ void UDynamicMeshBrushTool::Shutdown(EToolShutdownType ShutdownType)
 {
 	UBaseBrushTool::Shutdown(ShutdownType);
 
-	ComponentTarget->SetOwnerVisibility(true);
+	Cast<IPrimitiveComponentBackedTarget>(Target)->SetOwnerVisibility(true);
 
 	if (PreviewMesh != nullptr)
 	{
diff --git a/Engine/Plugins/Experimental/MeshModelingToolset/Source/MeshModelingTools/Private/DynamicMeshSculptTool.cpp b/Engine/Plugins/Experimental/MeshModelingToolset/Source/MeshModelingTools/Private/DynamicMeshSculptTool.cpp
--- a/Engine/Plugins/Experimental/MeshModelingToolset/Source/MeshModelingTools/Private/DynamicMeshSculptTool.cpp
+++ b/Engine/Plugins/Experimental/MeshModelingToolset/Source/MeshModelingTools/Private/DynamicMeshSculptTool.cpp
@@ -36,6 +36,11 @@
 #include "Materials/Material.h"
 #include "Materials/MaterialInstanceDynamic.h"
 
+#include "TargetInterfaces/MaterialProvider.h"
+#include "TargetInterfaces/MeshDescriptionCommitter.h"
+#include "TargetInterfaces/MeshDescriptionProvider.h"
+#include "TargetInterfaces/PrimitiveComponentBackedTarget.h"
+
 #include "ExplicitUseGeometryMathTypes.h"		// using UE::Geometry::(math types)
 using namespace UE::Geometry;
 
@@ -99,16 +104,17 @@ void UDynamicMeshSculptTool::Setup()
 	SetToolDisplayName(LOCTEXT("ToolName", "DynaSculpt"));
 
 	// create dynamic mesh component to use for live preview
-	DynamicMeshComponent = NewObject<UOctreeDynamicMeshComponent>(ComponentTarget->GetOwnerActor(), "DynamicMeshSculptToolMesh");
-	DynamicMeshComponent->SetupAttachment(ComponentTarget->GetOwnerActor()->GetRootComponent());
+	IPrimitiveComponentBackedTarget* TargetComponent = Cast<IPrimitiveComponentBackedTarget>(Target);
+	DynamicMeshComponent = NewObject<UOctreeDynamicMeshComponent>(TargetComponent->GetOwnerActor(), "DynamicMeshSculptToolMesh");
+	DynamicMeshComponent->SetupAttachment(TargetComponent->GetOwnerActor()->GetRootComponent());
 	DynamicMeshComponent->RegisterComponent();
 
 	// initialize from LOD-0 MeshDescription
-	DynamicMeshComponent->InitializeMesh(ComponentTarget->GetMesh());
+	DynamicMeshComponent->InitializeMesh(Cast<IMeshDescriptionProvider>(Target)->GetMeshDescription());
 
 	// transform mesh to world space because handling scaling inside brush is a mess
 	// Note: this transform does not include translation ( so only the 3x3 transform)
-	InitialTargetTransform = UE::Geometry::FTransform3d(ComponentTarget->GetWorldTransform());
+	InitialTargetTransform = UE::Geometry::FTransform3d(TargetComponent->GetWorldTransform());
 	// clamp scaling because if we allow zero-scale we cannot invert this transform on Accept
 	InitialTargetTransform.ClampMinimumScale(0.01);
 	FVector3d Translation = InitialTargetTransform.GetTranslation();
@@ -119,7 +125,7 @@ void UDynamicMeshSculptTool::Setup()
 	DynamicMeshComponent->SetWorldTransform((FTransform)CurTargetTransform);
 
 	// copy material if there is one
-	UMaterialInterface* Material = ComponentTarget->GetMaterial(0);
+	UMaterialInterface* Material = Cast<IMaterialProvider>(Target)->GetMaterial(0);
 	if (Material != nullptr)
 	{
 		DynamicMeshComponent->SetMaterial(0, Material);
@@ -157,7 +163,7 @@ void UDynamicMeshSculptTool::Setup()
 	InitialEdgeLength = EstimateIntialSafeTargetLength(*Mesh, 5000);
 
 	// hide input Component
-	ComponentTarget->SetOwnerVisibility(false);
+	TargetComponent->SetOwnerVisibility(false);
 
 	// init state flags flags
 	bInDrag = false;
@@ -286,7 +292,7 @@ void UDynamicMeshSculptTool::Shutdown(EToolShutdownType ShutdownType)
 	{
 		DynamicMeshComponent->OnMeshChanged.Remove(OnDynamicMeshComponentChangedHandle);
 
-		ComponentTarget->SetOwnerVisibility(true);
+		Cast<IPrimitiveComponentBackedTarget>(Target)->SetOwnerVisibility(true);
 
 		if (ShutdownType == EToolShutdownType::Accept)
 		{
@@ -295,9 +301,9 @@ void UDynamicMeshSculptTool::Shutdown(EToolShutdownType ShutdownType)
 
 			// this block bakes the modified DynamicMeshComponent back into the StaticMeshComponent inside an undo transaction
 			GetToolManager()->BeginUndoTransaction(LOCTEXT("SculptMeshToolTransactionName", "Sculpt Mesh"));
-			ComponentTarget->CommitMesh([=](const FPrimitiveComponentTarget::FCommitParams& CommitParams)
+			Cast<IMeshDescriptionCommitter>(Target)->CommitMeshDescription([=](const IMeshDescriptionCommitter::FCommitterParams& CommitParams)
 			{
-				DynamicMeshComponent->Bake(CommitParams.MeshDescription, bHaveRemeshed);
+				DynamicMeshComponent->Bake(CommitParams.MeshDescriptionOut, bHaveRemeshed);
 			});
 			GetToolManager()->EndUndoTransaction();
 		}
@@ -2439,7 +2445,7 @@ void UDynamicMeshSculptTool::UpdateMaterialMode(EMeshEditingMaterialModes Materi
 	if (MaterialMode == EMeshEditingMaterialModes::ExistingMaterial)
 	{
 		DynamicMeshComponent->ClearOverrideRenderMaterial();
-		DynamicMeshComponent->bCastDynamicShadow = ComponentTarget->GetOwnerComponent()->bCastDynamicShadow;
+		DynamicMeshComponent->bCastDynamicShadow = Cast<IPrimitiveComponentBackedTarget>(Target)->GetOwnerComponent()->bCastDynamicShadow;
 		ActiveOverrideMaterial = nullptr;
 	}
 	else 
diff --git a/Engine/Plugins/Experimental/MeshModelingToolset/Source/MeshModelingTools/Private/EditMeshPolygonsTool.cpp b/Engine/Plugins/Experimental/MeshModelingToolset/Source/MeshModelingTools/Private/EditMeshPolygonsTool.cpp
--- a/Engine/Plugins/Experimental/MeshModelingToolset/Source/MeshModelingTools/Private/EditMeshPolygonsTool.cpp
+++ b/Engine/Plugins/Experimental/MeshModelingToolset/Source/MeshModelingTools/Private/EditMeshPolygonsTool.cpp
@@ -40,6 +40,11 @@
 #include "Materials/MaterialInstanceDynamic.h"
 #include "Components/BrushComponent.h"
 
+#include "TargetInterfaces/MaterialProvider.h"
+#include "TargetInterfaces/MeshDescriptionCommitter.h"
+#include "TargetInterfaces/MeshDescriptionProvider.h"
+#include "TargetInterfaces/PrimitiveComponentBackedTarget.h"
+
 #include "ExplicitUseGeometryMathTypes.h"		// using UE::Geometry::(math types)
 using namespace UE::Geometry;
 
@@ -165,15 +170,16 @@ void UEditMeshPolygonsTool::Setup()
 	AddInputBehavior(ClickBehavior);
 
 	// create dynamic mesh component to use for live preview
-	DynamicMeshComponent = NewObject<USimpleDynamicMeshComponent>(ComponentTarget->GetOwnerActor(), "DynamicMesh");
-	DynamicMeshComponent->SetupAttachment(ComponentTarget->GetOwnerActor()->GetRootComponent());
+	IPrimitiveComponentBackedTarget* TargetComponent = Cast<IPrimitiveComponentBackedTarget>(Target);
+	DynamicMeshComponent = NewObject<USimpleDynamicMeshComponent>(TargetComponent->GetOwnerActor(), "DynamicMesh");
+	DynamicMeshComponent->SetupAttachment(TargetComponent->GetOwnerActor()->GetRootComponent());
 	DynamicMeshComponent->RegisterComponent();
-	DynamicMeshComponent->SetWorldTransform(ComponentTarget->GetWorldTransform());
+	DynamicMeshComponent->SetWorldTransform(TargetComponent->GetWorldTransform());
 	WorldTransform = FTransform3d(DynamicMeshComponent->GetComponentTransform());
 
 	// set materials
 	FComponentMaterialSet MaterialSet;
-	ComponentTarget->GetMaterialSet(MaterialSet);
+	Cast<IMaterialProvider>(Target)->GetMaterialSet(MaterialSet);
 	for (int k = 0; k < MaterialSet.Materials.Num(); ++k)
 	{
 		DynamicMeshComponent->SetMaterial(k, MaterialSet.Materials[k]);
@@ -196,7 +202,7 @@ void UEditMeshPolygonsTool::Setup()
 
 	// dynamic mesh configuration settings
 	DynamicMeshComponent->TangentsType = EDynamicMeshTangentCalcType::AutoCalculated;
-	DynamicMeshComponent->InitializeMesh(ComponentTarget->GetMesh());
+	DynamicMeshComponent->InitializeMesh(Cast<IMeshDescriptionProvider>(Target)->GetMeshDescription());
 	FMeshNormals::QuickComputeVertexNormals(*DynamicMeshComponent->GetMesh());
 	OnDynamicMeshComponentChangedHandle = DynamicMeshComponent->OnMeshChanged.Add(
 		FSimpleMulticastDelegate::FDelegate::CreateUObject(this, &UEditMeshPolygonsTool::OnDynamicMeshComponentChanged));
@@ -249,7 +255,7 @@ void UEditMeshPolygonsTool::Setup()
 	}
 
 	// hide input StaticMeshComponent
-	ComponentTarget->SetOwnerVisibility(false);
+	TargetComponent->SetOwnerVisibility(false);
 
 	// init state flags flags
 	bInDrag = false;
@@ -369,13 +375,13 @@ bool UEditMeshPolygonsTool::IsStoredToolSelectionUsable(const UGroupTopologyStor
 	// same vids unrecoverable. Once we have persistence of dynamic meshes, this will
 	// hopefully not become a problem, and this function (along with stored selection
 	// identifying info) will change.
-	return !Cast<UBrushComponent>(ComponentTarget->GetOwnerComponent())
+	return !Cast<UBrushComponent>(Cast<IPrimitiveComponentBackedTarget>(Target)->GetOwnerComponent())
 
 		&& StoredSelection
 		&& StoredSelection->IdentifyingInfo.TopologyType == (bTriangleMode ? 
 			UGroupTopologyStorableSelection::ETopologyType::FTriangleGroupTopology
 			: UGroupTopologyStorableSelection::ETopologyType::FGroupTopology)
-		&& StoredSelection->IdentifyingInfo.ComponentTarget == ComponentTarget->GetOwnerComponent()
+		&& StoredSelection->IdentifyingInfo.ComponentTarget == Cast<IPrimitiveComponentBackedTarget>(Target)->GetOwnerComponent()
 		&& !StoredSelection->IsEmpty();
 }
 
@@ -403,7 +409,8 @@ void UEditMeshPolygonsTool::Shutdown(EToolShutdownType ShutdownType)
 	{
 		DynamicMeshComponent->OnMeshChanged.Remove(OnDynamicMeshComponentChangedHandle);
 
-		ComponentTarget->SetOwnerVisibility(true);
+		IPrimitiveComponentBackedTarget* TargetComponent = Cast<IPrimitiveComponentBackedTarget>(Target);
+		TargetComponent->SetOwnerVisibility(true);
 
 		if (ShutdownType == EToolShutdownType::Accept)
 		{
@@ -413,10 +420,10 @@ void UEditMeshPolygonsTool::Shutdown(EToolShutdownType ShutdownType)
 			// Prep if we have a selection to store. We don't support storing selections for volumes
 			// because the conversion will change vids.
 			if (!SelectionMechanic->GetActiveSelection().IsEmpty()
-				&& !Cast<UBrushComponent>(ComponentTarget->GetOwnerComponent()))
+				&& !Cast<UBrushComponent>(TargetComponent->GetOwnerComponent()))
 			{
 				NewStoredToolSelection = NewObject<UGroupTopologyStorableSelection>();
-				NewStoredToolSelection->IdentifyingInfo.ComponentTarget = ComponentTarget->GetOwnerComponent();
+				NewStoredToolSelection->IdentifyingInfo.ComponentTarget = TargetComponent->GetOwnerComponent();
 				NewStoredToolSelection->IdentifyingInfo.TopologyType = (bTriangleMode ?
 					UGroupTopologyStorableSelection::ETopologyType::FTriangleGroupTopology
 					: UGroupTopologyStorableSelection::ETopologyType::FGroupTopology);
@@ -437,12 +444,12 @@ void UEditMeshPolygonsTool::Shutdown(EToolShutdownType ShutdownType)
 
 			// this block bakes the modified DynamicMeshComponent back into the StaticMeshComponent inside an undo transaction
 			GetToolManager()->BeginUndoTransaction(LOCTEXT("EditMeshPolygonsToolTransactionName", "Deform Mesh"));
-			ComponentTarget->CommitMesh([=](const FPrimitiveComponentTarget::FCommitParams& CommitParams)
+			Cast<IMeshDescriptionCommitter>(Target)->CommitMeshDescription([=](const IMeshDescriptionCommitter::FCommitterParams& CommitParams)
 			{
 				FConversionToMeshDescriptionOptions ConversionOptions;
 				bool bModifiedTopology = (ModifiedTopologyCounter > 0);
 				ConversionOptions.bSetPolyGroups = bModifiedTopology;
-				DynamicMeshComponent->Bake(CommitParams.MeshDescription, bModifiedTopology, ConversionOptions);
+				DynamicMeshComponent->Bake(CommitParams.MeshDescriptionOut, bModifiedTopology, ConversionOptions);
 			});
 
 			// The stored selection change should go into this transaction as well.
@@ -2168,7 +2175,7 @@ bool UEditMeshPolygonsTool::BeginMeshFaceEditChangeWithPreview()
 	if (bOK)
 	{
 		EditPreview = NewObject<UPolyEditPreviewMesh>(this);
-		EditPreview->CreateInWorld(ComponentTarget->GetOwnerActor()->GetWorld(), FTransform::Identity);
+		EditPreview->CreateInWorld(Cast<IPrimitiveComponentBackedTarget>(Target)->GetOwnerActor()->GetWorld(), FTransform::Identity);
 		UpdateEditPreviewMaterials(EPreviewMaterialType::PreviewMaterial);
 		EditPreview->EnableWireframe(true);
 
diff --git a/Engine/Plugins/Experimental/MeshModelingToolset/Source/MeshModelingTools/Private/EditUVIslandsTool.cpp b/Engine/Plugins/Experimental/MeshModelingToolset/Source/MeshModelingTools/Private/EditUVIslandsTool.cpp
--- a/Engine/Plugins/Experimental/MeshModelingToolset/Source/MeshModelingTools/Private/EditUVIslandsTool.cpp
+++ b/Engine/Plugins/Experimental/MeshModelingToolset/Source/MeshModelingTools/Private/EditUVIslandsTool.cpp
@@ -14,6 +14,11 @@
 
 #include "Materials/MaterialInstanceDynamic.h"
 
+#include "TargetInterfaces/MaterialProvider.h"
+#include "TargetInterfaces/MeshDescriptionCommitter.h"
+#include "TargetInterfaces/MeshDescriptionProvider.h"
+#include "TargetInterfaces/PrimitiveComponentBackedTarget.h"
+
 #include "ExplicitUseGeometryMathTypes.h"		// using UE::Geometry::(math types)
 using namespace UE::Geometry;
 
@@ -42,15 +47,16 @@ void UEditUVIslandsTool::Setup()
 	UMeshSurfacePointTool::Setup();
 
 	// create dynamic mesh component to use for live preview
-	DynamicMeshComponent = NewObject<USimpleDynamicMeshComponent>(ComponentTarget->GetOwnerActor(), "DynamicMesh");
-	DynamicMeshComponent->SetupAttachment(ComponentTarget->GetOwnerActor()->GetRootComponent());
+	IPrimitiveComponentBackedTarget* TargetComponent = Cast<IPrimitiveComponentBackedTarget>(Target);
+	DynamicMeshComponent = NewObject<USimpleDynamicMeshComponent>(TargetComponent->GetOwnerActor(), "DynamicMesh");
+	DynamicMeshComponent->SetupAttachment(TargetComponent->GetOwnerActor()->GetRootComponent());
 	DynamicMeshComponent->RegisterComponent();
-	DynamicMeshComponent->SetWorldTransform(ComponentTarget->GetWorldTransform());
+	DynamicMeshComponent->SetWorldTransform(TargetComponent->GetWorldTransform());
 	WorldTransform = UE::Geometry::FTransform3d(DynamicMeshComponent->GetComponentTransform());
 
 	// set materials
 	FComponentMaterialSet MaterialSet;
-	ComponentTarget->GetMaterialSet(MaterialSet);
+	Cast<IMaterialProvider>(Target)->GetMaterialSet(MaterialSet);
 	for (int k = 0; k < MaterialSet.Materials.Num(); ++k)
 	{
 		DynamicMeshComponent->SetMaterial(k, MaterialSet.Materials[k]);
@@ -65,7 +71,7 @@ void UEditUVIslandsTool::Setup()
 
 	// dynamic mesh configuration settings
 	DynamicMeshComponent->TangentsType = EDynamicMeshTangentCalcType::AutoCalculated;
-	DynamicMeshComponent->InitializeMesh(ComponentTarget->GetMesh());
+	DynamicMeshComponent->InitializeMesh(Cast<IMeshDescriptionProvider>(Target)->GetMeshDescription());
 	FMeshNormals::QuickComputeVertexNormals(*DynamicMeshComponent->GetMesh());
 	OnDynamicMeshComponentChangedHandle = DynamicMeshComponent->OnMeshChanged.Add(
 		FSimpleMulticastDelegate::FDelegate::CreateUObject(this, &UEditUVIslandsTool::OnDynamicMeshComponentChanged));
@@ -86,7 +92,7 @@ void UEditUVIslandsTool::Setup()
 	UVTranslateScale = 1.0 / DynamicMeshComponent->GetMesh()->GetBounds().MaxDim();
 
 	// hide input StaticMeshComponent
-	ComponentTarget->SetOwnerVisibility(false);
+	TargetComponent->SetOwnerVisibility(false);
 
 	// init state flags flags
 	bInDrag = false;
@@ -130,18 +136,18 @@ void UEditUVIslandsTool::Shutdown(EToolShutdownType ShutdownType)
 	{
 		DynamicMeshComponent->OnMeshChanged.Remove(OnDynamicMeshComponentChangedHandle);
 
-		ComponentTarget->SetOwnerVisibility(true);
+		Cast<IPrimitiveComponentBackedTarget>(Target)->SetOwnerVisibility(true);
 
 		if (ShutdownType == EToolShutdownType::Accept)
 		{
 			// this block bakes the modified DynamicMeshComponent back into the StaticMeshComponent inside an undo transaction
 			GetToolManager()->BeginUndoTransaction(LOCTEXT("EditUVIslandsToolTransactionName", "Edit UVs"));
-			ComponentTarget->CommitMesh([=](const FPrimitiveComponentTarget::FCommitParams& CommitParams)
+			Cast<IMeshDescriptionCommitter>(Target)->CommitMeshDescription([=](const IMeshDescriptionCommitter::FCommitterParams& CommitParams)
 			{
 				FConversionToMeshDescriptionOptions ConversionOptions;
 				ConversionOptions.bUpdateNormals = ConversionOptions.bUpdatePositions = false;
 				ConversionOptions.bUpdateUVs = true;
-				DynamicMeshComponent->Bake(CommitParams.MeshDescription, false, ConversionOptions);
+				DynamicMeshComponent->Bake(CommitParams.MeshDescriptionOut, false, ConversionOptions);
 			});
 			GetToolManager()->EndUndoTransaction();
 		}
diff --git a/Engine/Plugins/Experimental/MeshModelingToolset/Source/MeshModelingTools/Private/MeshAttributePaintTool.cpp b/Engine/Plugins/Experimental/MeshModelingToolset/Source/MeshModelingTools/Private/MeshAttributePaintTool.cpp
--- a/Engine/Plugins/Experimental/MeshModelingToolset/Source/MeshModelingTools/Private/MeshAttributePaintTool.cpp
+++ b/Engine/Plugins/Experimental/MeshModelingToolset/Source/MeshModelingTools/Private/MeshAttributePaintTool.cpp
@@ -9,6 +9,10 @@
 
 #include "MeshDescription.h"
 
+#include "TargetInterfaces/MeshDescriptionCommitter.h"
+#include "TargetInterfaces/MeshDescriptionProvider.h"
+#include "TargetInterfaces/PrimitiveComponentBackedTarget.h"
+
 #include "ExplicitUseGeometryMathTypes.h"		// using UE::Geometry::(math types)
 using namespace UE::Geometry;
 
@@ -194,7 +198,7 @@ void UMeshAttributePaintTool::Setup()
 	ColorMapper = MakeUnique<FFloatAttributeColorMapper>();
 
 	EditedMesh = MakeUnique<FMeshDescription>();
-	*EditedMesh = *ComponentTarget->GetMesh();
+	*EditedMesh = *Cast<IMeshDescriptionProvider>(Target)->GetMeshDescription();
 
 	AttributeSource = MakeUnique<FMeshDescriptionVertexAttributeSource>(EditedMesh.Get());
 	AttribProps->Attributes = AttributeSource->GetAttributeList();
@@ -341,7 +345,7 @@ bool UMeshAttributePaintTool::OnUpdateHover(const FInputDeviceRay& DevicePos)
 
 void UMeshAttributePaintTool::CalculateVertexROI(const FBrushStampData& Stamp, TArray<int>& VertexROI)
 {
-	UE::Geometry::FTransform3d Transform(ComponentTarget->GetWorldTransform());
+	UE::Geometry::FTransform3d Transform(Cast<IPrimitiveComponentBackedTarget>(Target)->GetWorldTransform());
 	FVector3d StampPosLocal = Transform.InverseTransformPosition(Stamp.WorldPosition);
 
 	float Radius = GetCurrentBrushRadiusLocal();
@@ -499,7 +503,7 @@ void UMeshAttributePaintTool::ApplyStamp(const FBrushStampData& Stamp)
 
 void UMeshAttributePaintTool::ApplyStamp_Paint(const FBrushStampData& Stamp, FStampActionData& ActionData)
 {
-	UE::Geometry::FTransform3d Transform(ComponentTarget->GetWorldTransform());
+	UE::Geometry::FTransform3d Transform(Cast<IPrimitiveComponentBackedTarget>(Target)->GetWorldTransform());
 	FVector3d StampPosLocal = Transform.InverseTransformPosition(Stamp.WorldPosition);
 
 	int32 NumVertices = ActionData.ROIVertices.Num();
@@ -640,9 +644,9 @@ void UMeshAttributePaintTool::OnShutdown(EToolShutdownType ShutdownType)
 		// this block bakes the modified DynamicMeshComponent back into the StaticMeshComponent inside an undo transaction
 		GetToolManager()->BeginUndoTransaction(LOCTEXT("MeshAttributePaintTool", "Edit Attributes"));
 
-		ComponentTarget->CommitMesh([=](const FPrimitiveComponentTarget::FCommitParams& CommitParams)
+		Cast<IMeshDescriptionCommitter>(Target)->CommitMeshDescription([=](const IMeshDescriptionCommitter::FCommitterParams& CommitParams)
 		{
-			*CommitParams.MeshDescription = *EditedMesh;
+			*CommitParams.MeshDescriptionOut = *EditedMesh;
 		});
 		GetToolManager()->EndUndoTransaction();
 	}
diff --git a/Engine/Plugins/Experimental/MeshModelingToolset/Source/MeshModelingTools/Private/MeshGroupPaintTool.cpp b/Engine/Plugins/Experimental/MeshModelingToolset/Source/MeshModelingTools/Private/MeshGroupPaintTool.cpp
--- a/Engine/Plugins/Experimental/MeshModelingToolset/Source/MeshModelingTools/Private/MeshGroupPaintTool.cpp
+++ b/Engine/Plugins/Experimental/MeshModelingToolset/Source/MeshModelingTools/Private/MeshGroupPaintTool.cpp
@@ -34,6 +34,10 @@
 
 #include "ToolSetupUtil.h"
 
+#include "TargetInterfaces/MaterialProvider.h"
+#include "TargetInterfaces/MeshDescriptionCommitter.h"
+#include "TargetInterfaces/PrimitiveComponentBackedTarget.h"
+
 #include "ExplicitUseGeometryMathTypes.h"		// using UE::Geometry::(math types)
 using namespace UE::Geometry;
 
@@ -87,12 +91,12 @@ void UMeshGroupPaintTool::Setup()
 	SetToolDisplayName(LOCTEXT("ToolName", "Paint PolyGroups"));
 
 	// create dynamic mesh component to use for live preview
-	DynamicMeshComponent = NewObject<USimpleDynamicMeshComponent>(ComponentTarget->GetOwnerActor());
+	DynamicMeshComponent = NewObject<USimpleDynamicMeshComponent>(Cast<IPrimitiveComponentBackedTarget>(Target)->GetOwnerActor());
 	InitializeSculptMeshComponent(DynamicMeshComponent);
 
 	// assign materials
 	FComponentMaterialSet MaterialSet;
-	ComponentTarget->GetMaterialSet(MaterialSet);
+	Cast<IMaterialProvider>(Target)->GetMaterialSet(MaterialSet);
 	for (int k = 0; k < MaterialSet.Materials.Num(); ++k)
 	{
 		DynamicMeshComponent->SetMaterial(k, MaterialSet.Materials[k]);
@@ -264,10 +268,10 @@ void UMeshGroupPaintTool::Shutdown(EToolShutdownType ShutdownType)
 	if (ShutdownType == EToolShutdownType::Accept)
 	{
 		GetToolManager()->BeginUndoTransaction(LOCTEXT("GroupPaintMeshToolTransactionName", "Paint Groups"));
-		ComponentTarget->CommitMesh([=](const FPrimitiveComponentTarget::FCommitParams& CommitParams)
+		Cast<IMeshDescriptionCommitter>(Target)->CommitMeshDescription([=](const IMeshDescriptionCommitter::FCommitterParams& CommitParams)
 		{
 			FConversionToMeshDescriptionOptions ConversionOptions;
-			DynamicMeshComponent->Bake(CommitParams.MeshDescription, true, ConversionOptions);
+			DynamicMeshComponent->Bake(CommitParams.MeshDescriptionOut, true, ConversionOptions);
 		});
 		GetToolManager()->EndUndoTransaction();
 	}
@@ -879,7 +883,7 @@ void UMeshGroupPaintTool::ApplyVisibilityFilter(const TArray<int32>& Triangles,
 
 	FViewCameraState StateOut;
 	GetToolManager()->GetContextQueriesAPI()->GetCurrentViewState(StateOut);
-	FVector3d LocalEyePosition(ComponentTarget->GetWorldTransform().InverseTransformPosition(StateOut.Position));
+	FVector3d LocalEyePosition(Cast<IPrimitiveComponentBackedTarget>(Target)->GetWorldTransform().InverseTransformPosition(StateOut.Position));
 
 	const FDynamicMesh3* Mesh = GetSculptMesh();
 
diff --git a/Engine/Plugins/Experimental/MeshModelingToolset/Source/MeshModelingTools/Private/MeshSelectionTool.cpp b/Engine/Plugins/Experimental/MeshModelingToolset/Source/MeshModelingTools/Private/MeshSelectionTool.cpp
--- a/Engine/Plugins/Experimental/MeshModelingToolset/Source/MeshModelingTools/Private/MeshSelectionTool.cpp
+++ b/Engine/Plugins/Experimental/MeshModelingToolset/Source/MeshModelingTools/Private/MeshSelectionTool.cpp
@@ -19,6 +19,10 @@
 
 #include "Algo/MaxElement.h"
 
+#include "TargetInterfaces/MaterialProvider.h"
+#include "TargetInterfaces/MeshDescriptionCommitter.h"
+#include "TargetInterfaces/PrimitiveComponentBackedTarget.h"
+
 #include "ExplicitUseGeometryMathTypes.h"		// using UE::Geometry::(math types)
 using namespace UE::Geometry;
 
@@ -169,9 +173,9 @@ void UMeshSelectionTool::OnShutdown(EToolShutdownType ShutdownType)
 		// this block bakes the modified DynamicMeshComponent back into the StaticMeshComponent inside an undo transaction
 		GetToolManager()->BeginUndoTransaction(LOCTEXT("MeshSelectionToolTransactionName", "Edit Mesh"));
 
-		ComponentTarget->CommitMesh([=](const FPrimitiveComponentTarget::FCommitParams& CommitParams)
+		Cast<IMeshDescriptionCommitter>(Target)->CommitMeshDescription([=](const IMeshDescriptionCommitter::FCommitterParams& CommitParams)
 		{
-			PreviewMesh->Bake(CommitParams.MeshDescription, true);
+			PreviewMesh->Bake(CommitParams.MeshDescriptionOut, true);
 		});
 		GetToolManager()->EndUndoTransaction();
 	}
@@ -327,7 +331,7 @@ bool UMeshSelectionTool::HitTest(const FRay& Ray, FHitResult& OutHit)
 		SourceMesh->GetTriInfo(OutHit.FaceIndex, Normal, Area, Centroid);
 		FViewCameraState StateOut;
 		GetToolManager()->GetContextQueriesAPI()->GetCurrentViewState(StateOut);
-		FVector3d LocalEyePosition(ComponentTarget->GetWorldTransform().InverseTransformPosition(StateOut.Position));
+		FVector3d LocalEyePosition(Cast<IPrimitiveComponentBackedTarget>(Target)->GetWorldTransform().InverseTransformPosition(StateOut.Position));
 
 		if (Normal.Dot((Centroid - LocalEyePosition)) > 0)
 		{
@@ -382,7 +386,7 @@ TUniquePtr<FDynamicMeshOctree3>& UMeshSelectionTool::GetOctree()
 
 void UMeshSelectionTool::CalculateVertexROI(const FBrushStampData& Stamp, TArray<int>& VertexROI)
 {
-	FTransform Transform = ComponentTarget->GetWorldTransform();
+	FTransform Transform = Cast<IPrimitiveComponentBackedTarget>(Target)->GetWorldTransform();
 	FVector StampPosLocal = Transform.InverseTransformPosition(Stamp.WorldPosition);
 
 	// TODO: need dynamic vertex hash table!
@@ -404,7 +408,7 @@ void UMeshSelectionTool::CalculateVertexROI(const FBrushStampData& Stamp, TArray
 
 void UMeshSelectionTool::CalculateTriangleROI(const FBrushStampData& Stamp, TArray<int>& TriangleROI)
 {
-	UE::Geometry::FTransform3d Transform(ComponentTarget->GetWorldTransform());
+	UE::Geometry::FTransform3d Transform(Cast<IPrimitiveComponentBackedTarget>(Target)->GetWorldTransform());
 	FVector3d StampPosLocal = Transform.InverseTransformPosition(Stamp.WorldPosition);
 
 	// always select first triangle
@@ -557,7 +561,7 @@ void UMeshSelectionTool::UpdateFaceSelection(const FBrushStampData& Stamp, const
 	{
 		FViewCameraState StateOut;
 		GetToolManager()->GetContextQueriesAPI()->GetCurrentViewState(StateOut);
-		FVector3d LocalEyePosition(ComponentTarget->GetWorldTransform().InverseTransformPosition(StateOut.Position));
+		FVector3d LocalEyePosition(Cast<IPrimitiveComponentBackedTarget>(Target)->GetWorldTransform().InverseTransformPosition(StateOut.Position));
 
 		for (int tid : TriangleROI)
 		{
@@ -760,7 +764,7 @@ void UMeshSelectionTool::Render(IToolsContextRenderAPI* RenderAPI)
 {
 	UDynamicMeshBrushTool::Render(RenderAPI);
 
-	FTransform WorldTransform = ComponentTarget->GetWorldTransform();
+	FTransform WorldTransform = Cast<IPrimitiveComponentBackedTarget>(Target)->GetWorldTransform();
 	const FDynamicMesh3* Mesh = PreviewMesh->GetMesh();
 
 	if (SelectionProps->bShowPoints)
@@ -1353,9 +1357,10 @@ void UMeshSelectionTool::SeparateSelectedTriangles()
 
 	// build array of materials from the original
 	TArray<UMaterialInterface*> Materials;
-	for (int MaterialIdx = 0, NumMaterials = ComponentTarget->GetNumMaterials(); MaterialIdx < NumMaterials; MaterialIdx++)
+	IMaterialProvider* TargetMaterial = Cast<IMaterialProvider>(Target);
+	for (int MaterialIdx = 0, NumMaterials = TargetMaterial->GetNumMaterials(); MaterialIdx < NumMaterials; MaterialIdx++)
 	{
-		Materials.Add(ComponentTarget->GetMaterial(MaterialIdx));
+		Materials.Add(TargetMaterial->GetMaterial(MaterialIdx));
 	}
 	AActor* NewActor = AssetGenerationUtil::GenerateStaticMeshActor(
 		AssetAPI, TargetWorld, &SeparatedMesh, Transform, TEXT("Submesh"), Materials);
diff --git a/Engine/Plugins/Experimental/MeshModelingToolset/Source/MeshModelingTools/Private/MeshVertexSculptTool.cpp b/Engine/Plugins/Experimental/MeshModelingToolset/Source/MeshModelingTools/Private/MeshVertexSculptTool.cpp
--- a/Engine/Plugins/Experimental/MeshModelingToolset/Source/MeshModelingTools/Private/MeshVertexSculptTool.cpp
+++ b/Engine/Plugins/Experimental/MeshModelingToolset/Source/MeshModelingTools/Private/MeshVertexSculptTool.cpp
@@ -25,6 +25,8 @@
 #include "Sculpting/StampFalloffs.h"
 #include "Sculpting/MeshSculptUtil.h"
 
+#include "TargetInterfaces/PrimitiveComponentBackedTarget.h"
+
 #include "ExplicitUseGeometryMathTypes.h"		// using UE::Geometry::(math types)
 using namespace UE::Geometry;
 
@@ -62,12 +64,12 @@ void UMeshVertexSculptTool::Setup()
 	SetToolDisplayName(LOCTEXT("ToolName", "Sculpt"));
 
 	// create dynamic mesh component to use for live preview
-	DynamicMeshComponent = NewObject<USimpleDynamicMeshComponent>(ComponentTarget->GetOwnerActor());
+	DynamicMeshComponent = NewObject<USimpleDynamicMeshComponent>(Cast<IPrimitiveComponentBackedTarget>(Target)->GetOwnerActor());
 	InitializeSculptMeshComponent(DynamicMeshComponent);
 
 	// assign materials
 	FComponentMaterialSet MaterialSet;
-	ComponentTarget->GetMaterialSet(MaterialSet);
+	Cast<IMaterialProvider>(Target)->GetMaterialSet(MaterialSet);
 	for (int k = 0; k < MaterialSet.Materials.Num(); ++k)
 	{
 		DynamicMeshComponent->SetMaterial(k, MaterialSet.Materials[k]);
diff --git a/Engine/Plugins/Experimental/MeshModelingToolset/Source/MeshModelingTools/Private/Sculpting/MeshSculptToolBase.cpp b/Engine/Plugins/Experimental/MeshModelingToolset/Source/MeshModelingTools/Private/Sculpting/MeshSculptToolBase.cpp
--- a/Engine/Plugins/Experimental/MeshModelingToolset/Source/MeshModelingTools/Private/Sculpting/MeshSculptToolBase.cpp
+++ b/Engine/Plugins/Experimental/MeshModelingToolset/Source/MeshModelingTools/Private/Sculpting/MeshSculptToolBase.cpp
@@ -15,6 +15,10 @@
 
 #include "Generators/SphereGenerator.h"
 
+#include "TargetInterfaces/MeshDescriptionCommitter.h"
+#include "TargetInterfaces/MeshDescriptionProvider.h"
+#include "TargetInterfaces/PrimitiveComponentBackedTarget.h"
+
 #include "ExplicitUseGeometryMathTypes.h"		// using UE::Geometry::(math types)
 using namespace UE::Geometry;
 
@@ -128,7 +132,7 @@ void UMeshSculptToolBase::Shutdown(EToolShutdownType ShutdownType)
 	UBaseDynamicMeshComponent* DynamicMeshComponent = GetSculptMeshComponent();
 	if (DynamicMeshComponent != nullptr)
 	{
-		ComponentTarget->SetOwnerVisibility(true);
+		Cast<IPrimitiveComponentBackedTarget>(Target)->SetOwnerVisibility(true);
 
 		if (ShutdownType == EToolShutdownType::Accept)
 		{
@@ -152,10 +156,10 @@ void UMeshSculptToolBase::Shutdown(EToolShutdownType ShutdownType)
 
 void UMeshSculptToolBase::CommitResult(UBaseDynamicMeshComponent* Component)
 {
-	ComponentTarget->CommitMesh([=](const FPrimitiveComponentTarget::FCommitParams& CommitParams)
+	Cast<IMeshDescriptionCommitter>(Target)->CommitMeshDescription([=](const IMeshDescriptionCommitter::FCommitterParams& CommitParams)
 	{
 		FConversionToMeshDescriptionOptions ConversionOptions;
-		Component->Bake(CommitParams.MeshDescription, false, ConversionOptions);
+		Component->Bake(CommitParams.MeshDescriptionOut, false, ConversionOptions);
 	});
 }
 
@@ -225,16 +229,17 @@ void UMeshSculptToolBase::Render(IToolsContextRenderAPI* RenderAPI)
 
 void UMeshSculptToolBase::InitializeSculptMeshComponent(UBaseDynamicMeshComponent* Component)
 {
-	Component->SetupAttachment(ComponentTarget->GetOwnerActor()->GetRootComponent());
+	IPrimitiveComponentBackedTarget* TargetComponent = Cast<IPrimitiveComponentBackedTarget>(Target);
+	Component->SetupAttachment(TargetComponent->GetOwnerActor()->GetRootComponent());
 	Component->RegisterComponent();
 
 	// initialize from LOD-0 MeshDescription
-	Component->InitializeMesh(ComponentTarget->GetMesh());
+	Component->InitializeMesh(Cast<IMeshDescriptionProvider>(Target)->GetMeshDescription());
 	double MaxDimension = Component->GetMesh()->GetCachedBounds().MaxDim();
 
 	// bake rotation and scaling into mesh because handling these inside sculpting is a mess
 	// Note: this transform does not include translation ( so only the 3x3 transform)
-	InitialTargetTransform = FTransform3d(ComponentTarget->GetWorldTransform());
+	InitialTargetTransform = FTransform3d(TargetComponent->GetWorldTransform());
 	// clamp scaling because if we allow zero-scale we cannot invert this transform on Accept
 	InitialTargetTransform.ClampMinimumScale(0.01);
 	FVector3d Translation = InitialTargetTransform.GetTranslation();
@@ -244,7 +249,7 @@ void UMeshSculptToolBase::InitializeSculptMeshComponent(UBaseDynamicMeshComponen
 	Component->SetWorldTransform((FTransform)CurTargetTransform);
 
 	// hide input Component
-	ComponentTarget->SetOwnerVisibility(false);
+	TargetComponent->SetOwnerVisibility(false);
 }
 
 
@@ -906,7 +911,7 @@ void UMeshSculptToolBase::UpdateMaterialMode(EMeshEditingMaterialModes MaterialM
 	if (MaterialMode == EMeshEditingMaterialModes::ExistingMaterial)
 	{
 		GetSculptMeshComponent()->ClearOverrideRenderMaterial();
-		GetSculptMeshComponent()->bCastDynamicShadow = ComponentTarget->GetOwnerComponent()->bCastDynamicShadow;
+		GetSculptMeshComponent()->bCastDynamicShadow = Cast<IPrimitiveComponentBackedTarget>(Target)->GetOwnerComponent()->bCastDynamicShadow;
 		ActiveOverrideMaterial = nullptr;
 	}
 	else
diff --git a/Engine/Plugins/Experimental/MeshModelingToolset/Source/MeshModelingTools/Private/SeamSculptTool.cpp b/Engine/Plugins/Experimental/MeshModelingToolset/Source/MeshModelingTools/Private/SeamSculptTool.cpp
--- a/Engine/Plugins/Experimental/MeshModelingToolset/Source/MeshModelingTools/Private/SeamSculptTool.cpp
+++ b/Engine/Plugins/Experimental/MeshModelingToolset/Source/MeshModelingTools/Private/SeamSculptTool.cpp
@@ -22,6 +22,9 @@
 #include "DynamicMeshChangeTracker.h"
 #include "DynamicMeshToMeshDescription.h"
 
+#include "TargetInterfaces/MeshDescriptionCommitter.h"
+#include "TargetInterfaces/PrimitiveComponentBackedTarget.h"
+
 #include "ExplicitUseGeometryMathTypes.h"		// using UE::Geometry::(math types)
 using namespace UE::Geometry;
 
@@ -60,7 +63,8 @@ void USeamSculptTool::Setup()
 	BrushProperties->bShowStrength = BrushProperties->bShowFalloff = false;
 	BrushProperties->RestoreProperties(this);
 
-	MeshTransform = UE::Geometry::FTransform3d(ComponentTarget->GetWorldTransform());
+	IPrimitiveComponentBackedTarget* TargetComponent = Cast<IPrimitiveComponentBackedTarget>(Target);
+	MeshTransform = UE::Geometry::FTransform3d(TargetComponent->GetWorldTransform());
 	InputMesh = MakeShared<FDynamicMesh3, ESPMode::ThreadSafe>(*PreviewMesh->GetMesh());
 	FMeshNormals::QuickComputeVertexNormals(*InputMesh);
 	NormalOffset = InputMesh->GetCachedBounds().MinDim() * 0.001;
@@ -80,7 +84,7 @@ void USeamSculptTool::Setup()
 	RecalculateBrushRadius();
 
 	PreviewGeom = NewObject<UPreviewGeometry>(this);
-	PreviewGeom->CreateInWorld(ComponentTarget->GetOwnerActor()->GetWorld(), ComponentTarget->GetWorldTransform());
+	PreviewGeom->CreateInWorld(TargetComponent->GetOwnerActor()->GetWorld(), TargetComponent->GetWorldTransform());
 	InitPreviewGeometry();
 
 	// regenerate preview geo if mesh changes due to undo/redo/etc
@@ -110,10 +114,10 @@ void USeamSculptTool::OnShutdown(EToolShutdownType ShutdownType)
 
 		GetToolManager()->BeginUndoTransaction(LOCTEXT("SeamSculptTransactionName", "UV Seam Edit"));
 
-		ComponentTarget->CommitMesh([&ResultMesh](const FPrimitiveComponentTarget::FCommitParams& CommitParams)
+		Cast<IMeshDescriptionCommitter>(Target)->CommitMeshDescription([&ResultMesh](const IMeshDescriptionCommitter::FCommitterParams& CommitParams)
 		{
 			FDynamicMeshToMeshDescription Converter;
-			Converter.Convert(&ResultMesh, *CommitParams.MeshDescription);
+			Converter.Convert(&ResultMesh, *CommitParams.MeshDescriptionOut);
 		});
 
 		GetToolManager()->EndUndoTransaction();
diff --git a/Engine/Plugins/Experimental/MeshModelingToolset/Source/MeshModelingToolsEditorOnly/Private/Hair/GroomCardsEditorTool.cpp b/Engine/Plugins/Experimental/MeshModelingToolset/Source/MeshModelingToolsEditorOnly/Private/Hair/GroomCardsEditorTool.cpp
--- a/Engine/Plugins/Experimental/MeshModelingToolset/Source/MeshModelingToolsEditorOnly/Private/Hair/GroomCardsEditorTool.cpp
+++ b/Engine/Plugins/Experimental/MeshModelingToolset/Source/MeshModelingToolsEditorOnly/Private/Hair/GroomCardsEditorTool.cpp
@@ -30,6 +30,10 @@
 #include "Engine/Classes/Engine/StaticMeshActor.h"
 #include "MeshDescriptionToDynamicMesh.h"
 
+#include "TargetInterfaces/MeshDescriptionCommitter.h"
+#include "TargetInterfaces/MeshDescriptionProvider.h"
+#include "TargetInterfaces/PrimitiveComponentBackedTarget.h"
+
 #define LOCTEXT_NAMESPACE "UGroomCardsEditorTool"
 
 #include "ExplicitUseGeometryMathTypes.h"		// using UE::Geometry::(math types)
@@ -39,29 +43,15 @@ using namespace UE::Geometry;
  * ToolBuilder
  */
 
-bool UGroomCardsEditorToolBuilder::CanBuildTool(const FToolBuilderState& SceneState) const
-{
-	return ToolBuilderUtil::CountComponents(SceneState, CanMakeComponentTarget) == 1;
-}
-
-UInteractiveTool* UGroomCardsEditorToolBuilder::BuildTool(const FToolBuilderState& SceneState) const
+UMeshSurfacePointTool* UGroomCardsEditorToolBuilder::CreateNewTool(const FToolBuilderState& SceneState) const
 {
 	UGroomCardsEditorTool* NewTool = NewObject<UGroomCardsEditorTool>(SceneState.ToolManager);
-
 	NewTool->SetWorld(SceneState.World);
 	NewTool->SetAssetAPI(AssetAPI);
-
-	TArray<UActorComponent*> Components = ToolBuilderUtil::FindAllComponents(SceneState, CanMakeComponentTarget);
-	check(Components.Num() > 0);
-	TUniquePtr<FPrimitiveComponentTarget> MeshTarget = MakeComponentTarget(Cast<UPrimitiveComponent>(Components[0]));
-	NewTool->SetSelection(MoveTemp(MeshTarget));
-
 	return NewTool;
 }
 
 
-
-
 void UEditGroomCardsToolActionPropertySet::PostAction(EEditGroomCardsToolActions Action)
 {
 	if (ParentTool.IsValid())
@@ -87,10 +77,11 @@ void UGroomCardsEditorTool::Setup()
 {
 	UMeshSurfacePointTool::Setup();
 
+	IPrimitiveComponentBackedTarget* TargetComponent = Cast<IPrimitiveComponentBackedTarget>(Target);
 	PreviewMesh = NewObject<UPreviewMesh>(this);
 	PreviewMesh->bBuildSpatialDataStructure = true;
 	PreviewMesh->CreateInWorld(TargetWorld, FTransform::Identity);
-	PreviewMesh->SetTransform(ComponentTarget->GetWorldTransform());
+	PreviewMesh->SetTransform(TargetComponent->GetWorldTransform());
 
 	MeshMaterial = ToolSetupUtil::GetDefaultSculptMaterial(GetToolManager());
 	UVMaterial = ToolSetupUtil::GetUVCheckerboardMaterial(50.0);
@@ -123,9 +114,9 @@ void UGroomCardsEditorTool::Setup()
 
 	PreviewGeom = NewObject<UPreviewGeometry>(this);
 	PreviewGeom->CreateInWorld(TargetWorld, FTransform::Identity);
-	PreviewGeom->GetActor()->SetActorTransform(ComponentTarget->GetWorldTransform());
+	PreviewGeom->GetActor()->SetActorTransform(TargetComponent->GetWorldTransform());
 
-	ComponentTarget->SetOwnerVisibility(false);
+	TargetComponent->SetOwnerVisibility(false);
 
 
 
@@ -195,9 +186,9 @@ void UGroomCardsEditorTool::Shutdown(EToolShutdownType ShutdownType)
 	if (ShutdownType == EToolShutdownType::Accept)
 	{
 		GetToolManager()->BeginUndoTransaction(LOCTEXT("CardsMeshApplyEditChange", "Update Cards Mesh"));
-		ComponentTarget->CommitMesh([this](const FPrimitiveComponentTarget::FCommitParams& CommitParams)
+		Cast<IMeshDescriptionCommitter>(Target)->CommitMeshDescription([this](const IMeshDescriptionCommitter::FCommitterParams& CommitParams)
 		{
-			PreviewMesh->Bake(CommitParams.MeshDescription, false);
+			PreviewMesh->Bake(CommitParams.MeshDescriptionOut, false);
 		});
 		GetToolManager()->EndUndoTransaction();
 	}
@@ -208,7 +199,7 @@ void UGroomCardsEditorTool::Shutdown(EToolShutdownType ShutdownType)
 	PreviewGeom->Disconnect();
 	PreviewGeom = nullptr;
 
-	ComponentTarget->SetOwnerVisibility(true);
+	Cast<IPrimitiveComponentBackedTarget>(Target)->SetOwnerVisibility(true);
 }
 
 
@@ -773,7 +764,7 @@ void UGroomCardsEditorTool::UpdateLineSet()
 void UGroomCardsEditorTool::InitializeMesh()
 {
 	EditableCardSet = MakePimpl<FEditableGroomCardSet>();
-	EditableCardSet->Initialize(ComponentTarget->GetMesh());
+	EditableCardSet->Initialize(Cast<IMeshDescriptionProvider>(Target)->GetMeshDescription());
 	const FDynamicMesh3* CardsMesh = EditableCardSet->FullCardMesh.Get();
 
 	PreviewMesh->UpdatePreview(CardsMesh);
diff --git a/Engine/Plugins/Experimental/MeshModelingToolset/Source/MeshModelingToolsEditorOnly/Private/ShapeSprayTool.cpp b/Engine/Plugins/Experimental/MeshModelingToolset/Source/MeshModelingToolsEditorOnly/Private/ShapeSprayTool.cpp
--- a/Engine/Plugins/Experimental/MeshModelingToolset/Source/MeshModelingToolsEditorOnly/Private/ShapeSprayTool.cpp
+++ b/Engine/Plugins/Experimental/MeshModelingToolset/Source/MeshModelingToolsEditorOnly/Private/ShapeSprayTool.cpp
@@ -8,6 +8,8 @@
 #include "AssetGenerationUtil.h"
 #include "MeshNormals.h"
 
+#include "TargetInterfaces/PrimitiveComponentBackedTarget.h"
+
 #include "ExplicitUseGeometryMathTypes.h"		// using UE::Geometry::(math types)
 using namespace UE::Geometry;
 
@@ -72,8 +74,9 @@ void UShapeSprayTool::Setup()
 	Random.Initialize(31337);
 
 	// create dynamic mesh component to use for live preview
-	AccumMeshComponent = NewObject<USimpleDynamicMeshComponent>(ComponentTarget->GetOwnerActor(), "SprayMesh");
-		AccumMeshComponent->SetupAttachment(ComponentTarget->GetOwnerActor()->GetRootComponent());
+	IPrimitiveComponentBackedTarget* TargetComponent = Cast<IPrimitiveComponentBackedTarget>(Target);
+	AccumMeshComponent = NewObject<USimpleDynamicMeshComponent>(TargetComponent->GetOwnerActor(), "SprayMesh");
+		AccumMeshComponent->SetupAttachment(TargetComponent->GetOwnerActor()->GetRootComponent());
 	AccumMeshComponent->RegisterComponent();
 
 	UMaterialInterface* VtxColorMaterial = GetToolManager()->GetContextQueriesAPI()->GetStandardMaterial(EStandardToolContextMaterials::VertexColorMaterial);
@@ -132,7 +135,7 @@ void UShapeSprayTool::OnUpdateDrag(const FRay& Ray)
 	UDynamicMeshBrushTool::OnUpdateDrag(Ray);
 
 	FFrame3f WorldFrame(LastBrushStamp.WorldPosition, LastBrushStamp.WorldNormal);
-	FTransform Transform = ComponentTarget->GetWorldTransform();
+	FTransform Transform = Cast<IPrimitiveComponentBackedTarget>(Target)->GetWorldTransform();
 
 
 	FDynamicMesh3* Mesh = AccumMeshComponent->GetMesh();
@@ -232,7 +235,7 @@ bool UShapeSprayTool::CanAccept() const
 void UShapeSprayTool::EmitResult()
 {
 	const FDynamicMesh3* Mesh = AccumMeshComponent->GetMesh();
-	UE::Geometry::FTransform3d UseTransform(ComponentTarget->GetOwnerActor()->GetTransform());
+	UE::Geometry::FTransform3d UseTransform(Cast<IPrimitiveComponentBackedTarget>(Target)->GetOwnerActor()->GetTransform());
 
 	GetToolManager()->BeginUndoTransaction(LOCTEXT("EmitShapeSpray", "Create ShapeSpray"));
 
diff --git a/Engine/Plugins/Experimental/MeshModelingToolset/Source/MeshModelingToolsEditorOnly/Public/Hair/GroomCardsEditorTool.h b/Engine/Plugins/Experimental/MeshModelingToolset/Source/MeshModelingToolsEditorOnly/Public/Hair/GroomCardsEditorTool.h
--- a/Engine/Plugins/Experimental/MeshModelingToolset/Source/MeshModelingToolsEditorOnly/Public/Hair/GroomCardsEditorTool.h
+++ b/Engine/Plugins/Experimental/MeshModelingToolset/Source/MeshModelingToolsEditorOnly/Public/Hair/GroomCardsEditorTool.h
@@ -29,14 +29,13 @@ class FMeshVertexChangeBuilder;
  *
  */
 UCLASS()
-class MESHMODELINGTOOLSEDITORONLY_API UGroomCardsEditorToolBuilder : public UInteractiveToolBuilder
+class MESHMODELINGTOOLSEDITORONLY_API UGroomCardsEditorToolBuilder : public UMeshSurfacePointToolBuilder
 {
 	GENERATED_BODY()
 public:
 	IToolsContextAssetAPI* AssetAPI = nullptr;
 
-	virtual bool CanBuildTool(const FToolBuilderState& SceneState) const override;
-	virtual UInteractiveTool* BuildTool(const FToolBuilderState& SceneState) const override;
+	virtual UMeshSurfacePointTool* CreateNewTool(const FToolBuilderState& SceneState) const override;
 };
 
 
diff --git a/Engine/Source/Runtime/Experimental/InteractiveToolsFramework/Private/BaseTools/MeshSurfacePointTool.cpp b/Engine/Source/Runtime/Experimental/InteractiveToolsFramework/Private/BaseTools/MeshSurfacePointTool.cpp
--- a/Engine/Source/Runtime/Experimental/InteractiveToolsFramework/Private/BaseTools/MeshSurfacePointTool.cpp
+++ b/Engine/Source/Runtime/Experimental/InteractiveToolsFramework/Private/BaseTools/MeshSurfacePointTool.cpp
@@ -7,6 +7,12 @@
 #include "ToolBuilderUtil.h"
 #include "Components/PrimitiveComponent.h"
 
+#include "TargetInterfaces/MaterialProvider.h"
+#include "TargetInterfaces/MeshDescriptionCommitter.h"
+#include "TargetInterfaces/MeshDescriptionProvider.h"
+#include "TargetInterfaces/PrimitiveComponentBackedTarget.h"
+#include "ToolTargetManager.h"
+
 #define LOCTEXT_NAMESPACE "UMeshSurfacePointTool"
 
 
@@ -14,9 +20,20 @@
  * ToolBuilder
  */
 
+const FToolTargetTypeRequirements& UMeshSurfacePointToolBuilder::GetTargetRequirements() const
+{
+	static FToolTargetTypeRequirements TypeRequirements({
+		UMaterialProvider::StaticClass(),
+		UMeshDescriptionCommitter::StaticClass(),
+		UMeshDescriptionProvider::StaticClass(),
+		UPrimitiveComponentBackedTarget::StaticClass()
+		});
+	return TypeRequirements;
+}
+
 bool UMeshSurfacePointToolBuilder::CanBuildTool(const FToolBuilderState& SceneState) const
 {
-	return ToolBuilderUtil::CountComponents(SceneState, CanMakeComponentTarget) == 1;
+	return SceneState.TargetManager->CountSelectedAndTargetable(SceneState, GetTargetRequirements()) == 1;
 }
 
 UInteractiveTool* UMeshSurfacePointToolBuilder::BuildTool(const FToolBuilderState& SceneState) const
@@ -33,11 +50,10 @@ UMeshSurfacePointTool* UMeshSurfacePointToolBuilder::CreateNewTool(const FToolBu
 
 void UMeshSurfacePointToolBuilder::InitializeNewTool(UMeshSurfacePointTool* NewTool, const FToolBuilderState& SceneState) const
 {
-	UActorComponent* ActorComponent = ToolBuilderUtil::FindFirstComponent(SceneState, CanMakeComponentTarget);
-	UPrimitiveComponent* MeshComponent = Cast<UPrimitiveComponent>(ActorComponent);
-	check(MeshComponent != nullptr);
+	UToolTarget* Target = SceneState.TargetManager->BuildFirstSelectedTargetable(SceneState, GetTargetRequirements());
+	check(Target);
+	NewTool->SetTarget(Target);
 	NewTool->SetStylusAPI(this->StylusAPI);
-	NewTool->SetSelection( MakeComponentTarget(MeshComponent) );
 }
 
 
@@ -76,7 +92,7 @@ void UMeshSurfacePointTool::SetStylusAPI(IToolStylusStateProviderAPI* StylusAPII
 
 bool UMeshSurfacePointTool::HitTest(const FRay& Ray, FHitResult& OutHit)
 {
-	return ComponentTarget->HitTest(Ray, OutHit);
+	return Cast<IPrimitiveComponentBackedTarget>(Target)->HitTestComponent(Ray, OutHit);
 }
 
 
diff --git a/Engine/Source/Runtime/Experimental/InteractiveToolsFramework/Public/BaseTools/MeshSurfacePointTool.h b/Engine/Source/Runtime/Experimental/InteractiveToolsFramework/Public/BaseTools/MeshSurfacePointTool.h
--- a/Engine/Source/Runtime/Experimental/InteractiveToolsFramework/Public/BaseTools/MeshSurfacePointTool.h
+++ b/Engine/Source/Runtime/Experimental/InteractiveToolsFramework/Public/BaseTools/MeshSurfacePointTool.h
@@ -44,6 +44,9 @@ public:
 
 	/** Called by BuildTool to configure the Tool with the input MeshSource based on the SceneState */
 	virtual void InitializeNewTool(UMeshSurfacePointTool* Tool, const FToolBuilderState& SceneState) const;
+
+protected:
+	virtual const FToolTargetTypeRequirements& GetTargetRequirements() const override;
 };
 
 
