ModelingTools: ToolTarget refactor of remaining SingleSelectionTool (pt.1)

#rnx

diff --git a/Engine/Plugins/Experimental/MeshModelingToolset/Source/MeshModelingTools/Private/ConvertToPolygonsTool.cpp b/Engine/Plugins/Experimental/MeshModelingToolset/Source/MeshModelingTools/Private/ConvertToPolygonsTool.cpp
--- a/Engine/Plugins/Experimental/MeshModelingToolset/Source/MeshModelingTools/Private/ConvertToPolygonsTool.cpp
+++ b/Engine/Plugins/Experimental/MeshModelingToolset/Source/MeshModelingTools/Private/ConvertToPolygonsTool.cpp
@@ -20,6 +20,11 @@
 #include "MeshOpPreviewHelpers.h"
 #include "ModelingOperators.h"
 
+#include "TargetInterfaces/MaterialProvider.h"
+#include "TargetInterfaces/MeshDescriptionCommitter.h"
+#include "TargetInterfaces/MeshDescriptionProvider.h"
+#include "TargetInterfaces/PrimitiveComponentBackedTarget.h"
+
 #include "ExplicitUseGeometryMathTypes.h"		// using UE::Geometry::(math types)
 using namespace UE::Geometry;
 
@@ -29,23 +34,9 @@ using namespace UE::Geometry;
  * ToolBuilder
  */
 
-
-bool UConvertToPolygonsToolBuilder::CanBuildTool(const FToolBuilderState& SceneState) const
+USingleSelectionMeshEditingTool* UConvertToPolygonsToolBuilder::CreateNewTool(const FToolBuilderState& SceneState) const
 {
-	return ToolBuilderUtil::CountComponents(SceneState, CanMakeComponentTarget) == 1;
-}
-
-UInteractiveTool* UConvertToPolygonsToolBuilder::BuildTool(const FToolBuilderState& SceneState) const
-{
-	UConvertToPolygonsTool* NewTool = NewObject<UConvertToPolygonsTool>(SceneState.ToolManager);
-
-	UActorComponent* ActorComponent = ToolBuilderUtil::FindFirstComponent(SceneState, CanMakeComponentTarget);
-	UPrimitiveComponent* MeshComponent = Cast<UPrimitiveComponent>(ActorComponent);
-	check(MeshComponent != nullptr);
-	NewTool->SetSelection(MakeComponentTarget(MeshComponent));
-	NewTool->SetWorld(SceneState.World);
-
-	return NewTool;
+	return NewObject<UConvertToPolygonsTool>(SceneState.ToolManager);
 }
 
 class FConvertToPolygonsOp : public  FDynamicMeshOperator
@@ -156,11 +147,6 @@ UConvertToPolygonsTool::UConvertToPolygonsTool()
 	SetToolDisplayName(LOCTEXT("ConvertToPolygonsToolName", "Generate PolyGroups"));
 }
 
-void UConvertToPolygonsTool::SetWorld(UWorld* World)
-{
-	this->TargetWorld = World;
-}
-
 bool UConvertToPolygonsTool::CanAccept() const
 {
 	return Super::CanAccept() && (PreviewWithBackgroundCompute == nullptr || PreviewWithBackgroundCompute->HaveValidResult());
@@ -170,22 +156,24 @@ void UConvertToPolygonsTool::Setup()
 {
 	UInteractiveTool::Setup();
 
+	IPrimitiveComponentBackedTarget* TargetComponent = Cast<IPrimitiveComponentBackedTarget>(Target);
+
 	FComponentMaterialSet MaterialSet;
-	ComponentTarget->GetMaterialSet(MaterialSet);
+	Cast<IMaterialProvider>(Target)->GetMaterialSet(MaterialSet);
 
 	// populate the OriginalDynamicMesh with a conversion of the input mesh.
 	{
 		OriginalDynamicMesh = MakeShared<FDynamicMesh3, ESPMode::ThreadSafe>();
 		FMeshDescriptionToDynamicMesh Converter;
-		Converter.Convert(ComponentTarget->GetMesh(), *OriginalDynamicMesh);
+		Converter.Convert(Cast<IMeshDescriptionProvider>(Target)->GetMeshDescription(), *OriginalDynamicMesh);
 	}
 
 	Settings = NewObject<UConvertToPolygonsToolProperties>(this);
 	Settings->RestoreProperties(this);
 	AddToolPropertySource(Settings);
-	FTransform MeshTransform = ComponentTarget->GetWorldTransform();
+	FTransform MeshTransform = TargetComponent->GetWorldTransform();
 	// hide existing mesh
-	ComponentTarget->SetOwnerVisibility(false);
+	TargetComponent->SetOwnerVisibility(false);
 	// Set up the preview object
 	{
 		// create the operator factory
@@ -198,7 +186,7 @@ void UConvertToPolygonsTool::Setup()
 		PreviewWithBackgroundCompute->SetIsMeshTopologyConstant(true, EMeshRenderAttributeFlags::Positions | EMeshRenderAttributeFlags::VertexNormals);
 
 		// Give the preview something to display
-		PreviewWithBackgroundCompute->PreviewMesh->SetTransform(ComponentTarget->GetWorldTransform());
+		PreviewWithBackgroundCompute->PreviewMesh->SetTransform(TargetComponent->GetWorldTransform());
 		PreviewWithBackgroundCompute->PreviewMesh->SetTangentsMode(EDynamicMeshTangentCalcType::AutoCalculated);
 		PreviewWithBackgroundCompute->PreviewMesh->UpdatePreview(OriginalDynamicMesh.Get());
 		
@@ -252,7 +240,7 @@ void UConvertToPolygonsTool::UpdateOpParameters(FConvertToPolygonsOp& ConvertToP
 	ConvertToPolygonsOp.AngleTolerance    = Settings->AngleTolerance;
 	ConvertToPolygonsOp.OriginalMesh = OriginalDynamicMesh;
 	
-	FTransform LocalToWorld = ComponentTarget->GetWorldTransform();
+	FTransform LocalToWorld = Cast<IPrimitiveComponentBackedTarget>(Target)->GetWorldTransform();
 	ConvertToPolygonsOp.SetTransform(LocalToWorld);
 }
 
@@ -263,10 +251,10 @@ void UConvertToPolygonsTool::GenerateAsset(const FDynamicMeshOpResult& Result)
 	FDynamicMesh3* DynamicMeshResult = Result.Mesh.Get();
 	check(DynamicMeshResult != nullptr);
 
-	ComponentTarget->CommitMesh([DynamicMeshResult](const FPrimitiveComponentTarget::FCommitParams& CommitParams)
+	Cast<IMeshDescriptionCommitter>(Target)->CommitMeshDescription([DynamicMeshResult](const IMeshDescriptionCommitter::FCommitterParams& CommitParams)
 	{
 		FDynamicMeshToMeshDescription Converter;
-		Converter.Convert(DynamicMeshResult, *CommitParams.MeshDescription);
+		Converter.Convert(DynamicMeshResult, *CommitParams.MeshDescriptionOut);
 	});
 
 	GetToolManager()->EndUndoTransaction();
@@ -275,7 +263,7 @@ void UConvertToPolygonsTool::GenerateAsset(const FDynamicMeshOpResult& Result)
 void UConvertToPolygonsTool::Shutdown(EToolShutdownType ShutdownType)
 {
 	Settings->SaveProperties(this);
-	ComponentTarget->SetOwnerVisibility(true);
+	Cast<IPrimitiveComponentBackedTarget>(Target)->SetOwnerVisibility(true);
 
 	if (PreviewWithBackgroundCompute)
 	{
@@ -304,7 +292,7 @@ void UConvertToPolygonsTool::Render(IToolsContextRenderAPI* RenderAPI)
 
 	FPrimitiveDrawInterface* PDI = RenderAPI->GetPrimitiveDrawInterface();
 	float PDIScale = RenderAPI->GetCameraState().GetPDIScalingFactor();
-	FTransform Transform = ComponentTarget->GetWorldTransform(); //Actor->GetTransform();
+	FTransform Transform = Cast<IPrimitiveComponentBackedTarget>(Target)->GetWorldTransform(); //Actor->GetTransform();
 
 	for (int eid : PolygonEdges)
 	{
@@ -324,10 +312,11 @@ void UConvertToPolygonsTool::UpdateVisualization()
 		return;
 	}
 
+	IMaterialProvider* MaterialTarget = Cast<IMaterialProvider>(Target);
 	FComponentMaterialSet MaterialSet;
 	if (Settings->bShowGroupColors)
 	{
-		int32 NumMaterials = ComponentTarget->GetNumMaterials();
+		int32 NumMaterials = MaterialTarget->GetNumMaterials();
 		for (int32 i = 0; i < NumMaterials; ++i)
 		{ 
 			MaterialSet.Materials.Add(ToolSetupUtil::GetSelectionMaterial(GetToolManager()));
@@ -340,7 +329,7 @@ void UConvertToPolygonsTool::UpdateVisualization()
 	}
 	else
 	{
-		ComponentTarget->GetMaterialSet(MaterialSet);
+		MaterialTarget->GetMaterialSet(MaterialSet);
 		PreviewWithBackgroundCompute->PreviewMesh->ClearTriangleColorFunction(UPreviewMesh::ERenderUpdateMode::FastUpdate);
 	}
 	PreviewWithBackgroundCompute->ConfigureMaterials(MaterialSet.Materials,
diff --git a/Engine/Plugins/Experimental/MeshModelingToolset/Source/MeshModelingTools/Private/DisplaceMeshTool.cpp b/Engine/Plugins/Experimental/MeshModelingToolset/Source/MeshModelingTools/Private/DisplaceMeshTool.cpp
--- a/Engine/Plugins/Experimental/MeshModelingToolset/Source/MeshModelingTools/Private/DisplaceMeshTool.cpp
+++ b/Engine/Plugins/Experimental/MeshModelingToolset/Source/MeshModelingTools/Private/DisplaceMeshTool.cpp
@@ -21,6 +21,11 @@
 #include "Engine/Classes/Engine/StaticMesh.h"
 #include "Engine/Classes/Components/StaticMeshComponent.h"
 
+#include "TargetInterfaces/MaterialProvider.h"
+#include "TargetInterfaces/MeshDescriptionCommitter.h"
+#include "TargetInterfaces/MeshDescriptionProvider.h"
+#include "TargetInterfaces/PrimitiveComponentBackedTarget.h"
+
 #include "ExplicitUseGeometryMathTypes.h"		// using UE::Geometry::(math types)
 using namespace UE::Geometry;
 
@@ -723,22 +728,9 @@ namespace DisplaceMeshToolLocals{
 /*
  * ToolBuilder
  */
-bool UDisplaceMeshToolBuilder::CanBuildTool(const FToolBuilderState& SceneState) const
-{
-	return ToolBuilderUtil::CountComponents(SceneState, CanMakeComponentTarget) == 1;
-}
-
-UInteractiveTool* UDisplaceMeshToolBuilder::BuildTool(const FToolBuilderState& SceneState) const
+USingleSelectionMeshEditingTool* UDisplaceMeshToolBuilder::CreateNewTool(const FToolBuilderState& SceneState) const
 {
-	UDisplaceMeshTool* NewTool = NewObject<UDisplaceMeshTool>(SceneState.ToolManager);
-
-	UActorComponent* ActorComponent = ToolBuilderUtil::FindFirstComponent(SceneState, CanMakeComponentTarget);
-	auto* MeshComponent = Cast<UPrimitiveComponent>(ActorComponent);
-	check(MeshComponent != nullptr);
-
-	NewTool->SetSelection(MakeComponentTarget(MeshComponent));
-
-	return NewTool;
+	return NewObject<UDisplaceMeshTool>(SceneState.ToolManager);
 }
 
 /*
@@ -791,8 +783,9 @@ void UDisplaceMeshTool::Setup()
 #endif
 	
 	// populate weight maps list
+	IMeshDescriptionProvider* TargetMeshProvider = Cast<IMeshDescriptionProvider>(Target);
 	TArray<FName> WeightMaps;
-	UE::WeightMaps::FindVertexWeightMaps(ComponentTarget->GetMesh(), WeightMaps);
+	UE::WeightMaps::FindVertexWeightMaps(TargetMeshProvider->GetMeshDescription(), WeightMaps);
 	CommonProperties->WeightMapsList.Add(TEXT("None"));
 	for (FName Name : WeightMaps)
 	{
@@ -806,22 +799,23 @@ void UDisplaceMeshTool::Setup()
 
 
 	// create dynamic mesh component to use for live preview
-	DynamicMeshComponent = NewObject<USimpleDynamicMeshComponent>(ComponentTarget->GetOwnerActor(), "DynamicMesh");
-	DynamicMeshComponent->SetupAttachment(ComponentTarget->GetOwnerActor()->GetRootComponent());
+	IPrimitiveComponentBackedTarget* TargetComponent = Cast<IPrimitiveComponentBackedTarget>(Target);
+	DynamicMeshComponent = NewObject<USimpleDynamicMeshComponent>(TargetComponent->GetOwnerActor(), "DynamicMesh");
+	DynamicMeshComponent->SetupAttachment(TargetComponent->GetOwnerActor()->GetRootComponent());
 	DynamicMeshComponent->RegisterComponent();
-	DynamicMeshComponent->SetWorldTransform(ComponentTarget->GetWorldTransform());
+	DynamicMeshComponent->SetWorldTransform(TargetComponent->GetWorldTransform());
 	DynamicMeshComponent->bExplicitShowWireframe = CommonProperties->bShowWireframe;
 
 	// transfer materials
 	FComponentMaterialSet MaterialSet;
-	ComponentTarget->GetMaterialSet(MaterialSet);
+	Cast<IMaterialProvider>(Target)->GetMaterialSet(MaterialSet);
 	for (int k = 0; k < MaterialSet.Materials.Num(); ++k)
 	{
 		DynamicMeshComponent->SetMaterial(k, MaterialSet.Materials[k]);
 	}
 
 	DynamicMeshComponent->TangentsType = EDynamicMeshTangentCalcType::AutoCalculated;
-	DynamicMeshComponent->InitializeMesh(ComponentTarget->GetMesh());
+	DynamicMeshComponent->InitializeMesh(TargetMeshProvider->GetMeshDescription());
 	OriginalMesh.Copy(*DynamicMeshComponent->GetMesh());
 	OriginalMeshSpatial.SetMesh(&OriginalMesh, true);
 
@@ -843,7 +837,7 @@ void UDisplaceMeshTool::Setup()
 	Displacer = MakeUnique<FDisplaceMeshOpFactory>(SubdividedMesh, Parameters, CommonProperties->DisplacementType);
 		
 	// hide input StaticMeshComponent
-	ComponentTarget->SetOwnerVisibility(false);
+	TargetComponent->SetOwnerVisibility(false);
 
 	// initialize our properties
 	ToolPropertyObjects.Add(this);
@@ -900,7 +894,8 @@ void UDisplaceMeshTool::Shutdown(EToolShutdownType ShutdownType)
 
 	if (DynamicMeshComponent != nullptr)
 	{
-		ComponentTarget->SetOwnerVisibility(true);
+		IPrimitiveComponentBackedTarget* TargetComponent = Cast<IPrimitiveComponentBackedTarget>(Target);
+		TargetComponent->SetOwnerVisibility(true);
 
 		if (ShutdownType == EToolShutdownType::Accept)
 		{
@@ -911,7 +906,7 @@ void UDisplaceMeshTool::Shutdown(EToolShutdownType ShutdownType)
 			// on the underlying StaticMesh Asset, or it will run on the Bake() below and the output result will not be the same as the preview
 			if (CommonProperties->DisplacementType == EDisplaceMeshToolDisplaceType::DisplacementMap && TextureMapProperties->bRecalcNormals == false)
 			{
-				UStaticMeshComponent* StaticMeshComponent = Cast<UStaticMeshComponent>(ComponentTarget->GetOwnerComponent());
+				UStaticMeshComponent* StaticMeshComponent = Cast<UStaticMeshComponent>(TargetComponent->GetOwnerComponent());
 				if (StaticMeshComponent)
 				{
 					if (UStaticMesh* StaticMesh = StaticMeshComponent->GetStaticMesh())
@@ -927,9 +922,9 @@ void UDisplaceMeshTool::Shutdown(EToolShutdownType ShutdownType)
 			}
 
 
-			ComponentTarget->CommitMesh([=](const FPrimitiveComponentTarget::FCommitParams& CommitParams)
+			Cast<IMeshDescriptionCommitter>(Target)->CommitMeshDescription([this](const IMeshDescriptionCommitter::FCommitterParams& CommitParams)
 			{
-				DynamicMeshComponent->Bake(CommitParams.MeshDescription, CommonProperties->Subdivisions > 0);
+				DynamicMeshComponent->Bake(CommitParams.MeshDescriptionOut, CommonProperties->Subdivisions > 0);
 			});
 			GetToolManager()->EndUndoTransaction();
 		}
@@ -1162,7 +1157,7 @@ void UDisplaceMeshTool::UpdateActiveWeightMap()
 	else
 	{
 		TSharedPtr<FIndexedWeightMap, ESPMode::ThreadSafe> NewWeightMap = MakeShared<FIndexedWeightMap, ESPMode::ThreadSafe>();
-		UE::WeightMaps::GetVertexWeightMap(ComponentTarget->GetMesh(), CommonProperties->WeightMap, *NewWeightMap, 1.0f);
+		UE::WeightMaps::GetVertexWeightMap(Cast<IMeshDescriptionProvider>(Target)->GetMeshDescription(), CommonProperties->WeightMap, *NewWeightMap, 1.0f);
 		if (CommonProperties->bInvertWeightMap)
 		{
 			NewWeightMap->InvertWeightMap();
diff --git a/Engine/Plugins/Experimental/MeshModelingToolset/Source/MeshModelingTools/Private/EdgeLoopInsertionTool.cpp b/Engine/Plugins/Experimental/MeshModelingToolset/Source/MeshModelingTools/Private/EdgeLoopInsertionTool.cpp
--- a/Engine/Plugins/Experimental/MeshModelingToolset/Source/MeshModelingTools/Private/EdgeLoopInsertionTool.cpp
+++ b/Engine/Plugins/Experimental/MeshModelingToolset/Source/MeshModelingTools/Private/EdgeLoopInsertionTool.cpp
@@ -13,29 +13,19 @@
 #include "ToolSceneQueriesUtil.h"
 #include "ToolSetupUtil.h"
 
+#include "TargetInterfaces/MaterialProvider.h"
+#include "TargetInterfaces/MeshDescriptionCommitter.h"
+#include "TargetInterfaces/MeshDescriptionProvider.h"
+#include "TargetInterfaces/PrimitiveComponentBackedTarget.h"
+
 #include "ExplicitUseGeometryMathTypes.h"		// using UE::Geometry::(math types)
 using namespace UE::Geometry;
 
 #define LOCTEXT_NAMESPACE "UEdgeLoopInsertionTool"
 
-bool UEdgeLoopInsertionToolBuilder::CanBuildTool(const FToolBuilderState& SceneState) const
-{
-	return AssetAPI != nullptr && ToolBuilderUtil::CountComponents(SceneState, CanMakeComponentTarget) == 1;
-}
-
-UInteractiveTool* UEdgeLoopInsertionToolBuilder::BuildTool(const FToolBuilderState& SceneState) const
+USingleSelectionMeshEditingTool* UEdgeLoopInsertionToolBuilder::CreateNewTool(const FToolBuilderState& SceneState) const
 {
-	UEdgeLoopInsertionTool* NewTool = NewObject<UEdgeLoopInsertionTool>(SceneState.ToolManager);
-
-	UActorComponent* ActorComponent = ToolBuilderUtil::FindFirstComponent(SceneState, CanMakeComponentTarget);
-	UPrimitiveComponent* MeshComponent = Cast<UPrimitiveComponent>(ActorComponent);
-	check(MeshComponent != nullptr);
-	NewTool->SetSelection(MakeComponentTarget(MeshComponent));
-
-	NewTool->SetWorld(SceneState.World);
-	NewTool->SetAssetAPI(AssetAPI);
-
-	return NewTool;
+	return NewObject<UEdgeLoopInsertionTool>(SceneState.ToolManager);
 }
 
 TUniquePtr<FDynamicMeshOperator> UEdgeLoopInsertionOperatorFactory::MakeNewOperator()
@@ -44,7 +34,7 @@ TUniquePtr<FDynamicMeshOperator> UEdgeLoopInsertionOperatorFactory::MakeNewOpera
 
 	Op->OriginalMesh = Tool->CurrentMesh;
 	Op->OriginalTopology = Tool->CurrentTopology;
-	Op->SetTransform(Tool->ComponentTarget->GetWorldTransform());
+	Op->SetTransform(Cast<IPrimitiveComponentBackedTarget>(Tool->Target)->GetWorldTransform());
 
 	if (Tool->bShowingBaseMesh)
 	{
@@ -101,7 +91,7 @@ void UEdgeLoopInsertionTool::Setup()
 {
 	USingleSelectionTool::Setup();
 
-	if (!ComponentTarget)
+	if (!Target)
 	{
 		return;
 	}
@@ -114,7 +104,7 @@ void UEdgeLoopInsertionTool::Setup()
 	// Initialize the mesh that we'll be operating on
 	CurrentMesh = MakeShared<FDynamicMesh3>();
 	FMeshDescriptionToDynamicMesh Converter;
-	Converter.Convert(ComponentTarget->GetMesh(), *CurrentMesh);
+	Converter.Convert(Cast<IMeshDescriptionProvider>(Target)->GetMeshDescription(), *CurrentMesh);
 	CurrentTopology = MakeShared<FGroupTopology>(CurrentMesh.Get(), true);
 	MeshSpatial.SetMesh(CurrentMesh.Get(), true);
 
@@ -143,7 +133,7 @@ void UEdgeLoopInsertionTool::Setup()
 	TopologySelector.Initialize(CurrentMesh.Get(), CurrentTopology.Get());
 	TopologySelector.SetSpatialSource([this]() {return &MeshSpatial; });
 	TopologySelector.PointsWithinToleranceTest = [this](const FVector3d& Position1, const FVector3d& Position2, double TolScale) {
-		UE::Geometry::FTransform3d Transform(ComponentTarget->GetWorldTransform());
+		UE::Geometry::FTransform3d Transform(Cast<IPrimitiveComponentBackedTarget>(Target)->GetWorldTransform());
 		return ToolSceneQueriesUtil::PointSnapQuery(CameraState, Transform.TransformPosition(Position1), Transform.TransformPosition(Position2),
 			ToolSceneQueriesUtil::GetDefaultVisualAngleSnapThreshD() * TolScale);
 	};
@@ -162,7 +152,7 @@ void UEdgeLoopInsertionTool::SetupPreview()
 	Preview->PreviewMesh->SetTangentsMode(EDynamicMeshTangentCalcType::AutoCalculated);
 
 	FComponentMaterialSet MaterialSet;
-	ComponentTarget->GetMaterialSet(MaterialSet);
+	Cast<IMaterialProvider>(Target)->GetMaterialSet(MaterialSet);
 	Preview->ConfigureMaterials(MaterialSet.Materials, ToolSetupUtil::GetDefaultWorkingMaterial(GetToolManager()));
 
 	// Whenever we get a new result from the op, we need to extract the preview edges so that
@@ -192,20 +182,21 @@ void UEdgeLoopInsertionTool::SetupPreview()
 		});
 
 	// Set initial preview to unprocessed mesh, so that things don't disappear initially
+	IPrimitiveComponentBackedTarget* TargetComponent = Cast<IPrimitiveComponentBackedTarget>(Target);
 	Preview->PreviewMesh->UpdatePreview(CurrentMesh.Get());
-	Preview->PreviewMesh->SetTransform(ComponentTarget->GetWorldTransform());
+	Preview->PreviewMesh->SetTransform(TargetComponent->GetWorldTransform());
 	Preview->PreviewMesh->EnableWireframe(Settings->bWireframe);
 	Preview->SetVisibility(true);
 	ClearPreview();
 
-	ComponentTarget->SetOwnerVisibility(false);
+	TargetComponent->SetOwnerVisibility(false);
 }
 
 void UEdgeLoopInsertionTool::Shutdown(EToolShutdownType ShutdownType)
 {
 	Settings->SaveProperties(this);
 	Preview->Shutdown();
-	ComponentTarget->SetOwnerVisibility(true);
+	Cast<IPrimitiveComponentBackedTarget>(Target)->SetOwnerVisibility(true);
 	CurrentMesh.Reset();
 	CurrentTopology.Reset();
 	ExpireChanges();
@@ -225,10 +216,10 @@ void UEdgeLoopInsertionTool::OnTick(float DeltaTime)
 				GetToolManager()->BeginUndoTransaction(LOCTEXT("EdgeLoopInsertionTransactionName", "Edge Loop Insertion"));
 
 				GetToolManager()->EmitObjectChange(this, MakeUnique<FEdgeLoopInsertionChangeBookend>(CurrentChangeStamp, true), LOCTEXT("EdgeLoopInsertion", "Edge Loop Insertion"));
-				ComponentTarget->CommitMesh([this](const FPrimitiveComponentTarget::FCommitParams& CommitParams)
+				Cast<IMeshDescriptionCommitter>(Target)->CommitMeshDescription([this](const IMeshDescriptionCommitter::FCommitterParams& CommitParams)
 					{
 						FDynamicMeshToMeshDescription Converter;
-						Converter.Convert(Preview->PreviewMesh->GetMesh(), *CommitParams.MeshDescription);
+						Converter.Convert(Preview->PreviewMesh->GetMesh(), *CommitParams.MeshDescriptionOut);
 					});
 				GetToolManager()->EmitObjectChange(this, MakeUnique<FEdgeLoopInsertionChangeBookend>(CurrentChangeStamp, false), LOCTEXT("EdgeLoopInsertion", "Edge Loop Insertion"));
 
@@ -295,7 +286,7 @@ FInputRayHit UEdgeLoopInsertionTool::HitTest(const FRay& WorldRay)
 	FInputRayHit Hit;
 
 	// See if we hit an edge
-	FTransform LocalToWorld = ComponentTarget->GetWorldTransform();
+	FTransform LocalToWorld = Cast<IPrimitiveComponentBackedTarget>(Target)->GetWorldTransform();
 	FRay3d LocalRay(LocalToWorld.InverseTransformPosition(WorldRay.Origin),
 		LocalToWorld.InverseTransformVector(WorldRay.Direction), false);
 	FGroupTopologySelection Selection;
@@ -320,7 +311,7 @@ FInputRayHit UEdgeLoopInsertionTool::HitTest(const FRay& WorldRay)
 bool UEdgeLoopInsertionTool::UpdateHoveredItem(const FRay& WorldRay)
 {
 	// Check that we hit an edge
-	FTransform LocalToWorld = ComponentTarget->GetWorldTransform();
+	FTransform LocalToWorld = Cast<IPrimitiveComponentBackedTarget>(Target)->GetWorldTransform();
 	FRay3d LocalRay(LocalToWorld.InverseTransformPosition(WorldRay.Origin),
 		LocalToWorld.InverseTransformVector(WorldRay.Direction), false);
 
@@ -503,7 +494,7 @@ void FEdgeLoopInsertionChangeBookend::Apply(UObject* Object)
 		UEdgeLoopInsertionTool* Tool = Cast<UEdgeLoopInsertionTool>(Object);
 		Tool->CurrentMesh->Clear();
 		FMeshDescriptionToDynamicMesh Converter;
-		Converter.Convert(Tool->ComponentTarget->GetMesh(), *Tool->CurrentMesh);
+		Converter.Convert(Cast<IMeshDescriptionProvider>(Tool->Target)->GetMeshDescription(), *Tool->CurrentMesh);
 		Tool->CurrentTopology->RebuildTopology();
 		Tool->MeshSpatial.Build();
 		Tool->TopologySelector.Invalidate(true, true);
@@ -518,7 +509,7 @@ void FEdgeLoopInsertionChangeBookend::Revert(UObject* Object)
 		UEdgeLoopInsertionTool* Tool = Cast<UEdgeLoopInsertionTool>(Object);
 		Tool->CurrentMesh->Clear();
 		FMeshDescriptionToDynamicMesh Converter;
-		Converter.Convert(Tool->ComponentTarget->GetMesh(), *Tool->CurrentMesh);
+		Converter.Convert(Cast<IMeshDescriptionProvider>(Tool->Target)->GetMeshDescription(), *Tool->CurrentMesh);
 		Tool->CurrentTopology->RebuildTopology();
 		Tool->MeshSpatial.Build();
 		Tool->TopologySelector.Invalidate(true, true);
diff --git a/Engine/Plugins/Experimental/MeshModelingToolset/Source/MeshModelingTools/Private/GroupEdgeInsertionTool.cpp b/Engine/Plugins/Experimental/MeshModelingToolset/Source/MeshModelingTools/Private/GroupEdgeInsertionTool.cpp
--- a/Engine/Plugins/Experimental/MeshModelingToolset/Source/MeshModelingTools/Private/GroupEdgeInsertionTool.cpp
+++ b/Engine/Plugins/Experimental/MeshModelingToolset/Source/MeshModelingTools/Private/GroupEdgeInsertionTool.cpp
@@ -12,6 +12,11 @@
 #include "ToolSceneQueriesUtil.h"
 #include "ToolSetupUtil.h"
 
+#include "TargetInterfaces/MaterialProvider.h"
+#include "TargetInterfaces/MeshDescriptionCommitter.h"
+#include "TargetInterfaces/MeshDescriptionProvider.h"
+#include "TargetInterfaces/PrimitiveComponentBackedTarget.h"
+
 #include "ExplicitUseGeometryMathTypes.h"		// using UE::Geometry::(math types)
 using namespace UE::Geometry;
 
@@ -24,25 +29,9 @@ bool GetSharedBoundary(const FGroupTopology& Topology,
 bool DoesBoundaryContainPoint(const FGroupTopology& Topology,
 	const FGroupTopology::FGroupBoundary& Boundary, int32 PointTopologyID, bool bPointIsCorner);
 
-
-bool UGroupEdgeInsertionToolBuilder::CanBuildTool(const FToolBuilderState& SceneState) const
-{
-	return AssetAPI != nullptr && ToolBuilderUtil::CountComponents(SceneState, CanMakeComponentTarget) == 1;
-}
-
-UInteractiveTool* UGroupEdgeInsertionToolBuilder::BuildTool(const FToolBuilderState& SceneState) const
+USingleSelectionMeshEditingTool* UGroupEdgeInsertionToolBuilder::CreateNewTool(const FToolBuilderState& SceneState) const
 {
-	UGroupEdgeInsertionTool* NewTool = NewObject<UGroupEdgeInsertionTool>(SceneState.ToolManager);
-
-	UActorComponent* ActorComponent = ToolBuilderUtil::FindFirstComponent(SceneState, CanMakeComponentTarget);
-	UPrimitiveComponent* MeshComponent = Cast<UPrimitiveComponent>(ActorComponent);
-	check(MeshComponent != nullptr);
-	NewTool->SetSelection(MakeComponentTarget(MeshComponent));
-
-	NewTool->SetWorld(SceneState.World);
-	NewTool->SetAssetAPI(AssetAPI);
-
-	return NewTool;
+	return NewObject<UGroupEdgeInsertionTool>(SceneState.ToolManager);
 }
 
 TUniquePtr<FDynamicMeshOperator> UGroupEdgeInsertionOperatorFactory::MakeNewOperator()
@@ -51,7 +40,7 @@ TUniquePtr<FDynamicMeshOperator> UGroupEdgeInsertionOperatorFactory::MakeNewOper
 
 	Op->OriginalMesh = Tool->CurrentMesh;
 	Op->OriginalTopology = Tool->CurrentTopology;
-	Op->SetTransform(Tool->ComponentTarget->GetWorldTransform());
+	Op->SetTransform(Cast<IPrimitiveComponentBackedTarget>(Tool->Target)->GetWorldTransform());
 
 	if (Tool->bShowingBaseMesh)
 	{
@@ -82,7 +71,7 @@ void UGroupEdgeInsertionTool::Setup()
 {
 	USingleSelectionTool::Setup();
 
-	if (!ComponentTarget)
+	if (!Target)
 	{
 		return;
 	}
@@ -95,7 +84,7 @@ void UGroupEdgeInsertionTool::Setup()
 	// Initialize the mesh that we'll be operating on
 	CurrentMesh = MakeShared<FDynamicMesh3, ESPMode::ThreadSafe>();
 	FMeshDescriptionToDynamicMesh Converter;
-	Converter.Convert(ComponentTarget->GetMesh(), *CurrentMesh);
+	Converter.Convert(Cast<IMeshDescriptionProvider>(Target)->GetMeshDescription(), *CurrentMesh);
 	CurrentTopology = MakeShared<FGroupTopology, ESPMode::ThreadSafe>(CurrentMesh.Get(), true);
 	MeshSpatial.SetMesh(CurrentMesh.Get(), true);
 
@@ -127,7 +116,7 @@ void UGroupEdgeInsertionTool::Setup()
 	TopologySelector.Initialize(CurrentMesh.Get(), CurrentTopology.Get());
 	TopologySelector.SetSpatialSource([this]() {return &MeshSpatial; });
 	TopologySelector.PointsWithinToleranceTest = [this](const FVector3d& Position1, const FVector3d& Position2, double TolScale) {
-		UE::Geometry::FTransform3d Transform(ComponentTarget->GetWorldTransform());
+		UE::Geometry::FTransform3d Transform(Cast<IPrimitiveComponentBackedTarget>(Target)->GetWorldTransform());
 		return ToolSceneQueriesUtil::PointSnapQuery(CameraState, Transform.TransformPosition(Position1), Transform.TransformPosition(Position2),
 			ToolSceneQueriesUtil::GetDefaultVisualAngleSnapThreshD() * TolScale);
 	};
@@ -146,7 +135,7 @@ void UGroupEdgeInsertionTool::SetupPreview()
 	Preview->PreviewMesh->SetTangentsMode(EDynamicMeshTangentCalcType::AutoCalculated);
 
 	FComponentMaterialSet MaterialSet;
-	ComponentTarget->GetMaterialSet(MaterialSet);
+	Cast<IMaterialProvider>(Target)->GetMaterialSet(MaterialSet);
 	Preview->ConfigureMaterials(MaterialSet.Materials, ToolSetupUtil::GetDefaultWorkingMaterial(GetToolManager()));
 
 	// Whenever we get a new result from the op, we need to extract the preview edges so that
@@ -179,20 +168,21 @@ void UGroupEdgeInsertionTool::SetupPreview()
 		});
 
 	// Set initial preview to unprocessed mesh, so that things don't disappear initially
+	IPrimitiveComponentBackedTarget* TargetComponent = Cast<IPrimitiveComponentBackedTarget>(Target);
 	Preview->PreviewMesh->UpdatePreview(CurrentMesh.Get());
-	Preview->PreviewMesh->SetTransform(ComponentTarget->GetWorldTransform());
+	Preview->PreviewMesh->SetTransform(TargetComponent->GetWorldTransform());
 	Preview->PreviewMesh->EnableWireframe(Settings->bWireframe);
 	Preview->SetVisibility(true);
 	ClearPreview();
 
-	ComponentTarget->SetOwnerVisibility(false);
+	TargetComponent->SetOwnerVisibility(false);
 }
 
 void UGroupEdgeInsertionTool::Shutdown(EToolShutdownType ShutdownType)
 {
 	Settings->SaveProperties(this);
 	Preview->Shutdown();
-	ComponentTarget->SetOwnerVisibility(true);
+	Cast<IPrimitiveComponentBackedTarget>(Target)->SetOwnerVisibility(true);
 	CurrentMesh.Reset();
 	CurrentTopology.Reset();
 	ExpireChanges();
@@ -212,10 +202,10 @@ void UGroupEdgeInsertionTool::OnTick(float DeltaTime)
 				GetToolManager()->BeginUndoTransaction(LOCTEXT("GroupEdgeInsertionTransactionName", "Group Edge Insertion"));
 
 				GetToolManager()->EmitObjectChange(this, MakeUnique<FGroupEdgeInsertionChangeBookend>(CurrentChangeStamp, true), LOCTEXT("GroupEdgeInsertion", "Group Edge Insertion"));
-				ComponentTarget->CommitMesh([this](const FPrimitiveComponentTarget::FCommitParams& CommitParams)
+				Cast<IMeshDescriptionCommitter>(Target)->CommitMeshDescription([this](const IMeshDescriptionCommitter::FCommitterParams& CommitParams)
 					{
 						FDynamicMeshToMeshDescription Converter;
-						Converter.Convert(Preview->PreviewMesh->GetMesh(), *CommitParams.MeshDescription);
+						Converter.Convert(Preview->PreviewMesh->GetMesh(), *CommitParams.MeshDescriptionOut);
 					});
 				GetToolManager()->EmitObjectChange(this, MakeUnique<FGroupEdgeInsertionChangeBookend>(CurrentChangeStamp, false), LOCTEXT("GroupEdgeInsertion", "Group Edge Insertion"));
 
@@ -540,7 +530,7 @@ void UGroupEdgeInsertionTool::OnClicked(const FInputDeviceRay& ClickPos)
 bool UGroupEdgeInsertionTool::TopologyHitTest(const FRay& WorldRay, 
 	FVector3d& RayPositionOut, FRay3d* LocalRayOut)
 {
-	FTransform LocalToWorld = ComponentTarget->GetWorldTransform();
+	FTransform LocalToWorld = Cast<IPrimitiveComponentBackedTarget>(Target)->GetWorldTransform();
 	FRay3d LocalRay(LocalToWorld.InverseTransformPosition(WorldRay.Origin),
 		LocalToWorld.InverseTransformVector(WorldRay.Direction), false);
 
@@ -569,7 +559,7 @@ bool UGroupEdgeInsertionTool::GetHoveredItem(const FRay& WorldRay,
 	PointOut.ElementID = FDynamicMesh3::InvalidID;
 
 	// Cast the ray to see what we hit.
-	FTransform LocalToWorld = ComponentTarget->GetWorldTransform();
+	FTransform LocalToWorld = Cast<IPrimitiveComponentBackedTarget>(Target)->GetWorldTransform();
 	FRay3d LocalRay(LocalToWorld.InverseTransformPosition(WorldRay.Origin),
 		LocalToWorld.InverseTransformVector(WorldRay.Direction), false);
 	if (LocalRayOut)
@@ -834,7 +824,7 @@ void FGroupEdgeInsertionChangeBookend::Revert(UObject* Object)
 		UGroupEdgeInsertionTool* Tool = Cast<UGroupEdgeInsertionTool>(Object);
 		Tool->CurrentMesh->Clear();
 		FMeshDescriptionToDynamicMesh Converter;
-		Converter.Convert(Tool->ComponentTarget->GetMesh(), *Tool->CurrentMesh);
+		Converter.Convert(Cast<IMeshDescriptionProvider>(Tool->Target)->GetMeshDescription(), *Tool->CurrentMesh);
 		Tool->CurrentTopology->RebuildTopology();
 		Tool->MeshSpatial.Build();
 		Tool->TopologySelector.Invalidate(true, true);
@@ -856,7 +846,7 @@ void FGroupEdgeInsertionChangeBookend::Apply(UObject* Object)
 		UGroupEdgeInsertionTool* Tool = Cast<UGroupEdgeInsertionTool>(Object);
 		Tool->CurrentMesh->Clear();
 		FMeshDescriptionToDynamicMesh Converter;
-		Converter.Convert(Tool->ComponentTarget->GetMesh(), *Tool->CurrentMesh);
+		Converter.Convert(Cast<IMeshDescriptionProvider>(Tool->Target)->GetMeshDescription(), *Tool->CurrentMesh);
 		Tool->CurrentTopology->RebuildTopology();
 		Tool->MeshSpatial.Build();
 		Tool->TopologySelector.Invalidate(true, true);
diff --git a/Engine/Plugins/Experimental/MeshModelingToolset/Source/MeshModelingTools/Private/SingleSelectionMeshEditingTool.cpp b/Engine/Plugins/Experimental/MeshModelingToolset/Source/MeshModelingTools/Private/SingleSelectionMeshEditingTool.cpp
new file mode 100644
--- /dev/null
+++ b/Engine/Plugins/Experimental/MeshModelingToolset/Source/MeshModelingTools/Private/SingleSelectionMeshEditingTool.cpp
@@ -0,0 +1,51 @@
+// Copyright Epic Games, Inc. All Rights Reserved.
+
+#include "SingleSelectionMeshEditingTool.h"
+
+#include "TargetInterfaces/MaterialProvider.h"
+#include "TargetInterfaces/MeshDescriptionCommitter.h"
+#include "TargetInterfaces/MeshDescriptionProvider.h"
+#include "TargetInterfaces/PrimitiveComponentBackedTarget.h"
+#include "ToolTargetManager.h"
+
+/*
+ * ToolBuilder
+ */
+
+const FToolTargetTypeRequirements& USingleSelectionMeshEditingToolBuilder::GetTargetRequirements() const
+{
+	static FToolTargetTypeRequirements TypeRequirements({
+		UMaterialProvider::StaticClass(),
+		UMeshDescriptionCommitter::StaticClass(),
+		UMeshDescriptionProvider::StaticClass(),
+		UPrimitiveComponentBackedTarget::StaticClass()
+		});
+	return TypeRequirements;
+}
+
+bool USingleSelectionMeshEditingToolBuilder::CanBuildTool(const FToolBuilderState& SceneState) const
+{
+	bool bValidAssetAPI = (RequiresAssetAPI() ? AssetAPI != nullptr : true);
+	return bValidAssetAPI && SceneState.TargetManager->CountSelectedAndTargetable(SceneState, GetTargetRequirements()) == 1;
+}
+
+UInteractiveTool* USingleSelectionMeshEditingToolBuilder::BuildTool(const FToolBuilderState& SceneState) const
+{
+	USingleSelectionMeshEditingTool* NewTool = CreateNewTool(SceneState);
+	InitializeNewTool(NewTool, SceneState);
+	return NewTool;
+}
+
+void USingleSelectionMeshEditingToolBuilder::InitializeNewTool(USingleSelectionMeshEditingTool* NewTool, const FToolBuilderState& SceneState) const
+{
+	UToolTarget* Target = SceneState.TargetManager->BuildFirstSelectedTargetable(SceneState, GetTargetRequirements());
+	check(Target);
+	NewTool->SetTarget(Target);
+	NewTool->SetWorld(SceneState.World);
+	if (RequiresAssetAPI())
+	{
+		NewTool->SetAssetAPI(AssetAPI);
+	}
+}
+
+
diff --git a/Engine/Plugins/Experimental/MeshModelingToolset/Source/MeshModelingTools/Public/ConvertToPolygonsTool.h b/Engine/Plugins/Experimental/MeshModelingToolset/Source/MeshModelingTools/Public/ConvertToPolygonsTool.h
--- a/Engine/Plugins/Experimental/MeshModelingToolset/Source/MeshModelingTools/Public/ConvertToPolygonsTool.h
+++ b/Engine/Plugins/Experimental/MeshModelingToolset/Source/MeshModelingTools/Public/ConvertToPolygonsTool.h
@@ -11,6 +11,7 @@
 #include "PreviewMesh.h"
 #include "ModelingOperators.h"
 #include "MeshOpPreviewHelpers.h"
+#include "SingleSelectionMeshEditingTool.h"
 #include "ConvertToPolygonsTool.generated.h"
 
 // predeclaration
@@ -21,13 +22,12 @@ class FConvertToPolygonsOp;
  *
  */
 UCLASS()
-class MESHMODELINGTOOLS_API UConvertToPolygonsToolBuilder : public UInteractiveToolBuilder
+class MESHMODELINGTOOLS_API UConvertToPolygonsToolBuilder : public USingleSelectionMeshEditingToolBuilder
 {
 	GENERATED_BODY()
 
 public:
-	virtual bool CanBuildTool(const FToolBuilderState & SceneState) const override;
-	virtual UInteractiveTool* BuildTool(const FToolBuilderState & SceneState) const override;
+	virtual USingleSelectionMeshEditingTool* CreateNewTool(const FToolBuilderState & SceneState) const override;
 };
 
 
@@ -86,7 +86,7 @@ public:
  *
  */
 UCLASS()
-class MESHMODELINGTOOLS_API UConvertToPolygonsTool : public USingleSelectionTool
+class MESHMODELINGTOOLS_API UConvertToPolygonsTool : public USingleSelectionMeshEditingTool
 {
 	GENERATED_BODY()
 
@@ -96,8 +96,6 @@ public:
 	virtual void Setup() override;
 	virtual void Shutdown(EToolShutdownType ShutdownType) override;
 
-	virtual void SetWorld(UWorld* World);
-
 	virtual void Render(IToolsContextRenderAPI* RenderAPI) override;
 	virtual void OnTick(float DeltaTime) override;
 
@@ -120,9 +118,6 @@ protected:
 	UMeshOpPreviewWithBackgroundCompute* PreviewWithBackgroundCompute = nullptr;
 
 protected:
-
-	UWorld* TargetWorld;
-
 	TSharedPtr<UE::Geometry::FDynamicMesh3, ESPMode::ThreadSafe> OriginalDynamicMesh;
 
 	// for visualization
diff --git a/Engine/Plugins/Experimental/MeshModelingToolset/Source/MeshModelingTools/Public/DisplaceMeshTool.h b/Engine/Plugins/Experimental/MeshModelingToolset/Source/MeshModelingTools/Public/DisplaceMeshTool.h
--- a/Engine/Plugins/Experimental/MeshModelingToolset/Source/MeshModelingTools/Public/DisplaceMeshTool.h
+++ b/Engine/Plugins/Experimental/MeshModelingToolset/Source/MeshModelingTools/Public/DisplaceMeshTool.h
@@ -12,6 +12,7 @@
 #include "MeshOpPreviewHelpers.h"
 #include "Spatial/SampledScalarField2.h"
 #include "WeightMapUtil.h"
+#include "SingleSelectionMeshEditingTool.h"
 #include "DisplaceMeshTool.generated.h"
 
 struct FMeshDescription;
@@ -225,19 +226,18 @@ public:
  * Builder for Simple Mesh Displacement Tool
  */
 UCLASS()
-class MESHMODELINGTOOLS_API UDisplaceMeshToolBuilder : public UInteractiveToolBuilder
+class MESHMODELINGTOOLS_API UDisplaceMeshToolBuilder : public USingleSelectionMeshEditingToolBuilder
 {
 	GENERATED_BODY()
 public:
-	bool CanBuildTool(const FToolBuilderState& SceneState) const final;
-	UInteractiveTool* BuildTool(const FToolBuilderState& SceneState) const final;
+	USingleSelectionMeshEditingTool* CreateNewTool(const FToolBuilderState& SceneState) const final;
 };
 
 /**
  * Simple Mesh Displacement Tool
  */
 UCLASS()
-class MESHMODELINGTOOLS_API UDisplaceMeshTool : public USingleSelectionTool
+class MESHMODELINGTOOLS_API UDisplaceMeshTool : public USingleSelectionMeshEditingTool
 {
 	GENERATED_BODY()
 public:
diff --git a/Engine/Plugins/Experimental/MeshModelingToolset/Source/MeshModelingTools/Public/EdgeLoopInsertionTool.h b/Engine/Plugins/Experimental/MeshModelingToolset/Source/MeshModelingTools/Public/EdgeLoopInsertionTool.h
--- a/Engine/Plugins/Experimental/MeshModelingToolset/Source/MeshModelingTools/Public/EdgeLoopInsertionTool.h
+++ b/Engine/Plugins/Experimental/MeshModelingToolset/Source/MeshModelingTools/Public/EdgeLoopInsertionTool.h
@@ -13,21 +13,19 @@
 #include "Selection/GroupTopologySelector.h"
 #include "SingleSelectionTool.h"
 #include "ToolDataVisualizer.h"
+#include "SingleSelectionMeshEditingTool.h"
 
 #include "EdgeLoopInsertionTool.generated.h"
 
 class UDynamicMeshReplacementChangeTarget;
 
 UCLASS()
-class MESHMODELINGTOOLS_API UEdgeLoopInsertionToolBuilder : public UInteractiveToolBuilder
+class MESHMODELINGTOOLS_API UEdgeLoopInsertionToolBuilder : public USingleSelectionMeshEditingToolBuilder
 {
 	GENERATED_BODY()
 
 public:
-	IToolsContextAssetAPI* AssetAPI = nullptr;
-
-	virtual bool CanBuildTool(const FToolBuilderState& SceneState) const override;
-	virtual UInteractiveTool* BuildTool(const FToolBuilderState& SceneState) const override;
+	virtual USingleSelectionMeshEditingTool* CreateNewTool(const FToolBuilderState& SceneState) const override;
 };
 
 UENUM()
@@ -115,7 +113,7 @@ public:
 
 /** Tool for inserting (group) edge loops into a mesh. */
 UCLASS()
-class MESHMODELINGTOOLS_API UEdgeLoopInsertionTool : public USingleSelectionTool, public IHoverBehaviorTarget, public IClickBehaviorTarget
+class MESHMODELINGTOOLS_API UEdgeLoopInsertionTool : public USingleSelectionMeshEditingTool, public IHoverBehaviorTarget, public IClickBehaviorTarget
 {
 	GENERATED_BODY()
 public:
@@ -128,9 +126,6 @@ public:
 	virtual void Setup() override;
 	virtual void Shutdown(EToolShutdownType ShutdownType) override;
 
-	virtual void SetWorld(UWorld* World) { TargetWorld = World; }
-	virtual void SetAssetAPI(IToolsContextAssetAPI* AssetAPIIn) { AssetAPI = AssetAPIIn; }
-
 	virtual void OnTick(float DeltaTime) override;
 	virtual void Render(IToolsContextRenderAPI* RenderAPI) override;
 
@@ -167,9 +162,6 @@ protected:
 	UPROPERTY()
 	UMeshOpPreviewWithBackgroundCompute* Preview;
 
-	UWorld* TargetWorld;
-	IToolsContextAssetAPI* AssetAPI;
-
 	FToolDataVisualizer ExistingEdgesRenderer;
 	FToolDataVisualizer PreviewEdgeRenderer;
 	FGroupTopologySelector::FSelectionSettings TopologySelectorSettings;
diff --git a/Engine/Plugins/Experimental/MeshModelingToolset/Source/MeshModelingTools/Public/GroupEdgeInsertionTool.h b/Engine/Plugins/Experimental/MeshModelingToolset/Source/MeshModelingTools/Public/GroupEdgeInsertionTool.h
--- a/Engine/Plugins/Experimental/MeshModelingToolset/Source/MeshModelingTools/Public/GroupEdgeInsertionTool.h
+++ b/Engine/Plugins/Experimental/MeshModelingToolset/Source/MeshModelingTools/Public/GroupEdgeInsertionTool.h
@@ -14,6 +14,7 @@
 #include "Selection/GroupTopologySelector.h"
 #include "SingleSelectionTool.h"
 #include "ToolDataVisualizer.h"
+#include "SingleSelectionMeshEditingTool.h"
 
 #include "GroupEdgeInsertionTool.generated.h"
 
@@ -23,15 +24,12 @@ using UE::Geometry::FGroupEdgeInserter;
 using UE::Geometry::FDynamicMesh3;
 
 UCLASS()
-class MESHMODELINGTOOLS_API UGroupEdgeInsertionToolBuilder : public UInteractiveToolBuilder
+class MESHMODELINGTOOLS_API UGroupEdgeInsertionToolBuilder : public USingleSelectionMeshEditingToolBuilder
 {
 	GENERATED_BODY()
 
 public:
-	IToolsContextAssetAPI* AssetAPI = nullptr;
-
-	virtual bool CanBuildTool(const FToolBuilderState& SceneState) const override;
-	virtual UInteractiveTool* BuildTool(const FToolBuilderState& SceneState) const override;
+	virtual USingleSelectionMeshEditingTool* CreateNewTool(const FToolBuilderState& SceneState) const override;
 };
 
 UENUM()
@@ -78,7 +76,7 @@ public:
 
 /** Tool for inserting group edges into polygons of the mesh. */
 UCLASS()
-class MESHMODELINGTOOLS_API UGroupEdgeInsertionTool : public USingleSelectionTool, public IHoverBehaviorTarget, public IClickBehaviorTarget
+class MESHMODELINGTOOLS_API UGroupEdgeInsertionTool : public USingleSelectionMeshEditingTool, public IHoverBehaviorTarget, public IClickBehaviorTarget
 {
 	GENERATED_BODY()
 
@@ -101,9 +99,6 @@ public:
 	virtual void Setup() override;
 	virtual void Shutdown(EToolShutdownType ShutdownType) override;
 
-	virtual void SetWorld(UWorld* World) { TargetWorld = World; }
-	virtual void SetAssetAPI(IToolsContextAssetAPI* AssetAPIIn) { AssetAPI = AssetAPIIn; }
-
 	virtual void OnTick(float DeltaTime) override;
 	virtual void Render(IToolsContextRenderAPI* RenderAPI) override;
 
@@ -130,9 +125,6 @@ protected:
 	UPROPERTY()
 	UMeshOpPreviewWithBackgroundCompute* Preview;
 
-	UWorld* TargetWorld;
-	IToolsContextAssetAPI* AssetAPI;
-
 	TSharedPtr<FDynamicMesh3, ESPMode::ThreadSafe> CurrentMesh;
 	TSharedPtr<FGroupTopology, ESPMode::ThreadSafe> CurrentTopology;
 	UE::Geometry::FDynamicMeshAABBTree3 MeshSpatial;
diff --git a/Engine/Plugins/Experimental/MeshModelingToolset/Source/MeshModelingTools/Public/SingleSelectionMeshEditingTool.h b/Engine/Plugins/Experimental/MeshModelingToolset/Source/MeshModelingTools/Public/SingleSelectionMeshEditingTool.h
new file mode 100644
--- /dev/null
+++ b/Engine/Plugins/Experimental/MeshModelingToolset/Source/MeshModelingTools/Public/SingleSelectionMeshEditingTool.h
@@ -0,0 +1,60 @@
+// Copyright Epic Games, Inc. All Rights Reserved.
+
+#pragma once
+
+#include "SingleSelectionTool.h"
+#include "InteractiveToolBuilder.h"
+#include "SingleSelectionMeshEditingTool.generated.h"
+
+class USingleSelectionMeshEditingTool;
+
+/**
+ * USingleSelectionMeshEditingToolBuilder is a base tool builder for single
+ * selection tools that define a common set of ToolTarget interfaces required
+ * for editing meshes.
+ */
+UCLASS(Transient, Abstract)
+class MESHMODELINGTOOLS_API USingleSelectionMeshEditingToolBuilder : public UInteractiveToolBuilder
+{
+	GENERATED_BODY()
+
+public:
+	IToolsContextAssetAPI* AssetAPI = nullptr;
+
+public:
+	/** @return true if a single mesh source can be found in the active selection */
+	virtual bool CanBuildTool(const FToolBuilderState& SceneState) const override;
+
+	/** @return new Tool instance initialized with selected mesh source */
+	virtual UInteractiveTool* BuildTool(const FToolBuilderState& SceneState) const override;
+
+	/** @return new Tool instance. Override this in subclasses to build a different Tool class type */
+	virtual USingleSelectionMeshEditingTool* CreateNewTool(const FToolBuilderState& SceneState) const PURE_VIRTUAL(USingleSelectionMeshEditingToolBuilder::CreateNewTool, return nullptr; );
+
+	/** Called by BuildTool to configure the Tool with the input MeshSource based on the SceneState */
+	virtual void InitializeNewTool(USingleSelectionMeshEditingTool* Tool, const FToolBuilderState& SceneState) const;
+
+	/** @return true if this tool requires an AssetAPI */
+	virtual bool RequiresAssetAPI() const { return false; }
+
+protected:
+	virtual const FToolTargetTypeRequirements& GetTargetRequirements() const override;
+};
+
+
+/**
+ * Single Selection Mesh Editing tool base class.
+ */
+UCLASS()
+class MESHMODELINGTOOLS_API USingleSelectionMeshEditingTool : public USingleSelectionTool
+{
+	GENERATED_BODY()
+
+public:
+	virtual void SetWorld(UWorld* World) { TargetWorld = World; }
+	virtual void SetAssetAPI(IToolsContextAssetAPI* AssetAPIIn) { AssetAPI = AssetAPIIn; }
+
+protected:
+	UWorld* TargetWorld = nullptr;
+	IToolsContextAssetAPI* AssetAPI = nullptr;
+};
