Add UI to control opacity/color on translucent cutting mesh for mesh CSG and mesh trimming tools.

#rnx

diff --git a/Engine/Plugins/Experimental/MeshModelingToolset/Source/MeshModelingTools/Private/CSGMeshesTool.cpp b/Engine/Plugins/Experimental/MeshModelingToolset/Source/MeshModelingTools/Private/CSGMeshesTool.cpp
--- a/Engine/Plugins/Experimental/MeshModelingToolset/Source/MeshModelingTools/Private/CSGMeshesTool.cpp
+++ b/Engine/Plugins/Experimental/MeshModelingToolset/Source/MeshModelingTools/Private/CSGMeshesTool.cpp
@@ -39,6 +39,14 @@ void UCSGMeshesTool::SetupProperties()
 		{
 			UpdatePreviewsVisibility();
 		});
+		TrimProperties->WatchProperty(TrimProperties->ColorOfTrimmingMesh, [this](FLinearColor)
+		{
+			UpdatePreviewsMaterial();
+		});
+		TrimProperties->WatchProperty(TrimProperties->OpacityOfTrimmingMesh, [this](float)
+		{
+			UpdatePreviewsMaterial();
+		});
 
 		SetToolDisplayName(LOCTEXT("TrimMeshesToolName", "Trim"));
 		GetToolManager()->DisplayMessage(
@@ -60,6 +68,14 @@ void UCSGMeshesTool::SetupProperties()
 		{
 			UpdatePreviewsVisibility();
 		});
+		CSGProperties->WatchProperty(CSGProperties->ColorOfSubtractedMesh, [this](FLinearColor)
+		{
+			UpdatePreviewsMaterial();
+		});
+		CSGProperties->WatchProperty(CSGProperties->OpacityOfSubtractedMesh, [this](float)
+		{
+			UpdatePreviewsMaterial();
+		});
 
 		SetToolDisplayName(LOCTEXT("CSGMeshesToolName", "Boolean"));
 		GetToolManager()->DisplayMessage(
@@ -68,6 +84,28 @@ void UCSGMeshesTool::SetupProperties()
 	}
 }
 
+void UCSGMeshesTool::UpdatePreviewsMaterial()
+{
+	if (!PreviewsGhostMaterial)
+	{
+		PreviewsGhostMaterial = ToolSetupUtil::GetSimpleCustomMaterial(GetToolManager(), FLinearColor::Black, .2);
+	}
+	FLinearColor Color;
+	float Opacity;
+	if (bTrimMode)
+	{
+		Color = TrimProperties->ColorOfTrimmingMesh;
+		Opacity = TrimProperties->OpacityOfTrimmingMesh;
+	}
+	else
+	{
+		Color = CSGProperties->ColorOfSubtractedMesh;
+		Opacity = CSGProperties->OpacityOfSubtractedMesh;
+	}
+	PreviewsGhostMaterial->SetVectorParameterValue(TEXT("Color"), Color);
+	PreviewsGhostMaterial->SetScalarParameterValue(TEXT("Opacity"), Opacity);
+}
+
 void UCSGMeshesTool::UpdatePreviewsVisibility()
 {
 	int32 ShowPreviewIdx = -1;
@@ -174,6 +212,7 @@ void UCSGMeshesTool::ConvertInputsAndSetPreviewMaterials(bool bSetPreviewMesh)
 		}
 	}
 
+	UpdatePreviewsMaterial();
 	for (int ComponentIdx = 0; ComponentIdx < Targets.Num(); ComponentIdx++)
 	{
 		OriginalDynamicMeshes[ComponentIdx] = MakeShared<FDynamicMesh3, ESPMode::ThreadSafe>();
@@ -192,7 +231,8 @@ void UCSGMeshesTool::ConvertInputsAndSetPreviewMaterials(bool bSetPreviewMesh)
 		UPreviewMesh* OriginalMeshPreview = OriginalMeshPreviews.Add_GetRef(NewObject<UPreviewMesh>());
 		OriginalMeshPreview->CreateInWorld(TargetWorld, TargetComponentInterface(ComponentIdx)->GetWorldTransform());
 		OriginalMeshPreview->UpdatePreview(OriginalDynamicMeshes[ComponentIdx].Get());
-		OriginalMeshPreview->SetMaterial(0, ToolSetupUtil::GetSimpleCustomMaterial(GetToolManager(), FLinearColor::White, 0.05));
+
+		OriginalMeshPreview->SetMaterial(0, PreviewsGhostMaterial);
 		OriginalMeshPreview->SetVisible(false);
 		TransformProxies[ComponentIdx]->AddComponent(OriginalMeshPreview->GetRootComponent());
 	}
diff --git a/Engine/Plugins/Experimental/MeshModelingToolset/Source/MeshModelingTools/Public/CSGMeshesTool.h b/Engine/Plugins/Experimental/MeshModelingToolset/Source/MeshModelingTools/Public/CSGMeshesTool.h
--- a/Engine/Plugins/Experimental/MeshModelingToolset/Source/MeshModelingTools/Public/CSGMeshesTool.h
+++ b/Engine/Plugins/Experimental/MeshModelingToolset/Source/MeshModelingTools/Public/CSGMeshesTool.h
@@ -25,24 +25,34 @@ class MESHMODELINGTOOLS_API UCSGMeshesToolProperties : public UInteractiveToolPr
 	GENERATED_BODY()
 public:
 	/** The type of operation */
-	UPROPERTY(EditAnywhere, Category = Options)
+	UPROPERTY(EditAnywhere, Category = Operation)
 	ECSGOperation Operation = ECSGOperation::DifferenceAB;
 
 	/** Show boundary edges created by the CSG operation -- often due to numerical error */
-	UPROPERTY(EditAnywhere, Category = Options)
+	UPROPERTY(EditAnywhere, Category = Operation)
 	bool bShowNewBoundaryEdges = true;
 
 	/** Automatically attempt to fill any holes left by CSG (e.g. due to numerical errors) */
-	UPROPERTY(EditAnywhere, Category = Options)
+	UPROPERTY(EditAnywhere, Category = Operation)
 	bool bAttemptFixHoles = false;
 
 	/** Try to collapse extra edges created by the Boolean operation */
-	UPROPERTY(EditAnywhere, Category = Options)
+	UPROPERTY(EditAnywhere, Category = Operation)
 	bool bCollapseExtraEdges = true;
 
 	/** Whether to show a translucent version of the subtracted mesh, to help visualize what is being removed */
-	UPROPERTY(EditAnywhere, Category = Options, meta = (EditCondition = "Operation == ECSGOperation::DifferenceAB || Operation == ECSGOperation::DifferenceBA", EditConditionHides))
+	UPROPERTY(EditAnywhere, Category = Display, meta = (EditCondition = "Operation == ECSGOperation::DifferenceAB || Operation == ECSGOperation::DifferenceBA", EditConditionHides))
 	bool bShowSubtractedMesh = true;
+	
+	/** Opacity of translucent version of the selected mesh */
+	UPROPERTY(EditAnywhere, Category = Display, meta = (ClampMin = "0", ClampMax = "1", 
+		EditCondition = "bShowSubtractedMesh && Operation == ECSGOperation::DifferenceAB || bShowSubtractedMesh && Operation == ECSGOperation::DifferenceBA", EditConditionHides))
+	float OpacityOfSubtractedMesh = .2;
+	
+	/** Color of translucent version of the subtracted mesh */
+	UPROPERTY(EditAnywhere, Category = Display, meta = (
+		EditCondition = "bShowSubtractedMesh && Operation == ECSGOperation::DifferenceAB || bShowSubtractedMesh && Operation == ECSGOperation::DifferenceBA", EditConditionHides), AdvancedDisplay)
+	FLinearColor ColorOfSubtractedMesh = FLinearColor::Black;
 
 	/** If true, only the first mesh will keep its materials assignments; all other triangles will be assigned material 0 */
 	UPROPERTY(EditAnywhere, Category = Materials)
@@ -59,16 +69,25 @@ class MESHMODELINGTOOLS_API UTrimMeshesToolProperties : public UInteractiveToolP
 	GENERATED_BODY()
 public:
 	/** Which object to trim */
-	UPROPERTY(EditAnywhere, Category = Options)
+	UPROPERTY(EditAnywhere, Category = Operation)
 	ETrimOperation WhichMesh = ETrimOperation::TrimA;
 
 	/** Whether to remove the surface inside or outside of the trimming geometry */
-	UPROPERTY(EditAnywhere, Category = Options)
+	UPROPERTY(EditAnywhere, Category = Operation)
 	ETrimSide TrimSide = ETrimSide::RemoveInside;
 
 	/** Whether to show a translucent version of the trimming mesh, to help visualize what is being cut */
-	UPROPERTY(EditAnywhere, Category = Options)
+	UPROPERTY(EditAnywhere, Category = Display)
 	bool bShowTrimmingMesh = true;
+
+	/** Opacity of translucent version of the trimming mesh */
+	UPROPERTY(EditAnywhere, Category = Display, meta = (EditCondition = "bShowTrimmingMesh", ClampMin = "0", ClampMax = "1"))
+	float OpacityOfTrimmingMesh = .2;
+
+	/** Color of translucent version of the trimming mesh */
+	UPROPERTY(EditAnywhere, Category = Display, meta = (EditCondition = "bShowTrimmingMesh"), AdvancedDisplay)
+	FLinearColor ColorOfTrimmingMesh = FLinearColor::Black;
+
 };
 
 
@@ -120,6 +139,10 @@ protected:
 	UPROPERTY()
 	TArray<UPreviewMesh*> OriginalMeshPreviews;
 
+	// Material used to show the otherwise-invisible cutting/trimming mesh
+	UPROPERTY()
+	UMaterialInstanceDynamic* PreviewsGhostMaterial;
+
 	UPROPERTY()
 	ULineSetComponent* DrawnLineSet;
 
@@ -132,6 +155,9 @@ protected:
 
 	// Update visibility of ghostly preview meshes (used to show trimming or subtracting surface)
 	void UpdatePreviewsVisibility();
+
+	// update the material of ghostly preview meshes (used to show trimming or subtracting surface)
+	void UpdatePreviewsMaterial();
 };
 
 
