MeshModelingToolset: add ToolSelectionUtil::SetNewActorSelection() variant that can select multiple objects. Const fixes.
#rnx

diff --git a/Engine/Plugins/Experimental/MeshModelingToolset/Source/ModelingComponents/Private/Selection/ToolSelectionUtil.cpp b/Engine/Plugins/Experimental/MeshModelingToolset/Source/ModelingComponents/Private/Selection/ToolSelectionUtil.cpp
--- a/Engine/Plugins/Experimental/MeshModelingToolset/Source/ModelingComponents/Private/Selection/ToolSelectionUtil.cpp
+++ b/Engine/Plugins/Experimental/MeshModelingToolset/Source/ModelingComponents/Private/Selection/ToolSelectionUtil.cpp
@@ -12,4 +12,16 @@ void ToolSelectionUtil::SetNewActorSelection(UInteractiveToolManager* ToolManage
 	NewSelection.ModificationType = ESelectedObjectsModificationType::Replace;
 	NewSelection.Actors.Add(Actor);
 	ToolManager->RequestSelectionChange(NewSelection);
+}
+
+
+void ToolSelectionUtil::SetNewActorSelection(UInteractiveToolManager* ToolManager, const TArray<AActor*>& Actors)
+{
+	FSelectedOjectsChangeList NewSelection;
+	NewSelection.ModificationType = ESelectedObjectsModificationType::Replace;
+	for (AActor* Actor : Actors)
+	{
+		NewSelection.Actors.Add(Actor);
+	}
+	ToolManager->RequestSelectionChange(NewSelection);
 }
\ No newline at end of file
diff --git a/Engine/Plugins/Experimental/MeshModelingToolset/Source/ModelingComponents/Public/Selection/ToolSelectionUtil.h b/Engine/Plugins/Experimental/MeshModelingToolset/Source/ModelingComponents/Public/Selection/ToolSelectionUtil.h
--- a/Engine/Plugins/Experimental/MeshModelingToolset/Source/ModelingComponents/Public/Selection/ToolSelectionUtil.h
+++ b/Engine/Plugins/Experimental/MeshModelingToolset/Source/ModelingComponents/Public/Selection/ToolSelectionUtil.h
@@ -17,4 +17,9 @@ namespace ToolSelectionUtil
 	 * Change the active selection to the given Actor, via given ToolManager. Replaces existing selection.
 	 */
 	MODELINGCOMPONENTS_API void SetNewActorSelection(UInteractiveToolManager* ToolManager, AActor* Actor);
+
+	/**
+	 * Change the active selection to the given Actors, via given ToolManager. Replaces existing selection.
+	 */
+	MODELINGCOMPONENTS_API void SetNewActorSelection(UInteractiveToolManager* ToolManager, const TArray<AActor*>& Actors);
 }
\ No newline at end of file
diff --git a/Engine/Plugins/Experimental/MeshModelingToolset/Source/ModelingOperators/Public/CompositionOps/BooleanMeshesOp.h b/Engine/Plugins/Experimental/MeshModelingToolset/Source/ModelingOperators/Public/CompositionOps/BooleanMeshesOp.h
--- a/Engine/Plugins/Experimental/MeshModelingToolset/Source/ModelingOperators/Public/CompositionOps/BooleanMeshesOp.h
+++ b/Engine/Plugins/Experimental/MeshModelingToolset/Source/ModelingOperators/Public/CompositionOps/BooleanMeshesOp.h
@@ -79,7 +79,7 @@ public:
 	virtual void CalculateResult(FProgressCancel* Progress) override;
 
 	// IDs of any newly-created boundary edges in the result mesh
-	TArray<int> GetCreatedBoundaryEdges() const
+	const TArray<int>& GetCreatedBoundaryEdges() const
 	{
 		return CreatedBoundaryEdges;
 	}
diff --git a/Engine/Plugins/Experimental/MeshModelingToolset/Source/ModelingOperators/Public/ModelingOperators.h b/Engine/Plugins/Experimental/MeshModelingToolset/Source/ModelingOperators/Public/ModelingOperators.h
--- a/Engine/Plugins/Experimental/MeshModelingToolset/Source/ModelingOperators/Public/ModelingOperators.h
+++ b/Engine/Plugins/Experimental/MeshModelingToolset/Source/ModelingOperators/Public/ModelingOperators.h
@@ -53,7 +53,7 @@ public:
 	/**
 	 * @return the transform applied to the mesh produced by CalculateResult()
 	 */
-	const FTransform3d& GetResultTransform()
+	const FTransform3d& GetResultTransform() const
 	{
 		return ResultTransform;
 	}
