ModelingTools: add CutMeshWithMesh Tool
#rnx

diff --git a/Engine/Plugins/Experimental/MeshModelingToolset/Source/MeshModelingTools/Private/CutMeshWithMeshTool.cpp b/Engine/Plugins/Experimental/MeshModelingToolset/Source/MeshModelingTools/Private/CutMeshWithMeshTool.cpp
new file mode 100644
--- /dev/null
+++ b/Engine/Plugins/Experimental/MeshModelingToolset/Source/MeshModelingTools/Private/CutMeshWithMeshTool.cpp
@@ -0,0 +1,388 @@
+// Copyright Epic Games, Inc. All Rights Reserved.
+
+#include "CutMeshWithMeshTool.h"
+#include "CompositionOps/BooleanMeshesOp.h"
+#include "ToolSetupUtil.h"
+#include "Selection/ToolSelectionUtil.h"
+#include "AssetGenerationUtil.h"
+#include "DynamicMesh3.h"
+#include "MeshTransforms.h"
+#include "MeshDescriptionToDynamicMesh.h"
+#include "DynamicMeshToMeshDescription.h"
+#include "Async/Async.h"
+
+namespace
+{
+	// probably should be something defined for the whole tool framework...
+#if WITH_EDITOR
+	static EAsyncExecution CutMeshWithMeshToolAsyncExecTarget = EAsyncExecution::LargeThreadPool;
+#else
+	static EAsyncExecution CutMeshWithMeshToolAsyncExecTarget = EAsyncExecution::ThreadPool;
+#endif
+}
+
+
+#define LOCTEXT_NAMESPACE "UCutMeshWithMeshTool"
+
+void UCutMeshWithMeshTool::SetupProperties()
+{
+	Super::SetupProperties();
+
+	CutProperties = NewObject<UCutMeshWithMeshToolProperties>(this);
+	CutProperties->RestoreProperties(this);
+	AddToolPropertySource(CutProperties);
+
+	SetToolDisplayName(LOCTEXT("CutMeshWithMeshToolName", "Cut With Mesh"));
+	GetToolManager()->DisplayMessage(
+		LOCTEXT("OnStartTool", "Cut first mesh with second. Use the transform gizmos to tweak the positions of the input objects (can help to resolve errors/failures)"),
+		EToolMessageLevel::UserNotification);
+
+
+	// create intersection preview mesh object
+	IntersectPreviewMesh = NewObject<UPreviewMesh>(this);
+	IntersectPreviewMesh->CreateInWorld(TargetWorld, FTransform::Identity);
+	IntersectPreviewMesh->SetVisible(true);
+	//IntersectPreviewMesh->SetMaterial(MaterialProperties->Material.Get());
+	//IntersectPreviewMesh->SetMaterial(ToolSetupUtil::GetDefaultMaterial(nullptr));
+	IntersectPreviewMesh->SetMaterial(ToolSetupUtil::GetDefaultBrushVolumeMaterial(GetToolManager()));
+}
+
+void UCutMeshWithMeshTool::SaveProperties()
+{
+	Super::SaveProperties();
+	CutProperties->SaveProperties(this);
+
+	IntersectPreviewMesh->Disconnect();
+}
+
+
+void UCutMeshWithMeshTool::ConvertInputsAndSetPreviewMaterials(bool bSetPreviewMesh)
+{
+	// disable output options
+	// (this property set is not registered yet in SetupProperties() above)
+	SetToolPropertySourceEnabled(HandleSourcesProperties, false);
+
+
+	FComponentMaterialSet AllMaterialSet;
+	TMap<UMaterialInterface*, int> KnownMaterials;
+	TArray<TArray<int>> MaterialRemap; MaterialRemap.SetNum(ComponentTargets.Num());
+
+	if (!CutProperties->bOnlyUseFirstMeshMaterials)
+	{
+		for (int ComponentIdx = 0; ComponentIdx < ComponentTargets.Num(); ComponentIdx++)
+		{
+			FComponentMaterialSet ComponentMaterialSet;
+			ComponentTargets[ComponentIdx]->GetMaterialSet(ComponentMaterialSet);
+			for (UMaterialInterface* Mat : ComponentMaterialSet.Materials)
+			{
+				int* FoundMatIdx = KnownMaterials.Find(Mat);
+				int MatIdx;
+				if (FoundMatIdx)
+				{
+					MatIdx = *FoundMatIdx;
+				}
+				else
+				{
+					MatIdx = AllMaterialSet.Materials.Add(Mat);
+					KnownMaterials.Add(Mat, MatIdx);
+				}
+				MaterialRemap[ComponentIdx].Add(MatIdx);
+			}
+		}
+	}
+	else
+	{
+		ComponentTargets[0]->GetMaterialSet(AllMaterialSet);
+		for (int MatIdx = 0; MatIdx < AllMaterialSet.Materials.Num(); MatIdx++)
+		{
+			MaterialRemap[0].Add(MatIdx);
+		}
+		for (int ComponentIdx = 1; ComponentIdx < ComponentTargets.Num(); ComponentIdx++)
+		{
+			MaterialRemap[ComponentIdx].Init(0, ComponentTargets[ComponentIdx]->GetNumMaterials());
+		}
+	}
+
+	for (int ComponentIdx = 0; ComponentIdx < ComponentTargets.Num(); ComponentIdx++)
+	{
+		TSharedPtr<FDynamicMesh3, ESPMode::ThreadSafe> Mesh = MakeShared<FDynamicMesh3, ESPMode::ThreadSafe>();
+		FMeshDescriptionToDynamicMesh Converter;
+		Converter.Convert(ComponentTargets[ComponentIdx]->GetMesh(), *Mesh);
+
+		// ensure materials and attributes are always enabled
+		Mesh->EnableAttributes();
+		Mesh->Attributes()->EnableMaterialID();
+		FDynamicMeshMaterialAttribute* MaterialIDs = Mesh->Attributes()->GetMaterialID();
+		for (int TID : Mesh->TriangleIndicesItr())
+		{
+			MaterialIDs->SetValue(TID, MaterialRemap[ComponentIdx][MaterialIDs->GetValue(TID)]);
+		}
+
+		if (ComponentIdx == 0)
+		{
+			OriginalTargetMesh = Mesh;
+		}
+		else
+		{
+			OriginalCuttingMesh = Mesh;
+		}
+	}
+	Preview->ConfigureMaterials(AllMaterialSet.Materials, ToolSetupUtil::GetDefaultWorkingMaterial(GetToolManager()));
+
+	// check if we have the same mesh on both inputs
+	if (ComponentTargets[0]->HasSameSourceData(*ComponentTargets[1]))
+	{
+		GetToolManager()->DisplayMessage(
+			LOCTEXT("SameSourceError", "WARNING: Target Mesh has same Asset as Cutting Mesh, both inputs will be affected"),
+			EToolMessageLevel::UserWarning);
+	}
+}
+
+
+
+
+class FCutMeshWithMeshOp : public FDynamicMeshOperator
+{
+public:
+	virtual ~FCutMeshWithMeshOp() {}
+
+	TSharedPtr<const FDynamicMesh3, ESPMode::ThreadSafe> TargetMesh;
+	FTransform TargetMeshTransform;
+	TSharedPtr<const FDynamicMesh3, ESPMode::ThreadSafe> CuttingMesh;
+	FTransform CuttingMeshTransform;
+
+	bool bAttemptToFixHoles = true;
+	bool bCollapseExtraEdges = true;
+
+	virtual void CalculateResult(FProgressCancel* Progress)
+	{
+		TUniquePtr<FBooleanMeshesOp> SubtractOp = MakeUnique<FBooleanMeshesOp>();
+		SubtractOp->CSGOperation = ECSGOperation::DifferenceAB;
+		SubtractOp->bAttemptFixHoles = bAttemptToFixHoles;
+		SubtractOp->bTryCollapseExtraEdges = bCollapseExtraEdges;
+		SubtractOp->Meshes.Add(TargetMesh);
+		SubtractOp->Transforms.Add(TargetMeshTransform);
+		SubtractOp->Meshes.Add(CuttingMesh);
+		SubtractOp->Transforms.Add(CuttingMeshTransform);
+
+		TUniquePtr<FBooleanMeshesOp> IntersectOp = MakeUnique<FBooleanMeshesOp>();
+		IntersectOp->CSGOperation = ECSGOperation::Intersect;
+		IntersectOp->bAttemptFixHoles = bAttemptToFixHoles;
+		IntersectOp->bTryCollapseExtraEdges = bCollapseExtraEdges;
+		IntersectOp->Meshes.Add(TargetMesh);
+		IntersectOp->Transforms.Add(TargetMeshTransform);
+		IntersectOp->Meshes.Add(CuttingMesh);
+		IntersectOp->Transforms.Add(CuttingMeshTransform);
+
+		TFuture<void> SubtractFuture = Async(CutMeshWithMeshToolAsyncExecTarget, [&]()
+		{
+			SubtractOp->CalculateResult(Progress);
+		});
+		TFuture<void> IntersectFuture = Async(CutMeshWithMeshToolAsyncExecTarget, [&]()
+		{
+			IntersectOp->CalculateResult(Progress);
+		});
+
+		SubtractFuture.Wait();
+		IntersectFuture.Wait();
+
+		this->ResultMesh = SubtractOp->ExtractResult();
+		SetResultTransform(SubtractOp->GetResultTransform());
+
+		IntersectMesh = IntersectOp->ExtractResult();
+
+		CreatedSubtractBoundaryEdges = SubtractOp->GetCreatedBoundaryEdges();
+		CreatedIntersectBoundaryEdges = IntersectOp->GetCreatedBoundaryEdges();
+	}
+
+	TUniquePtr<FDynamicMesh3> IntersectMesh;
+	TArray<int> CreatedSubtractBoundaryEdges;
+	TArray<int> CreatedIntersectBoundaryEdges;
+};
+
+
+
+
+
+void UCutMeshWithMeshTool::SetPreviewCallbacks()
+{	
+	DrawnLineSet = NewObject<ULineSetComponent>(Preview->PreviewMesh->GetRootComponent());
+	DrawnLineSet->SetupAttachment(Preview->PreviewMesh->GetRootComponent());
+	DrawnLineSet->SetLineMaterial(ToolSetupUtil::GetDefaultLineComponentMaterial(GetToolManager()));
+	DrawnLineSet->RegisterComponent();
+
+	Preview->OnOpCompleted.AddLambda(
+		[this](const FDynamicMeshOperator* Op)
+		{
+			const FCutMeshWithMeshOp* CuttingOp = (const FCutMeshWithMeshOp*)(Op);
+			CreatedSubtractBoundaryEdges = CuttingOp->CreatedSubtractBoundaryEdges;
+			CreatedIntersectBoundaryEdges = CuttingOp->CreatedIntersectBoundaryEdges;
+			IntersectionMesh = *CuttingOp->IntersectMesh;		// cannot steal this here because it is const...
+			IntersectPreviewMesh->UpdatePreview(&IntersectionMesh);
+			IntersectPreviewMesh->SetTransform((FTransform)Op->GetResultTransform());
+		}
+	);
+	Preview->OnMeshUpdated.AddLambda(
+		[this](const UMeshOpPreviewWithBackgroundCompute*)
+		{
+			GetToolManager()->PostInvalidation();
+			UpdateVisualization();
+		}
+	);
+}
+
+
+void UCutMeshWithMeshTool::UpdateVisualization()
+{
+	FColor BoundaryEdgeColor(240, 15, 15);
+	float BoundaryEdgeThickness = 2.0;
+	float BoundaryEdgeDepthBias = 2.0f;
+
+	DrawnLineSet->Clear();
+	if (CutProperties->bShowNewBoundaryEdges)
+	{
+		const FDynamicMesh3* TargetMesh = Preview->PreviewMesh->GetPreviewDynamicMesh();
+		FVector3d A, B;
+		for (int EID : CreatedSubtractBoundaryEdges)
+		{
+			TargetMesh->GetEdgeV(EID, A, B);
+			DrawnLineSet->AddLine((FVector)A, (FVector)B, BoundaryEdgeColor, BoundaryEdgeThickness, BoundaryEdgeDepthBias);
+		}
+
+		for (int EID : CreatedIntersectBoundaryEdges)
+		{
+			IntersectionMesh.GetEdgeV(EID, A, B);
+			DrawnLineSet->AddLine((FVector)A, (FVector)B, BoundaryEdgeColor, BoundaryEdgeThickness, BoundaryEdgeDepthBias);
+		}
+	}
+}
+
+
+
+TUniquePtr<FDynamicMeshOperator> UCutMeshWithMeshTool::MakeNewOperator()
+{
+	TUniquePtr<FCutMeshWithMeshOp> CuttingOp = MakeUnique<FCutMeshWithMeshOp>();
+	
+	CuttingOp->TargetMesh = OriginalTargetMesh;
+	CuttingOp->TargetMeshTransform = TransformProxies[0]->GetTransform();
+	CuttingOp->TargetMeshTransform.MultiplyScale3D(TransformInitialScales[0]);
+
+	CuttingOp->CuttingMesh = OriginalCuttingMesh;
+	CuttingOp->CuttingMeshTransform = TransformProxies[1]->GetTransform();
+	CuttingOp->CuttingMeshTransform.MultiplyScale3D(TransformInitialScales[1]);
+
+	CuttingOp->bAttemptToFixHoles = CutProperties->bAttemptFixHoles;
+	CuttingOp->bCollapseExtraEdges = CutProperties->bCollapseExtraEdges;
+
+	return CuttingOp;
+}
+
+
+
+void UCutMeshWithMeshTool::OnPropertyModified(UObject* PropertySet, FProperty* Property)
+{
+	if (Property && (Property->GetFName() == GET_MEMBER_NAME_CHECKED(UCutMeshWithMeshToolProperties, bOnlyUseFirstMeshMaterials)))
+	{
+		if (!AreAllTargetsValid())
+		{
+			GetToolManager()->DisplayMessage(LOCTEXT("InvalidTargets", "Target meshes are no longer valid"), EToolMessageLevel::UserWarning);
+			return;
+		}
+		ConvertInputsAndSetPreviewMaterials(false);
+		Preview->InvalidateResult();
+	}
+	else if (Property && (Property->GetFName() == GET_MEMBER_NAME_CHECKED(UCutMeshWithMeshToolProperties, bShowNewBoundaryEdges)))
+	{
+		GetToolManager()->PostInvalidation();
+		UpdateVisualization();
+	}
+	else
+	{
+		Super::OnPropertyModified(PropertySet, Property);
+	}
+}
+
+
+FString UCutMeshWithMeshTool::GetCreatedAssetName() const
+{
+	return TEXT("Boolean");
+}
+
+
+FText UCutMeshWithMeshTool::GetActionName() const
+{
+	return LOCTEXT("CutMeshWithMeshActionName", "Boolean Meshes");
+}
+
+
+void UCutMeshWithMeshTool::Shutdown(EToolShutdownType ShutdownType)
+{
+	SaveProperties();
+	HandleSourcesProperties->SaveProperties(this);
+	TransformProperties->SaveProperties(this);
+
+	FDynamicMeshOpResult Result = Preview->Shutdown();
+	// Restore (unhide) the source meshes
+	for (auto& ComponentTarget : ComponentTargets)
+	{
+		ComponentTarget->SetOwnerVisibility(true);
+	}
+
+	if (ShutdownType == EToolShutdownType::Accept)
+	{
+		GetToolManager()->BeginUndoTransaction(GetActionName());
+
+		TArray<AActor*> SelectActors;
+
+		FComponentMaterialSet MaterialSet;
+		MaterialSet.Materials = GetOutputMaterials();
+
+		// update subtract asset
+		TUniquePtr<FPrimitiveComponentTarget>& UpdateTarget = ComponentTargets[0];
+		FTransform3d TargetToWorld = (FTransform3d)UpdateTarget->GetWorldTransform();
+		FTransform3d WorldToTarget = TargetToWorld.Inverse();
+		{
+			if (Result.Mesh->TriangleCount() > 0)
+			{
+				MeshTransforms::ApplyTransform(*Result.Mesh, Result.Transform);
+				MeshTransforms::ApplyTransform(*Result.Mesh, WorldToTarget);
+				UpdateTarget->CommitMesh([&](const FPrimitiveComponentTarget::FCommitParams& CommitParams)
+				{
+					FDynamicMeshToMeshDescription Converter;
+					Converter.Convert(Result.Mesh.Get(), *CommitParams.MeshDescription);
+				});
+				UpdateTarget->CommitMaterialSetUpdate(MaterialSet, true);
+			}
+		}
+		SelectActors.Add(UpdateTarget->GetOwnerActor());
+
+		// create intersection asset
+		if ( IntersectionMesh.TriangleCount() > 0)
+		{
+			MeshTransforms::ApplyTransform(IntersectionMesh, Result.Transform);
+			MeshTransforms::ApplyTransform(IntersectionMesh, WorldToTarget);
+			FTransform3d NewTransform = TargetToWorld;
+
+			FString CurName = AssetGenerationUtil::GetComponentAssetBaseName(ComponentTargets[0]->GetOwnerComponent());
+			FString UseBaseName = FString::Printf(TEXT("%s_%s"), *CurName, TEXT("CutPart") );
+
+			TArray<UMaterialInterface*> Materials = GetOutputMaterials();
+			AActor* NewActor = AssetGenerationUtil::GenerateStaticMeshActor( AssetAPI, TargetWorld,
+				&IntersectionMesh, NewTransform, UseBaseName, Materials);
+			if (NewActor != nullptr)
+			{
+				SelectActors.Add(NewActor);
+			}
+		}
+
+		ToolSelectionUtil::SetNewActorSelection(GetToolManager(), SelectActors);
+		GetToolManager()->EndUndoTransaction();
+	}
+
+
+	UInteractiveGizmoManager* GizmoManager = GetToolManager()->GetPairedGizmoManager();
+	GizmoManager->DestroyAllGizmosByOwner(this);
+}
+
+
+#undef LOCTEXT_NAMESPACE
diff --git a/Engine/Plugins/Experimental/MeshModelingToolset/Source/MeshModelingTools/Public/CutMeshWithMeshTool.h b/Engine/Plugins/Experimental/MeshModelingToolset/Source/MeshModelingTools/Public/CutMeshWithMeshTool.h
new file mode 100644
--- /dev/null
+++ b/Engine/Plugins/Experimental/MeshModelingToolset/Source/MeshModelingTools/Public/CutMeshWithMeshTool.h
@@ -0,0 +1,125 @@
+// Copyright Epic Games, Inc. All Rights Reserved.
+
+#pragma once
+
+#include "CoreMinimal.h"
+#include "UObject/NoExportTypes.h"
+#include "InteractiveToolBuilder.h"
+#include "Drawing/LineSetComponent.h"
+#include "BaseTools/BaseCreateFromSelectedTool.h"
+
+#include "CutMeshWithMeshTool.generated.h"
+
+// predeclarations
+class UPreviewMesh;
+class FDynamicMesh3;
+
+
+/**
+ * Standard properties of the CutMeshWithMesh operation
+ */
+UCLASS()
+class MESHMODELINGTOOLS_API UCutMeshWithMeshToolProperties : public UInteractiveToolPropertySet
+{
+	GENERATED_BODY()
+public:
+	/** Show boundary edges created by the Boolean operations -- often due to numerical error */
+	UPROPERTY(EditAnywhere, Category = Options)
+	bool bShowNewBoundaryEdges = true;
+
+	/** Automatically attempt to fill any holes left by Booleans (e.g. due to numerical errors) */
+	UPROPERTY(EditAnywhere, Category = Options)
+	bool bAttemptFixHoles = false;
+
+	/** Try to collapse extra edges created by the Boolean operation */
+	UPROPERTY(EditAnywhere, Category = Options)
+	bool bCollapseExtraEdges = true;
+
+	/** If true, only the first mesh will keep its materials assignments; all other triangles will be assigned material 0 */
+	UPROPERTY(EditAnywhere, Category = Materials)
+	bool bOnlyUseFirstMeshMaterials = false;
+};
+
+
+/**
+ * UCutMeshWithMeshTool cuts an input mesh into two pieces based on a second input mesh.
+ * Essentially this just both a Boolean Subtract and a Boolean Intersection. However
+ * doing those as two separate operations involves quite a few steps, so this Tool
+ * does it in a single step and with some improved efficiency.
+ */
+UCLASS()
+class MESHMODELINGTOOLS_API UCutMeshWithMeshTool : public UBaseCreateFromSelectedTool
+{
+	GENERATED_BODY()
+
+public:
+
+	UCutMeshWithMeshTool() {}
+
+protected:
+
+	virtual void OnPropertyModified(UObject* PropertySet, FProperty* Property) override;
+
+	virtual void ConvertInputsAndSetPreviewMaterials(bool bSetPreviewMesh = true) override;
+
+	virtual void SetupProperties() override;
+	virtual void SaveProperties() override;
+	virtual void SetPreviewCallbacks() override;
+
+	virtual FString GetCreatedAssetName() const;
+	virtual FText GetActionName() const;
+
+	// IDynamicMeshOperatorFactory API
+	virtual TUniquePtr<FDynamicMeshOperator> MakeNewOperator() override;
+
+	virtual void Shutdown(EToolShutdownType ShutdownType) override;
+
+protected:
+	void UpdateVisualization();
+
+	UPROPERTY()
+	UCutMeshWithMeshToolProperties* CutProperties;
+
+	UPROPERTY()
+	UPreviewMesh* IntersectPreviewMesh;
+
+	TSharedPtr<FDynamicMesh3, ESPMode::ThreadSafe> OriginalTargetMesh;
+	TSharedPtr<FDynamicMesh3, ESPMode::ThreadSafe> OriginalCuttingMesh;
+
+	UPROPERTY()
+	ULineSetComponent* DrawnLineSet;
+
+	// for visualization of any errors in the currently-previewed CSG operation
+	TArray<int> CreatedSubtractBoundaryEdges;
+	TArray<int> CreatedIntersectBoundaryEdges;
+
+	FDynamicMesh3 IntersectionMesh;
+};
+
+
+
+
+UCLASS()
+class MESHMODELINGTOOLS_API UCutMeshWithMeshToolBuilder : public UBaseCreateFromSelectedToolBuilder
+{
+	GENERATED_BODY()
+
+public:
+	virtual TOptional<int32> MaxComponentsSupported() const override
+	{
+		return TOptional<int32>(2);
+	}
+
+	virtual int32 MinComponentsSupported() const override
+	{
+		return 2;
+	}
+
+	virtual UBaseCreateFromSelectedTool* MakeNewToolInstance(UObject* Outer) const override
+	{
+		return NewObject<UCutMeshWithMeshTool>(Outer);
+	}
+};
+
+
+
diff --git a/Engine/Plugins/Experimental/ModelingToolsEditorMode/Source/ModelingToolsEditorMode/Private/ModelingToolsEditorMode.cpp b/Engine/Plugins/Experimental/ModelingToolsEditorMode/Source/ModelingToolsEditorMode/Private/ModelingToolsEditorMode.cpp
--- a/Engine/Plugins/Experimental/ModelingToolsEditorMode/Source/ModelingToolsEditorMode/Private/ModelingToolsEditorMode.cpp
+++ b/Engine/Plugins/Experimental/ModelingToolsEditorMode/Source/ModelingToolsEditorMode/Private/ModelingToolsEditorMode.cpp
@@ -47,6 +47,7 @@
 #include "MirrorTool.h"
 #include "SelfUnionMeshesTool.h"
 #include "CSGMeshesTool.h"
+#include "CutMeshWithMeshTool.h"
 #include "BspConversionTool.h"
 #include "MeshToVolumeTool.h"
 #include "VolumeToMeshTool.h"
@@ -488,6 +489,10 @@ void UModelingToolsEditorMode::Enter()
 	CSGMeshesToolBuilder->AssetAPI = ModelingModeAssetGenerationAPI.Get();
 	RegisterTool(ToolManagerCommands.BeginMeshBooleanTool, TEXT("BeginMeshBooleanTool"), CSGMeshesToolBuilder);
 
+	auto CutMeshWithMeshToolBuilder = NewObject<UCutMeshWithMeshToolBuilder>();
+	CutMeshWithMeshToolBuilder->AssetAPI = ModelingModeAssetGenerationAPI.Get();
+	RegisterTool(ToolManagerCommands.BeginCutMeshWithMeshTool, TEXT("BeginCutMeshWithMeshTool"), CutMeshWithMeshToolBuilder);
+
 	auto TrimMeshesToolBuilder = NewObject<UCSGMeshesToolBuilder>();
 	TrimMeshesToolBuilder->bTrimMode = true;
 	TrimMeshesToolBuilder->AssetAPI = ModelingModeAssetGenerationAPI.Get();
diff --git a/Engine/Plugins/Experimental/ModelingToolsEditorMode/Source/ModelingToolsEditorMode/Private/ModelingToolsEditorModeStyle.cpp b/Engine/Plugins/Experimental/ModelingToolsEditorMode/Source/ModelingToolsEditorMode/Private/ModelingToolsEditorModeStyle.cpp
--- a/Engine/Plugins/Experimental/ModelingToolsEditorMode/Source/ModelingToolsEditorMode/Private/ModelingToolsEditorModeStyle.cpp
+++ b/Engine/Plugins/Experimental/ModelingToolsEditorMode/Source/ModelingToolsEditorMode/Private/ModelingToolsEditorModeStyle.cpp
@@ -234,6 +234,8 @@ void FModelingToolsEditorModeStyle::Initialize()
 		StyleSet->Set("ModelingToolsManagerCommands.BeginMeshBooleanTool.Small",             new IMAGE_PLUGIN_BRUSH("Icons/ModelingMeshBoolean_x20", Icon20x20));
 		StyleSet->Set("ModelingToolsManagerCommands.BeginMeshTrimTool",				 new IMAGE_BRUSH_SVG("Icons/ModelingTrim", Icon20x20));
 		StyleSet->Set("ModelingToolsManagerCommands.BeginMeshTrimTool.Small",			 new IMAGE_BRUSH_SVG("Icons/ModelingTrim", Icon20x20));
+		StyleSet->Set("ModelingToolsManagerCommands.BeginCutMeshWithMeshTool",               new IMAGE_BRUSH_SVG("Icons/ModelingTrim", Icon20x20));
+		StyleSet->Set("ModelingToolsManagerCommands.BeginCutMeshWithMeshTool.Small",         new IMAGE_BRUSH_SVG("Icons/ModelingTrim", Icon20x20));
 		StyleSet->Set("ModelingToolsManagerCommands.BeginSelfUnionTool",                     new IMAGE_PLUGIN_BRUSH("Icons/ModelingSelfUnion_x40", Icon20x20));
 		StyleSet->Set("ModelingToolsManagerCommands.BeginSelfUnionTool.Small",               new IMAGE_PLUGIN_BRUSH("Icons/ModelingSelfUnion_x20", Icon20x20));
 		StyleSet->Set("ModelingToolsManagerCommands.BeginVoxelSolidifyTool",                 new IMAGE_PLUGIN_BRUSH("Icons/ModelingVoxSolidify_x40", Icon20x20));
diff --git a/Engine/Plugins/Experimental/ModelingToolsEditorMode/Source/ModelingToolsEditorMode/Private/ModelingToolsEditorModeToolkit.cpp b/Engine/Plugins/Experimental/ModelingToolsEditorMode/Source/ModelingToolsEditorMode/Private/ModelingToolsEditorModeToolkit.cpp
--- a/Engine/Plugins/Experimental/ModelingToolsEditorMode/Source/ModelingToolsEditorMode/Private/ModelingToolsEditorModeToolkit.cpp
+++ b/Engine/Plugins/Experimental/ModelingToolsEditorMode/Source/ModelingToolsEditorMode/Private/ModelingToolsEditorModeToolkit.cpp
@@ -543,6 +543,7 @@ void FModelingToolsEditorModeToolkit::BuildToolPalette_Experimental(FName Palett
 		ToolbarBuilder.AddToolBarButton(Commands.BeginGroupEdgeInsertionTool);
 		ToolbarBuilder.AddToolBarButton(Commands.BeginEdgeLoopInsertionTool);
 		ToolbarBuilder.AddToolBarButton(Commands.BeginMeshBooleanTool);
+		ToolbarBuilder.AddToolBarButton(Commands.BeginCutMeshWithMeshTool);
 		ToolbarBuilder.AddToolBarButton(Commands.BeginSubdividePolyTool);
 	}
 	else if (PaletteIndex == AttributesTabName)
diff --git a/Engine/Plugins/Experimental/ModelingToolsEditorMode/Source/ModelingToolsEditorMode/Private/ModelingToolsManagerActions.cpp b/Engine/Plugins/Experimental/ModelingToolsEditorMode/Source/ModelingToolsEditorMode/Private/ModelingToolsManagerActions.cpp
--- a/Engine/Plugins/Experimental/ModelingToolsEditorMode/Source/ModelingToolsEditorMode/Private/ModelingToolsManagerActions.cpp
+++ b/Engine/Plugins/Experimental/ModelingToolsEditorMode/Source/ModelingToolsEditorMode/Private/ModelingToolsManagerActions.cpp
@@ -73,6 +73,7 @@ void FModelingToolsManagerCommands::RegisterCommands()
 	UI_COMMAND(BeginSelfUnionTool, "Merge", "Self-Union Selected Meshes to resolve Self-Intersections", EUserInterfaceActionType::ToggleButton, FInputChord());
 	UI_COMMAND(BeginMeshBooleanTool, "Boolean", "Apply Boolean/CSG operation to selected Mesh pair", EUserInterfaceActionType::ToggleButton, FInputChord());
 	UI_COMMAND(BeginMeshTrimTool, "Trim", "Trim/Cut selected mesh with second mesh", EUserInterfaceActionType::ToggleButton, FInputChord());
+	UI_COMMAND(BeginCutMeshWithMeshTool, "MshCut", "Split selected mesh into parts using second mesh", EUserInterfaceActionType::ToggleButton, FInputChord());
 	UI_COMMAND(BeginBspConversionTool, "BSPConv", "Convert BSP to StaticMesh Asset", EUserInterfaceActionType::ToggleButton, FInputChord());
 	UI_COMMAND(BeginMeshToVolumeTool, "Msh2Vol", "Convert Mesh to Volume", EUserInterfaceActionType::ToggleButton, FInputChord());
 	UI_COMMAND(BeginVolumeToMeshTool, "Vol2Msh", "Convert Volume to new Mesh Asset", EUserInterfaceActionType::ToggleButton, FInputChord());
diff --git a/Engine/Plugins/Experimental/ModelingToolsEditorMode/Source/ModelingToolsEditorMode/Public/ModelingToolsManagerActions.h b/Engine/Plugins/Experimental/ModelingToolsEditorMode/Source/ModelingToolsEditorMode/Public/ModelingToolsManagerActions.h
--- a/Engine/Plugins/Experimental/ModelingToolsEditorMode/Source/ModelingToolsEditorMode/Public/ModelingToolsManagerActions.h
+++ b/Engine/Plugins/Experimental/ModelingToolsEditorMode/Source/ModelingToolsEditorMode/Public/ModelingToolsManagerActions.h
@@ -72,6 +72,7 @@ public:
 	TSharedPtr<FUICommandInfo> BeginSelfUnionTool;
 	TSharedPtr<FUICommandInfo> BeginMeshBooleanTool;
 	TSharedPtr<FUICommandInfo> BeginMeshTrimTool;
+	TSharedPtr<FUICommandInfo> BeginCutMeshWithMeshTool;
 	TSharedPtr<FUICommandInfo> BeginMeshSelectionTool;
 	TSharedPtr<FUICommandInfo> BeginBspConversionTool;
 	TSharedPtr<FUICommandInfo> BeginMeshToVolumeTool;
