ModelingTools: update CutMeshWithMesh for Main API changes
#rnx

diff --git a/Engine/Plugins/Experimental/MeshModelingToolset/Source/MeshModelingTools/Private/CutMeshWithMeshTool.cpp b/Engine/Plugins/Experimental/MeshModelingToolset/Source/MeshModelingTools/Private/CutMeshWithMeshTool.cpp
--- a/Engine/Plugins/Experimental/MeshModelingToolset/Source/MeshModelingTools/Private/CutMeshWithMeshTool.cpp
+++ b/Engine/Plugins/Experimental/MeshModelingToolset/Source/MeshModelingTools/Private/CutMeshWithMeshTool.cpp
@@ -11,6 +11,15 @@
 #include "DynamicMeshToMeshDescription.h"
 #include "Async/Async.h"
 
+#include "TargetInterfaces/MaterialProvider.h"
+#include "TargetInterfaces/MeshDescriptionCommitter.h"
+#include "TargetInterfaces/MeshDescriptionProvider.h"
+#include "TargetInterfaces/PrimitiveComponentBackedTarget.h"
+#include "ToolTargetManager.h"
+
+#include "ExplicitUseGeometryMathTypes.h"		// using UE::Geometry::(math types)
+using namespace UE::Geometry;
+
 namespace
 {
 	// probably should be something defined for the whole tool framework...
@@ -65,14 +74,14 @@ void UCutMeshWithMeshTool::ConvertInputsAndSetPreviewMaterials(bool bSetPreviewM
 
 	FComponentMaterialSet AllMaterialSet;
 	TMap<UMaterialInterface*, int> KnownMaterials;
-	TArray<TArray<int>> MaterialRemap; MaterialRemap.SetNum(ComponentTargets.Num());
+	TArray<TArray<int>> MaterialRemap; MaterialRemap.SetNum(Targets.Num());
 
 	if (!CutProperties->bOnlyUseFirstMeshMaterials)
 	{
-		for (int ComponentIdx = 0; ComponentIdx < ComponentTargets.Num(); ComponentIdx++)
+		for (int ComponentIdx = 0; ComponentIdx < Targets.Num(); ComponentIdx++)
 		{
 			FComponentMaterialSet ComponentMaterialSet;
-			ComponentTargets[ComponentIdx]->GetMaterialSet(ComponentMaterialSet);
+			TargetMaterialInterface(ComponentIdx)->GetMaterialSet(ComponentMaterialSet);
 			for (UMaterialInterface* Mat : ComponentMaterialSet.Materials)
 			{
 				int* FoundMatIdx = KnownMaterials.Find(Mat);
@@ -92,22 +101,22 @@ void UCutMeshWithMeshTool::ConvertInputsAndSetPreviewMaterials(bool bSetPreviewM
 	}
 	else
 	{
-		ComponentTargets[0]->GetMaterialSet(AllMaterialSet);
+		TargetMaterialInterface(0)->GetMaterialSet(AllMaterialSet);
 		for (int MatIdx = 0; MatIdx < AllMaterialSet.Materials.Num(); MatIdx++)
 		{
 			MaterialRemap[0].Add(MatIdx);
 		}
-		for (int ComponentIdx = 1; ComponentIdx < ComponentTargets.Num(); ComponentIdx++)
+		for (int ComponentIdx = 1; ComponentIdx < Targets.Num(); ComponentIdx++)
 		{
-			MaterialRemap[ComponentIdx].Init(0, ComponentTargets[ComponentIdx]->GetNumMaterials());
+			MaterialRemap[ComponentIdx].Init(0, TargetMaterialInterface(ComponentIdx)->GetNumMaterials());
 		}
 	}
 
-	for (int ComponentIdx = 0; ComponentIdx < ComponentTargets.Num(); ComponentIdx++)
+	for (int ComponentIdx = 0; ComponentIdx < Targets.Num(); ComponentIdx++)
 	{
 		TSharedPtr<FDynamicMesh3, ESPMode::ThreadSafe> Mesh = MakeShared<FDynamicMesh3, ESPMode::ThreadSafe>();
 		FMeshDescriptionToDynamicMesh Converter;
-		Converter.Convert(ComponentTargets[ComponentIdx]->GetMesh(), *Mesh);
+		Converter.Convert(TargetMeshProviderInterface(ComponentIdx)->GetMeshDescription(), *Mesh);
 
 		// ensure materials and attributes are always enabled
 		Mesh->EnableAttributes();
@@ -129,13 +138,14 @@ void UCutMeshWithMeshTool::ConvertInputsAndSetPreviewMaterials(bool bSetPreviewM
 	}
 	Preview->ConfigureMaterials(AllMaterialSet.Materials, ToolSetupUtil::GetDefaultWorkingMaterial(GetToolManager()));
 
+	// TODO port this to ToolTargets
 	// check if we have the same mesh on both inputs
-	if (ComponentTargets[0]->HasSameSourceData(*ComponentTargets[1]))
-	{
-		GetToolManager()->DisplayMessage(
-			LOCTEXT("SameSourceError", "WARNING: Target Mesh has same Asset as Cutting Mesh, both inputs will be affected"),
-			EToolMessageLevel::UserWarning);
-	}
+	//if (ComponentTargets[0]->HasSameSourceData(*ComponentTargets[1]))
+	//{
+	//	GetToolManager()->DisplayMessage(
+	//		LOCTEXT("SameSourceError", "WARNING: Target Mesh has same Asset as Cutting Mesh, both inputs will be affected"),
+	//		EToolMessageLevel::UserWarning);
+	//}
 }
 
 
@@ -265,11 +275,15 @@ TUniquePtr<FDynamicMeshOperator> UCutMeshWithMeshTool::MakeNewOperator()
 	
 	CuttingOp->TargetMesh = OriginalTargetMesh;
 	CuttingOp->TargetMeshTransform = TransformProxies[0]->GetTransform();
-	CuttingOp->TargetMeshTransform.MultiplyScale3D(TransformInitialScales[0]);
+
+	// TODO FIX THIS
+	//CuttingOp->TargetMeshTransform.MultiplyScale3D(TransformInitialScales[0]);
 
 	CuttingOp->CuttingMesh = OriginalCuttingMesh;
 	CuttingOp->CuttingMeshTransform = TransformProxies[1]->GetTransform();
-	CuttingOp->CuttingMeshTransform.MultiplyScale3D(TransformInitialScales[1]);
+
+	// TODO FIX THIS
+	//CuttingOp->CuttingMeshTransform.MultiplyScale3D(TransformInitialScales[1]);
 
 	CuttingOp->bAttemptToFixHoles = CutProperties->bAttemptFixHoles;
 	CuttingOp->bCollapseExtraEdges = CutProperties->bCollapseExtraEdges;
@@ -323,9 +337,9 @@ void UCutMeshWithMeshTool::Shutdown(EToolShutdownType ShutdownType)
 
 	FDynamicMeshOpResult Result = Preview->Shutdown();
 	// Restore (unhide) the source meshes
-	for (auto& ComponentTarget : ComponentTargets)
+	for ( int32 ci = 0; ci < Targets.Num(); ++ci)
 	{
-		ComponentTarget->SetOwnerVisibility(true);
+		TargetComponentInterface(ci)->SetOwnerVisibility(true);
 	}
 
 	if (ShutdownType == EToolShutdownType::Accept)
@@ -338,7 +352,7 @@ void UCutMeshWithMeshTool::Shutdown(EToolShutdownType ShutdownType)
 		MaterialSet.Materials = GetOutputMaterials();
 
 		// update subtract asset
-		TUniquePtr<FPrimitiveComponentTarget>& UpdateTarget = ComponentTargets[0];
+		IPrimitiveComponentBackedTarget* UpdateTarget = TargetComponentInterface(0);
 		FTransform3d TargetToWorld = (FTransform3d)UpdateTarget->GetWorldTransform();
 		FTransform3d WorldToTarget = TargetToWorld.Inverse();
 		{
@@ -346,12 +360,12 @@ void UCutMeshWithMeshTool::Shutdown(EToolShutdownType ShutdownType)
 			{
 				MeshTransforms::ApplyTransform(*Result.Mesh, Result.Transform);
 				MeshTransforms::ApplyTransform(*Result.Mesh, WorldToTarget);
-				UpdateTarget->CommitMesh([&](const FPrimitiveComponentTarget::FCommitParams& CommitParams)
+				TargetMeshCommitterInterface(0)->CommitMeshDescription([&](const IMeshDescriptionCommitter::FCommitterParams& CommitParams)
 				{
 					FDynamicMeshToMeshDescription Converter;
-					Converter.Convert(Result.Mesh.Get(), *CommitParams.MeshDescription);
+					Converter.Convert(Result.Mesh.Get(), *CommitParams.MeshDescriptionOut);
 				});
-				UpdateTarget->CommitMaterialSetUpdate(MaterialSet, true);
+				TargetMaterialInterface(0)->CommitMaterialSetUpdate(MaterialSet, true);
 			}
 		}
 		SelectActors.Add(UpdateTarget->GetOwnerActor());
@@ -363,7 +377,7 @@ void UCutMeshWithMeshTool::Shutdown(EToolShutdownType ShutdownType)
 			MeshTransforms::ApplyTransform(IntersectionMesh, WorldToTarget);
 			FTransform3d NewTransform = TargetToWorld;
 
-			FString CurName = AssetGenerationUtil::GetComponentAssetBaseName(ComponentTargets[0]->GetOwnerComponent());
+			FString CurName = AssetGenerationUtil::GetComponentAssetBaseName(UpdateTarget->GetOwnerComponent());
 			FString UseBaseName = FString::Printf(TEXT("%s_%s"), *CurName, TEXT("CutPart") );
 
 			TArray<UMaterialInterface*> Materials = GetOutputMaterials();
diff --git a/Engine/Plugins/Experimental/MeshModelingToolset/Source/MeshModelingTools/Public/CutMeshWithMeshTool.h b/Engine/Plugins/Experimental/MeshModelingToolset/Source/MeshModelingTools/Public/CutMeshWithMeshTool.h
--- a/Engine/Plugins/Experimental/MeshModelingToolset/Source/MeshModelingTools/Public/CutMeshWithMeshTool.h
+++ b/Engine/Plugins/Experimental/MeshModelingToolset/Source/MeshModelingTools/Public/CutMeshWithMeshTool.h
@@ -70,7 +70,7 @@ protected:
 	virtual FText GetActionName() const;
 
 	// IDynamicMeshOperatorFactory API
-	virtual TUniquePtr<FDynamicMeshOperator> MakeNewOperator() override;
+	virtual TUniquePtr<UE::Geometry::FDynamicMeshOperator> MakeNewOperator() override;
 
 	virtual void Shutdown(EToolShutdownType ShutdownType) override;
 
