ModelingTools: add missing namespace on predeclaration
#rnx

diff --git a/Engine/Plugins/Experimental/MeshModelingToolset/Source/MeshModelingTools/Public/CutMeshWithMeshTool.h b/Engine/Plugins/Experimental/MeshModelingToolset/Source/MeshModelingTools/Public/CutMeshWithMeshTool.h
--- a/Engine/Plugins/Experimental/MeshModelingToolset/Source/MeshModelingTools/Public/CutMeshWithMeshTool.h
+++ b/Engine/Plugins/Experimental/MeshModelingToolset/Source/MeshModelingTools/Public/CutMeshWithMeshTool.h
@@ -12,7 +12,7 @@
 
 // predeclarations
 class UPreviewMesh;
-class FDynamicMesh3;
+PREDECLARE_USE_GEOMETRY_CLASS(FDynamicMesh3);
 
 
 /**
