clean unneeded commented-out lines leftover from merge in CutMeshWithMeshTool

#rnx

diff --git a/Engine/Plugins/Experimental/MeshModelingToolset/Source/MeshModelingTools/Private/CutMeshWithMeshTool.cpp b/Engine/Plugins/Experimental/MeshModelingToolset/Source/MeshModelingTools/Private/CutMeshWithMeshTool.cpp
--- a/Engine/Plugins/Experimental/MeshModelingToolset/Source/MeshModelingTools/Private/CutMeshWithMeshTool.cpp
+++ b/Engine/Plugins/Experimental/MeshModelingToolset/Source/MeshModelingTools/Private/CutMeshWithMeshTool.cpp
@@ -276,15 +276,9 @@ TUniquePtr<FDynamicMeshOperator> UCutMeshWithMeshTool::MakeNewOperator()
 	CuttingOp->TargetMesh = OriginalTargetMesh;
 	CuttingOp->TargetMeshTransform = TransformProxies[0]->GetTransform();
 
-	// TODO FIX THIS
-	//CuttingOp->TargetMeshTransform.MultiplyScale3D(TransformInitialScales[0]);
-
 	CuttingOp->CuttingMesh = OriginalCuttingMesh;
 	CuttingOp->CuttingMeshTransform = TransformProxies[1]->GetTransform();
 
-	// TODO FIX THIS
-	//CuttingOp->CuttingMeshTransform.MultiplyScale3D(TransformInitialScales[1]);
-
 	CuttingOp->bAttemptToFixHoles = CutProperties->bAttemptFixHoles;
 	CuttingOp->bCollapseExtraEdges = CutProperties->bCollapseExtraEdges;
 
