ModelingTools: Fix CutMeshWithMeshOp same source asset warning.

#rnx

diff --git a/Engine/Plugins/Experimental/MeshModelingToolset/Source/MeshModelingTools/Private/CutMeshWithMeshTool.cpp b/Engine/Plugins/Experimental/MeshModelingToolset/Source/MeshModelingTools/Private/CutMeshWithMeshTool.cpp
--- a/Engine/Plugins/Experimental/MeshModelingToolset/Source/MeshModelingTools/Private/CutMeshWithMeshTool.cpp
+++ b/Engine/Plugins/Experimental/MeshModelingToolset/Source/MeshModelingTools/Private/CutMeshWithMeshTool.cpp
@@ -15,6 +15,7 @@
 #include "TargetInterfaces/MeshDescriptionCommitter.h"
 #include "TargetInterfaces/MeshDescriptionProvider.h"
 #include "TargetInterfaces/PrimitiveComponentBackedTarget.h"
+#include "TargetInterfaces/AssetBackedTarget.h"
 #include "ToolTargetManager.h"
 
 #include "ExplicitUseGeometryMathTypes.h"		// using UE::Geometry::(math types)
@@ -138,14 +139,13 @@ void UCutMeshWithMeshTool::ConvertInputsAndSetPreviewMaterials(bool bSetPreviewM
 	}
 	Preview->ConfigureMaterials(AllMaterialSet.Materials, ToolSetupUtil::GetDefaultWorkingMaterial(GetToolManager()));
 
-	// TODO port this to ToolTargets
 	// check if we have the same mesh on both inputs
-	//if (ComponentTargets[0]->HasSameSourceData(*ComponentTargets[1]))
-	//{
-	//	GetToolManager()->DisplayMessage(
-	//		LOCTEXT("SameSourceError", "WARNING: Target Mesh has same Asset as Cutting Mesh, both inputs will be affected"),
-	//		EToolMessageLevel::UserWarning);
-	//}
+	if (Cast<IAssetBackedTarget>(Targets[0])->HasSameSourceData(Targets[1]))
+	{
+		GetToolManager()->DisplayMessage(
+			LOCTEXT("SameSourceError", "WARNING: Target Mesh has same Asset as Cutting Mesh, both inputs will be affected"),
+			EToolMessageLevel::UserWarning);
+	}
 }
 
 
diff --git a/Engine/Plugins/Experimental/MeshModelingToolset/Source/ModelingComponents/Private/BaseTools/BaseCreateFromSelectedTool.cpp b/Engine/Plugins/Experimental/MeshModelingToolset/Source/ModelingComponents/Private/BaseTools/BaseCreateFromSelectedTool.cpp
--- a/Engine/Plugins/Experimental/MeshModelingToolset/Source/ModelingComponents/Private/BaseTools/BaseCreateFromSelectedTool.cpp
+++ b/Engine/Plugins/Experimental/MeshModelingToolset/Source/ModelingComponents/Private/BaseTools/BaseCreateFromSelectedTool.cpp
@@ -19,6 +19,7 @@
 #include "TargetInterfaces/MeshDescriptionProvider.h"
 #include "TargetInterfaces/PrimitiveComponentBackedTarget.h"
 #include "TargetInterfaces/MaterialProvider.h"
+#include "TargetInterfaces/AssetBackedTarget.h"
 #include "ToolTargetManager.h"
 
 #define LOCTEXT_NAMESPACE "UBaseCreateFromSelectedTool"
@@ -35,7 +36,8 @@ const FToolTargetTypeRequirements& UBaseCreateFromSelectedToolBuilder::GetTarget
 		UMeshDescriptionCommitter::StaticClass(),
 		UMeshDescriptionProvider::StaticClass(),
 		UPrimitiveComponentBackedTarget::StaticClass(),
-		UMaterialProvider::StaticClass()
+		UMaterialProvider::StaticClass(),
+		UAssetBackedTarget::StaticClass()
 		});
 	return TypeRequirements;
 }
diff --git a/Engine/Source/Runtime/Experimental/InteractiveToolsFramework/Public/TargetInterfaces/AssetBackedTarget.h b/Engine/Source/Runtime/Experimental/InteractiveToolsFramework/Public/TargetInterfaces/AssetBackedTarget.h
--- a/Engine/Source/Runtime/Experimental/InteractiveToolsFramework/Public/TargetInterfaces/AssetBackedTarget.h
+++ b/Engine/Source/Runtime/Experimental/InteractiveToolsFramework/Public/TargetInterfaces/AssetBackedTarget.h
@@ -5,6 +5,7 @@
 #include "CoreMinimal.h"
 #include "Engine/EngineTypes.h"
 #include "UObject/Interface.h"
+#include "ToolTargets/ToolTarget.h"
 
 #include "AssetBackedTarget.generated.h"
 
@@ -23,4 +24,10 @@ public:
 	 * @return the underlying source asset for this Target.
 	 */
 	virtual UObject* GetSourceData() const = 0;
+
+	virtual bool HasSameSourceData(UToolTarget* OtherTarget) const
+	{
+		IAssetBackedTarget* OtherAsset = Cast<IAssetBackedTarget>(OtherTarget);
+		return OtherAsset && GetSourceData() == OtherAsset->GetSourceData();
+	}
 };
\ No newline at end of file
