fix CutMeshWithMesh tool's handling of non-uniform scales on output

#rnx

diff --git a/Engine/Plugins/Experimental/MeshModelingToolset/Source/MeshModelingTools/Private/CutMeshWithMeshTool.cpp b/Engine/Plugins/Experimental/MeshModelingToolset/Source/MeshModelingTools/Private/CutMeshWithMeshTool.cpp
--- a/Engine/Plugins/Experimental/MeshModelingToolset/Source/MeshModelingTools/Private/CutMeshWithMeshTool.cpp
+++ b/Engine/Plugins/Experimental/MeshModelingToolset/Source/MeshModelingTools/Private/CutMeshWithMeshTool.cpp
@@ -348,12 +348,11 @@ void UCutMeshWithMeshTool::Shutdown(EToolShutdownType ShutdownType)
 		// update subtract asset
 		IPrimitiveComponentBackedTarget* UpdateTarget = TargetComponentInterface(0);
 		FTransform3d TargetToWorld = (FTransform3d)UpdateTarget->GetWorldTransform();
-		FTransform3d WorldToTarget = TargetToWorld.Inverse();
 		{
 			if (Result.Mesh->TriangleCount() > 0)
 			{
 				MeshTransforms::ApplyTransform(*Result.Mesh, Result.Transform);
-				MeshTransforms::ApplyTransform(*Result.Mesh, WorldToTarget);
+				MeshTransforms::ApplyTransformInverse(*Result.Mesh, TargetToWorld);
 				TargetMeshCommitterInterface(0)->CommitMeshDescription([&](const IMeshDescriptionCommitter::FCommitterParams& CommitParams)
 				{
 					FDynamicMeshToMeshDescription Converter;
@@ -368,7 +367,7 @@ void UCutMeshWithMeshTool::Shutdown(EToolShutdownType ShutdownType)
 		if ( IntersectionMesh.TriangleCount() > 0)
 		{
 			MeshTransforms::ApplyTransform(IntersectionMesh, Result.Transform);
-			MeshTransforms::ApplyTransform(IntersectionMesh, WorldToTarget);
+			MeshTransforms::ApplyTransformInverse(IntersectionMesh, TargetToWorld);
 			FTransform3d NewTransform = TargetToWorld;
 
 			FString CurName = AssetGenerationUtil::GetComponentAssetBaseName(UpdateTarget->GetOwnerComponent());
