ModelingComponents: fix two different LOCTEXT strings sharing the same key
#rnx

diff --git a/Engine/Plugins/Experimental/MeshModelingToolset/Source/ModelingComponents/Private/Selection/PolygonSelectionMechanic.cpp b/Engine/Plugins/Experimental/MeshModelingToolset/Source/ModelingComponents/Private/Selection/PolygonSelectionMechanic.cpp
--- a/Engine/Plugins/Experimental/MeshModelingToolset/Source/ModelingComponents/Private/Selection/PolygonSelectionMechanic.cpp
+++ b/Engine/Plugins/Experimental/MeshModelingToolset/Source/ModelingComponents/Private/Selection/PolygonSelectionMechanic.cpp
@@ -666,7 +666,7 @@ bool UPolygonSelectionMechanic::EndChangeAndEmitIfModified()
 	if (SelectionTimestamp != ActiveChange->Timestamp)
 	{
 		GetParentTool()->GetToolManager()->EmitObjectChange(this, MoveTemp(ActiveChange),
-			LOCTEXT("SelectionChange", "Selection Change"));
+			LOCTEXT("SelectionChangeMessage", "Selection Change"));
 		return true;
 	}
 	ActiveChange = TUniquePtr<FPolygonSelectionMechanicSelectionChange>();
