Mesh Modeling Toolset:  GenerateLODMeshesTool.  Fix race condition when using MeshDescription and multiple copies of the UE simplifier.

#rnx

diff --git a/Engine/Plugins/Experimental/MeshModelingToolset/Source/MeshModelingToolsEditorOnly/Private/GenerateLODMeshesTool.cpp b/Engine/Plugins/Experimental/MeshModelingToolset/Source/MeshModelingToolsEditorOnly/Private/GenerateLODMeshesTool.cpp
--- a/Engine/Plugins/Experimental/MeshModelingToolset/Source/MeshModelingToolsEditorOnly/Private/GenerateLODMeshesTool.cpp
+++ b/Engine/Plugins/Experimental/MeshModelingToolset/Source/MeshModelingToolsEditorOnly/Private/GenerateLODMeshesTool.cpp
@@ -123,6 +123,9 @@ void UGenerateLODMeshesTool::Setup()
 #endif
 		OriginalMeshDescription = MakeShared<FMeshDescription, ESPMode::ThreadSafe>(*ComponentTarget->GetMesh());
 		EnterProgressFrame(1);
+		// aux-data isn't deep copied - by default it is built during initial evaluation (not thread safe)  
+		// so force aux-data rebuild now before multiple UE4 simplifiers try to use it in parallel.
+		OriginalMeshDescription->BuildIndexers();   
 		OriginalMesh = MakeShared<FDynamicMesh3, ESPMode::ThreadSafe>();
 		FMeshDescriptionToDynamicMesh Converter;
 		Converter.Convert(ComponentTarget->GetMesh(), *OriginalMesh);
