GeometryProcessing and MeshModelingToolset: remove UE4 comment/enum references
#rnx

diff --git a/Engine/Plugins/Experimental/Enterprise/DataprepGeometryOperations/Source/Private/DataprepGeometryOperations.cpp b/Engine/Plugins/Experimental/Enterprise/DataprepGeometryOperations/Source/Private/DataprepGeometryOperations.cpp
--- a/Engine/Plugins/Experimental/Enterprise/DataprepGeometryOperations/Source/Private/DataprepGeometryOperations.cpp
+++ b/Engine/Plugins/Experimental/Enterprise/DataprepGeometryOperations/Source/Private/DataprepGeometryOperations.cpp
@@ -176,7 +176,7 @@ void UDataprepSimplifyMeshOperation::OnExecution_Implementation(const FDataprepC
 		Op.bPreventNormalFlips = true;
 		Op.bPreserveSharpEdges = true;
 		Op.bReproject = false;
-		Op.SimplifierType = ESimplifyType::UE4Standard;
+		Op.SimplifierType = ESimplifyType::UEStandard;
 //		Op.TargetCount = TargetCount;
 //		Op.TargetEdgeLength = TargetEdgeLength;
 		Op.TargetMode = ESimplifyTargetType::Percentage;
diff --git a/Engine/Plugins/Experimental/GeometryProcessing/Source/DynamicMesh/Public/Parameterization/MeshUVPacking.h b/Engine/Plugins/Experimental/GeometryProcessing/Source/DynamicMesh/Public/Parameterization/MeshUVPacking.h
--- a/Engine/Plugins/Experimental/GeometryProcessing/Source/DynamicMesh/Public/Parameterization/MeshUVPacking.h
+++ b/Engine/Plugins/Experimental/GeometryProcessing/Source/DynamicMesh/Public/Parameterization/MeshUVPacking.h
@@ -35,7 +35,7 @@ public:
 
 
 	/**
-	 * Standard UE4 UV layout, similar to that used for Lightmap UVs. 
+	 * Standard UnrealEngine UV layout, similar to that used for Lightmap UVs. 
 	 * All UV islands are packed into standard positive-unit-square.
 	 * Only supports single-pixel border size.
 	 */
diff --git a/Engine/Plugins/Experimental/GeometryProcessing/Source/GeometricObjects/Public/Generators/GridBoxMeshGenerator.h b/Engine/Plugins/Experimental/GeometryProcessing/Source/GeometricObjects/Public/Generators/GridBoxMeshGenerator.h
--- a/Engine/Plugins/Experimental/GeometryProcessing/Source/GeometricObjects/Public/Generators/GridBoxMeshGenerator.h
+++ b/Engine/Plugins/Experimental/GeometryProcessing/Source/GeometricObjects/Public/Generators/GridBoxMeshGenerator.h
@@ -175,7 +175,7 @@ public:
 		{
 			int FaceIdxBase = FaceDimOrder[Dim]*2;
 
-			// UV-specific minor axes + flips; manually set to match default ue4 cube texture arrangement
+			// UV-specific minor axes + flips; manually set to match default UnrealEngine cube texture arrangement
 			int Minor1Flip[3] = { -1, 1, 1 };
 			int Minor2Flip[3] = { -1, -1, 1 };
 
diff --git a/Engine/Plugins/Experimental/GeometryProcessing/Source/GeometricObjects/Public/Parameterization/UVPacking.h b/Engine/Plugins/Experimental/GeometryProcessing/Source/GeometricObjects/Public/Parameterization/UVPacking.h
--- a/Engine/Plugins/Experimental/GeometryProcessing/Source/GeometricObjects/Public/Parameterization/UVPacking.h
+++ b/Engine/Plugins/Experimental/GeometryProcessing/Source/GeometricObjects/Public/Parameterization/UVPacking.h
@@ -46,7 +46,7 @@ public:
 	bool bScaleIslandsByWorldSpaceTexelRatio = false;
 
 	/**
-	 * Standard UE4 UV layout, similar to that used for Lightmap UVs. 
+	 * Standard UE UV layout, similar to that used for Lightmap UVs. 
 	 * All UV islands are packed into standard positive-unit-square.
 	 * Only supports single-pixel border size.
 	 */
diff --git a/Engine/Plugins/Experimental/MeshModelingToolset/Source/MeshModelingTools/MeshModelingTools.Build.cs b/Engine/Plugins/Experimental/MeshModelingToolset/Source/MeshModelingTools/MeshModelingTools.Build.cs
--- a/Engine/Plugins/Experimental/MeshModelingToolset/Source/MeshModelingTools/MeshModelingTools.Build.cs
+++ b/Engine/Plugins/Experimental/MeshModelingToolset/Source/MeshModelingTools/MeshModelingTools.Build.cs
@@ -65,7 +65,7 @@ public class MeshModelingTools : ModuleRules
 				//"MeshBuilder",
                 //"MeshDescriptionOperations",
 				//"MeshUtilitiesCommon",
-				//"MeshReductionInterface", // for UE4 standard simplification
+				//"MeshReductionInterface", // for UE standard simplification
                 //"ProxyLODMeshReduction", // for mesh merging voxel-based csg
 				//"Slate",
 				//"SlateCore",
diff --git a/Engine/Plugins/Experimental/MeshModelingToolset/Source/MeshModelingToolsEditorOnly/MeshModelingToolsEditorOnly.Build.cs b/Engine/Plugins/Experimental/MeshModelingToolset/Source/MeshModelingToolsEditorOnly/MeshModelingToolsEditorOnly.Build.cs
--- a/Engine/Plugins/Experimental/MeshModelingToolset/Source/MeshModelingToolsEditorOnly/MeshModelingToolsEditorOnly.Build.cs
+++ b/Engine/Plugins/Experimental/MeshModelingToolset/Source/MeshModelingToolsEditorOnly/MeshModelingToolsEditorOnly.Build.cs
@@ -58,7 +58,7 @@ public class MeshModelingToolsEditorOnly : ModuleRules
 				"UnrealEd",
 				"MeshBuilder",
 				"MeshUtilitiesCommon",  
-				"MeshReductionInterface", // for UE4 standard simplification 
+				"MeshReductionInterface", // for UE standard simplification 
                 "ProxyLODMeshReduction", // for mesh merging voxel-based csg
 				"OpenSubdiv",
 				//"Slate",
diff --git a/Engine/Plugins/Experimental/MeshModelingToolset/Source/MeshModelingToolsEditorOnly/Private/GenerateLODMeshesTool.cpp b/Engine/Plugins/Experimental/MeshModelingToolset/Source/MeshModelingToolsEditorOnly/Private/GenerateLODMeshesTool.cpp
--- a/Engine/Plugins/Experimental/MeshModelingToolset/Source/MeshModelingToolsEditorOnly/Private/GenerateLODMeshesTool.cpp
+++ b/Engine/Plugins/Experimental/MeshModelingToolset/Source/MeshModelingToolsEditorOnly/Private/GenerateLODMeshesTool.cpp
@@ -67,21 +67,21 @@ UGenerateLODMeshesToolProperties::UGenerateLODMeshesToolProperties()
 	// hardcoded for hair helmet
 
 	FLODLevelGenerateSettings LOD0;
-	LOD0.SimplifierType = ESimplifyType::UE4Standard;
+	LOD0.SimplifierType = ESimplifyType::UEStandard;
 	LOD0.TargetMode = ESimplifyTargetType::VertexCount;
 	LOD0.TargetPercentage = 50;
 	LOD0.TargetCount = 500;
 	LODLevels.Add(LOD0);
 
 	FLODLevelGenerateSettings LOD1;
-	LOD1.SimplifierType = ESimplifyType::UE4Standard;
+	LOD1.SimplifierType = ESimplifyType::UEStandard;
 	LOD1.TargetMode = ESimplifyTargetType::VertexCount;
 	LOD1.TargetPercentage = 30;
 	LOD1.TargetCount = 250;
 	LODLevels.Add(LOD1);
 
 	FLODLevelGenerateSettings LOD2;
-	LOD2.SimplifierType = ESimplifyType::UE4Standard;
+	LOD2.SimplifierType = ESimplifyType::UEStandard;
 	LOD2.TargetMode = ESimplifyTargetType::VertexCount;
 	LOD2.TargetPercentage = 15;
 	LOD2.TargetCount = 150;
diff --git a/Engine/Plugins/Experimental/MeshModelingToolset/Source/MeshModelingToolsEditorOnly/Public/GenerateLODMeshesTool.h b/Engine/Plugins/Experimental/MeshModelingToolset/Source/MeshModelingToolsEditorOnly/Public/GenerateLODMeshesTool.h
--- a/Engine/Plugins/Experimental/MeshModelingToolset/Source/MeshModelingToolsEditorOnly/Public/GenerateLODMeshesTool.h
+++ b/Engine/Plugins/Experimental/MeshModelingToolset/Source/MeshModelingToolsEditorOnly/Public/GenerateLODMeshesTool.h
@@ -42,7 +42,7 @@ struct MESHMODELINGTOOLSEDITORONLY_API FLODLevelGenerateSettings
 
 	/** Simplification Scheme  */
 	UPROPERTY(EditAnywhere, Category = Options)
-	ESimplifyType SimplifierType = ESimplifyType::UE4Standard;
+	ESimplifyType SimplifierType = ESimplifyType::UEStandard;
 
 	/** Simplification Target Type  */
 	UPROPERTY(EditAnywhere, Category = Options)
@@ -89,7 +89,7 @@ public:
 	/** Simplification Scheme  */
 	//UPROPERTY(EditAnywhere, Category = Options)
 	UPROPERTY()
-	ESimplifyType SimplifierType = ESimplifyType::UE4Standard;
+	ESimplifyType SimplifierType = ESimplifyType::UEStandard;
 
 	/** Output LOD Assets will be numbered starting at this number */
 	UPROPERTY(EditAnywhere, Category = Options)
@@ -101,7 +101,7 @@ public:
 	int TargetPercentage = 50;
 
 	/** Target edge length */
-	//UPROPERTY(EditAnywhere, Category = Options, meta = (UIMin = "3.0", UIMax = "10.0", ClampMin = "0.001", ClampMax = "1000.0", EditCondition = "TargetMode == ESimplifyTargetType::EdgeLength && SimplifierType != ESimplifyType::UE4Standard"))
+	//UPROPERTY(EditAnywhere, Category = Options, meta = (UIMin = "3.0", UIMax = "10.0", ClampMin = "0.001", ClampMax = "1000.0", EditCondition = "TargetMode == ESimplifyTargetType::EdgeLength && SimplifierType != ESimplifyType::UEStandard"))
 	UPROPERTY()
 	float TargetEdgeLength;
 
diff --git a/Engine/Plugins/Experimental/MeshModelingToolset/Source/MeshModelingToolsEditorOnly/Public/SimplifyMeshTool.h b/Engine/Plugins/Experimental/MeshModelingToolset/Source/MeshModelingToolsEditorOnly/Public/SimplifyMeshTool.h
--- a/Engine/Plugins/Experimental/MeshModelingToolset/Source/MeshModelingToolsEditorOnly/Public/SimplifyMeshTool.h
+++ b/Engine/Plugins/Experimental/MeshModelingToolset/Source/MeshModelingToolsEditorOnly/Public/SimplifyMeshTool.h
@@ -58,7 +58,7 @@ public:
 	int TargetPercentage;
 
 	/** Target edge length */
-	UPROPERTY(EditAnywhere, Category = Options, meta = (UIMin = "3.0", UIMax = "10.0", ClampMin = "0.001", ClampMax = "1000.0", EditCondition = "TargetMode == ESimplifyTargetType::EdgeLength && SimplifierType != ESimplifyType::UE4Standard"))
+	UPROPERTY(EditAnywhere, Category = Options, meta = (UIMin = "3.0", UIMax = "10.0", ClampMin = "0.001", ClampMax = "1000.0", EditCondition = "TargetMode == ESimplifyTargetType::EdgeLength && SimplifierType != ESimplifyType::UEStandard"))
 	float TargetEdgeLength;
 
 	/** Target triangle count */
diff --git a/Engine/Plugins/Experimental/MeshModelingToolset/Source/ModelingOperators/ModelingOperators.Build.cs b/Engine/Plugins/Experimental/MeshModelingToolset/Source/ModelingOperators/ModelingOperators.Build.cs
--- a/Engine/Plugins/Experimental/MeshModelingToolset/Source/ModelingOperators/ModelingOperators.Build.cs
+++ b/Engine/Plugins/Experimental/MeshModelingToolset/Source/ModelingOperators/ModelingOperators.Build.cs
@@ -55,7 +55,7 @@ public class ModelingOperators : ModuleRules
 				//"MeshBuilder",
                 //"MeshDescriptionOperations",
 				//"MeshUtilitiesCommon",
-				//"MeshReductionInterface", // for UE4 standard simplification
+				//"MeshReductionInterface", // for standard UE simplification
 
 				"CoreUObject",
 				"Engine",
diff --git a/Engine/Plugins/Experimental/MeshModelingToolset/Source/ModelingOperatorsEditorOnly/ModelingOperatorsEditorOnly.Build.cs b/Engine/Plugins/Experimental/MeshModelingToolset/Source/ModelingOperatorsEditorOnly/ModelingOperatorsEditorOnly.Build.cs
--- a/Engine/Plugins/Experimental/MeshModelingToolset/Source/ModelingOperatorsEditorOnly/ModelingOperatorsEditorOnly.Build.cs
+++ b/Engine/Plugins/Experimental/MeshModelingToolset/Source/ModelingOperatorsEditorOnly/ModelingOperatorsEditorOnly.Build.cs
@@ -49,7 +49,7 @@ public class ModelingOperatorsEditorOnly : ModuleRules
 
 				"MeshBuilder",
 				"MeshUtilitiesCommon",  
-				"MeshReductionInterface", // for UE4 standard simplification 
+				"MeshReductionInterface", // for UE standard simplification 
 				"MeshUtilities",			// for tangents calculation
 			}
 			);
diff --git a/Engine/Plugins/Experimental/MeshModelingToolset/Source/ModelingOperatorsEditorOnly/Private/CleaningOps/SimplifyMeshOp.cpp b/Engine/Plugins/Experimental/MeshModelingToolset/Source/ModelingOperatorsEditorOnly/Private/CleaningOps/SimplifyMeshOp.cpp
--- a/Engine/Plugins/Experimental/MeshModelingToolset/Source/ModelingOperatorsEditorOnly/Private/CleaningOps/SimplifyMeshOp.cpp
+++ b/Engine/Plugins/Experimental/MeshModelingToolset/Source/ModelingOperatorsEditorOnly/Private/CleaningOps/SimplifyMeshOp.cpp
@@ -249,7 +249,7 @@ void FSimplifyMeshOp::CalculateResult(FProgressCancel* Progress)
 
 
 		bool bFailedModifyNeedsRegen = false;
-		// The UE4 tool will split the UV boundaries.  Need to weld this.
+		// The UEStandard simplifier will split the UV boundaries.  Need to weld this.
 		{
 			FDynamicMesh3* ComponentMesh = ResultMesh.Get();
 
diff --git a/Engine/Plugins/Experimental/MeshModelingToolset/Source/ModelingOperatorsEditorOnly/Public/CleaningOps/SimplifyMeshOp.h b/Engine/Plugins/Experimental/MeshModelingToolset/Source/ModelingOperatorsEditorOnly/Public/CleaningOps/SimplifyMeshOp.h
--- a/Engine/Plugins/Experimental/MeshModelingToolset/Source/ModelingOperatorsEditorOnly/Public/CleaningOps/SimplifyMeshOp.h
+++ b/Engine/Plugins/Experimental/MeshModelingToolset/Source/ModelingOperatorsEditorOnly/Public/CleaningOps/SimplifyMeshOp.h
@@ -43,7 +43,7 @@ enum class ESimplifyType : uint8
 	Attribute = 1 UMETA(DisplayName = "Normal Aware"),
 
 	/** Highest quality reduction. */
-	UE4Standard = 2 UMETA(DisplayName = "UE4 Standard"),
+	UEStandard = 2 UMETA(DisplayName = "UE Standard"),
 
 	/** Edge collapse to existing vertices only.  Quality may suffer.*/
 	MinimalExistingVertex = 3 UMETA(DisplayName = "Existing Positions"),
@@ -76,7 +76,7 @@ public:
 	/** Angle threshold in degrees used for testing if two triangles should be considered coplanar, or two lines collinear */
 	float MinimalPlanarAngleThresh = 0.01f;
 
-	// stored for the UE4 Standard path
+	// stored for the UEStandard path
 	TSharedPtr<FMeshDescription, ESPMode::ThreadSafe> OriginalMeshDescription;
 	// stored for the GeometryProcessing custom simplifier paths (currently precomputed once in tool setup)
 	TSharedPtr<FDynamicMesh3, ESPMode::ThreadSafe> OriginalMesh;
