ModelingTools: Refactor SingleSelectionTool to ToolTargets (pt.2)

#rnx

diff --git a/Engine/Plugins/Experimental/MeshModelingToolset/Source/MeshModelingTools/Private/HoleFillTool.cpp b/Engine/Plugins/Experimental/MeshModelingToolset/Source/MeshModelingTools/Private/HoleFillTool.cpp
--- a/Engine/Plugins/Experimental/MeshModelingToolset/Source/MeshModelingTools/Private/HoleFillTool.cpp
+++ b/Engine/Plugins/Experimental/MeshModelingToolset/Source/MeshModelingTools/Private/HoleFillTool.cpp
@@ -14,6 +14,11 @@
 #include "MeshOpPreviewHelpers.h"
 #include "Selection/PolygonSelectionMechanic.h"
 
+#include "TargetInterfaces/MaterialProvider.h"
+#include "TargetInterfaces/MeshDescriptionCommitter.h"
+#include "TargetInterfaces/MeshDescriptionProvider.h"
+#include "TargetInterfaces/PrimitiveComponentBackedTarget.h"
+
 #include "ExplicitUseGeometryMathTypes.h"		// using UE::Geometry::(math types)
 using namespace UE::Geometry;
 
@@ -23,22 +28,9 @@ using namespace UE::Geometry;
  * ToolBuilder
  */
 
-bool UHoleFillToolBuilder::CanBuildTool(const FToolBuilderState& SceneState) const
+USingleSelectionMeshEditingTool* UHoleFillToolBuilder::CreateNewTool(const FToolBuilderState& SceneState) const
 {
-	return ToolBuilderUtil::CountComponents(SceneState, CanMakeComponentTarget) == 1;
-}
-
-UInteractiveTool* UHoleFillToolBuilder::BuildTool(const FToolBuilderState& SceneState) const
-{
-	UActorComponent* ActorComponent = ToolBuilderUtil::FindFirstComponent(SceneState, CanMakeComponentTarget);
-	auto* MeshComponent = Cast<UPrimitiveComponent>(ActorComponent);
-	check(MeshComponent != nullptr);
-
-	UHoleFillTool* NewTool = NewObject<UHoleFillTool>(SceneState.ToolManager);
-	NewTool->SetSelection(MakeComponentTarget(MeshComponent));
-	NewTool->SetWorld(SceneState.World);
-
-	return NewTool;
+	return NewObject<UHoleFillTool>(SceneState.ToolManager);
 }
 
 /*
@@ -101,7 +93,7 @@ TUniquePtr<FDynamicMeshOperator> UHoleFillOperatorFactory::MakeNewOperator()
 {
 	TUniquePtr<FHoleFillOp> FillOp = MakeUnique<FHoleFillOp>();
 
-	FTransform LocalToWorld = FillTool->ComponentTarget->GetWorldTransform();
+	FTransform LocalToWorld = Cast<IPrimitiveComponentBackedTarget>(FillTool->Target)->GetWorldTransform();
 	FillOp->SetResultTransform((UE::Geometry::FTransform3d)LocalToWorld);
 	FillOp->OriginalMesh = FillTool->OriginalMesh;
 	FillOp->MeshUVScaleFactor = FillTool->MeshUVScaleFactor;
@@ -125,7 +117,7 @@ void UHoleFillTool::Setup()
 {
 	USingleSelectionTool::Setup();
 
-	if (!ComponentTarget)
+	if (!Target)
 	{
 		return;
 	}
@@ -133,7 +125,7 @@ void UHoleFillTool::Setup()
 	// create mesh to operate on
 	OriginalMesh = MakeShared<FDynamicMesh3, ESPMode::ThreadSafe>();
 	FMeshDescriptionToDynamicMesh Converter;
-	Converter.Convert(ComponentTarget->GetMesh(), *OriginalMesh);
+	Converter.Convert(Cast<IMeshDescriptionProvider>(Target)->GetMeshDescription(), *OriginalMesh);
 
 	// initialize properties
 	Properties = NewObject<UHoleFillToolProperties>(this, TEXT("Hole Fill Settings"));
@@ -172,6 +164,7 @@ void UHoleFillTool::Setup()
 	bool bTopologyOK = Topology->RebuildTopology();
 
 	// Set up selection mechanic to find and select edges
+	IPrimitiveComponentBackedTarget* TargetComponent = Cast<IPrimitiveComponentBackedTarget>(Target);
 	SelectionMechanic = NewObject<UPolygonSelectionMechanic>(this);
 	SelectionMechanic->bAddSelectionFilterPropertiesToParentTool = false;
 	SelectionMechanic->Setup(this);
@@ -179,7 +172,7 @@ void UHoleFillTool::Setup()
 	SelectionMechanic->Properties->bSelectFaces = false;
 	SelectionMechanic->Properties->bSelectVertices = false;
 	SelectionMechanic->Initialize(OriginalMesh.Get(),
-		ComponentTarget->GetWorldTransform(),
+		TargetComponent->GetWorldTransform(),
 		TargetWorld,
 		Topology.Get(),
 		[this]() { return &MeshSpatial; }
@@ -226,7 +219,7 @@ void UHoleFillTool::Setup()
 			EToolMessageLevel::UserNotification);
 
 		// Hide all meshes except the Preview
-		ComponentTarget->SetOwnerVisibility(false);
+		TargetComponent->SetOwnerVisibility(false);
 	}
 
 	SetToolDisplayName(LOCTEXT("ToolName", "Fill Holes"));
@@ -270,7 +263,7 @@ void UHoleFillTool::Shutdown(EToolShutdownType ShutdownType)
 		SelectionMechanic->Shutdown();
 	}
 
-	ComponentTarget->SetOwnerVisibility(true);
+	Cast<IPrimitiveComponentBackedTarget>(Target)->SetOwnerVisibility(true);
 
 	FDynamicMeshOpResult Result = Preview->Shutdown();
 	if (ShutdownType == EToolShutdownType::Accept)
@@ -278,12 +271,12 @@ void UHoleFillTool::Shutdown(EToolShutdownType ShutdownType)
 		GetToolManager()->BeginUndoTransaction(LOCTEXT("HoleFillToolTransactionName", "Hole Fill Tool"));
 
 		check(Result.Mesh.Get() != nullptr);
-		ComponentTarget->CommitMesh([&Result](const FPrimitiveComponentTarget::FCommitParams& CommitParams)
+		Cast<IMeshDescriptionCommitter>(Target)->CommitMeshDescription([&Result](const IMeshDescriptionCommitter::FCommitterParams& CommitParams)
 		{
 			FDynamicMeshToMeshDescription Converter;
 
 			// full conversion if normal topology changed or faces were inverted
-			Converter.Convert(Result.Mesh.Get(), *CommitParams.MeshDescription);
+			Converter.Convert(Result.Mesh.Get(), *CommitParams.MeshDescriptionOut);
 		});
 
 		GetToolManager()->EndUndoTransaction();
@@ -307,11 +300,6 @@ void UHoleFillTool::RequestAction(EHoleFillToolActions ActionType)
 	bHavePendingAction = true;
 }
 
-void UHoleFillTool::SetWorld(UWorld* World)
-{
-	this->TargetWorld = World;
-}
-
 
 void UHoleFillTool::InvalidatePreviewResult()
 {
@@ -329,7 +317,7 @@ void UHoleFillTool::SetupPreview()
 	Preview->Setup(this->TargetWorld, OpFactory);
 
 	FComponentMaterialSet MaterialSet;
-	ComponentTarget->GetMaterialSet(MaterialSet);
+	Cast<IMaterialProvider>(Target)->GetMaterialSet(MaterialSet);
 	Preview->ConfigureMaterials(MaterialSet.Materials,
 		ToolSetupUtil::GetDefaultWorkingMaterial(GetToolManager())
 	);
@@ -366,7 +354,7 @@ void UHoleFillTool::SetupPreview()
 	});
 
 	// set initial preview to un-processed mesh
-	Preview->PreviewMesh->SetTransform(ComponentTarget->GetWorldTransform());
+	Preview->PreviewMesh->SetTransform(Cast<IPrimitiveComponentBackedTarget>(Target)->GetWorldTransform());
 	Preview->PreviewMesh->UpdatePreview(OriginalMesh.Get());
 
 	Preview->SetVisibility(true);
diff --git a/Engine/Plugins/Experimental/MeshModelingToolset/Source/MeshModelingTools/Private/LatticeDeformerTool.cpp b/Engine/Plugins/Experimental/MeshModelingToolset/Source/MeshModelingTools/Private/LatticeDeformerTool.cpp
--- a/Engine/Plugins/Experimental/MeshModelingToolset/Source/MeshModelingTools/Private/LatticeDeformerTool.cpp
+++ b/Engine/Plugins/Experimental/MeshModelingToolset/Source/MeshModelingTools/Private/LatticeDeformerTool.cpp
@@ -19,6 +19,11 @@
 #include "Algo/ForEach.h"
 #include "Operations/FFDLattice.h"
 
+#include "TargetInterfaces/MaterialProvider.h"
+#include "TargetInterfaces/MeshDescriptionCommitter.h"
+#include "TargetInterfaces/MeshDescriptionProvider.h"
+#include "TargetInterfaces/PrimitiveComponentBackedTarget.h"
+
 #include "ExplicitUseGeometryMathTypes.h"		// using UE::Geometry::(math types)
 using namespace UE::Geometry;
 
@@ -26,22 +31,9 @@ using namespace UE::Geometry;
 
 // Tool builder
 
-bool ULatticeDeformerToolBuilder::CanBuildTool(const FToolBuilderState& SceneState) const
-{
-	// TODO: Do we want the ability to handle multiple meshes with the same lattice?
-	return ToolBuilderUtil::CountComponents(SceneState, CanMakeComponentTarget) == 1;
-}
-
-UInteractiveTool* ULatticeDeformerToolBuilder::BuildTool(const FToolBuilderState& SceneState) const
+USingleSelectionMeshEditingTool* ULatticeDeformerToolBuilder::CreateNewTool(const FToolBuilderState& SceneState) const
 {
-	ULatticeDeformerTool* NewTool = NewObject<ULatticeDeformerTool>(SceneState.ToolManager);
-
-	UActorComponent* ActorComponent = ToolBuilderUtil::FindFirstComponent(SceneState, CanMakeComponentTarget);
-	auto* MeshComponent = Cast<UPrimitiveComponent>(ActorComponent);
-	check(MeshComponent != nullptr);
-	NewTool->SetSelection(MakeComponentTarget(MeshComponent));
-	NewTool->SetWorld(SceneState.World);
-	return NewTool;
+	return NewObject<ULatticeDeformerTool>(SceneState.ToolManager);
 }
 
 
@@ -89,7 +81,7 @@ void ULatticeDeformerTool::InitializeLattice(TArray<FVector3d>& OutLatticePoints
 	Lattice->GenerateInitialLatticePositions(OutLatticePoints);
 
 	// Put the lattice in world space
-	UE::Geometry::FTransform3d LocalToWorld(ComponentTarget->GetWorldTransform());
+	UE::Geometry::FTransform3d LocalToWorld(Cast<IPrimitiveComponentBackedTarget>(Target)->GetWorldTransform());
 	Algo::ForEach(OutLatticePoints, [&LocalToWorld](FVector3d& Point) {
 		Point = LocalToWorld.TransformPosition(Point);
 	});
@@ -107,7 +99,7 @@ void ULatticeDeformerTool::Setup()
 
 	OriginalMesh = MakeShared<FDynamicMesh3, ESPMode::ThreadSafe>();
 	FMeshDescriptionToDynamicMesh Converter;
-	Converter.Convert(ComponentTarget->GetMesh(), *OriginalMesh);
+	Converter.Convert(Cast<IMeshDescriptionProvider>(Target)->GetMeshDescription(), *OriginalMesh);
 
 	Settings = NewObject<ULatticeDeformerToolProperties>(this, TEXT("Lattice Deformer Tool Settings"));
 	Settings->RestoreProperties(this);
@@ -144,7 +136,7 @@ void ULatticeDeformerTool::Setup()
 	ControlPointsMechanic = NewObject<ULatticeControlPointsMechanic>(this);
 	ControlPointsMechanic->Setup(this);
 	ControlPointsMechanic->SetWorld(TargetWorld);
-	UE::Geometry::FTransform3d LocalToWorld(ComponentTarget->GetWorldTransform());
+	UE::Geometry::FTransform3d LocalToWorld(Cast<IPrimitiveComponentBackedTarget>(Target)->GetWorldTransform());
 	ControlPointsMechanic->Initialize(LatticePoints, LatticeEdges, LocalToWorld);
 
 	auto OnPointsChangedLambda = [this]()
@@ -165,7 +157,8 @@ void ULatticeDeformerTool::Shutdown(EToolShutdownType ShutdownType)
 	Settings->SaveProperties(this);
 	ControlPointsMechanic->Shutdown();
 
-	ComponentTarget->SetOwnerVisibility(true);
+	IPrimitiveComponentBackedTarget* TargetComponent = Cast<IPrimitiveComponentBackedTarget>(Target);
+	TargetComponent->SetOwnerVisibility(true);
 
 	if (Preview)
 	{
@@ -180,13 +173,13 @@ void ULatticeDeformerTool::Shutdown(EToolShutdownType ShutdownType)
 
 			// The lattice and its output mesh are in world space, so get them in local space.
 			// TODO: Would it make more sense to do all the lattice computation in local space?
-			UE::Geometry::FTransform3d LocalToWorld(ComponentTarget->GetWorldTransform());
+			UE::Geometry::FTransform3d LocalToWorld(TargetComponent->GetWorldTransform());
 			MeshTransforms::ApplyTransformInverse(*DynamicMeshResult, LocalToWorld);
 
-			ComponentTarget->CommitMesh([DynamicMeshResult](const FPrimitiveComponentTarget::FCommitParams& CommitParams)
+			Cast<IMeshDescriptionCommitter>(Target)->CommitMeshDescription([DynamicMeshResult](const IMeshDescriptionCommitter::FCommitterParams& CommitParams)
 			{
 				FDynamicMeshToMeshDescription Converter;
-				Converter.Convert(DynamicMeshResult, *CommitParams.MeshDescription);
+				Converter.Convert(DynamicMeshResult, *CommitParams.MeshDescriptionOut);
 			});
 
 			GetToolManager()->EndUndoTransaction();
@@ -207,7 +200,7 @@ void ULatticeDeformerTool::StartPreview()
 	Preview->SetIsMeshTopologyConstant(true, EMeshRenderAttributeFlags::Positions | EMeshRenderAttributeFlags::VertexNormals);
 
 	FComponentMaterialSet MaterialSet;
-	ComponentTarget->GetMaterialSet(MaterialSet);
+	Cast<IMaterialProvider>(Target)->GetMaterialSet(MaterialSet);
 	Preview->ConfigureMaterials(MaterialSet.Materials,
 								ToolSetupUtil::GetDefaultWorkingMaterial(GetToolManager())
 	);
@@ -223,7 +216,7 @@ void ULatticeDeformerTool::StartPreview()
 	Preview->SetVisibility(true);
 	Preview->InvalidateResult();
 
-	ComponentTarget->SetOwnerVisibility(false);
+	Cast<IPrimitiveComponentBackedTarget>(Target)->SetOwnerVisibility(false);
 }
 
 void ULatticeDeformerTool::OnTick(float DeltaTime)
@@ -235,7 +228,7 @@ void ULatticeDeformerTool::OnTick(float DeltaTime)
 			TArray<FVector3d> LatticePoints;
 			TArray<FVector2i> LatticeEdges;
 			InitializeLattice(LatticePoints, LatticeEdges);
-			UE::Geometry::FTransform3d LocalToWorld(ComponentTarget->GetWorldTransform());
+			UE::Geometry::FTransform3d LocalToWorld(Cast<IPrimitiveComponentBackedTarget>(Target)->GetWorldTransform());
 			ControlPointsMechanic->Initialize(LatticePoints, LatticeEdges, LocalToWorld);
 			Preview->InvalidateResult();
 			bShouldRebuild = false;
diff --git a/Engine/Plugins/Experimental/MeshModelingToolset/Source/MeshModelingTools/Private/MeshBoundaryToolBase.cpp b/Engine/Plugins/Experimental/MeshModelingToolset/Source/MeshModelingTools/Private/MeshBoundaryToolBase.cpp
--- a/Engine/Plugins/Experimental/MeshModelingToolset/Source/MeshModelingTools/Private/MeshBoundaryToolBase.cpp
+++ b/Engine/Plugins/Experimental/MeshModelingToolset/Source/MeshModelingTools/Private/MeshBoundaryToolBase.cpp
@@ -9,6 +9,9 @@
 #include "MeshDescriptionToDynamicMesh.h"
 #include "Selection/PolygonSelectionMechanic.h"
 
+#include "TargetInterfaces/MeshDescriptionProvider.h"
+#include "TargetInterfaces/PrimitiveComponentBackedTarget.h"
+
 #include "ExplicitUseGeometryMathTypes.h"		// using UE::Geometry::(math types)
 using namespace UE::Geometry;
 
@@ -18,7 +21,7 @@ void UMeshBoundaryToolBase::Setup()
 {
 	USingleSelectionTool::Setup();
 
-	if (!ComponentTarget)
+	if (!Target)
 	{
 		return;
 	}
@@ -26,7 +29,7 @@ void UMeshBoundaryToolBase::Setup()
 	// create mesh to operate on
 	OriginalMesh = MakeShared<FDynamicMesh3, ESPMode::ThreadSafe>();
 	FMeshDescriptionToDynamicMesh Converter;
-	Converter.Convert(ComponentTarget->GetMesh(), *OriginalMesh);
+	Converter.Convert(Cast<IMeshDescriptionProvider>(Target)->GetMeshDescription(), *OriginalMesh);
 
 	// initialize hit query
 	MeshSpatial.SetMesh(OriginalMesh.Get());
@@ -43,7 +46,7 @@ void UMeshBoundaryToolBase::Setup()
 	SelectionMechanic->Properties->bSelectFaces = false;
 	SelectionMechanic->Properties->bSelectVertices = false;
 	SelectionMechanic->Initialize(OriginalMesh.Get(),
-		ComponentTarget->GetWorldTransform(),
+		Cast<IPrimitiveComponentBackedTarget>(Target)->GetWorldTransform(),
 		TargetWorld,
 		Topology.Get(),
 		[this]() { return &MeshSpatial; }
diff --git a/Engine/Plugins/Experimental/MeshModelingToolset/Source/MeshModelingTools/Private/MeshInspectorTool.cpp b/Engine/Plugins/Experimental/MeshModelingToolset/Source/MeshModelingTools/Private/MeshInspectorTool.cpp
--- a/Engine/Plugins/Experimental/MeshModelingToolset/Source/MeshModelingTools/Private/MeshInspectorTool.cpp
+++ b/Engine/Plugins/Experimental/MeshModelingToolset/Source/MeshModelingTools/Private/MeshInspectorTool.cpp
@@ -22,6 +22,10 @@
 #include "ToolSetupUtil.h"
 #include "AssetUtils/MeshDescriptionUtil.h"
 
+#include "TargetInterfaces/MaterialProvider.h"
+#include "TargetInterfaces/MeshDescriptionProvider.h"
+#include "TargetInterfaces/PrimitiveComponentBackedTarget.h"
+
 #include "ExplicitUseGeometryMathTypes.h"		// using UE::Geometry::(math types)
 using namespace UE::Geometry;
 
@@ -32,21 +36,9 @@ using namespace UE::Geometry;
  */
 
 
-bool UMeshInspectorToolBuilder::CanBuildTool(const FToolBuilderState& SceneState) const
-{
-	return ToolBuilderUtil::CountComponents(SceneState, CanMakeComponentTarget) == 1;
-}
-
-UInteractiveTool* UMeshInspectorToolBuilder::BuildTool(const FToolBuilderState& SceneState) const
+USingleSelectionMeshEditingTool* UMeshInspectorToolBuilder::CreateNewTool(const FToolBuilderState& SceneState) const
 {
-	UMeshInspectorTool* NewTool = NewObject<UMeshInspectorTool>(SceneState.ToolManager);
-
-	UActorComponent* ActorComponent = ToolBuilderUtil::FindFirstComponent(SceneState, CanMakeComponentTarget);
-	auto* MeshComponent = Cast<UPrimitiveComponent>(ActorComponent);
-	check(MeshComponent != nullptr);
-	NewTool->SetSelection(MakeComponentTarget(MeshComponent));
-
-	return NewTool;
+	return NewObject<UMeshInspectorTool>(SceneState.ToolManager);
 }
 
 /*
@@ -61,16 +53,17 @@ void UMeshInspectorTool::Setup()
 	UInteractiveTool::Setup();
 
 	// copy input MeshDescription and make sure it has initialized normals/tangents
-	InputMeshDescription = MakeShared<FMeshDescription>(*ComponentTarget->GetMesh());
-	UE::MeshDescription::InitializeAutoGeneratedAttributes(*InputMeshDescription, ComponentTarget->GetOwnerComponent(), 0);
+	IPrimitiveComponentBackedTarget* TargetComponent = Cast<IPrimitiveComponentBackedTarget>(Target);
+	InputMeshDescription = MakeShared<FMeshDescription>(*Cast<IMeshDescriptionProvider>(Target)->GetMeshDescription());
+	UE::MeshDescription::InitializeAutoGeneratedAttributes(*InputMeshDescription, TargetComponent->GetOwnerComponent(), 0);
 
 	PreviewMesh = NewObject<UPreviewMesh>(this);
 	PreviewMesh->bBuildSpatialDataStructure = false;
-	PreviewMesh->CreateInWorld(ComponentTarget->GetOwnerActor()->GetWorld(), FTransform::Identity);
-	PreviewMesh->SetTransform(ComponentTarget->GetWorldTransform());
+	PreviewMesh->CreateInWorld(TargetComponent->GetOwnerActor()->GetWorld(), FTransform::Identity);
+	PreviewMesh->SetTransform(TargetComponent->GetWorldTransform());
 
 	FComponentMaterialSet MaterialSet;
-	ComponentTarget->GetMaterialSet(MaterialSet);
+	Cast<IMaterialProvider>(Target)->GetMaterialSet(MaterialSet);
 	PreviewMesh->SetMaterials(MaterialSet.Materials);
 	DefaultMaterial = PreviewMesh->GetMaterial(0);
 
@@ -87,7 +80,7 @@ void UMeshInspectorTool::Setup()
 	Precompute();
 
 	// hide input StaticMeshComponent
-	ComponentTarget->SetOwnerVisibility(false);
+	TargetComponent->SetOwnerVisibility(false);
 
 	// initialize our properties
 	Settings = NewObject<UMeshInspectorProperties>(this);
@@ -105,7 +98,7 @@ void UMeshInspectorTool::Setup()
 	AddToolPropertySource(Statistics);
 
 	UMeshAnalysisProperties* MeshAnalysis = NewObject<UMeshAnalysisProperties>(this);
-	MeshAnalysis->Update(*PreviewMesh->GetPreviewDynamicMesh(), ComponentTarget->GetWorldTransform());
+	MeshAnalysis->Update(*PreviewMesh->GetPreviewDynamicMesh(), TargetComponent->GetWorldTransform());
 	AddToolPropertySource(MeshAnalysis);
 
 	UpdateVisualization();
@@ -119,7 +112,7 @@ void UMeshInspectorTool::Setup()
 
 void UMeshInspectorTool::Shutdown(EToolShutdownType ShutdownType)
 {
-	ComponentTarget->SetOwnerVisibility(true);
+	Cast<IPrimitiveComponentBackedTarget>(Target)->SetOwnerVisibility(true);
 
 	if (PreviewMesh != nullptr)
 	{
diff --git a/Engine/Plugins/Experimental/MeshModelingToolset/Source/MeshModelingTools/Private/MeshSpaceDeformerTool.cpp b/Engine/Plugins/Experimental/MeshModelingToolset/Source/MeshModelingTools/Private/MeshSpaceDeformerTool.cpp
--- a/Engine/Plugins/Experimental/MeshModelingToolset/Source/MeshModelingTools/Private/MeshSpaceDeformerTool.cpp
+++ b/Engine/Plugins/Experimental/MeshModelingToolset/Source/MeshModelingTools/Private/MeshSpaceDeformerTool.cpp
@@ -27,6 +27,10 @@
 #include "CoreMinimal.h"
 #include "Math/Matrix.h"
 
+#include "TargetInterfaces/MeshDescriptionCommitter.h"
+#include "TargetInterfaces/MeshDescriptionProvider.h"
+#include "TargetInterfaces/PrimitiveComponentBackedTarget.h"
+
 #include "ExplicitUseGeometryMathTypes.h"		// using UE::Geometry::(math types)
 using namespace UE::Geometry;
 
@@ -35,28 +39,12 @@ using namespace UE::Geometry;
 /*
  * ToolBuilder
  */
-UInteractiveTool* UMeshSpaceDeformerToolBuilder::BuildTool(const FToolBuilderState& SceneState) const
+USingleSelectionMeshEditingTool* UMeshSpaceDeformerToolBuilder::CreateNewTool(const FToolBuilderState& SceneState) const
 {
-	UMeshSpaceDeformerTool* MeshSpaceDeformerTool = NewObject<UMeshSpaceDeformerTool>(SceneState.ToolManager);
-	
-	UActorComponent* ActorComponent = ToolBuilderUtil::FindFirstComponent(SceneState, CanMakeComponentTarget);
-	auto* MeshComponent = Cast<UPrimitiveComponent>(ActorComponent);
-	check(MeshComponent != nullptr);
-
-	MeshSpaceDeformerTool->SetSelection(MakeComponentTarget(MeshComponent));
-	MeshSpaceDeformerTool->SetWorld(SceneState.World);
-
-	return MeshSpaceDeformerTool;
+	return NewObject<UMeshSpaceDeformerTool>(SceneState.ToolManager);
 }
 
 
-bool UMeshSpaceDeformerToolBuilder::CanBuildTool(const FToolBuilderState& SceneState) const
-{
-	return ToolBuilderUtil::CountComponents(SceneState, CanMakeComponentTarget) == 1;
-}
-
-
-
 TUniquePtr<FDynamicMeshOperator> USpaceDeformerOperatorFactory::MakeNewOperator()
 {
 
@@ -121,18 +109,6 @@ UMeshSpaceDeformerTool::UMeshSpaceDeformerTool()
 {
 }
 
-void UMeshSpaceDeformerTool::SetWorld(UWorld* World)
-{
-	this->TargetWorld = World;
-}
-
-
-void UMeshSpaceDeformerTool::SetAssetAPI(IToolsContextAssetAPI* AssetAPIIn)
-{
-	this->AssetAPI = AssetAPIIn;
-}
-
-
 bool UMeshSpaceDeformerTool::CanAccept() const 
 {
 	return Super::CanAccept() && (Preview == nullptr || Preview->HaveValidResult());
@@ -154,14 +130,15 @@ void UMeshSpaceDeformerTool::Setup()
 	{
 		OriginalDynamicMesh = MakeShared<FDynamicMesh3, ESPMode::ThreadSafe>();
 		FMeshDescriptionToDynamicMesh Converter;
-		Converter.Convert(ComponentTarget->GetMesh(), *OriginalDynamicMesh);
+		Converter.Convert(Cast<IMeshDescriptionProvider>(Target)->GetMeshDescription(), *OriginalDynamicMesh);
 	}
 
-	FTransform MeshTransform = ComponentTarget->GetWorldTransform();
+	IPrimitiveComponentBackedTarget* TargetComponent = Cast<IPrimitiveComponentBackedTarget>(Target);
+	FTransform MeshTransform = TargetComponent->GetWorldTransform();
 
 	// Hide the mesh, and potentially put a semi-transparent copy in its place. We could
 	// update the materials and restore them later, but this seems safer.
-	ComponentTarget->SetOwnerVisibility(false);
+	TargetComponent->SetOwnerVisibility(false);
 
 	OriginalMeshPreview = NewObject<UPreviewMesh>();
 	OriginalMeshPreview->CreateInWorld(TargetWorld, MeshTransform);
@@ -207,7 +184,7 @@ void UMeshSpaceDeformerTool::Setup()
 	// add click to set plane behavior
 	SetPointInWorldConnector = MakePimpl<FSelectClickedAction>();
 	SetPointInWorldConnector->World = this->TargetWorld;
-	SetPointInWorldConnector->InvisibleComponentsToHitTest.Add(ComponentTarget->GetOwnerComponent());
+	SetPointInWorldConnector->InvisibleComponentsToHitTest.Add(TargetComponent->GetOwnerComponent());
 	SetPointInWorldConnector->OnClickedPositionFunc = [this](const FHitResult& Hit)
 	{
 		SetGizmoFrameFromWorldPos(Hit.ImpactPoint, Hit.ImpactNormal, Settings->bAlignToNormalOnCtrlClick);
@@ -315,7 +292,7 @@ void UMeshSpaceDeformerTool::Setup()
 	// We want to align to the original mesh, even though it is hidden (our stand-in preview mesh that
 	// we use to display a transparent version does not get hit tested by our normal raycasts into
 	// the world).
-	TArray<const UPrimitiveComponent*> ComponentsToInclude{ ComponentTarget->GetOwnerComponent() } ;
+	TArray<const UPrimitiveComponent*> ComponentsToInclude{ TargetComponent->GetOwnerComponent() } ;
 	DragAlignmentMechanic->AddToGizmo(TransformGizmo, nullptr, &ComponentsToInclude);
 	DragAlignmentMechanic->AddToGizmo(IntervalGizmo, nullptr, &ComponentsToInclude);
 
@@ -330,7 +307,7 @@ void UMeshSpaceDeformerTool::Shutdown(EToolShutdownType ShutdownType)
 	Settings->SaveProperties(this);
 
 	// Restore source mesh and remove our stand-in
-	ComponentTarget->SetOwnerVisibility(true);
+	Cast<IPrimitiveComponentBackedTarget>(Target)->SetOwnerVisibility(true);
 	OriginalMeshPreview->SetVisible(false);
 	OriginalMeshPreview->Disconnect();
 	OriginalMeshPreview = nullptr;
@@ -346,10 +323,10 @@ void UMeshSpaceDeformerTool::Shutdown(EToolShutdownType ShutdownType)
 			FDynamicMesh3* DynamicMeshResult = Result.Mesh.Get();
 			check(DynamicMeshResult != nullptr);
 
-			ComponentTarget->CommitMesh([DynamicMeshResult](const FPrimitiveComponentTarget::FCommitParams& CommitParams)
+			Cast<IMeshDescriptionCommitter>(Target)->CommitMeshDescription([DynamicMeshResult](const IMeshDescriptionCommitter::FCommitterParams& CommitParams)
 			{
 				FDynamicMeshToMeshDescription Converter;
-				Converter.Convert(DynamicMeshResult, *CommitParams.MeshDescription);
+				Converter.Convert(DynamicMeshResult, *CommitParams.MeshDescriptionOut);
 			});
 
 
@@ -388,7 +365,7 @@ void  UMeshSpaceDeformerTool::OnPropertyModified(UObject* PropertySet, FProperty
 void UMeshSpaceDeformerTool::UpdateOpParameters(FMeshSpaceDeformerOp& MeshSpaceDeformerOp) const
 {
 	MeshSpaceDeformerOp.OriginalMesh = OriginalDynamicMesh;
-	MeshSpaceDeformerOp.SetTransform(ComponentTarget->GetWorldTransform());
+	MeshSpaceDeformerOp.SetTransform(Cast<IPrimitiveComponentBackedTarget>(Target)->GetWorldTransform());
 	MeshSpaceDeformerOp.GizmoFrame = GizmoFrame;
 
 	// set the bound range
diff --git a/Engine/Plugins/Experimental/MeshModelingToolset/Source/MeshModelingTools/Private/Physics/ExtractCollisionGeometryTool.cpp b/Engine/Plugins/Experimental/MeshModelingToolset/Source/MeshModelingTools/Private/Physics/ExtractCollisionGeometryTool.cpp
--- a/Engine/Plugins/Experimental/MeshModelingToolset/Source/MeshModelingTools/Private/Physics/ExtractCollisionGeometryTool.cpp
+++ b/Engine/Plugins/Experimental/MeshModelingToolset/Source/MeshModelingTools/Private/Physics/ExtractCollisionGeometryTool.cpp
@@ -27,34 +27,33 @@
 #include "Engine/Classes/PhysicsEngine/BodySetup.h"
 #include "Engine/Classes/PhysicsEngine/AggregateGeom.h"
 
+#include "TargetInterfaces/MaterialProvider.h"
+#include "TargetInterfaces/MeshDescriptionCommitter.h"
+#include "TargetInterfaces/MeshDescriptionProvider.h"
+#include "TargetInterfaces/PrimitiveComponentBackedTarget.h"
+#include "TargetInterfaces/StaticMeshBackedTarget.h"
+
 #include "ExplicitUseGeometryMathTypes.h"		// using UE::Geometry::(math types)
 using namespace UE::Geometry;
 
 #define LOCTEXT_NAMESPACE "UExtractCollisionGeometryTool"
 
-bool UExtractCollisionGeometryToolBuilder::CanBuildTool(const FToolBuilderState& SceneState) const
+const FToolTargetTypeRequirements& UExtractCollisionGeometryToolBuilder::GetTargetRequirements() const
 {
-	int32 NumStaticMeshes = ToolBuilderUtil::CountComponents(SceneState, [&](UActorComponent* Comp) { return Cast<UStaticMeshComponent>(Comp) != nullptr; });
-	int32 NumComponentTargets = ToolBuilderUtil::CountComponents(SceneState, CanMakeComponentTarget);
-	return (NumStaticMeshes == 1 && NumStaticMeshes == NumComponentTargets);
+	static FToolTargetTypeRequirements TypeRequirements({
+		UMaterialProvider::StaticClass(),
+		UMeshDescriptionCommitter::StaticClass(),
+		UMeshDescriptionProvider::StaticClass(),
+		UPrimitiveComponentBackedTarget::StaticClass(),
+		UStaticMeshBackedTarget::StaticClass()
+		});
+	return TypeRequirements;
 }
 
 
-UInteractiveTool* UExtractCollisionGeometryToolBuilder::BuildTool(const FToolBuilderState& SceneState) const
+USingleSelectionMeshEditingTool* UExtractCollisionGeometryToolBuilder::CreateNewTool(const FToolBuilderState& SceneState) const
 {
-	UExtractCollisionGeometryTool* NewTool = NewObject<UExtractCollisionGeometryTool>(SceneState.ToolManager);
-	NewTool->SetWorld(SceneState.World);
-	check(AssetAPI);
-	NewTool->SetAssetAPI(AssetAPI);
-
-	TArray<UActorComponent*> ValidComponents = ToolBuilderUtil::FindAllComponents(SceneState,
-		[&](UActorComponent* Comp) { return Cast<UStaticMeshComponent>(Comp) != nullptr; });
-	check(ValidComponents.Num() == 1);
-
-	TUniquePtr<FPrimitiveComponentTarget> ComponentTargets;
-	UStaticMeshComponent* MeshComponent = Cast<UStaticMeshComponent>(ValidComponents[0]);
-	NewTool->SetSelection(MakeComponentTarget(MeshComponent));
-	return NewTool;
+	return NewObject<UExtractCollisionGeometryTool>(SceneState.ToolManager);
 }
 
 
@@ -63,10 +62,11 @@ void UExtractCollisionGeometryTool::Setup()
 	UInteractiveTool::Setup();
 
 	// create preview mesh
+	IPrimitiveComponentBackedTarget* TargetComponent = Cast<IPrimitiveComponentBackedTarget>(Target);
 	PreviewMesh = NewObject<UPreviewMesh>(this);
 	PreviewMesh->bBuildSpatialDataStructure = false;
 	PreviewMesh->CreateInWorld(TargetWorld, FTransform::Identity);
-	PreviewMesh->SetTransform(ComponentTarget->GetWorldTransform());
+	PreviewMesh->SetTransform(TargetComponent->GetWorldTransform());
 	PreviewMesh->SetMaterial(ToolSetupUtil::GetDefaultSculptMaterial(GetToolManager()));
 	PreviewMesh->SetOverrideRenderMaterial(ToolSetupUtil::GetSelectionMaterial(GetToolManager()));
 	PreviewMesh->SetTriangleColorFunction([this](const FDynamicMesh3* Mesh, int TriangleID)
@@ -81,7 +81,7 @@ void UExtractCollisionGeometryTool::Setup()
 	VizSettings->WatchProperty(VizSettings->Color, [this](FColor NewValue) { bVisualizationDirty = true; });
 
 
-	const UStaticMeshComponent* Component = CastChecked<UStaticMeshComponent>(ComponentTarget->GetOwnerComponent());
+	const UStaticMeshComponent* Component = CastChecked<UStaticMeshComponent>(TargetComponent->GetOwnerComponent());
 	const UStaticMesh* StaticMesh = Component->GetStaticMesh();
 	if (ensure(StaticMesh && StaticMesh->GetBodySetup()))
 	{
@@ -89,10 +89,10 @@ void UExtractCollisionGeometryTool::Setup()
 		PhysicsInfo->InitializeFromComponent(Component, true);
 
 		PreviewElements = NewObject<UPreviewGeometry>(this);
-		FTransform TargetTransform = ComponentTarget->GetWorldTransform();
+		FTransform TargetTransform = TargetComponent->GetWorldTransform();
 		//PhysicsInfo->ExternalScale3D = TargetTransform.GetScale3D();
 		//TargetTransform.SetScale3D(FVector::OneVector);
-		PreviewElements->CreateInWorld(ComponentTarget->GetOwnerActor()->GetWorld(), TargetTransform);
+		PreviewElements->CreateInWorld(TargetComponent->GetOwnerActor()->GetWorld(), TargetTransform);
 		
 		UE::PhysicsTools::InitializePreviewGeometryLines(*PhysicsInfo, PreviewElements,
 			VizSettings->Color, VizSettings->LineThickness, 0.0f, 16);
@@ -125,8 +125,8 @@ void UExtractCollisionGeometryTool::Shutdown(EToolShutdownType ShutdownType)
 	{
 		UMaterialInterface* UseMaterial = UMaterial::GetDefaultMaterial(MD_Surface);
 
-		FString NewName = ComponentTarget->IsValid() ?
-			FString::Printf(TEXT("%s_Collision"), *ComponentTarget->GetOwnerComponent()->GetName()) : TEXT("CollisionGeo");
+		FString NewName = Target->IsValid() ?
+			FString::Printf(TEXT("%s_Collision"), *Cast<IPrimitiveComponentBackedTarget>(Target)->GetOwnerComponent()->GetName()) : TEXT("CollisionGeo");
 
 		GetToolManager()->BeginUndoTransaction(LOCTEXT("CreateCollisionMesh", "Collision To Mesh"));
 
diff --git a/Engine/Plugins/Experimental/MeshModelingToolset/Source/MeshModelingTools/Private/RevolveBoundaryTool.cpp b/Engine/Plugins/Experimental/MeshModelingToolset/Source/MeshModelingTools/Private/RevolveBoundaryTool.cpp
--- a/Engine/Plugins/Experimental/MeshModelingToolset/Source/MeshModelingTools/Private/RevolveBoundaryTool.cpp
+++ b/Engine/Plugins/Experimental/MeshModelingToolset/Source/MeshModelingTools/Private/RevolveBoundaryTool.cpp
@@ -15,6 +15,8 @@
 #include "ToolSceneQueriesUtil.h"
 #include "ToolSetupUtil.h"
 
+#include "TargetInterfaces/PrimitiveComponentBackedTarget.h"
+
 #include "ExplicitUseGeometryMathTypes.h"		// using UE::Geometry::(math types)
 using namespace UE::Geometry;
 
@@ -22,24 +24,9 @@ using namespace UE::Geometry;
 
 // Tool builder
 
-bool URevolveBoundaryToolBuilder::CanBuildTool(const FToolBuilderState& SceneState) const
-{
-	return this->AssetAPI != nullptr && ToolBuilderUtil::CountComponents(SceneState, CanMakeComponentTarget) == 1;
-}
-
-UInteractiveTool* URevolveBoundaryToolBuilder::BuildTool(const FToolBuilderState& SceneState) const
+USingleSelectionMeshEditingTool* URevolveBoundaryToolBuilder::CreateNewTool(const FToolBuilderState& SceneState) const
 {
-	UActorComponent* ActorComponent = ToolBuilderUtil::FindFirstComponent(SceneState, CanMakeComponentTarget);
-	auto* MeshComponent = Cast<UPrimitiveComponent>(ActorComponent);
-	check(MeshComponent != nullptr);
-
-	URevolveBoundaryTool* NewTool = NewObject<URevolveBoundaryTool>(SceneState.ToolManager);
-
-	NewTool->SetSelection(MakeComponentTarget(MeshComponent));
-	NewTool->SetWorld(SceneState.World);
-	NewTool->SetAssetAPI(AssetAPI);
-
-	return NewTool;
+	return NewObject<URevolveBoundaryTool>(SceneState.ToolManager);
 }
 
 
@@ -57,7 +44,7 @@ TUniquePtr<FDynamicMeshOperator> URevolveBoundaryOperatorFactory::MakeNewOperato
 		if (RevolveBoundaryTool->Topology->IsBoundaryEdge(EdgeID))
 		{
 			const TArray<int32>& VertexIndices = RevolveBoundaryTool->Topology->GetGroupEdgeVertices(EdgeID);
-			FTransform ToWorld = RevolveBoundaryTool->ComponentTarget->GetWorldTransform();
+			FTransform ToWorld = Cast<IPrimitiveComponentBackedTarget>(RevolveBoundaryTool->Target)->GetWorldTransform();
 
 			// Boundary loop includes the last vertex twice, so stop early.
 			CurveSweepOp->ProfileCurve.Reserve(VertexIndices.Num() - 1);
@@ -123,7 +110,7 @@ void URevolveBoundaryTool::Setup()
 
 	PlaneMechanic->SetEnableGridSnaping(Settings->bSnapToWorldGrid);
 
-	ComponentTarget->SetOwnerVisibility(Settings->bDisplayOriginalMesh);
+	Cast<IPrimitiveComponentBackedTarget>(Target)->SetOwnerVisibility(Settings->bDisplayOriginalMesh);
 
 	SetToolDisplayName(LOCTEXT("ToolName", "Revolve Boundary"));
 	GetToolManager()->DisplayMessage(
@@ -189,7 +176,7 @@ void URevolveBoundaryTool::OnClicked(const FInputDeviceRay& ClickPos)
 			
 			FVector3d VertexA, VertexB;
 			OriginalMesh->GetEdgeV(ClickedEid, VertexA, VertexB);
-			FTransform ToWorldTranform = ComponentTarget->GetWorldTransform();
+			FTransform ToWorldTranform = Cast<IPrimitiveComponentBackedTarget>(Target)->GetWorldTransform();
 			FLine3d EdgeLine = FLine3d::FromPoints(ToWorldTranform.TransformPosition((FVector)VertexA), 
 				ToWorldTranform.TransformPosition((FVector)VertexB));
 			
@@ -269,7 +256,7 @@ void URevolveBoundaryTool::Shutdown(EToolShutdownType ShutdownType)
 
 	PlaneMechanic->Shutdown();
 
-	ComponentTarget->SetOwnerVisibility(true);
+	Cast<IPrimitiveComponentBackedTarget>(Target)->SetOwnerVisibility(true);
 
 	if (Preview)
 	{
@@ -350,7 +337,7 @@ void URevolveBoundaryTool::OnPropertyModified(UObject* PropertySet, FProperty* P
 		FRotator(Settings->AxisPitch, Settings->AxisYaw, 0).Quaternion()));
 	UpdateRevolutionAxis();
 
-	ComponentTarget->SetOwnerVisibility(Settings->bDisplayOriginalMesh);
+	Cast<IPrimitiveComponentBackedTarget>(Target)->SetOwnerVisibility(Settings->bDisplayOriginalMesh);
 	PlaneMechanic->SetEnableGridSnaping(Settings->bSnapToWorldGrid);
 
 	if (Preview)
diff --git a/Engine/Plugins/Experimental/MeshModelingToolset/Source/MeshModelingTools/Private/SingleSelectionMeshEditingTool.cpp b/Engine/Plugins/Experimental/MeshModelingToolset/Source/MeshModelingTools/Private/SingleSelectionMeshEditingTool.cpp
deleted file mode 100644
diff --git a/Engine/Plugins/Experimental/MeshModelingToolset/Source/MeshModelingTools/Private/WeldMeshEdgesTool.cpp b/Engine/Plugins/Experimental/MeshModelingToolset/Source/MeshModelingTools/Private/WeldMeshEdgesTool.cpp
--- a/Engine/Plugins/Experimental/MeshModelingToolset/Source/MeshModelingTools/Private/WeldMeshEdgesTool.cpp
+++ b/Engine/Plugins/Experimental/MeshModelingToolset/Source/MeshModelingTools/Private/WeldMeshEdgesTool.cpp
@@ -13,6 +13,11 @@
 #include "SceneManagement.h" // for FPrimitiveDrawInterface
 #include "Async/ParallelFor.h"
 
+#include "TargetInterfaces/MaterialProvider.h"
+#include "TargetInterfaces/MeshDescriptionCommitter.h"
+#include "TargetInterfaces/MeshDescriptionProvider.h"
+#include "TargetInterfaces/PrimitiveComponentBackedTarget.h"
+
 #include "ExplicitUseGeometryMathTypes.h"		// using UE::Geometry::(math types)
 using namespace UE::Geometry;
 
@@ -22,23 +27,9 @@ using namespace UE::Geometry;
  * ToolBuilder
  */
 
-
-bool UWeldMeshEdgesToolBuilder::CanBuildTool(const FToolBuilderState& SceneState) const
-{
-	return ToolBuilderUtil::CountComponents(SceneState, CanMakeComponentTarget) == 1;
-}
-
-UInteractiveTool* UWeldMeshEdgesToolBuilder::BuildTool(const FToolBuilderState& SceneState) const
+USingleSelectionMeshEditingTool* UWeldMeshEdgesToolBuilder::CreateNewTool(const FToolBuilderState& SceneState) const
 {
-	UWeldMeshEdgesTool* NewTool = NewObject<UWeldMeshEdgesTool>(SceneState.ToolManager);
-
-	UActorComponent* ActorComponent = ToolBuilderUtil::FindFirstComponent(SceneState, CanMakeComponentTarget);
-	auto* MeshComponent = Cast<UPrimitiveComponent>(ActorComponent);
-	check(MeshComponent != nullptr);
-
-	NewTool->SetSelection(MakeComponentTarget(MeshComponent));
-
-	return NewTool;
+	return NewObject<UWeldMeshEdgesTool>(SceneState.ToolManager);
 }
 
 
@@ -58,25 +49,26 @@ void UWeldMeshEdgesTool::Setup()
 	UInteractiveTool::Setup();
 
 	// create dynamic mesh component to use for live preview
-	DynamicMeshComponent = NewObject<USimpleDynamicMeshComponent>(ComponentTarget->GetOwnerActor(), "DynamicMesh");
-	DynamicMeshComponent->SetupAttachment(ComponentTarget->GetOwnerActor()->GetRootComponent());
+	IPrimitiveComponentBackedTarget* TargetComponent = Cast<IPrimitiveComponentBackedTarget>(Target);
+	DynamicMeshComponent = NewObject<USimpleDynamicMeshComponent>(TargetComponent->GetOwnerActor(), "DynamicMesh");
+	DynamicMeshComponent->SetupAttachment(TargetComponent->GetOwnerActor()->GetRootComponent());
 	DynamicMeshComponent->RegisterComponent();
-	DynamicMeshComponent->SetWorldTransform(ComponentTarget->GetWorldTransform());
+	DynamicMeshComponent->SetWorldTransform(TargetComponent->GetWorldTransform());
 
 	// transfer materials
 	FComponentMaterialSet MaterialSet;
-	ComponentTarget->GetMaterialSet(MaterialSet);
+	Cast<IMaterialProvider>(Target)->GetMaterialSet(MaterialSet);
 	for (int k = 0; k < MaterialSet.Materials.Num(); ++k)
 	{
 		DynamicMeshComponent->SetMaterial(k, MaterialSet.Materials[k]);
 	}
 
 	DynamicMeshComponent->TangentsType = EDynamicMeshTangentCalcType::AutoCalculated;
-	DynamicMeshComponent->InitializeMesh(ComponentTarget->GetMesh());
+	DynamicMeshComponent->InitializeMesh(Cast<IMeshDescriptionProvider>(Target)->GetMeshDescription());
 	OriginalMesh.Copy(*DynamicMeshComponent->GetMesh());
 
 	// hide input StaticMeshComponent
-	ComponentTarget->SetOwnerVisibility(false);
+	TargetComponent->SetOwnerVisibility(false);
 
 	// initialize our properties
 	ToolPropertyObjects.Add(this);
@@ -94,15 +86,15 @@ void UWeldMeshEdgesTool::Shutdown(EToolShutdownType ShutdownType)
 {
 	if (DynamicMeshComponent != nullptr)
 	{
-		ComponentTarget->SetOwnerVisibility(true);
+		Cast<IPrimitiveComponentBackedTarget>(Target)->SetOwnerVisibility(true);
 
 		if (ShutdownType == EToolShutdownType::Accept)
 		{
 			// this block bakes the modified DynamicMeshComponent back into the StaticMeshComponent inside an undo transaction
 			GetToolManager()->BeginUndoTransaction(LOCTEXT("WeldMeshEdgesToolTransactionName", "Remesh Mesh"));
-			ComponentTarget->CommitMesh([=](const FPrimitiveComponentTarget::FCommitParams& CommitParams)
+			Cast<IMeshDescriptionCommitter>(Target)->CommitMeshDescription([this](const IMeshDescriptionCommitter::FCommitterParams& CommitParams)
 			{
-				DynamicMeshComponent->Bake(CommitParams.MeshDescription, true);
+				DynamicMeshComponent->Bake(CommitParams.MeshDescriptionOut, true);
 			});
 			GetToolManager()->EndUndoTransaction();
 		}
@@ -119,7 +111,7 @@ void UWeldMeshEdgesTool::Render(IToolsContextRenderAPI* RenderAPI)
 	UpdateResult();
 
 	FPrimitiveDrawInterface* PDI = RenderAPI->GetPrimitiveDrawInterface();
-	FTransform Transform = ComponentTarget->GetWorldTransform(); //Actor->GetTransform();
+	FTransform Transform = Cast<IPrimitiveComponentBackedTarget>(Target)->GetWorldTransform(); //Actor->GetTransform();
 
 	FColor LineColor(200, 200, 200);
 	float PDIScale = RenderAPI->GetCameraState().GetPDIScalingFactor();
diff --git a/Engine/Plugins/Experimental/MeshModelingToolset/Source/MeshModelingTools/Public/ConvertToPolygonsTool.h b/Engine/Plugins/Experimental/MeshModelingToolset/Source/MeshModelingTools/Public/ConvertToPolygonsTool.h
--- a/Engine/Plugins/Experimental/MeshModelingToolset/Source/MeshModelingTools/Public/ConvertToPolygonsTool.h
+++ b/Engine/Plugins/Experimental/MeshModelingToolset/Source/MeshModelingTools/Public/ConvertToPolygonsTool.h
@@ -11,7 +11,7 @@
 #include "PreviewMesh.h"
 #include "ModelingOperators.h"
 #include "MeshOpPreviewHelpers.h"
-#include "SingleSelectionMeshEditingTool.h"
+#include "BaseTools/SingleSelectionMeshEditingTool.h"
 #include "ConvertToPolygonsTool.generated.h"
 
 // predeclaration
diff --git a/Engine/Plugins/Experimental/MeshModelingToolset/Source/MeshModelingTools/Public/DisplaceMeshTool.h b/Engine/Plugins/Experimental/MeshModelingToolset/Source/MeshModelingTools/Public/DisplaceMeshTool.h
--- a/Engine/Plugins/Experimental/MeshModelingToolset/Source/MeshModelingTools/Public/DisplaceMeshTool.h
+++ b/Engine/Plugins/Experimental/MeshModelingToolset/Source/MeshModelingTools/Public/DisplaceMeshTool.h
@@ -12,7 +12,7 @@
 #include "MeshOpPreviewHelpers.h"
 #include "Spatial/SampledScalarField2.h"
 #include "WeightMapUtil.h"
-#include "SingleSelectionMeshEditingTool.h"
+#include "BaseTools/SingleSelectionMeshEditingTool.h"
 #include "DisplaceMeshTool.generated.h"
 
 struct FMeshDescription;
diff --git a/Engine/Plugins/Experimental/MeshModelingToolset/Source/MeshModelingTools/Public/EdgeLoopInsertionTool.h b/Engine/Plugins/Experimental/MeshModelingToolset/Source/MeshModelingTools/Public/EdgeLoopInsertionTool.h
--- a/Engine/Plugins/Experimental/MeshModelingToolset/Source/MeshModelingTools/Public/EdgeLoopInsertionTool.h
+++ b/Engine/Plugins/Experimental/MeshModelingToolset/Source/MeshModelingTools/Public/EdgeLoopInsertionTool.h
@@ -13,7 +13,7 @@
 #include "Selection/GroupTopologySelector.h"
 #include "SingleSelectionTool.h"
 #include "ToolDataVisualizer.h"
-#include "SingleSelectionMeshEditingTool.h"
+#include "BaseTools/SingleSelectionMeshEditingTool.h"
 
 #include "EdgeLoopInsertionTool.generated.h"
 
diff --git a/Engine/Plugins/Experimental/MeshModelingToolset/Source/MeshModelingTools/Public/GroupEdgeInsertionTool.h b/Engine/Plugins/Experimental/MeshModelingToolset/Source/MeshModelingTools/Public/GroupEdgeInsertionTool.h
--- a/Engine/Plugins/Experimental/MeshModelingToolset/Source/MeshModelingTools/Public/GroupEdgeInsertionTool.h
+++ b/Engine/Plugins/Experimental/MeshModelingToolset/Source/MeshModelingTools/Public/GroupEdgeInsertionTool.h
@@ -14,7 +14,7 @@
 #include "Selection/GroupTopologySelector.h"
 #include "SingleSelectionTool.h"
 #include "ToolDataVisualizer.h"
-#include "SingleSelectionMeshEditingTool.h"
+#include "BaseTools/SingleSelectionMeshEditingTool.h"
 
 #include "GroupEdgeInsertionTool.generated.h"
 
diff --git a/Engine/Plugins/Experimental/MeshModelingToolset/Source/MeshModelingTools/Public/HoleFillTool.h b/Engine/Plugins/Experimental/MeshModelingToolset/Source/MeshModelingTools/Public/HoleFillTool.h
--- a/Engine/Plugins/Experimental/MeshModelingToolset/Source/MeshModelingTools/Public/HoleFillTool.h
+++ b/Engine/Plugins/Experimental/MeshModelingToolset/Source/MeshModelingTools/Public/HoleFillTool.h
@@ -12,6 +12,7 @@
 #include "BaseBehaviors/BehaviorTargetInterfaces.h"
 #include "DynamicMeshAABBTree3.h"
 #include "GroupTopology.h"
+#include "BaseTools/SingleSelectionMeshEditingTool.h"
 #include "HoleFillTool.generated.h"
 
 class UPolygonSelectionMechanic;
@@ -26,13 +27,12 @@ class UHoleFillTool;
  */
 
 UCLASS()
-class MESHMODELINGTOOLS_API UHoleFillToolBuilder : public UInteractiveToolBuilder
+class MESHMODELINGTOOLS_API UHoleFillToolBuilder : public USingleSelectionMeshEditingToolBuilder
 {
 	GENERATED_BODY()
 
 public:
-	bool CanBuildTool(const FToolBuilderState& SceneState) const override;
-	UInteractiveTool* BuildTool(const FToolBuilderState& SceneState) const override;
+	virtual USingleSelectionMeshEditingTool* CreateNewTool(const FToolBuilderState& SceneState) const override;
 };
 
 /*
@@ -211,7 +211,7 @@ public:
  */
 
 UCLASS()
-class MESHMODELINGTOOLS_API UHoleFillTool : public USingleSelectionTool
+class MESHMODELINGTOOLS_API UHoleFillTool : public USingleSelectionMeshEditingTool
 {
 	GENERATED_BODY()
 
@@ -260,16 +260,11 @@ protected:
 	// UV Scale factor is cached based on the bounding box of the mesh before any fills are performed
 	float MeshUVScaleFactor = 0.0f;
 
-	// World in which to create the Preview mesh actor
-	UWorld* TargetWorld = nullptr;
-
 	// Used for hit querying
 	UE::Geometry::FDynamicMeshAABBTree3 MeshSpatial;
 
 	TSet<int32> NewTriangleIDs;
 
-	void SetWorld(UWorld* World);
-
 	// Create the Preview object
 	void SetupPreview();
 
diff --git a/Engine/Plugins/Experimental/MeshModelingToolset/Source/MeshModelingTools/Public/LatticeDeformerTool.h b/Engine/Plugins/Experimental/MeshModelingToolset/Source/MeshModelingTools/Public/LatticeDeformerTool.h
--- a/Engine/Plugins/Experimental/MeshModelingToolset/Source/MeshModelingTools/Public/LatticeDeformerTool.h
+++ b/Engine/Plugins/Experimental/MeshModelingToolset/Source/MeshModelingTools/Public/LatticeDeformerTool.h
@@ -8,19 +8,19 @@
 #include "InteractiveToolChange.h" //FToolCommandChange
 #include "SingleSelectionTool.h" //USingleSelectionTool
 #include "Operations/FFDLattice.h"
+#include "BaseTools/SingleSelectionMeshEditingTool.h"
 #include "LatticeDeformerTool.generated.h"
 
 class ULatticeControlPointsMechanic;
 class UMeshOpPreviewWithBackgroundCompute;
 
 UCLASS()
-class MESHMODELINGTOOLS_API ULatticeDeformerToolBuilder : public UInteractiveToolBuilder
+class MESHMODELINGTOOLS_API ULatticeDeformerToolBuilder : public USingleSelectionMeshEditingToolBuilder
 {
 	GENERATED_BODY()
 
 public:
-	virtual bool CanBuildTool(const FToolBuilderState& SceneState) const override;
-	virtual UInteractiveTool* BuildTool(const FToolBuilderState& SceneState) const override;
+	virtual USingleSelectionMeshEditingTool* CreateNewTool(const FToolBuilderState& SceneState) const override;
 };
 
 
@@ -97,7 +97,7 @@ public:
 
 /** Deform a mesh using a regular hexahedral lattice */
 UCLASS()
-class MESHMODELINGTOOLS_API ULatticeDeformerTool : public USingleSelectionTool
+class MESHMODELINGTOOLS_API ULatticeDeformerTool : public USingleSelectionMeshEditingTool
 {
 	GENERATED_BODY()
 
@@ -105,9 +105,6 @@ public:
 
 	virtual void DrawHUD(FCanvas* Canvas, IToolsContextRenderAPI* RenderAPI) override;
 
-	virtual void SetWorld(UWorld* World) { TargetWorld = World; }
-	virtual void SetAssetAPI(IToolsContextAssetAPI* NewAssetApi) { AssetAPI = NewAssetApi; }
-
 	virtual bool HasCancel() const override { return true; }
 	virtual bool HasAccept() const override { return true; }
 	virtual bool CanAccept() const override;
@@ -122,9 +119,6 @@ public:
 
 protected:
 
-	UWorld* TargetWorld;
-	IToolsContextAssetAPI* AssetAPI;
-
 	// Input mesh
 	TSharedPtr<UE::Geometry::FDynamicMesh3, ESPMode::ThreadSafe> OriginalMesh;
 
diff --git a/Engine/Plugins/Experimental/MeshModelingToolset/Source/MeshModelingTools/Public/MeshBoundaryToolBase.h b/Engine/Plugins/Experimental/MeshModelingToolset/Source/MeshModelingTools/Public/MeshBoundaryToolBase.h
--- a/Engine/Plugins/Experimental/MeshModelingToolset/Source/MeshModelingTools/Public/MeshBoundaryToolBase.h
+++ b/Engine/Plugins/Experimental/MeshModelingToolset/Source/MeshModelingTools/Public/MeshBoundaryToolBase.h
@@ -8,7 +8,7 @@
 #include "BaseBehaviors/BehaviorTargetInterfaces.h"
 #include "DynamicMeshAABBTree3.h"
 #include "GroupTopology.h"
-#include "SingleSelectionTool.h"
+#include "BaseTools/SingleSelectionMeshEditingTool.h"
 
 #include "MeshBoundaryToolBase.generated.h"
 
@@ -22,12 +22,11 @@ class USingleClickInputBehavior;
   *	TODO: We can refactor to make the HoleFiller tool inherit from this.
   */
 UCLASS()
-class MESHMODELINGTOOLS_API UMeshBoundaryToolBase : public USingleSelectionTool
+class MESHMODELINGTOOLS_API UMeshBoundaryToolBase : public USingleSelectionMeshEditingTool
 {
 	GENERATED_BODY()
 
 public:
-	virtual void SetWorld(UWorld* World) { TargetWorld = World; }
 
 	virtual bool HasCancel() const override { return true; }
 	virtual bool HasAccept() const override { return true; }
@@ -40,8 +39,6 @@ public:
 
 protected:
 
-	UWorld* TargetWorld;
-
 	TSharedPtr<FDynamicMesh3, ESPMode::ThreadSafe> OriginalMesh;
 
 	// Used for hit querying
diff --git a/Engine/Plugins/Experimental/MeshModelingToolset/Source/MeshModelingTools/Public/MeshInspectorTool.h b/Engine/Plugins/Experimental/MeshModelingToolset/Source/MeshModelingTools/Public/MeshInspectorTool.h
--- a/Engine/Plugins/Experimental/MeshModelingToolset/Source/MeshModelingTools/Public/MeshInspectorTool.h
+++ b/Engine/Plugins/Experimental/MeshModelingToolset/Source/MeshModelingTools/Public/MeshInspectorTool.h
@@ -11,6 +11,7 @@
 #include "DynamicMeshAABBTree3.h"
 #include "Properties/MeshMaterialProperties.h"
 #include "PreviewMesh.h"
+#include "BaseTools/SingleSelectionMeshEditingTool.h"
 #include "MeshInspectorTool.generated.h"
 
 
@@ -22,13 +23,12 @@ class UMaterialInstanceDynamic;
  *
  */
 UCLASS()
-class MESHMODELINGTOOLS_API UMeshInspectorToolBuilder : public UInteractiveToolBuilder
+class MESHMODELINGTOOLS_API UMeshInspectorToolBuilder : public USingleSelectionMeshEditingToolBuilder
 {
 	GENERATED_BODY()
 
 public:
-	virtual bool CanBuildTool(const FToolBuilderState& SceneState) const override;
-	virtual UInteractiveTool* BuildTool(const FToolBuilderState& SceneState) const override;
+	virtual USingleSelectionMeshEditingTool* CreateNewTool(const FToolBuilderState& SceneState) const override;
 };
 
 
@@ -90,7 +90,7 @@ public:
  * Mesh Inspector Tool for visualizing mesh information
  */
 UCLASS()
-class MESHMODELINGTOOLS_API UMeshInspectorTool : public USingleSelectionTool
+class MESHMODELINGTOOLS_API UMeshInspectorTool : public USingleSelectionMeshEditingTool
 {
 	GENERATED_BODY()
 
diff --git a/Engine/Plugins/Experimental/MeshModelingToolset/Source/MeshModelingTools/Public/MeshSpaceDeformerTool.h b/Engine/Plugins/Experimental/MeshModelingToolset/Source/MeshModelingTools/Public/MeshSpaceDeformerTool.h
--- a/Engine/Plugins/Experimental/MeshModelingToolset/Source/MeshModelingTools/Public/MeshSpaceDeformerTool.h
+++ b/Engine/Plugins/Experimental/MeshModelingToolset/Source/MeshModelingTools/Public/MeshSpaceDeformerTool.h
@@ -13,6 +13,7 @@
 #include "SimpleDynamicMeshComponent.h"
 #include "ToolDataVisualizer.h"
 #include "BaseGizmos/GizmoInterfaces.h"
+#include "BaseTools/SingleSelectionMeshEditingTool.h"
 
 #include "MeshSpaceDeformerTool.generated.h"
 
@@ -32,17 +33,12 @@ class FSelectClickedAction;
  * ToolBuilder
  */
 UCLASS()
-class MESHMODELINGTOOLS_API UMeshSpaceDeformerToolBuilder : public UInteractiveToolBuilder
+class MESHMODELINGTOOLS_API UMeshSpaceDeformerToolBuilder : public USingleSelectionMeshEditingToolBuilder
 {
 	GENERATED_BODY()
 
 public:
-	UMeshSpaceDeformerToolBuilder()
-	{
-	}
-
-	virtual bool CanBuildTool(const FToolBuilderState& SceneState) const override;
-	virtual UInteractiveTool* BuildTool(const FToolBuilderState& SceneState) const override;
+	virtual USingleSelectionMeshEditingTool* CreateNewTool(const FToolBuilderState& SceneState) const override;
 };
 
 /** ENonlinearOperation determines which type of nonlinear deformation will be applied*/
@@ -193,7 +189,7 @@ public:
  * Applies non-linear deformations to a mesh 
  */
 UCLASS()
-class MESHMODELINGTOOLS_API UMeshSpaceDeformerTool : public USingleSelectionTool
+class MESHMODELINGTOOLS_API UMeshSpaceDeformerTool : public USingleSelectionMeshEditingTool
 {
 	GENERATED_BODY()
 public:
@@ -202,9 +198,6 @@ public:
 	virtual void Setup() override;
 	virtual void Shutdown(EToolShutdownType ShutdownType) override;
 
-	virtual void SetWorld(UWorld* World);
-	virtual void SetAssetAPI(IToolsContextAssetAPI* AssetAPI);
-
 	void RequestAction(EMeshSpaceDeformerToolAction ActionType);
 
 	virtual void OnTick(float DeltaTime) override;
@@ -245,9 +238,6 @@ protected:
 	UPROPERTY()
 	UPreviewMesh* OriginalMeshPreview;
 
-	UWorld* TargetWorld;
-	IToolsContextAssetAPI* AssetAPI;
-
 	// Stores the materials of the original object, used to make the object semi-transparent.
 	FComponentMaterialSet StoredMaterialSet;
 
diff --git a/Engine/Plugins/Experimental/MeshModelingToolset/Source/MeshModelingTools/Public/Physics/ExtractCollisionGeometryTool.h b/Engine/Plugins/Experimental/MeshModelingToolset/Source/MeshModelingTools/Public/Physics/ExtractCollisionGeometryTool.h
--- a/Engine/Plugins/Experimental/MeshModelingToolset/Source/MeshModelingTools/Public/Physics/ExtractCollisionGeometryTool.h
+++ b/Engine/Plugins/Experimental/MeshModelingToolset/Source/MeshModelingTools/Public/Physics/ExtractCollisionGeometryTool.h
@@ -7,6 +7,7 @@
 #include "InteractiveToolBuilder.h"
 #include "DynamicMesh3.h"
 #include "Physics/CollisionPropertySets.h"
+#include "BaseTools/SingleSelectionMeshEditingTool.h"
 #include "ExtractCollisionGeometryTool.generated.h"
 
 class UPreviewGeometry;
@@ -14,15 +15,16 @@ class UPreviewMesh;
 class IAssetGenerationAPI;
 
 UCLASS()
-class MESHMODELINGTOOLS_API UExtractCollisionGeometryToolBuilder : public UInteractiveToolBuilder
+class MESHMODELINGTOOLS_API UExtractCollisionGeometryToolBuilder : public USingleSelectionMeshEditingToolBuilder
 {
 	GENERATED_BODY()
 
 public:
-	IAssetGenerationAPI* AssetAPI = nullptr;
+	virtual USingleSelectionMeshEditingTool* CreateNewTool(const FToolBuilderState& SceneState) const override;
+	virtual bool RequiresAssetAPI() const override { return true; }
 
-	virtual bool CanBuildTool(const FToolBuilderState& SceneState) const override;
-	virtual UInteractiveTool* BuildTool(const FToolBuilderState& SceneState) const override;
+protected:
+	virtual const FToolTargetTypeRequirements& GetTargetRequirements() const override;
 };
 
 
@@ -32,13 +34,10 @@ public:
  * Mesh Inspector Tool for visualizing mesh information
  */
 UCLASS()
-class MESHMODELINGTOOLS_API UExtractCollisionGeometryTool : public USingleSelectionTool
+class MESHMODELINGTOOLS_API UExtractCollisionGeometryTool : public USingleSelectionMeshEditingTool
 {
 	GENERATED_BODY()
 public:
-	virtual void SetWorld(UWorld* World) { this->TargetWorld = World; }
-	virtual void SetAssetAPI(IAssetGenerationAPI* InAssetAPI) { this->AssetAPI = InAssetAPI; }
-
 	virtual void Setup() override;
 	virtual void Shutdown(EToolShutdownType ShutdownType) override;
 
@@ -66,9 +65,6 @@ protected:
 	// these are TSharedPtr because TPimplPtr cannot currently be added to a TArray?
 	TSharedPtr<FPhysicsDataCollection> PhysicsInfo;
 
-	UWorld* TargetWorld = nullptr;
-	IAssetGenerationAPI* AssetAPI = nullptr;
-
 	UE::Geometry::FDynamicMesh3 CurrentMesh;
 	bool bResultValid = false;
 	void RecalculateMesh();
diff --git a/Engine/Plugins/Experimental/MeshModelingToolset/Source/MeshModelingTools/Public/RevolveBoundaryTool.h b/Engine/Plugins/Experimental/MeshModelingToolset/Source/MeshModelingTools/Public/RevolveBoundaryTool.h
--- a/Engine/Plugins/Experimental/MeshModelingToolset/Source/MeshModelingTools/Public/RevolveBoundaryTool.h
+++ b/Engine/Plugins/Experimental/MeshModelingToolset/Source/MeshModelingTools/Public/RevolveBoundaryTool.h
@@ -24,15 +24,13 @@ class IAssetGenerationAPI;
 // Tool Builder
 
 UCLASS()
-class MESHMODELINGTOOLS_API URevolveBoundaryToolBuilder : public UInteractiveToolBuilder
+class MESHMODELINGTOOLS_API URevolveBoundaryToolBuilder : public USingleSelectionMeshEditingToolBuilder
 {
 	GENERATED_BODY()
 
 public:
-	IAssetGenerationAPI* AssetAPI = nullptr;
-
-	bool CanBuildTool(const FToolBuilderState& SceneState) const override;
-	UInteractiveTool* BuildTool(const FToolBuilderState& SceneState) const override;
+	virtual USingleSelectionMeshEditingTool* CreateNewTool(const FToolBuilderState& SceneState) const override;
+	virtual bool RequiresAssetAPI() const override { return true; }
 };
 
 UCLASS()
@@ -86,7 +84,6 @@ class MESHMODELINGTOOLS_API URevolveBoundaryTool : public UMeshBoundaryToolBase,
 	GENERATED_BODY()
 
 public:
-	virtual void SetAssetAPI(IAssetGenerationAPI* NewAssetApi) { AssetAPI = NewAssetApi; }
 
 	virtual bool CanAccept() const override;
 
@@ -104,8 +101,6 @@ protected:
 	bool bAlignAxisOnClick = false;
 	int32 AlignAxisModifier = 2;
 
-	IAssetGenerationAPI* AssetAPI = nullptr;
-
 	UPROPERTY()
 	URevolveBoundaryToolProperties* Settings = nullptr;
 
diff --git a/Engine/Plugins/Experimental/MeshModelingToolset/Source/MeshModelingTools/Public/SingleSelectionMeshEditingTool.h b/Engine/Plugins/Experimental/MeshModelingToolset/Source/MeshModelingTools/Public/SingleSelectionMeshEditingTool.h
deleted file mode 100644
diff --git a/Engine/Plugins/Experimental/MeshModelingToolset/Source/MeshModelingTools/Public/WeldMeshEdgesTool.h b/Engine/Plugins/Experimental/MeshModelingToolset/Source/MeshModelingTools/Public/WeldMeshEdgesTool.h
--- a/Engine/Plugins/Experimental/MeshModelingToolset/Source/MeshModelingTools/Public/WeldMeshEdgesTool.h
+++ b/Engine/Plugins/Experimental/MeshModelingToolset/Source/MeshModelingTools/Public/WeldMeshEdgesTool.h
@@ -8,6 +8,7 @@
 #include "InteractiveToolBuilder.h"
 #include "DynamicMesh3.h"
 #include "DynamicMeshAABBTree3.h"
+#include "BaseTools/SingleSelectionMeshEditingTool.h"
 #include "WeldMeshEdgesTool.generated.h"
 
 
@@ -19,20 +20,19 @@ class USimpleDynamicMeshComponent;
  *
  */
 UCLASS()
-class MESHMODELINGTOOLS_API UWeldMeshEdgesToolBuilder : public UInteractiveToolBuilder
+class MESHMODELINGTOOLS_API UWeldMeshEdgesToolBuilder : public USingleSelectionMeshEditingToolBuilder
 {
 	GENERATED_BODY()
 
 public:
-	virtual bool CanBuildTool(const FToolBuilderState& SceneState) const override;
-	virtual UInteractiveTool* BuildTool(const FToolBuilderState& SceneState) const override;
+	virtual USingleSelectionMeshEditingTool* CreateNewTool(const FToolBuilderState& SceneState) const override;
 };
 
 /**
  * Mesh Weld Edges Tool
  */
 UCLASS()
-class MESHMODELINGTOOLS_API UWeldMeshEdgesTool : public USingleSelectionTool
+class MESHMODELINGTOOLS_API UWeldMeshEdgesTool : public USingleSelectionMeshEditingTool
 {
 	GENERATED_BODY()
 
diff --git a/Engine/Plugins/Experimental/MeshModelingToolset/Source/MeshModelingToolsEditorOnly/Private/EditMeshMaterialsTool.cpp b/Engine/Plugins/Experimental/MeshModelingToolset/Source/MeshModelingToolsEditorOnly/Private/EditMeshMaterialsTool.cpp
--- a/Engine/Plugins/Experimental/MeshModelingToolset/Source/MeshModelingToolsEditorOnly/Private/EditMeshMaterialsTool.cpp
+++ b/Engine/Plugins/Experimental/MeshModelingToolset/Source/MeshModelingToolsEditorOnly/Private/EditMeshMaterialsTool.cpp
@@ -15,6 +15,10 @@
 #include "AssetGenerationUtil.h"
 #include "ToolSetupUtil.h"
 
+#include "TargetInterfaces/MaterialProvider.h"
+#include "TargetInterfaces/MeshDescriptionCommitter.h"
+#include "TargetInterfaces/PrimitiveComponentBackedTarget.h"
+
 #include "ExplicitUseGeometryMathTypes.h"		// using UE::Geometry::(math types)
 using namespace UE::Geometry;
 
@@ -52,8 +56,9 @@ void UEditMeshMaterialsTool::Setup()
 
 	PreviewMesh->ClearOverrideRenderMaterial();
 
+	IMaterialProvider* TargetMaterial = Cast<IMaterialProvider>(Target);
 	FComponentMaterialSet AssetMaterials;
-	ComponentTarget->GetMaterialSet(AssetMaterials, true);
+	TargetMaterial->GetMaterialSet(AssetMaterials, true);
 	MaterialProps->Materials = AssetMaterials.Materials;
 	CurrentMaterials = MaterialProps->Materials;
 	InitialMaterialKey = GetMaterialKey();
@@ -63,7 +68,7 @@ void UEditMeshMaterialsTool::Setup()
 		[this](FMaterialSetKey NewKey) { OnMaterialSetChanged(); });
 
 	FComponentMaterialSet ComponentMaterials;
-	ComponentTarget->GetMaterialSet(ComponentMaterials, false);
+	TargetMaterial->GetMaterialSet(ComponentMaterials, false);
 	if (ComponentMaterials != AssetMaterials)
 	{
 		GetToolManager()->DisplayMessage(
@@ -224,7 +229,7 @@ void UEditMeshMaterialsTool::OnShutdown(EToolShutdownType ShutdownType)
 		{
 			FComponentMaterialSet NewMaterialSet;
 			NewMaterialSet.Materials = CurrentMaterials;
-			ComponentTarget->CommitMaterialSetUpdate(NewMaterialSet, true);
+			Cast<IMaterialProvider>(Target)->CommitMaterialSetUpdate(NewMaterialSet, true);
 		}
 
 		UMeshSelectionTool::OnShutdown(ShutdownType);
diff --git a/Engine/Plugins/Experimental/MeshModelingToolset/Source/MeshModelingToolsEditorOnly/Private/GenerateLODMeshesTool.cpp b/Engine/Plugins/Experimental/MeshModelingToolset/Source/MeshModelingToolsEditorOnly/Private/GenerateLODMeshesTool.cpp
--- a/Engine/Plugins/Experimental/MeshModelingToolset/Source/MeshModelingToolsEditorOnly/Private/GenerateLODMeshesTool.cpp
+++ b/Engine/Plugins/Experimental/MeshModelingToolset/Source/MeshModelingToolsEditorOnly/Private/GenerateLODMeshesTool.cpp
@@ -19,6 +19,9 @@
 
 #include "Modules/ModuleManager.h"
 
+#include "TargetInterfaces/MaterialProvider.h"
+#include "TargetInterfaces/MeshDescriptionProvider.h"
+#include "TargetInterfaces/PrimitiveComponentBackedTarget.h"
 
 #include "IMeshReductionManagerModule.h"
 #include "IMeshReductionInterfaces.h"
@@ -36,24 +39,9 @@ using namespace UE::Geometry;
 /*
  * ToolBuilder
  */
-bool UGenerateLODMeshesToolBuilder::CanBuildTool(const FToolBuilderState& SceneState) const
+USingleSelectionMeshEditingTool* UGenerateLODMeshesToolBuilder::CreateNewTool(const FToolBuilderState& SceneState) const
 {
-	return ToolBuilderUtil::CountComponents(SceneState, CanMakeComponentTarget) == 1;
-}
-
-UInteractiveTool* UGenerateLODMeshesToolBuilder::BuildTool(const FToolBuilderState& SceneState) const
-{
-	UGenerateLODMeshesTool* NewTool = NewObject<UGenerateLODMeshesTool>(SceneState.ToolManager);
-
-	UActorComponent* ActorComponent = ToolBuilderUtil::FindFirstComponent(SceneState, CanMakeComponentTarget);
-	auto* MeshComponent = Cast<UPrimitiveComponent>(ActorComponent);
-	check(MeshComponent != nullptr);
-
-	NewTool->SetSelection(MakeComponentTarget(MeshComponent));
-	NewTool->SetWorld(SceneState.World);
-	NewTool->SetAssetAPI(AssetAPI);
-
-	return NewTool;
+	return NewObject<UGenerateLODMeshesTool>(SceneState.ToolManager);
 }
 
 /*
@@ -88,16 +76,6 @@ UGenerateLODMeshesToolProperties::UGenerateLODMeshesToolProperties()
 	LODLevels.Add(LOD2);
 }
 
-void UGenerateLODMeshesTool::SetWorld(UWorld* World)
-{
-	this->TargetWorld = World;
-}
-
-void UGenerateLODMeshesTool::SetAssetAPI(IAssetGenerationAPI* AssetAPIIn)
-{
-	this->AssetAPI = AssetAPIIn;
-}
-
 void UGenerateLODMeshesTool::Setup()
 {
 	UInteractiveTool::Setup();
@@ -121,19 +99,20 @@ void UGenerateLODMeshesTool::Setup()
 #else
 		auto EnterProgressFrame = [](int Progress) {};
 #endif
-		OriginalMeshDescription = MakeShared<FMeshDescription, ESPMode::ThreadSafe>(*ComponentTarget->GetMesh());
+		IMeshDescriptionProvider* TargetMeshProvider = Cast<IMeshDescriptionProvider>(Target);
+		OriginalMeshDescription = MakeShared<FMeshDescription, ESPMode::ThreadSafe>(*TargetMeshProvider->GetMeshDescription());
 		EnterProgressFrame(1);
 		// aux-data isn't deep copied - by default it is built during initial evaluation (not thread safe)  
 		// so force aux-data rebuild now before multiple UE4 simplifiers try to use it in parallel.
 		OriginalMeshDescription->BuildIndexers();   
 		OriginalMesh = MakeShared<FDynamicMesh3, ESPMode::ThreadSafe>();
 		FMeshDescriptionToDynamicMesh Converter;
-		Converter.Convert(ComponentTarget->GetMesh(), *OriginalMesh);
+		Converter.Convert(TargetMeshProvider->GetMeshDescription(), *OriginalMesh);
 		EnterProgressFrame(2);
 		OriginalMeshSpatial = MakeShared<FDynamicMeshAABBTree3, ESPMode::ThreadSafe>(OriginalMesh.Get(), true);
 	}
 
-	WorldBounds = FAxisAlignedBox3d(ComponentTarget->GetOwnerActor()->GetComponentsBoundingBox());
+	WorldBounds = FAxisAlignedBox3d(Cast<IPrimitiveComponentBackedTarget>(Target)->GetOwnerActor()->GetComponentsBoundingBox());
 
 	// initialize our properties
 	SimplifyProperties = NewObject<UGenerateLODMeshesToolProperties>(this);
@@ -162,7 +141,7 @@ void UGenerateLODMeshesTool::Setup()
 void UGenerateLODMeshesTool::Shutdown(EToolShutdownType ShutdownType)
 {
 	SimplifyProperties->SaveProperties(this);
-	ComponentTarget->SetOwnerVisibility(true);
+	Cast<IPrimitiveComponentBackedTarget>(Target)->SetOwnerVisibility(true);
 
 	if (ShutdownType == EToolShutdownType::Accept)
 	{
@@ -206,7 +185,7 @@ void UGenerateLODMeshesTool::UpdateNumPreviews()
 
 	check(CurNumPreviews < NumPreviews);		// todo: support less
 
-	FTransform OrigTransform = ComponentTarget->GetWorldTransform();
+	FTransform OrigTransform = Cast<IPrimitiveComponentBackedTarget>(Target)->GetWorldTransform();
 	FVector WorldShift = (FVector)((WorldBounds.Width() * 1.1) * FVector3d::UnitX());
 
 	for (int32 k = CurNumPreviews; k < NumPreviews; ++k)
@@ -218,7 +197,7 @@ void UGenerateLODMeshesTool::UpdateNumPreviews()
 		NewPreview->Setup(this->TargetWorld, Factory.Get());
 
 		FComponentMaterialSet MaterialSet;
-		ComponentTarget->GetMaterialSet(MaterialSet);
+		Cast<IMaterialProvider>(Target)->GetMaterialSet(MaterialSet);
 		NewPreview->ConfigureMaterials(MaterialSet.Materials,
 			ToolSetupUtil::GetDefaultWorkingMaterial(GetToolManager())
 		);
@@ -320,7 +299,7 @@ TUniquePtr<UE::Geometry::FDynamicMeshOperator> FGenerateLODOperatorFactory::Make
 void UGenerateLODMeshesTool::Render(IToolsContextRenderAPI* RenderAPI)
 {
 	//FPrimitiveDrawInterface* PDI = RenderAPI->GetPrimitiveDrawInterface();
-	//FTransform Transform = ComponentTarget->GetWorldTransform(); //Actor->GetTransform();
+	//FTransform Transform = Cast<IPrimitiveComponentBackedTarget>(Target)->GetWorldTransform(); //Actor->GetTransform();
 
 	//FColor LineColor(255, 0, 0);
 	//const FDynamicMesh3* TargetMesh = Preview->PreviewMesh->GetPreviewDynamicMesh();
@@ -346,8 +325,9 @@ void UGenerateLODMeshesTool::Render(IToolsContextRenderAPI* RenderAPI)
 
 void UGenerateLODMeshesTool::UpdateVisualization()
 {
+	IMaterialProvider* TargetMaterial = Cast<IMaterialProvider>(Target);
 	FComponentMaterialSet MaterialSet;
-	ComponentTarget->GetMaterialSet(MaterialSet);
+	TargetMaterial->GetMaterialSet(MaterialSet);
 
 	for (UMeshOpPreviewWithBackgroundCompute* Preview : Previews)
 	{
@@ -369,7 +349,7 @@ void UGenerateLODMeshesTool::UpdateVisualization()
 	//}
 	//else
 	//{
-	//	ComponentTarget->GetMaterialSet(MaterialSet);
+	//	TargetMaterial->GetMaterialSet(MaterialSet);
 	//	Preview->PreviewMesh->ClearTriangleColorFunction(UPreviewMesh::ERenderUpdateMode::FastUpdate);
 	//}
 
@@ -393,7 +373,8 @@ void UGenerateLODMeshesTool::GenerateAssets()
 {
 	GetToolManager()->BeginUndoTransaction(LOCTEXT("GenerateLODMeshesToolTransactionName", "Simplify Mesh"));
 
-	FTransform Transform = ComponentTarget->GetWorldTransform();
+	IPrimitiveComponentBackedTarget* TargetComponent = Cast<IPrimitiveComponentBackedTarget>(Target);
+	FTransform Transform = TargetComponent->GetWorldTransform();
 
 	int32 NumLODs = Previews.Num();
 	for (int32 k = 0; k < NumLODs; ++k)
@@ -407,9 +388,9 @@ void UGenerateLODMeshesTool::GenerateAssets()
 		}
 
 		FComponentMaterialSet MaterialSet;
-		ComponentTarget->GetMaterialSet(MaterialSet);
+		Cast<IMaterialProvider>(Target)->GetMaterialSet(MaterialSet);
 
-		FString BaseName = AssetGenerationUtil::GetComponentAssetBaseName(ComponentTarget->GetOwnerComponent());
+		FString BaseName = AssetGenerationUtil::GetComponentAssetBaseName(TargetComponent->GetOwnerComponent());
 		FString Name = FString::Printf( TEXT("%s_LOD%d"), *BaseName, (SimplifyProperties->NameIndexBase+k) );
 
 		AActor* NewActor = AssetGenerationUtil::GenerateStaticMeshActor(
diff --git a/Engine/Plugins/Experimental/MeshModelingToolset/Source/MeshModelingToolsEditorOnly/Private/MeshTangentsTool.cpp b/Engine/Plugins/Experimental/MeshModelingToolset/Source/MeshModelingToolsEditorOnly/Private/MeshTangentsTool.cpp
--- a/Engine/Plugins/Experimental/MeshModelingToolset/Source/MeshModelingToolsEditorOnly/Private/MeshTangentsTool.cpp
+++ b/Engine/Plugins/Experimental/MeshModelingToolset/Source/MeshModelingToolsEditorOnly/Private/MeshTangentsTool.cpp
@@ -15,6 +15,12 @@
 #include "Engine/Classes/Engine/StaticMesh.h"
 #include "Engine/Classes/Components/StaticMeshComponent.h"
 
+#include "TargetInterfaces/MaterialProvider.h"
+#include "TargetInterfaces/MeshDescriptionCommitter.h"
+#include "TargetInterfaces/MeshDescriptionProvider.h"
+#include "TargetInterfaces/PrimitiveComponentBackedTarget.h"
+#include "TargetInterfaces/StaticMeshBackedTarget.h"
+
 #include "ExplicitUseGeometryMathTypes.h"		// using UE::Geometry::(math types)
 using namespace UE::Geometry;
 
@@ -24,24 +30,21 @@ using namespace UE::Geometry;
  * ToolBuilder
  */
 
-bool UMeshTangentsToolBuilder::CanBuildTool(const FToolBuilderState& SceneState) const
+const FToolTargetTypeRequirements& UMeshTangentsToolBuilder::GetTargetRequirements() const
 {
-	return (ToolBuilderUtil::CountComponents(SceneState, CanMakeComponentTarget) == 1)
-		&& (ToolBuilderUtil::CountComponents(SceneState, [&](UActorComponent* Comp) { return Cast<UStaticMeshComponent>(Comp) != nullptr; }) == 1);		
+	static FToolTargetTypeRequirements TypeRequirements({
+		UMaterialProvider::StaticClass(),
+		UMeshDescriptionCommitter::StaticClass(),
+		UMeshDescriptionProvider::StaticClass(),
+		UPrimitiveComponentBackedTarget::StaticClass(),
+		UStaticMeshBackedTarget::StaticClass()
+		});
+	return TypeRequirements;
 }
 
-UInteractiveTool* UMeshTangentsToolBuilder::BuildTool(const FToolBuilderState& SceneState) const
+USingleSelectionMeshEditingTool* UMeshTangentsToolBuilder::CreateNewTool(const FToolBuilderState& SceneState) const
 {
-	UMeshTangentsTool* NewTool = NewObject<UMeshTangentsTool>(SceneState.ToolManager);
-
-	UActorComponent* ActorComponent = ToolBuilderUtil::FindFirstComponent(SceneState, CanMakeComponentTarget);
-	check(Cast<UStaticMeshComponent>(ActorComponent) != nullptr);
-	auto* MeshComponent = Cast<UPrimitiveComponent>(ActorComponent);
-	check(MeshComponent != nullptr);
-	NewTool->SetSelection(MakeComponentTarget(MeshComponent));
-	NewTool->SetWorld(SceneState.World);
-
-	return NewTool;
+	return NewObject<UMeshTangentsTool>(SceneState.ToolManager);
 }
 
 /*
@@ -51,11 +54,6 @@ UMeshTangentsTool::UMeshTangentsTool()
 {
 }
 
-void UMeshTangentsTool::SetWorld(UWorld* World)
-{
-	this->TargetWorld = World;
-}
-
 
 
 void UMeshTangentsTool::Setup()
@@ -63,20 +61,21 @@ void UMeshTangentsTool::Setup()
 	UInteractiveTool::Setup();
 
 	// copy input MeshDescription and make sure it has initialized normals/tangents
-	InputMeshDescription = MakeShared<FMeshDescription, ESPMode::ThreadSafe>(*ComponentTarget->GetMesh());
-	UE::MeshDescription::InitializeAutoGeneratedAttributes(*InputMeshDescription, ComponentTarget->GetOwnerComponent(), 0);
+	IPrimitiveComponentBackedTarget* TargetComponent = Cast<IPrimitiveComponentBackedTarget>(Target);
+	InputMeshDescription = MakeShared<FMeshDescription, ESPMode::ThreadSafe>(*Cast<IMeshDescriptionProvider>(Target)->GetMeshDescription());
+	UE::MeshDescription::InitializeAutoGeneratedAttributes(*InputMeshDescription, TargetComponent->GetOwnerComponent(), 0);
 
 	// hide input StaticMeshComponent
-	ComponentTarget->SetOwnerVisibility(false);
+	TargetComponent->SetOwnerVisibility(false);
 
 	// make our preview mesh
 	PreviewMesh = NewObject<UPreviewMesh>(this);
 	PreviewMesh->bBuildSpatialDataStructure = false;
-	PreviewMesh->CreateInWorld(ComponentTarget->GetOwnerActor()->GetWorld(), FTransform::Identity);
-	PreviewMesh->SetTransform(ComponentTarget->GetWorldTransform());
+	PreviewMesh->CreateInWorld(TargetComponent->GetOwnerActor()->GetWorld(), FTransform::Identity);
+	PreviewMesh->SetTransform(TargetComponent->GetWorldTransform());
 	// configure materials
 	FComponentMaterialSet MaterialSet;
-	ComponentTarget->GetMaterialSet(MaterialSet);
+	Cast<IMaterialProvider>(Target)->GetMaterialSet(MaterialSet);
 	PreviewMesh->SetMaterials(MaterialSet.Materials);
 	// initialize with meshdescription created above, that has valid tangents
 	PreviewMesh->SetTangentsMode(EDynamicMeshTangentCalcType::ExternallyCalculated);
@@ -99,7 +98,7 @@ void UMeshTangentsTool::Setup()
 	Settings->WatchProperty(Settings->bShowNormals, [this](float) { bVisibilityChanged = true; });
 
 	PreviewGeometry = NewObject<UPreviewGeometry>(this);
-	PreviewGeometry->CreateInWorld(ComponentTarget->GetOwnerActor()->GetWorld(), PreviewMesh->GetTransform());
+	PreviewGeometry->CreateInWorld(TargetComponent->GetOwnerActor()->GetWorld(), PreviewMesh->GetTransform());
 
 	Compute = MakeUnique<TGenericDataBackgroundCompute<FMeshTangentsd>>();
 	Compute->Setup(this);
@@ -122,14 +121,15 @@ void UMeshTangentsTool::Shutdown(EToolShutdownType ShutdownType)
 	Settings->SaveProperties(this);
 
 	// Restore (unhide) the source meshes
-	ComponentTarget->SetOwnerVisibility(true);
+	IPrimitiveComponentBackedTarget* TargetComponent = Cast<IPrimitiveComponentBackedTarget>(Target);
+	TargetComponent->SetOwnerVisibility(true);
 
 	TUniquePtr<FMeshTangentsd> Tangents = Compute->Shutdown();
 	if (ShutdownType == EToolShutdownType::Accept)
 	{
 		GetToolManager()->BeginUndoTransaction(LOCTEXT("UpdateTangents", "Update Tangents"));
 
-		UStaticMeshComponent* StaticMeshComponent = Cast<UStaticMeshComponent>(ComponentTarget->GetOwnerComponent());
+		UStaticMeshComponent* StaticMeshComponent = Cast<UStaticMeshComponent>(TargetComponent->GetOwnerComponent());
 		UStaticMesh* StaticMesh = StaticMeshComponent->GetStaticMesh();
 		StaticMesh->Modify();
 
@@ -139,16 +139,16 @@ void UMeshTangentsTool::Shutdown(EToolShutdownType ShutdownType)
 		SettingsChange.AutoGeneratedTangents = UE::MeshDescription::EBuildSettingBoolChange::Disable;
 		UE::MeshDescription::ConfigureBuildSettings(StaticMesh, 0, SettingsChange);
 
-		ComponentTarget->CommitMesh([&](const FPrimitiveComponentTarget::FCommitParams& CommitParams)
+		Cast<IMeshDescriptionCommitter>(Target)->CommitMeshDescription([&](const IMeshDescriptionCommitter::FCommitterParams& CommitParams)
 		{
 			FDynamicMeshToMeshDescription Converter;
 			{
 				const bool bCopyNormals = true; 
 				const bool bCopyOverlayTangents = false;
 				const bool bCopyOverlayUVs = false;
-				Converter.UpdateAttributes(InputMesh.Get(), *CommitParams.MeshDescription, bCopyNormals, bCopyOverlayTangents, bCopyOverlayUVs);
+				Converter.UpdateAttributes(InputMesh.Get(), *CommitParams.MeshDescriptionOut, bCopyNormals, bCopyOverlayTangents, bCopyOverlayUVs);
 			}
-			Converter.UpdateTangents(InputMesh.Get(), *CommitParams.MeshDescription, Tangents.Get());
+			Converter.UpdateTangents(InputMesh.Get(), *CommitParams.MeshDescriptionOut, Tangents.Get());
 		});
 
 		GetToolManager()->EndUndoTransaction();
diff --git a/Engine/Plugins/Experimental/MeshModelingToolset/Source/MeshModelingToolsEditorOnly/Private/MeshToVolumeTool.cpp b/Engine/Plugins/Experimental/MeshModelingToolset/Source/MeshModelingToolsEditorOnly/Private/MeshToVolumeTool.cpp
--- a/Engine/Plugins/Experimental/MeshModelingToolset/Source/MeshModelingToolsEditorOnly/Private/MeshToVolumeTool.cpp
+++ b/Engine/Plugins/Experimental/MeshModelingToolset/Source/MeshModelingToolsEditorOnly/Private/MeshToVolumeTool.cpp
@@ -32,32 +32,15 @@ using namespace UE::Geometry;
  * ToolBuilder
  */
 
-const FToolTargetTypeRequirements& UMeshToVolumeToolBuilder::GetTargetRequirements() const
-{
-	static FToolTargetTypeRequirements TypeRequirements({
-		UMeshDescriptionProvider::StaticClass(),
-		UPrimitiveComponentBackedTarget::StaticClass(),
-		UMaterialProvider::StaticClass()
-		});
-	return TypeRequirements;
-}
-
 bool UMeshToVolumeToolBuilder::CanBuildTool(const FToolBuilderState& SceneState) const
 {
 	// We don't want to allow this tool to run on selected volumes
-	return ToolBuilderUtil::CountSelectedActorsOfType<AVolume>(SceneState) == 0 
-		&& SceneState.TargetManager->CountSelectedAndTargetable(SceneState, GetTargetRequirements()) == 1;
+	return ToolBuilderUtil::CountSelectedActorsOfType<AVolume>(SceneState) == 0 && Super::CanBuildTool(SceneState);
 }
 
-UInteractiveTool* UMeshToVolumeToolBuilder::BuildTool(const FToolBuilderState& SceneState) const
+USingleSelectionMeshEditingTool* UMeshToVolumeToolBuilder::CreateNewTool(const FToolBuilderState& SceneState) const
 {
-	UMeshToVolumeTool* NewTool = NewObject<UMeshToVolumeTool>(SceneState.ToolManager);
-
-	UToolTarget* Target = SceneState.TargetManager->BuildFirstSelectedTargetable(SceneState, GetTargetRequirements());
-	check(Target);
-	NewTool->SetTarget(Target);
-
-	return NewTool;
+	return NewObject<UMeshToVolumeTool>(SceneState.ToolManager);
 }
 
 /*
@@ -137,7 +120,7 @@ void UMeshToVolumeTool::Shutdown(EToolShutdownType ShutdownType)
 	IPrimitiveComponentBackedTarget* TargetComponent = Cast<IPrimitiveComponentBackedTarget>(Target);
 	TargetComponent->SetOwnerVisibility(true);
 
-	UWorld* TargetWorld = TargetComponent->GetOwnerActor()->GetWorld();
+	UWorld* TargetOwnerWorld = TargetComponent->GetOwnerActor()->GetWorld();
 
 
 	if (ShutdownType == EToolShutdownType::Accept)
@@ -157,11 +140,11 @@ void UMeshToVolumeTool::Shutdown(EToolShutdownType ShutdownType)
 			UClass* VolumeClass = Settings->NewVolumeType.Get();
 			if (VolumeClass)
 			{
-				TargetVolume = (AVolume*)TargetWorld->SpawnActor(VolumeClass, &NewActorTransform, SpawnInfo);
+				TargetVolume = (AVolume*)TargetOwnerWorld->SpawnActor(VolumeClass, &NewActorTransform, SpawnInfo);
 			}
 			else
 			{
-				TargetVolume = TargetWorld->SpawnActor<ABlockingVolume>(FVector::ZeroVector, Rotation, SpawnInfo);
+				TargetVolume = TargetOwnerWorld->SpawnActor<ABlockingVolume>(FVector::ZeroVector, Rotation, SpawnInfo);
 			}
 			TargetVolume->BrushType = EBrushType::Brush_Add;
 			UModel* Model = NewObject<UModel>(TargetVolume);
diff --git a/Engine/Plugins/Experimental/MeshModelingToolset/Source/MeshModelingToolsEditorOnly/Private/ParameterizeMeshTool.cpp b/Engine/Plugins/Experimental/MeshModelingToolset/Source/MeshModelingToolsEditorOnly/Private/ParameterizeMeshTool.cpp
--- a/Engine/Plugins/Experimental/MeshModelingToolset/Source/MeshModelingToolsEditorOnly/Private/ParameterizeMeshTool.cpp
+++ b/Engine/Plugins/Experimental/MeshModelingToolset/Source/MeshModelingToolsEditorOnly/Private/ParameterizeMeshTool.cpp
@@ -16,6 +16,11 @@
 #include "ParameterizationOps/ParameterizeMeshOp.h"
 #include "ToolSetupUtil.h"
 
+#include "TargetInterfaces/MaterialProvider.h"
+#include "TargetInterfaces/MeshDescriptionCommitter.h"
+#include "TargetInterfaces/MeshDescriptionProvider.h"
+#include "TargetInterfaces/PrimitiveComponentBackedTarget.h"
+
 #include "ExplicitUseGeometryMathTypes.h"		// using UE::Geometry::(math types)
 using namespace UE::Geometry;
 
@@ -27,23 +32,10 @@ DEFINE_LOG_CATEGORY_STATIC(LogParameterizeMeshTool, Log, All);
  * ToolBuilder
  */
 
-bool UParameterizeMeshToolBuilder::CanBuildTool(const FToolBuilderState& SceneState) const
-{
-	return ToolBuilderUtil::CountComponents(SceneState, CanMakeComponentTarget) == 1;
-}
-
-UInteractiveTool* UParameterizeMeshToolBuilder::BuildTool(const FToolBuilderState& SceneState) const
+USingleSelectionMeshEditingTool* UParameterizeMeshToolBuilder::CreateNewTool(const FToolBuilderState& SceneState) const
 {
 	UParameterizeMeshTool* NewTool = NewObject<UParameterizeMeshTool>(SceneState.ToolManager);
-
-	UActorComponent* ActorComponent = ToolBuilderUtil::FindFirstComponent(SceneState, CanMakeComponentTarget);
-	auto* MeshComponent = Cast<UPrimitiveComponent>(ActorComponent);
-	check(MeshComponent != nullptr);
-	NewTool->SetSelection(MakeComponentTarget(MeshComponent));
-	NewTool->SetWorld(SceneState.World);
-	NewTool->SetAssetAPI(AssetAPI);
 	NewTool->SetUseAutoGlobalParameterizationMode(bDoAutomaticGlobalUnwrap);
-
 	return NewTool;
 }
 
@@ -54,16 +46,6 @@ UParameterizeMeshTool::UParameterizeMeshTool()
 {
 }
 
-void UParameterizeMeshTool::SetWorld(UWorld* World)
-{
-	this->TargetWorld = World;
-}
-
-void UParameterizeMeshTool::SetAssetAPI(IAssetGenerationAPI* AssetAPIIn)
-{
-	this->AssetAPI = AssetAPIIn;
-}
-
 void UParameterizeMeshTool::SetUseAutoGlobalParameterizationMode(bool bEnable)
 {
 	bDoAutomaticGlobalUnwrap = bEnable;
@@ -75,21 +57,24 @@ void UParameterizeMeshTool::Setup()
 	UInteractiveTool::Setup();
 
 	// Copy existing material if there is one	
-	DefaultMaterial = ComponentTarget->GetMaterial(0);
+	IMaterialProvider* TargetMaterial = Cast<IMaterialProvider>(Target);
+	DefaultMaterial = TargetMaterial->GetMaterial(0);
 	if (DefaultMaterial == nullptr)
 	{
 		DefaultMaterial = LoadObject<UMaterial>(nullptr, TEXT("/Engine/EngineMaterials/DefaultMaterial"));  
 	}
 
 	// hide input StaticMeshComponent
-	ComponentTarget->SetOwnerVisibility(false);
+	IPrimitiveComponentBackedTarget* TargetComponent = Cast<IPrimitiveComponentBackedTarget>(Target);
+	TargetComponent->SetOwnerVisibility(false);
 
 	// Construct the preview object and set the material on it
+	IMeshDescriptionProvider* TargetMeshProvider = Cast<IMeshDescriptionProvider>(Target);
 	Preview = NewObject<UMeshOpPreviewWithBackgroundCompute>(this, "Preview");
 	Preview->Setup(this->TargetWorld, this);
 	Preview->PreviewMesh->SetTangentsMode(EDynamicMeshTangentCalcType::AutoCalculated);
 
-	Preview->PreviewMesh->InitializeMesh(ComponentTarget->GetMesh());
+	Preview->PreviewMesh->InitializeMesh(TargetMeshProvider->GetMeshDescription());
 
 	Preview->OnMeshUpdated.AddLambda([this](UMeshOpPreviewWithBackgroundCompute* Op)
 	{
@@ -104,12 +89,12 @@ void UParameterizeMeshTool::Setup()
 		bHasGroups = FaceGroupUtil::HasMultipleGroups(*InputMesh);
 
 		FComponentMaterialSet MaterialSet;
-		ComponentTarget->GetMaterialSet(MaterialSet);
+		TargetMaterial->GetMaterialSet(MaterialSet);
 		Preview->ConfigureMaterials(MaterialSet.Materials,
 			ToolSetupUtil::GetDefaultWorkingMaterial(GetToolManager())
 		);
 
-		Preview->PreviewMesh->SetTransform(ComponentTarget->GetWorldTransform());
+		Preview->PreviewMesh->SetTransform(TargetComponent->GetWorldTransform());
 	}
 
 	if (bDoAutomaticGlobalUnwrap == false && bHasGroups == false)
@@ -124,7 +109,7 @@ void UParameterizeMeshTool::Setup()
 
 	UVChannelProperties = NewObject<UMeshUVChannelProperties>(this);
 	UVChannelProperties->RestoreProperties(this);
-	UVChannelProperties->Initialize(ComponentTarget->GetMesh(), false);
+	UVChannelProperties->Initialize(TargetMeshProvider->GetMeshDescription(), false);
 	UVChannelProperties->ValidateSelection(true);
 	UVChannelProperties->WatchProperty(UVChannelProperties->UVChannel, [this](const FString& NewValue) 
 	{
@@ -193,7 +178,7 @@ void UParameterizeMeshTool::Shutdown(EToolShutdownType ShutdownType)
 	FDynamicMeshOpResult Result = Preview->Shutdown();
 	
 	// Restore (unhide) the source meshes
-	ComponentTarget->SetOwnerVisibility(true);
+	Cast<IPrimitiveComponentBackedTarget>(Target)->SetOwnerVisibility(true);
 
 	if (ShutdownType == EToolShutdownType::Accept)
 	{
@@ -201,9 +186,9 @@ void UParameterizeMeshTool::Shutdown(EToolShutdownType ShutdownType)
 		check(DynamicMeshResult != nullptr);
 		GetToolManager()->BeginUndoTransaction(LOCTEXT("ParameterizeMesh", "Parameterize Mesh"));
 
-		ComponentTarget->CommitMesh([DynamicMeshResult](const FPrimitiveComponentTarget::FCommitParams& CommitParams)
+		Cast<IMeshDescriptionCommitter>(Target)->CommitMeshDescription([DynamicMeshResult](const IMeshDescriptionCommitter::FCommitterParams& CommitParams)
 		{
-			FMeshDescription* MeshDescription = CommitParams.MeshDescription;
+			FMeshDescription* MeshDescription = CommitParams.MeshDescriptionOut;
 
 			bool bVerticesOnly = false;
 			bool bAttributesOnly = true;
@@ -272,7 +257,7 @@ TUniquePtr<FDynamicMeshOperator> UParameterizeMeshTool::MakeNewOperator()
 
 	ParameterizeMeshOp->Method = (EParamOpBackend)(int)Settings->Method;
 
-	UE::Geometry::FTransform3d XForm3d(ComponentTarget->GetWorldTransform());
+	UE::Geometry::FTransform3d XForm3d(Cast<IPrimitiveComponentBackedTarget>(Target)->GetWorldTransform());
 	ParameterizeMeshOp->SetTransform(XForm3d);
 
 	return ParameterizeMeshOp;
diff --git a/Engine/Plugins/Experimental/MeshModelingToolset/Source/MeshModelingToolsEditorOnly/Private/PolygonOnMeshTool.cpp b/Engine/Plugins/Experimental/MeshModelingToolset/Source/MeshModelingToolsEditorOnly/Private/PolygonOnMeshTool.cpp
--- a/Engine/Plugins/Experimental/MeshModelingToolset/Source/MeshModelingToolsEditorOnly/Private/PolygonOnMeshTool.cpp
+++ b/Engine/Plugins/Experimental/MeshModelingToolset/Source/MeshModelingToolsEditorOnly/Private/PolygonOnMeshTool.cpp
@@ -15,6 +15,11 @@
 #include "MeshDescriptionToDynamicMesh.h"
 #include "DynamicMeshToMeshDescription.h"
 
+#include "TargetInterfaces/MaterialProvider.h"
+#include "TargetInterfaces/MeshDescriptionCommitter.h"
+#include "TargetInterfaces/MeshDescriptionProvider.h"
+#include "TargetInterfaces/PrimitiveComponentBackedTarget.h"
+
 #include "ExplicitUseGeometryMathTypes.h"		// using UE::Geometry::(math types)
 using namespace UE::Geometry;
 
@@ -25,22 +30,9 @@ using namespace UE::Geometry;
  */
 
 
-bool UPolygonOnMeshToolBuilder::CanBuildTool(const FToolBuilderState& SceneState) const
-{
-	return ToolBuilderUtil::CountComponents(SceneState, CanMakeComponentTarget) == 1;
-}
-
-UInteractiveTool* UPolygonOnMeshToolBuilder::BuildTool(const FToolBuilderState& SceneState) const
+USingleSelectionMeshEditingTool* UPolygonOnMeshToolBuilder::CreateNewTool(const FToolBuilderState& SceneState) const
 {
-	UPolygonOnMeshTool* NewTool = NewObject<UPolygonOnMeshTool>(SceneState.ToolManager);
-
-	UActorComponent* ActorComponent = ToolBuilderUtil::FindFirstComponent(SceneState, CanMakeComponentTarget);
-	auto* MeshComponent = Cast<UPrimitiveComponent>(ActorComponent);
-	check(MeshComponent != nullptr);
-	NewTool->SetSelection(MakeComponentTarget(MeshComponent));
-
-	NewTool->SetWorld(SceneState.World);
-	return NewTool;
+	return NewObject<UPolygonOnMeshTool>(SceneState.ToolManager);
 }
 
 
@@ -84,10 +76,11 @@ void UPolygonOnMeshTool::Setup()
 	HoverBehavior->Initialize(this);
 	AddInputBehavior(HoverBehavior);
 
-	WorldTransform = UE::Geometry::FTransform3d(ComponentTarget->GetWorldTransform());
+	IPrimitiveComponentBackedTarget* TargetComponent = Cast<IPrimitiveComponentBackedTarget>(Target);
+	WorldTransform = UE::Geometry::FTransform3d(TargetComponent->GetWorldTransform());
 
 	// hide input StaticMeshComponent
-	ComponentTarget->SetOwnerVisibility(false);
+	TargetComponent->SetOwnerVisibility(false);
 
 	BasicProperties = NewObject<UPolygonOnMeshToolProperties>(this);
 	AddToolPropertySource(BasicProperties);
@@ -132,12 +125,12 @@ void UPolygonOnMeshTool::Setup()
 		DrawPlaneWorld = PlaneMechanic->Plane;
 		UpdateDrawPlane();
 	});
-	PlaneMechanic->SetPlaneCtrlClickBehaviorTarget->InvisibleComponentsToHitTest.Add(ComponentTarget->GetOwnerComponent());
+	PlaneMechanic->SetPlaneCtrlClickBehaviorTarget->InvisibleComponentsToHitTest.Add(TargetComponent->GetOwnerComponent());
 
 	// Convert input mesh description to dynamic mesh
 	OriginalDynamicMesh = MakeShared<FDynamicMesh3, ESPMode::ThreadSafe>();
 	FMeshDescriptionToDynamicMesh Converter;
-	Converter.Convert(ComponentTarget->GetMesh(), *OriginalDynamicMesh);
+	Converter.Convert(Cast<IMeshDescriptionProvider>(Target)->GetMeshDescription(), *OriginalDynamicMesh);
 	// TODO: consider adding an AABB tree construction here?  tradeoff vs doing a raycast against full every time a param change happens ...
 
 	LastDrawnPolygon = FPolygon2d();
@@ -226,7 +219,7 @@ void UPolygonOnMeshTool::SetupPreview()
 	Preview->PreviewMesh->SetTangentsMode(EDynamicMeshTangentCalcType::AutoCalculated);
 
 	FComponentMaterialSet MaterialSet;
-	ComponentTarget->GetMaterialSet(MaterialSet);
+	Cast<IMaterialProvider>(Target)->GetMaterialSet(MaterialSet);
 	Preview->ConfigureMaterials(MaterialSet.Materials,
 		ToolSetupUtil::GetDefaultWorkingMaterial(GetToolManager())
 	);
@@ -244,7 +237,7 @@ void UPolygonOnMeshTool::Shutdown(EToolShutdownType ShutdownType)
 	}
 
 	// Restore (unhide) the source meshes
-	ComponentTarget->SetOwnerVisibility(true);
+	Cast<IPrimitiveComponentBackedTarget>(Target)->SetOwnerVisibility(true);
 
 	TArray<FDynamicMeshOpResult> Results;
 	Results.Emplace(Preview->Shutdown());
@@ -408,10 +401,10 @@ void UPolygonOnMeshTool::GenerateAsset(const TArray<FDynamicMeshOpResult>& Resul
 	
 	check(Results.Num() > 0);
 	check(Results[0].Mesh.Get() != nullptr);
-	ComponentTarget->CommitMesh([&Results](const FPrimitiveComponentTarget::FCommitParams& CommitParams)
+	Cast<IMeshDescriptionCommitter>(Target)->CommitMeshDescription([&Results](const IMeshDescriptionCommitter::FCommitterParams& CommitParams)
 	{
 		FDynamicMeshToMeshDescription Converter;
-		Converter.Convert(Results[0].Mesh.Get(), *CommitParams.MeshDescription);
+		Converter.Convert(Results[0].Mesh.Get(), *CommitParams.MeshDescriptionOut);
 	});
 
 	GetToolManager()->EndUndoTransaction();
diff --git a/Engine/Plugins/Experimental/MeshModelingToolset/Source/MeshModelingToolsEditorOnly/Private/SimplifyMeshTool.cpp b/Engine/Plugins/Experimental/MeshModelingToolset/Source/MeshModelingToolsEditorOnly/Private/SimplifyMeshTool.cpp
--- a/Engine/Plugins/Experimental/MeshModelingToolset/Source/MeshModelingToolsEditorOnly/Private/SimplifyMeshTool.cpp
+++ b/Engine/Plugins/Experimental/MeshModelingToolset/Source/MeshModelingToolsEditorOnly/Private/SimplifyMeshTool.cpp
@@ -18,6 +18,11 @@
 
 #include "SceneManagement.h" // for FPrimitiveDrawInterface
 
+#include "TargetInterfaces/MaterialProvider.h"
+#include "TargetInterfaces/MeshDescriptionCommitter.h"
+#include "TargetInterfaces/MeshDescriptionProvider.h"
+#include "TargetInterfaces/PrimitiveComponentBackedTarget.h"
+
 //#include "ProfilingDebugging/ScopedTimers.h" // enable this to use the timer.
 #include "Modules/ModuleManager.h"
 
@@ -40,24 +45,9 @@ DEFINE_LOG_CATEGORY_STATIC(LogMeshSimplification, Log, All);
 /*
  * ToolBuilder
  */
-bool USimplifyMeshToolBuilder::CanBuildTool(const FToolBuilderState& SceneState) const
-{
-	return ToolBuilderUtil::CountComponents(SceneState, CanMakeComponentTarget) == 1;
-}
-
-UInteractiveTool* USimplifyMeshToolBuilder::BuildTool(const FToolBuilderState& SceneState) const
+USingleSelectionMeshEditingTool* USimplifyMeshToolBuilder::CreateNewTool(const FToolBuilderState& SceneState) const
 {
-	USimplifyMeshTool* NewTool = NewObject<USimplifyMeshTool>(SceneState.ToolManager);
-
-	UActorComponent* ActorComponent = ToolBuilderUtil::FindFirstComponent(SceneState, CanMakeComponentTarget);
-	auto* MeshComponent = Cast<UPrimitiveComponent>(ActorComponent);
-	check(MeshComponent != nullptr);
-
-	NewTool->SetSelection(MakeComponentTarget(MeshComponent));
-	NewTool->SetWorld(SceneState.World);
-	NewTool->SetAssetAPI(AssetAPI);
-
-	return NewTool;
+	return NewObject<USimplifyMeshTool>(SceneState.ToolManager);
 }
 
 /*
@@ -79,27 +69,18 @@ USimplifyMeshToolProperties::USimplifyMeshToolProperties()
 	MaterialBoundaryConstraint = EMaterialBoundaryConstraint::Ignore;
 }
 
-void USimplifyMeshTool::SetWorld(UWorld* World)
-{
-	this->TargetWorld = World;
-}
-
-void USimplifyMeshTool::SetAssetAPI(IToolsContextAssetAPI* AssetAPIIn)
-{
-	this->AssetAPI = AssetAPIIn;
-}
-
 void USimplifyMeshTool::Setup()
 {
 	UInteractiveTool::Setup();
 
 
 	// hide component and create + show preview
-	ComponentTarget->SetOwnerVisibility(false);
+	IPrimitiveComponentBackedTarget* TargetComponent = Cast<IPrimitiveComponentBackedTarget>(Target);
+	TargetComponent->SetOwnerVisibility(false);
 	Preview = NewObject<UMeshOpPreviewWithBackgroundCompute>(this, "Preview");
 	Preview->Setup(this->TargetWorld, this);
 	FComponentMaterialSet MaterialSet;
-	ComponentTarget->GetMaterialSet(MaterialSet);
+	Cast<IMaterialProvider>(Target)->GetMaterialSet(MaterialSet);
 	Preview->ConfigureMaterials( MaterialSet.Materials,
 		ToolSetupUtil::GetDefaultWorkingMaterial(GetToolManager())
 	);
@@ -121,7 +102,8 @@ void USimplifyMeshTool::Setup()
 #else
 		auto EnterProgressFrame = [](int Progress) {};
 #endif
-		OriginalMeshDescription = MakeShared<FMeshDescription, ESPMode::ThreadSafe>(*ComponentTarget->GetMesh());
+		IMeshDescriptionProvider* TargetMeshProvider = Cast<IMeshDescriptionProvider>(Target);
+		OriginalMeshDescription = MakeShared<FMeshDescription, ESPMode::ThreadSafe>(*TargetMeshProvider->GetMeshDescription());
 		// make sure the mesh description has tangents. 
 		{
 			// if the source was imported with auto-generate tangents, the mesh description tangents won't have been populated
@@ -136,12 +118,12 @@ void USimplifyMeshTool::Setup()
 		FMeshDescriptionToDynamicMesh Converter; 
 		// convert with tangent overlay
 		const bool bUseTangents = true;
-		Converter.Convert(ComponentTarget->GetMesh(), *OriginalMesh, bUseTangents);
+		Converter.Convert(TargetMeshProvider->GetMeshDescription(), *OriginalMesh, bUseTangents);
 		EnterProgressFrame(2);
 		OriginalMeshSpatial = MakeShared<FDynamicMeshAABBTree3, ESPMode::ThreadSafe>(OriginalMesh.Get(), true);
 	}
 
-	Preview->PreviewMesh->SetTransform(ComponentTarget->GetWorldTransform());
+	Preview->PreviewMesh->SetTransform(TargetComponent->GetWorldTransform());
 	Preview->PreviewMesh->SetTangentsMode(EDynamicMeshTangentCalcType::AutoCalculated);
 	Preview->PreviewMesh->UpdatePreview(OriginalMesh.Get());
 
@@ -182,7 +164,7 @@ bool USimplifyMeshTool::CanAccept() const
 void USimplifyMeshTool::Shutdown(EToolShutdownType ShutdownType)
 {
 	SimplifyProperties->SaveProperties(this);
-	ComponentTarget->SetOwnerVisibility(true);
+	Cast<IPrimitiveComponentBackedTarget>(Target)->SetOwnerVisibility(true);
 	FDynamicMeshOpResult Result = Preview->Shutdown();
 	if (ShutdownType == EToolShutdownType::Accept)
 	{
@@ -213,7 +195,7 @@ TUniquePtr<FDynamicMeshOperator> USimplifyMeshTool::MakeNewOperator()
 	Op->MeshBoundaryConstraint = (EEdgeRefineFlags)SimplifyProperties->MeshBoundaryConstraint;
 	Op->GroupBoundaryConstraint = (EEdgeRefineFlags)SimplifyProperties->GroupBoundaryConstraint;
 	Op->MaterialBoundaryConstraint = (EEdgeRefineFlags)SimplifyProperties->MaterialBoundaryConstraint;
-	FTransform LocalToWorld = ComponentTarget->GetWorldTransform();
+	FTransform LocalToWorld = Cast<IPrimitiveComponentBackedTarget>(Target)->GetWorldTransform();
 	Op->SetTransform(LocalToWorld);
 
 	Op->OriginalMeshDescription = OriginalMeshDescription;
@@ -229,7 +211,7 @@ TUniquePtr<FDynamicMeshOperator> USimplifyMeshTool::MakeNewOperator()
 void USimplifyMeshTool::Render(IToolsContextRenderAPI* RenderAPI)
 {
 	FPrimitiveDrawInterface* PDI = RenderAPI->GetPrimitiveDrawInterface();
-	FTransform Transform = ComponentTarget->GetWorldTransform(); //Actor->GetTransform();
+	FTransform Transform = Cast<IPrimitiveComponentBackedTarget>(Target)->GetWorldTransform(); //Actor->GetTransform();
 
 	if (SimplifyProperties->bShowConstraintEdges)
 	{
@@ -285,7 +267,7 @@ void USimplifyMeshTool::UpdateVisualization()
 	}
 	else
 	{
-		ComponentTarget->GetMaterialSet(MaterialSet);
+		Cast<IMaterialProvider>(Target)->GetMaterialSet(MaterialSet);
 		Preview->PreviewMesh->ClearTriangleColorFunction(UPreviewMesh::ERenderUpdateMode::FastUpdate);
 	}
 	Preview->ConfigureMaterials(MaterialSet.Materials,
@@ -297,12 +279,12 @@ void USimplifyMeshTool::GenerateAsset(const FDynamicMeshOpResult& Result)
 	GetToolManager()->BeginUndoTransaction(LOCTEXT("SimplifyMeshToolTransactionName", "Simplify Mesh"));
 
 	check(Result.Mesh.Get() != nullptr);
-	ComponentTarget->CommitMesh([&Result](const FPrimitiveComponentTarget::FCommitParams& CommitParams)
+	Cast<IMeshDescriptionCommitter>(Target)->CommitMeshDescription([&Result](const IMeshDescriptionCommitter::FCommitterParams& CommitParams)
 	{
 		FDynamicMeshToMeshDescription Converter;
 
 		// full conversion if normal topology changed or faces were inverted
-		Converter.Convert(Result.Mesh.Get(), *CommitParams.MeshDescription);
+		Converter.Convert(Result.Mesh.Get(), *CommitParams.MeshDescriptionOut);
 	});
 
 	GetToolManager()->EndUndoTransaction();
diff --git a/Engine/Plugins/Experimental/MeshModelingToolset/Source/MeshModelingToolsEditorOnly/Private/SubdividePolyTool.cpp b/Engine/Plugins/Experimental/MeshModelingToolset/Source/MeshModelingToolsEditorOnly/Private/SubdividePolyTool.cpp
--- a/Engine/Plugins/Experimental/MeshModelingToolset/Source/MeshModelingToolsEditorOnly/Private/SubdividePolyTool.cpp
+++ b/Engine/Plugins/Experimental/MeshModelingToolset/Source/MeshModelingToolsEditorOnly/Private/SubdividePolyTool.cpp
@@ -12,6 +12,11 @@
 #include "SimpleDynamicMeshComponent.h"
 #include "Drawing/PreviewGeometryActor.h"
 
+#include "TargetInterfaces/MaterialProvider.h"
+#include "TargetInterfaces/MeshDescriptionCommitter.h"
+#include "TargetInterfaces/MeshDescriptionProvider.h"
+#include "TargetInterfaces/PrimitiveComponentBackedTarget.h"
+
 #include "ExplicitUseGeometryMathTypes.h"		// using UE::Geometry::(math types)
 using namespace UE::Geometry;
 
@@ -54,22 +59,9 @@ public:
 
 // Tool builder
 
-bool USubdividePolyToolBuilder::CanBuildTool(const FToolBuilderState& SceneState) const
-{
-	return ToolBuilderUtil::CountComponents(SceneState, CanMakeComponentTarget) == 1;
-}
-
-UInteractiveTool* USubdividePolyToolBuilder::BuildTool(const FToolBuilderState& SceneState) const
+USingleSelectionMeshEditingTool* USubdividePolyToolBuilder::CreateNewTool(const FToolBuilderState& SceneState) const
 {
-	UActorComponent* ActorComponent = ToolBuilderUtil::FindFirstComponent(SceneState, CanMakeComponentTarget);
-	auto* MeshComponent = Cast<UPrimitiveComponent>(ActorComponent);
-	check(MeshComponent != nullptr);
-
-	USubdividePolyTool* NewTool = NewObject<USubdividePolyTool>(SceneState.ToolManager);
-	NewTool->SetSelection(MakeComponentTarget(MeshComponent));
-	NewTool->SetWorld(SceneState.World);
-
-	return NewTool;
+	return NewObject<USubdividePolyTool>(SceneState.ToolManager);
 }
 
 
@@ -165,7 +157,7 @@ void USubdividePolyTool::Setup()
 	UInteractiveTool::Setup();
 	SetToolDisplayName(LOCTEXT("ToolName", "Subdivide"));
 
-	if (!ComponentTarget)
+	if (!Target)
 	{
 		return;
 	}
@@ -173,7 +165,7 @@ void USubdividePolyTool::Setup()
 	bool bWantVertexNormals = false;
 	OriginalMesh = MakeShared<FDynamicMesh3, ESPMode::ThreadSafe>(bWantVertexNormals, false, false, false);
 	FMeshDescriptionToDynamicMesh Converter;
-	Converter.Convert(ComponentTarget->GetMesh(), *OriginalMesh);
+	Converter.Convert(Cast<IMeshDescriptionProvider>(Target)->GetMeshDescription(), *OriginalMesh);
 
 	FText ErrorMessage;
 	bool bCatmullClarkOK = CheckGroupTopology(ErrorMessage);
@@ -201,6 +193,7 @@ void USubdividePolyTool::Setup()
 	AddToolPropertySource(Properties);
 	SetToolPropertySourceEnabled(Properties, true);
 
+	IPrimitiveComponentBackedTarget* TargetComponent = Cast<IPrimitiveComponentBackedTarget>(Target);
 	PreviewMesh = NewObject<UPreviewMesh>(this);
 	if (PreviewMesh == nullptr)
 	{
@@ -208,7 +201,7 @@ void USubdividePolyTool::Setup()
 	}
 	PreviewMesh->CreateInWorld(TargetWorld, FTransform::Identity);
 
-	PreviewMesh->SetTransform(ComponentTarget->GetWorldTransform());
+	PreviewMesh->SetTransform(TargetComponent->GetWorldTransform());
 	PreviewMesh->UpdatePreview(OriginalMesh.Get());
 
 	USimpleDynamicMeshComponent* PreviewDynamicMeshComponent = (USimpleDynamicMeshComponent*)PreviewMesh->GetRootComponent();
@@ -227,7 +220,7 @@ void USubdividePolyTool::Setup()
 
 	// Use the input mesh's material on the preview
 	FComponentMaterialSet MaterialSet;
-	ComponentTarget->GetMaterialSet(MaterialSet);
+	Cast<IMaterialProvider>(Target)->GetMaterialSet(MaterialSet);
 	for (int k = 0; k < MaterialSet.Materials.Num(); ++k)
 	{
 		PreviewMesh->SetMaterial(k, MaterialSet.Materials[k]);
@@ -303,13 +296,13 @@ void USubdividePolyTool::Setup()
 	});
 
 	PreviewGeometry = NewObject<UPreviewGeometry>(this);
-	PreviewGeometry->CreateInWorld(ComponentTarget->GetOwnerActor()->GetWorld(), ComponentTarget->GetWorldTransform());
+	PreviewGeometry->CreateInWorld(TargetComponent->GetOwnerActor()->GetWorld(), TargetComponent->GetWorldTransform());
 	CreateOrUpdatePreviewGeometry();
 
 	// regenerate preview geo if mesh changes due to undo/redo/etc
 	PreviewDynamicMeshComponent->OnMeshChanged.AddLambda([this]() { bPreviewGeometryNeedsUpdate = true; });
 
-	ComponentTarget->SetOwnerVisibility(false);
+	TargetComponent->SetOwnerVisibility(false);
 	PreviewMesh->SetVisible(true);
 }
 
@@ -389,7 +382,7 @@ void USubdividePolyTool::Shutdown(EToolShutdownType ShutdownType)
 		PreviewGeometry->Disconnect();
 	}
 	
-	ComponentTarget->SetOwnerVisibility(true);
+	Cast<IPrimitiveComponentBackedTarget>(Target)->SetOwnerVisibility(true);
 
 	if (PreviewMesh)
 	{
@@ -400,10 +393,10 @@ void USubdividePolyTool::Shutdown(EToolShutdownType ShutdownType)
 			USimpleDynamicMeshComponent* PreviewDynamicMeshComponent = (USimpleDynamicMeshComponent*)PreviewMesh->GetRootComponent();
 			FDynamicMesh3* DynamicMeshResult = PreviewDynamicMeshComponent->GetRenderMesh();
 
-			ComponentTarget->CommitMesh([DynamicMeshResult](const FPrimitiveComponentTarget::FCommitParams& CommitParams)
+			Cast<IMeshDescriptionCommitter>(Target)->CommitMeshDescription([DynamicMeshResult](const IMeshDescriptionCommitter::FCommitterParams& CommitParams)
 			{
 				FDynamicMeshToMeshDescription Converter;
-				Converter.Convert(DynamicMeshResult, *CommitParams.MeshDescription);
+				Converter.Convert(DynamicMeshResult, *CommitParams.MeshDescriptionOut);
 			});
 
 			GetToolManager()->EndUndoTransaction();
diff --git a/Engine/Plugins/Experimental/MeshModelingToolset/Source/MeshModelingToolsEditorOnly/Public/GenerateLODMeshesTool.h b/Engine/Plugins/Experimental/MeshModelingToolset/Source/MeshModelingToolsEditorOnly/Public/GenerateLODMeshesTool.h
--- a/Engine/Plugins/Experimental/MeshModelingToolset/Source/MeshModelingToolsEditorOnly/Public/GenerateLODMeshesTool.h
+++ b/Engine/Plugins/Experimental/MeshModelingToolset/Source/MeshModelingToolsEditorOnly/Public/GenerateLODMeshesTool.h
@@ -12,6 +12,7 @@
 #include "CleaningOps/SimplifyMeshOp.h"
 #include "Properties/MeshStatisticsProperties.h"
 #include "Properties/RemeshProperties.h"
+#include "BaseTools/SingleSelectionMeshEditingTool.h"
 #include "GenerateLODMeshesTool.generated.h"
 
 
@@ -22,15 +23,13 @@ class IAssetGenerationAPI;
  *
  */
 UCLASS()
-class MESHMODELINGTOOLSEDITORONLY_API UGenerateLODMeshesToolBuilder : public UInteractiveToolBuilder
+class MESHMODELINGTOOLSEDITORONLY_API UGenerateLODMeshesToolBuilder : public USingleSelectionMeshEditingToolBuilder
 {
 	GENERATED_BODY()
 
 public:
-	IAssetGenerationAPI* AssetAPI = nullptr;
-
-	virtual bool CanBuildTool(const FToolBuilderState& SceneState) const override;
-	virtual UInteractiveTool* BuildTool(const FToolBuilderState& SceneState) const override;
+	virtual USingleSelectionMeshEditingTool* CreateNewTool(const FToolBuilderState& SceneState) const override;
+	virtual bool RequiresAssetAPI() const override { return true; }
 };
 
 
@@ -153,14 +152,11 @@ public:
  * Simple Mesh Simplifying Tool
  */
 UCLASS()
-class MESHMODELINGTOOLSEDITORONLY_API UGenerateLODMeshesTool : public USingleSelectionTool
+class MESHMODELINGTOOLSEDITORONLY_API UGenerateLODMeshesTool : public USingleSelectionMeshEditingTool
 {
 	GENERATED_BODY()
 
 public:
-	virtual void SetWorld(UWorld* World);
-	virtual void SetAssetAPI(IAssetGenerationAPI* AssetAPI);
-
 	virtual void Setup() override;
 	virtual void Shutdown(EToolShutdownType ShutdownType) override;
 
@@ -186,9 +182,6 @@ private:
 	void UpdateNumPreviews();
 	void InvalidateAllPreviews();
 
-	UWorld* TargetWorld;
-	IAssetGenerationAPI* AssetAPI;
-
 	TSharedPtr<FMeshDescription, ESPMode::ThreadSafe> OriginalMeshDescription;
 	// Dynamic Mesh versions precomputed in Setup (rather than recomputed for every simplify op)
 	TSharedPtr<UE::Geometry::FDynamicMesh3, ESPMode::ThreadSafe> OriginalMesh;
diff --git a/Engine/Plugins/Experimental/MeshModelingToolset/Source/MeshModelingToolsEditorOnly/Public/MeshTangentsTool.h b/Engine/Plugins/Experimental/MeshModelingToolset/Source/MeshModelingToolsEditorOnly/Public/MeshTangentsTool.h
--- a/Engine/Plugins/Experimental/MeshModelingToolset/Source/MeshModelingToolsEditorOnly/Public/MeshTangentsTool.h
+++ b/Engine/Plugins/Experimental/MeshModelingToolset/Source/MeshModelingToolsEditorOnly/Public/MeshTangentsTool.h
@@ -11,6 +11,7 @@
 #include "DynamicMesh3.h"
 #include "MeshTangents.h"
 #include "ParameterizationOps/CalculateTangentsOp.h"
+#include "BaseTools/SingleSelectionMeshEditingTool.h"
 #include "MeshTangentsTool.generated.h"
 
 
@@ -25,13 +26,15 @@ class UMaterialInstanceDynamic;
  *
  */
 UCLASS()
-class MESHMODELINGTOOLSEDITORONLY_API UMeshTangentsToolBuilder : public UInteractiveToolBuilder
+class MESHMODELINGTOOLSEDITORONLY_API UMeshTangentsToolBuilder : public USingleSelectionMeshEditingToolBuilder
 {
 	GENERATED_BODY()
 
 public:
-	virtual bool CanBuildTool(const FToolBuilderState& SceneState) const override;
-	virtual UInteractiveTool* BuildTool(const FToolBuilderState& SceneState) const override;
+	virtual USingleSelectionMeshEditingTool* CreateNewTool(const FToolBuilderState& SceneState) const override;
+
+protected:
+	virtual const FToolTargetTypeRequirements& GetTargetRequirements() const override;
 };
 
 
@@ -78,14 +81,12 @@ public:
  * Simple Mesh Simplifying Tool
  */
 UCLASS()
-class MESHMODELINGTOOLSEDITORONLY_API UMeshTangentsTool : public USingleSelectionTool, public UE::Geometry::IGenericDataOperatorFactory<UE::Geometry::FMeshTangentsd>
+class MESHMODELINGTOOLSEDITORONLY_API UMeshTangentsTool : public USingleSelectionMeshEditingTool, public UE::Geometry::IGenericDataOperatorFactory<UE::Geometry::FMeshTangentsd>
 {
 	GENERATED_BODY()
 public:
 	UMeshTangentsTool();
 
-	virtual void SetWorld(UWorld* World);
-
 	virtual void Setup() override;
 	virtual void Shutdown(EToolShutdownType ShutdownType) override;
 
@@ -116,8 +117,6 @@ protected:
 	TUniquePtr<TGenericDataBackgroundCompute<UE::Geometry::FMeshTangentsd>> Compute = nullptr;
 
 protected:
-	UWorld* TargetWorld;
-
 	TSharedPtr<FMeshDescription, ESPMode::ThreadSafe> InputMeshDescription;
 	TSharedPtr<UE::Geometry::FMeshTangentsf, ESPMode::ThreadSafe> InitialTangents;
 	TSharedPtr<UE::Geometry::FDynamicMesh3, ESPMode::ThreadSafe> InputMesh;
diff --git a/Engine/Plugins/Experimental/MeshModelingToolset/Source/MeshModelingToolsEditorOnly/Public/MeshToVolumeTool.h b/Engine/Plugins/Experimental/MeshModelingToolset/Source/MeshModelingToolsEditorOnly/Public/MeshToVolumeTool.h
--- a/Engine/Plugins/Experimental/MeshModelingToolset/Source/MeshModelingToolsEditorOnly/Public/MeshToVolumeTool.h
+++ b/Engine/Plugins/Experimental/MeshModelingToolset/Source/MeshModelingToolsEditorOnly/Public/MeshToVolumeTool.h
@@ -13,6 +13,7 @@
 #include "PreviewMesh.h"
 #include "Drawing/LineSetComponent.h"
 #include "PropertySets/OnAcceptProperties.h"
+#include "BaseTools/SingleSelectionMeshEditingTool.h"
 #include "MeshToVolumeTool.generated.h"
 
 class IAssetGenerationAPI;
@@ -21,17 +22,12 @@ class IAssetGenerationAPI;
  *
  */
 UCLASS()
-class MESHMODELINGTOOLSEDITORONLY_API UMeshToVolumeToolBuilder : public UInteractiveToolBuilder
+class MESHMODELINGTOOLSEDITORONLY_API UMeshToVolumeToolBuilder : public USingleSelectionMeshEditingToolBuilder
 {
 	GENERATED_BODY()
 public:
-	IAssetGenerationAPI* AssetAPI = nullptr;
-
 	virtual bool CanBuildTool(const FToolBuilderState & SceneState) const override;
-	virtual UInteractiveTool* BuildTool(const FToolBuilderState & SceneState) const override;
-
-protected:
-	virtual const FToolTargetTypeRequirements& GetTargetRequirements() const override;
+	virtual USingleSelectionMeshEditingTool* CreateNewTool(const FToolBuilderState & SceneState) const override;
 };
 
 
@@ -75,15 +71,13 @@ public:
  * tool can be moved out of the editor-only section and put with VolumeToMeshTool.
  */
 UCLASS()
-class MESHMODELINGTOOLSEDITORONLY_API UMeshToVolumeTool : public USingleSelectionTool
+class MESHMODELINGTOOLSEDITORONLY_API UMeshToVolumeTool : public USingleSelectionMeshEditingTool
 {
 	GENERATED_BODY()
 
 public:
 	UMeshToVolumeTool();
 
-	virtual void SetAssetAPI(IAssetGenerationAPI* InAssetAPI) { this->AssetAPI = InAssetAPI; }
-
 	virtual void Setup() override;
 	virtual void Shutdown(EToolShutdownType ShutdownType) override;
 
@@ -107,8 +101,6 @@ protected:
 	ULineSetComponent* VolumeEdgesSet;
 
 protected:
-	IAssetGenerationAPI* AssetAPI = nullptr;
-
 	UE::Geometry::FDynamicMesh3 InputMesh;
 
 	void RecalculateVolume();
diff --git a/Engine/Plugins/Experimental/MeshModelingToolset/Source/MeshModelingToolsEditorOnly/Public/ParameterizeMeshTool.h b/Engine/Plugins/Experimental/MeshModelingToolset/Source/MeshModelingToolsEditorOnly/Public/ParameterizeMeshTool.h
--- a/Engine/Plugins/Experimental/MeshModelingToolset/Source/MeshModelingToolsEditorOnly/Public/ParameterizeMeshTool.h
+++ b/Engine/Plugins/Experimental/MeshModelingToolset/Source/MeshModelingToolsEditorOnly/Public/ParameterizeMeshTool.h
@@ -10,6 +10,7 @@
 #include "MeshOpPreviewHelpers.h"
 #include "Properties/MeshMaterialProperties.h"
 #include "Properties/MeshUVChannelProperties.h"
+#include "BaseTools/SingleSelectionMeshEditingTool.h"
 
 #include "ParameterizeMeshTool.generated.h"
 
@@ -18,7 +19,6 @@
 class USimpleDynamicMeshComponent;
 class UMaterialInterface;
 class UMaterialInstanceDynamic;
-class IAssetGenerationAPI;
 
 
 
@@ -40,22 +40,13 @@ enum class EParameterizeMeshMaterialMode : uint8
  *
  */
 UCLASS()
-class MESHMODELINGTOOLSEDITORONLY_API UParameterizeMeshToolBuilder : public UInteractiveToolBuilder
+class MESHMODELINGTOOLSEDITORONLY_API UParameterizeMeshToolBuilder : public USingleSelectionMeshEditingToolBuilder
 {
 	GENERATED_BODY()
 
 public:
-
-	IAssetGenerationAPI* AssetAPI;
 	bool bDoAutomaticGlobalUnwrap = false;
-
-	UParameterizeMeshToolBuilder()
-	{
-		AssetAPI = nullptr;
-	}
-
-	virtual bool CanBuildTool(const FToolBuilderState& SceneState) const override;
-	virtual UInteractiveTool* BuildTool(const FToolBuilderState& SceneState) const override;
+	virtual USingleSelectionMeshEditingTool* CreateNewTool(const FToolBuilderState& SceneState) const override;
 };
 
 
@@ -149,15 +140,13 @@ public:
  * Simple Mesh Simplifying Tool
  */
 UCLASS()
-class MESHMODELINGTOOLSEDITORONLY_API UParameterizeMeshTool : public USingleSelectionTool, public UE::Geometry::IDynamicMeshOperatorFactory
+class MESHMODELINGTOOLSEDITORONLY_API UParameterizeMeshTool : public USingleSelectionMeshEditingTool, public UE::Geometry::IDynamicMeshOperatorFactory
 {
 	GENERATED_BODY()
 
 public:
 	UParameterizeMeshTool();
 
-	virtual void SetWorld(UWorld* World);
-	virtual void SetAssetAPI(IAssetGenerationAPI* AssetAPI);
 	virtual void SetUseAutoGlobalParameterizationMode(bool bEnable);
 
 	virtual void Setup() override;
@@ -202,9 +191,6 @@ protected:
 	bool bDoAutomaticGlobalUnwrap = false;
 
 protected:
-	UWorld* TargetWorld;
-	IAssetGenerationAPI* AssetAPI;
-
 	TSharedPtr<UE::Geometry::FDynamicMesh3, ESPMode::ThreadSafe> InputMesh;
 
 };
diff --git a/Engine/Plugins/Experimental/MeshModelingToolset/Source/MeshModelingToolsEditorOnly/Public/PolygonOnMeshTool.h b/Engine/Plugins/Experimental/MeshModelingToolset/Source/MeshModelingToolsEditorOnly/Public/PolygonOnMeshTool.h
--- a/Engine/Plugins/Experimental/MeshModelingToolset/Source/MeshModelingToolsEditorOnly/Public/PolygonOnMeshTool.h
+++ b/Engine/Plugins/Experimental/MeshModelingToolset/Source/MeshModelingToolsEditorOnly/Public/PolygonOnMeshTool.h
@@ -14,6 +14,7 @@
 #include "Mechanics/ConstructionPlaneMechanic.h"
 #include "Mechanics/CollectSurfacePathMechanic.h"
 #include "Polygon2.h"
+#include "BaseTools/SingleSelectionMeshEditingTool.h"
 #include "PolygonOnMeshTool.generated.h"
 
 
@@ -31,13 +32,12 @@ class ULineSetComponent;
  *
  */
 UCLASS()
-class MESHMODELINGTOOLSEDITORONLY_API UPolygonOnMeshToolBuilder : public UInteractiveToolBuilder
+class MESHMODELINGTOOLSEDITORONLY_API UPolygonOnMeshToolBuilder : public USingleSelectionMeshEditingToolBuilder
 {
 	GENERATED_BODY()
 
 public:
-	virtual bool CanBuildTool(const FToolBuilderState& SceneState) const override;
-	virtual UInteractiveTool* BuildTool(const FToolBuilderState& SceneState) const override;
+	virtual USingleSelectionMeshEditingTool* CreateNewTool(const FToolBuilderState& SceneState) const override;
 };
 
 
@@ -144,7 +144,7 @@ public:
  * Simple Mesh Plane Cutting Tool
  */
 UCLASS()
-class MESHMODELINGTOOLSEDITORONLY_API UPolygonOnMeshTool : public USingleSelectionTool, public UE::Geometry::IDynamicMeshOperatorFactory, public IClickBehaviorTarget, public IHoverBehaviorTarget
+class MESHMODELINGTOOLSEDITORONLY_API UPolygonOnMeshTool : public USingleSelectionMeshEditingTool, public UE::Geometry::IDynamicMeshOperatorFactory, public IClickBehaviorTarget, public IHoverBehaviorTarget
 {
 	GENERATED_BODY()
 
diff --git a/Engine/Plugins/Experimental/MeshModelingToolset/Source/MeshModelingToolsEditorOnly/Public/SimplifyMeshTool.h b/Engine/Plugins/Experimental/MeshModelingToolset/Source/MeshModelingToolsEditorOnly/Public/SimplifyMeshTool.h
--- a/Engine/Plugins/Experimental/MeshModelingToolset/Source/MeshModelingToolsEditorOnly/Public/SimplifyMeshTool.h
+++ b/Engine/Plugins/Experimental/MeshModelingToolset/Source/MeshModelingToolsEditorOnly/Public/SimplifyMeshTool.h
@@ -12,6 +12,7 @@
 #include "CleaningOps/SimplifyMeshOp.h"
 #include "Properties/MeshStatisticsProperties.h"
 #include "Properties/RemeshProperties.h"
+#include "BaseTools/SingleSelectionMeshEditingTool.h"
 #include "SimplifyMeshTool.generated.h"
 
 
@@ -22,15 +23,12 @@
  *
  */
 UCLASS()
-class MESHMODELINGTOOLSEDITORONLY_API USimplifyMeshToolBuilder : public UInteractiveToolBuilder
+class MESHMODELINGTOOLSEDITORONLY_API USimplifyMeshToolBuilder : public USingleSelectionMeshEditingToolBuilder
 {
 	GENERATED_BODY()
 
 public:
-	IToolsContextAssetAPI* AssetAPI = nullptr;
-
-	virtual bool CanBuildTool(const FToolBuilderState& SceneState) const override;
-	virtual UInteractiveTool* BuildTool(const FToolBuilderState& SceneState) const override;
+	virtual USingleSelectionMeshEditingTool* CreateNewTool(const FToolBuilderState& SceneState) const override;
 };
 
 
@@ -93,14 +91,11 @@ public:
  * Simple Mesh Simplifying Tool
  */
 UCLASS()
-class MESHMODELINGTOOLSEDITORONLY_API USimplifyMeshTool : public USingleSelectionTool, public UE::Geometry::IDynamicMeshOperatorFactory
+class MESHMODELINGTOOLSEDITORONLY_API USimplifyMeshTool : public USingleSelectionMeshEditingTool, public UE::Geometry::IDynamicMeshOperatorFactory
 {
 	GENERATED_BODY()
 
 public:
-	virtual void SetWorld(UWorld* World);
-	virtual void SetAssetAPI(IToolsContextAssetAPI* AssetAPI);
-
 	virtual void Setup() override;
 	virtual void Shutdown(EToolShutdownType ShutdownType) override;
 
@@ -126,9 +121,6 @@ private:
 	UPROPERTY()
 	UMeshOpPreviewWithBackgroundCompute* Preview;
 
-	UWorld* TargetWorld;
-	IToolsContextAssetAPI* AssetAPI;
-
 	TSharedPtr<FMeshDescription, ESPMode::ThreadSafe> OriginalMeshDescription;
 	// Dynamic Mesh versions precomputed in Setup (rather than recomputed for every simplify op)
 	TSharedPtr<UE::Geometry::FDynamicMesh3, ESPMode::ThreadSafe> OriginalMesh;
diff --git a/Engine/Plugins/Experimental/MeshModelingToolset/Source/MeshModelingToolsEditorOnly/Public/SubdividePolyTool.h b/Engine/Plugins/Experimental/MeshModelingToolset/Source/MeshModelingToolsEditorOnly/Public/SubdividePolyTool.h
--- a/Engine/Plugins/Experimental/MeshModelingToolset/Source/MeshModelingToolsEditorOnly/Public/SubdividePolyTool.h
+++ b/Engine/Plugins/Experimental/MeshModelingToolset/Source/MeshModelingToolsEditorOnly/Public/SubdividePolyTool.h
@@ -10,6 +10,7 @@
 #include "SingleSelectionTool.h"
 #include "SimpleDynamicMeshComponent.h"
 #include "MeshOpPreviewHelpers.h"
+#include "BaseTools/SingleSelectionMeshEditingTool.h"
 #include "SubdividePolyTool.generated.h"
 
 class USubdividePolyTool;
@@ -19,14 +20,11 @@ class UPreviewGeometry;
  * Tool builder
  */
 UCLASS()
-class MESHMODELINGTOOLSEDITORONLY_API USubdividePolyToolBuilder : public UInteractiveToolBuilder
+class MESHMODELINGTOOLSEDITORONLY_API USubdividePolyToolBuilder : public USingleSelectionMeshEditingToolBuilder
 {
 	GENERATED_BODY()
 public:
-
-	bool CanBuildTool(const FToolBuilderState& SceneState) const override;
-
-	UInteractiveTool* BuildTool(const FToolBuilderState& SceneState) const override;
+	virtual USingleSelectionMeshEditingTool* CreateNewTool(const FToolBuilderState& SceneState) const override;
 };
 
 /**
@@ -69,14 +67,12 @@ public:
  * Tool actual
  */
 UCLASS()
-class MESHMODELINGTOOLSEDITORONLY_API USubdividePolyTool : public USingleSelectionTool
+class MESHMODELINGTOOLSEDITORONLY_API USubdividePolyTool : public USingleSelectionMeshEditingTool
 {
 	GENERATED_BODY()
 
 public:
 
-	virtual void SetWorld(UWorld* World) { TargetWorld = World; }
-
 	virtual void Setup() override;
 	virtual void Shutdown(EToolShutdownType ShutdownType) override;
 
@@ -90,8 +86,6 @@ protected:
 
 	friend class USubdividePolyOperatorFactory;
 
-	UWorld* TargetWorld;
-
 	UPROPERTY()
 	UPreviewMesh* PreviewMesh = nullptr;
 
diff --git a/Engine/Plugins/Experimental/MeshModelingToolset/Source/ModelingComponents/Private/BaseTools/BaseMeshProcessingTool.cpp b/Engine/Plugins/Experimental/MeshModelingToolset/Source/ModelingComponents/Private/BaseTools/BaseMeshProcessingTool.cpp
--- a/Engine/Plugins/Experimental/MeshModelingToolset/Source/ModelingComponents/Private/BaseTools/BaseMeshProcessingTool.cpp
+++ b/Engine/Plugins/Experimental/MeshModelingToolset/Source/ModelingComponents/Private/BaseTools/BaseMeshProcessingTool.cpp
@@ -14,6 +14,11 @@
 #include "DynamicMeshToMeshDescription.h"
 #include "MeshDescriptionToDynamicMesh.h"
 
+#include "TargetInterfaces/MaterialProvider.h"
+#include "TargetInterfaces/MeshDescriptionCommitter.h"
+#include "TargetInterfaces/MeshDescriptionProvider.h"
+#include "TargetInterfaces/PrimitiveComponentBackedTarget.h"
+#include "ToolTargetManager.h"
 
 #define LOCTEXT_NAMESPACE "UBaseMeshProcessingTool"
 
@@ -23,39 +28,29 @@ using namespace UE::Geometry;
  * ToolBuilder
  */
 
+const FToolTargetTypeRequirements& UBaseMeshProcessingToolBuilder::GetTargetRequirements() const
+{
+	static FToolTargetTypeRequirements TypeRequirements({
+		UMaterialProvider::StaticClass(),
+		UMeshDescriptionCommitter::StaticClass(),
+		UMeshDescriptionProvider::StaticClass(),
+		UPrimitiveComponentBackedTarget::StaticClass()
+		});
+	return TypeRequirements;
+}
 
 bool UBaseMeshProcessingToolBuilder::CanBuildTool(const FToolBuilderState& SceneState) const
 {
-	if (SupportsMultipleObjects())
-	{
-		return ToolBuilderUtil::CountComponents(SceneState, CanMakeComponentTarget) > 0;
-	}
-	else
-	{
-		return ToolBuilderUtil::CountComponents(SceneState, CanMakeComponentTarget) == 1;
-	}
+	return SceneState.TargetManager->CountSelectedAndTargetable(SceneState, GetTargetRequirements()) == 1;
 }
 
 UInteractiveTool* UBaseMeshProcessingToolBuilder::BuildTool(const FToolBuilderState& SceneState) const
 {
-	check(SupportsMultipleObjects() == false); // not supported yet
-
 	UBaseMeshProcessingTool* NewTool = MakeNewToolInstance(SceneState.ToolManager);
 
-	TArray<UActorComponent*> Components = ToolBuilderUtil::FindAllComponents(SceneState, CanMakeComponentTarget);
-	check(Components.Num() > 0);
-
-	TArray<TUniquePtr<FPrimitiveComponentTarget>> ComponentTargets;
-	for (UActorComponent* ActorComponent : Components)
-	{
-		auto* MeshComponent = Cast<UPrimitiveComponent>(ActorComponent);
-		if (MeshComponent)
-		{
-			ComponentTargets.Add(MakeComponentTarget(MeshComponent));
-		}
-	}
-
-	NewTool->SetSelection( MoveTemp(ComponentTargets[0]) );
+	UToolTarget* Target = SceneState.TargetManager->BuildFirstSelectedTargetable(SceneState, GetTargetRequirements());
+	check(Target);
+	NewTool->SetTarget(Target);
 	NewTool->SetWorld(SceneState.World);
 
 	return NewTool;
@@ -79,14 +74,15 @@ void UBaseMeshProcessingTool::Setup()
 	UInteractiveTool::Setup();
 
 	// hide input StaticMeshComponent
-	ComponentTarget->SetOwnerVisibility(false);
+	IPrimitiveComponentBackedTarget* TargetComponent = Cast<IPrimitiveComponentBackedTarget>(Target);
+	TargetComponent->SetOwnerVisibility(false);
 
 	// initialize our properties
 	ToolPropertyObjects.Add(this);
 
 	// populate the BaseMesh with a conversion of the input mesh.
 	FMeshDescriptionToDynamicMesh Converter;
-	Converter.Convert(ComponentTarget->GetMesh(), InitialMesh);
+	Converter.Convert(Cast<IMeshDescriptionProvider>(Target)->GetMeshDescription(), InitialMesh);
 
 	if (RequiresScaleNormalization())
 	{
@@ -102,7 +98,7 @@ void UBaseMeshProcessingTool::Setup()
 		MeshTransforms::Scale(InitialMesh, (1.0 / SrcScale) * FVector3d::One(), FVector3d::Zero());
 
 		// apply that transform to target transform so that visible mesh stays in the same spot
-		OverrideTransform = ComponentTarget->GetWorldTransform();
+		OverrideTransform = TargetComponent->GetWorldTransform();
 		FVector TranslateDelta = OverrideTransform.TransformVector((FVector)SrcTranslate);
 		FVector CurScale = OverrideTransform.GetScale3D();
 		OverrideTransform.AddToTranslation(TranslateDelta);
@@ -117,7 +113,7 @@ void UBaseMeshProcessingTool::Setup()
 	{
 		SrcTranslate = FVector3d::Zero();
 		SrcScale = 1.0;
-		OverrideTransform = ComponentTarget->GetWorldTransform();
+		OverrideTransform = TargetComponent->GetWorldTransform();
 		bIsScaleNormalizationApplied = false;
 	}
 
@@ -149,7 +145,7 @@ void UBaseMeshProcessingTool::Setup()
 	Preview->PreviewMesh->SetTangentsMode(EDynamicMeshTangentCalcType::AutoCalculated);
 
 	FComponentMaterialSet MaterialSet;
-	ComponentTarget->GetMaterialSet(MaterialSet);
+	Cast<IMaterialProvider>(Target)->GetMaterialSet(MaterialSet);
 	Preview->ConfigureMaterials(MaterialSet.Materials,
 		ToolSetupUtil::GetDefaultWorkingMaterial(GetToolManager())
 	);
@@ -220,7 +216,7 @@ void UBaseMeshProcessingTool::Shutdown(EToolShutdownType ShutdownType)
 	SavePropertySets();
 
 	// Restore (unhide) the source meshes
-	ComponentTarget->SetOwnerVisibility(true);
+	Cast<IPrimitiveComponentBackedTarget>(Target)->SetOwnerVisibility(true);
 
 	if (Preview != nullptr)
 	{
@@ -241,16 +237,16 @@ void UBaseMeshProcessingTool::Shutdown(EToolShutdownType ShutdownType)
 			}
 
 			bool bTopologyChanged = HasMeshTopologyChanged();
-			ComponentTarget->CommitMesh([DynamicMeshResult, bTopologyChanged](const FPrimitiveComponentTarget::FCommitParams& CommitParams)
+			Cast<IMeshDescriptionCommitter>(Target)->CommitMeshDescription([DynamicMeshResult, bTopologyChanged](const IMeshDescriptionCommitter::FCommitterParams& CommitParams)
 			{
 				FDynamicMeshToMeshDescription Converter;
 				if (bTopologyChanged)
 				{
-					Converter.Convert(DynamicMeshResult, *CommitParams.MeshDescription);
+					Converter.Convert(DynamicMeshResult, *CommitParams.MeshDescriptionOut);
 				}
 				else
 				{
-					Converter.Update(DynamicMeshResult, *CommitParams.MeshDescription);
+					Converter.Update(DynamicMeshResult, *CommitParams.MeshDescriptionOut);
 				}
 			});
 
@@ -375,7 +371,7 @@ void UBaseMeshProcessingTool::SetupWeightMapPropertySet(UWeightMapSetProperties*
 	WeightMapPropertySet = Properties;
 
 	// initialize property list
-	Properties->InitializeFromMesh(ComponentTarget->GetMesh());
+	Properties->InitializeFromMesh(Cast<IMeshDescriptionProvider>(Target)->GetMeshDescription());
 
 	Properties->WatchProperty(Properties->WeightMap,
 		[&](FName) { OnSelectedWeightMapChanged(true); });
@@ -391,7 +387,7 @@ void UBaseMeshProcessingTool::OnSelectedWeightMapChanged(bool bInvalidate)
 	TSharedPtr<FIndexedWeightMap1f> NewWeightMap = MakeShared<FIndexedWeightMap1f>();
 
 	// this will return all-ones weight map if None is selected
-	bool bFound = UE::WeightMaps::GetVertexWeightMap(ComponentTarget->GetMesh(), WeightMapPropertySet->WeightMap, *NewWeightMap, 1.0f);
+	bool bFound = UE::WeightMaps::GetVertexWeightMap(Cast<IMeshDescriptionProvider>(Target)->GetMeshDescription(), WeightMapPropertySet->WeightMap, *NewWeightMap, 1.0f);
 	if (bFound && WeightMapPropertySet->bInvertWeightMap)
 	{
 		NewWeightMap->InvertWeightMap();
diff --git a/Engine/Plugins/Experimental/MeshModelingToolset/Source/ModelingComponents/Private/BaseTools/SingleSelectionMeshEditingTool.cpp b/Engine/Plugins/Experimental/MeshModelingToolset/Source/ModelingComponents/Private/BaseTools/SingleSelectionMeshEditingTool.cpp
new file mode 100644
--- /dev/null
+++ b/Engine/Plugins/Experimental/MeshModelingToolset/Source/ModelingComponents/Private/BaseTools/SingleSelectionMeshEditingTool.cpp
@@ -0,0 +1,51 @@
+// Copyright Epic Games, Inc. All Rights Reserved.
+
+#include "BaseTools/SingleSelectionMeshEditingTool.h"
+
+#include "TargetInterfaces/MaterialProvider.h"
+#include "TargetInterfaces/MeshDescriptionCommitter.h"
+#include "TargetInterfaces/MeshDescriptionProvider.h"
+#include "TargetInterfaces/PrimitiveComponentBackedTarget.h"
+#include "ToolTargetManager.h"
+
+/*
+ * ToolBuilder
+ */
+
+const FToolTargetTypeRequirements& USingleSelectionMeshEditingToolBuilder::GetTargetRequirements() const
+{
+	static FToolTargetTypeRequirements TypeRequirements({
+		UMaterialProvider::StaticClass(),
+		UMeshDescriptionCommitter::StaticClass(),
+		UMeshDescriptionProvider::StaticClass(),
+		UPrimitiveComponentBackedTarget::StaticClass()
+		});
+	return TypeRequirements;
+}
+
+bool USingleSelectionMeshEditingToolBuilder::CanBuildTool(const FToolBuilderState& SceneState) const
+{
+	bool bValidAssetAPI = (RequiresAssetAPI() ? AssetAPI != nullptr : true);
+	return bValidAssetAPI && SceneState.TargetManager->CountSelectedAndTargetable(SceneState, GetTargetRequirements()) == 1;
+}
+
+UInteractiveTool* USingleSelectionMeshEditingToolBuilder::BuildTool(const FToolBuilderState& SceneState) const
+{
+	USingleSelectionMeshEditingTool* NewTool = CreateNewTool(SceneState);
+	InitializeNewTool(NewTool, SceneState);
+	return NewTool;
+}
+
+void USingleSelectionMeshEditingToolBuilder::InitializeNewTool(USingleSelectionMeshEditingTool* NewTool, const FToolBuilderState& SceneState) const
+{
+	UToolTarget* Target = SceneState.TargetManager->BuildFirstSelectedTargetable(SceneState, GetTargetRequirements());
+	check(Target);
+	NewTool->SetTarget(Target);
+	NewTool->SetWorld(SceneState.World);
+	if (RequiresAssetAPI())
+	{
+		NewTool->SetAssetAPI(AssetAPI);
+	}
+}
+
+
diff --git a/Engine/Plugins/Experimental/MeshModelingToolset/Source/ModelingComponents/Public/BaseTools/BaseMeshProcessingTool.h b/Engine/Plugins/Experimental/MeshModelingToolset/Source/ModelingComponents/Public/BaseTools/BaseMeshProcessingTool.h
--- a/Engine/Plugins/Experimental/MeshModelingToolset/Source/ModelingComponents/Public/BaseTools/BaseMeshProcessingTool.h
+++ b/Engine/Plugins/Experimental/MeshModelingToolset/Source/ModelingComponents/Public/BaseTools/BaseMeshProcessingTool.h
@@ -33,12 +33,11 @@ public:
 	virtual UInteractiveTool* BuildTool(const FToolBuilderState& SceneState) const override;
 
 public:
-	
-	virtual bool SupportsMultipleObjects() const { return false;  }
-	
 	// subclass must override!
 	virtual UBaseMeshProcessingTool* MakeNewToolInstance(UObject* Outer) const { check(false); return nullptr; }
 
+protected:
+	virtual const FToolTargetTypeRequirements& GetTargetRequirements() const override;
 };
 
 
diff --git a/Engine/Plugins/Experimental/MeshModelingToolset/Source/ModelingComponents/Public/BaseTools/SingleSelectionMeshEditingTool.h b/Engine/Plugins/Experimental/MeshModelingToolset/Source/ModelingComponents/Public/BaseTools/SingleSelectionMeshEditingTool.h
new file mode 100644
--- /dev/null
+++ b/Engine/Plugins/Experimental/MeshModelingToolset/Source/ModelingComponents/Public/BaseTools/SingleSelectionMeshEditingTool.h
@@ -0,0 +1,61 @@
+// Copyright Epic Games, Inc. All Rights Reserved.
+
+#pragma once
+
+#include "SingleSelectionTool.h"
+#include "InteractiveToolBuilder.h"
+#include "SingleSelectionMeshEditingTool.generated.h"
+
+class IAssetGenerationAPI;
+class USingleSelectionMeshEditingTool;
+
+/**
+ * USingleSelectionMeshEditingToolBuilder is a base tool builder for single
+ * selection tools that define a common set of ToolTarget interfaces required
+ * for editing meshes.
+ */
+UCLASS(Transient, Abstract)
+class MODELINGCOMPONENTS_API USingleSelectionMeshEditingToolBuilder : public UInteractiveToolBuilder
+{
+	GENERATED_BODY()
+
+public:
+	IAssetGenerationAPI* AssetAPI = nullptr;
+
+public:
+	/** @return true if a single mesh source can be found in the active selection */
+	virtual bool CanBuildTool(const FToolBuilderState& SceneState) const override;
+
+	/** @return new Tool instance initialized with selected mesh source */
+	virtual UInteractiveTool* BuildTool(const FToolBuilderState& SceneState) const override;
+
+	/** @return new Tool instance. Override this in subclasses to build a different Tool class type */
+	virtual USingleSelectionMeshEditingTool* CreateNewTool(const FToolBuilderState& SceneState) const PURE_VIRTUAL(USingleSelectionMeshEditingToolBuilder::CreateNewTool, return nullptr; );
+
+	/** Called by BuildTool to configure the Tool with the input MeshSource based on the SceneState */
+	virtual void InitializeNewTool(USingleSelectionMeshEditingTool* Tool, const FToolBuilderState& SceneState) const;
+
+	/** @return true if this tool requires an AssetAPI */
+	virtual bool RequiresAssetAPI() const { return false; }
+
+protected:
+	virtual const FToolTargetTypeRequirements& GetTargetRequirements() const override;
+};
+
+
+/**
+ * Single Selection Mesh Editing tool base class.
+ */
+UCLASS()
+class MODELINGCOMPONENTS_API USingleSelectionMeshEditingTool : public USingleSelectionTool
+{
+	GENERATED_BODY()
+
+public:
+	virtual void SetWorld(UWorld* World) { TargetWorld = World; }
+	virtual void SetAssetAPI(IAssetGenerationAPI* AssetAPIIn) { AssetAPI = AssetAPIIn; }
+
+protected:
+	UWorld* TargetWorld = nullptr;
+	IAssetGenerationAPI* AssetAPI = nullptr;
+};
diff --git a/Engine/Plugins/Experimental/ModelingToolsEditorMode/Source/ModelingToolsEditorMode/Private/ModelingToolsEditorMode.cpp b/Engine/Plugins/Experimental/ModelingToolsEditorMode/Source/ModelingToolsEditorMode/Private/ModelingToolsEditorMode.cpp
--- a/Engine/Plugins/Experimental/ModelingToolsEditorMode/Source/ModelingToolsEditorMode/Private/ModelingToolsEditorMode.cpp
+++ b/Engine/Plugins/Experimental/ModelingToolsEditorMode/Source/ModelingToolsEditorMode/Private/ModelingToolsEditorMode.cpp
@@ -432,11 +432,9 @@ void UModelingToolsEditorMode::Enter()
 	RegisterTool(ToolManagerCommands.BeginSimplifyMeshTool, TEXT("BeginSimplifyMeshTool"), NewObject<USimplifyMeshToolBuilder>());
 
 	auto GroupEdgeInsertionToolBuilder = NewObject<UGroupEdgeInsertionToolBuilder>();
-	GroupEdgeInsertionToolBuilder->AssetAPI = ToolsContext->GetAssetAPI();
 	RegisterTool(ToolManagerCommands.BeginGroupEdgeInsertionTool, TEXT("BeginGroupEdgeInsertionTool"), GroupEdgeInsertionToolBuilder);
 
 	auto EdgeLoopInsertionToolBuilder = NewObject<UEdgeLoopInsertionToolBuilder>();
-	EdgeLoopInsertionToolBuilder->AssetAPI = ToolsContext->GetAssetAPI();
 	RegisterTool(ToolManagerCommands.BeginEdgeLoopInsertionTool, TEXT("BeginEdgeLoopInsertionTool"), EdgeLoopInsertionToolBuilder);
 
 	auto EditNormalsToolBuilder = NewObject<UEditNormalsToolBuilder>();
diff --git a/Engine/Source/Runtime/Experimental/InteractiveToolsFramework/Public/SingleSelectionTool.h b/Engine/Source/Runtime/Experimental/InteractiveToolsFramework/Public/SingleSelectionTool.h
--- a/Engine/Source/Runtime/Experimental/InteractiveToolsFramework/Public/SingleSelectionTool.h
+++ b/Engine/Source/Runtime/Experimental/InteractiveToolsFramework/Public/SingleSelectionTool.h
@@ -13,24 +13,17 @@ class INTERACTIVETOOLSFRAMEWORK_API USingleSelectionTool : public UInteractiveTo
 {
 GENERATED_BODY()
 public:
-	// @deprecated Use SetTarget instead, and don't use FPrimitiveComponentTarget.
-	void SetSelection(TUniquePtr<FPrimitiveComponentTarget> ComponentTargetIn)
-    {
-		ComponentTarget = MoveTemp(ComponentTargetIn);
-	}
-
 	void SetTarget(UToolTarget* TargetIn)
 	{
 		Target = TargetIn;
 	}
 
 	/**
-	 * @return true if all ComponentTargets of this tool are still valid
+	 * @return true if all ToolTargets of this tool are still valid
 	 */
 	virtual bool AreAllTargetsValid() const
 	{
-		// TODO: This needs to be updated once tools no longer use ComponentTarget.
-		return Target ? Target->IsValid() : ComponentTarget->IsValid();
+		return Target ? Target->IsValid() : false;
 	}
 
 
@@ -41,9 +34,6 @@ public:
 	}
 
 protected:
-	/** @deprecated Tools should use Target instead. */
-	TUniquePtr<FPrimitiveComponentTarget> ComponentTarget{};
-
 	UPROPERTY()
 	TObjectPtr<UToolTarget> Target;
 };
