MeshModelingTools: Made it possible to use delete key in TriEdit, PolyEdit, and draw and revolve. Also corrected transaction names in PolyEdit.

#rnx

diff --git a/Engine/Plugins/Experimental/MeshModelingToolset/Source/MeshModelingTools/Private/DrawAndRevolveTool.cpp b/Engine/Plugins/Experimental/MeshModelingToolset/Source/MeshModelingTools/Private/DrawAndRevolveTool.cpp
--- a/Engine/Plugins/Experimental/MeshModelingToolset/Source/MeshModelingTools/Private/DrawAndRevolveTool.cpp
+++ b/Engine/Plugins/Experimental/MeshModelingToolset/Source/MeshModelingTools/Private/DrawAndRevolveTool.cpp
@@ -83,17 +83,21 @@ TUniquePtr<FDynamicMeshOperator> URevolveOperatorFactory::MakeNewOperator()
 void UDrawAndRevolveTool::RegisterActions(FInteractiveToolActionSet& ActionSet)
 {
 	ActionSet.RegisterAction(this, (int32)EStandardToolActions::BaseClientDefinedActionID + 1,
-		TEXT("DeletePoint"),
+		TEXT("DeletePointBackspaceKey"),
 		LOCTEXT("DeletePointUIName", "Delete Point"),
 		LOCTEXT("DeletePointTooltip", "Delete currently selected point(s)"),
 		EModifierKey::None, EKeys::BackSpace,
-		[this]() { OnBackspacePress(); });
+		[this]() { OnPointDeletionKeyPress(); });
+	ActionSet.RegisterAction(this, (int32)EStandardToolActions::BaseClientDefinedActionID + 2,
+		TEXT("DeletePointDeleteKey"),
+		LOCTEXT("DeletePointUIName", "Delete Point"),
+		LOCTEXT("DeletePointTooltip", "Delete currently selected point(s)"),
+		EModifierKey::None, EKeys::Delete,
+		[this]() { OnPointDeletionKeyPress(); });
 }
 
-void UDrawAndRevolveTool::OnBackspacePress()
+void UDrawAndRevolveTool::OnPointDeletionKeyPress()
 {
-	// TODO: Someday we'd like the mechanic to register the action itself and not have to catch and forward
-	// it in the tool.
 	if (ControlPointsMechanic)
 	{
 		ControlPointsMechanic->DeleteSelectedPoints();
diff --git a/Engine/Plugins/Experimental/MeshModelingToolset/Source/MeshModelingTools/Private/EditMeshPolygonsTool.cpp b/Engine/Plugins/Experimental/MeshModelingToolset/Source/MeshModelingTools/Private/EditMeshPolygonsTool.cpp
--- a/Engine/Plugins/Experimental/MeshModelingToolset/Source/MeshModelingTools/Private/EditMeshPolygonsTool.cpp
+++ b/Engine/Plugins/Experimental/MeshModelingToolset/Source/MeshModelingTools/Private/EditMeshPolygonsTool.cpp
@@ -489,6 +489,27 @@ void UEditMeshPolygonsTool::RegisterActions(FInteractiveToolActionSet& ActionSet
 		LOCTEXT("ToggleLockRotationTooltip", "Toggle Frame Rotation Lock on and off"),
 		EModifierKey::None, EKeys::Q,
 		[this]() { CommonProps->bLockRotation = !CommonProps->bLockRotation; });
+	
+	// Backspace and delete both trigger deletion (as long as the delete button is also enabled)
+	auto OnDeletionKeyPress = [this]() 
+	{
+		if ((EditActions && EditActions->IsPropertySetEnabled())
+			|| (EditActions_Triangles && EditActions_Triangles->IsPropertySetEnabled()))
+		{
+			RequestAction(EEditMeshPolygonsToolActions::Delete);
+		}
+	};
+	ActionSet.RegisterAction(this, (int32)EStandardToolActions::BaseClientDefinedActionID + 3,
+		TEXT("DeleteSelectionBackSpaceKey"),
+		LOCTEXT("DeleteSelectionUIName", "Delete Selection"),
+		LOCTEXT("DeleteSelectionTooltip", "Delete Selection"),
+		EModifierKey::None, EKeys::BackSpace, OnDeletionKeyPress);
+
+	ActionSet.RegisterAction(this, (int32)EStandardToolActions::BaseClientDefinedActionID + 4,
+		TEXT("DeleteSelectionDeleteKey"),
+		LOCTEXT("DeleteSelectionUIName", "Delete Selection"),
+		LOCTEXT("DeleteSelectionTooltip", "Delete Selection"),
+		EModifierKey::None, EKeys::Delete, OnDeletionKeyPress);
 }
 
 
@@ -870,7 +891,7 @@ void UEditMeshPolygonsTool::OnTick(float DeltaTime)
 
 		if (PendingAction == EEditMeshPolygonsToolActions::Extrude || PendingAction == EEditMeshPolygonsToolActions::Offset)
 		{
-			GetToolManager()->EmitObjectChange(this, MakeUnique<FBeginInteractivePolyEditChange>(CurrentOperationTimestamp), LOCTEXT("PolyMeshEditBeginExtrude", "Extrude"));
+			GetToolManager()->EmitObjectChange(this, MakeUnique<FBeginInteractivePolyEditChange>(CurrentOperationTimestamp), LOCTEXT("PolyMeshEditBeginExtrude", "Begin Extrude"));
 			BeginExtrude( (PendingAction == EEditMeshPolygonsToolActions::Offset) );
 		}
 		else if (PendingAction == EEditMeshPolygonsToolActions::Inset)
@@ -885,12 +906,12 @@ void UEditMeshPolygonsTool::OnTick(float DeltaTime)
 		}
 		else if (PendingAction == EEditMeshPolygonsToolActions::CutFaces)
 		{
-			GetToolManager()->EmitObjectChange(this, MakeUnique<FBeginInteractivePolyEditChange>(CurrentOperationTimestamp), LOCTEXT("PolyMeshEditBeginCutFaces", "Cut Faces"));
+			GetToolManager()->EmitObjectChange(this, MakeUnique<FBeginInteractivePolyEditChange>(CurrentOperationTimestamp), LOCTEXT("PolyMeshEditBeginCutFaces", "Begin Cut Faces"));
 			BeginCutFaces();
 		}
 		else if (PendingAction == EEditMeshPolygonsToolActions::PlanarProjectionUV)
 		{
-			GetToolManager()->EmitObjectChange(this, MakeUnique<FBeginInteractivePolyEditChange>(CurrentOperationTimestamp), LOCTEXT("PolyMeshEditBeginUVPlanarProjection", "Set UVs"));
+			GetToolManager()->EmitObjectChange(this, MakeUnique<FBeginInteractivePolyEditChange>(CurrentOperationTimestamp), LOCTEXT("PolyMeshEditBeginUVPlanarProjection", "Begin Set UVs"));
 			BeginSetUVs();
 		}
 		else if (PendingAction == EEditMeshPolygonsToolActions::Merge)
diff --git a/Engine/Plugins/Experimental/MeshModelingToolset/Source/MeshModelingTools/Public/DrawAndRevolveTool.h b/Engine/Plugins/Experimental/MeshModelingToolset/Source/MeshModelingTools/Public/DrawAndRevolveTool.h
--- a/Engine/Plugins/Experimental/MeshModelingToolset/Source/MeshModelingTools/Public/DrawAndRevolveTool.h
+++ b/Engine/Plugins/Experimental/MeshModelingToolset/Source/MeshModelingTools/Public/DrawAndRevolveTool.h
@@ -93,7 +93,7 @@ public:
 	virtual void SetAssetAPI(IAssetGenerationAPI* NewAssetApi) { AssetAPI = NewAssetApi; }
 
 	virtual void RegisterActions(FInteractiveToolActionSet& ActionSet) override;
-	void OnBackspacePress();
+	virtual void OnPointDeletionKeyPress();
 
 	virtual bool HasCancel() const override { return true; }
 	virtual bool HasAccept() const override { return true; }
