ModelingTools: clear transactional flag on tool property sets added via CreateDefaultSubobjects
#rnx

diff --git a/Engine/Plugins/Experimental/MeshModelingToolset/Source/MeshModelingTools/Private/AddPrimitiveTool.cpp b/Engine/Plugins/Experimental/MeshModelingToolset/Source/MeshModelingTools/Private/AddPrimitiveTool.cpp
--- a/Engine/Plugins/Experimental/MeshModelingToolset/Source/MeshModelingTools/Private/AddPrimitiveTool.cpp
+++ b/Engine/Plugins/Experimental/MeshModelingToolset/Source/MeshModelingTools/Private/AddPrimitiveTool.cpp
@@ -117,6 +117,8 @@ void UAddPrimitiveTool::SetAssetAPI(IAssetGenerationAPI* AssetAPIIn)
 UAddPrimitiveTool::UAddPrimitiveTool(const FObjectInitializer&)
 {
 	ShapeSettings = CreateDefaultSubobject<UProceduralShapeToolProperties>(TEXT("ShapeSettings"));
+	// CreateDefaultSubobject automatically sets RF_Transactional flag, we need to clear it so that undo/redo doesn't affect tool properties
+	ShapeSettings->ClearFlags(RF_Transactional);
 }
 
 void UAddPrimitiveTool::Setup()
diff --git a/Engine/Plugins/Experimental/MeshModelingToolset/Source/MeshModelingTools/Private/RemeshMeshTool.cpp b/Engine/Plugins/Experimental/MeshModelingToolset/Source/MeshModelingTools/Private/RemeshMeshTool.cpp
--- a/Engine/Plugins/Experimental/MeshModelingToolset/Source/MeshModelingTools/Private/RemeshMeshTool.cpp
+++ b/Engine/Plugins/Experimental/MeshModelingToolset/Source/MeshModelingTools/Private/RemeshMeshTool.cpp
@@ -88,6 +88,8 @@ URemeshMeshToolProperties::URemeshMeshToolProperties()
 URemeshMeshTool::URemeshMeshTool(const FObjectInitializer&)
 {
 	BasicProperties = CreateDefaultSubobject<URemeshMeshToolProperties>(TEXT("RemeshProperties"));
+	// CreateDefaultSubobject automatically sets RF_Transactional flag, we need to clear it so that undo/redo doesn't affect tool properties
+	BasicProperties->ClearFlags(RF_Transactional);
 }
 
 void URemeshMeshTool::SetWorld(UWorld* World)
