ModelingTools: disable lazy/spacing/flowrate settings in DynamicMeshSculptTool because they aren't supported there yet
#rnx

diff --git a/Engine/Plugins/Experimental/MeshModelingToolset/Source/MeshModelingTools/Private/DynamicMeshSculptTool.cpp b/Engine/Plugins/Experimental/MeshModelingToolset/Source/MeshModelingTools/Private/DynamicMeshSculptTool.cpp
--- a/Engine/Plugins/Experimental/MeshModelingToolset/Source/MeshModelingTools/Private/DynamicMeshSculptTool.cpp
+++ b/Engine/Plugins/Experimental/MeshModelingToolset/Source/MeshModelingTools/Private/DynamicMeshSculptTool.cpp
@@ -151,6 +151,9 @@ void UDynamicMeshSculptTool::Setup()
 	BrushRelativeSizeRange = FInterval1d(MaxDimension*0.01, MaxDimension);
 	BrushProperties = NewObject<USculptBrushProperties>(this);
 	BrushProperties->bShowStrength = false;
+	BrushProperties->bShowFlowRate = false;
+	BrushProperties->bShowSpacing = false;
+	BrushProperties->bShowLazyness = false;
 	CalculateBrushRadius();
 
 	// initialize other properties
