HoleFill tool: optionally try to fill single-triangle holes first. Don't create new UV islands for these filled triangles, instead try to attach them to adjacent UV islands.

diff --git a/Engine/Plugins/Experimental/GeometryProcessing/Source/DynamicMesh/Public/DynamicMeshOverlay.h b/Engine/Plugins/Experimental/GeometryProcessing/Source/DynamicMesh/Public/DynamicMeshOverlay.h
--- a/Engine/Plugins/Experimental/GeometryProcessing/Source/DynamicMesh/Public/DynamicMeshOverlay.h
+++ b/Engine/Plugins/Experimental/GeometryProcessing/Source/DynamicMesh/Public/DynamicMeshOverlay.h
@@ -222,6 +222,11 @@ public:
 	/** Allocate a new element with the given value */
 	int AppendElement(const RealType* Value);
 
+	void SetParentVertex(int ElementIndex, int ParentVertexIndex)
+	{
+		ParentVertices[ElementIndex] = ParentVertexIndex;
+	}
+
 	/** Initialize the triangle list to the given size, and set all triangles to InvalidID */
 	void InitializeTriangles(int MaxTriangleID);
 	/** Set the triangle to the given Element index tuple, and increment element reference counts */
diff --git a/Engine/Plugins/Experimental/MeshModelingToolset/Source/MeshModelingTools/Private/HoleFillTool.cpp b/Engine/Plugins/Experimental/MeshModelingToolset/Source/MeshModelingTools/Private/HoleFillTool.cpp
--- a/Engine/Plugins/Experimental/MeshModelingToolset/Source/MeshModelingTools/Private/HoleFillTool.cpp
+++ b/Engine/Plugins/Experimental/MeshModelingToolset/Source/MeshModelingTools/Private/HoleFillTool.cpp
@@ -101,6 +101,7 @@ TUniquePtr<FDynamicMeshOperator> UHoleFillOperatorFactory::MakeNewOperator()
 	FillOp->FillType = FillTool->Properties->FillType;
 
 	FillOp->FillOptions.bRemoveIsolatedTriangles = FillTool->Properties->bRemoveIsolatedTriangles;
+	FillOp->FillOptions.bQuickFillSmallHoles = FillTool->Properties->bQuickFillSmallHoles;
 
 	// Smooth fill properties
 	FillOp->SmoothFillOptions = FillTool->SmoothHoleFillProperties->ToSmoothFillOptions();
diff --git a/Engine/Plugins/Experimental/MeshModelingToolset/Source/MeshModelingTools/Public/HoleFillTool.h b/Engine/Plugins/Experimental/MeshModelingToolset/Source/MeshModelingTools/Public/HoleFillTool.h
--- a/Engine/Plugins/Experimental/MeshModelingToolset/Source/MeshModelingTools/Public/HoleFillTool.h
+++ b/Engine/Plugins/Experimental/MeshModelingToolset/Source/MeshModelingTools/Public/HoleFillTool.h
@@ -120,6 +120,10 @@ public:
 	UPROPERTY(EditAnywhere, Category = Options)
 	bool bRemoveIsolatedTriangles = false;
 
+	/** Identify and quickly fill single-triangle holes */
+	UPROPERTY(EditAnywhere, Category = Options)
+	bool bQuickFillSmallHoles = false;
+
 };
 
 
diff --git a/Engine/Plugins/Experimental/MeshModelingToolset/Source/ModelingOperators/Private/CleaningOps/HoleFillOp.cpp b/Engine/Plugins/Experimental/MeshModelingToolset/Source/ModelingOperators/Private/CleaningOps/HoleFillOp.cpp
--- a/Engine/Plugins/Experimental/MeshModelingToolset/Source/ModelingOperators/Private/CleaningOps/HoleFillOp.cpp
+++ b/Engine/Plugins/Experimental/MeshModelingToolset/Source/ModelingOperators/Private/CleaningOps/HoleFillOp.cpp
@@ -8,6 +8,7 @@
 #include "ConstrainedDelaunay2.h"
 #include "CompGeom/PolygonTriangulation.h"
 #include "MeshQueries.h"
+#include "Selections/MeshConnectedComponents.h"
 
 #include "ExplicitUseGeometryMathTypes.h"		// using UE::Geometry::(math types)
 using namespace UE::Geometry;
@@ -83,8 +84,211 @@ namespace
 		return Filler;
 	}
 
+	
+	// return index in the parent triangle (0, 1, 2) or -1 if not found
+	int FindParentVertexInBaseTriangle(const FDynamicMeshUVOverlay& UVOverlay, int ElementID, int TriangleID)
+	{
+		const FDynamicMesh3* Mesh = UVOverlay.GetParentMesh();
+		const FIndex3i BaseTriangle = Mesh->GetTriangle(TriangleID);
+		int ParentVertex = UVOverlay.GetParentVertex(ElementID);
+		return BaseTriangle.IndexOf(ParentVertex);
+	}
+
+	// Given two triangles adjacent to a newly added triangle in the base mesh, find the elements that should comprise 
+	// the newly added triangle in the UV mesh.
+	void SetUVTriangleFromExistingTriangles(FDynamicMeshUVOverlay& UVOverlay,
+											int TriangleA,
+											int TriangleB,
+											int NewBaseTriangleIndex,
+											FIndex3i& NewTriangleElements)
+	{
+		const FDynamicMesh3* Mesh = UVOverlay.GetParentMesh();
+		const FIndex3i BaseTriangle = Mesh->GetTriangle(NewBaseTriangleIndex);
+
+		for (const FIndex3i& TriElements : { UVOverlay.GetTriangle(TriangleA), UVOverlay.GetTriangle(TriangleB) })
+		{
+			for (int TriangleElementIndex = 0; TriangleElementIndex < 3; TriangleElementIndex++)
+			{
+				const int ElementID = TriElements[TriangleElementIndex];
+				int IndexInBaseTriangle = FindParentVertexInBaseTriangle(UVOverlay, ElementID, NewBaseTriangleIndex);
+				if (IndexInBaseTriangle >= 0)
+				{
+					NewTriangleElements[IndexInBaseTriangle] = ElementID;
+				}
+			}
+		}
+	}
+
+	float Area(FDynamicMeshUVOverlay& UVOverlay, FIndex3i& Tri)
+	{
+		const FVector2f AB = UVOverlay.GetElement(Tri[1]) - UVOverlay.GetElement(Tri[0]);
+		const FVector2f AC = UVOverlay.GetElement(Tri[2]) - UVOverlay.GetElement(Tri[0]);
+		return 0.5f * FMathd::Abs(AB.X * AC.Y - AB.Y * AC.X);
+	}
+
+	// Given a single triangle adjacent to a newly added triangle in the base mesh, find the two elements that should 
+	// go into a new triangle in the UV mesh (i.e. the shared edge between the adjacent triangle and the new UV triangle.) 
+	// Also add a third UV element, taking the average UV coordinates of the shared edge.
+	void SetUVTriangleFromAdjacentTriangle(FDynamicMeshUVOverlay& UVOverlay,
+										   int AdjacentTriangle,
+										   int NewTriangleID,
+										   FIndex3i& NewTriangleElements)
+	{
+		const FDynamicMesh3* Mesh = UVOverlay.GetParentMesh();
+		FIndex3i AdjacentTriangleElements = UVOverlay.GetTriangle(AdjacentTriangle);
+
+		// Find valid elements on the adjacent triangle that point to the base triangle's vertices. There should be 2.
+		TStaticArray<FVector2f, 2> EdgeElements;
+		int FoundElements = 0;
+		FVector2f OppositeVertexPosition;
+		for (int AdjacentTriVertexIndex = 0; AdjacentTriVertexIndex < 3; ++AdjacentTriVertexIndex)
+		{
+			int IndexInBaseTriangle = FindParentVertexInBaseTriangle(UVOverlay, AdjacentTriangleElements[AdjacentTriVertexIndex], NewTriangleID);
+			if (IndexInBaseTriangle >= 0)
+			{
+				NewTriangleElements[IndexInBaseTriangle] = AdjacentTriangleElements[AdjacentTriVertexIndex];
+				EdgeElements[FoundElements] = UVOverlay.GetElement(AdjacentTriangleElements[AdjacentTriVertexIndex]);
+				++FoundElements;
+			}
+			else
+			{
+				OppositeVertexPosition = UVOverlay.GetElement(AdjacentTriangleElements[AdjacentTriVertexIndex]);
+			}
+		}
+		check(FoundElements == 2);
+
+		// Now insert the new element
+		FIndex3i TriangleVertices = Mesh->GetTriangle(NewTriangleID);
+		FVector2f NewElement = 0.5f * (EdgeElements[0] + EdgeElements[1]);
+		
+		bool bIsDegen = (EdgeElements[0].Distance(EdgeElements[1]) == 0.0f);
+		if (!bIsDegen)
+		{
+			// push the new element off the edge slightly to avoid degenerate triangle (assuming the adjacent triangle 
+			// wasn't degenerate before)
+			FVector2f Delta = (NewElement - OppositeVertexPosition);
+			if (Delta.Length() > 0.0)
+			{
+				NewElement += KINDA_SMALL_NUMBER * Delta.Normalized();
+			}
+			else
+			{
+				bIsDegen = true;
+			}
+		}
+
+		int NewElementIndex = NewTriangleElements.IndexOf(FDynamicMesh3::InvalidID);
+		NewTriangleElements[NewElementIndex] = UVOverlay.AppendElement(NewElement);
+		UVOverlay.SetParentVertex(NewTriangleElements[NewElementIndex], TriangleVertices[NewElementIndex]);
+
+		if (!bIsDegen)
+		{
+			ensure(Area(UVOverlay, NewTriangleElements) > 0.0);
+		}
+	}
+
+} // unnamed namespace
+
+
+bool FHoleFillOp::FillSingleTriangleHole(const FEdgeLoop& Loop, int32& NewGroupID)
+{
+	check(Loop.GetEdgeCount() == 3);
+	check(Loop.GetVertexCount() == 3);
+
+	NewGroupID = ResultMesh->AllocateTriangleGroup();
+
+	FIndex3i Vertices{ Loop.Vertices[0], Loop.Vertices[1], Loop.Vertices[2] };
+
+	// Find adjacent edges and triangles
+	FIndex3i ExistingEdges;
+	FIndex3i ExistingTriangles;
+	for (int Nbr = 0; Nbr < 3; ++Nbr)
+	{
+		ExistingEdges[Nbr] = ResultMesh->FindEdge(Vertices[Nbr], Vertices[(Nbr + 1) % 3]);
+		ensure(ExistingEdges[Nbr] != FDynamicMesh3::InvalidID);
+
+		FIndex2i EdgeTris = ResultMesh->GetEdgeT(ExistingEdges[Nbr]);
+		ensure(EdgeTris[0] != FDynamicMesh3::InvalidID && EdgeTris[1] == FDynamicMesh3::InvalidID);
+
+		ExistingTriangles[Nbr] = EdgeTris[0];
+	}
+
+	FIndex3i ExistingTri0Vertices = ResultMesh->GetTriangle(ExistingTriangles[0]);
+
+	// Check orientation of the loop
+	if (IndexUtil::FindTriOrderedEdge(Vertices[0], Vertices[1], ExistingTri0Vertices) != FDynamicMesh3::InvalidID)
+	{
+		// Orientation matches an existing triangle, so reverse it
+		Swap(Vertices[0], Vertices[1]);
+	}
+
+	int NewTriangleID = ResultMesh->AppendTriangle(Vertices, NewGroupID);
+	if (NewTriangleID < 0)
+	{
+		return false;
+	}
+
+	if (ResultMesh->HasAttributes() && ResultMesh->Attributes()->PrimaryUV())
+	{
+		FDynamicMeshUVOverlay* UVOverlay = ResultMesh->Attributes()->PrimaryUV();
+
+		// check we have a hole in the UV mesh as well
+		check(UVOverlay->IsSeamEdge(ExistingEdges[0]));
+		check(UVOverlay->IsSeamEdge(ExistingEdges[1]));
+		check(UVOverlay->IsSeamEdge(ExistingEdges[2]));
+
+		FIndex3i NewTriangleElements;
+		int NewTriangleUVIsland = -1;
+
+		if (TriangleUVIsland[ExistingTriangles[0]] == TriangleUVIsland[ExistingTriangles[1]])
+		{
+			// Connect new triangle to existing triangles in a common UV island
+			SetUVTriangleFromExistingTriangles(*UVOverlay, ExistingTriangles[0], ExistingTriangles[1], NewTriangleID, NewTriangleElements);
+			NewTriangleUVIsland = TriangleUVIsland[ExistingTriangles[0]];
+		}
+		else if (TriangleUVIsland[ExistingTriangles[1]] == TriangleUVIsland[ExistingTriangles[2]])
+		{
+			SetUVTriangleFromExistingTriangles(*UVOverlay, ExistingTriangles[1], ExistingTriangles[2], NewTriangleID, NewTriangleElements);
+			NewTriangleUVIsland = TriangleUVIsland[ExistingTriangles[1]];
+		}
+		else if (TriangleUVIsland[ExistingTriangles[2]] == TriangleUVIsland[ExistingTriangles[0]])
+		{
+			SetUVTriangleFromExistingTriangles(*UVOverlay, ExistingTriangles[2], ExistingTriangles[0], NewTriangleID, NewTriangleElements);
+			NewTriangleUVIsland = TriangleUVIsland[ExistingTriangles[2]];
+		}
+		else
+		{
+			// All edges are in different UV islands.
+			// Pick an arbitrary edge and add a new third UV element
+			SetUVTriangleFromAdjacentTriangle(*UVOverlay, ExistingTriangles[0], NewTriangleID, NewTriangleElements);
+			NewTriangleUVIsland = TriangleUVIsland[ExistingTriangles[0]];
+		}
+
+		EMeshResult Result = UVOverlay->SetTriangle(NewTriangleID, NewTriangleElements);
+
+		if (!ensure(Result == EMeshResult::Ok))
+		{
+			// TODO: undo changes ResultMesh and UVOverlay before returning? 
+			// Or create a totally new triangle with unique UVs and pass that in as a fallback?
+			return false;
+		}
+
+		TriangleUVIsland.SetNum(NewTriangleID + 1);
+		TriangleUVIsland[NewTriangleID] = NewTriangleUVIsland;
+	}
+
+
+	// And normals
+	if (ResultMesh->HasAttributes())
+	{
+		FDynamicMeshEditor Editor(ResultMesh.Get());
+		Editor.SetTriangleNormals({ NewTriangleID });
+	}
+
+	return true;
 }
 
+
 void FHoleFillOp::CalculateResult(FProgressCancel* Progress)
 {
 	NumFailedLoops = 0;
@@ -106,6 +310,26 @@ void FHoleFillOp::CalculateResult(FProgressCancel* Progress)
 		return;
 	}
 
+	// Initialize TriangleUVIsland, the UVIsland ID for each triangle in the mesh
+	// TODO: Initialize this only if really absolutely necessary
+	if (ResultMesh->HasAttributes() && ResultMesh->Attributes()->PrimaryUV() && FillOptions.bQuickFillSmallHoles)
+	{
+		FDynamicMeshUVOverlay* UVOverlay = ResultMesh->Attributes()->PrimaryUV();
+		FMeshConnectedComponents UVComponents(ResultMesh.Get());
+		UVComponents.FindConnectedTriangles([UVOverlay](int32 Triangle0, int32 Triangle1) {
+			return UVOverlay->AreTrianglesConnected(Triangle0, Triangle1);
+		});
+
+		TriangleUVIsland.Init( FDynamicMesh3::InvalidID, ResultMesh->MaxTriangleID());
+		for (int IslandID = 0; IslandID < UVComponents.Num(); ++IslandID)
+		{
+			for (int TriangleID : UVComponents.GetComponent(IslandID).Indices)
+			{
+				TriangleUVIsland[TriangleID] = IslandID;
+			}
+		}
+	}
+
 	TSet<int32> NewGroupIDs;
 
 	for (FEdgeLoop& Loop : Loops)
@@ -126,6 +350,20 @@ void FHoleFillOp::CalculateResult(FProgressCancel* Progress)
 			}
 		}
 
+		if (FillOptions.bQuickFillSmallHoles && Loop.GetEdgeCount() == 3)
+		{
+			int32 PossibleNewGroupID;
+			bool bFilledSmallHole = FillSingleTriangleHole(Loop, PossibleNewGroupID);
+
+			if (bFilledSmallHole)
+			{
+				// TODO: If FillSmallHole fails, this newly created group ID won't be used. Is that okay?
+				NewGroupIDs.Add(PossibleNewGroupID);
+				continue;
+			}
+		}
+
+
 		int32 NewGroupID = ResultMesh->AllocateTriangleGroup();
 		NewGroupIDs.Add(NewGroupID);
 
diff --git a/Engine/Plugins/Experimental/MeshModelingToolset/Source/ModelingOperators/Public/CleaningOps/HoleFillOp.h b/Engine/Plugins/Experimental/MeshModelingToolset/Source/ModelingOperators/Public/CleaningOps/HoleFillOp.h
--- a/Engine/Plugins/Experimental/MeshModelingToolset/Source/ModelingOperators/Public/CleaningOps/HoleFillOp.h
+++ b/Engine/Plugins/Experimental/MeshModelingToolset/Source/ModelingOperators/Public/CleaningOps/HoleFillOp.h
@@ -35,7 +35,11 @@ namespace Geometry
 
 struct FFillOptions
 {
+	/** Clean up triangles that have no neighbors */
 	bool bRemoveIsolatedTriangles = false;
+
+	/** Identify and quickly fill single-triangle holes */
+	bool bQuickFillSmallHoles = false;
 };
 
 class MODELINGOPERATORS_API FHoleFillOp : public FDynamicMeshOperator
@@ -60,6 +64,8 @@ public:
 	// FDynamicMeshOperator implementation
 	void CalculateResult(FProgressCancel* Progress) override;
 
+	TArray<int> TriangleUVIsland;
+	bool FillSingleTriangleHole(const FEdgeLoop& Loop, int32& NewGroupID);
 };
 
 } // end namespace UE::Geometry
