Fix robomerge build break. Vector member functions have been changed to free functions.

#rnx

diff --git a/Engine/Plugins/Experimental/MeshModelingToolset/Source/ModelingOperators/Private/CleaningOps/HoleFillOp.cpp b/Engine/Plugins/Experimental/MeshModelingToolset/Source/ModelingOperators/Private/CleaningOps/HoleFillOp.cpp
--- a/Engine/Plugins/Experimental/MeshModelingToolset/Source/ModelingOperators/Private/CleaningOps/HoleFillOp.cpp
+++ b/Engine/Plugins/Experimental/MeshModelingToolset/Source/ModelingOperators/Private/CleaningOps/HoleFillOp.cpp
@@ -169,7 +169,7 @@ namespace
 			FVector2f Delta = (NewElement - OppositeVertexPosition);
 			if (Delta.Length() > 0.0)
 			{
-				NewElement += KINDA_SMALL_NUMBER * Delta.Normalized();
+				NewElement += KINDA_SMALL_NUMBER * UE::Geometry::Normalized(Delta);
 			}
 			else
 			{
