Fixed a shadowing variable error with clang on Windows.

#trivial
#rnx

diff --git a/Engine/Plugins/Experimental/MeshModelingToolset/Source/MeshModelingToolsEditorOnly/Private/Hair/GroomToMeshTool.cpp b/Engine/Plugins/Experimental/MeshModelingToolset/Source/MeshModelingToolsEditorOnly/Private/Hair/GroomToMeshTool.cpp
--- a/Engine/Plugins/Experimental/MeshModelingToolset/Source/MeshModelingToolsEditorOnly/Private/Hair/GroomToMeshTool.cpp
+++ b/Engine/Plugins/Experimental/MeshModelingToolset/Source/MeshModelingToolsEditorOnly/Private/Hair/GroomToMeshTool.cpp
@@ -243,13 +243,13 @@ void UGroomToMeshTool::RecalculateMesh()
 	TSharedPtr<FDynamicMesh3, ESPMode::ThreadSafe> CurrentResult = UpdateVoxelization();
 
 	// helper to only use NewMesh if it was valid, otherwise ignore it
-	auto SelectValidMesh = [](TSharedPtr<FDynamicMesh3, ESPMode::ThreadSafe> CurrentMesh, TSharedPtr<FDynamicMesh3, ESPMode::ThreadSafe> NewMesh)
+	auto SelectValidMesh = [](TSharedPtr<FDynamicMesh3, ESPMode::ThreadSafe> InCurrentMesh, TSharedPtr<FDynamicMesh3, ESPMode::ThreadSafe> InNewMesh)
 	{
-		if (NewMesh->TriangleCount() > 0)
+		if (InNewMesh->TriangleCount() > 0)
 		{
-			return NewMesh;
+			return InNewMesh;
 		}
-		return CurrentMesh;
+		return InCurrentMesh;
 	};
 
 	if (Settings->bApplyMorphology && CurrentResult->TriangleCount() > 0)
