MeshModelingTools: Added the ability to use a transparent or custom material to the sculpt tools.

#rnx

diff --git a/Engine/Plugins/Experimental/MeshModelingToolset/Source/MeshModelingTools/Private/DynamicMeshSculptTool.cpp b/Engine/Plugins/Experimental/MeshModelingToolset/Source/MeshModelingTools/Private/DynamicMeshSculptTool.cpp
--- a/Engine/Plugins/Experimental/MeshModelingToolset/Source/MeshModelingTools/Private/DynamicMeshSculptTool.cpp
+++ b/Engine/Plugins/Experimental/MeshModelingToolset/Source/MeshModelingTools/Private/DynamicMeshSculptTool.cpp
@@ -231,6 +231,9 @@ void UDynamicMeshSculptTool::Setup()
 	MaterialModeWatcher.Initialize(
 		[this]() { return ViewProperties->MaterialMode; },
 		[this](EMeshEditingMaterialModes NewMode) { UpdateMaterialMode(NewMode); }, EMeshEditingMaterialModes::ExistingMaterial);
+	CustomMaterialWatcher.Initialize( 
+		[this]() { return ViewProperties->CustomMaterial; },
+		[this](TWeakObjectPtr<UMaterialInterface> NewMaterial) { UpdateCustomMaterial(NewMaterial); }, ViewProperties->CustomMaterial);
 	FlatShadingWatcher.Initialize(
 		[this]() { return ViewProperties->bFlatShading; },
 		[this](bool bNewValue) { UpdateFlatShadingSetting(bNewValue); }, ViewProperties->bFlatShading);
@@ -240,6 +243,15 @@ void UDynamicMeshSculptTool::Setup()
 	ImageWatcher.Initialize(
 		[this]() { return ViewProperties->Image; },
 		[this](UTexture2D* NewImage) { UpdateImageSetting(NewImage); }, ViewProperties->Image);
+	TransparentColorWatcher.Initialize(
+		[this]() { return ViewProperties->TransparentMaterialColor; },
+		[this](FLinearColor NewColor) { UpdateColorSetting(NewColor); }, ViewProperties->TransparentMaterialColor);
+	OpacityWatcher.Initialize(
+		[this]() { return ViewProperties->Opacity; },
+		[this](double Opacity) { UpdateOpacitySetting(Opacity); }, ViewProperties->Opacity);
+	TwoSidedWatcher.Initialize(
+		[this]() { return ViewProperties->bTwoSided; },
+		[this](bool bOn) { UpdateTwoSidedSetting(bOn); }, ViewProperties->bTwoSided);
 	BrushTypeWatcher.Initialize(
 		[this]() { return SculptProperties->PrimaryBrushType; },
 		[this](EDynamicMeshSculptBrushType NewBrushType) { UpdateBrushType(NewBrushType); }, SculptProperties->PrimaryBrushType);
@@ -1505,8 +1517,12 @@ void UDynamicMeshSculptTool::OnTick(float DeltaTime)
 
 	ShowWireframeWatcher.CheckAndUpdate();
 	MaterialModeWatcher.CheckAndUpdate();
+	CustomMaterialWatcher.CheckAndUpdate();
 	FlatShadingWatcher.CheckAndUpdate();
 	ColorWatcher.CheckAndUpdate();
+	TransparentColorWatcher.CheckAndUpdate();
+	OpacityWatcher.CheckAndUpdate();
+	TwoSidedWatcher.CheckAndUpdate();
 	ImageWatcher.CheckAndUpdate();
 	BrushTypeWatcher.CheckAndUpdate();
 	GizmoPositionWatcher.CheckAndUpdate();
@@ -2455,6 +2471,18 @@ void UDynamicMeshSculptTool::UpdateMaterialMode(EMeshEditingMaterialModes Materi
 	else 
 	{
 		if (MaterialMode == EMeshEditingMaterialModes::Custom)
+		{
+			if (ViewProperties->CustomMaterial.IsValid())
+			{
+				ActiveOverrideMaterial = UMaterialInstanceDynamic::Create(ViewProperties->CustomMaterial.Get(), this);
+			}
+			else
+			{
+				DynamicMeshComponent->ClearOverrideRenderMaterial();
+				ActiveOverrideMaterial = nullptr;
+			}
+		}
+		else if (MaterialMode == EMeshEditingMaterialModes::CustomImage)
 		{
 			ActiveOverrideMaterial = ToolSetupUtil::GetCustomImageBasedSculptMaterial(GetToolManager(), ViewProperties->Image);
 			if (ViewProperties->Image != nullptr)
@@ -2462,6 +2490,11 @@ void UDynamicMeshSculptTool::UpdateMaterialMode(EMeshEditingMaterialModes Materi
 				ActiveOverrideMaterial->SetTextureParameterValue(TEXT("ImageTexture"), ViewProperties->Image);
 			}
 		}
+		else if (MaterialMode == EMeshEditingMaterialModes::Transparent)
+		{
+			ActiveOverrideMaterial = ToolSetupUtil::GetTransparentSculptMaterial(GetToolManager(), 
+				ViewProperties->TransparentMaterialColor, ViewProperties->Opacity, ViewProperties->bTwoSided);
+		}
 		else
 		{
 			UMaterialInterface* SculptMaterial = nullptr;
@@ -2517,6 +2550,41 @@ void UDynamicMeshSculptTool::UpdateColorSetting(FLinearColor NewColor)
 	}
 }
 
+void UDynamicMeshSculptTool::UpdateOpacitySetting(double Opacity)
+{
+	if (ActiveOverrideMaterial != nullptr)
+	{
+		ActiveOverrideMaterial->SetScalarParameterValue(TEXT("Opacity"), Opacity);
+	}
+}
+
+void UDynamicMeshSculptTool::UpdateTwoSidedSetting(bool bOn)
+{
+	ActiveOverrideMaterial = ToolSetupUtil::GetTransparentSculptMaterial(GetToolManager(),
+		ViewProperties->TransparentMaterialColor, ViewProperties->Opacity, bOn);
+	if (ActiveOverrideMaterial)
+	{
+		DynamicMeshComponent->SetOverrideRenderMaterial(ActiveOverrideMaterial);
+	}
+}
+
+void UDynamicMeshSculptTool::UpdateCustomMaterial(TWeakObjectPtr<UMaterialInterface> NewMaterial)
+{
+	if (ViewProperties->MaterialMode == EMeshEditingMaterialModes::Custom)
+	{
+		if (NewMaterial.IsValid())
+		{
+			ActiveOverrideMaterial = UMaterialInstanceDynamic::Create(NewMaterial.Get(), this);
+			DynamicMeshComponent->SetOverrideRenderMaterial(ActiveOverrideMaterial);
+		}
+		else
+		{
+			DynamicMeshComponent->ClearOverrideRenderMaterial();
+			ActiveOverrideMaterial = nullptr;
+		}
+	}
+}
+
 void UDynamicMeshSculptTool::UpdateImageSetting(UTexture2D* NewImage)
 {
 	if (ActiveOverrideMaterial != nullptr)
diff --git a/Engine/Plugins/Experimental/MeshModelingToolset/Source/MeshModelingTools/Private/Sculpting/MeshSculptToolBase.cpp b/Engine/Plugins/Experimental/MeshModelingToolset/Source/MeshModelingTools/Private/Sculpting/MeshSculptToolBase.cpp
--- a/Engine/Plugins/Experimental/MeshModelingToolset/Source/MeshModelingTools/Private/Sculpting/MeshSculptToolBase.cpp
+++ b/Engine/Plugins/Experimental/MeshModelingToolset/Source/MeshModelingTools/Private/Sculpting/MeshSculptToolBase.cpp
@@ -80,10 +80,19 @@ void UMeshSculptToolBase::Setup()
 		[this](bool bNewValue) { UpdateWireframeVisibility(bNewValue); });
 	ViewProperties->WatchProperty(ViewProperties->MaterialMode,
 		[this](EMeshEditingMaterialModes NewMode) { UpdateMaterialMode(NewMode); });
+	ViewProperties->WatchProperty(ViewProperties->CustomMaterial,
+		[this](TWeakObjectPtr<UMaterialInterface> NewMaterial) { UpdateCustomMaterial(NewMaterial); });
 	ViewProperties->WatchProperty(ViewProperties->bFlatShading,
 		[this](bool bNewValue) { UpdateFlatShadingSetting(bNewValue); });
 	ViewProperties->WatchProperty(ViewProperties->Color,
+		[this](FLinearColor NewColor) { UpdateColorSetting(NewColor); }); 
+	// This can actually use the same function since the parameter names for the material are the same
+	ViewProperties->WatchProperty(ViewProperties->TransparentMaterialColor,
 		[this](FLinearColor NewColor) { UpdateColorSetting(NewColor); });
+	ViewProperties->WatchProperty(ViewProperties->Opacity,
+		[this](double NewValue) { UpdateOpacitySetting(NewValue); });
+	ViewProperties->WatchProperty(ViewProperties->bTwoSided,
+		[this](bool bOn) { UpdateTwoSidedSetting(bOn); });
 	ViewProperties->WatchProperty(ViewProperties->Image,
 		[this](UTexture2D* NewImage) { UpdateImageSetting(NewImage); });
 }
@@ -908,6 +917,43 @@ void UMeshSculptToolBase::UpdateImageSetting(UTexture2D* NewImage)
 	}
 }
 
+void UMeshSculptToolBase::UpdateOpacitySetting(double Opacity)
+{
+	if (ActiveOverrideMaterial != nullptr)
+	{
+		ActiveOverrideMaterial->SetScalarParameterValue(TEXT("Opacity"), Opacity);
+	}
+}
+
+void UMeshSculptToolBase::UpdateTwoSidedSetting(bool bOn)
+{
+	if (ViewProperties->MaterialMode == EMeshEditingMaterialModes::Transparent)
+	{
+		ActiveOverrideMaterial = ToolSetupUtil::GetTransparentSculptMaterial(GetToolManager(),
+			ViewProperties->TransparentMaterialColor, ViewProperties->Opacity, bOn);
+		if (ActiveOverrideMaterial)
+		{
+			GetSculptMeshComponent()->SetOverrideRenderMaterial(ActiveOverrideMaterial);
+		}
+	}
+}
+
+void UMeshSculptToolBase::UpdateCustomMaterial(TWeakObjectPtr<UMaterialInterface> NewMaterial)
+{
+	if (ViewProperties->MaterialMode == EMeshEditingMaterialModes::Custom)
+	{
+		if (NewMaterial.IsValid())
+		{
+			ActiveOverrideMaterial = UMaterialInstanceDynamic::Create(NewMaterial.Get(), this);
+			GetSculptMeshComponent()->SetOverrideRenderMaterial(ActiveOverrideMaterial);
+		}
+		else
+		{
+			GetSculptMeshComponent()->ClearOverrideRenderMaterial();
+			ActiveOverrideMaterial = nullptr;
+		}
+	}
+}
 
 
 void UMeshSculptToolBase::UpdateMaterialMode(EMeshEditingMaterialModes MaterialMode)
@@ -921,6 +967,18 @@ void UMeshSculptToolBase::UpdateMaterialMode(EMeshEditingMaterialModes MaterialM
 	else
 	{
 		if (MaterialMode == EMeshEditingMaterialModes::Custom)
+		{
+			if (ViewProperties->CustomMaterial.IsValid())
+			{
+				ActiveOverrideMaterial = UMaterialInstanceDynamic::Create(ViewProperties->CustomMaterial.Get(), this);
+			}
+			else
+			{
+				GetSculptMeshComponent()->ClearOverrideRenderMaterial();
+				ActiveOverrideMaterial = nullptr;
+			}
+		}
+		else if (MaterialMode == EMeshEditingMaterialModes::CustomImage)
 		{
 			ActiveOverrideMaterial = ToolSetupUtil::GetCustomImageBasedSculptMaterial(GetToolManager(), ViewProperties->Image);
 			if (ViewProperties->Image != nullptr)
@@ -932,6 +990,11 @@ void UMeshSculptToolBase::UpdateMaterialMode(EMeshEditingMaterialModes MaterialM
 		{
 			ActiveOverrideMaterial = ToolSetupUtil::GetVertexColorMaterial(GetToolManager());
 		}
+		else if (MaterialMode == EMeshEditingMaterialModes::Transparent)
+		{
+			ActiveOverrideMaterial = ToolSetupUtil::GetTransparentSculptMaterial(GetToolManager(), 
+				ViewProperties->TransparentMaterialColor, ViewProperties->Opacity, ViewProperties->bTwoSided);
+		}
 		else
 		{
 			UMaterialInterface* SculptMaterial = nullptr;
diff --git a/Engine/Plugins/Experimental/MeshModelingToolset/Source/MeshModelingTools/Public/DynamicMeshSculptTool.h b/Engine/Plugins/Experimental/MeshModelingToolset/Source/MeshModelingTools/Public/DynamicMeshSculptTool.h
--- a/Engine/Plugins/Experimental/MeshModelingToolset/Source/MeshModelingTools/Public/DynamicMeshSculptTool.h
+++ b/Engine/Plugins/Experimental/MeshModelingToolset/Source/MeshModelingTools/Public/DynamicMeshSculptTool.h
@@ -329,8 +329,12 @@ private:
 
 	TValueWatcher<bool> ShowWireframeWatcher;
 	TValueWatcher<EMeshEditingMaterialModes> MaterialModeWatcher;
+	TValueWatcher<TWeakObjectPtr<UMaterialInterface>> CustomMaterialWatcher;
 	TValueWatcher<bool> FlatShadingWatcher;
 	TValueWatcher<FLinearColor> ColorWatcher;
+	TValueWatcher<FLinearColor> TransparentColorWatcher;
+	TValueWatcher<double> OpacityWatcher;
+	TValueWatcher<bool> TwoSidedWatcher;
 	TValueWatcher<UTexture2D*> ImageWatcher;
 	TValueWatcher<EDynamicMeshSculptBrushType> BrushTypeWatcher;
 	TValueWatcher<FVector> GizmoPositionWatcher;
@@ -338,6 +342,9 @@ private:
 	void UpdateMaterialMode(EMeshEditingMaterialModes NewMode);
 	void UpdateFlatShadingSetting(bool bNewValue);
 	void UpdateColorSetting(FLinearColor NewColor);
+	void UpdateOpacitySetting(double Opacity);
+	void UpdateTwoSidedSetting(bool bOn);
+	void UpdateCustomMaterial(TWeakObjectPtr<UMaterialInterface> NewMaterial);
 	void UpdateImageSetting(UTexture2D* NewImage);
 	void UpdateBrushType(EDynamicMeshSculptBrushType BrushType);
 	void UpdateGizmoFromProperties();
diff --git a/Engine/Plugins/Experimental/MeshModelingToolset/Source/MeshModelingTools/Public/Properties/MeshMaterialProperties.h b/Engine/Plugins/Experimental/MeshModelingToolset/Source/MeshModelingTools/Public/Properties/MeshMaterialProperties.h
--- a/Engine/Plugins/Experimental/MeshModelingToolset/Source/MeshModelingTools/Public/Properties/MeshMaterialProperties.h
+++ b/Engine/Plugins/Experimental/MeshModelingToolset/Source/MeshModelingTools/Public/Properties/MeshMaterialProperties.h
@@ -93,13 +93,15 @@ public:
 UENUM()
 enum class EMeshEditingMaterialModes
 {
-	ExistingMaterial = 0,
-	Diffuse = 1,
-	Grey = 2,
-	Soft = 3,
-	TangentNormal = 4,
-	VertexColor = 5,
-	Custom = 6
+	ExistingMaterial,
+	Diffuse,
+	Grey,
+	Soft,
+	Transparent,
+	TangentNormal,
+	VertexColor,
+	CustomImage,
+	Custom
 };
 
 
@@ -118,7 +120,8 @@ public:
 	EMeshEditingMaterialModes MaterialMode = EMeshEditingMaterialModes::Diffuse;
 
 	/** Toggle flat shading on/off */
-	UPROPERTY(EditAnywhere, Category = Rendering, meta = (EditConditionHides, EditCondition = "MaterialMode != EMeshEditingMaterialModes::ExistingMaterial") )
+	UPROPERTY(EditAnywhere, Category = Rendering, meta = (EditConditionHides, 
+		EditCondition = "MaterialMode != EMeshEditingMaterialModes::ExistingMaterial && MaterialMode != EMeshEditingMaterialModes::Transparent && MaterialMode != EMeshEditingMaterialModes::Custom") )
 	bool bFlatShading = true;
 
 	/** Main Color of Material */
@@ -126,6 +129,20 @@ public:
 	FLinearColor Color = FLinearColor(0.4f, 0.4f, 0.4f);
 
 	/** Image used in Image-Based Material */
-	UPROPERTY(EditAnywhere, Category = Rendering, meta = (EditConditionHides, EditCondition = "MaterialMode == EMeshEditingMaterialModes::Custom", TransientToolProperty) )
+	UPROPERTY(EditAnywhere, Category = Rendering, meta = (EditConditionHides, EditCondition = "MaterialMode == EMeshEditingMaterialModes::CustomImage", TransientToolProperty) )
 	UTexture2D* Image;
+
+	//~ Could have used the same property as Color, above, but the user may want different saved values for the two
+	UPROPERTY(EditAnywhere, Category = Rendering, AdvancedDisplay, meta = (EditConditionHides, EditCondition = "MaterialMode == EMeshEditingMaterialModes::Transparent", DisplayName = "Color"))
+	FLinearColor TransparentMaterialColor = FLinearColor(0.0606, 0.309, 0.842);
+
+	UPROPERTY(EditAnywhere, Category = Rendering, meta = (EditConditionHides, EditCondition = "MaterialMode == EMeshEditingMaterialModes::Transparent", ClampMin = "0", ClampMax = "1.0"))
+	double Opacity = 0.65;
+
+	/** Although a two-sided transparent material causes rendering issues with overlapping faces, it is still frequently useful to see the shape when sculpting around other objects. */
+	UPROPERTY(EditAnywhere, Category = Rendering, meta = (EditConditionHides, EditCondition = "MaterialMode == EMeshEditingMaterialModes::Transparent"))
+	bool bTwoSided = true;
+
+	UPROPERTY(EditAnywhere, Category = Rendering, meta = (EditConditionHides, EditCondition = "MaterialMode == EMeshEditingMaterialModes::Custom"))
+	TWeakObjectPtr<UMaterialInterface> CustomMaterial;
 };
diff --git a/Engine/Plugins/Experimental/MeshModelingToolset/Source/MeshModelingTools/Public/Sculpting/MeshSculptToolBase.h b/Engine/Plugins/Experimental/MeshModelingToolset/Source/MeshModelingTools/Public/Sculpting/MeshSculptToolBase.h
--- a/Engine/Plugins/Experimental/MeshModelingToolset/Source/MeshModelingTools/Public/Sculpting/MeshSculptToolBase.h
+++ b/Engine/Plugins/Experimental/MeshModelingToolset/Source/MeshModelingTools/Public/Sculpting/MeshSculptToolBase.h
@@ -428,6 +428,9 @@ protected:
 	virtual void UpdateMaterialMode(EMeshEditingMaterialModes NewMode);
 	virtual void UpdateFlatShadingSetting(bool bNewValue);
 	virtual void UpdateColorSetting(FLinearColor NewColor);
+	virtual void UpdateOpacitySetting(double Opacity);
+	virtual void UpdateTwoSidedSetting(bool bOn);
+	virtual void UpdateCustomMaterial(TWeakObjectPtr<UMaterialInterface> NewMaterial);
 	virtual void UpdateImageSetting(UTexture2D* NewImage);
 
 
diff --git a/Engine/Plugins/Experimental/MeshModelingToolset/Source/ModelingComponents/Private/BaseDynamicMeshComponent.cpp b/Engine/Plugins/Experimental/MeshModelingToolset/Source/ModelingComponents/Private/BaseDynamicMeshComponent.cpp
--- a/Engine/Plugins/Experimental/MeshModelingToolset/Source/ModelingComponents/Private/BaseDynamicMeshComponent.cpp
+++ b/Engine/Plugins/Experimental/MeshModelingToolset/Source/ModelingComponents/Private/BaseDynamicMeshComponent.cpp
@@ -76,6 +76,20 @@ UMaterialInterface* UBaseDynamicMeshComponent::GetMaterial(int32 ElementIndex) c
 	return (ElementIndex >= 0 && ElementIndex < BaseMaterials.Num()) ? BaseMaterials[ElementIndex] : nullptr;
 }
 
+FMaterialRelevance UBaseDynamicMeshComponent::GetMaterialRelevance(ERHIFeatureLevel::Type InFeatureLevel) const
+{
+	FMaterialRelevance Result = UMeshComponent::GetMaterialRelevance(InFeatureLevel);
+	if (OverrideRenderMaterial)
+	{
+		Result |= OverrideRenderMaterial->GetRelevance_Concurrent(InFeatureLevel);
+	}
+	if (SecondaryRenderMaterial)
+	{
+		Result |= SecondaryRenderMaterial->GetRelevance_Concurrent(InFeatureLevel);
+	}
+	return Result;
+}
+
 void UBaseDynamicMeshComponent::SetMaterial(int32 ElementIndex, UMaterialInterface* Material)
 {
 	check(ElementIndex >= 0);
diff --git a/Engine/Plugins/Experimental/MeshModelingToolset/Source/ModelingComponents/Private/OctreeDynamicMeshSceneProxy.h b/Engine/Plugins/Experimental/MeshModelingToolset/Source/ModelingComponents/Private/OctreeDynamicMeshSceneProxy.h
--- a/Engine/Plugins/Experimental/MeshModelingToolset/Source/ModelingComponents/Private/OctreeDynamicMeshSceneProxy.h
+++ b/Engine/Plugins/Experimental/MeshModelingToolset/Source/ModelingComponents/Private/OctreeDynamicMeshSceneProxy.h
@@ -224,6 +224,13 @@ public:
 		return Result;
 	}
 
+	virtual void UpdatedReferencedMaterials() override
+	{
+		FBaseDynamicMeshSceneProxy::UpdatedReferencedMaterials();
+
+		// The material relevance may need updating.
+		MaterialRelevance = ParentComponent->GetMaterialRelevance(GetScene().GetFeatureLevel());
+	}
 
 	virtual void GetLightRelevance(const FLightSceneProxy* LightSceneProxy, bool& bDynamic, bool& bRelevant, bool& bLightMapped, bool& bShadowMapped) const override
 	{
diff --git a/Engine/Plugins/Experimental/MeshModelingToolset/Source/ModelingComponents/Private/SimpleDynamicMeshSceneProxy.h b/Engine/Plugins/Experimental/MeshModelingToolset/Source/ModelingComponents/Private/SimpleDynamicMeshSceneProxy.h
--- a/Engine/Plugins/Experimental/MeshModelingToolset/Source/ModelingComponents/Private/SimpleDynamicMeshSceneProxy.h
+++ b/Engine/Plugins/Experimental/MeshModelingToolset/Source/ModelingComponents/Private/SimpleDynamicMeshSceneProxy.h
@@ -610,6 +610,7 @@ public:
 			Result.bUsesLightingChannels = GetLightingChannelMask() != GetDefaultLightingChannelMask();
 			Result.bTranslucentSelfShadow = bCastVolumetricTranslucentShadow;
 			Result.bRenderCustomDepth = ShouldRenderCustomDepth();
+			// Note that this is actually a getter. One may argue that it is poorly named.
 			MaterialRelevance.SetPrimitiveViewRelevance(Result);
 			Result.bVelocityRelevance = DrawsVelocity() && Result.bOpaque && Result.bRenderInMainPass;
 		}
@@ -617,6 +618,14 @@ public:
 		return Result;
 	}
 
+	virtual void UpdatedReferencedMaterials() override
+	{
+		FBaseDynamicMeshSceneProxy::UpdatedReferencedMaterials();
+
+		// The material relevance may need updating.
+		MaterialRelevance = ParentComponent->GetMaterialRelevance(GetScene().GetFeatureLevel());
+	}
+
 
 	virtual void GetLightRelevance(const FLightSceneProxy* LightSceneProxy, bool& bDynamic, bool& bRelevant, bool& bLightMapped, bool& bShadowMapped) const override
 	{
diff --git a/Engine/Plugins/Experimental/MeshModelingToolset/Source/ModelingComponents/Private/ToolSetupUtil.cpp b/Engine/Plugins/Experimental/MeshModelingToolset/Source/ModelingComponents/Private/ToolSetupUtil.cpp
--- a/Engine/Plugins/Experimental/MeshModelingToolset/Source/ModelingComponents/Private/ToolSetupUtil.cpp
+++ b/Engine/Plugins/Experimental/MeshModelingToolset/Source/ModelingComponents/Private/ToolSetupUtil.cpp
@@ -91,6 +91,24 @@ UMaterialInterface* ToolSetupUtil::GetDefaultSculptMaterial(UInteractiveToolMana
 
 
 
+UMaterialInstanceDynamic* ToolSetupUtil::GetTransparentSculptMaterial(UInteractiveToolManager* ToolManager, const FLinearColor& Color, double Opacity, bool bTwoSided)
+{
+	// Unfortunately the two-sided flag is not something that we can give as a runtime 
+	// parameter, so we need separate versions of the material.
+	UMaterial* Material = bTwoSided ? LoadObject<UMaterial>(nullptr, TEXT("/MeshModelingToolset/Materials/SculptMaterial_TransparentTwoSided"))
+		: LoadObject<UMaterial>(nullptr, TEXT("/MeshModelingToolset/Materials/SculptMaterial_Transparent"));
+
+	if (Material != nullptr)
+	{
+		UMaterialInstanceDynamic* MatInstance = UMaterialInstanceDynamic::Create(Material, ToolManager);
+		MatInstance->SetVectorParameterValue(TEXT("Color"), Color);
+		MatInstance->SetScalarParameterValue(TEXT("Opacity"), Opacity);
+		return MatInstance;
+	}
+	return nullptr;
+}
+
+
 
 UMaterialInterface* ToolSetupUtil::GetImageBasedSculptMaterial(UInteractiveToolManager* ToolManager, ImageMaterialType Type)
 {
diff --git a/Engine/Plugins/Experimental/MeshModelingToolset/Source/ModelingComponents/Public/BaseDynamicMeshComponent.h b/Engine/Plugins/Experimental/MeshModelingToolset/Source/ModelingComponents/Public/BaseDynamicMeshComponent.h
--- a/Engine/Plugins/Experimental/MeshModelingToolset/Source/ModelingComponents/Public/BaseDynamicMeshComponent.h
+++ b/Engine/Plugins/Experimental/MeshModelingToolset/Source/ModelingComponents/Public/BaseDynamicMeshComponent.h
@@ -268,6 +268,7 @@ public:
 	//~ Begin UMeshComponent Interface.
 	virtual int32 GetNumMaterials() const override;
 	virtual UMaterialInterface* GetMaterial(int32 ElementIndex) const override;
+	virtual FMaterialRelevance GetMaterialRelevance(ERHIFeatureLevel::Type InFeatureLevel) const override;
 	virtual void SetMaterial(int32 ElementIndex, UMaterialInterface* Material) override;
 	virtual void GetUsedMaterials(TArray<UMaterialInterface*>& OutMaterials, bool bGetDebugMaterials = false) const override;
 	//~ End UMeshComponent Interface.
diff --git a/Engine/Plugins/Experimental/MeshModelingToolset/Source/ModelingComponents/Public/ToolSetupUtil.h b/Engine/Plugins/Experimental/MeshModelingToolset/Source/ModelingComponents/Public/ToolSetupUtil.h
--- a/Engine/Plugins/Experimental/MeshModelingToolset/Source/ModelingComponents/Public/ToolSetupUtil.h
+++ b/Engine/Plugins/Experimental/MeshModelingToolset/Source/ModelingComponents/Public/ToolSetupUtil.h
@@ -59,6 +59,13 @@ namespace ToolSetupUtil
 	 */
 	MODELINGCOMPONENTS_API UMaterialInterface* GetDefaultSculptMaterial(UInteractiveToolManager* ToolManager);
 
+	/**
+	 * @param bTwoSided A two sided material has some rendering artifacts in a transparent material because of indeterminate
+	 *  ordering of triangles within the mesh. Still, it is sometimes useful despite these flaws.
+	 * @return Transparent two-sided material suitable for sculpting (has some shine and a Fresnel effect)
+	 */
+	MODELINGCOMPONENTS_API UMaterialInstanceDynamic* GetTransparentSculptMaterial(UInteractiveToolManager* ToolManager, 
+		const FLinearColor& Color, double Opacity, bool bTwoSided);
 
 	/** Types of image-based material that we can create */
 	enum class ImageMaterialType
diff --git a/Engine/Source/Runtime/Engine/Classes/Components/MeshComponent.h b/Engine/Source/Runtime/Engine/Classes/Components/MeshComponent.h
--- a/Engine/Source/Runtime/Engine/Classes/Components/MeshComponent.h
+++ b/Engine/Source/Runtime/Engine/Classes/Components/MeshComponent.h
@@ -69,7 +69,7 @@ public:
 	//~ End UPrimitiveComponent Interface
 
 	/** Accesses the scene relevance information for the materials applied to the mesh. Valid from game thread only. */
-	FMaterialRelevance GetMaterialRelevance(ERHIFeatureLevel::Type InFeatureLevel) const;
+	virtual FMaterialRelevance GetMaterialRelevance(ERHIFeatureLevel::Type InFeatureLevel) const;
 
 	/**
 	 *	Tell the streaming system whether or not all mip levels of all textures used by this component should be loaded and remain loaded.
