Add a plane flip button/hotket and rotation snapping option to the plane cut tool

#rnx

diff --git a/Engine/Plugins/Experimental/MeshModelingToolset/Source/MeshModelingTools/Private/PlaneCutTool.cpp b/Engine/Plugins/Experimental/MeshModelingToolset/Source/MeshModelingTools/Private/PlaneCutTool.cpp
--- a/Engine/Plugins/Experimental/MeshModelingToolset/Source/MeshModelingTools/Private/PlaneCutTool.cpp
+++ b/Engine/Plugins/Experimental/MeshModelingToolset/Source/MeshModelingTools/Private/PlaneCutTool.cpp
@@ -213,7 +213,7 @@ void UPlaneCutTool::Setup()
 
 	SetToolDisplayName(LOCTEXT("ToolName", "Plane Cut"));
 	GetToolManager()->DisplayMessage(
-		LOCTEXT("OnStartPlaneCutTool", "Press 'A' or use the Cut button to cut the mesh without leaving the tool.  When grid snapping is enabled, you can toggle snapping with the shift key."),
+		LOCTEXT("OnStartPlaneCutTool", "Press 'A' or use the Cut button to cut the mesh without leaving the tool.  Press 'S' to flip the plane direction.  When grid snapping is enabled, you can toggle snapping with the shift key."),
 		EToolMessageLevel::UserNotification);
 }
 
@@ -228,6 +228,12 @@ void UPlaneCutTool::RegisterActions(FInteractiveToolActionSet& ActionSet)
 		LOCTEXT("DoPlaneCutTooltip", "Cut the mesh with the current cutting plane, without exiting the tool"),
 		EModifierKey::None, EKeys::A,
 		[this]() { this->Cut(); } );
+	ActionSet.RegisterAction(this, (int32)EStandardToolActions::BaseClientDefinedActionID + 102,
+		TEXT("Flip Cutting Plane"),
+		LOCTEXT("FlipCutPlane", "Flip Cutting Plane"),
+		LOCTEXT("FlipCutPlaneTooltip", "Flip the cutting plane"),
+		EModifierKey::None, EKeys::S,
+		[this]() { this->FlipPlane(); });
 }
 
 
@@ -261,6 +267,18 @@ void UPlaneCutTool::SetupPreviews()
 }
 
 
+void UPlaneCutTool::FlipPlane()
+{
+	GetToolManager()->BeginUndoTransaction(LOCTEXT("FlipPlaneTransactionName", "Flip Plane"));
+
+	FQuat Flipper(CutPlaneOrientation.GetRightVector(), FMathf::Pi);
+	CutPlaneOrientation = Flipper * CutPlaneOrientation;
+	FTransform NewTransform(CutPlaneOrientation, CutPlaneOrigin);
+	PlaneTransformGizmo->SetNewGizmoTransform(NewTransform);
+
+	GetToolManager()->EndUndoTransaction();
+}
+
 
 void UPlaneCutTool::Cut()
 {
@@ -389,6 +407,8 @@ void UPlaneCutTool::OnTick(float DeltaTime)
 	{
 		PlaneTransformGizmo->bSnapToWorldGrid =
 			BasicProperties->bSnapToWorldGrid && bIgnoreSnappingToggle == false;
+		PlaneTransformGizmo->bSnapToWorldRotGrid =
+			BasicProperties->bSnapRotationToWorldGrid && bIgnoreSnappingToggle == false;
 	}
 
 	for (UMeshOpPreviewWithBackgroundCompute* Preview : Previews)
diff --git a/Engine/Plugins/Experimental/MeshModelingToolset/Source/MeshModelingTools/Public/PlaneCutTool.h b/Engine/Plugins/Experimental/MeshModelingToolset/Source/MeshModelingTools/Public/PlaneCutTool.h
--- a/Engine/Plugins/Experimental/MeshModelingToolset/Source/MeshModelingTools/Public/PlaneCutTool.h
+++ b/Engine/Plugins/Experimental/MeshModelingToolset/Source/MeshModelingTools/Public/PlaneCutTool.h
@@ -77,6 +77,10 @@ public:
 	UPROPERTY(EditAnywhere, Category = Snapping)
 	bool bSnapToWorldGrid = false;
 
+	/** Snap rotation of the cut plane to the world grid */
+	UPROPERTY(EditAnywhere, Category = Snapping)
+	bool bSnapRotationToWorldGrid = false;
+
 	/** If true, both halves of the cut are computed */
 	UPROPERTY(EditAnywhere, Category = Options)
 	bool bKeepBothHalves;
@@ -176,6 +180,10 @@ protected:
 	UFUNCTION(CallInEditor, Category = Actions, meta = (DisplayName = "Cut"))
 	void Cut();
 
+	/** Flip the cutting plane */
+	UFUNCTION(CallInEditor, Category = Actions, meta = (DisplayName = "Flip Plane"))
+	void FlipPlane();
+
 protected:
 
 	UPROPERTY()
