ModelingTools: properly save ROI of Decompose and Merge operations in PolyEdit. Also handle case of deleting entire mesh in PolyEdit and TriSelect tools (not allowed).
#rnx

diff --git a/Engine/Plugins/Experimental/MeshModelingToolset/Source/MeshModelingTools/Private/EditMeshPolygonsTool.cpp b/Engine/Plugins/Experimental/MeshModelingToolset/Source/MeshModelingTools/Private/EditMeshPolygonsTool.cpp
--- a/Engine/Plugins/Experimental/MeshModelingToolset/Source/MeshModelingTools/Private/EditMeshPolygonsTool.cpp
+++ b/Engine/Plugins/Experimental/MeshModelingToolset/Source/MeshModelingTools/Private/EditMeshPolygonsTool.cpp
@@ -1575,7 +1575,7 @@ void UEditMeshPolygonsTool::ApplySetUVs()
 
 void UEditMeshPolygonsTool::ApplyMerge()
 {
-	if (BeginMeshFaceEditChangeWithPreview() == false)
+	if (BeginMeshFaceEditChange() == false)
 	{
 		GetToolManager()->DisplayMessage(
 			LOCTEXT("OnMergeFailedMessage", "Cannot Merge Current Selection"),
@@ -1586,7 +1586,7 @@ void UEditMeshPolygonsTool::ApplyMerge()
 	FDynamicMesh3* Mesh = DynamicMeshComponent->GetMesh();
 	FDynamicMeshChangeTracker ChangeTracker(Mesh);
 	ChangeTracker.BeginChange();
-	ChangeTracker.SaveTriangles(ActiveTriangleSelection, false);
+	ChangeTracker.SaveTriangles(ActiveTriangleSelection, true);
 	FMeshConnectedComponents Components(Mesh);
 	Components.FindConnectedTriangles(ActiveTriangleSelection);
 	FGroupTopologySelection NewSelection;
@@ -1621,6 +1621,15 @@ void UEditMeshPolygonsTool::ApplyDelete()
 
 	FDynamicMesh3* Mesh = DynamicMeshComponent->GetMesh();
 
+	// prevent deleting all triangles
+	if (ActiveTriangleSelection.Num() >= Mesh->TriangleCount())
+	{
+		GetToolManager()->DisplayMessage(
+			LOCTEXT("OnDeleteAllFailedMessage", "Cannot Delete Entire Mesh"),
+			EToolMessageLevel::UserWarning);
+		return;
+	}
+
 	FDynamicMeshChangeTracker ChangeTracker(Mesh);
 	ChangeTracker.BeginChange();
 	ChangeTracker.SaveTriangles(ActiveTriangleSelection, true);
@@ -1788,7 +1797,7 @@ void UEditMeshPolygonsTool::ApplyDecompose()
 	for (int32 GroupID : SelectionMechanic->GetActiveSelection().SelectedGroupIDs)
 	{
 		const TArray<int32>& Triangles = Topology->GetGroupTriangles(GroupID);
-		ChangeTracker.SaveTriangles(Triangles, false);
+		ChangeTracker.SaveTriangles(Triangles, true);
 		for (int32 tid : Triangles)
 		{
 			int32 NewGroupID = Mesh->AllocateTriangleGroup();
@@ -2021,7 +2030,7 @@ void UEditMeshPolygonsTool::ApplyPokeSingleFace()
 	FDynamicMesh3* Mesh = DynamicMeshComponent->GetMesh();
 	FDynamicMeshChangeTracker ChangeTracker(Mesh);
 	ChangeTracker.BeginChange();
-	ChangeTracker.SaveTriangles(ActiveTriangleSelection, false);
+	ChangeTracker.SaveTriangles(ActiveTriangleSelection, true);
 	FGroupTopologySelection NewSelection;
 	for (int32 tid : ActiveTriangleSelection)
 	{
diff --git a/Engine/Plugins/Experimental/MeshModelingToolset/Source/MeshModelingTools/Private/MeshSelectionTool.cpp b/Engine/Plugins/Experimental/MeshModelingToolset/Source/MeshModelingTools/Private/MeshSelectionTool.cpp
--- a/Engine/Plugins/Experimental/MeshModelingToolset/Source/MeshModelingTools/Private/MeshSelectionTool.cpp
+++ b/Engine/Plugins/Experimental/MeshModelingToolset/Source/MeshModelingTools/Private/MeshSelectionTool.cpp
@@ -1258,6 +1258,10 @@ void UMeshSelectionTool::DeleteSelectedTriangles()
 	{
 		return;
 	}
+	if (SelectedFaces.Num() >= PreviewMesh->GetMesh()->TriangleCount())
+	{
+		return;
+	}
 
 	TUniquePtr<FToolCommandChangeSequence> ChangeSeq = MakeUnique<FToolCommandChangeSequence>();
 
