ModelingTools: Restore constrast curve asset that was mistakenly deleted. Return a transient duplicate of contrast curve in ToolSetupUtil, so that users are not editing an Engine asset every time they modify the curve. If Property is reset to default, set back to new clone of default asset. Make the contrast curve property editable so users can assign their own curve asset, and handle register/unregister of curve listener
#rnx

diff --git a/Engine/Plugins/Experimental/MeshModelingToolset/Source/MeshModelingTools/Private/DisplaceMeshTool.cpp b/Engine/Plugins/Experimental/MeshModelingToolset/Source/MeshModelingTools/Private/DisplaceMeshTool.cpp
--- a/Engine/Plugins/Experimental/MeshModelingToolset/Source/MeshModelingTools/Private/DisplaceMeshTool.cpp
+++ b/Engine/Plugins/Experimental/MeshModelingToolset/Source/MeshModelingTools/Private/DisplaceMeshTool.cpp
@@ -763,24 +763,11 @@ void UDisplaceMeshTool::Setup()
 	{
 		TextureMapProperties->DisplacementMap = nullptr;
 	}
-	TextureMapProperties->AdjustmentCurve = ToolSetupUtil::GetContrastAdjustmentCurve(GetToolManager());
-
-	// In editor, we can respond directly to curve updates.
-#if WITH_EDITORONLY_DATA
-	if (TextureMapProperties->AdjustmentCurve)
+	if (TextureMapProperties->AdjustmentCurve == nullptr)
 	{
-		TextureMapProperties->AdjustmentCurve->OnUpdateCurve.AddWeakLambda(this,
-			[this](UCurveBase* Curve, EPropertyChangeType::Type ChangeType) {
-				if (TextureMapProperties->bApplyAdjustmentCurve)
-				{
-					FDisplaceMeshOpFactory* DisplacerDownCast = static_cast<FDisplaceMeshOpFactory*>(Displacer.Get());
-					DisplacerDownCast->SetAdjustmentCurve(TextureMapProperties->AdjustmentCurve);
-					bNeedsDisplaced = true;
-					StartComputation();
-				}
-			});
+		// if curve is null, create from default
+		TextureMapProperties->AdjustmentCurve = ToolSetupUtil::GetContrastAdjustmentCurve(GetToolManager());
 	}
-#endif
 	
 	// populate weight maps list
 	IMeshDescriptionProvider* TargetMeshProvider = Cast<IMeshDescriptionProvider>(Target);
@@ -879,12 +866,8 @@ void UDisplaceMeshTool::Setup()
 
 void UDisplaceMeshTool::Shutdown(EToolShutdownType ShutdownType)
 {
-#if WITH_EDITORONLY_DATA
-	if (TextureMapProperties->AdjustmentCurve)
-	{
-		TextureMapProperties->AdjustmentCurve->OnUpdateCurve.RemoveAll(this);
-	}
-#endif
+	// unhook any active listener for contrast curve
+	DisconnectActiveContrastCurveTarget();
 
 	CommonProperties->SaveProperties(this);
 	NoiseProperties->SaveProperties(this);
@@ -935,6 +918,7 @@ void UDisplaceMeshTool::Shutdown(EToolShutdownType ShutdownType)
 	}
 }
 
+
 void UDisplaceMeshTool::ValidateSubdivisions()
 {
 	if (CommonProperties->bDisableSizeWarning)
@@ -1092,8 +1076,66 @@ void UDisplaceMeshTool::OnPropertyModified(UObject* PropertySet, FProperty* Prop
 }
 #endif
 
+
+void UDisplaceMeshTool::UpdateActiveContrastCurveTarget()
+{
+	using namespace DisplaceMeshToolLocals;
+
+	// if user resets the AdjustmentCurve field, it will go to nullptr, in this case we will force it 
+	// back to a new default curve
+	if (TextureMapProperties->AdjustmentCurve == nullptr)
+	{
+		using namespace DisplaceMeshToolLocals;
+		TextureMapProperties->AdjustmentCurve = ToolSetupUtil::GetContrastAdjustmentCurve(GetToolManager());
+		bNeedsDisplaced = true;
+	}
+
+#if WITH_EDITORONLY_DATA
+	// if the AdjustmentCurve changes, then we need to change which one we are listening to for CurveUpdate events
+	if (TextureMapProperties->AdjustmentCurve != ActiveContrastCurveTarget)
+	{
+		DisconnectActiveContrastCurveTarget();
+
+		if (TextureMapProperties->AdjustmentCurve != nullptr)
+		{
+			ActiveContrastCurveTarget = TextureMapProperties->AdjustmentCurve;
+			ActiveContrastCurveListenerHandle = ActiveContrastCurveTarget->OnUpdateCurve.AddWeakLambda(this,
+				[this](UCurveBase* Curve, EPropertyChangeType::Type ChangeType) {
+				if (TextureMapProperties->bApplyAdjustmentCurve)
+				{
+					FDisplaceMeshOpFactory* DisplacerDownCast = static_cast<FDisplaceMeshOpFactory*>(Displacer.Get());
+					DisplacerDownCast->SetAdjustmentCurve(TextureMapProperties->AdjustmentCurve);
+					bNeedsDisplaced = true;
+					StartComputation();
+				}
+			});
+		}
+	}
+#endif
+}
+
+void UDisplaceMeshTool::DisconnectActiveContrastCurveTarget()
+{
+	using namespace DisplaceMeshToolLocals;
+
+#if WITH_EDITORONLY_DATA
+	if (ActiveContrastCurveTarget != nullptr)
+	{
+		ActiveContrastCurveTarget->OnUpdateCurve.Remove(ActiveContrastCurveListenerHandle);
+		ActiveContrastCurveListenerHandle = FDelegateHandle();
+		ActiveContrastCurveTarget = nullptr;
+
+		FDisplaceMeshOpFactory* DisplacerDownCast = static_cast<FDisplaceMeshOpFactory*>(Displacer.Get());
+		DisplacerDownCast->SetAdjustmentCurve(nullptr);
+	}
+#endif
+}
+
+
 void UDisplaceMeshTool::OnTick(float DeltaTime)
 {
+	UpdateActiveContrastCurveTarget();
+
 	AdvanceComputation();
 }
 
@@ -1122,6 +1164,8 @@ void UDisplaceMeshTool::StartComputation()
 
 void UDisplaceMeshTool::AdvanceComputation()
 {
+	using namespace DisplaceMeshToolLocals;
+
 	if (SubdivideTask && SubdivideTask->IsDone())
 	{
 		SubdividedMesh = TSharedPtr<FDynamicMesh3, ESPMode::ThreadSafe>(SubdivideTask->GetTask().ExtractOperator()->ExtractResult().Release());
@@ -1130,6 +1174,10 @@ void UDisplaceMeshTool::AdvanceComputation()
 	}
 	if (SubdividedMesh && bNeedsDisplaced)
 	{
+		// force update of contrast curve
+		FDisplaceMeshOpFactory* DisplacerDownCast = static_cast<FDisplaceMeshOpFactory*>(Displacer.Get());
+		DisplacerDownCast->SetAdjustmentCurve(TextureMapProperties->bApplyAdjustmentCurve ? TextureMapProperties->AdjustmentCurve : nullptr);
+
 		DisplaceTask = new FAsyncTaskExecuterWithAbort<TModelingOpTask<FDynamicMeshOperator>>(Displacer->MakeNewOperator());
 		DisplaceTask->StartBackgroundTask();
 		bNeedsDisplaced = false;
diff --git a/Engine/Plugins/Experimental/MeshModelingToolset/Source/MeshModelingTools/Public/DisplaceMeshTool.h b/Engine/Plugins/Experimental/MeshModelingToolset/Source/MeshModelingTools/Public/DisplaceMeshTool.h
--- a/Engine/Plugins/Experimental/MeshModelingToolset/Source/MeshModelingTools/Public/DisplaceMeshTool.h
+++ b/Engine/Plugins/Experimental/MeshModelingToolset/Source/MeshModelingTools/Public/DisplaceMeshTool.h
@@ -126,7 +126,7 @@ public:
 	bool bApplyAdjustmentCurve = false;
 
 	/** This curve is queried in the range [0,1] to adjust contrast of the displacement map. */
-	UPROPERTY(VisibleAnywhere, Category = Options, meta = (EditCondition = "bApplyAdjustmentCurve", EditConditionHides))
+	UPROPERTY(EditAnywhere, Category = Options, meta = (EditCondition = "bApplyAdjustmentCurve", EditConditionHides))
 	UCurveFloat* AdjustmentCurve = nullptr;
 
 	/** Recalculate normals from displaced mesh. Disable this if you are applying Displacements that are paired with an existing Normal Map in your Material. */
@@ -272,6 +272,10 @@ public:
 	UPROPERTY()
 	UDisplaceMeshSineWaveProperties* SineWaveProperties;
 
+	/** Contrast Curve we are actively listening to */
+	UPROPERTY()
+	UCurveFloat* ActiveContrastCurveTarget = nullptr;
+
 private:
 	void StartComputation();
 	void AdvanceComputation();
@@ -293,4 +297,9 @@ private:
 
 	TUniquePtr<UE::Geometry::IDynamicMeshOperatorFactory> Subdivider = nullptr;
 	TUniquePtr<UE::Geometry::IDynamicMeshOperatorFactory> Displacer = nullptr;
+
+	// management of TextureMapProperties->AdjustmentCurve.OnUpdateCurve listener
+	FDelegateHandle ActiveContrastCurveListenerHandle;
+	void UpdateActiveContrastCurveTarget();
+	void DisconnectActiveContrastCurveTarget();
 };
diff --git a/Engine/Plugins/Experimental/MeshModelingToolset/Source/ModelingComponents/Private/ToolSetupUtil.cpp b/Engine/Plugins/Experimental/MeshModelingToolset/Source/ModelingComponents/Private/ToolSetupUtil.cpp
--- a/Engine/Plugins/Experimental/MeshModelingToolset/Source/ModelingComponents/Private/ToolSetupUtil.cpp
+++ b/Engine/Plugins/Experimental/MeshModelingToolset/Source/ModelingComponents/Private/ToolSetupUtil.cpp
@@ -248,5 +248,10 @@ UCurveFloat* ToolSetupUtil::GetContrastAdjustmentCurve(UInteractiveToolManager*
 	// This curve would currently be shared across any tools that need such a curve. We'll probably want to revisit this
 	// once it is used in multiple tools.
 	UCurveFloat* Curve = LoadObject<UCurveFloat>(nullptr, TEXT("/MeshModelingToolset/Curves/ContrastAdjustmentCurve"));
-	return Curve;
+
+	// Create a transient duplicate of the curve, as we are going to expose this curve to the user, and we
+	// do not want them editing the default Asset.
+	UCurveFloat* CurveCopy = DuplicateObject<UCurveFloat>(Curve, GetTransientPackage());
+
+	return CurveCopy;
 }
\ No newline at end of file
