ModelingTools: minor fixes. Use world grid snapping regardless of world/local setting in AddPrimitiveTool. Disable instead of Hide unavailable settings in DrawPolygonTool
#rnx

diff --git a/Engine/Plugins/Experimental/MeshModelingToolset/Source/MeshModelingTools/Private/AddPrimitiveTool.cpp b/Engine/Plugins/Experimental/MeshModelingToolset/Source/MeshModelingTools/Private/AddPrimitiveTool.cpp
--- a/Engine/Plugins/Experimental/MeshModelingToolset/Source/MeshModelingTools/Private/AddPrimitiveTool.cpp
+++ b/Engine/Plugins/Experimental/MeshModelingToolset/Source/MeshModelingTools/Private/AddPrimitiveTool.cpp
@@ -246,8 +246,7 @@ void UAddPrimitiveTool::UpdatePreviewPosition(const FInputDeviceRay& DeviceClick
 	}
 
 	// Snap to grid if applicable
-	if (ShapeSettings->bSnapToGrid
-		&& GetToolManager()->GetContextQueriesAPI()->GetCurrentCoordinateSystem() == EToolContextCoordinateSystem::World)
+	if (ShapeSettings->bSnapToGrid)
 	{
 		FSceneSnapQueryRequest Request;
 		Request.RequestType = ESceneSnapQueryType::Position;
diff --git a/Engine/Plugins/Experimental/MeshModelingToolset/Source/MeshModelingTools/Public/DrawPolygonTool.h b/Engine/Plugins/Experimental/MeshModelingToolset/Source/MeshModelingTools/Public/DrawPolygonTool.h
--- a/Engine/Plugins/Experimental/MeshModelingToolset/Source/MeshModelingTools/Public/DrawPolygonTool.h
+++ b/Engine/Plugins/Experimental/MeshModelingToolset/Source/MeshModelingTools/Public/DrawPolygonTool.h
@@ -97,27 +97,27 @@ public:
 	UPROPERTY(EditAnywhere, NonTransactional, Category = Polygon)
 	EDrawPolygonDrawMode PolygonType = EDrawPolygonDrawMode::Freehand;
 
-	UPROPERTY(EditAnywhere, NonTransactional, Category = Polygon)
-	EDrawPolygonOutputMode OutputMode = EDrawPolygonOutputMode::ExtrudedInteractive;
-
 	/** Feature size as fraction of overall shape size, for shapes with secondary features like the rounded corners of a Rounded Rectangle */
 	UPROPERTY(EditAnywhere, NonTransactional, Category = Polygon, meta = (UIMin = "0.01", UIMax = "0.99", ClampMin = "0.01", ClampMax = "0.99",
-									EditCondition = "PolygonType == EDrawPolygonDrawMode::RoundedRectangle || PolygonType == EDrawPolygonDrawMode::HoleyCircle", EditConditionHides))
+									EditCondition = "PolygonType == EDrawPolygonDrawMode::RoundedRectangle || PolygonType == EDrawPolygonDrawMode::HoleyCircle"))
 	float FeatureSizeRatio = .25;
 
-	/** Extrusion Distance in non-interactive mode */
-	UPROPERTY(EditAnywhere, NonTransactional, Category = Polygon, meta = (UIMin = "-1000", UIMax = "1000", ClampMin = "-10000", ClampMax = "10000",
-									EditCondition = "OutputMode == EDrawPolygonOutputMode::ExtrudedConstant", EditConditionHides))
-	float ExtrudeHeight = 100.0f;
-
 	/** Number of vertices in round features */
 	UPROPERTY(EditAnywhere, NonTransactional, Category = Polygon, meta = (UIMin = "3", UIMax = "100", ClampMin = "3", ClampMax = "10000",
-				  EditCondition = "PolygonType == EDrawPolygonDrawMode::Circle || PolygonType == EDrawPolygonDrawMode::RoundedRectangle || PolygonType == EDrawPolygonDrawMode::HoleyCircle", EditConditionHides))
+				  EditCondition = "PolygonType == EDrawPolygonDrawMode::Circle || PolygonType == EDrawPolygonDrawMode::RoundedRectangle || PolygonType == EDrawPolygonDrawMode::HoleyCircle"))
 	int Steps = 16;
 
-	UPROPERTY(EditAnywhere, NonTransactional, Category = Polygon, meta = (EditCondition = "PolygonType == EDrawPolygonDrawMode::Freehand", EditConditionHides))
+	UPROPERTY(EditAnywhere, NonTransactional, Category = Polygon, meta = (EditCondition = "PolygonType == EDrawPolygonDrawMode::Freehand"))
 	bool bAllowSelfIntersections = false;
 
+	UPROPERTY(EditAnywhere, NonTransactional, Category = Polygon)
+	EDrawPolygonOutputMode OutputMode = EDrawPolygonOutputMode::ExtrudedInteractive;
+
+	/** Extrusion Distance in non-interactive mode */
+	UPROPERTY(EditAnywhere, NonTransactional, Category = Polygon, meta = (UIMin = "-1000", UIMax = "1000", ClampMin = "-10000", ClampMax = "10000",
+									EditCondition = "OutputMode == EDrawPolygonOutputMode::ExtrudedConstant"))
+	float ExtrudeHeight = 100.0f;
+
 	UPROPERTY(EditAnywhere, NonTransactional, Category = Polygon)
 	bool bShowGizmo = true;
 };
