MeshModelingTools: Various fixes to make volumes work better with our tools. Specifically:
- Changed triangulation of volume polygons to Delaunay to avoid degenerate tris when possible.
- Added ability to set degenerate tri normals to neighbor normal to avoid creating nonplanar volume polygons if degenerate tris do appear.
- Created a path to use FPolys instead of bsp nodes to convert from volumes, as nodes get broken up arbitrarily.
- Gave volume conversion the ability to preserve group boundary edges.
- Changed PlanarFlipsOptimization to correctly deal with degenerate tris, and to respect group boundaries.
- Added legacy cleanup code that Ryan ported to make volume faces convex.
- Changed our version of the convexity code to avoid arbitrary precision cutoff, and to remove one of the places we were arbitrarily throwing away colinear verts.
- Changed on-the-fly conversion options to not try to fill cracks, as this seems to worsen any problems further.

#rnx

diff --git a/Engine/Plugins/Experimental/GeometryProcessing/Source/DynamicMesh/Private/MeshNormals.cpp b/Engine/Plugins/Experimental/GeometryProcessing/Source/DynamicMesh/Private/MeshNormals.cpp
--- a/Engine/Plugins/Experimental/GeometryProcessing/Source/DynamicMesh/Private/MeshNormals.cpp
+++ b/Engine/Plugins/Experimental/GeometryProcessing/Source/DynamicMesh/Private/MeshNormals.cpp
@@ -120,6 +120,140 @@ void FMeshNormals::Compute_Triangle()
 	});
 }
 
+void FMeshNormals::SetDegenerateTriangleNormalsToNeighborNormal()
+{
+	check(Normals.Num() >= Mesh->MaxTriangleID());
+
+	// We're going to look through the triangles and set any zero normals
+	// to the normal of their neighbor, preferring to go toward the neighbor
+	// with the longer side when possible. Since we could have multiple degenerate
+	// triangles linked together, this may require a little neighborhood walk.
+
+	// Our walking function
+	auto GetNeighborThatHasNormal = [this](int32 StartTid, TSet<int32>& WalkedTidsOut, int32& NonDegenerateNeighborTidOut)
+	{
+		WalkedTidsOut.Reset();
+		NonDegenerateNeighborTidOut = FDynamicMesh3::InvalidID;
+
+		check(StartTid != FDynamicMesh3::InvalidID && Normals[StartTid] == FVector3d::Zero());
+
+		// We don't like recursion, so we use a little stack instead to help us prioritize 
+		// the longer-side neighbors in our walk.
+		TArray<int32> TidsToSearch;
+		TidsToSearch.Push(StartTid);
+
+		while (!TidsToSearch.IsEmpty())
+		{
+			int32 CurrentTid = TidsToSearch.Pop();
+
+			if (WalkedTidsOut.Contains(CurrentTid))
+			{
+				continue;
+			}
+
+			// See if we've reached a non-degenerate triangle
+			if (Normals[CurrentTid] != FVector3d::Zero())
+			{
+				NonDegenerateNeighborTidOut = CurrentTid;
+				return;
+			}
+
+			WalkedTidsOut.Add(CurrentTid);
+
+			// Sanity check so we don't go forever
+			if (WalkedTidsOut.Num() > Mesh->MaxTriangleID())
+			{
+				check(false);
+				return;
+			}
+
+			// Otherwise, get neighbors and corresponding squared edge lengths.
+			int32 NeighborTids[3];
+			double SquaredEdgeLengths[3];
+			FIndex3i TriEdges = Mesh->GetTriEdges(CurrentTid);
+			for (int i = 0; i < 3; ++i)
+			{
+				FIndex2i Tids = Mesh->GetEdgeT(TriEdges[i]);
+				int32 OtherTid = (Tids.A == CurrentTid) ? Tids.B : Tids.A;
+				NeighborTids[i] = OtherTid;
+
+				if (OtherTid == FDynamicMesh3::InvalidID)
+				{
+					SquaredEdgeLengths[i] = 0;
+				}
+				else
+				{
+					FVector3d Vert1, Vert2;
+					Mesh->GetEdgeV(TriEdges[i], Vert1, Vert2);
+					SquaredEdgeLengths[i] = FVector3d::DistSquared(Vert1, Vert2);
+				}
+			}
+
+			// Order neighbors by ascending length
+			if (SquaredEdgeLengths[0] > SquaredEdgeLengths[1])
+			{
+				Swap(NeighborTids[0], NeighborTids[1]);
+				Swap(SquaredEdgeLengths[0], SquaredEdgeLengths[1]);
+			}
+			if (SquaredEdgeLengths[1] > SquaredEdgeLengths[2])
+			{
+				Swap(NeighborTids[1], NeighborTids[2]);
+				Swap(SquaredEdgeLengths[1], SquaredEdgeLengths[2]);
+			}
+			if (SquaredEdgeLengths[0] > SquaredEdgeLengths[1])
+			{
+				Swap(NeighborTids[0], NeighborTids[1]);
+				Swap(SquaredEdgeLengths[0], SquaredEdgeLengths[1]);
+			}
+
+			// Add onto stack. Longest length neighbor is at top of stack
+			for (int32 i = 0; i < 3; ++i)
+			{
+				if (NeighborTids[i] != FDynamicMesh3::InvalidID)
+				{
+					TidsToSearch.Push(NeighborTids[i]);
+				}
+			}
+		}
+	};
+
+	// It's possible that we could have an island of degenerates, ie no normal neighbor.
+	// In that case we might as well not waste time starting the same walk from each one.
+	TSet<int32> IslandDegenerates;
+
+	TSet<int32> CurrentWalkedTids;
+
+	for (int32 Tid : Mesh->TriangleIndicesItr())
+	{
+		if (Normals[Tid] == FVector3d::Zero() && !IslandDegenerates.Contains(Tid))
+		{
+			// Find a normal to use
+			CurrentWalkedTids.Reset();
+			int32 NonDegenerateNeighborTid = FDynamicMesh3::InvalidID;
+			GetNeighborThatHasNormal(Tid, CurrentWalkedTids, NonDegenerateNeighborTid);
+
+			// Make sure there was a non-degenerate neighbor.
+			if (NonDegenerateNeighborTid == FDynamicMesh3::InvalidID)
+			{
+				ensureMsgf(false, TEXT("FMeshNormals::SetDegenerateTriangleNormalsToNeighborNormal: "
+					"Had a component entirely composed of degenerate triangle normals."));
+				IslandDegenerates.Append(CurrentWalkedTids);
+			}
+			else
+			{
+				// Apply the neighbor normal.
+				FVector3d NormalToUse = Normals[NonDegenerateNeighborTid];
+				check(NormalToUse != FVector3d::Zero());
+
+				for (int32 WalkedTid : CurrentWalkedTids)
+				{
+					Normals[WalkedTid] = NormalToUse;
+				}
+			}
+		}//end if normal is zero
+	}//end for all triangles
+}
+
 
 
 
diff --git a/Engine/Plugins/Experimental/GeometryProcessing/Source/DynamicMesh/Private/Operations/PlanarFlipsOptimization.cpp b/Engine/Plugins/Experimental/GeometryProcessing/Source/DynamicMesh/Private/Operations/PlanarFlipsOptimization.cpp
--- a/Engine/Plugins/Experimental/GeometryProcessing/Source/DynamicMesh/Private/Operations/PlanarFlipsOptimization.cpp
+++ b/Engine/Plugins/Experimental/GeometryProcessing/Source/DynamicMesh/Private/Operations/PlanarFlipsOptimization.cpp
@@ -70,9 +70,20 @@ void FPlanarFlipsOptimization::ApplySinglePass()
 			{
 				continue;		// if both are degenerate we can't fix by flipping edge between them
 			}
+			if (bRespectGroupBoundaries && Mesh->IsGroupBoundaryEdge(eid))
+			{
+				continue;
+			}
 			double MinAspect = FMathd::Min(AspectRatios[EdgeT.A], AspectRatios[EdgeT.B]);
+			
+			// The triangles need to be coplanar, but we also check for zero normals, which are a sign
+			// of a degenerate triangle. We want to try to flip those.
+			// TODO: Perhaps should also look for edges where an opposite vert is almost in the plane
+			// of the other triangle even if the normals disagree (i.e., triangle is practically degenerate).
 			double NormDot = Normals[EdgeT.A].Dot(Normals[EdgeT.B]);
-			if (NormDot > PlanarDotThresh)
+			if (NormDot > PlanarDotThresh 
+				|| Normals[EdgeT.A] == FVector3d::Zero() 
+				|| Normals[EdgeT.B] == FVector3d::Zero())
 			{
 				Flips.Add({ eid, MinAspect });
 			}
@@ -104,7 +115,9 @@ void FPlanarFlipsOptimization::ApplySinglePass()
 		FVector3d FlipNormal2 = VectorUtil::Normal(D, C, A);
 		if (FlipNormal1.Dot(Normal) < PlanarDotThresh || FlipNormal2.Dot(Normal) < PlanarDotThresh)
 		{
-			continue;		// should not happen?
+			// This should only happen if flipping would result in a degenerate tri (which
+			// will have a zero normal). We don't want the flip in that case.
+			continue;
 		}
 
 		if (FMathd::Min(FlipAspect1, FlipAspect2) > Metric)
diff --git a/Engine/Plugins/Experimental/GeometryProcessing/Source/DynamicMesh/Public/MeshNormals.h b/Engine/Plugins/Experimental/GeometryProcessing/Source/DynamicMesh/Public/MeshNormals.h
--- a/Engine/Plugins/Experimental/GeometryProcessing/Source/DynamicMesh/Public/MeshNormals.h
+++ b/Engine/Plugins/Experimental/GeometryProcessing/Source/DynamicMesh/Public/MeshNormals.h
@@ -69,6 +69,18 @@ public:
 		Compute_Triangle();
 	}
 
+	/**
+	 * Assumes that ComputeTriangleNormals() has already been run.
+	 * 
+	 * Looks through the normals for any that are zero, and tries to set them to the normal of 
+	 * some neighboring triangle, preferring the triangle on its longer side if possible. This is 
+	 * useful for finding connected components based on normals, as it prevents degenerate triangles
+	 * with 0 normals from connecting otherwise disconnected components.
+	 * This won't help if a whole mesh componenent is entirely made of such degenerate triangles
+	 * (which is usually ok).
+	 */
+	void SetDegenerateTriangleNormalsToNeighborNormal();
+
 	/**
 	 * Recompute the per-element normals of the given overlay by averaging one-ring face normals
 	 * @warning NormalOverlay must be attached to ParentMesh or an exact copy
diff --git a/Engine/Plugins/Experimental/GeometryProcessing/Source/DynamicMesh/Public/Operations/PlanarFlipsOptimization.h b/Engine/Plugins/Experimental/GeometryProcessing/Source/DynamicMesh/Public/Operations/PlanarFlipsOptimization.h
--- a/Engine/Plugins/Experimental/GeometryProcessing/Source/DynamicMesh/Public/Operations/PlanarFlipsOptimization.h
+++ b/Engine/Plugins/Experimental/GeometryProcessing/Source/DynamicMesh/Public/Operations/PlanarFlipsOptimization.h
@@ -31,6 +31,7 @@ public:
 	FDynamicMesh3* Mesh;
 	int32 NumPasses;
 	double PlanarDotThresh;
+	bool bRespectGroupBoundaries = true;
 
 	/**
 	 * Apply the operation to the mesh.
diff --git a/Engine/Plugins/Experimental/MeshModelingToolset/Source/MeshModelingTools/Private/ConversionUtils/VolumeToDynamicMesh.cpp b/Engine/Plugins/Experimental/MeshModelingToolset/Source/MeshModelingTools/Private/ConversionUtils/VolumeToDynamicMesh.cpp
--- a/Engine/Plugins/Experimental/MeshModelingToolset/Source/MeshModelingTools/Private/ConversionUtils/VolumeToDynamicMesh.cpp
+++ b/Engine/Plugins/Experimental/MeshModelingToolset/Source/MeshModelingTools/Private/ConversionUtils/VolumeToDynamicMesh.cpp
@@ -2,7 +2,8 @@
 
 #include "ConversionUtils/VolumeToDynamicMesh.h"
 
-#include "CompGeom/PolygonTriangulation.h"
+#include "ConstrainedDelaunay2.h"
+#include "Polygon2.h"
 #include "DynamicMesh3.h"
 #include "GameFramework/Volume.h"
 #include "MeshBoundaryLoops.h"
@@ -13,6 +14,10 @@
 #include "Operations/MinimalHoleFiller.h"
 #include "Operations/PlanarFlipsOptimization.h"
 
+#if WITH_EDITOR
+#include "Engine/Classes/Engine/Polys.h"
+#endif
+
 #include "ExplicitUseGeometryMathTypes.h"		// using UE::Geometry::(math types)
 using namespace UE::Geometry;
 
@@ -22,7 +27,7 @@ namespace Conversion {
 void VolumeToDynamicMesh(AVolume* Volume, FDynamicMesh3& Mesh, 
 	const FVolumeToMeshOptions& Options)
 {
-	Mesh.DiscardAttributes();
+	Mesh.Clear();
 	if (Options.bSetGroups)
 	{
 		Mesh.EnableTriangleGroups();
@@ -32,6 +37,56 @@ void VolumeToDynamicMesh(AVolume* Volume, FDynamicMesh3& Mesh,
 	UE::Geometry::FTransform3d XForm = (Options.bInWorldSpace) ? 
 		UE::Geometry::FTransform3d(Volume->GetTransform()) : UE::Geometry::FTransform3d::Identity();
 
+#if WITH_EDITOR
+	// In the editor, the preferred source of geometry for a volume is the Polys array,
+	// which the bsp nodes are generated from, because the polys may be broken up into
+	// pieces unnecessarily as bsp nodes.
+	// Polys are planar polygons.
+
+	// We do not try to merge any vertices yet.
+
+	const TArray<FPoly>& Polygons = Model->Polys->Element;
+	for (FPoly Poly : Polygons)
+	{
+		int32 NumVerts = Poly.Vertices.Num();
+		if (NumVerts < 3)
+		{
+			continue;
+		}
+
+		FVector3d Normal = (FVector3d)Poly.Normal;
+		FFrame3d Plane((FVector3d)Poly.Base, Normal);
+
+		TArray<int32> Vids;
+		FPolygon2d ToTriangulate;
+
+		Vids.SetNum(NumVerts);
+		for (int32 VertexIndex = 0; VertexIndex < NumVerts; ++VertexIndex)
+		{
+			FVector3d Point = XForm.TransformPosition((FVector3d)Poly.Vertices[VertexIndex]);
+			Vids[VertexIndex] = Mesh.AppendVertex(Point);
+			ToTriangulate.AppendVertex(Plane.ToPlaneUV(Point, 2));
+		}
+
+		// Note that this call gives triangles with the reverse orientation compared to the
+		// polygon, but the polygons we get are oriented opposite of what we want (they
+		// are clockwise if the normal is towards us), so this ends up giving us the triangle
+		// orientation that we want.
+		TArray<FIndex3i> PolyTriangles = ConstrainedDelaunayTriangulate<double>(ToTriangulate);
+
+		int32 GroupID = FDynamicMesh3::InvalidID;
+		if (Options.bSetGroups)
+		{
+			GroupID = Mesh.AllocateTriangleGroup();
+		}
+
+		for (FIndex3i Tri : PolyTriangles)
+		{
+			Mesh.AppendTriangle(Vids[Tri.A], Vids[Tri.B], Vids[Tri.C], GroupID);
+		}
+	}
+
+#else
 	// Each "BspNode" is a planar polygon, triangulate each polygon and accumulate in a mesh.
 	// Note that this does not make any attempt to weld vertices/edges
 	for (const FBspNode& Node : Model->Nodes)
@@ -42,21 +97,23 @@ void VolumeToDynamicMesh(AVolume* Volume, FDynamicMesh3& Mesh,
 		int32 NumVerts = (Node.NodeFlags & PF_TwoSided) ? Node.NumVertices / 2 : Node.NumVertices;  // ??
 		if (NumVerts > 0)
 		{
-			TArray<int32> VertIndices;
-			TArray<FVector2d> VertPositions2d;
-			VertIndices.SetNum(NumVerts);
-			VertPositions2d.SetNum(NumVerts);
+			TArray<int32> Vids;
+			FPolygon2d ToTriangulate;
+			Vids.SetNum(NumVerts);
 			for (int32 VertexIndex = 0; VertexIndex < NumVerts; ++VertexIndex)
 			{
 				const FVert& Vert = Model->Verts[Node.iVertPool + VertexIndex];
 				FVector3d Point = (FVector3d)Model->Points[Vert.pVertex];
 				Point = XForm.TransformPosition(Point);
-				VertIndices[VertexIndex] = Mesh.AppendVertex(Point);
-				VertPositions2d[VertexIndex] = Plane.ToPlaneUV(Point, 2);
+				Vids[VertexIndex] = Mesh.AppendVertex(Point);
+				ToTriangulate.AppendVertex(Plane.ToPlaneUV(Point, 2));
 			}
 
-			TArray<FIndex3i> PolyTriangles;
-			PolygonTriangulation::TriangulateSimplePolygon(VertPositions2d, PolyTriangles);
+			// Note that this call gives triangles with the reverse orientation compared to the
+			// polygon, but the polygons we get are oriented opposite of what we want (they
+			// are clockwise if the normal is towards us), so this ends up giving us the triangle
+			// orientation that we want.
+			TArray<FIndex3i> PolyTriangles = ConstrainedDelaunayTriangulate<double>(ToTriangulate);
 
 			int32 GroupID = FDynamicMesh3::InvalidID;
 			if (Options.bSetGroups)
@@ -66,11 +123,11 @@ void VolumeToDynamicMesh(AVolume* Volume, FDynamicMesh3& Mesh,
 
 			for (FIndex3i Tri : PolyTriangles)
 			{
-				// flip orientation here...
-				Mesh.AppendTriangle(VertIndices[Tri.A], VertIndices[Tri.C], VertIndices[Tri.B], GroupID);
+				Mesh.AppendTriangle(Vids[Tri.A], Vids[Tri.B], Vids[Tri.C], GroupID);
 			}
 		}
 	}
+#endif
 
 	if (Options.bMergeVertices)
 	{
diff --git a/Engine/Plugins/Experimental/MeshModelingToolset/Source/MeshModelingToolsEditorOnly/Private/ConversionUtils/DynamicMeshToVolume.cpp b/Engine/Plugins/Experimental/MeshModelingToolset/Source/MeshModelingToolsEditorOnly/Private/ConversionUtils/DynamicMeshToVolume.cpp
--- a/Engine/Plugins/Experimental/MeshModelingToolset/Source/MeshModelingToolsEditorOnly/Private/ConversionUtils/DynamicMeshToVolume.cpp
+++ b/Engine/Plugins/Experimental/MeshModelingToolset/Source/MeshModelingToolsEditorOnly/Private/ConversionUtils/DynamicMeshToVolume.cpp
@@ -17,10 +17,205 @@ using namespace UE::Geometry;
 namespace UE {
 namespace Conversion {
 
-void DynamicMeshToVolume(const FDynamicMesh3& InputMesh, AVolume* TargetVolume)
+/**
+ * A slightly modified version of FPoly::OptimizeIntoConvexPolys. Unlike the original, it won't 
+ * arbitrarily delete colinear verts or snap things to grid. Like the original, it assumes that 
+ * the passed in polys are convex and coplanar (usually just a triangulation), and it will not 
+ * merge polygons sharing a multi-edge (colinear) side.
+ *
+ * @param	InOwnerBrush	The brush that owns the polygons.
+ * @param	InPolygons		An array of FPolys that will be replaced with a new set of polygons that are merged together as much as possible.
+ */
+static void OptimizeIntoConvexPolys(ABrush* InOwnerBrush, TArray<FPoly>& InPolygons)
+{
+	bool bDidMergePolygons = true;
+
+	while (bDidMergePolygons)
+	{
+		bDidMergePolygons = false;
+
+		for (int32 PolyIndex = 0; PolyIndex < InPolygons.Num() && !bDidMergePolygons; ++PolyIndex)
+		{
+			const FPoly* PolyMain = &InPolygons[PolyIndex];
+
+			// Find all the polygons that neighbor this one (aka share an edge)
+
+			for (int32 p2 = PolyIndex + 1; p2 < InPolygons.Num() && !bDidMergePolygons; ++p2)
+			{
+				const FPoly* PolyNeighbor = &InPolygons[p2];
+
+				if (PolyMain->Normal.Equals(PolyNeighbor->Normal))
+				{
+					// See if PolyNeighbor is sharing an edge with Poly
+
+					int32 PolyMainIndex1 = INDEX_NONE, PolyMainIndex2 = INDEX_NONE;
+					int32 PolyNeighborIndex1 = INDEX_NONE, PolyNeighborIndex2 = INDEX_NONE;
+
+					int32 PolyMainNumVerts = PolyMain->Vertices.Num();
+					int32 PolyNeighborNumVerts = PolyNeighbor->Vertices.Num();
+
+					for (int32 v = 0; v < PolyMainNumVerts; ++v)
+					{
+						FVector vtx = PolyMain->Vertices[v];
+
+						int32 idx = INDEX_NONE;
+
+						for (int32 v2 = 0; v2 < PolyNeighborNumVerts; ++v2)
+						{
+							const FVector* vtx2 = &PolyNeighbor->Vertices[v2];
+
+							if (vtx.Equals(*vtx2))
+							{
+								idx = v2;
+								break;
+							}
+						}
+
+						// TODO: Could clean this up. We only need to find the first shared vert
+						// and then check the next neighbor (and maybe previous, if at start)
+						if (idx != INDEX_NONE)
+						{
+							if (PolyMainIndex1 == INDEX_NONE)
+							{
+								PolyMainIndex1 = v;
+								PolyNeighborIndex1 = idx;
+							}
+							else if (PolyMainIndex2 == INDEX_NONE)
+							{
+								PolyMainIndex2 = v;
+								PolyNeighborIndex2 = idx;
+								break;
+							}
+						}
+					}
+
+					// Check that we found two shared verts and that their indices are adjacent in PolyMain.
+					if (!(PolyMainIndex1 != INDEX_NONE && PolyMainIndex2 != INDEX_NONE
+						&& (PolyMainIndex1 + 1 == PolyMainIndex2 || (PolyMainIndex1 == 0 && PolyMainIndex2 == PolyMainNumVerts - 1))))
+					{
+						continue; // skip neighbor
+					}
+
+					// Swap indices around so that they correspond to the start and end of the edge counterclockwise
+					// in PolyMain (this will only not be the case if the edge spans end and start of vert array).
+					if (PolyMainIndex1 + 1 != PolyMainIndex2)
+					{
+						Swap(PolyMainIndex1, PolyMainIndex2);
+						Swap(PolyNeighborIndex1, PolyNeighborIndex2);
+					}
+
+					// Check that the edge goes the opposite direction in PolyNeighbor
+					if ((PolyNeighborIndex2 + 1) % PolyNeighborNumVerts != PolyNeighborIndex1)
+					{
+						continue; // skip neighbor
+					}
+
+					// Found a shared edge.  Let's see if we can merge these polygons together.
+					// The most robust thing to do is to just go ahead and try the merge and see if FPoly::IsConvex() 
+					// returns true. 
+
+					FPoly PolyMerged = *PolyMain;
+					int32 NumTotalVerts = PolyMainNumVerts + PolyNeighborNumVerts - 2;
+					PolyMerged.Vertices.SetNumUninitialized(NumTotalVerts);
+
+					// The new vertices get added after PolyMainIndex1. First shift over everything that
+					// is to the right of PolyMainIndex1
+					int32 NumVertsToShift = PolyMainNumVerts - 1 - PolyMainIndex1;
+					for (int32 i = 0; i < NumVertsToShift; ++i)
+					{
+						PolyMerged.Vertices[NumTotalVerts - 1 - i] = PolyMerged.Vertices[PolyMainNumVerts - 1 - i];
+					}
+					// Now splice in PolyNeighbor vertices.
+					for (int32 i = 1; i < PolyNeighborNumVerts - 1; ++i) // Skipping the shared verts
+					{
+						int32 NeighborIndex = (PolyNeighborIndex1 + i) % PolyNeighborNumVerts;
+						PolyMerged.Vertices[PolyMainIndex1 + i] = PolyNeighbor->Vertices[NeighborIndex];
+					}
+
+					PolyMerged.CalcNormal(true);
+
+					// Check if the result is valid
+					if (PolyMerged.IsConvex() && PolyMerged.Finalize(InOwnerBrush, 1) == 0)
+					{
+						// Remove the original polygons from the list. Uses knowledge that p2 > PolyIndex
+						InPolygons.RemoveAtSwap(p2);
+						InPolygons.RemoveAtSwap(PolyIndex);
+
+						// Add the newly merged polygon into the list
+						InPolygons.Add(PolyMerged);
+
+						// Tell the outside loop that we merged polygons and need to run through it all again
+						bDidMergePolygons = true;
+					}
+				}
+			}
+		}
+	}
+}
+
+/**
+ * Largely copied from FGeomObject::FinalizeSourceData(). Does various cleanup
+ * operations, particularly to make the volume polygons all planar and convex.
+ */
+static void ApplyLegacyVolumeCleanup(ABrush* brush)
+{
+	// Remove invalid polygons from the brush
+	for (int32 x = 0; x < brush->Brush->Polys->Element.Num(); ++x)
+	{
+		FPoly* Poly = &brush->Brush->Polys->Element[x];
+		if (Poly->Vertices.Num() < 3)
+		{
+			brush->Brush->Polys->Element.RemoveAt(x);
+			x = -1;
+		}
+	}
+	for (int32 p = 0; p < brush->Brush->Polys->Element.Num(); ++p)
+	{
+		FPoly* Poly = &(brush->Brush->Polys->Element[p]);
+		Poly->iLink = p;
+		int32 SaveNumVertices = Poly->Vertices.Num();
+		if (!Poly->IsCoplanar() || !Poly->IsConvex())
+		{
+			// If the polygon is no longer coplanar and/or convex, break it up into separate triangles.
+			FPoly WkPoly = *Poly;
+			brush->Brush->Polys->Element.RemoveAt(p);
+			TArray<FPoly> Polygons;
+			if (WkPoly.Triangulate(brush, Polygons) > 0)
+			{
+				OptimizeIntoConvexPolys(brush, Polygons);
+				for (int32 t = 0; t < Polygons.Num(); ++t)
+				{
+					brush->Brush->Polys->Element.Add(Polygons[t]);
+				}
+			}
+			p = -1;
+		}
+		else
+		{
+			int32 FixResult = Poly->Fix();
+			if (FixResult != SaveNumVertices)
+			{
+				// If the polygon collapses after running "Fix" against it, it needs to be
+				// removed from the brushes polygon list.
+				if (FixResult == 0)
+				{
+					brush->Brush->Polys->Element.RemoveAt(p);
+				}
+				p = -1;
+			}
+			else
+			{
+				// If we get here, the polygon is valid and needs to be kept.  Finalize its internals.
+				Poly->Finalize(brush, 1);
+			}
+		}
+	}
+}
+
+void DynamicMeshToVolume(const FDynamicMesh3& InputMesh, AVolume* TargetVolume, const FMeshToVolumeOptions& Options)
 {
 	TArray<FDynamicMeshFace> Faces;
-	GetPolygonFaces(InputMesh, Faces);
+	GetPolygonFaces(InputMesh, Faces, Options.bRespectGroupBoundaries);
 	DynamicMeshToVolume(InputMesh, Faces, TargetVolume);
 }
 
@@ -126,6 +321,9 @@ void DynamicMeshToVolume(const FDynamicMesh3& InputMesh, TArray<FDynamicMeshFace
 		*/
 	}
 
+	ApplyLegacyVolumeCleanup(TargetVolume);
+	FBSPOps::bspUnlinkPolys(Model);
+
 	// requires ed
 	FBSPOps::csgPrepMovingBrush(TargetVolume);
 
@@ -134,30 +332,29 @@ void DynamicMeshToVolume(const FDynamicMesh3& InputMesh, TArray<FDynamicMeshFace
 }
 
 
-void GetPolygonFaces(const FDynamicMesh3& InputMesh, TArray<FDynamicMeshFace>& Faces)
+void GetPolygonFaces(const FDynamicMesh3& InputMesh, TArray<FDynamicMeshFace>& Faces, bool bRespectGroupBoundaries)
 {
 	Faces.SetNum(0);
 
+	// We'll find faces using a connected component search based on normals. Note that we give
+	// degenerate tris the normal of their neighbor so that they don't connect non-planar
+	// components.
+
 	FMeshNormals Normals(&InputMesh);
 	Normals.ComputeTriangleNormals();
+	Normals.SetDegenerateTriangleNormalsToNeighborNormal(); // See comment above
 
-	double PlanarTolerance = FMathf::ZeroTolerance;
+	double NormalTolerance = FMathf::ZeroTolerance;
 
 	FMeshConnectedComponents Components(&InputMesh);
-	Components.FindConnectedTriangles([&](int32 Triangle0, int32 Triangle1)
-		{
-			FVector3d Origin = InputMesh.GetTriCentroid(Triangle0);
-			FVector3d Normal = Normals[Triangle0];
-
-			FVector3d A, B, C;
-			InputMesh.GetTriVertices(Triangle1, A, B, C);;
-			double DistA = (A - Origin).Dot(Normal);
-			double DistB = (B - Origin).Dot(Normal);
-			double DistC = (C - Origin).Dot(Normal);
-			double MaxDist = FMathd::Max3(FMathd::Abs(DistA), FMathd::Abs(DistB), FMathd::Abs(DistC));
-
-			return MaxDist < PlanarTolerance;
-		});
+	Components.FindConnectedTriangles([&InputMesh, &Normals, NormalTolerance, bRespectGroupBoundaries](int32 Triangle0, int32 Triangle1)
+	{
+		return (!bRespectGroupBoundaries || InputMesh.GetTriangleGroup(Triangle0) == InputMesh.GetTriangleGroup(Triangle1))
+
+			// This test is only performed if triangles share an edge, so checking the normal is 
+			// sufficient for coplanarity.
+			&& Normals[Triangle0].Dot(Normals[Triangle1]) >= 1 - NormalTolerance;
+	});
 
 	for (const FMeshConnectedComponents::FComponent& Component : Components)
 	{
diff --git a/Engine/Plugins/Experimental/MeshModelingToolset/Source/MeshModelingToolsEditorOnly/Private/ConversionUtils/VolumeComponentTarget.cpp b/Engine/Plugins/Experimental/MeshModelingToolset/Source/MeshModelingToolsEditorOnly/Private/ConversionUtils/VolumeComponentTarget.cpp
--- a/Engine/Plugins/Experimental/MeshModelingToolset/Source/MeshModelingToolsEditorOnly/Private/ConversionUtils/VolumeComponentTarget.cpp
+++ b/Engine/Plugins/Experimental/MeshModelingToolset/Source/MeshModelingToolsEditorOnly/Private/ConversionUtils/VolumeComponentTarget.cpp
@@ -23,7 +23,7 @@ FVolumeComponentTarget::FVolumeComponentTarget(UPrimitiveComponent* Component)
 	VolumeToMeshOptions.bInWorldSpace = false;
 	VolumeToMeshOptions.bSetGroups = true;
 	VolumeToMeshOptions.bMergeVertices = true;
-	VolumeToMeshOptions.bAutoRepairMesh = true;
+	VolumeToMeshOptions.bAutoRepairMesh = false;
 	VolumeToMeshOptions.bOptimizeMesh = true;
 }
 
diff --git a/Engine/Plugins/Experimental/MeshModelingToolset/Source/MeshModelingToolsEditorOnly/Private/ConversionUtils/VolumeMeshDescriptionToolTarget.cpp b/Engine/Plugins/Experimental/MeshModelingToolset/Source/MeshModelingToolsEditorOnly/Private/ConversionUtils/VolumeMeshDescriptionToolTarget.cpp
--- a/Engine/Plugins/Experimental/MeshModelingToolset/Source/MeshModelingToolsEditorOnly/Private/ConversionUtils/VolumeMeshDescriptionToolTarget.cpp
+++ b/Engine/Plugins/Experimental/MeshModelingToolset/Source/MeshModelingToolsEditorOnly/Private/ConversionUtils/VolumeMeshDescriptionToolTarget.cpp
@@ -24,7 +24,12 @@ UVolumeMeshDescriptionToolTarget::UVolumeMeshDescriptionToolTarget()
 	VolumeToMeshOptions.bInWorldSpace = false;
 	VolumeToMeshOptions.bSetGroups = true;
 	VolumeToMeshOptions.bMergeVertices = true;
-	VolumeToMeshOptions.bAutoRepairMesh = true;
+
+	// When a volume has cracks, this option seems to make the geometry
+	// worse rather than better, since the filled in triangles are sometimes
+	// degenerate, folded in on themselves, etc.
+	VolumeToMeshOptions.bAutoRepairMesh = false;
+	
 	VolumeToMeshOptions.bOptimizeMesh = true;
 }
 
diff --git a/Engine/Plugins/Experimental/MeshModelingToolset/Source/MeshModelingToolsEditorOnly/Private/MeshToVolumeTool.cpp b/Engine/Plugins/Experimental/MeshModelingToolset/Source/MeshModelingToolsEditorOnly/Private/MeshToVolumeTool.cpp
--- a/Engine/Plugins/Experimental/MeshModelingToolset/Source/MeshModelingToolsEditorOnly/Private/MeshToVolumeTool.cpp
+++ b/Engine/Plugins/Experimental/MeshModelingToolset/Source/MeshModelingToolsEditorOnly/Private/MeshToVolumeTool.cpp
@@ -212,7 +212,11 @@ void UMeshToVolumeTool::RecalculateVolume()
 {
 	if (Settings->ConversionMode == EMeshToVolumeMode::MinimalPolygons)
 	{
-		UE::Conversion::GetPolygonFaces(InputMesh, Faces);
+		// Since this tool is likely to be a sink, there isn't much reason to keep
+		// the group differentiations if they are coplanar.
+		bool bRespectGroupBoundaries = false;
+
+		UE::Conversion::GetPolygonFaces(InputMesh, Faces, bRespectGroupBoundaries);
 	}
 	else
 	{
diff --git a/Engine/Plugins/Experimental/MeshModelingToolset/Source/MeshModelingToolsEditorOnly/Public/ConversionUtils/DynamicMeshToVolume.h b/Engine/Plugins/Experimental/MeshModelingToolset/Source/MeshModelingToolsEditorOnly/Public/ConversionUtils/DynamicMeshToVolume.h
--- a/Engine/Plugins/Experimental/MeshModelingToolset/Source/MeshModelingToolsEditorOnly/Public/ConversionUtils/DynamicMeshToVolume.h
+++ b/Engine/Plugins/Experimental/MeshModelingToolset/Source/MeshModelingToolsEditorOnly/Public/ConversionUtils/DynamicMeshToVolume.h
@@ -15,9 +15,18 @@ class AVolume;
 PREDECLARE_USE_GEOMETRY_CLASS(FDynamicMesh3);
 
 namespace UE {
-namespace Conversion { 
+namespace Conversion {
 
-using namespace UE::Geometry;
+using namespace UE::Geometry; 
+
+struct FMeshToVolumeOptions
+{
+	/**
+	 * When true, coplanar components with different group id's will not be
+	 * merged into a single volume polygon.
+	 */
+	bool bRespectGroupBoundaries = true;
+};
 
 struct FDynamicMeshFace
 {
@@ -29,7 +38,8 @@ struct FDynamicMeshFace
  * Gets an array of face objects that can be used to convert a dynamic mesh to a volume. This version tries to
  * merge coplanar triangles into polygons.
  */
-void MESHMODELINGTOOLSEDITORONLY_API GetPolygonFaces(const FDynamicMesh3& InputMesh, TArray<FDynamicMeshFace>& Faces);
+void MESHMODELINGTOOLSEDITORONLY_API GetPolygonFaces(const FDynamicMesh3& InputMesh, TArray<FDynamicMeshFace>& Faces,
+	bool bRespectGroupBoundaries);
 /**
  * Gets an array of face objects that can be used to convert a dynamic mesh to a volume. This version makes
  * each triangle its own face.
@@ -39,7 +49,8 @@ void MESHMODELINGTOOLSEDITORONLY_API GetTriangleFaces(const FDynamicMesh3& Input
 /**
  * Converts a dynamic mesh to a volume.
  */
-void MESHMODELINGTOOLSEDITORONLY_API DynamicMeshToVolume(const FDynamicMesh3& InputMesh, AVolume* TargetVolume);
+void MESHMODELINGTOOLSEDITORONLY_API DynamicMeshToVolume(const FDynamicMesh3& InputMesh, AVolume* TargetVolume,
+	const FMeshToVolumeOptions& Options = FMeshToVolumeOptions());
 void MESHMODELINGTOOLSEDITORONLY_API DynamicMeshToVolume(const FDynamicMesh3& InputMesh, TArray<FDynamicMeshFace>& Faces, AVolume* TargetVolume);
 
 }}//end namespace UE::Conversion
\ No newline at end of file
