ModelingComponents: add missing thread locking in FMeshRenderDecomposition, otherwise random crashes may occur when decomposing for multiple materials
#rnx

diff --git a/Engine/Plugins/Experimental/MeshModelingToolset/Source/ModelingComponents/Private/Drawing/MeshRenderDecomposition.cpp b/Engine/Plugins/Experimental/MeshModelingToolset/Source/ModelingComponents/Private/Drawing/MeshRenderDecomposition.cpp
--- a/Engine/Plugins/Experimental/MeshModelingToolset/Source/ModelingComponents/Private/Drawing/MeshRenderDecomposition.cpp
+++ b/Engine/Plugins/Experimental/MeshModelingToolset/Source/ModelingComponents/Private/Drawing/MeshRenderDecomposition.cpp
@@ -71,8 +71,10 @@ static void CollectSubDecomposition(
 	int32 MaxTrisPerGroup = MaxChunkSize;
 	if (Triangles.Num() < MaxTrisPerGroup)
 	{
+		DecompLock.Lock();
 		int32 i = Decomp.AppendGroup();
 		FMeshRenderDecomposition::FGroup& Group = Decomp.GetGroup(i);
+		DecompLock.Unlock();
 		Group.Triangles = Triangles;
 		Group.Material = Material;
 		return;
