ModelingTools: Add background compute for BakeMeshAttributeMapsTool.

#rnx

diff --git a/Engine/Plugins/Experimental/GeometryProcessing/Source/DynamicMesh/Private/Sampling/MeshImageBakingCache.cpp b/Engine/Plugins/Experimental/GeometryProcessing/Source/DynamicMesh/Private/Sampling/MeshImageBakingCache.cpp
--- a/Engine/Plugins/Experimental/GeometryProcessing/Source/DynamicMesh/Private/Sampling/MeshImageBakingCache.cpp
+++ b/Engine/Plugins/Experimental/GeometryProcessing/Source/DynamicMesh/Private/Sampling/MeshImageBakingCache.cpp
@@ -257,6 +257,11 @@ bool FMeshImageBakingCache::ValidateCache()
 		// calculate interior texels
 		ParallelFor(Dimensions.Num(), [&](int64 LinearIdx)
 		{
+			if (CancelF())
+			{
+				return;
+			}
+
 			if (OccupancyMap->IsInterior(LinearIdx) == false)
 			{
 				return;
@@ -287,6 +292,11 @@ void FMeshImageBakingCache::EvaluateSamples(
 
 	ParallelFor(Dimensions.Num(), [&](int64 LinearIdx)
 	{
+		if (CancelF())
+		{
+			return;
+		}
+
 		if (OccupancyMap->IsInterior(LinearIdx) == false)
 		{
 			return;
diff --git a/Engine/Plugins/Experimental/GeometryProcessing/Source/DynamicMesh/Public/Sampling/MeshImageBaker.h b/Engine/Plugins/Experimental/GeometryProcessing/Source/DynamicMesh/Public/Sampling/MeshImageBaker.h
--- a/Engine/Plugins/Experimental/GeometryProcessing/Source/DynamicMesh/Public/Sampling/MeshImageBaker.h
+++ b/Engine/Plugins/Experimental/GeometryProcessing/Source/DynamicMesh/Public/Sampling/MeshImageBaker.h
@@ -24,8 +24,10 @@ public:
 		return ImageBakeCache;
 	}
 
+	virtual void Bake()
+	{
 
-	virtual void Bake() = 0;
+	}
 
 
 protected:
diff --git a/Engine/Plugins/Experimental/GeometryProcessing/Source/DynamicMesh/Public/Sampling/MeshImageBakingCache.h b/Engine/Plugins/Experimental/GeometryProcessing/Source/DynamicMesh/Public/Sampling/MeshImageBakingCache.h
--- a/Engine/Plugins/Experimental/GeometryProcessing/Source/DynamicMesh/Public/Sampling/MeshImageBakingCache.h
+++ b/Engine/Plugins/Experimental/GeometryProcessing/Source/DynamicMesh/Public/Sampling/MeshImageBakingCache.h
@@ -59,6 +59,8 @@ public:
 
 	const FImageOccupancyMap* GetOccupancyMap() const;
 
+	/** if this function returns true, we should abort calculation */
+	TFunction<bool(void)> CancelF = []() { return false; };
 
 protected:
 	const FDynamicMesh3* DetailMesh = nullptr;
@@ -67,7 +69,7 @@ protected:
 
 	FImageDimensions Dimensions;
 	int32 UVLayer;
-	double Thickness;
+	double Thickness = 3.0;
 
 	TDenseGrid2<FCorrespondenceSample> SampleMap;
 	bool bSamplesValid = false;
diff --git a/Engine/Plugins/Experimental/MeshModelingToolset/Source/MeshModelingTools/Private/BakeMeshAttributeMapsTool.cpp b/Engine/Plugins/Experimental/MeshModelingToolset/Source/MeshModelingTools/Private/BakeMeshAttributeMapsTool.cpp
--- a/Engine/Plugins/Experimental/MeshModelingToolset/Source/MeshModelingTools/Private/BakeMeshAttributeMapsTool.cpp
+++ b/Engine/Plugins/Experimental/MeshModelingToolset/Source/MeshModelingTools/Private/BakeMeshAttributeMapsTool.cpp
@@ -82,7 +82,230 @@ TArray<FString> UBakeMeshAttributeMapsToolProperties::GetUVLayerNamesFunc()
 }
 
 
+/*
+ * Operators
+ */
+
+class FBakeMapBaseOp : public TGenericDataOperator<FMeshImageBaker>
+{
+public:
+	TSharedPtr<UE::Geometry::FDynamicMesh3, ESPMode::ThreadSafe> DetailMesh;
+	TSharedPtr<UE::Geometry::FDynamicMeshAABBTree3, ESPMode::ThreadSafe> DetailSpatial;
+	UE::Geometry::FDynamicMesh3* BaseMesh;
+	TPimplPtr<UE::Geometry::FMeshImageBakingCache> BakeCache;
+	UBakeMeshAttributeMapsTool::FBakeCacheSettings BakeCacheSettings;
+
+public:
+	virtual ~FBakeMapBaseOp() {}
+
+	//
+	// TGenericDataOperator implementation
+	// 
+	virtual void CalculateResult(FProgressCancel* Progress) override
+	{
+		BakeCache = MakePimpl<FMeshImageBakingCache>();
+		BakeCache->CancelF = [Progress]() {
+			return Progress && Progress->Cancelled();
+		};
+		BakeCache->SetDetailMesh(DetailMesh.Get(), DetailSpatial.Get());
+		BakeCache->SetBakeTargetMesh(BaseMesh);
+		BakeCache->SetDimensions(BakeCacheSettings.Dimensions);
+		BakeCache->SetUVLayer(BakeCacheSettings.UVLayer);
+		BakeCache->SetThickness(BakeCacheSettings.Thickness);
+		BakeCache->ValidateCache();
+	}
+};
+
+class FBakeNormalMapOp : public FBakeMapBaseOp
+{
+public:
+	typedef FBakeMapBaseOp Super;
+
+	// inputs
+	const TMeshTangents<double>* BaseMeshTangents = nullptr;
+	UBakeMeshAttributeMapsTool::FNormalMapSettings Settings;
+
+public:
+	virtual ~FBakeNormalMapOp() {}
+
+	//
+	// TGenericDataOperator implementation
+	// 
+	virtual void CalculateResult(FProgressCancel* Progress) override
+	{
+		Super::CalculateResult(Progress);
+
+		if (Progress && Progress->Cancelled())
+		{
+			return;
+		}
+
+		TUniquePtr<FMeshNormalMapBaker> Baker = MakeUnique<FMeshNormalMapBaker>();
+		Baker->SetCache(BakeCache.Get());
+		Baker->BaseMeshTangents = BaseMeshTangents;
+		Baker->Bake();
+		SetResult(MoveTemp(Baker));
+	}
+};
+
+class FBakeOcclusionMapOp : public FBakeMapBaseOp
+{
+public:
+	typedef FBakeMapBaseOp Super;
+
+	// inputs
+	const TMeshTangents<double>* BaseMeshTangents = nullptr;
+	UBakeMeshAttributeMapsTool::FOcclusionMapSettings Settings;
+
+public:
+	virtual ~FBakeOcclusionMapOp() {}
+
+	//
+	// TGenericDataOperator implementation
+	// 
+	virtual void CalculateResult(FProgressCancel* Progress) override
+	{
+		Super::CalculateResult(Progress);
+		 
+		if (Progress && Progress->Cancelled())
+		{
+			return;
+		}
+
+		TUniquePtr<FMeshOcclusionMapBaker> Baker = MakeUnique<FMeshOcclusionMapBaker>();
+		Baker->SetCache(BakeCache.Get());
+        Baker->BaseMeshTangents = BaseMeshTangents;
+		Baker->NumOcclusionRays = Settings.OcclusionRays;
+        Baker->MaxDistance = Settings.MaxDistance;
+        Baker->SpreadAngle = Settings.SpreadAngle;
+		Baker->BlurRadius = Settings.BlurRadius;
+		Baker->BiasAngleDeg = Settings.BiasAngle;
+        switch (Settings.Distribution)
+		{
+		case EOcclusionMapDistribution::Cosine:
+			Baker->Distribution = FMeshOcclusionMapBaker::EDistribution::Cosine;
+			break;
+		case EOcclusionMapDistribution::Uniform:
+			Baker->Distribution = FMeshOcclusionMapBaker::EDistribution::Uniform;
+			break;
+		}
+        switch (Settings.NormalSpace)
+		{
+		case ENormalMapSpace::Tangent:
+			Baker->NormalSpace = FMeshOcclusionMapBaker::ESpace::Tangent;
+			break;
+		case ENormalMapSpace::Object:
+			Baker->NormalSpace = FMeshOcclusionMapBaker::ESpace::Object;
+			break;
+		}
+		Baker->Bake();
+		SetResult(MoveTemp(Baker));
+	}
+};
+
+class FBakeCurvatureMapOp : public FBakeMapBaseOp
+{
+public:
+	typedef FBakeMapBaseOp Super;
+
+	// inputs
+	UBakeMeshAttributeMapsTool::FCurvatureMapSettings Settings;
 
+public:
+	virtual ~FBakeCurvatureMapOp() {}
+
+	//
+	// TGenericDataOperator implementation
+	// 
+	virtual void CalculateResult(FProgressCancel* Progress) override
+	{
+		Super::CalculateResult(Progress);
+
+		if (Progress && Progress->Cancelled())
+		{
+			return;
+		}
+
+		TUniquePtr<FMeshCurvatureMapBaker> Baker = MakeUnique<FMeshCurvatureMapBaker>();
+		Baker->SetCache(BakeCache.Get());
+		Baker->RangeScale = FMathd::Clamp(Settings.RangeMultiplier, 0.0001, 1000.0);
+		Baker->MinRangeScale = FMathd::Clamp(Settings.MinRangeMultiplier, 0.0, 1.0);
+		Baker->UseCurvatureType = (FMeshCurvatureMapBaker::ECurvatureType)Settings.CurvatureType;
+		Baker->UseColorMode = (FMeshCurvatureMapBaker::EColorMode)Settings.ColorMode;
+		Baker->UseClampMode = (FMeshCurvatureMapBaker::EClampMode)Settings.ClampMode;
+		Baker->BlurRadius = Settings.BlurRadius;
+		Baker->Bake();
+		SetResult(MoveTemp(Baker));
+	}
+};
+
+class FBakeMeshPropertyMapOp : public FBakeMapBaseOp
+{
+public:
+	typedef FBakeMapBaseOp Super;
+
+	// inputs
+	UBakeMeshAttributeMapsTool::FMeshPropertyMapSettings Settings;
+
+public:
+	virtual ~FBakeMeshPropertyMapOp() {}
+
+	//
+	// TGenericDataOperator implementation
+	// 
+	virtual void CalculateResult(FProgressCancel* Progress) override
+	{
+		Super::CalculateResult(Progress);
+
+		if (Progress && Progress->Cancelled())
+		{
+			return;
+		}
+
+		TUniquePtr<FMeshPropertyMapBaker> Baker = MakeUnique<FMeshPropertyMapBaker>();
+		Baker->SetCache(BakeCache.Get());
+		Baker->Property = (EMeshPropertyBakeType)Settings.PropertyTypeIndex;
+		//Baker->BlurRadius = Settings.BlurRadius;
+		Baker->Bake();
+		SetResult(MoveTemp(Baker));
+	}
+};
+
+class FBakeTexture2DImageMapOp : public FBakeMapBaseOp
+{
+public:
+	typedef FBakeMapBaseOp Super;
+
+	// inputs
+	const FDynamicMeshUVOverlay* UVOverlay = nullptr;
+	TSharedPtr<UE::Geometry::TImageBuilder<FVector4f>, ESPMode::ThreadSafe> TextureImage;
+	UBakeMeshAttributeMapsTool::FTexture2DImageSettings Settings;
+
+public:
+	virtual ~FBakeTexture2DImageMapOp() {}
+
+	//
+	// TGenericDataOperator implementation
+	// 
+	virtual void CalculateResult(FProgressCancel* Progress) override
+	{
+		Super::CalculateResult(Progress);
+
+		if (Progress && Progress->Cancelled())
+		{
+			return;
+		}
+
+		TUniquePtr<FMeshResampleImageBaker> Baker = MakeUnique<FMeshResampleImageBaker>();
+		Baker->SetCache(BakeCache.Get());
+		Baker->DetailUVOverlay = UVOverlay;
+		Baker->SampleFunction = [this](FVector2d UVCoord) {
+			return TextureImage->BilinearSampleUV<float>(UVCoord, FVector4f(0, 0, 0, 1));
+		};
+		Baker->Bake();
+		SetResult(MoveTemp(Baker));
+	}
+};
 
 /*
  * Tool
@@ -141,10 +364,16 @@ void UBakeMeshAttributeMapsTool::Setup()
 	}
 	UMaterial* BentNormalMaterial = LoadObject<UMaterial>(nullptr, TEXT("/MeshModelingToolset/Materials/BakeBentNormalPreviewMaterial"));
 	check(BentNormalMaterial);
-	if (Material != nullptr)
+	if (BentNormalMaterial != nullptr)
 	{
 		BentNormalPreviewMaterial = UMaterialInstanceDynamic::Create(BentNormalMaterial, GetToolManager());
 	}
+	UMaterial* WorkingMaterial = LoadObject<UMaterial>(nullptr, TEXT("/MeshModelingToolset/Materials/InProgressMaterial"));
+	check(WorkingMaterial);
+	if (WorkingMaterial != nullptr)
+	{
+		WorkingPreviewMaterial = UMaterialInstanceDynamic::Create(WorkingMaterial, GetToolManager());
+	}
 
 	bIsBakeToSelf = (Targets.Num() == 1);
 
@@ -170,11 +399,11 @@ void UBakeMeshAttributeMapsTool::Setup()
 	}
 	AddToolPropertySource(Settings);
 
-	Settings->WatchProperty(Settings->MapType, [this](EBakeMapType) { bResultValid = false; UpdateOnModeChange(); });
-	Settings->WatchProperty(Settings->Resolution, [this](EBakeTextureResolution) { bResultValid = false; });
-	Settings->WatchProperty(Settings->UVLayer, [this](FString) { bResultValid = false; });
-	Settings->WatchProperty(Settings->bUseWorldSpace, [this](bool) { bDetailMeshValid = false; bResultValid = false; });
-	Settings->WatchProperty(Settings->Thickness, [this](float) { bResultValid = false; });
+	Settings->WatchProperty(Settings->MapType, [this](EBakeMapType) { bInputsDirty = true; UpdateOnModeChange(); });
+	Settings->WatchProperty(Settings->Resolution, [this](EBakeTextureResolution) { bInputsDirty = true; });
+	Settings->WatchProperty(Settings->UVLayer, [this](FString) { bInputsDirty = true; });
+	Settings->WatchProperty(Settings->bUseWorldSpace, [this](bool) { bDetailMeshValid = false; bInputsDirty = true; });
+	Settings->WatchProperty(Settings->Thickness, [this](float) { bInputsDirty = true; });
 
 	NormalMapProps = NewObject<UBakedNormalMapToolProperties>(this);
 	NormalMapProps->RestoreProperties(this);
@@ -186,36 +415,36 @@ void UBakeMeshAttributeMapsTool::Setup()
 	OcclusionMapProps->RestoreProperties(this);
 	AddToolPropertySource(OcclusionMapProps);
 	SetToolPropertySourceEnabled(OcclusionMapProps, false);
-	OcclusionMapProps->WatchProperty(OcclusionMapProps->OcclusionRays, [this](int32) { bResultValid = false; });
-	OcclusionMapProps->WatchProperty(OcclusionMapProps->MaxDistance, [this](float) { bResultValid = false; });
-	OcclusionMapProps->WatchProperty(OcclusionMapProps->SpreadAngle, [this](float) { bResultValid = false; });
-	OcclusionMapProps->WatchProperty(OcclusionMapProps->Distribution, [this](EOcclusionMapDistribution) { bResultValid = false; });
+	OcclusionMapProps->WatchProperty(OcclusionMapProps->OcclusionRays, [this](int32) { bInputsDirty = true; });
+	OcclusionMapProps->WatchProperty(OcclusionMapProps->MaxDistance, [this](float) { bInputsDirty = true; });
+	OcclusionMapProps->WatchProperty(OcclusionMapProps->SpreadAngle, [this](float) { bInputsDirty = true; });
+	OcclusionMapProps->WatchProperty(OcclusionMapProps->Distribution, [this](EOcclusionMapDistribution) { bInputsDirty = true; });
 	OcclusionMapProps->WatchProperty(OcclusionMapProps->Preview, [this](EOcclusionMapPreview) { UpdateVisualization(); GetToolManager()->PostInvalidation(); });
-	OcclusionMapProps->WatchProperty(OcclusionMapProps->BlurRadius, [this](float) { bResultValid = false; });
-	OcclusionMapProps->WatchProperty(OcclusionMapProps->bGaussianBlur, [this](float) { bResultValid = false; });
-	OcclusionMapProps->WatchProperty(OcclusionMapProps->BiasAngle, [this](float) { bResultValid = false; });
-	OcclusionMapProps->WatchProperty(OcclusionMapProps->NormalSpace, [this](ENormalMapSpace) { bResultValid = false; });
+	OcclusionMapProps->WatchProperty(OcclusionMapProps->BlurRadius, [this](float) { bInputsDirty = true; });
+	OcclusionMapProps->WatchProperty(OcclusionMapProps->bGaussianBlur, [this](float) { bInputsDirty = true; });
+	OcclusionMapProps->WatchProperty(OcclusionMapProps->BiasAngle, [this](float) { bInputsDirty = true; });
+	OcclusionMapProps->WatchProperty(OcclusionMapProps->NormalSpace, [this](ENormalMapSpace) { bInputsDirty = true; });
 
 
 	CurvatureMapProps = NewObject<UBakedCurvatureMapToolProperties>(this);
 	CurvatureMapProps->RestoreProperties(this);
 	AddToolPropertySource(CurvatureMapProps);
 	SetToolPropertySourceEnabled(CurvatureMapProps, false);
-	CurvatureMapProps->WatchProperty(CurvatureMapProps->RangeMultiplier, [this](float) { bResultValid = false; });
-	CurvatureMapProps->WatchProperty(CurvatureMapProps->MinRangeMultiplier, [this](float) { bResultValid = false; });
-	CurvatureMapProps->WatchProperty(CurvatureMapProps->CurvatureType, [this](EBakedCurvatureTypeMode) { bResultValid = false; });
-	CurvatureMapProps->WatchProperty(CurvatureMapProps->ColorMode, [this](EBakedCurvatureColorMode) { bResultValid = false; });
-	CurvatureMapProps->WatchProperty(CurvatureMapProps->Clamping, [this](EBakedCurvatureClampMode) { bResultValid = false; });
-	CurvatureMapProps->WatchProperty(CurvatureMapProps->BlurRadius, [this](float) { bResultValid = false; });
-	CurvatureMapProps->WatchProperty(CurvatureMapProps->bGaussianBlur, [this](float) { bResultValid = false; });
+	CurvatureMapProps->WatchProperty(CurvatureMapProps->RangeMultiplier, [this](float) { bInputsDirty = true; });
+	CurvatureMapProps->WatchProperty(CurvatureMapProps->MinRangeMultiplier, [this](float) { bInputsDirty = true; });
+	CurvatureMapProps->WatchProperty(CurvatureMapProps->CurvatureType, [this](EBakedCurvatureTypeMode) { bInputsDirty = true; });
+	CurvatureMapProps->WatchProperty(CurvatureMapProps->ColorMode, [this](EBakedCurvatureColorMode) { bInputsDirty = true; });
+	CurvatureMapProps->WatchProperty(CurvatureMapProps->Clamping, [this](EBakedCurvatureClampMode) { bInputsDirty = true; });
+	CurvatureMapProps->WatchProperty(CurvatureMapProps->BlurRadius, [this](float) { bInputsDirty = true; });
+	CurvatureMapProps->WatchProperty(CurvatureMapProps->bGaussianBlur, [this](float) { bInputsDirty = true; });
 
 
 	Texture2DProps = NewObject<UBakedTexture2DImageProperties>(this);
 	Texture2DProps->RestoreProperties(this);
 	AddToolPropertySource(Texture2DProps);
 	SetToolPropertySourceEnabled(Texture2DProps, false);
-	Texture2DProps->WatchProperty(Texture2DProps->UVLayer, [this](float) { bResultValid = false; });
-	Texture2DProps->WatchProperty(Texture2DProps->SourceTexture, [this](UTexture2D*) { bResultValid = false; });
+	Texture2DProps->WatchProperty(Texture2DProps->UVLayer, [this](float) { bInputsDirty = true; });
+	Texture2DProps->WatchProperty(Texture2DProps->SourceTexture, [this](UTexture2D*) { bInputsDirty = true; });
 
 
 	VisualizationProps = NewObject<UBakedOcclusionMapVisualizationProperties>(this);
@@ -224,7 +453,7 @@ void UBakeMeshAttributeMapsTool::Setup()
 
 	InitializeEmptyMaps();
 
-	bResultValid = false;
+	bInputsDirty = true;
 	bDetailMeshValid = false;
 
 	SetToolDisplayName(LOCTEXT("ToolName", "Bake Textures"));
@@ -238,15 +467,85 @@ void UBakeMeshAttributeMapsTool::Setup()
 
 bool UBakeMeshAttributeMapsTool::CanAccept() const
 {
-	bool bCanAccept = Settings->Result.Num() > 0;
-	for (UTexture2D* Result : Settings->Result)
+	bool bCanAccept = Compute ? Compute->HaveValidResult() : false;
+	if (bCanAccept)
 	{
-		bCanAccept = bCanAccept && Result;
+		bCanAccept = Settings->Result.Num() > 0;
+		for (UTexture2D* Result : Settings->Result)
+		{
+			bCanAccept = bCanAccept && Result;
+		}
 	}
 	return bCanAccept;
 }
 
 
+TUniquePtr<UE::Geometry::TGenericDataOperator<FMeshImageBaker>> UBakeMeshAttributeMapsTool::MakeNewOperator()
+{
+	switch (Settings->MapType)
+	{
+	case EBakeMapType::TangentSpaceNormalMap:
+	{
+		TUniquePtr<FBakeNormalMapOp> Op = MakeUnique<FBakeNormalMapOp>();
+		Op->DetailMesh = DetailMesh;
+		Op->DetailSpatial = DetailSpatial;
+		Op->BaseMesh = &BaseMesh;
+		Op->BakeCacheSettings = CachedBakeCacheSettings;
+		Op->BaseMeshTangents = BaseMeshTangents.Get();
+		Op->Settings = CachedNormalMapSettings;
+		return Op;
+	}
+	case EBakeMapType::Occlusion:
+	{
+		TUniquePtr<FBakeOcclusionMapOp> Op = MakeUnique<FBakeOcclusionMapOp>();
+		Op->DetailMesh = DetailMesh;
+		Op->DetailSpatial = DetailSpatial;
+		Op->BaseMesh = &BaseMesh;
+		Op->BakeCacheSettings = CachedBakeCacheSettings;
+		Op->BaseMeshTangents = BaseMeshTangents.Get();
+		Op->Settings = CachedOcclusionMapSettings;
+		return Op;
+	}
+	case EBakeMapType::Curvature:
+	{
+		TUniquePtr<FBakeCurvatureMapOp> Op = MakeUnique<FBakeCurvatureMapOp>();
+		Op->DetailMesh = DetailMesh;
+		Op->DetailSpatial = DetailSpatial;
+		Op->BaseMesh = &BaseMesh;
+		Op->BakeCacheSettings = CachedBakeCacheSettings;
+		Op->Settings = CachedCurvatureMapSettings;
+		return Op;
+	}
+	case EBakeMapType::PositionImage:
+	case EBakeMapType::NormalImage:
+	case EBakeMapType::FaceNormalImage:
+	{
+		TUniquePtr<FBakeMeshPropertyMapOp> Op = MakeUnique<FBakeMeshPropertyMapOp>();
+		Op->DetailMesh = DetailMesh;
+		Op->DetailSpatial = DetailSpatial;
+		Op->BaseMesh = &BaseMesh;
+		Op->BakeCacheSettings = CachedBakeCacheSettings;
+		Op->Settings = CachedMeshPropertyMapSettings;
+		return Op;
+	}
+	case EBakeMapType::Texture2DImage:
+	{
+		TUniquePtr<FBakeTexture2DImageMapOp> Op = MakeUnique<FBakeTexture2DImageMapOp>();
+		Op->DetailMesh = DetailMesh;
+		Op->DetailSpatial = DetailSpatial;
+		Op->BaseMesh = &BaseMesh;
+		Op->BakeCacheSettings = CachedBakeCacheSettings;
+		Op->UVOverlay = DetailMesh->Attributes()->GetUVLayer(CachedTexture2DImageSettings.UVLayer);
+		Op->TextureImage = CachedTextureImage;
+		Op->Settings = CachedTexture2DImageSettings;
+		return Op;
+	}
+	}
+	check(false);
+	return nullptr;
+}
+
+
 void UBakeMeshAttributeMapsTool::Shutdown(EToolShutdownType ShutdownType)
 {
 	IPrimitiveComponentBackedTarget* TargetComponent = TargetComponentInterface(0);
@@ -258,6 +557,10 @@ void UBakeMeshAttributeMapsTool::Shutdown(EToolShutdownType ShutdownType)
 	Texture2DProps->SaveProperties(this);
 	VisualizationProps->SaveProperties(this);
 
+	if (Compute)
+	{
+		Compute->Shutdown();
+	}
 	if (DynamicMeshComponent != nullptr)
 	{
 		TargetComponent->SetOwnerVisibility(true);
@@ -335,6 +638,21 @@ void UBakeMeshAttributeMapsTool::Shutdown(EToolShutdownType ShutdownType)
 }
 
 
+void UBakeMeshAttributeMapsTool::OnTick(float DeltaTime)
+{
+	if (Compute)
+	{
+		Compute->Tick(DeltaTime);
+
+		float ElapsedComputeTime = Compute->GetElapsedComputeTime();
+		if (!CanAccept() && ElapsedComputeTime > SecondsBeforeWorkingMaterial)
+		{
+			DynamicMeshComponent->SetOverrideRenderMaterial(WorkingPreviewMaterial);
+		}
+	}
+}
+
+
 void UBakeMeshAttributeMapsTool::Render(IToolsContextRenderAPI* RenderAPI)
 {
 	UpdateResult();
@@ -373,7 +691,7 @@ void UBakeMeshAttributeMapsTool::UpdateDetailMesh()
 	DetailSpatial = MakeShared<FDynamicMeshAABBTree3, ESPMode::ThreadSafe>();
 	DetailSpatial->SetMesh(DetailMesh.Get(), true);
 
-	bResultValid = false;
+	bInputsDirty = true;
 	DetailMeshTimestamp++;
 }
 
@@ -396,7 +714,9 @@ void UBakeMeshAttributeMapsTool::UpdateResult()
 		CachedBakeCacheSettings = FBakeCacheSettings();
 	}
 
-	if (bResultValid)
+	// bInputsDirty ensures that we only validate parameters once per param
+	// change. Parameter validation can be expensive (ex. UpdateResult_Texture2DImage).
+	if (!bInputsDirty)
 	{
 		return;
 	}
@@ -413,16 +733,9 @@ void UBakeMeshAttributeMapsTool::UpdateResult()
 	BakeCacheSettings.DetailTimestamp = this->DetailMeshTimestamp;
 	BakeCacheSettings.Thickness = Settings->Thickness;
 
-	// rebuild bake cache if we need to
+	// update bake cache settings
 	if (!(CachedBakeCacheSettings == BakeCacheSettings))
 	{
-		BakeCache = MakePimpl<FMeshImageBakingCache>();
-		BakeCache->SetDetailMesh(DetailMesh.Get(), DetailSpatial.Get());
-		BakeCache->SetBakeTargetMesh(&BaseMesh);
-		BakeCache->SetDimensions(Dimensions);
-		BakeCache->SetUVLayer(BakeCacheSettings.UVLayer);
-		BakeCache->SetThickness(BakeCacheSettings.Thickness);
-		BakeCache->ValidateCache();
 		CachedBakeCacheSettings = BakeCacheSettings;
 
 		CachedNormalMapSettings = FNormalMapSettings();
@@ -432,43 +745,59 @@ void UBakeMeshAttributeMapsTool::UpdateResult()
 		CachedTexture2DImageSettings = FTexture2DImageSettings();
 	}
 
-	// rebuild select map type
+	// update map type settings
 	switch (Settings->MapType)
 	{
 		default:
 		case EBakeMapType::TangentSpaceNormalMap:
-			UpdateResult_Normal();
+			OpState = UpdateResult_Normal();
 			break;
 		case EBakeMapType::Occlusion:
-			UpdateResult_Occlusion();
+			OpState = UpdateResult_Occlusion();
 			break;
 		case EBakeMapType::Curvature:
-			UpdateResult_Curvature();
+			OpState = UpdateResult_Curvature();
 			break;
 
 		case EBakeMapType::NormalImage:
 		case EBakeMapType::FaceNormalImage:
 		case EBakeMapType::PositionImage:
-			UpdateResult_MeshProperty();
+			OpState = UpdateResult_MeshProperty();
 			break;
 
 		case EBakeMapType::Texture2DImage:
-			UpdateResult_Texture2DImage();
+			OpState = UpdateResult_Texture2DImage();
 			break;
 	}
 
+	// Early exit if op input parameters are invalid.
+	if (OpState == EOpState::Invalid)
+	{
+		return;
+	}
 
-	UpdateVisualization();
-	GetToolManager()->PostInvalidation();
-
-	bResultValid = true;
+	// This should be the only point of compute invalidation to
+	// minimize synchronization issues.
+	bool bInvalidate = bInputsDirty || (OpState == EOpState::Evaluate);
+	if (!Compute)
+	{
+		Compute = MakeUnique<TGenericDataBackgroundCompute<FMeshImageBaker>>();
+		Compute->Setup(this);
+		Compute->OnResultUpdated.AddLambda([this](const TUniquePtr<FMeshImageBaker>& NewResult) { OnMapsUpdated(NewResult); });
+		Compute->InvalidateResult();
+	}
+	else if (bInvalidate)
+	{
+		Compute->InvalidateResult();
+	}
+	bInputsDirty = false;
 }
 
 
 
-void UBakeMeshAttributeMapsTool::UpdateResult_Normal()
+UBakeMeshAttributeMapsTool::EOpState UBakeMeshAttributeMapsTool::UpdateResult_Normal()
 {
-	check(BakeCache->IsCacheValid());
+	EOpState ResultState = EOpState::Complete;
 
 	int32 ImageSize = (int32)Settings->Resolution;
 	FImageDimensions Dimensions(ImageSize, ImageSize);
@@ -478,24 +807,16 @@ void UBakeMeshAttributeMapsTool::UpdateResult_Normal()
 
 	if (!(CachedNormalMapSettings == NormalMapSettings))
 	{
-		FMeshNormalMapBaker NormalBaker;
-		NormalBaker.SetCache(BakeCache.Get());
-		NormalBaker.BaseMeshTangents = BaseMeshTangents.Get();
-		NormalBaker.Bake();
-
-		FTexture2DBuilder TextureBuilder;
-		TextureBuilder.Initialize(FTexture2DBuilder::ETextureType::NormalMap, Dimensions);
-		TextureBuilder.Copy(*NormalBaker.GetResult());
-		TextureBuilder.Commit(false);
-		CachedNormalMap = TextureBuilder.GetTexture2D();
 		CachedNormalMapSettings = NormalMapSettings;
+		ResultState = EOpState::Evaluate;
 	}
+	return ResultState;
 }
 
 
-void UBakeMeshAttributeMapsTool::UpdateResult_Occlusion()
+UBakeMeshAttributeMapsTool::EOpState UBakeMeshAttributeMapsTool::UpdateResult_Occlusion()
 {
-	check(BakeCache->IsCacheValid());
+	EOpState ResultState = EOpState::Complete;
 
 	int32 ImageSize = (int32)Settings->Resolution;
 	FImageDimensions Dimensions(ImageSize, ImageSize);
@@ -512,55 +833,17 @@ void UBakeMeshAttributeMapsTool::UpdateResult_Occlusion()
 
 	if ( !(CachedOcclusionMapSettings == OcclusionMapSettings) )
 	{
-		FMeshOcclusionMapBaker OcclusionBaker;
-		OcclusionBaker.SetCache(BakeCache.Get());
-		OcclusionBaker.BaseMeshTangents = BaseMeshTangents.Get();
-		OcclusionBaker.NumOcclusionRays = OcclusionMapSettings.OcclusionRays;
-		OcclusionBaker.MaxDistance = OcclusionMapSettings.MaxDistance;
-		OcclusionBaker.SpreadAngle = OcclusionMapSettings.SpreadAngle;
-		OcclusionBaker.BlurRadius = OcclusionMapSettings.BlurRadius;
-		OcclusionBaker.BiasAngleDeg = OcclusionMapSettings.BiasAngle;
-		switch (OcclusionMapSettings.Distribution)
-		{
-		case EOcclusionMapDistribution::Cosine:
-			OcclusionBaker.Distribution = FMeshOcclusionMapBaker::EDistribution::Cosine;
-			break;
-		case EOcclusionMapDistribution::Uniform:
-			OcclusionBaker.Distribution = FMeshOcclusionMapBaker::EDistribution::Uniform;
-			break;
-		}
-		switch (OcclusionMapSettings.NormalSpace)
-		{
-		case ENormalMapSpace::Tangent:
-			OcclusionBaker.NormalSpace = FMeshOcclusionMapBaker::ESpace::Tangent;
-			break;
-		case ENormalMapSpace::Object:
-			OcclusionBaker.NormalSpace = FMeshOcclusionMapBaker::ESpace::Object;
-			break;
-		}
-		OcclusionBaker.Bake();
-
-		FTexture2DBuilder TextureBuilder;
-		TextureBuilder.Initialize(FTexture2DBuilder::ETextureType::AmbientOcclusion, Dimensions);
-		TextureBuilder.Copy(*OcclusionBaker.GetResult(FMeshOcclusionMapBaker::EResult::AmbientOcclusion));
-		TextureBuilder.Commit(false);
-		CachedOcclusionMap = TextureBuilder.GetTexture2D();
-
-		FTexture2DBuilder TextureNormalBuilder;
-		TextureNormalBuilder.Initialize(FTexture2DBuilder::ETextureType::NormalMap, Dimensions);
-		TextureNormalBuilder.Copy(*OcclusionBaker.GetResult(FMeshOcclusionMapBaker::EResult::BentNormal));
-		TextureNormalBuilder.Commit(false);
-		CachedBentNormalMap = TextureNormalBuilder.GetTexture2D();
-
 		CachedOcclusionMapSettings = OcclusionMapSettings;
+		ResultState = EOpState::Evaluate;
 	}
+	return ResultState;
 }
 
 
 
-void UBakeMeshAttributeMapsTool::UpdateResult_Curvature()
+UBakeMeshAttributeMapsTool::EOpState UBakeMeshAttributeMapsTool::UpdateResult_Curvature()
 {
-	check(BakeCache->IsCacheValid());
+	EOpState ResultState = EOpState::Complete;
 
 	int32 ImageSize = (int32)Settings->Resolution;
 	FImageDimensions Dimensions(ImageSize, ImageSize);
@@ -615,30 +898,17 @@ void UBakeMeshAttributeMapsTool::UpdateResult_Curvature()
 
 	if (!(CachedCurvatureMapSettings == CurvatureMapSettings))
 	{
-		FMeshCurvatureMapBaker CurvatureBaker;
-		CurvatureBaker.SetCache(BakeCache.Get());
-		CurvatureBaker.RangeScale = FMathd::Clamp(CurvatureMapSettings.RangeMultiplier, 0.0001, 1000.0);
-		CurvatureBaker.MinRangeScale = FMathd::Clamp(CurvatureMapSettings.MinRangeMultiplier, 0.0, 1.0);
-		CurvatureBaker.UseCurvatureType = (FMeshCurvatureMapBaker::ECurvatureType)CurvatureMapSettings.CurvatureType;
-		CurvatureBaker.UseColorMode = (FMeshCurvatureMapBaker::EColorMode)CurvatureMapSettings.ColorMode;
-		CurvatureBaker.UseClampMode = (FMeshCurvatureMapBaker::EClampMode)CurvatureMapSettings.ClampMode;
-		CurvatureBaker.BlurRadius = CurvatureMapSettings.BlurRadius;
-		CurvatureBaker.Bake();
-
-		FTexture2DBuilder TextureBuilder;
-		TextureBuilder.Initialize(FTexture2DBuilder::ETextureType::Color, Dimensions);
-		TextureBuilder.Copy(*CurvatureBaker.GetResult());
-		TextureBuilder.Commit(false);
-		CachedCurvatureMap = TextureBuilder.GetTexture2D();
 		CachedCurvatureMapSettings = CurvatureMapSettings;
+		ResultState = EOpState::Evaluate;
 	}
+	return ResultState;
 }
 
 
 
-void UBakeMeshAttributeMapsTool::UpdateResult_MeshProperty()
+UBakeMeshAttributeMapsTool::EOpState UBakeMeshAttributeMapsTool::UpdateResult_MeshProperty()
 {
-	check(BakeCache->IsCacheValid());
+	EOpState ResultState = EOpState::Complete;
 
 	int32 ImageSize = (int32)Settings->Resolution;
 	FImageDimensions Dimensions(ImageSize, ImageSize);
@@ -663,19 +933,10 @@ void UBakeMeshAttributeMapsTool::UpdateResult_MeshProperty()
 
 	if (!(CachedMeshPropertyMapSettings == MeshPropertyMapSettings))
 	{
-		FMeshPropertyMapBaker MeshPropertyBaker;
-		MeshPropertyBaker.SetCache(BakeCache.Get());
-		MeshPropertyBaker.Property = (EMeshPropertyBakeType)MeshPropertyMapSettings.PropertyTypeIndex;
-		//MeshPropertyBaker.BlurRadius = CurvatureMapSettings.BlurRadius;
-		MeshPropertyBaker.Bake();
-
-		FTexture2DBuilder TextureBuilder;
-		TextureBuilder.Initialize(FTexture2DBuilder::ETextureType::Color, Dimensions);
-		TextureBuilder.Copy(*MeshPropertyBaker.GetResult());
-		TextureBuilder.Commit(false);
-		CachedMeshPropertyMap = TextureBuilder.GetTexture2D();
 		CachedMeshPropertyMapSettings = MeshPropertyMapSettings;
+		ResultState = EOpState::Evaluate;
 	}
+	return ResultState;
 }
 
 
@@ -775,9 +1036,9 @@ private:
 
 
 
-void UBakeMeshAttributeMapsTool::UpdateResult_Texture2DImage()
+UBakeMeshAttributeMapsTool::EOpState UBakeMeshAttributeMapsTool::UpdateResult_Texture2DImage()
 {
-	check(BakeCache->IsCacheValid());
+	EOpState ResultState = EOpState::Complete;
 
 	int32 ImageSize = (int32)Settings->Resolution;
 	FImageDimensions Dimensions(ImageSize, ImageSize);
@@ -790,44 +1051,33 @@ void UBakeMeshAttributeMapsTool::UpdateResult_Texture2DImage()
 	if (UVOverlay == nullptr)
 	{
 		GetToolManager()->DisplayMessage(LOCTEXT("InvalidUVWarning", "The Source Mesh does not have the selected UV layer"), EToolMessageLevel::UserWarning);
-		return;
+		return EOpState::Invalid;
 	}
 	
 	if (Texture2DProps->SourceTexture == nullptr)
 	{
 		GetToolManager()->DisplayMessage(LOCTEXT("InvalidTextureWarning", "The Source Texture is not valid"), EToolMessageLevel::UserWarning);
-		return;
+		return EOpState::Invalid;
 	}
 
 
-	TImageBuilder<FVector4f> TextureImage;
 	{
 		FTempTextureAccess TextureAccess(Texture2DProps->SourceTexture);
-		TextureImage.SetDimensions(TextureAccess.GetDimensions());
-		if (!TextureAccess.CopyTo(TextureImage))
+		CachedTextureImage = MakeShared<UE::Geometry::TImageBuilder<FVector4f>, ESPMode::ThreadSafe>();
+		CachedTextureImage->SetDimensions(TextureAccess.GetDimensions());
+		if (!TextureAccess.CopyTo(*CachedTextureImage))
 		{
 			GetToolManager()->DisplayMessage(LOCTEXT("CannotReadTextureWarning", "Cannot read from the source texture"), EToolMessageLevel::UserWarning);
-			return;
+			return EOpState::Invalid;
 		}
 	}
 
 	if (!(CachedTexture2DImageSettings == NewSettings))
 	{
-		FMeshResampleImageBaker Baker;
-		Baker.SetCache(BakeCache.Get());
-		Baker.DetailUVOverlay = UVOverlay;
-		Baker.SampleFunction = [&](FVector2d UVCoord) { 
-			return TextureImage.BilinearSampleUV<float>(UVCoord, FVector4f(0,0,0,1));
-		};
-		Baker.Bake();
-
-		FTexture2DBuilder TextureBuilder;
-		TextureBuilder.Initialize(FTexture2DBuilder::ETextureType::Color, Dimensions);
-		TextureBuilder.Copy(*Baker.GetResult(), true);
-		TextureBuilder.Commit(false);
-		CachedTexture2DImageMap = TextureBuilder.GetTexture2D();
 		CachedTexture2DImageSettings = NewSettings;
+		ResultState = EOpState::Evaluate;
 	}
+	return ResultState;
 }
 
 
@@ -927,6 +1177,73 @@ void UBakeMeshAttributeMapsTool::UpdateOnModeChange()
 }
 
 
+void UBakeMeshAttributeMapsTool::OnMapsUpdated(const TUniquePtr<FMeshImageBaker>& NewResult)
+{
+	switch (Settings->MapType)
+	{
+	case EBakeMapType::TangentSpaceNormalMap:
+	{
+		FMeshNormalMapBaker* Baker = static_cast<FMeshNormalMapBaker*>(NewResult.Get());
+		FTexture2DBuilder TextureBuilder;
+		TextureBuilder.Initialize(FTexture2DBuilder::ETextureType::NormalMap, CachedNormalMapSettings.Dimensions);
+		TextureBuilder.Copy(*Baker->GetResult());
+		TextureBuilder.Commit(false);
+		CachedNormalMap = TextureBuilder.GetTexture2D();
+		break;
+	}
+	case EBakeMapType::Occlusion:
+	{
+		FMeshOcclusionMapBaker* Baker = static_cast<FMeshOcclusionMapBaker*>(NewResult.Get());
+		FTexture2DBuilder TextureBuilder;
+		TextureBuilder.Initialize(FTexture2DBuilder::ETextureType::AmbientOcclusion, CachedOcclusionMapSettings.Dimensions);
+		TextureBuilder.Copy(*Baker->GetResult(FMeshOcclusionMapBaker::EResult::AmbientOcclusion));
+		TextureBuilder.Commit(false);
+		CachedOcclusionMap = TextureBuilder.GetTexture2D();
+
+		FTexture2DBuilder TextureNormalBuilder;
+		TextureNormalBuilder.Initialize(FTexture2DBuilder::ETextureType::NormalMap, CachedOcclusionMapSettings.Dimensions);
+		TextureNormalBuilder.Copy(*Baker->GetResult(FMeshOcclusionMapBaker::EResult::BentNormal));
+		TextureNormalBuilder.Commit(false);
+		CachedBentNormalMap = TextureNormalBuilder.GetTexture2D();
+		break;
+	}
+	case EBakeMapType::Curvature:
+	{
+		FMeshCurvatureMapBaker* Baker = static_cast<FMeshCurvatureMapBaker*>(NewResult.Get());
+		FTexture2DBuilder TextureBuilder;
+		TextureBuilder.Initialize(FTexture2DBuilder::ETextureType::Color, CachedCurvatureMapSettings.Dimensions);
+		TextureBuilder.Copy(*Baker->GetResult());
+		TextureBuilder.Commit(false);
+		CachedCurvatureMap = TextureBuilder.GetTexture2D();
+		break;
+	}
+	case EBakeMapType::PositionImage:
+	case EBakeMapType::NormalImage:
+	case EBakeMapType::FaceNormalImage:
+	{
+		FMeshPropertyMapBaker* Baker = static_cast<FMeshPropertyMapBaker*>(NewResult.Get());
+		FTexture2DBuilder TextureBuilder;
+		TextureBuilder.Initialize(FTexture2DBuilder::ETextureType::Color, CachedMeshPropertyMapSettings.Dimensions);
+		TextureBuilder.Copy(*Baker->GetResult());
+		TextureBuilder.Commit(false);
+		CachedMeshPropertyMap = TextureBuilder.GetTexture2D();
+		break;
+	}
+	case EBakeMapType::Texture2DImage:
+	{
+		FMeshResampleImageBaker* Baker = static_cast<FMeshResampleImageBaker*>(NewResult.Get());
+		FTexture2DBuilder TextureBuilder;
+		TextureBuilder.Initialize(FTexture2DBuilder::ETextureType::Color, CachedTexture2DImageSettings.Dimensions);
+		TextureBuilder.Copy(*Baker->GetResult(), true);
+		TextureBuilder.Commit(false);
+		CachedTexture2DImageMap = TextureBuilder.GetTexture2D();
+		break;
+	}
+	}
+
+	UpdateVisualization();
+	GetToolManager()->PostInvalidation();
+}
 
 
 void UBakeMeshAttributeMapsTool::InitializeEmptyMaps()
diff --git a/Engine/Plugins/Experimental/MeshModelingToolset/Source/MeshModelingTools/Public/BakeMeshAttributeMapsTool.h b/Engine/Plugins/Experimental/MeshModelingToolset/Source/MeshModelingTools/Public/BakeMeshAttributeMapsTool.h
--- a/Engine/Plugins/Experimental/MeshModelingToolset/Source/MeshModelingTools/Public/BakeMeshAttributeMapsTool.h
+++ b/Engine/Plugins/Experimental/MeshModelingToolset/Source/MeshModelingTools/Public/BakeMeshAttributeMapsTool.h
@@ -9,6 +9,10 @@
 #include "DynamicMesh3.h"
 #include "DynamicMeshAABBTree3.h"
 #include "Image/ImageDimensions.h"
+#include "Image/ImageBuilder.h"
+#include "Sampling/MeshImageBaker.h"
+#include "ModelingOperators.h"
+#include "MeshOpPreviewHelpers.h"
 #include "BakeMeshAttributeMapsTool.generated.h"
 
 
@@ -24,9 +28,14 @@ using UE::Geometry::FImageDimensions;
 class IPrimitiveComponentBackedTarget;
 class IMeshDescriptionProvider;
 class IMaterialProvider;
+class FBakeNormalMapOp;
+class FBakeOcclusionMapOp;
+class FBakeCurvatureMapOp;
+class FBakeMeshPropertyMapOp;
+class FBakeTexture2DImageMapOp;
 
 /**
- *
+ * Tool Builder
  */
 UCLASS()
 class MESHMODELINGTOOLS_API UBakeMeshAttributeMapsToolBuilder : public UInteractiveToolBuilder
@@ -301,7 +310,7 @@ public:
  * Detail Map Baking Tool
  */
 UCLASS()
-class MESHMODELINGTOOLS_API UBakeMeshAttributeMapsTool : public UMultiSelectionTool
+class MESHMODELINGTOOLS_API UBakeMeshAttributeMapsTool : public UMultiSelectionTool, public UE::Geometry::IGenericDataOperatorFactory<UE::Geometry::FMeshImageBaker>
 {
 	GENERATED_BODY()
 
@@ -313,12 +322,16 @@ public:
 	virtual void Setup() override;
 	virtual void Shutdown(EToolShutdownType ShutdownType) override;
 
+	virtual void OnTick(float DeltaTime) override;
 	virtual void Render(IToolsContextRenderAPI* RenderAPI) override;
 
 	virtual bool HasCancel() const override { return true; }
 	virtual bool HasAccept() const override { return true; }
 	virtual bool CanAccept() const override;
 
+	// IGenericDataOperatorFactory API
+	virtual TUniquePtr<UE::Geometry::TGenericDataOperator<UE::Geometry::FMeshImageBaker>> MakeNewOperator() override;
+
 protected:
 	// need to update bResultValid if these are modified, so we don't publicly expose them. 
 	// @todo setters/getters for these
@@ -344,6 +357,13 @@ protected:
 
 
 protected:
+	friend class FBakeMapBaseOp;
+	friend class FBakeNormalMapOp;
+	friend class FBakeOcclusionMapOp;
+	friend class FBakeCurvatureMapOp;
+	friend class FBakeTexture2DImageMapOp;
+	friend class FBakeMeshPropertyMapOp;
+
 	IAssetGenerationAPI* AssetAPI = nullptr;
 
 	USimpleDynamicMeshComponent* DynamicMeshComponent;
@@ -354,6 +374,10 @@ protected:
     UPROPERTY()
 	UMaterialInstanceDynamic* BentNormalPreviewMaterial;
 
+	UPROPERTY()
+	UMaterialInstanceDynamic* WorkingPreviewMaterial;
+	float SecondsBeforeWorkingMaterial = 0.75;
+
 	TSharedPtr<FMeshDescription, ESPMode::ThreadSafe> BaseMeshDescription;
 	TSharedPtr<UE::Geometry::TMeshTangents<double>, ESPMode::ThreadSafe> BaseMeshTangents;
 	UE::Geometry::FDynamicMesh3 BaseMesh;
@@ -367,13 +391,15 @@ protected:
 	void UpdateDetailMesh();
 	bool bDetailMeshValid = false;
 
-	bool bResultValid;
+	bool bInputsDirty = false;
 	void UpdateResult();
 
 	void UpdateOnModeChange();
 	void UpdateVisualization();
 
-	TPimplPtr<UE::Geometry::FMeshImageBakingCache> BakeCache;
+	TUniquePtr<TGenericDataBackgroundCompute<UE::Geometry::FMeshImageBaker>> Compute = nullptr;
+	void OnMapsUpdated(const TUniquePtr<UE::Geometry::FMeshImageBaker>& NewResult);
+
 	struct FBakeCacheSettings
 	{
 		FImageDimensions Dimensions;
@@ -388,7 +414,13 @@ protected:
 	};
 	FBakeCacheSettings CachedBakeCacheSettings;
 
-
+	enum class EOpState
+	{
+		Complete,	// Inputs valid & Result is valid - no-op.
+		Evaluate,	// Inputs valid & Result is invalid - re-evaluate.
+		Invalid		// Inputs invalid - pause eval.
+	};
+	EOpState OpState = EOpState::Evaluate;
 
 	struct FNormalMapSettings
 	{
@@ -403,7 +435,7 @@ protected:
 	UPROPERTY()
 	UTexture2D* CachedNormalMap;
 
-	void UpdateResult_Normal();
+	EOpState UpdateResult_Normal();
 
 
 
@@ -436,7 +468,7 @@ protected:
 	UPROPERTY()
 	UTexture2D* CachedBentNormalMap;
 
-	void UpdateResult_Occlusion();
+	EOpState UpdateResult_Occlusion();
 
 
 
@@ -461,7 +493,7 @@ protected:
 	UPROPERTY()
 	UTexture2D* CachedCurvatureMap;
 
-	void UpdateResult_Curvature();
+	EOpState UpdateResult_Curvature();
 
 
 
@@ -480,7 +512,7 @@ protected:
 	UPROPERTY()
 	UTexture2D* CachedMeshPropertyMap;
 
-	void UpdateResult_MeshProperty();
+	EOpState UpdateResult_MeshProperty();
 
 
 
@@ -495,11 +527,12 @@ protected:
 			return Dimensions == Other.Dimensions && UVLayer == Other.UVLayer;
 		}
 	};
+	TSharedPtr<UE::Geometry::TImageBuilder<FVector4f>, ESPMode::ThreadSafe> CachedTextureImage;
 	FTexture2DImageSettings CachedTexture2DImageSettings;
 	UPROPERTY()
 	UTexture2D* CachedTexture2DImageMap;
 
-	void UpdateResult_Texture2DImage();
+	EOpState UpdateResult_Texture2DImage();
 
 
 
diff --git a/Engine/Plugins/Experimental/MeshModelingToolset/Source/ModelingComponents/Public/MeshOpPreviewHelpers.h b/Engine/Plugins/Experimental/MeshModelingToolset/Source/ModelingComponents/Public/MeshOpPreviewHelpers.h
--- a/Engine/Plugins/Experimental/MeshModelingToolset/Source/ModelingComponents/Public/MeshOpPreviewHelpers.h
+++ b/Engine/Plugins/Experimental/MeshModelingToolset/Source/ModelingComponents/Public/MeshOpPreviewHelpers.h
@@ -298,6 +298,19 @@ public:
 	 */
 	bool HaveValidResult() const { return bResultValid; }
 
+	/**
+	 * @return the elapsed compute time
+	 */
+	float GetElapsedComputeTime() const
+	{
+		check(BackgroundCompute);
+		if (BackgroundCompute)
+		{
+			return BackgroundCompute->GetElapsedComputeTime();
+		}
+		return 0.0;
+	}
+
 
 	//
 	// Change notification
