ModelingTools: add option to place Pivot at World Origin in EditPivotTool
#rnx

diff --git a/Engine/Plugins/Experimental/MeshModelingToolset/Source/MeshModelingTools/Private/EditPivotTool.cpp b/Engine/Plugins/Experimental/MeshModelingToolset/Source/MeshModelingTools/Private/EditPivotTool.cpp
--- a/Engine/Plugins/Experimental/MeshModelingToolset/Source/MeshModelingTools/Private/EditPivotTool.cpp
+++ b/Engine/Plugins/Experimental/MeshModelingToolset/Source/MeshModelingTools/Private/EditPivotTool.cpp
@@ -256,6 +256,10 @@ void UEditPivotTool::ApplyAction(EEditPivotToolActions ActionType)
 		case EEditPivotToolActions::Back:
 			SetPivotToBoxPoint(ActionType);
 			break;
+
+		case EEditPivotToolActions::WorldOrigin:
+			SetPivotToWorldOrigin();
+			break;
 	}
 }
 
@@ -294,6 +298,11 @@ void UEditPivotTool::SetPivotToBoxPoint(EEditPivotToolActions ActionPoint)
 }
 
 
+void UEditPivotTool::SetPivotToWorldOrigin()
+{
+	ActiveGizmos[0].TransformGizmo->SetNewGizmoTransform(FTransform());
+}
+
 
 void UEditPivotTool::SetActiveGizmos_Single(bool bLocalRotations)
 {
diff --git a/Engine/Plugins/Experimental/MeshModelingToolset/Source/MeshModelingTools/Public/EditPivotTool.h b/Engine/Plugins/Experimental/MeshModelingToolset/Source/MeshModelingTools/Public/EditPivotTool.h
--- a/Engine/Plugins/Experimental/MeshModelingToolset/Source/MeshModelingTools/Public/EditPivotTool.h
+++ b/Engine/Plugins/Experimental/MeshModelingToolset/Source/MeshModelingTools/Public/EditPivotTool.h
@@ -99,7 +99,8 @@ enum class EEditPivotToolActions
 	Left,
 	Right,
 	Front,
-	Back
+	Back,
+	WorldOrigin
 };
 
 
@@ -140,6 +141,8 @@ public:
 	UFUNCTION(CallInEditor, Category = BoxPositions, meta = (DisplayPriority = 4))
 	void Back() { PostAction(EEditPivotToolActions::Back); }
 
+	UFUNCTION(CallInEditor, Category = BoxPositions, meta = (DisplayPriority = 5))
+	void WorldOrigin() { PostAction(EEditPivotToolActions::WorldOrigin); }
 };
 
 
@@ -213,6 +216,7 @@ protected:
 	EEditPivotToolActions PendingAction;
 	virtual void ApplyAction(EEditPivotToolActions ActionType);
 	virtual void SetPivotToBoxPoint(EEditPivotToolActions ActionPoint);
+	virtual void SetPivotToWorldOrigin();
 
 	void UpdateAssets(const UE::Geometry::FFrame3d& NewPivotWorldFrame);
 };
