Zero-initialized to uninitialized UStruct members in LODManagerTool and CollisionPeroertySets

diff --git a/Engine/Plugins/Experimental/MeshLODToolset/Source/MeshLODToolset/Public/Tools/LODManagerTool.h b/Engine/Plugins/Experimental/MeshLODToolset/Source/MeshLODToolset/Public/Tools/LODManagerTool.h
--- a/Engine/Plugins/Experimental/MeshLODToolset/Source/MeshLODToolset/Public/Tools/LODManagerTool.h
+++ b/Engine/Plugins/Experimental/MeshLODToolset/Source/MeshLODToolset/Public/Tools/LODManagerTool.h
@@ -40,10 +40,10 @@ struct MESHLODTOOLSET_API FLODManagerLODInfo
 	GENERATED_BODY()
 
 	UPROPERTY(VisibleAnywhere, Category = LODInformation)
-	int32 VertexCount;
+	int32 VertexCount = 0;
 
 	UPROPERTY(VisibleAnywhere, Category = LODInformation)
-	int32 TriangleCount;
+	int32 TriangleCount = 0;
 };
 
 
diff --git a/Engine/Plugins/Experimental/MeshModelingToolset/Source/MeshModelingTools/Public/Physics/CollisionPropertySets.h b/Engine/Plugins/Experimental/MeshModelingToolset/Source/MeshModelingTools/Public/Physics/CollisionPropertySets.h
--- a/Engine/Plugins/Experimental/MeshModelingToolset/Source/MeshModelingTools/Public/Physics/CollisionPropertySets.h
+++ b/Engine/Plugins/Experimental/MeshModelingToolset/Source/MeshModelingTools/Public/Physics/CollisionPropertySets.h
@@ -47,7 +47,7 @@ struct MESHMODELINGTOOLS_API FPhysicsBoxData
 	GENERATED_BODY()
 
 	UPROPERTY(VisibleAnywhere, Category = Sphere)
-	FVector Dimensions;
+	FVector Dimensions = FVector::ZeroVector;
 
 	UPROPERTY(VisibleAnywhere, Category = Sphere)
 	FTransform Transform;
