Copy the separate asset and component materials from the source mesh to the new assets generated by plane cut tool (+ generally support creating an asset with separate component and asset materials)

#rnx

diff --git a/Engine/Plugins/Experimental/MeshModelingToolset/Source/MeshModelingTools/Private/PlaneCutTool.cpp b/Engine/Plugins/Experimental/MeshModelingToolset/Source/MeshModelingTools/Private/PlaneCutTool.cpp
--- a/Engine/Plugins/Experimental/MeshModelingToolset/Source/MeshModelingTools/Private/PlaneCutTool.cpp
+++ b/Engine/Plugins/Experimental/MeshModelingToolset/Source/MeshModelingTools/Private/PlaneCutTool.cpp
@@ -599,20 +599,18 @@ void UPlaneCutTool::GenerateAsset(const TArray<FDynamicMeshOpResult>& Results)
 				continue;
 			}
 
-			// build array of materials from the original
-			TArray<UMaterialInterface*> Materials;
+			// get materials for both the component and the asset
 			IMaterialProvider* TargetMaterial = TargetMaterialInterface(OrigMeshIdx);
-			for (int MaterialIdx = 0, NumMaterials = TargetMaterial->GetNumMaterials(); MaterialIdx < NumMaterials; MaterialIdx++)
-			{
-				Materials.Add(TargetMaterial->GetMaterial(MaterialIdx));
-			}
+			FComponentMaterialSet ComponentMaterialSet, AssetMaterialSet;
+			TargetMaterial->GetMaterialSet(ComponentMaterialSet, false);
+			TargetMaterial->GetMaterialSet(AssetMaterialSet, true /*prefer asset materials*/);
 
 			// add all the additional meshes
 			for (int AddMeshIdx = 1; AddMeshIdx < SplitMeshes.Num(); AddMeshIdx++)
 			{
 				AActor* NewActor = AssetGenerationUtil::GenerateStaticMeshActor(
 					AssetAPI, TargetWorld,
-					&SplitMeshes[AddMeshIdx], Results[OrigMeshIdx].Transform, TEXT("PlaneCutOtherPart"), Materials);
+					&SplitMeshes[AddMeshIdx], Results[OrigMeshIdx].Transform, TEXT("PlaneCutOtherPart"), ComponentMaterialSet.Materials, AssetMaterialSet.Materials);
 				if (NewActor != nullptr)
 				{
 					NewSelection.Actors.Add(NewActor);
diff --git a/Engine/Plugins/Experimental/MeshModelingToolset/Source/ModelingComponents/Private/AssetGenerationUtil.cpp b/Engine/Plugins/Experimental/MeshModelingToolset/Source/ModelingComponents/Private/AssetGenerationUtil.cpp
--- a/Engine/Plugins/Experimental/MeshModelingToolset/Source/ModelingComponents/Private/AssetGenerationUtil.cpp
+++ b/Engine/Plugins/Experimental/MeshModelingToolset/Source/ModelingComponents/Private/AssetGenerationUtil.cpp
@@ -31,14 +31,22 @@ AActor* AssetGenerationUtil::GenerateStaticMeshActor(
 	const FDynamicMesh3* Mesh,
 	const UE::Geometry::FTransform3d& Transform,
 	FString ObjectName,
-	const TArrayView<UMaterialInterface*>& Materials
+	const TArrayView<UMaterialInterface*>& ComponentMaterials,
+	const TArrayView<UMaterialInterface*>& AssetMaterials
 )
 {
 	FGeneratedStaticMeshAssetConfig AssetConfig;
-	for (UMaterialInterface* Material : Materials)
+	for (UMaterialInterface* Material : ComponentMaterials)
 	{
 		AssetConfig.Materials.Add(Material);
 	}
+	if (AssetMaterials.Num() > 0 && ensure(AssetMaterials.Num() == ComponentMaterials.Num()))
+	{
+		for (UMaterialInterface* Material : AssetMaterials)
+		{
+			AssetConfig.AssetMaterials.Add(Material);
+		}
+	}
 
 	AssetConfig.MeshDescription = MakeUnique<FMeshDescription>();
 	FStaticMeshAttributes Attributes(*AssetConfig.MeshDescription);
diff --git a/Engine/Plugins/Experimental/MeshModelingToolset/Source/ModelingComponents/Public/AssetGenerationUtil.h b/Engine/Plugins/Experimental/MeshModelingToolset/Source/ModelingComponents/Public/AssetGenerationUtil.h
--- a/Engine/Plugins/Experimental/MeshModelingToolset/Source/ModelingComponents/Public/AssetGenerationUtil.h
+++ b/Engine/Plugins/Experimental/MeshModelingToolset/Source/ModelingComponents/Public/AssetGenerationUtil.h
@@ -84,7 +84,8 @@ namespace AssetGenerationUtil
 	* @param Transform transformation for the new actor
 	* @param ObjectName name of the new asset
 	* @param PackagePath path of the new asset
-	* @param Materials materials to set on generated actor
+	* @param ComponentMaterials materials to set on generated actor component (and asset if AssetMaterials not provided)
+	* @param AssetMaterials materials to set on generated actor mesh asset (optional; if omitted the component materials will be used)
 	* @return new mesh actor
 	*/
 	MODELINGCOMPONENTS_API AActor* GenerateStaticMeshActor(
@@ -93,7 +94,8 @@ namespace AssetGenerationUtil
 		const FDynamicMesh3* Mesh,
 		const UE::Geometry::FTransform3d& Transform,
 		FString ObjectName,
-		const TArrayView<UMaterialInterface*>& Materials
+		const TArrayView<UMaterialInterface*>& ComponentMaterials,
+		const TArrayView<UMaterialInterface*>& AssetMaterials = TArrayView<UMaterialInterface*>()
 	);
 
 
diff --git a/Engine/Plugins/Experimental/ModelingToolsEditorMode/Source/ModelingToolsEditorMode/Private/ModelingModeAssetAPI.cpp b/Engine/Plugins/Experimental/ModelingToolsEditorMode/Source/ModelingToolsEditorMode/Private/ModelingModeAssetAPI.cpp
--- a/Engine/Plugins/Experimental/ModelingToolsEditorMode/Source/ModelingToolsEditorMode/Private/ModelingModeAssetAPI.cpp
+++ b/Engine/Plugins/Experimental/ModelingToolsEditorMode/Source/ModelingToolsEditorMode/Private/ModelingModeAssetAPI.cpp
@@ -201,9 +201,14 @@ AActor* FModelingModeAssetAPI::GenerateStaticMeshActor(
 		Builder.NewMeshComponent->SetMaterial(MatIdx, AssetConfig.Materials[MatIdx]);
 	}
 	// also put the materials on the asset
-	for (int MatIdx = 0, NumMats = AssetConfig.Materials.Num(); MatIdx < NumMats; MatIdx++)
+	// use the asset-specific AssetMaterials array only if it matches the size of the Materials array
+	//  otherwise favor the Materials array
+	TArray<UMaterialInterface*>* AssetMaterials =
+		(AssetConfig.Materials.Num() == AssetConfig.AssetMaterials.Num()) ?
+			&AssetConfig.AssetMaterials : &AssetConfig.Materials;
+	for (int MatIdx = 0, NumMats = AssetMaterials->Num(); MatIdx < NumMats; MatIdx++)
 	{
-		Builder.NewStaticMesh->SetMaterial(MatIdx, AssetConfig.Materials[MatIdx]);
+		Builder.NewStaticMesh->SetMaterial(MatIdx, (*AssetMaterials)[MatIdx]);
 	}
 
 	// save the new asset (or don't, if that's what the user wants)
diff --git a/Engine/Source/Runtime/Experimental/InteractiveToolsFramework/Public/ToolContextInterfaces.h b/Engine/Source/Runtime/Experimental/InteractiveToolsFramework/Public/ToolContextInterfaces.h
--- a/Engine/Source/Runtime/Experimental/InteractiveToolsFramework/Public/ToolContextInterfaces.h
+++ b/Engine/Source/Runtime/Experimental/InteractiveToolsFramework/Public/ToolContextInterfaces.h
@@ -426,7 +426,8 @@ struct FGeneratedStaticMeshAssetConfig
 	INTERACTIVETOOLSFRAMEWORK_API ~FGeneratedStaticMeshAssetConfig();
 
 	TUniquePtr<FMeshDescription> MeshDescription;
-	TArray<UMaterialInterface*> Materials;
+	TArray<UMaterialInterface*> Materials; // Materials for the component (and asset if no AssetMaterials)
+	TArray<UMaterialInterface*> AssetMaterials; // Optional, for the asset only
 
 	bool bEnableRecomputeNormals = false;
 	bool bEnableRecomputeTangents = false;
