add lower bound and comment/tooltip description for Thickness parameter in mesh attribute baker

#rnx

diff --git a/Engine/Plugins/Experimental/MeshModelingToolset/Source/MeshModelingTools/Public/BakeMeshAttributeMapsTool.h b/Engine/Plugins/Experimental/MeshModelingToolset/Source/MeshModelingTools/Public/BakeMeshAttributeMapsTool.h
--- a/Engine/Plugins/Experimental/MeshModelingToolset/Source/MeshModelingTools/Public/BakeMeshAttributeMapsTool.h
+++ b/Engine/Plugins/Experimental/MeshModelingToolset/Source/MeshModelingTools/Public/BakeMeshAttributeMapsTool.h
@@ -100,7 +100,8 @@ public:
 	UPROPERTY(EditAnywhere, Category = MapSettings)
 	bool bUseWorldSpace = false;
 
-	UPROPERTY(EditAnywhere, Category = MapSettings)
+	/** Distance to search for the correspondence between the source and target meshes */
+	UPROPERTY(EditAnywhere, Category = MapSettings, meta = (ClampMin = "0.001"))
 	float Thickness = 3.0;
 
 	/** Which UV layer to use to create the map */
