plane cut cleanup:
 - add comment and editcondition for FillSpans, move next to FillCutHole
 - move default setting values consistently in to the struct, rather than being half there and half in the constructor
 - move action buttons to use a PendingAction variable and actually do the action in the OnTick, rather than doing the action in the handler -- this lets the actions get the correct transaction labels in the undo history

#rnx

diff --git a/Engine/Plugins/Experimental/MeshModelingToolset/Source/MeshModelingTools/Private/PlaneCutTool.cpp b/Engine/Plugins/Experimental/MeshModelingToolset/Source/MeshModelingTools/Private/PlaneCutTool.cpp
--- a/Engine/Plugins/Experimental/MeshModelingToolset/Source/MeshModelingTools/Private/PlaneCutTool.cpp
+++ b/Engine/Plugins/Experimental/MeshModelingToolset/Source/MeshModelingTools/Private/PlaneCutTool.cpp
@@ -81,15 +81,6 @@ UInteractiveTool* UPlaneCutToolBuilder::BuildTool(const FToolBuilderState& Scene
  * Tool
  */
 
-UPlaneCutToolProperties::UPlaneCutToolProperties() :
-	bKeepBothHalves(false),
-	SpacingBetweenHalves(1),
-	bFillCutHole(true),
-	bShowPreview(true),
-	bFillSpans(false)
-{
-}
-
 UPlaneCutTool::UPlaneCutTool()
 {
 	CutPlaneOrigin = FVector::ZeroVector;
@@ -267,7 +258,7 @@ void UPlaneCutTool::SetupPreviews()
 }
 
 
-void UPlaneCutTool::FlipPlane()
+void UPlaneCutTool::DoFlipPlane()
 {
 	GetToolManager()->BeginUndoTransaction(LOCTEXT("FlipPlaneTransactionName", "Flip Plane"));
 
@@ -280,7 +271,7 @@ void UPlaneCutTool::FlipPlane()
 }
 
 
-void UPlaneCutTool::Cut()
+void UPlaneCutTool::DoCut()
 {
 	if (!CanAccept())
 	{
@@ -403,6 +394,20 @@ void UPlaneCutTool::Render(IToolsContextRenderAPI* RenderAPI)
 
 void UPlaneCutTool::OnTick(float DeltaTime)
 {
+	if (PendingAction != EPlaneCutToolActions::NoAction)
+	{
+		if (PendingAction == EPlaneCutToolActions::Cut)
+		{
+			DoCut();
+		}
+		else if (PendingAction == EPlaneCutToolActions::FlipPlane)
+		{
+			DoFlipPlane();
+		}
+
+		PendingAction = EPlaneCutToolActions::NoAction;
+	}
+
 	if (PlaneTransformGizmo != nullptr)
 	{
 		PlaneTransformGizmo->bSnapToWorldGrid =
diff --git a/Engine/Plugins/Experimental/MeshModelingToolset/Source/MeshModelingTools/Public/PlaneCutTool.h b/Engine/Plugins/Experimental/MeshModelingToolset/Source/MeshModelingTools/Public/PlaneCutTool.h
--- a/Engine/Plugins/Experimental/MeshModelingToolset/Source/MeshModelingTools/Public/PlaneCutTool.h
+++ b/Engine/Plugins/Experimental/MeshModelingToolset/Source/MeshModelingTools/Public/PlaneCutTool.h
@@ -71,7 +71,8 @@ class MESHMODELINGTOOLS_API UPlaneCutToolProperties : public UInteractiveToolPro
 	GENERATED_BODY()
 
 public:
-	UPlaneCutToolProperties();
+	UPlaneCutToolProperties()
+	{}
 
 	/** Snap the cut plane to the world grid */
 	UPROPERTY(EditAnywhere, Category = Snapping)
@@ -83,21 +84,32 @@ public:
 
 	/** If true, both halves of the cut are computed */
 	UPROPERTY(EditAnywhere, Category = Options)
-	bool bKeepBothHalves;
+	bool bKeepBothHalves = false;
 
 	/** If keeping both halves, separate the two pieces by this amount */
 	UPROPERTY(EditAnywhere, Category = Options, meta = (EditCondition = "bKeepBothHalves == true", UIMin = "0", ClampMin = "0") )
-	float SpacingBetweenHalves;
+	float SpacingBetweenHalves = 0;
 
-	/** If true, the cut surface is filled with simple planar hole fill surface(s) */
 	UPROPERTY(EditAnywhere, Category = Options)
-	bool bFillCutHole;
+	bool bShowPreview = true;
 
+	/** If true, the cut surface is filled with simple planar hole fill surface(s) */
 	UPROPERTY(EditAnywhere, Category = Options)
-	bool bShowPreview;
+	bool bFillCutHole = true;
 
-	UPROPERTY(EditAnywhere, Category = Options, AdvancedDisplay)
-	bool bFillSpans;
+	/** If true, will attempt to fill cut holes even if they're ill-formed (e.g. because they connect to pre-existing holes in the geometry) */
+	UPROPERTY(EditAnywhere, Category = Options, AdvancedDisplay, meta = (EditCondition = "bFillCutHole", EditConditionHides))
+	bool bFillSpans = false;
+};
+
+
+
+UENUM()
+enum class EPlaneCutToolActions
+{
+	NoAction,
+	Cut,
+	FlipPlane
 };
 
 
@@ -176,13 +188,24 @@ protected:
 	UPROPERTY()
 	TArray<UMeshOpPreviewWithBackgroundCompute*> Previews;
 
+
+	/// Action buttons.
+	/// Note these set a flag to call the action later (in OnTick)
+	/// Otherwise, the actions in undo history will end up being generically named by an outer UI handler transaction
+
 	/** Cut with the current plane without exiting the tool */
 	UFUNCTION(CallInEditor, Category = Actions, meta = (DisplayName = "Cut"))
-	void Cut();
+	void Cut()
+	{
+		PendingAction = EPlaneCutToolActions::Cut;
+	}
 
 	/** Flip the cutting plane */
 	UFUNCTION(CallInEditor, Category = Actions, meta = (DisplayName = "Flip Plane"))
-	void FlipPlane();
+	void FlipPlane()
+	{
+		PendingAction = EPlaneCutToolActions::FlipPlane;
+	}
 
 protected:
 
@@ -201,6 +224,8 @@ protected:
 	static const int IgnoreSnappingModifier = 1;
 	bool bIgnoreSnappingToggle = false;		// toggled by hotkey (shift)
 
+	EPlaneCutToolActions PendingAction = EPlaneCutToolActions::NoAction;
+
 	UPROPERTY()
 	UTransformGizmo* PlaneTransformGizmo;
 
@@ -210,6 +235,9 @@ protected:
 	void TransformChanged(UTransformProxy* Proxy, FTransform Transform);
 	void MeshChanged();
 
+	void DoCut();
+	void DoFlipPlane();
+
 	void SetupPreviews();
 
 	IClickBehaviorTarget* SetPointInWorldConnector = nullptr;
