fix transforms breaking on CSG meshes when toggling only-use-first-mesh-materials setting

#rnx

diff --git a/Engine/Plugins/Experimental/MeshModelingToolset/Source/MeshModelingTools/Private/CSGMeshesTool.cpp b/Engine/Plugins/Experimental/MeshModelingToolset/Source/MeshModelingTools/Private/CSGMeshesTool.cpp
--- a/Engine/Plugins/Experimental/MeshModelingToolset/Source/MeshModelingTools/Private/CSGMeshesTool.cpp
+++ b/Engine/Plugins/Experimental/MeshModelingToolset/Source/MeshModelingTools/Private/CSGMeshesTool.cpp
@@ -228,13 +228,16 @@ void UCSGMeshesTool::ConvertInputsAndSetPreviewMaterials(bool bSetPreviewMesh)
 			MaterialIDs->SetValue(TID, MaterialRemap[ComponentIdx][MaterialIDs->GetValue(TID)]);
 		}
 
-		UPreviewMesh* OriginalMeshPreview = OriginalMeshPreviews.Add_GetRef(NewObject<UPreviewMesh>());
-		OriginalMeshPreview->CreateInWorld(TargetWorld, TargetComponentInterface(ComponentIdx)->GetWorldTransform());
-		OriginalMeshPreview->UpdatePreview(OriginalDynamicMeshes[ComponentIdx].Get());
+		if (bSetPreviewMesh)
+		{
+			UPreviewMesh* OriginalMeshPreview = OriginalMeshPreviews.Add_GetRef(NewObject<UPreviewMesh>());
+			OriginalMeshPreview->CreateInWorld(TargetWorld, TargetComponentInterface(ComponentIdx)->GetWorldTransform());
+			OriginalMeshPreview->UpdatePreview(OriginalDynamicMeshes[ComponentIdx].Get());
 
-		OriginalMeshPreview->SetMaterial(0, PreviewsGhostMaterial);
-		OriginalMeshPreview->SetVisible(false);
-		TransformProxies[ComponentIdx]->AddComponent(OriginalMeshPreview->GetRootComponent());
+			OriginalMeshPreview->SetMaterial(0, PreviewsGhostMaterial);
+			OriginalMeshPreview->SetVisible(false);
+			TransformProxies[ComponentIdx]->AddComponent(OriginalMeshPreview->GetRootComponent());
+		}
 	}
 	Preview->ConfigureMaterials(AllMaterialSet.Materials, ToolSetupUtil::GetDefaultWorkingMaterial(GetToolManager()));
 }
