Don't use the sharp normal setting when computing new normals for simplified meshes on the voxel op (tends to give confusing normal seams)

+ compute per-vertex normals on the overlay if the mesh has one (when not auto-simplifying)

#rnx

diff --git a/Engine/Plugins/Experimental/MeshModelingToolset/Source/ModelingOperators/Private/BaseOps/VoxelBaseOp.cpp b/Engine/Plugins/Experimental/MeshModelingToolset/Source/ModelingOperators/Private/BaseOps/VoxelBaseOp.cpp
--- a/Engine/Plugins/Experimental/MeshModelingToolset/Source/ModelingOperators/Private/BaseOps/VoxelBaseOp.cpp
+++ b/Engine/Plugins/Experimental/MeshModelingToolset/Source/ModelingOperators/Private/BaseOps/VoxelBaseOp.cpp
@@ -60,7 +60,7 @@ void FVoxelBaseOp::PostProcessResult(FProgressCancel* Progress, double MeshCellS
 		EditNormalsOp.bRecomputeNormals = true;
 		EditNormalsOp.NormalCalculationMethod = ENormalCalculationMethod::AreaAngleWeighting;
 		EditNormalsOp.SplitNormalMethod = ESplitNormalMethod::FaceNormalThreshold;
-		EditNormalsOp.bAllowSharpVertices = true;
+		EditNormalsOp.bAllowSharpVertices = false;
 		EditNormalsOp.NormalSplitThreshold = 60.f;
 
 		EditNormalsOp.SetTransform(FTransform(ResultTransform));
@@ -70,8 +70,15 @@ void FVoxelBaseOp::PostProcessResult(FProgressCancel* Progress, double MeshCellS
 	}
 	else
 	{
-		// if nothing was simplified, just use quick vertex normals
-		FMeshNormals::QuickComputeVertexNormals(*ResultMesh.Get());
+		// if nothing was simplified, just use per-vertex normals
+		if (ResultMesh->HasAttributes())
+		{
+			FMeshNormals::InitializeOverlayToPerVertexNormals(ResultMesh->Attributes()->PrimaryNormals(), false);
+		}
+		else
+		{
+			FMeshNormals::QuickComputeVertexNormals(*ResultMesh.Get());
+		}
 	}
 
 	if (Progress && Progress->Cancelled())
