remove editconditionhides from the single 'advanced' option on plane cut tool (FillSpans), because when it's hidden that leaves a confusing 'show advanced options' arrow toggle that toggles zero options

#rnx

diff --git a/Engine/Plugins/Experimental/MeshModelingToolset/Source/MeshModelingTools/Public/PlaneCutTool.h b/Engine/Plugins/Experimental/MeshModelingToolset/Source/MeshModelingTools/Public/PlaneCutTool.h
--- a/Engine/Plugins/Experimental/MeshModelingToolset/Source/MeshModelingTools/Public/PlaneCutTool.h
+++ b/Engine/Plugins/Experimental/MeshModelingToolset/Source/MeshModelingTools/Public/PlaneCutTool.h
@@ -98,7 +98,7 @@ public:
 	bool bFillCutHole = true;
 
 	/** If true, will attempt to fill cut holes even if they're ill-formed (e.g. because they connect to pre-existing holes in the geometry) */
-	UPROPERTY(EditAnywhere, Category = Options, AdvancedDisplay, meta = (EditCondition = "bFillCutHole", EditConditionHides))
+	UPROPERTY(EditAnywhere, Category = Options, AdvancedDisplay, meta = (EditCondition = "bFillCutHole"))
 	bool bFillSpans = false;
 };
 
