change plane cut tool to use the construction plane mechanic

#rnx

diff --git a/Engine/Plugins/Experimental/MeshModelingToolset/Source/MeshModelingTools/Private/PlaneCutTool.cpp b/Engine/Plugins/Experimental/MeshModelingToolset/Source/MeshModelingTools/Private/PlaneCutTool.cpp
--- a/Engine/Plugins/Experimental/MeshModelingToolset/Source/MeshModelingTools/Private/PlaneCutTool.cpp
+++ b/Engine/Plugins/Experimental/MeshModelingToolset/Source/MeshModelingTools/Private/PlaneCutTool.cpp
@@ -83,8 +83,6 @@ UInteractiveTool* UPlaneCutToolBuilder::BuildTool(const FToolBuilderState& Scene
 
 UPlaneCutTool::UPlaneCutTool()
 {
-	CutPlaneOrigin = FVector::ZeroVector;
-	CutPlaneOrientation = FQuat::Identity;
 }
 
 void UPlaneCutTool::SetWorld(UWorld* World)
@@ -98,7 +96,7 @@ void UPlaneCutTool::Setup()
 
 	// add modifier button for snapping
 	UKeyAsModifierInputBehavior* SnapToggleBehavior = NewObject<UKeyAsModifierInputBehavior>();
-	SnapToggleBehavior->Initialize(this, IgnoreSnappingModifier, FInputDeviceState::IsShiftKeyDown);
+	SnapToggleBehavior->Initialize(this, SnappingModifier, FInputDeviceState::IsShiftKeyDown);
 	AddInputBehavior(SnapToggleBehavior);
 
 	// hide input StaticMeshComponents
@@ -139,38 +137,6 @@ void UPlaneCutTool::Setup()
 		Target->OnMeshChanged.AddLambda([this, Idx]() { Previews[Idx]->InvalidateResult(); });
 	}
 
-	// click to set plane behavior
-	FSelectClickedAction* SetPlaneAction = new FSelectClickedAction();
-	SetPlaneAction->World = this->TargetWorld;
-
-	// Include the original components even though we made them invisible, since we want
-	// to be able to reposition the plane onto the original mesh.
-	for (int Idx = 0; Idx < Targets.Num(); Idx++)
-	{
-		SetPlaneAction->InvisibleComponentsToHitTest.Add(TargetComponentInterface(Idx)->GetOwnerComponent());
-	}
-
-	SetPlaneAction->OnClickedPositionFunc = [this](const FHitResult& Hit)
-	{
-		SetCutPlaneFromWorldPos(Hit.ImpactPoint, Hit.ImpactNormal, false);
-		for (UMeshOpPreviewWithBackgroundCompute* Preview : Previews)
-		{
-			Preview->InvalidateResult();
-		}
-	};
-	SetPointInWorldConnector = SetPlaneAction;
-
-	USingleClickInputBehavior* ClickToSetPlaneBehavior = NewObject<USingleClickInputBehavior>();
-	ClickToSetPlaneBehavior->ModifierCheckFunc = FInputDeviceState::IsCtrlKeyDown;
-	ClickToSetPlaneBehavior->Initialize(SetPointInWorldConnector);
-	AddInputBehavior(ClickToSetPlaneBehavior);
-
-	// create proxy and gizmo (but don't attach yet)
-	UInteractiveGizmoManager* GizmoManager = GetToolManager()->GetPairedGizmoManager();
-	PlaneTransformProxy = NewObject<UTransformProxy>(this);
-	PlaneTransformGizmo = GizmoManager->CreateCustomTransformGizmo(
-		ETransformGizmoSubElements::StandardTranslateRotate, this);
-
 	// initialize our properties
 	BasicProperties = NewObject<UPlaneCutToolProperties>(this, TEXT("Plane Cut Settings"));
 	BasicProperties->RestoreProperties(this);
@@ -193,9 +159,21 @@ void UPlaneCutTool::Setup()
 		TargetComponentInterface(Idx)->GetOwnerActor()->GetActorBounds(false, ComponentOrigin, ComponentExtents);
 		CombinedBounds += FBox::BuildAABB(ComponentOrigin, ComponentExtents);
 	}
-	SetCutPlaneFromWorldPos(CombinedBounds.GetCenter(), FVector::UpVector, true);
-	// hook up callback so further changes trigger recut
-	PlaneTransformProxy->OnTransformChanged.AddUObject(this, &UPlaneCutTool::TransformChanged);
+
+	CutPlaneWorld.Origin = (FVector3d)CombinedBounds.GetCenter();
+	PlaneMechanic = NewObject<UConstructionPlaneMechanic>(this);
+	PlaneMechanic->Setup(this);
+	PlaneMechanic->Initialize(TargetWorld, CutPlaneWorld);
+	PlaneMechanic->OnPlaneChanged.AddLambda([this]() {
+		CutPlaneWorld = PlaneMechanic->Plane;
+		InvalidatePreviews();
+		});
+	for (int Idx = 0; Idx < Targets.Num(); Idx++)
+	{
+		PlaneMechanic->SetPlaneCtrlClickBehaviorTarget->InvisibleComponentsToHitTest.Add(TargetComponentInterface(Idx)->GetOwnerComponent());
+	}
+
+	InvalidatePreviews();
 
 	for (UMeshOpPreviewWithBackgroundCompute* Preview : Previews)
 	{
@@ -204,7 +182,7 @@ void UPlaneCutTool::Setup()
 
 	SetToolDisplayName(LOCTEXT("ToolName", "Plane Cut"));
 	GetToolManager()->DisplayMessage(
-		LOCTEXT("OnStartPlaneCutTool", "Press 'A' or use the Cut button to cut the mesh without leaving the tool.  Press 'S' to flip the plane direction.  When grid snapping is enabled, you can toggle snapping with the shift key."),
+		LOCTEXT("OnStartPlaneCutTool", "Press 'A' or use the Cut button to cut the mesh without leaving the tool.  Press 'S' to flip the plane direction.  Hold 'Shift' to toggle grid snapping."),
 		EToolMessageLevel::UserNotification);
 }
 
@@ -262,10 +240,9 @@ void UPlaneCutTool::DoFlipPlane()
 {
 	GetToolManager()->BeginUndoTransaction(LOCTEXT("FlipPlaneTransactionName", "Flip Plane"));
 
-	FQuat Flipper(CutPlaneOrientation.GetRightVector(), FMathf::Pi);
-	CutPlaneOrientation = Flipper * CutPlaneOrientation;
-	FTransform NewTransform(CutPlaneOrientation, CutPlaneOrigin);
-	PlaneTransformGizmo->SetNewGizmoTransform(NewTransform);
+	FVector3d Origin = CutPlaneWorld.Origin;
+	FVector3d Normal = CutPlaneWorld.GetAxis(2);
+	PlaneMechanic->SetDrawPlaneFromWorldPos(Origin, -Normal, false);
 
 	GetToolManager()->EndUndoTransaction();
 }
@@ -304,6 +281,7 @@ void UPlaneCutTool::DoCut()
 
 void UPlaneCutTool::Shutdown(EToolShutdownType ShutdownType)
 {
+	PlaneMechanic->Shutdown();
 	BasicProperties->SaveProperties(this);
 	AcceptProperties->SaveProperties(this);
 
@@ -322,13 +300,6 @@ void UPlaneCutTool::Shutdown(EToolShutdownType ShutdownType)
 	{
 		GenerateAsset(Results);
 	}
-
-	if (SetPointInWorldConnector != nullptr)
-	{
-		delete SetPointInWorldConnector;
-	}
-	UInteractiveGizmoManager* GizmoManager = GetToolManager()->GetPairedGizmoManager();
-	GizmoManager->DestroyAllGizmosByOwner(this);
 }
 
 void UPlaneCutTool::SetAssetAPI(IAssetGenerationAPI* AssetAPIIn)
@@ -357,12 +328,13 @@ TUniquePtr<FDynamicMeshOperator> UPlaneCutOperatorFactory::MakeNewOperator()
 	}
 	LocalToWorld.SetScale3D(LocalToWorldScale);
 	FTransform WorldToLocal = LocalToWorld.Inverse();
-	FVector LocalOrigin = WorldToLocal.TransformPosition(CutTool->CutPlaneOrigin);
-	FVector WorldNormal = CutTool->CutPlaneOrientation.GetAxisZ();
+
+	FVector LocalOrigin = WorldToLocal.TransformPosition((FVector)CutTool->CutPlaneWorld.Origin);
+	FVector3d WorldNormal = CutTool->CutPlaneWorld.GetAxis(2);
 	UE::Geometry::FTransform3d W2LForNormal(WorldToLocal);
-	FVector LocalNormal = (FVector)W2LForNormal.TransformNormal((FVector3d)WorldNormal);
+	FVector LocalNormal = (FVector)W2LForNormal.TransformNormal(WorldNormal);
 	FVector BackTransformed = LocalToWorld.TransformVector(LocalNormal);
-	float NormalScaleFactor = FVector::DotProduct(BackTransformed, WorldNormal);
+	float NormalScaleFactor = FVector::DotProduct(BackTransformed, (FVector)WorldNormal);
 	if (NormalScaleFactor >= FLT_MIN)
 	{
 		NormalScaleFactor = 1.0 / NormalScaleFactor;
@@ -382,18 +354,15 @@ TUniquePtr<FDynamicMeshOperator> UPlaneCutOperatorFactory::MakeNewOperator()
 
 void UPlaneCutTool::Render(IToolsContextRenderAPI* RenderAPI)
 {
-	FViewCameraState RenderCameraState = RenderAPI->GetCameraState();
-	FPrimitiveDrawInterface* PDI = RenderAPI->GetPrimitiveDrawInterface();
-	FColor GridColor(128, 128, 128, 32);
-	float GridThickness = 0.5f*RenderCameraState.GetPDIScalingFactor();
-	int NumGridLines = 10;
-	
-	FFrame3f DrawFrame(CutPlaneOrigin, CutPlaneOrientation);
-	MeshDebugDraw::DrawSimpleFixedScreenAreaGrid(RenderCameraState, DrawFrame, NumGridLines, 45.0, GridThickness, GridColor, false, PDI, FTransform::Identity);
+	PlaneMechanic->Render(RenderAPI);
 }
 
 void UPlaneCutTool::OnTick(float DeltaTime)
 {
+	PlaneMechanic->Tick(DeltaTime);
+	PlaneMechanic->SetEnableGridSnaping(BasicProperties->bSnapToWorldGrid ^ bSnappingToggle);
+	PlaneMechanic->SetEnableGridRotationSnapping(BasicProperties->bSnapRotationToWorldGrid ^ bSnappingToggle);
+
 	if (PendingAction != EPlaneCutToolActions::NoAction)
 	{
 		if (PendingAction == EPlaneCutToolActions::Cut)
@@ -408,14 +377,6 @@ void UPlaneCutTool::OnTick(float DeltaTime)
 		PendingAction = EPlaneCutToolActions::NoAction;
 	}
 
-	if (PlaneTransformGizmo != nullptr)
-	{
-		PlaneTransformGizmo->bSnapToWorldGrid =
-			BasicProperties->bSnapToWorldGrid && bIgnoreSnappingToggle == false;
-		PlaneTransformGizmo->bSnapToWorldRotGrid =
-			BasicProperties->bSnapRotationToWorldGrid && bIgnoreSnappingToggle == false;
-	}
-
 	for (UMeshOpPreviewWithBackgroundCompute* Preview : Previews)
 	{
 		Preview->Tick(DeltaTime);
@@ -426,13 +387,11 @@ void UPlaneCutTool::OnTick(float DeltaTime)
 #if WITH_EDITOR
 void UPlaneCutTool::PostEditChangeProperty(FPropertyChangedEvent& PropertyChangedEvent)
 {
-	for (UMeshOpPreviewWithBackgroundCompute* Preview : Previews)
-	{
-		Preview->InvalidateResult();
-	}
+	InvalidatePreviews();
 }
 #endif
 
+
 void UPlaneCutTool::OnPropertyModified(UObject* PropertySet, FProperty* Property)
 {
 	if (Property && (Property->GetFName() == GET_MEMBER_NAME_CHECKED(UPlaneCutToolProperties, bShowPreview)))
@@ -447,26 +406,21 @@ void UPlaneCutTool::OnPropertyModified(UObject* PropertySet, FProperty* Property
 		}
 	}
 
-	for (UMeshOpPreviewWithBackgroundCompute* Preview : Previews)
-	{
-		Preview->InvalidateResult();
-	}
+	InvalidatePreviews();
 }
 
 
 void UPlaneCutTool::OnUpdateModifierState(int ModifierID, bool bIsOn)
 {
-	if (ModifierID == IgnoreSnappingModifier)
+	if (ModifierID == SnappingModifier)
 	{
-		bIgnoreSnappingToggle = bIsOn;
+		bSnappingToggle = bIsOn;
 	}
 }
 
 
-void UPlaneCutTool::TransformChanged(UTransformProxy* Proxy, FTransform Transform)
+void UPlaneCutTool::InvalidatePreviews()
 {
-	CutPlaneOrientation = Transform.GetRotation();
-	CutPlaneOrigin = (FVector)Transform.GetTranslation();
 	for (UMeshOpPreviewWithBackgroundCompute* Preview : Previews)
 	{
 		Preview->InvalidateResult();
@@ -474,23 +428,8 @@ void UPlaneCutTool::TransformChanged(UTransformProxy* Proxy, FTransform Transfor
 }
 
 
-void UPlaneCutTool::SetCutPlaneFromWorldPos(const FVector& Position, const FVector& Normal, bool bIsInitializing)
-{
-	CutPlaneOrigin = Position;
 
-	FFrame3f CutPlane((FVector3f)Position, (FVector3f)Normal);
-	CutPlaneOrientation = (FQuat)CutPlane.Rotation;
 
-	PlaneTransformGizmo->SetActiveTarget(PlaneTransformProxy, GetToolManager());
-	if (bIsInitializing)
-	{
-		PlaneTransformGizmo->ReinitializeGizmoTransform(CutPlane.ToFTransform());
-	}
-	else
-	{
-		PlaneTransformGizmo->SetNewGizmoTransform(CutPlane.ToFTransform());
-	}
-}
 
 bool UPlaneCutTool::CanAccept() const
 {
diff --git a/Engine/Plugins/Experimental/MeshModelingToolset/Source/MeshModelingTools/Public/PlaneCutTool.h b/Engine/Plugins/Experimental/MeshModelingToolset/Source/MeshModelingTools/Public/PlaneCutTool.h
--- a/Engine/Plugins/Experimental/MeshModelingToolset/Source/MeshModelingTools/Public/PlaneCutTool.h
+++ b/Engine/Plugins/Experimental/MeshModelingToolset/Source/MeshModelingTools/Public/PlaneCutTool.h
@@ -11,6 +11,7 @@
 #include "DynamicMesh3.h"
 #include "Changes/DynamicMeshChangeTarget.h"
 #include "BaseTools/SingleClickTool.h"
+#include "Mechanics/ConstructionPlaneMechanic.h"
 #include "PlaneCutTool.generated.h"
 
 
@@ -177,14 +178,6 @@ protected:
 	UPROPERTY()
 	UAcceptOutputProperties* AcceptProperties;
 
-	/** Origin of cutting plane */
-	UPROPERTY()
-	FVector CutPlaneOrigin;
-
-	/** Orientation of cutting plane */
-	UPROPERTY()
-	FQuat CutPlaneOrientation;
-
 	UPROPERTY()
 	TArray<UMeshOpPreviewWithBackgroundCompute*> Previews;
 
@@ -212,6 +205,12 @@ protected:
 	UPROPERTY()
 	TArray<UDynamicMeshReplacementChangeTarget*> MeshesToCut;
 
+	UPROPERTY()
+	UConstructionPlaneMechanic* PlaneMechanic = nullptr;
+
+	// Cutting plane
+	UE::Geometry::FFrame3d CutPlaneWorld;
+
 	// UV Scale factor is cached based on the bounding box of the mesh before any cuts are performed, so you don't get inconsistent UVs if you multi-cut the object to smaller sizes
 	TArray<float> MeshUVScaleFactor;
 
@@ -221,28 +220,18 @@ protected:
 	FViewCameraState CameraState;
 
 	// flags used to identify modifier keys/buttons
-	static const int IgnoreSnappingModifier = 1;
-	bool bIgnoreSnappingToggle = false;		// toggled by hotkey (shift)
+	static const int SnappingModifier = 1;
+	bool bSnappingToggle = false;		// toggled by hotkey (shift)
 
 	EPlaneCutToolActions PendingAction = EPlaneCutToolActions::NoAction;
 
-	UPROPERTY()
-	UTransformGizmo* PlaneTransformGizmo;
-
-	UPROPERTY()
-	UTransformProxy* PlaneTransformProxy;
-
-	void TransformChanged(UTransformProxy* Proxy, FTransform Transform);
-	void MeshChanged();
 
 	void DoCut();
 	void DoFlipPlane();
 
 	void SetupPreviews();
 
-	IClickBehaviorTarget* SetPointInWorldConnector = nullptr;
-
-	virtual void SetCutPlaneFromWorldPos(const FVector& Position, const FVector& Normal, bool bIsInitializing);
+	void InvalidatePreviews();
 
 	void GenerateAsset(const TArray<FDynamicMeshOpResult>& Results);
 };
diff --git a/Engine/Plugins/Experimental/MeshModelingToolset/Source/ModelingComponents/Private/Mechanics/ConstructionPlaneMechanic.cpp b/Engine/Plugins/Experimental/MeshModelingToolset/Source/ModelingComponents/Private/Mechanics/ConstructionPlaneMechanic.cpp
--- a/Engine/Plugins/Experimental/MeshModelingToolset/Source/ModelingComponents/Private/Mechanics/ConstructionPlaneMechanic.cpp
+++ b/Engine/Plugins/Experimental/MeshModelingToolset/Source/ModelingComponents/Private/Mechanics/ConstructionPlaneMechanic.cpp
@@ -62,6 +62,11 @@ void UConstructionPlaneMechanic::SetEnableGridSnaping(bool bEnable)
 	bEnableSnapToWorldGrid = bEnable;
 }
 
+void UConstructionPlaneMechanic::SetEnableGridRotationSnapping(bool bEnableRotation)
+{
+	bEnableSnapToWorldRotGrid = bEnableRotation;
+}
+
 void UConstructionPlaneMechanic::UpdateClickPriority(FInputCapturePriority NewPriority)
 {
 	ensure(ClickToSetPlaneBehavior != nullptr);
@@ -118,6 +123,7 @@ void UConstructionPlaneMechanic::Tick(float DeltaTime)
 	if (PlaneTransformGizmo != nullptr)
 	{
 		PlaneTransformGizmo->bSnapToWorldGrid = bEnableSnapToWorldGrid;
+		PlaneTransformGizmo->bSnapToWorldRotGrid = bEnableSnapToWorldRotGrid;
 		PlaneTransformGizmo->SetVisibility(CanUpdatePlaneFunc());
 	}
 }
diff --git a/Engine/Plugins/Experimental/MeshModelingToolset/Source/ModelingComponents/Public/Mechanics/ConstructionPlaneMechanic.h b/Engine/Plugins/Experimental/MeshModelingToolset/Source/ModelingComponents/Public/Mechanics/ConstructionPlaneMechanic.h
--- a/Engine/Plugins/Experimental/MeshModelingToolset/Source/ModelingComponents/Public/Mechanics/ConstructionPlaneMechanic.h
+++ b/Engine/Plugins/Experimental/MeshModelingToolset/Source/ModelingComponents/Public/Mechanics/ConstructionPlaneMechanic.h
@@ -58,6 +58,8 @@ public:
 
 	void SetEnableGridSnaping(bool bEnable);
 
+	void SetEnableGridRotationSnapping(bool bEnableRotationSnapping);
+
 	void UpdateClickPriority(FInputCapturePriority NewPriority);
 
 public:
@@ -76,6 +78,7 @@ public:
 
 protected:
 	bool bEnableSnapToWorldGrid = false;
+	bool bEnableSnapToWorldRotGrid = false;
 
 	UPROPERTY()
 	USingleClickInputBehavior* ClickToSetPlaneBehavior;
