MeshModelingTools: Trivial tweaks- updated PolyPath help text to not refer to inexistent hotkeys. Updated tooltips in group edge/loop insertion tools since retriangulation doesn't lose UV's. Set default edge insertion mode in group edge tool to the same as edge loop tool.

#rnx

diff --git a/Engine/Plugins/Experimental/MeshModelingToolset/Source/MeshModelingTools/Private/DrawPolyPathTool.cpp b/Engine/Plugins/Experimental/MeshModelingToolset/Source/MeshModelingTools/Private/DrawPolyPathTool.cpp
--- a/Engine/Plugins/Experimental/MeshModelingToolset/Source/MeshModelingTools/Private/DrawPolyPathTool.cpp
+++ b/Engine/Plugins/Experimental/MeshModelingToolset/Source/MeshModelingTools/Private/DrawPolyPathTool.cpp
@@ -791,7 +791,7 @@ void UDrawPolyPathTool::EmitNewObject(EDrawPolyPathOutputMode OutputMode)
 void UDrawPolyPathTool::ShowStartupMessage()
 {
 	GetToolManager()->DisplayMessage(
-		LOCTEXT("OnStartDraw", "Use this Tool to draw a path on the Drawing Plane, and then Extrude. Left-click to place points, Double-click (or close) to complete path. Ctrl-click on the scene to reposition the Plane (Shift+Ctrl-click to ignore Normal). [A] toggles Gizmo. Hold Shift to ignore Snapping."),
+		LOCTEXT("OnStartDraw", "Use this Tool to draw a path on the Drawing Plane, and then Extrude. Left-click to place points, Double-click (or close) to complete path. Ctrl-click on the scene to reposition the Plane (Shift+Ctrl-click to ignore Normal)."),
 		EToolMessageLevel::UserNotification);
 }
 
diff --git a/Engine/Plugins/Experimental/MeshModelingToolset/Source/MeshModelingTools/Public/EdgeLoopInsertionTool.h b/Engine/Plugins/Experimental/MeshModelingToolset/Source/MeshModelingTools/Public/EdgeLoopInsertionTool.h
--- a/Engine/Plugins/Experimental/MeshModelingToolset/Source/MeshModelingTools/Public/EdgeLoopInsertionTool.h
+++ b/Engine/Plugins/Experimental/MeshModelingToolset/Source/MeshModelingTools/Public/EdgeLoopInsertionTool.h
@@ -46,7 +46,7 @@ UENUM()
 enum class EEdgeLoopInsertionMode
 {
 	/** Existing groups will be deleted and new triangles will be created for the new groups. 
-	 Keeps topology simple but breaks non-planar groups and loses the UV's. */
+	 Keeps topology simple but breaks non-planar groups. */
 	Retriangulate,
 
 	/** Keeps existing triangles and cuts them to create a new path. May result in fragmented triangles over time.*/
diff --git a/Engine/Plugins/Experimental/MeshModelingToolset/Source/MeshModelingTools/Public/GroupEdgeInsertionTool.h b/Engine/Plugins/Experimental/MeshModelingToolset/Source/MeshModelingTools/Public/GroupEdgeInsertionTool.h
--- a/Engine/Plugins/Experimental/MeshModelingToolset/Source/MeshModelingTools/Public/GroupEdgeInsertionTool.h
+++ b/Engine/Plugins/Experimental/MeshModelingToolset/Source/MeshModelingTools/Public/GroupEdgeInsertionTool.h
@@ -36,7 +36,7 @@ UENUM()
 enum class EGroupEdgeInsertionMode
 {
 	/** Existing groups will be deleted and new triangles will be created for the new groups.
-	 Keeps topology simple but breaks non-planar groups and loses the UV's. */
+	 Keeps topology simple but breaks non-planar groups. */
 	Retriangulate,
 
 	/** Keeps existing triangles and cuts them to create a new path. May result in fragmented triangles over time.*/
@@ -51,7 +51,7 @@ class MESHMODELINGTOOLS_API UGroupEdgeInsertionProperties : public UInteractiveT
 public:
 	/** Determines how group edges are added to the geometry */
 	UPROPERTY(EditAnywhere, Category = InsertEdge)
-	EGroupEdgeInsertionMode InsertionMode = EGroupEdgeInsertionMode::PlaneCut;
+	EGroupEdgeInsertionMode InsertionMode = EGroupEdgeInsertionMode::Retriangulate;
 
 	UPROPERTY(EditAnywhere, Category = InsertEdge)
 	bool bWireframe = true;
