MeshModelingTools: Fix gizmo location calculation in CurveControlPointsMechanic.

#rnx

diff --git a/Engine/Plugins/Experimental/MeshModelingToolset/Source/ModelingComponents/Private/Mechanics/CurveControlPointsMechanic.cpp b/Engine/Plugins/Experimental/MeshModelingToolset/Source/ModelingComponents/Private/Mechanics/CurveControlPointsMechanic.cpp
--- a/Engine/Plugins/Experimental/MeshModelingToolset/Source/ModelingComponents/Private/Mechanics/CurveControlPointsMechanic.cpp
+++ b/Engine/Plugins/Experimental/MeshModelingToolset/Source/ModelingComponents/Private/Mechanics/CurveControlPointsMechanic.cpp
@@ -775,7 +775,7 @@ void UCurveControlPointsMechanic::UpdateGizmoLocation()
 
 	if (SelectedPointIDs.Num() > 0)
 	{
-		FVector3d NewGizmoLocation;
+		FVector3d NewGizmoLocation = FVector3d::Zero();
 		for (int32 PointID : SelectedPointIDs)
 		{
 			NewGizmoLocation += ControlPoints.GetPointCoordinates(PointID);
