change plane cut hotkeys for cut/flip-plane from A,S to T,R, to avoid overlap with WASD + have some mneumonic connection (cuT, Reverse)

#rnx

diff --git a/Engine/Plugins/Experimental/MeshModelingToolset/Source/MeshModelingTools/Private/PlaneCutTool.cpp b/Engine/Plugins/Experimental/MeshModelingToolset/Source/MeshModelingTools/Private/PlaneCutTool.cpp
--- a/Engine/Plugins/Experimental/MeshModelingToolset/Source/MeshModelingTools/Private/PlaneCutTool.cpp
+++ b/Engine/Plugins/Experimental/MeshModelingToolset/Source/MeshModelingTools/Private/PlaneCutTool.cpp
@@ -182,7 +182,7 @@ void UPlaneCutTool::Setup()
 
 	SetToolDisplayName(LOCTEXT("ToolName", "Plane Cut"));
 	GetToolManager()->DisplayMessage(
-		LOCTEXT("OnStartPlaneCutTool", "Press 'A' or use the Cut button to cut the mesh without leaving the tool.  Press 'S' to flip the plane direction.  Hold 'Shift' to toggle grid snapping."),
+		LOCTEXT("OnStartPlaneCutTool", "Press 'T' or use the Cut button to cut the mesh without leaving the tool.  Press 'R' to flip the plane direction.  Hold 'Shift' to toggle grid snapping."),
 		EToolMessageLevel::UserNotification);
 }
 
@@ -195,13 +195,13 @@ void UPlaneCutTool::RegisterActions(FInteractiveToolActionSet& ActionSet)
 		TEXT("Do Plane Cut"), 
 		LOCTEXT("DoPlaneCut", "Do Plane Cut"),
 		LOCTEXT("DoPlaneCutTooltip", "Cut the mesh with the current cutting plane, without exiting the tool"),
-		EModifierKey::None, EKeys::A,
+		EModifierKey::None, EKeys::T,
 		[this]() { this->Cut(); } );
 	ActionSet.RegisterAction(this, (int32)EStandardToolActions::BaseClientDefinedActionID + 102,
 		TEXT("Flip Cutting Plane"),
 		LOCTEXT("FlipCutPlane", "Flip Cutting Plane"),
 		LOCTEXT("FlipCutPlaneTooltip", "Flip the cutting plane"),
-		EModifierKey::None, EKeys::S,
+		EModifierKey::None, EKeys::R,
 		[this]() { this->FlipPlane(); });
 }
 
diff --git a/Engine/Plugins/Experimental/MeshModelingToolset/Source/MeshModelingTools/Public/PlaneCutTool.h b/Engine/Plugins/Experimental/MeshModelingToolset/Source/MeshModelingTools/Public/PlaneCutTool.h
--- a/Engine/Plugins/Experimental/MeshModelingToolset/Source/MeshModelingTools/Public/PlaneCutTool.h
+++ b/Engine/Plugins/Experimental/MeshModelingToolset/Source/MeshModelingTools/Public/PlaneCutTool.h
@@ -186,14 +186,14 @@ protected:
 	/// Note these set a flag to call the action later (in OnTick)
 	/// Otherwise, the actions in undo history will end up being generically named by an outer UI handler transaction
 
-	/** Cut with the current plane without exiting the tool */
+	/** Cut with the current plane without exiting the tool (Hotkey: T)*/
 	UFUNCTION(CallInEditor, Category = Actions, meta = (DisplayName = "Cut"))
 	void Cut()
 	{
 		PendingAction = EPlaneCutToolActions::Cut;
 	}
 
-	/** Flip the cutting plane */
+	/** Flip the cutting plane (Hotkey: R) */
 	UFUNCTION(CallInEditor, Category = Actions, meta = (DisplayName = "Flip Plane"))
 	void FlipPlane()
 	{
