MeshModelingTools: Fix height subdivision property min when adding box shape primitive.

#rnx

diff --git a/Engine/Plugins/Experimental/MeshModelingToolset/Source/MeshModelingTools/Public/AddPrimitiveTool.h b/Engine/Plugins/Experimental/MeshModelingToolset/Source/MeshModelingTools/Public/AddPrimitiveTool.h
--- a/Engine/Plugins/Experimental/MeshModelingToolset/Source/MeshModelingTools/Public/AddPrimitiveTool.h
+++ b/Engine/Plugins/Experimental/MeshModelingToolset/Source/MeshModelingTools/Public/AddPrimitiveTool.h
@@ -154,7 +154,7 @@ public:
 	int DepthSubdivisions = 1;
 
 	/** Number of Subdivisions Along the Height */
-	UPROPERTY(EditAnywhere, Category = "ShapeSettings", meta = (DisplayName = "Height Subdivisions", UIMin = "0", UIMax = "100", ClampMin = "0", ClampMax = "500", ProceduralShapeSetting))
+	UPROPERTY(EditAnywhere, Category = "ShapeSettings", meta = (DisplayName = "Height Subdivisions", UIMin = "1", UIMax = "100", ClampMin = "1", ClampMax = "500", ProceduralShapeSetting))
 	int HeightSubdivisions = 1;
 };
 
