MeshModelingTools : MeshAttributePaintTool - fix bug that resulted from change to color overlay.
#rnx

diff --git a/Engine/Plugins/Experimental/MeshModelingToolset/Source/MeshModelingTools/Private/MeshAttributePaintTool.cpp b/Engine/Plugins/Experimental/MeshModelingToolset/Source/MeshModelingTools/Private/MeshAttributePaintTool.cpp
--- a/Engine/Plugins/Experimental/MeshModelingToolset/Source/MeshModelingTools/Private/MeshAttributePaintTool.cpp
+++ b/Engine/Plugins/Experimental/MeshModelingToolset/Source/MeshModelingTools/Private/MeshAttributePaintTool.cpp
@@ -175,8 +175,11 @@ void UMeshAttributePaintTool::Setup()
 	// enable vtx colors on preview mesh
 	PreviewMesh->EditMesh([](FDynamicMesh3& Mesh)
 	{
-		Mesh.DiscardVertexColors();
-		Mesh.EnableVertexColors(FVector3f::Zero());
+		Mesh.EnableAttributes();
+		Mesh.Attributes()->DisablePrimaryColors();
+		Mesh.Attributes()->EnablePrimaryColors();
+		// Create an overlay that has no split elements, init with zero value.
+		Mesh.Attributes()->PrimaryColors()->CreateFromPredicate([](int ParentVID, int TriIDA, int TriIDB){return true;}, 0.f);
 	});
 
 	// build octree
@@ -419,11 +422,13 @@ void UMeshAttributePaintTool::UpdateVisibleAttribute()
 		// update mesh with new value colors
 		PreviewMesh->EditMesh([&](FDynamicMesh3& Mesh)
 		{
-			for (int32 vid : Mesh.VertexIndicesItr())
+			FDynamicMeshColorOverlay* ColorOverlay = Mesh.Attributes()->PrimaryColors();
+			for (int32 elid : ColorOverlay->ElementIndicesItr())
 			{
-				float Value = AttribData.CurrentValues[vid];
-				FVector3f Color = ColorMapper->ToColor<FVector3f>(Value);
-				Mesh.SetVertexColor(vid, Color);
+				const int32 vid = ColorOverlay->GetParentVertex(elid);
+				const float Value = AttribData.CurrentValues[vid];
+				const FVector4f Color4f(ColorMapper->ToColor(Value));
+				ColorOverlay->SetElement(elid,  Color4f);
 			}
 		});
 
@@ -488,13 +493,20 @@ void UMeshAttributePaintTool::ApplyStamp(const FBrushStampData& Stamp)
 	// update values and colors
 	PreviewMesh->DeferredEditMesh([&](FDynamicMesh3& Mesh)
 	{
+		TArray<int> ElIDs;
+		FDynamicMeshColorOverlay* ColorOverlay = Mesh.Attributes()->PrimaryColors();
 		int32 NumVertices = ActionData.ROIVertices.Num();
 		for (int32 k = 0; k < NumVertices; ++k)
 		{
 			int32 vid = ActionData.ROIVertices[k];
 			AttribData.CurrentValues[vid] = ActionData.ROIAfter[k];
-			FVector3f NewColor = ColorMapper->ToColor<FVector3f>(ActionData.ROIAfter[k]);
-			Mesh.SetVertexColor(vid, NewColor);
+			FVector4f NewColor(ColorMapper->ToColor(ActionData.ROIAfter[k]));
+			ColorOverlay->GetVertexElements(vid, ElIDs);
+			for (int elid : ElIDs)
+			{
+				ColorOverlay->SetElement(elid, NewColor);
+			}
+			ElIDs.Reset();
 		}
 	}, false);
 	PreviewMesh->NotifyDeferredEditCompleted(UPreviewMesh::ERenderUpdateMode::FastUpdate, EMeshRenderAttributeFlags::VertexColors, false);
@@ -702,10 +714,17 @@ void UMeshAttributePaintTool::ExternalUpdateValues(int32 AttribIndex, const TArr
 	{
 		PreviewMesh->EditMesh([&](FDynamicMesh3& Mesh)
 		{
+			TArray<int> ElIDs;
+			FDynamicMeshColorOverlay* ColorOverlay = Mesh.Attributes()->PrimaryColors();
 			for (int32 vid : VertexIndices)
 			{
-				FVector3f NewColor = ColorMapper->ToColor<FVector3f>(AttribData.CurrentValues[vid]);
-				Mesh.SetVertexColor(vid, NewColor);
+				FVector4f NewColor(ColorMapper->ToColor(AttribData.CurrentValues[vid]));
+				ColorOverlay->GetVertexElements(vid, ElIDs);
+				for (int elid : ElIDs)
+				{
+					ColorOverlay->SetElement(elid, NewColor);
+				}
+				ElIDs.Reset();
 			}
 		});
 	}
