MeshModelingTools: Making some PolyEdit properties disable themselves when they would not do anything, and making them work more predictably.

#rnx

diff --git a/Engine/Plugins/Experimental/MeshModelingToolset/Source/MeshModelingTools/Private/EditMeshPolygonsTool.cpp b/Engine/Plugins/Experimental/MeshModelingToolset/Source/MeshModelingTools/Private/EditMeshPolygonsTool.cpp
--- a/Engine/Plugins/Experimental/MeshModelingToolset/Source/MeshModelingTools/Private/EditMeshPolygonsTool.cpp
+++ b/Engine/Plugins/Experimental/MeshModelingToolset/Source/MeshModelingTools/Private/EditMeshPolygonsTool.cpp
@@ -218,7 +218,7 @@ void UEditMeshPolygonsTool::Setup()
 	CommonProps->WatchProperty(CommonProps->bLockRotation,
 								  [this](bool)
 								  {
-									  LockedTransfomerFrame = LastTransformerFrame; UpdateMultiTransformerFrame();
+									  LockedTransfomerFrame = LastTransformerFrame;
 								  });
 	// We are going to SilentUpdate here because otherwise the Watches above will immediately fire (why??)
 	// and cause UpdateMultiTransformerFrame() to be called for each, emitting two spurious Transform changes. 
@@ -268,6 +268,13 @@ void UEditMeshPolygonsTool::Setup()
 	MultiTransformer->OnTransformStarted.AddUObject(this, &UEditMeshPolygonsTool::OnMultiTransformerTransformBegin);
 	MultiTransformer->OnTransformUpdated.AddUObject(this, &UEditMeshPolygonsTool::OnMultiTransformerTransformUpdate);
 	MultiTransformer->OnTransformCompleted.AddUObject(this, &UEditMeshPolygonsTool::OnMultiTransformerTransformEnd);
+	MultiTransformer->OnEndPivotEdit.AddWeakLambda(this, [this]() {
+		LastTransformerFrame = MultiTransformer->GetCurrentGizmoFrame();
+		if (CommonProps->bLockRotation)
+		{
+			LockedTransfomerFrame = LastTransformerFrame;
+		}
+	});
 	MultiTransformer->SetSnapToWorldGridSourceFunc([this]() {
 		return CommonProps->bSnapToWorldGrid
 			&& GetToolManager()->GetContextQueriesAPI()->GetCurrentCoordinateSystem() == EToolContextCoordinateSystem::World;
@@ -709,7 +716,18 @@ void UEditMeshPolygonsTool::OnMultiTransformerTransformEnd()
 	bSpatialDirty = true;
 	SelectionMechanic->NotifyMeshChanged(false);
 
-	MultiTransformer->ResetScale();
+	if (CommonProps->bLockRotation)
+	{
+		FFrame3d SetFrame = MultiTransformer->GetCurrentGizmoFrame();
+		SetFrame.Rotation = LockedTransfomerFrame.Rotation;
+		MultiTransformer->InitializeGizmoPositionFromWorldFrame(SetFrame, true);
+	}
+	else
+	{
+		MultiTransformer->ResetScale();
+	}
+
+	LastTransformerFrame = MultiTransformer->GetCurrentGizmoFrame();
 
 	// close change record
 	EndChange();
@@ -859,6 +877,14 @@ void UEditMeshPolygonsTool::OnTick(float DeltaTime)
 {
 	MultiTransformer->Tick(DeltaTime);
 
+	bool bLocalCoordSystem = GetToolManager()->GetPairedGizmoManager()
+		->GetContextQueriesAPI()->GetCurrentCoordinateSystem() == EToolContextCoordinateSystem::Local;
+	if (CommonProps->bLocalCoordSystem != bLocalCoordSystem)
+	{
+		CommonProps->bLocalCoordSystem = bLocalCoordSystem;
+		NotifyOfPropertyChangeByTool(CommonProps);
+	}
+
 	if (bGizmoUpdatePending)
 	{
 		ComputeUpdate_Gizmo();
diff --git a/Engine/Plugins/Experimental/MeshModelingToolset/Source/MeshModelingTools/Public/EditMeshPolygonsTool.h b/Engine/Plugins/Experimental/MeshModelingToolset/Source/MeshModelingTools/Public/EditMeshPolygonsTool.h
--- a/Engine/Plugins/Experimental/MeshModelingToolset/Source/MeshModelingTools/Public/EditMeshPolygonsTool.h
+++ b/Engine/Plugins/Experimental/MeshModelingToolset/Source/MeshModelingTools/Public/EditMeshPolygonsTool.h
@@ -63,14 +63,23 @@ public:
 	UPROPERTY(EditAnywhere, Category = Options)
 	bool bShowWireframe = false;
 
-	UPROPERTY(EditAnywhere, Category = Gizmo)
+	/** Determines whether, on selection changes, the gizmo's rotation is taken from the object transform, or from the geometry
+	 elements selected. Only relevant with a local coordinate system and when rotation is not locked. */
+	UPROPERTY(EditAnywhere, Category = Gizmo, meta = (HideEditConditionToggle, EditCondition = "bLocalCoordSystem && !bLockRotation"))
 	ELocalFrameMode LocalFrameMode = ELocalFrameMode::FromGeometry;
 
-	UPROPERTY(EditAnywhere, Category = Gizmo)
+	/** When true, keeps rotation of gizmo constant through selection changes and manipulations 
+	 (but not middle-click repositions). Only active with a local coordinate system.*/
+	UPROPERTY(EditAnywhere, Category = Gizmo, meta = (HideEditConditionToggle, EditCondition = "bLocalCoordSystem"))
 	bool bLockRotation = false;
 
-	UPROPERTY(EditAnywhere, Category = Gizmo)
+	/** Snaps to grid in translation. Only active with a global coordinate system. */
+	UPROPERTY(EditAnywhere, Category = Gizmo, meta = (HideEditConditionToggle, EditCondition = "!bLocalCoordSystem"))
 	bool bSnapToWorldGrid = false;
+
+	/** This gets updated internally so that properties can respond to whether the coordinate system is set to local or global */
+	UPROPERTY()
+	bool bLocalCoordSystem = true;
 };
 
 
diff --git a/Engine/Plugins/Experimental/MeshModelingToolset/Source/ModelingComponents/Private/Transforms/MultiTransformer.cpp b/Engine/Plugins/Experimental/MeshModelingToolset/Source/ModelingComponents/Private/Transforms/MultiTransformer.cpp
--- a/Engine/Plugins/Experimental/MeshModelingToolset/Source/ModelingComponents/Private/Transforms/MultiTransformer.cpp
+++ b/Engine/Plugins/Experimental/MeshModelingToolset/Source/ModelingComponents/Private/Transforms/MultiTransformer.cpp
@@ -30,6 +30,9 @@ void UMultiTransformer::Setup(UInteractiveGizmoManager* GizmoManagerIn, IToolCon
 		ActiveGizmoFrame = FFrame3d(Transform);
 		ActiveGizmoScale = FVector3d(Transform.GetScale3D());
 	});
+	TransformProxy->OnEndPivotEdit.AddWeakLambda(this, [this](UTransformProxy* Proxy) {
+		OnEndPivotEdit.Broadcast();
+	});
 }
 
 
diff --git a/Engine/Plugins/Experimental/MeshModelingToolset/Source/ModelingComponents/Public/Transforms/MultiTransformer.h b/Engine/Plugins/Experimental/MeshModelingToolset/Source/ModelingComponents/Public/Transforms/MultiTransformer.h
--- a/Engine/Plugins/Experimental/MeshModelingToolset/Source/ModelingComponents/Public/Transforms/MultiTransformer.h
+++ b/Engine/Plugins/Experimental/MeshModelingToolset/Source/ModelingComponents/Public/Transforms/MultiTransformer.h
@@ -25,6 +25,7 @@ enum class EMultiTransformerMode
 
 /**
  * UMultiTransformer abstracts both a default TRS Gizmo, and the "Quick" translate/rotate Gizmos.
+ * The "Quick" part is not yet implemented, and we might end up phasing out this class.
  */
 UCLASS()
 class MODELINGCOMPONENTS_API UMultiTransformer : public UObject
@@ -81,6 +82,12 @@ public:
 	/** This delegate is fired when the drag is completed */
 	FMultiTransformerEvent OnTransformCompleted;
 
+	/** This delegate is fired when a repositioning drag is completed */
+	FMultiTransformerEvent OnEndPivotEdit;
+
+	// Note: we could have more pivot change delegates, but we don't yet need them,
+	// and we might phase out the UMultiTransformer.
+
 protected:
 	UPROPERTY()
 	UInteractiveGizmoManager* GizmoManager;
