MeshModelingTools: Trivial PolyExt edit condition fix (snap to lengths depends on snap to angles).

#rnx

diff --git a/Engine/Plugins/Experimental/MeshModelingToolset/Source/MeshModelingTools/Public/DrawPolygonTool.h b/Engine/Plugins/Experimental/MeshModelingToolset/Source/MeshModelingTools/Public/DrawPolygonTool.h
--- a/Engine/Plugins/Experimental/MeshModelingToolset/Source/MeshModelingTools/Public/DrawPolygonTool.h
+++ b/Engine/Plugins/Experimental/MeshModelingToolset/Source/MeshModelingTools/Public/DrawPolygonTool.h
@@ -143,7 +143,8 @@ public:
 	UPROPERTY(EditAnywhere, NonTransactional, Category = Snapping, Meta = (EditCondition = "bEnableSnapping", EditConditionHides))
 	bool bSnapToAngles = true;
 
-	UPROPERTY(EditAnywhere, NonTransactional, Category = Snapping, Meta = (EditCondition = "bEnableSnapping", EditConditionHides))
+	/** When snapping to an angle, also prefer to snap to an existing segment length */
+	UPROPERTY(EditAnywhere, NonTransactional, Category = Snapping, Meta = (EditCondition = "bEnableSnapping && bSnapToAngles", EditConditionHides))
 	bool bSnapToLengths = true;
 
 	UPROPERTY(VisibleAnywhere, NonTransactional, Category = Snapping, meta = (TransientToolProperty))
