ModelingComponents: add support for capturing scene renders and doing photo-based texturing
- new FWorldRenderCapture class can capture various render buffers of a set of Actors from a given 3D camera
- new TSpatialPhoto represents a 3D-oriented 2D image (ie a "photo")
- new TSpatialPhotoSet stores a set of photos and supports quering them for an attribute value "seen by" the photo set at an arbitrary 3D position (eg basically photo-reprojection onto the point)
- new FSceneCapturePhotoSet calculates TSpatialPhotoSet for a set of render attributes using FWorldRenderCapture, and then supports querying to effectively texture a 3D point using the Photo Set
- FTexture2DBuilder now knows about various other texture types, clarified SRGB conversion parameters
#rnx

diff --git a/Engine/Plugins/Experimental/GeometryProcessing/Source/GeometricObjects/Public/Image/SpatialPhotoSet.h b/Engine/Plugins/Experimental/GeometryProcessing/Source/GeometricObjects/Public/Image/SpatialPhotoSet.h
new file mode 100644
--- /dev/null
+++ b/Engine/Plugins/Experimental/GeometryProcessing/Source/GeometricObjects/Public/Image/SpatialPhotoSet.h
@@ -0,0 +1,154 @@
+// Copyright Epic Games, Inc. All Rights Reserved.
+
+#pragma once
+
+#include "VectorTypes.h"
+#include "FrameTypes.h"
+#include "Image/ImageBuilder.h"
+#include "Image/ImageDimensions.h"
+
+
+namespace UE
+{
+namespace Geometry
+{
+
+/**
+ * TSpatialPhoto represents a 2D image located in 3D space, ie the image plus camera parameters, 
+ * which is essentially a "Photograph" of some 3D scene (hence the name)
+ */
+template<typename PixelType>
+struct TSpatialPhoto
+{
+	/** Coordinate system of the view camera - X() is forward, Z() is up */
+	FFrame3d Frame;
+
+	/** Near-plane distance for the camera, image pixels lie on this plane */
+	double NearPlaneDist = 1.0;
+	/** Horizontal Field-of-View of the camera in degrees (full FOV, so generally calculations will use half this value) */
+	double HorzFOVDegrees = 90.0;
+	/** Pixel dimensions of the photo image */
+	FImageDimensions Dimensions;
+	/** Pixels of the image */
+	TImageBuilder<PixelType> Image;
+};
+typedef TSpatialPhoto<FVector4f> FSpatialPhoto4f;
+
+
+
+/**
+ * TSpatialPhotoSet is a set of TSpatialPhotos. 
+ * The ComputeSample() function can be used to determine the value "seen"
+ * by the photo set at a given 3D position/normal, if possible.
+ */
+template<typename PixelType, typename RealType>
+class TSpatialPhotoSet
+{
+public:
+	
+	/** Add a photo to the photo set via move operation */
+	void Add(TSpatialPhoto<PixelType>&& Photo)
+	{
+		TSharedPtr<TSpatialPhoto<PixelType>, ESPMode::ThreadSafe> NewPhoto = MakeShared<TSpatialPhoto<PixelType>, ESPMode::ThreadSafe>(MoveTemp(Photo));
+		Photos.Add(NewPhoto);
+	}
+
+	/** @return the number of photos in the photo set */
+	int32 Num() const { return Photos.Num(); }
+
+	/** @return the photo at the given index */
+	const TSpatialPhoto<PixelType>& Get(int32 Index) const { return *Photos[Index]; }
+
+	/**
+	 * Estimate a pixel value at the given 3D Position/Normal using the PhotoSet. 
+	 * This is effectively a reprojection process, that tries to find the "best"
+	 * pixel value in the photo set that projects onto the given Position/Normal.
+	 * 
+	 * A position may be visible from multiple photos, in this case the dot-product
+	 * between the view vector and normal is used to decide which photo pixel to use.
+	 * 
+	 * VisibilityFunction is used to determine if a 3D point is visible from the given photo point.
+	 * Generally the caller would implement some kind of raycast to do this.
+	 * 
+	 * @returns the best valid sample if found, or DefaultValue if no suitable sample is available
+	 */
+	PixelType ComputeSample(
+		const FVector3d& Position, 
+		const FVector3d& Normal, 
+		TFunctionRef<bool(const FVector3d&,const FVector3d&)> VisibilityFunction,
+		const PixelType& DefaultValue
+	) const;
+
+protected:
+	TArray<TSharedPtr<TSpatialPhoto<PixelType>, ESPMode::ThreadSafe>> Photos;
+};
+typedef TSpatialPhotoSet<FVector4f, float> FSpatialPhotoSet4f;
+
+
+
+
+template<typename PixelType, typename RealType>
+PixelType TSpatialPhotoSet<PixelType, RealType>::ComputeSample(
+	const FVector3d& Position, 
+	const FVector3d& Normal, 
+	TFunctionRef<bool(const FVector3d&, const FVector3d&)> VisibilityFunction,
+	const PixelType& DefaultValue ) const
+{
+	double DotTolerance = -0.1;		// dot should be negative for normal pointing towards photo
+
+	PixelType BestSample = DefaultValue;
+	double MinDot = 1.0;
+
+	int32 NumPhotos = Num();
+	for (int32 pi = 0; pi < NumPhotos; ++pi)
+	{
+		const TSpatialPhoto<PixelType>& Photo = *Photos[pi];
+		check(Photo.Dimensions.IsSquare());
+
+		FVector3d ViewDirection = Photo.Frame.X();
+		double ViewDot = ViewDirection.Dot(Normal);
+		if (ViewDot > DotTolerance || ViewDot > MinDot)
+		{
+			continue;
+		}
+
+		FFrame3d ViewPlane = Photo.Frame;
+		ViewPlane.Origin += Photo.NearPlaneDist * ViewDirection;
+
+		double ViewPlaneWidthWorld = Photo.NearPlaneDist * FMathd::Tan(Photo.HorzFOVDegrees * 0.5 * FMathd::DegToRad);
+		double ViewPlaneHeightWorld = ViewPlaneWidthWorld;
+
+		FVector3d RayOrigin = Photo.Frame.Origin;
+		FVector3d RayDir = Normalized(Position - RayOrigin);
+		FVector3d HitPoint;
+		bool bHit = ViewPlane.RayPlaneIntersection(RayOrigin, RayDir, 0, HitPoint);
+		if (bHit)
+		{
+			bool bVisible = VisibilityFunction(Position, HitPoint);
+			if ( bVisible )
+			{
+				double PlaneX = (HitPoint - ViewPlane.Origin).Dot(ViewPlane.Y());
+				double PlaneY = (HitPoint - ViewPlane.Origin).Dot(ViewPlane.Z());
+
+				//FVector2d PlanePos = ViewPlane.ToPlaneUV(HitPoint, 0);
+				double u = PlaneX / ViewPlaneWidthWorld;
+				double v = -(PlaneY / ViewPlaneHeightWorld);
+				if (FMathd::Abs(u) < 1 && FMathd::Abs(v) < 1)
+				{
+					double x = (u/2.0 + 0.5) * (double)Photo.Dimensions.GetWidth();
+					double y = (v/2.0 + 0.5) * (double)Photo.Dimensions.GetHeight();
+					PixelType Sample = Photo.Image.template BilinearSample<RealType>(FVector2d(x, y), DefaultValue);
+
+					MinDot = ViewDot;
+					BestSample = Sample;
+				}
+			}
+		}
+	}
+
+	return BestSample;
+}
+
+
+} // end namespace UE::Geometry
+} // end namespace UE
\ No newline at end of file
diff --git a/Engine/Plugins/Experimental/MeshModelingToolset/Source/ModelingComponents/Private/Scene/SceneCapturePhotoSet.cpp b/Engine/Plugins/Experimental/MeshModelingToolset/Source/ModelingComponents/Private/Scene/SceneCapturePhotoSet.cpp
new file mode 100644
--- /dev/null
+++ b/Engine/Plugins/Experimental/MeshModelingToolset/Source/ModelingComponents/Private/Scene/SceneCapturePhotoSet.cpp
@@ -0,0 +1,245 @@
+// Copyright Epic Games, Inc. All Rights Reserved.
+
+#include "Scene/SceneCapturePhotoSet.h"
+
+using namespace UE::Geometry;
+
+void FSceneCapturePhotoSet::SetCaptureSceneActors(UWorld* World, const TArray<AActor*>& Actors)
+{
+	this->TargetWorld = World;
+	this->VisibleActors = Actors;
+}
+
+
+
+void FSceneCapturePhotoSet::SetCaptureTypeEnabled(ERenderCaptureType CaptureType, bool bEnabled)
+{
+	switch (CaptureType)
+	{
+	case ERenderCaptureType::BaseColor:
+			bEnableBaseColor = bEnabled;
+			break;
+		case ERenderCaptureType::WorldNormal:
+			bEnableWorldNormal = bEnabled;
+			break;
+		case ERenderCaptureType::Roughness:
+			bEnableRoughness = bEnabled;
+			break;
+		case ERenderCaptureType::Metallic:
+			bEnableMetallic = bEnabled;
+			break;
+		case ERenderCaptureType::Specular:
+			bEnableSpecular = bEnabled;
+			break;
+		case ERenderCaptureType::Emissive:
+			bEnableEmissive = bEnabled;
+			break;
+		default:
+			check(false);
+	}
+}
+
+
+
+void FSceneCapturePhotoSet::AddStandardExteriorCapturesFromBoundingBox(
+	FImageDimensions PhotoDimensions,
+	double HorizontalFOVDegrees,
+	double NearPlaneDist,
+	bool bFaces,
+	bool bUpperCorners,
+	bool bLowerCorners)
+{
+	TArray<FVector3d> Directions;
+
+	if (bFaces)
+	{
+		Directions.Add(FVector3d::UnitX());
+		Directions.Add(-FVector3d::UnitX());
+		Directions.Add(FVector3d::UnitY());
+		Directions.Add(-FVector3d::UnitY());
+		Directions.Add(FVector3d::UnitZ());
+		Directions.Add(-FVector3d::UnitZ());
+	}
+	if (bUpperCorners)
+	{
+		Directions.Add(Normalized(FVector3d(1, 1, -1)));
+		Directions.Add(Normalized(FVector3d(-1, 1, -1)));
+		Directions.Add(Normalized(FVector3d(1, -1, -1)));
+		Directions.Add(Normalized(FVector3d(-1, -1, -1)));
+	}
+	if (bLowerCorners)
+	{
+		Directions.Add(Normalized(FVector3d(1, 1, 1)));
+		Directions.Add(Normalized(FVector3d(-1, 1, 1)));
+		Directions.Add(Normalized(FVector3d(1, -1, 1)));
+		Directions.Add(Normalized(FVector3d(-1, -1, 1)));
+	}
+	AddExteriorCaptures(PhotoDimensions, HorizontalFOVDegrees, NearPlaneDist, Directions);
+}
+
+
+
+void FSceneCapturePhotoSet::AddExteriorCaptures(
+	FImageDimensions PhotoDimensions,
+	double HorizontalFOVDegrees,
+	double NearPlaneDist,
+	const TArray<FVector3d>& Directions)
+{
+	check(this->TargetWorld != nullptr);
+
+	FWorldRenderCapture RenderCapture;
+	RenderCapture.SetWorld(TargetWorld);
+	RenderCapture.SetVisibleActors(VisibleActors);
+	RenderCapture.SetDimensions(PhotoDimensions);
+
+	// this tells us origin and radius - could be view-dependent...
+	FSphere RenderSphere = RenderCapture.ComputeContainingRenderSphere(HorizontalFOVDegrees);
+
+	int32 NumDirections = Directions.Num();
+	for (int32 di = 0; di < NumDirections; ++di)
+	{
+		FVector3d ViewDirection = Directions[di];
+		ViewDirection.Normalize();
+
+		FFrame3d ViewFrame;
+		ViewFrame.AlignAxis(0, ViewDirection);
+		ViewFrame.ConstrainedAlignAxis(2, FVector3d::UnitZ(), ViewFrame.X());
+		ViewFrame.Origin = (FVector3d)RenderSphere.Center;
+		ViewFrame.Origin -= (double)RenderSphere.W * ViewFrame.X();
+
+		FSpatialPhoto4f BasePhoto;
+		BasePhoto.Frame = ViewFrame;
+		BasePhoto.NearPlaneDist = NearPlaneDist;
+		BasePhoto.HorzFOVDegrees = HorizontalFOVDegrees;
+		BasePhoto.Dimensions = PhotoDimensions;
+
+		auto CaptureImageTypeFunc = [&BasePhoto, &RenderCapture](ERenderCaptureType CaptureType, FSpatialPhotoSet4f& PhotoSet)
+		{
+			FSpatialPhoto4f NewPhoto = BasePhoto;
+			RenderCapture.CaptureFromPosition(CaptureType, NewPhoto.Frame, NewPhoto.HorzFOVDegrees, NewPhoto.NearPlaneDist, NewPhoto.Image);
+			PhotoSet.Add(MoveTemp(NewPhoto));
+		};
+
+		if (bEnableBaseColor)
+		{
+			CaptureImageTypeFunc(ERenderCaptureType::BaseColor, BaseColorPhotoSet);
+		}
+		if (bEnableRoughness)
+		{
+			CaptureImageTypeFunc(ERenderCaptureType::Roughness, RoughnessPhotoSet);
+		}
+		if (bEnableSpecular)
+		{
+			CaptureImageTypeFunc(ERenderCaptureType::Specular, SpecularPhotoSet);
+		}
+		if (bEnableMetallic)
+		{
+			CaptureImageTypeFunc(ERenderCaptureType::Metallic, MetallicPhotoSet);
+		}
+		if (bEnableWorldNormal)
+		{
+			CaptureImageTypeFunc(ERenderCaptureType::WorldNormal, WorldNormalPhotoSet);
+		}
+		if (bEnableEmissive)
+		{
+			CaptureImageTypeFunc(ERenderCaptureType::Emissive, EmissivePhotoSet);
+		}
+	}
+
+}
+
+
+
+void FSceneCapturePhotoSet::OptimizePhotoSets()
+{
+	// todo:
+	//  1) crop photos to regions with actual pixels
+	//  2) pack into fewer photos  (eg pack spec/rough/metallic)
+	//  3) RLE encoding or other compression
+}
+
+
+
+FSceneCapturePhotoSet::FSceneSample::FSceneSample()
+{
+	HaveValues = FRenderCaptureTypeFlags::None();
+	BaseColor = FVector4f(0, 0, 0, 1);
+	Roughness = FVector4f(0, 0, 0, 1);
+	Specular = FVector4f(0, 0, 0, 1);
+	Metallic = FVector4f(0, 0, 0, 1);
+	Emissive = FVector4f(0, 0, 0, 1);
+	WorldNormal = FVector4f(0, 0, 1, 1);
+}
+
+FVector4f FSceneCapturePhotoSet::FSceneSample::GetValue(ERenderCaptureType CaptureType) const
+{
+	switch (CaptureType)
+	{
+	case ERenderCaptureType::BaseColor:
+		return BaseColor;
+	case ERenderCaptureType::WorldNormal:
+		return WorldNormal;
+	case ERenderCaptureType::Roughness:
+		return Roughness;
+	case ERenderCaptureType::Metallic:
+		return Metallic;
+	case ERenderCaptureType::Specular:
+		return Specular;
+	case ERenderCaptureType::Emissive:
+		return Emissive;
+	default:
+		check(false);
+	}
+	return FVector4f::Zero();
+}
+
+
+bool FSceneCapturePhotoSet::ComputeSample(
+	const FRenderCaptureTypeFlags& SampleChannels,
+	const FVector3d& Position,
+	const FVector3d& Normal,
+	TFunctionRef<bool(const FVector3d&, const FVector3d&)> VisibilityFunction,
+	FSceneSample& DefaultsInResultsOut) const
+{
+	// This could be much more efficient if (eg) we knew that all the photo sets have
+	// the same captures, then the query only has to be done once and can be used to sample each specific photo
+
+	if (SampleChannels.bBaseColor)
+	{
+		DefaultsInResultsOut.BaseColor =
+			BaseColorPhotoSet.ComputeSample(Position, Normal, VisibilityFunction, DefaultsInResultsOut.BaseColor);
+		DefaultsInResultsOut.HaveValues.bBaseColor = true;
+	}
+	if (SampleChannels.bRoughness)
+	{
+		DefaultsInResultsOut.Roughness =
+			RoughnessPhotoSet.ComputeSample(Position, Normal, VisibilityFunction, DefaultsInResultsOut.Roughness);
+		DefaultsInResultsOut.HaveValues.bRoughness = true;
+	}
+	if (SampleChannels.bSpecular)
+	{
+		DefaultsInResultsOut.Specular =
+			SpecularPhotoSet.ComputeSample(Position, Normal, VisibilityFunction, DefaultsInResultsOut.Specular);
+		DefaultsInResultsOut.HaveValues.bSpecular = true;
+	}
+	if (SampleChannels.bMetallic)
+	{
+		DefaultsInResultsOut.Metallic =
+			MetallicPhotoSet.ComputeSample(Position, Normal, VisibilityFunction, DefaultsInResultsOut.Metallic);
+		DefaultsInResultsOut.HaveValues.bMetallic = true;
+	}
+	if (SampleChannels.bEmissive)
+	{
+		DefaultsInResultsOut.Emissive =
+			EmissivePhotoSet.ComputeSample(Position, Normal, VisibilityFunction, DefaultsInResultsOut.Emissive);
+		DefaultsInResultsOut.HaveValues.bEmissive = true;
+	}
+	if (SampleChannels.bWorldNormal)
+	{
+		DefaultsInResultsOut.WorldNormal =
+			WorldNormalPhotoSet.ComputeSample(Position, Normal, VisibilityFunction, DefaultsInResultsOut.WorldNormal);
+		DefaultsInResultsOut.HaveValues.bWorldNormal = true;
+	}
+
+	return true;
+}
\ No newline at end of file
diff --git a/Engine/Plugins/Experimental/MeshModelingToolset/Source/ModelingComponents/Private/Scene/WorldRenderCapture.cpp b/Engine/Plugins/Experimental/MeshModelingToolset/Source/ModelingComponents/Private/Scene/WorldRenderCapture.cpp
new file mode 100644
--- /dev/null
+++ b/Engine/Plugins/Experimental/MeshModelingToolset/Source/ModelingComponents/Private/Scene/WorldRenderCapture.cpp
@@ -0,0 +1,528 @@
+// Copyright Epic Games, Inc. All Rights Reserved.
+
+#include "Scene/WorldRenderCapture.h"
+
+#include "Engine/Canvas.h"
+#include "Engine/TextureRenderTarget2D.h"
+#include "LegacyScreenPercentageDriver.h"
+#include "EngineModule.h"
+#include "Rendering/Texture2DResource.h"
+
+#include "PreviewScene.h"
+#include "Components/DirectionalLightComponent.h"
+#include "Components/ChildActorComponent.h"
+
+using namespace UE::Geometry;
+
+
+
+FRenderCaptureTypeFlags FRenderCaptureTypeFlags::All()
+{
+	return FRenderCaptureTypeFlags{ true, true, true, true, true, true };
+}
+
+FRenderCaptureTypeFlags FRenderCaptureTypeFlags::None()
+{
+	return FRenderCaptureTypeFlags{ false, false, false, false, false, false };
+}
+
+FRenderCaptureTypeFlags FRenderCaptureTypeFlags::BaseColor()
+{
+	return FRenderCaptureTypeFlags{ true, false, false, false, false, false };
+}
+
+FRenderCaptureTypeFlags FRenderCaptureTypeFlags::WorldNormal()
+{
+	return FRenderCaptureTypeFlags{ false, false, false, false, false, true };
+}
+
+FRenderCaptureTypeFlags FRenderCaptureTypeFlags::Single(ERenderCaptureType CaptureType)
+{
+	FRenderCaptureTypeFlags Flags = None();
+	Flags.SetEnabled(CaptureType, true);
+	return Flags;
+}
+
+void FRenderCaptureTypeFlags::SetEnabled(ERenderCaptureType CaptureType, bool bEnabled)
+{
+	switch (CaptureType)
+	{
+	case ERenderCaptureType::BaseColor:
+		bBaseColor = bEnabled;
+		break;
+	case ERenderCaptureType::WorldNormal:
+		bWorldNormal = bEnabled;
+		break;
+	case ERenderCaptureType::Roughness:
+		bRoughness = bEnabled;
+		break;
+	case ERenderCaptureType::Metallic:
+		bMetallic = bEnabled;
+		break;
+	case ERenderCaptureType::Specular:
+		bSpecular = bEnabled;
+		break;
+	case ERenderCaptureType::Emissive:
+		bEmissive = bEnabled;
+		break;
+	default:
+		check(false);
+	}
+}
+
+
+
+
+FWorldRenderCapture::FWorldRenderCapture()
+{
+	Dimensions = FImageDimensions(128, 128);
+}
+
+FWorldRenderCapture::~FWorldRenderCapture()
+{
+	Shutdown();
+}
+
+void FWorldRenderCapture::Shutdown()
+{
+	if (LinearRenderTexture != nullptr)
+	{
+		LinearRenderTexture->RemoveFromRoot();
+		LinearRenderTexture = nullptr;
+	}
+	if (GammaRenderTexture != nullptr)
+	{
+		GammaRenderTexture->RemoveFromRoot();
+		GammaRenderTexture = nullptr;
+	}
+}
+
+
+
+void FWorldRenderCapture::SetWorld(UWorld* WorldIn)
+{
+	this->World = WorldIn;
+}
+
+
+void FWorldRenderCapture::SetVisibleActors(const TArray<AActor*>& Actors)
+{
+	CaptureActors = Actors;
+
+	VisiblePrimitives.Reset();
+
+	VisibleBounds = FBoxSphereBounds();
+	bool bFirst = true;
+
+	// Find all components that need to be included in rendering.
+	// This also descends into any ChildActorComponents
+	TArray<UActorComponent*> ComponentQueue;
+	for (AActor* Actor : Actors)
+	{
+		FVector ActorOrigin, ActorExtent;
+		Actor->GetActorBounds(false, ActorOrigin, ActorExtent, true);
+		FBoxSphereBounds ActorBounds(ActorOrigin, ActorExtent, ActorExtent.Size());
+		VisibleBounds = (bFirst) ? ActorBounds : (VisibleBounds + ActorBounds);
+		bFirst = false;
+
+		ComponentQueue.Reset();
+		for (UActorComponent* Component : Actor->GetComponents())
+		{
+			ComponentQueue.Add(Component);
+		}
+		while (ComponentQueue.Num() > 0)
+		{
+			UActorComponent* Component = ComponentQueue.Pop(false);
+			if (Cast<UPrimitiveComponent>(Component) != nullptr)
+			{
+				UPrimitiveComponent* PrimitiveComponent = Cast<UPrimitiveComponent>(Component);
+				VisiblePrimitives.Add(PrimitiveComponent->ComponentId);
+			}
+			else if (Cast<UChildActorComponent>(Component) != nullptr)
+			{
+				AActor* ChildActor = Cast<UChildActorComponent>(Component)->GetChildActor();
+				for (UActorComponent* SubComponent : ChildActor->GetComponents())
+				{
+					ComponentQueue.Add(SubComponent);
+				}
+			}
+		}
+	}
+}
+
+
+void FWorldRenderCapture::SetDimensions(const FImageDimensions& DimensionsIn)
+{
+	this->Dimensions = DimensionsIn;
+}
+
+
+FSphere FWorldRenderCapture::ComputeContainingRenderSphere(float HorzFOVDegrees, float SafetyBoundsScale) const
+{
+	if (VisiblePrimitives.Num() == 0)
+	{
+		ensure(false);		// unclear what we should do here - get bounds of all actors?
+		return FSphere(FVector(0,0,0), 1000.0f);
+	}
+
+	// todo: I think this maybe needs to be based on the box corners? 
+
+	const float HalfFOVRadians = FMath::DegreesToRadians<float>(HorzFOVDegrees) * 0.5f;
+	const float HalfMeshSize = VisibleBounds.SphereRadius * SafetyBoundsScale;
+	const float TargetDistance = HalfMeshSize / FMath::Tan(HalfFOVRadians);
+	return FSphere(VisibleBounds.Origin, TargetDistance);
+}
+
+
+
+
+
+
+
+UTextureRenderTarget2D* FWorldRenderCapture::GetRenderTexture(bool bLinear)
+{
+	if (RenderTextureDimensions != this->Dimensions)
+	{
+		if (LinearRenderTexture != nullptr)
+		{
+			LinearRenderTexture->RemoveFromRoot();
+			LinearRenderTexture = nullptr;
+		}
+		if (GammaRenderTexture != nullptr)
+		{
+			GammaRenderTexture->RemoveFromRoot();
+			GammaRenderTexture = nullptr;
+		}
+	}
+
+	UTextureRenderTarget2D** WhichTexture = (bLinear) ? &LinearRenderTexture : &GammaRenderTexture;
+
+	if ( *WhichTexture != nullptr )
+	{
+		return *WhichTexture;
+	}
+
+	RenderTextureDimensions = Dimensions;
+	int32 Width = Dimensions.GetWidth();
+	int32 Height = Dimensions.GetHeight();
+
+	*WhichTexture = NewObject<UTextureRenderTarget2D>();
+	if (ensure(*WhichTexture))
+	{
+		(*WhichTexture)->AddToRoot();		// keep alive for GC
+		(*WhichTexture)->ClearColor = FLinearColor::Transparent;
+		(*WhichTexture)->TargetGamma = (bLinear) ? 1.0f : 2.2f;
+		(*WhichTexture)->InitCustomFormat(Width, Height, PF_FloatRGBA, false);
+	}
+
+	return *WhichTexture;
+}
+
+
+
+bool FWorldRenderCapture::CaptureEmissiveFromPosition(
+	const FFrame3d& ViewFrame,
+	double HorzFOVDegrees,
+	double NearPlaneDist,
+	TImageBuilder<FVector4f>& ResultImageOut)
+{
+	// This function is a combination of CaptureFromPosition() and RenderSceneVisualizationToTexture() below,
+	// that attempts to capture Emissive. The difficulty with capturing Emissive is there is no Visualization
+	// Buffer Mode specifically for it, like there is for BaseColor, Specular, WorldNormal, etc. So, the
+	// strategy we take here is basically to disable all scene lights but not disable lighting. This should
+	// result in only Emissive being rendered, ideally without any tone mapping, gamma, bloom, etc, so that we 
+	// directly capture the raw Emissive shader output.
+	//
+	// One known possible issue here is with baked lighting, which may also appear with the current setup.
+	// Currently untested.
+
+	check(this->World);
+	FSceneInterface* Scene = this->World->Scene;
+
+	UTextureRenderTarget2D* RenderTargetTexture = GetRenderTexture(true);
+	if (!ensure(RenderTargetTexture))
+	{
+		return false;
+	}
+	FRenderTarget* RenderTargetResource = RenderTargetTexture->GameThread_GetRenderTargetResource();
+
+	int32 Width = Dimensions.GetWidth();
+	int32 Height = Dimensions.GetHeight();
+
+	FQuat ViewOrientation = (FQuat)ViewFrame.Rotation;
+	FMatrix ViewRotationMatrix = FInverseRotationMatrix(ViewOrientation.Rotator());
+	FVector ViewOrigin = (FVector)ViewFrame.Origin;
+
+	// convert to rendering coordinate system
+	ViewRotationMatrix = ViewRotationMatrix * FMatrix(
+		FPlane(0, 0, 1, 0),
+		FPlane(1, 0, 0, 0),
+		FPlane(0, 1, 0, 0),
+		FPlane(0, 0, 0, 1));
+
+	const float HalfFOVRadians = FMath::DegreesToRadians<float>(HorzFOVDegrees) * 0.5f;
+	static_assert((int32)ERHIZBuffer::IsInverted != 0, "Check NearPlane and Projection Matrix");
+	FMatrix ProjectionMatrix = FReversedZPerspectiveMatrix(HalfFOVRadians, 1.0f, 1.0f, (float)NearPlaneDist);
+
+	EViewModeIndex ViewModeIndex = EViewModeIndex::VMI_Unlit;		// VMI_Lit, VMI_LightingOnly, VMI_VisualizeBuffer
+
+	FEngineShowFlags ShowFlags = FEngineShowFlags(EShowFlagInitMode::ESFIM_Game);
+	ApplyViewMode(ViewModeIndex, true, ShowFlags);
+
+	// unclear if these flags need to be set before creating ViewFamily
+	ShowFlags.SetAntiAliasing(false);
+	ShowFlags.SetDepthOfField(false);
+	ShowFlags.SetMotionBlur(false);
+	ShowFlags.SetBloom(false);
+	ShowFlags.SetSceneColorFringe(false);
+
+	FSceneViewFamilyContext ViewFamily(FSceneViewFamily::ConstructionValues(
+		RenderTargetResource, Scene, ShowFlags)
+		.SetWorldTimes(0, 0, 0)
+		.SetRealtimeUpdate(false));
+
+	// This set of flags currently seems to work for capturing Emissive. Possibly many are unnecessary or
+	// ignored due to ther rendering config, but it's quite hard to know without extensive A/B testing
+
+	ViewFamily.EngineShowFlags.DisableAdvancedFeatures();
+	ViewFamily.EngineShowFlags.MotionBlur = 0;
+	ViewFamily.EngineShowFlags.LOD = 0;
+
+	ViewFamily.EngineShowFlags.SetTonemapper(false);
+	ViewFamily.EngineShowFlags.SetColorGrading(false);
+	ViewFamily.EngineShowFlags.SetToneCurve(false);
+
+	ViewFamily.EngineShowFlags.SetPostProcessing(false);
+	ViewFamily.EngineShowFlags.SetFog(false);
+	ViewFamily.EngineShowFlags.SetGlobalIllumination(false);
+	ViewFamily.EngineShowFlags.SetEyeAdaptation(false);
+	ViewFamily.EngineShowFlags.SetDirectionalLights(false);
+	ViewFamily.EngineShowFlags.SetPointLights(false);
+	ViewFamily.EngineShowFlags.SetSpotLights(false);
+	ViewFamily.EngineShowFlags.SetRectLights(false);
+
+	ViewFamily.EngineShowFlags.SetDiffuse(false);
+	ViewFamily.EngineShowFlags.SetSpecular(false);
+
+	ViewFamily.EngineShowFlags.SetDynamicShadows(false);
+	ViewFamily.EngineShowFlags.SetCapsuleShadows(false);
+	ViewFamily.EngineShowFlags.SetContactShadows(false);
+
+	// some of these settings may need to be configured in more complex cases, currently untested
+	//ViewFamily.SceneCaptureSource = SCS_FinalColorHDR;
+	//ViewFamily.bWorldIsPaused = true;
+	//ViewFamily.ViewMode = ViewModeIndex;
+	// since we are capturing actual render we are not using visualize-buffer mode (?)
+	//ViewFamily.EngineShowFlags.SetVisualizeBuffer(false);
+
+	//ViewFamily.EngineShowFlags.SetScreenPercentage(false);
+	ViewFamily.SetScreenPercentageInterface(new FLegacyScreenPercentageDriver(ViewFamily, 1.f, false));
+
+	// This is called in various other places, unclear if we should be doing this too
+	//EngineShowFlagOverride(EShowFlagInitMode::ESFIM_Game, ViewModeIndex, ViewFamily.EngineShowFlags, true);
+
+	FSceneViewInitOptions ViewInitOptions;
+	ViewInitOptions.SetViewRectangle(FIntRect(0, 0, Width, Height));
+	ViewInitOptions.ViewFamily = &ViewFamily;
+	if (VisiblePrimitives.Num() > 0)
+	{
+		ViewInitOptions.ShowOnlyPrimitives = VisiblePrimitives;
+	}
+	ViewInitOptions.ViewOrigin = ViewOrigin;
+	ViewInitOptions.ViewRotationMatrix = ViewRotationMatrix;
+	ViewInitOptions.ProjectionMatrix = ProjectionMatrix;
+	ViewInitOptions.FOV = HorzFOVDegrees;
+
+	FSceneView* NewView = new FSceneView(ViewInitOptions);
+	ViewFamily.Views.Add(NewView);
+
+	NewView->StartFinalPostprocessSettings(ViewInitOptions.ViewOrigin);
+	NewView->EndFinalPostprocessSettings(ViewInitOptions);
+
+	// other FSceneView settings that may need configuring to properly capture Emissive - needs testing
+	//NewView->SetupRayTracedRendering();
+	//NewView->bIsOfflineRender = true;
+	//NewView->bIsSceneCapture = false;
+	//NewView->AntiAliasingMethod = IsAntiAliasingSupported() ? InOutSampleState.AntiAliasingMethod : EAntiAliasingMethod::AAM_None;
+	//NewView->AntiAliasingMethod = EAntiAliasingMethod::AAM_None;
+
+	// can we fully disable auto-exposure?
+	NewView->FinalPostProcessSettings.AutoExposureMethod = EAutoExposureMethod::AEM_Manual;
+
+	// do we actually need for force SM5 here? the other FCanvas() constructor below does not pass these flags...
+	FCanvas Canvas = FCanvas(RenderTargetResource, nullptr, this->World, ERHIFeatureLevel::SM5, FCanvas::CDM_DeferDrawing, 1.0f);
+	Canvas.Clear(FLinearColor::Transparent);
+	GetRendererModule().BeginRenderingViewFamily(&Canvas, &ViewFamily);
+
+	ReadImageBuffer.SetNumUninitialized(Width * Height);
+	FReadSurfaceDataFlags ReadSurfaceDataFlags;
+	ReadSurfaceDataFlags.SetLinearToGamma(false);
+	RenderTargetResource->ReadLinearColorPixels(ReadImageBuffer, ReadSurfaceDataFlags, FIntRect(0, 0, Width, Height));
+
+	FlushRenderingCommands();
+
+	ResultImageOut.SetDimensions(Dimensions);
+	for (int32 yi = 0; yi < Height; ++yi)
+	{
+		for (int32 xi = 0; xi < Width; ++xi)
+		{
+			FLinearColor PixelColorf = ReadImageBuffer[yi * Width + xi];
+			PixelColorf.A = 1.0f;		// ?
+			ResultImageOut.SetPixel(FVector2i(xi, yi), FVector4f(PixelColorf));
+		}
+	}
+
+	return true;
+}
+
+
+
+
+/**
+ * internal Utility function to render the given Scene to a render target and capture
+ * one of the render buffers, defined by VisualizationMode. Not clear where
+ * the valid VisualizationMode FNames are defined, possibly this list: "BaseColor,Specular,SubsurfaceColor,WorldNormal,SeparateTranslucencyRGB,,,WorldTangent,SeparateTranslucencyA,,,Opacity,SceneDepth,Roughness,Metallic,ShadingModel,,SceneDepthWorldUnits,SceneColor,PreTonemapHDRColor,PostTonemapHDRColor"
+ */
+namespace UE
+{
+namespace Internal
+{ 
+	static void RenderSceneVisualizationToTexture(
+		UTextureRenderTarget2D* RenderTargetTexture,
+		FImageDimensions Dimensions,
+		FSceneInterface* Scene,
+		const FName& VisualizationMode,
+		const FVector& ViewOrigin,
+		const FMatrix& ViewRotationMatrix,
+		const FMatrix& ProjectionMatrix,
+		const TSet<FPrimitiveComponentId>& HiddenPrimitives,		// these primitives will be hidden
+		const TSet<FPrimitiveComponentId>& VisiblePrimitives,		// if non-empty, only these primitives are shown
+		TArray<FLinearColor>& OutSamples
+	)
+	{
+		int32 Width = Dimensions.GetWidth();
+		int32 Height = Dimensions.GetHeight();
+		FTextureRenderTargetResource* RenderTargetResource = RenderTargetTexture->GameThread_GetRenderTargetResource();
+
+		FSceneViewFamilyContext ViewFamily(
+			FSceneViewFamily::ConstructionValues(RenderTargetResource, Scene, FEngineShowFlags(ESFIM_Game))
+			.SetWorldTimes(FApp::GetCurrentTime() - GStartTime, FApp::GetDeltaTime(), FApp::GetCurrentTime() - GStartTime)
+		);
+
+		// To enable visualization mode
+		ViewFamily.EngineShowFlags.SetPostProcessing(true);
+		ViewFamily.EngineShowFlags.SetVisualizeBuffer(true);
+		ViewFamily.EngineShowFlags.SetTonemapper(false);
+		ViewFamily.EngineShowFlags.SetScreenPercentage(false);
+
+		FSceneViewInitOptions ViewInitOptions;
+		ViewInitOptions.SetViewRectangle(FIntRect(0, 0, Width, Height));
+		ViewInitOptions.ViewFamily = &ViewFamily;
+		ViewInitOptions.HiddenPrimitives = HiddenPrimitives;
+		if (VisiblePrimitives.Num() > 0)
+		{
+			ViewInitOptions.ShowOnlyPrimitives = VisiblePrimitives;
+		}
+		ViewInitOptions.ViewOrigin = ViewOrigin;
+		ViewInitOptions.ViewRotationMatrix = ViewRotationMatrix;
+		ViewInitOptions.ProjectionMatrix = ProjectionMatrix;
+
+		FSceneView* NewView = new FSceneView(ViewInitOptions);
+		NewView->CurrentBufferVisualizationMode = VisualizationMode;
+		ViewFamily.Views.Add(NewView);
+
+		ViewFamily.SetScreenPercentageInterface(new FLegacyScreenPercentageDriver(ViewFamily, 1.0f, false));
+
+		// should we cache the FCanvas?
+		FCanvas Canvas(RenderTargetResource, NULL, FApp::GetCurrentTime() - GStartTime, FApp::GetDeltaTime(), FApp::GetCurrentTime() - GStartTime, Scene->GetFeatureLevel());
+		Canvas.Clear(FLinearColor::Transparent);
+
+		GetRendererModule().BeginRenderingViewFamily(&Canvas, &ViewFamily);
+
+		// Copy the contents of the remote texture to system memory
+		OutSamples.SetNumUninitialized(Width * Height);
+		FReadSurfaceDataFlags ReadSurfaceDataFlags;
+		ReadSurfaceDataFlags.SetLinearToGamma(false);
+		RenderTargetResource->ReadLinearColorPixels(OutSamples, ReadSurfaceDataFlags, FIntRect(0, 0, Width, Height));
+
+		FlushRenderingCommands();
+	}
+}
+}
+
+
+
+bool FWorldRenderCapture::CaptureFromPosition(
+	ERenderCaptureType CaptureType,
+	const FFrame3d& ViewFrame,
+	double HorzFOVDegrees,
+	double NearPlaneDist,
+	TImageBuilder<FVector4f>& ResultImageOut)
+{
+	if (CaptureType == ERenderCaptureType::Emissive)
+	{
+		return CaptureEmissiveFromPosition(ViewFrame, HorzFOVDegrees, NearPlaneDist, ResultImageOut);
+	}
+
+	// Roughness visualization is rendered with gamma correction (unclear why)
+	bool bLinear = (CaptureType != ERenderCaptureType::Roughness);
+	UTextureRenderTarget2D* RenderTargetTexture = GetRenderTexture(bLinear);
+	if (ensure(RenderTargetTexture) == false)
+	{
+		return false;
+	}
+
+	int32 Width = Dimensions.GetWidth();
+	int32 Height = Dimensions.GetHeight();
+
+	const FRotator RotationOffsetToViewCenter(0.f, 90.f, 0.f);
+	FQuat ViewOrientation = (FQuat)ViewFrame.Rotation;
+	FMatrix ViewRotationMatrix = FInverseRotationMatrix(ViewOrientation.Rotator());
+	FVector ViewOrigin = (FVector)ViewFrame.Origin;
+
+	// convert to rendering coordinate system
+	ViewRotationMatrix = ViewRotationMatrix * FMatrix(
+		FPlane(0, 0, 1, 0),
+		FPlane(1, 0, 0, 0),
+		FPlane(0, 1, 0, 0),
+		FPlane(0, 0, 0, 1));
+
+	const float HalfFOVRadians = FMath::DegreesToRadians<float>(HorzFOVDegrees) * 0.5f;
+	static_assert((int32)ERHIZBuffer::IsInverted != 0, "Check NearPlane and Projection Matrix");
+	FMatrix ProjectionMatrix = FReversedZPerspectiveMatrix( HalfFOVRadians, 1.0f, 1.0f, (float)NearPlaneDist );
+
+	FName CaptureTypeName = FName("BaseColor");
+	switch (CaptureType)
+	{
+	case ERenderCaptureType::WorldNormal:	CaptureTypeName = FName("WorldNormal");	break;
+	case ERenderCaptureType::Roughness:		CaptureTypeName = FName("Roughness");	break;
+	case ERenderCaptureType::Metallic:		CaptureTypeName = FName("Metallic");	break;
+	case ERenderCaptureType::Specular:		CaptureTypeName = FName("Specular");	break;
+	case ERenderCaptureType::BaseColor:		CaptureTypeName = FName("BaseColor");	break;
+	}
+
+	ReadImageBuffer.Reset();
+	TSet<FPrimitiveComponentId> HiddenPrimitives;
+	UE::Internal::RenderSceneVisualizationToTexture( 
+		RenderTargetTexture, this->Dimensions,
+		World->Scene, CaptureTypeName,
+		ViewOrigin, ViewRotationMatrix, ProjectionMatrix,
+		HiddenPrimitives,
+		VisiblePrimitives,
+		ReadImageBuffer);
+
+	ResultImageOut.SetDimensions(Dimensions);
+	for (int32 yi = 0; yi < Height; ++yi)
+	{
+		for (int32 xi = 0; xi < Width; ++xi)
+		{
+			//FColor PixelColor = OutputImage[yi * Width + xi];
+			//FLinearColor PixelColorf = PixelColor.ReinterpretAsLinear();
+			FLinearColor PixelColorf = ReadImageBuffer[yi * Width + xi];
+			PixelColorf.A = 1.0f;		// ?
+			ResultImageOut.SetPixel(FVector2i(xi, yi), FVector4f(PixelColorf));
+		}
+	}
+
+	return true;
+}
diff --git a/Engine/Plugins/Experimental/MeshModelingToolset/Source/ModelingComponents/Public/AssetUtils/Texture2DBuilder.h b/Engine/Plugins/Experimental/MeshModelingToolset/Source/ModelingComponents/Public/AssetUtils/Texture2DBuilder.h
--- a/Engine/Plugins/Experimental/MeshModelingToolset/Source/ModelingComponents/Public/AssetUtils/Texture2DBuilder.h
+++ b/Engine/Plugins/Experimental/MeshModelingToolset/Source/ModelingComponents/Public/AssetUtils/Texture2DBuilder.h
@@ -34,6 +34,9 @@ public:
 	enum class ETextureType
 	{
 		Color,
+		Roughness,
+		Metallic,
+		Specular,
 		NormalMap,
 		AmbientOcclusion
 	};
@@ -85,6 +88,11 @@ public:
 			return false;
 		}
 
+		if (BuildType == ETextureType::Roughness || BuildType == ETextureType::Metallic || BuildType == ETextureType::Specular)
+		{
+			RawTexture2D->SRGB = false;
+			RawTexture2D->UpdateResource();
+		}
 		if (BuildType == ETextureType::NormalMap)
 		{
 			RawTexture2D->CompressionSettings = TC_Normalmap;
@@ -333,7 +341,7 @@ public:
 	/**
 	 * populate texel values from floating-point SourceImage
 	 */
-	bool Copy(const TImageBuilder<FVector3f>& SourceImage, const bool bSRGB = false)
+	bool Copy(const TImageBuilder<FVector3f>& SourceImage, const bool bConvertToSRGB = false)
 	{
 		if (ensure(SourceImage.GetDimensions() == Dimensions) == false)
 		{
@@ -346,16 +354,17 @@ public:
 			Pixel.X = FMathf::Clamp(Pixel.X, 0.0, 1.0);
 			Pixel.Y = FMathf::Clamp(Pixel.Y, 0.0, 1.0);
 			Pixel.Z = FMathf::Clamp(Pixel.Z, 0.0, 1.0);
-			FColor Texel = ToLinearColor(Pixel).ToFColor(bSRGB);
+			FColor Texel = ToLinearColor(Pixel).ToFColor(bConvertToSRGB);
 			SetTexel(i, Texel);
 		}
 		return true;
 	}
 
 	/**
-	 * populate texel values from floating-point SourceImage
+	 * Populate texel values from floating-point SourceImage.
+	 * @param bConvertToSRGB if true, SourceImage is assumed to be Linear, and converted to SRGB for the Texture
 	 */
-	bool Copy(const TImageBuilder<FVector4f>& SourceImage, const bool bSRGB = false)
+	bool Copy(const TImageBuilder<FVector4f>& SourceImage, const bool bConvertToSRGB = false)
 	{
 		if (ensure(SourceImage.GetDimensions() == Dimensions) == false)
 		{
@@ -369,7 +378,7 @@ public:
 			Pixel.Y = FMathf::Clamp(Pixel.Y, 0.0, 1.0);
 			Pixel.Z = FMathf::Clamp(Pixel.Z, 0.0, 1.0);
 			Pixel.W = FMathf::Clamp(Pixel.W, 0.0, 1.0);
-			FColor Texel = ((FLinearColor)Pixel).ToFColor(bSRGB);
+			FColor Texel = ((FLinearColor)Pixel).ToFColor(bConvertToSRGB);
 			SetTexel(i, Texel);
 		}
 		return true;
@@ -377,7 +386,7 @@ public:
 
 
 	/**
-	 * copy existing texel values to floating-point DestImage
+	 * copy existing texel values to floating-point DestImage. Assumed to be same color space (eg both Linear)
 	 */
 	bool CopyTo(TImageBuilder<FVector4f>& DestImage) const
 	{
@@ -423,6 +432,9 @@ public:
 	const FColor& GetClearColor() const
 	{
 		static const FColor DefaultColor = FColor::Black;
+		static const FColor DefaultRoughness(128, 128, 128);
+		static const FColor DefaultSpecular(100, 100, 100);
+		static const FColor DefaultMetallic(16, 16, 16);
 		static const FColor DefaultNormalColor(128, 128, 255);
 		static const FColor DefaultAOColor = FColor::White;
 
@@ -431,6 +443,12 @@ public:
 		default:
 		case ETextureType::Color:
 			return DefaultColor;
+		case ETextureType::Roughness:
+			return DefaultRoughness;
+		case ETextureType::Metallic:
+			return DefaultMetallic;
+		case ETextureType::Specular:
+			return DefaultSpecular;
 		case ETextureType::NormalMap:
 			return DefaultNormalColor;
 		case ETextureType::AmbientOcclusion:
@@ -453,6 +471,17 @@ public:
 		return bOK;
 	}
 
+
+	template<typename PixelType>
+	static UTexture2D* BuildTextureFromImage(const TImageBuilder<PixelType>& SourceImage, FTexture2DBuilder::ETextureType TextureType, const bool bConvertToSRGB, bool bPopulateSourceData = true)
+	{
+		FTexture2DBuilder TexBuilder;
+		TexBuilder.Initialize(TextureType, SourceImage.GetDimensions());
+		TexBuilder.Copy(SourceImage, bConvertToSRGB);
+		TexBuilder.Commit(bPopulateSourceData);
+		return TexBuilder.GetTexture2D();
+	}
+
 };
 
 
diff --git a/Engine/Plugins/Experimental/MeshModelingToolset/Source/ModelingComponents/Public/Scene/SceneCapturePhotoSet.h b/Engine/Plugins/Experimental/MeshModelingToolset/Source/ModelingComponents/Public/Scene/SceneCapturePhotoSet.h
new file mode 100644
--- /dev/null
+++ b/Engine/Plugins/Experimental/MeshModelingToolset/Source/ModelingComponents/Public/Scene/SceneCapturePhotoSet.h
@@ -0,0 +1,136 @@
+// Copyright Epic Games, Inc. All Rights Reserved.
+
+#pragma once
+
+#include "VectorTypes.h"
+#include "Image/SpatialPhotoSet.h"
+#include "Scene/WorldRenderCapture.h"
+
+class UWorld;
+class AActor;
+
+namespace UE
+{
+namespace Geometry
+{
+
+/**
+ * FSceneCapturePhotoSet creates a set of render captures for a given World and set of Actors,
+ * stored as a SpatialPhotoSet for each desired render buffer type. Currently the set of buffers are
+ * defined by ERenderCaptureType:
+ *		BaseColor
+ *		Roughness
+ *		Metallic
+ *		Specular
+ *		Emissive
+ *		WorldNormal
+ * There are various efficiences possible by doing these captures as a group, rather
+ * than doing each one individually.
+ * 
+ * One the capture set is computed, the ComputeSample() function can be used to 
+ * call SpatialPhotoSet::ComputeSample() on each photo set, ie to estimate the
+ * value of the different channels at a given 3D position/normal by raycasting
+ * against the photo set. Again, it can be more efficient to do this on the
+ * group, rather than each individually.
+ * 
+ */
+class MODELINGCOMPONENTS_API FSceneCapturePhotoSet
+{
+public:
+	/**
+	 * Set the target World and set of Actors
+	 */
+	void SetCaptureSceneActors(UWorld* World, const TArray<AActor*>& Actors);
+
+	/**
+	 * Enable/Disable a particular capture type. Currently all are enabled by default.
+	 */
+	void SetCaptureTypeEnabled(ERenderCaptureType CaptureType, bool bEnabled);
+
+	/**
+	 * Add captures at the corners and face centers of the "view box",
+	 * ie the bounding box that contains the view sphere (see AddExteriorCaptures)
+	 */
+	void AddStandardExteriorCapturesFromBoundingBox(
+		FImageDimensions PhotoDimensions,
+		double HorizontalFOVDegrees,
+		double NearPlaneDist,
+		bool bFaces,
+		bool bUpperCorners,
+		bool bLowerCorners);
+
+	/**
+	 * Add captures on the "view sphere", ie a sphere centered/sized such that the target actors
+	 * will be fully contained inside a square image rendered from locations on the sphere, where
+	 * the view direction is towards the sphere center. The Directions array defines the directions.
+	 */
+	void AddExteriorCaptures(
+		FImageDimensions PhotoDimensions,
+		double HorizontalFOVDegrees,
+		double NearPlaneDist,
+		const TArray<FVector3d>& Directions);
+
+	/**
+	 * Post-process the various PhotoSets after capture, to reduce memory usage and sampling cost.
+	 */
+	void OptimizePhotoSets();
+
+
+	/**
+	 * FSceneSample stores a full sample of all possible channels, some
+	 * values may be default-values though
+	 */
+	struct MODELINGCOMPONENTS_API FSceneSample
+	{
+		FRenderCaptureTypeFlags HaveValues;		// defines which channels have non-default values
+		FVector4f BaseColor;
+		FVector4f Roughness;
+		FVector4f Specular;
+		FVector4f Metallic;
+		FVector4f Emissive;
+		FVector4f WorldNormal;
+
+		FSceneSample();
+
+		/** @return value for the given captured channel, or default value */
+		FVector4f GetValue(ERenderCaptureType CaptureType) const;
+	};
+
+
+	/**
+	 * Sample the requested SampleChannels from the available PhotoSets to determine
+	 * values at the given 3D Position/Normal. This calls TSpatialPhotoSet::ComputeSample()
+	 * internally, see that function for more details.
+	 * @param DefaultsInResultsOut this value is passed in by caller with suitable Default values, and returned with any available computed values updated
+	 */
+	bool ComputeSample(
+		const FRenderCaptureTypeFlags& SampleChannels,
+		const FVector3d& Position,
+		const FVector3d& Normal,
+		TFunctionRef<bool(const FVector3d&, const FVector3d&)> VisibilityFunction,
+		FSceneSample& DefaultsInResultsOut) const;
+
+
+protected:
+	UWorld* TargetWorld = nullptr;
+	TArray<AActor*> VisibleActors;
+
+	bool bEnableBaseColor = true;
+	bool bEnableRoughness = true;
+	bool bEnableSpecular = true;
+	bool bEnableMetallic = true;
+	bool bEnableWorldNormal = true;
+	bool bEnableEmissive = true;
+
+	FSpatialPhotoSet4f BaseColorPhotoSet;
+	FSpatialPhotoSet4f RoughnessPhotoSet;
+	FSpatialPhotoSet4f SpecularPhotoSet;
+	FSpatialPhotoSet4f MetallicPhotoSet;
+	FSpatialPhotoSet4f WorldNormalPhotoSet;
+	FSpatialPhotoSet4f EmissivePhotoSet;
+};
+
+
+
+} // end namespace UE::Geometry
+} // end namespace UE
diff --git a/Engine/Plugins/Experimental/MeshModelingToolset/Source/ModelingComponents/Public/Scene/WorldRenderCapture.h b/Engine/Plugins/Experimental/MeshModelingToolset/Source/ModelingComponents/Public/Scene/WorldRenderCapture.h
new file mode 100644
--- /dev/null
+++ b/Engine/Plugins/Experimental/MeshModelingToolset/Source/ModelingComponents/Public/Scene/WorldRenderCapture.h
@@ -0,0 +1,142 @@
+// Copyright Epic Games, Inc. All Rights Reserved.
+
+#pragma once
+
+#include "SceneTypes.h"		// FPrimitiveComponentId
+#include "Image/SpatialPhotoSet.h"
+
+class UTextureRenderTarget2D;
+class AActor;
+
+namespace UE
+{
+namespace Geometry
+{
+
+/**
+ * ERenderCaptureType defines which type of render buffer should be captured.
+ */
+enum class ERenderCaptureType : uint8
+{
+	BaseColor = 1,
+	Roughness = 2,
+	Metallic = 4,
+	Specular = 8,
+	Emissive = 16,
+	WorldNormal = 32
+};
+
+/**
+ * FRenderCaptureTypeFlags is a set of per-capture-type booleans
+ */
+struct MODELINGCOMPONENTS_API FRenderCaptureTypeFlags
+{
+	bool bBaseColor = false;
+	bool bRoughness = false;
+	bool bMetallic = false;
+	bool bSpecular = false;
+	bool bEmissive = false;
+	bool bWorldNormal = false;
+
+	/** @return FRenderCaptureTypeFlags with all types enabled/true */
+	static FRenderCaptureTypeFlags All();
+
+	/** @return FRenderCaptureTypeFlags with all types disabled/false */
+	static FRenderCaptureTypeFlags None();
+
+	/** @return FRenderCaptureTypeFlags with only BaseColor enabled/true */
+	static FRenderCaptureTypeFlags BaseColor();
+
+	/** @return FRenderCaptureTypeFlags with only WorldNormal enabled/true */
+	static FRenderCaptureTypeFlags WorldNormal();
+
+	/** @return FRenderCaptureTypeFlags with a single capture type enabled/true */
+	static FRenderCaptureTypeFlags Single(ERenderCaptureType CaptureType);
+
+	/** Set the indicated CaptureType to enabled */
+	void SetEnabled(ERenderCaptureType CaptureType, bool bEnabled);
+};
+
+
+/**
+ * FWorldRenderCapture captures a rendering of a set of Actors in a World from a
+ * specific viewpoint. Various types of rendering are supported, as defined by ERenderCaptureType.
+ */
+class MODELINGCOMPONENTS_API FWorldRenderCapture
+{
+public:
+	FWorldRenderCapture();
+	~FWorldRenderCapture();
+
+	/** Explicitly release any allocated textures or other data structres */
+	void Shutdown();
+
+	/** Set the target World */
+	void SetWorld(UWorld* World);
+
+	/** 
+	 *  Set the set of Actors in the target World that should be included in the Rendering.
+	 *  Currently rendering an entire World, ie without an explicit list of Actors, is not supported.
+	 */
+	void SetVisibleActors(const TArray<AActor*>& Actors);
+
+	/** Get bounding-box of the Visible actors */
+	FBoxSphereBounds GetVisibleActorBounds() const { return VisibleBounds; }
+
+	/** 
+	 * Compute a sphere where, if the camera is on the sphere pointed at the center, then the Visible Actors
+	 * will be fully visible (ie a square capture will not have any clipping), for the given Field of View.
+	 * SafetyBoundsScale multiplier is applied to the bounding box.
+	 */
+	FSphere ComputeContainingRenderSphere(float HorzFOVDegrees, float SafetyBoundsScale = 1.25f) const;
+
+	/** Set desired pixel dimensions of the rendered target image */
+	void SetDimensions(const FImageDimensions& Dimensions);
+
+	/** @return pixel dimensions that target image will be rendered at */
+	const FImageDimensions& GetDimensions() const { return Dimensions; }
+
+	/**
+	 * Capture the desired buffer type CaptureType with the given view/camera parameters.
+	 * @param ResultImageOut output iamge of size GetDimensions() is stored here.
+	 * @return true if capture could be rendered successfully
+	 */
+	bool CaptureFromPosition(
+		ERenderCaptureType CaptureType,
+		const FFrame3d& ViewFrame,
+		double HorzFOVDegrees,
+		double NearPlaneDist,
+		TImageBuilder<FVector4f>& ResultImageOut);
+
+protected:
+	UWorld* World;
+
+	TArray<AActor*> CaptureActors;
+	TSet<FPrimitiveComponentId> VisiblePrimitives;
+	FBoxSphereBounds VisibleBounds;
+
+	FImageDimensions Dimensions;
+
+	// Temporary textures used as render targets. We explicitly prevent this from being GC'd internally
+	UTextureRenderTarget2D* LinearRenderTexture = nullptr;
+	UTextureRenderTarget2D* GammaRenderTexture = nullptr;
+	FImageDimensions RenderTextureDimensions;
+	UTextureRenderTarget2D* GetRenderTexture(bool bLinear);
+
+	// temporary buffer used to read from texture
+	TArray<FLinearColor> ReadImageBuffer;
+
+	/** Emissive is a special case and uses different code than capture of color/property channels */
+	bool CaptureEmissiveFromPosition(
+		const FFrame3d& Frame,
+		double HorzFOVDegrees,
+		double NearPlaneDist,
+		TImageBuilder<FVector4f>& ResultImageOut);
+};
+
+
+
+
+
+} // end namespace UE::Geometry
+} // end namespace UE
