GeometryProcessing: add basic TTransformSequence3, allows storing multiple separate transforms in list, transforming points
ModelingComponents: add FMeshSceneAdapter, generates "flattened" mesh representation for a set of Actors, with support for instancing to avoid duplication, and creates minimal set of AABBTree/FWNTrees for the unique meshes. Does not make copies of any mesh data, uses MeshAdapter to allow queries directly on MeshDescriptions.
#rnx

diff --git a/Engine/Plugins/Experimental/GeometryProcessing/Source/GeometricObjects/Public/TransformSequence.h b/Engine/Plugins/Experimental/GeometryProcessing/Source/GeometricObjects/Public/TransformSequence.h
new file mode 100644
--- /dev/null
+++ b/Engine/Plugins/Experimental/GeometryProcessing/Source/GeometricObjects/Public/TransformSequence.h
@@ -0,0 +1,101 @@
+// Copyright Epic Games, Inc. All Rights Reserved.
+
+#pragma once
+
+#include "CoreMinimal.h"
+#include "TransformTypes.h"
+
+namespace UE
+{
+namespace Geometry
+{
+
+
+/**
+ * TTransformSequence3 represents a sequence of 3D transforms. 
+ */
+template<typename RealType>
+class TTransformSequence3
+{
+protected:
+	TArray<TTransform3<RealType>, TInlineAllocator<2>> Transforms;
+
+public:
+
+	/**
+	 * Add Transform to the end of the sequence, ie Seq(P) becomes NewTransform * Seq(P)
+	 */
+	void Append(const TTransform3<RealType>& Transform)
+	{
+		Transforms.Add(Transform);
+	}
+
+	/**
+	 * Add Transform to the end of the sequence, ie Seq(P) becomes NewTransform * Seq(P)
+	 */
+	void Append(const FTransform& Transform)
+	{
+		Transforms.Add( TTransform3<RealType>(Transform) );
+	}
+
+	/**
+	 * @return number of transforms in the sequence
+	 */
+	int32 Num() const { return Transforms.Num(); }
+
+	/**
+	 * @return transforms in the sequence
+	 */
+	const TArray<TTransform3<RealType>,TInlineAllocator<2>>& GetTransforms() const { return Transforms; }
+
+	/**
+	 * @return point P with transform sequence applied
+	 */
+	FVector3d TransformPosition(FVector3d P) const
+	{
+		for (const TTransform3<RealType>& Transform : Transforms)
+		{
+			P = Transform.TransformPosition(P);
+		}
+		return P;
+	}
+
+	/**
+	 * @return point P with inverse transform sequence applied
+	 */
+	FVector3d InverseTransformPosition(FVector3d P) const
+	{
+		int32 N = Transforms.Num();
+		for (int32 k = N - 1; k >= 0; k--)
+		{
+			P = Transforms[k].InverseTransformPosition(P);
+		}
+		return P;
+	}
+
+	/**
+	 * Create an reversed transform sequence such that InverseSeq.Transform(P) is equivalent to Seq.InverseTransform(P).
+	 * This is more efficient if applying inverse many times as it saves lots of checking for invalid scaling.
+	 * @return inverse transform sequence
+	 */
+	TTransformSequence3<RealType> Inverse() const
+	{
+		TTransformSequence3 Inverse;
+		int32 N = Transforms.Num();
+		Inverse.Transforms.Reserve(N);
+		for (int32 k = N - 1; k >= 0; --k)
+		{
+			Inverse.Transforms.Add(Transforms[k].Inverse());
+		}
+		return Inverse;
+	}
+
+
+};
+
+typedef TTransformSequence3<float> FTransformSequence3f;
+typedef TTransformSequence3<double> FTransformSequence3d;
+
+
+} // end namespace UE::Geometry
+} // end namespace UE
\ No newline at end of file
diff --git a/Engine/Plugins/Experimental/MeshModelingToolset/Source/ModelingComponents/Private/Scene/MeshSceneAdapter.cpp b/Engine/Plugins/Experimental/MeshModelingToolset/Source/ModelingComponents/Private/Scene/MeshSceneAdapter.cpp
new file mode 100644
--- /dev/null
+++ b/Engine/Plugins/Experimental/MeshModelingToolset/Source/ModelingComponents/Private/Scene/MeshSceneAdapter.cpp
@@ -0,0 +1,270 @@
+// Copyright Epic Games, Inc. All Rights Reserved.
+
+
+#include "Scene/MeshSceneAdapter.h"
+
+#include "Spatial/MeshAABBTree3.h"
+#include "Spatial/FastWinding.h"
+#include "MeshDescriptionAdapter.h"
+
+#include "GameFramework/Actor.h"
+#include "Components/StaticMeshComponent.h"
+#include "Components/InstancedStaticMeshComponent.h"
+#include "Engine/StaticMesh.h"
+
+#include "Async/ParallelFor.h"
+
+using namespace UE::Geometry;
+
+
+class FStaticMeshSpatialWrapper : public IMeshSpatialWrapper
+{
+public:
+	UStaticMesh* StaticMesh = nullptr;
+	int32 LODIndex = 0;
+	
+	FMeshDescription* SourceMesh = nullptr;
+	TUniquePtr<FMeshDescriptionTriangleMeshAdapter> Adapter;
+	TUniquePtr<TMeshAABBTree3<FMeshDescriptionTriangleMeshAdapter>> AABBTree;
+	TUniquePtr<TFastWindingTree<FMeshDescriptionTriangleMeshAdapter>> FWNTree;
+
+	virtual bool Build() override
+	{
+		check(StaticMesh);
+#if WITH_EDITOR
+		SourceMesh = StaticMesh->GetMeshDescription(LODIndex);
+#else
+		checkf(false, TEXT("Not currently supported - to build at Runtime it is necessary to read from the StaticMesh RenderBuffers"));
+		SourceMesh = nullptr;
+#endif
+		if (SourceMesh)
+		{
+			Adapter = MakeUnique<FMeshDescriptionTriangleMeshAdapter>(SourceMesh);
+
+#if WITH_EDITOR
+			// respect BuildScale build setting
+			const FMeshBuildSettings& LODBuildSettings = StaticMesh->GetSourceModel(LODIndex).BuildSettings;
+			Adapter->SetBuildScale((FVector3d)LODBuildSettings.BuildScale3D, false);
+#endif
+
+			AABBTree = MakeUnique<TMeshAABBTree3<FMeshDescriptionTriangleMeshAdapter>>(Adapter.Get(), true);
+			FWNTree = MakeUnique<TFastWindingTree<FMeshDescriptionTriangleMeshAdapter>>(AABBTree.Get(), true);
+			return true;
+		}
+
+		SourceMesh = nullptr;
+		return false;
+	}
+
+
+	virtual void CollectSeedPoints(TArray<FVector3d>& WorldPoints, TFunctionRef<FVector3d(const FVector3d&)> LocalToWorldFunc) override
+	{
+		if (!SourceMesh) return;
+
+		int32 NumVertices = Adapter->VertexCount();
+		int32 LogNumVertices = FMathd::Log(NumVertices);
+		int32 SeedPointCount = (int)(5.0 * LogNumVertices * LogNumVertices);
+		ensure(SeedPointCount < 500);		// just a sanity check
+		int32 Skip = FMath::Max(NumVertices / SeedPointCount, 2);
+		for (int32 k = 0; k < NumVertices; k += Skip)
+		{
+			Adapter->GetNormal(k);
+			WorldPoints.Add(LocalToWorldFunc(Adapter->GetVertex(k)));
+		}
+	}
+
+	virtual double FastWindingNumber(const FVector3d& P) override
+	{
+		return (SourceMesh != nullptr) ? FWNTree->FastWindingNumber(P) : 0.0;
+	}
+};
+
+
+static TSharedPtr<IMeshSpatialWrapper> SpatialWrapperFactory(const FMeshTypeContainer& MeshContainer)
+{
+	if (MeshContainer.MeshType == ESceneMeshType::StaticMeshAsset)
+	{
+		TSharedPtr<FStaticMeshSpatialWrapper> SMWrapper = MakeShared<FStaticMeshSpatialWrapper>();
+		SMWrapper->StaticMesh = MeshContainer.GetStaticMesh();
+		if (ensure(SMWrapper->StaticMesh != nullptr))
+		{
+			return SMWrapper;
+		}
+	}
+
+	return TSharedPtr<IMeshSpatialWrapper>();
+}
+
+
+
+
+static void CollectActorChildMeshes(AActor* Actor, UActorComponent* Component, FActorAdapter& Adapter)
+{
+	UStaticMeshComponent* StaticMeshComponent = Cast<UStaticMeshComponent>(Component);
+	if (StaticMeshComponent != nullptr)
+	{
+		UStaticMesh* Mesh = StaticMeshComponent->GetStaticMesh();
+		if (Mesh != nullptr)
+		{
+			FActorChildMesh ChildMesh;
+			ChildMesh.SourceComponent = Component;
+			ChildMesh.MeshContainer = FMeshTypeContainer{ Mesh, ESceneMeshType::StaticMeshAsset };
+
+			UInstancedStaticMeshComponent* ISMComponent = Cast<UInstancedStaticMeshComponent>(StaticMeshComponent);
+			if (ISMComponent != nullptr)
+			{
+				// TODO: HISMC
+
+				ChildMesh.ComponentType = EActorMeshComponentType::InstancedStaticMesh;
+
+				int32 NumInstances = ISMComponent->GetInstanceCount();
+				for (int32 i = 0; i < NumInstances; ++i)
+				{
+					if (ISMComponent->IsValidInstance(i))
+					{
+						FTransform InstanceTransform;
+						if (ensure(ISMComponent->GetInstanceTransform(i, InstanceTransform, true)))
+						{
+							FActorChildMesh InstanceChild = ChildMesh;
+							InstanceChild.ComponentIndex = i;
+							InstanceChild.WorldTransform.Append(InstanceTransform);
+							InstanceChild.WorldTransformInverse = InstanceChild.WorldTransform.Inverse();
+							Adapter.ChildMeshes.Add(InstanceChild);
+						}
+					}
+				}
+
+			}
+			else
+			{
+				// base StaticMeshComponent
+				ChildMesh.ComponentType = EActorMeshComponentType::StaticMesh;
+				ChildMesh.ComponentIndex = 0;
+				ChildMesh.WorldTransform.Append(StaticMeshComponent->GetComponentTransform());
+				ChildMesh.WorldTransformInverse = ChildMesh.WorldTransform.Inverse();
+				Adapter.ChildMeshes.Add(ChildMesh);
+			}
+		}
+		
+	}
+
+}
+
+
+void FMeshSceneAdapter::AddActors(const TArray<AActor*>& ActorsSetIn)
+{
+	// build an FActorAdapter for each Actor, that contains all mesh Components we know 
+	// how to process, including those contained in ChildActorComponents
+	TArray<AActor*> ChildActors;
+	for (AActor* Actor : ActorsSetIn)
+	{
+		TUniquePtr<FActorAdapter> Adapter = MakeUnique<FActorAdapter>();
+		Adapter->SourceActor = Actor;
+
+		for (UActorComponent* Component : Actor->GetComponents())
+		{
+			CollectActorChildMeshes(Actor, Component, *Adapter);
+		}
+
+		ChildActors.Reset();
+		Actor->GetAllChildActors(ChildActors, true);
+		for (AActor* ChildActor : ChildActors)
+		{
+			for (UActorComponent* Component : ChildActor->GetComponents())
+			{
+				CollectActorChildMeshes(ChildActor, Component, *Adapter);
+			}
+		}
+
+		FVector ActorOrigin, ActorExtent;
+		Actor->GetActorBounds(false, ActorOrigin, ActorExtent, true);
+		Adapter->WorldBounds = FAxisAlignedBox3d(
+			(FVector3d)ActorOrigin - (FVector3d)ActorExtent, (FVector3d)ActorOrigin + (FVector3d)ActorExtent);
+
+		SceneActors.Add(MoveTemp(Adapter));
+	}
+
+
+	// Find IMeshSpatialWrapper for each child mesh component. If one does not exist
+	// and we have not seen the underlying unique mesh (eg StaticMesh Asset, etc, construct a new one
+	TArray<IMeshSpatialWrapper*> NewSpatialsToBuild;
+	for (TUniquePtr<FActorAdapter>& Actor : SceneActors)
+	{
+		for (FActorChildMesh& ChildMesh : Actor->ChildMeshes)
+		{
+			void* MeshKey = ChildMesh.MeshContainer.GetMeshKey();
+			TSharedPtr<IMeshSpatialWrapper>* Found = SpatialAdapters.Find(MeshKey);
+			if (Found == nullptr)
+			{
+				TSharedPtr<IMeshSpatialWrapper> NewWrapper = SpatialWrapperFactory(ChildMesh.MeshContainer);
+				if (ensure(NewWrapper.IsValid()))
+				{
+					ChildMesh.MeshSpatial = NewWrapper.Get();
+					NewSpatialsToBuild.Add(NewWrapper.Get());
+					SpatialAdapters.Add(MeshKey, NewWrapper);
+				}
+			}
+			else
+			{
+				ChildMesh.MeshSpatial = Found->Get();
+			}
+		}
+	}
+
+	
+	// parallel build of all the spatial data structures
+	ParallelFor(NewSpatialsToBuild.Num(), [&](int32 i)
+	{
+		bool bOK = NewSpatialsToBuild[i]->Build();
+		ensure(bOK);	// assumption is that the wrapper will handle failure gracefully
+	});
+
+}
+
+
+FAxisAlignedBox3d FMeshSceneAdapter::GetBoundingBox()
+{
+	FAxisAlignedBox3d SceneBounds = FAxisAlignedBox3d::Empty();
+	for (const TUniquePtr<FActorAdapter>& Actor : SceneActors)
+	{
+		SceneBounds.Contain(Actor->WorldBounds);
+	}
+	return SceneBounds;
+}
+
+
+void FMeshSceneAdapter::CollectMeshSeedPoints(TArray<FVector3d>& Points)
+{
+	for (const TUniquePtr<FActorAdapter>& Actor : SceneActors)
+	{
+		for (const FActorChildMesh& ChildMesh : Actor->ChildMeshes)
+		{
+			if (ChildMesh.MeshSpatial != nullptr)
+			{
+				ChildMesh.MeshSpatial->CollectSeedPoints(Points,
+					[&](const FVector3d& P) { return ChildMesh.WorldTransform.TransformPosition(P); } );
+			}
+		}
+	}
+}
+
+double FMeshSceneAdapter::MaxFastWindingNumber(const FVector3d& P)
+{
+	double MaxWinding = 0.0;
+	for (const TUniquePtr<FActorAdapter>& Actor : SceneActors)
+	{
+		if (Actor->WorldBounds.Contains(P))
+		{
+			for (const FActorChildMesh& ChildMesh : Actor->ChildMeshes)
+			{
+				if (ChildMesh.MeshSpatial != nullptr)
+				{
+					FVector3d LocalP = ChildMesh.WorldTransformInverse.TransformPosition(P);
+					double LocalWinding = ChildMesh.MeshSpatial->FastWindingNumber(LocalP);
+					MaxWinding = FMathd::Max(LocalWinding, MaxWinding);
+				}
+			}
+		}
+	}
+	return MaxWinding;
+}
\ No newline at end of file
diff --git a/Engine/Plugins/Experimental/MeshModelingToolset/Source/ModelingComponents/Public/Scene/MeshSceneAdapter.h b/Engine/Plugins/Experimental/MeshModelingToolset/Source/ModelingComponents/Public/Scene/MeshSceneAdapter.h
new file mode 100644
--- /dev/null
+++ b/Engine/Plugins/Experimental/MeshModelingToolset/Source/ModelingComponents/Public/Scene/MeshSceneAdapter.h
@@ -0,0 +1,181 @@
+// Copyright Epic Games, Inc. All Rights Reserved.
+
+#pragma once
+
+#include "CoreMinimal.h"
+#include "VectorTypes.h"
+#include "BoxTypes.h"
+#include "TransformTypes.h"
+#include "TransformSequence.h"
+
+class AActor;
+class UActorComponent;
+class UStaticMesh;
+
+namespace UE
+{
+namespace Geometry
+{
+
+
+/**
+ * Abstract interface to a spatial data structure for a mesh
+ */
+class MODELINGCOMPONENTS_API IMeshSpatialWrapper
+{
+public:
+	virtual ~IMeshSpatialWrapper() {}
+
+	/** If possible, spatial data structure should defer construction until this function, which will be called off the game thread (in ParallelFor) */
+	virtual bool Build() = 0;
+
+	/** Calculate the mesh winding number at the given Position. Must be callable in parallel from any thread.  */
+	virtual double FastWindingNumber(const FVector3d& P) = 0;
+
+	/** Collect a set of seed points from this Mesh, mapped through LocalToWorldFunc to world space. Must be callable in parallel from any thread. */
+	virtual void CollectSeedPoints(TArray<FVector3d>& WorldPoints, TFunctionRef<FVector3d(const FVector3d&)> LocalToWorldFunc) = 0;
+};
+
+
+/**
+ * ESceneMeshType is used to indicate which type of Mesh a FMeshTypeContainer contains.
+ */
+enum class ESceneMeshType
+{
+	StaticMeshAsset,
+	Unknown
+};
+
+
+/**
+ * FMeshTypeContainer is a wrapper for an object that contains a unique Mesh of some kind,
+ * which is used by an FActorChildMesh to represent that unique mesh. For example this could be 
+ * a UStaticMesh asset, a FMeshDescription, a FDynamicMesh3, and so on.
+ * (Currently only UStaticMesh is supported)
+ */
+struct MODELINGCOMPONENTS_API FMeshTypeContainer
+{
+	/** raw pointer to the Mesh, used as a key to identify this mesh in various maps/etc */
+	void* MeshPointer = nullptr;
+	/** type of unique Mesh object this container contains */
+	ESceneMeshType MeshType = ESceneMeshType::Unknown;
+
+	/** @return key for this mesh */
+	void* GetMeshKey() const { return MeshPointer; }
+
+	/** @return the UStaticMesh this container contains, if this is a StaticMeshAsset container (otherwise nullptr) */
+	UStaticMesh* GetStaticMesh() const
+	{
+		if (ensure(MeshType == ESceneMeshType::StaticMeshAsset))
+		{
+			return (UStaticMesh*)MeshPointer;
+		}
+		return nullptr;
+	}
+};
+
+
+/**
+ * EActorMeshComponentType enum is used to determine which type of Component 
+ * an FActorChildMesh represents.
+ */
+enum class EActorMeshComponentType
+{
+	StaticMesh,
+	InstancedStaticMesh,
+	HierarchicalInstancedStaticMesh,
+
+	Unknown
+};
+
+
+/**
+ * FActorChildMesh represents a 3D Mesh attached to an Actor. This generally comes from a Component,
+ * however in some cases a Component generates multiple FActorChildMesh (eg an InstancedStaticMeshComponent),
+ * and potentially some Actors may store/represent a Mesh directly (no examples currently).
+ */
+struct MODELINGCOMPONENTS_API FActorChildMesh
+{
+public:
+	/** the Component this Mesh was generated from, if there is one. */
+	UActorComponent* SourceComponent = nullptr;
+	/** Type of SourceComponent, if known */
+	EActorMeshComponentType ComponentType = EActorMeshComponentType::Unknown;
+	/** Index of this Mesh in the SourceComponent, if such an index exists (eg Instance Index in InstancedStaticMeshComponent) */
+	int32 ComponentIndex = 0;
+
+	/** Wrapper around the Mesh this FActorChildMesh refers to (eg from a StaticMeshAsset, etc) */
+	FMeshTypeContainer MeshContainer;
+	/** Local-to-World transformation of the Mesh in the MeshContainer */
+	UE::Geometry::FTransformSequence3d WorldTransform;
+	UE::Geometry::FTransformSequence3d WorldTransformInverse;
+
+	/** Spatial data structure that represents the Mesh in MeshContainer - assumption is this is owned externally */
+	IMeshSpatialWrapper* MeshSpatial;
+};
+
+/**
+ * FActorAdapter is used by FMeshSceneAdapter to represent all the child info for an AActor.
+ * This is primarily a list of FActorChildMesh, which represent the spatially-positioned meshes
+ * of any child StaticMeshComponents or other mesh Components that can be identified and represented.
+ * Note that ChildActorComponents will be flatted into the parent Actor.
+ */
+struct MODELINGCOMPONENTS_API FActorAdapter
+{
+public:
+	FActorAdapter() {}
+	FActorAdapter(const FActorAdapter&) = delete;		// must delete to allow TArray<TUniquePtr> in FMeshSceneAdapter
+	FActorAdapter(FActorAdapter&&) = delete;
+
+	// the AActor this Adapter represents
+	AActor* SourceActor = nullptr;
+	// set of child Meshes with transforms
+	TArray<FActorChildMesh> ChildMeshes;
+	// World-space bounds of this Actor (meshes)
+	UE::Geometry::FAxisAlignedBox3d WorldBounds;
+};
+
+
+/**
+ * FMeshSceneAdapter creates an internal representation of an Actor/Component/Asset hierarchy,
+ * so that a minimal set of Mesh data structures can be constructed for the unique Meshes (generally Assets).
+ * This allows queries against the Actor set to be computed without requiring mesh copies or 
+ * duplicates of the mesh data structures (ie, saving memory, at the cost of some computation overhead).
+ * 
+ * Currenly this builds an AABBTree and FastWindingTree for each unique Mesh,
+ * and only FastWinding and "Seed Point" queries are (currently) exposed.
+ * (SeedPoint query is necessary for mesh the scalar Winding field)
+ */
+class MODELINGCOMPONENTS_API FMeshSceneAdapter
+{
+public:
+	virtual ~FMeshSceneAdapter() {}
+
+	FMeshSceneAdapter() {}
+	FMeshSceneAdapter(const FMeshSceneAdapter&) = delete;		// must delete this due to TArray<TUniquePtr> member
+	FMeshSceneAdapter(FMeshSceneAdapter&&) = delete;
+
+	/** Add the  given Actors to our Actor set */
+	void AddActors(const TArray<AActor*>& ActorsSetIn);
+
+	/** @return bounding box for the Actor set */
+	virtual UE::Geometry::FAxisAlignedBox3d GetBoundingBox();
+
+	/** @return a set of points on the surface of the meshes, can be used to initialize the MarchingCubes mesher */
+	virtual void CollectMeshSeedPoints(TArray<FVector3d>& PointsOut);
+
+	/** @return max FastWindingNumber computed for all mesh Actors/Components */
+	virtual double MaxFastWindingNumber(const FVector3d& P);
+
+protected:
+
+	// top-level list of ActorAdapters, which represent each Actor and set of Components
+	TArray<TUniquePtr<FActorAdapter>> SceneActors;
+
+	// Unique set of spatial data structure query interfaces, one for each Mesh object, which is identified by void* pointer
+	TMap<void*, TSharedPtr<IMeshSpatialWrapper>> SpatialAdapters;
+};
+
+
+}
+}
