MeshModelingToolset: add new ModelingComponentsEditorOnly module, with utility functions for creating UStaticMesh/UTexture2D/UMaterial assets.
#rnx

diff --git a/Engine/Plugins/Experimental/MeshModelingToolset/MeshModelingToolset.uplugin b/Engine/Plugins/Experimental/MeshModelingToolset/MeshModelingToolset.uplugin
--- a/Engine/Plugins/Experimental/MeshModelingToolset/MeshModelingToolset.uplugin
+++ b/Engine/Plugins/Experimental/MeshModelingToolset/MeshModelingToolset.uplugin
@@ -20,6 +20,12 @@
 			"Type": "Runtime",
 			"LoadingPhase": "Default"
 		},
+		{
+			"Name": "ModelingComponentsEditorOnly",
+			"Type": "Editor",
+			"LoadingPhase": "Default",
+			"WhitelistPlatforms": [ "Win64" ]
+		},		
 		{
 			"Name": "ModelingOperators",
 			"Type": "Runtime",
diff --git a/Engine/Plugins/Experimental/MeshModelingToolset/Source/ModelingComponentsEditorOnly/ModelingComponentsEditorOnly.Build.cs b/Engine/Plugins/Experimental/MeshModelingToolset/Source/ModelingComponentsEditorOnly/ModelingComponentsEditorOnly.Build.cs
new file mode 100644
--- /dev/null
+++ b/Engine/Plugins/Experimental/MeshModelingToolset/Source/ModelingComponentsEditorOnly/ModelingComponentsEditorOnly.Build.cs
@@ -0,0 +1,65 @@
+// Copyright Epic Games, Inc. All Rights Reserved.
+
+using UnrealBuildTool;
+
+public class ModelingComponentsEditorOnly : ModuleRules
+{
+	public ModelingComponentsEditorOnly(ReadOnlyTargetRules Target) : base(Target)
+	{
+		PCHUsage = ModuleRules.PCHUsageMode.UseExplicitOrSharedPCHs;
+		
+		PublicIncludePaths.AddRange(
+			new string[] {
+				// ... add public include paths required here ...
+			}
+			);
+				
+		
+		PrivateIncludePaths.AddRange(
+			new string[] {
+				// ... add other private include paths required here ...
+			}
+			);
+			
+		
+		PublicDependencyModuleNames.AddRange(
+			new string[]
+			{
+				"Core",
+                "InteractiveToolsFramework",
+                "MeshDescription",
+				"StaticMeshDescription",
+				"GeometricObjects",
+				"GeometryAlgorithms",
+				"DynamicMesh",
+				"MeshConversion",
+				"ModelingOperators",
+				"ModelingComponents"
+				// ... add other public dependencies that you statically link with here ...
+			}
+            );
+			
+		
+		PrivateDependencyModuleNames.AddRange(
+			new string[]
+			{
+				"CoreUObject",
+				"Engine",
+				"InputCore",
+                "RenderCore",
+                "RHI",
+
+				"AssetTools"
+
+			}
+			);
+		
+		
+		DynamicallyLoadedModuleNames.AddRange(
+			new string[]
+			{
+				// ... add any modules that your module loads dynamically here ...
+			}
+			);
+	}
+}
diff --git a/Engine/Plugins/Experimental/MeshModelingToolset/Source/ModelingComponentsEditorOnly/Private/AssetUtils/CreateMaterialUtil.cpp b/Engine/Plugins/Experimental/MeshModelingToolset/Source/ModelingComponentsEditorOnly/Private/AssetUtils/CreateMaterialUtil.cpp
new file mode 100644
--- /dev/null
+++ b/Engine/Plugins/Experimental/MeshModelingToolset/Source/ModelingComponentsEditorOnly/Private/AssetUtils/CreateMaterialUtil.cpp
@@ -0,0 +1,65 @@
+// Copyright Epic Games, Inc. All Rights Reserved. 
+
+#include "AssetUtils/CreateMaterialUtil.h"
+
+#include "Materials/Material.h"
+#include "Materials/MaterialInstanceDynamic.h"
+#include "Materials/MaterialInstanceConstant.h"
+#include "Factories/MaterialInstanceConstantFactoryNew.h"
+
+#include "Misc/Paths.h"
+#include "AssetToolsModule.h"
+#include "IAssetTools.h"
+#include "FileHelpers.h"
+
+
+UE::AssetUtils::ECreateMaterialResult UE::AssetUtils::CreateDuplicateMaterial(
+	UMaterialInterface* BaseMaterial,
+	FMaterialAssetOptions& Options,
+	FMaterialAssetResults& ResultsOut)
+{
+	// Duplicate operation will create a new Package, so calling code should not!
+	if (ensure(Options.UsePackage == nullptr) == false)
+	{
+		return ECreateMaterialResult::InvalidPackage;
+	}
+
+	if (ensure(BaseMaterial) == false)
+	{
+		return ECreateMaterialResult::InvalidBaseMaterial;
+	}
+	UMaterial* SourceMaterial = BaseMaterial->GetBaseMaterial();
+	if (ensure(SourceMaterial) == false)
+	{
+		return ECreateMaterialResult::InvalidBaseMaterial;
+	}
+
+	// cannot create new package if path already exists
+	if (FPackageName::DoesPackageExist(Options.NewAssetPath, nullptr, nullptr))
+	{
+		ensure(false);
+		return ECreateMaterialResult::NameError;
+	}
+
+	FString DestinationLongPackagePath = FPackageName::GetLongPackagePath(Options.NewAssetPath);
+	//FString DestinationObjectName = FPackageName::ObjectPathToObjectName(Options.NewAssetPath);
+	FString DestinationObjectName = FPackageName::GetLongPackageAssetName(Options.NewAssetPath);
+
+	IAssetTools& AssetTools = FModuleManager::Get().LoadModuleChecked<FAssetToolsModule>("AssetTools").Get();
+	UObject* NewAsset = AssetTools.DuplicateAsset(DestinationObjectName, DestinationLongPackagePath, SourceMaterial);
+	UMaterial* NewMaterial = Cast<UMaterial>(NewAsset);
+
+	if ( ensure(NewMaterial) == false )
+	{
+		return ECreateMaterialResult::DuplicateFailed;
+	}
+
+	if (Options.bDeferPostEditChange == false)
+	{
+		NewMaterial->PostEditChange();
+	}
+
+	ResultsOut.NewMaterial = NewMaterial;
+	return ECreateMaterialResult::Ok;
+}
+
diff --git a/Engine/Plugins/Experimental/MeshModelingToolset/Source/ModelingComponentsEditorOnly/Private/AssetUtils/CreateStaticMeshUtil.cpp b/Engine/Plugins/Experimental/MeshModelingToolset/Source/ModelingComponentsEditorOnly/Private/AssetUtils/CreateStaticMeshUtil.cpp
new file mode 100644
--- /dev/null
+++ b/Engine/Plugins/Experimental/MeshModelingToolset/Source/ModelingComponentsEditorOnly/Private/AssetUtils/CreateStaticMeshUtil.cpp
@@ -0,0 +1,137 @@
+// Copyright Epic Games, Inc. All Rights Reserved.
+
+#include "AssetUtils/CreateStaticMeshUtil.h"
+
+#include "Engine/StaticMesh.h"
+#include "Engine/StaticMeshActor.h"
+#include "Components/StaticMeshComponent.h"
+
+#include "Engine/CollisionProfile.h"
+#include "PhysicsEngine/BodySetup.h"
+
+#include "MeshDescription.h"
+#include "DynamicMeshToMeshDescription.h"
+
+using namespace UE::AssetUtils;
+
+#include "ExplicitUseGeometryMathTypes.h"		// using UE::Geometry::(math types)
+using namespace UE::Geometry;
+
+
+
+UE::AssetUtils::ECreateStaticMeshResult UE::AssetUtils::CreateStaticMeshAsset(
+	FStaticMeshAssetOptions& Options,
+	FStaticMeshResults& ResultsOut)
+{
+	FString NewObjectName = FPackageName::GetLongPackageAssetName(Options.NewAssetPath);
+
+	UPackage* UsePackage = nullptr;
+	if (Options.UsePackage != nullptr)
+	{
+		UsePackage = Options.UsePackage;
+	}
+	else
+	{
+		UsePackage = CreatePackage(*Options.NewAssetPath);
+	}
+	if (ensure(UsePackage != nullptr) == false)
+	{
+		return ECreateStaticMeshResult::InvalidPackage;
+	}
+
+	// create new UStaticMesh object
+	EObjectFlags UseFlags = EObjectFlags::RF_Public | EObjectFlags::RF_Standalone;
+	UStaticMesh* NewStaticMesh = NewObject<UStaticMesh>(UsePackage, FName(*NewObjectName), UseFlags);
+	if (ensure(NewStaticMesh != nullptr) == false)
+	{
+		return ECreateStaticMeshResult::UnknownError;
+	}
+
+	// initialize the MeshDescription SourceModel LODs
+	int32 UseNumSourceModels = FMath::Max(1, Options.NumSourceModels);
+	NewStaticMesh->SetNumSourceModels(UseNumSourceModels);
+	for (int32 k = 0; k < UseNumSourceModels; ++k)
+	{
+		NewStaticMesh->GetSourceModel(k).BuildSettings.bRecomputeNormals = Options.bEnableRecomputeNormals;
+		NewStaticMesh->GetSourceModel(k).BuildSettings.bRecomputeTangents = Options.bEnableRecomputeTangents;
+		NewStaticMesh->CreateMeshDescription(k);
+	}
+
+	// create physics body and configure appropriately
+	if (Options.bCreatePhysicsBody)
+	{
+		NewStaticMesh->CreateBodySetup();
+		if (Options.bEnableComplexAsSimple)
+		{
+			// enable complex as simple collision to use mesh directly
+			NewStaticMesh->GetBodySetup()->CollisionTraceFlag = ECollisionTraceFlag::CTF_UseComplexAsSimple;
+		}
+	}
+
+	// add a material slot. Must always have one material slot.
+	int32 UseNumMaterialSlots = FMath::Max(1, Options.NumMaterialSlots);
+	for (int MatIdx = 0; MatIdx < UseNumMaterialSlots; MatIdx++)
+	{
+		NewStaticMesh->GetStaticMaterials().Add(FStaticMaterial());
+	}
+
+	// set materials if the count matches
+	if (Options.AssetMaterials.Num() == UseNumMaterialSlots)
+	{
+		for (int MatIdx = 0; MatIdx < UseNumMaterialSlots; MatIdx++)
+		{
+			NewStaticMesh->SetMaterial(MatIdx, Options.AssetMaterials[MatIdx]);
+		}
+	}
+
+	// if options included SourceModel meshes, copy them over
+	if (Options.SourceMeshes.MoveMeshDescriptions.Num() > 0)
+	{
+		if (ensure(Options.SourceMeshes.MoveMeshDescriptions.Num() == UseNumSourceModels))
+		{
+			for (int32 k = 0; k < UseNumSourceModels; ++k)
+			{
+				FMeshDescription* Mesh = NewStaticMesh->GetMeshDescription(k);
+				*Mesh = MoveTemp(*Options.SourceMeshes.MoveMeshDescriptions[k]);
+				NewStaticMesh->CommitMeshDescription(k);
+			}
+		}
+	}
+	else if (Options.SourceMeshes.MeshDescriptions.Num() > 0)
+	{
+		if (ensure(Options.SourceMeshes.MeshDescriptions.Num() == UseNumSourceModels))
+		{
+			for (int32 k = 0; k < UseNumSourceModels; ++k)
+			{
+				FMeshDescription* Mesh = NewStaticMesh->GetMeshDescription(k);
+				*Mesh = *Options.SourceMeshes.MeshDescriptions[k];
+				NewStaticMesh->CommitMeshDescription(k);
+			}
+		}
+	}
+	else if (Options.SourceMeshes.DynamicMeshes.Num() > 0)
+	{
+		if (ensure(Options.SourceMeshes.DynamicMeshes.Num() == UseNumSourceModels))
+		{
+			for (int32 k = 0; k < UseNumSourceModels; ++k)
+			{
+				FMeshDescription* Mesh = NewStaticMesh->GetMeshDescription(k);
+				FDynamicMeshToMeshDescription Converter;
+				Converter.Convert(Options.SourceMeshes.DynamicMeshes[k], *Mesh, !Options.bEnableRecomputeTangents);
+				NewStaticMesh->CommitMeshDescription(k);
+			}
+		}
+	}
+
+	NewStaticMesh->MarkPackageDirty();
+	if (Options.bDeferPostEditChange == false)
+	{
+		NewStaticMesh->PostEditChange();
+	}
+
+	ResultsOut.StaticMesh = NewStaticMesh;
+	return ECreateStaticMeshResult::Ok;
+}
+
+
+
diff --git a/Engine/Plugins/Experimental/MeshModelingToolset/Source/ModelingComponentsEditorOnly/Private/AssetUtils/CreateTexture2DUtil.cpp b/Engine/Plugins/Experimental/MeshModelingToolset/Source/ModelingComponentsEditorOnly/Private/AssetUtils/CreateTexture2DUtil.cpp
new file mode 100644
--- /dev/null
+++ b/Engine/Plugins/Experimental/MeshModelingToolset/Source/ModelingComponentsEditorOnly/Private/AssetUtils/CreateTexture2DUtil.cpp
@@ -0,0 +1,65 @@
+// Copyright Epic Games, Inc. All Rights Reserved. 
+
+#include "AssetUtils/CreateTexture2DUtil.h"
+
+UE::AssetUtils::ECreateTexture2DResult UE::AssetUtils::SaveGeneratedTexture2DAsset(
+	UTexture2D* GeneratedTexture,
+	FTexture2DAssetOptions& Options,
+	FTexture2DAssetResults& ResultsOut)
+{
+	// validate input texture
+	if (ensure(GeneratedTexture) == false)
+	{
+		return ECreateTexture2DResult::InvalidInputTexture;
+	}
+	if (ensure(GeneratedTexture->GetOuter() == GetTransientPackage()) == false)
+	{
+		return ECreateTexture2DResult::InvalidInputTexture;
+	}
+	if (ensure(GeneratedTexture->Source.IsValid()) == false)
+	{
+		return ECreateTexture2DResult::InvalidInputTexture;
+	}
+
+	FString NewObjectName = FPackageName::GetLongPackageAssetName(Options.NewAssetPath);
+
+	UPackage* UsePackage = nullptr;
+	if (Options.UsePackage != nullptr)
+	{
+		UsePackage = Options.UsePackage;
+	}
+	else
+	{
+		UsePackage = CreatePackage(*Options.NewAssetPath);
+	}
+	if (ensure(UsePackage != nullptr) == false)
+	{
+		return ECreateTexture2DResult::InvalidPackage;
+	}
+
+	// move texture from Transient package to real package
+	bool bRenameOK = GeneratedTexture->Rename(*NewObjectName, UsePackage, REN_None);
+	if (ensure(bRenameOK) == false)
+	{
+		return ECreateTexture2DResult::NameError;
+	}
+
+	// remove transient flag, add public/standalone/transactional
+	GeneratedTexture->ClearFlags(RF_Transient);
+	GeneratedTexture->SetFlags(RF_Public | RF_Standalone | RF_Transactional);
+	GeneratedTexture->MarkPackageDirty();
+
+	// do we need to Modify() it? we are not doing any undo/redo
+	GeneratedTexture->Modify();
+
+	// force texture to rebuild (can we defer this?)
+	GeneratedTexture->UpdateResource();
+
+	if (Options.bDeferPostEditChange == false)
+	{
+		GeneratedTexture->PostEditChange();		// this may be necessary if any Materials are using this texture
+	}
+
+	ResultsOut.Texture = GeneratedTexture;
+	return ECreateTexture2DResult::Ok;
+}
\ No newline at end of file
diff --git a/Engine/Plugins/Experimental/MeshModelingToolset/Source/ModelingComponentsEditorOnly/Private/ModelingComponentsEditorOnlyModule.cpp b/Engine/Plugins/Experimental/MeshModelingToolset/Source/ModelingComponentsEditorOnly/Private/ModelingComponentsEditorOnlyModule.cpp
new file mode 100644
--- /dev/null
+++ b/Engine/Plugins/Experimental/MeshModelingToolset/Source/ModelingComponentsEditorOnly/Private/ModelingComponentsEditorOnlyModule.cpp
@@ -0,0 +1,20 @@
+// Copyright Epic Games, Inc. All Rights Reserved.
+
+#include "ModelingComponentsEditorOnlyModule.h"
+
+#define LOCTEXT_NAMESPACE "FModelingComponentsEditorOnlyModule"
+
+void FModelingComponentsEditorOnlyModule::StartupModule()
+{
+	// This code will execute after your module is loaded into memory; the exact timing is specified in the .uplugin file per-module
+}
+
+void FModelingComponentsEditorOnlyModule::ShutdownModule()
+{
+	// This function may be called during shutdown to clean up your module.  For modules that support dynamic reloading,
+	// we call this function before unloading the module.
+}
+
+#undef LOCTEXT_NAMESPACE
+	
+IMPLEMENT_MODULE(FModelingComponentsEditorOnlyModule, ModelingComponentsEditorOnly)
\ No newline at end of file
diff --git a/Engine/Plugins/Experimental/MeshModelingToolset/Source/ModelingComponentsEditorOnly/Public/AssetUtils/CreateMaterialUtil.h b/Engine/Plugins/Experimental/MeshModelingToolset/Source/ModelingComponentsEditorOnly/Public/AssetUtils/CreateMaterialUtil.h
new file mode 100644
--- /dev/null
+++ b/Engine/Plugins/Experimental/MeshModelingToolset/Source/ModelingComponentsEditorOnly/Public/AssetUtils/CreateMaterialUtil.h
@@ -0,0 +1,65 @@
+// Copyright Epic Games, Inc. All Rights Reserved.
+
+#pragma once
+
+#include "MathUtil.h"
+#include "VectorTypes.h"
+#include "Image/ImageDimensions.h"
+#include "Image/ImageBuilder.h"
+
+class UMaterial;
+class UMaterialInterface;
+
+namespace UE
+{
+namespace AssetUtils
+{
+	/**
+	 * Result enum returned by Create() functions below to indicate succes/error conditions
+	 */
+	enum class ECreateMaterialResult
+	{
+		Ok = 0,
+		InvalidPackage = 1,
+		InvalidBaseMaterial = 2,
+		NameError = 3,
+		DuplicateFailed = 4,
+
+		UnknownError = 100
+	};
+
+	/**
+	 * Options for new UMaterial asset created by Create() functions below.
+	 */
+	struct MODELINGCOMPONENTSEDITORONLY_API FMaterialAssetOptions
+	{
+		// this package will be used if it is not null, otherwise a new package will be created at NewAssetPath
+		UPackage* UsePackage = nullptr;
+		FString NewAssetPath;
+
+		bool bDeferPostEditChange = false;
+	};
+
+	/**
+	 * Output information about a newly-created UMaterial, returned by Create functions below.
+	 * Some fields may be null, if the relevant function did not create that type of object
+	 */
+	struct MODELINGCOMPONENTSEDITORONLY_API FMaterialAssetResults
+	{
+		UMaterial* NewMaterial = nullptr;
+	};
+
+	/**
+	 * Create a new UMaterial Asset by duplicating the parent UMaterial of BaseMaterial
+	 * @param Options defines configuration for the new UMaterial Asset
+	 * @param ResultOut new UMaterial is returned here
+	 * @return Ok, or error flag if some part of the creation process failed. On failure, no Asset is created.
+	 */
+	MODELINGCOMPONENTSEDITORONLY_API ECreateMaterialResult CreateDuplicateMaterial(
+		UMaterialInterface* BaseMaterial,
+		FMaterialAssetOptions& Options,
+		FMaterialAssetResults& ResultsOut);
+
+}  // end namespace UE
+}  // end namespace AssetUtils
+
diff --git a/Engine/Plugins/Experimental/MeshModelingToolset/Source/ModelingComponentsEditorOnly/Public/AssetUtils/CreateStaticMeshUtil.h b/Engine/Plugins/Experimental/MeshModelingToolset/Source/ModelingComponentsEditorOnly/Public/AssetUtils/CreateStaticMeshUtil.h
new file mode 100644
--- /dev/null
+++ b/Engine/Plugins/Experimental/MeshModelingToolset/Source/ModelingComponentsEditorOnly/Public/AssetUtils/CreateStaticMeshUtil.h
@@ -0,0 +1,107 @@
+// Copyright Epic Games, Inc. All Rights Reserved.
+
+#pragma once
+
+#include "CoreMinimal.h"
+#include "DynamicMesh3.h"
+
+class UPackage;
+class UStaticMesh;
+class UStaticMeshComponent;
+class UStaticMeshActor;
+struct FMeshDescription;
+class UMaterialInterface;
+
+namespace UE
+{
+namespace AssetUtils
+{
+
+	using namespace UE::Geometry;
+
+	/**
+	 * Result enum returned by Create functions below to indicate succes/error conditions
+	 */
+	enum class ECreateStaticMeshResult
+	{
+		Ok = 0,
+		InvalidPackage = 1,
+
+		UnknownError = 100
+	};
+
+	/**
+	 * Set of input meshes for StaticMesh Create() functions below.
+	 * Only one of the arrays should be initialized, and it's size should be equal to 
+	 * the number of LODs specified for the Asset.
+	 */
+	struct MODELINGCOMPONENTSEDITORONLY_API FStaticMeshAssetMeshes
+	{
+		TArray<const UE::Geometry::FDynamicMesh3*> DynamicMeshes;
+		TArray<const FMeshDescription*> MeshDescriptions;
+		TArray<FMeshDescription*> MoveMeshDescriptions;		// these will be MoveTemp()'d into the StaticMesh, so if initialized, this array will be invalid after Create call()
+	};
+
+	/**
+	 * Options for new UStaticMesh asset created by Create() functions below.
+	 */
+	struct MODELINGCOMPONENTSEDITORONLY_API FStaticMeshAssetOptions
+	{
+		// this package will be used if it is not null, otherwise a new package will be created at NewAssetPath
+		UPackage* UsePackage = nullptr;
+		FString NewAssetPath;
+
+		// number of SourceModels (ie LODs)
+		int32 NumSourceModels = 1;
+		// number of Material Slots desired on the Asset. At least one will always be created.
+		int32 NumMaterialSlots = 1;
+
+		// Controls whether or not the RecomputeNormals option will be enabled on the Asset
+		bool bEnableRecomputeNormals = false;
+		// Controls whether or not the RecomputeTangents option will be enabled on the Asset
+		bool bEnableRecomputeTangents = true;
+
+		// Controls whether the UBodySetup on the Asset will be created (generally should be true)
+		bool bCreatePhysicsBody = true;
+		// If true, Asset will be configured to UseComplexAsSimple mode
+		bool bEnableComplexAsSimple = true;
+		// TODO: option for simple collision geo
+
+		// Optional list of materials to initialize the Asset with. If defined, size must be the same as NumMaterialSlots
+		TArray<UMaterialInterface*> AssetMaterials;
+
+		// Optional list of meshes to use to initialize SourceModels. If defined, only one array should be non-empty, and must be the same length as NumSourceModels.
+		FStaticMeshAssetMeshes SourceMeshes;
+
+		// by default, PostEditChange() will be called to rebuild mesh, set true to skip this call
+		bool bDeferPostEditChange = false;
+	};
+
+	/**
+	 * Output information about a newly-created StaticMesh, returned by Create functions below.
+	 * Some fields may be null, if the relevant function did not create that type of object
+	 */
+	struct MODELINGCOMPONENTSEDITORONLY_API FStaticMeshResults
+	{
+		/** UStaticMesh asset that was created */
+		UStaticMesh* StaticMesh = nullptr;
+		/** New Component that was created that references the Asset (or null, if no component is created) */
+		UStaticMeshComponent* Component = nullptr;
+		/** New Actor that was created for the Component (or null, if no component is created) */
+		UStaticMeshActor* Actor = nullptr;
+	};
+
+
+	/**
+	 * Create a new UStaticMesh Asset based on the input Options
+	 * @param Options defines configuration for the new UStaticMesh Asset
+	 * @param ResultOut new StaticMesh is returned here
+	 * @return Ok, or error flag if some part of the creation process failed. On failure, no Asset is created.
+	 */
+	MODELINGCOMPONENTSEDITORONLY_API ECreateStaticMeshResult CreateStaticMeshAsset(
+		FStaticMeshAssetOptions& Options,
+		FStaticMeshResults& ResultsOut);
+
+
+}
+}
diff --git a/Engine/Plugins/Experimental/MeshModelingToolset/Source/ModelingComponentsEditorOnly/Public/AssetUtils/CreateTexture2DUtil.h b/Engine/Plugins/Experimental/MeshModelingToolset/Source/ModelingComponentsEditorOnly/Public/AssetUtils/CreateTexture2DUtil.h
new file mode 100644
--- /dev/null
+++ b/Engine/Plugins/Experimental/MeshModelingToolset/Source/ModelingComponentsEditorOnly/Public/AssetUtils/CreateTexture2DUtil.h
@@ -0,0 +1,63 @@
+// Copyright Epic Games, Inc. All Rights Reserved.
+
+#pragma once
+
+#include "MathUtil.h"
+#include "VectorTypes.h"
+#include "Engine/Texture2D.h"
+
+
+
+namespace UE
+{
+namespace AssetUtils
+{
+	/**
+	 * Result enum returned by Create() functions below to indicate succes/error conditions
+	 */
+	enum class ECreateTexture2DResult
+	{
+		Ok = 0,
+		InvalidPackage = 1,
+		InvalidInputTexture = 2,
+		NameError = 3,
+
+		UnknownError = 100
+	};
+
+
+	/**
+	 * Options for new UTexture asset created by Create() functions below.
+	 */
+	struct MODELINGCOMPONENTSEDITORONLY_API FTexture2DAssetOptions
+	{
+		// this package will be used if it is not null, otherwise a new package will be created at NewAssetPath
+		UPackage* UsePackage = nullptr;
+		FString NewAssetPath;
+
+		bool bDeferPostEditChange = false;
+	};
+
+	/**
+	 * Output information about a newly-created UTexture2D, returned by Create functions below.
+	 * Some fields may be null, if the relevant function did not create that type of object
+	 */
+	struct MODELINGCOMPONENTSEDITORONLY_API FTexture2DAssetResults
+	{
+		UTexture2D* Texture = nullptr;
+	};
+
+	/**
+	 * Convert a temporary UTexture2D in the Transient package into a serialized Asset
+	 * @param Options defines configuration for the new UTexture2D Asset
+	 * @param ResultOut new UTexture2D is returned here
+	 * @return Ok, or error flag if some part of the creation process failed. On failure, no Asset is created.
+	 */
+	MODELINGCOMPONENTSEDITORONLY_API ECreateTexture2DResult SaveGeneratedTexture2DAsset(
+		UTexture2D* GeneratedTexture,
+		FTexture2DAssetOptions& Options,
+		FTexture2DAssetResults& ResultsOut);
+
+}
+}
+
diff --git a/Engine/Plugins/Experimental/MeshModelingToolset/Source/ModelingComponentsEditorOnly/Public/ModelingComponentsEditorOnlyModule.h b/Engine/Plugins/Experimental/MeshModelingToolset/Source/ModelingComponentsEditorOnly/Public/ModelingComponentsEditorOnlyModule.h
new file mode 100644
--- /dev/null
+++ b/Engine/Plugins/Experimental/MeshModelingToolset/Source/ModelingComponentsEditorOnly/Public/ModelingComponentsEditorOnlyModule.h
@@ -0,0 +1,15 @@
+// Copyright Epic Games, Inc. All Rights Reserved.
+
+#pragma once
+
+#include "CoreMinimal.h"
+#include "Modules/ModuleManager.h"
+
+class FModelingComponentsEditorOnlyModule : public IModuleInterface
+{
+public:
+
+	/** IModuleInterface implementation */
+	virtual void StartupModule() override;
+	virtual void ShutdownModule() override;
+};
