Add GeometryProcessingInterfaces module in /Developer. This module exists to define abstract interfaces to high-level geometry operations, which the core Engine/Editor can see, but which can be implemented by plugins. Currently only ApproximateActors interface is defined, which will allow the ProxyLOD and HLOD systems to use implementations that depend on the MeshModelingToolset and GeometryProcessing plugins.

Add GeometryProcessingAdapters module to MeshModelingToolset, this is an Editor-only module that implements GeometryProcessingInterfaces APIs

Add ApproximateActors implementation in GeometryProcessingAdapters that uses GeometryProcessing (FastWindingNumber / MeshMorphology / Simplify / UVGen / Tangents) and a new material-baking process based on render captures to generate a new StaticMesh Asset / Material / Textures that approximate the input Actor set.

#rnx

diff --git a/Engine/Plugins/Experimental/MeshModelingToolset/MeshModelingToolset.uplugin b/Engine/Plugins/Experimental/MeshModelingToolset/MeshModelingToolset.uplugin
--- a/Engine/Plugins/Experimental/MeshModelingToolset/MeshModelingToolset.uplugin
+++ b/Engine/Plugins/Experimental/MeshModelingToolset/MeshModelingToolset.uplugin
@@ -47,6 +47,12 @@
 			"Type": "Editor",
 			"LoadingPhase": "Default",
 			"WhitelistPlatforms": [ "Win64" ]
+		},
+		{
+			"Name": "GeometryProcessingAdapters",
+			"Type": "Editor",
+			"LoadingPhase": "Default",
+			"WhitelistPlatforms": [ "Win64" ]
 		}
 	],
 	"Plugins": [
diff --git a/Engine/Plugins/Experimental/MeshModelingToolset/Source/GeometryProcessingAdapters/GeometryProcessingAdapters.Build.cs b/Engine/Plugins/Experimental/MeshModelingToolset/Source/GeometryProcessingAdapters/GeometryProcessingAdapters.Build.cs
new file mode 100644
--- /dev/null
+++ b/Engine/Plugins/Experimental/MeshModelingToolset/Source/GeometryProcessingAdapters/GeometryProcessingAdapters.Build.cs
@@ -0,0 +1,74 @@
+// Copyright Epic Games, Inc. All Rights Reserved.
+
+using UnrealBuildTool;
+
+public class GeometryProcessingAdapters : ModuleRules
+{
+	public GeometryProcessingAdapters(ReadOnlyTargetRules Target) : base(Target)
+	{
+		PCHUsage = ModuleRules.PCHUsageMode.UseExplicitOrSharedPCHs;
+		
+		PublicIncludePaths.AddRange(
+			new string[] {
+				// ... add public include paths required here ...
+			}
+			);
+				
+		
+		PrivateIncludePaths.AddRange(
+			new string[] {
+				// ... add other private include paths required here ...
+			}
+			);
+			
+		
+		PublicDependencyModuleNames.AddRange(
+			new string[]
+			{
+				"Core",
+				"GeometryProcessingInterfaces",
+				
+                "InteractiveToolsFramework",
+                "MeshDescription",
+				"StaticMeshDescription",
+				"ModelingOperators",
+				"GeometricObjects",
+				"DynamicMesh",
+				"MeshConversion",
+				"GeometryAlgorithms", // required for constrained Delaunay triangulation
+				"MeshUtilitiesCommon", // required by uvlayoutop
+				"ModelingComponents",
+				"ModelingComponentsEditorOnly",
+				"ModelingOperators",
+				"ModelingOperatorsEditorOnly"
+				// ... add other public dependencies that you statically link with here ...
+			}
+            );
+			
+		
+		PrivateDependencyModuleNames.AddRange(
+			new string[]
+			{
+				"CoreUObject",
+				"Engine",
+
+				"MeshBuilder",
+				"MeshUtilitiesCommon",  
+				"MeshReductionInterface", // for UE standard simplification 
+				"MeshUtilities",			// for tangents calculation
+			}
+			);
+		
+		
+		DynamicallyLoadedModuleNames.AddRange(
+			new string[]
+			{
+				// ... add any modules that your module loads dynamically here ...
+			}
+			);
+
+
+		AddEngineThirdPartyPrivateStaticDependencies(Target, "MikkTSpace");
+
+	}
+}
diff --git a/Engine/Plugins/Experimental/MeshModelingToolset/Source/GeometryProcessingAdapters/Private/ApproximateActorsImpl.cpp b/Engine/Plugins/Experimental/MeshModelingToolset/Source/GeometryProcessingAdapters/Private/ApproximateActorsImpl.cpp
new file mode 100644
--- /dev/null
+++ b/Engine/Plugins/Experimental/MeshModelingToolset/Source/GeometryProcessingAdapters/Private/ApproximateActorsImpl.cpp
@@ -0,0 +1,452 @@
+// Copyright Epic Games, Inc. All Rights Reserved.
+
+#include "GeometryProcessing/ApproximateActorsImpl.h"
+
+#include "Scene/MeshSceneAdapter.h"
+#include "DynamicMesh3.h"
+#include "MeshNormals.h"
+#include "MeshTangents.h"
+#include "MeshSimplification.h"
+#include "MeshConstraintsUtil.h"
+#include "Implicit/Solidify.h"
+#include "Implicit/Morphology.h"
+#include "CleaningOps/RemoveOccludedTrianglesOp.h"
+#include "ParameterizationOps/ParameterizeMeshOp.h"
+#include "Parameterization/MeshUVPacking.h"
+#include "MeshQueries.h"
+
+#include "AssetUtils/CreateStaticMeshUtil.h"
+#include "AssetUtils/CreateTexture2DUtil.h"
+#include "AssetUtils/CreateMaterialUtil.h"
+#include "UObject/UObjectGlobals.h"		// for CreatePackage
+
+#include "ImageUtils.h"
+#include "Image/ImageInfilling.h"
+#include "Sampling/MeshGenericWorldPositionBaker.h"
+#include "Scene/SceneCapturePhotoSet.h"
+#include "AssetUtils/Texture2DBuilder.h"
+
+#include "Materials/Material.h"
+
+#include "Misc/ScopedSlowTask.h"
+
+using namespace UE::Geometry;
+using namespace UE::AssetUtils;
+
+#define LOCTEXT_NAMESPACE "ApproximateActorsImpl"
+
+struct FGeneratedResultTextures
+{
+	UTexture2D* BaseColorMap;
+	UTexture2D* RoughnessMap;
+	UTexture2D* MetallicMap;
+	UTexture2D* SpecularMap;
+	UTexture2D* EmissiveMap;
+	UTexture2D* NormalMap;
+};
+
+
+static void BakeTexturesFromPhotoCapture(
+	const TArray<AActor*>& Actors, 
+	const IGeometryProcessing_ApproximateActors::FOptions& Options, 
+	FGeneratedResultTextures& GeneratedTextures,
+	const FDynamicMesh3* WorldTargetMesh,
+	const FMeshTangentsd* MeshTangents
+	)
+{
+	int32 UVLayer = 0;
+	double FieldOfView = Options.FieldOfViewDegrees;
+	double NearPlaneDist = Options.NearPlaneDist;
+	double RayOffsetHackDist = (double)(100.0f * FMathf::ZeroTolerance);
+
+	FImageDimensions CaptureDimensions(Options.RenderCaptureImageSize, Options.RenderCaptureImageSize);
+	FImageDimensions OutputDimensions(Options.TextureImageSize, Options.TextureImageSize);
+
+	FSceneCapturePhotoSet SceneCapture;
+	SceneCapture.SetCaptureSceneActors(Actors[0]->GetWorld(), Actors);
+
+	SceneCapture.AddStandardExteriorCapturesFromBoundingBox(
+		CaptureDimensions, FieldOfView, NearPlaneDist,
+		true, true, true);
+
+	FScopedSlowTask Progress(8.f, LOCTEXT("BakingTextures", "Baking Textures..."));
+	Progress.MakeDialog(true);
+
+	Progress.EnterProgressFrame(1.f, LOCTEXT("BakingSetup", "Setup..."));
+
+	FDynamicMeshAABBTree3 Spatial(WorldTargetMesh, true);
+
+	FMeshImageBakingCache TempBakeCache;
+	TempBakeCache.SetDetailMesh(WorldTargetMesh, &Spatial);
+	TempBakeCache.SetBakeTargetMesh(WorldTargetMesh);
+	TempBakeCache.SetDimensions(OutputDimensions);
+	TempBakeCache.SetUVLayer(UVLayer);
+	TempBakeCache.SetThickness(0.1);
+	TempBakeCache.SetCorrespondenceStrategy(FMeshImageBakingCache::ECorrespondenceStrategy::Identity);
+	TempBakeCache.ValidateCache();
+
+	Progress.EnterProgressFrame(1.f, LOCTEXT("BakingBaseColor", "Baking Base Color..."));
+
+	auto VisibilityFunction = [&Spatial, RayOffsetHackDist](const FVector3d& SurfPos, const FVector3d& ImagePosWorld)
+	{
+		FVector3d RayDir = ImagePosWorld - SurfPos;
+		double Dist = Normalize(RayDir);
+		FVector3d RayOrigin = SurfPos + RayOffsetHackDist * RayDir;
+		int32 HitTID = Spatial.FindNearestHitTriangle(FRay3d(RayOrigin, RayDir), IMeshSpatial::FQueryOptions(Dist));
+		return (HitTID == IndexConstants::InvalidID);
+	};
+
+	FSceneCapturePhotoSet::FSceneSample DefaultSample;
+	DefaultSample.BaseColor = FVector4f(0, -1, 0, 0);
+
+	FMeshGenericWorldPositionColorBaker BaseColorBaker;
+	BaseColorBaker.SetCache(&TempBakeCache);
+	BaseColorBaker.ColorSampleFunction = [&](FVector3d Position, FVector3d Normal) {
+		FSceneCapturePhotoSet::FSceneSample Sample = DefaultSample;
+		SceneCapture.ComputeSample(FRenderCaptureTypeFlags::BaseColor(),
+			Position, Normal, VisibilityFunction, Sample);
+		return Sample.BaseColor;
+	};
+	BaseColorBaker.Bake();
+
+	// find "hole" pixels
+	TArray<FVector2i> MissingPixels;
+	TUniquePtr<TImageBuilder<FVector4f>> ColorImage = BaseColorBaker.TakeResult();
+	TempBakeCache.FindSamplingHoles([&](const FVector2i& Coords)
+	{
+		return ColorImage->GetPixel(Coords) == DefaultSample.BaseColor;
+	}, MissingPixels);
+
+
+	TMarchingPixelInfill<FVector4f> Infill;
+	Infill.ComputeInfill(*ColorImage, MissingPixels, DefaultSample.BaseColor,
+		[](FVector4f SumValue, int32 Count) {
+		float InvSum = (Count == 0) ? 1.0f : (1.0f / Count);
+		return FVector4f(SumValue.X * InvSum, SumValue.Y * InvSum, SumValue.Z * InvSum, 1.0f);
+	});
+
+
+	auto ProcessChannelFunc = [&](ERenderCaptureType CaptureType)
+	{
+		FVector4f DefaultValue(0, 0, 0, 0);
+		FMeshGenericWorldPositionColorBaker ChannelBaker;
+		ChannelBaker.SetCache(&TempBakeCache);
+		ChannelBaker.ColorSampleFunction = [&](FVector3d Position, FVector3d Normal) {
+			FSceneCapturePhotoSet::FSceneSample Sample = DefaultSample;
+			SceneCapture.ComputeSample(FRenderCaptureTypeFlags::Single(CaptureType), Position, Normal, VisibilityFunction, Sample);
+			return Sample.GetValue(CaptureType);
+		};
+		ChannelBaker.Bake();
+		TUniquePtr<TImageBuilder<FVector4f>> Image = ChannelBaker.TakeResult();
+
+		Infill.ApplyInfill(*Image,
+			[](FVector4f SumValue, int32 Count) {
+			float InvSum = (Count == 0) ? 1.0f : (1.0f / Count);
+			return FVector4f(SumValue.X * InvSum, SumValue.Y * InvSum, SumValue.Z * InvSum, 1.0f);
+		});
+
+		return MoveTemp(Image);
+	};
+
+	Progress.EnterProgressFrame(1.f, LOCTEXT("BakingRoughness", "Baking Roughness..."));
+	TUniquePtr<TImageBuilder<FVector4f>> RoughnessImage = ProcessChannelFunc(ERenderCaptureType::Roughness);
+	Progress.EnterProgressFrame(1.f, LOCTEXT("BakingMetallic", "Baking Metallic..."));
+	TUniquePtr<TImageBuilder<FVector4f>> MetallicImage = ProcessChannelFunc(ERenderCaptureType::Metallic);
+	Progress.EnterProgressFrame(1.f, LOCTEXT("BakingSpecular", "Baking Specular..."));
+	TUniquePtr<TImageBuilder<FVector4f>> SpecularImage = ProcessChannelFunc(ERenderCaptureType::Specular);
+	Progress.EnterProgressFrame(1.f, LOCTEXT("BakingEmissive", "Baking Emissive..."));
+	TUniquePtr<TImageBuilder<FVector4f>> EmissiveImage = ProcessChannelFunc(ERenderCaptureType::Emissive);
+
+	Progress.EnterProgressFrame(1.f, LOCTEXT("BakingNormals", "Baking Normals..."));
+
+	// no infill on normal map for now, doesn't make sense to do after mapping to tangent space!
+	//  (should we build baked normal map in world space, and then resample to tangent space??)
+	FVector4f DefaultNormalValue(0, 0, 1, 1);
+	FMeshGenericWorldPositionNormalBaker NormalMapBaker;
+	NormalMapBaker.SetCache(&TempBakeCache);
+	NormalMapBaker.BaseMeshTangents = MeshTangents;
+	NormalMapBaker.NormalSampleFunction = [&](FVector3d Position, FVector3d Normal) {
+		FSceneCapturePhotoSet::FSceneSample Sample = DefaultSample;
+		SceneCapture.ComputeSample(FRenderCaptureTypeFlags::WorldNormal(),
+			Position, Normal, VisibilityFunction, Sample);
+		FVector4f NormalColor = Sample.WorldNormal;
+		float x = (NormalColor.X - 0.5f) * 2.0f;
+		float y = (NormalColor.Y - 0.5f) * 2.0f;
+		float z = (NormalColor.Z - 0.5f) * 2.0f;
+		return FVector3f(x, y, z);
+	};
+	NormalMapBaker.Bake();
+	TUniquePtr<TImageBuilder<FVector3f>> NormalImage = NormalMapBaker.TakeResult();
+
+
+	// build textures
+	Progress.EnterProgressFrame(1.f, LOCTEXT("BuildingTextures", "Building Textures..."));
+	{
+		FScopedSlowTask BuildTexProgress(6.f, LOCTEXT("BuildingTextures", "Building Textures..."));
+		BuildTexProgress.MakeDialog(true);
+		BuildTexProgress.EnterProgressFrame(1.f);
+		GeneratedTextures.BaseColorMap = FTexture2DBuilder::BuildTextureFromImage(*ColorImage, FTexture2DBuilder::ETextureType::Color, true, false);
+		BuildTexProgress.EnterProgressFrame(1.f);
+		GeneratedTextures.RoughnessMap = FTexture2DBuilder::BuildTextureFromImage(*RoughnessImage, FTexture2DBuilder::ETextureType::Roughness, false, false);
+		BuildTexProgress.EnterProgressFrame(1.f);
+		GeneratedTextures.MetallicMap = FTexture2DBuilder::BuildTextureFromImage(*MetallicImage, FTexture2DBuilder::ETextureType::Metallic, false, false);
+		BuildTexProgress.EnterProgressFrame(1.f);
+		GeneratedTextures.SpecularMap = FTexture2DBuilder::BuildTextureFromImage(*SpecularImage, FTexture2DBuilder::ETextureType::Specular, false, false);
+		BuildTexProgress.EnterProgressFrame(1.f);
+		GeneratedTextures.EmissiveMap = FTexture2DBuilder::BuildTextureFromImage(*EmissiveImage, FTexture2DBuilder::ETextureType::Color, true, false);
+		BuildTexProgress.EnterProgressFrame(1.f);
+		GeneratedTextures.NormalMap = FTexture2DBuilder::BuildTextureFromImage(*NormalImage, FTexture2DBuilder::ETextureType::NormalMap, false, false);
+	}
+}
+
+void FApproximateActorsImpl::ApproximateActors(const TArray<AActor*>& Actors, const FOptions& Options, FResults& ResultsOut)
+{
+	int32 ActorClusters = 1;
+	FScopedSlowTask Progress(1.f, LOCTEXT("ApproximatingActors", "Generating Actor Approximation..."));
+	Progress.MakeDialog(true);
+	Progress.EnterProgressFrame(1.f);
+	GenerateApproximationForActorSet(Actors, Options, ResultsOut);
+}
+
+
+void FApproximateActorsImpl::GenerateApproximationForActorSet(const TArray<AActor*>& Actors, const FOptions& Options, FResults& ResultsOut)
+{
+	//
+	// Future Optimizations
+	// 	   - can do most of the mesh processing at the same time as capturing the photo set (if that matters)
+	//     - some parts of mesh gen can be done simultaneously (maybe?)
+	//
+
+	FScopedSlowTask Progress(10.f, LOCTEXT("ApproximatingActors", "Generating Actor Approximation..."));
+
+	Progress.EnterProgressFrame(1.f, LOCTEXT("BuildingScene", "Building Scene..."));
+
+	FMeshSceneAdapter Scene;
+	Scene.AddActors(Actors);
+
+	TArray<FVector3d> SeedPoints;
+	Scene.CollectMeshSeedPoints(SeedPoints);
+	FAxisAlignedBox3d SceneBounds = Scene.GetBoundingBox();
+
+	// calculate a voxel size based on target world-space approximation accuracy
+	float WorldBoundsSize = SceneBounds.DiagonalLength();
+	float ApproxAccuracy = Options.WorldSpaceApproximationAccuracyMeters * 100.0;		// convert to cm
+	int32 VoxelDimTarget = (int)(WorldBoundsSize / ApproxAccuracy);
+	if (VoxelDimTarget < 64)
+	{
+		VoxelDimTarget = 64;		// use a sane minimum in case the parameter is super-wrong
+	}
+
+	// avoid insane memory usage
+	if (ensure(VoxelDimTarget < Options.ClampVoxelDimension) == false)
+	{
+		UE_LOG(LogTemp, Warning, TEXT("FApproximateActorsImpl - very large voxel size %d clamped to %d"), VoxelDimTarget, Options.ClampVoxelDimension);
+		VoxelDimTarget = Options.ClampVoxelDimension;
+	}
+
+	Progress.EnterProgressFrame(1.f, LOCTEXT("GeneratingMesh", "Generating Mesh..."));
+
+	FWindingNumberBasedSolidify Solidify(
+		[&Scene](const FVector3d& Position) { return Scene.MaxFastWindingNumber(Position); },
+		SceneBounds, SeedPoints);
+	Solidify.SetCellSizeAndExtendBounds(SceneBounds, 0, VoxelDimTarget);
+
+	FDynamicMesh3 SolidMesh(&Solidify.Generate());
+	SolidMesh.DiscardAttributes();
+	FDynamicMesh3* CurResultMesh = &SolidMesh;		// this pointer will be updated as we recompute the mesh
+
+	Progress.EnterProgressFrame(1.f, LOCTEXT("ClosingMesh", "Topological Operations..."));
+
+	// do topological closure to fix small gaps/etc
+	FDynamicMesh3 MorphologyMesh;
+	if (Options.bApplyMorphology)
+	{
+		double MorphologyDistance = Options.MorphologyDistanceMeters * 100.0;		// convert to cm
+		FAxisAlignedBox3d MorphologyBounds = CurResultMesh->GetBounds();
+		FDynamicMeshAABBTree3 MorphologyBVTree(CurResultMesh);
+		TImplicitMorphology<FDynamicMesh3> ImplicitMorphology;
+		ImplicitMorphology.MorphologyOp = TImplicitMorphology<FDynamicMesh3>::EMorphologyOp::Close;
+		ImplicitMorphology.Source = CurResultMesh;
+		ImplicitMorphology.SourceSpatial = &MorphologyBVTree;
+		ImplicitMorphology.SetCellSizesAndDistance(MorphologyBounds, MorphologyDistance, VoxelDimTarget, VoxelDimTarget);
+		MorphologyMesh = FDynamicMesh3(&ImplicitMorphology.Generate());
+		MorphologyMesh.DiscardAttributes();
+		CurResultMesh = &MorphologyMesh;
+	}
+
+	// if mesh has no triangles, something has gone wrong
+	if (CurResultMesh == nullptr || CurResultMesh->TriangleCount() == 0)
+	{
+		ResultsOut.ResultCode = EResultCode::MeshGenerationFailed;
+		return;
+	}
+
+	Progress.EnterProgressFrame(1.f, LOCTEXT("RemoveHidden", "Removing Hidden Geometry..."));
+
+	//  TODO: remove interior components that will never be visible (Not the same as hidden triangles)
+
+	Progress.EnterProgressFrame(1.f, LOCTEXT("SimplifyingMesh", "Simplifying Mesh..."));
+
+	int32 SimplifyTargetTriCount = Options.FixedTriangleCount;
+	if (Options.MeshSimplificationPolicy == ESimplificationPolicy::TrianglesPerUnitSqMeter)
+	{
+		FVector2d VolArea = TMeshQueries<FDynamicMesh3>::GetVolumeArea(*CurResultMesh);
+		double MeshAreaMeterSqr = VolArea.Y * 0.0001;
+		SimplifyTargetTriCount = MeshAreaMeterSqr * Options.SimplificationTargetMetric;
+	}
+
+	FVolPresMeshSimplification Simplifier(CurResultMesh);
+	Simplifier.ProjectionMode = FVolPresMeshSimplification::ETargetProjectionMode::NoProjection;
+	Simplifier.DEBUG_CHECK_LEVEL = 0;
+	Simplifier.bAllowSeamCollapse = false;
+	Simplifier.SimplifyToTriangleCount(SimplifyTargetTriCount);
+
+	// re-enable attributes
+	CurResultMesh->EnableAttributes();
+
+	//  TODO: clip hidden triangles against occluder geo like landscape (should be no hidden coming out of meshing)
+
+	// compute normals
+	FMeshNormals::InitializeOverlayToPerVertexNormals(CurResultMesh->Attributes()->PrimaryNormals());
+
+	// exit here if we are just generating a merged collision mesh
+	if (Options.BasePolicy == EApproximationPolicy::CollisionMesh)
+	{
+		EmitGeneratedMeshAsset(Actors, Options, ResultsOut, CurResultMesh, nullptr);
+		ResultsOut.ResultCode = EResultCode::Success;
+		return;
+	}
+
+	Progress.EnterProgressFrame(1.f, LOCTEXT("ComputingUVs", "Computing UVs..."));
+
+	// compute UVs
+	TSharedPtr<FDynamicMesh3, ESPMode::ThreadSafe> UVInputMesh = MakeShared<FDynamicMesh3, ESPMode::ThreadSafe>();
+	*UVInputMesh = MoveTemp(*CurResultMesh);
+	FParameterizeMeshOp ParameterizeMeshOp;
+	ParameterizeMeshOp.Stretch = 0.1;
+	ParameterizeMeshOp.NumCharts = 0;
+	ParameterizeMeshOp.InputMesh = UVInputMesh;
+	ParameterizeMeshOp.IslandMode = EParamOpIslandMode::Auto;
+	ParameterizeMeshOp.UnwrapType = EParamOpUnwrapType::MinStretch;
+	FProgressCancel UVProgressCancel;
+	ParameterizeMeshOp.CalculateResult(&UVProgressCancel);
+
+	TUniquePtr<FDynamicMesh3> FinalMesh = ParameterizeMeshOp.ExtractResult();
+
+	Progress.EnterProgressFrame(1.f, LOCTEXT("PackingUVs", "Packing UVs..."));
+
+	// repack UVs
+	FDynamicMeshUVOverlay* RepackUVLayer = FinalMesh->Attributes()->PrimaryUV();
+	RepackUVLayer->SplitBowties();
+	FDynamicMeshUVPacker Packer(RepackUVLayer);
+	Packer.TextureResolution = Options.TextureImageSize / 4;		// maybe too conservative? We don't have gutter control currently.
+	Packer.GutterSize = 1.0;		// not clear this works
+	Packer.bAllowFlips = false;
+	bool bOK = Packer.StandardPack();
+	ensure(bOK);
+
+	Progress.EnterProgressFrame(1.f, LOCTEXT("ComputingTangents", "Computing Tangents..."));
+
+	// compute tangents
+	FMeshTangentsd Tangents(FinalMesh.Get());
+	FComputeTangentsOptions TangentsOptions;
+	TangentsOptions.bAveraged = true;
+	Tangents.ComputeTriVertexTangents(
+		FinalMesh->Attributes()->PrimaryNormals(),
+		FinalMesh->Attributes()->PrimaryUV(),
+		TangentsOptions);
+
+	// TODO: copy/calc into mesh tangents attrib, so it can be transferred to MeshDescription
+	 
+	Progress.EnterProgressFrame(1.f, LOCTEXT("BakingTextures", "Baking Textures..."));
+
+	// bake textures for Actor
+	FGeneratedResultTextures GeneratedTextures;
+	BakeTexturesFromPhotoCapture(Actors,
+		Options,
+		GeneratedTextures,
+		FinalMesh.Get(),
+		&Tangents);
+
+	Progress.EnterProgressFrame(1.f, LOCTEXT("Writing Assets", "Writing Assets..."));
+
+	// Make material for textures by duplicating input material (hardcoded!!)
+	UMaterial* BakeMaterial = LoadObject<UMaterial>(nullptr, TEXT("/MeshModelingToolset/Materials/FullMaterialBakePreviewMaterial"));
+	FMaterialAssetOptions MatOptions;
+	MatOptions.NewAssetPath = Options.BasePackagePath + TEXT("_Material");
+	FMaterialAssetResults MatResults;
+	ECreateMaterialResult MatResult = UE::AssetUtils::CreateDuplicateMaterial(BakeMaterial, MatOptions, MatResults);
+	UMaterial* NewMaterial = nullptr;
+	if (ensure(MatResult == ECreateMaterialResult::Ok))
+	{
+		NewMaterial = MatResults.NewMaterial;
+		ResultsOut.NewMaterials.Add(NewMaterial);
+	}
+
+	// this lambda converts a generated texture to an Asset, and then assigns it to a parameter of the Material
+	FString BaseTexturePath = MatOptions.NewAssetPath;
+	auto WriteTextureLambda = [BaseTexturePath, NewMaterial, &ResultsOut](
+		UTexture2D* Texture, 
+		FString TextureTypeSuffix, 
+		FTexture2DBuilder::ETextureType Type,
+		FName MaterialParamName )
+	{
+		if (ensure(Texture != nullptr) == false) return;
+
+		FTexture2DBuilder::CopyPlatformDataToSourceData(Texture, Type);
+
+		FTexture2DAssetOptions TexOptions;
+		TexOptions.NewAssetPath = BaseTexturePath + TextureTypeSuffix;
+		FTexture2DAssetResults Results;
+		ECreateTexture2DResult TexResult = UE::AssetUtils::SaveGeneratedTexture2DAsset(Texture, TexOptions, Results);
+		if (ensure(TexResult == ECreateTexture2DResult::Ok))
+		{
+			ResultsOut.NewTextures.Add(Texture);
+			if (NewMaterial != nullptr)
+			{
+				NewMaterial->SetTextureParameterValueEditorOnly(MaterialParamName, Texture);
+			}
+		}
+	};
+	// process the generated textures
+	WriteTextureLambda(GeneratedTextures.BaseColorMap, TEXT("_BaseColor"), FTexture2DBuilder::ETextureType::Color, FName("BaseColor"));
+	WriteTextureLambda(GeneratedTextures.RoughnessMap, TEXT("_Roughness"), FTexture2DBuilder::ETextureType::Color, FName("Roughness"));
+	WriteTextureLambda(GeneratedTextures.MetallicMap, TEXT("_Metallic"), FTexture2DBuilder::ETextureType::Color, FName("Metallic"));
+	WriteTextureLambda(GeneratedTextures.SpecularMap, TEXT("_Specular"), FTexture2DBuilder::ETextureType::Color, FName("Specular"));
+	WriteTextureLambda(GeneratedTextures.EmissiveMap, TEXT("_Emissive"), FTexture2DBuilder::ETextureType::Color, FName("Emissive"));
+	WriteTextureLambda(GeneratedTextures.NormalMap, TEXT("_Normal"), FTexture2DBuilder::ETextureType::NormalMap, FName("NormalMap"));
+
+	// force material update now that we have updated texture parameters
+	// (does this do that? Let calling code do it?)
+	NewMaterial->PostEditChange();
+
+	EmitGeneratedMeshAsset(Actors, Options, ResultsOut, FinalMesh.Get(), NewMaterial);
+	ResultsOut.ResultCode = EResultCode::Success;
+}
+
+
+UStaticMesh* FApproximateActorsImpl::EmitGeneratedMeshAsset(
+	const TArray<AActor*>& Actors, 
+	const FOptions& Options, 
+	FResults& ResultsOut,
+	FDynamicMesh3* FinalMesh,
+	UMaterialInterface* Material)
+{
+	FStaticMeshAssetOptions MeshAssetOptions;
+	MeshAssetOptions.NewAssetPath = Options.BasePackagePath;
+	MeshAssetOptions.SourceMeshes.DynamicMeshes.Add(FinalMesh);
+	if (Material)
+	{
+		MeshAssetOptions.AssetMaterials.Add(Material);
+	}
+	FStaticMeshResults MeshAssetOutputs;
+	ECreateStaticMeshResult ResultCode = UE::AssetUtils::CreateStaticMeshAsset(MeshAssetOptions, MeshAssetOutputs);
+	ensure(ResultCode == ECreateStaticMeshResult::Ok);
+
+	ResultsOut.NewMeshAssets.Add(MeshAssetOutputs.StaticMesh);
+
+	return MeshAssetOutputs.StaticMesh;
+}
+
+
+#undef LOCTEXT_NAMESPACE
\ No newline at end of file
diff --git a/Engine/Plugins/Experimental/MeshModelingToolset/Source/GeometryProcessingAdapters/Private/GeometryProcessingAdaptersModule.cpp b/Engine/Plugins/Experimental/MeshModelingToolset/Source/GeometryProcessingAdapters/Private/GeometryProcessingAdaptersModule.cpp
new file mode 100644
--- /dev/null
+++ b/Engine/Plugins/Experimental/MeshModelingToolset/Source/GeometryProcessingAdapters/Private/GeometryProcessingAdaptersModule.cpp
@@ -0,0 +1,37 @@
+// Copyright Epic Games, Inc. All Rights Reserved.
+
+#include "GeometryProcessingAdaptersModule.h"
+#include "Features/IModularFeatures.h"
+
+#define LOCTEXT_NAMESPACE "FGeometryProcessingAdaptersModule"
+
+using namespace UE::Geometry;
+
+void FGeometryProcessingAdaptersModule::StartupModule()
+{
+	// This code will execute after your module is loaded into memory; the exact timing is specified in the .uplugin file per-module
+	
+	ApproximateActors = MakeShared<FApproximateActorsImpl>();
+	if (ApproximateActors.IsValid())
+	{
+		IModularFeatures::Get().RegisterModularFeature(IGeometryProcessing_ApproximateActors::GetModularFeatureName(), ApproximateActors.Get());
+	}
+	
+}
+
+void FGeometryProcessingAdaptersModule::ShutdownModule()
+{
+	// This function may be called during shutdown to clean up your module.  For modules that support dynamic reloading,
+	// we call this function before unloading the module.
+	
+	if ( ApproximateActors.IsValid() )
+	{
+		IModularFeatures::Get().UnregisterModularFeature(IGeometryProcessing_ApproximateActors::GetModularFeatureName(), ApproximateActors.Get());
+		ApproximateActors = nullptr;
+	}
+	
+}
+
+#undef LOCTEXT_NAMESPACE
+	
+IMPLEMENT_MODULE(FGeometryProcessingAdaptersModule, GeometryProcessingAdapters)
diff --git a/Engine/Plugins/Experimental/MeshModelingToolset/Source/GeometryProcessingAdapters/Public/GeometryProcessing/ApproximateActorsImpl.h b/Engine/Plugins/Experimental/MeshModelingToolset/Source/GeometryProcessingAdapters/Public/GeometryProcessing/ApproximateActorsImpl.h
new file mode 100644
--- /dev/null
+++ b/Engine/Plugins/Experimental/MeshModelingToolset/Source/GeometryProcessingAdapters/Public/GeometryProcessing/ApproximateActorsImpl.h
@@ -0,0 +1,42 @@
+// Copyright Epic Games, Inc. All Rights Reserved.
+
+#pragma once
+
+#include "CoreMinimal.h"
+#include "GeometryProcessingInterfaces/ApproximateActors.h"
+
+class UMaterialInterface;
+
+namespace UE
+{
+namespace Geometry
+{
+
+class FDynamicMesh3;
+
+/**
+ * Implementation of IGeometryProcessing_ApproximateActors
+ */
+class GEOMETRYPROCESSINGADAPTERS_API FApproximateActorsImpl : public IGeometryProcessing_ApproximateActors
+{
+public:	
+	
+	virtual void ApproximateActors(const TArray<AActor*>& Actors, const FOptions& Options, FResults& ResultsOut) override;
+	
+
+
+protected:
+
+	virtual void GenerateApproximationForActorSet(const TArray<AActor*>& Actors, const FOptions& Options, FResults& ResultsOut);
+
+	virtual UStaticMesh* EmitGeneratedMeshAsset(
+		const TArray<AActor*>& Actors,
+		const FOptions& Options,
+		FResults& ResultsOut,
+		FDynamicMesh3* FinalMesh,
+		UMaterialInterface* Material = nullptr);
+};
+
+
+}
+}
diff --git a/Engine/Plugins/Experimental/MeshModelingToolset/Source/GeometryProcessingAdapters/Public/GeometryProcessingAdaptersModule.h b/Engine/Plugins/Experimental/MeshModelingToolset/Source/GeometryProcessingAdapters/Public/GeometryProcessingAdaptersModule.h
new file mode 100644
--- /dev/null
+++ b/Engine/Plugins/Experimental/MeshModelingToolset/Source/GeometryProcessingAdapters/Public/GeometryProcessingAdaptersModule.h
@@ -0,0 +1,21 @@
+// Copyright Epic Games, Inc. All Rights Reserved.
+
+#pragma once
+
+#include "CoreMinimal.h"
+#include "Modules/ModuleManager.h"
+
+#include "GeometryProcessing/ApproximateActorsImpl.h"
+
+class FGeometryProcessingAdaptersModule : public IModuleInterface
+{
+public:
+
+	/** IModuleInterface implementation */
+	virtual void StartupModule() override;
+	virtual void ShutdownModule() override;
+	
+	
+protected:
+	TSharedPtr<UE::Geometry::FApproximateActorsImpl> ApproximateActors;
+};
diff --git a/Engine/Source/Developer/GeometryProcessingInterfaces/GeometryProcessingInterfaces.Build.cs b/Engine/Source/Developer/GeometryProcessingInterfaces/GeometryProcessingInterfaces.Build.cs
new file mode 100644
--- /dev/null
+++ b/Engine/Source/Developer/GeometryProcessingInterfaces/GeometryProcessingInterfaces.Build.cs
@@ -0,0 +1,14 @@
+// Copyright Epic Games, Inc. All Rights Reserved.
+
+using UnrealBuildTool;
+using System.IO;
+
+public class GeometryProcessingInterfaces : ModuleRules
+{
+	public GeometryProcessingInterfaces(ReadOnlyTargetRules Target) : base(Target)
+	{		
+		PrivateDependencyModuleNames.Add("Core");
+		PrivateDependencyModuleNames.Add("CoreUObject");
+		PrivateDependencyModuleNames.Add("Engine");
+    }
+}
diff --git a/Engine/Source/Developer/GeometryProcessingInterfaces/Private/GeometryProcessingInterfacesModule.cpp b/Engine/Source/Developer/GeometryProcessingInterfaces/Private/GeometryProcessingInterfacesModule.cpp
new file mode 100644
--- /dev/null
+++ b/Engine/Source/Developer/GeometryProcessingInterfaces/Private/GeometryProcessingInterfacesModule.cpp
@@ -0,0 +1,40 @@
+// Copyright Epic Games, Inc. All Rights Reserved.
+
+#include "GeometryProcessingInterfacesModule.h"
+#include "CoreGlobals.h"
+#include "Features/IModularFeatures.h"
+#include "Modules/ModuleManager.h"
+
+#include "GeometryProcessingInterfaces/ApproximateActors.h"
+
+
+//DEFINE_LOG_CATEGORY_STATIC(LogMeshReduction, Verbose, All);
+
+IMPLEMENT_MODULE(FGeometryProcessingInterfacesModule, GeometryProcessingInterfaces);
+
+
+void FGeometryProcessingInterfacesModule::StartupModule()
+{
+
+}
+
+
+void FGeometryProcessingInterfacesModule::ShutdownModule()
+{
+	ApproximateActors = nullptr;
+}
+
+
+IGeometryProcessing_ApproximateActors* FGeometryProcessingInterfacesModule::GetApproximateActorsImplementation()
+{
+	if (ApproximateActors == nullptr)
+	{
+		TArray<IGeometryProcessing_ApproximateActors*> ApproximateActorsOptions =
+			IModularFeatures::Get().GetModularFeatureImplementations<IGeometryProcessing_ApproximateActors>(IGeometryProcessing_ApproximateActors::GetModularFeatureName());
+
+		ApproximateActors = (ApproximateActorsOptions.Num() > 0) ? ApproximateActorsOptions[0] : nullptr;
+	}
+
+	return ApproximateActors;
+}
+
diff --git a/Engine/Source/Developer/GeometryProcessingInterfaces/Private/GeometryProcessingInterfacesModule.h b/Engine/Source/Developer/GeometryProcessingInterfaces/Private/GeometryProcessingInterfacesModule.h
new file mode 100644
--- /dev/null
+++ b/Engine/Source/Developer/GeometryProcessingInterfaces/Private/GeometryProcessingInterfacesModule.h
@@ -0,0 +1,31 @@
+// Copyright Epic Games, Inc. All Rights Reserved.
+
+#pragma once
+
+#include "IGeometryProcessingInterfacesModule.h"
+
+
+class IGeometryProcessing_ApproximateActors;
+
+
+/**
+ * Implementation of IGeometryProcessingInterfacesModule (which extends the standard IModuleInterface)
+ * to provide access to various "Operation" interfaces
+ */
+class FGeometryProcessingInterfacesModule : public IGeometryProcessingInterfacesModule
+{
+public:
+
+	virtual void StartupModule() override;
+	virtual void ShutdownModule() override;
+	
+	
+	/**
+	 * @return implementation of IGeometryProcessing_ApproximateActors, if available, or nullptr (result is cached internally)
+	 */
+	virtual IGeometryProcessing_ApproximateActors* GetApproximateActorsImplementation() override;
+
+
+private:
+	IGeometryProcessing_ApproximateActors* ApproximateActors = nullptr;
+};
diff --git a/Engine/Source/Developer/GeometryProcessingInterfaces/Public/GeometryProcessingInterfaces/ApproximateActors.h b/Engine/Source/Developer/GeometryProcessingInterfaces/Public/GeometryProcessingInterfaces/ApproximateActors.h
new file mode 100644
--- /dev/null
+++ b/Engine/Source/Developer/GeometryProcessingInterfaces/Public/GeometryProcessingInterfaces/ApproximateActors.h
@@ -0,0 +1,143 @@
+// Copyright Epic Games, Inc. All Rights Reserved.
+
+#pragma once
+
+#include "CoreMinimal.h"
+#include "Features/IModularFeature.h"
+
+class AActor;
+class UStaticMesh;
+class UMaterial;
+class UTexture2D;
+
+
+/**
+ * IGeometryProcessing_ApproximateActors is a generic high-level interface to a function/system that
+ * outputs a set of new Mesh assets that approximate a set of input Actors.
+ * This can be used to provide LOD-generation implementations for systems like ProxyLOD and HLOD.
+ * 
+ * This is an IModularFeature, and so clients can query/enumerate the available ApproximateActors implementations 
+ * based on the ::GetModularFeatureName(). However, the preferred way is to use code like the following:
+ * 
+ *     IGeometryProcessingInterfacesModule& GeomProcInterfaces = FModuleManager::Get().LoadModuleChecked<IGeometryProcessingInterfacesModule>("GeometryProcessingInterfaces");
+ *     IGeometryProcessing_ApproximateActors* ApproxActorsAPI = GeomProcInterfaces.GetApproximateActorsImplementation();
+ * 
+ * This will automatically determine which available implementation to use (based on any applicable config settings/etc),
+ * and cache the result of that decision.
+ */
+class IGeometryProcessing_ApproximateActors : public IModularFeature
+{
+public:
+	virtual ~IGeometryProcessing_ApproximateActors() {}
+
+	enum class EApproximationPolicy
+	{
+		MeshAndGeneratedMaterial,
+		CollisionMesh
+	};
+
+
+	enum class ESimplificationPolicy
+	{
+		FixedTriangleCount = 0,
+		TrianglesPerUnitSqMeter = 1
+	};
+
+
+	/**
+	 * Input options to Actor Approximation process
+	 */
+	struct FOptions
+	{
+		// Base path for generated assets. Names will be generated by appending strings to this path.
+		FString BasePackagePath;
+
+		// high-level control of the overall approximation process
+		EApproximationPolicy BasePolicy = EApproximationPolicy::MeshAndGeneratedMaterial;
+
+		//
+		// Shape Approximation settings
+		//
+
+		// Meshing settings (ie for voxelization)
+		float WorldSpaceApproximationAccuracyMeters = 1.0f;
+		int32 ClampVoxelDimension = 512;
+
+		bool bApplyMorphology = false;
+		float MorphologyDistanceMeters = 0.1;
+
+		ESimplificationPolicy MeshSimplificationPolicy = ESimplificationPolicy::FixedTriangleCount;
+		int32 FixedTriangleCount = 5000;
+		float SimplificationTargetMetric = 0.1;		// interpretation varies depending on MeshSimplificationPolicy
+
+		//
+		// Material approximation settings
+		//
+		int32 RenderCaptureImageSize = 1024;
+
+		// render capture parameters
+		float FieldOfViewDegrees = 45.0f;
+		float NearPlaneDist = 1.0f;
+
+
+		//
+		// Material output settings
+		//
+
+		// output texture options
+		int32 TextureImageSize = 1024;
+
+
+
+
+		// todo
+		//   - which texture types to generate (BaseColor, Specular, WorldNormal, etc)
+		//   - per-texture resolution, and channel/compression settings
+		//   - which material to use as input, whether to make new UMaterial or a MIC
+		//   - custom Generated-Texture to Material Parameter Name mapping
+	};
+
+	enum class EResultCode
+	{
+		Success,
+		MeshGenerationFailed,
+		UnknownError
+	};
+
+	/**
+	 * Outputs of an Actor Approximation process
+	 */
+	struct FResults
+	{
+		EResultCode ResultCode = EResultCode::UnknownError;
+
+		TArray<UStaticMesh*> NewMeshAssets;
+
+		TArray<UMaterial*> NewMaterials;
+		
+		TArray<UTexture2D*> NewTextures;
+
+		// todo
+	};
+
+	/**
+	 * Top-level driver function that clients call to generate the approximation for a set of input Actors.
+	 */
+	virtual void ApproximateActors(const TArray<AActor*>& Actors, const FOptions& Options, FResults& ResultsOut) 
+	{
+		check(false);		// not implemented in base class
+	}
+
+
+	// Modular feature name to register for retrieval during runtime
+	static const FName GetModularFeatureName()
+	{
+		return TEXT("GeometryProcessing_ApproximateActors");
+	}
+
+
+	// delegates to pass back info?
+	DECLARE_DELEGATE_OneParam(FApproximateActorsCompleteDelegate, FResults&);
+	FApproximateActorsCompleteDelegate CompleteDelegate;
+
+};
\ No newline at end of file
diff --git a/Engine/Source/Developer/GeometryProcessingInterfaces/Public/IGeometryProcessingInterfacesModule.h b/Engine/Source/Developer/GeometryProcessingInterfaces/Public/IGeometryProcessingInterfacesModule.h
new file mode 100644
--- /dev/null
+++ b/Engine/Source/Developer/GeometryProcessingInterfaces/Public/IGeometryProcessingInterfacesModule.h
@@ -0,0 +1,30 @@
+// Copyright Epic Games, Inc. All Rights Reserved.
+
+#pragma once
+
+#include "Modules/ModuleInterface.h"
+
+class IGeometryProcessing_ApproximateActors;
+
+/**
+ * Abstract interface to a Module that provides functions to access
+ * different "Operation" interfaces, which are high-level APIs to
+ * run complex geometry operations on engine inputs without the caller
+ * needing to know about the potential dependencies (ie, the implementations 
+ * can be provided by plugins, and this non-plugin module can find them for 
+ * the Engine/Editor core, avoiding the no-core-dependencies-on-plugins limitation)
+ * 
+ * This interface is implemented by FGeometryProcessingInterfacesModule. 
+ * Client code that wants access to the operations should get this interface like so:
+ * 
+ *		IGeometryProcessingInterfacesModule& GeomProcInterfaces = FModuleManager::Get().LoadModuleChecked<IGeometryProcessingInterfacesModule>("GeometryProcessingInterfaces");
+ *
+ */
+class IGeometryProcessingInterfacesModule : public IModuleInterface
+{
+public:	
+	/**
+	 * @return implementation of IGeometryProcessing_ApproximateActors, if available
+	 */
+	virtual IGeometryProcessing_ApproximateActors* GetApproximateActorsImplementation() = 0;
+};
