MeshModelingTools: Multiple default material fixes. Made LineSetComponentMaterial thicken in the correct direction in ortho, added dashed/dimmed occlusion for overlaid version, disabled EyeAdaptationInverse for ortho in translucent materials because it does not work there, made depth offset work in ortho, split translucent-mode materials from opaque ones since opaque has fewer limitations, changed several mechanics to get material via ToolSetupUtil instead of directly, changed the "soft" point material to be a proper round points material since this is what it was used for. Made MeshInspector not use too large values for its depth offset.

#rnx

diff --git a/Engine/Plugins/Experimental/MeshModelingToolset/Source/MeshModelingTools/Private/MeshInspectorTool.cpp b/Engine/Plugins/Experimental/MeshModelingToolset/Source/MeshModelingTools/Private/MeshInspectorTool.cpp
--- a/Engine/Plugins/Experimental/MeshModelingToolset/Source/MeshModelingTools/Private/MeshInspectorTool.cpp
+++ b/Engine/Plugins/Experimental/MeshModelingToolset/Source/MeshModelingTools/Private/MeshInspectorTool.cpp
@@ -246,12 +246,12 @@ void UMeshInspectorTool::UpdateVisualization()
 	FColor BinormalColor(15, 240, 15);
 	float TangentThickness = LineWidthMultiplier * 2.0f;
 
-	float BoundaryEdgeDepthBias = 2.0f;
-	float UVSeamDepthBias = 3.0f;
-	float NormalSeamDepthBias = 3.0f;
-	float PolygonBorderDepthBias = 2.0f;
+	float BoundaryEdgeDepthBias = 0.2f;
+	float UVSeamDepthBias = 0.3f;
+	float NormalSeamDepthBias = 0.3f;
+	float PolygonBorderDepthBias = 0.2f;
 	float NormalDepthBias = 0.0f;
-	float TangentDepthBias = 3.5f;
+	float TangentDepthBias = 0.35f;
 
 	// Used to scale normals and tangents appropriately
 	FVector ComponentScale = DrawnLineSet->GetComponentTransform().GetScale3D();
diff --git a/Engine/Plugins/Experimental/MeshModelingToolset/Source/ModelingComponents/Private/Mechanics/CurveControlPointsMechanic.cpp b/Engine/Plugins/Experimental/MeshModelingToolset/Source/ModelingComponents/Private/Mechanics/CurveControlPointsMechanic.cpp
--- a/Engine/Plugins/Experimental/MeshModelingToolset/Source/ModelingComponents/Private/Mechanics/CurveControlPointsMechanic.cpp
+++ b/Engine/Plugins/Experimental/MeshModelingToolset/Source/ModelingComponents/Private/Mechanics/CurveControlPointsMechanic.cpp
@@ -53,16 +53,16 @@ void UCurveControlPointsMechanic::Setup(UInteractiveTool* ParentToolIn)
 	// TODO: This probably should be configurable. For instance, we may want the segments to be dashed, or not visible at all.
 	DrawnControlPoints = NewObject<UPointSetComponent>();
 	DrawnControlPoints->SetPointMaterial(
-		LoadObject<UMaterial>(nullptr, TEXT("/MeshModelingToolset/Materials/PointSetOverlaidComponentMaterial")));
+		ToolSetupUtil::GetDefaultPointComponentMaterial(GetParentTool()->GetToolManager(), /*bDepthTested*/ false));
 	DrawnControlSegments = NewObject<ULineSetComponent>();
 	DrawnControlSegments->SetLineMaterial(
-		LoadObject<UMaterial>(nullptr, TEXT("/MeshModelingToolset/Materials/LineSetOverlaidComponentMaterial")));
+		ToolSetupUtil::GetDefaultLineComponentMaterial(GetParentTool()->GetToolManager(), /*bDepthTested*/ false));
 	PreviewPoint = NewObject<UPointSetComponent>();
 	PreviewPoint->SetPointMaterial(
-		LoadObject<UMaterial>(nullptr, TEXT("/MeshModelingToolset/Materials/PointSetOverlaidComponentMaterial")));
+		ToolSetupUtil::GetDefaultPointComponentMaterial(GetParentTool()->GetToolManager(), /*bDepthTested*/ false));
 	PreviewSegment = NewObject<ULineSetComponent>();
 	PreviewSegment->SetLineMaterial(
-		LoadObject<UMaterial>(nullptr, TEXT("/MeshModelingToolset/Materials/LineSetOverlaidComponentMaterial")));
+		ToolSetupUtil::GetDefaultLineComponentMaterial(GetParentTool()->GetToolManager(), /*bDepthTested*/ false));
 
 	InitializationCurveColor = FColor::Orange;
 	NormalCurveColor = FColor::Red;
diff --git a/Engine/Plugins/Experimental/MeshModelingToolset/Source/ModelingComponents/Private/Mechanics/LatticeControlPointsMechanic.cpp b/Engine/Plugins/Experimental/MeshModelingToolset/Source/ModelingComponents/Private/Mechanics/LatticeControlPointsMechanic.cpp
--- a/Engine/Plugins/Experimental/MeshModelingToolset/Source/ModelingComponents/Private/Mechanics/LatticeControlPointsMechanic.cpp
+++ b/Engine/Plugins/Experimental/MeshModelingToolset/Source/ModelingComponents/Private/Mechanics/LatticeControlPointsMechanic.cpp
@@ -63,10 +63,10 @@ void ULatticeControlPointsMechanic::Setup(UInteractiveTool* ParentToolIn)
 
 	DrawnControlPoints = NewObject<UPointSetComponent>();
 	DrawnControlPoints->SetPointMaterial(
-		LoadObject<UMaterial>(nullptr, TEXT("/MeshModelingToolset/Materials/PointSetComponentMaterial")));
+		ToolSetupUtil::GetDefaultPointComponentMaterial(GetParentTool()->GetToolManager(), /*bDepthTested*/ true));
 	DrawnLatticeEdges = NewObject<ULineSetComponent>();
 	DrawnLatticeEdges->SetLineMaterial(
-		LoadObject<UMaterial>(nullptr, TEXT("/MeshModelingToolset/Materials/LineSetComponentMaterial")));
+		ToolSetupUtil::GetDefaultLineComponentMaterial(GetParentTool()->GetToolManager(), /*bDepthTested*/ true));
 
 	NormalPointColor = FColor::Red;
 	NormalSegmentColor = FColor::Red;
diff --git a/Engine/Plugins/Experimental/MeshModelingToolset/Source/ModelingComponents/Private/Mechanics/SpaceCurveDeformationMechanic.cpp b/Engine/Plugins/Experimental/MeshModelingToolset/Source/ModelingComponents/Private/Mechanics/SpaceCurveDeformationMechanic.cpp
--- a/Engine/Plugins/Experimental/MeshModelingToolset/Source/ModelingComponents/Private/Mechanics/SpaceCurveDeformationMechanic.cpp
+++ b/Engine/Plugins/Experimental/MeshModelingToolset/Source/ModelingComponents/Private/Mechanics/SpaceCurveDeformationMechanic.cpp
@@ -51,10 +51,10 @@ void USpaceCurveDeformationMechanic::Setup(UInteractiveTool* ParentToolIn)
 	// TODO: This probably should be configurable. For instance, we may want the segments to be dashed, or not visible at all.
 	RenderPoints = NewObject<UPointSetComponent>();
 	RenderPoints->SetPointMaterial(
-		LoadObject<UMaterial>(nullptr, TEXT("/MeshModelingToolset/Materials/PointSetOverlaidComponentMaterial")));
+		ToolSetupUtil::GetDefaultPointComponentMaterial(GetParentTool()->GetToolManager(), /*bDepthTested*/ false));
 	RenderSegments = NewObject<ULineSetComponent>();
 	RenderSegments->SetLineMaterial(
-		LoadObject<UMaterial>(nullptr, TEXT("/MeshModelingToolset/Materials/LineSetOverlaidComponentMaterial")));
+		ToolSetupUtil::GetDefaultLineComponentMaterial(GetParentTool()->GetToolManager(), /*bDepthTested*/ false));
 
 	NormalCurveColor = FColor::Red;
 	CurrentSegmentsColor = FColor::Orange;
diff --git a/Engine/Plugins/Experimental/MeshModelingToolset/Source/ModelingComponents/Private/ToolSetupUtil.cpp b/Engine/Plugins/Experimental/MeshModelingToolset/Source/ModelingComponents/Private/ToolSetupUtil.cpp
--- a/Engine/Plugins/Experimental/MeshModelingToolset/Source/ModelingComponents/Private/ToolSetupUtil.cpp
+++ b/Engine/Plugins/Experimental/MeshModelingToolset/Source/ModelingComponents/Private/ToolSetupUtil.cpp
@@ -187,7 +187,7 @@ UMaterialInterface* ToolSetupUtil::GetSelectionMaterial(const FLinearColor& UseC
 
 UMaterialInstanceDynamic* ToolSetupUtil::GetSimpleCustomMaterial(UInteractiveToolManager* ToolManager, const FLinearColor& Color, float Opacity)
 {
-	UMaterialInterface* Material = LoadObject<UMaterial>(nullptr, TEXT("/MeshModelingToolset/Materials/SimpleBaseMaterial"));
+	UMaterialInterface* Material = LoadObject<UMaterial>(nullptr, TEXT("/MeshModelingToolset/Materials/SimpleBaseMaterial_Transparent"));
 	if (Material != nullptr)
 	{
 		UMaterialInstanceDynamic* MatInstance = UMaterialInstanceDynamic::Create(Material, ToolManager);
@@ -198,9 +198,21 @@ UMaterialInstanceDynamic* ToolSetupUtil::GetSimpleCustomMaterial(UInteractiveToo
 	return nullptr;
 }
 
+MODELINGCOMPONENTS_API UMaterialInstanceDynamic* ToolSetupUtil::GetSimpleCustomMaterial(UInteractiveToolManager* ToolManager, const FLinearColor& Color)
+{
+	UMaterialInterface* Material = LoadObject<UMaterial>(nullptr, TEXT("/MeshModelingToolset/Materials/SimpleBaseMaterial"));
+	if (Material != nullptr)
+	{
+		UMaterialInstanceDynamic* MatInstance = UMaterialInstanceDynamic::Create(Material, ToolManager);
+		MatInstance->SetVectorParameterValue(TEXT("Color"), Color);
+		return MatInstance;
+	}
+	return nullptr;
+}
+
 UMaterialInstanceDynamic* ToolSetupUtil::GetCustomTwoSidedDepthOffsetMaterial(UInteractiveToolManager* ToolManager, const FLinearColor& Color, float PercentDepthOffset, float Opacity)
 {
-	UMaterialInterface* Material = LoadObject<UMaterial>(nullptr, TEXT("/MeshModelingToolset/Materials/SimpleTwoSidedOffsetMaterial"));
+	UMaterialInterface* Material = LoadObject<UMaterial>(nullptr, TEXT("/MeshModelingToolset/Materials/SimpleTwoSidedOffsetMaterial_Transparent"));
 	if (Material != nullptr)
 	{
 		UMaterialInstanceDynamic* MatInstance = UMaterialInstanceDynamic::Create(Material, ToolManager);
@@ -212,9 +224,22 @@ UMaterialInstanceDynamic* ToolSetupUtil::GetCustomTwoSidedDepthOffsetMaterial(UI
 	return nullptr;
 }
 
+MODELINGCOMPONENTS_API UMaterialInstanceDynamic* ToolSetupUtil::GetCustomTwoSidedDepthOffsetMaterial(UInteractiveToolManager* ToolManager, const FLinearColor& Color, float PercentDepthOffset)
+{
+	UMaterialInterface* Material = LoadObject<UMaterial>(nullptr, TEXT("/MeshModelingToolset/Materials/SimpleTwoSidedOffsetMaterial"));
+	if (Material != nullptr)
+	{
+		UMaterialInstanceDynamic* MatInstance = UMaterialInstanceDynamic::Create(Material, ToolManager);
+		MatInstance->SetVectorParameterValue(TEXT("Color"), Color);
+		MatInstance->SetScalarParameterValue(TEXT("PercentDepthOffset"), PercentDepthOffset);
+		return MatInstance;
+	}
+	return nullptr;
+}
+
 UMaterialInstanceDynamic* ToolSetupUtil::GetCustomDepthOffsetMaterial(UInteractiveToolManager* ToolManager, const FLinearColor& Color, float PercentDepthOffset, float Opacity)
 {
-	UMaterialInterface* Material = LoadObject<UMaterial>(nullptr, TEXT("/MeshModelingToolset/Materials/SimpleOffsetMaterial"));
+	UMaterialInterface* Material = LoadObject<UMaterial>(nullptr, TEXT("/MeshModelingToolset/Materials/SimpleOffsetMaterial_Transparent"));
 	if (Material != nullptr)
 	{
 		UMaterialInstanceDynamic* MatInstance = UMaterialInstanceDynamic::Create(Material, ToolManager);
@@ -226,17 +251,42 @@ UMaterialInstanceDynamic* ToolSetupUtil::GetCustomDepthOffsetMaterial(UInteracti
 	return nullptr;
 }
 
+MODELINGCOMPONENTS_API UMaterialInstanceDynamic* ToolSetupUtil::GetCustomDepthOffsetMaterial(UInteractiveToolManager* ToolManager, const FLinearColor& Color, float PercentDepthOffset)
+{
+	UMaterialInterface* Material = LoadObject<UMaterial>(nullptr, TEXT("/MeshModelingToolset/Materials/SimpleOffsetMaterial"));
+	if (Material != nullptr)
+	{
+		UMaterialInstanceDynamic* MatInstance = UMaterialInstanceDynamic::Create(Material, ToolManager);
+		MatInstance->SetVectorParameterValue(TEXT("Color"), Color);
+		MatInstance->SetScalarParameterValue(TEXT("PercentDepthOffset"), PercentDepthOffset);
+		return MatInstance;
+	}
+	return nullptr;
+}
+
 UMaterialInterface* ToolSetupUtil::GetDefaultEditVolumeMaterial()
 {
 	return LoadObject<UMaterial>(nullptr, TEXT("/MeshModelingToolset/Materials/VolumeEditMaterial"));
 }
 
 
-UMaterialInterface* ToolSetupUtil::GetDefaultPointComponentMaterial(bool bRoundPoints, UInteractiveToolManager* ToolManager)
+UMaterialInterface* ToolSetupUtil::GetDefaultPointComponentMaterial(UInteractiveToolManager* ToolManager, bool bDepthTested)
+{
+	UMaterialInterface* Material = (bDepthTested) ?
+		LoadObject<UMaterial>(nullptr, TEXT("/MeshModelingToolset/Materials/PointSetComponentMaterial"))
+		: LoadObject<UMaterial>(nullptr, TEXT("/MeshModelingToolset/Materials/PointSetOverlaidComponentMaterial"));
+	if (Material == nullptr && ToolManager != nullptr)
+	{
+		return ToolManager->GetContextQueriesAPI()->GetStandardMaterial(EStandardToolContextMaterials::VertexColorMaterial);
+	}
+	return Material;
+}
+
+UMaterialInterface* ToolSetupUtil::GetRoundPointComponentMaterial(UInteractiveToolManager* ToolManager, bool bDepthTested)
 {
-	UMaterialInterface* Material = (bRoundPoints) ?
-		LoadObject<UMaterial>(nullptr, TEXT("/MeshModelingToolset/Materials/PointSetComponentMaterialSoft")) :
-		LoadObject<UMaterial>(nullptr, TEXT("/MeshModelingToolset/Materials/PointSetComponentMaterial"));
+	UMaterialInterface* Material = (bDepthTested) ?
+		LoadObject<UMaterial>(nullptr, TEXT("/MeshModelingToolset/Materials/PointSetComponentMaterial_Round"))
+		: LoadObject<UMaterial>(nullptr, TEXT("/MeshModelingToolset/Materials/PointSetOverlaidComponentMaterial_Round"));
 	if (Material == nullptr && ToolManager != nullptr)
 	{
 		return ToolManager->GetContextQueriesAPI()->GetStandardMaterial(EStandardToolContextMaterials::VertexColorMaterial);
diff --git a/Engine/Plugins/Experimental/MeshModelingToolset/Source/ModelingComponents/Public/ToolSetupUtil.h b/Engine/Plugins/Experimental/MeshModelingToolset/Source/ModelingComponents/Public/ToolSetupUtil.h
--- a/Engine/Plugins/Experimental/MeshModelingToolset/Source/ModelingComponents/Public/ToolSetupUtil.h
+++ b/Engine/Plugins/Experimental/MeshModelingToolset/Source/ModelingComponents/Public/ToolSetupUtil.h
@@ -97,33 +97,78 @@ namespace ToolSetupUtil
 	MODELINGCOMPONENTS_API UMaterialInterface* GetSelectionMaterial(const FLinearColor& UseColor, UInteractiveToolManager* ToolManager, float PercentDepthOffset = 0.0f);
 
 	/**
-	 * @return Simple material with configurable color and opacity.
+	 * @return Simple material with configurable color and opacity. Note that the material 
+	 *  will have translucent blend mode, which can interact poorly with overlapping translucent
+	 *  objects, so use the other overload if you do not need opacity control.
 	 */
 	MODELINGCOMPONENTS_API UMaterialInstanceDynamic* GetSimpleCustomMaterial(UInteractiveToolManager* ToolManager, const FLinearColor& Color, float Opacity);
 
 	/**
-	 * @return Simple material with configurable depth offset, color, and opacity.
+	 * @return Simple material with configurable color. The material will have opaque blend mode.
 	 */
+	MODELINGCOMPONENTS_API UMaterialInstanceDynamic* GetSimpleCustomMaterial(UInteractiveToolManager* ToolManager, const FLinearColor& Color);
+
+	 /**
+	  * @return Simple material with configurable depth offset, color, and opacity. Note that the material
+	  *  will have translucent blend mode, which can interact poorly with overlapping translucent
+	  *  objects, so use the other overload if you do not need opacity control.
+	  */
 	MODELINGCOMPONENTS_API UMaterialInstanceDynamic* GetCustomDepthOffsetMaterial(UInteractiveToolManager* ToolManager, const FLinearColor& Color, float PercentDepthOffset, float Opacity);
 
 	/**
-	 * @return Simple two-sided material with configurable depth offset, color, and opacity.
+	 * @return Simple material with configurable depth offset and color. The material will have opaque blend mode.
+	 */
+	MODELINGCOMPONENTS_API UMaterialInstanceDynamic* GetCustomDepthOffsetMaterial(UInteractiveToolManager* ToolManager, const FLinearColor& Color, float PercentDepthOffset);
+	
+	/**
+	 * @return Simple two-sided material with configurable depth offset, color, and opacity. Note that 
+	  *  the material will have translucent blend mode, which can interact poorly with overlapping translucent
+	  *  objects, so use the other overload if you do not need opacity control.
 	 */
 	MODELINGCOMPONENTS_API UMaterialInstanceDynamic* GetCustomTwoSidedDepthOffsetMaterial(UInteractiveToolManager* ToolManager, const FLinearColor& Color, float PercentDepthOffset, float Opacity);
 
+	/**
+	 * @return Simple two-sided material with configurable depth offset and color. The material will have opaque blend mode.
+	 */
+	MODELINGCOMPONENTS_API UMaterialInstanceDynamic* GetCustomTwoSidedDepthOffsetMaterial(UInteractiveToolManager* ToolManager, const FLinearColor& Color, float PercentDepthOffset);
+
 	/**
 	 * @return Material used when editing AVolume objects using our tools.
 	 */
 	MODELINGCOMPONENTS_API UMaterialInterface* GetDefaultEditVolumeMaterial();
 
 	/**
-	 * @param bRoundPoints true for round points, false for square
-	 * @return custom material suitable for use with UPointSetComponent
+	 * Gets a custom material suitable for use with UPointSetComponent for square points.
+	 * 
+	 * @param bDepthTested If true, the material will be depth tested as normal. If false, occluded points will still be
+	 *  displayed but dimmed.
+	 *  Note that the current implementations of the depth-tested and non-depth-tested modes use opaque and translucent 
+	 *  blend modes, respectively, and so inherit their limitations. Specifically, opaque does not support opacity, 
+	 *  and translucent does not always follow correct draw order relative to other translucent objects, which means
+	 *  that depth offset cannot reliably order lines within a non-depth-tested line set component.
+	 */
+	MODELINGCOMPONENTS_API UMaterialInterface* GetDefaultPointComponentMaterial(UInteractiveToolManager* ToolManager, bool bDepthTested = true);
+
+	/**
+	 * Gets a custom material suitable for use with UPointSetComponent for round points.
+	 * Note that this material uses translucent blend mode, and therefore can't always follow the correct draw order
+	 * relative to other translucent objects (and within the point set). If this is not acceptible, you will need to use
+	 * the depth tested square point material (GetDefaultPointComponentMaterial with bDepthTested = true).
+	 * 
+	 * @param bDepthTested If true, the material will be depth tested as normal. If false, occluded portions of lines
+	 *  will still be displayed, but dashed and dimmed.
 	 */
-	MODELINGCOMPONENTS_API UMaterialInterface* GetDefaultPointComponentMaterial(bool bRoundPoints, UInteractiveToolManager* ToolManager);
+	MODELINGCOMPONENTS_API UMaterialInterface* GetRoundPointComponentMaterial(UInteractiveToolManager* ToolManager, bool bDepthTested = true);
 
 	/**
-	 * @return custom material suitable for use with ULineSetComponent
+	 * Gets a custom material suitable for use with ULineSetComponent.
+	 * 
+	 * @param bDepthTested If true, the material will be depth tested as normal. If false, occluded portions of lines
+	 *  will still be displayed, but dashed and dimmed. 
+	 *  Note that the current implementations of the depth-tested and non-depth-tested modes use opaque and translucent 
+	 *  blend modes, respectively, and so inherit their limitations. Specifically, opaque does not support opacity, 
+	 *  and translucent does not always follow correct draw order relative to other translucent objects, which means
+	 *  that depth offset cannot reliably order lines within a non-depth-tested line set component.
 	 */
 	MODELINGCOMPONENTS_API UMaterialInterface* GetDefaultLineComponentMaterial(UInteractiveToolManager* ToolManager, bool bDepthTested = true);
 
