Adding a multi-texture bake mode to the texture bake tool. User can bake a different texture per material region. Defaults to trying to find an Albedo texture for each material on the input mesh, but the textures to bake can be specified by the user.

Also adding a mode to bake MaterialIDs as colors.

diff --git a/Engine/Plugins/Experimental/GeometryProcessing/Source/DynamicMesh/Private/Sampling/MeshPropertyMapBaker.cpp b/Engine/Plugins/Experimental/GeometryProcessing/Source/DynamicMesh/Private/Sampling/MeshPropertyMapBaker.cpp
--- a/Engine/Plugins/Experimental/GeometryProcessing/Source/DynamicMesh/Private/Sampling/MeshPropertyMapBaker.cpp
+++ b/Engine/Plugins/Experimental/GeometryProcessing/Source/DynamicMesh/Private/Sampling/MeshPropertyMapBaker.cpp
@@ -1,6 +1,7 @@
 // Copyright Epic Games, Inc. All Rights Reserved.
 
 #include "Sampling/MeshPropertyMapBaker.h"
+#include "Util/ColorConstants.h"
 
 #include "ExplicitUseGeometryMathTypes.h"		// using UE::Geometry::(math types)
 using namespace UE::Geometry;
@@ -105,6 +106,20 @@ void FMeshPropertyMapBaker::Bake()
 						}
 					}
 					break;
+				case EMeshPropertyBakeType::MaterialID:
+					{
+						if (DetailMesh->Attributes() && DetailMesh->Attributes()->HasMaterialID())
+						{
+							const FDynamicMeshMaterialAttribute* DetailMaterialIDAttrib = DetailMesh->Attributes()->GetMaterialID();
+							const int32 MatID = DetailMaterialIDAttrib->GetValue(DetailTriID);
+							Color = LinearColors::SelectColor<FVector3f>(MatID);
+						}
+						else
+						{
+							Color = FVector3f(LinearColors::LightPink3f());
+						}
+					}
+					break;
 			}
 		}
 		return Color;
diff --git a/Engine/Plugins/Experimental/GeometryProcessing/Source/DynamicMesh/Private/Sampling/MeshResampleImageBaker.cpp b/Engine/Plugins/Experimental/GeometryProcessing/Source/DynamicMesh/Private/Sampling/MeshResampleImageBaker.cpp
--- a/Engine/Plugins/Experimental/GeometryProcessing/Source/DynamicMesh/Private/Sampling/MeshResampleImageBaker.cpp
+++ b/Engine/Plugins/Experimental/GeometryProcessing/Source/DynamicMesh/Private/Sampling/MeshResampleImageBaker.cpp
@@ -47,4 +47,85 @@ void FMeshResampleImageBaker::Bake()
 		ResultBuilder->CopyPixel(GutterTexel.Value, GutterTexel.Key);
 	}
 
-}
\ No newline at end of file
+}
+
+
+
+
+void FMeshMultiResampleImageBaker::InitResult()
+{
+	const FMeshImageBakingCache* BakeCache = GetCache();
+	ResultBuilder = MakeUnique<TImageBuilder<FVector4f>>();
+	ResultBuilder->SetDimensions(BakeCache->GetDimensions());
+	ResultBuilder->Clear(DefaultColor);
+}
+
+
+void FMeshMultiResampleImageBaker::BakeMaterial(int32 MaterialID)
+{
+	const FMeshImageBakingCache* BakeCache = GetCache();
+	if (!ensure(BakeCache))
+	{
+		return;
+	}
+	const FDynamicMesh3* DetailMesh = BakeCache->GetDetailMesh();
+	if (!ensure(DetailMesh))
+	{
+		return;
+	}
+	const FDynamicMeshMaterialAttribute* DetailMaterialIDAttrib = DetailMesh->Attributes()->GetMaterialID();
+	if (!ensure(DetailMaterialIDAttrib))
+	{
+		return;
+	}
+	if (!ensure(DetailUVOverlay))
+	{
+		return;
+	}
+
+	BakeCache->EvaluateSamples([this, DetailMesh, DetailMaterialIDAttrib, MaterialID](const FVector2i& Coords, const FMeshImageBakingCache::FCorrespondenceSample& Sample)
+	{
+		int32 DetailTriID = Sample.DetailTriID;
+
+		if (DetailMesh->IsTriangle(DetailTriID))
+		{
+			if (DetailMaterialIDAttrib->GetValue(DetailTriID) == MaterialID)
+			{
+				FVector2d DetailUV;
+				DetailUVOverlay->GetTriBaryInterpolate<double>(DetailTriID, &Sample.DetailBaryCoords.X, &DetailUV.X);
+				FVector4f Color = SampleFunction(DetailUV);
+				ResultBuilder->SetPixel(Coords, Color);
+			}
+			// otherwise leave the pixel alone
+		}
+	});
+
+	const FImageOccupancyMap& Occupancy = *BakeCache->GetOccupancyMap();
+
+	for (int64 k = 0; k < Occupancy.GutterTexels.Num(); k++)
+	{
+		TPair<int64, int64> GutterTexel = Occupancy.GutterTexels[k];
+		ResultBuilder->CopyPixel(GutterTexel.Value, GutterTexel.Key);
+	}
+}
+
+
+void FMeshMultiResampleImageBaker::Bake()
+{
+	InitResult();
+
+	// Write into the sample buffer, separate pass for each material ID
+	for (TPair< int32, TSharedPtr<UE::Geometry::TImageBuilder<FVector4f>, ESPMode::ThreadSafe>>& MaterialTexture : MultiTextures)
+	{
+		int32 MaterialID = MaterialTexture.Key;
+		TSharedPtr<UE::Geometry::TImageBuilder<FVector4f>, ESPMode::ThreadSafe> TextureImage = MaterialTexture.Value;
+
+		this->SampleFunction = [&TextureImage](FVector2d UVCoord) 
+		{
+			return TextureImage->BilinearSampleUV<float>(UVCoord, FVector4f(0, 0, 0, 1));
+		};
+
+		BakeMaterial(MaterialID);
+	}
+
+}
diff --git a/Engine/Plugins/Experimental/GeometryProcessing/Source/DynamicMesh/Public/Sampling/MeshPropertyMapBaker.h b/Engine/Plugins/Experimental/GeometryProcessing/Source/DynamicMesh/Public/Sampling/MeshPropertyMapBaker.h
--- a/Engine/Plugins/Experimental/GeometryProcessing/Source/DynamicMesh/Public/Sampling/MeshPropertyMapBaker.h
+++ b/Engine/Plugins/Experimental/GeometryProcessing/Source/DynamicMesh/Public/Sampling/MeshPropertyMapBaker.h
@@ -17,7 +17,8 @@ enum class EMeshPropertyBakeType
 	Position = 1,
 	Normal = 2,
 	FacetNormal = 3,
-	UVPosition = 4
+	UVPosition = 4,
+	MaterialID = 5
 };
 
 
diff --git a/Engine/Plugins/Experimental/GeometryProcessing/Source/DynamicMesh/Public/Sampling/MeshResampleImageBaker.h b/Engine/Plugins/Experimental/GeometryProcessing/Source/DynamicMesh/Public/Sampling/MeshResampleImageBaker.h
--- a/Engine/Plugins/Experimental/GeometryProcessing/Source/DynamicMesh/Public/Sampling/MeshResampleImageBaker.h
+++ b/Engine/Plugins/Experimental/GeometryProcessing/Source/DynamicMesh/Public/Sampling/MeshResampleImageBaker.h
@@ -44,5 +44,21 @@ protected:
 	TUniquePtr<TImageBuilder<FVector4f>> ResultBuilder;
 };
 
+
+
+class DYNAMICMESH_API FMeshMultiResampleImageBaker : public FMeshResampleImageBaker
+{
+public:
+
+	TMap<int32, TSharedPtr<UE::Geometry::TImageBuilder<FVector4f>, ESPMode::ThreadSafe>> MultiTextures;
+
+	virtual void Bake() override;
+
+protected:
+
+	void InitResult();
+	void BakeMaterial(int32 MaterialID);
+};
+
 } // end namespace UE::Geometry
 } // end namespace UE
diff --git a/Engine/Plugins/Experimental/MeshModelingToolset/Source/MeshModelingTools/Private/BakeMeshAttributeMapsTool.cpp b/Engine/Plugins/Experimental/MeshModelingToolset/Source/MeshModelingTools/Private/BakeMeshAttributeMapsTool.cpp
--- a/Engine/Plugins/Experimental/MeshModelingToolset/Source/MeshModelingTools/Private/BakeMeshAttributeMapsTool.cpp
+++ b/Engine/Plugins/Experimental/MeshModelingToolset/Source/MeshModelingTools/Private/BakeMeshAttributeMapsTool.cpp
@@ -308,6 +308,41 @@ public:
 	}
 };
 
+
+
+class FBakeMultiTextureOp : public FBakeMapBaseOp
+{
+public:
+	typedef FBakeMapBaseOp Super;
+	using ImagePtr = TSharedPtr<UE::Geometry::TImageBuilder<FVector4f>, ESPMode::ThreadSafe>;
+
+	// inputs
+	const FDynamicMeshUVOverlay* UVOverlay = nullptr;
+	TMap<int32, ImagePtr> MaterialToTextureImageMap;
+
+	UBakeMeshAttributeMapsTool::FTexture2DImageSettings Settings;
+
+public:
+
+	virtual ~FBakeMultiTextureOp() {}
+
+	virtual void CalculateResult(FProgressCancel* Progress) override
+	{
+		Super::CalculateResult(Progress);
+		if (Progress && Progress->Cancelled())
+		{
+			return;
+		}
+
+		TUniquePtr<FMeshMultiResampleImageBaker> Baker = MakeUnique<FMeshMultiResampleImageBaker>();
+		Baker->SetCache(BakeCache.Get());
+		Baker->DetailUVOverlay = UVOverlay;
+		Baker->MultiTextures = MaterialToTextureImageMap;
+		Baker->Bake();
+		SetResult(MoveTemp(Baker));
+	}
+};
+
 /*
  * Tool
  */
@@ -447,6 +482,15 @@ void UBakeMeshAttributeMapsTool::Setup()
 	Texture2DProps->WatchProperty(Texture2DProps->UVLayer, [this](float) { bInputsDirty = true; });
 	Texture2DProps->WatchProperty(Texture2DProps->SourceTexture, [this](UTexture2D*) { bInputsDirty = true; });
 
+	MultiTextureProps = NewObject<UBakedMultiTexture2DImageProperties>(this);
+	MultiTextureProps->RestoreProperties(this);
+	AddToolPropertySource(MultiTextureProps);
+	SetToolPropertySourceEnabled(MultiTextureProps, false);
+
+	auto SetDirtyCallback = [this](TMap<int32, UTexture2D*>) { bInputsDirty = true; };
+	auto EqualsCallback = [](const TMap<int32, UTexture2D*>& A, const TMap<int32, UTexture2D*>& B) -> bool { return A.OrderIndependentCompareEqual(B); };
+	MultiTextureProps->WatchProperty(MultiTextureProps->MaterialIDSourceTextureMap, SetDirtyCallback, EqualsCallback);
+	MultiTextureProps->WatchProperty(MultiTextureProps->UVLayer, [this](float) { bInputsDirty = true; });
 
 	VisualizationProps = NewObject<UBakedOcclusionMapVisualizationProperties>(this);
 	VisualizationProps->RestoreProperties(this);
@@ -520,6 +564,7 @@ TUniquePtr<UE::Geometry::TGenericDataOperator<FMeshImageBaker>> UBakeMeshAttribu
 	case EBakeMapType::PositionImage:
 	case EBakeMapType::NormalImage:
 	case EBakeMapType::FaceNormalImage:
+	case EBakeMapType::MaterialID:
 	{
 		TUniquePtr<FBakeMeshPropertyMapOp> Op = MakeUnique<FBakeMeshPropertyMapOp>();
 		Op->DetailMesh = DetailMesh;
@@ -541,6 +586,18 @@ TUniquePtr<UE::Geometry::TGenericDataOperator<FMeshImageBaker>> UBakeMeshAttribu
 		Op->Settings = CachedTexture2DImageSettings;
 		return Op;
 	}
+	case EBakeMapType::MultiTexture:
+	{
+		TUniquePtr<FBakeMultiTextureOp> Op = MakeUnique<FBakeMultiTextureOp>();
+		Op->DetailMesh = DetailMesh;
+		Op->DetailSpatial = DetailSpatial;
+		Op->BaseMesh = &BaseMesh;
+		Op->BakeCacheSettings = CachedBakeCacheSettings;
+		Op->UVOverlay = DetailMesh->Attributes()->GetUVLayer(CachedTexture2DImageSettings.UVLayer);
+		Op->MaterialToTextureImageMap = CachedMultiTextures;
+		Op->Settings = CachedTexture2DImageSettings;
+		return Op;
+	}
 	}
 	check(false);
 	return nullptr;
@@ -556,6 +613,7 @@ void UBakeMeshAttributeMapsTool::Shutdown(EToolShutdownType ShutdownType)
 	NormalMapProps->SaveProperties(this);
 	CurvatureMapProps->SaveProperties(this);
 	Texture2DProps->SaveProperties(this);
+	MultiTextureProps->SaveProperties(this);
 	VisualizationProps->SaveProperties(this);
 
 	if (Compute)
@@ -614,7 +672,11 @@ void UBakeMeshAttributeMapsTool::Shutdown(EToolShutdownType ShutdownType)
 					bCreatedAssetOK = AssetGenerationUtil::SaveGeneratedTexture2D(AssetAPI, Settings->Result[0],
 						FString::Printf(TEXT("%s_FaceNormalImg"), *BaseName), StaticMeshAsset);
 					break;
-
+				case EBakeMapType::MaterialID:
+					FTexture2DBuilder::CopyPlatformDataToSourceData(Settings->Result[0], FTexture2DBuilder::ETextureType::Color);
+					bCreatedAssetOK = AssetGenerationUtil::SaveGeneratedTexture2D(AssetAPI, Settings->Result[0],
+																				  FString::Printf(TEXT("%s_MaterialIDImg"), *BaseName), StaticMeshAsset);
+					break;
 				case EBakeMapType::PositionImage:
 					FTexture2DBuilder::CopyPlatformDataToSourceData(Settings->Result[0], FTexture2DBuilder::ETextureType::Color);
 					bCreatedAssetOK = AssetGenerationUtil::SaveGeneratedTexture2D(AssetAPI, Settings->Result[0],
@@ -626,6 +688,13 @@ void UBakeMeshAttributeMapsTool::Shutdown(EToolShutdownType ShutdownType)
 					bCreatedAssetOK = AssetGenerationUtil::SaveGeneratedTexture2D(AssetAPI, Settings->Result[0],
 						FString::Printf(TEXT("%s_TextureImg"), *BaseName), StaticMeshAsset);
 					break;
+
+				case EBakeMapType::MultiTexture:
+					FTexture2DBuilder::CopyPlatformDataToSourceData(Settings->Result[0], FTexture2DBuilder::ETextureType::Color);
+					bCreatedAssetOK = AssetGenerationUtil::SaveGeneratedTexture2D(AssetAPI, Settings->Result[0],
+																				  FString::Printf(TEXT("%s_MultiTextureImg"), *BaseName), StaticMeshAsset);
+					break;
+
 				}
 				ensure(bCreatedAssetOK);
 			}
@@ -672,6 +741,102 @@ void UBakeMeshAttributeMapsTool::Render(IToolsContextRenderAPI* RenderAPI)
 
 
 
+int SelectTextureToBake(const TArray<UTexture*>& Textures)
+{
+	TArray<int> TextureVotes;
+	TextureVotes.Init(0, Textures.Num());
+
+	for (int TextureIndex = 0; TextureIndex < Textures.Num(); ++TextureIndex)
+	{
+		UTexture* Tex = Textures[TextureIndex];
+		UTexture2D* Tex2D = Cast<UTexture2D>(Tex);
+
+		if (Tex2D)
+		{
+			// Texture uses SRGB
+			if (Tex->SRGB != 0)
+			{
+				++TextureVotes[TextureIndex];
+			}
+
+#if WITH_EDITORONLY_DATA
+			// Texture has multiple channels
+			ETextureSourceFormat Format = Tex->Source.GetFormat();
+			if (Format == TSF_BGRA8 || Format == TSF_BGRE8 || Format == TSF_RGBA16 || Format == TSF_RGBA16F)
+			{
+				++TextureVotes[TextureIndex];
+			}
+#endif
+
+			// What else? Largest texture? Most layers? Most mipmaps?
+		}
+	}
+
+	int MaxIndex = -1;
+	int MaxVotes = -1;
+	for (int TextureIndex = 0; TextureIndex < Textures.Num(); ++TextureIndex)
+	{
+		if (TextureVotes[TextureIndex] > MaxVotes)
+		{
+			MaxIndex = TextureIndex;
+			MaxVotes = TextureVotes[TextureIndex];
+		}
+	}
+
+	return MaxIndex;
+}
+
+
+void UBakeMeshAttributeMapsTool::GetTexturesFromDetailMesh(const IPrimitiveComponentBackedTarget* DetailComponent)
+{
+	constexpr bool bGuessAtTextures = true;
+
+	MultiTextureProps->AllSourceTextures.Reset();
+	MultiTextureProps->MaterialIDSourceTextureMap.Reset();
+
+	TArray<UMaterialInterface*> Materials;
+	DetailComponent->GetOwnerComponent()->GetUsedMaterials(Materials);
+	
+	for (int32 MaterialID = 0; MaterialID < Materials.Num(); ++MaterialID)	// TODO: This won't match MaterialIDs on the FDynamicMesh3 in general, will it?
+	{
+		UMaterialInterface* MaterialInterface = Materials[MaterialID];
+		if (MaterialInterface == nullptr)
+		{
+			continue;
+		}
+
+		TArray<UTexture*> Textures;
+		MaterialInterface->GetUsedTextures(Textures, EMaterialQualityLevel::High, true, ERHIFeatureLevel::SM5, true);
+		
+		for (UTexture* Tex : Textures)
+		{
+			UTexture2D* Tex2D = Cast<UTexture2D>(Tex);
+			if (Tex2D)
+			{
+				MultiTextureProps->AllSourceTextures.Add(Tex2D);
+			}
+		}
+
+		if (bGuessAtTextures)
+		{
+			int SelectedTextureIndex = SelectTextureToBake(Textures);
+			if (SelectedTextureIndex >= 0)
+			{
+				UTexture2D* Tex2D = Cast<UTexture2D>(Textures[SelectedTextureIndex]);
+
+				// if cast fails, this will set the value to nullptr, which is fine
+				MultiTextureProps->MaterialIDSourceTextureMap.Add(MaterialID, Tex2D);	
+			}
+		}
+		else
+		{
+			MultiTextureProps->MaterialIDSourceTextureMap.Add(MaterialID, nullptr);
+		}
+	}
+
+}
+
+
 void UBakeMeshAttributeMapsTool::UpdateDetailMesh()
 {
 	IPrimitiveComponentBackedTarget* TargetComponent = TargetComponentInterface(0);
@@ -692,8 +857,11 @@ void UBakeMeshAttributeMapsTool::UpdateDetailMesh()
 	DetailSpatial = MakeShared<FDynamicMeshAABBTree3, ESPMode::ThreadSafe>();
 	DetailSpatial->SetMesh(DetailMesh.Get(), true);
 
+	GetTexturesFromDetailMesh(DetailComponent);
+
 	bInputsDirty = true;
 	DetailMeshTimestamp++;
+
 }
 
 
@@ -763,12 +931,18 @@ void UBakeMeshAttributeMapsTool::UpdateResult()
 		case EBakeMapType::NormalImage:
 		case EBakeMapType::FaceNormalImage:
 		case EBakeMapType::PositionImage:
+		case EBakeMapType::MaterialID:
 			OpState = UpdateResult_MeshProperty();
 			break;
 
 		case EBakeMapType::Texture2DImage:
 			OpState = UpdateResult_Texture2DImage();
 			break;
+
+		case EBakeMapType::MultiTexture:
+			OpState = UpdateResult_MultiTexture();
+			break;
+
 	}
 
 	// Early exit if op input parameters are invalid.
@@ -927,6 +1101,9 @@ UBakeMeshAttributeMapsTool::EOpState UBakeMeshAttributeMapsTool::UpdateResult_Me
 		case EBakeMapType::PositionImage:
 			MeshPropertyMapSettings.PropertyTypeIndex = (int32)EMeshPropertyBakeType::Position;
 			break;
+		case EBakeMapType::MaterialID:
+			MeshPropertyMapSettings.PropertyTypeIndex = (int32)EMeshPropertyBakeType::MaterialID;
+			break;
 		default:
 			check(false);		// should not be possible!
 	}
@@ -1082,6 +1259,65 @@ UBakeMeshAttributeMapsTool::EOpState UBakeMeshAttributeMapsTool::UpdateResult_Te
 }
 
 
+UBakeMeshAttributeMapsTool::EOpState UBakeMeshAttributeMapsTool::UpdateResult_MultiTexture()
+{
+	EOpState ResultState = EOpState::Complete;
+
+	int32 ImageSize = (int32)Settings->Resolution;
+	FImageDimensions Dimensions(ImageSize, ImageSize);
+
+	FTexture2DImageSettings NewSettings;
+	NewSettings.Dimensions = Dimensions;
+	NewSettings.UVLayer = MultiTextureProps->UVLayer;
+
+	const FDynamicMeshUVOverlay* UVOverlay = DetailMesh->Attributes()->GetUVLayer(NewSettings.UVLayer);
+	if (UVOverlay == nullptr)
+	{
+		GetToolManager()->DisplayMessage(LOCTEXT("InvalidUVWarning", "The Source Mesh does not have the selected UV layer"), EToolMessageLevel::UserWarning);
+		return EOpState::Invalid;
+	}
+
+	for (TPair<int32, UTexture2D*>& InputTexture : MultiTextureProps->MaterialIDSourceTextureMap)
+	{
+		if (InputTexture.Value == nullptr)
+		{
+			GetToolManager()->DisplayMessage(LOCTEXT("InvalidTextureWarning", "The Source Texture is not valid"), EToolMessageLevel::UserWarning);
+			return EOpState::Invalid;
+		}
+	}
+
+
+	CachedMultiTextures.Reset();
+
+	for ( TPair<int32, UTexture2D*>& InputTexture : MultiTextureProps->MaterialIDSourceTextureMap)
+	{
+		UTexture2D* Texture = InputTexture.Value;
+		if (!ensure(Texture != nullptr))
+		{
+			GetToolManager()->DisplayMessage(LOCTEXT("InvalidTextureWarning", "The Source Texture is not valid"), EToolMessageLevel::UserWarning);
+			return EOpState::Invalid;
+		}
+
+		int32 MaterialID = InputTexture.Key;
+		FTempTextureAccess TextureAccess(Texture);
+		CachedMultiTextures.Add(MaterialID, MakeShared<UE::Geometry::TImageBuilder<FVector4f>, ESPMode::ThreadSafe>());
+		CachedMultiTextures[MaterialID]->SetDimensions(TextureAccess.GetDimensions());
+
+		if (!TextureAccess.CopyTo(*CachedMultiTextures[MaterialID]))
+		{
+			GetToolManager()->DisplayMessage(LOCTEXT("CannotReadTextureWarning", "Cannot read from the source texture"), EToolMessageLevel::UserWarning);
+			return EOpState::Invalid;
+		}
+	}
+
+	if (!(CachedTexture2DImageSettings == NewSettings))
+	{
+		CachedTexture2DImageSettings = NewSettings;
+		ResultState = EOpState::Evaluate;
+	}
+	return ResultState;
+}
+
 
 
 
@@ -1132,6 +1368,7 @@ void UBakeMeshAttributeMapsTool::UpdateVisualization()
 	case EBakeMapType::NormalImage:
 	case EBakeMapType::FaceNormalImage:
 	case EBakeMapType::PositionImage:
+	case EBakeMapType::MaterialID:
 		Settings->Result[0] = CachedMeshPropertyMap;
 		PreviewMaterial->SetTextureParameterValue(TEXT("NormalMap"), EmptyNormalMap);
 		PreviewMaterial->SetTextureParameterValue(TEXT("OcclusionMap"), EmptyColorMapWhite);
@@ -1139,6 +1376,7 @@ void UBakeMeshAttributeMapsTool::UpdateVisualization()
 		break;
 
 	case EBakeMapType::Texture2DImage:
+	case EBakeMapType::MultiTexture:
 		Settings->Result[0] = CachedTexture2DImageMap;
 		PreviewMaterial->SetTextureParameterValue(TEXT("NormalMap"), EmptyNormalMap);
 		PreviewMaterial->SetTextureParameterValue(TEXT("OcclusionMap"), EmptyColorMapWhite);
@@ -1155,6 +1393,7 @@ void UBakeMeshAttributeMapsTool::UpdateOnModeChange()
 	SetToolPropertySourceEnabled(OcclusionMapProps, false);
 	SetToolPropertySourceEnabled(CurvatureMapProps, false);
 	SetToolPropertySourceEnabled(Texture2DProps, false);
+	SetToolPropertySourceEnabled(MultiTextureProps, false);
 
 	Settings->Result.Empty();
 	Settings->Result.Add(nullptr);
@@ -1173,6 +1412,9 @@ void UBakeMeshAttributeMapsTool::UpdateOnModeChange()
 	case EBakeMapType::Texture2DImage:
 		SetToolPropertySourceEnabled(Texture2DProps, true);
 		break;
+	case EBakeMapType::MultiTexture:
+		SetToolPropertySourceEnabled(MultiTextureProps, true);
+		break;
 	}
 
 }
@@ -1229,6 +1471,7 @@ void UBakeMeshAttributeMapsTool::OnMapsUpdated(const TUniquePtr<FMeshImageBaker>
 	case EBakeMapType::PositionImage:
 	case EBakeMapType::NormalImage:
 	case EBakeMapType::FaceNormalImage:
+	case EBakeMapType::MaterialID:
 	{
 		FMeshPropertyMapBaker* Baker = static_cast<FMeshPropertyMapBaker*>(NewResult.Get());
 		FTexture2DBuilder TextureBuilder;
@@ -1248,6 +1491,15 @@ void UBakeMeshAttributeMapsTool::OnMapsUpdated(const TUniquePtr<FMeshImageBaker>
 		CachedTexture2DImageMap = TextureBuilder.GetTexture2D();
 		break;
 	}
+	case EBakeMapType::MultiTexture:
+	{
+		FMeshMultiResampleImageBaker* Baker = static_cast<FMeshMultiResampleImageBaker*>(NewResult.Get());
+		FTexture2DBuilder TextureBuilder;
+		TextureBuilder.Initialize(FTexture2DBuilder::ETextureType::Color, CachedTexture2DImageSettings.Dimensions);
+		TextureBuilder.Copy(*Baker->GetResult(), true);
+		TextureBuilder.Commit(false);
+		CachedTexture2DImageMap = TextureBuilder.GetTexture2D();
+	}
 	}
 
 	UpdateVisualization();
diff --git a/Engine/Plugins/Experimental/MeshModelingToolset/Source/MeshModelingTools/Public/BakeMeshAttributeMapsTool.h b/Engine/Plugins/Experimental/MeshModelingToolset/Source/MeshModelingTools/Public/BakeMeshAttributeMapsTool.h
--- a/Engine/Plugins/Experimental/MeshModelingToolset/Source/MeshModelingTools/Public/BakeMeshAttributeMapsTool.h
+++ b/Engine/Plugins/Experimental/MeshModelingToolset/Source/MeshModelingTools/Public/BakeMeshAttributeMapsTool.h
@@ -63,7 +63,9 @@ enum class EBakeMapType
 	Texture2DImage,
 	NormalImage,
 	FaceNormalImage,
-	PositionImage
+	PositionImage,
+	MaterialID,
+	MultiTexture
 };
 
 
@@ -305,6 +307,26 @@ public:
 };
 
 
+UCLASS()
+class MESHMODELINGTOOLS_API UBakedMultiTexture2DImageProperties : public UInteractiveToolPropertySet
+{
+	GENERATED_BODY()
+public:
+
+	/** For each material ID, the source texture that will be resampled in that material's region*/
+	UPROPERTY(EditAnywhere, Category = MultiTexture, meta = (DisplayName = "Material IDs / Source Textures"))
+	TMap<int32, UTexture2D*> MaterialIDSourceTextureMap;
+
+	/** UV layer to sample from on the input mesh */
+	UPROPERTY(EditAnywhere, Category = MultiTexture)
+	int32 UVLayer = 0;
+
+	/** The set of all source textures from all input materials */
+	UPROPERTY(VisibleAnywhere, BlueprintReadOnly, AdvancedDisplay, Category = MultiTexture, meta = (DisplayName = "Source Textures"))
+	TArray<UTexture2D*> AllSourceTextures;
+
+};
+
 
 
 /**
@@ -352,6 +374,9 @@ protected:
 	UPROPERTY()
 	UBakedTexture2DImageProperties* Texture2DProps;
 
+	UPROPERTY()
+	UBakedMultiTexture2DImageProperties* MultiTextureProps;
+
 	UPROPERTY()
 	UBakedOcclusionMapVisualizationProperties* VisualizationProps;
 
@@ -364,6 +389,8 @@ protected:
 	friend class FBakeCurvatureMapOp;
 	friend class FBakeTexture2DImageMapOp;
 	friend class FBakeMeshPropertyMapOp;
+	friend class FBakeMultiTextureOp;
+
 
 	IAssetGenerationAPI* AssetAPI = nullptr;
 
@@ -536,6 +563,10 @@ protected:
 	EOpState UpdateResult_Texture2DImage();
 
 
+	TMap<int32, TSharedPtr<UE::Geometry::TImageBuilder<FVector4f>, ESPMode::ThreadSafe>> CachedMultiTextures;
+	EOpState UpdateResult_MultiTexture();
+
+
 
 	// empty maps are shown when nothing is computed
 
@@ -550,4 +581,6 @@ protected:
 
 	void InitializeEmptyMaps();
 
+	void GetTexturesFromDetailMesh(const IPrimitiveComponentBackedTarget* DetailComponent);
+
 };
