move fast iterative smoothing algorithm over from ModelingOperators into DynamicMesh + associated cleanup

#rnx

diff --git a/Engine/Plugins/Experimental/GeometryProcessing/Source/DynamicMesh/Private/Solvers/MeshSmoothing.cpp b/Engine/Plugins/Experimental/GeometryProcessing/Source/DynamicMesh/Private/Solvers/MeshSmoothing.cpp
--- a/Engine/Plugins/Experimental/GeometryProcessing/Source/DynamicMesh/Private/Solvers/MeshSmoothing.cpp
+++ b/Engine/Plugins/Experimental/GeometryProcessing/Source/DynamicMesh/Private/Solvers/MeshSmoothing.cpp
@@ -3,6 +3,7 @@
 #include "Solvers/MeshSmoothing.h"
 #include "Solvers/Internal/MeshDiffusionSmoothing.h"
 #include "Solvers/Internal/ConstrainedMeshSmoothers.h"
+#include "MeshWeights.h"
 
 #include "ExplicitUseGeometryMathTypes.h"		// using UE::Geometry::(math types)
 using namespace UE::Geometry;
@@ -127,3 +128,92 @@ void UE::MeshDeformation::ComputeSmoothing_Diffusion(const ELaplacianWeightSchem
 	Smoother.GetPositions(PositionArray);
 };
 
+
+void UE::MeshDeformation::ComputeSmoothing_Forward(bool bUniformWeightScheme, bool bSmoothBoundary, const FDynamicMesh3& OriginalMesh,
+	TFunctionRef<double(int VID, bool bOnBoundary)> GetSmoothingAlpha, const int32 NumIterations, TArray<FVector3d>& PositionArray)
+{
+	int32 NV = OriginalMesh.MaxVertexID();
+
+	// cache boundary verts info
+	TArray<bool> IsBoundary;
+	TArray<int32> BoundaryVerts;
+	IsBoundary.SetNum(NV);
+	ParallelFor(NV, [&OriginalMesh, &IsBoundary](int32 VID)
+	{
+		IsBoundary[VID] = OriginalMesh.IsBoundaryVertex(VID) && OriginalMesh.IsReferencedVertex(VID);
+	});
+
+	for (int32 VID = 0; VID < NV; ++VID)
+	{
+		if (IsBoundary[VID])
+		{
+			BoundaryVerts.Add(VID);
+		}
+	}
+
+	TArray<FVector3d> SecondBuffer;
+	SecondBuffer.SetNumUninitialized(NV);
+
+	TArray<FVector3d>* CurrentPositions = &PositionArray, *NextPositions = &SecondBuffer;
+
+	for (int32 Iter = 0; Iter < NumIterations; ++Iter)
+	{
+		// calculate smoothed positions of interior vertices
+		ParallelFor(NV, [bUniformWeightScheme, &OriginalMesh, &GetSmoothingAlpha, &IsBoundary, &CurrentPositions, &NextPositions](int32 VID)
+		{
+			if (OriginalMesh.IsReferencedVertex(VID) == false || IsBoundary[VID])
+			{
+				(*NextPositions)[VID] = (*CurrentPositions)[VID];
+				return;
+			}
+
+			FVector3d Centroid;
+			if (bUniformWeightScheme)
+			{
+				Centroid = FMeshWeights::UniformCentroid(OriginalMesh, VID, [&CurrentPositions](int32 NbrVID) { return (*CurrentPositions)[NbrVID]; });
+			}
+			else
+			{
+				Centroid = FMeshWeights::CotanCentroidSafe(OriginalMesh, VID, [&CurrentPositions](int32 NbrVID) { return (*CurrentPositions)[NbrVID]; }, 1.0);
+
+				// This code does not work because the mean curvature increases as things get smaller. Need to normalize it,
+				// however that *also* doesn't really work because there is a maximum step size based on edge length, and as edges
+				// collapse to nearly zero length, progress stops. Need to refine while smoothing.
+
+				//FVector3d Uniform = FMeshWeights::UniformCentroid(*ResultMesh, vid, [&](int32 nbrvid) { return PositionBuffer[nbrvid]; });
+				//FVector3d MeanCurvNorm = UE::MeshCurvature::MeanCurvatureNormal(*ResultMesh, vid, [&](int32 nbrvid) { return PositionBuffer[nbrvid]; });
+				//Centroid = PositionBuffer[vid] - 0.5*MeanCurvNorm;
+				//if (Centroid.DistanceSquared(PositionBuffer[vid]) > Uniform.DistanceSquared(PositionBuffer[vid]))
+				//{
+				//	Centroid = Uniform;
+				//}
+			}
+
+			double UseAlpha = GetSmoothingAlpha(VID, false);
+			(*NextPositions)[VID] = Lerp((*CurrentPositions)[VID], Centroid, UseAlpha);
+
+		});
+
+		// calculate boundary vertices
+		if (bSmoothBoundary)
+		{
+			ParallelFor(BoundaryVerts.Num(),
+				[&OriginalMesh, &GetSmoothingAlpha, &CurrentPositions, &NextPositions, &IsBoundary, &BoundaryVerts](int32 BdryIdx)
+			{
+				int32 VID = BoundaryVerts[BdryIdx];
+				FVector3d Centroid = FMeshWeights::FilteredUniformCentroid(OriginalMesh, VID,
+					[&CurrentPositions](int32 NbrVID) { return (*CurrentPositions)[NbrVID]; },
+					[&IsBoundary](int32 NbrVID) { return IsBoundary[NbrVID]; });
+				double UseAlpha = GetSmoothingAlpha(VID, true);
+				(*NextPositions)[VID] = Lerp((*CurrentPositions)[VID], Centroid, UseAlpha);
+			});
+		}
+
+		Swap(CurrentPositions, NextPositions);
+	}
+
+	if (CurrentPositions != &PositionArray)
+	{
+		PositionArray = MoveTemp(*CurrentPositions);
+	}
+}
\ No newline at end of file
diff --git a/Engine/Plugins/Experimental/GeometryProcessing/Source/DynamicMesh/Public/Solvers/MeshSmoothing.h b/Engine/Plugins/Experimental/GeometryProcessing/Source/DynamicMesh/Public/Solvers/MeshSmoothing.h
--- a/Engine/Plugins/Experimental/GeometryProcessing/Source/DynamicMesh/Public/Solvers/MeshSmoothing.h
+++ b/Engine/Plugins/Experimental/GeometryProcessing/Source/DynamicMesh/Public/Solvers/MeshSmoothing.h
@@ -56,5 +56,19 @@ namespace UE
 			const double Speed, double Weight, const int32 NumIterations, TArray<FVector3d>& PositionArray);
 
 
+		/**
+		 * Simple iterative smoothing with an optional weight map
+		 * Each iteration Lerps towards the weighted neighbor centroid by the GetSmoothingAlpha amount
+		 * 
+		 * @param bUniformWeightScheme If true, uses uniform weights for centroid; else uses cotan weights
+		 * @param bSmoothBoundary If true, boundary vertices are smoothed (toward the centroid of neighboring boundary vertices); else boundary is fixed
+		 * @param OriginalMesh The source mesh, used for topology
+		 * @param GetSmoothingAlpha The amount to lerp each vertex towards the centroid, per iteration
+		 * @param NumIterations Number of iterations to smooth
+		 * @param PositionsArray Vertex positions (input and output)
+		 */
+		void DYNAMICMESH_API ComputeSmoothing_Forward(bool bUniformWeightScheme, bool bSmoothBoundary, const FDynamicMesh3& OriginalMesh,
+			TFunctionRef<double(int VID, bool bBoundary)> GetSmoothingAlpha, const int32 NumIterations, TArray<FVector3d>& PositionArray);
+
 	}
 }
\ No newline at end of file
diff --git a/Engine/Plugins/Experimental/MeshModelingToolset/Source/ModelingOperators/Private/SmoothingOps/IterativeSmoothingOp.cpp b/Engine/Plugins/Experimental/MeshModelingToolset/Source/ModelingOperators/Private/SmoothingOps/IterativeSmoothingOp.cpp
--- a/Engine/Plugins/Experimental/MeshModelingToolset/Source/ModelingOperators/Private/SmoothingOps/IterativeSmoothingOp.cpp
+++ b/Engine/Plugins/Experimental/MeshModelingToolset/Source/ModelingOperators/Private/SmoothingOps/IterativeSmoothingOp.cpp
@@ -2,7 +2,6 @@
 
 #include "SmoothingOps/IterativeSmoothingOp.h"
 #include "Solvers/MeshSmoothing.h"
-#include "MeshWeights.h"
 #include "MeshCurvature.h"
 #include "Async/ParallelFor.h"
 
@@ -12,7 +11,6 @@ using namespace UE::Geometry;
 FIterativeSmoothingOp::FIterativeSmoothingOp(const FDynamicMesh3* Mesh, const FSmoothingOpBase::FOptions& OptionsIn) :
 	FSmoothingOpBase(Mesh, OptionsIn)
 {
-	SmoothedBuffer = PositionBuffer;
 }
 
 void FIterativeSmoothingOp::CalculateResult(FProgressCancel* Progress)
@@ -68,94 +66,8 @@ void FIterativeSmoothingOp::Smooth_MeanValue()
 }
 
 
-// Each iteration is equivalent to time advancing the diffusion system 
-// 
-//  d p / dt = k*k L[ p ]
-// 
-// with forward Euler
-//  
-// p^{n+1} = p^{n} + dt k * k L[ p^{n} ]
-// 
-// Where { L[p] }_i = Sum[ w_{ij} p_j ]  
-//        
-//        and w_{ii} = - Sum[ w_{ij}, j != i }
-//
-// Here SmoothSpeed = -dt * k * k * w_{ii}
 void FIterativeSmoothingOp::Smooth_Forward(bool bUniform)
 {
-	int32 NV = ResultMesh->MaxVertexID();
-
-	// cache boundary verts info
-	TArray<bool> bIsBoundary;
-	TArray<int32> BoundaryVerts;
-	bIsBoundary.SetNum(NV);
-	ParallelFor(NV, [&](int32 vid)
-	{
-		bIsBoundary[vid] = ResultMesh->IsBoundaryVertex(vid) && ResultMesh->IsReferencedVertex(vid);
-	});
-	for (int32 vid = 0; vid < NV; ++vid)
-	{
-		if (bIsBoundary[vid])
-		{
-			BoundaryVerts.Add(vid);
-		}
-	}
-
-	for (int32 k = 0; k < SmoothOptions.Iterations; ++k)
-	{
-		// calculate smoothed positions of interior vertices
-		ParallelFor(NV, [&](int32 vid) 
-		{
-			if (ResultMesh->IsReferencedVertex(vid) == false || bIsBoundary[vid])
-			{
-				SmoothedBuffer[vid] = PositionBuffer[vid];
-				return;
-			}
-
-			FVector3d Centroid;
-			if (bUniform)
-			{
-				Centroid = FMeshWeights::UniformCentroid(*ResultMesh, vid, [&](int32 nbrvid) { return PositionBuffer[nbrvid]; });
-			}
-			else
-			{
-				Centroid = FMeshWeights::CotanCentroidSafe(*ResultMesh, vid, [&](int32 nbrvid) { return PositionBuffer[nbrvid]; }, 1.0);
-
-				// This code does not work because the mean curvature increases as things get smaller. Need to normalize it,
-				// however that *also* doesn't really work because there is a maximum step size based on edge length, and as edges
-				// collapse to nearly zero length, progress stops. Need to refine while smoothing.
-
-				//FVector3d Uniform = FMeshWeights::UniformCentroid(*ResultMesh, vid, [&](int32 nbrvid) { return PositionBuffer[nbrvid]; });
-				//FVector3d MeanCurvNorm = UE::MeshCurvature::MeanCurvatureNormal(*ResultMesh, vid, [&](int32 nbrvid) { return PositionBuffer[nbrvid]; });
-				//Centroid = PositionBuffer[vid] - 0.5*MeanCurvNorm;
-				//if (Centroid.DistanceSquared(PositionBuffer[vid]) > Uniform.DistanceSquared(PositionBuffer[vid]))
-				//{
-				//	Centroid = Uniform;
-				//}
-			}
-
-			double UseAlpha = GetSmoothAlpha(vid, false);
-			SmoothedBuffer[vid] = Lerp(PositionBuffer[vid], Centroid, UseAlpha);
-
-		} /*, EParallelForFlags::ForceSingleThread*/ );
-
-		// calculate boundary vertices
-		if (SmoothOptions.bSmoothBoundary)
-		{
-			ParallelFor(BoundaryVerts.Num(), [=](int32 idx)
-			{
-				int32 vid = BoundaryVerts[idx];
-				FVector3d Centroid = FMeshWeights::FilteredUniformCentroid(*ResultMesh, vid,
-					[&](int32 nbrvid) { return PositionBuffer[nbrvid]; },
-					[&](int32 nbrvid) { return bIsBoundary[nbrvid]; });
-				double UseAlpha = GetSmoothAlpha(vid, true);
-				SmoothedBuffer[vid] = Lerp(PositionBuffer[vid], Centroid, UseAlpha);
-			});
-		}
-
-		for (int32 vid = 0; vid < NV; ++vid)
-		{
-			PositionBuffer[vid] = SmoothedBuffer[vid];
-		}
-	}
+	UE::MeshDeformation::ComputeSmoothing_Forward(bUniform, SmoothOptions.bSmoothBoundary,
+		*ResultMesh, [this](int VID, bool bBoundary) { return GetSmoothAlpha(VID, bBoundary); }, SmoothOptions.Iterations, PositionBuffer);
 }
diff --git a/Engine/Plugins/Experimental/MeshModelingToolset/Source/ModelingOperators/Public/SmoothingOps/IterativeSmoothingOp.h b/Engine/Plugins/Experimental/MeshModelingToolset/Source/ModelingOperators/Public/SmoothingOps/IterativeSmoothingOp.h
--- a/Engine/Plugins/Experimental/MeshModelingToolset/Source/ModelingOperators/Public/SmoothingOps/IterativeSmoothingOp.h
+++ b/Engine/Plugins/Experimental/MeshModelingToolset/Source/ModelingOperators/Public/SmoothingOps/IterativeSmoothingOp.h
@@ -33,9 +33,6 @@ private:
 
 	// mean value smoothing iterations
 	void Smooth_MeanValue();
-
-private:
-	TArray<FVector3d> SmoothedBuffer;
 };
 
 } // end namespace UE::Geometry
diff --git a/Engine/Plugins/Experimental/PlanarCutPlugin/Source/PlanarCut/Private/FractureAutoUV.cpp b/Engine/Plugins/Experimental/PlanarCutPlugin/Source/PlanarCut/Private/FractureAutoUV.cpp
--- a/Engine/Plugins/Experimental/PlanarCutPlugin/Source/PlanarCut/Private/FractureAutoUV.cpp
+++ b/Engine/Plugins/Experimental/PlanarCutPlugin/Source/PlanarCut/Private/FractureAutoUV.cpp
@@ -19,6 +19,7 @@
 #include "Parameterization/UVPacking.h"
 #include "Sampling/MeshCurvatureMapBaker.h"
 #include "Sampling/MeshOcclusionMapBaker.h"
+#include "Solvers/MeshSmoothing.h"
 #include "DisjointSet.h"
 
 #include "Image/ImageOccupancyMap.h"
@@ -578,43 +579,20 @@ void TextureInternalSurfaces(
 
 			FDynamicMesh3 SolidMesh(&Solidify.Generate());
 
-			// Smooth mesh (there is a more general version in MeshModelingToolset's IterativeSmoothingOp.cpp, but don't want to depend on MeshModelingToolset just for that)
-			// (note it is extra simple b/c we shouldn't have boundary vertices)
-			auto IterativeSmoothMesh = [](FDynamicMesh3& Mesh, int Iters, double SmoothAlpha)
+			// Smooth mesh
+			double SmoothAlpha = AttributeSettings.Curvature_SmoothingPerStep;
+			TArray<FVector3d> PositionBuffer;
+			PositionBuffer.SetNumUninitialized(Mesh.MaxVertexID());
+			for (int VID : Mesh.VertexIndicesItr())
 			{
-				TArray<FVector3d> SmoothedBuffer, PositionBuffer;
-				int NV = Mesh.MaxVertexID();
-				SmoothedBuffer.SetNumZeroed(NV);
-				PositionBuffer.SetNumZeroed(NV);
-				for (int VID : Mesh.VertexIndicesItr())
-				{
-					PositionBuffer[VID] = Mesh.GetVertex(VID);
-				}
-				for (int32 k = 0; k < Iters; ++k)
-				{
-					// calculate smoothed positions of interior vertices
-					ParallelFor(NV, [&Mesh, &SmoothedBuffer, &PositionBuffer, SmoothAlpha](int32 vid)
-						{
-							if (Mesh.IsReferencedVertex(vid) == false)
-							{
-								return;
-							}
-
-							FVector3d Centroid = FMeshWeights::UniformCentroid(Mesh, vid, [&](int32 nbrvid) { return PositionBuffer[nbrvid]; });
-							SmoothedBuffer[vid] = Lerp(PositionBuffer[vid], Centroid, SmoothAlpha);
-						});
-
-					for (int32 vid = 0; vid < NV; ++vid)
-					{
-						PositionBuffer[vid] = SmoothedBuffer[vid];
-					}
-				}
-				for (int VID : Mesh.VertexIndicesItr())
-				{
-					Mesh.SetVertex(VID, PositionBuffer[VID]);
-				}
-			};
-			IterativeSmoothMesh(SolidMesh, AttributeSettings.Curvature_SmoothingSteps, AttributeSettings.Curvature_SmoothingPerStep);
+				PositionBuffer[VID] = Mesh.GetVertex(VID);
+			}
+			UE::MeshDeformation::ComputeSmoothing_Forward(true, false, SolidMesh, [SmoothAlpha](int VID, bool bIsBoundary) { return SmoothAlpha; },
+				AttributeSettings.Curvature_SmoothingSteps, PositionBuffer);
+			for (int VID : Mesh.VertexIndicesItr())
+			{
+				Mesh.SetVertex(VID, PositionBuffer[VID]);
+			}
 
 			FDynamicMeshAABBTree3 SolidSpatial(&SolidMesh);
 
