GeometryProcesing: add TImageBuilder::FastDownsample. Add AntiAliasMultiSampling parameter to ApproximateActors API and Impl. Expose WindingNumber parameter in ApproximateActors API. Modify FMeshSceneAdapter to compute winding number sum instead of max.
#rnx

diff --git a/Engine/Plugins/Experimental/GeometryProcessing/Source/GeometricObjects/Public/Image/ImageBuilder.h b/Engine/Plugins/Experimental/GeometryProcessing/Source/GeometricObjects/Public/Image/ImageBuilder.h
--- a/Engine/Plugins/Experimental/GeometryProcessing/Source/GeometricObjects/Public/Image/ImageBuilder.h
+++ b/Engine/Plugins/Experimental/GeometryProcessing/Source/GeometricObjects/Public/Image/ImageBuilder.h
@@ -43,6 +43,14 @@ public:
 		Image.AssignAll(ClearValue);
 	}
 
+	/**
+	 * @return true if the given X/Y coordinates are inside the image pixel bounds, ie can be used to index a pixel
+	 */
+	bool ContainsPixel(int32 X, int32 Y) const
+	{
+		return X >= 0 && Y >= 0 && X < Dimensions.GetWidth() && Y < Dimensions.GetHeight();
+	}
+
 	/**
 	 * @return true if the given X/Y coordinates are inside the image pixel bounds, ie can be used to index a pixel
 	 */
@@ -52,6 +60,15 @@ public:
 			ImageCoords.X < Dimensions.GetWidth() && ImageCoords.Y < Dimensions.GetHeight();
 	}
 
+	/**
+	 * Get the Pixel at the given X/Y coordinates
+	 */
+	const PixelType& GetPixel(int32 X, int32 Y) const
+	{
+		int64 LinearIndex = Dimensions.GetIndex(FVector2i(X,Y));
+		return Image[LinearIndex];
+	}
+
 	/**
 	 * Get the Pixel at the given X/Y coordinates
 	 */
@@ -61,7 +78,6 @@ public:
 		return Image[LinearIndex];
 	}
 
-
 	/**
 	 * Get the Pixel at the given linear index
 	 */
@@ -71,6 +87,15 @@ public:
 	}
 
 
+	/**
+	 * Set the Pixel at the given X/Y coordinates to the given PixelType
+	 */
+	void SetPixel(int32 X, int32 Y, const PixelType& NewValue)
+	{
+		int64 LinearIndex = Dimensions.GetIndex(FVector2i(X,Y));
+		Image[LinearIndex] = NewValue;
+	}
+
 	/**
 	 * Set the Pixel at the given X/Y coordinates to the given PixelType
 	 */
@@ -166,6 +191,58 @@ public:
 		return BilinearSample<ScalarType>(PixelCoords, InvalidValue);
 	}
 
+	/**
+	 * Very basic downsampling technqiue that just averages NxN pixel blocks. Multi-threaded.
+	 * @param SubSteps each NxN pixel block of this size is averaged into 1 pixel in the output image
+	 * @param Zero value to use with template PixelType must be provided, pixel values will be added to this value
+	 * @param AverageFunc Called with Sum(Pixels) and PixelCount, return value is used as new pixel value (eg should be average)
+	 */
+	TImageBuilder<PixelType> FastDownsample(
+		int32 SubSteps, 
+		const PixelType& ZeroValue, 
+		TFunctionRef<PixelType(PixelType, int)> AverageFunc) const
+	{
+		TImageBuilder<PixelType> DownsampleImage;
+		int32 Width = Dimensions.GetWidth();
+		int32 Height = Dimensions.GetHeight();
+
+		// can only fast-downsample by even multiple of image size
+		if (ensure((Width % SubSteps == 0) && (Height % SubSteps == 0)) == false)
+		{
+			DownsampleImage = *this;
+			return DownsampleImage;
+		}
+
+		int32 SubWidth = Width / SubSteps;
+		int32 SubHeight = Height / SubSteps;
+
+		FImageDimensions SubDimensions(SubWidth, SubHeight);
+		DownsampleImage.SetDimensions(SubDimensions);
+
+		ParallelFor(SubHeight, [&](int32 yi)
+		{
+			int32 baseyi = yi * SubSteps;
+			for (int32 xi = 0; xi < SubWidth; ++xi)
+			{
+				int32 basexi = xi * SubSteps;
+
+				PixelType AccumBasePixels = ZeroValue;
+				for (int32 dy = 0; dy < SubSteps; ++dy)
+				{
+					for (int32 dx = 0; dx < SubSteps; ++dx)
+					{
+						AccumBasePixels += GetPixel(basexi + dx, baseyi + dy);
+					}
+				}
+				PixelType SubPixel = AverageFunc(AccumBasePixels, SubSteps * SubSteps);
+				DownsampleImage.SetPixel(xi, yi, SubPixel);
+			}
+		});
+
+		return DownsampleImage;
+	}
+	
+
 };
 
 
diff --git a/Engine/Plugins/Experimental/MeshModelingToolset/Source/GeometryProcessingAdapters/Private/ApproximateActorsImpl.cpp b/Engine/Plugins/Experimental/MeshModelingToolset/Source/GeometryProcessingAdapters/Private/ApproximateActorsImpl.cpp
--- a/Engine/Plugins/Experimental/MeshModelingToolset/Source/GeometryProcessingAdapters/Private/ApproximateActorsImpl.cpp
+++ b/Engine/Plugins/Experimental/MeshModelingToolset/Source/GeometryProcessingAdapters/Private/ApproximateActorsImpl.cpp
@@ -59,8 +59,15 @@ static void BakeTexturesFromPhotoCapture(
 	double NearPlaneDist = Options.NearPlaneDist;
 	double RayOffsetHackDist = (double)(100.0f * FMathf::ZeroTolerance);
 
+	int32 Supersample = FMath::Max(1, Options.AntiAliasMultiSampling);
+	if ( (Options.TextureImageSize * Supersample) > 16384)
+	{
+		UE_LOG(LogTemp, Warning, TEXT("[ApproximateActors] Ignoring requested supersampling rate %d because it would require image buffers with resolution %d, please try lower value."), Supersample, Options.TextureImageSize * Supersample);
+		Supersample = 1;
+	}
+
 	FImageDimensions CaptureDimensions(Options.RenderCaptureImageSize, Options.RenderCaptureImageSize);
-	FImageDimensions OutputDimensions(Options.TextureImageSize, Options.TextureImageSize);
+	FImageDimensions OutputDimensions(Options.TextureImageSize*Supersample, Options.TextureImageSize*Supersample);
 
 	FSceneCapturePhotoSet SceneCapture;
 	SceneCapture.SetCaptureSceneActors(Actors[0]->GetWorld(), Actors);
@@ -117,7 +124,7 @@ static void BakeTexturesFromPhotoCapture(
 		return ColorImage->GetPixel(Coords) == DefaultSample.BaseColor;
 	}, MissingPixels);
 
-
+	// solve infill for the holes while also caching infill information
 	TMarchingPixelInfill<FVector4f> Infill;
 	Infill.ComputeInfill(*ColorImage, MissingPixels, DefaultSample.BaseColor,
 		[](FVector4f SumValue, int32 Count) {
@@ -125,7 +132,14 @@ static void BakeTexturesFromPhotoCapture(
 		return FVector4f(SumValue.X * InvSum, SumValue.Y * InvSum, SumValue.Z * InvSum, 1.0f);
 	});
 
+	// downsample the image if necessary
+	if (Supersample > 1)
+	{
+		TImageBuilder<FVector4f> Downsampled = ColorImage->FastDownsample(Supersample, FVector4f::Zero(), [](FVector4f V, int N) { return V / (float)N; });
+		*ColorImage = MoveTemp(Downsampled);
+	}
 
+	// this lambda is used to process the per-channel images. It does the bake, applies infill, and downsamples if necessary
 	auto ProcessChannelFunc = [&](ERenderCaptureType CaptureType)
 	{
 		FVector4f DefaultValue(0, 0, 0, 0);
@@ -145,6 +159,12 @@ static void BakeTexturesFromPhotoCapture(
 			return FVector4f(SumValue.X * InvSum, SumValue.Y * InvSum, SumValue.Z * InvSum, 1.0f);
 		});
 
+		if (Supersample > 1)
+		{
+			TImageBuilder<FVector4f> Downsampled = Image->FastDownsample(Supersample, FVector4f::Zero(), [](FVector4f V, int N) { return V / (float)N; });
+			*Image = MoveTemp(Downsampled);
+		}
+
 		return MoveTemp(Image);
 	};
 
@@ -178,6 +198,11 @@ static void BakeTexturesFromPhotoCapture(
 	NormalMapBaker.Bake();
 	TUniquePtr<TImageBuilder<FVector3f>> NormalImage = NormalMapBaker.TakeResult();
 
+	if (Supersample > 1)
+	{
+		TImageBuilder<FVector3f> Downsampled = NormalImage->FastDownsample(Supersample, FVector3f::Zero(), [](FVector3f V, int N) { return V / (float)N; });
+		*NormalImage = MoveTemp(Downsampled);
+	}
 
 	// build textures
 	Progress.EnterProgressFrame(1.f, LOCTEXT("BuildingTextures", "Building Textures..."));
@@ -247,9 +272,10 @@ void FApproximateActorsImpl::GenerateApproximationForActorSet(const TArray<AActo
 	Progress.EnterProgressFrame(1.f, LOCTEXT("GeneratingMesh", "Generating Mesh..."));
 
 	FWindingNumberBasedSolidify Solidify(
-		[&Scene](const FVector3d& Position) { return Scene.MaxFastWindingNumber(Position); },
+		[&Scene](const FVector3d& Position) { return Scene.FastWindingNumber(Position); },
 		SceneBounds, SeedPoints);
 	Solidify.SetCellSizeAndExtendBounds(SceneBounds, 0, VoxelDimTarget);
+	Solidify.WindingThreshold = Options.WindingThreshold;
 
 	FDynamicMesh3 SolidMesh(&Solidify.Generate());
 	SolidMesh.DiscardAttributes();
diff --git a/Engine/Plugins/Experimental/MeshModelingToolset/Source/ModelingComponents/Private/Scene/MeshSceneAdapter.cpp b/Engine/Plugins/Experimental/MeshModelingToolset/Source/ModelingComponents/Private/Scene/MeshSceneAdapter.cpp
--- a/Engine/Plugins/Experimental/MeshModelingToolset/Source/ModelingComponents/Private/Scene/MeshSceneAdapter.cpp
+++ b/Engine/Plugins/Experimental/MeshModelingToolset/Source/ModelingComponents/Private/Scene/MeshSceneAdapter.cpp
@@ -62,13 +62,12 @@ public:
 		if (!SourceMesh) return;
 
 		int32 NumVertices = Adapter->VertexCount();
-		int32 LogNumVertices = FMathd::Log(NumVertices);
-		int32 SeedPointCount = (int)(5.0 * LogNumVertices * LogNumVertices);
-		ensure(SeedPointCount < 500);		// just a sanity check
+		int32 LogNumVertices = FMath::Max(1, (int32)FMathd::Ceil(FMathd::Log(NumVertices)));
+		int32 SeedPointCount = (int)(10 * LogNumVertices);
+		SeedPointCount = FMath::Min(SeedPointCount, 5000);
 		int32 Skip = FMath::Max(NumVertices / SeedPointCount, 2);
 		for (int32 k = 0; k < NumVertices; k += Skip)
 		{
-			Adapter->GetNormal(k);
 			WorldPoints.Add(LocalToWorldFunc(Adapter->GetVertex(k)));
 		}
 	}
@@ -248,9 +247,9 @@ void FMeshSceneAdapter::CollectMeshSeedPoints(TArray<FVector3d>& Points)
 	}
 }
 
-double FMeshSceneAdapter::MaxFastWindingNumber(const FVector3d& P)
+double FMeshSceneAdapter::FastWindingNumber(const FVector3d& P)
 {
-	double MaxWinding = 0.0;
+	double SumWinding = 0.0;
 	for (const TUniquePtr<FActorAdapter>& Actor : SceneActors)
 	{
 		if (Actor->WorldBounds.Contains(P))
@@ -260,11 +259,11 @@ double FMeshSceneAdapter::MaxFastWindingNumber(const FVector3d& P)
 				if (ChildMesh.MeshSpatial != nullptr)
 				{
 					FVector3d LocalP = ChildMesh.WorldTransformInverse.TransformPosition(P);
-					double LocalWinding = ChildMesh.MeshSpatial->FastWindingNumber(LocalP);
-					MaxWinding = FMathd::Max(LocalWinding, MaxWinding);
+					double MeshWinding = ChildMesh.MeshSpatial->FastWindingNumber(LocalP);
+					SumWinding += MeshWinding;
 				}
 			}
 		}
 	}
-	return MaxWinding;
+	return SumWinding;
 }
\ No newline at end of file
diff --git a/Engine/Plugins/Experimental/MeshModelingToolset/Source/ModelingComponents/Public/Scene/MeshSceneAdapter.h b/Engine/Plugins/Experimental/MeshModelingToolset/Source/ModelingComponents/Public/Scene/MeshSceneAdapter.h
--- a/Engine/Plugins/Experimental/MeshModelingToolset/Source/ModelingComponents/Public/Scene/MeshSceneAdapter.h
+++ b/Engine/Plugins/Experimental/MeshModelingToolset/Source/ModelingComponents/Public/Scene/MeshSceneAdapter.h
@@ -164,8 +164,8 @@ public:
 	/** @return a set of points on the surface of the meshes, can be used to initialize the MarchingCubes mesher */
 	virtual void CollectMeshSeedPoints(TArray<FVector3d>& PointsOut);
 
-	/** @return max FastWindingNumber computed for all mesh Actors/Components */
-	virtual double MaxFastWindingNumber(const FVector3d& P);
+	/** @return FastWindingNumber computed across all mesh Actors/Components */
+	virtual double FastWindingNumber(const FVector3d& P);
 
 protected:
 
diff --git a/Engine/Source/Developer/GeometryProcessingInterfaces/Public/GeometryProcessingInterfaces/ApproximateActors.h b/Engine/Source/Developer/GeometryProcessingInterfaces/Public/GeometryProcessingInterfaces/ApproximateActors.h
--- a/Engine/Source/Developer/GeometryProcessingInterfaces/Public/GeometryProcessingInterfaces/ApproximateActors.h
+++ b/Engine/Source/Developer/GeometryProcessingInterfaces/Public/GeometryProcessingInterfaces/ApproximateActors.h
@@ -63,6 +63,8 @@ public:
 		float WorldSpaceApproximationAccuracyMeters = 1.0f;
 		int32 ClampVoxelDimension = 512;
 
+		float WindingThreshold = 0.5;
+
 		bool bApplyMorphology = false;
 		float MorphologyDistanceMeters = 0.1;
 
@@ -87,6 +89,8 @@ public:
 		// output texture options
 		int32 TextureImageSize = 1024;
 
+		// supersampling parameter
+		int32 AntiAliasMultiSampling = 0;
 
 
 
