ModelingComponents: convert checks to ensures where possible in Base/SimpleDynamicMeshSceneProxy, to avoid hard crashes when these things come up during debugging.
#rnx

diff --git a/Engine/Plugins/Experimental/MeshModelingToolset/Source/ModelingComponents/Private/BaseDynamicMeshSceneProxy.h b/Engine/Plugins/Experimental/MeshModelingToolset/Source/ModelingComponents/Private/BaseDynamicMeshSceneProxy.h
--- a/Engine/Plugins/Experimental/MeshModelingToolset/Source/ModelingComponents/Private/BaseDynamicMeshSceneProxy.h
+++ b/Engine/Plugins/Experimental/MeshModelingToolset/Source/ModelingComponents/Private/BaseDynamicMeshSceneProxy.h
@@ -391,12 +391,14 @@ public:
 	 */
 	virtual void ReleaseRenderBufferSet(FMeshRenderBufferSet* BufferSet)
 	{
-		AllocatedSetsLock.Lock();
-		check(AllocatedBufferSets.Contains(BufferSet));
-		AllocatedBufferSets.Remove(BufferSet);
-		AllocatedSetsLock.Unlock();
+		FScopeLock Lock(&AllocatedSetsLock);
+		if (ensure(AllocatedBufferSets.Contains(BufferSet)))
+		{
+			AllocatedBufferSets.Remove(BufferSet);
+			Lock.Unlock();
 
-		FMeshRenderBufferSet::DestroyRenderBufferSet(BufferSet);
+			FMeshRenderBufferSet::DestroyRenderBufferSet(BufferSet);
+		}
 	}
 
 
@@ -546,8 +548,10 @@ public:
 		const FDynamicMesh3* Mesh,
 		bool bDuplicate)
 	{
-		check(bUseSecondaryTriBuffers == true);
-		check(RenderBuffers->Triangles.IsSet());
+		if (ensure(bUseSecondaryTriBuffers == true && RenderBuffers->Triangles.IsSet()) == false)
+		{
+			return;
+		}
 
 		const TArray<int32>& TriangleIDs = RenderBuffers->Triangles.GetValue();
 		int NumTris = TriangleIDs.Num();
@@ -592,7 +596,10 @@ public:
 		{
 			return;
 		}
-		check(RenderBuffers->Triangles.IsSet() && RenderBuffers->Triangles->Num() > 0);
+		if (ensure(RenderBuffers->Triangles.IsSet() && RenderBuffers->Triangles->Num() > 0) == false)
+		{
+			return;
+		}
 
 		//bool bDuplicate = false;		// flag for future use, in case we want to draw all triangles in primary and duplicates in secondary...
 		RenderBuffers->IndexBuffer.Indices.Reset();
@@ -648,14 +655,11 @@ public:
 		bool bHaveColors = (ColorOverlay != nullptr) && (bIgnoreVertexColors == false);
 
 		int NumVertices = NumTriangles * 3;
-		check(RenderBuffers->PositionVertexBuffer.GetNumVertices() == NumVertices);
-		if (bUpdateNormals)
-		{
-			check(RenderBuffers->StaticMeshVertexBuffer.GetNumVertices() == NumVertices);
-		}
-		if (bUpdateColors)
+		if ( (bUpdatePositions && ensure(RenderBuffers->PositionVertexBuffer.GetNumVertices() == NumVertices) == false )
+			|| (bUpdateNormals && ensure(RenderBuffers->StaticMeshVertexBuffer.GetNumVertices() == NumVertices) == false )
+			|| (bUpdateColors && ensure(RenderBuffers->ColorVertexBuffer.GetNumVertices() == NumVertices) == false ) )
 		{
-			check(RenderBuffers->ColorVertexBuffer.GetNumVertices() == NumVertices);
+			return;
 		}
 
 		int VertIdx = 0;
@@ -721,7 +725,10 @@ public:
 			return;
 		}
 		int NumVertices = NumTriangles * 3;
-		check(RenderBuffers->StaticMeshVertexBuffer.GetNumVertices() == NumVertices);
+		if ( ensure(RenderBuffers->StaticMeshVertexBuffer.GetNumVertices() == NumVertices) == false)
+		{
+			return;
+		}
 
 		int VertIdx = 0;
 		for (int TriangleID : Enumerable)
diff --git a/Engine/Plugins/Experimental/MeshModelingToolset/Source/ModelingComponents/Private/SimpleDynamicMeshSceneProxy.h b/Engine/Plugins/Experimental/MeshModelingToolset/Source/ModelingComponents/Private/SimpleDynamicMeshSceneProxy.h
--- a/Engine/Plugins/Experimental/MeshModelingToolset/Source/ModelingComponents/Private/SimpleDynamicMeshSceneProxy.h
+++ b/Engine/Plugins/Experimental/MeshModelingToolset/Source/ModelingComponents/Private/SimpleDynamicMeshSceneProxy.h
@@ -50,7 +50,7 @@ public:
 	void Initialize()
 	{
 		// allocate buffer sets based on materials
-		check(RenderBufferSets.Num() == 0);
+		ensure(RenderBufferSets.Num() == 0);
 		int32 NumMaterials = GetNumMaterials();
 		if (NumMaterials == 0)
 		{
@@ -89,7 +89,7 @@ public:
 	 */
 	void InitializeFromDecomposition(TUniquePtr<FMeshRenderDecomposition>& Decomposition)
 	{
-		check(RenderBufferSets.Num() == 0);
+		ensure(RenderBufferSets.Num() == 0);
 		int32 NumSets = Decomposition->Num();
 		RenderBufferSets.SetNum(NumSets);
 		for (int32 k = 0; k < NumSets; ++k)
@@ -201,7 +201,10 @@ public:
 	void InitializeByMaterial(TArray<FMeshRenderBufferSet*>& BufferSets)
 	{
 		const FDynamicMesh3* Mesh = ParentComponent->GetRenderMesh();
-		check(Mesh->HasAttributes() && Mesh->Attributes()->HasMaterialID());
+		if (ensure(Mesh->HasAttributes() && Mesh->Attributes()->HasMaterialID()) == false)
+		{
+			return;
+		}
 
 		const FDynamicMeshAttributeSet* Attributes = Mesh->Attributes();
 
@@ -319,31 +322,34 @@ public:
 
 		if (bIsSingleBuffer)
 		{
-			check(RenderBufferSets.Num() == 1);
-
-			FMeshRenderBufferSet* BufferSet = RenderBufferSets[0];
-			if (BufferSet->TriangleCount != Mesh->TriangleCount())
+			if (ensure(RenderBufferSets.Num() == 1))
 			{
-				return false;
-			}
+				FMeshRenderBufferSet* BufferSet = RenderBufferSets[0];
+				if (BufferSet->TriangleCount != Mesh->TriangleCount())
+				{
+					return false;
+				}
 
-			int NumTriangles = Mesh->TriangleCount();
+				int NumTriangles = Mesh->TriangleCount();
 
-			if (!CheckBufferSet(Mesh, BufferSet, NumTriangles))
-			{
-				return false;
+				if (!CheckBufferSet(Mesh, BufferSet, NumTriangles))
+				{
+					return false;
+				}
 			}
 		}
 		else
 		{
 			for (FMeshRenderBufferSet* BufferSet : RenderBufferSets)
 			{
-				check(BufferSet->Triangles);
-				int NumTriangles = BufferSet->Triangles->Num();
-
-				if (!CheckBufferSet(Mesh, BufferSet, NumTriangles))
+				if (ensure(BufferSet->Triangles))
 				{
-					return false;
+					int NumTriangles = BufferSet->Triangles->Num();
+
+					if (!CheckBufferSet(Mesh, BufferSet, NumTriangles))
+					{
+						return false;
+					}
 				}
 			}
 		}
@@ -367,9 +373,8 @@ public:
 		FDynamicMeshNormalOverlay* NormalOverlay = nullptr;
 		TFunction<void(int, int, int, const FVector3f&, FVector3f&, FVector3f&)> TangentsFunc =
 			[](int, int, int, const FVector3f& Normal, FVector3f& TangentX, FVector3f& TangentY) { VectorUtil::MakePerpVectors(Normal, TangentX, TangentY); };
-		if (bNormals)
+		if (bNormals && ensure(Mesh->HasAttributes()) )
 		{
-			check(Mesh->HasAttributes());
 			NormalOverlay = Mesh->Attributes()->PrimaryNormals();
 
 			const UE::Geometry::FMeshTangentsf* Tangents = ParentComponent->GetTangents();
@@ -382,70 +387,72 @@ public:
 			}
 		}
 		FDynamicMeshUVOverlay* UVOVerlay = nullptr;
-		if (bUVs)
+		if (bUVs && ensure(Mesh->HasAttributes()) )
 		{
-			check(Mesh->HasAttributes());
 			UVOVerlay = Mesh->Attributes()->PrimaryUV();
 		}
 		FDynamicMeshColorOverlay* ColorOverlay = nullptr;
-		if (bColors)
+		if (bColors && ensure(Mesh->HasAttributes()) )
 		{
-			check(Mesh->HasAttributes());
 			ColorOverlay = Mesh->Attributes()->PrimaryColors();
 					
 		}
 		if (bIsSingleBuffer)
 		{
-			check(RenderBufferSets.Num() == 1);
-			FMeshRenderBufferSet* Buffers = RenderBufferSets[0];
-			if (bPositions || bNormals || bColors)
-			{
-				UpdateVertexBuffersFromOverlays(Buffers, Mesh,
-					Mesh->TriangleCount(), Mesh->TriangleIndicesItr(),
-					NormalOverlay, ColorOverlay, TangentsFunc,
-					bPositions, bNormals, bColors);
-			}
-			if (bUVs)
-			{
-				UpdateVertexUVBufferFromOverlays(Buffers, Mesh,
-					Mesh->TriangleCount(), Mesh->TriangleIndicesItr(), UVOVerlay, 0);
-			}
-
-
-			ENQUEUE_RENDER_COMMAND(FSimpleDynamicMeshSceneProxyFastUpdateVertices)(
-				[Buffers, bPositions, bNormals, bColors, bUVs](FRHICommandListImmediate& RHICmdList)
-			{
-				Buffers->UploadVertexUpdate(bPositions, bNormals || bUVs, bColors);
-			});
-		}
-		else
-		{
-			ParallelFor(RenderBufferSets.Num(), [&](int i)
+			if (ensure(RenderBufferSets.Num() == 1))
 			{
-				FMeshRenderBufferSet* Buffers = RenderBufferSets[i];
-				if (Buffers->TriangleCount == 0)
-				{
-					return;
-				}
-				check(Buffers->Triangles.IsSet());
+				FMeshRenderBufferSet* Buffers = RenderBufferSets[0];
 				if (bPositions || bNormals || bColors)
 				{
 					UpdateVertexBuffersFromOverlays(Buffers, Mesh,
-						Buffers->Triangles->Num(), Buffers->Triangles.GetValue(),
+						Mesh->TriangleCount(), Mesh->TriangleIndicesItr(),
 						NormalOverlay, ColorOverlay, TangentsFunc,
 						bPositions, bNormals, bColors);
 				}
 				if (bUVs)
 				{
 					UpdateVertexUVBufferFromOverlays(Buffers, Mesh,
-						Buffers->Triangles->Num(), Buffers->Triangles.GetValue(), UVOVerlay, 0);
+						Mesh->TriangleCount(), Mesh->TriangleIndicesItr(), UVOVerlay, 0);
 				}
 
+
 				ENQUEUE_RENDER_COMMAND(FSimpleDynamicMeshSceneProxyFastUpdateVertices)(
 					[Buffers, bPositions, bNormals, bColors, bUVs](FRHICommandListImmediate& RHICmdList)
 				{
 					Buffers->UploadVertexUpdate(bPositions, bNormals || bUVs, bColors);
 				});
+			}
+		}
+		else
+		{
+			ParallelFor(RenderBufferSets.Num(), [&](int i)
+			{
+				FMeshRenderBufferSet* Buffers = RenderBufferSets[i];
+				if (Buffers->TriangleCount == 0)
+				{
+					return;
+				}
+				if (ensure(Buffers->Triangles.IsSet()))
+				{
+					if (bPositions || bNormals || bColors)
+					{
+						UpdateVertexBuffersFromOverlays(Buffers, Mesh,
+							Buffers->Triangles->Num(), Buffers->Triangles.GetValue(),
+							NormalOverlay, ColorOverlay, TangentsFunc,
+							bPositions, bNormals, bColors);
+					}
+					if (bUVs)
+					{
+						UpdateVertexUVBufferFromOverlays(Buffers, Mesh,
+							Buffers->Triangles->Num(), Buffers->Triangles.GetValue(), UVOVerlay, 0);
+					}
+
+					ENQUEUE_RENDER_COMMAND(FSimpleDynamicMeshSceneProxyFastUpdateVertices)(
+						[Buffers, bPositions, bNormals, bColors, bUVs](FRHICommandListImmediate& RHICmdList)
+					{
+						Buffers->UploadVertexUpdate(bPositions, bNormals || bUVs, bColors);
+					});
+				}
 			});
 		}
 	}
@@ -474,9 +481,8 @@ public:
 		FDynamicMeshNormalOverlay* NormalOverlay = nullptr;
 		TFunction<void(int, int, int, const FVector3f&, FVector3f&, FVector3f&)> TangentsFunc =
 			[](int, int, int, const FVector3f& Normal, FVector3f& TangentX, FVector3f& TangentY) { VectorUtil::MakePerpVectors(Normal, TangentX, TangentY); };
-		if (bNormals)
+		if (bNormals && ensure(Mesh->HasAttributes()))
 		{
-			check(Mesh->HasAttributes());
 			NormalOverlay = Mesh->Attributes()->PrimaryNormals();
 
 			const UE::Geometry::FMeshTangentsf* Tangents = ParentComponent->GetTangents();
@@ -489,15 +495,13 @@ public:
 			}
 		}
 		FDynamicMeshUVOverlay* UVOVerlay = nullptr;
-		if (bUVs)
+		if (bUVs && ensure(Mesh->HasAttributes()))
 		{
-			check(Mesh->HasAttributes());
 			UVOVerlay = Mesh->Attributes()->PrimaryUV();
 		}
 		FDynamicMeshColorOverlay* ColorOverlay = nullptr;
-		if (bColors)
+		if (bColors && ensure(Mesh->HasAttributes()))
 		{
-			check(Mesh->HasAttributes());
 			ColorOverlay = Mesh->Attributes()->PrimaryColors();
 
 		}
@@ -509,25 +513,27 @@ public:
 				return;
 			}
 			FMeshRenderBufferSet* Buffers = RenderBufferSets[BufferIndex];
-			check(Buffers->Triangles.IsSet());
-			if (bPositions || bNormals || bColors)
+			if (ensure(Buffers->Triangles.IsSet()))
 			{
-				UpdateVertexBuffersFromOverlays(Buffers, Mesh,
-					Buffers->Triangles->Num(), Buffers->Triangles.GetValue(),
-					NormalOverlay, ColorOverlay, TangentsFunc,
-					bPositions, bNormals, bColors);
-			}
-			if (bUVs)
-			{
-				UpdateVertexUVBufferFromOverlays(Buffers, Mesh,
-					Buffers->Triangles->Num(), Buffers->Triangles.GetValue(), UVOVerlay, 0);
-			}
+				if (bPositions || bNormals || bColors)
+				{
+					UpdateVertexBuffersFromOverlays(Buffers, Mesh,
+						Buffers->Triangles->Num(), Buffers->Triangles.GetValue(),
+						NormalOverlay, ColorOverlay, TangentsFunc,
+						bPositions, bNormals, bColors);
+				}
+				if (bUVs)
+				{
+					UpdateVertexUVBufferFromOverlays(Buffers, Mesh,
+						Buffers->Triangles->Num(), Buffers->Triangles.GetValue(), UVOVerlay, 0);
+				}
 
-			ENQUEUE_RENDER_COMMAND(FSimpleDynamicMeshSceneProxyFastUpdateVerticesBufferList)(
-				[Buffers, bPositions, bNormals, bColors, bUVs](FRHICommandListImmediate& RHICmdList)
-			{
-				Buffers->TransferVertexUpdateToGPU(bPositions, bNormals, bUVs, bColors);
-			});
+				ENQUEUE_RENDER_COMMAND(FSimpleDynamicMeshSceneProxyFastUpdateVerticesBufferList)(
+					[Buffers, bPositions, bNormals, bColors, bUVs](FRHICommandListImmediate& RHICmdList)
+				{
+					Buffers->TransferVertexUpdateToGPU(bPositions, bNormals, bUVs, bColors);
+				});
+			}
 		});
 	}
 
