MeshModelingTools: Disable material verification for in FBaseDynamicMeshSceneProxy while enqueued render command is in flight.

#rnx

diff --git a/Engine/Plugins/Experimental/MeshModelingToolset/Source/ModelingComponents/Private/BaseDynamicMeshSceneProxy.h b/Engine/Plugins/Experimental/MeshModelingToolset/Source/ModelingComponents/Private/BaseDynamicMeshSceneProxy.h
--- a/Engine/Plugins/Experimental/MeshModelingToolset/Source/ModelingComponents/Private/BaseDynamicMeshSceneProxy.h
+++ b/Engine/Plugins/Experimental/MeshModelingToolset/Source/ModelingComponents/Private/BaseDynamicMeshSceneProxy.h
@@ -771,14 +771,22 @@ public:
 	 */
 	virtual void UpdatedReferencedMaterials()
 	{
-		// copied from FPrimitiveSceneProxy::FPrimitiveSceneProxy()
 #if WITH_EDITOR
 		TArray<UMaterialInterface*> Materials;
 		ParentBaseComponent->GetUsedMaterials(Materials, true);
+
+		// Temporarily disable material verification while the enqueued render command is in flight.
+		// The original value for bVerifyUsedMaterials gets restored when the command is executed.
+		// If we do not do this, material verification might spuriously fail in cases where the render command for changing
+		// the verfifcation material is still in flight but the render thread is already trying to render the mesh.
+		const uint8 bRestoreVerifyUsedMaterials = bVerifyUsedMaterials;
+		bVerifyUsedMaterials = false;
+
 		ENQUEUE_RENDER_COMMAND(FMeshRenderBufferSetDestroy)(
-			[this, Materials](FRHICommandListImmediate& RHICmdList)
+			[this, Materials, bRestoreVerifyUsedMaterials](FRHICommandListImmediate& RHICmdList)
 		{
 			this->SetUsedMaterialForVerification(Materials);
+			this->bVerifyUsedMaterials = bRestoreVerifyUsedMaterials;
 		});
 #endif
 	}
