GeometryTools:  Simplify Mesh Tool.  Don't recompute simplification when toggling display of UV seams (it was recomputing the simplification when you toggled this display option).
                 Change color of UV seams to match the color used in the mesh inspector.
                 Also update the option name to "Show UV Seams" from "Show Constraints"

diff --git a/Engine/Plugins/Experimental/MeshModelingToolset/Source/MeshModelingToolsEditorOnly/Private/SimplifyMeshTool.cpp b/Engine/Plugins/Experimental/MeshModelingToolset/Source/MeshModelingToolsEditorOnly/Private/SimplifyMeshTool.cpp
--- a/Engine/Plugins/Experimental/MeshModelingToolset/Source/MeshModelingToolsEditorOnly/Private/SimplifyMeshTool.cpp
+++ b/Engine/Plugins/Experimental/MeshModelingToolset/Source/MeshModelingToolsEditorOnly/Private/SimplifyMeshTool.cpp
@@ -105,6 +105,8 @@ void USimplifyMeshTool::Setup()
 		IMeshDescriptionProvider* TargetMeshProvider = Cast<IMeshDescriptionProvider>(Target);
 		OriginalMeshDescription = MakeShared<FMeshDescription, ESPMode::ThreadSafe>(*TargetMeshProvider->GetMeshDescription());
 		// make sure the mesh description has tangents. 
+		const bool bUseTangents = true;
+		if (bUseTangents)
 		{
 			// if the source was imported with auto-generate tangents, the mesh description tangents won't have been populated
 			FMeshBuildSettings BuildSettings; 
@@ -117,7 +119,6 @@ void USimplifyMeshTool::Setup()
 		OriginalMesh = MakeShared<FDynamicMesh3, ESPMode::ThreadSafe>();
 		FMeshDescriptionToDynamicMesh Converter; 
 		// convert with tangent overlay
-		const bool bUseTangents = true;
 		Converter.Convert(TargetMeshProvider->GetMeshDescription(), *OriginalMesh, bUseTangents);
 		EnterProgressFrame(2);
 		OriginalMeshSpatial = MakeShared<FDynamicMeshAABBTree3, ESPMode::ThreadSafe>(OriginalMesh.Get(), true);
@@ -213,9 +214,9 @@ void USimplifyMeshTool::Render(IToolsContextRenderAPI* RenderAPI)
 	FPrimitiveDrawInterface* PDI = RenderAPI->GetPrimitiveDrawInterface();
 	FTransform Transform = Cast<IPrimitiveComponentBackedTarget>(Target)->GetWorldTransform(); //Actor->GetTransform();
 
-	if (SimplifyProperties->bShowConstraintEdges)
+	if (SimplifyProperties->bShowUVSeams)
 	{
-		FColor LineColor(255, 0, 0);
+		FColor UVSeamColor(15, 240, 15); // same color used in mesh inspector
 		const FDynamicMesh3* TargetMesh = Preview->PreviewMesh->GetPreviewDynamicMesh();
 		if (TargetMesh->HasAttributes())
 		{
@@ -229,7 +230,7 @@ void USimplifyMeshTool::Render(IToolsContextRenderAPI* RenderAPI)
 					FVector3d A, B;
 					TargetMesh->GetEdgeV(eid, A, B);
 					PDI->DrawLine(Transform.TransformPosition((FVector)A), Transform.TransformPosition((FVector)B),
-						LineColor, 0, 2.0 * PDIScale, 1.0f, true);
+						UVSeamColor, 0, 2.0 * PDIScale, 1.0f, true);
 				}
 			}
 		}
@@ -240,8 +241,13 @@ void USimplifyMeshTool::OnPropertyModified(UObject* PropertySet, FProperty* Prop
 {
 	if ( Property )
 	{
+		if ( Property->GetFName() == GET_MEMBER_NAME_CHECKED(USimplifyMeshToolProperties, bShowUVSeams) )
+		{
+			// nothing required.  Render() handles this case
+			return;
+		}
 		if ( ( Property->GetFName() == GET_MEMBER_NAME_CHECKED(USimplifyMeshToolProperties, bShowWireframe) ) ||
-			 ( Property->GetFName() == GET_MEMBER_NAME_CHECKED(USimplifyMeshToolProperties, bShowGroupColors) ) )
+			 ( Property->GetFName() == GET_MEMBER_NAME_CHECKED(USimplifyMeshToolProperties, bShowGroupColors) )  )
 		{
 			UpdateVisualization();
 		}
@@ -282,8 +288,6 @@ void USimplifyMeshTool::GenerateAsset(const FDynamicMeshOpResult& Result)
 	Cast<IMeshDescriptionCommitter>(Target)->CommitMeshDescription([&Result](const IMeshDescriptionCommitter::FCommitterParams& CommitParams)
 	{
 		FDynamicMeshToMeshDescription Converter;
-
-		// full conversion if normal topology changed or faces were inverted
 		Converter.Convert(Result.Mesh.Get(), *CommitParams.MeshDescriptionOut);
 	});
 
diff --git a/Engine/Plugins/Experimental/MeshModelingToolset/Source/MeshModelingToolsEditorOnly/Public/SimplifyMeshTool.h b/Engine/Plugins/Experimental/MeshModelingToolset/Source/MeshModelingToolsEditorOnly/Public/SimplifyMeshTool.h
--- a/Engine/Plugins/Experimental/MeshModelingToolset/Source/MeshModelingToolsEditorOnly/Public/SimplifyMeshTool.h
+++ b/Engine/Plugins/Experimental/MeshModelingToolset/Source/MeshModelingToolsEditorOnly/Public/SimplifyMeshTool.h
@@ -75,10 +75,10 @@ public:
 	UPROPERTY(EditAnywhere, Category = Display)
 	bool bShowGroupColors = false;
 
-	/** Display constrained edges */
-	UPROPERTY(EditAnywhere, Category = Display)
-	bool bShowConstraintEdges = false;
-
+	/** Display seams in first UV channel */
+	UPROPERTY(EditAnywhere, Category = Display, DisplayName ="Show UV Seams")
+	bool bShowUVSeams = false;
+	
 	/** Enable projection back to input mesh */
 	UPROPERTY(EditAnywhere, Category = Options, AdvancedDisplay)
 	bool bReproject;
