ModelingTools: fix backface-hit issues w/ VSculpt and DSculpt. Behavior is now consistent between the tools, and when backface-hit is disabled, the brush will not be able to jump to back-facing surfaces during the stroke. The Nearest-hit is now always computed, and then rejected if it is back-facing, instead of filtering out backfaces during the raycast (which allowed "hidden" front-facing surfaces to still be hit).
#rnx

diff --git a/Engine/Plugins/Experimental/MeshModelingToolset/Source/MeshModelingTools/Private/DynamicMeshSculptTool.cpp b/Engine/Plugins/Experimental/MeshModelingToolset/Source/MeshModelingTools/Private/DynamicMeshSculptTool.cpp
--- a/Engine/Plugins/Experimental/MeshModelingToolset/Source/MeshModelingTools/Private/DynamicMeshSculptTool.cpp
+++ b/Engine/Plugins/Experimental/MeshModelingToolset/Source/MeshModelingTools/Private/DynamicMeshSculptTool.cpp
@@ -1231,54 +1231,42 @@ bool UDynamicMeshSculptTool::ApplyScaleKelvinBrush(const FRay& WorldRay)
 	return true;
 }
 
-int UDynamicMeshSculptTool::FindHitSculptMeshTriangle(const FRay3d& LocalRay)
+bool UDynamicMeshSculptTool::IsHitTriangleBackFacing(int32 TriangleID, const FDynamicMesh3* QueryMesh)
 {
-	if (BrushProperties->bHitBackFaces)
-	{
-		return DynamicMeshComponent->GetOctree()->FindNearestHitObject(LocalRay);
-	}
-	else
+	if (TriangleID != IndexConstants::InvalidID)
 	{
-		FDynamicMesh3* Mesh = DynamicMeshComponent->GetMesh();
-
 		FViewCameraState StateOut;
 		GetToolManager()->GetContextQueriesAPI()->GetCurrentViewState(StateOut);
 		FVector3d LocalEyePosition(CurTargetTransform.InverseTransformPosition((FVector3d)StateOut.Position));
-		int HitTID = DynamicMeshComponent->GetOctree()->FindNearestHitObject(LocalRay,
-			[this, Mesh, &LocalEyePosition](int TriangleID) {
-			FVector3d Normal, Centroid;
-			double Area;
-			Mesh->GetTriInfo(TriangleID, Normal, Area, Centroid);
-			return Normal.Dot((Centroid - LocalEyePosition)) < 0;
-		});
-		return HitTID;
+
+		FVector3d Normal, Centroid;
+		double Area;
+		QueryMesh->GetTriInfo(TriangleID, Normal, Area, Centroid);
+
+		return (Normal.Dot((Centroid - LocalEyePosition)) >= 0);
 	}
+	return false;
+}
+
+int UDynamicMeshSculptTool::FindHitSculptMeshTriangle(const FRay3d& LocalRay)
+{
+	int32 HitTID = DynamicMeshComponent->GetOctree()->FindNearestHitObject(LocalRay);
+	if (BrushProperties->bHitBackFaces == false && IsHitTriangleBackFacing(HitTID, DynamicMeshComponent->GetMesh()) )
+	{
+		HitTID = IndexConstants::InvalidID;
+	}
+	return HitTID;
 }
 
 int UDynamicMeshSculptTool::FindHitTargetMeshTriangle(const FRay3d& LocalRay)
 {
 	PendingTargetUpdate.Wait();
 
-	IMeshSpatial::FQueryOptions RaycastOptions;
-
-	if (BrushProperties->bHitBackFaces == false)
+	int32 HitTID = BrushTargetMeshSpatial.FindNearestHitTriangle(LocalRay);
+	if (BrushProperties->bHitBackFaces == false && IsHitTriangleBackFacing(HitTID, &BrushTargetMesh))
 	{
-		FDynamicMesh3* Mesh = DynamicMeshComponent->GetMesh();
-
-		FViewCameraState StateOut;
-		GetToolManager()->GetContextQueriesAPI()->GetCurrentViewState(StateOut);
-		FVector3d LocalEyePosition(CurTargetTransform.InverseTransformPosition((FVector3d)StateOut.Position));
-
-		RaycastOptions.TriangleFilterF = [this, Mesh, LocalEyePosition](int TriangleID) {
-			FVector3d Normal, Centroid;
-			double Area;
-			Mesh->GetTriInfo(TriangleID, Normal, Area, Centroid);
-			return Normal.Dot((Centroid - LocalEyePosition)) < 0;
-		};
+		HitTID = IndexConstants::InvalidID;
 	}
-
-	int HitTID =  BrushTargetMeshSpatial.FindNearestHitTriangle(LocalRay, RaycastOptions);
-
 	return HitTID;
 }
 
diff --git a/Engine/Plugins/Experimental/MeshModelingToolset/Source/MeshModelingTools/Private/MeshVertexSculptTool.cpp b/Engine/Plugins/Experimental/MeshModelingToolset/Source/MeshModelingTools/Private/MeshVertexSculptTool.cpp
--- a/Engine/Plugins/Experimental/MeshModelingToolset/Source/MeshModelingTools/Private/MeshVertexSculptTool.cpp
+++ b/Engine/Plugins/Experimental/MeshModelingToolset/Source/MeshModelingTools/Private/MeshVertexSculptTool.cpp
@@ -631,7 +631,22 @@ TFuture<void> UMeshVertexSculptTool::ApplyStamp()
 
 
 
+bool UMeshVertexSculptTool::IsHitTriangleBackFacing(int32 TriangleID, const FDynamicMesh3* QueryMesh)
+{
+	if (TriangleID != IndexConstants::InvalidID)
+	{
+		FViewCameraState StateOut;
+		GetToolManager()->GetContextQueriesAPI()->GetCurrentViewState(StateOut);
+		FVector3d LocalEyePosition(CurTargetTransform.InverseTransformPosition((FVector3d)StateOut.Position));
+
+		FVector3d Normal, Centroid;
+		double Area;
+		QueryMesh->GetTriInfo(TriangleID, Normal, Area, Centroid);
 
+		return (Normal.Dot((Centroid - LocalEyePosition)) >= 0);
+	}
+	return false;
+}
 
 
 int32 UMeshVertexSculptTool::FindHitSculptMeshTriangle(const FRay3d& LocalRay)
@@ -639,53 +654,22 @@ int32 UMeshVertexSculptTool::FindHitSculptMeshTriangle(const FRay3d& LocalRay)
 	// need this to finish before we can touch Octree
 	WaitForPendingStampUpdate();
 
-	if (GetBrushCanHitBackFaces())
+	int32 HitTID = Octree.FindNearestHitObject(LocalRay);
+	if (GetBrushCanHitBackFaces() == false && IsHitTriangleBackFacing(HitTID, GetSculptMesh()))
 	{
-		return Octree.FindNearestHitObject(LocalRay);
-	}
-	else
-	{
-		FDynamicMesh3* Mesh = GetSculptMesh();
-
-		FViewCameraState StateOut;
-		GetToolManager()->GetContextQueriesAPI()->GetCurrentViewState(StateOut);
-		FVector3d LocalEyePosition(CurTargetTransform.InverseTransformPosition((FVector3d)StateOut.Position));
-		int HitTID = Octree.FindNearestHitObject(LocalRay,
-			[this, Mesh, &LocalEyePosition](int TriangleID) {
-			FVector3d Normal, Centroid;
-			double Area;
-			Mesh->GetTriInfo(TriangleID, Normal, Area, Centroid);
-			return Normal.Dot((Centroid - LocalEyePosition)) < 0;
-		});
-		return HitTID;
+		HitTID = IndexConstants::InvalidID;
 	}
+	return HitTID;
 }
 
 int32 UMeshVertexSculptTool::FindHitTargetMeshTriangle(const FRay3d& LocalRay)
 {
-	IMeshSpatial::FQueryOptions RaycastOptions;
-
-	if (GetBrushCanHitBackFaces())
-	{
-		FDynamicMesh3* Mesh = GetSculptMesh();
-
-		FViewCameraState StateOut;
-		GetToolManager()->GetContextQueriesAPI()->GetCurrentViewState(StateOut);
-		FVector3d LocalEyePosition(CurTargetTransform.InverseTransformPosition((FVector3d)StateOut.Position));
-
-		RaycastOptions.TriangleFilterF = [this, Mesh, LocalEyePosition](int TriangleID) {
-			FVector3d Normal, Centroid;
-			double Area;
-			Mesh->GetTriInfo(TriangleID, Normal, Area, Centroid);
-			return Normal.Dot((Centroid - LocalEyePosition)) < 0;
-		};
-
-		return BaseMeshSpatial.FindNearestHitObject(LocalRay, RaycastOptions.TriangleFilterF);
-	}
-	else
+	int32 HitTID = BaseMeshSpatial.FindNearestHitObject(LocalRay);
+	if (GetBrushCanHitBackFaces() == false && IsHitTriangleBackFacing(HitTID, GetBaseMesh()))
 	{
-		return BaseMeshSpatial.FindNearestHitObject(LocalRay);
+		HitTID = IndexConstants::InvalidID;
 	}
+	return HitTID;
 }
 
 
diff --git a/Engine/Plugins/Experimental/MeshModelingToolset/Source/MeshModelingTools/Public/DynamicMeshSculptTool.h b/Engine/Plugins/Experimental/MeshModelingToolset/Source/MeshModelingTools/Public/DynamicMeshSculptTool.h
--- a/Engine/Plugins/Experimental/MeshModelingToolset/Source/MeshModelingTools/Public/DynamicMeshSculptTool.h
+++ b/Engine/Plugins/Experimental/MeshModelingToolset/Source/MeshModelingTools/Public/DynamicMeshSculptTool.h
@@ -413,6 +413,7 @@ private:
 
 	int FindHitSculptMeshTriangle(const FRay3d& LocalRay);
 	int FindHitTargetMeshTriangle(const FRay3d& LocalRay);
+	bool IsHitTriangleBackFacing(int32 TriangleID, const FDynamicMesh3* QueryMesh);
 
 	bool UpdateBrushPosition(const FRay& WorldRay);
 	bool UpdateBrushPositionOnActivePlane(const FRay& WorldRay);
diff --git a/Engine/Plugins/Experimental/MeshModelingToolset/Source/MeshModelingTools/Public/MeshVertexSculptTool.h b/Engine/Plugins/Experimental/MeshModelingToolset/Source/MeshModelingTools/Public/MeshVertexSculptTool.h
--- a/Engine/Plugins/Experimental/MeshModelingToolset/Source/MeshModelingTools/Public/MeshVertexSculptTool.h
+++ b/Engine/Plugins/Experimental/MeshModelingToolset/Source/MeshModelingTools/Public/MeshVertexSculptTool.h
@@ -221,6 +221,7 @@ protected:
 
 	virtual int32 FindHitSculptMeshTriangle(const FRay3d& LocalRay) override;
 	virtual int32 FindHitTargetMeshTriangle(const FRay3d& LocalRay) override;
+	bool IsHitTriangleBackFacing(int32 TriangleID, const FDynamicMesh3* QueryMesh);
 
 	virtual void OnBeginStroke(const FRay& WorldRay) override;
 	virtual void OnEndStroke() override;
