MeshModelingToolset: fix incorrect ClampMin
#rnx

diff --git a/Engine/Plugins/Experimental/MeshModelingToolset/Source/MeshModelingTools/Public/Physics/SetCollisionGeometryTool.h b/Engine/Plugins/Experimental/MeshModelingToolset/Source/MeshModelingTools/Public/Physics/SetCollisionGeometryTool.h
--- a/Engine/Plugins/Experimental/MeshModelingToolset/Source/MeshModelingTools/Public/Physics/SetCollisionGeometryTool.h
+++ b/Engine/Plugins/Experimental/MeshModelingToolset/Source/MeshModelingTools/Public/Physics/SetCollisionGeometryTool.h
@@ -91,7 +91,7 @@ public:
 	UPROPERTY(EditAnywhere, Category = Options)
 	bool bEnableMaxCount = true;
 
-	UPROPERTY(EditAnywhere, Category = Options, meta = (UIMin = "1", UIMax = "100", ClampMin = "4", ClampMax = "9999999", EditCondition = "bEnableMaxCount"))
+	UPROPERTY(EditAnywhere, Category = Options, meta = (UIMin = "1", UIMax = "100", ClampMin = "1", ClampMax = "9999999", EditCondition = "bEnableMaxCount"))
 	int32 MaxCount = 50;
 
 	UPROPERTY(EditAnywhere, Category = Options, AdvancedDisplay)
