MeshModelingTools: Add DynamicMesh provider/committer interfaces and update existing targets to support them. Move interface target files.

#rnx

diff --git a/Engine/Plugins/Experimental/GeometryProcessing/Source/DynamicMesh/DynamicMesh.Build.cs b/Engine/Plugins/Experimental/GeometryProcessing/Source/DynamicMesh/DynamicMesh.Build.cs
--- a/Engine/Plugins/Experimental/GeometryProcessing/Source/DynamicMesh/DynamicMesh.Build.cs
+++ b/Engine/Plugins/Experimental/GeometryProcessing/Source/DynamicMesh/DynamicMesh.Build.cs
@@ -27,6 +27,7 @@ public class DynamicMesh : ModuleRules
 
 		PrivateDependencyModuleNames.AddRange(
 			new string[] {
+				"CoreUObject",
 				"Eigen"
 			}
 		);
diff --git a/Engine/Plugins/Experimental/GeometryProcessing/Source/DynamicMesh/Public/TargetInterfaces/DynamicMeshCommitter.h b/Engine/Plugins/Experimental/GeometryProcessing/Source/DynamicMesh/Public/TargetInterfaces/DynamicMeshCommitter.h
new file mode 100644
--- /dev/null
+++ b/Engine/Plugins/Experimental/GeometryProcessing/Source/DynamicMesh/Public/TargetInterfaces/DynamicMeshCommitter.h
@@ -0,0 +1,66 @@
+// Copyright Epic Games, Inc. All Rights Reserved.
+
+#pragma once
+
+#include "CoreMinimal.h"
+#include "UObject/Interface.h"
+
+#include "GeometryBase.h"
+
+#include "DynamicMeshCommitter.generated.h"
+
+PREDECLARE_GEOMETRY(class FDynamicMesh3);
+
+
+UINTERFACE()
+class DYNAMICMESH_API UDynamicMeshCommitter : public UInterface
+{
+	GENERATED_BODY()
+};
+
+class DYNAMICMESH_API IDynamicMeshCommitter
+{
+	GENERATED_BODY()
+
+public:
+
+	/**
+	 * Extra information that can be passed to a CommitMesh call to potentially make
+	 * the commit faster. Note that setting any of these to false doesn't mean that
+	 * the corresponding data won't be updated, because a target may choose to always
+	 * update everything. But it may help some targets do faster updates by not
+	 * updating things that stayed the same.
+	 */
+	struct DYNAMICMESH_API FDynamicMeshCommitInfo
+	{
+		/** Initializes everything to bInitValue */
+		FDynamicMeshCommitInfo(bool bInitValue)
+		{
+			bPositionsChanged =
+			bTopologyChanged =
+			bPolygroupsChanged =
+			bNormalsChanged =
+			bTangentsChanged =
+			bUVsChanged =
+			bVertexColorsChanged = bInitValue;
+		}
+
+		/** Leaves everything initialized to default (true) */
+		FDynamicMeshCommitInfo() {}
+
+		bool bPositionsChanged = true;
+		bool bTopologyChanged = true;
+		bool bPolygroupsChanged = true;
+		bool bNormalsChanged = true;
+		bool bTangentsChanged = true;
+		bool bUVsChanged = true;
+		bool bVertexColorsChanged = true;
+	};
+
+	virtual void CommitDynamicMesh(const UE::Geometry::FDynamicMesh3& Mesh)
+	{
+		FDynamicMeshCommitInfo CommitInfo;
+		CommitDynamicMesh(Mesh, CommitInfo);
+	};
+	virtual void CommitDynamicMesh(const UE::Geometry::FDynamicMesh3& Mesh, const FDynamicMeshCommitInfo& CommitInfo) = 0;
+};
diff --git a/Engine/Plugins/Experimental/GeometryProcessing/Source/DynamicMesh/Public/TargetInterfaces/DynamicMeshProvider.h b/Engine/Plugins/Experimental/GeometryProcessing/Source/DynamicMesh/Public/TargetInterfaces/DynamicMeshProvider.h
new file mode 100644
--- /dev/null
+++ b/Engine/Plugins/Experimental/GeometryProcessing/Source/DynamicMesh/Public/TargetInterfaces/DynamicMeshProvider.h
@@ -0,0 +1,35 @@
+// Copyright Epic Games, Inc. All Rights Reserved.
+
+#pragma once
+
+#include "CoreMinimal.h"
+#include "UObject/Interface.h"
+
+#include "GeometryBase.h"
+
+#include "DynamicMeshProvider.generated.h"
+
+PREDECLARE_GEOMETRY(class FDynamicMesh3);
+
+UINTERFACE()
+class DYNAMICMESH_API UDynamicMeshProvider : public UInterface
+{
+	GENERATED_BODY()
+};
+
+class DYNAMICMESH_API IDynamicMeshProvider
+{
+	GENERATED_BODY()
+
+public:
+	/**
+	 * Gives back an editable dynamic mesh. Changes may or may not affect the target, but
+	 * target guarantees that they are safe to perform without crashing. In particular,
+	 * note that it's unspecified whether the next GetMesh() call after an edit will give
+	 * back the edited mesh or a copy of the original.
+	 * 
+	 * To guarantee changes having an effect, target must also be a IDynamicMeshCommitter 
+	 * and the mesh must be committed.
+	 */
+	virtual TSharedPtr<UE::Geometry::FDynamicMesh3> GetDynamicMesh() = 0;
+};
diff --git a/Engine/Plugins/Experimental/GeometryProcessing/Source/MeshConversion/Private/DynamicMeshToMeshDescription.cpp b/Engine/Plugins/Experimental/GeometryProcessing/Source/MeshConversion/Private/DynamicMeshToMeshDescription.cpp
--- a/Engine/Plugins/Experimental/GeometryProcessing/Source/MeshConversion/Private/DynamicMeshToMeshDescription.cpp
+++ b/Engine/Plugins/Experimental/GeometryProcessing/Source/MeshConversion/Private/DynamicMeshToMeshDescription.cpp
@@ -341,6 +341,32 @@ void FDynamicMeshToMeshDescription::Convert(const FDynamicMesh3* MeshIn, FMeshDe
 	}
 }
 
+void FDynamicMeshToMeshDescription::UpdateUsingConversionOptions(const FDynamicMesh3* MeshIn, FMeshDescription& MeshOut)
+{
+	// See if we can do a buffer update without having to alter triangles.
+	if (HaveMatchingElementCounts(MeshIn, &MeshOut))
+	{
+		if (ConversionOptions.bUpdatePositions)
+		{
+			Update(MeshIn, MeshOut, ConversionOptions.bUpdateNormals, ConversionOptions.bUpdateTangents, ConversionOptions.bUpdateUVs);
+		}
+		else if (ConversionOptions.bUpdateNormals || ConversionOptions.bUpdateTangents || ConversionOptions.bUpdateUVs)
+		{
+			UpdateAttributes(MeshIn, MeshOut, ConversionOptions.bUpdateNormals, ConversionOptions.bUpdateTangents, ConversionOptions.bUpdateUVs);
+		}
+
+		if (ConversionOptions.bUpdateVtxColors)
+		{
+			UpdateVertexColors(MeshIn, MeshOut);
+		}
+	}
+	else
+	{
+		// Do a full conversion.
+		Convert(MeshIn, MeshOut);
+	}
+}
+
 
 void FDynamicMeshToMeshDescription::Convert_NoAttributes(const FDynamicMesh3* MeshIn, FMeshDescription& MeshOut)
 {
diff --git a/Engine/Plugins/Experimental/GeometryProcessing/Source/MeshConversion/Public/DynamicMeshToMeshDescription.h b/Engine/Plugins/Experimental/GeometryProcessing/Source/MeshConversion/Public/DynamicMeshToMeshDescription.h
--- a/Engine/Plugins/Experimental/GeometryProcessing/Source/MeshConversion/Public/DynamicMeshToMeshDescription.h
+++ b/Engine/Plugins/Experimental/GeometryProcessing/Source/MeshConversion/Public/DynamicMeshToMeshDescription.h
@@ -79,6 +79,14 @@ public:
 	 */
 	void Convert(const FDynamicMesh3* MeshIn, FMeshDescription& MeshOut, bool bCopyTangents = false);
 
+	/**
+	 * Updates the given mesh description based conversion options provided in the constructor. Assumes
+	 * the mesh topology has not changed. 
+	 * Annoyingly, this can't just be named Update() due to ambiguity with the function below, which
+	 * existed beforehand and should probably have been this function instead.
+	 */
+	void UpdateUsingConversionOptions(const FDynamicMesh3* MeshIn, FMeshDescription& MeshOut);
+
 	/**
 	 * Update existing MeshDescription based on DynamicMesh. Assumes mesh topology has not changed.
 	 * Copies positions 
diff --git a/Engine/Plugins/Experimental/MeshModelingToolset/Source/MeshModelingToolsEditorOnly/Private/ConversionUtils/VolumeMeshDescriptionToolTarget.cpp b/Engine/Plugins/Experimental/MeshModelingToolset/Source/MeshModelingToolsEditorOnly/Private/ConversionUtils/VolumeMeshDescriptionToolTarget.cpp
deleted file mode 100644
diff --git a/Engine/Plugins/Experimental/MeshModelingToolset/Source/MeshModelingToolsEditorOnly/Private/ToolTargets/SkeletalMeshComponentToolTarget.cpp b/Engine/Plugins/Experimental/MeshModelingToolset/Source/MeshModelingToolsEditorOnly/Private/ToolTargets/SkeletalMeshComponentToolTarget.cpp
new file mode 100644
--- /dev/null
+++ b/Engine/Plugins/Experimental/MeshModelingToolset/Source/MeshModelingToolsEditorOnly/Private/ToolTargets/SkeletalMeshComponentToolTarget.cpp
@@ -0,0 +1,261 @@
+﻿// Copyright Epic Games, Inc. All Rights Reserved.
+
+#include "ToolTargets/SkeletalMeshComponentToolTarget.h"
+
+#include "Components/SkinnedMeshComponent.h"
+#include "DynamicMeshToMeshDescription.h"
+#include "MeshDescriptionToDynamicMesh.h"
+#include "Rendering/SkeletalMeshModel.h"
+
+
+namespace USkeletalMeshComponentToolTargetLocals
+{
+	int32 LODIndex = 0;
+}
+
+
+int32 USkeletalMeshComponentToolTarget::GetNumMaterials() const
+{
+	return ensure(IsValid()) ? Component->GetNumMaterials() : 0;
+}
+
+UMaterialInterface* USkeletalMeshComponentToolTarget::GetMaterial(int32 MaterialIndex) const
+{
+	return ensure(IsValid()) ? Component->GetMaterial(MaterialIndex) : nullptr;
+}
+
+void USkeletalMeshComponentToolTarget::GetMaterialSet(FComponentMaterialSet& MaterialSetOut, bool bPreferAssetMaterials) const
+{
+	if (!ensure(IsValid())) return;
+
+	if (bPreferAssetMaterials)
+	{
+		const USkeletalMesh* SkeletalMesh = Cast<USkinnedMeshComponent>(Component)->SkeletalMesh;
+		const TArray<FSkeletalMaterial>& Materials = SkeletalMesh->GetMaterials(); 
+		MaterialSetOut.Materials.SetNum(Materials.Num());
+		for (int32 k = 0; k < Materials.Num(); ++k)
+		{
+			MaterialSetOut.Materials[k] = Materials[k].MaterialInterface;
+		}
+	}
+	else
+	{
+		int32 NumMaterials = Component->GetNumMaterials();
+		MaterialSetOut.Materials.SetNum(NumMaterials);
+		for (int32 k = 0; k < NumMaterials; ++k)
+		{
+			MaterialSetOut.Materials[k] = Component->GetMaterial(k);
+		}
+	}
+}
+
+bool USkeletalMeshComponentToolTarget::CommitMaterialSetUpdate(const FComponentMaterialSet& MaterialSet, bool bApplyToAsset)
+{
+	if (!ensure(IsValid())) return false;
+
+	if (bApplyToAsset)
+	{
+		USkeletalMesh* SkeletalMesh = Cast<USkinnedMeshComponent>(Component)->SkeletalMesh;
+
+		if (SkeletalMesh->GetPathName().StartsWith(TEXT("/Engine/")))
+		{
+			UE_LOG(LogTemp, Warning, TEXT("CANNOT MODIFY BUILT-IN ENGINE ASSET %s"), *SkeletalMesh->GetPathName());
+			return false;
+		}
+
+		// flush any pending rendering commands, which might touch this component while we are rebuilding its mesh
+		FlushRenderingCommands();
+
+		// unregister the component while we update it's static mesh
+		TUniquePtr<FComponentReregisterContext> ComponentReregisterContext = MakeUnique<FComponentReregisterContext>(Component);
+
+		// make sure transactional flag is on
+		SkeletalMesh->SetFlags(RF_Transactional);
+
+		SkeletalMesh->Modify();
+
+		const int NewNumMaterials = MaterialSet.Materials.Num();
+		TArray<FSkeletalMaterial> &SkeletalMaterials = SkeletalMesh->GetMaterials(); 
+		if (NewNumMaterials != SkeletalMaterials.Num())
+		{
+			SkeletalMaterials.SetNum(NewNumMaterials);
+		}
+		for (int k = 0; k < NewNumMaterials; ++k)
+		{
+			if (SkeletalMaterials[k].MaterialInterface != MaterialSet.Materials[k])
+			{
+				SkeletalMaterials[k].MaterialInterface = MaterialSet.Materials[k];
+				if (SkeletalMaterials[k].MaterialSlotName.IsNone())
+				{
+					SkeletalMaterials[k].MaterialSlotName = MaterialSet.Materials[k]->GetFName();
+				}
+			}
+		}
+
+		SkeletalMesh->PostEditChange();
+	}
+	else
+	{
+		const int32 NumMaterialsNeeded = Component->GetNumMaterials();
+		const int32 NumMaterialsGiven = MaterialSet.Materials.Num();
+
+		// We wrote the below code to support a mismatch in the number of materials.
+		// However, it is not yet clear whether this might be desirable, and we don't
+		// want to inadvertantly hide bugs in the meantime. So, we keep this ensure here
+		// for now, and we can remove it if we decide that we want the ability.
+		ensure(NumMaterialsNeeded == NumMaterialsGiven);
+
+		check(NumMaterialsGiven > 0);
+
+		for (int32 i = 0; i < NumMaterialsNeeded; ++i)
+		{
+			const int32 MaterialToUseIndex = FMath::Min(i, NumMaterialsGiven - 1);
+			Component->SetMaterial(i, MaterialSet.Materials[MaterialToUseIndex]);
+		}
+	}
+
+	return true;
+}
+
+FMeshDescription* USkeletalMeshComponentToolTarget::GetMeshDescription()
+{
+	if (!ensure(IsValid()))
+	{
+		return nullptr;
+	}
+
+	if (!CachedMeshDescription.IsValid())
+	{
+		using namespace USkeletalMeshComponentToolTargetLocals;
+
+		CachedMeshDescription = MakeUnique<FMeshDescription>();
+		const USkeletalMesh* SkeletalMesh = Cast<USkinnedMeshComponent>(Component)->SkeletalMesh;
+
+		// Check first if we have bulk data available.
+		if (SkeletalMesh->IsLODImportedDataBuildAvailable(LODIndex))
+		{
+			FSkeletalMeshImportData SkeletalMeshImportData;
+			SkeletalMesh->LoadLODImportedData(LODIndex, SkeletalMeshImportData);
+			SkeletalMeshImportData.GetMeshDescription(*CachedMeshDescription);
+		}
+		else
+		{
+			// Fall back on the LOD model directly if no bulk data exists. When we commit
+			// the mesh description, we override using the bulk data. This can happen for older
+			// skeletal meshes, from UE 4.24 and earlier.
+			const FSkeletalMeshModel* SkeletalMeshModel = SkeletalMesh->GetImportedModel();
+			if (SkeletalMeshModel && SkeletalMeshModel->LODModels.IsValidIndex(LODIndex))
+			{
+				SkeletalMeshModel->LODModels[LODIndex].GetMeshDescription(*CachedMeshDescription, SkeletalMesh);
+			}			
+		}
+	}
+	
+	return CachedMeshDescription.Get();
+}
+
+void USkeletalMeshComponentToolTarget::CommitMeshDescription(const FCommitter& Committer)
+{
+	using namespace USkeletalMeshComponentToolTargetLocals;
+	if (ensure(IsValid()) == false) return;
+
+	USkeletalMesh* SkeletalMesh = Cast<USkinnedMeshComponent>(Component)->SkeletalMesh;
+
+	if (SkeletalMesh->GetPathName().StartsWith(TEXT("/Engine/")))
+	{
+		const FString DebugMessage = FString::Printf(TEXT("CANNOT MODIFY BUILT-IN ENGINE ASSET %s"), *SkeletalMesh->GetPathName());
+		if (GAreScreenMessagesEnabled)
+		{
+			GEngine->AddOnScreenDebugMessage(INDEX_NONE, 10.0f, FColor::Red, DebugMessage);
+		}
+		UE_LOG(LogTemp, Warning, TEXT("%s"), *DebugMessage);
+		return;
+	}
+
+	// flush any pending rendering commands, which might touch this component while we are rebuilding it's mesh
+	FlushRenderingCommands();
+
+	// unregister the component while we update its skeletal mesh
+	FComponentReregisterContext ComponentReregisterContext(Component);
+
+	// make sure transactional flag is on for this asset
+	SkeletalMesh->SetFlags(RF_Transactional);
+
+	verify(SkeletalMesh->Modify());
+
+	FCommitterParams CommitterParams;
+
+	CommitterParams.MeshDescriptionOut = GetMeshDescription();
+
+	Committer(CommitterParams);
+
+	FSkeletalMeshImportData SkeletalMeshImportData = FSkeletalMeshImportData::CreateFromMeshDescription(*CommitterParams.MeshDescriptionOut);
+	SkeletalMesh->SaveLODImportedData(LODIndex, SkeletalMeshImportData);
+
+	// Make sure the mesh builder knows it's the latest variety, so that the render data gets
+	// properly rebuilt.
+	SkeletalMesh->SetLODImportedDataVersions(LODIndex, ESkeletalMeshGeoImportVersions::LatestVersion, ESkeletalMeshSkinningImportVersions::LatestVersion);
+	SkeletalMesh->SetUseLegacyMeshDerivedDataKey(false);
+
+	SkeletalMesh->PostEditChange();
+
+	// this rebuilds physics, but it doesn't undo!
+	Component->RecreatePhysicsState();
+}
+
+TSharedPtr<FDynamicMesh3> USkeletalMeshComponentToolTarget::GetDynamicMesh()
+{
+	TSharedPtr<FDynamicMesh3> DynamicMesh = MakeShared<FDynamicMesh3>();
+	FMeshDescriptionToDynamicMesh Converter;
+	Converter.Convert(GetMeshDescription(), *DynamicMesh);
+	return DynamicMesh;
+}
+
+void USkeletalMeshComponentToolTarget::CommitDynamicMesh(const FDynamicMesh3& Mesh, const FDynamicMeshCommitInfo& CommitInfo)
+{
+	FConversionToMeshDescriptionOptions ConversionOptions;
+	ConversionOptions.bSetPolyGroups = CommitInfo.bPolygroupsChanged;
+	ConversionOptions.bUpdatePositions = CommitInfo.bPositionsChanged;
+	ConversionOptions.bUpdateNormals = CommitInfo.bNormalsChanged;
+	ConversionOptions.bUpdateTangents = CommitInfo.bTangentsChanged;
+	ConversionOptions.bUpdateUVs = CommitInfo.bUVsChanged;
+	ConversionOptions.bUpdateVtxColors = CommitInfo.bVertexColorsChanged;
+
+	CommitMeshDescription([&CommitInfo, &ConversionOptions, &Mesh](const IMeshDescriptionCommitter::FCommitterParams& CommitParams)
+		{
+			FDynamicMeshToMeshDescription Converter(ConversionOptions);
+
+			if (!CommitInfo.bTopologyChanged)
+			{
+				Converter.UpdateUsingConversionOptions(&Mesh, *CommitParams.MeshDescriptionOut);
+			}
+			else
+			{
+				// Do a full conversion.
+				Converter.Convert(&Mesh, *CommitParams.MeshDescriptionOut);
+			}
+		});
+}
+
+USkeletalMesh* USkeletalMeshComponentToolTarget::GetSkeletalMesh() const
+{
+	return IsValid() ? Cast<USkinnedMeshComponent>(Component)->SkeletalMesh : nullptr;
+}
+
+
+// Factory
+
+bool USkeletalMeshComponentToolTargetFactory::CanBuildTarget(UObject* SourceObject, const FToolTargetTypeRequirements& Requirements) const
+{
+	return Cast<USkinnedMeshComponent>(SourceObject) && Cast<USkinnedMeshComponent>(SourceObject)->SkeletalMesh &&
+		Requirements.AreSatisfiedBy(USkeletalMeshComponentToolTarget::StaticClass());
+}
+
+UToolTarget* USkeletalMeshComponentToolTargetFactory::BuildTarget(UObject* SourceObject, const FToolTargetTypeRequirements& Requirements)
+{
+	USkeletalMeshComponentToolTarget* Target = NewObject<USkeletalMeshComponentToolTarget>();
+	Target->Component = Cast<USkinnedMeshComponent>(SourceObject);
+	check(Target->Component && Requirements.AreSatisfiedBy(Target));
+
+	return Target;
+}
diff --git a/Engine/Plugins/Experimental/MeshModelingToolset/Source/MeshModelingToolsEditorOnly/Private/ToolTargets/StaticMeshComponentToolTarget.cpp b/Engine/Plugins/Experimental/MeshModelingToolset/Source/MeshModelingToolsEditorOnly/Private/ToolTargets/StaticMeshComponentToolTarget.cpp
new file mode 100644
--- /dev/null
+++ b/Engine/Plugins/Experimental/MeshModelingToolset/Source/MeshModelingToolsEditorOnly/Private/ToolTargets/StaticMeshComponentToolTarget.cpp
@@ -0,0 +1,321 @@
+// Copyright Epic Games, Inc. All Rights Reserved.
+
+#include "ToolTargets/StaticMeshComponentToolTarget.h"
+
+#include "ComponentReregisterContext.h"
+#include "Components/StaticMeshComponent.h"
+#include "DynamicMeshToMeshDescription.h"
+#include "Engine/StaticMesh.h"
+#include "MeshDescriptionToDynamicMesh.h"
+#include "RenderingThread.h"
+
+static void DisplayCriticalWarningMessage(const FString& Message)
+{
+	if (GAreScreenMessagesEnabled)
+	{
+		GEngine->AddOnScreenDebugMessage(INDEX_NONE, 10.0f, FColor::Red, Message);
+	}
+	UE_LOG(LogTemp, Warning, TEXT("%s"), *Message);
+}
+
+
+void UStaticMeshComponentToolTarget::SetEditingLOD(EStaticMeshEditingLOD RequestedEditingLOD)
+{
+	EStaticMeshEditingLOD ValidEditingLOD = EStaticMeshEditingLOD::LOD0;
+
+	UStaticMeshComponent* StaticMeshComponent = Cast<UStaticMeshComponent>(Component);
+	if (ensure(StaticMeshComponent != nullptr))
+	{
+		UStaticMesh* StaticMeshAsset = StaticMeshComponent->GetStaticMesh();
+		if (ensure(StaticMeshAsset != nullptr))
+		{
+			if (RequestedEditingLOD == EStaticMeshEditingLOD::MaxQuality)
+			{
+				ValidEditingLOD = StaticMeshAsset->IsHiResMeshDescriptionValid() ? EStaticMeshEditingLOD::HiResSource : EStaticMeshEditingLOD::LOD0;
+			}
+			else if (RequestedEditingLOD == EStaticMeshEditingLOD::HiResSource)
+			{
+				ValidEditingLOD = StaticMeshAsset->IsHiResMeshDescriptionValid() ? EStaticMeshEditingLOD::HiResSource : EStaticMeshEditingLOD::LOD0;
+				if (ValidEditingLOD != EStaticMeshEditingLOD::HiResSource)
+				{
+					DisplayCriticalWarningMessage(FString(TEXT("HiRes Source selected but not available - Falling Back to LOD0")));
+				}
+			}
+			else
+			{
+				ValidEditingLOD = RequestedEditingLOD;
+				int32 MaxExistingLOD = StaticMeshAsset->GetNumSourceModels() - 1;
+				if ((int32)ValidEditingLOD > MaxExistingLOD)
+				{
+					DisplayCriticalWarningMessage(FString::Printf(TEXT("LOD%d Requested but not available - Falling Back to LOD%d"), (int32)ValidEditingLOD, MaxExistingLOD));
+					ValidEditingLOD = (EStaticMeshEditingLOD)MaxExistingLOD;
+				}
+			}
+		}
+	}
+
+	EditingLOD = ValidEditingLOD;
+}
+
+
+bool UStaticMeshComponentToolTarget::IsValid() const
+{
+	if (!UPrimitiveComponentToolTarget::IsValid())
+	{
+		return false;
+	}
+	UStaticMeshComponent* StaticMeshComponent = Cast<UStaticMeshComponent>(Component);
+	if (StaticMeshComponent == nullptr)
+	{
+		return false;
+	}
+	UStaticMesh* StaticMesh = StaticMeshComponent->GetStaticMesh();
+	if (StaticMesh == nullptr)
+	{
+		return false;
+	}
+
+	if (EditingLOD == EStaticMeshEditingLOD::HiResSource)
+	{
+		if (StaticMesh->IsHiResMeshDescriptionValid() == false)
+		{
+			return false;
+		}
+	}
+	else if ((int32)EditingLOD >= StaticMesh->GetNumSourceModels())
+	{
+		return false;
+	}
+
+	return true;
+}
+
+
+
+int32 UStaticMeshComponentToolTarget::GetNumMaterials() const
+{
+	return ensure(IsValid()) ? Component->GetNumMaterials() : 0;
+}
+
+UMaterialInterface* UStaticMeshComponentToolTarget::GetMaterial(int32 MaterialIndex) const
+{
+	return ensure(IsValid()) ? Component->GetMaterial(MaterialIndex) : nullptr;
+}
+
+void UStaticMeshComponentToolTarget::GetMaterialSet(FComponentMaterialSet& MaterialSetOut, bool bPreferAssetMaterials) const
+{
+	if (!ensure(IsValid())) return;
+
+	if (bPreferAssetMaterials)
+	{
+		UStaticMesh* StaticMesh = Cast<UStaticMeshComponent>(Component)->GetStaticMesh();
+		int32 NumMaterials = Component->GetNumMaterials();
+		MaterialSetOut.Materials.SetNum(NumMaterials);
+		for (int32 k = 0; k < NumMaterials; ++k)
+		{
+			MaterialSetOut.Materials[k] = StaticMesh->GetMaterial(k);
+		}
+	}
+	else
+	{
+		int32 NumMaterials = Component->GetNumMaterials();
+		MaterialSetOut.Materials.SetNum(NumMaterials);
+		for (int32 k = 0; k < NumMaterials; ++k)
+		{
+			MaterialSetOut.Materials[k] = Component->GetMaterial(k);
+		}
+	}
+}
+
+bool UStaticMeshComponentToolTarget::CommitMaterialSetUpdate(const FComponentMaterialSet& MaterialSet, bool bApplyToAsset)
+{
+	if (!ensure(IsValid())) return false;
+
+	if (bApplyToAsset)
+	{
+		UStaticMesh* StaticMesh = Cast<UStaticMeshComponent>(Component)->GetStaticMesh();
+
+		if (StaticMesh->GetPathName().StartsWith(TEXT("/Engine/")))
+		{
+			UE_LOG(LogTemp, Warning, TEXT("CANNOT MODIFY BUILT-IN ENGINE ASSET %s"), *StaticMesh->GetPathName());
+			return false;
+		}
+
+		// flush any pending rendering commands, which might touch this component while we are rebuilding its mesh
+		FlushRenderingCommands();
+
+		// unregister the component while we update it's static mesh
+		TUniquePtr<FComponentReregisterContext> ComponentReregisterContext = MakeUnique<FComponentReregisterContext>(Component);
+
+		// make sure transactional flag is on
+		StaticMesh->SetFlags(RF_Transactional);
+
+		StaticMesh->Modify();
+
+		int NewNumMaterials = MaterialSet.Materials.Num();
+		if (NewNumMaterials != StaticMesh->GetStaticMaterials().Num())
+		{
+			StaticMesh->GetStaticMaterials().SetNum(NewNumMaterials);
+		}
+		for (int k = 0; k < NewNumMaterials; ++k)
+		{
+			if (StaticMesh->GetMaterial(k) != MaterialSet.Materials[k])
+			{
+				StaticMesh->SetMaterial(k, MaterialSet.Materials[k]);
+			}
+		}
+
+		StaticMesh->PostEditChange();
+	}
+	else
+	{
+		int32 NumMaterialsNeeded = Component->GetNumMaterials();
+		int32 NumMaterialsGiven = MaterialSet.Materials.Num();
+
+		// We wrote the below code to support a mismatch in the number of materials.
+		// However, it is not yet clear whether this might be desirable, and we don't
+		// want to inadvertantly hide bugs in the meantime. So, we keep this ensure here
+		// for now, and we can remove it if we decide that we want the ability.
+		ensure(NumMaterialsNeeded == NumMaterialsGiven);
+
+		check(NumMaterialsGiven > 0);
+
+		for (int32 i = 0; i < NumMaterialsNeeded; ++i)
+		{
+			int32 MaterialToUseIndex = FMath::Min(i, NumMaterialsGiven - 1);
+			Component->SetMaterial(i, MaterialSet.Materials[MaterialToUseIndex]);
+		}
+	}
+
+	return true;
+}
+
+
+FMeshDescription* UStaticMeshComponentToolTarget::GetMeshDescription()
+{
+	if (ensure(IsValid()))
+	{
+		UStaticMesh* StaticMesh = Cast<UStaticMeshComponent>(Component)->GetStaticMesh();
+		return (EditingLOD == EStaticMeshEditingLOD::HiResSource) ?
+			StaticMesh->GetHiResMeshDescription() : StaticMesh->GetMeshDescription((int32)EditingLOD);
+	}
+	return nullptr;
+}
+
+
+void UStaticMeshComponentToolTarget::CommitMeshDescription(const FCommitter& Committer)
+{
+	if (ensure(IsValid()) == false) return;
+
+	UStaticMesh* StaticMesh = Cast<UStaticMeshComponent>(Component)->GetStaticMesh();
+
+	if (StaticMesh->GetPathName().StartsWith(TEXT("/Engine/")))
+	{
+		DisplayCriticalWarningMessage(FString::Printf(TEXT("CANNOT MODIFY BUILT-IN ENGINE ASSET %s"), *StaticMesh->GetPathName()));
+		return;
+	}
+
+	// flush any pending rendering commands, which might touch this component while we are rebuilding it's mesh
+	FlushRenderingCommands();
+
+	// unregister the component while we update its static mesh
+	FComponentReregisterContext ComponentReregisterContext(Component);
+
+	// make sure transactional flag is on for this asset
+	StaticMesh->SetFlags(RF_Transactional);
+
+	verify(StaticMesh->Modify());
+	if (EditingLOD == EStaticMeshEditingLOD::HiResSource)
+	{
+		verify(StaticMesh->ModifyHiResMeshDescription());
+	}
+	else
+	{
+		verify(StaticMesh->ModifyMeshDescription((int32)EditingLOD));
+	}
+
+	FCommitterParams CommitterParams;
+	CommitterParams.MeshDescriptionOut = GetMeshDescription();
+
+	Committer(CommitterParams);
+
+	if (EditingLOD == EStaticMeshEditingLOD::HiResSource)
+	{
+		StaticMesh->CommitHiResMeshDescription();
+	}
+	else
+	{
+		StaticMesh->CommitMeshDescription((int32)EditingLOD);
+	}
+
+	StaticMesh->PostEditChange();
+
+	// this rebuilds physics, but it doesn't undo!
+	Component->RecreatePhysicsState();
+}
+
+TSharedPtr<FDynamicMesh3> UStaticMeshComponentToolTarget::GetDynamicMesh()
+{
+	TSharedPtr<FDynamicMesh3> DynamicMesh = MakeShared<FDynamicMesh3>();
+	FMeshDescriptionToDynamicMesh Converter;
+	Converter.Convert(GetMeshDescription(), *DynamicMesh);
+	return DynamicMesh;
+}
+
+void UStaticMeshComponentToolTarget::CommitDynamicMesh(const FDynamicMesh3& Mesh, const FDynamicMeshCommitInfo& CommitInfo)
+{
+	FConversionToMeshDescriptionOptions ConversionOptions;
+	ConversionOptions.bSetPolyGroups = CommitInfo.bPolygroupsChanged;
+	ConversionOptions.bUpdatePositions = CommitInfo.bPositionsChanged;
+	ConversionOptions.bUpdateNormals = CommitInfo.bNormalsChanged;
+	ConversionOptions.bUpdateTangents = CommitInfo.bTangentsChanged;
+	ConversionOptions.bUpdateUVs = CommitInfo.bUVsChanged;
+	ConversionOptions.bUpdateVtxColors = CommitInfo.bVertexColorsChanged;
+
+	CommitMeshDescription([&CommitInfo, &ConversionOptions, &Mesh](const IMeshDescriptionCommitter::FCommitterParams& CommitParams)
+	{
+		FDynamicMeshToMeshDescription Converter(ConversionOptions);
+
+		if (!CommitInfo.bTopologyChanged)
+		{
+			Converter.UpdateUsingConversionOptions(&Mesh, *CommitParams.MeshDescriptionOut);
+		}
+		else
+		{
+			// Do a full conversion.
+			Converter.Convert(&Mesh, *CommitParams.MeshDescriptionOut);
+		}
+	});
+}
+
+UStaticMesh* UStaticMeshComponentToolTarget::GetStaticMesh() const
+{
+	return IsValid() ? Cast<UStaticMeshComponent>(Component)->GetStaticMesh() : nullptr;
+}
+
+
+// Factory
+
+bool UStaticMeshComponentToolTargetFactory::CanBuildTarget(UObject* SourceObject, const FToolTargetTypeRequirements& Requirements) const
+{
+	const UStaticMeshComponent* Component = Cast<UStaticMeshComponent>(SourceObject);
+	return Component && !Component->IsPendingKillOrUnreachable() && Component->IsValidLowLevel() && Component->GetStaticMesh()
+		&& (Component->GetStaticMesh()->GetNumSourceModels() > 0)
+		&& Requirements.AreSatisfiedBy(UStaticMeshComponentToolTarget::StaticClass());
+}
+
+UToolTarget* UStaticMeshComponentToolTargetFactory::BuildTarget(UObject* SourceObject, const FToolTargetTypeRequirements& Requirements)
+{
+	UStaticMeshComponentToolTarget* Target = NewObject<UStaticMeshComponentToolTarget>();// TODO: Should we set an outer here?
+	Target->Component = Cast<UStaticMeshComponent>(SourceObject);
+	Target->SetEditingLOD(EditingLOD);
+	check(Target->Component && Requirements.AreSatisfiedBy(Target));
+
+	return Target;
+}
+
+
+void UStaticMeshComponentToolTargetFactory::SetActiveEditingLOD(EStaticMeshEditingLOD NewEditingLOD)
+{
+	EditingLOD = NewEditingLOD;
+}
\ No newline at end of file
diff --git a/Engine/Plugins/Experimental/MeshModelingToolset/Source/MeshModelingToolsEditorOnly/Private/ToolTargets/VolumeDynamicMeshToolTarget.cpp b/Engine/Plugins/Experimental/MeshModelingToolset/Source/MeshModelingToolsEditorOnly/Private/ToolTargets/VolumeDynamicMeshToolTarget.cpp
new file mode 100644
--- /dev/null
+++ b/Engine/Plugins/Experimental/MeshModelingToolset/Source/MeshModelingToolsEditorOnly/Private/ToolTargets/VolumeDynamicMeshToolTarget.cpp
@@ -0,0 +1,125 @@
+// Copyright Epic Games, Inc. All Rights Reserved.
+
+#pragma once
+
+#include "ToolTargets/VolumeDynamicMeshToolTarget.h"
+
+#include "Components/BrushComponent.h"
+#include "ConversionUtils/VolumeToDynamicMesh.h"
+#include "ConversionUtils/DynamicMeshToVolume.h"
+#include "DynamicMeshToMeshDescription.h"
+#include "Engine/Brush.h"
+#include "GameFramework/Volume.h"
+#include "MeshDescriptionToDynamicMesh.h"
+#include "MeshNormals.h"
+#include "StaticMeshAttributes.h"
+#include "ToolSetupUtil.h"
+
+#include "ExplicitUseGeometryMathTypes.h"		// using UE::Geometry::(math types)
+using namespace UE::Geometry;
+
+UVolumeDynamicMeshToolTarget::UVolumeDynamicMeshToolTarget()
+{
+	// TODO: These should be user-configurable somewhere
+	VolumeToMeshOptions.bInWorldSpace = false;
+	VolumeToMeshOptions.bSetGroups = true;
+	VolumeToMeshOptions.bMergeVertices = true;
+
+	// When a volume has cracks, this option seems to make the geometry
+	// worse rather than better, since the filled in triangles are sometimes
+	// degenerate, folded in on themselves, etc.
+	VolumeToMeshOptions.bAutoRepairMesh = false;
+
+	VolumeToMeshOptions.bOptimizeMesh = true;
+}
+
+int32 UVolumeDynamicMeshToolTarget::GetNumMaterials() const
+{
+	return IsValid() ? 1 : 0;
+}
+
+UMaterialInterface* UVolumeDynamicMeshToolTarget::GetMaterial(int32 MaterialIndex) const
+{
+	if (!IsValid() || MaterialIndex > 0)
+	{
+		return nullptr;
+	}
+
+	return ToolSetupUtil::GetDefaultEditVolumeMaterial();
+}
+
+void UVolumeDynamicMeshToolTarget::GetMaterialSet(FComponentMaterialSet& MaterialSetOut, bool bPreferAssetMaterials) const
+{
+	MaterialSetOut.Materials.Reset();
+	if (IsValid() == false)
+	{
+		return;
+	}
+
+	UMaterialInterface* Material = ToolSetupUtil::GetDefaultEditVolumeMaterial();
+	if (Material)
+	{
+		MaterialSetOut.Materials.Add(Material);
+	}
+}
+
+TSharedPtr<UE::Geometry::FDynamicMesh3> UVolumeDynamicMeshToolTarget::GetDynamicMesh()
+{
+	UBrushComponent* BrushComponent = Cast<UBrushComponent>(Component);
+	if (!BrushComponent)
+	{
+		return nullptr;
+	}
+	AVolume* Volume = Cast<AVolume>(BrushComponent->GetOwner());
+	if (!Volume)
+	{
+		return nullptr;
+	}
+
+	TSharedPtr<FDynamicMesh3> DynamicMesh = MakeShared<FDynamicMesh3>();
+	UE::Conversion::VolumeToDynamicMesh(Volume, *DynamicMesh, GetVolumeToMeshOptions());
+	FMeshNormals::InitializeMeshToPerTriangleNormals(DynamicMesh.Get());
+
+	return DynamicMesh;
+}
+
+void UVolumeDynamicMeshToolTarget::CommitDynamicMesh(const UE::Geometry::FDynamicMesh3& Mesh, const FDynamicMeshCommitInfo&)
+{
+	check(IsValid());
+
+	UBrushComponent* BrushComponent = Cast<UBrushComponent>(Component);
+	check(BrushComponent);
+	AVolume* Volume = Cast<AVolume>(BrushComponent->GetOwner());
+	check(Volume);
+
+	FTransform Transform = GetWorldTransform();
+	UE::Conversion::DynamicMeshToVolume(Mesh, Volume);
+
+	Volume->SetActorTransform(Transform);
+	Volume->PostEditChange();
+}
+
+
+// Factory
+
+bool UVolumeDynamicMeshToolTargetFactory::CanBuildTarget(UObject* SourceObject, const FToolTargetTypeRequirements& Requirements) const
+{
+	UBrushComponent* BrushComponent = Cast<UBrushComponent>(SourceObject);
+	if (!BrushComponent)
+	{
+		return false;
+	}
+
+	return Cast<AVolume>(BrushComponent->GetOwner()) && Requirements.AreSatisfiedBy(UVolumeDynamicMeshToolTarget::StaticClass());
+}
+
+UToolTarget* UVolumeDynamicMeshToolTargetFactory::BuildTarget(UObject* SourceObject, const FToolTargetTypeRequirements& Requirements)
+{
+	UVolumeDynamicMeshToolTarget* Target = NewObject<UVolumeDynamicMeshToolTarget>();
+	Target->Component = Cast<UBrushComponent>(SourceObject);
+
+	check(Target->Component && Requirements.AreSatisfiedBy(Target)
+		&& Cast<AVolume>(Cast<UBrushComponent>(SourceObject)->GetOwner()));
+
+	return Target;
+}
\ No newline at end of file
diff --git a/Engine/Plugins/Experimental/MeshModelingToolset/Source/MeshModelingToolsEditorOnly/Private/ToolTargets/VolumeMeshDescriptionToolTarget.cpp b/Engine/Plugins/Experimental/MeshModelingToolset/Source/MeshModelingToolsEditorOnly/Private/ToolTargets/VolumeMeshDescriptionToolTarget.cpp
new file mode 100644
--- /dev/null
+++ b/Engine/Plugins/Experimental/MeshModelingToolset/Source/MeshModelingToolsEditorOnly/Private/ToolTargets/VolumeMeshDescriptionToolTarget.cpp
@@ -0,0 +1,91 @@
+// Copyright Epic Games, Inc. All Rights Reserved.
+
+#pragma once
+
+#include "ToolTargets/VolumeMeshDescriptionToolTarget.h"
+
+#include "Components/BrushComponent.h"
+#include "ConversionUtils/VolumeToDynamicMesh.h"
+#include "ConversionUtils/DynamicMeshToVolume.h"
+#include "DynamicMeshToMeshDescription.h"
+#include "Engine/Brush.h"
+#include "GameFramework/Volume.h"
+#include "MeshDescriptionToDynamicMesh.h"
+#include "MeshNormals.h"
+#include "StaticMeshAttributes.h"
+#include "ToolSetupUtil.h"
+
+#include "ExplicitUseGeometryMathTypes.h"		// using UE::Geometry::(math types)
+using namespace UE::Geometry;
+
+FMeshDescription* UVolumeMeshDescriptionToolTarget::GetMeshDescription()
+{
+	if (!ConvertedMeshDescription.IsValid())
+	{
+		// Note: We can go directly from a volume to a mesh description using GetBrushMesh() in
+		// Editor.h. However, that path doesn't assign polygroups to the result, which we
+		// typically want when using this target, hence the path through a dynamic mesh.
+
+		TSharedPtr<FDynamicMesh3> DynamicMesh = GetDynamicMesh();
+		if (!DynamicMesh)
+		{
+			return nullptr;
+		}
+
+		ConvertedMeshDescription = MakeShared<FMeshDescription>();
+		FStaticMeshAttributes Attributes(*ConvertedMeshDescription);
+		Attributes.Register();
+		FDynamicMeshToMeshDescription Converter;
+		Converter.Convert(DynamicMesh.Get(), *ConvertedMeshDescription);
+	}
+	return ConvertedMeshDescription.Get();
+}
+
+void UVolumeMeshDescriptionToolTarget::CommitMeshDescription(const FCommitter& Committer)
+{
+	check(IsValid());
+
+	// Let the user fill our mesh description with the Committer
+	if (!ConvertedMeshDescription.IsValid())
+	{
+		ConvertedMeshDescription = MakeShared<FMeshDescription>();
+		FStaticMeshAttributes Attributes(*ConvertedMeshDescription);
+		Attributes.Register();
+	}
+	FCommitterParams CommitParams;
+	CommitParams.MeshDescriptionOut = ConvertedMeshDescription.Get();
+	Committer(CommitParams);
+
+	// The conversion we have right now is from dynamic mesh to volume, so we convert
+	// to dynamic mesh first.
+	FDynamicMesh3 DynamicMesh;
+	FMeshDescriptionToDynamicMesh Converter;
+	Converter.Convert(ConvertedMeshDescription.Get(), DynamicMesh);
+
+	CommitDynamicMesh(DynamicMesh);
+}
+
+
+// Factory
+
+bool UVolumeMeshDescriptionToolTargetFactory::CanBuildTarget(UObject* SourceObject, const FToolTargetTypeRequirements& Requirements) const
+{
+	UBrushComponent* BrushComponent = Cast<UBrushComponent>(SourceObject);
+	if (!BrushComponent)
+	{
+		return false;
+	}
+
+	return Cast<AVolume>(BrushComponent->GetOwner()) && Requirements.AreSatisfiedBy(UVolumeMeshDescriptionToolTarget::StaticClass());
+}
+
+UToolTarget* UVolumeMeshDescriptionToolTargetFactory::BuildTarget(UObject* SourceObject, const FToolTargetTypeRequirements& Requirements)
+{
+	UVolumeMeshDescriptionToolTarget* Target = NewObject<UVolumeMeshDescriptionToolTarget>();
+	Target->Component = Cast<UBrushComponent>(SourceObject);
+
+	check(Target->Component && Requirements.AreSatisfiedBy(Target)
+		&& Cast<AVolume>(Cast<UBrushComponent>(SourceObject)->GetOwner()));
+
+	return Target;
+}
\ No newline at end of file
diff --git a/Engine/Plugins/Experimental/MeshModelingToolset/Source/MeshModelingToolsEditorOnly/Public/ConversionUtils/VolumeMeshDescriptionToolTarget.h b/Engine/Plugins/Experimental/MeshModelingToolset/Source/MeshModelingToolsEditorOnly/Public/ConversionUtils/VolumeMeshDescriptionToolTarget.h
deleted file mode 100644
diff --git a/Engine/Plugins/Experimental/MeshModelingToolset/Source/MeshModelingToolsEditorOnly/Public/ToolTargets/SkeletalMeshComponentToolTarget.h b/Engine/Plugins/Experimental/MeshModelingToolset/Source/MeshModelingToolsEditorOnly/Public/ToolTargets/SkeletalMeshComponentToolTarget.h
new file mode 100644
--- /dev/null
+++ b/Engine/Plugins/Experimental/MeshModelingToolset/Source/MeshModelingToolsEditorOnly/Public/ToolTargets/SkeletalMeshComponentToolTarget.h
@@ -0,0 +1,81 @@
+﻿// Copyright Epic Games, Inc. All Rights Reserved.
+
+#pragma once
+
+#include "CoreMinimal.h"
+
+#include "TargetInterfaces/DynamicMeshCommitter.h"
+#include "TargetInterfaces/DynamicMeshProvider.h"
+#include "TargetInterfaces/MaterialProvider.h"
+#include "TargetInterfaces/MeshDescriptionCommitter.h"
+#include "TargetInterfaces/MeshDescriptionProvider.h"
+#include "TargetInterfaces/SkeletalMeshBackedTarget.h"
+#include "ToolTargets/PrimitiveComponentToolTarget.h"
+
+#include "SkeletalMeshComponentToolTarget.generated.h"
+
+
+class USkeletalMesh;
+
+/**
+ * A tool target backed by a skeletal mesh component that can provide and take a mesh
+ * description.
+ */
+UCLASS(Transient)
+class MESHMODELINGTOOLSEDITORONLY_API USkeletalMeshComponentToolTarget :
+	public UPrimitiveComponentToolTarget,
+	public IMeshDescriptionCommitter,
+	public IMeshDescriptionProvider,
+	public IDynamicMeshProvider, 
+	public IDynamicMeshCommitter,
+	public IMaterialProvider,
+	public ISkeletalMeshBackedTarget
+{
+	GENERATED_BODY()
+
+public:
+	// IMeshDescriptionProvider implementation
+	FMeshDescription* GetMeshDescription() override;
+
+	// IMeshDescritpionCommitter implementation
+	void CommitMeshDescription(const FCommitter& Committer) override;
+	using IMeshDescriptionCommitter::CommitMeshDescription; // unhide the other overload
+
+	// IMaterialProvider implementation
+	int32 GetNumMaterials() const override;
+	UMaterialInterface* GetMaterial(int32 MaterialIndex) const override;
+	void GetMaterialSet(FComponentMaterialSet& MaterialSetOut, bool bPreferAssetMaterials) const override;
+	bool CommitMaterialSetUpdate(const FComponentMaterialSet& MaterialSet, bool bApplyToAsset) override;	
+
+	// IDynamicMeshProvider
+	virtual TSharedPtr<UE::Geometry::FDynamicMesh3> GetDynamicMesh() override;
+
+	// IDynamicMeshCommitter
+	virtual void CommitDynamicMesh(const UE::Geometry::FDynamicMesh3& Mesh, const FDynamicMeshCommitInfo& CommitInfo) override;
+	using IDynamicMeshCommitter::CommitDynamicMesh; // unhide the other overload
+
+	// ISkeletalMeshBackedTarget implementation
+	USkeletalMesh* GetSkeletalMesh() const override;
+
+protected:
+	// So that the tool target factory can poke into Component.
+	friend class USkeletalMeshComponentToolTargetFactory;
+
+private:
+	// Until USkeletalMesh stores its internal representation as FMeshDescription, we need to
+	// retain the storage here to cover the lifetime of the pointer returned by GetMeshDescription(). 
+	TUniquePtr<FMeshDescription> CachedMeshDescription;	
+};
+
+/** Factory for USkeletalMeshComponentToolTarget to be used by the target manager. */
+UCLASS(Transient)
+class MESHMODELINGTOOLSEDITORONLY_API USkeletalMeshComponentToolTargetFactory : public UToolTargetFactory
+{
+	GENERATED_BODY()
+
+public:
+
+	bool CanBuildTarget(UObject* SourceObject, const FToolTargetTypeRequirements& TargetTypeInfo) const override;
+
+	UToolTarget* BuildTarget(UObject* SourceObject, const FToolTargetTypeRequirements& TargetTypeInfo) override;
+};
diff --git a/Engine/Plugins/Experimental/MeshModelingToolset/Source/MeshModelingToolsEditorOnly/Public/ToolTargets/StaticMeshComponentToolTarget.h b/Engine/Plugins/Experimental/MeshModelingToolset/Source/MeshModelingToolsEditorOnly/Public/ToolTargets/StaticMeshComponentToolTarget.h
new file mode 100644
--- /dev/null
+++ b/Engine/Plugins/Experimental/MeshModelingToolset/Source/MeshModelingToolsEditorOnly/Public/ToolTargets/StaticMeshComponentToolTarget.h
@@ -0,0 +1,101 @@
+// Copyright Epic Games, Inc. All Rights Reserved.
+
+#pragma once
+
+#include "CoreMinimal.h"
+
+#include "TargetInterfaces/DynamicMeshCommitter.h"
+#include "TargetInterfaces/DynamicMeshProvider.h"
+#include "TargetInterfaces/MaterialProvider.h"
+#include "TargetInterfaces/MeshDescriptionCommitter.h"
+#include "TargetInterfaces/MeshDescriptionProvider.h"
+#include "TargetInterfaces/StaticMeshBackedTarget.h"
+#include "ToolTargets/PrimitiveComponentToolTarget.h"
+#include "ComponentSourceInterfaces.h"  // for EStaticMeshEditingLOD
+
+#include "StaticMeshComponentToolTarget.generated.h"
+
+class UStaticMesh;
+
+/**
+ * A tool target backed by a static mesh component that can provide and take a mesh
+ * description.
+ */
+UCLASS(Transient)
+class MESHMODELINGTOOLSEDITORONLY_API UStaticMeshComponentToolTarget : public UPrimitiveComponentToolTarget,
+	public IMeshDescriptionCommitter, public IMeshDescriptionProvider, public IMaterialProvider, public IStaticMeshBackedTarget,
+	public IDynamicMeshProvider, public IDynamicMeshCommitter
+{
+	GENERATED_BODY()
+
+public:
+	/**
+	 * Configure active LOD to edit. Can only call this after Component is configured in base UPrimitiveComponentToolTarget.
+	 * If requested LOD does not exist, fall back to one that does.
+	 */
+	virtual void SetEditingLOD(EStaticMeshEditingLOD RequestedEditingLOD);
+
+	/** @return current editing LOD */
+	virtual EStaticMeshEditingLOD GetEditingLOD() const { return EditingLOD; }
+
+public:
+	virtual bool IsValid() const override;
+
+
+public:
+	// IMeshDescriptionProvider implementation
+	FMeshDescription* GetMeshDescription() override;
+
+	// IMeshDescritpionCommitter implementation
+	virtual void CommitMeshDescription(const FCommitter& Committer) override;
+	using IMeshDescriptionCommitter::CommitMeshDescription; // unhide the other overload
+
+	// IDynamicMeshProvider
+	virtual TSharedPtr<UE::Geometry::FDynamicMesh3> GetDynamicMesh() override;
+
+	// IDynamicMeshCommitter
+	virtual void CommitDynamicMesh(const UE::Geometry::FDynamicMesh3& Mesh, const FDynamicMeshCommitInfo& CommitInfo) override;
+	using IDynamicMeshCommitter::CommitDynamicMesh; // unhide the other overload
+
+	// IMaterialProvider implementation
+	int32 GetNumMaterials() const override;
+	UMaterialInterface* GetMaterial(int32 MaterialIndex) const override;
+	void GetMaterialSet(FComponentMaterialSet& MaterialSetOut, bool bPreferAssetMaterials) const override;
+	virtual bool CommitMaterialSetUpdate(const FComponentMaterialSet& MaterialSet, bool bApplyToAsset) override;	
+
+	// IStaticMeshBackedTarget
+	UStaticMesh* GetStaticMesh() const override;
+
+	// Rest provided by parent class
+
+public:
+
+protected:
+	EStaticMeshEditingLOD EditingLOD = EStaticMeshEditingLOD::LOD0;
+
+	friend class UStaticMeshComponentToolTargetFactory;
+};
+
+
+/** Factory for UStaticMeshComponentToolTarget to be used by the target manager. */
+UCLASS(Transient)
+class MESHMODELINGTOOLSEDITORONLY_API UStaticMeshComponentToolTargetFactory : public UToolTargetFactory
+{
+	GENERATED_BODY()
+
+public:
+
+	virtual bool CanBuildTarget(UObject* SourceObject, const FToolTargetTypeRequirements& TargetTypeInfo) const override;
+
+	virtual UToolTarget* BuildTarget(UObject* SourceObject, const FToolTargetTypeRequirements& TargetTypeInfo) override;
+
+
+
+public:
+	virtual EStaticMeshEditingLOD GetActiveEditingLOD() const { return EditingLOD; }
+	virtual void SetActiveEditingLOD(EStaticMeshEditingLOD NewEditingLOD);
+
+protected:
+	// LOD to edit, default is to edit LOD0
+	EStaticMeshEditingLOD EditingLOD = EStaticMeshEditingLOD::LOD0;
+};
\ No newline at end of file
diff --git a/Engine/Plugins/Experimental/MeshModelingToolset/Source/MeshModelingToolsEditorOnly/Public/ToolTargets/VolumeDynamicMeshToolTarget.h b/Engine/Plugins/Experimental/MeshModelingToolset/Source/MeshModelingToolsEditorOnly/Public/ToolTargets/VolumeDynamicMeshToolTarget.h
new file mode 100644
--- /dev/null
+++ b/Engine/Plugins/Experimental/MeshModelingToolset/Source/MeshModelingToolsEditorOnly/Public/ToolTargets/VolumeDynamicMeshToolTarget.h
@@ -0,0 +1,65 @@
+// Copyright Epic Games, Inc. All Rights Reserved.
+
+#pragma once
+
+#include "CoreMinimal.h"
+
+#include "ConversionUtils/VolumeToDynamicMesh.h" // FVolumeToMeshOptions
+#include "TargetInterfaces/DynamicMeshCommitter.h"
+#include "TargetInterfaces/DynamicMeshProvider.h"
+#include "TargetInterfaces/MaterialProvider.h"
+#include "ToolTargets/PrimitiveComponentToolTarget.h"
+
+#include "VolumeDynamicMeshToolTarget.generated.h"
+
+/**
+ * A tool target backed by AVolume
+ */
+UCLASS(Transient)
+class MESHMODELINGTOOLSEDITORONLY_API UVolumeDynamicMeshToolTarget : public UPrimitiveComponentToolTarget,
+	public IDynamicMeshCommitter, public IDynamicMeshProvider, public IMaterialProvider
+{
+	GENERATED_BODY()
+
+public:
+	UVolumeDynamicMeshToolTarget();
+
+	const UE::Conversion::FVolumeToMeshOptions& GetVolumeToMeshOptions() { return VolumeToMeshOptions; }
+
+	// IDynamicMeshProvider implementation
+	virtual TSharedPtr<UE::Geometry::FDynamicMesh3> GetDynamicMesh() override;
+
+	// IDynamicMeshCommitter implementation
+	virtual void CommitDynamicMesh(const UE::Geometry::FDynamicMesh3& Mesh, const FDynamicMeshCommitInfo&) override;
+	using IDynamicMeshCommitter::CommitDynamicMesh; // unhide the other overload
+
+	// IMaterialProvider implementation
+	virtual int32 GetNumMaterials() const override;
+	virtual UMaterialInterface* GetMaterial(int32 MaterialIndex) const override;
+	// Ignores bPreferAssetMaterials
+	virtual void GetMaterialSet(FComponentMaterialSet& MaterialSetOut, bool bPreferAssetMaterials) const override;
+
+	// Doesn't actually do anything for a volume
+	virtual bool CommitMaterialSetUpdate(const FComponentMaterialSet& MaterialSet, bool bApplyToAsset) override { return false; }
+
+	// Rest provided by parent class
+
+protected:
+	TSharedPtr<FMeshDescription> ConvertedMeshDescription;
+	UE::Conversion::FVolumeToMeshOptions VolumeToMeshOptions;
+
+	friend class UVolumeDynamicMeshToolTargetFactory;
+};
+
+/** Factory for UVolumeDynamicMeshToolTarget to be used by the target manager. */
+UCLASS(Transient)
+class MESHMODELINGTOOLSEDITORONLY_API UVolumeDynamicMeshToolTargetFactory : public UToolTargetFactory
+{
+	GENERATED_BODY()
+
+public:
+
+	virtual bool CanBuildTarget(UObject* SourceObject, const FToolTargetTypeRequirements& TargetTypeInfo) const override;
+
+	virtual UToolTarget* BuildTarget(UObject* SourceObject, const FToolTargetTypeRequirements& TargetTypeInfo) override;
+};
\ No newline at end of file
diff --git a/Engine/Plugins/Experimental/MeshModelingToolset/Source/MeshModelingToolsEditorOnly/Public/ToolTargets/VolumeMeshDescriptionToolTarget.h b/Engine/Plugins/Experimental/MeshModelingToolset/Source/MeshModelingToolsEditorOnly/Public/ToolTargets/VolumeMeshDescriptionToolTarget.h
new file mode 100644
--- /dev/null
+++ b/Engine/Plugins/Experimental/MeshModelingToolset/Source/MeshModelingToolsEditorOnly/Public/ToolTargets/VolumeMeshDescriptionToolTarget.h
@@ -0,0 +1,50 @@
+// Copyright Epic Games, Inc. All Rights Reserved.
+
+#pragma once
+
+#include "CoreMinimal.h"
+
+#include "TargetInterfaces/MeshDescriptionCommitter.h"
+#include "TargetInterfaces/MeshDescriptionProvider.h"
+#include "ToolTargets/VolumeDynamicMeshToolTarget.h"
+
+#include "VolumeMeshDescriptionToolTarget.generated.h"
+
+/**
+ * A tool target backed by AVolume
+ */
+UCLASS(Transient)
+class MESHMODELINGTOOLSEDITORONLY_API UVolumeMeshDescriptionToolTarget : public UVolumeDynamicMeshToolTarget,
+	public IMeshDescriptionCommitter, public IMeshDescriptionProvider
+{
+	GENERATED_BODY()
+
+public:
+
+	// IMeshDescriptionProvider implementation
+	FMeshDescription* GetMeshDescription() override;
+
+	// IMeshDescritpionCommitter implementation
+	virtual void CommitMeshDescription(const FCommitter& Committer) override;
+	using IMeshDescriptionCommitter::CommitMeshDescription; // unhide the other overload
+
+protected:
+	// This isn't for caching- we have to take ownership of the mesh description because it is
+	// expected for things like a static mesh.
+	TSharedPtr<FMeshDescription> ConvertedMeshDescription;
+
+	friend class UVolumeMeshDescriptionToolTargetFactory;
+};
+
+/** Factory for UVolumeMeshDescriptionToolTarget to be used by the target manager. */
+UCLASS(Transient)
+class MESHMODELINGTOOLSEDITORONLY_API UVolumeMeshDescriptionToolTargetFactory : public UToolTargetFactory
+{
+	GENERATED_BODY()
+
+public:
+
+	virtual bool CanBuildTarget(UObject* SourceObject, const FToolTargetTypeRequirements& TargetTypeInfo) const override;
+
+	virtual UToolTarget* BuildTarget(UObject* SourceObject, const FToolTargetTypeRequirements& TargetTypeInfo) override;
+};
\ No newline at end of file
diff --git a/Engine/Plugins/Experimental/MeshModelingToolset/Source/ModelingComponents/Private/SimpleDynamicMeshComponent.cpp b/Engine/Plugins/Experimental/MeshModelingToolset/Source/ModelingComponents/Private/SimpleDynamicMeshComponent.cpp
--- a/Engine/Plugins/Experimental/MeshModelingToolset/Source/ModelingComponents/Private/SimpleDynamicMeshComponent.cpp
+++ b/Engine/Plugins/Experimental/MeshModelingToolset/Source/ModelingComponents/Private/SimpleDynamicMeshComponent.cpp
@@ -179,29 +179,14 @@ void USimpleDynamicMeshComponent::ApplyTransform(const UE::Geometry::FTransform3
 
 void USimpleDynamicMeshComponent::Bake(FMeshDescription* MeshDescription, bool bHaveModifiedTopology, const FConversionToMeshDescriptionOptions& ConversionOptions)
 {
-	
 	FDynamicMeshToMeshDescription Converter(ConversionOptions);
-	if (bHaveModifiedTopology == false && Converter.HaveMatchingElementCounts(Mesh.Get(), MeshDescription))
+	if (!bHaveModifiedTopology)
 	{
-		if (ConversionOptions.bUpdatePositions)
-		{
-			Converter.Update(Mesh.Get(), *MeshDescription, ConversionOptions.bUpdateNormals, ConversionOptions.bUpdateTangents, ConversionOptions.bUpdateUVs);
-		}
-		else if (ConversionOptions.bUpdateNormals || ConversionOptions.bUpdateTangents || ConversionOptions.bUpdateUVs)
-		{
-			Converter.UpdateAttributes(Mesh.Get(), *MeshDescription, ConversionOptions.bUpdateNormals, ConversionOptions.bUpdateTangents, ConversionOptions.bUpdateUVs);
-		}
-
-		if (ConversionOptions.bUpdateVtxColors)
-		{
-			Converter.UpdateVertexColors(Mesh.Get(), *MeshDescription);
-		}
+		Converter.UpdateUsingConversionOptions(Mesh.Get(), *MeshDescription);
 	}
 	else
 	{
 		Converter.Convert(Mesh.Get(), *MeshDescription);
-
-		//UE_LOG(LogTemp, Warning, TEXT("MeshDescription has %d instances"), MeshDescription->VertexInstances().Num());
 	}
 }
 
diff --git a/Engine/Plugins/Experimental/ModelingToolsEditorMode/Source/ModelingToolsEditorMode/Private/ModelingToolsEditorMode.cpp b/Engine/Plugins/Experimental/ModelingToolsEditorMode/Source/ModelingToolsEditorMode/Private/ModelingToolsEditorMode.cpp
--- a/Engine/Plugins/Experimental/ModelingToolsEditorMode/Source/ModelingToolsEditorMode/Private/ModelingToolsEditorMode.cpp
+++ b/Engine/Plugins/Experimental/ModelingToolsEditorMode/Source/ModelingToolsEditorMode/Private/ModelingToolsEditorMode.cpp
@@ -7,7 +7,8 @@
 #include "Toolkits/ToolkitManager.h"
 #include "ToolTargetManager.h"
 #include "ToolTargets/StaticMeshComponentToolTarget.h"
-#include "ConversionUtils/VolumeMeshDescriptionToolTarget.h"
+#include "ToolTargets/VolumeDynamicMeshToolTarget.h"
+#include "ToolTargets/VolumeMeshDescriptionToolTarget.h"
 #include "Framework/Commands/UICommandList.h"
 #include "Framework/Application/SlateApplication.h"
 #include "Framework/MultiBox/MultiBoxBuilder.h"
@@ -296,6 +297,7 @@ void UModelingToolsEditorMode::Enter()
 
 	// Register builders for tool targets that the mode uses.
 	ToolsContext->TargetManager->AddTargetFactory(NewObject<UStaticMeshComponentToolTargetFactory>(ToolsContext->TargetManager));
+	ToolsContext->TargetManager->AddTargetFactory(NewObject<UVolumeDynamicMeshToolTargetFactory>(ToolsContext->TargetManager));
 	ToolsContext->TargetManager->AddTargetFactory(NewObject<UVolumeMeshDescriptionToolTargetFactory>(ToolsContext->TargetManager));
 
 	// register stylus event handler
diff --git a/Engine/Source/Editor/Experimental/EditorInteractiveToolsFramework/Private/ToolTargets/SkeletalMeshComponentToolTarget.cpp b/Engine/Source/Editor/Experimental/EditorInteractiveToolsFramework/Private/ToolTargets/SkeletalMeshComponentToolTarget.cpp
deleted file mode 100644
diff --git a/Engine/Source/Editor/Experimental/EditorInteractiveToolsFramework/Private/ToolTargets/StaticMeshComponentToolTarget.cpp b/Engine/Source/Editor/Experimental/EditorInteractiveToolsFramework/Private/ToolTargets/StaticMeshComponentToolTarget.cpp
deleted file mode 100644
diff --git a/Engine/Source/Editor/Experimental/EditorInteractiveToolsFramework/Public/ToolTargets/SkeletalMeshComponentToolTarget.h b/Engine/Source/Editor/Experimental/EditorInteractiveToolsFramework/Public/ToolTargets/SkeletalMeshComponentToolTarget.h
deleted file mode 100644
diff --git a/Engine/Source/Editor/Experimental/EditorInteractiveToolsFramework/Public/ToolTargets/StaticMeshComponentToolTarget.h b/Engine/Source/Editor/Experimental/EditorInteractiveToolsFramework/Public/ToolTargets/StaticMeshComponentToolTarget.h
deleted file mode 100644
