A primitive skin weights paint tool for skeletal meshes.

diff --git a/Engine/Plugins/Experimental/MeshModelingToolset/Source/MeshModelingTools/MeshModelingTools.Build.cs b/Engine/Plugins/Experimental/MeshModelingToolset/Source/MeshModelingTools/MeshModelingTools.Build.cs
--- a/Engine/Plugins/Experimental/MeshModelingToolset/Source/MeshModelingTools/MeshModelingTools.Build.cs
+++ b/Engine/Plugins/Experimental/MeshModelingToolset/Source/MeshModelingTools/MeshModelingTools.Build.cs
@@ -14,14 +14,12 @@ public class MeshModelingTools : ModuleRules
 			}
 			);
 				
-		
-		PrivateIncludePaths.AddRange(
+		PublicIncludePathModuleNames.AddRange(
 			new string[] {
-				// ... add other private include paths required here ...
+				"AnimationCore",			// For the BoneWeights.h include
 			}
-			);
-			
-		
+		);
+	
 		PublicDependencyModuleNames.AddRange(
 			new string[]
 			{
@@ -42,6 +40,7 @@ public class MeshModelingTools : ModuleRules
 				"MeshConversion",
 				"MeshDescription",
                 "StaticMeshDescription",
+                "SkeletalMeshDescription",
 				"ModelingComponents",
 				"ModelingOperators",
 
@@ -75,7 +74,8 @@ public class MeshModelingTools : ModuleRules
 				"RenderCore",
 				"ModelingOperators",
 				"InputCore",
-				"PhysicsCore"
+				"PhysicsCore",
+
 				// ... add private dependencies that you statically link with here ...
 			}
 			);
diff --git a/Engine/Plugins/Experimental/MeshModelingToolset/Source/MeshModelingTools/Private/SkinWeightsPaintTool.cpp b/Engine/Plugins/Experimental/MeshModelingToolset/Source/MeshModelingTools/Private/SkinWeightsPaintTool.cpp
new file mode 100644
--- /dev/null
+++ b/Engine/Plugins/Experimental/MeshModelingToolset/Source/MeshModelingTools/Private/SkinWeightsPaintTool.cpp
@@ -0,0 +1,732 @@
+// Copyright Epic Games, Inc. All Rights Reserved.
+
+#include "SkinWeightsPaintTool.h"
+
+#include "InteractiveToolManager.h"
+#include "ToolBuilderUtil.h"
+#include "SkeletalMeshAttributes.h"
+#include "SkeletalDebugRendering.h"
+#include "Math/UnrealMathUtility.h"
+#include "Components/SkeletalMeshComponent.h"
+#include "TargetInterfaces/MeshDescriptionProvider.h"
+#include "TargetInterfaces/PrimitiveComponentBackedTarget.h"
+
+#include "MeshDescription.h"
+
+
+#define LOCTEXT_NAMESPACE "USkinWeightsPaintTool"
+
+// Any weight below this value is ignored, since it won't be representable in a uint8.
+const float MinimumWeightThreshold = 1.0f / 255.0f;
+
+
+void FMeshSkinWeightsChange::Apply(UObject* Object)
+{
+	USkinWeightsPaintTool* Tool = CastChecked<USkinWeightsPaintTool>(Object);
+
+	Tool->ExternalUpdateValues(BoneName, NewWeights);
+}
+
+void FMeshSkinWeightsChange::Revert(UObject* Object)
+{
+	USkinWeightsPaintTool* Tool = CastChecked<USkinWeightsPaintTool>(Object);
+
+	Tool->ExternalUpdateValues(BoneName, OldWeights);
+}
+
+void FMeshSkinWeightsChange::UpdateValues(const TArray<int32>& Indices, const TArray<float>& OldValues, const TArray<float>& NewValues)
+{
+	const int32 NumIndices = Indices.Num();
+	for (int32 i = 0; i < NumIndices; i++)
+	{
+		NewWeights.Add(Indices[i], NewValues[i]);
+		OldWeights.FindOrAdd(Indices[i], OldValues[i]);
+	}
+}
+
+
+/*
+ * ToolBuilder
+ */
+
+UMeshSurfacePointTool* USkinWeightsPaintToolBuilder::CreateNewTool(const FToolBuilderState& SceneState) const
+{
+	return NewObject<USkinWeightsPaintTool>(SceneState.ToolManager);
+}
+
+
+void USkinWeightsPaintTool::Setup()
+{
+	UDynamicMeshBrushTool::Setup();
+
+	IPrimitiveComponentBackedTarget* TargetComponent = Cast<IPrimitiveComponentBackedTarget>(Target);
+	USkeletalMeshComponent* Component = Cast<USkeletalMeshComponent>(TargetComponent->GetOwnerComponent());
+
+	// hide strength and falloff
+	BrushProperties->RestoreProperties(this);
+
+	ToolProps = NewObject<USkinWeightsPaintToolProperties>(this);
+	ToolProps->RestoreProperties(this);
+	if (Component && Component->SkeletalMesh)
+	{
+		// Get all non-virtual bones
+		// TArray<int32> BoneIndices;
+		USkeletalMesh& SkeletalMesh = *Component->SkeletalMesh;
+		// SkeletalMesh.RefSkeleton.GetRawRefBoneInfo();
+		
+
+		// Initialize the bone browser
+		FCurveEvaluationOption CurveEvalOption(
+				Component->GetAllowedAnimCurveEvaluate(),
+				&Component->GetDisallowedAnimCurvesEvaluation(),
+				0 /* Always use the highest LOD */
+				);
+		BoneContainer.InitializeTo(Component->RequiredBones, CurveEvalOption, SkeletalMesh);
+
+		ToolProps->SkeletalMesh = Component->SkeletalMesh;
+		ToolProps->CurrentBone.Initialize(BoneContainer);
+
+		// Pick the first root bone as the initial selection.
+		PendingCurrentBone = Component->SkeletalMesh->GetRefSkeleton().GetBoneName(0);
+		ToolProps->CurrentBone.BoneName = PendingCurrentBone.GetValue();
+	}
+	AddToolPropertySource(ToolProps);
+
+	// Update the skeleton drawing information from the original bind pose
+	MaxDrawRadius = Component->Bounds.SphereRadius * 0.01f;
+
+	UpdateBonePositionInfos(MaxDrawRadius);
+
+	// configure preview mesh
+	PreviewMesh->SetTangentsMode(EDynamicMeshTangentCalcType::AutoCalculated);
+	//PreviewMesh->EnableWireframe(SelectionProps->bShowWireframe);
+	PreviewMesh->GetRootComponent()->bCastDynamicShadow = false;
+
+	// enable vtx colors on preview mesh
+	PreviewMesh->EditMesh([](FDynamicMesh3& Mesh)
+	{
+		Mesh.EnableAttributes();
+		Mesh.Attributes()->DisablePrimaryColors();
+		Mesh.Attributes()->EnablePrimaryColors();
+		// Create an overlay that has no split elements, init with zero value.
+		Mesh.Attributes()->PrimaryColors()->CreateFromPredicate([](int ParentVID, int TriIDA, int TriIDB){return true;}, 0.f);
+	});
+
+	// build octree
+	VerticesOctree.Initialize(PreviewMesh->GetMesh(), true);
+
+	UMaterialInterface* VtxColorMaterial = GetToolManager()->GetContextQueriesAPI()->GetStandardMaterial(EStandardToolContextMaterials::VertexColorMaterial);
+	if (VtxColorMaterial != nullptr)
+	{
+		PreviewMesh->SetOverrideRenderMaterial(VtxColorMaterial);
+	}
+
+	RecalculateBrushRadius();
+
+	GetToolManager()->DisplayMessage(
+		LOCTEXT("OnStartSkinWeightsPaint", "Paint per-bone skin weights. [ and ] change brush size, Ctrl to Erase/Subtract, Shift to Smooth"),
+		EToolMessageLevel::UserNotification);
+
+	EditedMesh = MakeUnique<FMeshDescription>();
+	*EditedMesh = *Cast<IMeshDescriptionProvider>(Target)->GetMeshDescription();
+
+	InitializeSkinWeights();
+
+	CurrentBoneWatcher.Initialize([this]() { return ToolProps->CurrentBone; },
+		[this](FBoneReference NewValue) { PendingCurrentBone = NewValue.BoneName; }, ToolProps->CurrentBone);
+
+	bVisibleWeightsValid = false;
+}
+
+
+
+void USkinWeightsPaintTool::RegisterActions(FInteractiveToolActionSet& ActionSet)
+{
+	UDynamicMeshBrushTool::RegisterActions(ActionSet);
+}
+
+
+
+
+void USkinWeightsPaintTool::OnTick(float DeltaTime)
+{
+	CurrentBoneWatcher.CheckAndUpdate();
+
+	if (bStampPending)
+	{
+		ApplyStamp(LastStamp);
+		bStampPending = false;
+	}
+
+	if (PendingCurrentBone.IsSet())
+	{
+		UpdateCurrentBone(*PendingCurrentBone);
+		PendingCurrentBone.Reset();
+	}
+
+	if (bVisibleWeightsValid == false)
+	{
+		UpdateBoneVisualization();
+		bVisibleWeightsValid = true;
+	}
+}
+
+
+
+
+void USkinWeightsPaintTool::Render(IToolsContextRenderAPI* RenderAPI)
+{
+	UDynamicMeshBrushTool::Render(RenderAPI);
+
+	// FIXME: Make selective.
+	RenderBonePositions(RenderAPI->GetPrimitiveDrawInterface());
+}
+
+
+
+
+
+
+
+bool USkinWeightsPaintTool::HitTest(const FRay& Ray, FHitResult& OutHit)
+{
+	bool bHit = UDynamicMeshBrushTool::HitTest(Ray, OutHit);
+	//if (bHit && SelectionProps->bHitBackFaces == false)
+	//{
+	//	const FDynamicMesh3* SourceMesh = PreviewMesh->GetPreviewDynamicMesh();
+	//	FVector3d Normal, Centroid;
+	//	double Area;
+	//	SourceMesh->GetTriInfo(OutHit.FaceIndex, Normal, Area, Centroid);
+	//	FViewCameraState StateOut;
+	//	GetToolManager()->GetContextQueriesAPI()->GetCurrentViewState(StateOut);
+	//	FVector3d LocalEyePosition(ComponentTarget->GetWorldTransform().InverseTransformPosition(StateOut.Position));
+
+	//	if (Normal.Dot((Centroid - LocalEyePosition)) > 0)
+	//	{
+	//		bHit = false;
+	//	}
+	//}
+	return bHit;
+}
+
+
+void USkinWeightsPaintTool::OnBeginDrag(const FRay& WorldRay)
+{
+	UDynamicMeshBrushTool::OnBeginDrag(WorldRay);
+
+	PreviewBrushROI.Reset();
+	if (IsInBrushStroke())
+	{
+		bInRemoveStroke = GetCtrlToggle();
+		bInSmoothStroke = GetShiftToggle();
+		BeginChange();
+		StartStamp = UBaseBrushTool::LastBrushStamp;
+		LastStamp = StartStamp;
+		bStampPending = true;
+	}
+}
+
+
+
+void USkinWeightsPaintTool::OnUpdateDrag(const FRay& WorldRay)
+{
+	UDynamicMeshBrushTool::OnUpdateDrag(WorldRay);
+	if (IsInBrushStroke())
+	{
+		LastStamp = UBaseBrushTool::LastBrushStamp;
+		bStampPending = true;
+	}
+}
+
+
+
+void USkinWeightsPaintTool::OnEndDrag(const FRay& Ray)
+{
+	UDynamicMeshBrushTool::OnEndDrag(Ray);
+
+	bInRemoveStroke = false;
+	bInSmoothStroke = false;
+	bStampPending = false;
+
+	// close change record
+	TUniquePtr<FMeshSkinWeightsChange> Change = EndChange();
+
+	GetToolManager()->BeginUndoTransaction(LOCTEXT("BoneWeightValuesChange", "Paint"));
+
+	GetToolManager()->EmitObjectChange(this, MoveTemp(Change), LOCTEXT("BoneWeightValuesChange", "Paint"));
+
+	GetToolManager()->EndUndoTransaction();
+}
+
+
+bool USkinWeightsPaintTool::OnUpdateHover(const FInputDeviceRay& DevicePos)
+{
+	UDynamicMeshBrushTool::OnUpdateHover(DevicePos);
+
+	// todo get rid of this redundant hit test!
+	FHitResult OutHit;
+	if (UDynamicMeshBrushTool::HitTest(DevicePos.WorldRay, OutHit))
+	{
+		PreviewBrushROI.Reset();
+		CalculateVertexROI(LastBrushStamp, PreviewBrushROI);
+	}
+
+	return true;
+}
+
+
+
+void USkinWeightsPaintTool::CalculateVertexROI(const FBrushStampData& Stamp, TArray<int>& VertexROI)
+{
+	using namespace UE::Geometry;
+	IPrimitiveComponentBackedTarget* TargetComponent = Cast<IPrimitiveComponentBackedTarget>(Target);
+	FTransform3d Transform(TargetComponent->GetWorldTransform());
+	FVector3d StampPosLocal = Transform.InverseTransformPosition((FVector3d)Stamp.WorldPosition);
+
+	float RadiusSqr = CurrentBrushRadius * CurrentBrushRadius;
+	const FDynamicMesh3* Mesh = PreviewMesh->GetPreviewDynamicMesh();
+	FAxisAlignedBox3d QueryBox(StampPosLocal, CurrentBrushRadius);
+	VerticesOctree.RangeQuery(QueryBox,
+		[&](int32 VertexID) { return FVector3d::DistSquared(Mesh->GetVertex(VertexID), StampPosLocal) < RadiusSqr; },
+		VertexROI);
+}
+
+
+FVector4f USkinWeightsPaintTool::WeightToColor(float Value)
+{
+	Value = FMath::Clamp(Value, 0.0f, 1.0f);
+
+	{
+		// A close approximation of the skeletal mesh editor's bone weight ramp. 
+		const FLinearColor HSV((1.0f - Value) * 285.0f, 100.0f, 85.0f);
+		return FVector4f(HSV.HSVToLinearRGB());
+	}
+}
+
+
+void USkinWeightsPaintTool::UpdateBoneVisualization()
+{
+	if (!SkinWeightsMap.Contains(CurrentBone))
+		return;
+
+	TArray<float>& SkinWeightsData = *SkinWeightsMap.Find(CurrentBone);
+
+	// update mesh with new value colors
+	PreviewMesh->EditMesh([&](FDynamicMesh3& Mesh)
+	{
+		UE::Geometry::FDynamicMeshColorOverlay* ColorOverlay = Mesh.Attributes()->PrimaryColors();
+		for (int32 ElementId : ColorOverlay->ElementIndicesItr())
+		{
+			const int32 VertexId = ColorOverlay->GetParentVertex(ElementId);
+			const float Value = SkinWeightsData[VertexId];
+			const FVector4f Color(WeightToColor(Value));
+			ColorOverlay->SetElement(ElementId, Color);
+		}
+	});
+}
+
+
+double USkinWeightsPaintTool::CalculateBrushFalloff(double Distance) const
+{
+	double f = FMathd::Clamp(1.0 - BrushProperties->BrushFalloffAmount, 0.0, 1.0);
+	double d = Distance / CurrentBrushRadius;
+	double w = 1;
+	if (d > f)
+	{
+		d = FMathd::Clamp((d - f) / (1.0 - f), 0.0, 1.0);
+		w = (1.0 - d * d);
+		w = w * w * w;
+	}
+	return w;
+}
+
+
+
+void USkinWeightsPaintTool::ApplyStamp(const FBrushStampData& Stamp)
+{
+	using namespace UE::Geometry;
+	
+	// FIXME: Move to earlier.
+	if (!SkinWeightsMap.Contains(CurrentBone))
+		return;
+
+	IPrimitiveComponentBackedTarget* TargetComponent = Cast<IPrimitiveComponentBackedTarget>(Target);
+	FTransform3d Transform(TargetComponent->GetWorldTransform());
+	FVector3d StampPosLocal = Transform.InverseTransformPosition((FVector3d)Stamp.WorldPosition);
+
+	TArray<int32> ROIVertices;
+	CalculateVertexROI(Stamp, ROIVertices);
+	const int32 NumROIVertices = ROIVertices.Num();
+
+	TArray<float> ROIBefore, ROIAfter;
+	ROIBefore.SetNum(NumROIVertices);
+	ROIAfter.SetNum(NumROIVertices);
+
+	TArray<float>& SkinWeightsData = *SkinWeightsMap.Find(CurrentBone);
+
+	const FDynamicMesh3* CurrentMesh = PreviewMesh->GetMesh();
+	if (bInSmoothStroke)
+	{
+		const float SmoothSpeed = 0.25f;
+
+		for (int32 Index = 0; Index < NumROIVertices; ++Index)
+		{
+			const int32 VertexId = ROIVertices[Index];
+			FVector3d Position = CurrentMesh->GetVertex(VertexId);
+			float ValueSum = 0, WeightSum = 0;
+			for (int32 NeighborVertexId : CurrentMesh->VtxVerticesItr(VertexId))
+			{
+				FVector3d NbrPos = CurrentMesh->GetVertex(NeighborVertexId);
+				const float Weight = FMathf::Clamp(1.0f / FVector3d::DistSquared(NbrPos, Position), 0.0001f, 1000.0f);
+				ValueSum += Weight * SkinWeightsData[NeighborVertexId];
+				WeightSum += Weight;
+			}
+			ValueSum /= WeightSum;
+
+			const float Falloff = float(CalculateBrushFalloff(FVector3d::Dist(Position, StampPosLocal)));
+			const float NewValue = FMathf::Lerp(SkinWeightsData[VertexId], ValueSum, SmoothSpeed*Falloff);
+
+			ROIBefore[Index] = SkinWeightsData[VertexId];
+			ROIAfter[Index] = FMath::Clamp(NewValue, 0.0f, 1.0f);
+		}
+	}
+	else
+	{
+		const bool bInvert = bInRemoveStroke;
+		const float Sign = (bInvert) ? -1.0f : 1.0f;
+		const float UseStrength = Sign * BrushProperties->BrushStrength;
+		for (int32 Index = 0; Index < NumROIVertices; ++Index)
+		{
+			const int32 VertexId = ROIVertices[Index];
+			const FVector3d Position = CurrentMesh->GetVertex(VertexId);
+			const float Falloff = (float)CalculateBrushFalloff(FVector3d::Dist(Position, StampPosLocal));
+			ROIBefore[Index] = SkinWeightsData[VertexId];
+			ROIAfter[Index] = FMath::Clamp(ROIBefore[Index] + UseStrength * Falloff, 0.0f, 1.0f);
+		}
+	}
+
+	// track changes
+	if (ActiveChange)
+	{
+		ActiveChange->UpdateValues(ROIVertices, ROIBefore, ROIAfter);
+	}
+
+	// update values and colors
+	PreviewMesh->DeferredEditMesh([&](FDynamicMesh3& Mesh)
+	{
+		TArray<int> ElementIds;
+		FDynamicMeshColorOverlay* ColorOverlay = Mesh.Attributes()->PrimaryColors();
+		for (int32 Index = 0; Index < NumROIVertices; ++Index)
+		{
+			const int32 VertexId = ROIVertices[Index];
+			SkinWeightsData[VertexId] = ROIAfter[Index];
+			FVector4f NewColor(WeightToColor(ROIAfter[Index]));
+			ColorOverlay->GetVertexElements(VertexId, ElementIds);
+			for (int ElementId : ElementIds)
+			{
+				ColorOverlay->SetElement(ElementId, NewColor);
+			}
+			ElementIds.Reset();
+		}
+	}, false);
+	PreviewMesh->NotifyDeferredEditCompleted(UPreviewMesh::ERenderUpdateMode::FastUpdate, EMeshRenderAttributeFlags::VertexColors, false);
+
+}
+
+
+void USkinWeightsPaintTool::UpdateCurrentBone(const FName& BoneName)
+{
+	CurrentBone = BoneName;
+
+	bVisibleWeightsValid = false;
+}
+
+
+void USkinWeightsPaintTool::OnShutdown(EToolShutdownType ShutdownType)
+{
+	BrushProperties->SaveProperties(this);
+
+	if (ShutdownType == EToolShutdownType::Accept)
+	{
+		// this block bakes the modified DynamicMeshComponent back into the StaticMeshComponent inside an undo transaction
+		GetToolManager()->BeginUndoTransaction(LOCTEXT("SkinWeightsPaintTool", "Paint Skin Weights"));
+
+		Cast<IMeshDescriptionCommitter>(Target)->CommitMeshDescription([this](const IMeshDescriptionCommitter::FCommitterParams& CommitParams)
+			{
+				CommitSkinWeights(CommitParams);
+			});
+		GetToolManager()->EndUndoTransaction();
+	}
+}
+
+
+void USkinWeightsPaintTool::UpdateBonePositionInfos(float MinRadius)
+{
+	const FReferenceSkeleton& RefSkeleton = BoneContainer.GetReferenceSkeleton();
+	const TArray<FMeshBoneInfo>& BoneInfos = RefSkeleton.GetRefBoneInfo();
+	const TArray<FTransform>& BonePoses = RefSkeleton.GetRefBonePose();
+
+	BonePositionInfos.Reset();
+
+	// Exclude virtual bones.
+	for (int BoneIndex = 0; BoneIndex < RefSkeleton.GetRawBoneNum(); BoneIndex++)
+	{
+		FTransform Xform = BonePoses[BoneIndex];
+		int32 ParentBoneIndex = BoneInfos[BoneIndex].ParentIndex;
+
+		while (ParentBoneIndex != INDEX_NONE)
+		{
+			Xform = Xform * BonePoses[ParentBoneIndex];
+			ParentBoneIndex = BoneInfos[ParentBoneIndex].ParentIndex;
+		}
+
+		BonePositionInfos.Add({ BoneInfos[BoneIndex].Name, BoneInfos[BoneIndex].ParentIndex, Xform.GetLocation(), -1.0f });
+	}
+
+	// Populate the children.
+	for (int BoneIndex = 0; BoneIndex < BonePositionInfos.Num(); BoneIndex++)
+	{
+		FBonePositionInfo& BoneInfo = BonePositionInfos[BoneIndex];
+		if (BoneInfo.ParentBoneIndex != INDEX_NONE)
+		{
+			BonePositionInfos[BoneInfo.ParentBoneIndex].ChildBones.Add(BoneInfo.BoneName, BoneIndex);
+		}
+	}
+
+	bool bComputedRadius = true;
+	while (bComputedRadius)
+	{
+		bComputedRadius = false;
+
+		for (int BoneIndex = 0; BoneIndex < BonePositionInfos.Num(); BoneIndex++)
+		{
+			FBonePositionInfo& BoneInfo = BonePositionInfos[BoneIndex];
+			if (BoneInfo.Radius > 0.0f)
+			{
+				continue;
+			}
+
+			if (BoneInfo.ParentBoneIndex == INDEX_NONE)
+			{
+				if (BoneInfo.ChildBones.Num())
+				{
+					int32 Count = 0;
+					float RadiusSum = 0.0f;
+					for (const auto& CB : BoneInfo.ChildBones)
+					{
+						const FBonePositionInfo& ChildBoneInfo = BonePositionInfos[CB.Value];
+						if (ChildBoneInfo.Radius > 0.0f)
+						{
+							RadiusSum += ChildBoneInfo.Radius;
+							Count++;
+						}
+					}
+					if (BoneInfo.ChildBones.Num() == Count)
+					{
+						BoneInfo.Radius = RadiusSum / float(Count);
+						bComputedRadius = true;
+					}
+				}
+				else
+				{
+					// No children either? Take the whole mesh.
+					BoneInfo.Radius = EditedMesh->GetBounds().SphereRadius;
+				}
+			}
+			else 
+			{
+				BoneInfo.Radius = FVector::Dist(BoneInfo.Position, BonePositionInfos[BoneInfo.ParentBoneIndex].Position) / 2.0f;
+				bComputedRadius = true;
+			}
+
+			if (bComputedRadius)
+			{
+				BoneInfo.Radius = FMath::Max(BoneInfo.Radius, MinRadius);
+			}
+		}
+	}
+}
+
+
+void USkinWeightsPaintTool::RenderBonePositions(FPrimitiveDrawInterface* PDI)
+{
+	static const int32 NumSphereSides = 10;
+	static const int32 NumConeSides = 4;
+
+	IPrimitiveComponentBackedTarget* TargetComponent = Cast<IPrimitiveComponentBackedTarget>(Target);
+	FTransform WorldTransform = TargetComponent->GetWorldTransform();
+
+	for (const FBonePositionInfo& BoneInfo : BonePositionInfos)
+	{
+		FLinearColor BoneColor;
+		FVector Start, End;
+
+		End = BoneInfo.Position;
+		End = WorldTransform.TransformPosition(End);
+
+		if (BoneInfo.ParentBoneIndex != INDEX_NONE)
+		{
+			Start = BonePositionInfos[BoneInfo.ParentBoneIndex].Position;
+			Start = WorldTransform.TransformPosition(Start);
+			BoneColor = FLinearColor::White;
+		}
+		else
+		{
+			// Root bone.
+			BoneColor = FLinearColor::Red;
+		}
+		BoneColor.A = 0.10f;
+
+		if (BoneInfo.BoneName == CurrentBone)
+		{
+			BoneColor = FLinearColor(1.0f, 0.34f, 0.0f, 0.75f);
+		}
+
+		const float BoneLength = (End - Start).Size();
+		// clamp by bound, we don't want too long or big
+		const float Radius = FMath::Clamp(BoneLength * 0.05f, 0.1f, MaxDrawRadius);
+
+		// Render Sphere for bone end point and a cone between it and its parent.
+		DrawWireSphere(PDI, End, BoneColor, Radius, NumSphereSides, SDPG_Foreground, 0.0f, 1.0f);
+
+		if (BoneInfo.ParentBoneIndex != INDEX_NONE)
+		{
+			// Calc cone size 
+			const FVector EndToStart = (Start - End);
+			const float ConeLength = EndToStart.Size();
+			const float Angle = FMath::RadiansToDegrees(FMath::Atan(Radius / ConeLength));
+
+			TArray<FVector> Verts;
+			DrawWireCone(PDI, Verts, FRotationMatrix::MakeFromX(EndToStart) * FTranslationMatrix(End), ConeLength, Angle, NumConeSides, BoneColor, SDPG_Foreground, 0.0f, 1.0f);
+		}
+
+		// SkeletalDebugRendering::DrawWireBone(PDI, Start, End, BoneColor, SDPG_Foreground, Radius);
+	}
+}
+
+
+void USkinWeightsPaintTool::BeginChange()
+{
+	ActiveChange = MakeUnique<FMeshSkinWeightsChange>(CurrentBone);
+}
+
+
+TUniquePtr<FMeshSkinWeightsChange> USkinWeightsPaintTool::EndChange()
+{
+	return MoveTemp(ActiveChange);
+}
+
+
+void USkinWeightsPaintTool::ExternalUpdateValues(const FName& BoneName, const TMap<int32, float>& NewValues)
+{
+	TArray<float>* SkinWeightValues = SkinWeightsMap.Find(BoneName);
+	if (SkinWeightValues == nullptr)
+	{
+		return;
+	}
+
+	for (const auto& IV : NewValues)
+	{
+		SkinWeightValues->GetData()[IV.Key] = IV.Value;
+	}
+
+	if (BoneName == CurrentBone)
+		UpdateBoneVisualization();
+}
+
+
+void USkinWeightsPaintTool::InitializeSkinWeights()
+{
+	const FReferenceSkeleton& RefSkeleton = BoneContainer.GetReferenceSkeleton();
+
+	const FSkeletalMeshConstAttributes MeshAttribs(*EditedMesh);
+	const FSkinWeightsVertexAttributesConstRef VertexSkinWeights = MeshAttribs.GetVertexSkinWeights();
+	const int32 NumVertices = EditedMesh->Vertices().Num();
+
+	// Create a map of all bones to their per-vertex weights.
+	SkinWeightsMap.Reset();
+	for (const FMeshBoneInfo& BoneInfo : RefSkeleton.GetRefBoneInfo())
+	{
+		SkinWeightsMap.Add(BoneInfo.Name, {}).AddZeroed(NumVertices);
+	}
+
+	for (int32 VertexIndex = 0; VertexIndex < NumVertices; VertexIndex++)
+	{
+		const FVertexID VertexID(VertexIndex);
+
+		for (UE::AnimationCore::FBoneWeight BoneWeight: VertexSkinWeights.Get(VertexID))
+		{
+			FName BoneName = RefSkeleton.GetBoneName(static_cast<int32>(BoneWeight.GetBoneIndex()));
+			const float Weight = BoneWeight.GetWeight();
+
+			if (Weight >= MinimumWeightThreshold)
+			{
+				// If the source mesh has a bone that we don't recognize, we ignore it. It's
+				// weight will get cleared when the new weights are updated back to the 
+				// source mesh.
+				TArray<float>* PerVertexWeights = SkinWeightsMap.Find(BoneName);
+
+				if (PerVertexWeights)
+				{
+					PerVertexWeights->GetData()[VertexIndex] = Weight;
+				}
+			}
+		}
+	}
+
+	// Pick a root bone.
+	CurrentBone = RefSkeleton.GetBoneName(0);
+	PendingCurrentBone.Reset();
+}
+
+
+void USkinWeightsPaintTool::CommitSkinWeights(const IMeshDescriptionCommitter::FCommitterParams& CommitParams)
+{
+	using namespace UE::AnimationCore;
+	
+	FSkeletalMeshAttributes MeshAttribs(*EditedMesh);
+	FSkinWeightsVertexAttributesRef VertexSkinWeights = MeshAttribs.GetVertexSkinWeights();
+	
+	const FReferenceSkeleton& RefSkeleton = BoneContainer.GetReferenceSkeleton();
+
+	TMap<FBoneIndexType, const TArray<float>*> BoneIndexWeightMap;
+	for (const auto& BoneNameAndWeights : SkinWeightsMap)
+	{
+		const int32 BoneIndex = RefSkeleton.FindBoneIndex(BoneNameAndWeights.Key);
+		if (BoneIndex != INDEX_NONE)
+		{
+			BoneIndexWeightMap.Add(static_cast<FBoneIndexType>(BoneIndex), &BoneNameAndWeights.Value);
+		}
+	}
+
+	FBoneWeightsSettings Settings;
+	Settings.SetNormalizeType(EBoneWeightNormalizeType::AboveOne);
+
+	TArray<FBoneWeight> SourceBoneWeights;
+	SourceBoneWeights.Reserve(MaxInlineBoneWeightCount);
+
+	const int32 NumVertices = EditedMesh->Vertices().Num();
+	for (int32 VertexIndex = 0; VertexIndex < NumVertices; VertexIndex++)
+	{
+		SourceBoneWeights.Reset();
+
+		for (const auto& BoneIndexWeight : BoneIndexWeightMap)
+		{
+			SourceBoneWeights.Add(FBoneWeight(BoneIndexWeight.Key, (*BoneIndexWeight.Value)[VertexIndex]));
+		}
+
+		VertexSkinWeights.Set(FVertexID(VertexIndex), FBoneWeights::Create(SourceBoneWeights, Settings));
+	}
+
+	*CommitParams.MeshDescriptionOut = *EditedMesh;
+}
+
+
+USkeleton* USkinWeightsPaintToolProperties::GetSkeleton(bool& bInvalidSkeletonIsError)
+{
+	bInvalidSkeletonIsError = false;
+	return SkeletalMesh ? SkeletalMesh->GetSkeleton() : nullptr;
+}
+
+
+#undef LOCTEXT_NAMESPACE
diff --git a/Engine/Plugins/Experimental/MeshModelingToolset/Source/MeshModelingTools/Public/SkinWeightsPaintTool.h b/Engine/Plugins/Experimental/MeshModelingToolset/Source/MeshModelingTools/Public/SkinWeightsPaintTool.h
new file mode 100644
--- /dev/null
+++ b/Engine/Plugins/Experimental/MeshModelingToolset/Source/MeshModelingTools/Public/SkinWeightsPaintTool.h
@@ -0,0 +1,172 @@
+// Copyright Epic Games, Inc. All Rights Reserved.
+
+#pragma once
+
+#include "DynamicMeshBrushTool.h"
+#include "Changes/ValueWatcher.h"
+#include "MeshDescription.h"
+#include "DynamicVerticesOctree3.h"
+#include "BoneContainer.h"
+#include "Interfaces/Interface_BoneReferenceSkeletonProvider.h"
+#include "Misc/Optional.h"
+#include "Containers/Map.h"
+#include "TargetInterfaces/MeshDescriptionCommitter.h"
+
+#include "SkinWeightsPaintTool.generated.h"
+
+
+struct FMeshDescription;
+class USkinWeightsPaintTool;
+class IMeshDescriptionCommitter;
+
+
+
+class MESHMODELINGTOOLS_API FMeshSkinWeightsChange : public FToolCommandChange
+{
+public:
+	FMeshSkinWeightsChange(const FName& InBoneName) :
+		FToolCommandChange(),
+		BoneName(InBoneName)
+	{ }
+
+	virtual FString ToString() const override
+	{
+		return FString(TEXT("Paint Skin Weights"));
+	}
+
+	void Apply(UObject* Object) override;
+
+	void Revert(UObject* Object) override;
+
+	void UpdateValues(const TArray<int32>& Indices, const TArray<float>& OldValues, const TArray<float>& NewValues);
+
+private:
+	FName BoneName;
+	TMap<int32, float> OldWeights;
+	TMap<int32, float> NewWeights;
+};
+
+
+/**
+ *
+ */
+UCLASS()
+class MESHMODELINGTOOLS_API USkinWeightsPaintToolBuilder : public UMeshSurfacePointToolBuilder
+{
+	GENERATED_BODY()
+
+public:
+	UMeshSurfacePointTool* CreateNewTool(const FToolBuilderState& SceneState) const override;
+};
+
+
+UCLASS()
+class MESHMODELINGTOOLS_API USkinWeightsPaintToolProperties : 
+	public UInteractiveToolPropertySet,
+	public IBoneReferenceSkeletonProvider
+{
+	GENERATED_BODY()
+public:
+	UPROPERTY(VisibleAnywhere, Category = Skeleton)
+	FBoneReference CurrentBone;
+
+	// IBoneReferenceSkeletonProvider
+	USkeleton* GetSkeleton(bool& bInvalidSkeletonIsError) override;
+
+	TObjectPtr<USkeletalMesh> SkeletalMesh;
+};
+
+
+/**
+ *
+ */
+UCLASS()
+class MESHMODELINGTOOLS_API USkinWeightsPaintTool : public UDynamicMeshBrushTool
+{
+	GENERATED_BODY()
+
+public:
+	void RegisterActions(FInteractiveToolActionSet& ActionSet) override;
+
+	void Setup() override;
+	void Render(IToolsContextRenderAPI* RenderAPI) override;
+
+	bool HasCancel() const override { return true; }
+	bool HasAccept() const override { return true; }
+	bool CanAccept() const override { return true; }
+
+	// UBaseBrushTool overrides
+	bool HitTest(const FRay& Ray, FHitResult& OutHit) override;
+	void OnBeginDrag(const FRay& Ray) override;
+	void OnUpdateDrag(const FRay& Ray) override;
+	void OnEndDrag(const FRay& Ray) override;
+	bool OnUpdateHover(const FInputDeviceRay& DevicePos) override;
+
+protected:
+
+	virtual void ApplyStamp(const FBrushStampData& Stamp);
+
+	void OnShutdown(EToolShutdownType ShutdownType) override;
+	void OnTick(float DeltaTime) override;
+
+	double CalculateBrushFalloff(double Distance) const;
+	void CalculateVertexROI(const FBrushStampData& Stamp, TArray<int>& VertexROI);
+
+	UPROPERTY()
+	USkinWeightsPaintToolProperties* ToolProps;
+
+	TValueWatcher<FBoneReference> CurrentBoneWatcher;
+
+	bool bInRemoveStroke = false;
+	bool bInSmoothStroke = false;
+	FBrushStampData StartStamp;
+	FBrushStampData LastStamp;
+	bool bStampPending;
+
+	TUniquePtr<FMeshDescription> EditedMesh;
+
+	FBoneContainer BoneContainer;
+
+	UE::Geometry::TDynamicVerticesOctree3<FDynamicMesh3> VerticesOctree;
+	TArray<int> PreviewBrushROI;
+
+	using BoneInfluenceMapType = TMap<FName, TArray<float>>;
+	BoneInfluenceMapType SkinWeightsMap;
+	FName CurrentBone = NAME_None;
+	TOptional<FName> PendingCurrentBone;
+
+
+	bool bVisibleWeightsValid = false;
+
+	static FVector4f WeightToColor(float Value);
+	
+	void InitializeSkinWeights();
+	void UpdateBoneVisualization();
+	void UpdateCurrentBone(const FName &BoneName);
+
+	void ExternalUpdateAttributes(TUniquePtr<TAttributesSet<FVertexID>>& NewAttrib, int32 NewSelecedIndex);
+
+	struct FBonePositionInfo
+	{
+		FName BoneName;
+		int32 ParentBoneIndex;
+		FVector Position;
+		float Radius;
+		TMap<FName, int32> ChildBones;
+	};
+	TArray<FBonePositionInfo> BonePositionInfos;
+	float MaxDrawRadius;
+
+	void UpdateBonePositionInfos(float MinRadius);
+	void RenderBonePositions(FPrimitiveDrawInterface *PDI);
+
+	TUniquePtr<FMeshSkinWeightsChange> ActiveChange;
+
+	void BeginChange();
+	TUniquePtr<FMeshSkinWeightsChange> EndChange();
+
+	friend class FMeshSkinWeightsChange;
+	void ExternalUpdateValues(const FName &BoneName, const TMap<int32, float>& IndexValues);
+
+	void CommitSkinWeights(const IMeshDescriptionCommitter::FCommitterParams& CommitParams);
+};
diff --git a/Engine/Plugins/Experimental/ModelingToolsEditorMode/Source/ModelingToolsEditorMode/Private/ModelingToolsManagerActions.cpp b/Engine/Plugins/Experimental/ModelingToolsEditorMode/Source/ModelingToolsEditorMode/Private/ModelingToolsManagerActions.cpp
--- a/Engine/Plugins/Experimental/ModelingToolsEditorMode/Source/ModelingToolsEditorMode/Private/ModelingToolsManagerActions.cpp
+++ b/Engine/Plugins/Experimental/ModelingToolsEditorMode/Source/ModelingToolsEditorMode/Private/ModelingToolsManagerActions.cpp
@@ -101,6 +101,8 @@ void FModelingToolsManagerCommands::RegisterCommands()
 	UI_COMMAND(BeginBakeMeshAttributeMapsTool, "BakeTx", "Bake Image Maps for Target Mesh (optionally from second Source Mesh)", EUserInterfaceActionType::ToggleButton, FInputChord());
 	UI_COMMAND(BeginUVSeamEditTool, "SeamEd", "Add UV Seams to Mesh", EUserInterfaceActionType::ToggleButton, FInputChord());
 
+	UI_COMMAND(BeginSkinWeightsPaintTool, "SkinWts", "Start the Paint Skin Weights Tool", EUserInterfaceActionType::Button, FInputChord());
+
 	UI_COMMAND(BeginGroomToMeshTool, "Helmet", "Generate Helmet Mesh for Selected Groom", EUserInterfaceActionType::ToggleButton, FInputChord());
 	UI_COMMAND(BeginGenerateLODMeshesTool, "HlmLOD", "Generate LODS for Hair Helmet", EUserInterfaceActionType::ToggleButton, FInputChord());
 	UI_COMMAND(BeginGroomCardsToMeshTool, "CardsToMesh", "Hair Cards to Mesh Tool", EUserInterfaceActionType::ToggleButton, FInputChord());
diff --git a/Engine/Plugins/Experimental/ModelingToolsEditorMode/Source/ModelingToolsEditorMode/Public/ModelingToolsManagerActions.h b/Engine/Plugins/Experimental/ModelingToolsEditorMode/Source/ModelingToolsEditorMode/Public/ModelingToolsManagerActions.h
--- a/Engine/Plugins/Experimental/ModelingToolsEditorMode/Source/ModelingToolsEditorMode/Public/ModelingToolsManagerActions.h
+++ b/Engine/Plugins/Experimental/ModelingToolsEditorMode/Source/ModelingToolsEditorMode/Public/ModelingToolsManagerActions.h
@@ -92,6 +92,7 @@ public:
 	TSharedPtr<FUICommandInfo> BeginEditMeshMaterialsTool;
 	TSharedPtr<FUICommandInfo> BeginTransformUVIslandsTool;
 	TSharedPtr<FUICommandInfo> BeginMeshAttributePaintTool;
+	TSharedPtr<FUICommandInfo> BeginSkinWeightsPaintTool;
 	TSharedPtr<FUICommandInfo> BeginAttributeEditorTool;
 	TSharedPtr<FUICommandInfo> BeginBakeMeshAttributeMapsTool;
 	TSharedPtr<FUICommandInfo> BeginUVSeamEditTool;
