Do fallback on model data also if the bulk data is empty, not just if it's missing.
Always clear the cached mesh descriptor for now. If we do undo after the commit, the cached descriptor is out-of-date w.r.t the actual skelmesh. Deal with it better once skelmeshes actually store their data as mesh descripors.

#trivial
#rnx

diff --git a/Engine/Plugins/Experimental/MeshModelingToolset/Source/MeshModelingToolsEditorOnly/Private/ToolTargets/SkeletalMeshComponentToolTarget.cpp b/Engine/Plugins/Experimental/MeshModelingToolset/Source/MeshModelingToolsEditorOnly/Private/ToolTargets/SkeletalMeshComponentToolTarget.cpp
--- a/Engine/Plugins/Experimental/MeshModelingToolset/Source/MeshModelingToolsEditorOnly/Private/ToolTargets/SkeletalMeshComponentToolTarget.cpp
+++ b/Engine/Plugins/Experimental/MeshModelingToolset/Source/MeshModelingToolsEditorOnly/Private/ToolTargets/SkeletalMeshComponentToolTarget.cpp
@@ -131,8 +131,8 @@ FMeshDescription* USkeletalMeshComponentToolTarget::GetMeshDescription()
 		CachedMeshDescription = MakeUnique<FMeshDescription>();
 		const USkeletalMesh* SkeletalMesh = Cast<USkinnedMeshComponent>(Component)->SkeletalMesh;
 
-		// Check first if we have bulk data available.
-		if (SkeletalMesh->IsLODImportedDataBuildAvailable(LODIndex))
+		// Check first if we have bulk data available and non-empty.
+		if (SkeletalMesh->IsLODImportedDataBuildAvailable(LODIndex) && !SkeletalMesh->IsLODImportedDataEmpty(LODIndex))
 		{
 			FSkeletalMeshImportData SkeletalMeshImportData;
 			SkeletalMesh->LoadLODImportedData(LODIndex, SkeletalMeshImportData);
@@ -201,6 +201,8 @@ void USkeletalMeshComponentToolTarget::CommitMeshDescription(const FCommitter& C
 
 	// this rebuilds physics, but it doesn't undo!
 	Component->RecreatePhysicsState();
+	
+	CachedMeshDescription.Reset();
 }
 
 TSharedPtr<FDynamicMesh3> USkeletalMeshComponentToolTarget::GetDynamicMesh()
