Fix static analysis warnings in the skin weights paint tool.

#trivial
#rnx

diff --git a/Engine/Plugins/Experimental/MeshModelingToolset/Source/MeshModelingTools/Private/SkinWeightsPaintTool.cpp b/Engine/Plugins/Experimental/MeshModelingToolset/Source/MeshModelingTools/Private/SkinWeightsPaintTool.cpp
--- a/Engine/Plugins/Experimental/MeshModelingToolset/Source/MeshModelingTools/Private/SkinWeightsPaintTool.cpp
+++ b/Engine/Plugins/Experimental/MeshModelingToolset/Source/MeshModelingTools/Private/SkinWeightsPaintTool.cpp
@@ -89,12 +89,12 @@ void USkinWeightsPaintTool::Setup()
 		// Pick the first root bone as the initial selection.
 		PendingCurrentBone = Component->SkeletalMesh->GetRefSkeleton().GetBoneName(0);
 		ToolProps->CurrentBone.BoneName = PendingCurrentBone.GetValue();
+
+		// Update the skeleton drawing information from the original bind pose
+		MaxDrawRadius = Component->Bounds.SphereRadius * 0.0025f;
 	}
 	AddToolPropertySource(ToolProps);
 
-	// Update the skeleton drawing information from the original bind pose
-	MaxDrawRadius = Component->Bounds.SphereRadius * 0.01f;
-
 	UpdateBonePositionInfos(MaxDrawRadius);
 
 	// configure preview mesh
diff --git a/Engine/Plugins/Experimental/MeshModelingToolset/Source/MeshModelingTools/Public/SkinWeightsPaintTool.h b/Engine/Plugins/Experimental/MeshModelingToolset/Source/MeshModelingTools/Public/SkinWeightsPaintTool.h
--- a/Engine/Plugins/Experimental/MeshModelingToolset/Source/MeshModelingTools/Public/SkinWeightsPaintTool.h
+++ b/Engine/Plugins/Experimental/MeshModelingToolset/Source/MeshModelingTools/Public/SkinWeightsPaintTool.h
@@ -144,8 +144,6 @@ protected:
 	void UpdateBoneVisualization();
 	void UpdateCurrentBone(const FName &BoneName);
 
-	void ExternalUpdateAttributes(TUniquePtr<TAttributesSet<FVertexID>>& NewAttrib, int32 NewSelecedIndex);
-
 	struct FBonePositionInfo
 	{
 		FName BoneName;
@@ -155,7 +153,7 @@ protected:
 		TMap<FName, int32> ChildBones;
 	};
 	TArray<FBonePositionInfo> BonePositionInfos;
-	float MaxDrawRadius;
+	float MaxDrawRadius = 0.0f;
 
 	void UpdateBonePositionInfos(float MinRadius);
 	void RenderBonePositions(FPrimitiveDrawInterface *PDI);
