ModelingComponents:
Add Texture2DUtil::ConvertToSingleChannel(), converts TSF_BGRA8 and TSF_BGRE8 textures to TSF_G8 (based on cast-to-FColor Red channel) to reduce memory usage of a texture
Add CreateDerivedMaterialInstance(), makes a new MIC for input material (either by duplicating or using MIC asset factory)
#rnx

diff --git a/Engine/Plugins/Experimental/MeshModelingToolset/Source/ModelingComponents/Private/AssetUtils/Texture2DUtil.cpp b/Engine/Plugins/Experimental/MeshModelingToolset/Source/ModelingComponents/Private/AssetUtils/Texture2DUtil.cpp
--- a/Engine/Plugins/Experimental/MeshModelingToolset/Source/ModelingComponents/Private/AssetUtils/Texture2DUtil.cpp
+++ b/Engine/Plugins/Experimental/MeshModelingToolset/Source/ModelingComponents/Private/AssetUtils/Texture2DUtil.cpp
@@ -146,4 +146,62 @@ bool UE::AssetUtils::ReadTexture(
 #endif
 
 	return ReadTexture_PlatformData(TextureMap, Dimensions, DestImage);
-}
\ No newline at end of file
+}
+
+
+
+bool UE::AssetUtils::ConvertToSingleChannel(UTexture2D* TextureMap)
+{
+	if (ensure(TextureMap) == false) return false;
+
+#if WITH_EDITOR
+	bool bHasVTData = TextureMap->GetPlatformData()->VTData != nullptr;
+	bool bHasMips = TextureMap->GetPlatformData()->Mips.Num() != 0;
+	ensure(bHasVTData != bHasMips);		// should be one or the other
+
+	if (ensure(TextureMap->Source.IsValid()) && bHasVTData == false)
+	{
+		FTextureSource& TextureSource = TextureMap->Source;
+		int32 Width = TextureSource.GetSizeX();
+		int32 Height = TextureSource.GetSizeY();
+		int64 Num = Width * Height;
+		ETextureSourceFormat SourceFormat = TextureSource.GetFormat();
+		if (SourceFormat == TSF_G8)
+		{
+			return true;		// already single channel
+		}
+		if (SourceFormat != TSF_BGRA8 && SourceFormat != TSF_BGRE8)
+		{
+			ensureMsgf(false, TEXT("ConvertToSingleChannel currently only supports RGBA8 textures"));
+			return false;
+		}
+
+		TArray64<uint8> NewSourceData;
+		NewSourceData.SetNum(Width * Height);
+
+		TArray64<uint8> SourceData;
+		TextureMap->Source.GetMipData(SourceData, 0, 0, 0);
+		int32 BytesPerPixel = TextureSource.GetBytesPerPixel();
+		check(BytesPerPixel == sizeof(FColor));
+		const uint8* SourceDataPtr = SourceData.GetData();
+		for (int32 i = 0; i < Num; ++i)
+		{
+			const uint8* PixelPtr = SourceDataPtr + (i * BytesPerPixel);
+			FColor PixelColor = *((FColor*)PixelPtr);
+			NewSourceData[i] = PixelColor.R;
+		}
+
+		TextureSource.Init(Width, Height, 1, 1, TSF_G8, &NewSourceData[0]);
+
+		TextureMap->UpdateResource();
+
+		return true;
+	}
+	return false;
+#else
+
+	ensureMsgf(false, TEXT("ConvertToSingleChannel currently requires editor-only SourceData"));
+	return false;
+#endif
+
+}
diff --git a/Engine/Plugins/Experimental/MeshModelingToolset/Source/ModelingComponents/Public/AssetUtils/Texture2DUtil.h b/Engine/Plugins/Experimental/MeshModelingToolset/Source/ModelingComponents/Public/AssetUtils/Texture2DUtil.h
--- a/Engine/Plugins/Experimental/MeshModelingToolset/Source/ModelingComponents/Public/AssetUtils/Texture2DUtil.h
+++ b/Engine/Plugins/Experimental/MeshModelingToolset/Source/ModelingComponents/Public/AssetUtils/Texture2DUtil.h
@@ -22,6 +22,8 @@ namespace AssetUtils
 		TImageBuilder<FVector4f>& DestImageOut,
 		bool bPreferPlatformData = false);
 
+	MODELINGCOMPONENTS_API bool ConvertToSingleChannel(UTexture2D* TextureMap);
+
 }
 }
 
diff --git a/Engine/Plugins/Experimental/MeshModelingToolset/Source/ModelingComponentsEditorOnly/ModelingComponentsEditorOnly.Build.cs b/Engine/Plugins/Experimental/MeshModelingToolset/Source/ModelingComponentsEditorOnly/ModelingComponentsEditorOnly.Build.cs
--- a/Engine/Plugins/Experimental/MeshModelingToolset/Source/ModelingComponentsEditorOnly/ModelingComponentsEditorOnly.Build.cs
+++ b/Engine/Plugins/Experimental/MeshModelingToolset/Source/ModelingComponentsEditorOnly/ModelingComponentsEditorOnly.Build.cs
@@ -49,7 +49,8 @@ public class ModelingComponentsEditorOnly : ModuleRules
                 "RenderCore",
                 "RHI",
 
-				"AssetTools"
+				"AssetTools",
+				"UnrealEd"			// required for asset factories
 
 			}
 			);
diff --git a/Engine/Plugins/Experimental/MeshModelingToolset/Source/ModelingComponentsEditorOnly/Private/AssetUtils/CreateMaterialUtil.cpp b/Engine/Plugins/Experimental/MeshModelingToolset/Source/ModelingComponentsEditorOnly/Private/AssetUtils/CreateMaterialUtil.cpp
--- a/Engine/Plugins/Experimental/MeshModelingToolset/Source/ModelingComponentsEditorOnly/Private/AssetUtils/CreateMaterialUtil.cpp
+++ b/Engine/Plugins/Experimental/MeshModelingToolset/Source/ModelingComponentsEditorOnly/Private/AssetUtils/CreateMaterialUtil.cpp
@@ -63,3 +63,75 @@ UE::AssetUtils::ECreateMaterialResult UE::AssetUtils::CreateDuplicateMaterial(
 	return ECreateMaterialResult::Ok;
 }
 
+
+
+
+
+
+
+UE::AssetUtils::ECreateMaterialResult UE::AssetUtils::CreateDerivedMaterialInstance(
+	UMaterialInterface* BaseMaterial,
+	FMaterialAssetOptions& Options,
+	FMaterialAssetResults& ResultsOut)
+{
+	// Duplicate operation will create a new Package, so calling code should not!
+	if (ensure(Options.UsePackage == nullptr) == false)
+	{
+		return ECreateMaterialResult::InvalidPackage;
+	}
+
+	if (ensure(BaseMaterial) == false)
+	{
+		return ECreateMaterialResult::InvalidBaseMaterial;
+	}
+	UMaterial* SourceMaterial = BaseMaterial->GetBaseMaterial();
+	if (ensure(SourceMaterial) == false)
+	{
+		return ECreateMaterialResult::InvalidBaseMaterial;
+	}
+
+	// cannot create new package if path already exists
+	if (FPackageName::DoesPackageExist(Options.NewAssetPath, nullptr, nullptr))
+	{
+		ensure(false);
+		return ECreateMaterialResult::NameError;
+	}
+
+	FString DestinationLongPackagePath = FPackageName::GetLongPackagePath(Options.NewAssetPath);
+	//FString DestinationObjectName = FPackageName::ObjectPathToObjectName(Options.NewAssetPath);
+	FString DestinationObjectName = FPackageName::GetLongPackageAssetName(Options.NewAssetPath);
+
+	IAssetTools& AssetTools = FModuleManager::Get().LoadModuleChecked<FAssetToolsModule>("AssetTools").Get();
+
+	bool bSourceIsMIC = (Cast<UMaterialInstanceConstant>(BaseMaterial) != nullptr);
+
+	UMaterialInstanceConstant* NewMaterialInstance = nullptr;
+	if (bSourceIsMIC)
+	{
+		UObject* DupeAsset = AssetTools.DuplicateAsset(DestinationObjectName, DestinationLongPackagePath, BaseMaterial);
+		NewMaterialInstance = Cast<UMaterialInstanceConstant>(DupeAsset);
+	}
+	else
+	{
+		UMaterialInstanceConstantFactoryNew* Factory = NewObject<UMaterialInstanceConstantFactoryNew>();
+		Factory->InitialParent = SourceMaterial;
+
+		UObject* NewAsset = AssetTools.CreateAsset(DestinationObjectName, DestinationLongPackagePath,
+			UMaterialInstanceConstant::StaticClass(), Factory);
+
+		NewMaterialInstance = Cast<UMaterialInstanceConstant>(NewAsset);
+	}
+
+	if (ensure(NewMaterialInstance) == false)
+	{
+		return ECreateMaterialResult::DuplicateFailed;
+	}
+
+	if (Options.bDeferPostEditChange == false)
+	{
+		NewMaterialInstance->PostEditChange();
+	}
+
+	ResultsOut.NewMaterialInstance = NewMaterialInstance;
+	return ECreateMaterialResult::Ok;
+}
diff --git a/Engine/Plugins/Experimental/MeshModelingToolset/Source/ModelingComponentsEditorOnly/Public/AssetUtils/CreateMaterialUtil.h b/Engine/Plugins/Experimental/MeshModelingToolset/Source/ModelingComponentsEditorOnly/Public/AssetUtils/CreateMaterialUtil.h
--- a/Engine/Plugins/Experimental/MeshModelingToolset/Source/ModelingComponentsEditorOnly/Public/AssetUtils/CreateMaterialUtil.h
+++ b/Engine/Plugins/Experimental/MeshModelingToolset/Source/ModelingComponentsEditorOnly/Public/AssetUtils/CreateMaterialUtil.h
@@ -8,6 +8,7 @@
 #include "Image/ImageBuilder.h"
 
 class UMaterial;
+class UMaterialInstanceConstant;
 class UMaterialInterface;
 
 namespace UE
@@ -47,6 +48,7 @@ namespace AssetUtils
 	struct MODELINGCOMPONENTSEDITORONLY_API FMaterialAssetResults
 	{
 		UMaterial* NewMaterial = nullptr;
+		UMaterialInstanceConstant* NewMaterialInstance = nullptr;
 	};
 
 	/**
@@ -60,6 +62,20 @@ namespace AssetUtils
 		FMaterialAssetOptions& Options,
 		FMaterialAssetResults& ResultsOut);
 
+
+	/**
+	 * Create a new UMaterialInstanceConstant Asset derived from the parent UMaterial of BaseMaterial.
+	 * If the input is a UMaterialInstanceConstant, duplicate it to preserve existing parameters/etc
+	 * @param Options defines configuration for the new UMaterialInstanceConstant Asset
+	 * @param ResultOut new UMateriaInstanceConstant is returned here
+	 * @return Ok, or error flag if some part of the creation process failed. On failure, no Asset is created.
+	 */
+	MODELINGCOMPONENTSEDITORONLY_API ECreateMaterialResult CreateDerivedMaterialInstance(
+		UMaterialInterface* BaseMaterial,
+		FMaterialAssetOptions& Options,
+		FMaterialAssetResults& ResultsOut);
+
+
 }  // end namespace UE
 }  // end namespace AssetUtils
 
