GeometryProcessing:
- Add FImageAdapter, wraps TImageBuilder of different types and provides a consistent interface, so that other code doesn't need specializations for each image type
- FMeshSceneAdapter now supports the ability to detect and thicken "thin" meshes, by flattening the instances into unique mesh copies (generally necessary due to scaling). Can also now return statistics about the mesh scene, return a full flattened mesh (for debugging), and automatically generate a "closing mesh" that can help to cap off the base of mesh assemblies when trying to generate a solid/etc.

ModelingComponents:
- FWorldRenderCapture now writes via FImageAdapter instead of an Image4f
- FSceneCapturePhotoSet now stores less data for photo sets. BaseColor/Emissive/WorldNormal are 3f, Roughness/Specular/Metallic are 1f
- IApproximateActors and implementation now support optional base-capping, occluded mesh removal, geometric-tolerance simplification, thin mesh auth-thickening. Now creates MICs based on an input Material, instead of new Material. Added ability to emit a flattened debug mesh, and print debug info.

MergeActors:
- Expose above improvements to IApproximateActors

#rnx

diff --git a/Engine/Plugins/Experimental/GeometryProcessing/Source/GeometricObjects/Public/Image/ImageBuilder.h b/Engine/Plugins/Experimental/GeometryProcessing/Source/GeometricObjects/Public/Image/ImageBuilder.h
--- a/Engine/Plugins/Experimental/GeometryProcessing/Source/GeometricObjects/Public/Image/ImageBuilder.h
+++ b/Engine/Plugins/Experimental/GeometryProcessing/Source/GeometricObjects/Public/Image/ImageBuilder.h
@@ -241,10 +241,83 @@ public:
 
 		return DownsampleImage;
 	}
-	
+};
+
+
+/**
+ * FImageAdapter is a wrapper around different types of TImageBuilder that provides a
+ * standard interface, which allows functions that don't need to know about the specific
+ * image type to operate on "any" image. For example to allow code that works on a 4-channel
+ * image to operate on a single-channel image (in which case the Adapter will expand/collapse
+ * the channels automatically).
+ */
+class FImageAdapter
+{
+public:
+
+	enum class EImageType
+	{
+		Float1,
+		Float3,
+		Float4
+	};
+
+	FImageAdapter(TImageBuilder<float>* Image)
+	{
+		ImageType = EImageType::Float1;
+		Image1f = Image;
+	}
 
+	FImageAdapter(TImageBuilder<FVector3f>* Image)
+	{
+		ImageType = EImageType::Float3;
+		Image3f = Image;
+	}
+
+	FImageAdapter(TImageBuilder<FVector4f>* Image)
+	{
+		ImageType = EImageType::Float4;
+		Image4f = Image;
+	}
+
+	void SetDimensions(FImageDimensions Dimensions)
+	{
+		switch (ImageType)
+		{
+		case EImageType::Float1:
+			Image1f->SetDimensions(Dimensions); break;
+		case EImageType::Float3:
+			Image3f->SetDimensions(Dimensions); break;
+		case EImageType::Float4:
+			Image4f->SetDimensions(Dimensions); break;
+		}
+	}
+
+	void SetPixel(const FVector2i& PixelCoords, const FVector4f& FloatPixel)
+	{
+		switch (ImageType)
+		{
+		case EImageType::Float1:
+			Image1f->SetPixel(PixelCoords, FloatPixel.X); 
+			break;
+		case EImageType::Float3:
+			Image3f->SetPixel(PixelCoords, FVector3f(FloatPixel.X, FloatPixel.Y, FloatPixel.Z)); 
+			break;
+		case EImageType::Float4:
+			Image4f->SetPixel(PixelCoords, FloatPixel); 
+			break;
+		}
+	}
+
+protected:
+	EImageType ImageType;
+
+	TImageBuilder<float>* Image1f = nullptr;
+	TImageBuilder<FVector3f>* Image3f = nullptr;
+	TImageBuilder<FVector4f>* Image4f = nullptr;
 };
 
 
+
 } // end namespace UE::Geometry
 } // end namespace UE
\ No newline at end of file
diff --git a/Engine/Plugins/Experimental/GeometryProcessing/Source/GeometricObjects/Public/Image/SpatialPhotoSet.h b/Engine/Plugins/Experimental/GeometryProcessing/Source/GeometricObjects/Public/Image/SpatialPhotoSet.h
--- a/Engine/Plugins/Experimental/GeometryProcessing/Source/GeometricObjects/Public/Image/SpatialPhotoSet.h
+++ b/Engine/Plugins/Experimental/GeometryProcessing/Source/GeometricObjects/Public/Image/SpatialPhotoSet.h
@@ -33,6 +33,8 @@ struct TSpatialPhoto
 	TImageBuilder<PixelType> Image;
 };
 typedef TSpatialPhoto<FVector4f> FSpatialPhoto4f;
+typedef TSpatialPhoto<FVector3f> FSpatialPhoto3f;
+typedef TSpatialPhoto<float> FSpatialPhoto1f;
 
 
 
@@ -83,6 +85,8 @@ protected:
 	TArray<TSharedPtr<TSpatialPhoto<PixelType>, ESPMode::ThreadSafe>> Photos;
 };
 typedef TSpatialPhotoSet<FVector4f, float> FSpatialPhotoSet4f;
+typedef TSpatialPhotoSet<FVector3f, float> FSpatialPhotoSet3f;
+typedef TSpatialPhotoSet<float, float> FSpatialPhotoSet1f;
 
 
 
diff --git a/Engine/Plugins/Experimental/GeometryProcessing/Source/GeometricObjects/Public/TransformSequence.h b/Engine/Plugins/Experimental/GeometryProcessing/Source/GeometricObjects/Public/TransformSequence.h
--- a/Engine/Plugins/Experimental/GeometryProcessing/Source/GeometricObjects/Public/TransformSequence.h
+++ b/Engine/Plugins/Experimental/GeometryProcessing/Source/GeometricObjects/Public/TransformSequence.h
@@ -48,10 +48,26 @@ public:
 	 */
 	const TArray<TTransform3<RealType>,TInlineAllocator<2>>& GetTransforms() const { return Transforms; }
 
+
+	/**
+	 * @return true if any transform in the sequence has nonuniform scaling
+	 */
+	bool HasNonUniformScale(RealType Tolerance = TMathUtil<RealType>::ZeroTolerance) const
+	{
+		for (const TTransform3<RealType>& Transform : Transforms)
+		{
+			if (Transform.HasNonUniformScale(Tolerance))
+			{
+				return true;
+			}
+		}
+		return false;
+	}
+
 	/**
 	 * @return point P with transform sequence applied
 	 */
-	FVector3d TransformPosition(FVector3d P) const
+	FVector3<RealType> TransformPosition(UE::Math::TVector<RealType> P) const
 	{
 		for (const TTransform3<RealType>& Transform : Transforms)
 		{
@@ -63,7 +79,7 @@ public:
 	/**
 	 * @return point P with inverse transform sequence applied
 	 */
-	FVector3d InverseTransformPosition(FVector3d P) const
+	FVector3<RealType> InverseTransformPosition(UE::Math::TVector<RealType> P) const
 	{
 		int32 N = Transforms.Num();
 		for (int32 k = N - 1; k >= 0; k--)
@@ -74,20 +90,28 @@ public:
 	}
 
 	/**
-	 * Create an reversed transform sequence such that InverseSeq.Transform(P) is equivalent to Seq.InverseTransform(P).
-	 * This is more efficient if applying inverse many times as it saves lots of checking for invalid scaling.
-	 * @return inverse transform sequence
+	 * @return Vector V with transform sequence applied
 	 */
-	TTransformSequence3<RealType> Inverse() const
+	FVector3<RealType> TransformVector(UE::Math::TVector<RealType> V) const
 	{
-		TTransformSequence3 Inverse;
-		int32 N = Transforms.Num();
-		Inverse.Transforms.Reserve(N);
-		for (int32 k = N - 1; k >= 0; --k)
+		for (const TTransform3<RealType>& Transform : Transforms)
+		{
+			V = Transform.TransformVector(V);
+		}
+		return V;
+	}
+
+
+	/**
+	 * @return Normal with transform sequence applied
+	 */
+	FVector3<RealType> TransformNormal(UE::Math::TVector<RealType> Normal) const
+	{
+		for (const TTransform3<RealType>& Transform : Transforms)
 		{
-			Inverse.Transforms.Add(Transforms[k].Inverse());
+			Normal = Transform.TransformNormal(Normal);
 		}
-		return Inverse;
+		return Normal;
 	}
 
 
diff --git a/Engine/Plugins/Experimental/MeshModelingToolset/Source/GeometryProcessingAdapters/Private/ApproximateActorsImpl.cpp b/Engine/Plugins/Experimental/MeshModelingToolset/Source/GeometryProcessingAdapters/Private/ApproximateActorsImpl.cpp
--- a/Engine/Plugins/Experimental/MeshModelingToolset/Source/GeometryProcessingAdapters/Private/ApproximateActorsImpl.cpp
+++ b/Engine/Plugins/Experimental/MeshModelingToolset/Source/GeometryProcessingAdapters/Private/ApproximateActorsImpl.cpp
@@ -10,14 +10,17 @@
 #include "MeshConstraintsUtil.h"
 #include "Implicit/Solidify.h"
 #include "Implicit/Morphology.h"
-#include "CleaningOps/RemoveOccludedTrianglesOp.h"
+#include "Operations/RemoveOccludedTriangles.h"
 #include "ParameterizationOps/ParameterizeMeshOp.h"
 #include "Parameterization/MeshUVPacking.h"
 #include "MeshQueries.h"
+#include "ProjectionTargets.h"
+#include "Selections/MeshFaceSelection.h"
 
 #include "AssetUtils/CreateStaticMeshUtil.h"
 #include "AssetUtils/CreateTexture2DUtil.h"
 #include "AssetUtils/CreateMaterialUtil.h"
+#include "AssetUtils/Texture2DUtil.h"
 #include "UObject/UObjectGlobals.h"		// for CreatePackage
 
 #include "ImageUtils.h"
@@ -27,7 +30,9 @@
 #include "AssetUtils/Texture2DBuilder.h"
 
 #include "Materials/Material.h"
+#include "Materials/MaterialInstanceConstant.h"
 
+#include "Async/Async.h"
 #include "Misc/ScopedSlowTask.h"
 
 using namespace UE::Geometry;
@@ -46,19 +51,39 @@ struct FGeneratedResultTextures
 };
 
 
+static TUniquePtr<FSceneCapturePhotoSet> CapturePhotoSet(
+	const TArray<AActor*>& Actors,
+	const IGeometryProcessing_ApproximateActors::FOptions& Options
+)
+{
+	TRACE_CPUPROFILER_EVENT_SCOPE(ApproximateActorsImpl_Captures);
+
+	double FieldOfView = Options.FieldOfViewDegrees;
+	double NearPlaneDist = Options.NearPlaneDist;
+
+	FImageDimensions CaptureDimensions(Options.RenderCaptureImageSize, Options.RenderCaptureImageSize);
+
+	TUniquePtr<FSceneCapturePhotoSet> SceneCapture = MakeUnique<FSceneCapturePhotoSet>();
+	SceneCapture->SetCaptureSceneActors(Actors[0]->GetWorld(), Actors);
+
+	SceneCapture->AddStandardExteriorCapturesFromBoundingBox(
+		CaptureDimensions, FieldOfView, NearPlaneDist,
+		true, true, true);
+	
+	return SceneCapture;
+}
+
 static void BakeTexturesFromPhotoCapture(
-	const TArray<AActor*>& Actors, 
+	TUniquePtr<FSceneCapturePhotoSet>& SceneCapture,
 	const IGeometryProcessing_ApproximateActors::FOptions& Options, 
 	FGeneratedResultTextures& GeneratedTextures,
 	const FDynamicMesh3* WorldTargetMesh,
 	const FMeshTangentsd* MeshTangents
 	)
 {
-	int32 UVLayer = 0;
-	double FieldOfView = Options.FieldOfViewDegrees;
-	double NearPlaneDist = Options.NearPlaneDist;
-	double RayOffsetHackDist = (double)(100.0f * FMathf::ZeroTolerance);
+	TRACE_CPUPROFILER_EVENT_SCOPE(ApproximateActorsImpl_Textures);
 
+	int32 UVLayer = 0;
 	int32 Supersample = FMath::Max(1, Options.AntiAliasMultiSampling);
 	if ( (Options.TextureImageSize * Supersample) > 16384)
 	{
@@ -66,15 +91,8 @@ static void BakeTexturesFromPhotoCapture(
 		Supersample = 1;
 	}
 
-	FImageDimensions CaptureDimensions(Options.RenderCaptureImageSize, Options.RenderCaptureImageSize);
 	FImageDimensions OutputDimensions(Options.TextureImageSize*Supersample, Options.TextureImageSize*Supersample);
 
-	FSceneCapturePhotoSet SceneCapture;
-	SceneCapture.SetCaptureSceneActors(Actors[0]->GetWorld(), Actors);
-
-	SceneCapture.AddStandardExteriorCapturesFromBoundingBox(
-		CaptureDimensions, FieldOfView, NearPlaneDist,
-		true, true, true);
 
 	FScopedSlowTask Progress(8.f, LOCTEXT("BakingTextures", "Baking Textures..."));
 	Progress.MakeDialog(true);
@@ -84,16 +102,22 @@ static void BakeTexturesFromPhotoCapture(
 	FDynamicMeshAABBTree3 Spatial(WorldTargetMesh, true);
 
 	FMeshImageBakingCache TempBakeCache;
-	TempBakeCache.SetDetailMesh(WorldTargetMesh, &Spatial);
-	TempBakeCache.SetBakeTargetMesh(WorldTargetMesh);
-	TempBakeCache.SetDimensions(OutputDimensions);
-	TempBakeCache.SetUVLayer(UVLayer);
-	TempBakeCache.SetThickness(0.1);
-	TempBakeCache.SetCorrespondenceStrategy(FMeshImageBakingCache::ECorrespondenceStrategy::Identity);
-	TempBakeCache.ValidateCache();
+	{
+		TRACE_CPUPROFILER_EVENT_SCOPE(ApproximateActorsImpl_Textures_MakeCache);
+		TempBakeCache.SetDetailMesh(WorldTargetMesh, &Spatial);
+		TempBakeCache.SetBakeTargetMesh(WorldTargetMesh);
+		TempBakeCache.SetDimensions(OutputDimensions);
+		TempBakeCache.SetUVLayer(UVLayer);
+		TempBakeCache.SetThickness(0.1);
+		TempBakeCache.SetCorrespondenceStrategy(FMeshImageBakingCache::ECorrespondenceStrategy::Identity);
+		TempBakeCache.ValidateCache();
+	}
 
 	Progress.EnterProgressFrame(1.f, LOCTEXT("BakingBaseColor", "Baking Base Color..."));
 
+	FAxisAlignedBox3d TargetBounds = WorldTargetMesh->GetBounds();
+	double RayOffsetHackDist = (double)(100.0f * FMathf::ZeroTolerance * TargetBounds.MinDim() );
+
 	auto VisibilityFunction = [&Spatial, RayOffsetHackDist](const FVector3d& SurfPos, const FVector3d& ImagePosWorld)
 	{
 		FVector3d RayDir = ImagePosWorld - SurfPos;
@@ -104,37 +128,46 @@ static void BakeTexturesFromPhotoCapture(
 	};
 
 	FSceneCapturePhotoSet::FSceneSample DefaultSample;
-	DefaultSample.BaseColor = FVector4f(0, -1, 0, 0);
+	FVector4f InvalidColor(0, -1, 0, 1);
+	DefaultSample.BaseColor = FVector3f(InvalidColor.X, InvalidColor.Y, InvalidColor.Z);
 
 	FMeshGenericWorldPositionColorBaker BaseColorBaker;
 	BaseColorBaker.SetCache(&TempBakeCache);
 	BaseColorBaker.ColorSampleFunction = [&](FVector3d Position, FVector3d Normal) {
 		FSceneCapturePhotoSet::FSceneSample Sample = DefaultSample;
-		SceneCapture.ComputeSample(FRenderCaptureTypeFlags::BaseColor(),
+		SceneCapture->ComputeSample(FRenderCaptureTypeFlags::BaseColor(),
 			Position, Normal, VisibilityFunction, Sample);
-		return Sample.BaseColor;
+		return Sample.GetValue4f(ERenderCaptureType::BaseColor);
 	};
-	BaseColorBaker.Bake();
+	{
+		TRACE_CPUPROFILER_EVENT_SCOPE(ApproximateActorsImpl_Textures_BakeColor);
+		BaseColorBaker.Bake();
+	}
 
 	// find "hole" pixels
 	TArray<FVector2i> MissingPixels;
 	TUniquePtr<TImageBuilder<FVector4f>> ColorImage = BaseColorBaker.TakeResult();
-	TempBakeCache.FindSamplingHoles([&](const FVector2i& Coords)
+	TMarchingPixelInfill<FVector4f> Infill;
+
 	{
-		return ColorImage->GetPixel(Coords) == DefaultSample.BaseColor;
-	}, MissingPixels);
+		TRACE_CPUPROFILER_EVENT_SCOPE(ApproximateActorsImpl_Textures_ComputeInfill);
+		TempBakeCache.FindSamplingHoles([&](const FVector2i& Coords)
+		{
+			return ColorImage->GetPixel(Coords) == InvalidColor;
+		}, MissingPixels);
 
-	// solve infill for the holes while also caching infill information
-	TMarchingPixelInfill<FVector4f> Infill;
-	Infill.ComputeInfill(*ColorImage, MissingPixels, DefaultSample.BaseColor,
-		[](FVector4f SumValue, int32 Count) {
-		float InvSum = (Count == 0) ? 1.0f : (1.0f / Count);
-		return FVector4f(SumValue.X * InvSum, SumValue.Y * InvSum, SumValue.Z * InvSum, 1.0f);
-	});
+		// solve infill for the holes while also caching infill information
+		Infill.ComputeInfill(*ColorImage, MissingPixels, InvalidColor,
+			[](FVector4f SumValue, int32 Count) {
+			float InvSum = (Count == 0) ? 1.0f : (1.0f / Count);
+			return FVector4f(SumValue.X * InvSum, SumValue.Y * InvSum, SumValue.Z * InvSum, 1.0f);
+		});
+	}
 
 	// downsample the image if necessary
 	if (Supersample > 1)
 	{
+		TRACE_CPUPROFILER_EVENT_SCOPE(ApproximateActorsImpl_Textures_Downsample);
 		TImageBuilder<FVector4f> Downsampled = ColorImage->FastDownsample(Supersample, FVector4f::Zero(), [](FVector4f V, int N) { return V / (float)N; });
 		*ColorImage = MoveTemp(Downsampled);
 	}
@@ -147,8 +180,8 @@ static void BakeTexturesFromPhotoCapture(
 		ChannelBaker.SetCache(&TempBakeCache);
 		ChannelBaker.ColorSampleFunction = [&](FVector3d Position, FVector3d Normal) {
 			FSceneCapturePhotoSet::FSceneSample Sample = DefaultSample;
-			SceneCapture.ComputeSample(FRenderCaptureTypeFlags::Single(CaptureType), Position, Normal, VisibilityFunction, Sample);
-			return Sample.GetValue(CaptureType);
+			SceneCapture->ComputeSample(FRenderCaptureTypeFlags::Single(CaptureType), Position, Normal, VisibilityFunction, Sample);
+			return Sample.GetValue4f(CaptureType);
 		};
 		ChannelBaker.Bake();
 		TUniquePtr<TImageBuilder<FVector4f>> Image = ChannelBaker.TakeResult();
@@ -168,14 +201,19 @@ static void BakeTexturesFromPhotoCapture(
 		return MoveTemp(Image);
 	};
 
-	Progress.EnterProgressFrame(1.f, LOCTEXT("BakingRoughness", "Baking Roughness..."));
-	TUniquePtr<TImageBuilder<FVector4f>> RoughnessImage = ProcessChannelFunc(ERenderCaptureType::Roughness);
-	Progress.EnterProgressFrame(1.f, LOCTEXT("BakingMetallic", "Baking Metallic..."));
-	TUniquePtr<TImageBuilder<FVector4f>> MetallicImage = ProcessChannelFunc(ERenderCaptureType::Metallic);
-	Progress.EnterProgressFrame(1.f, LOCTEXT("BakingSpecular", "Baking Specular..."));
-	TUniquePtr<TImageBuilder<FVector4f>> SpecularImage = ProcessChannelFunc(ERenderCaptureType::Specular);
-	Progress.EnterProgressFrame(1.f, LOCTEXT("BakingEmissive", "Baking Emissive..."));
-	TUniquePtr<TImageBuilder<FVector4f>> EmissiveImage = ProcessChannelFunc(ERenderCaptureType::Emissive);
+	TUniquePtr<TImageBuilder<FVector4f>> RoughnessImage, MetallicImage, SpecularImage, EmissiveImage;
+	{
+		TRACE_CPUPROFILER_EVENT_SCOPE(ApproximateActorsImpl_Textures_OtherChannels);
+
+		Progress.EnterProgressFrame(1.f, LOCTEXT("BakingRoughness", "Baking Roughness..."));
+		RoughnessImage = ProcessChannelFunc(ERenderCaptureType::Roughness);
+		Progress.EnterProgressFrame(1.f, LOCTEXT("BakingMetallic", "Baking Metallic..."));
+		MetallicImage = ProcessChannelFunc(ERenderCaptureType::Metallic);
+		Progress.EnterProgressFrame(1.f, LOCTEXT("BakingSpecular", "Baking Specular..."));
+		SpecularImage = ProcessChannelFunc(ERenderCaptureType::Specular);
+		Progress.EnterProgressFrame(1.f, LOCTEXT("BakingEmissive", "Baking Emissive..."));
+		EmissiveImage = ProcessChannelFunc(ERenderCaptureType::Emissive);
+	}
 
 	Progress.EnterProgressFrame(1.f, LOCTEXT("BakingNormals", "Baking Normals..."));
 
@@ -187,15 +225,18 @@ static void BakeTexturesFromPhotoCapture(
 	NormalMapBaker.BaseMeshTangents = MeshTangents;
 	NormalMapBaker.NormalSampleFunction = [&](FVector3d Position, FVector3d Normal) {
 		FSceneCapturePhotoSet::FSceneSample Sample = DefaultSample;
-		SceneCapture.ComputeSample(FRenderCaptureTypeFlags::WorldNormal(),
+		SceneCapture->ComputeSample(FRenderCaptureTypeFlags::WorldNormal(),
 			Position, Normal, VisibilityFunction, Sample);
-		FVector4f NormalColor = Sample.WorldNormal;
+		FVector3f NormalColor = Sample.WorldNormal;
 		float x = (NormalColor.X - 0.5f) * 2.0f;
 		float y = (NormalColor.Y - 0.5f) * 2.0f;
 		float z = (NormalColor.Z - 0.5f) * 2.0f;
 		return FVector3f(x, y, z);
 	};
-	NormalMapBaker.Bake();
+	{
+		TRACE_CPUPROFILER_EVENT_SCOPE(ApproximateActorsImpl_Textures_NormalMapBake);
+		NormalMapBaker.Bake();
+	}
 	TUniquePtr<TImageBuilder<FVector3f>> NormalImage = NormalMapBaker.TakeResult();
 
 	if (Supersample > 1)
@@ -207,6 +248,8 @@ static void BakeTexturesFromPhotoCapture(
 	// build textures
 	Progress.EnterProgressFrame(1.f, LOCTEXT("BuildingTextures", "Building Textures..."));
 	{
+		TRACE_CPUPROFILER_EVENT_SCOPE(ApproximateActorsImpl_Textures_BuildTextures);
+
 		FScopedSlowTask BuildTexProgress(6.f, LOCTEXT("BuildingTextures", "Building Textures..."));
 		BuildTexProgress.MakeDialog(true);
 		BuildTexProgress.EnterProgressFrame(1.f);
@@ -224,69 +267,79 @@ static void BakeTexturesFromPhotoCapture(
 	}
 }
 
-void FApproximateActorsImpl::ApproximateActors(const TArray<AActor*>& Actors, const FOptions& Options, FResults& ResultsOut)
+struct FApproximationMeshData
 {
-	int32 ActorClusters = 1;
-	FScopedSlowTask Progress(1.f, LOCTEXT("ApproximatingActors", "Generating Actor Approximation..."));
-	Progress.MakeDialog(true);
-	Progress.EnterProgressFrame(1.f);
-	GenerateApproximationForActorSet(Actors, Options, ResultsOut);
-}
+	IGeometryProcessing_ApproximateActors::EResultCode ResultCode = IGeometryProcessing_ApproximateActors::EResultCode::UnknownError;
 
+	bool bHaveMesh = false;
+	FDynamicMesh3 Mesh;
 
-void FApproximateActorsImpl::GenerateApproximationForActorSet(const TArray<AActor*>& Actors, const FOptions& Options, FResults& ResultsOut)
-{
-	//
-	// Future Optimizations
-	// 	   - can do most of the mesh processing at the same time as capturing the photo set (if that matters)
-	//     - some parts of mesh gen can be done simultaneously (maybe?)
-	//
+	bool bHaveTangents = false;
+	FMeshTangentsd Tangents;
+};
 
-	FScopedSlowTask Progress(10.f, LOCTEXT("ApproximatingActors", "Generating Actor Approximation..."));
 
-	Progress.EnterProgressFrame(1.f, LOCTEXT("BuildingScene", "Building Scene..."));
+static TSharedPtr<FApproximationMeshData> GenerateApproximationMesh(
+	FMeshSceneAdapter& Scene,
+	const IGeometryProcessing_ApproximateActors::FOptions& Options,
+	double ApproxAccuracy
+)
+{
+	FScopedSlowTask Progress(8.f, LOCTEXT("Generating Mesh", "Generating Mesh.."));
 
-	FMeshSceneAdapter Scene;
-	Scene.AddActors(Actors);
+	TSharedPtr<FApproximationMeshData> Result = MakeShared<FApproximationMeshData>();
 
+	// collect seed poitns
 	TArray<FVector3d> SeedPoints;
-	Scene.CollectMeshSeedPoints(SeedPoints);
+	{
+		TRACE_CPUPROFILER_EVENT_SCOPE(ApproximateActorsImpl_Generate_SeedPoints);
+		Scene.CollectMeshSeedPoints(SeedPoints);
+	}
 	FAxisAlignedBox3d SceneBounds = Scene.GetBoundingBox();
 
 	// calculate a voxel size based on target world-space approximation accuracy
 	float WorldBoundsSize = SceneBounds.DiagonalLength();
-	float ApproxAccuracy = Options.WorldSpaceApproximationAccuracyMeters * 100.0;		// convert to cm
-	int32 VoxelDimTarget = (int)(WorldBoundsSize / ApproxAccuracy);
+	int32 VoxelDimTarget = (int)(WorldBoundsSize / ApproxAccuracy) + 1;
 	if (VoxelDimTarget < 64)
 	{
 		VoxelDimTarget = 64;		// use a sane minimum in case the parameter is super-wrong
 	}
 
 	// avoid insane memory usage
-	if (ensure(VoxelDimTarget < Options.ClampVoxelDimension) == false)
+	if (VoxelDimTarget > Options.ClampVoxelDimension)
 	{
 		UE_LOG(LogTemp, Warning, TEXT("FApproximateActorsImpl - very large voxel size %d clamped to %d"), VoxelDimTarget, Options.ClampVoxelDimension);
 		VoxelDimTarget = Options.ClampVoxelDimension;
 	}
 
-	Progress.EnterProgressFrame(1.f, LOCTEXT("GeneratingMesh", "Generating Mesh..."));
+	Progress.EnterProgressFrame(1.f, LOCTEXT("SolidifyMesh", "Approximating Mesh..."));
 
 	FWindingNumberBasedSolidify Solidify(
 		[&Scene](const FVector3d& Position) { return Scene.FastWindingNumber(Position); },
 		SceneBounds, SeedPoints);
-	Solidify.SetCellSizeAndExtendBounds(SceneBounds, 0, VoxelDimTarget);
+	Solidify.SetCellSizeAndExtendBounds(SceneBounds, 2.0 * ApproxAccuracy, VoxelDimTarget);
 	Solidify.WindingThreshold = Options.WindingThreshold;
 
-	FDynamicMesh3 SolidMesh(&Solidify.Generate());
+	FDynamicMesh3 SolidMesh;
+	{
+		TRACE_CPUPROFILER_EVENT_SCOPE(ApproximateActorsImpl_Generate_Solidify);
+		SolidMesh = FDynamicMesh3(&Solidify.Generate());
+	}
 	SolidMesh.DiscardAttributes();
 	FDynamicMesh3* CurResultMesh = &SolidMesh;		// this pointer will be updated as we recompute the mesh
 
+	if (Options.bVerbose)
+	{
+		UE_LOG(LogTemp, Warning, TEXT("[ApproximateActors] Solidify mesh has %d triangles"), CurResultMesh->TriangleCount());
+	}
+
 	Progress.EnterProgressFrame(1.f, LOCTEXT("ClosingMesh", "Topological Operations..."));
 
 	// do topological closure to fix small gaps/etc
 	FDynamicMesh3 MorphologyMesh;
 	if (Options.bApplyMorphology)
 	{
+		TRACE_CPUPROFILER_EVENT_SCOPE(ApproximateActorsImpl_Generate_Morphology);
 		double MorphologyDistance = Options.MorphologyDistanceMeters * 100.0;		// convert to cm
 		FAxisAlignedBox3d MorphologyBounds = CurResultMesh->GetBounds();
 		FDynamicMeshAABBTree3 MorphologyBVTree(CurResultMesh);
@@ -303,44 +356,111 @@ void FApproximateActorsImpl::GenerateApproximationForActorSet(const TArray<AActo
 	// if mesh has no triangles, something has gone wrong
 	if (CurResultMesh == nullptr || CurResultMesh->TriangleCount() == 0)
 	{
-		ResultsOut.ResultCode = EResultCode::MeshGenerationFailed;
-		return;
+		Result->ResultCode = IGeometryProcessing_ApproximateActors::EResultCode::MeshGenerationFailed;
+		return Result;
 	}
 
 	Progress.EnterProgressFrame(1.f, LOCTEXT("RemoveHidden", "Removing Hidden Geometry..."));
 
-	//  TODO: remove interior components that will never be visible (Not the same as hidden triangles)
-
-	Progress.EnterProgressFrame(1.f, LOCTEXT("SimplifyingMesh", "Simplifying Mesh..."));
-
-	int32 SimplifyTargetTriCount = Options.FixedTriangleCount;
-	if (Options.MeshSimplificationPolicy == ESimplificationPolicy::TrianglesPerUnitSqMeter)
+	if (Options.OcclusionPolicy == IGeometryProcessing_ApproximateActors::EOcclusionPolicy::VisibilityBased)
 	{
-		FVector2d VolArea = TMeshQueries<FDynamicMesh3>::GetVolumeArea(*CurResultMesh);
-		double MeshAreaMeterSqr = VolArea.Y * 0.0001;
-		SimplifyTargetTriCount = MeshAreaMeterSqr * Options.SimplificationTargetMetric;
+		TRACE_CPUPROFILER_EVENT_SCOPE(ApproximateActorsImpl_Generate_Occlusion);
+		TRemoveOccludedTriangles<FDynamicMesh3> Remover(CurResultMesh);
+		Remover.InsideMode = EOcclusionCalculationMode::SimpleOcclusionTest;
+		Remover.TriangleSamplingMethod = EOcclusionTriangleSampling::Centroids;
+		Remover.AddTriangleSamples = 5;
+		FDynamicMeshAABBTree3 CurResultMeshSpatial(CurResultMesh, true);
+
+		TArray<FTransform3d> NoTransforms;
+		NoTransforms.Add(FTransform3d::Identity());
+		TArray<FDynamicMeshAABBTree3*> Spatials;
+		Spatials.Add(&CurResultMeshSpatial);
+		Remover.Select(NoTransforms, Spatials, {}, NoTransforms);
+		if (Remover.RemovedT.Num() > 0)
+		{
+			FMeshFaceSelection Selection(CurResultMesh);
+			{
+				TRACE_CPUPROFILER_EVENT_SCOPE(ApproximateActorsImpl_Generate_Occlusion_Build);
+				Selection.Select(Remover.RemovedT);
+				Selection.ExpandToOneRingNeighbours(1);
+				Selection.ContractBorderByOneRingNeighbours(4);
+			}
+			FDynamicMeshEditor Editor(CurResultMesh);
+			{
+				TRACE_CPUPROFILER_EVENT_SCOPE(ApproximateActorsImpl_Generate_Occlusion_Delete);
+				Editor.RemoveTriangles(Selection.AsArray(), true);
+			}
+		}
 	}
 
+	Progress.EnterProgressFrame(1.f, LOCTEXT("SimplifyingMesh", "Simplifying Mesh..."));
+
 	FVolPresMeshSimplification Simplifier(CurResultMesh);
 	Simplifier.ProjectionMode = FVolPresMeshSimplification::ETargetProjectionMode::NoProjection;
 	Simplifier.DEBUG_CHECK_LEVEL = 0;
 	Simplifier.bAllowSeamCollapse = false;
-	Simplifier.SimplifyToTriangleCount(SimplifyTargetTriCount);
+
+	int32 BaseTargeTriCount = Options.FixedTriangleCount;
+	{
+		int32 BeforeCount = CurResultMesh->TriangleCount();
+
+		TRACE_CPUPROFILER_EVENT_SCOPE(ApproximateActorsImpl_Generate_Simplification);
+		if (Options.MeshSimplificationPolicy == IGeometryProcessing_ApproximateActors::ESimplificationPolicy::TrianglesPerUnitSqMeter)
+		{
+			FVector2d VolArea = TMeshQueries<FDynamicMesh3>::GetVolumeArea(*CurResultMesh);
+			double MeshAreaMeterSqr = VolArea.Y * 0.0001;
+			int32 AreaBaseTargetTriCount = MeshAreaMeterSqr * Options.SimplificationTargetMetric;
+			Simplifier.SimplifyToTriangleCount(AreaBaseTargetTriCount);
+		}
+		else if (Options.MeshSimplificationPolicy == IGeometryProcessing_ApproximateActors::ESimplificationPolicy::GeometricTolerance)
+		{
+			double UseTargetTolerance = Options.SimplificationTargetMetric * 100.0;		// convert to cm (UE Units)
+
+			// first do fast collapse
+			Simplifier.FastCollapsePass(0.1 * UseTargetTolerance, 5);
+
+			// now simplify down to a reasonable tri count, as geometric metric is (relatively) expensive
+			// (still, this is all incredibly cheap compared to the cost of the rest of this method in practice)
+			Simplifier.SimplifyToTriangleCount(50000);
+
+			FDynamicMesh3 MeshCopy(*CurResultMesh);
+			FDynamicMeshAABBTree3 MeshCopySpatial(&MeshCopy, true);
+			FMeshProjectionTarget ProjectionTarget(&MeshCopy, &MeshCopySpatial);
+			Simplifier.SetProjectionTarget(&ProjectionTarget);
+			Simplifier.GeometricErrorConstraint = FVolPresMeshSimplification::EGeometricErrorCriteria::PredictedPointToProjectionTarget;
+			Simplifier.GeometricErrorTolerance = UseTargetTolerance;
+			Simplifier.SimplifyToTriangleCount(8);
+		}
+		else
+		{
+			Simplifier.SimplifyToTriangleCount(BaseTargeTriCount);
+		}
+
+		int32 AfterCount = CurResultMesh->TriangleCount();
+		if (Options.bVerbose)
+		{
+			UE_LOG(LogTemp, Warning, TEXT("[ApproximateActors] Simplified mesh from %d to %d triangles"), BeforeCount, AfterCount);
+		}
+	}
 
 	// re-enable attributes
 	CurResultMesh->EnableAttributes();
 
-	//  TODO: clip hidden triangles against occluder geo like landscape (should be no hidden coming out of meshing)
+	//  TODO: clip hidden triangles against occluder geo like landscape
 
 	// compute normals
-	FMeshNormals::InitializeOverlayToPerVertexNormals(CurResultMesh->Attributes()->PrimaryNormals());
+	{
+		TRACE_CPUPROFILER_EVENT_SCOPE(ApproximateActorsImpl_Generate_Normals);
+		FMeshNormals::InitializeOverlayToPerVertexNormals(CurResultMesh->Attributes()->PrimaryNormals());
+	}
 
 	// exit here if we are just generating a merged collision mesh
-	if (Options.BasePolicy == EApproximationPolicy::CollisionMesh)
+	if (Options.BasePolicy == IGeometryProcessing_ApproximateActors::EApproximationPolicy::CollisionMesh)
 	{
-		EmitGeneratedMeshAsset(Actors, Options, ResultsOut, CurResultMesh, nullptr);
-		ResultsOut.ResultCode = EResultCode::Success;
-		return;
+		Result->ResultCode = IGeometryProcessing_ApproximateActors::EResultCode::Success;
+		Result->bHaveMesh = true;
+		Result->Mesh = MoveTemp(*CurResultMesh);
+		return Result;
 	}
 
 	Progress.EnterProgressFrame(1.f, LOCTEXT("ComputingUVs", "Computing UVs..."));
@@ -355,7 +475,10 @@ void FApproximateActorsImpl::GenerateApproximationForActorSet(const TArray<AActo
 	ParameterizeMeshOp.IslandMode = EParamOpIslandMode::Auto;
 	ParameterizeMeshOp.UnwrapType = EParamOpUnwrapType::MinStretch;
 	FProgressCancel UVProgressCancel;
-	ParameterizeMeshOp.CalculateResult(&UVProgressCancel);
+	{
+		TRACE_CPUPROFILER_EVENT_SCOPE(ApproximateActorsImpl_Generate_GenerateUVs);
+		ParameterizeMeshOp.CalculateResult(&UVProgressCancel);
+	}
 
 	TUniquePtr<FDynamicMesh3> FinalMesh = ParameterizeMeshOp.ExtractResult();
 
@@ -368,44 +491,159 @@ void FApproximateActorsImpl::GenerateApproximationForActorSet(const TArray<AActo
 	Packer.TextureResolution = Options.TextureImageSize / 4;		// maybe too conservative? We don't have gutter control currently.
 	Packer.GutterSize = 1.0;		// not clear this works
 	Packer.bAllowFlips = false;
-	bool bOK = Packer.StandardPack();
-	ensure(bOK);
+	{
+		TRACE_CPUPROFILER_EVENT_SCOPE(ApproximateActorsImpl_Generate_PackUVs);
+		bool bOK = Packer.StandardPack();
+		ensure(bOK);
+	}
 
 	Progress.EnterProgressFrame(1.f, LOCTEXT("ComputingTangents", "Computing Tangents..."));
 
+	Result->ResultCode = IGeometryProcessing_ApproximateActors::EResultCode::Success;
+	Result->bHaveMesh = true;
+	Result->Mesh = MoveTemp(*FinalMesh);
+
 	// compute tangents
-	FMeshTangentsd Tangents(FinalMesh.Get());
+	Result->bHaveTangents = true;
+	Result->Tangents.SetMesh(&Result->Mesh);
 	FComputeTangentsOptions TangentsOptions;
 	TangentsOptions.bAveraged = true;
-	Tangents.ComputeTriVertexTangents(
-		FinalMesh->Attributes()->PrimaryNormals(),
-		FinalMesh->Attributes()->PrimaryUV(),
-		TangentsOptions);
+	{
+		TRACE_CPUPROFILER_EVENT_SCOPE(ApproximateActorsImpl_Generate_Tangents);
+		Result->Tangents.ComputeTriVertexTangents(
+			Result->Mesh.Attributes()->PrimaryNormals(),
+			Result->Mesh.Attributes()->PrimaryUV(),
+			TangentsOptions);
+	}
+
+	return Result;
+}
+
+
+
+
+void FApproximateActorsImpl::ApproximateActors(const TArray<AActor*>& Actors, const FOptions& Options, FResults& ResultsOut)
+{
+	int32 ActorClusters = 1;
+	FScopedSlowTask Progress(1.f, LOCTEXT("ApproximatingActors", "Generating Actor Approximation..."));
+	Progress.MakeDialog(true);
+	Progress.EnterProgressFrame(1.f);
+	GenerateApproximationForActorSet(Actors, Options, ResultsOut);
+}
+
+
+
+
+void FApproximateActorsImpl::GenerateApproximationForActorSet(const TArray<AActor*>& Actors, const FOptions& Options, FResults& ResultsOut)
+{
+	TRACE_CPUPROFILER_EVENT_SCOPE(ApproximateActorsImpl_Generate);
+
+	//
+	// Future Optimizations
+	// 	   - can do most of the mesh processing at the same time as capturing the photo set (if that matters)
+	//     - some parts of mesh gen can be done simultaneously (maybe?)
+	//
+
+	FScopedSlowTask Progress(11.f, LOCTEXT("ApproximatingActors", "Generating Actor Approximation..."));
+
+	Progress.EnterProgressFrame(1.f, LOCTEXT("BuildingScene", "Building Scene..."));
+
+	float ApproxAccuracy = Options.WorldSpaceApproximationAccuracyMeters * 100.0;		// convert to cm (UE Units)
+
+	FMeshSceneAdapter Scene;
+	FMeshSceneAdapterBuildOptions SceneBuildOptions;
+	SceneBuildOptions.bThickenThinMeshes = Options.bAutoThickenThinParts;
+	SceneBuildOptions.DesiredMinThickness = Options.AutoThickenThicknessMeters * 100.0;		// convert to cm (UE Units)
+	SceneBuildOptions.bPrintDebugMessages = Options.bVerbose;
+	{
+		TRACE_CPUPROFILER_EVENT_SCOPE(ApproximateActorsImpl_Generate_BuildScene);
+		Scene.AddActors(Actors);
+		Scene.Build(SceneBuildOptions);
+	}
+
+	// todo: make optional
+	if (Options.bVerbose)
+	{
+		FMeshSceneAdapter::FStatistics Stats;
+		Scene.GetGeometryStatistics(Stats);
+		UE_LOG(LogTemp, Warning, TEXT("[ApproximateActors] %d triangles in %d unique meshes, total %d triangles in %d instances"),
+			Stats.UniqueMeshTriangleCount, Stats.UniqueMeshCount, Stats.InstanceMeshTriangleCount, Stats.InstanceMeshCount);
+	}
+
+	if (Options.BaseCappingPolicy != EBaseCappingPolicy::NoBaseCapping)
+	{
+		TRACE_CPUPROFILER_EVENT_SCOPE(ApproximateActorsImpl_Generate_Capping);
+		double UseThickness = (Options.BaseThicknessOverrideMeters != 0) ? (Options.BaseThicknessOverrideMeters * 100.0) :
+			(Options.bAutoThickenThinParts ? SceneBuildOptions.DesiredMinThickness : 1.25 * ApproxAccuracy);
+		double UseHeight = (Options.BaseHeightOverrideMeters != 0) ? (Options.BaseHeightOverrideMeters * 100.0) : (2.0 * ApproxAccuracy);
+		Scene.GenerateBaseClosingMesh(UseHeight, UseThickness);
+	}
+
+	FDynamicMesh3 DebugMesh;
+	FDynamicMesh3* WriteDebugMesh = nullptr;
+	if (Options.bWriteDebugMesh)
+	{
+		TRACE_CPUPROFILER_EVENT_SCOPE(ApproximateActorsImpl_Generate_DebugMesh);
+		DebugMesh.EnableAttributes();
+		Scene.GetAccumulatedMesh(DebugMesh);
+		FMeshNormals::InitializeMeshToPerTriangleNormals(&DebugMesh);
+		WriteDebugMesh = &DebugMesh;
+	}
+
+	// if we are only generating collision mesh, we are going to exit after mesh generation
+	if (Options.BasePolicy == IGeometryProcessing_ApproximateActors::EApproximationPolicy::CollisionMesh)
+	{
+		TSharedPtr<FApproximationMeshData> ApproximationMeshData = GenerateApproximationMesh(Scene, Options, ApproxAccuracy);
+		ResultsOut.ResultCode = ApproximationMeshData->ResultCode;
+		if (ResultsOut.ResultCode == EResultCode::Success)
+		{
+			EmitGeneratedMeshAsset(Actors, Options, ResultsOut, &ApproximationMeshData->Mesh, nullptr, WriteDebugMesh);
+		}
+		return;
+	}
+
+	// launch async mesh compute which can run while we do (relatively) expensive render captures
+	TFuture<TSharedPtr<FApproximationMeshData>> MeshComputeFuture = Async(EAsyncExecution::Thread,
+		[&Scene, &Options, &ApproxAccuracy]() {
+			return GenerateApproximationMesh(Scene, Options, ApproxAccuracy);
+		});
+
+	Progress.EnterProgressFrame(1.f, LOCTEXT("CapturingScene", "Capturing Scene..."));
+
+	TUniquePtr<FSceneCapturePhotoSet> SceneCapture = CapturePhotoSet(Actors, Options);
 
-	// TODO: copy/calc into mesh tangents attrib, so it can be transferred to MeshDescription
-	 
 	Progress.EnterProgressFrame(1.f, LOCTEXT("BakingTextures", "Baking Textures..."));
 
+	// need to wait for mesh to finish computing
+	MeshComputeFuture.Wait();
+	TSharedPtr<FApproximationMeshData> ApproximationMeshData = MeshComputeFuture.Get();
+	if (ApproximationMeshData->ResultCode != EResultCode::Success)
+	{
+		ResultsOut.ResultCode = ApproximationMeshData->ResultCode;
+		return;
+	}
+	FDynamicMesh3 FinalMesh = MoveTemp(ApproximationMeshData->Mesh);
+	FMeshTangentsd FinalMeshTangents = MoveTemp(ApproximationMeshData->Tangents);
+
 	// bake textures for Actor
 	FGeneratedResultTextures GeneratedTextures;
-	BakeTexturesFromPhotoCapture(Actors,
-		Options,
+	BakeTexturesFromPhotoCapture(SceneCapture, Options,
 		GeneratedTextures,
-		FinalMesh.Get(),
-		&Tangents);
+		&FinalMesh, &FinalMeshTangents);
 
 	Progress.EnterProgressFrame(1.f, LOCTEXT("Writing Assets", "Writing Assets..."));
 
 	// Make material for textures by duplicating input material (hardcoded!!)
-	UMaterial* BakeMaterial = LoadObject<UMaterial>(nullptr, TEXT("/MeshModelingToolset/Materials/FullMaterialBakePreviewMaterial"));
+	UMaterialInterface* UseBaseMaterial = (Options.BakeMaterial != nullptr) ? 
+		Options.BakeMaterial : LoadObject<UMaterial>(nullptr, TEXT("/MeshModelingToolset/Materials/FullMaterialBakePreviewMaterial"));
 	FMaterialAssetOptions MatOptions;
 	MatOptions.NewAssetPath = Options.BasePackagePath + TEXT("_Material");
 	FMaterialAssetResults MatResults;
-	ECreateMaterialResult MatResult = UE::AssetUtils::CreateDuplicateMaterial(BakeMaterial, MatOptions, MatResults);
-	UMaterial* NewMaterial = nullptr;
+	ECreateMaterialResult MatResult = UE::AssetUtils::CreateDerivedMaterialInstance(UseBaseMaterial, MatOptions, MatResults);
+	UMaterialInstanceConstant* NewMaterial = nullptr;
 	if (ensure(MatResult == ECreateMaterialResult::Ok))
 	{
-		NewMaterial = MatResults.NewMaterial;
+		NewMaterial = MatResults.NewMaterialInstance;
 		ResultsOut.NewMaterials.Add(NewMaterial);
 	}
 
@@ -421,6 +659,13 @@ void FApproximateActorsImpl::GenerateApproximationForActorSet(const TArray<AActo
 
 		FTexture2DBuilder::CopyPlatformDataToSourceData(Texture, Type);
 
+		if (Type == FTexture2DBuilder::ETextureType::Roughness
+			|| Type == FTexture2DBuilder::ETextureType::Metallic
+			|| Type == FTexture2DBuilder::ETextureType::Specular)
+		{
+			UE::AssetUtils::ConvertToSingleChannel(Texture);
+		}
+
 		FTexture2DAssetOptions TexOptions;
 		TexOptions.NewAssetPath = BaseTexturePath + TextureTypeSuffix;
 		FTexture2DAssetResults Results;
@@ -434,19 +679,39 @@ void FApproximateActorsImpl::GenerateApproximationForActorSet(const TArray<AActo
 			}
 		}
 	};
+
+
 	// process the generated textures
-	WriteTextureLambda(GeneratedTextures.BaseColorMap, TEXT("_BaseColor"), FTexture2DBuilder::ETextureType::Color, FName("BaseColor"));
-	WriteTextureLambda(GeneratedTextures.RoughnessMap, TEXT("_Roughness"), FTexture2DBuilder::ETextureType::Color, FName("Roughness"));
-	WriteTextureLambda(GeneratedTextures.MetallicMap, TEXT("_Metallic"), FTexture2DBuilder::ETextureType::Color, FName("Metallic"));
-	WriteTextureLambda(GeneratedTextures.SpecularMap, TEXT("_Specular"), FTexture2DBuilder::ETextureType::Color, FName("Specular"));
-	WriteTextureLambda(GeneratedTextures.EmissiveMap, TEXT("_Emissive"), FTexture2DBuilder::ETextureType::Color, FName("Emissive"));
-	WriteTextureLambda(GeneratedTextures.NormalMap, TEXT("_Normal"), FTexture2DBuilder::ETextureType::NormalMap, FName("NormalMap"));
+	if (Options.bBakeBaseColor && GeneratedTextures.BaseColorMap)
+	{
+		WriteTextureLambda(GeneratedTextures.BaseColorMap, TEXT("_BaseColor"), FTexture2DBuilder::ETextureType::Color, Options.BaseColorTexParamName);
+	}
+	if (Options.bBakeRoughness && GeneratedTextures.RoughnessMap)
+	{
+		WriteTextureLambda(GeneratedTextures.RoughnessMap, TEXT("_Roughness"), FTexture2DBuilder::ETextureType::Roughness, Options.RoughnessTexParamName);
+	}
+	if (Options.bBakeMetallic && GeneratedTextures.MetallicMap)
+	{
+		WriteTextureLambda(GeneratedTextures.MetallicMap, TEXT("_Metallic"), FTexture2DBuilder::ETextureType::Metallic, Options.MetallicTexParamName);
+	}
+	if (Options.bBakeSpecular && GeneratedTextures.SpecularMap)
+	{
+		WriteTextureLambda(GeneratedTextures.SpecularMap, TEXT("_Specular"), FTexture2DBuilder::ETextureType::Specular, Options.SpecularTexParamName);
+	}
+	if (Options.bBakeEmissive && GeneratedTextures.EmissiveMap)
+	{
+		WriteTextureLambda(GeneratedTextures.EmissiveMap, TEXT("_Emissive"), FTexture2DBuilder::ETextureType::Color, Options.EmissiveTexParamName);
+	}
+	if (Options.bBakeNormalMap && GeneratedTextures.NormalMap)
+	{
+		WriteTextureLambda(GeneratedTextures.NormalMap, TEXT("_Normal"), FTexture2DBuilder::ETextureType::NormalMap, Options.NormalTexParamName);
+	}
 
 	// force material update now that we have updated texture parameters
 	// (does this do that? Let calling code do it?)
 	NewMaterial->PostEditChange();
 
-	EmitGeneratedMeshAsset(Actors, Options, ResultsOut, FinalMesh.Get(), NewMaterial);
+	EmitGeneratedMeshAsset(Actors, Options, ResultsOut, &FinalMesh, NewMaterial, WriteDebugMesh);
 	ResultsOut.ResultCode = EResultCode::Success;
 }
 
@@ -456,7 +721,8 @@ UStaticMesh* FApproximateActorsImpl::EmitGeneratedMeshAsset(
 	const FOptions& Options, 
 	FResults& ResultsOut,
 	FDynamicMesh3* FinalMesh,
-	UMaterialInterface* Material)
+	UMaterialInterface* Material,
+	FDynamicMesh3* DebugMesh)
 {
 	FStaticMeshAssetOptions MeshAssetOptions;
 	MeshAssetOptions.NewAssetPath = Options.BasePackagePath;
@@ -471,6 +737,15 @@ UStaticMesh* FApproximateActorsImpl::EmitGeneratedMeshAsset(
 
 	ResultsOut.NewMeshAssets.Add(MeshAssetOutputs.StaticMesh);
 
+	if (DebugMesh != nullptr)
+	{
+		FStaticMeshAssetOptions DebugMeshAssetOptions;
+		DebugMeshAssetOptions.NewAssetPath = Options.BasePackagePath + TEXT("_DEBUG");
+		DebugMeshAssetOptions.SourceMeshes.DynamicMeshes.Add(DebugMesh);
+		FStaticMeshResults DebugMeshAssetOutputs;
+		UE::AssetUtils::CreateStaticMeshAsset(DebugMeshAssetOptions, DebugMeshAssetOutputs);
+	}
+
 	return MeshAssetOutputs.StaticMesh;
 }
 
diff --git a/Engine/Plugins/Experimental/MeshModelingToolset/Source/GeometryProcessingAdapters/Public/GeometryProcessing/ApproximateActorsImpl.h b/Engine/Plugins/Experimental/MeshModelingToolset/Source/GeometryProcessingAdapters/Public/GeometryProcessing/ApproximateActorsImpl.h
--- a/Engine/Plugins/Experimental/MeshModelingToolset/Source/GeometryProcessingAdapters/Public/GeometryProcessing/ApproximateActorsImpl.h
+++ b/Engine/Plugins/Experimental/MeshModelingToolset/Source/GeometryProcessingAdapters/Public/GeometryProcessing/ApproximateActorsImpl.h
@@ -34,7 +34,8 @@ protected:
 		const FOptions& Options,
 		FResults& ResultsOut,
 		FDynamicMesh3* FinalMesh,
-		UMaterialInterface* Material = nullptr);
+		UMaterialInterface* Material = nullptr,
+		FDynamicMesh3* DebugMesh = nullptr);
 };
 
 
diff --git a/Engine/Plugins/Experimental/MeshModelingToolset/Source/ModelingComponents/Private/Scene/MeshSceneAdapter.cpp b/Engine/Plugins/Experimental/MeshModelingToolset/Source/ModelingComponents/Private/Scene/MeshSceneAdapter.cpp
--- a/Engine/Plugins/Experimental/MeshModelingToolset/Source/ModelingComponents/Private/Scene/MeshSceneAdapter.cpp
+++ b/Engine/Plugins/Experimental/MeshModelingToolset/Source/ModelingComponents/Private/Scene/MeshSceneAdapter.cpp
@@ -3,9 +3,19 @@
 
 #include "Scene/MeshSceneAdapter.h"
 
+#include "DynamicMesh3.h"
 #include "Spatial/MeshAABBTree3.h"
 #include "Spatial/FastWinding.h"
 #include "MeshDescriptionAdapter.h"
+#include "MeshDescriptionToDynamicMesh.h"
+#include "MeshTransforms.h"
+#include "BoxTypes.h"
+#include "FrameTypes.h"
+#include "Selections/MeshConnectedComponents.h"
+#include "Operations/OffsetMeshRegion.h"
+#include "DynamicMeshEditor.h"
+#include "ConvexHull2.h"
+#include "Generators/PlanarPolygonMeshGenerator.h"
 
 #include "GameFramework/Actor.h"
 #include "Components/StaticMeshComponent.h"
@@ -17,6 +27,466 @@
 using namespace UE::Geometry;
 
 
+template<typename MeshType>
+FAxisAlignedBox3d GetTransformedVertexBounds(const MeshType& Mesh, TFunctionRef<FVector3d(const FVector3d&)> TransformFunc)
+{
+	FAxisAlignedBox3d Bounds = FAxisAlignedBox3d::Empty();
+	int32 NumVertices = Mesh.VertexCount();
+	for (int32 k = 0; k < NumVertices; ++k)
+	{
+		if (Mesh.IsVertex(k))
+		{
+			Bounds.Contain(TransformFunc(Mesh.GetVertex(k)));
+		}
+	}
+	return Bounds;
+}
+
+
+template<typename MeshType>
+void CollectSeedPointsFromMeshVertices(const MeshType& Mesh, TFunctionRef<FVector3d(const FVector3d&)> TransformFunc, TArray<FVector3d>& AccumPoints)
+{
+	int32 NumVertices = Mesh.VertexCount();
+	int32 LogNumVertices = FMath::Max(1, (int32)FMathd::Ceil(FMathd::Log(NumVertices)));
+	int32 SeedPointCount = (int)(10 * LogNumVertices);
+	SeedPointCount = FMath::Min(SeedPointCount, 5000);
+	int32 Skip = FMath::Max(NumVertices / SeedPointCount, 2);
+	for (int32 k = 0; k < NumVertices; k += Skip)
+	{
+		AccumPoints.Add(TransformFunc(Mesh.GetVertex(k)));
+	}
+}
+
+template<typename MeshType>
+FDynamicMesh3 FastConvertToDynamicMesh(const MeshType& Mesh)
+{
+	FDynamicMesh3 NewMesh;
+	int32 VertexCount = Mesh.VertexCount();
+	TArray<int32> VertexIDMap;
+	VertexIDMap.SetNum(VertexCount);
+	for (int32 vid = 0; vid < VertexCount; ++vid)
+	{
+		if (Mesh.IsVertex(vid))
+		{
+			int32 newvid = NewMesh.AppendVertex(Mesh.GetVertex(vid));
+			VertexIDMap[vid] = newvid;
+		}
+	}
+
+	int32 TriCount = Mesh.TriangleCount();
+	for (int32 tid = 0; tid < TriCount; ++tid)
+	{
+		if (Mesh.IsTriangle(tid))
+		{
+			FIndex3i TriVerts = Mesh.GetTriangle(tid);
+			NewMesh.AppendTriangle(VertexIDMap[TriVerts.A], VertexIDMap[TriVerts.B], VertexIDMap[TriVerts.C]);
+		}
+	}
+
+	return NewMesh;
+}
+
+
+template<typename MeshType>
+bool IsThinPlanarMesh(const MeshType& Mesh, double ThinTolerance, FFrame3d& ThinPlane)
+{
+	int32 TriCount = Mesh.TriangleCount();
+	double MaxArea = 0;
+	FVector3d MaxAreaNormal;
+	FVector3d MaxAreaPoint;
+	for (int32 tid = 0; tid < TriCount; ++tid)
+	{
+		if (Mesh.IsTriangle(tid))
+		{
+			FVector3d A, B, C;
+			Mesh.GetTriVertices(tid, A, B, C);
+			double TriArea;
+			FVector3d TriNormal = VectorUtil::NormalArea(A, B, C, TriArea);
+			if (TriArea > MaxArea)
+			{
+				MaxArea = TriArea;
+				MaxAreaNormal = TriNormal;
+				MaxAreaPoint = A;
+			}
+		}
+	}
+	ThinPlane = FFrame3d(MaxAreaPoint, MaxAreaNormal);
+	FAxisAlignedBox3d PlaneExtents = FAxisAlignedBox3d::Empty();
+	int32 VertexCount = Mesh.VertexCount();
+	for (int32 vid = 0; vid < VertexCount; ++vid)
+	{
+		if (Mesh.IsVertex(vid))
+		{
+			FVector3d Pos = Mesh.GetVertex(vid);
+			Pos = ThinPlane.ToFramePoint(Pos);
+			PlaneExtents.Contain(Pos);
+		}
+	}
+	if (PlaneExtents.Depth() > ThinTolerance)
+	{
+		return false;
+	}
+	FVector3d Center = PlaneExtents.Center();
+	ThinPlane.Origin += Center.X*ThinPlane.X() + Center.Y*ThinPlane.Y() + Center.Z*ThinPlane.Z();
+	return true;
+}
+
+
+static bool IsClosedRegion(const FDynamicMesh3& Mesh, const TArray<int32>& Triangles)
+{
+	for (int32 tid : Triangles)
+	{
+		FIndex3i TriEdges = Mesh.GetTriEdges(tid);
+		if (Mesh.IsBoundaryEdge(TriEdges.A) ||
+			Mesh.IsBoundaryEdge(TriEdges.B) ||
+			Mesh.IsBoundaryEdge(TriEdges.C))
+		{
+			return false;
+		}
+	}
+	return true;
+}
+
+
+template<typename Func>
+void ApplyToTriangleVertexIndices(const FDynamicMesh3& Mesh, const TArray<int32>& Triangles, Func VertexFunc)
+{
+	for (int32 tid : Triangles)
+	{
+		FIndex3i TriVerts = Mesh.GetTriangle(tid);
+		VertexFunc(TriVerts.A);
+		VertexFunc(TriVerts.B);
+		VertexFunc(TriVerts.C);
+	}
+}
+
+template<typename Func>
+void ApplyToTriangleVertices(const FDynamicMesh3& Mesh, const TArray<int32>& Triangles, Func VertexFunc)
+{
+	FVector3d A, B, C;
+	for (int32 tid : Triangles)
+	{
+		Mesh.GetTriVertices(tid, A, B, C);
+		VertexFunc(A);
+		VertexFunc(B);
+		VertexFunc(C);
+	}
+}
+
+static void ThickenThinPlanarMesh(FDynamicMesh3& Mesh, double TargetThickness, const FFrame3d& ThinPlane)
+{
+	FVector3d PlaneZ = ThinPlane.Z();
+
+	TArray<bool> VertFlags;
+	VertFlags.Init(false, Mesh.VertexCount());
+
+	FMeshConnectedComponents Components(&Mesh);
+	Components.FindConnectedTriangles();
+	int32 NumComponents = Components.Num();
+	for (const FMeshConnectedComponents::FComponent& Component : Components)
+	{
+		const TArray<int32>& Triangles = Component.Indices;
+		bool bIsClosed = IsClosedRegion(Mesh, Triangles);
+
+		if (bIsClosed)
+		{
+			FInterval1d ThicknessRange = FInterval1d::Empty();
+			ApplyToTriangleVertices(Mesh, Triangles, [&](const FVector3d& V) {
+				ThicknessRange.Contain((V - ThinPlane.Origin).Dot(PlaneZ));
+			});
+			double CurThickness = ThicknessRange.Length();
+
+			double MoveDelta = (TargetThickness - CurThickness) * 0.5;
+
+			int32 NumTriangles = Triangles.Num();
+			TArray<FVector3d> TriNormals;
+			TriNormals.SetNum(NumTriangles);
+			for (int32 k = 0; k < NumTriangles; ++k)
+			{
+				TriNormals[k] = Mesh.GetTriNormal(Triangles[k]);
+			}
+
+			//for (int32 k = 0; k < NumTriangles; ++k)
+			//{
+			//	int32 tid = Triangles[k];
+			//	FIndex3i TriVerts = Mesh.GetTriangle(tid);
+			//	FVector3d MoveDir = TriNormals[k].Dot(PlaneZ) > 0 ? PlaneZ : -PlaneZ;
+			//	for (int32 j = 0; j < 3; ++j)
+			//	{
+			//		int32 vid = TriVerts[j];
+			//		if (VertFlags[vid] == false)
+			//		{
+			//			Mesh.SetVertex(vid, Mesh.GetVertex(vid) + MoveDir * MoveDelta);
+			//			VertFlags[vid] = true;
+			//		}
+			//	}
+			//}
+
+		}
+		else   // bIsClosed == false
+		{
+			FOffsetMeshRegion Offset(&Mesh);
+			Offset.Triangles = Triangles;
+			//Offset.bUseFaceNormals = true;
+			Offset.DefaultOffsetDistance = -TargetThickness;
+			Offset.bIsPositiveOffset = false;
+			Offset.Apply();
+		}
+	}
+}
+
+
+
+
+
+static bool CheckIfStaticMeshUsageRequiresDynamicMeshDecomposition(
+	UStaticMesh* StaticMesh, 
+	int32 LODIndex,
+	const FMeshSceneAdapterBuildOptions& BuildOptions,
+	const TArray<FActorChildMesh*>& ParentMeshes,
+	int32 NonUniformScaleCount )
+{
+	FMeshDescription* SourceMesh;
+	check(StaticMesh);
+#if WITH_EDITOR
+	SourceMesh = StaticMesh->GetMeshDescription(LODIndex);
+#else
+	checkf(false, TEXT("Not currently supported - to build at Runtime it is necessary to read from the StaticMesh RenderBuffers"));
+	return false;
+#endif
+
+	if (SourceMesh)
+	{
+		// does not need this - copypasta
+		FMeshDescriptionTriangleMeshAdapter Adapter(SourceMesh);
+
+		FVector3d BuildScale;
+#if WITH_EDITOR
+		// respect BuildScale build setting
+		const FMeshBuildSettings& LODBuildSettings = StaticMesh->GetSourceModel(LODIndex).BuildSettings;
+		Adapter.SetBuildScale((FVector3d)LODBuildSettings.BuildScale3D, false);
+#endif
+		if (BuildOptions.bThickenThinMeshes)
+		{
+			FFrame3d ThinMeshPlane;
+			bool IsThinMesh = IsThinPlanarMesh<FMeshDescriptionTriangleMeshAdapter>(Adapter, BuildOptions.DesiredMinThickness, ThinMeshPlane);
+			if (IsThinMesh)
+			{
+				if (Adapter.TriangleCount() > 100000 && ParentMeshes.Num() > 1)
+				{
+					return false;		// skip thickening of huge meshes that are frequently used
+				}
+
+				return true;
+			}
+		}
+	}
+	return false;
+}
+
+
+
+
+
+class FDynamicMeshSpatialWrapper : public IMeshSpatialWrapper
+{
+public:
+	FDynamicMesh3 Mesh;
+	bool bHasBakedTransform = false;
+	TUniquePtr<TMeshAABBTree3<FDynamicMesh3>> AABBTree;
+	TUniquePtr<TFastWindingTree<FDynamicMesh3>> FWNTree;
+
+	virtual bool Build(const FMeshSceneAdapterBuildOptions& BuildOptions) override
+	{
+		ensure(Mesh.TriangleCount() > 0);
+		AABBTree = MakeUnique<TMeshAABBTree3<FDynamicMesh3>>(&Mesh, true);
+		FWNTree = MakeUnique<TFastWindingTree<FDynamicMesh3>>(AABBTree.Get(), true);
+		return true;
+	}
+
+	virtual int32 GetTriangleCount() override
+	{
+		return Mesh.TriangleCount();
+	}
+
+	virtual FAxisAlignedBox3d GetWorldBounds(TFunctionRef<FVector3d(const FVector3d&)> LocalToWorldFunc) override
+	{
+		return bHasBakedTransform ?
+			GetTransformedVertexBounds<FDynamicMesh3>(Mesh, [&](const FVector3d& P) {return P;}) :
+			GetTransformedVertexBounds<FDynamicMesh3>(Mesh, LocalToWorldFunc);
+	}
+
+	virtual void CollectSeedPoints(TArray<FVector3d>& WorldPoints, TFunctionRef<FVector3d(const FVector3d&)> LocalToWorldFunc) override
+	{
+		return bHasBakedTransform ?
+			CollectSeedPointsFromMeshVertices<FDynamicMesh3>(Mesh, [&](const FVector3d& P) {return P;}, WorldPoints) :
+			CollectSeedPointsFromMeshVertices<FDynamicMesh3>(Mesh, LocalToWorldFunc, WorldPoints);
+	}
+
+	virtual double FastWindingNumber(const FVector3d& P, const FTransformSequence3d& LocalToWorldTransform) override
+	{
+		return bHasBakedTransform ?
+			FWNTree->FastWindingNumber(P) :
+			FWNTree->FastWindingNumber(LocalToWorldTransform.InverseTransformPosition(P));
+	}
+
+	virtual void ProcessVerticesInWorld(TFunctionRef<void(const FVector3d&)> ProcessFunc, const FTransformSequence3d& LocalToWorldTransform) override
+	{
+		if (bHasBakedTransform)
+		{
+			for (FVector3d P : Mesh.VerticesItr())
+			{
+				ProcessFunc(P);
+			}
+		}
+		else
+		{
+			for (FVector3d P : Mesh.VerticesItr())
+			{
+				ProcessFunc(LocalToWorldTransform.TransformPosition(P));
+			}
+		}
+	}
+
+
+	virtual void AppendMesh(FDynamicMesh3& AppendTo, const FTransformSequence3d& TransformSeq) override
+	{
+		FDynamicMeshEditor Editor(&AppendTo);
+		FMeshIndexMappings Mappings;
+		if (bHasBakedTransform)
+		{
+			Editor.AppendMesh(&Mesh, Mappings,
+				[&](int, const FVector3d& Pos) { return Pos; },
+				[&](int, const FVector3d& Normal) { return Normal; });
+		}
+		else
+		{
+			Editor.AppendMesh(&Mesh, Mappings,
+				[&](int, const FVector3d& Pos) { return TransformSeq.TransformPosition(Pos); },
+				[&](int, const FVector3d& Normal) { return TransformSeq.TransformNormal(Normal); });
+		}
+	}
+
+};
+
+
+
+
+
+
+class FStaticMeshDynamicSpatialWrapper : public IMeshSpatialWrapper
+{
+public:
+	UStaticMesh* StaticMesh = nullptr;
+	int32 LODIndex = 0;
+	
+	FMeshDescription* SourceMesh = nullptr;
+
+	TUniquePtr<FDynamicMeshSpatialWrapper> DynamicWrapper;
+
+	virtual bool Build(const FMeshSceneAdapterBuildOptions& BuildOptions) override
+	{
+		return ExtendedBuild(BuildOptions, FTransformSequence3d(), false);
+	}
+
+	virtual bool ExtendedBuild(const FMeshSceneAdapterBuildOptions& BuildOptions, const FTransformSequence3d& TransformSeq, bool bBakeTransform)
+	{
+		check(StaticMesh);
+#if WITH_EDITOR
+		SourceMesh = StaticMesh->GetMeshDescription(LODIndex);
+#else
+		checkf(false, TEXT("Not currently supported - to build at Runtime it is necessary to read from the StaticMesh RenderBuffers"));
+		SourceMesh = nullptr;
+#endif
+		if (SourceMesh)
+		{
+			// does not need this - copypasta
+			TUniquePtr<FMeshDescriptionTriangleMeshAdapter> Adapter = MakeUnique<FMeshDescriptionTriangleMeshAdapter>(SourceMesh);
+
+			FVector3d BuildScale;
+#if WITH_EDITOR
+			// respect BuildScale build setting
+			const FMeshBuildSettings& LODBuildSettings = StaticMesh->GetSourceModel(LODIndex).BuildSettings;
+			BuildScale = (FVector3d)LODBuildSettings.BuildScale3D;
+			Adapter->SetBuildScale(BuildScale, false);
+#endif
+			FDynamicMesh3 ConvertedMesh = FastConvertToDynamicMesh<FMeshDescriptionTriangleMeshAdapter>(*Adapter);
+
+			if (bBakeTransform)
+			{
+				MeshTransforms::ApplyTransform(ConvertedMesh,
+					[&](const FVector3d& P) { return TransformSeq.TransformPosition(P); },
+					[&](const FVector3f& N) { return (FVector3f)TransformSeq.TransformNormal((FVector3d)N); });
+			}
+
+			if (BuildOptions.bThickenThinMeshes)
+			{
+				FFrame3d ThinMeshPlane;
+				bool IsThinMesh = IsThinPlanarMesh<FMeshDescriptionTriangleMeshAdapter>(*Adapter, BuildOptions.DesiredMinThickness, ThinMeshPlane);
+				if (IsThinMesh)
+				{
+					ThickenThinPlanarMesh(ConvertedMesh, BuildOptions.DesiredMinThickness, ThinMeshPlane);
+				}
+			}
+
+			DynamicWrapper = MakeUnique<FDynamicMeshSpatialWrapper>();
+			DynamicWrapper->Mesh = MoveTemp(ConvertedMesh);
+			DynamicWrapper->bHasBakedTransform = bBakeTransform;
+			DynamicWrapper->Build(BuildOptions);
+
+			return true;
+		}
+
+		SourceMesh = nullptr;
+		return false;
+	}
+
+	virtual int32 GetTriangleCount() override
+	{
+		return (SourceMesh) ? DynamicWrapper->GetTriangleCount() : 0;
+	}
+
+	virtual FAxisAlignedBox3d GetWorldBounds(TFunctionRef<FVector3d(const FVector3d&)> LocalToWorldFunc) override
+	{
+		return (SourceMesh) ? DynamicWrapper->GetWorldBounds(LocalToWorldFunc) : FAxisAlignedBox3d::Empty();
+	}
+
+	virtual void CollectSeedPoints(TArray<FVector3d>& WorldPoints, TFunctionRef<FVector3d(const FVector3d&)> LocalToWorldFunc) override
+	{
+		if (SourceMesh)
+		{
+			DynamicWrapper->CollectSeedPoints(WorldPoints, LocalToWorldFunc);
+		}
+	}
+
+	virtual double FastWindingNumber(const FVector3d& P, const FTransformSequence3d& LocalToWorldTransform) override
+	{
+		return (SourceMesh) ? DynamicWrapper->FastWindingNumber(P, LocalToWorldTransform) : 0.0;
+	}
+
+	virtual void ProcessVerticesInWorld(TFunctionRef<void(const FVector3d&)> ProcessFunc, const FTransformSequence3d& LocalToWorldTransform) override
+	{
+		if (SourceMesh)
+		{
+			DynamicWrapper->ProcessVerticesInWorld(ProcessFunc, LocalToWorldTransform);
+		}
+	}
+
+
+	virtual void AppendMesh(FDynamicMesh3& AppendTo, const FTransformSequence3d& TransformSeq) override
+	{
+		if (SourceMesh)
+		{
+			DynamicWrapper->AppendMesh(AppendTo, TransformSeq);
+		}
+	}
+
+};
+
+
+
+
+
 class FStaticMeshSpatialWrapper : public IMeshSpatialWrapper
 {
 public:
@@ -24,11 +494,12 @@ public:
 	int32 LODIndex = 0;
 	
 	FMeshDescription* SourceMesh = nullptr;
+
 	TUniquePtr<FMeshDescriptionTriangleMeshAdapter> Adapter;
 	TUniquePtr<TMeshAABBTree3<FMeshDescriptionTriangleMeshAdapter>> AABBTree;
 	TUniquePtr<TFastWindingTree<FMeshDescriptionTriangleMeshAdapter>> FWNTree;
 
-	virtual bool Build() override
+	virtual bool Build(const FMeshSceneAdapterBuildOptions& BuildOptions) override
 	{
 		check(StaticMesh);
 #if WITH_EDITOR
@@ -41,12 +512,13 @@ public:
 		{
 			Adapter = MakeUnique<FMeshDescriptionTriangleMeshAdapter>(SourceMesh);
 
+			FVector3d BuildScale;
 #if WITH_EDITOR
 			// respect BuildScale build setting
 			const FMeshBuildSettings& LODBuildSettings = StaticMesh->GetSourceModel(LODIndex).BuildSettings;
-			Adapter->SetBuildScale((FVector3d)LODBuildSettings.BuildScale3D, false);
+			BuildScale = (FVector3d)LODBuildSettings.BuildScale3D;
+			Adapter->SetBuildScale(BuildScale, false);
 #endif
-
 			AABBTree = MakeUnique<TMeshAABBTree3<FMeshDescriptionTriangleMeshAdapter>>(Adapter.Get(), true);
 			FWNTree = MakeUnique<TFastWindingTree<FMeshDescriptionTriangleMeshAdapter>>(AABBTree.Get(), true);
 			return true;
@@ -56,34 +528,63 @@ public:
 		return false;
 	}
 
+	virtual int32 GetTriangleCount() override
+	{
+		if (!SourceMesh) return 0;
+		return Adapter->TriangleCount();
+	}
+
+
+	virtual FAxisAlignedBox3d GetWorldBounds(TFunctionRef<FVector3d(const FVector3d&)> LocalToWorldFunc) override
+	{
+		if (!SourceMesh) return FAxisAlignedBox3d::Empty();
+		return GetTransformedVertexBounds<FMeshDescriptionTriangleMeshAdapter>(*Adapter, LocalToWorldFunc);
+	}
 
 	virtual void CollectSeedPoints(TArray<FVector3d>& WorldPoints, TFunctionRef<FVector3d(const FVector3d&)> LocalToWorldFunc) override
 	{
 		if (!SourceMesh) return;
+		CollectSeedPointsFromMeshVertices<FMeshDescriptionTriangleMeshAdapter>(*Adapter, LocalToWorldFunc, WorldPoints);
+	}
+
+	virtual double FastWindingNumber(const FVector3d& P, const FTransformSequence3d& LocalToWorldTransform) override
+	{
+		return (SourceMesh) ? FWNTree->FastWindingNumber(LocalToWorldTransform.InverseTransformPosition(P)) : 0.0;
+	}
 
-		int32 NumVertices = Adapter->VertexCount();
-		int32 LogNumVertices = FMath::Max(1, (int32)FMathd::Ceil(FMathd::Log(NumVertices)));
-		int32 SeedPointCount = (int)(10 * LogNumVertices);
-		SeedPointCount = FMath::Min(SeedPointCount, 5000);
-		int32 Skip = FMath::Max(NumVertices / SeedPointCount, 2);
-		for (int32 k = 0; k < NumVertices; k += Skip)
+	virtual void ProcessVerticesInWorld(TFunctionRef<void(const FVector3d&)> ProcessFunc, const FTransformSequence3d& LocalToWorldTransform) override
+	{
+		int32 NumVertices = (SourceMesh) ? Adapter->VertexCount() : 0;
+		for (int32 vi = 0; vi < NumVertices; ++vi)
 		{
-			WorldPoints.Add(LocalToWorldFunc(Adapter->GetVertex(k)));
+			if (Adapter->IsVertex(vi))
+			{
+				ProcessFunc(LocalToWorldTransform.TransformPosition(Adapter->GetVertex(vi)));
+			}
 		}
 	}
 
-	virtual double FastWindingNumber(const FVector3d& P) override
+	virtual void AppendMesh(FDynamicMesh3& AppendTo, const FTransformSequence3d& TransformSeq) override
 	{
-		return (SourceMesh != nullptr) ? FWNTree->FastWindingNumber(P) : 0.0;
+		if (!SourceMesh) return;
+		FDynamicMeshEditor Editor(&AppendTo);
+		FMeshIndexMappings Mappings;
+		FMeshDescriptionMeshAdapterd AdapterWrapper(*Adapter);
+		Editor.AppendMesh(&AdapterWrapper, Mappings,
+			[&](int, const FVector3d& Pos) { return TransformSeq.TransformPosition(Pos); });
 	}
+
 };
 
 
-static TSharedPtr<IMeshSpatialWrapper> SpatialWrapperFactory(const FMeshTypeContainer& MeshContainer)
+
+
+static TUniquePtr<IMeshSpatialWrapper> SpatialWrapperFactory( const FMeshTypeContainer& MeshContainer )
 {
 	if (MeshContainer.MeshType == ESceneMeshType::StaticMeshAsset)
 	{
-		TSharedPtr<FStaticMeshSpatialWrapper> SMWrapper = MakeShared<FStaticMeshSpatialWrapper>();
+		TUniquePtr<FStaticMeshSpatialWrapper> SMWrapper = MakeUnique<FStaticMeshSpatialWrapper>();
+		SMWrapper->SourceContainer = MeshContainer;
 		SMWrapper->StaticMesh = MeshContainer.GetStaticMesh();
 		if (ensure(SMWrapper->StaticMesh != nullptr))
 		{
@@ -91,7 +592,7 @@ static TSharedPtr<IMeshSpatialWrapper> SpatialWrapperFactory(const FMeshTypeCont
 		}
 	}
 
-	return TSharedPtr<IMeshSpatialWrapper>();
+	return TUniquePtr<IMeshSpatialWrapper>();
 }
 
 
@@ -105,16 +606,16 @@ static void CollectActorChildMeshes(AActor* Actor, UActorComponent* Component, F
 		UStaticMesh* Mesh = StaticMeshComponent->GetStaticMesh();
 		if (Mesh != nullptr)
 		{
-			FActorChildMesh ChildMesh;
-			ChildMesh.SourceComponent = Component;
-			ChildMesh.MeshContainer = FMeshTypeContainer{ Mesh, ESceneMeshType::StaticMeshAsset };
+			TUniquePtr<FActorChildMesh> ChildMesh = MakeUnique<FActorChildMesh>();
+			ChildMesh->SourceComponent = Component;
+			ChildMesh->MeshContainer = FMeshTypeContainer{ Mesh, ESceneMeshType::StaticMeshAsset };
 
 			UInstancedStaticMeshComponent* ISMComponent = Cast<UInstancedStaticMeshComponent>(StaticMeshComponent);
 			if (ISMComponent != nullptr)
 			{
-				// TODO: HISMC
+				// does anything additional need to happen here for HISMC?
 
-				ChildMesh.ComponentType = EActorMeshComponentType::InstancedStaticMesh;
+				ChildMesh->ComponentType = EActorMeshComponentType::InstancedStaticMesh;
 
 				int32 NumInstances = ISMComponent->GetInstanceCount();
 				for (int32 i = 0; i < NumInstances; ++i)
@@ -124,11 +625,14 @@ static void CollectActorChildMeshes(AActor* Actor, UActorComponent* Component, F
 						FTransform InstanceTransform;
 						if (ensure(ISMComponent->GetInstanceTransform(i, InstanceTransform, true)))
 						{
-							FActorChildMesh InstanceChild = ChildMesh;
-							InstanceChild.ComponentIndex = i;
-							InstanceChild.WorldTransform.Append(InstanceTransform);
-							InstanceChild.WorldTransformInverse = InstanceChild.WorldTransform.Inverse();
-							Adapter.ChildMeshes.Add(InstanceChild);
+							TUniquePtr<FActorChildMesh> InstanceChild = MakeUnique<FActorChildMesh>();
+							InstanceChild->SourceComponent = ChildMesh->SourceComponent;
+							InstanceChild->MeshContainer = ChildMesh->MeshContainer;
+							InstanceChild->ComponentType = ChildMesh->ComponentType;
+							InstanceChild->ComponentIndex = i;
+							InstanceChild->WorldTransform.Append(InstanceTransform);
+							InstanceChild->bIsNonUniformScaled = InstanceChild->WorldTransform.HasNonUniformScale();
+							Adapter.ChildMeshes.Add(MoveTemp(InstanceChild));
 						}
 					}
 				}
@@ -137,11 +641,11 @@ static void CollectActorChildMeshes(AActor* Actor, UActorComponent* Component, F
 			else
 			{
 				// base StaticMeshComponent
-				ChildMesh.ComponentType = EActorMeshComponentType::StaticMesh;
-				ChildMesh.ComponentIndex = 0;
-				ChildMesh.WorldTransform.Append(StaticMeshComponent->GetComponentTransform());
-				ChildMesh.WorldTransformInverse = ChildMesh.WorldTransform.Inverse();
-				Adapter.ChildMeshes.Add(ChildMesh);
+				ChildMesh->ComponentType = EActorMeshComponentType::StaticMesh;
+				ChildMesh->ComponentIndex = 0;
+				ChildMesh->WorldTransform.Append(StaticMeshComponent->GetComponentTransform());
+				ChildMesh->bIsNonUniformScaled = ChildMesh->WorldTransform.HasNonUniformScale();
+				Adapter.ChildMeshes.Add(MoveTemp(ChildMesh));
 			}
 		}
 		
@@ -175,58 +679,171 @@ void FMeshSceneAdapter::AddActors(const TArray<AActor*>& ActorsSetIn)
 			}
 		}
 
-		FVector ActorOrigin, ActorExtent;
-		Actor->GetActorBounds(false, ActorOrigin, ActorExtent, true);
-		Adapter->WorldBounds = FAxisAlignedBox3d(
-			(FVector3d)ActorOrigin - (FVector3d)ActorExtent, (FVector3d)ActorOrigin + (FVector3d)ActorExtent);
-
 		SceneActors.Add(MoveTemp(Adapter));
 	}
 
 
 	// Find IMeshSpatialWrapper for each child mesh component. If one does not exist
 	// and we have not seen the underlying unique mesh (eg StaticMesh Asset, etc, construct a new one
-	TArray<IMeshSpatialWrapper*> NewSpatialsToBuild;
 	for (TUniquePtr<FActorAdapter>& Actor : SceneActors)
 	{
-		for (FActorChildMesh& ChildMesh : Actor->ChildMeshes)
+		for (TUniquePtr<FActorChildMesh>& ChildMesh : Actor->ChildMeshes)
 		{
-			void* MeshKey = ChildMesh.MeshContainer.GetMeshKey();
-			TSharedPtr<IMeshSpatialWrapper>* Found = SpatialAdapters.Find(MeshKey);
+			void* MeshKey = ChildMesh->MeshContainer.GetMeshKey();
+			TSharedPtr<FSpatialWrapperInfo>* Found = SpatialAdapters.Find(MeshKey);
+			FSpatialWrapperInfo* MeshInfo = nullptr;
 			if (Found == nullptr)
 			{
-				TSharedPtr<IMeshSpatialWrapper> NewWrapper = SpatialWrapperFactory(ChildMesh.MeshContainer);
-				if (ensure(NewWrapper.IsValid()))
-				{
-					ChildMesh.MeshSpatial = NewWrapper.Get();
-					NewSpatialsToBuild.Add(NewWrapper.Get());
-					SpatialAdapters.Add(MeshKey, NewWrapper);
-				}
+				TSharedPtr<FSpatialWrapperInfo> NewWrapperInfo = MakeShared<FSpatialWrapperInfo>();
+				SpatialAdapters.Add(MeshKey, NewWrapperInfo);
+
+				NewWrapperInfo->SourceContainer = ChildMesh->MeshContainer;
+				NewWrapperInfo->SpatialWrapper = SpatialWrapperFactory(ChildMesh->MeshContainer);
+				MeshInfo = NewWrapperInfo.Get();
 			}
 			else
 			{
-				ChildMesh.MeshSpatial = Found->Get();
+				MeshInfo = (*Found).Get();
 			}
+
+			MeshInfo->ParentMeshes.Add(ChildMesh.Get());
+			if (ChildMesh->bIsNonUniformScaled)
+			{
+				MeshInfo->NonUniformScaleCount++;
+			}
+			ChildMesh->MeshSpatial = MeshInfo->SpatialWrapper.Get();
 		}
 	}
 
-	
+
+	// TODO: compute per-actor world bounds
+	//FVector ActorOrigin, ActorExtent;
+	//Actor->GetActorBounds(false, ActorOrigin, ActorExtent, true);
+	//Adapter->WorldBounds = FAxisAlignedBox3d(
+	//	(FVector3d)ActorOrigin - (FVector3d)ActorExtent, (FVector3d)ActorOrigin + (FVector3d)ActorExtent);
+}
+
+
+
+
+
+void FMeshSceneAdapter::Build(const FMeshSceneAdapterBuildOptions& BuildOptions)
+{
+	TArray<FSpatialWrapperInfo*> ToBuild;
+	for (TPair<void*, TSharedPtr<FSpatialWrapperInfo>> Pair : SpatialAdapters)
+	{
+		ToBuild.Add(Pair.Value.Get());
+	}
+
+	FCriticalSection ListsLock;
+
+	std::atomic<int32> DecomposedSourceMeshCount;
+	DecomposedSourceMeshCount = 0;
+	std::atomic<int32> DecomposedMeshesCount;
+	DecomposedMeshesCount = 0;
+	int32 AddedTrisCount = 0;
+
 	// parallel build of all the spatial data structures
-	ParallelFor(NewSpatialsToBuild.Num(), [&](int32 i)
+	ParallelFor(ToBuild.Num(), [&](int32 i)
 	{
-		bool bOK = NewSpatialsToBuild[i]->Build();
+		FSpatialWrapperInfo* WrapperInfo = ToBuild[i];
+
+		if (WrapperInfo->SourceContainer.MeshType == ESceneMeshType::StaticMeshAsset)
+		{
+			UStaticMesh* StaticMesh = WrapperInfo->SourceContainer.GetStaticMesh();
+			int32 LODIndex = 0;
+			bool bDecompose = CheckIfStaticMeshUsageRequiresDynamicMeshDecomposition(
+				StaticMesh, LODIndex,
+				BuildOptions,
+				WrapperInfo->ParentMeshes, WrapperInfo->NonUniformScaleCount);
+
+			if (bDecompose)
+			{
+				DecomposedSourceMeshCount++;
+
+				TArray<FActorChildMesh*> MeshesToDecompose = WrapperInfo->ParentMeshes;
+				for (FActorChildMesh* MeshInstance : MeshesToDecompose)
+				{
+					TSharedPtr<FSpatialWrapperInfo> NewWrapperInfo = MakeShared<FSpatialWrapperInfo>();
+					NewWrapperInfo->SourceContainer = MeshInstance->MeshContainer;
+
+					TUniquePtr<FStaticMeshDynamicSpatialWrapper> DynamicMeshWrapper = MakeUnique<FStaticMeshDynamicSpatialWrapper>();
+					DynamicMeshWrapper->SourceContainer = NewWrapperInfo->SourceContainer;
+					DynamicMeshWrapper->StaticMesh = StaticMesh;
+					DynamicMeshWrapper->ExtendedBuild(BuildOptions, MeshInstance->WorldTransform, true);
+					NewWrapperInfo->SpatialWrapper = MoveTemp(DynamicMeshWrapper);
+					MeshInstance->MeshSpatial = NewWrapperInfo->SpatialWrapper.Get();
+
+					ListsLock.Lock();
+					void* OldMeshKey = MeshInstance->MeshContainer.GetMeshKey();
+					SpatialAdapters.Remove(OldMeshKey);			// remove existing
+					void* NewMeshKey = (void*)MeshInstance;		// need to key on instance, not unique mesh
+					SpatialAdapters.Add(NewMeshKey, NewWrapperInfo);
+					AddedTrisCount += NewWrapperInfo->SpatialWrapper->GetTriangleCount();
+					ListsLock.Unlock();
+
+					DecomposedMeshesCount++;
+				}
+				return;
+			}
+		}
+
+		TUniquePtr<IMeshSpatialWrapper>& Wrapper = WrapperInfo->SpatialWrapper;
+		bool bOK = Wrapper->Build(BuildOptions);
 		ensure(bOK);	// assumption is that the wrapper will handle failure gracefully
 	});
 
+	if (BuildOptions.bPrintDebugMessages)
+	{
+		UE_LOG(LogTemp, Warning, TEXT("[FMeshSceneAdapter] decomposed %d source meshes into %d unique meshes containing %d triangles"), DecomposedSourceMeshCount.load(), DecomposedMeshesCount.load(), AddedTrisCount)
+	}
+}
+
+
+
+
+
+void FMeshSceneAdapter::GetGeometryStatistics(FStatistics& StatsOut)
+{
+	StatsOut.UniqueMeshCount = 0;
+	StatsOut.UniqueMeshTriangleCount = 0;
+	for (TPair<void*, TSharedPtr<FSpatialWrapperInfo>> Pair : SpatialAdapters)
+	{
+		StatsOut.UniqueMeshCount++;
+		StatsOut.UniqueMeshTriangleCount += Pair.Value->SpatialWrapper->GetTriangleCount();
+	}
+
+	StatsOut.InstanceMeshCount = 0;
+	StatsOut.InstanceMeshTriangleCount = 0;
+	for (const TUniquePtr<FActorAdapter>& Actor : SceneActors)
+	{
+		for (const TUniquePtr<FActorChildMesh>& ChildMesh : Actor->ChildMeshes)
+		{
+			StatsOut.InstanceMeshCount++;
+			if (ChildMesh->MeshSpatial != nullptr)
+			{
+				StatsOut.InstanceMeshTriangleCount += ChildMesh->MeshSpatial->GetTriangleCount();
+			}
+		}
+	}
 }
 
 
 FAxisAlignedBox3d FMeshSceneAdapter::GetBoundingBox()
 {
 	FAxisAlignedBox3d SceneBounds = FAxisAlignedBox3d::Empty();
+
 	for (const TUniquePtr<FActorAdapter>& Actor : SceneActors)
 	{
-		SceneBounds.Contain(Actor->WorldBounds);
+		for (const TUniquePtr<FActorChildMesh>& ChildMesh : Actor->ChildMeshes)
+		{
+			if (ChildMesh->MeshSpatial != nullptr)
+			{
+				FAxisAlignedBox3d ChildBounds = ChildMesh->MeshSpatial->GetWorldBounds(
+					[&](const FVector3d& P) { return ChildMesh->WorldTransform.TransformPosition(P); });
+				SceneBounds.Contain(ChildBounds);
+			}
+		}
 	}
 	return SceneBounds;
 }
@@ -236,12 +853,12 @@ void FMeshSceneAdapter::CollectMeshSeedPoints(TArray<FVector3d>& Points)
 {
 	for (const TUniquePtr<FActorAdapter>& Actor : SceneActors)
 	{
-		for (const FActorChildMesh& ChildMesh : Actor->ChildMeshes)
+		for (const TUniquePtr<FActorChildMesh>& ChildMesh : Actor->ChildMeshes)
 		{
-			if (ChildMesh.MeshSpatial != nullptr)
+			if (ChildMesh->MeshSpatial != nullptr)
 			{
-				ChildMesh.MeshSpatial->CollectSeedPoints(Points,
-					[&](const FVector3d& P) { return ChildMesh.WorldTransform.TransformPosition(P); } );
+				ChildMesh->MeshSpatial->CollectSeedPoints(Points,
+					[&](const FVector3d& P) { return ChildMesh->WorldTransform.TransformPosition(P); } );
 			}
 		}
 	}
@@ -252,18 +869,150 @@ double FMeshSceneAdapter::FastWindingNumber(const FVector3d& P)
 	double SumWinding = 0.0;
 	for (const TUniquePtr<FActorAdapter>& Actor : SceneActors)
 	{
-		if (Actor->WorldBounds.Contains(P))
+		// TODO: cannot skip here because fast winding number for open region will extend...
+		if ( true /**Actor->WorldBounds.Contains(P)*/ )
 		{
-			for (const FActorChildMesh& ChildMesh : Actor->ChildMeshes)
+			for (const TUniquePtr<FActorChildMesh>& ChildMesh : Actor->ChildMeshes)
 			{
-				if (ChildMesh.MeshSpatial != nullptr)
+				if (ChildMesh->MeshSpatial != nullptr)
 				{
-					FVector3d LocalP = ChildMesh.WorldTransformInverse.TransformPosition(P);
-					double MeshWinding = ChildMesh.MeshSpatial->FastWindingNumber(LocalP);
+					double MeshWinding = ChildMesh->MeshSpatial->FastWindingNumber(P, ChildMesh->WorldTransform);
 					SumWinding += MeshWinding;
 				}
 			}
 		}
 	}
 	return SumWinding;
-}
\ No newline at end of file
+}
+
+
+
+
+void FMeshSceneAdapter::GetAccumulatedMesh(FDynamicMesh3& AccumMesh)
+{
+	for (const TUniquePtr<FActorAdapter>& Actor : SceneActors)
+	{
+		for (const TUniquePtr<FActorChildMesh>& ChildMesh : Actor->ChildMeshes)
+		{
+			if (ChildMesh->MeshSpatial != nullptr)
+			{
+				ChildMesh->MeshSpatial->AppendMesh(AccumMesh, ChildMesh->WorldTransform);
+			}
+		}
+	}
+}
+
+
+
+
+void FMeshSceneAdapter::GenerateBaseClosingMesh(double BaseHeight, double ExtrudeHeight)
+{
+	FAxisAlignedBox3d WorldBounds = GetBoundingBox();
+	FInterval1d ZRange(WorldBounds.Min.Z, WorldBounds.Min.Z + BaseHeight);
+
+	TArray<FActorChildMesh*> AllChildMeshes;
+	for (TUniquePtr<FActorAdapter>& Actor : SceneActors)
+	{
+		for (TUniquePtr<FActorChildMesh>& ChildMesh : Actor->ChildMeshes)
+		{
+			if (ChildMesh->MeshSpatial != nullptr)
+			{
+				AllChildMeshes.Add(ChildMesh.Get());
+			}
+		}
+	}
+
+	TArray<FVector2d> WorldHullPoints;
+	FCriticalSection WorldHullPointsLock;
+
+	int32 NumChildren = AllChildMeshes.Num();
+	ParallelFor(NumChildren, [&](int32 ci)
+	{
+		FActorChildMesh* ChildMesh = AllChildMeshes[ci];
+		TArray<FVector2d> LocalHullPoints;
+		ChildMesh->MeshSpatial->ProcessVerticesInWorld([&](const FVector3d& WorldPos)
+		{
+			if (ZRange.Contains(WorldPos.Z))
+			{
+				LocalHullPoints.Add(FVector2d(WorldPos.X, WorldPos.Y));
+			}
+		}, ChildMesh->WorldTransform);
+
+		if (LocalHullPoints.Num() > 0)
+		{
+			FConvexHull2d HullSolver;
+			if (HullSolver.Solve(LocalHullPoints))
+			{
+				WorldHullPointsLock.Lock();
+				for (int32 idx : HullSolver.GetPolygonIndices())
+				{
+					WorldHullPoints.Add(LocalHullPoints[idx]);
+				}
+				WorldHullPointsLock.Unlock();
+			}
+		}
+	});
+
+	FConvexHull2d FinalHullSolver;
+	FinalHullSolver.Solve(WorldHullPoints);
+	FPolygon2d ConvexHullPoly;
+	for (int32 idx : FinalHullSolver.GetPolygonIndices())
+	{
+		ConvexHullPoly.AppendVertex(WorldHullPoints[idx]);
+	}
+
+	FPlanarPolygonMeshGenerator MeshGen;
+	MeshGen.Polygon = ConvexHullPoly;
+	FDynamicMesh3 BasePolygonMesh(&MeshGen.Generate());
+	MeshTransforms::Translate(BasePolygonMesh, ZRange.Min * FVector3d::UnitZ());
+
+	if (ExtrudeHeight == 0)
+	{
+		BasePolygonMesh.ReverseOrientation();		// flip so it points down
+	}
+	else
+	{
+		FOffsetMeshRegion Offset(&BasePolygonMesh);
+		for (int32 tid : BasePolygonMesh.TriangleIndicesItr())
+		{
+			Offset.Triangles.Add(tid);
+		}
+		Offset.bUseFaceNormals = true;
+		Offset.DefaultOffsetDistance = ExtrudeHeight;
+		Offset.bIsPositiveOffset = (ExtrudeHeight > 0);
+		Offset.Apply();
+	}
+
+	//
+	// append a fake actor/mesh
+	//
+
+	TUniquePtr<FActorAdapter> Adapter = MakeUnique<FActorAdapter>();
+	Adapter->SourceActor = nullptr;
+
+	TUniquePtr<FActorChildMesh> ChildMesh = MakeUnique<FActorChildMesh>();
+	ChildMesh->SourceComponent = nullptr;
+	//InstanceChild->MeshContainer = ;
+	ChildMesh->ComponentType = EActorMeshComponentType::InternallyGeneratedComponent;
+	ChildMesh->ComponentIndex = 0;
+	//ChildMesh->WorldTransform.Append(InstanceTransform);
+	ChildMesh->bIsNonUniformScaled = false;
+
+	TSharedPtr<FSpatialWrapperInfo> NewWrapperInfo = MakeShared<FSpatialWrapperInfo>();
+	SpatialAdapters.Add(ChildMesh.Get(), NewWrapperInfo);
+	TUniquePtr<FDynamicMeshSpatialWrapper> DynamicMeshWrapper = MakeUnique<FDynamicMeshSpatialWrapper>();
+	//DynamicMeshWrapper->SourceContainer = NewWrapperInfo->SourceContainer;
+	DynamicMeshWrapper->Mesh = MoveTemp(BasePolygonMesh);
+	DynamicMeshWrapper->bHasBakedTransform = true;
+	FMeshSceneAdapterBuildOptions UseBuildOptions;
+	DynamicMeshWrapper->Build(UseBuildOptions);
+
+	//NewWrapperInfo->SourceContainer = ChildMesh->MeshContainer;
+	NewWrapperInfo->SpatialWrapper = MoveTemp(DynamicMeshWrapper);
+	NewWrapperInfo->ParentMeshes.Add(ChildMesh.Get());
+	ChildMesh->MeshSpatial = NewWrapperInfo->SpatialWrapper.Get();
+	Adapter->ChildMeshes.Add(MoveTemp(ChildMesh));
+	SceneActors.Add(MoveTemp(Adapter));
+
+}
+
diff --git a/Engine/Plugins/Experimental/MeshModelingToolset/Source/ModelingComponents/Private/Scene/SceneCapturePhotoSet.cpp b/Engine/Plugins/Experimental/MeshModelingToolset/Source/ModelingComponents/Private/Scene/SceneCapturePhotoSet.cpp
--- a/Engine/Plugins/Experimental/MeshModelingToolset/Source/ModelingComponents/Private/Scene/SceneCapturePhotoSet.cpp
+++ b/Engine/Plugins/Experimental/MeshModelingToolset/Source/ModelingComponents/Private/Scene/SceneCapturePhotoSet.cpp
@@ -107,42 +107,57 @@ void FSceneCapturePhotoSet::AddExteriorCaptures(
 		ViewFrame.Origin = (FVector3d)RenderSphere.Center;
 		ViewFrame.Origin -= (double)RenderSphere.W * ViewFrame.X();
 
-		FSpatialPhoto4f BasePhoto;
-		BasePhoto.Frame = ViewFrame;
-		BasePhoto.NearPlaneDist = NearPlaneDist;
-		BasePhoto.HorzFOVDegrees = HorizontalFOVDegrees;
-		BasePhoto.Dimensions = PhotoDimensions;
+		FSpatialPhoto3f BasePhoto3f;
+		BasePhoto3f.Frame = ViewFrame;
+		BasePhoto3f.NearPlaneDist = NearPlaneDist;
+		BasePhoto3f.HorzFOVDegrees = HorizontalFOVDegrees;
+		BasePhoto3f.Dimensions = PhotoDimensions;
 
-		auto CaptureImageTypeFunc = [&BasePhoto, &RenderCapture](ERenderCaptureType CaptureType, FSpatialPhotoSet4f& PhotoSet)
+		auto CaptureImageTypeFunc_3f = [&BasePhoto3f, &RenderCapture](ERenderCaptureType CaptureType, FSpatialPhotoSet3f& PhotoSet)
 		{
-			FSpatialPhoto4f NewPhoto = BasePhoto;
-			RenderCapture.CaptureFromPosition(CaptureType, NewPhoto.Frame, NewPhoto.HorzFOVDegrees, NewPhoto.NearPlaneDist, NewPhoto.Image);
+			FSpatialPhoto3f NewPhoto = BasePhoto3f;
+			FImageAdapter Image(&NewPhoto.Image);
+			RenderCapture.CaptureFromPosition(CaptureType, NewPhoto.Frame, NewPhoto.HorzFOVDegrees, NewPhoto.NearPlaneDist, Image);
+			PhotoSet.Add(MoveTemp(NewPhoto));
+		};
+
+		FSpatialPhoto1f BasePhoto1f;
+		BasePhoto1f.Frame = ViewFrame;
+		BasePhoto1f.NearPlaneDist = NearPlaneDist;
+		BasePhoto1f.HorzFOVDegrees = HorizontalFOVDegrees;
+		BasePhoto1f.Dimensions = PhotoDimensions;
+
+		auto CaptureImageTypeFunc_1f = [&BasePhoto1f, &RenderCapture](ERenderCaptureType CaptureType, FSpatialPhotoSet1f& PhotoSet)
+		{
+			FSpatialPhoto1f NewPhoto = BasePhoto1f;
+			FImageAdapter Image(&NewPhoto.Image);
+			RenderCapture.CaptureFromPosition(CaptureType, NewPhoto.Frame, NewPhoto.HorzFOVDegrees, NewPhoto.NearPlaneDist, Image);
 			PhotoSet.Add(MoveTemp(NewPhoto));
 		};
 
 		if (bEnableBaseColor)
 		{
-			CaptureImageTypeFunc(ERenderCaptureType::BaseColor, BaseColorPhotoSet);
+			CaptureImageTypeFunc_3f(ERenderCaptureType::BaseColor, BaseColorPhotoSet);
 		}
 		if (bEnableRoughness)
 		{
-			CaptureImageTypeFunc(ERenderCaptureType::Roughness, RoughnessPhotoSet);
+			CaptureImageTypeFunc_1f(ERenderCaptureType::Roughness, RoughnessPhotoSet);
 		}
 		if (bEnableSpecular)
 		{
-			CaptureImageTypeFunc(ERenderCaptureType::Specular, SpecularPhotoSet);
+			CaptureImageTypeFunc_1f(ERenderCaptureType::Specular, SpecularPhotoSet);
 		}
 		if (bEnableMetallic)
 		{
-			CaptureImageTypeFunc(ERenderCaptureType::Metallic, MetallicPhotoSet);
+			CaptureImageTypeFunc_1f(ERenderCaptureType::Metallic, MetallicPhotoSet);
 		}
 		if (bEnableWorldNormal)
 		{
-			CaptureImageTypeFunc(ERenderCaptureType::WorldNormal, WorldNormalPhotoSet);
+			CaptureImageTypeFunc_3f(ERenderCaptureType::WorldNormal, WorldNormalPhotoSet);
 		}
 		if (bEnableEmissive)
 		{
-			CaptureImageTypeFunc(ERenderCaptureType::Emissive, EmissivePhotoSet);
+			CaptureImageTypeFunc_3f(ERenderCaptureType::Emissive, EmissivePhotoSet);
 		}
 	}
 
@@ -163,15 +178,15 @@ void FSceneCapturePhotoSet::OptimizePhotoSets()
 FSceneCapturePhotoSet::FSceneSample::FSceneSample()
 {
 	HaveValues = FRenderCaptureTypeFlags::None();
-	BaseColor = FVector4f(0, 0, 0, 1);
-	Roughness = FVector4f(0, 0, 0, 1);
-	Specular = FVector4f(0, 0, 0, 1);
-	Metallic = FVector4f(0, 0, 0, 1);
-	Emissive = FVector4f(0, 0, 0, 1);
-	WorldNormal = FVector4f(0, 0, 1, 1);
+	BaseColor = FVector3f(0, 0, 0);
+	Roughness = 0.0f;
+	Specular = 0.0f;
+	Metallic = 0.0f;
+	Emissive = FVector3f(0, 0, 0);
+	WorldNormal = FVector3f(0, 0, 1);
 }
 
-FVector4f FSceneCapturePhotoSet::FSceneSample::GetValue(ERenderCaptureType CaptureType) const
+FVector3f FSceneCapturePhotoSet::FSceneSample::GetValue3f(ERenderCaptureType CaptureType) const
 {
 	switch (CaptureType)
 	{
@@ -180,20 +195,43 @@ FVector4f FSceneCapturePhotoSet::FSceneSample::GetValue(ERenderCaptureType Captu
 	case ERenderCaptureType::WorldNormal:
 		return WorldNormal;
 	case ERenderCaptureType::Roughness:
-		return Roughness;
+		return Roughness * FVector3f::One();
 	case ERenderCaptureType::Metallic:
-		return Metallic;
+		return Metallic * FVector3f::One();
 	case ERenderCaptureType::Specular:
-		return Specular;
+		return Specular * FVector3f::One();
 	case ERenderCaptureType::Emissive:
 		return Emissive;
 	default:
 		check(false);
 	}
+	return FVector3f::Zero();
+}
+
+FVector4f FSceneCapturePhotoSet::FSceneSample::GetValue4f(ERenderCaptureType CaptureType) const
+{
+	switch (CaptureType)
+	{
+	case ERenderCaptureType::BaseColor:
+		return FVector4f(BaseColor.X, BaseColor.Y, BaseColor.Z, 1.0f);
+	case ERenderCaptureType::WorldNormal:
+		return FVector4f(WorldNormal.X, WorldNormal.Y, WorldNormal.Z, 1.0f);
+	case ERenderCaptureType::Roughness:
+		return FVector4f(Roughness, Roughness, Roughness, 1.0f);
+	case ERenderCaptureType::Metallic:
+		return FVector4f(Metallic, Metallic, Metallic, 1.0f);
+	case ERenderCaptureType::Specular:
+		return FVector4f(Specular, Specular, Specular, 1.0f);
+	case ERenderCaptureType::Emissive:
+		return FVector4f(Emissive.X, Emissive.Y, Emissive.Z, 1.0f);
+	default:
+		check(false);
+	}
 	return FVector4f::Zero();
 }
 
 
+
 bool FSceneCapturePhotoSet::ComputeSample(
 	const FRenderCaptureTypeFlags& SampleChannels,
 	const FVector3d& Position,
diff --git a/Engine/Plugins/Experimental/MeshModelingToolset/Source/ModelingComponents/Private/Scene/WorldRenderCapture.cpp b/Engine/Plugins/Experimental/MeshModelingToolset/Source/ModelingComponents/Private/Scene/WorldRenderCapture.cpp
--- a/Engine/Plugins/Experimental/MeshModelingToolset/Source/ModelingComponents/Private/Scene/WorldRenderCapture.cpp
+++ b/Engine/Plugins/Experimental/MeshModelingToolset/Source/ModelingComponents/Private/Scene/WorldRenderCapture.cpp
@@ -224,7 +224,7 @@ bool FWorldRenderCapture::CaptureEmissiveFromPosition(
 	const FFrame3d& ViewFrame,
 	double HorzFOVDegrees,
 	double NearPlaneDist,
-	TImageBuilder<FVector4f>& ResultImageOut)
+	FImageAdapter& ResultImageOut)
 {
 	// This function is a combination of CaptureFromPosition() and RenderSceneVisualizationToTexture() below,
 	// that attempts to capture Emissive. The difficulty with capturing Emissive is there is no Visualization
@@ -457,7 +457,7 @@ bool FWorldRenderCapture::CaptureFromPosition(
 	const FFrame3d& ViewFrame,
 	double HorzFOVDegrees,
 	double NearPlaneDist,
-	TImageBuilder<FVector4f>& ResultImageOut)
+	FImageAdapter& ResultImageOut)
 {
 	if (CaptureType == ERenderCaptureType::Emissive)
 	{
diff --git a/Engine/Plugins/Experimental/MeshModelingToolset/Source/ModelingComponents/Public/Scene/MeshSceneAdapter.h b/Engine/Plugins/Experimental/MeshModelingToolset/Source/ModelingComponents/Public/Scene/MeshSceneAdapter.h
--- a/Engine/Plugins/Experimental/MeshModelingToolset/Source/ModelingComponents/Public/Scene/MeshSceneAdapter.h
+++ b/Engine/Plugins/Experimental/MeshModelingToolset/Source/ModelingComponents/Public/Scene/MeshSceneAdapter.h
@@ -7,6 +7,7 @@
 #include "BoxTypes.h"
 #include "TransformTypes.h"
 #include "TransformSequence.h"
+#include "DynamicMesh3.h"
 
 class AActor;
 class UActorComponent;
@@ -18,25 +19,6 @@ namespace Geometry
 {
 
 
-/**
- * Abstract interface to a spatial data structure for a mesh
- */
-class MODELINGCOMPONENTS_API IMeshSpatialWrapper
-{
-public:
-	virtual ~IMeshSpatialWrapper() {}
-
-	/** If possible, spatial data structure should defer construction until this function, which will be called off the game thread (in ParallelFor) */
-	virtual bool Build() = 0;
-
-	/** Calculate the mesh winding number at the given Position. Must be callable in parallel from any thread.  */
-	virtual double FastWindingNumber(const FVector3d& P) = 0;
-
-	/** Collect a set of seed points from this Mesh, mapped through LocalToWorldFunc to world space. Must be callable in parallel from any thread. */
-	virtual void CollectSeedPoints(TArray<FVector3d>& WorldPoints, TFunctionRef<FVector3d(const FVector3d&)> LocalToWorldFunc) = 0;
-};
-
-
 /**
  * ESceneMeshType is used to indicate which type of Mesh a FMeshTypeContainer contains.
  */
@@ -75,6 +57,57 @@ struct MODELINGCOMPONENTS_API FMeshTypeContainer
 };
 
 
+
+/**
+ * Configuration for FMeshSceneAdapter::Build()
+ */
+struct FMeshSceneAdapterBuildOptions
+{
+	bool bPrintDebugMessages = false;
+
+	/** If true, find approximately-planar meshes with a main dimension below DesiredMinThickness and thicken them to DesiredMinThickness  */
+	bool bThickenThinMeshes = false;
+	/** Thickness used for bThickenThinMeshes processing */
+	double DesiredMinThickness = 0.1;
+};
+
+
+
+
+
+/**
+ * Abstract interface to a spatial data structure for a mesh
+ */
+class MODELINGCOMPONENTS_API IMeshSpatialWrapper
+{
+public:
+	virtual ~IMeshSpatialWrapper() {}
+
+	FMeshTypeContainer SourceContainer;
+
+	/** If possible, spatial data structure should defer construction until this function, which will be called off the game thread (in ParallelFor) */
+	virtual bool Build(const FMeshSceneAdapterBuildOptions& BuildOptions) = 0;
+
+	/*** @return triangle count for this mesh */
+	virtual int32 GetTriangleCount() = 0;
+
+	/** Calculate bounding box for this Mesh */
+	virtual FAxisAlignedBox3d GetWorldBounds(TFunctionRef<FVector3d(const FVector3d&)> LocalToWorldFunc) = 0;
+
+	/** Calculate the mesh winding number at the given Position. Must be callable in parallel from any thread.  */
+	virtual double FastWindingNumber(const FVector3d& P, const FTransformSequence3d& LocalToWorldTransform) = 0;
+
+	/** Collect a set of seed points from this Mesh, mapped through LocalToWorldFunc to world space. Must be callable in parallel from any thread. */
+	virtual void CollectSeedPoints(TArray<FVector3d>& WorldPoints, TFunctionRef<FVector3d(const FVector3d&)> LocalToWorldFunc) = 0;
+
+	/** apply ProcessFunc to each vertex in world space */
+	virtual void ProcessVerticesInWorld(TFunctionRef<void(const FVector3d&)> ProcessFunc, const FTransformSequence3d& LocalToWorldTransform) = 0;
+
+	/** Append the geometry represented by this wrapper to the accumulated AppendTo mesh, under the given world transform */
+	virtual void AppendMesh(FDynamicMesh3& AppendTo, const FTransformSequence3d& TransformSeq) = 0;
+};
+
+
 /**
  * EActorMeshComponentType enum is used to determine which type of Component 
  * an FActorChildMesh represents.
@@ -85,6 +118,8 @@ enum class EActorMeshComponentType
 	InstancedStaticMesh,
 	HierarchicalInstancedStaticMesh,
 
+	InternallyGeneratedComponent,
+
 	Unknown
 };
 
@@ -97,6 +132,10 @@ enum class EActorMeshComponentType
 struct MODELINGCOMPONENTS_API FActorChildMesh
 {
 public:
+	FActorChildMesh() {}
+	FActorChildMesh(const FActorChildMesh&) = delete;		// must delete to allow TArray<TUniquePtr> in FActorAdapter
+	FActorChildMesh(FActorChildMesh&&) = delete;
+
 	/** the Component this Mesh was generated from, if there is one. */
 	UActorComponent* SourceComponent = nullptr;
 	/** Type of SourceComponent, if known */
@@ -107,8 +146,8 @@ public:
 	/** Wrapper around the Mesh this FActorChildMesh refers to (eg from a StaticMeshAsset, etc) */
 	FMeshTypeContainer MeshContainer;
 	/** Local-to-World transformation of the Mesh in the MeshContainer */
-	UE::Geometry::FTransformSequence3d WorldTransform;
-	UE::Geometry::FTransformSequence3d WorldTransformInverse;
+	FTransformSequence3d WorldTransform;
+	bool bIsNonUniformScaled = false;
 
 	/** Spatial data structure that represents the Mesh in MeshContainer - assumption is this is owned externally */
 	IMeshSpatialWrapper* MeshSpatial;
@@ -130,9 +169,9 @@ public:
 	// the AActor this Adapter represents
 	AActor* SourceActor = nullptr;
 	// set of child Meshes with transforms
-	TArray<FActorChildMesh> ChildMeshes;
+	TArray<TUniquePtr<FActorChildMesh>> ChildMeshes;
 	// World-space bounds of this Actor (meshes)
-	UE::Geometry::FAxisAlignedBox3d WorldBounds;
+	FAxisAlignedBox3d WorldBounds;
 };
 
 
@@ -155,11 +194,39 @@ public:
 	FMeshSceneAdapter(const FMeshSceneAdapter&) = delete;		// must delete this due to TArray<TUniquePtr> member
 	FMeshSceneAdapter(FMeshSceneAdapter&&) = delete;
 
-	/** Add the  given Actors to our Actor set */
+	/** Add the given Actors to our Actor set. */
 	void AddActors(const TArray<AActor*>& ActorsSetIn);
 
+	/** Build */
+	void Build(const FMeshSceneAdapterBuildOptions& BuildOptions);
+
+	/**
+	 * Generate a new mesh that "caps" the mesh scene on the bottom. This can be used in cases where
+	 * the geometry is open on the base, to fill in the hole, which allows things like mesh-solidification
+	 * to work better. The base mesh is a polygon which can be optionally extruded.
+	 * Currently the closing mesh is a convex hull, todo: implement a better option
+	 * @param BaseHeight the height in world units from the bounding-box MinZ to consider as part of the "base".
+	 * @param ExtrudeHeight height in world units to extrude the generated base. Positive is in +Z direction. If zero, an open-boundary polygon is generated instead.
+	 */
+	void GenerateBaseClosingMesh(double BaseHeight = 1.0, double ExtrudeHeight = 0.0);
+
+	/**
+	 * Statistics about the mesh scene returned by GetGeometryStatistics()
+	 */
+	struct FStatistics
+	{
+		int32 UniqueMeshCount = 0;
+		int32 UniqueMeshTriangleCount = 0;
+
+		int32 InstanceMeshCount = 0;
+		int32 InstanceMeshTriangleCount = 0;
+	};
+
+	/** @return bounding box for the Actor set */
+	virtual void GetGeometryStatistics(FStatistics& StatsOut);
+
 	/** @return bounding box for the Actor set */
-	virtual UE::Geometry::FAxisAlignedBox3d GetBoundingBox();
+	virtual FAxisAlignedBox3d GetBoundingBox();
 
 	/** @return a set of points on the surface of the meshes, can be used to initialize the MarchingCubes mesher */
 	virtual void CollectMeshSeedPoints(TArray<FVector3d>& PointsOut);
@@ -167,13 +234,27 @@ public:
 	/** @return FastWindingNumber computed across all mesh Actors/Components */
 	virtual double FastWindingNumber(const FVector3d& P);
 
+	/** Append all instance triangles to a single mesh. May be very large. */
+	virtual void GetAccumulatedMesh(FDynamicMesh3& AccumMesh);
+
 protected:
 
 	// top-level list of ActorAdapters, which represent each Actor and set of Components
 	TArray<TUniquePtr<FActorAdapter>> SceneActors;
 
+	struct FSpatialWrapperInfo
+	{
+		FMeshTypeContainer SourceContainer;
+		TArray<FActorChildMesh*> ParentMeshes;
+		int32 NonUniformScaleCount = 0;
+		TUniquePtr<IMeshSpatialWrapper> SpatialWrapper;
+	};
+
 	// Unique set of spatial data structure query interfaces, one for each Mesh object, which is identified by void* pointer
-	TMap<void*, TSharedPtr<IMeshSpatialWrapper>> SpatialAdapters;
+	TMap<void*, TSharedPtr<FSpatialWrapperInfo>> SpatialAdapters;
+
+	bool bEnableClipPlane = false;
+	FFrame3d ClipPlane;
 };
 
 
diff --git a/Engine/Plugins/Experimental/MeshModelingToolset/Source/ModelingComponents/Public/Scene/SceneCapturePhotoSet.h b/Engine/Plugins/Experimental/MeshModelingToolset/Source/ModelingComponents/Public/Scene/SceneCapturePhotoSet.h
--- a/Engine/Plugins/Experimental/MeshModelingToolset/Source/ModelingComponents/Public/Scene/SceneCapturePhotoSet.h
+++ b/Engine/Plugins/Experimental/MeshModelingToolset/Source/ModelingComponents/Public/Scene/SceneCapturePhotoSet.h
@@ -83,17 +83,18 @@ public:
 	struct MODELINGCOMPONENTS_API FSceneSample
 	{
 		FRenderCaptureTypeFlags HaveValues;		// defines which channels have non-default values
-		FVector4f BaseColor;
-		FVector4f Roughness;
-		FVector4f Specular;
-		FVector4f Metallic;
-		FVector4f Emissive;
-		FVector4f WorldNormal;
+		FVector3f BaseColor;
+		float Roughness;
+		float Specular;
+		float Metallic;
+		FVector3f Emissive;
+		FVector3f WorldNormal;
 
 		FSceneSample();
 
 		/** @return value for the given captured channel, or default value */
-		FVector4f GetValue(ERenderCaptureType CaptureType) const;
+		FVector3f GetValue3f(ERenderCaptureType CaptureType) const;
+		FVector4f GetValue4f(ERenderCaptureType CaptureType) const;
 	};
 
 
@@ -122,12 +123,12 @@ protected:
 	bool bEnableWorldNormal = true;
 	bool bEnableEmissive = true;
 
-	FSpatialPhotoSet4f BaseColorPhotoSet;
-	FSpatialPhotoSet4f RoughnessPhotoSet;
-	FSpatialPhotoSet4f SpecularPhotoSet;
-	FSpatialPhotoSet4f MetallicPhotoSet;
-	FSpatialPhotoSet4f WorldNormalPhotoSet;
-	FSpatialPhotoSet4f EmissivePhotoSet;
+	FSpatialPhotoSet3f BaseColorPhotoSet;
+	FSpatialPhotoSet1f RoughnessPhotoSet;
+	FSpatialPhotoSet1f SpecularPhotoSet;
+	FSpatialPhotoSet1f MetallicPhotoSet;
+	FSpatialPhotoSet3f WorldNormalPhotoSet;
+	FSpatialPhotoSet3f EmissivePhotoSet;
 };
 
 
diff --git a/Engine/Plugins/Experimental/MeshModelingToolset/Source/ModelingComponents/Public/Scene/WorldRenderCapture.h b/Engine/Plugins/Experimental/MeshModelingToolset/Source/ModelingComponents/Public/Scene/WorldRenderCapture.h
--- a/Engine/Plugins/Experimental/MeshModelingToolset/Source/ModelingComponents/Public/Scene/WorldRenderCapture.h
+++ b/Engine/Plugins/Experimental/MeshModelingToolset/Source/ModelingComponents/Public/Scene/WorldRenderCapture.h
@@ -106,7 +106,7 @@ public:
 		const FFrame3d& ViewFrame,
 		double HorzFOVDegrees,
 		double NearPlaneDist,
-		TImageBuilder<FVector4f>& ResultImageOut);
+		FImageAdapter& ResultImageOut);
 
 protected:
 	UWorld* World;
@@ -131,7 +131,7 @@ protected:
 		const FFrame3d& Frame,
 		double HorzFOVDegrees,
 		double NearPlaneDist,
-		TImageBuilder<FVector4f>& ResultImageOut);
+		FImageAdapter& ResultImageOut);
 };
 
 
diff --git a/Engine/Source/Developer/GeometryProcessingInterfaces/Public/GeometryProcessingInterfaces/ApproximateActors.h b/Engine/Source/Developer/GeometryProcessingInterfaces/Public/GeometryProcessingInterfaces/ApproximateActors.h
--- a/Engine/Source/Developer/GeometryProcessingInterfaces/Public/GeometryProcessingInterfaces/ApproximateActors.h
+++ b/Engine/Source/Developer/GeometryProcessingInterfaces/Public/GeometryProcessingInterfaces/ApproximateActors.h
@@ -8,6 +8,7 @@
 class AActor;
 class UStaticMesh;
 class UMaterial;
+class UMaterialInterface;
 class UTexture2D;
 
 
@@ -36,11 +37,24 @@ public:
 		CollisionMesh
 	};
 
+	enum class EBaseCappingPolicy
+	{
+		NoBaseCapping = 0,
+		ConvexPolygon = 1,
+		ConvexSolid = 2
+	};
+
+	enum class EOcclusionPolicy : uint8
+	{
+		None = 0,
+		VisibilityBased = 1
+	};
 
 	enum class ESimplificationPolicy
 	{
 		FixedTriangleCount = 0,
-		TrianglesPerUnitSqMeter = 1
+		TrianglesPerUnitSqMeter = 1,
+		GeometricTolerance = 2
 	};
 
 
@@ -55,22 +69,33 @@ public:
 		// high-level control of the overall approximation process
 		EApproximationPolicy BasePolicy = EApproximationPolicy::MeshAndGeneratedMaterial;
 
+		//
+		// Mesh Preprocessing settings
+		//
+		bool bAutoThickenThinParts = false;
+		double AutoThickenThicknessMeters = 0.1;
+
+		EBaseCappingPolicy BaseCappingPolicy = EBaseCappingPolicy::NoBaseCapping;
+		double BaseThicknessOverrideMeters = 0.0;		// use this thickness for Solid. If zero, use AutoThickenThickness if bAutoThickenThinParts=true, otherwise use WorldSpaceApproximationAccuracyMeters
+		double BaseHeightOverrideMeters = 0.0;		// consider this height from MinZ as "base" region. If zero, use 2.0*WorldSpaceApproximationAccuracyMeters
+
 		//
 		// Shape Approximation settings
 		//
 
 		// Meshing settings (ie for voxelization)
-		float WorldSpaceApproximationAccuracyMeters = 1.0f;
-		int32 ClampVoxelDimension = 512;
+		double WorldSpaceApproximationAccuracyMeters = 1.0;
+		int32 ClampVoxelDimension = 1024;
 
-		float WindingThreshold = 0.5;
+		double WindingThreshold = 0.5;
 
 		bool bApplyMorphology = false;
-		float MorphologyDistanceMeters = 0.1;
+		double MorphologyDistanceMeters = 0.1;
 
+		EOcclusionPolicy OcclusionPolicy = EOcclusionPolicy::VisibilityBased;
 		ESimplificationPolicy MeshSimplificationPolicy = ESimplificationPolicy::FixedTriangleCount;
 		int32 FixedTriangleCount = 5000;
-		float SimplificationTargetMetric = 0.1;		// interpretation varies depending on MeshSimplificationPolicy
+		double SimplificationTargetMetric = 0.1;		// interpretation varies depending on MeshSimplificationPolicy
 
 		//
 		// Material approximation settings
@@ -78,14 +103,29 @@ public:
 		int32 RenderCaptureImageSize = 1024;
 
 		// render capture parameters
-		float FieldOfViewDegrees = 45.0f;
-		float NearPlaneDist = 1.0f;
+		double FieldOfViewDegrees = 45.0;
+		double NearPlaneDist = 1.0;
 
 
 		//
 		// Material output settings
 		//
 
+		// A new MIC derived from this material will be created and assigned to the generated mesh
+		UMaterialInterface* BakeMaterial = nullptr;		// if null, will use /MeshModelingToolset/Materials/FullMaterialBakePreviewMaterial instead
+		FName BaseColorTexParamName = FName("BaseColor");
+		bool bBakeBaseColor = true;
+		FName RoughnessTexParamName = FName("Roughness");
+		bool bBakeRoughness = true;
+		FName MetallicTexParamName = FName("Metallic");
+		bool bBakeMetallic = true;
+		FName SpecularTexParamName = FName("Specular");
+		bool bBakeSpecular = true;
+		FName EmissiveTexParamName = FName("Emissive");
+		bool bBakeEmissive = true;
+		FName NormalTexParamName = FName("NormalMap");
+		bool bBakeNormalMap = true;
+
 		// output texture options
 		int32 TextureImageSize = 1024;
 
@@ -93,6 +133,16 @@ public:
 		int32 AntiAliasMultiSampling = 0;
 
 
+		//
+		// Debug settings
+		//
+
+		// print useful information to the Output Log 
+		bool bVerbose = false;
+
+		// create a flattened (ie non-instanced) mesh and save it with _DEBUG suffix. Warning often absolutely enormous!
+		bool bWriteDebugMesh = false;
+
 
 		// todo
 		//   - which texture types to generate (BaseColor, Specular, WorldNormal, etc)
@@ -117,7 +167,7 @@ public:
 
 		TArray<UStaticMesh*> NewMeshAssets;
 
-		TArray<UMaterial*> NewMaterials;
+		TArray<UMaterialInterface*> NewMaterials;
 		
 		TArray<UTexture2D*> NewTextures;
 
diff --git a/Engine/Source/Editor/MergeActors/Private/MeshApproximationTool/MeshApproximationTool.cpp b/Engine/Source/Editor/MergeActors/Private/MeshApproximationTool/MeshApproximationTool.cpp
--- a/Engine/Source/Editor/MergeActors/Private/MeshApproximationTool/MeshApproximationTool.cpp
+++ b/Engine/Source/Editor/MergeActors/Private/MeshApproximationTool/MeshApproximationTool.cpp
@@ -168,17 +168,42 @@ bool FMeshApproximationTool::RunMerge(const FString& PackageName, const TArray<T
 
 	IGeometryProcessing_ApproximateActors::FOptions Options;
 	Options.BasePackagePath = PackageName;
+	Options.BasePolicy = (UseSettings.OutputType == EMeshApproximationType::MeshShapeOnly) ?
+		IGeometryProcessing_ApproximateActors::EApproximationPolicy::CollisionMesh :
+		IGeometryProcessing_ApproximateActors::EApproximationPolicy::MeshAndGeneratedMaterial;
 	Options.WorldSpaceApproximationAccuracyMeters = UseSettings.ApproximationAccuracy;
+
+	Options.bAutoThickenThinParts = UseSettings.bAttemptAutoThickening;
+	Options.AutoThickenThicknessMeters = UseSettings.TargetMinThicknessMultiplier * UseSettings.ApproximationAccuracy;
+
+	Options.BaseCappingPolicy = IGeometryProcessing_ApproximateActors::EBaseCappingPolicy::NoBaseCapping;
+	if (UseSettings.BaseCapping == EMeshApproximationBaseCappingType::ConvexPolygon)
+	{
+		Options.BaseCappingPolicy = IGeometryProcessing_ApproximateActors::EBaseCappingPolicy::ConvexPolygon;
+	}
+	else if (UseSettings.BaseCapping == EMeshApproximationBaseCappingType::ConvexSolid)
+	{
+		Options.BaseCappingPolicy = IGeometryProcessing_ApproximateActors::EBaseCappingPolicy::ConvexSolid;
+	}
+
 	Options.ClampVoxelDimension = UseSettings.ClampVoxelDimension;
 	Options.WindingThreshold = UseSettings.WindingThreshold;
 	Options.bApplyMorphology = UseSettings.bFillGaps;
 	Options.MorphologyDistanceMeters = UseSettings.GapDistance;
+
+	Options.OcclusionPolicy = (UseSettings.OcclusionMethod == EOccludedGeometryFilteringPolicy::VisibilityBasedFiltering) ?
+		IGeometryProcessing_ApproximateActors::EOcclusionPolicy::VisibilityBased : IGeometryProcessing_ApproximateActors::EOcclusionPolicy::None;
 	Options.FixedTriangleCount = UseSettings.TargetTriCount;
 	if (UseSettings.SimplifyMethod == EMeshApproximationSimplificationPolicy::TrianglesPerArea)
 	{
 		Options.MeshSimplificationPolicy = IGeometryProcessing_ApproximateActors::ESimplificationPolicy::TrianglesPerUnitSqMeter;
 		Options.SimplificationTargetMetric = UseSettings.TrianglesPerM;
 	}
+	else if (UseSettings.SimplifyMethod == EMeshApproximationSimplificationPolicy::GeometricTolerance)
+	{
+		Options.MeshSimplificationPolicy = IGeometryProcessing_ApproximateActors::ESimplificationPolicy::GeometricTolerance;
+		Options.SimplificationTargetMetric = UseSettings.GeometricDeviation;
+	}
 	else
 	{
 		Options.MeshSimplificationPolicy = IGeometryProcessing_ApproximateActors::ESimplificationPolicy::FixedTriangleCount;
@@ -192,6 +217,8 @@ bool FMeshApproximationTool::RunMerge(const FString& PackageName, const TArray<T
 	Options.FieldOfViewDegrees = UseSettings.CaptureFieldOfView;
 	Options.NearPlaneDist = UseSettings.NearPlaneDist;
 
+	Options.bVerbose = UseSettings.bPrintDebugMessages;
+	Options.bWriteDebugMesh = UseSettings.bEmitFullDebugMesh;
 
 	// run actor approximation computation
 	IGeometryProcessing_ApproximateActors::FResults Results;
@@ -202,7 +229,7 @@ bool FMeshApproximationTool::RunMerge(const FString& PackageName, const TArray<T
 	{
 		FAssetRegistryModule::AssetCreated(StaticMesh);
 	}
-	for (UMaterial* Material : Results.NewMaterials)
+	for (UMaterialInterface* Material : Results.NewMaterials)
 	{
 		FAssetRegistryModule::AssetCreated(Material);
 	}
diff --git a/Engine/Source/Runtime/Engine/Classes/Engine/MeshMerging.h b/Engine/Source/Runtime/Engine/Classes/Engine/MeshMerging.h
--- a/Engine/Source/Runtime/Engine/Classes/Engine/MeshMerging.h
+++ b/Engine/Source/Runtime/Engine/Classes/Engine/MeshMerging.h
@@ -679,11 +679,36 @@ struct FMeshInstancingSettings
 	TSubclassOf<UInstancedStaticMeshComponent> ISMComponentToUse;
 };
 
+
+UENUM()
+enum class EMeshApproximationType : uint8
+{
+	MeshAndMaterials,
+	MeshShapeOnly
+};
+
+UENUM()
+enum class EMeshApproximationBaseCappingType : uint8
+{
+	NoBaseCapping = 0,
+	ConvexPolygon = 1,
+	ConvexSolid = 2
+};
+
+
+UENUM()
+enum class EOccludedGeometryFilteringPolicy : uint8
+{
+	NoOcclusionFiltering = 0,
+	VisibilityBasedFiltering = 1
+};
+
 UENUM()
 enum class EMeshApproximationSimplificationPolicy : uint8
 {
 	FixedTriangleCount = 0,
-	TrianglesPerArea = 1
+	TrianglesPerArea = 1,
+	GeometricTolerance = 2
 };
 
 
@@ -692,13 +717,30 @@ struct FMeshApproximationSettings
 {
 	GENERATED_BODY()
 
+	/** Type of output from mesh approximation process */
+	UPROPERTY(BlueprintReadWrite, EditAnywhere, Category = ShapeSettings)
+	EMeshApproximationType OutputType = EMeshApproximationType::MeshAndMaterials;
+
 	/** Approximation Accuracy in Meters, will determine (eg) voxel resolution */
 	UPROPERTY(BlueprintReadWrite, EditAnywhere, Category = ShapeSettings, meta = (ClampMin = "0.1"))
 	float ApproximationAccuracy = 1.0;
 
 	/** Maximum allowable voxel count along main directions. This is a limit on ApproximationAccuracy */
 	UPROPERTY(BlueprintReadWrite, EditAnywhere, AdvancedDisplay, Category = ShapeSettings, meta = (ClampMin = "64"))
-	int32 ClampVoxelDimension = 512;
+	int32 ClampVoxelDimension = 1024;
+
+	/** if enabled, we will attempt to auto-thicken thin parts or flat sheets */
+	UPROPERTY(BlueprintReadWrite, EditAnywhere, Category = ShapeSettings)
+	bool bAttemptAutoThickening = true;
+
+	/** Multiplier on Approximation Accuracy used for auto-thickening */
+	UPROPERTY(BlueprintReadWrite, EditAnywhere, Category = ShapeSettings, meta = (ClampMin = "0.1", EditCondition = "bAttemptAutoThickening"))
+	float TargetMinThicknessMultiplier = 1.5;
+
+	/** Optional methods to attempt to close off the bottom of open meshes */
+	UPROPERTY(BlueprintReadWrite, EditAnywhere, Category = ShapeSettings)
+	EMeshApproximationBaseCappingType BaseCapping = EMeshApproximationBaseCappingType::NoBaseCapping;
+
 
 	/** */
 	UPROPERTY(BlueprintReadWrite, EditAnywhere, Category = ShapeSettings, meta = (ClampMin = "0.01", ClampMax = "0.99"))
@@ -711,40 +753,62 @@ struct FMeshApproximationSettings
 	UPROPERTY(BlueprintReadWrite, EditAnywhere, Category = ShapeSettings, meta = (ClampMin = "0.1", EditCondition = "bFillGaps"))
 	float GapDistance = 0.1;
 
-	UPROPERTY(BlueprintReadWrite, EditAnywhere, Category = ShapeSettings)
-	EMeshApproximationSimplificationPolicy SimplifyMethod = EMeshApproximationSimplificationPolicy::FixedTriangleCount;
 
-	UPROPERTY(BlueprintReadWrite, EditAnywhere, Category = ShapeSettings, meta = (ClampMin = "16", EditCondition = "SimplifyMethod == EMeshApproximationSimplificationPolicy::FixedTriangleCount" ))
+	/** Type of hidden geometry removal to apply */
+	UPROPERTY(BlueprintReadWrite, EditAnywhere, Category = SimplifySettings)
+	EOccludedGeometryFilteringPolicy OcclusionMethod = EOccludedGeometryFilteringPolicy::VisibilityBasedFiltering;
+
+	/** Mesh simplification criteria */
+	UPROPERTY(BlueprintReadWrite, EditAnywhere, Category = SimplifySettings)
+	EMeshApproximationSimplificationPolicy SimplifyMethod = EMeshApproximationSimplificationPolicy::GeometricTolerance;
+
+	UPROPERTY(BlueprintReadWrite, EditAnywhere, Category = SimplifySettings, meta = (ClampMin = "16", EditCondition = "SimplifyMethod == EMeshApproximationSimplificationPolicy::FixedTriangleCount" ))
 	int32 TargetTriCount = 2000;
 
 	/** Approximate Number of triangles per Square Meter */
-	UPROPERTY(BlueprintReadWrite, EditAnywhere, Category = ShapeSettings, meta = (ClampMin = "0.01", EditCondition = "SimplifyMethod == EMeshApproximationSimplificationPolicy::TrianglesPerArea" ))
+	UPROPERTY(BlueprintReadWrite, EditAnywhere, Category = SimplifySettings, meta = (ClampMin = "0.01", EditCondition = "SimplifyMethod == EMeshApproximationSimplificationPolicy::TrianglesPerArea" ))
 	float TrianglesPerM = 2.0f;
 
+	/** Allowable Geometric Deviation in Meters when SimplifyMethod incorporates a Geometric Tolerance */
+	UPROPERTY(BlueprintReadWrite, EditAnywhere, Category = SimplifySettings, meta = (ClampMin = "0.0001", EditCondition = "SimplifyMethod == EMeshApproximationSimplificationPolicy::GeometricTolerance"))
+	float GeometricDeviation = 0.1f;
 
 
-	/** Material simplification */
-	UPROPERTY(BlueprintReadWrite, EditAnywhere, Category = MaterialSettings)
-	FMaterialProxySettings MaterialSettings;
 
 	/** If Value is > 1, Multisample output baked textures by this amount in each direction (eg 4 == 16x supersampling) */
-	UPROPERTY(BlueprintReadWrite, EditAnywhere, AdvancedDisplay, Category = MaterialSettings, meta = (ClampMin = "0", ClampMax = "8", UIMin = "0", UIMax = "4"))
+	UPROPERTY(BlueprintReadWrite, EditAnywhere, Category = MaterialSettings, meta = (ClampMin = "0", ClampMax = "8", UIMin = "0", UIMax = "4"))
 	int32 MultiSamplingAA = 0;
 
+	/** If Value is zero, use MaterialSettings resolution, otherwise override the render capture resolution */
+	UPROPERTY(BlueprintReadWrite, EditAnywhere, Category = MaterialSettings, meta = (ClampMin = "128"))
+	int32 RenderCaptureResolution = 2048;
+
+	/** Material generation settings */
+	UPROPERTY(BlueprintReadWrite, EditAnywhere, Category = MaterialSettings)
+	FMaterialProxySettings MaterialSettings;
+
 	UPROPERTY(BlueprintReadWrite, EditAnywhere, AdvancedDisplay, Category = MaterialSettings, meta = (ClampMin = "5.0", ClampMax = "160.0"))
 	float CaptureFieldOfView = 30.0f;
 
 	UPROPERTY(BlueprintReadWrite, EditAnywhere, AdvancedDisplay, Category = MaterialSettings, meta = (ClampMin = "0.001", ClampMax = "1000.0"))
 	float NearPlaneDist = 1.0f;
 
-	/** If Value is zero, use MaterialSettings resolution, otherwise override the render capture resolution */
-	UPROPERTY(BlueprintReadWrite, EditAnywhere, AdvancedDisplay, Category = MaterialSettings, meta = (ClampMin = "128"))
-	int32 RenderCaptureResolution = 0;
+
+
+	/** If true, print out debugging messages */
+	UPROPERTY(BlueprintReadWrite, EditAnywhere, AdvancedDisplay, Category = DebugSettings)
+	bool bPrintDebugMessages = false;
+
+	/** If true, write the full mesh triangle set (ie flattened, non-instanced) used for mesh generation. Warning: this asset may be extremely large!! */
+	UPROPERTY(BlueprintReadWrite, EditAnywhere, AdvancedDisplay, Category = DebugSettings)
+	bool bEmitFullDebugMesh = false;
+
 
 	/** Equality operator. */
 	bool operator==(const FMeshApproximationSettings& Other) const
 	{
-		return ApproximationAccuracy == Other.ApproximationAccuracy
+		return OutputType == Other.OutputType
+			&& ApproximationAccuracy == Other.ApproximationAccuracy
 			&& ClampVoxelDimension == Other.ClampVoxelDimension
 			&& WindingThreshold == Other.WindingThreshold
 			&& bFillGaps == Other.bFillGaps
