ModelingTools: 1st pass Mesh Texture bake refactor w/ multisampling support.

#rnx

diff --git a/Engine/Plugins/Experimental/GeometryProcessing/Source/DynamicMesh/Private/Sampling/MeshCurvatureMapBaker.cpp b/Engine/Plugins/Experimental/GeometryProcessing/Source/DynamicMesh/Private/Sampling/MeshCurvatureMapBaker.cpp
--- a/Engine/Plugins/Experimental/GeometryProcessing/Source/DynamicMesh/Private/Sampling/MeshCurvatureMapBaker.cpp
+++ b/Engine/Plugins/Experimental/GeometryProcessing/Source/DynamicMesh/Private/Sampling/MeshCurvatureMapBaker.cpp
@@ -1,6 +1,7 @@
 // Copyright Epic Games, Inc. All Rights Reserved.
 
 #include "Sampling/MeshCurvatureMapBaker.h"
+#include "Sampling/MeshMapBaker.h"
 #include "Sampling/SphericalFibonacci.h"
 #include "Sampling/Gaussians.h"
 #include "MeshCurvature.h"
@@ -12,14 +13,14 @@
 #include "ExplicitUseGeometryMathTypes.h"		// using UE::Geometry::(math types)
 using namespace UE::Geometry;
 
-void FMeshCurvatureMapBaker::CacheDetailCurvatures(const FDynamicMesh3* DetailMesh)
+void FMeshCurvatureMapBaker::CacheDetailCurvatures(const FDynamicMesh3* DetailMeshIn)
 {
 	if (! Curvatures)
 	{
 		Curvatures = MakeShared<FMeshVertexCurvatureCache>();
-		Curvatures->BuildAll(*DetailMesh);
+		Curvatures->BuildAll(*DetailMeshIn);
 	}
-	check(Curvatures->Num() == DetailMesh->MaxVertexID());
+	check(Curvatures->Num() == DetailMeshIn->MaxVertexID());
 
 }
 
@@ -28,7 +29,7 @@ void FMeshCurvatureMapBaker::Bake()
 {
 	const FMeshImageBakingCache* BakeCache = GetCache();
 	check(BakeCache);
-	const FDynamicMesh3* DetailMesh = BakeCache->GetDetailMesh();
+	DetailMesh = BakeCache->GetDetailMesh();
 
 	CacheDetailCurvatures(DetailMesh);
 
@@ -69,12 +70,10 @@ void FMeshCurvatureMapBaker::Bake()
 void FMeshCurvatureMapBaker::Bake_Single()
 {
 	const FMeshImageBakingCache* BakeCache = GetCache();
-	const FDynamicMesh3* DetailMesh = BakeCache->GetDetailMesh();
 	const FMeshVertexCurvatureCache& Cache = *Curvatures;
 
-
-	double MinPreClamp = -TNumericLimits<double>::Max();
-	double MaxPreClamp = TNumericLimits<double>::Max();
+	MinPreClamp = -TNumericLimits<double>::Max();
+	MaxPreClamp = TNumericLimits<double>::Max();
 	if (UseClampMode == EClampMode::Positive)
 	{
 		MinPreClamp = 0;
@@ -84,23 +83,6 @@ void FMeshCurvatureMapBaker::Bake_Single()
 		MaxPreClamp = 0;
 	}
 
-
-
-	auto GetCurvature = [&](int32 vid) {
-		switch (UseCurvatureType)
-		{
-		default:
-		case ECurvatureType::Mean:
-			return FMathd::Clamp(Cache[vid].Mean, MinPreClamp, MaxPreClamp);
-		case ECurvatureType::Gaussian:
-			return FMathd::Clamp(Cache[vid].Gaussian, MinPreClamp, MaxPreClamp);
-		case ECurvatureType::MaxPrincipal:
-			return FMathd::Clamp(Cache[vid].MaxPrincipal, MinPreClamp, MaxPreClamp);
-		case ECurvatureType::MinPrincipal:
-			return FMathd::Clamp(Cache[vid].MinPrincipal, MinPreClamp, MaxPreClamp);
-		}
-	};
-
 	FSampleSetStatisticsd CurvatureStats = Cache.MeanStats;
 	switch (UseCurvatureType)
 	{
@@ -125,35 +107,13 @@ void FMeshCurvatureMapBaker::Bake_Single()
 		ClampMax = RangeScale * OverrideRangeMax;
 	}
 	double MinClamp = MinRangeScale * ClampMax;
-	FInterval1d ClampRange(MinClamp, ClampMax);
-
-	// sample curvature of detail surface in tangent-space of base surface
-	auto CurvatureSampleFunction = [&](const FMeshImageBakingCache::FCorrespondenceSample& SampleData)
-	{
-		int32 DetailTriID = SampleData.DetailTriID;
-		if (DetailMesh->IsTriangle(DetailTriID))
-		{
-			FIndex3i DetailTri = DetailMesh->GetTriangle(DetailTriID);
-			double CurvatureA = GetCurvature(DetailTri.A);
-			double CurvatureB = GetCurvature(DetailTri.B);
-			double CurvatureC = GetCurvature(DetailTri.C);
-			double InterpCurvature = SampleData.DetailBaryCoords.X * CurvatureA
-				+ SampleData.DetailBaryCoords.Y * CurvatureB
-				+ SampleData.DetailBaryCoords.Z * CurvatureC;
-
-			return InterpCurvature;
-		}
-		return 0.0;
-	};
-
-
+	ClampRange = FInterval1d(MinClamp, ClampMax);
 
-	FVector3f NegativeColor, ZeroColor, PositiveColor;
 	GetColorMapRange(NegativeColor, ZeroColor, PositiveColor);
 
-	BakeCache->EvaluateSamples([&](const FVector2i& Coords, const FMeshImageBakingCache::FCorrespondenceSample& Sample)
+	BakeCache->EvaluateSamples([this](const FVector2i& Coords, const FMeshImageBakingCache::FCorrespondenceSample& Sample)
 	{
-		double Curvature = CurvatureSampleFunction(Sample);
+		double Curvature = SampleFunction<FMeshImageBakingCache::FCorrespondenceSample>(Sample);
 
 		double Sign = FMathd::Sign(Curvature);
 		float T = (float)ClampRange.GetT(FMathd::Abs(Curvature));
@@ -171,26 +131,117 @@ void FMeshCurvatureMapBaker::Bake_Multi()
 
 }
 
+double FMeshCurvatureMapBaker::GetCurvature(int32 vid)
+{
+	const FMeshVertexCurvatureCache& Cache = *Curvatures;
+	switch (UseCurvatureType)
+	{
+	default:
+	case ECurvatureType::Mean:
+		return FMathd::Clamp(Cache[vid].Mean, MinPreClamp, MaxPreClamp);
+	case ECurvatureType::Gaussian:
+		return FMathd::Clamp(Cache[vid].Gaussian, MinPreClamp, MaxPreClamp);
+	case ECurvatureType::MaxPrincipal:
+		return FMathd::Clamp(Cache[vid].MaxPrincipal, MinPreClamp, MaxPreClamp);
+	case ECurvatureType::MinPrincipal:
+		return FMathd::Clamp(Cache[vid].MinPrincipal, MinPreClamp, MaxPreClamp);
+	}
+}
 
-void FMeshCurvatureMapBaker::GetColorMapRange(FVector3f& NegativeColor, FVector3f& ZeroColor, FVector3f& PositiveColor)
+void FMeshCurvatureMapBaker::GetColorMapRange(FVector3f& NegativeColorOut, FVector3f& ZeroColorOut, FVector3f& PositiveColorOut)
 {
 	switch (UseColorMode)
 	{
 	default:
 	case EColorMode::BlackGrayWhite:
-		NegativeColor = FVector3f(0, 0, 0);
-		ZeroColor = FVector3f(0.5, 0.5, 0.5);
-		PositiveColor = FVector3f(1, 1, 1);
+		NegativeColorOut = FVector3f(0, 0, 0);
+		ZeroColorOut = FVector3f(0.5, 0.5, 0.5);
+		PositiveColorOut = FVector3f(1, 1, 1);
 		break;
 	case EColorMode::RedGreenBlue:
-		NegativeColor = FVector3f(1, 0, 0);
-		ZeroColor = FVector3f(0, 1, 0);
-		PositiveColor = FVector3f(0, 0, 1);
+		NegativeColorOut = FVector3f(1, 0, 0);
+		ZeroColorOut = FVector3f(0, 1, 0);
+		PositiveColorOut = FVector3f(0, 0, 1);
 		break;
 	case EColorMode::RedBlue:
-		NegativeColor = FVector3f(1, 0, 0);
-		ZeroColor = FVector3f(0, 0, 0);
-		PositiveColor = FVector3f(0, 0, 1);
+		NegativeColorOut = FVector3f(1, 0, 0);
+		ZeroColorOut = FVector3f(0, 0, 0);
+		PositiveColorOut = FVector3f(0, 0, 1);
+		break;
+	}
+}
+
+void FMeshCurvatureMapBaker::PreEvaluate(const FMeshMapBaker& Baker)
+{
+	DetailMesh = Baker.GetDetailMesh();
+	CacheDetailCurvatures(DetailMesh);
+
+	MinPreClamp = -TNumericLimits<double>::Max();
+	MaxPreClamp = TNumericLimits<double>::Max();
+	if (UseClampMode == EClampMode::Positive)
+	{
+		MinPreClamp = 0;
+	}
+	else if (UseClampMode == EClampMode::Negative)
+	{
+		MaxPreClamp = 0;
+	}
+
+	const FMeshVertexCurvatureCache& Cache = *Curvatures;
+	FSampleSetStatisticsd CurvatureStats = Cache.MeanStats;
+	switch (UseCurvatureType)
+	{
+	default:
+	case ECurvatureType::Mean:
+		CurvatureStats = Cache.MeanStats;
+		break;
+	case ECurvatureType::Gaussian:
+		CurvatureStats = Cache.GaussianStats;
+		break;
+	case ECurvatureType::MaxPrincipal:
+		CurvatureStats = Cache.MaxPrincipalStats;
+		break;
+	case ECurvatureType::MinPrincipal:
+		CurvatureStats = Cache.MinPrincipalStats;
 		break;
 	}
-}
\ No newline at end of file
+
+	double ClampMax = RangeScale * (CurvatureStats.Mean + CurvatureStats.StandardDeviation);
+	if (bOverrideCurvatureRange)
+	{
+		ClampMax = RangeScale * OverrideRangeMax;
+	}
+	double MinClamp = MinRangeScale * ClampMax;
+	ClampRange = FInterval1d(MinClamp, ClampMax);
+
+	GetColorMapRange(NegativeColor, ZeroColor, PositiveColor);
+}
+
+FVector4f FMeshCurvatureMapBaker::EvaluateSample(const FMeshMapBaker& Baker, const FCorrespondenceSample& Sample)
+{
+	double Curvature = SampleFunction<FCorrespondenceSample>(Sample);
+	double Sign = FMathd::Sign(Curvature);
+	float T = (float)ClampRange.GetT(FMathd::Abs(Curvature));
+	FVector3f CurvatureColor = (Sign < 0) ?
+		Lerp(ZeroColor, NegativeColor, T) : Lerp(ZeroColor, PositiveColor, T);
+	return FVector4f(CurvatureColor.X, CurvatureColor.Y, CurvatureColor.Z, 1.0f);
+}
+
+template<class SampleType>
+double FMeshCurvatureMapBaker::SampleFunction(const SampleType& SampleData)
+{
+	int32 DetailTriID = SampleData.DetailTriID;
+	if (DetailMesh->IsTriangle(DetailTriID))
+	{
+		FIndex3i DetailTri = DetailMesh->GetTriangle(DetailTriID);
+		double CurvatureA = GetCurvature(DetailTri.A);
+		double CurvatureB = GetCurvature(DetailTri.B);
+		double CurvatureC = GetCurvature(DetailTri.C);
+		double InterpCurvature = SampleData.DetailBaryCoords.X * CurvatureA
+			+ SampleData.DetailBaryCoords.Y * CurvatureB
+			+ SampleData.DetailBaryCoords.Z * CurvatureC;
+
+		return InterpCurvature;
+	}
+	return 0.0;
+}
diff --git a/Engine/Plugins/Experimental/GeometryProcessing/Source/DynamicMesh/Private/Sampling/MeshMapBaker.cpp b/Engine/Plugins/Experimental/GeometryProcessing/Source/DynamicMesh/Private/Sampling/MeshMapBaker.cpp
new file mode 100644
--- /dev/null
+++ b/Engine/Plugins/Experimental/GeometryProcessing/Source/DynamicMesh/Private/Sampling/MeshMapBaker.cpp
@@ -0,0 +1,415 @@
+// Copyright Epic Games, Inc. All Rights Reserved.
+
+#include "Sampling/MeshMapBaker.h"
+#include "Image/ImageOccupancyMap.h"
+#include "Spatial/DenseGrid2.h"
+
+#include "ExplicitUseGeometryMathTypes.h"		// using UE::Geometry::(math types)
+using namespace UE::Geometry;
+
+/**
+ * Find point on Detail mesh that corresponds to point on Base mesh.
+ * Strategy is:
+ *    1) cast a ray inwards along -Normal from BasePoint + Thickness*Normal
+ *    2) cast a ray outwards along Normal from BasePoint
+ *    3) cast a ray inwards along -Normal from BasePoint
+ * We take (1) preferentially, and then (2), and then (3)
+ *
+ * If all of those fail, if bFailToNearestPoint is true we fall back to nearest-point,
+ *
+ * If all the above fail, return false
+ */
+static bool GetDetailMeshTrianglePoint_Raycast(
+	const FDynamicMesh3& DetailMesh,
+	const FDynamicMeshAABBTree3& DetailSpatial,
+	const FVector3d& BasePoint,
+	const FVector3d& BaseNormal,
+	int32& DetailTriangleOut,
+	FVector3d& DetailTriBaryCoords,
+	double Thickness,
+	bool bFailToNearestPoint
+)
+{
+	// TODO: should we check normals here? inverse normal should probably not be considered valid
+
+	// shoot rays forwards and backwards
+	FRay3d InwardRay = FRay3d(BasePoint + Thickness * BaseNormal, -BaseNormal);
+	FRay3d ForwardRay(BasePoint, BaseNormal);
+	FRay3d BackwardRay(BasePoint, -BaseNormal);
+	int32 ForwardHitTID = IndexConstants::InvalidID, InwardHitTID = IndexConstants::InvalidID, BackwardHitTID = IndexConstants::InvalidID;
+	double ForwardHitDist, InwardHitDist, BackwardHitDist;
+
+	IMeshSpatial::FQueryOptions Options;
+	Options.MaxDistance = Thickness;
+	bool bHitInward = DetailSpatial.FindNearestHitTriangle(InwardRay, InwardHitDist, InwardHitTID, Options);
+	bool bHitForward = DetailSpatial.FindNearestHitTriangle(ForwardRay, ForwardHitDist, ForwardHitTID, Options);
+	bool bHitBackward = DetailSpatial.FindNearestHitTriangle(BackwardRay, BackwardHitDist, BackwardHitTID, Options);
+
+	FRay3d HitRay;
+	int32 HitTID = IndexConstants::InvalidID;
+	double HitDist = TNumericLimits<double>::Max();
+
+	if (bHitInward)
+	{
+		HitRay = InwardRay;
+		HitTID = InwardHitTID;
+		HitDist = InwardHitDist;
+	}
+	else if (bHitForward)
+	{
+		HitRay = ForwardRay;
+		HitTID = ForwardHitTID;
+		HitDist = ForwardHitDist;
+	}
+	else if (bHitBackward)
+	{
+		HitRay = BackwardRay;
+		HitTID = BackwardHitTID;
+		HitDist = BackwardHitDist;
+	}
+
+	// if we got a valid ray hit, use it
+	if (DetailMesh.IsTriangle(HitTID))
+	{
+		DetailTriangleOut = HitTID;
+		FIntrRay3Triangle3d IntrQuery = TMeshQueries<FDynamicMesh3>::TriangleIntersection(DetailMesh, HitTID, HitRay);
+		DetailTriBaryCoords = IntrQuery.TriangleBaryCoords;
+		return true;
+	}
+	else
+	{
+		// if we did not find any hits, try nearest-point
+		IMeshSpatial::FQueryOptions OnSurfQueryOptions;
+		OnSurfQueryOptions.MaxDistance = Thickness;
+		double NearDistSqr = 0;
+		int32 NearestTriID = -1;
+		// if we are using absolute nearest point as a fallback, then ignore max distance
+		if (bFailToNearestPoint)
+		{
+			NearestTriID = DetailSpatial.FindNearestTriangle(BasePoint, NearDistSqr);
+		}
+		else
+		{
+			NearestTriID = DetailSpatial.FindNearestTriangle(BasePoint, NearDistSqr, OnSurfQueryOptions);
+		}
+		if (DetailMesh.IsTriangle(NearestTriID))
+		{
+			DetailTriangleOut = NearestTriID;
+			FDistPoint3Triangle3d DistQuery = TMeshQueries<FDynamicMesh3>::TriangleDistance(DetailMesh, NearestTriID, BasePoint);
+			DetailTriBaryCoords = DistQuery.TriangleBaryCoords;
+			return true;
+		}
+	}
+
+	return false;
+}
+
+/**
+ * Find point on Detail mesh that corresponds to point on Base mesh using minimum distance
+ */
+static bool GetDetailMeshTrianglePoint_Nearest(
+	const FDynamicMesh3& DetailMesh,
+	const FDynamicMeshAABBTree3& DetailSpatial,
+	const FVector3d& BasePoint,
+	int32& DetailTriangleOut,
+	FVector3d& DetailTriBaryCoords)
+{
+	double NearDistSqr = 0;
+	int32 NearestTriID = DetailSpatial.FindNearestTriangle(BasePoint, NearDistSqr);
+	if (DetailMesh.IsTriangle(NearestTriID))
+	{
+		DetailTriangleOut = NearestTriID;
+		FDistPoint3Triangle3d DistQuery = TMeshQueries<FDynamicMesh3>::TriangleDistance(DetailMesh, NearestTriID, BasePoint);
+		DetailTriBaryCoords = DistQuery.TriangleBaryCoords;
+		return true;
+	}
+
+	return false;
+}
+
+void FMeshMapBaker::Bake()
+{
+	if (Bakers.IsEmpty())
+	{
+		return;
+	}
+
+	// Initialize Bake targets and results.
+	BakeResults.SetNum(Bakers.Num());
+	for (int32 Idx = 0; Idx < Bakers.Num(); Idx++)
+	{
+		BakeResults[Idx] = MakeShared<TImageBuilder<FVector4f>, ESPMode::ThreadSafe>();
+		BakeResults[Idx]->SetDimensions(Dimensions);
+		FVector4f ClearColor = Bakers[Idx]->DefaultSample();
+		BakeResults[Idx]->Clear(ClearColor);
+		Bakers[Idx]->PreEvaluate(*this);
+	}
+
+	// Generate UV space mesh
+	const FDynamicMesh3* Mesh = TargetMesh;
+	const FDynamicMeshUVOverlay* UVOverlay = GetTargetMeshUVs();
+	const FDynamicMeshNormalOverlay* NormalOverlay = GetTargetMeshNormals();
+
+	FlatMesh = FDynamicMesh3(EMeshComponents::FaceGroups);
+	for (int32 tid : Mesh->TriangleIndicesItr())
+	{
+		if (UVOverlay->IsSetTriangle(tid))
+		{
+			FVector2f A, B, C;
+			UVOverlay->GetTriElements(tid, A, B, C);
+			int32 VertA = FlatMesh.AppendVertex(FVector3d(A.X, A.Y, 0));
+			int32 VertB = FlatMesh.AppendVertex(FVector3d(B.X, B.Y, 0));
+			int32 VertC = FlatMesh.AppendVertex(FVector3d(C.X, C.Y, 0));
+			int32 NewTriID = FlatMesh.AppendTriangle(VertA, VertB, VertC, tid);
+		}
+	}
+
+	ECorrespondenceStrategy UseStrategy = this->CorrespondenceStrategy;
+	if (UseStrategy == ECorrespondenceStrategy::Identity && ensure(DetailMesh == Mesh) == false)
+	{
+		// Identity strategy requires mesh to be the same. Could potentially have two copies, in which
+		// case this ensure is too conservative, but for now we will assume this
+		UseStrategy = ECorrespondenceStrategy::NearestPoint;
+	}
+
+	// this sampler finds the correspondence between base surface and detail surface
+	DetailMeshSampler.Initialize(Mesh, UVOverlay, EMeshSurfaceSamplerQueryType::TriangleAndUV, FMeshImageBaker::FCorrespondenceSample(),
+		[Mesh, NormalOverlay, UseStrategy, this](const FMeshUVSampleInfo& SampleInfo, FMeshImageBaker::FCorrespondenceSample& ValueOut)
+	{
+		NormalOverlay->GetTriBaryInterpolate<double>(SampleInfo.TriangleIndex, &SampleInfo.BaryCoords.X, &ValueOut.BaseNormal.X);
+		Normalize(ValueOut.BaseNormal);
+		FVector3d RayDir = ValueOut.BaseNormal;
+
+		ValueOut.BaseSample = SampleInfo;
+		ValueOut.DetailTriID = FDynamicMesh3::InvalidID;
+
+		if (UseStrategy == ECorrespondenceStrategy::Identity)
+		{
+			ValueOut.DetailTriID = SampleInfo.TriangleIndex;
+			ValueOut.DetailBaryCoords = SampleInfo.BaryCoords;
+		}
+		else if (UseStrategy == ECorrespondenceStrategy::NearestPoint)
+		{
+			bool bFoundTri = GetDetailMeshTrianglePoint_Nearest(*DetailMesh, *DetailSpatial, SampleInfo.SurfacePoint,
+				ValueOut.DetailTriID, ValueOut.DetailBaryCoords);
+		}
+		else	// fall back to raycast strategy
+		{
+			double SampleThickness = this->GetThickness();		// could modulate w/ a map here...
+
+			// find detail mesh triangle point
+			bool bFoundTri = GetDetailMeshTrianglePoint_Raycast(*DetailMesh, *DetailSpatial, SampleInfo.SurfacePoint, RayDir,
+				ValueOut.DetailTriID, ValueOut.DetailBaryCoords, SampleThickness,
+				(UseStrategy == ECorrespondenceStrategy::RaycastStandardThenNearest));
+		}
+
+	});
+
+	// Setup image tiling
+	FImageTiling Tiles(Dimensions, TileSize, TileSize, 0);
+	ParallelFor(Tiles.Num(), [this, &Tiles](int32 TileIdx)
+	{
+		FImageDimensions Tile = Tiles.GetTile(TileIdx);
+
+		FImageOccupancyMap OccupancyMap;
+		OccupancyMap.GutterSize = GutterSize;
+		OccupancyMap.Initialize(Tile, Multisampling);
+		OccupancyMap.ComputeFromUVSpaceMesh(FlatMesh, [this](int32 TriangleID) { return FlatMesh.GetTriangleGroup(TriangleID); });
+
+		// calculate interior texels
+		for (int32 ImgY = 0; ImgY < Tile.GetHeight(); ImgY++)
+		{
+			for (int32 ImgX = 0; ImgX < Tile.GetWidth(); ImgX++)
+			{
+				if (CancelF())
+				{
+					return;
+				}
+
+				BakePixel(OccupancyMap, Tile, ImgX, ImgY);
+			}
+		}
+
+		for (int64 k = 0; k < OccupancyMap.GutterTexels.Num(); k++)
+		{
+			TPair<int64, int64> GutterTexel = OccupancyMap.GutterTexels[k];
+			for (int32 Idx = 0; Idx < Bakers.Num(); Idx++)
+			{
+				BakeResults[Idx]->CopyPixel(GutterTexel.Value, GutterTexel.Key);
+			}
+		}
+	}, !bParallel ? EParallelForFlags::ForceSingleThread : EParallelForFlags::None);
+
+	for (int32 Idx = 0; Idx < Bakers.Num(); Idx++)
+	{
+		Bakers[Idx]->PostEvaluate(*this, *BakeResults[Idx]);
+	}
+}
+
+void FMeshMapBaker::BakePixel(FImageOccupancyMap& OccupancyMap, const FImageDimensions& Tile, int32 ImgX, int32 ImgY)
+{
+	FVector2i ImageCoords = Tile.GetSourceCoords(ImgX, ImgY);
+
+	// If this texel contains any interior samples, we will accumulate
+	// the result, reset the clear color to black.
+	if (OccupancyMap.TexelNumSamples(Tile.GetIndex(ImgX, ImgY)) > 0)
+	{
+		for (int32 Idx = 0; Idx < Bakers.Num(); Idx++)
+		{
+			if (Bakers[Idx]->SupportsMultisampling())
+			{
+				BakeResults[Idx]->SetPixel(ImageCoords, FVector4f::Zero());
+			}
+		}
+	}
+	else
+	{
+		return;
+	}
+
+	for (int32 SampleIdx = 0; SampleIdx < OccupancyMap.Multisampler.Num(); SampleIdx++)
+	{
+		int64 LinearIdx = Tile.GetIndex(ImgX, ImgY) * OccupancyMap.Multisampler.Num() + SampleIdx;
+		double OneOverNumSamples = 1.0 / (double)OccupancyMap.Multisampler.Num();
+		if (OccupancyMap.IsInterior(LinearIdx) == false)
+		{
+			for (int32 Idx = 0; Idx < Bakers.Num(); Idx++)
+			{
+				if (Bakers[Idx]->SupportsMultisampling())
+				{
+					FVector4f Result = BakeResults[Idx]->GetPixel(ImageCoords);
+					Result += Bakers[Idx]->DefaultSample() * OneOverNumSamples;
+					BakeResults[Idx]->SetPixel(ImageCoords, Result);
+				}
+			}
+		}
+		else
+		{
+			FVector2d UVPosition = (FVector2d)OccupancyMap.TexelQueryUV[LinearIdx];
+			int32 UVTriangleID = OccupancyMap.TexelQueryTriangle[LinearIdx];
+
+			FMeshImageBaker::FCorrespondenceSample Sample;
+			DetailMeshSampler.SampleUV(UVTriangleID, UVPosition, Sample);
+
+			for (int32 Idx = 0; Idx < Bakers.Num(); Idx++)
+			{
+				FVector4f SampleResult = Bakers[Idx]->EvaluateSample(*this, Sample);
+				if (!Bakers[Idx]->SupportsMultisampling())
+				{
+					// TODO/FIXME: Can be more correct & optimal here by not multisampling at all.
+					// Currently the samples are precomputed in the occupancy map.
+					BakeResults[Idx]->SetPixel(ImageCoords, SampleResult);
+				}
+				else
+				{
+					FVector4f Result = BakeResults[Idx]->GetPixel(ImageCoords);
+					Result += SampleResult * OneOverNumSamples;
+					BakeResults[Idx]->SetPixel(ImageCoords, Result);
+				}
+			}
+		}
+	}
+}
+
+int32 FMeshMapBaker::AddBaker(TSharedPtr<FMeshImageBaker> Sampler)
+{
+	return Bakers.Add(Sampler);
+}
+
+void FMeshMapBaker::Reset()
+{
+	Bakers.Empty();
+	BakeResults.Empty();
+}
+
+void FMeshMapBaker::SetDetailMesh(const FDynamicMesh3* Mesh, const FDynamicMeshAABBTree3* Spatial)
+{
+	DetailMesh = Mesh;
+	DetailSpatial = Spatial;
+}
+
+void FMeshMapBaker::SetTargetMesh(const FDynamicMesh3* Mesh)
+{
+	TargetMesh = Mesh;
+}
+
+void FMeshMapBaker::SetTargetMeshTangents(TSharedPtr<FMeshTangentsd, ESPMode::ThreadSafe> Tangents)
+{
+	TargetMeshTangents = Tangents;
+}
+
+const FDynamicMesh3* FMeshMapBaker::GetDetailMesh() const
+{
+	return DetailMesh;
+}
+
+const FDynamicMeshAABBTree3* FMeshMapBaker::GetDetailMeshSpatial() const
+{
+	return DetailSpatial;
+}
+
+const FDynamicMeshNormalOverlay* FMeshMapBaker::GetDetailMeshNormals() const
+{
+	check(DetailMesh && DetailMesh->HasAttributes());
+	return DetailMesh->Attributes()->PrimaryNormals();
+}
+
+const FDynamicMeshUVOverlay* FMeshMapBaker::GetDetailMeshUVs(int32 UVLayerIn /*=0*/) const
+{
+	check(DetailMesh && DetailMesh->HasAttributes());
+	return DetailMesh->Attributes()->GetUVLayer(UVLayerIn);
+}
+
+const FDynamicMesh3* FMeshMapBaker::GetTargetMesh() const
+{
+	return TargetMesh;
+}
+
+const FDynamicMeshUVOverlay* FMeshMapBaker::GetTargetMeshUVs() const
+{
+	check(TargetMesh && TargetMesh->HasAttributes() && UVLayer < TargetMesh->Attributes()->NumUVLayers());
+	return TargetMesh->Attributes()->GetUVLayer(UVLayer);
+}
+
+const FDynamicMeshNormalOverlay* FMeshMapBaker::GetTargetMeshNormals() const
+{
+	check(TargetMesh && TargetMesh->HasAttributes());
+	return TargetMesh->Attributes()->PrimaryNormals();
+}
+
+TSharedPtr<FMeshTangentsd, ESPMode::ThreadSafe> FMeshMapBaker::GetTargetMeshTangents() const
+{
+	return TargetMeshTangents;
+}
+
+void FMeshMapBaker::SetDimensions(FImageDimensions DimensionsIn)
+{
+	Dimensions = DimensionsIn;
+}
+
+void FMeshMapBaker::SetUVLayer(int32 UVLayerIn)
+{
+	UVLayer = UVLayerIn;
+}
+
+void FMeshMapBaker::SetThickness(double ThicknessIn)
+{
+	Thickness = ThicknessIn;
+}
+
+void FMeshMapBaker::SetGutterSize(int32 GutterSizeIn)
+{
+	GutterSize = GutterSizeIn;
+}
+
+void FMeshMapBaker::SetMultisampling(int32 MultisamplingIn)
+{
+	Multisampling = MultisamplingIn;
+}
+
+void FMeshMapBaker::SetCorrespondenceStrategy(ECorrespondenceStrategy Strategy)
+{
+	CorrespondenceStrategy = Strategy;
+}
+
+
diff --git a/Engine/Plugins/Experimental/GeometryProcessing/Source/DynamicMesh/Private/Sampling/MeshNormalMapBaker.cpp b/Engine/Plugins/Experimental/GeometryProcessing/Source/DynamicMesh/Private/Sampling/MeshNormalMapBaker.cpp
--- a/Engine/Plugins/Experimental/GeometryProcessing/Source/DynamicMesh/Private/Sampling/MeshNormalMapBaker.cpp
+++ b/Engine/Plugins/Experimental/GeometryProcessing/Source/DynamicMesh/Private/Sampling/MeshNormalMapBaker.cpp
@@ -1,6 +1,7 @@
 // Copyright Epic Games, Inc. All Rights Reserved.
 
 #include "Sampling/MeshNormalMapBaker.h"
+#include "Sampling/MeshMapBaker.h"
 
 #include "ExplicitUseGeometryMathTypes.h"		// using UE::Geometry::(math types)
 using namespace UE::Geometry;
@@ -9,62 +10,78 @@ void FMeshNormalMapBaker::Bake()
 {
 	const FMeshImageBakingCache* BakeCache = GetCache();
 	check(BakeCache);
-	const FDynamicMesh3* DetailMesh = BakeCache->GetDetailMesh();
-	const FDynamicMeshNormalOverlay* DetailNormalOverlay = BakeCache->GetDetailNormals();
+	DetailMesh = BakeCache->GetDetailMesh();
+	DetailNormalOverlay = BakeCache->GetDetailNormals();
 	check(DetailNormalOverlay);
 
 	// make sure we have per-triangle tangents computed so that GetInterpolatedTriangleTangent() below works
 	int32 MaxTriID = BakeCache->GetBakeTargetMesh()->MaxTriangleID();
 	check( BaseMeshTangents->GetTangents().Num() >= 3*BakeCache->GetBakeTargetMesh()->MaxTriangleID());
 
-	auto NormalSampleFunction = [&](const FMeshImageBakingCache::FCorrespondenceSample& SampleData)
-	{
-		int32 DetailTriID = SampleData.DetailTriID;
-		if (DetailMesh->IsTriangle(DetailTriID))
-		{
-			// get tangents on base mesh
-			FVector3d BaseTangentX, BaseTangentY;
-			BaseMeshTangents->GetInterpolatedTriangleTangent(
-				SampleData.BaseSample.TriangleIndex, 
-				SampleData.BaseSample.BaryCoords, 
-				BaseTangentX, BaseTangentY);
-
-			// sample normal on detail mesh
-			FVector3d DetailNormal;
-			DetailNormalOverlay->GetTriBaryInterpolate<double>(DetailTriID, &SampleData.DetailBaryCoords.X, &DetailNormal.X);
-			Normalize(DetailNormal);
-
-			// compute normal in tangent space
-			double dx = DetailNormal.Dot(BaseTangentX);
-			double dy = DetailNormal.Dot(BaseTangentY);
-			double dz = DetailNormal.Dot(SampleData.BaseNormal);
-
-			return FVector3f((float)dx, (float)dy, (float)dz);
-		}
-		return FVector3f::UnitZ();
-	};
-
-
 	NormalsBuilder = MakeUnique<TImageBuilder<FVector3f>>();
 	NormalsBuilder->SetDimensions(BakeCache->GetDimensions());
 	FVector3f DefaultMapNormal = (DefaultNormal + FVector3f::One()) * 0.5f;
 	NormalsBuilder->Clear(DefaultMapNormal);
 
-
 	BakeCache->EvaluateSamples([&](const FVector2i& Coords, const FMeshImageBakingCache::FCorrespondenceSample& Sample)
 	{
-		FVector3f RelativeDetailNormal = NormalSampleFunction(Sample);
+		FVector3f RelativeDetailNormal = SampleFunction<FMeshImageBakingCache::FCorrespondenceSample>(Sample);
 		FVector3f MapNormal = (RelativeDetailNormal + FVector3f::One()) * 0.5f;
 		NormalsBuilder->SetPixel(Coords, MapNormal);
 	});
 
 	const FImageOccupancyMap& Occupancy = *BakeCache->GetOccupancyMap();
-
 	for (int64 k = 0; k < Occupancy.GutterTexels.Num(); k++)
 	{
 		TPair<int64, int64> GutterTexel = Occupancy.GutterTexels[k];
 		NormalsBuilder->CopyPixel(GutterTexel.Value, GutterTexel.Key);
 	}
+}
 
+void FMeshNormalMapBaker::PreEvaluate(const FMeshMapBaker& Baker)
+{
+	DetailMesh = Baker.GetDetailMesh();
+	DetailNormalOverlay = Baker.GetDetailMeshNormals();
+	check(DetailNormalOverlay);
+	BaseMeshTangents = Baker.GetTargetMeshTangents().Get();
+
+	// make sure we have per-triangle tangents computed so that GetInterpolatedTriangleTangent() below works
+	int32 MaxTriID = Baker.GetTargetMesh()->MaxTriangleID();
+	check(BaseMeshTangents->GetTangents().Num() >= 3 * Baker.GetTargetMesh()->MaxTriangleID());
+}
 
-}
\ No newline at end of file
+FVector4f FMeshNormalMapBaker::EvaluateSample(const FMeshMapBaker& Baker, const FCorrespondenceSample& Sample)
+{
+	FVector3f SampleResult = SampleFunction<FCorrespondenceSample>(Sample);
+	// Map normal space [-1,1] to floating point color space [0,1]
+	FVector3f Result = (SampleResult + FVector3f::One()) * 0.5f;
+	return FVector4f(Result.X, Result.Y, Result.Z, 1.0f);
+}
+
+template<class SampleType>
+FVector3f FMeshNormalMapBaker::SampleFunction(const SampleType& SampleData)
+{
+	int32 DetailTriID = SampleData.DetailTriID;
+	if (DetailMesh->IsTriangle(DetailTriID))
+	{
+		// get tangents on base mesh
+		FVector3d BaseTangentX, BaseTangentY;
+		BaseMeshTangents->GetInterpolatedTriangleTangent(
+			SampleData.BaseSample.TriangleIndex,
+			SampleData.BaseSample.BaryCoords,
+			BaseTangentX, BaseTangentY);
+
+		// sample normal on detail mesh
+		FVector3d DetailNormal;
+		DetailNormalOverlay->GetTriBaryInterpolate<double>(DetailTriID, &SampleData.DetailBaryCoords.X, &DetailNormal.X);
+		Normalize(DetailNormal);
+
+		// compute normal in tangent space
+		double dx = DetailNormal.Dot(BaseTangentX);
+		double dy = DetailNormal.Dot(BaseTangentY);
+		double dz = DetailNormal.Dot(SampleData.BaseNormal);
+
+		return FVector3f((float)dx, (float)dy, (float)dz);
+	}
+	return FVector3f::UnitZ();
+}
diff --git a/Engine/Plugins/Experimental/GeometryProcessing/Source/DynamicMesh/Private/Sampling/MeshOcclusionMapBaker.cpp b/Engine/Plugins/Experimental/GeometryProcessing/Source/DynamicMesh/Private/Sampling/MeshOcclusionMapBaker.cpp
--- a/Engine/Plugins/Experimental/GeometryProcessing/Source/DynamicMesh/Private/Sampling/MeshOcclusionMapBaker.cpp
+++ b/Engine/Plugins/Experimental/GeometryProcessing/Source/DynamicMesh/Private/Sampling/MeshOcclusionMapBaker.cpp
@@ -1,11 +1,10 @@
 // Copyright Epic Games, Inc. All Rights Reserved.
 
 #include "Sampling/MeshOcclusionMapBaker.h"
+#include "Sampling/MeshMapBaker.h"
 #include "Sampling/SphericalFibonacci.h"
 #include "Sampling/Gaussians.h"
 
-#include "Math/RandomStream.h"
-
 #include "ExplicitUseGeometryMathTypes.h"		// using UE::Geometry::(math types)
 using namespace UE::Geometry;
 
@@ -18,15 +17,15 @@ void FMeshOcclusionMapBaker::Bake()
 
 	const FMeshImageBakingCache* BakeCache = GetCache();
 	check(BakeCache);
-	const FDynamicMesh3* DetailMesh = BakeCache->GetDetailMesh();
-	const FDynamicMeshAABBTree3* DetailSpatial = BakeCache->GetDetailSpatial();
-	const FDynamicMeshNormalOverlay* DetailNormalOverlay = BakeCache->GetDetailNormals();
+	DetailMesh = BakeCache->GetDetailMesh();
+	DetailSpatial = BakeCache->GetDetailSpatial();
+	DetailNormalOverlay = BakeCache->GetDetailNormals();
 	check(DetailNormalOverlay);
 
-	double BiasDotThreshold = FMathd::Cos(FMathd::Clamp(90.0 - BiasAngleDeg, 0.0, 90.0) * FMathd::DegToRad);
+	BiasDotThreshold = FMathd::Cos(FMathd::Clamp(90.0 - BiasAngleDeg, 0.0, 90.0) * FMathd::DegToRad);
 
 	// Precompute occlusion ray directions
-	TArray<FVector3d> RayDirections;
+	RayDirections.Empty(NumOcclusionRays);
 	THemisphericalFibonacci<double>::EDistribution Dist = THemisphericalFibonacci<double>::EDistribution::Cosine;
 	switch (Distribution)
 	{
@@ -59,136 +58,11 @@ void FMeshOcclusionMapBaker::Bake()
 		RayDir = Rotation * RayDir;
 	}
 
-	// Randomized rotation of occlusion rays about the normal
-	FRandomStream RotationGen(31337);
-	FCriticalSection RotationLock;
-	auto GetRandomRotation = [&RotationGen, &RotationLock]() {
-		RotationLock.Lock();
-		double Angle = RotationGen.GetFraction() * FMathd::TwoPi;
-		RotationLock.Unlock();
-		return Angle;
-	};
-
-	auto OcclusionSampleFunction = [&](const FMeshImageBakingCache::FCorrespondenceSample& SampleData, double& OutOcclusion, FVector3d& OutNormal)
-	{
-		// Fallback normal if invalid Tri or fully occluded
-		FVector3d DefaultNormal = FVector3d::UnitZ();
-		switch (NormalSpace)
-		{
-		case ESpace::Tangent:
-			break;
-		case ESpace::Object:
-			DefaultNormal = SampleData.BaseNormal;
-			break;
-		}
-
-		int32 DetailTriID = SampleData.DetailTriID;
-		if (DetailMesh->IsTriangle(DetailTriID))
-		{
-			FIndex3i DetailTri = DetailMesh->GetTriangle(DetailTriID);
-			//FVector3d DetailTriNormal = DetailMesh.GetTriNormal(DetailTriID);
-			FVector3d DetailTriNormal;
-			DetailNormalOverlay->GetTriBaryInterpolate<double>(DetailTriID, &SampleData.DetailBaryCoords.X, &DetailTriNormal.X);
-			Normalize(DetailTriNormal);
-
-			FVector3d BaseTangentX, BaseTangentY;
-			if (WantBentNormal() && NormalSpace == ESpace::Tangent)
-			{
-				check(BaseMeshTangents);
-				BaseMeshTangents->GetInterpolatedTriangleTangent(
-					SampleData.BaseSample.TriangleIndex,
-					SampleData.BaseSample.BaryCoords,
-					BaseTangentX, BaseTangentY);
-			}
-
-			FVector3d DetailBaryCoords = SampleData.DetailBaryCoords;
-			FVector3d DetailPos = DetailMesh->GetTriBaryPoint(DetailTriID, DetailBaryCoords.X, DetailBaryCoords.Y, DetailBaryCoords.Z);
-			DetailPos += 10.0f * FMathf::ZeroTolerance * DetailTriNormal;
-			FFrame3d SurfaceFrame(DetailPos, DetailTriNormal);
-
-			double RotationAngle = GetRandomRotation();
-			SurfaceFrame.Rotate(FQuaterniond(SurfaceFrame.Z(), RotationAngle, false));
-
-			IMeshSpatial::FQueryOptions QueryOptions;
-			QueryOptions.MaxDistance = MaxDistance;
-
-			double AccumOcclusion = 0;
-			FVector3d AccumNormal(FVector3d::Zero());
-			double TotalPointWeight = 0;
-			double TotalNormalWeight = 0;
-			for (FVector3d SphereDir : RayDirections)
-			{
-				FRay3d OcclusionRay(DetailPos, SurfaceFrame.FromFrameVector(SphereDir));
-				ensure(OcclusionRay.Direction.Dot(DetailTriNormal) > 0);
-
-				bool bHit = DetailSpatial->TestAnyHitTriangle(OcclusionRay, QueryOptions);
-
-				if (WantAmbientOcclusion())
-				{
-					// Have weight of point fall off as it becomes more coplanar with face. 
-					// This reduces faceting artifacts that we would otherwise see because geometry does not vary smoothly
-					double PointWeight = 1.0;
-					double BiasDot = OcclusionRay.Direction.Dot(DetailTriNormal);
-					if (BiasDot < BiasDotThreshold)
-					{
-						PointWeight = FMathd::Lerp(0.0, 1.0, FMathd::Clamp(BiasDot / BiasDotThreshold, 0.0, 1.0));
-						PointWeight *= PointWeight;
-					}
-					TotalPointWeight += PointWeight;
-
-					if (bHit)
-					{
-						AccumOcclusion += PointWeight;
-					}
-				}
-
-				if (WantBentNormal())
-				{
-					FVector3d BentNormal = OcclusionRay.Direction;
-					switch (NormalSpace)
-					{
-					case ESpace::Tangent:
-					{
-						// compute normal in tangent space
-						double dx = BentNormal.Dot(BaseTangentX);
-						double dy = BentNormal.Dot(BaseTangentY);
-						double dz = BentNormal.Dot(SampleData.BaseNormal);
-						BentNormal = FVector3d(dx, dy, dz);;
-						break;
-					}
-					case ESpace::Object:
-						break;
-					}
-					TotalNormalWeight += 1.0;
-
-					if (!bHit)
-					{
-						AccumNormal += BentNormal;
-					}
-				}
-			}
-
-			if (WantAmbientOcclusion())
-			{
-				AccumOcclusion = (TotalPointWeight > 0.0001) ? (AccumOcclusion / TotalPointWeight) : 0.0;
-			}
-			if (WantBentNormal())
-			{
-				AccumNormal = (TotalNormalWeight > 0.0 && AccumNormal.Length() > 0.0) ? Normalized(AccumNormal / TotalNormalWeight) : DefaultNormal;
-			}
-			OutOcclusion = AccumOcclusion;
-			OutNormal = AccumNormal;
-			return;
-		}
-		OutOcclusion = 0.0;
-		OutNormal = DefaultNormal;
-	};
-
 	if (WantAmbientOcclusion())
 	{
 		OcclusionBuilder = MakeUnique<TImageBuilder<FVector3f>>();
 		OcclusionBuilder->SetDimensions(BakeCache->GetDimensions());
-		OcclusionBuilder->Clear(FVector3f::Zero());
+		OcclusionBuilder->Clear(FVector3f::One());
 	}
 	if (WantBentNormal())
 	{
@@ -210,8 +84,8 @@ void FMeshOcclusionMapBaker::Bake()
 	BakeCache->EvaluateSamples([&](const FVector2i& Coords, const FMeshImageBakingCache::FCorrespondenceSample& Sample)
 	{
 		double Occlusion = 0.0;
-		FVector3d BentNormal;
-		OcclusionSampleFunction(Sample, Occlusion, BentNormal);
+		FVector3d BentNormal = FVector3f::UnitZ();
+		SampleFunction<FMeshImageBakingCache::FCorrespondenceSample>(Sample, Occlusion, BentNormal);
 		if (WantAmbientOcclusion())
 		{
 			FVector3f OcclusionColor = FMathd::Clamp(1.0f - (float)Occlusion, 0.0f, 1.0f) * FVector3f::One();
@@ -252,7 +126,189 @@ void FMeshOcclusionMapBaker::Bake()
 			BlurKernel2d.IntRadius,
 			AOBlurBuffer);
 	}
+}
+
+void FMeshOcclusionMapBaker::PreEvaluate(const FMeshMapBaker& Baker)
+{
+	// TODO: Support evaluation all map types in single pass.
+	check(OcclusionType != EOcclusionMapType::None && OcclusionType != EOcclusionMapType::All);
+
+	DetailMesh = Baker.GetDetailMesh();
+	DetailSpatial = Baker.GetDetailMeshSpatial();
+	DetailNormalOverlay = Baker.GetDetailMeshNormals();
+	check(DetailNormalOverlay);
+
+	if (WantBentNormal() && NormalSpace == ESpace::Tangent)
+	{
+		BaseMeshTangents = Baker.GetTargetMeshTangents().Get();
+	}
+
+	BiasDotThreshold = FMathd::Cos(FMathd::Clamp(90.0 - BiasAngleDeg, 0.0, 90.0) * FMathd::DegToRad);
+
+	// Precompute occlusion ray directions
+	RayDirections.Empty(NumOcclusionRays);
+	THemisphericalFibonacci<double>::EDistribution Dist = THemisphericalFibonacci<double>::EDistribution::Cosine;
+	switch (Distribution)
+	{
+	case EDistribution::Uniform:
+		Dist = THemisphericalFibonacci<double>::EDistribution::Uniform;
+		break;
+	case EDistribution::Cosine:
+		Dist = THemisphericalFibonacci<double>::EDistribution::Cosine;
+		break;
+	}
+	THemisphericalFibonacci<double> Points(NumOcclusionRays, Dist);
+	for (int32 k = 0; k < Points.Num(); ++k)
+	{
+		FVector3d P = Points[k];
+		if (P.Z > 0)
+		{
+			RayDirections.Add(Normalized(P));
+		}
+	}
+
+	// Map occlusion ray hemisphere to conical area (SpreadAngle/2)
+	double ConicalAngle = FMathd::Clamp(SpreadAngle * 0.5, 0.0001, 90.0);
+	for (int32 k = 0; k < RayDirections.Num(); ++k)
+	{
+		FVector3d& RayDir = RayDirections[k];
+		double RayAngle = AngleD(RayDir, FVector3d::UnitZ());
+		FVector3d RayCross = RayDir.Cross(FVector3d::UnitZ());
+		double RotationAngle = RayAngle - FMathd::Lerp(0.0, ConicalAngle, RayAngle / 90.0);
+		FQuaterniond Rotation(RayCross, RotationAngle, true);
+		RayDir = Rotation * RayDir;
+	}
+}
+
+FVector4f FMeshOcclusionMapBaker::EvaluateSample(const FMeshMapBaker& Baker, const FCorrespondenceSample& Sample)
+{
+	double Occlusion = 0.0;
+	FVector3d BentNormal = FVector3f::UnitZ();
+	SampleFunction<FCorrespondenceSample>(Sample, Occlusion, BentNormal);
+	if (WantAmbientOcclusion())
+	{
+		FVector3f OcclusionColor = FMathd::Clamp(1.0f - (float)Occlusion, 0.0f, 1.0f) * FVector3f::One();
+		return FVector4f(OcclusionColor.X, OcclusionColor.Y, OcclusionColor.Z, 1.0f);
+	}
+	else //if (WantBentNormal())
+	{
+		// Map normal space [-1,1] to floating point color space [0,1]
+		FVector3f NormalColor = (FVector3f(BentNormal.X, BentNormal.Y, BentNormal.Z) + FVector3f::One()) * 0.5;
+		return FVector4f(NormalColor.X, NormalColor.Y, NormalColor.Z, 1.0f);
+	}
+}
+
+template<class SampleType>
+void FMeshOcclusionMapBaker::SampleFunction(const SampleType& SampleData, double& OutOcclusion, FVector3d& OutNormal)
+{
+	// Fallback normal if invalid Tri or fully occluded
+	FVector3d DefaultNormal = FVector3d::UnitZ();
+	switch (NormalSpace)
+	{
+	case ESpace::Tangent:
+		break;
+	case ESpace::Object:
+		DefaultNormal = SampleData.BaseNormal;
+		break;
+	}
+
+	int32 DetailTriID = SampleData.DetailTriID;
+	if (!DetailMesh->IsTriangle(DetailTriID))
+	{
+		OutOcclusion = 0.0;
+		OutNormal = DefaultNormal;
+		return;
+	}
+
+	FIndex3i DetailTri = DetailMesh->GetTriangle(DetailTriID);
+	//FVector3d DetailTriNormal = DetailMesh.GetTriNormal(DetailTriID);
+	FVector3d DetailTriNormal;
+	DetailNormalOverlay->GetTriBaryInterpolate<double>(DetailTriID, &SampleData.DetailBaryCoords.X, &DetailTriNormal.X);
+	Normalize(DetailTriNormal);
+
+	FVector3d BaseTangentX, BaseTangentY;
+	if (WantBentNormal() && NormalSpace == ESpace::Tangent)
+	{
+		check(BaseMeshTangents);
+		BaseMeshTangents->GetInterpolatedTriangleTangent(
+			SampleData.BaseSample.TriangleIndex,
+			SampleData.BaseSample.BaryCoords,
+			BaseTangentX, BaseTangentY);
+	}
+
+	FVector3d DetailBaryCoords = SampleData.DetailBaryCoords;
+	FVector3d DetailPos = DetailMesh->GetTriBaryPoint(DetailTriID, DetailBaryCoords.X, DetailBaryCoords.Y, DetailBaryCoords.Z);
+	DetailPos += 10.0f * FMathf::ZeroTolerance * DetailTriNormal;
+	FFrame3d SurfaceFrame(DetailPos, DetailTriNormal);
+
+	double RotationAngle = GetRandomRotation();
+	SurfaceFrame.Rotate(FQuaterniond(SurfaceFrame.Z(), RotationAngle, false));
+
+	IMeshSpatial::FQueryOptions QueryOptions;
+	QueryOptions.MaxDistance = MaxDistance;
 
+	double AccumOcclusion = 0;
+	FVector3d AccumNormal(FVector3d::Zero());
+	double TotalPointWeight = 0;
+	for (FVector3d SphereDir : RayDirections)
+	{
+		FRay3d OcclusionRay(DetailPos, SurfaceFrame.FromFrameVector(SphereDir));
+		ensure(OcclusionRay.Direction.Dot(DetailTriNormal) > 0);
+
+		bool bHit = DetailSpatial->TestAnyHitTriangle(OcclusionRay, QueryOptions);
+
+		if (WantAmbientOcclusion())
+		{
+			// Have weight of point fall off as it becomes more coplanar with face. 
+			// This reduces faceting artifacts that we would otherwise see because geometry does not vary smoothly
+			double PointWeight = 1.0;
+			double BiasDot = OcclusionRay.Direction.Dot(DetailTriNormal);
+			if (BiasDot < BiasDotThreshold)
+			{
+				PointWeight = FMathd::Lerp(0.0, 1.0, FMathd::Clamp(BiasDot / BiasDotThreshold, 0.0, 1.0));
+				PointWeight *= PointWeight;
+			}
+			TotalPointWeight += PointWeight;
 
+			if (bHit)
+			{
+				AccumOcclusion += PointWeight;
+			}
+		}
 
+		if (WantBentNormal() && !bHit)
+		{
+			FVector3d BentNormal = OcclusionRay.Direction;
+			if (NormalSpace == ESpace::Tangent)
+			{
+				// compute normal in tangent space
+				double dx = BentNormal.Dot(BaseTangentX);
+				double dy = BentNormal.Dot(BaseTangentY);
+				double dz = BentNormal.Dot(SampleData.BaseNormal);
+				BentNormal = FVector3d(dx, dy, dz);;
+			}
+			AccumNormal += BentNormal;
+		}
+	}
+
+	if (WantAmbientOcclusion())
+	{
+		AccumOcclusion = (TotalPointWeight > 0.0001) ? (AccumOcclusion / TotalPointWeight) : 0.0;
+	}
+	if (WantBentNormal())
+	{
+		AccumNormal = (AccumNormal.Length() > 0.0) ? Normalized(AccumNormal) : DefaultNormal;
+	}
+	OutOcclusion = AccumOcclusion;
+	OutNormal = AccumNormal;
 }
+
+double FMeshOcclusionMapBaker::GetRandomRotation()
+{
+	// Randomized rotation of occlusion rays about the normal
+	RotationLock.Lock();
+	double Angle = RotationGen.GetFraction() * FMathd::TwoPi;
+	RotationLock.Unlock();
+	return Angle;
+}
+
diff --git a/Engine/Plugins/Experimental/GeometryProcessing/Source/DynamicMesh/Private/Sampling/MeshPropertyMapBaker.cpp b/Engine/Plugins/Experimental/GeometryProcessing/Source/DynamicMesh/Private/Sampling/MeshPropertyMapBaker.cpp
--- a/Engine/Plugins/Experimental/GeometryProcessing/Source/DynamicMesh/Private/Sampling/MeshPropertyMapBaker.cpp
+++ b/Engine/Plugins/Experimental/GeometryProcessing/Source/DynamicMesh/Private/Sampling/MeshPropertyMapBaker.cpp
@@ -1,6 +1,7 @@
 // Copyright Epic Games, Inc. All Rights Reserved.
 
 #include "Sampling/MeshPropertyMapBaker.h"
+#include "Sampling/MeshMapBaker.h"
 #include "Util/ColorConstants.h"
 
 #include "ExplicitUseGeometryMathTypes.h"		// using UE::Geometry::(math types)
@@ -31,12 +32,12 @@ void FMeshPropertyMapBaker::Bake()
 {
 	const FMeshImageBakingCache* BakeCache = GetCache();
 	check(BakeCache);
-	const FDynamicMesh3* DetailMesh = BakeCache->GetDetailMesh();
-	const FDynamicMeshNormalOverlay* DetailNormalOverlay = BakeCache->GetDetailNormals();
+	DetailMesh = BakeCache->GetDetailMesh();
+	DetailNormalOverlay = BakeCache->GetDetailNormals();
 	check(DetailNormalOverlay);
-	const FDynamicMeshUVOverlay* DetailUVOverlay = DetailMesh->Attributes()->PrimaryUV();
+	DetailUVOverlay = DetailMesh->Attributes()->PrimaryUV();
 
-	FAxisAlignedBox3d Bounds = DetailMesh->GetBounds();
+	Bounds = DetailMesh->GetBounds();
 	for (int32 j = 0; j < 3; ++j)
 	{
 		if (Bounds.Diagonal()[j] < FMathf::ZeroTolerance)
@@ -47,7 +48,7 @@ void FMeshPropertyMapBaker::Bake()
 		}
 	}
 
-	FVector3f DefaultValue(0, 0, 0);
+	DefaultValue = FVector3f::Zero();
 	switch (this->Property)
 	{
 		case EMeshPropertyBakeType::Position:
@@ -63,75 +64,12 @@ void FMeshPropertyMapBaker::Bake()
 			break;
 	}
 
-	auto PropertySampleFunction = [&](const FMeshImageBakingCache::FCorrespondenceSample& SampleData)
-	{
-		FVector3f Color = DefaultValue;
-		int32 DetailTriID = SampleData.DetailTriID;
-		if (DetailMesh->IsTriangle(DetailTriID))
-		{
-			switch (this->Property)
-			{
-				case EMeshPropertyBakeType::Position:
-					{
-						FVector3d Position = DetailMesh->GetTriBaryPoint(DetailTriID, SampleData.DetailBaryCoords[0], SampleData.DetailBaryCoords[1], SampleData.DetailBaryCoords[2]);
-						Color = PositionToColor(Position, Bounds);
-					}
-					break;
-				default:
-				case EMeshPropertyBakeType::FacetNormal:
-					{
-						FVector3d FacetNormal = DetailMesh->GetTriNormal(DetailTriID);
-						Color = NormalToColor(FacetNormal);
-					}
-					break;
-				break;
-				case EMeshPropertyBakeType::Normal:
-					{
-						if (DetailNormalOverlay->IsSetTriangle(DetailTriID))
-						{
-							FVector3d DetailNormal;
-							DetailNormalOverlay->GetTriBaryInterpolate<double>(DetailTriID, &SampleData.DetailBaryCoords.X, &DetailNormal.X);
-							Normalize(DetailNormal);
-							Color = NormalToColor(DetailNormal);
-						}
-					}
-					break;
-				case EMeshPropertyBakeType::UVPosition:
-					{
-						if (DetailUVOverlay && DetailUVOverlay->IsSetTriangle(DetailTriID))
-						{
-							FVector2d DetailUV;
-							DetailUVOverlay->GetTriBaryInterpolate<double>(DetailTriID, &SampleData.DetailBaryCoords.X, &DetailUV.X);
-							Color = UVToColor(DetailUV);
-						}
-					}
-					break;
-				case EMeshPropertyBakeType::MaterialID:
-					{
-						if (DetailMesh->Attributes() && DetailMesh->Attributes()->HasMaterialID())
-						{
-							const FDynamicMeshMaterialAttribute* DetailMaterialIDAttrib = DetailMesh->Attributes()->GetMaterialID();
-							const int32 MatID = DetailMaterialIDAttrib->GetValue(DetailTriID);
-							Color = LinearColors::SelectColor<FVector3f>(MatID);
-						}
-						else
-						{
-							Color = FVector3f(LinearColors::LightPink3f());
-						}
-					}
-					break;
-			}
-		}
-		return Color;
-	};
-
-
 	ResultBuilder = MakeUnique<TImageBuilder<FVector3f>>();
 	ResultBuilder->SetDimensions(BakeCache->GetDimensions());
 
 	BakeCache->EvaluateSamples([&](const FVector2i& Coords, const FMeshImageBakingCache::FCorrespondenceSample& Sample)
 	{
-		FVector3f Color = PropertySampleFunction(Sample);
+		FVector3f Color = SampleFunction<FMeshImageBakingCache::FCorrespondenceSample>(Sample);
 		ResultBuilder->SetPixel(Coords, Color);
 	});
 
@@ -144,4 +82,108 @@ void FMeshPropertyMapBaker::Bake()
 	}
 
 
-}
\ No newline at end of file
+}
+
+void FMeshPropertyMapBaker::PreEvaluate(const FMeshMapBaker& Baker)
+{
+	DetailMesh = Baker.GetDetailMesh();
+	DetailNormalOverlay = Baker.GetDetailMeshNormals();
+	check(DetailNormalOverlay);
+	DetailUVOverlay = Baker.GetDetailMeshUVs();
+
+	Bounds = DetailMesh->GetBounds();
+	for (int32 j = 0; j < 3; ++j)
+	{
+		if (Bounds.Diagonal()[j] < FMathf::ZeroTolerance)
+		{
+			Bounds.Min[j] = Bounds.Center()[j] - FMathf::ZeroTolerance;
+			Bounds.Max[j] = Bounds.Center()[j] + FMathf::ZeroTolerance;
+
+		}
+	}
+
+	DefaultValue = FVector3f::Zero();
+	switch (this->Property)
+	{
+	case EMeshPropertyBakeType::Position:
+		DefaultValue = PositionToColor(Bounds.Center(), Bounds);
+		break;
+	default:
+	case EMeshPropertyBakeType::FacetNormal:
+	case EMeshPropertyBakeType::Normal:
+		DefaultValue = NormalToColor(FVector3d::UnitZ());
+		break;
+	case EMeshPropertyBakeType::UVPosition:
+		DefaultValue = UVToColor(FVector2d::Zero());
+		break;
+	}
+}
+
+FVector4f FMeshPropertyMapBaker::EvaluateSample(const FMeshMapBaker& Baker, const FCorrespondenceSample& Sample)
+{
+	FVector3f SampleResult = SampleFunction<FCorrespondenceSample>(Sample);
+	return FVector4f(SampleResult.X, SampleResult.Y, SampleResult.Z, 1.0f);
+}
+
+template<class SampleType>
+FVector3f FMeshPropertyMapBaker::SampleFunction(const SampleType& SampleData)
+{
+	FVector3f Color = DefaultValue;
+	int32 DetailTriID = SampleData.DetailTriID;
+	if (DetailMesh->IsTriangle(DetailTriID))
+	{
+		switch (this->Property)
+		{
+		case EMeshPropertyBakeType::Position:
+		{
+			FVector3d Position = DetailMesh->GetTriBaryPoint(DetailTriID, SampleData.DetailBaryCoords[0], SampleData.DetailBaryCoords[1], SampleData.DetailBaryCoords[2]);
+			Color = PositionToColor(Position, Bounds);
+		}
+		break;
+		default:
+		case EMeshPropertyBakeType::FacetNormal:
+		{
+			FVector3d FacetNormal = DetailMesh->GetTriNormal(DetailTriID);
+			Color = NormalToColor(FacetNormal);
+		}
+		break;
+		break;
+		case EMeshPropertyBakeType::Normal:
+		{
+			if (DetailNormalOverlay->IsSetTriangle(DetailTriID))
+			{
+				FVector3d DetailNormal;
+				DetailNormalOverlay->GetTriBaryInterpolate<double>(DetailTriID, &SampleData.DetailBaryCoords.X, &DetailNormal.X);
+				Normalize(DetailNormal);
+				Color = NormalToColor(DetailNormal);
+			}
+		}
+		break;
+		case EMeshPropertyBakeType::UVPosition:
+		{
+			if (DetailUVOverlay && DetailUVOverlay->IsSetTriangle(DetailTriID))
+			{
+				FVector2d DetailUV;
+				DetailUVOverlay->GetTriBaryInterpolate<double>(DetailTriID, &SampleData.DetailBaryCoords.X, &DetailUV.X);
+				Color = UVToColor(DetailUV);
+			}
+		}
+		break;
+		case EMeshPropertyBakeType::MaterialID:
+		{
+			if (DetailMesh->Attributes() && DetailMesh->Attributes()->HasMaterialID())
+			{
+				const FDynamicMeshMaterialAttribute* DetailMaterialIDAttrib = DetailMesh->Attributes()->GetMaterialID();
+				const int32 MatID = DetailMaterialIDAttrib->GetValue(DetailTriID);
+				Color = LinearColors::SelectColor<FVector3f>(MatID);
+			}
+			else
+			{
+				Color = FVector3f(LinearColors::LightPink3f());
+			}
+		}
+		break;
+		}
+	}
+	return Color;
+}
diff --git a/Engine/Plugins/Experimental/GeometryProcessing/Source/DynamicMesh/Private/Sampling/MeshResampleImageBaker.cpp b/Engine/Plugins/Experimental/GeometryProcessing/Source/DynamicMesh/Private/Sampling/MeshResampleImageBaker.cpp
--- a/Engine/Plugins/Experimental/GeometryProcessing/Source/DynamicMesh/Private/Sampling/MeshResampleImageBaker.cpp
+++ b/Engine/Plugins/Experimental/GeometryProcessing/Source/DynamicMesh/Private/Sampling/MeshResampleImageBaker.cpp
@@ -1,6 +1,7 @@
 // Copyright Epic Games, Inc. All Rights Reserved.
 
 #include "Sampling/MeshResampleImageBaker.h"
+#include "Sampling/MeshMapBaker.h"
 
 #include "ExplicitUseGeometryMathTypes.h"		// using UE::Geometry::(math types)
 using namespace UE::Geometry;
@@ -9,33 +10,17 @@ void FMeshResampleImageBaker::Bake()
 {
 	const FMeshImageBakingCache* BakeCache = GetCache();
 	check(BakeCache);
-	const FDynamicMesh3* DetailMesh = BakeCache->GetDetailMesh();
+	DetailMesh = BakeCache->GetDetailMesh();
 
 	check(DetailUVOverlay);
 
-	FVector4f DefaultValue(0, 0, 0, 1.0);
-
-	auto PropertySampleFunction = [&](const FMeshImageBakingCache::FCorrespondenceSample& SampleData)
-	{
-		FVector4f Color = DefaultValue;
-		int32 DetailTriID = SampleData.DetailTriID;
-		if (DetailMesh->IsTriangle(SampleData.DetailTriID) && DetailUVOverlay)
-		{
-			FVector2d DetailUV;
-			DetailUVOverlay->GetTriBaryInterpolate<double>(DetailTriID, &SampleData.DetailBaryCoords.X, &DetailUV.X);
-
-			Color = SampleFunction(DetailUV);
-		}
-		return Color;
-	};
-
 	ResultBuilder = MakeUnique<TImageBuilder<FVector4f>>();
 	ResultBuilder->SetDimensions(BakeCache->GetDimensions());
 	ResultBuilder->Clear(DefaultColor);
 
 	BakeCache->EvaluateSamples([&](const FVector2i& Coords, const FMeshImageBakingCache::FCorrespondenceSample& Sample)
 	{
-		FVector4f Color = PropertySampleFunction(Sample);
+		FVector4f Color = PropertySampleFunction<FMeshImageBakingCache::FCorrespondenceSample>(Sample);
 		ResultBuilder->SetPixel(Coords, Color);
 	});
 
@@ -49,6 +34,31 @@ void FMeshResampleImageBaker::Bake()
 
 }
 
+void FMeshResampleImageBaker::PreEvaluate(const FMeshMapBaker& Baker)
+{
+	DetailMesh = Baker.GetDetailMesh();
+}
+
+FVector4f FMeshResampleImageBaker::EvaluateSample(const FMeshMapBaker& Baker, const FCorrespondenceSample& Sample)
+{
+	return PropertySampleFunction<FCorrespondenceSample>(Sample);
+}
+
+template <class SampleType>
+FVector4f FMeshResampleImageBaker::PropertySampleFunction(const SampleType& SampleData)
+{
+	FVector4f Color = DefaultColor;
+	int32 DetailTriID = SampleData.DetailTriID;
+	if (DetailMesh->IsTriangle(SampleData.DetailTriID) && DetailUVOverlay)
+	{
+		FVector2d DetailUV;
+		DetailUVOverlay->GetTriBaryInterpolate<double>(DetailTriID, &SampleData.DetailBaryCoords.X, &DetailUV.X);
+
+		Color = SampleFunction(DetailUV);
+	}
+	return Color;
+}
+
 
 
 
@@ -68,12 +78,12 @@ void FMeshMultiResampleImageBaker::BakeMaterial(int32 MaterialID)
 	{
 		return;
 	}
-	const FDynamicMesh3* DetailMesh = BakeCache->GetDetailMesh();
+	DetailMesh = BakeCache->GetDetailMesh();
 	if (!ensure(DetailMesh))
 	{
 		return;
 	}
-	const FDynamicMeshMaterialAttribute* DetailMaterialIDAttrib = DetailMesh->Attributes()->GetMaterialID();
+	DetailMaterialIDAttrib = DetailMesh->Attributes()->GetMaterialID();
 	if (!ensure(DetailMaterialIDAttrib))
 	{
 		return;
@@ -83,7 +93,7 @@ void FMeshMultiResampleImageBaker::BakeMaterial(int32 MaterialID)
 		return;
 	}
 
-	BakeCache->EvaluateSamples([this, DetailMesh, DetailMaterialIDAttrib, MaterialID](const FVector2i& Coords, const FMeshImageBakingCache::FCorrespondenceSample& Sample)
+	BakeCache->EvaluateSamples([this, MaterialID](const FVector2i& Coords, const FMeshImageBakingCache::FCorrespondenceSample& Sample)
 	{
 		int32 DetailTriID = Sample.DetailTriID;
 
@@ -129,3 +139,34 @@ void FMeshMultiResampleImageBaker::Bake()
 	}
 
 }
+
+void FMeshMultiResampleImageBaker::PreEvaluate(const FMeshMapBaker& Baker)
+{
+	DetailMesh = Baker.GetDetailMesh();
+	DetailMaterialIDAttrib = DetailMesh->Attributes()->GetMaterialID();
+
+	bValidDetailMesh = ensure(DetailMesh) && ensure(DetailMaterialIDAttrib) && ensure(DetailUVOverlay);
+}
+
+FVector4f FMeshMultiResampleImageBaker::EvaluateSample(const FMeshMapBaker& Baker, const FCorrespondenceSample& Sample)
+{
+	if (!bValidDetailMesh)
+	{
+		return DefaultSample();
+	}
+
+	FVector4f Color = DefaultSample();
+	int32 DetailTriID = Sample.DetailTriID;
+	if (DetailMesh->IsTriangle(DetailTriID))
+	{
+		int32 MaterialID = DetailMaterialIDAttrib->GetValue(DetailTriID);
+		TSharedPtr<UE::Geometry::TImageBuilder<FVector4f>, ESPMode::ThreadSafe> TextureImage = MultiTextures.FindRef(MaterialID);
+		if (TextureImage)
+		{
+			FVector2d DetailUV;
+			DetailUVOverlay->GetTriBaryInterpolate<double>(DetailTriID, &Sample.DetailBaryCoords.X, &DetailUV.X);
+			Color = TextureImage->BilinearSampleUV<float>(DetailUV, FVector4f(0, 0, 0, 1));
+		}
+	}
+	return Color;
+}
diff --git a/Engine/Plugins/Experimental/GeometryProcessing/Source/DynamicMesh/Public/Sampling/MeshCurvatureMapBaker.h b/Engine/Plugins/Experimental/GeometryProcessing/Source/DynamicMesh/Public/Sampling/MeshCurvatureMapBaker.h
--- a/Engine/Plugins/Experimental/GeometryProcessing/Source/DynamicMesh/Public/Sampling/MeshCurvatureMapBaker.h
+++ b/Engine/Plugins/Experimental/GeometryProcessing/Source/DynamicMesh/Public/Sampling/MeshCurvatureMapBaker.h
@@ -90,8 +90,37 @@ protected:
 
 	void Bake_Multi();
 
+	double GetCurvature(int32 vid);
 
 	void GetColorMapRange(FVector3f& NegativeColor, FVector3f& ZeroColor, FVector3f& PositiveColor);
+
+public:
+	//
+	// New Baker interface
+	//
+
+	/** Invoked at start of bake to initialize baker. */
+	virtual void PreEvaluate(const FMeshMapBaker& Baker) override;
+
+	/** Evaluate the sample at this mesh correspondence point. */
+	virtual FVector4f EvaluateSample(const FMeshMapBaker& Baker, const FCorrespondenceSample& Sample) override;
+
+	/** @return the default sample value for the baker. */
+	virtual FVector4f DefaultSample() const override { return FVector4f(0.0f, 0.0f, 0.0f, 1.0f); }
+
+protected:
+	// Cached data
+	const FDynamicMesh3* DetailMesh = nullptr;
+	double MinPreClamp = -TNumericLimits<double>::Max();
+	double MaxPreClamp = TNumericLimits<double>::Max();
+	FInterval1d ClampRange;
+	FVector3f NegativeColor;
+	FVector3f ZeroColor;
+	FVector3f PositiveColor;
+
+private:
+	template <class SampleType>
+	double SampleFunction(const SampleType& Sample);
 };
 
 
diff --git a/Engine/Plugins/Experimental/GeometryProcessing/Source/DynamicMesh/Public/Sampling/MeshImageBaker.h b/Engine/Plugins/Experimental/GeometryProcessing/Source/DynamicMesh/Public/Sampling/MeshImageBaker.h
--- a/Engine/Plugins/Experimental/GeometryProcessing/Source/DynamicMesh/Public/Sampling/MeshImageBaker.h
+++ b/Engine/Plugins/Experimental/GeometryProcessing/Source/DynamicMesh/Public/Sampling/MeshImageBaker.h
@@ -3,17 +3,23 @@
 #pragma once
 
 #include "Sampling/MeshImageBakingCache.h"
+#include "Image/ImageBuilder.h"
 
 namespace UE
 {
 namespace Geometry
 {
 
+class FMeshMapBaker;
+
 class DYNAMICMESH_API FMeshImageBaker
 {
 public:
 	virtual ~FMeshImageBaker() {}
 
+	//
+	// Old Baker interface
+	//
 	void SetCache(const FMeshImageBakingCache* CacheIn)
 	{
 		ImageBakeCache = CacheIn;
@@ -26,14 +32,38 @@ public:
 
 	virtual void Bake()
 	{
-
 	}
 
-
 protected:
-
 	const FMeshImageBakingCache* ImageBakeCache;
 
+public:
+	//
+	// New Baker interface
+	//
+	struct FCorrespondenceSample
+	{
+		FMeshUVSampleInfo BaseSample;
+		FVector3d BaseNormal;
+
+		int32 DetailTriID;
+		FVector3d DetailBaryCoords;
+	};
+
+	/** Invoked at start of bake to initialize baker. */
+	virtual void PreEvaluate(const FMeshMapBaker& Baker) {}
+
+	/** Evaluate the sample at this mesh correspondence point. */
+	virtual FVector4f EvaluateSample(const FMeshMapBaker& Baker, const FCorrespondenceSample& Sample) { return DefaultSample(); }
+
+	/** Invoked at the end of the bake to finalize the bake result. */
+	virtual void PostEvaluate(const FMeshMapBaker& Baker, TImageBuilder<FVector4f>& Result) {}
+
+	/** @return the default sample value for the baker. */
+	virtual FVector4f DefaultSample() const { return FVector4f::Zero(); }
+
+	/** @return true if this baker supports multisampling. */
+	virtual bool SupportsMultisampling() const { return true; }
 };
 
 } // end namespace UE::Geometry
diff --git a/Engine/Plugins/Experimental/GeometryProcessing/Source/DynamicMesh/Public/Sampling/MeshMapBaker.h b/Engine/Plugins/Experimental/GeometryProcessing/Source/DynamicMesh/Public/Sampling/MeshMapBaker.h
new file mode 100644
--- /dev/null
+++ b/Engine/Plugins/Experimental/GeometryProcessing/Source/DynamicMesh/Public/Sampling/MeshMapBaker.h
@@ -0,0 +1,122 @@
+// Copyright Epic Games, Inc. All Rights Reserved.
+
+#pragma once
+
+#include "DynamicMesh3.h"
+#include "DynamicMeshAttributeSet.h"
+#include "DynamicMeshAABBTree3.h"
+#include "MeshTangents.h"
+#include "Sampling/MeshImageBaker.h"
+#include "Sampling/MeshSurfaceSampler.h"
+#include "Spatial/DenseGrid2.h"
+#include "Image/ImageBuilder.h"
+#include "Image/ImageDimensions.h"
+
+namespace UE
+{
+namespace Geometry
+{
+
+class DYNAMICMESH_API FMeshMapBaker
+{
+public:
+
+	//
+	// Bake
+	//
+
+	/** Process all bakers to generate image results for each. */
+	void Bake();
+
+	/** Add a baker to be processed. */
+	int32 AddBaker(TSharedPtr<FMeshImageBaker> Target);
+
+	/** Reset the list of bakers. */
+	void Reset();
+
+	/** @return the bake result image for a given baker index. */
+	const TSharedPtr<TImageBuilder<FVector4f>, ESPMode::ThreadSafe>& GetBakeResult(int32 Index) const { return BakeResults[Index]; }
+
+	/** if this function returns true, we should abort calculation */
+	TFunction<bool(void)> CancelF = []() { return false; };
+
+
+	//
+	// Parameters
+	//
+
+	/**
+	 * ECorrespondenceStrategy determines the basic approach that will be used to establish a
+	 * mapping from points on the BakeTarget Mesh (usually low-poly) to points on the Detail Mesh (eg highpoly).
+	 * Geometrically this is not a 1-1 mapping so there are various options
+	 */
+	enum class ECorrespondenceStrategy
+	{
+		/** Raycast inwards from Point+Thickness*Normal, if that misses, try Outwards from Point, then Inwards from Point */
+		RaycastStandard,
+		/** Use geometrically nearest point. Thickness is ignored */
+		NearestPoint,
+		/** Use RaycastStandard but fall back to NearestPoint if none of the rays hit */
+		RaycastStandardThenNearest,
+		/** Assume that BakeTarget == DetailMesh and so no mapping is necessary */
+		Identity
+	};
+
+	void SetDetailMesh(const FDynamicMesh3* Mesh, const FDynamicMeshAABBTree3* Spatial);
+	void SetTargetMesh(const FDynamicMesh3* Mesh);
+
+	void SetTargetMeshTangents(TSharedPtr<FMeshTangentsd, ESPMode::ThreadSafe> Tangents);
+
+	const FDynamicMesh3* GetTargetMesh() const;
+	const FDynamicMeshUVOverlay* GetTargetMeshUVs() const;
+	const FDynamicMeshNormalOverlay* GetTargetMeshNormals() const;
+	TSharedPtr<FMeshTangentsd, ESPMode::ThreadSafe> GetTargetMeshTangents() const;
+
+	const FDynamicMesh3* GetDetailMesh() const;
+	const FDynamicMeshAABBTree3* GetDetailMeshSpatial() const;
+	const FDynamicMeshNormalOverlay* GetDetailMeshNormals() const;
+	const FDynamicMeshUVOverlay* GetDetailMeshUVs(int32 UVLayer=0) const;
+
+	void SetDimensions(FImageDimensions DimensionsIn);
+	void SetUVLayer(int32 UVLayerIn);
+	void SetThickness(double ThicknessIn);
+	void SetGutterSize(int32 GutterSizeIn);
+	void SetMultisampling(int32 MultisamplingIn);
+	void SetCorrespondenceStrategy(ECorrespondenceStrategy Strategy);
+
+	FImageDimensions GetDimensions() const { return Dimensions; }
+	int32 GetUVLayer() const { return UVLayer; }
+	double GetThickness() const { return Thickness; }
+	int32 GetGutterSize() const { return GutterSize; }
+	int32 GetMultisampling() const { return Multisampling; }
+	ECorrespondenceStrategy GetCorrespondenceStrategy() const { return CorrespondenceStrategy; }
+
+protected:
+	void BakePixel(FImageOccupancyMap& OccupancyMap, const FImageDimensions& Tile, int32 ImgX, int32 ImgY);
+
+protected:
+	bool bParallel = true;
+
+	const FDynamicMesh3* DetailMesh = nullptr;
+	const FDynamicMeshAABBTree3* DetailSpatial = nullptr;
+	const FDynamicMesh3* TargetMesh = nullptr;
+	TSharedPtr<FMeshTangentsd, ESPMode::ThreadSafe> TargetMeshTangents;
+
+	FDynamicMesh3 FlatMesh;
+	TMeshSurfaceUVSampler<FMeshImageBaker::FCorrespondenceSample> DetailMeshSampler;
+
+	FImageDimensions Dimensions = FImageDimensions(128, 128);
+	int32 UVLayer = 0;
+	double Thickness = 3.0;
+	ECorrespondenceStrategy CorrespondenceStrategy = ECorrespondenceStrategy::RaycastStandard;
+	int32 GutterSize = 4;
+	int32 Multisampling = 1;
+
+	const int32 TileSize = 32;
+
+	TArray<TSharedPtr<FMeshImageBaker, ESPMode::ThreadSafe>> Bakers;
+	TArray<TSharedPtr<TImageBuilder<FVector4f>, ESPMode::ThreadSafe>> BakeResults;
+};
+
+} // end namespace UE::Geometry
+} // end namespace UE
\ No newline at end of file
diff --git a/Engine/Plugins/Experimental/GeometryProcessing/Source/DynamicMesh/Public/Sampling/MeshNormalMapBaker.h b/Engine/Plugins/Experimental/GeometryProcessing/Source/DynamicMesh/Public/Sampling/MeshNormalMapBaker.h
--- a/Engine/Plugins/Experimental/GeometryProcessing/Source/DynamicMesh/Public/Sampling/MeshNormalMapBaker.h
+++ b/Engine/Plugins/Experimental/GeometryProcessing/Source/DynamicMesh/Public/Sampling/MeshNormalMapBaker.h
@@ -24,7 +24,7 @@ public:
 
 	const TMeshTangents<double>* BaseMeshTangents;
 
-	FVector3f DefaultNormal = FVector3f(0, 0, 1);
+	FVector3f DefaultNormal = FVector3f::UnitZ();
 
 	//
 	// Compute functions
@@ -43,6 +43,33 @@ public:
 protected:
 	TUniquePtr<TImageBuilder<FVector3f>> NormalsBuilder;
 
+public:
+	//
+	// New Baker interface
+	//
+
+	/** Invoked at start of bake to initialize baker. */
+	virtual void PreEvaluate(const FMeshMapBaker& Baker) override;
+
+	/** Evaluate the sample at this mesh correspondence point. */
+	virtual FVector4f EvaluateSample(const FMeshMapBaker& Baker, const FCorrespondenceSample& Sample) override;
+
+	/** @return the default sample value for the baker. */
+	virtual FVector4f DefaultSample() const override
+	{
+		// Map normal space [-1,1] to floating point color space [0,1]
+		FVector3f Normal = (DefaultNormal + FVector3f::One()) * 0.5f;
+		return FVector4f(Normal.X, Normal.Y, Normal.Z, 1.0f);
+	}
+
+protected:
+	// Cached data
+	const FDynamicMesh3* DetailMesh = nullptr;
+	const FDynamicMeshNormalOverlay* DetailNormalOverlay = nullptr;
+
+private:
+	template <class SampleType>
+	FVector3f SampleFunction(const SampleType& Sample);
 };
 
 
diff --git a/Engine/Plugins/Experimental/GeometryProcessing/Source/DynamicMesh/Public/Sampling/MeshOcclusionMapBaker.h b/Engine/Plugins/Experimental/GeometryProcessing/Source/DynamicMesh/Public/Sampling/MeshOcclusionMapBaker.h
--- a/Engine/Plugins/Experimental/GeometryProcessing/Source/DynamicMesh/Public/Sampling/MeshOcclusionMapBaker.h
+++ b/Engine/Plugins/Experimental/GeometryProcessing/Source/DynamicMesh/Public/Sampling/MeshOcclusionMapBaker.h
@@ -7,6 +7,7 @@
 #include "Sampling/MeshImageBakingCache.h"
 #include "Image/ImageBuilder.h"
 #include "MeshTangents.h"
+#include "Math/RandomStream.h"
 
 namespace UE
 {
@@ -124,6 +125,57 @@ public:
 protected:
 	TUniquePtr<TImageBuilder<FVector3f>> OcclusionBuilder;
 	TUniquePtr<TImageBuilder<FVector3f>> NormalBuilder;
+
+public:
+	//
+	// New Baker interface
+	//
+
+	/** Invoked at start of bake to initialize baker. */
+	virtual void PreEvaluate(const FMeshMapBaker& Baker) override;
+
+	/** Evaluate the sample at this mesh correspondence point. */
+	virtual FVector4f EvaluateSample(const FMeshMapBaker& Baker, const FCorrespondenceSample& Sample) override;
+
+	/** @return the default sample value for the baker. */
+	virtual FVector4f DefaultSample() const override
+	{
+		if (WantAmbientOcclusion())
+		{
+			return FVector4f::One();
+		}
+		else //if (WantBentNormal())
+		{
+			FVector3f DefaultNormal = FVector3f::UnitZ();
+			switch (NormalSpace)
+			{
+			case ESpace::Tangent:
+				break;
+			case ESpace::Object:
+				DefaultNormal = FVector3f::Zero();
+				break;
+			}
+			// Map normal space [-1,1] to floating point color space [0,1]
+			DefaultNormal = (DefaultNormal + FVector3f::One()) * 0.5f;
+			return FVector4f(DefaultNormal.X, DefaultNormal.Y, DefaultNormal.Z, 1.0f);
+		}
+	}
+
+protected:
+	const FDynamicMesh3* DetailMesh = nullptr;
+	const FDynamicMeshAABBTree3* DetailSpatial = nullptr;
+	const FDynamicMeshNormalOverlay* DetailNormalOverlay = nullptr;
+	double BiasDotThreshold = 0.25;
+	TArray<FVector3d> RayDirections;
+
+	FRandomStream RotationGen = FRandomStream(31337);
+	FCriticalSection RotationLock;
+
+private:
+	template<class SampleType>
+	void SampleFunction(const SampleType& Sample, double& Occlusion, FVector3d& Normal);
+
+	double GetRandomRotation();
 };
 
 } // end namespace UE::Geometry
diff --git a/Engine/Plugins/Experimental/GeometryProcessing/Source/DynamicMesh/Public/Sampling/MeshPropertyMapBaker.h b/Engine/Plugins/Experimental/GeometryProcessing/Source/DynamicMesh/Public/Sampling/MeshPropertyMapBaker.h
--- a/Engine/Plugins/Experimental/GeometryProcessing/Source/DynamicMesh/Public/Sampling/MeshPropertyMapBaker.h
+++ b/Engine/Plugins/Experimental/GeometryProcessing/Source/DynamicMesh/Public/Sampling/MeshPropertyMapBaker.h
@@ -49,6 +49,35 @@ public:
 
 protected:
 	TUniquePtr<TImageBuilder<FVector3f>> ResultBuilder;
+
+public:
+	//
+	// New Baker interface
+	//
+
+	/** Invoked at start of bake to initialize baker. */
+	virtual void PreEvaluate(const FMeshMapBaker& Baker) override;
+
+	/** Evaluate the sample at this mesh correspondence point. */
+	virtual FVector4f EvaluateSample(const FMeshMapBaker& Baker, const FCorrespondenceSample& Sample) override;
+
+	/** @return the default sample value for the baker. */
+	virtual FVector4f DefaultSample() const override { return FVector4f(0.0f, 0.0f, 0.0f, 1.0f); }
+
+	/** @return true if this baker supports multisampling. */
+	virtual bool SupportsMultisampling() const override { return Property != EMeshPropertyBakeType::MaterialID; }
+
+protected:
+	// Cached data
+	const FDynamicMesh3* DetailMesh = nullptr;
+	const FDynamicMeshNormalOverlay* DetailNormalOverlay = nullptr;
+	const FDynamicMeshUVOverlay* DetailUVOverlay = nullptr;
+	FAxisAlignedBox3d Bounds;
+	FVector3f DefaultValue = FVector3f::Zero();
+
+private:
+	template <class SampleType>
+	FVector3f SampleFunction(const SampleType& Sample);
 };
 
 } // end namespace UE::Geometry
diff --git a/Engine/Plugins/Experimental/GeometryProcessing/Source/DynamicMesh/Public/Sampling/MeshResampleImageBaker.h b/Engine/Plugins/Experimental/GeometryProcessing/Source/DynamicMesh/Public/Sampling/MeshResampleImageBaker.h
--- a/Engine/Plugins/Experimental/GeometryProcessing/Source/DynamicMesh/Public/Sampling/MeshResampleImageBaker.h
+++ b/Engine/Plugins/Experimental/GeometryProcessing/Source/DynamicMesh/Public/Sampling/MeshResampleImageBaker.h
@@ -42,6 +42,26 @@ public:
 
 protected:
 	TUniquePtr<TImageBuilder<FVector4f>> ResultBuilder;
+
+public:
+	//
+	// New Baker interface
+	//
+
+	/** Invoked at start of bake to initialize sampler. */
+	virtual void PreEvaluate(const FMeshMapBaker& Baker) override;
+
+	virtual FVector4f EvaluateSample(const FMeshMapBaker& Baker, const FCorrespondenceSample& Sample) override;
+
+	virtual FVector4f DefaultSample() const override { return FVector4f(0.0f, 0.0f, 0.0f, 1.0f); }
+
+protected:
+	// Cached data
+	const FDynamicMesh3* DetailMesh = nullptr;
+
+private:
+	template <class SampleType>
+	FVector4f PropertySampleFunction(const SampleType& Sample);
 };
 
 
@@ -58,6 +78,24 @@ protected:
 
 	void InitResult();
 	void BakeMaterial(int32 MaterialID);
+
+public:
+	//
+	// New Baker interface
+	//
+
+	/** Invoked at start of bake to initialize baker. */
+	virtual void PreEvaluate(const FMeshMapBaker& Baker) override;
+
+	/** Evaluate the sample at this mesh correspondence point. */
+	virtual FVector4f EvaluateSample(const FMeshMapBaker& Baker, const FCorrespondenceSample& Sample) override;
+
+	/** @return the default sample value for the baker. */
+	virtual FVector4f DefaultSample() const override { return FVector4f(0.0f, 0.0f, 0.0f, 1.0f); }
+
+protected:
+	const FDynamicMeshMaterialAttribute* DetailMaterialIDAttrib = nullptr;
+	bool bValidDetailMesh = false;
 };
 
 } // end namespace UE::Geometry
diff --git a/Engine/Plugins/Experimental/GeometryProcessing/Source/GeometricObjects/Public/Image/ImageDimensions.h b/Engine/Plugins/Experimental/GeometryProcessing/Source/GeometricObjects/Public/Image/ImageDimensions.h
--- a/Engine/Plugins/Experimental/GeometryProcessing/Source/GeometricObjects/Public/Image/ImageDimensions.h
+++ b/Engine/Plugins/Experimental/GeometryProcessing/Source/GeometricObjects/Public/Image/ImageDimensions.h
@@ -20,12 +20,28 @@ protected:
 	int32 Width = 0;
 	int32 Height = 0;
 
+	bool bIsTile = false;
+	int32 SourceWidth = 0;
+	int32 SourceHeight = 0;
+	FVector2i SourceOffset = FVector2i::Zero();
+
 public:
 	FImageDimensions(int32 WidthIn = 0, int32 HeightIn = 0)
+		: SourceWidth(WidthIn), SourceHeight(HeightIn)
+	{
+		SetDimensions(WidthIn, HeightIn);
+	}
+
+	/** Tile constructor. */
+	FImageDimensions(int32 WidthIn, int32 HeightIn, int32 SourceWidthIn, int32 SourceHeightIn, FVector2i SourceOffsetIn)
+		: bIsTile(true), SourceWidth(SourceWidthIn), SourceHeight(SourceHeightIn), SourceOffset(SourceOffsetIn)
 	{
 		SetDimensions(WidthIn, HeightIn);
+		checkSlow(Width >= 0 && Width <= SourceWidth && SourceOffset.X >= 0 && Width + SourceOffset.X <= SourceWidth);
+		checkSlow(Height >= 0 && Height <= SourceHeight && SourceOffset.Y >= 0 && Height + SourceOffset.Y <= SourceHeight);
 	}
 
+	/** Set the dimensions of this image. */
 	void SetDimensions(int32 WidthIn, int32 HeightIn)
 	{
 		check(WidthIn >= 0 && HeightIn >= 0);
@@ -41,12 +57,12 @@ public:
 	int64 Num() const { return (int64)Width * (int64)Height; }
 
 	/** @return true if image is square */
-	bool IsSquare() const { return Width == Height; }
+	bool IsSquare() const { return bIsTile ? SourceWidth == SourceHeight : Width == Height; }
 
 	/** @return true if coordinates are valid, ie in-bounds of image dimensions */
 	bool IsValidCoords(const FVector2i& Coords) const
 	{
-		return (Coords.X >= 0 && Coords.X < Width&& Coords.Y >= 0 && Coords.Y < Height);
+		return (Coords.X >= 0 && Coords.X < Width && Coords.Y >= 0 && Coords.Y < Height);
 	}
 
 	/** Clamp input coordinates to valid range of image coordinates */
@@ -95,15 +111,24 @@ public:
 	/** @return Real-valued dimensions of a pixel/texel in the image, relative to default UV space [0..1]^2 */
 	FVector2d GetTexelSize() const
 	{
-		return FVector2d(1.0 / (double)Width, 1.0 / (double)Height);
+		return FVector2d(1.0 / (double)SourceWidth, 1.0 / (double)SourceHeight);
 	}
 
 	/** @return Real-valued position of given texel center in default UV-space [0..1]^2 */
 	FVector2d GetTexelUV(const FVector2i& Coords) const
 	{
-		return FVector2d(
-			((double)Coords.X + 0.5) / (double)Width,
-			((double)Coords.Y + 0.5) / (double)Height);
+		if (bIsTile)
+		{
+			return FVector2d(
+				((double)SourceOffset.X + (double)Coords.X + 0.5) / (double)SourceWidth,
+				((double)SourceOffset.Y + (double)Coords.Y + 0.5) / (double)SourceWidth);
+		}
+		else
+		{
+			return FVector2d(
+				((double)Coords.X + 0.5) / (double)Width,
+				((double)Coords.Y + 0.5) / (double)Height);
+		}
 	}
 
 	/** @return Real-valued position of given texel center in default UV-space [0..1]^2 */
@@ -115,30 +140,129 @@ public:
 	/** @return integer XY coordinates for real-valued XY coordinates (ie texel that contains value, if texel origin is in bottom-left */
 	FVector2i PixelToCoords(const FVector2d& PixelPosition) const
 	{
-		int32 X = FMath::Clamp((int32)PixelPosition.X, 0, Width - 1);
-		int32 Y = FMath::Clamp((int32)PixelPosition.Y, 0, Height - 1);
-		return FVector2i(X, Y);
+		if (bIsTile)
+		{
+			int32 X = FMath::Clamp((int32)PixelPosition.X, 0, SourceWidth - 1);
+			int32 Y = FMath::Clamp((int32)PixelPosition.Y, 0, SourceHeight - 1);
+			return FVector2i(X, Y);
+		}
+		else
+		{
+			int32 X = FMath::Clamp((int32)PixelPosition.X, 0, Width - 1);
+			int32 Y = FMath::Clamp((int32)PixelPosition.Y, 0, Height - 1);
+			return FVector2i(X, Y);
+		}
 	}
 
 	/** @return integer XY coordinates for UV coordinates, assuming default UV space [0..1]^2 */
 	FVector2i UVToCoords(const FVector2d& UVPosition) const
 	{
-		double PixelX = (UVPosition.X * (double)Width /*- 0.5*/);
-		double PixelY = (UVPosition.Y * (double)Height /*- 0.5*/);
-		return PixelToCoords(FVector2d(PixelX, PixelY));
+		if (bIsTile)
+		{
+			double PixelX = (UVPosition.X * (double)SourceWidth /*- 0.5*/);
+			double PixelY = (UVPosition.Y * (double)SourceHeight /*- 0.5*/);
+			return PixelToCoords(FVector2d(PixelX, PixelY));
+		}
+		else
+		{
+			double PixelX = (UVPosition.X * (double)Width /*- 0.5*/);
+			double PixelY = (UVPosition.Y * (double)Height /*- 0.5*/);
+			return PixelToCoords(FVector2d(PixelX, PixelY));
+		}
 	}
 
 	bool operator==(const FImageDimensions& Other) const
 	{
-		return Width == Other.Width && Height == Other.Height;
+		return Width == Other.Width && Height == Other.Height &&
+			bIsTile == Other.bIsTile &&
+			SourceWidth == Other.SourceWidth &&
+			SourceHeight == Other.SourceHeight &&
+			SourceOffset == Other.SourceOffset;
 	}
 
 	bool operator!=(const FImageDimensions& Other) const
 	{
-		return Width != Other.Width || Height != Other.Height;
+		return !(*this == Other);
+	}
+
+	// Tiling
+
+	/** @return true if this image is a tile. */
+	bool IsTile() const
+	{
+		return bIsTile;
+	}
+
+	/** @return source image dimensions for this tile. */
+	FImageDimensions GetSourceDimensions() const
+	{
+		return FImageDimensions(SourceWidth, SourceHeight);
+	}
+
+	/** @return source image linear index from tile coordinates. */
+	int64 GetSourceIndex(const FVector2i& Coords) const
+	{
+		FVector2i SourceCoords = GetSourceCoords(Coords.X, Coords.Y);
+		return GetSourceDimensions().GetIndex(SourceCoords);
+	}
+
+	/** @return source image coordinates from tile coordinates. */
+	FVector2i GetSourceCoords(int32 TileX, int32 TileY) const
+	{
+		checkSlow(IsValidCoords(FVector2i(TileX, TileY)));
+		return FVector2i(SourceOffset.X + TileX, SourceOffset.Y + TileY);
 	}
 };
 
+class FImageTiling
+{
+private:
+	FImageDimensions Dimensions;
+	int32 TileWidth = 32;
+	int32 TileHeight = 32;
+	int32 TilePadding = 0;
+
+public:
+	FImageTiling(FImageDimensions DimensionsIn, int32 TileWidthIn, int32 TileHeightIn, int32 TilePaddingIn = 0)
+	{
+		SetTiling(DimensionsIn, TileWidthIn, TileHeightIn, TilePaddingIn);
+	}
+
+	void SetTiling(FImageDimensions DimensionsIn, int32 TileWidthIn, int32 TileHeightIn, int32 TilePaddingIn)
+	{
+		checkSlow(TileWidthIn >= 1 && TileHeightIn >= 1 && TilePaddingIn >= 0);
+		Dimensions = DimensionsIn;
+		TileWidth = FMath::Clamp(TileWidthIn, 1, Dimensions.GetWidth());
+		TileHeight = FMath::Clamp(TileHeightIn, 1, Dimensions.GetHeight());
+		TilePadding = TilePaddingIn < 0 ? 0 : TilePaddingIn;
+	}
+
+	int32 NumTilesX() const
+	{
+		return Dimensions.GetWidth() / TileWidth + (Dimensions.GetWidth() % TileWidth == 0 ? 0 : 1);
+	}
+
+	int32 NumTilesY() const
+	{
+		return Dimensions.GetHeight() / TileHeight + (Dimensions.GetHeight() % TileHeight == 0 ? 0 : 1);
+	}
+
+	int32 Num() const
+	{
+		return NumTilesX() * NumTilesY();
+	}
+
+	/** @return the tile corresponding to the linear index [0, NumTiles()-1]. */
+	FImageDimensions GetTile(int32 LinearTileIndex) const
+	{
+		FVector2i TileIndex(LinearTileIndex % NumTilesX(), LinearTileIndex / NumTilesX());
+		FVector2i TileOffset(TileIndex.X * TileWidth, TileIndex.Y * TileHeight);
+		int32 ClampedWidth = FMath::Clamp(TileWidth, 0, Dimensions.GetWidth() - TileIndex.X * TileWidth);
+		int32 ClampedHeight = FMath::Clamp(TileHeight, 0, Dimensions.GetHeight() - TileIndex.Y * TileHeight);
+		// TODO: Padding support
+		return FImageDimensions(ClampedWidth, ClampedHeight, Dimensions.GetWidth(), Dimensions.GetHeight(), TileOffset);
+	}
+};
 
 } // end namespace UE::Geometry
 } // end namespace UE
diff --git a/Engine/Plugins/Experimental/GeometryProcessing/Source/GeometricObjects/Public/Image/ImageOccupancyMap.h b/Engine/Plugins/Experimental/GeometryProcessing/Source/GeometricObjects/Public/Image/ImageOccupancyMap.h
--- a/Engine/Plugins/Experimental/GeometryProcessing/Source/GeometricObjects/Public/Image/ImageOccupancyMap.h
+++ b/Engine/Plugins/Experimental/GeometryProcessing/Source/GeometricObjects/Public/Image/ImageOccupancyMap.h
@@ -6,6 +6,7 @@
 #include "VectorTypes.h"
 #include "Image/ImageDimensions.h"
 #include "Spatial/MeshAABBTree3.h"
+#include "Sampling/GridSampler.h"
 
 namespace UE
 {
@@ -40,9 +41,10 @@ public:
 	//const int8 BorderTexel = 2;
 	const int8 GutterTexel = 3;
 
-	/** Texel Type for each texel in image, Size = Width x Height */
+	/** Texel Type for each texel in image, Size = Width x Height x Multisamples */
 	TArray64<int8> TexelType;
-
+	/** Count of samples in the texel that are interior. */
+	TArray64<int32> TexelInteriorSamples;
 	/** UV for each texel in image. Only set for Interior texels */
 	TArray64<FVector2f> TexelQueryUV;
 	/** integer/Triangle ID for each texel in image. Only set for Interior texels. */
@@ -52,22 +54,33 @@ public:
 	    Gutter can be filled by directly copying from source to target. */
 	TArray64<TTuple<int64, int64>> GutterTexels;
 
+	using FGridSampler = TGridSampler<double>;
+	FGridSampler Multisampler = FGridSampler(1);
+
+	bool bParallel = true;
+
 
-	void Initialize(FImageDimensions DimensionsIn)
+	void Initialize(FImageDimensions DimensionsIn, int32 Multisampling = 1)
 	{
 		check(DimensionsIn.IsSquare()); // are we sure it works otherwise?
 		Dimensions = DimensionsIn;
+		Multisampler = FGridSampler(Multisampling);
 	}
 
-
 	/**
-	 * @return true if texel at this linear index is an Interior texel
+	 * @return true if texel sample at this sample linear index is an Interior texel
 	 */
 	bool IsInterior(int64 LinearIndex) const
 	{
 		return TexelType[LinearIndex] == InteriorTexel;
 	}
 
+	/** @return number of interior samples for the texel at this image linear index. */
+	int32 TexelNumSamples(int64 LinearIndex) const
+	{
+		return TexelInteriorSamples[LinearIndex];
+	}
+
 	/** Set texel type at the given X/Y */
 	//void SetTexelType(int64 X, int64 Y, int8 Type)
 	//{
@@ -82,12 +95,15 @@ public:
 		// make flat mesh
 		TMeshAABBTree3<MeshType> FlatSpatial(&UVSpaceMesh, true);
 
-		TexelType.Init(EmptyTexel, Dimensions.Num());
-		TexelQueryUV.Init(FVector2f::Zero(), Dimensions.Num());
-		TexelQueryTriangle.Init(IndexConstants::InvalidID, Dimensions.Num());
+		const int32 LinearImageSize = Dimensions.Num();
+		const int32 LinearSampleSize = Dimensions.Num() * Multisampler.Num();
+		TexelType.Init(EmptyTexel, LinearSampleSize);
+		TexelInteriorSamples.Init(0, LinearImageSize);
+		TexelQueryUV.Init(FVector2f::Zero(), LinearSampleSize);
+		TexelQueryTriangle.Init(IndexConstants::InvalidID, LinearSampleSize);
 
-		FVector2d TexelDims = Dimensions.GetTexelSize();
-		double TexelDiag = TexelDims.Length();
+		const FVector2d TexelSize = Dimensions.GetTexelSize();
+		const double TexelDiag = TexelSize.Length();
 
 		IMeshSpatial::FQueryOptions QueryOptions;
 		QueryOptions.MaxDistance = (double)GutterSize * TexelDiag;
@@ -96,56 +112,99 @@ public:
 		TAtomic<int64> InteriorCounter;
 		FCriticalSection GutterLock;
 		ParallelFor(Dimensions.GetHeight(), 
-			[this, &UVSpaceMesh, &GetTriangleIDFunc, &FlatSpatial, TexelDiag, &QueryOptions, &InteriorCounter, &GutterLock]
+			[this, &UVSpaceMesh, &GetTriangleIDFunc, &FlatSpatial, TexelDiag, &QueryOptions, &InteriorCounter, &GutterLock, &TexelSize]
 		(int32 ImgY)
 		{
+			// To minimize locks, accumulate gutter texel data locally and append
+			// to our gutter texel array at the end.
+			TArray64<TTuple<int64, int64>> LocalGutterTexels;
+			int64 LocalInteriorCounter = 0;
 			for (int32 ImgX = 0; ImgX < Dimensions.GetWidth(); ++ImgX)
 			{
-				int64 LinearIdx = Dimensions.GetIndex(ImgX, ImgY);
-				FVector2d UVPoint = Dimensions.GetTexelUV(LinearIdx);
-				FVector3d UVPoint3d(UVPoint.X, UVPoint.Y, 0);
-
-				double NearDistSqr;
-				int32 NearestTriID = FlatSpatial.FindNearestTriangle(UVPoint3d, NearDistSqr, QueryOptions);
-				if (NearestTriID >= 0)
+				// With multiple samples, a texel may contain multiple gutter samples.
+				// Since we copy nearest interior texel over top our gutter texels, we
+				// should only consider a texel a gutter texel iff all samples in a
+				// texel are gutter texels. This bool will track that as we iterate
+				// over each sample of the texel.
+				bool bIsGutterTexel = true;
+				TTuple<int64, int64> GutterNearestTexel;
+				for (int32 Sample = 0; Sample < Multisampler.Num(); ++Sample)
 				{
-					FVector3d A, B, C;
-					UVSpaceMesh.GetTriVertices(NearestTriID, A, B, C);
-					FTriangle2d UVTriangle(GetXY(A), GetXY(B), GetXY(C));
-
-					if (UVTriangle.IsInsideOrOn(UVPoint))
-					{
-						TexelType[LinearIdx] = InteriorTexel;
-						TexelQueryUV[LinearIdx] = (FVector2f)UVPoint;
-						TexelQueryTriangle[LinearIdx] = GetTriangleIDFunc(NearestTriID);
-						InteriorCounter.IncrementExchange();
-					}
-					else if (NearDistSqr < TexelDiag * TexelDiag)
+					const int64 TexelLinearIdx = Dimensions.GetIndex(ImgX, ImgY);
+					const int64 SampleLinearIdx = TexelLinearIdx * Multisampler.Num() + Sample;
+					const FVector2d SampleUV = Multisampler.Sample(Sample);
+					FVector2d UVPoint = Dimensions.GetTexelUV(TexelLinearIdx);
+					UVPoint = UVPoint - 0.5 * TexelSize + SampleUV * TexelSize;
+					const FVector3d UVPoint3d(UVPoint.X, UVPoint.Y, 0);
+
+					double NearDistSqr;
+					const int32 NearestTriID = FlatSpatial.FindNearestTriangle(UVPoint3d, NearDistSqr, QueryOptions);
+					if (NearestTriID >= 0)
 					{
-						FDistPoint3Triangle3d DistQuery = TMeshQueries<MeshType>::TriangleDistance(UVSpaceMesh, NearestTriID, UVPoint3d);
-						FVector2d NearestUV = GetXY(DistQuery.ClosestTrianglePoint);
-						// nudge point into triangle to improve numerical behavior of things like barycentric coord calculation
-						NearestUV += (10.0 * FMathf::ZeroTolerance) * Normalized(NearestUV - UVPoint);
-
-						TexelType[LinearIdx] = InteriorTexel;
-						TexelQueryUV[LinearIdx] = (FVector2f)NearestUV;
-						TexelQueryTriangle[LinearIdx] = GetTriangleIDFunc(NearestTriID);
+						FVector3d A, B, C;
+						UVSpaceMesh.GetTriVertices(NearestTriID, A, B, C);
+						const FTriangle2d UVTriangle(GetXY(A), GetXY(B), GetXY(C));
+
+						if (UVTriangle.IsInsideOrOn(UVPoint))
+						{
+							TexelType[SampleLinearIdx] = InteriorTexel;
+							TexelQueryUV[SampleLinearIdx] = (FVector2f)UVPoint;
+							TexelQueryTriangle[SampleLinearIdx] = GetTriangleIDFunc(NearestTriID);
+							++TexelInteriorSamples[TexelLinearIdx];
+							++LocalInteriorCounter;
+							bIsGutterTexel = false;
+						}
+						else if (NearDistSqr < TexelDiag * TexelDiag)
+						{
+							const FDistPoint3Triangle3d DistQuery = TMeshQueries<MeshType>::TriangleDistance(UVSpaceMesh, NearestTriID, UVPoint3d);
+							FVector2d NearestUV = GetXY(DistQuery.ClosestTrianglePoint);
+							// nudge point into triangle to improve numerical behavior of things like barycentric coord calculation
+							NearestUV += (10.0 * FMathf::ZeroTolerance) * Normalized(NearestUV - UVPoint);
+
+							TexelType[SampleLinearIdx] = InteriorTexel;
+							TexelQueryUV[SampleLinearIdx] = (FVector2f)NearestUV;
+							TexelQueryTriangle[SampleLinearIdx] = GetTriangleIDFunc(NearestTriID);
+							++TexelInteriorSamples[TexelLinearIdx];
+							++LocalInteriorCounter;
+							bIsGutterTexel = false;
+						}
+						else
+						{
+							TexelType[SampleLinearIdx] = GutterTexel;
+							const FDistPoint3Triangle3d DistQuery = TMeshQueries<MeshType>::TriangleDistance(UVSpaceMesh, NearestTriID, UVPoint3d);
+							const FVector2d NearestUV = GetXY(DistQuery.ClosestTrianglePoint);
+							const FVector2i NearestCoords = Dimensions.UVToCoords(NearestUV);
+
+							// To avoid artifacts with mipmapping the gutter texels store the nearest
+							// interior texel and a post pass copies those nearest interior texel values
+							// to each gutter pixel.
+							//
+							// The gutter pixel structure stores pixel indices to copy in the source
+							// image's linear index.
+							//
+							// For tiled images:
+							//    NearestCoords are in source space --> GetSource().GetIndex() will return source linear index.
+							//    ImgX, ImgY are in local tile space --> GetSourceIndex() will return source linear index.
+							const int64 NearestLinearIdx = Dimensions.IsTile() ? Dimensions.GetSourceDimensions().GetIndex(NearestCoords) : Dimensions.GetIndex(NearestCoords);
+							const int32 LinearSourceIdx = Dimensions.IsTile() ? Dimensions.GetSourceIndex(FVector2i(ImgX, ImgY)) : TexelLinearIdx;
+
+							GutterNearestTexel = TTuple<int64, int64>(LinearSourceIdx, NearestLinearIdx);
+						}
 					}
-					else
+					if (bIsGutterTexel)
 					{
-						TexelType[LinearIdx] = GutterTexel;
-						FDistPoint3Triangle3d DistQuery = TMeshQueries<MeshType>::TriangleDistance(UVSpaceMesh, NearestTriID, UVPoint3d);
-						FVector2d NearestUV = GetXY(DistQuery.ClosestTrianglePoint);
-						FVector2i NearestCoords = Dimensions.UVToCoords(NearestUV);
-						int64 NearestLinearIdx = Dimensions.GetIndex(NearestCoords);
-
-						GutterLock.Lock();
-						GutterTexels.Add(TTuple<int64, int64>(LinearIdx, NearestLinearIdx));
-						GutterLock.Unlock();
+						LocalGutterTexels.Add(GutterNearestTexel);
 					}
 				}
 			}
-		});
+			if (LocalGutterTexels.Num() > 0)
+			{
+				GutterLock.Lock();
+				GutterTexels += LocalGutterTexels;
+				GutterLock.Unlock();
+			}
+			InteriorCounter += LocalInteriorCounter;
+		}, !bParallel ? EParallelForFlags::ForceSingleThread : EParallelForFlags::None);
 
 		return true;
 	}
diff --git a/Engine/Plugins/Experimental/GeometryProcessing/Source/GeometricObjects/Public/Sampling/GridSampler.h b/Engine/Plugins/Experimental/GeometryProcessing/Source/GeometricObjects/Public/Sampling/GridSampler.h
new file mode 100644
--- /dev/null
+++ b/Engine/Plugins/Experimental/GeometryProcessing/Source/GeometricObjects/Public/Sampling/GridSampler.h
@@ -0,0 +1,52 @@
+// Copyright Epic Games, Inc. All Rights Reserved.
+
+#pragma once
+
+#include "MathUtil.h"
+#include "VectorTypes.h"
+
+namespace UE
+{
+namespace Geometry
+{
+
+/**
+ * Given a linear index into a square grid of given size, gives back the [0,1)
+ * XY coordinate of the grid cell center (assuming linear index increases first
+ * in X and then in Y).
+ */
+template<typename RealType>
+class TGridSampler
+{
+private:
+	int32 Size = 1;
+
+public:
+	TGridSampler(int32 SizeIn)
+		: Size(SizeIn)
+	{
+		checkSlow(Size >= 1);
+	}
+
+	/** @return total number of grid points. */
+	int32 Num() const
+	{
+		return Size * Size;
+	}
+
+	/**
+	 * @param Index linear point index in range [0,Num()-1]
+	 * @return grid sample for given linear index
+	 */
+	FVector2<RealType> Sample(int32 Index) const
+	{
+		checkSlow(Index >= 0 && Index < Num());
+		RealType X = (static_cast<RealType>(Index % Size) + RealType{ 1 } / RealType{ 2 }) / static_cast<RealType>(Size);
+		RealType Y = (static_cast<RealType>(Index / Size) + RealType{ 1 } / RealType{ 2 }) / static_cast<RealType>(Size);
+		return FVector2<RealType>(X, Y);
+	}
+};
+
+} // end namespace UE::Geometry
+} // end namespace UE
+
diff --git a/Engine/Plugins/Experimental/MeshModelingToolset/Source/MeshModelingTools/Private/BakeMeshAttributeMapsTool.cpp b/Engine/Plugins/Experimental/MeshModelingToolset/Source/MeshModelingTools/Private/BakeMeshAttributeMapsTool.cpp
--- a/Engine/Plugins/Experimental/MeshModelingToolset/Source/MeshModelingTools/Private/BakeMeshAttributeMapsTool.cpp
+++ b/Engine/Plugins/Experimental/MeshModelingToolset/Source/MeshModelingTools/Private/BakeMeshAttributeMapsTool.cpp
@@ -86,33 +86,32 @@ TArray<FString> UBakeMeshAttributeMapsToolProperties::GetUVLayerNamesFunc()
  * Operators
  */
 
-class FBakeMapBaseOp : public TGenericDataOperator<FMeshImageBaker>
+class FBakeMapBaseOp : public TGenericDataOperator<FMeshMapBaker>
 {
 public:
+	virtual ~FBakeMapBaseOp() {}
+
 	TSharedPtr<UE::Geometry::FDynamicMesh3, ESPMode::ThreadSafe> DetailMesh;
 	TSharedPtr<UE::Geometry::FDynamicMeshAABBTree3, ESPMode::ThreadSafe> DetailSpatial;
 	UE::Geometry::FDynamicMesh3* BaseMesh;
-	TPimplPtr<UE::Geometry::FMeshImageBakingCache> BakeCache;
+	TUniquePtr<UE::Geometry::FMeshMapBaker> Baker;
 	UBakeMeshAttributeMapsTool::FBakeCacheSettings BakeCacheSettings;
 
-public:
-	virtual ~FBakeMapBaseOp() {}
-
 	//
 	// TGenericDataOperator implementation
-	// 
+	//
 	virtual void CalculateResult(FProgressCancel* Progress) override
 	{
-		BakeCache = MakePimpl<FMeshImageBakingCache>();
-		BakeCache->CancelF = [Progress]() {
+		Baker = MakeUnique<FMeshMapBaker>();
+		Baker->CancelF = [Progress]() {
 			return Progress && Progress->Cancelled();
 		};
-		BakeCache->SetDetailMesh(DetailMesh.Get(), DetailSpatial.Get());
-		BakeCache->SetBakeTargetMesh(BaseMesh);
-		BakeCache->SetDimensions(BakeCacheSettings.Dimensions);
-		BakeCache->SetUVLayer(BakeCacheSettings.UVLayer);
-		BakeCache->SetThickness(BakeCacheSettings.Thickness);
-		BakeCache->ValidateCache();
+		Baker->SetTargetMesh(BaseMesh);
+		Baker->SetDetailMesh(DetailMesh.Get(), DetailSpatial.Get());
+		Baker->SetDimensions(BakeCacheSettings.Dimensions);
+		Baker->SetUVLayer(BakeCacheSettings.UVLayer);
+		Baker->SetThickness(BakeCacheSettings.Thickness);
+		Baker->SetMultisampling(BakeCacheSettings.Multisampling);
 	}
 };
 
@@ -122,7 +121,7 @@ public:
 	typedef FBakeMapBaseOp Super;
 
 	// inputs
-	const TMeshTangents<double>* BaseMeshTangents = nullptr;
+	TSharedPtr<UE::Geometry::TMeshTangents<double>, ESPMode::ThreadSafe> BaseMeshTangents;
 	UBakeMeshAttributeMapsTool::FNormalMapSettings Settings;
 
 public:
@@ -140,9 +139,9 @@ public:
 			return;
 		}
 
-		TUniquePtr<FMeshNormalMapBaker> Baker = MakeUnique<FMeshNormalMapBaker>();
-		Baker->SetCache(BakeCache.Get());
-		Baker->BaseMeshTangents = BaseMeshTangents;
+		TSharedPtr<FMeshNormalMapBaker, ESPMode::ThreadSafe> NormalBaker = MakeShared<FMeshNormalMapBaker, ESPMode::ThreadSafe>();
+		Baker->AddBaker(NormalBaker);
+		Baker->SetTargetMeshTangents(BaseMeshTangents);
 		Baker->Bake();
 		SetResult(MoveTemp(Baker));
 	}
@@ -154,7 +153,7 @@ public:
 	typedef FBakeMapBaseOp Super;
 
 	// inputs
-	const TMeshTangents<double>* BaseMeshTangents = nullptr;
+	TSharedPtr<UE::Geometry::TMeshTangents<double>, ESPMode::ThreadSafe> BaseMeshTangents;
 	UBakeMeshAttributeMapsTool::FOcclusionMapSettings Settings;
 
 public:
@@ -172,33 +171,45 @@ public:
 			return;
 		}
 
-		TUniquePtr<FMeshOcclusionMapBaker> Baker = MakeUnique<FMeshOcclusionMapBaker>();
-		Baker->SetCache(BakeCache.Get());
-		Baker->OcclusionType = EOcclusionMapType::All;
-        Baker->BaseMeshTangents = BaseMeshTangents;
-		Baker->NumOcclusionRays = Settings.OcclusionRays;
-        Baker->MaxDistance = Settings.MaxDistance;
-        Baker->SpreadAngle = Settings.SpreadAngle;
-		Baker->BlurRadius = Settings.BlurRadius;
-		Baker->BiasAngleDeg = Settings.BiasAngle;
-        switch (Settings.Distribution)
-		{
-		case EOcclusionMapDistribution::Cosine:
-			Baker->Distribution = FMeshOcclusionMapBaker::EDistribution::Cosine;
-			break;
-		case EOcclusionMapDistribution::Uniform:
-			Baker->Distribution = FMeshOcclusionMapBaker::EDistribution::Uniform;
-			break;
-		}
-        switch (Settings.NormalSpace)
-		{
-		case ENormalMapSpace::Tangent:
-			Baker->NormalSpace = FMeshOcclusionMapBaker::ESpace::Tangent;
-			break;
-		case ENormalMapSpace::Object:
-			Baker->NormalSpace = FMeshOcclusionMapBaker::ESpace::Object;
-			break;
-		}
+		auto InitializeBaker = [this](TSharedPtr<FMeshOcclusionMapBaker>& BakerIn) {
+			BakerIn->NumOcclusionRays = Settings.OcclusionRays;
+			BakerIn->MaxDistance = Settings.MaxDistance;
+			BakerIn->SpreadAngle = Settings.SpreadAngle;
+			BakerIn->BlurRadius = Settings.BlurRadius;
+			BakerIn->BiasAngleDeg = Settings.BiasAngle;
+
+			switch (Settings.Distribution)
+			{
+			case EOcclusionMapDistribution::Cosine:
+				BakerIn->Distribution = FMeshOcclusionMapBaker::EDistribution::Cosine;
+				break;
+			case EOcclusionMapDistribution::Uniform:
+				BakerIn->Distribution = FMeshOcclusionMapBaker::EDistribution::Uniform;
+				break;
+			}
+
+			switch (Settings.NormalSpace)
+			{
+			case ENormalMapSpace::Tangent:
+				BakerIn->NormalSpace = FMeshOcclusionMapBaker::ESpace::Tangent;
+				break;
+			case ENormalMapSpace::Object:
+				BakerIn->NormalSpace = FMeshOcclusionMapBaker::ESpace::Object;
+				break;
+			}
+		};
+
+		TSharedPtr<FMeshOcclusionMapBaker> OcclusionBaker = MakeShared<FMeshOcclusionMapBaker>();
+		InitializeBaker(OcclusionBaker);
+		OcclusionBaker->OcclusionType = EOcclusionMapType::AmbientOcclusion;
+
+		TSharedPtr<FMeshOcclusionMapBaker> BentNormalBaker = MakeShared<FMeshOcclusionMapBaker>();
+		InitializeBaker(BentNormalBaker);
+		BentNormalBaker->OcclusionType = EOcclusionMapType::BentNormal;
+
+		Baker->SetTargetMeshTangents(BaseMeshTangents);
+		Baker->AddBaker(OcclusionBaker);
+		Baker->AddBaker(BentNormalBaker);
 		Baker->Bake();
 		SetResult(MoveTemp(Baker));
 	}
@@ -227,14 +238,14 @@ public:
 			return;
 		}
 
-		TUniquePtr<FMeshCurvatureMapBaker> Baker = MakeUnique<FMeshCurvatureMapBaker>();
-		Baker->SetCache(BakeCache.Get());
-		Baker->RangeScale = FMathd::Clamp(Settings.RangeMultiplier, 0.0001, 1000.0);
-		Baker->MinRangeScale = FMathd::Clamp(Settings.MinRangeMultiplier, 0.0, 1.0);
-		Baker->UseCurvatureType = (FMeshCurvatureMapBaker::ECurvatureType)Settings.CurvatureType;
-		Baker->UseColorMode = (FMeshCurvatureMapBaker::EColorMode)Settings.ColorMode;
-		Baker->UseClampMode = (FMeshCurvatureMapBaker::EClampMode)Settings.ClampMode;
-		Baker->BlurRadius = Settings.BlurRadius;
+		TSharedPtr<FMeshCurvatureMapBaker> CurvatureBaker = MakeShared<FMeshCurvatureMapBaker>();
+		CurvatureBaker->RangeScale = FMathd::Clamp(Settings.RangeMultiplier, 0.0001, 1000.0);
+		CurvatureBaker->MinRangeScale = FMathd::Clamp(Settings.MinRangeMultiplier, 0.0, 1.0);
+		CurvatureBaker->UseCurvatureType = (FMeshCurvatureMapBaker::ECurvatureType)Settings.CurvatureType;
+		CurvatureBaker->UseColorMode = (FMeshCurvatureMapBaker::EColorMode)Settings.ColorMode;
+		CurvatureBaker->UseClampMode = (FMeshCurvatureMapBaker::EClampMode)Settings.ClampMode;
+		CurvatureBaker->BlurRadius = Settings.BlurRadius;
+		Baker->AddBaker(CurvatureBaker);
 		Baker->Bake();
 		SetResult(MoveTemp(Baker));
 	}
@@ -263,10 +274,9 @@ public:
 			return;
 		}
 
-		TUniquePtr<FMeshPropertyMapBaker> Baker = MakeUnique<FMeshPropertyMapBaker>();
-		Baker->SetCache(BakeCache.Get());
-		Baker->Property = (EMeshPropertyBakeType)Settings.PropertyTypeIndex;
-		//Baker->BlurRadius = Settings.BlurRadius;
+		TSharedPtr<FMeshPropertyMapBaker> PropertyBaker = MakeShared<FMeshPropertyMapBaker>();
+		PropertyBaker->Property = (EMeshPropertyBakeType)Settings.PropertyTypeIndex;
+		Baker->AddBaker(PropertyBaker);
 		Baker->Bake();
 		SetResult(MoveTemp(Baker));
 	}
@@ -297,12 +307,12 @@ public:
 			return;
 		}
 
-		TUniquePtr<FMeshResampleImageBaker> Baker = MakeUnique<FMeshResampleImageBaker>();
-		Baker->SetCache(BakeCache.Get());
-		Baker->DetailUVOverlay = UVOverlay;
-		Baker->SampleFunction = [this](FVector2d UVCoord) {
+		TSharedPtr<FMeshResampleImageBaker> ResampleBaker = MakeShared<FMeshResampleImageBaker>();
+		ResampleBaker->DetailUVOverlay = UVOverlay;
+		ResampleBaker->SampleFunction = [this](FVector2d UVCoord) {
 			return TextureImage->BilinearSampleUV<float>(UVCoord, FVector4f(0, 0, 0, 1));
 		};
+		Baker->AddBaker(ResampleBaker);
 		Baker->Bake();
 		SetResult(MoveTemp(Baker));
 	}
@@ -334,10 +344,10 @@ public:
 			return;
 		}
 
-		TUniquePtr<FMeshMultiResampleImageBaker> Baker = MakeUnique<FMeshMultiResampleImageBaker>();
-		Baker->SetCache(BakeCache.Get());
-		Baker->DetailUVOverlay = UVOverlay;
-		Baker->MultiTextures = MaterialToTextureImageMap;
+		TSharedPtr<FMeshMultiResampleImageBaker> TextureBaker = MakeShared<FMeshMultiResampleImageBaker>();
+		TextureBaker->DetailUVOverlay = UVOverlay;
+		TextureBaker->MultiTextures = MaterialToTextureImageMap;
+		Baker->AddBaker(TextureBaker);
 		Baker->Bake();
 		SetResult(MoveTemp(Baker));
 	}
@@ -347,11 +357,6 @@ public:
  * Tool
  */
 
-UBakeMeshAttributeMapsTool::UBakeMeshAttributeMapsTool()
-{
-}
-
-
 void UBakeMeshAttributeMapsTool::SetAssetAPI(IAssetGenerationAPI* AssetAPIIn)
 {
 	AssetAPI = AssetAPIIn;
@@ -440,6 +445,7 @@ void UBakeMeshAttributeMapsTool::Setup()
 	Settings->WatchProperty(Settings->UVLayer, [this](FString) { bInputsDirty = true; });
 	Settings->WatchProperty(Settings->bUseWorldSpace, [this](bool) { bDetailMeshValid = false; bInputsDirty = true; });
 	Settings->WatchProperty(Settings->Thickness, [this](float) { bInputsDirty = true; });
+	Settings->WatchProperty(Settings->Multisampling, [this](EBakeMultisampling) { bInputsDirty = true; });
 
 	NormalMapProps = NewObject<UBakedNormalMapToolProperties>(this);
 	NormalMapProps->RestoreProperties(this);
@@ -525,39 +531,38 @@ bool UBakeMeshAttributeMapsTool::CanAccept() const
 }
 
 
-TUniquePtr<UE::Geometry::TGenericDataOperator<FMeshImageBaker>> UBakeMeshAttributeMapsTool::MakeNewOperator()
+TUniquePtr<UE::Geometry::TGenericDataOperator<FMeshMapBaker>> UBakeMeshAttributeMapsTool::MakeNewOperator()
 {
-	switch (Settings->MapType)
-	{
-	case EBakeMapType::TangentSpaceNormalMap:
+	auto InitializeBaseOp = [this](FBakeMapBaseOp* Op)
 	{
-		TUniquePtr<FBakeNormalMapOp> Op = MakeUnique<FBakeNormalMapOp>();
 		Op->DetailMesh = DetailMesh;
 		Op->DetailSpatial = DetailSpatial;
 		Op->BaseMesh = &BaseMesh;
 		Op->BakeCacheSettings = CachedBakeCacheSettings;
-		Op->BaseMeshTangents = BaseMeshTangents.Get();
+	};
+
+	switch (Settings->MapType)
+	{
+	case EBakeMapType::TangentSpaceNormalMap:
+	{
+		TUniquePtr<FBakeNormalMapOp> Op = MakeUnique<FBakeNormalMapOp>();
+		InitializeBaseOp(Op.Get());
+		Op->BaseMeshTangents = BaseMeshTangents;
 		Op->Settings = CachedNormalMapSettings;
 		return Op;
 	}
 	case EBakeMapType::Occlusion:
 	{
 		TUniquePtr<FBakeOcclusionMapOp> Op = MakeUnique<FBakeOcclusionMapOp>();
-		Op->DetailMesh = DetailMesh;
-		Op->DetailSpatial = DetailSpatial;
-		Op->BaseMesh = &BaseMesh;
-		Op->BakeCacheSettings = CachedBakeCacheSettings;
-		Op->BaseMeshTangents = BaseMeshTangents.Get();
+		InitializeBaseOp(Op.Get());
+		Op->BaseMeshTangents = BaseMeshTangents;
 		Op->Settings = CachedOcclusionMapSettings;
 		return Op;
 	}
 	case EBakeMapType::Curvature:
 	{
 		TUniquePtr<FBakeCurvatureMapOp> Op = MakeUnique<FBakeCurvatureMapOp>();
-		Op->DetailMesh = DetailMesh;
-		Op->DetailSpatial = DetailSpatial;
-		Op->BaseMesh = &BaseMesh;
-		Op->BakeCacheSettings = CachedBakeCacheSettings;
+		InitializeBaseOp(Op.Get());
 		Op->Settings = CachedCurvatureMapSettings;
 		return Op;
 	}
@@ -567,20 +572,14 @@ TUniquePtr<UE::Geometry::TGenericDataOperator<FMeshImageBaker>> UBakeMeshAttribu
 	case EBakeMapType::MaterialID:
 	{
 		TUniquePtr<FBakeMeshPropertyMapOp> Op = MakeUnique<FBakeMeshPropertyMapOp>();
-		Op->DetailMesh = DetailMesh;
-		Op->DetailSpatial = DetailSpatial;
-		Op->BaseMesh = &BaseMesh;
-		Op->BakeCacheSettings = CachedBakeCacheSettings;
+		InitializeBaseOp(Op.Get());
 		Op->Settings = CachedMeshPropertyMapSettings;
 		return Op;
 	}
 	case EBakeMapType::Texture2DImage:
 	{
 		TUniquePtr<FBakeTexture2DImageMapOp> Op = MakeUnique<FBakeTexture2DImageMapOp>();
-		Op->DetailMesh = DetailMesh;
-		Op->DetailSpatial = DetailSpatial;
-		Op->BaseMesh = &BaseMesh;
-		Op->BakeCacheSettings = CachedBakeCacheSettings;
+		InitializeBaseOp(Op.Get());
 		Op->UVOverlay = DetailMesh->Attributes()->GetUVLayer(CachedTexture2DImageSettings.UVLayer);
 		Op->TextureImage = CachedTextureImage;
 		Op->Settings = CachedTexture2DImageSettings;
@@ -589,10 +588,7 @@ TUniquePtr<UE::Geometry::TGenericDataOperator<FMeshImageBaker>> UBakeMeshAttribu
 	case EBakeMapType::MultiTexture:
 	{
 		TUniquePtr<FBakeMultiTextureOp> Op = MakeUnique<FBakeMultiTextureOp>();
-		Op->DetailMesh = DetailMesh;
-		Op->DetailSpatial = DetailSpatial;
-		Op->BaseMesh = &BaseMesh;
-		Op->BakeCacheSettings = CachedBakeCacheSettings;
+		InitializeBaseOp(Op.Get());
 		Op->UVOverlay = DetailMesh->Attributes()->GetUVLayer(CachedTexture2DImageSettings.UVLayer);
 		Op->MaterialToTextureImageMap = CachedMultiTextures;
 		Op->Settings = CachedTexture2DImageSettings;
@@ -901,6 +897,7 @@ void UBakeMeshAttributeMapsTool::UpdateResult()
 	BakeCacheSettings.UVLayer = FCString::Atoi(*Settings->UVLayer);
 	BakeCacheSettings.DetailTimestamp = this->DetailMeshTimestamp;
 	BakeCacheSettings.Thickness = Settings->Thickness;
+	BakeCacheSettings.Multisampling = (int32)Settings->Multisampling;
 
 	// update bake cache settings
 	if (!(CachedBakeCacheSettings == BakeCacheSettings))
@@ -956,9 +953,9 @@ void UBakeMeshAttributeMapsTool::UpdateResult()
 	bool bInvalidate = bInputsDirty || (OpState == EOpState::Evaluate);
 	if (!Compute)
 	{
-		Compute = MakeUnique<TGenericDataBackgroundCompute<FMeshImageBaker>>();
+		Compute = MakeUnique<TGenericDataBackgroundCompute<FMeshMapBaker>>();
 		Compute->Setup(this);
-		Compute->OnResultUpdated.AddLambda([this](const TUniquePtr<FMeshImageBaker>& NewResult) { OnMapsUpdated(NewResult); });
+		Compute->OnResultUpdated.AddLambda([this](const TUniquePtr<FMeshMapBaker>& NewResult) { OnMapsUpdated(NewResult); });
 		Compute->InvalidateResult();
 	}
 	else if (bInvalidate)
@@ -1309,6 +1306,11 @@ UBakeMeshAttributeMapsTool::EOpState UBakeMeshAttributeMapsTool::UpdateResult_Mu
 			return EOpState::Invalid;
 		}
 	}
+	if (CachedMultiTextures.Num() == 0)
+	{
+		GetToolManager()->DisplayMessage(LOCTEXT("InvalidTextureWarning", "The Source Texture is not valid"), EToolMessageLevel::UserWarning);
+		return EOpState::Invalid;
+	}
 
 	if (!(CachedTexture2DImageSettings == NewSettings))
 	{
@@ -1420,50 +1422,39 @@ void UBakeMeshAttributeMapsTool::UpdateOnModeChange()
 }
 
 
-void UBakeMeshAttributeMapsTool::OnMapsUpdated(const TUniquePtr<FMeshImageBaker>& NewResult)
+void UBakeMeshAttributeMapsTool::OnMapsUpdated(const TUniquePtr<FMeshMapBaker>& NewResult)
 {
 	switch (Settings->MapType)
 	{
 	case EBakeMapType::TangentSpaceNormalMap:
 	{
-		FMeshNormalMapBaker* Baker = static_cast<FMeshNormalMapBaker*>(NewResult.Get());
 		FTexture2DBuilder TextureBuilder;
 		TextureBuilder.Initialize(FTexture2DBuilder::ETextureType::NormalMap, CachedNormalMapSettings.Dimensions);
-		TextureBuilder.Copy(*Baker->GetResult());
+		TextureBuilder.Copy(*NewResult->GetBakeResult(0));
 		TextureBuilder.Commit(false);
 		CachedNormalMap = TextureBuilder.GetTexture2D();
 		break;
 	}
 	case EBakeMapType::Occlusion:
 	{
-		FMeshOcclusionMapBaker* Baker = static_cast<FMeshOcclusionMapBaker*>(NewResult.Get());
-		const TUniquePtr<TImageBuilder<FVector3f>>& AOResult = Baker->GetResult(FMeshOcclusionMapBaker::EResult::AmbientOcclusion);
-		if (ensure(AOResult))
-		{
-			FTexture2DBuilder TextureBuilder;
-			TextureBuilder.Initialize(FTexture2DBuilder::ETextureType::AmbientOcclusion, CachedOcclusionMapSettings.Dimensions);
-			TextureBuilder.Copy(*AOResult);
-			TextureBuilder.Commit(false);
-			CachedOcclusionMap = TextureBuilder.GetTexture2D();
-		}
+		FTexture2DBuilder TextureBuilder;
+		TextureBuilder.Initialize(FTexture2DBuilder::ETextureType::AmbientOcclusion, CachedOcclusionMapSettings.Dimensions);
+		TextureBuilder.Copy(*NewResult->GetBakeResult(0));
+		TextureBuilder.Commit(false);
+		CachedOcclusionMap = TextureBuilder.GetTexture2D();
 
-		const TUniquePtr<TImageBuilder<FVector3f>>& BNResult = Baker->GetResult(FMeshOcclusionMapBaker::EResult::BentNormal);
-		if (ensure(BNResult))
-		{
-			FTexture2DBuilder TextureNormalBuilder;
-			TextureNormalBuilder.Initialize(FTexture2DBuilder::ETextureType::NormalMap, CachedOcclusionMapSettings.Dimensions);
-			TextureNormalBuilder.Copy(*BNResult);
-			TextureNormalBuilder.Commit(false);
-			CachedBentNormalMap = TextureNormalBuilder.GetTexture2D();
-		}
+		FTexture2DBuilder TextureNormalBuilder;
+		TextureNormalBuilder.Initialize(FTexture2DBuilder::ETextureType::NormalMap, CachedOcclusionMapSettings.Dimensions);
+		TextureNormalBuilder.Copy(*NewResult->GetBakeResult(1));
+		TextureNormalBuilder.Commit(false);
+		CachedBentNormalMap = TextureNormalBuilder.GetTexture2D();
 		break;
 	}
 	case EBakeMapType::Curvature:
 	{
-		FMeshCurvatureMapBaker* Baker = static_cast<FMeshCurvatureMapBaker*>(NewResult.Get());
 		FTexture2DBuilder TextureBuilder;
 		TextureBuilder.Initialize(FTexture2DBuilder::ETextureType::Color, CachedCurvatureMapSettings.Dimensions);
-		TextureBuilder.Copy(*Baker->GetResult());
+		TextureBuilder.Copy(*NewResult->GetBakeResult(0));
 		TextureBuilder.Commit(false);
 		CachedCurvatureMap = TextureBuilder.GetTexture2D();
 		break;
@@ -1473,32 +1464,30 @@ void UBakeMeshAttributeMapsTool::OnMapsUpdated(const TUniquePtr<FMeshImageBaker>
 	case EBakeMapType::FaceNormalImage:
 	case EBakeMapType::MaterialID:
 	{
-		FMeshPropertyMapBaker* Baker = static_cast<FMeshPropertyMapBaker*>(NewResult.Get());
 		FTexture2DBuilder TextureBuilder;
 		TextureBuilder.Initialize(FTexture2DBuilder::ETextureType::Color, CachedMeshPropertyMapSettings.Dimensions);
-		TextureBuilder.Copy(*Baker->GetResult());
+		TextureBuilder.Copy(*NewResult->GetBakeResult(0));
 		TextureBuilder.Commit(false);
 		CachedMeshPropertyMap = TextureBuilder.GetTexture2D();
 		break;
 	}
 	case EBakeMapType::Texture2DImage:
 	{
-		FMeshResampleImageBaker* Baker = static_cast<FMeshResampleImageBaker*>(NewResult.Get());
 		FTexture2DBuilder TextureBuilder;
 		TextureBuilder.Initialize(FTexture2DBuilder::ETextureType::Color, CachedTexture2DImageSettings.Dimensions);
-		TextureBuilder.Copy(*Baker->GetResult(), true);
+		TextureBuilder.Copy(*NewResult->GetBakeResult(0), true);
 		TextureBuilder.Commit(false);
 		CachedTexture2DImageMap = TextureBuilder.GetTexture2D();
 		break;
 	}
 	case EBakeMapType::MultiTexture:
 	{
-		FMeshMultiResampleImageBaker* Baker = static_cast<FMeshMultiResampleImageBaker*>(NewResult.Get());
 		FTexture2DBuilder TextureBuilder;
 		TextureBuilder.Initialize(FTexture2DBuilder::ETextureType::Color, CachedTexture2DImageSettings.Dimensions);
-		TextureBuilder.Copy(*Baker->GetResult(), true);
+		TextureBuilder.Copy(*NewResult->GetBakeResult(0), true);
 		TextureBuilder.Commit(false);
 		CachedTexture2DImageMap = TextureBuilder.GetTexture2D();
+		break;
 	}
 	}
 
diff --git a/Engine/Plugins/Experimental/MeshModelingToolset/Source/MeshModelingTools/Public/BakeMeshAttributeMapsTool.h b/Engine/Plugins/Experimental/MeshModelingToolset/Source/MeshModelingTools/Public/BakeMeshAttributeMapsTool.h
--- a/Engine/Plugins/Experimental/MeshModelingToolset/Source/MeshModelingTools/Public/BakeMeshAttributeMapsTool.h
+++ b/Engine/Plugins/Experimental/MeshModelingToolset/Source/MeshModelingTools/Public/BakeMeshAttributeMapsTool.h
@@ -10,7 +10,7 @@
 #include "DynamicMeshAABBTree3.h"
 #include "Image/ImageDimensions.h"
 #include "Image/ImageBuilder.h"
-#include "Sampling/MeshImageBaker.h"
+#include "Sampling/MeshMapBaker.h"
 #include "ModelingOperators.h"
 #include "MeshOpPreviewHelpers.h"
 #include "BakeMeshAttributeMapsTool.generated.h"
@@ -85,6 +85,17 @@ enum class EBakeTextureResolution
 };
 
 
+UENUM()
+enum class EBakeMultisampling
+{
+	None = 1 UMETA(DisplayName = "None"),
+	Sample2x2 = 2 UMETA(DisplayName = "2 x 2"),
+	Sample4x4 = 4 UMETA(DisplayName = "4 x 4"),
+	Sample8x8 = 8 UMETA(DisplayName = "8 x 8"),
+	Sample16x16 = 16 UMETA(DisplayName = "16 x 16")
+};
+
+
 UCLASS()
 class MESHMODELINGTOOLS_API UBakeMeshAttributeMapsToolProperties : public UInteractiveToolPropertySet
 {
@@ -99,6 +110,10 @@ public:
 	UPROPERTY(EditAnywhere, Category = MapSettings, meta = (TransientToolProperty))
 	EBakeTextureResolution Resolution = EBakeTextureResolution::Resolution256;
 
+	/** The multisampling configuration per texel */
+	UPROPERTY(EditAnywhere, Category = MapSettings)
+	EBakeMultisampling Multisampling = EBakeMultisampling::None;
+
 	UPROPERTY(EditAnywhere, Category = MapSettings)
 	bool bUseWorldSpace = false;
 
@@ -333,12 +348,12 @@ public:
  * Detail Map Baking Tool
  */
 UCLASS()
-class MESHMODELINGTOOLS_API UBakeMeshAttributeMapsTool : public UMultiSelectionTool, public UE::Geometry::IGenericDataOperatorFactory<UE::Geometry::FMeshImageBaker>
+class MESHMODELINGTOOLS_API UBakeMeshAttributeMapsTool : public UMultiSelectionTool, public UE::Geometry::IGenericDataOperatorFactory<UE::Geometry::FMeshMapBaker>
 {
 	GENERATED_BODY()
 
 public:
-	UBakeMeshAttributeMapsTool();
+	UBakeMeshAttributeMapsTool() = default;
 
 	virtual void SetAssetAPI(IAssetGenerationAPI* AssetAPI);
 
@@ -353,7 +368,7 @@ public:
 	virtual bool CanAccept() const override;
 
 	// IGenericDataOperatorFactory API
-	virtual TUniquePtr<UE::Geometry::TGenericDataOperator<UE::Geometry::FMeshImageBaker>> MakeNewOperator() override;
+	virtual TUniquePtr<UE::Geometry::TGenericDataOperator<UE::Geometry::FMeshMapBaker>> MakeNewOperator() override;
 
 protected:
 	// need to update bResultValid if these are modified, so we don't publicly expose them. 
@@ -391,6 +406,8 @@ protected:
 	friend class FBakeMeshPropertyMapOp;
 	friend class FBakeMultiTextureOp;
 
+	friend class FMeshMapBakerOp;
+
 
 	IAssetGenerationAPI* AssetAPI = nullptr;
 
@@ -425,8 +442,8 @@ protected:
 	void UpdateOnModeChange();
 	void UpdateVisualization();
 
-	TUniquePtr<TGenericDataBackgroundCompute<UE::Geometry::FMeshImageBaker>> Compute = nullptr;
-	void OnMapsUpdated(const TUniquePtr<UE::Geometry::FMeshImageBaker>& NewResult);
+	TUniquePtr<TGenericDataBackgroundCompute<UE::Geometry::FMeshMapBaker>> Compute = nullptr;
+	void OnMapsUpdated(const TUniquePtr<UE::Geometry::FMeshMapBaker>& NewResult);
 
 	struct FBakeCacheSettings
 	{
@@ -434,10 +451,11 @@ protected:
 		int32 UVLayer;
 		int32 DetailTimestamp;
 		float Thickness;
+		int32 Multisampling;
 
 		bool operator==(const FBakeCacheSettings& Other) const
 		{
-			return Dimensions == Other.Dimensions && UVLayer == Other.UVLayer && DetailTimestamp == Other.DetailTimestamp && Thickness == Other.Thickness;
+			return Dimensions == Other.Dimensions && UVLayer == Other.UVLayer && DetailTimestamp == Other.DetailTimestamp && Thickness == Other.Thickness && Multisampling == Other.Multisampling;
 		}
 	};
 	FBakeCacheSettings CachedBakeCacheSettings;
