SceneCapturePhotoSet - "Hide" Nanite actors by temporarily unregistering components

diff --git a/Engine/Plugins/Experimental/MeshModelingToolset/Source/ModelingComponents/Private/Scene/SceneCapturePhotoSet.cpp b/Engine/Plugins/Experimental/MeshModelingToolset/Source/ModelingComponents/Private/Scene/SceneCapturePhotoSet.cpp
--- a/Engine/Plugins/Experimental/MeshModelingToolset/Source/ModelingComponents/Private/Scene/SceneCapturePhotoSet.cpp
+++ b/Engine/Plugins/Experimental/MeshModelingToolset/Source/ModelingComponents/Private/Scene/SceneCapturePhotoSet.cpp
@@ -1,6 +1,7 @@
 // Copyright Epic Games, Inc. All Rights Reserved.
 
 #include "Scene/SceneCapturePhotoSet.h"
+#include "EngineUtils.h"
 
 using namespace UE::Geometry;
 
@@ -87,6 +88,19 @@ void FSceneCapturePhotoSet::AddExteriorCaptures(
 {
 	check(this->TargetWorld != nullptr);
 
+	// Workaround for Nanite scene proxies visibility
+	// Unregister all components to remove unwanted proxies from the scene. This is currently the only way to "hide" nanite meshes.
+	TSet<AActor*> VisibleActorsSet(VisibleActors);
+	TArray<AActor*> ActorsToRegister;
+	for (TActorIterator<AActor> Actor(TargetWorld); Actor; ++Actor)
+	{
+		if (!VisibleActorsSet.Contains(*Actor))
+		{
+			Actor->UnregisterAllComponents();
+			ActorsToRegister.Add(*Actor);
+		}
+	}
+
 	FWorldRenderCapture RenderCapture;
 	RenderCapture.SetWorld(TargetWorld);
 	RenderCapture.SetVisibleActors(VisibleActors);
@@ -161,6 +175,12 @@ void FSceneCapturePhotoSet::AddExteriorCaptures(
 		}
 	}
 
+	// Workaround for Nanite scene proxies visibility
+	// Reregister all components we previously unregistered
+	for (AActor* Actor : ActorsToRegister)
+	{
+		Actor->RegisterAllComponents();
+	}
 }
 
 
