Simple skin binding tool. Supports closest distance and geodesic voxel skinning methods.

diff --git a/Engine/Plugins/Experimental/MeshModelingToolset/Source/MeshModelingTools/Private/SkinWeightsBindingTool.cpp b/Engine/Plugins/Experimental/MeshModelingToolset/Source/MeshModelingTools/Private/SkinWeightsBindingTool.cpp
new file mode 100644
--- /dev/null
+++ b/Engine/Plugins/Experimental/MeshModelingToolset/Source/MeshModelingTools/Private/SkinWeightsBindingTool.cpp
@@ -0,0 +1,1095 @@
+// Copyright Epic Games, Inc. All Rights Reserved.
+
+#include "SkinWeightsBindingTool.h"
+
+#include "BoneWeights.h"
+#include "DynamicMeshToMeshDescription.h"
+#include "DynamicVertexSkinWeightsAttribute.h"
+#include "InteractiveToolManager.h"
+#include "MeshDescriptionToDynamicMesh.h"
+#include "MeshOpPreviewHelpers.h"
+#include "SkeletalMeshAttributes.h"
+#include "ToolSetupUtil.h"
+#include "ToolTargetManager.h"
+#include "Components/SkeletalMeshComponent.h"
+#include "Engine/SkeletalMesh.h"
+#include "Spatial/FastWinding.h"
+#include "Spatial/MeshWindingNumberGrid.h"
+
+// #pragma optimize( "", off )
+
+#define LOCTEXT_NAMESPACE "USkinWeightsBindingTool"
+
+// TODO: Move to a helper function.
+
+// A simple FIFO queue. Maintains a set of blocks, rather than allocate for each element.
+template<typename T, int32 BlockSize=256>
+class TFIFOQueue
+{
+public:
+	TFIFOQueue() = default;
+
+	void Push(T InElem)
+	{
+		// Container Empty?
+		if (Blocks.IsEmpty())
+		{
+			Blocks.SetNum(1);
+			Blocks[0].bIsFree = false;
+			PushIndex = -1;
+		}
+		// Or are we at the end of current block?
+		else if (PushIndex == (BlockSize - 1))
+		{
+			PushBlock = AllocateNewPushBlock();
+			PushIndex = -1;
+		}
+
+		PushIndex++;
+
+		Blocks[PushBlock].Data[PushIndex] = InElem;
+	}
+
+	T Pop()
+	{
+		check(!IsEmpty());
+
+		if (PopIndex == (BlockSize - 1))
+		{
+			// Reached the end. Free this block and move onto the next one.
+			PopBlock = FreeCurrentPopBlock();
+			PopIndex = -1;
+		}
+
+		PopIndex++;
+		return Blocks[PopBlock].Data[PopIndex];
+	}
+
+	bool TryPop(T &OutValue)
+	{
+		if (IsEmpty())
+		{
+			return false;
+		}
+
+		OutValue = Pop();
+		return true;
+	}
+
+	void Reset()
+	{
+		Blocks.Reset();
+		PushBlock = 0;
+		PushIndex = -1;
+		PopBlock = 0;
+		PopIndex = -1;
+	}
+	
+
+	bool IsEmpty() const
+	{
+		return PushBlock == PopBlock && PushIndex == PopIndex;
+	}
+
+private:
+	int32 AllocateNewPushBlock()
+	{
+		checkSlow(PushBlock != -1);
+
+		int NewBlockIndex = INDEX_NONE;
+		
+		if (FreeCount > 0)
+		{
+			// Find a free block
+			for (int32 Index = 0; Index < Blocks.Num(); Index++)
+			{
+				if (Blocks[Index].bIsFree)
+				{
+					NewBlockIndex = Index;
+					FreeCount--;
+					break;
+				}
+			}
+			checkfSlow(NewBlockIndex != INDEX_NONE, TEXT("We should have found a free block."));
+		}
+		else
+		{
+			NewBlockIndex = Blocks.Num();
+			Blocks.AddDefaulted();
+		}
+
+		Blocks[PushBlock].NextBlock = NewBlockIndex;
+		Blocks[NewBlockIndex].bIsFree = false;
+		Blocks[NewBlockIndex].NextBlock = INDEX_NONE;
+		return NewBlockIndex;
+	}
+
+	int32 FreeCurrentPopBlock()
+	{
+		const int32 NextBlock = Blocks[PopBlock].NextBlock;
+		checkSlow(NextBlock != -1);
+		Blocks[PopBlock].bIsFree = true;
+		FreeCount++;
+
+		return NextBlock;
+	}
+	
+
+	int32 PushBlock = 0;
+	int32 PushIndex = -1;
+
+	int32 PopBlock = 0;
+	int32 PopIndex = -1;
+
+	int32 FreeCount = 0;
+	
+	struct FBlock
+	{
+		FBlock()
+		{
+			Data = new T[BlockSize];
+		}
+		~FBlock()
+		{
+			delete [] Data;
+		}
+
+		FBlock(FBlock &&InOther) noexcept
+		{
+			Data = InOther.Data;
+			InOther.Data = nullptr;
+		}
+		
+		bool bIsFree = true;
+		int32 NextBlock = INDEX_NONE;
+		T *Data;
+	};
+	
+	TArray<FBlock> Blocks;
+};
+
+
+static float DistanceToLineSegment(const FVector& P, const FVector& A, const FVector& B)
+{
+	const FVector M = B - A;
+	const FVector T = P - A;
+
+	const float C1 = FVector::DotProduct(M, T);
+	if (C1 <= 0.0f)
+	{
+		return FVector::Dist(P, A);
+	}
+
+	const float C2 = FVector::DotProduct(M, M);
+	if (C2 <= C1)
+	{
+		return FVector::Dist(P, B);
+	}
+
+	// Project the point onto the line and get the distance between them.
+	const FVector PT = A + M * (C1 / C2);
+	return FVector::Dist(P, PT);
+}
+
+
+// List of bones as used by the binding class. In this case for each bone transform, we want to
+// store a list of line segments going from the bone transform to all the child bone transforms.
+struct FTransformHierarchyQuery
+{
+	explicit FTransformHierarchyQuery(const TArray<TPair<FTransform, int32>>& InTransformHierarchy)
+	{
+		for (int Index = 0; Index < InTransformHierarchy.Num(); Index++)
+		{
+			FTransform Xform = InTransformHierarchy[Index].Key;
+			int32 ParentIndex = InTransformHierarchy[Index].Value;
+
+			while (ParentIndex != INDEX_NONE)
+			{
+				Xform = Xform * InTransformHierarchy[ParentIndex].Key;
+				ParentIndex = InTransformHierarchy[ParentIndex].Value;
+			}
+
+			BoneFans.Add({ Xform.GetLocation() });
+		}
+
+		// Fill in the fan tips, as needed.
+		for (int Index = 0; Index < InTransformHierarchy.Num(); Index++)
+		{
+			const int32 ParentIndex = InTransformHierarchy[Index].Value;
+			if (ParentIndex != INDEX_NONE)
+			{
+				BoneFans[ParentIndex].TipsPos.Add(BoneFans[Index].RootPos);
+			}
+		}		
+	}
+
+	float GetDistanceToBoneFan(const int32 InBoneIndex, const FVector& InPoint) const
+	{
+		return BoneFans[InBoneIndex].GetDistance(InPoint);
+	}
+	
+	FBox GetBoneFanBBox(const int32 InBoneIndex) const
+	{
+		return BoneFans[InBoneIndex].GetBBox();
+	}
+
+	bool GetBoneFanIntersectsBox(const int32 InBoneIndex, const FBox &InBox) const
+	{
+		return BoneFans[InBoneIndex].IntersectsBox(InBox);
+	}
+
+private:
+	struct FBoneFan
+	{
+		FVector RootPos;
+		TArray<FVector> TipsPos;
+
+		float GetDistance(const FVector& InPoint) const
+		{
+			if (TipsPos.IsEmpty())
+			{
+				return FVector::Distance(RootPos, InPoint);
+			}
+			else
+			{
+				float Distance = std::numeric_limits<float>::max();
+				for (const FVector& TipPos: TipsPos)
+				{
+					Distance = FMath::Min(Distance, DistanceToLineSegment(InPoint, RootPos, TipPos));
+				}
+				return Distance;
+			}
+		}
+
+		FBox GetBBox() const
+		{
+			FBox Box(RootPos, RootPos);
+			for (const FVector& TipPos: TipsPos)
+			{
+				Box += TipPos;
+			}
+
+			return Box;
+		}
+
+		bool IntersectsBox(const FBox &InBox) const
+		{
+			if (TipsPos.IsEmpty())
+			{
+				return FMath::PointBoxIntersection(RootPos, InBox);
+			}
+
+			if (GetBBox().Intersect(InBox))
+			{
+				for (const FVector& TipPos: TipsPos)
+				{
+					if (FMath::LineBoxIntersection(InBox, RootPos, TipPos, TipPos - RootPos))
+					{
+						return true;
+					}
+				}
+			}
+			return false;
+		}				
+	};
+
+	 TArray<FBoneFan> BoneFans;
+};
+
+
+struct FOccupancyGrid
+{
+	enum EDomain : int32
+	{
+		Exterior,
+		Boundary,
+		Interior
+	};
+
+	UE::Geometry::FDenseGrid3i Occupancy;
+	float CellSize;
+	FVector3f GridOrigin;
+	FVector3f CellMidPoint;
+
+	FOccupancyGrid(
+		const FDynamicMesh3& InMesh,
+		int32 InVoxelResolution
+		)
+	{
+		using namespace UE::Geometry;
+	
+		// Compute a voxel grid 
+		FDynamicMeshAABBTree3 Spatial(&InMesh);
+		TFastWindingTree FastWinding(&Spatial);
+		FAxisAlignedBox3d Bounds = Spatial.GetBoundingBox();
+		CellSize = Bounds.MaxDim() / InVoxelResolution;
+		CellMidPoint = FVector3f(CellSize / 2.0f, CellSize / 2.0f, CellSize / 2.0f);
+
+		TMeshWindingNumberGrid WindingGrid(&InMesh, &FastWinding, CellSize);
+	
+		WindingGrid.Compute();
+
+		// Our occupancy grid is computed on the winding number grid's cell centers.
+		const FVector3i WindingDims = WindingGrid.Dimensions();
+		Occupancy = FDenseGrid3i(WindingDims.X - 1, WindingDims.Y - 1, WindingDims.Z - 1, Exterior);
+
+		GridOrigin = WindingGrid.GridOrigin + FVector3f(CellSize / 2.0f, CellSize / 2.0f, CellSize / 2.0f);
+
+		static const FVector3i CornerOffsets[] = {
+			FVector3i(0, 0, 0),
+			FVector3i(0, 0, 1),
+			FVector3i(0, 1, 0),
+			FVector3i(0, 1, 1),
+			FVector3i(1, 0, 0),
+			FVector3i(1, 0, 1),
+			FVector3i(1, 1, 0),
+			FVector3i(1, 1, 1),
+		};
+		
+		// TODO: Parallel-for
+		for (int32 OccupancyId = 0; OccupancyId < Occupancy.Size(); OccupancyId++)
+		{
+			const FVector3i OccupancyIndex(Occupancy.ToIndex(OccupancyId));
+			int32 Count = 0;
+			for (int32 CornerId = 0; CornerId < 8; CornerId++)
+			{
+				const FVector3i CornerIndex(OccupancyIndex + CornerOffsets[CornerId]);
+				if (WindingGrid.GetValue(CornerIndex) >= WindingGrid.WindingIsoValue)
+				{
+					Count++;
+				}
+			}
+
+			if (Count == 8)
+			{
+				Occupancy[OccupancyIndex] = Interior;
+			}
+			else if (Count > 0)
+			{
+				Occupancy[OccupancyIndex] = Boundary;
+			}
+		}
+
+		// Make sure we include all the vertices of the mesh as a part of the boundary, if
+		// the vertex areas are marked as being exterior.
+		for (int32 VertexIdx = 0; VertexIdx < InMesh.VertexCount(); VertexIdx++)
+		{
+			const FVector3d& Pos = InMesh.GetVertex(VertexIdx);
+			const FVector3i OccupancyIndex = GetCellIndexFromPoint(FVector(Pos));
+			if (Occupancy[OccupancyIndex] == Exterior)
+			{
+				Occupancy[OccupancyIndex] = Boundary;
+			}
+		}
+	}
+	
+	UE::Geometry::FVector3i GetCellIndexFromPoint(const FVector &InPoint) const
+	{
+		FVector3f PP(InPoint);
+		PP -= GridOrigin;
+		PP += CellMidPoint;
+
+		return { FMath::FloorToInt(PP.X / CellSize),
+				FMath::FloorToInt(PP.Y / CellSize),
+				FMath::FloorToInt(PP.Z / CellSize) };
+	}
+
+	FVector3d GetCellCenterFromIndex(const UE::Geometry::FVector3i &Index) const
+	{
+		const float CS = CellSize;
+		return {Index.X * CS + GridOrigin.X, Index.Y * CS + GridOrigin.Y, Index.Z * CS + GridOrigin.Z};
+	}
+
+	FBox GetCellBoxFromIndex(const UE::Geometry::FVector3i &Index) const
+	{
+		const FVector3f P = GetCellCenterFromIndex(Index);
+		return {P - CellMidPoint, P + CellMidPoint};
+	}
+};
+
+
+namespace
+{
+	struct FCreateSkinWeights_Closest_WorkData final :
+		TThreadSingleton<FCreateSkinWeights_Closest_WorkData>
+	{
+		TArray<TPair<FBoneIndexType, float>> RawBoneWeights;
+		TArray<UE::AnimationCore::FBoneWeight> BoneWeights;
+	};
+
+	float ComputeWeightStiffness(const float InWeight, const float InStiffness)
+	{
+		return (1.0f - InStiffness) * InWeight + InStiffness * InWeight * InWeight;
+	}
+	
+}
+
+
+class FComputeSkinWeightsBindingOp : public UE::Geometry::FDynamicMeshOperator
+{
+public:
+	virtual ~FComputeSkinWeightsBindingOp() override {}
+
+	// The transform hierarchy to bind to. Listed in the same order as the bones in the
+	// reference skeleton that this skelmesh is tied to.
+	TSharedPtr<FDynamicMesh3, ESPMode::ThreadSafe> OriginalMesh;
+	TArray<TPair<FTransform, int32>> TransformHierarchy;
+
+	FName ProfileName = FSkeletalMeshAttributes::DefaultSkinWeightProfileName;
+	
+	ESkinWeightsBindType BindType = ESkinWeightsBindType::DirectDistance;
+	float Stiffness = 0.2f;
+	int32 MaxInfluences = 5;
+	int32 VoxelResolution = 256;
+	
+	void CalculateResult(FProgressCancel* InProgress) override
+	{
+		if (InProgress && InProgress->Cancelled())
+		{
+			return;
+		}
+
+		ResultMesh->Copy(*OriginalMesh, true, true, true, true);
+		
+		if (InProgress && InProgress->Cancelled())
+		{
+			return;
+		}
+
+		const float ClampedStiffness = FMath::Clamp(Stiffness, 0.0f, 1.0f);	
+
+		UE::AnimationCore::FBoneWeightsSettings Settings;
+		Settings.SetMaxWeightCount(MaxInfluences);
+
+		switch(BindType)
+		{
+		case ESkinWeightsBindType::DirectDistance:
+			CreateSkinWeights_DirectDistance(*ResultMesh, ClampedStiffness, Settings);
+			break;
+			
+		case ESkinWeightsBindType::GeodesicVoxel:
+			CreateSkinWeights_GeodesicVoxel(*ResultMesh, ClampedStiffness, Settings);
+			break;
+		}
+	}
+
+private:
+	static UE::Geometry::FDynamicMeshVertexSkinWeightsAttribute *GetOrCreateSkinWeightsAttribute(
+		FDynamicMesh3& InMesh,
+		FName InProfileName
+		)
+	{
+		using namespace UE::Geometry;
+		
+		FDynamicMeshVertexSkinWeightsAttribute *Attribute = InMesh.Attributes()->GetSkinWeightsAttribute(InProfileName);
+		if (!Attribute)
+		{
+			Attribute = new FDynamicMeshVertexSkinWeightsAttribute(&InMesh);
+			InMesh.Attributes()->AttachSkinWeightsAttribute(InProfileName, Attribute);
+		}
+		return Attribute;
+	}
+	
+	void CreateSkinWeights_DirectDistance(
+		FDynamicMesh3& InMesh,
+		float InStiffness,
+		const UE::AnimationCore::FBoneWeightsSettings& InSettings
+		)
+	{
+		using namespace UE::AnimationCore;
+		using namespace UE::Geometry;
+
+		FDynamicMeshVertexSkinWeightsAttribute *SkinWeights = InMesh.Attributes()->GetSkinWeightsAttribute(ProfileName);
+		
+		const int32 NumVertices = InMesh.VertexCount();
+
+		// Use the diagonal size of the bbox to make the bone distance falloff scale invariant.
+		const float DiagBounds = InMesh.GetCachedBounds().DiagonalLength();
+	
+		const FTransformHierarchyQuery Skeleton(TransformHierarchy);
+
+		ParallelFor(NumVertices, [&](const int32 VertexIdx)
+		{
+			const FVector3d& Pos = InMesh.GetVertex(VertexIdx);
+
+			FCreateSkinWeights_Closest_WorkData &WorkData = FCreateSkinWeights_Closest_WorkData::Get();
+		
+			WorkData.RawBoneWeights.Reset(TransformHierarchy.Num());
+
+			float TotalWeight = 0.0f;
+			for (int32 BoneIndex = 0; BoneIndex < TransformHierarchy.Num(); BoneIndex++)
+			{
+				// Normalize the distance by the diagonal size of the bbox to maintain scale invariance.
+				float Weight = Skeleton.GetDistanceToBoneFan(BoneIndex, Pos) / DiagBounds;
+
+				// Avoid div-by-zero but allow for the possibility that multiple bones may
+				// touch this vertex.
+				Weight = FMath::Max(Weight, KINDA_SMALL_NUMBER);
+
+				// Compute the actual weight, factoring in the stiffness value. W = (1/S(D))^2
+				// Where S(x) is the stiffness function.
+				Weight = FMath::Square(1.0f / ComputeWeightStiffness(Weight, InStiffness));
+				TotalWeight += Weight;
+				WorkData.RawBoneWeights.Add(MakeTuple(static_cast<FBoneIndexType>(BoneIndex), Weight));
+			}
+
+			// Normalize
+			for (TPair<FBoneIndexType, float> &BoneWeight: WorkData.RawBoneWeights)
+			{
+				BoneWeight.Value /= TotalWeight;
+			}
+			
+			WorkData.RawBoneWeights.Sort([](const TPair<FBoneIndexType, float> &A, const TPair<FBoneIndexType, float> &B)
+			{
+				return A.Value > B.Value;
+			});
+
+			WorkData.BoneWeights.Reset(InSettings.GetMaxWeightCount());
+			for (int32 BoneIndex = 0; BoneIndex < FMath::Min(InSettings.GetMaxWeightCount(), WorkData.RawBoneWeights.Num()); BoneIndex++)
+			{
+				const TPair<FBoneIndexType, float>& BoneWeight = WorkData.RawBoneWeights[BoneIndex];
+				WorkData.BoneWeights.Add(FBoneWeight(BoneWeight.Key, BoneWeight.Value));
+			}
+
+			SkinWeights->SetValue(VertexIdx, FBoneWeights::Create(WorkData.BoneWeights, InSettings));
+		});
+	}
+	
+	void CreateSkinWeights_GeodesicVoxel(
+		FDynamicMesh3& InMesh,
+		float InStiffness,
+		const UE::AnimationCore::FBoneWeightsSettings& InSettings
+		)
+	{
+		using namespace UE::AnimationCore;
+		using namespace UE::Geometry;
+
+		FDynamicMeshVertexSkinWeightsAttribute *SkinWeights = InMesh.Attributes()->GetSkinWeightsAttribute(ProfileName);
+		
+		const int32 NumVertices = InMesh.VertexCount();
+
+		// Use the diagonal size of the bbox to make the bone distance falloff scale invariant.
+		const float DiagBounds = InMesh.GetCachedBounds().DiagonalLength();
+		
+		const FTransformHierarchyQuery Skeleton(TransformHierarchy);
+
+		// This is grossly inefficient but tricky to do otherwise, since each bone distance
+		// computation is done per-thread. We could possibly solve this by chunking instead
+		// and accumulating partial results.
+		TArray<float> Weights;
+		Weights.SetNumUninitialized(NumVertices * TransformHierarchy.Num());
+
+		FOccupancyGrid Occupancy(InMesh, VoxelResolution);
+
+		const FVector3i Dimensions = Occupancy.Occupancy.GetDimensions();
+		
+		ParallelFor(TransformHierarchy.Num(), [&](int32 BoneIndex) {
+			TFIFOQueue<FVector3i> WorkingSet;  
+			FDenseGrid3f BoneDistance(Dimensions.X, Dimensions.Y, Dimensions.Z, DiagBounds);
+
+			// Mark all the cells that the bone intersects with distance of 0 and put them
+			// on the work queue.
+			const FBox BoneBox = Skeleton.GetBoneFanBBox(BoneIndex);
+			const FVector3i BoneMin = Occupancy.GetCellIndexFromPoint(BoneBox.Min);
+			const FVector3i BoneMax = Occupancy.GetCellIndexFromPoint(BoneBox.Max);
+
+			for (int32 I = BoneMin.X; I <= BoneMax.X; I++)
+			{
+				for (int32 J = BoneMin.Y; J <= BoneMax.Y; J++)
+				{
+					for (int32 K = BoneMin.Z; K <= BoneMax.Z; K++)
+					{
+						const FVector3i Candidate(I, J, K);
+						const FBox CellBox = Occupancy.GetCellBoxFromIndex(Candidate);
+						if (Skeleton.GetBoneFanIntersectsBox(BoneIndex, CellBox))
+						{
+							WorkingSet.Push(Candidate);
+							BoneDistance[Candidate] = 0.0f;
+						}
+					}
+				}
+			}
+
+			// Iterate over all the voxels until we have constructed shortest distance paths
+			// throughout the level set.
+			while (!WorkingSet.IsEmpty())
+			{
+				const FVector3i WorkItem = WorkingSet.Pop();
+
+				// Loop through each of the neighbours (6 face neighbours, 12 edge neighbors,
+				// and 8 corner neighbours) and see if any of them are closer to the bone
+				// than their current marked distance.
+				float CurrentDistance = BoneDistance[WorkItem];
+				for (int32 N = 0; N < 26; N++)
+				{
+					FVector3i Offset(IndexUtil::GridOffsets26[N]);
+					FVector3i Candidate(WorkItem + Offset);
+
+					if (!BoneDistance.IsValidIndex(Candidate))
+					{
+						continue;
+					}
+					
+					// Ensure this entry is either a part of the interior or boundary domain.
+					if (Occupancy.Occupancy[Candidate] == FOccupancyGrid::Exterior)
+					{
+						continue;
+					}
+
+					const float CellDistance = (FVector3f(Offset) * Occupancy.CellSize).Length();
+					const float CandidateDistance = CurrentDistance + CellDistance;
+					const float OldDistance = BoneDistance[Candidate]; 
+
+					if (OldDistance > CandidateDistance)
+					{
+						WorkingSet.Push(Candidate);
+						BoneDistance[Candidate] = CandidateDistance;
+					}
+				}
+			}
+
+			// Loop through all the vertices, find the voxel each belongs to, and compute
+			// the distance from the voxel to the vertex (assuming the distance stored in the
+			// voxel is based on traversing from voxel center to voxel center).
+			for (int32 VertexIdx = 0; VertexIdx < NumVertices; VertexIdx++)
+			{
+				const FVector3d& Pos = InMesh.GetVertex(VertexIdx);
+				const FVector3i CellIndex = Occupancy.GetCellIndexFromPoint(Pos);
+				const FVector3d CellCenter = Occupancy.GetCellCenterFromIndex(CellIndex);
+
+				float Distance = BoneDistance[CellIndex];
+				const FOccupancyGrid::EDomain Domain = static_cast<FOccupancyGrid::EDomain>(Occupancy.Occupancy[CellIndex]);
+				// check(Distance != std::numeric_limits<float>::max());
+
+				Distance += FVector3d::Distance(CellCenter, Pos);
+				
+				// Normalize the distance by the diagonal size of the bbox to maintain scale invariance.
+				float Weight = Distance / DiagBounds;
+
+				// Avoid div-by-zero but allow for the possibility that multiple bones may
+				// touch this vertex.
+				Weight = FMath::Max(Weight, KINDA_SMALL_NUMBER);
+
+				// Compute the actual weight, factoring in the stiffness value. W = (1/S(D))^2
+				// Where S(x) is the stiffness function.
+				Weight = FMath::Square(1.0f / ComputeWeightStiffness(Weight, InStiffness));
+
+				Weights[VertexIdx * TransformHierarchy.Num() + BoneIndex] = Weight;
+			}
+		});	
+
+		ParallelFor(NumVertices, [&](const int32 VertexIdx)
+		{
+			FCreateSkinWeights_Closest_WorkData &WorkData = FCreateSkinWeights_Closest_WorkData::Get();
+			
+			WorkData.RawBoneWeights.Reset(TransformHierarchy.Num());
+
+			float TotalWeight = 0.0f;
+			for (int32 BoneIndex = 0; BoneIndex < TransformHierarchy.Num(); BoneIndex++)
+			{
+				const float Weight = Weights[VertexIdx * TransformHierarchy.Num() + BoneIndex];
+				TotalWeight += Weight;
+				WorkData.RawBoneWeights.Add(MakeTuple(static_cast<FBoneIndexType>(BoneIndex), Weight));
+			}
+
+			// Normalize
+			for (TPair<FBoneIndexType, float> &BoneWeight: WorkData.RawBoneWeights)
+			{
+				BoneWeight.Value /= TotalWeight;
+			}
+				
+			WorkData.RawBoneWeights.Sort([](const TPair<FBoneIndexType, float> &A, const TPair<FBoneIndexType, float> &B)
+			{
+				return A.Value > B.Value;
+			});
+
+			WorkData.BoneWeights.Reset(InSettings.GetMaxWeightCount());
+			for (int32 BoneIndex = 0; BoneIndex < FMath::Min(InSettings.GetMaxWeightCount(), WorkData.RawBoneWeights.Num()); BoneIndex++)
+			{
+				const TPair<FBoneIndexType, float>& BoneWeight = WorkData.RawBoneWeights[BoneIndex];
+				WorkData.BoneWeights.Add(FBoneWeight(BoneWeight.Key, BoneWeight.Value));
+			}
+
+			SkinWeights->SetValue(VertexIdx, FBoneWeights::Create(WorkData.BoneWeights, InSettings));
+		});
+	}	
+};
+
+
+
+
+bool USkinWeightsBindingToolBuilder::CanBuildTool(const FToolBuilderState& SceneState) const
+{
+	return SceneState.TargetManager->CountSelectedAndTargetable(SceneState, GetTargetRequirements()) == 1;
+}
+
+
+UInteractiveTool* USkinWeightsBindingToolBuilder::BuildTool(const FToolBuilderState& SceneState) const
+{
+	USkinWeightsBindingTool* NewTool = NewObject<USkinWeightsBindingTool>(SceneState.ToolManager);
+
+	TArray<TObjectPtr<UToolTarget>> Targets = SceneState.TargetManager->BuildAllSelectedTargetable(SceneState, GetTargetRequirements());
+	NewTool->SetTargets(MoveTemp(Targets));
+	NewTool->SetWorld(SceneState.World);
+
+	return NewTool;
+}
+
+
+USkeleton* USkinWeightsBindingToolProperties::GetSkeleton(bool& bInvalidSkeletonIsError)
+{
+	bInvalidSkeletonIsError = false;
+	return SkeletalMesh ? SkeletalMesh->GetSkeleton() : nullptr;
+}
+
+
+USkinWeightsBindingTool::USkinWeightsBindingTool()
+{
+	Properties = CreateDefaultSubobject<USkinWeightsBindingToolProperties>(TEXT("SkinWeightsBindingProperties"));
+	// CreateDefaultSubobject automatically sets RF_Transactional flag, we need to clear it so that undo/redo doesn't affect tool properties
+	Properties->ClearFlags(RF_Transactional);
+}
+
+
+USkinWeightsBindingTool::~USkinWeightsBindingTool()
+{
+}
+
+
+void USkinWeightsBindingTool::SetWorld(UWorld* InWorld)
+{
+	TargetWorld = InWorld;
+}
+
+
+void USkinWeightsBindingTool::Setup()
+{
+	Super::Setup();
+
+	if (ensure(Properties))
+	{
+		Properties->RestoreProperties(this);
+	}
+
+	if (!ensure(Targets.Num() > 0) || !ensure(Targets[0]))
+	{
+		return;
+	}
+	
+	IPrimitiveComponentBackedTarget* TargetComponent = TargetComponentInterface(0);
+	IMeshDescriptionProvider* TargetMeshProvider = TargetMeshProviderInterface(0);
+	IMaterialProvider* TargetMaterial = TargetMaterialInterface(0);
+
+	const USkeletalMeshComponent* SkelMeshComponent = Cast<USkeletalMeshComponent>(TargetComponent->GetOwnerComponent());
+	
+	if (SkelMeshComponent && SkelMeshComponent->SkeletalMesh)
+	{
+		USkeletalMesh* SkeletalMesh = SkelMeshComponent->SkeletalMesh;
+
+		// Initialize the bone browser
+		FCurveEvaluationOption CurveEvalOption(
+				SkelMeshComponent->GetAllowedAnimCurveEvaluate(),
+				&SkelMeshComponent->GetDisallowedAnimCurvesEvaluation(),
+				0 /* Always use the highest LOD */
+				);
+		BoneContainer.InitializeTo(SkelMeshComponent->RequiredBones, CurveEvalOption, *SkeletalMesh);
+
+		Properties->SkeletalMesh = SkeletalMesh;
+		Properties->CurrentBone.Initialize(BoneContainer);
+
+		const FReferenceSkeleton& RefSkeleton = SkeletalMesh->GetRefSkeleton();
+		for (int32 Index = 0; Index < RefSkeleton.GetRawBoneNum(); Index++)
+		{
+			BoneToIndex.Add(RefSkeleton.GetRawRefBoneInfo()[Index].Name, Index);
+		}
+
+		// Pick the first root bone as the initial selection.
+		Properties->CurrentBone.BoneName = SkeletalMesh->GetRefSkeleton().GetBoneName(0);
+
+		const TArray<FMeshBoneInfo>& BoneInfo = RefSkeleton.GetRawRefBoneInfo();
+		const TArray<FTransform>& BonePose = RefSkeleton.GetRawRefBonePose();
+		
+		TransformHierarchy.Reserve(BoneInfo.Num());
+		
+		for (int32 Index = 0; Index < BoneInfo.Num(); Index++)
+		{
+			TransformHierarchy.Add(MakeTuple(BonePose[Index], BoneInfo[Index].ParentIndex));
+		}
+	}
+	
+	TargetComponent->SetOwnerVisibility(false);
+	Preview = NewObject<UMeshOpPreviewWithBackgroundCompute>(this, "Preview");
+	Preview->Setup(this->TargetWorld, this);
+	Preview->SetIsMeshTopologyConstant(true, EMeshRenderAttributeFlags::VertexColors);
+	Preview->OnMeshUpdated.AddLambda([this](UMeshOpPreviewWithBackgroundCompute* Compute)
+	{
+		UpdateVisualization();
+	});
+	
+	FComponentMaterialSet MaterialSet;
+	TargetMaterial->GetMaterialSet(MaterialSet);
+
+	UMaterialInterface* VtxColorMaterial = GetToolManager()->GetContextQueriesAPI()->GetStandardMaterial(EStandardToolContextMaterials::VertexColorMaterial);
+	if (VtxColorMaterial != nullptr)
+	{
+		for (UMaterialInterface*& Material: MaterialSet.Materials)
+		{
+			Material = VtxColorMaterial;
+		}
+	}
+	
+	Preview->ConfigureMaterials( MaterialSet.Materials,
+								 ToolSetupUtil::GetDefaultWorkingMaterial(GetToolManager())
+	);
+
+	Properties->WatchProperty(Properties->CurrentBone.BoneName,
+							  [this](FName) { UpdateVisualization();});
+
+	OriginalMesh = MakeShared<FDynamicMesh3, ESPMode::ThreadSafe>();
+	FMeshDescriptionToDynamicMesh Converter;
+	Converter.Convert(TargetMeshProvider->GetMeshDescription(), *OriginalMesh);
+
+	// Enable or override vertex colors on the original mesh.
+	OriginalMesh->EnableAttributes();
+	OriginalMesh->Attributes()->DisablePrimaryColors();
+	OriginalMesh->Attributes()->EnablePrimaryColors();
+	// Create an overlay that has no split elements, init with zero value.
+	OriginalMesh->Attributes()->PrimaryColors()->CreateFromPredicate([](int ParentVID, int TriIDA, int TriIDB){return true;}, 0.f);
+
+	Preview->PreviewMesh->SetTransform(TargetComponent->GetWorldTransform());
+	Preview->PreviewMesh->SetTangentsMode(EDynamicMeshTangentCalcType::AutoCalculated);
+	Preview->PreviewMesh->GetRootComponent()->bCastDynamicShadow = false;
+	Preview->PreviewMesh->UpdatePreview(OriginalMesh.Get());
+
+	Occupancy = MakeShared<FOccupancyGrid>(*OriginalMesh, Properties->VoxelResolution);
+	
+	UpdateVisualization(/*bForce=*/true);
+	
+	// add properties to GUI
+	AddToolPropertySource(Properties);
+
+	Preview->InvalidateResult();
+
+	SetToolDisplayName(LOCTEXT("ToolName", "Bind Skin"));
+	GetToolManager()->DisplayMessage(
+		LOCTEXT("OnStartTool", "Creates a rigid binding for the skin weights."),
+		EToolMessageLevel::UserNotification);
+}
+
+
+void USkinWeightsBindingTool::Shutdown(EToolShutdownType ShutdownType)
+{
+	Properties->SaveProperties(this);
+
+	TargetComponentInterface(0)->SetOwnerVisibility(true);
+	
+	FDynamicMeshOpResult Result = Preview->Shutdown();
+	if (ShutdownType == EToolShutdownType::Accept)
+	{
+		GenerateAsset(Result);
+	}
+}
+
+
+void USkinWeightsBindingTool::OnTick(float DeltaTime)
+{
+	Preview->Tick(DeltaTime);
+}
+
+static void DrawBox(IToolsContextRenderAPI* RenderAPI, const FTransform& Transform, const FBox &Box, const FLinearColor &Color, float LineThickness)
+{
+	FPrimitiveDrawInterface* PDI = RenderAPI->GetPrimitiveDrawInterface();
+	const float PDIScale = RenderAPI->GetCameraState().GetPDIScalingFactor();
+
+	FVector Corners[2] = {
+		Transform.TransformPosition(Box.Min),
+		Transform.TransformPosition(Box.Max)
+	};
+
+	static const UE::Geometry::FVector3i Offsets[12][2] =
+	{
+		// Bottom
+		{{0, 0, 0}, {1, 0, 0}},
+		{{1, 0, 0}, {1, 1, 0}},
+		{{1, 1, 0}, {0, 1, 0}},
+		{{0, 1, 0}, {0, 0, 0}},
+		
+		// Top
+		{{0, 0, 1}, {1, 0, 1}},
+		{{1, 0, 1}, {1, 1, 1}},
+		{{1, 1, 1}, {0, 1, 1}},
+		{{0, 1, 1}, {0, 0, 1}},
+		
+		// Sides
+		{{0, 0, 0}, {0, 0, 1}},
+		{{1, 0, 0}, {1, 0, 1}},
+		{{1, 1, 0}, {1, 1, 1}},
+		{{0, 1, 0}, {0, 1, 1}},
+	}; 
+
+	for (int32 Index = 0; Index < 12; Index++)
+	{
+		const UE::Geometry::FVector3i* LineOffsets = Offsets[Index];
+		FVector  A(Corners[LineOffsets[0].X].X, Corners[LineOffsets[0].Y].Y, Corners[LineOffsets[0].Z].Z);
+		FVector  B(Corners[LineOffsets[1].X].X, Corners[LineOffsets[1].Y].Y, Corners[LineOffsets[1].Z].Z);
+		
+		PDI->DrawTranslucentLine(A, B, Color, 1, LineThickness * PDIScale);
+	}
+}
+
+
+void USkinWeightsBindingTool::Render(IToolsContextRenderAPI* RenderAPI)
+{
+	/**/
+
+	if (Occupancy && Properties->bDebugDraw)
+	{
+		bool bShowInterior = false;
+		bool bShowBoundary = true;
+
+		FPrimitiveDrawInterface* PDI = RenderAPI->GetPrimitiveDrawInterface();
+		FTransform Transform = TargetComponentInterface(0)->GetWorldTransform();
+		float PDIScale = RenderAPI->GetCameraState().GetPDIScalingFactor();
+
+		for (int32 I = 0; I < Occupancy->Occupancy.Size(); I++)
+		{
+			UE::Geometry::FVector3i G(Occupancy->Occupancy.ToIndex(I));
+
+			FOccupancyGrid::EDomain Domain = static_cast<FOccupancyGrid::EDomain>(Occupancy->Occupancy[G]);
+			if (bShowBoundary && Domain == FOccupancyGrid::Boundary)
+			{
+				FBox Box = Occupancy->GetCellBoxFromIndex(G);
+				DrawBox(RenderAPI, Transform, Box, FLinearColor(1.0, 1.0, 0.0, 0.5), 0.5f);
+			}
+			
+		}
+	}
+	
+}
+
+
+bool USkinWeightsBindingTool::CanAccept() const
+{
+	return Super::CanAccept() && Preview->HaveValidResult();
+}
+
+
+void USkinWeightsBindingTool::OnPropertyModified(UObject* PropertySet, FProperty* Property)
+{
+	if ( Property )
+	{
+		if (Property->GetFName() == GET_MEMBER_NAME_CHECKED(FBoneReference, BoneName) /* ||
+			Property->GetFName() == GET_MEMBER_NAME_CHECKED(USkinWeightsBindingToolProperties, bDebugDraw) */ )
+		{
+			// Handled by the property watcher.
+		}
+		else
+		{
+			Occupancy = MakeShared<FOccupancyGrid>(*OriginalMesh, Properties->VoxelResolution);
+			
+			Preview->InvalidateResult();
+		}
+	}
+}
+
+
+TUniquePtr<UE::Geometry::FDynamicMeshOperator> USkinWeightsBindingTool::MakeNewOperator()
+{
+	TUniquePtr<FComputeSkinWeightsBindingOp> Op = MakeUnique<FComputeSkinWeightsBindingOp>();
+
+	Op->ProfileName = FSkeletalMeshAttributes::DefaultSkinWeightProfileName;
+	Op->BindType = Properties->BindingType;
+	Op->Stiffness = Properties->Stiffness;
+	Op->MaxInfluences = Properties->MaxInfluences;
+	Op->VoxelResolution = Properties->VoxelResolution;
+	
+	Op->OriginalMesh = OriginalMesh;
+	Op->TransformHierarchy = TransformHierarchy;
+
+	const FTransform LocalToWorld = TargetComponentInterface(0)->GetWorldTransform();
+	Op->SetResultTransform(static_cast<UE::Geometry::FTransform3d>(LocalToWorld));
+	
+	return Op;
+}
+
+
+void USkinWeightsBindingTool::GenerateAsset(const FDynamicMeshOpResult& Result)
+{
+	// TODO: Update FDynamicMeshToMeshDescription to allow update the skin weights only.
+	GetToolManager()->BeginUndoTransaction(LOCTEXT("SkinWeightsBindingToolTransactionName", "Create Rigid Binding"));
+
+	check(Result.Mesh.Get() != nullptr);
+	TargetMeshCommitterInterface(0)->CommitMeshDescription([&Result](const IMeshDescriptionCommitter::FCommitterParams& CommitParams)
+	{
+		FDynamicMeshToMeshDescription Converter;
+
+		// full conversion if normal topology changed or faces were inverted
+		Converter.Convert(Result.Mesh.Get(), *CommitParams.MeshDescriptionOut);
+	});
+
+	GetToolManager()->EndUndoTransaction();	
+}
+
+FVector4f USkinWeightsBindingTool::WeightToColor(float Value)
+{
+	Value = FMath::Clamp(Value, 0.0f, 1.0f);
+
+	{
+		// A close approximation of the skeletal mesh editor's bone weight ramp. 
+		const FLinearColor HSV((1.0f - Value) * 285.0f, 100.0f, 85.0f);
+		return FVector4f(HSV.HSVToLinearRGB());
+	}
+}
+
+void USkinWeightsBindingTool::UpdateVisualization(bool bInForce)
+{
+	using namespace UE::AnimationCore;
+	using namespace UE::Geometry;
+	
+	if ((bInForce || Preview->HaveValidNonEmptyResult()) && BoneToIndex.Contains(Properties->CurrentBone.BoneName))
+	{
+		const FBoneIndexType BoneIndex = BoneToIndex[Properties->CurrentBone.BoneName];
+
+		// update mesh with new value colors
+		Preview->PreviewMesh->EditMesh([&](FDynamicMesh3& InMesh)
+		{
+			FDynamicMeshVertexSkinWeightsAttribute *SkinWeights = InMesh.Attributes()->GetSkinWeightsAttribute(FSkeletalMeshAttributes::DefaultSkinWeightProfileName);
+			FDynamicMeshColorOverlay* ColorOverlay = InMesh.Attributes()->PrimaryColors();
+
+			if (!ColorOverlay)
+			{
+				InMesh.EnableAttributes();
+				InMesh.Attributes()->EnablePrimaryColors();
+				// Create an overlay that has no split elements, init with zero value.
+				ColorOverlay = InMesh.Attributes()->PrimaryColors();
+				ColorOverlay->CreateFromPredicate([](int ParentVID, int TriIDA, int TriIDB){return true;}, 0.f);
+			}
+			
+			FBoneWeights BoneWeights;
+			
+			for (int32 ElementId : ColorOverlay->ElementIndicesItr())
+			{
+				const int32 VertexId = ColorOverlay->GetParentVertex(ElementId);
+				SkinWeights->GetValue(VertexId, BoneWeights);
+
+				float Weight = 0.0f;
+				for (FBoneWeight BW: BoneWeights)
+				{
+					if (BW.GetBoneIndex() == BoneIndex)
+					{
+						Weight = BW.GetWeight();
+						break;
+					}
+				}
+				
+				const FVector4f Color(WeightToColor(Weight));
+				ColorOverlay->SetElement(ElementId, Color);
+			}
+		});
+	}
+}
+
+
+#undef LOCTEXT_NAMESPACE
+
+// #pragma optimize( "", on )
diff --git a/Engine/Plugins/Experimental/MeshModelingToolset/Source/MeshModelingTools/Public/SkinWeightsBindingTool.h b/Engine/Plugins/Experimental/MeshModelingToolset/Source/MeshModelingTools/Public/SkinWeightsBindingTool.h
new file mode 100644
--- /dev/null
+++ b/Engine/Plugins/Experimental/MeshModelingToolset/Source/MeshModelingTools/Public/SkinWeightsBindingTool.h
@@ -0,0 +1,135 @@
+// Copyright Epic Games, Inc. All Rights Reserved.
+
+#pragma once
+
+#include "CoreMinimal.h"
+
+#include "BoneContainer.h"
+#include "InteractiveToolBuilder.h"
+#include "ModelingOperators.h"
+#include "MultiSelectionTool.h"
+#include "Interfaces/Interface_BoneReferenceSkeletonProvider.h"
+
+#include "SkinWeightsBindingTool.generated.h"
+
+class UMeshOpPreviewWithBackgroundCompute;
+struct FDynamicMeshOpResult;
+struct FMeshDescription;
+class IMeshDescriptionCommitter;
+
+
+struct FOccupancyGrid;
+
+/**
+ *
+ */
+UCLASS()
+class MESHMODELINGTOOLS_API USkinWeightsBindingToolBuilder : public UInteractiveToolBuilder
+{
+	GENERATED_BODY()
+
+public:
+	IToolsContextAssetAPI* AssetAPI = nullptr;
+
+	bool CanBuildTool(const FToolBuilderState& SceneState) const override;
+
+	UInteractiveTool* BuildTool(const FToolBuilderState& SceneState) const override;
+};
+
+UENUM()
+enum class ESkinWeightsBindType : uint8
+{
+	DirectDistance = 0 UMETA(DisplayName = "Direct Distance"),
+	// GeodesicDistance UMETA(DisplayName = "Geodesic Distance"),
+	// Heatmaps UMETA(DisplayName = "Heatmap"),
+	GeodesicVoxel = 3 UMETA(DisplayName = "Geodesic Voxel"),
+};
+
+UCLASS()
+class MESHMODELINGTOOLS_API USkinWeightsBindingToolProperties : 
+	public UInteractiveToolPropertySet,
+	public IBoneReferenceSkeletonProvider
+{
+	GENERATED_BODY()
+public:
+	UPROPERTY(VisibleAnywhere, Category = Skeleton)
+	FBoneReference CurrentBone;
+
+	/** Binding type to use */
+	UPROPERTY(EditAnywhere, Category = Binding)
+	ESkinWeightsBindType BindingType = ESkinWeightsBindType::DirectDistance;
+
+	/** Stiffness of binding. Lower values allow more distant bones to contribute more */
+	UPROPERTY(EditAnywhere, Category = Binding)
+	float Stiffness = 0.2f;
+
+	/** Maximum bones that will influence each vertex */
+	UPROPERTY(EditAnywhere, Category = Binding, meta=(ClampMin="1", UIMin="1", UIMax="10"))
+	int32 MaxInfluences = 5;
+
+	/** The resolution of the voxel grid if doing geodesic voxel binding */
+	UPROPERTY(EditAnywhere, Category = Binding, meta=(EditCondition = "BindingType == ESkinWeightsBindType::GeodesicVoxel", ClampMin="1", UIMin="8", UIMax="1024"))
+	int32 VoxelResolution = 128;
+	
+	// UPROPERTY(EditAnywhere, Category = Debug)
+	bool bDebugDraw = false;
+	
+	// IBoneReferenceSkeletonProvider
+	USkeleton* GetSkeleton(bool& bInvalidSkeletonIsError) override;
+
+	TObjectPtr<USkeletalMesh> SkeletalMesh;
+};
+
+/**
+ *
+ */
+UCLASS()
+class MESHMODELINGTOOLS_API USkinWeightsBindingTool :
+	public UMultiSelectionTool,
+	public UE::Geometry::IDynamicMeshOperatorFactory
+{
+	GENERATED_BODY()
+
+public:
+	USkinWeightsBindingTool();
+	~USkinWeightsBindingTool();
+
+	void SetWorld(UWorld* World);
+	
+	void Setup() override;
+	void Shutdown(EToolShutdownType ShutdownType) override;
+	
+	void OnTick(float DeltaTime) override;
+	void Render(IToolsContextRenderAPI* RenderAPI) override;
+
+	bool HasCancel() const override { return true; }
+	bool HasAccept() const override { return true; }
+	bool CanAccept() const override;
+
+	void OnPropertyModified(UObject* PropertySet, FProperty* Property) override;
+
+	// IDynamicMeshOperatorFactory API
+	TUniquePtr<UE::Geometry::FDynamicMeshOperator> MakeNewOperator() override;
+
+	UPROPERTY()
+	USkinWeightsBindingToolProperties* Properties = nullptr;
+	
+	UPROPERTY()
+	UMeshOpPreviewWithBackgroundCompute* Preview = nullptr;
+
+protected:
+	UWorld* TargetWorld = nullptr;
+
+	TSharedPtr<FOccupancyGrid> Occupancy;
+	
+	TSharedPtr<UE::Geometry::FDynamicMesh3, ESPMode::ThreadSafe> OriginalMesh;
+	
+	FBoneContainer BoneContainer;
+	TMap<FName, FBoneIndexType> BoneToIndex;
+	TArray<TPair<FTransform, int32>> TransformHierarchy;
+
+	void GenerateAsset(const FDynamicMeshOpResult& Result);
+
+	static FVector4f WeightToColor(float Value);
+	void UpdateVisualization(bool bInForce = false);
+};
diff --git a/Engine/Plugins/Experimental/ModelingToolsEditorMode/Source/ModelingToolsEditorMode/Private/ModelingToolsManagerActions.cpp b/Engine/Plugins/Experimental/ModelingToolsEditorMode/Source/ModelingToolsEditorMode/Private/ModelingToolsManagerActions.cpp
--- a/Engine/Plugins/Experimental/ModelingToolsEditorMode/Source/ModelingToolsEditorMode/Private/ModelingToolsManagerActions.cpp
+++ b/Engine/Plugins/Experimental/ModelingToolsEditorMode/Source/ModelingToolsEditorMode/Private/ModelingToolsManagerActions.cpp
@@ -102,6 +102,7 @@ void FModelingToolsManagerCommands::RegisterCommands()
 	UI_COMMAND(BeginUVSeamEditTool, "SeamEd", "Add UV Seams to Mesh", EUserInterfaceActionType::ToggleButton, FInputChord());
 
 	UI_COMMAND(BeginSkinWeightsPaintTool, "SkinWts", "Start the Paint Skin Weights Tool", EUserInterfaceActionType::Button, FInputChord());
+	UI_COMMAND(BeginSkinWeightsBindingTool, "SkinBind", "Start the Skin Weights Binding Tool", EUserInterfaceActionType::Button, FInputChord());
 
 	UI_COMMAND(BeginGroomToMeshTool, "Helmet", "Generate Helmet Mesh for Selected Groom", EUserInterfaceActionType::ToggleButton, FInputChord());
 	UI_COMMAND(BeginGenerateLODMeshesTool, "HlmLOD", "Generate LODS for Hair Helmet", EUserInterfaceActionType::ToggleButton, FInputChord());
diff --git a/Engine/Plugins/Experimental/ModelingToolsEditorMode/Source/ModelingToolsEditorMode/Public/ModelingToolsManagerActions.h b/Engine/Plugins/Experimental/ModelingToolsEditorMode/Source/ModelingToolsEditorMode/Public/ModelingToolsManagerActions.h
--- a/Engine/Plugins/Experimental/ModelingToolsEditorMode/Source/ModelingToolsEditorMode/Public/ModelingToolsManagerActions.h
+++ b/Engine/Plugins/Experimental/ModelingToolsEditorMode/Source/ModelingToolsEditorMode/Public/ModelingToolsManagerActions.h
@@ -92,11 +92,13 @@ public:
 	TSharedPtr<FUICommandInfo> BeginEditMeshMaterialsTool;
 	TSharedPtr<FUICommandInfo> BeginTransformUVIslandsTool;
 	TSharedPtr<FUICommandInfo> BeginMeshAttributePaintTool;
-	TSharedPtr<FUICommandInfo> BeginSkinWeightsPaintTool;
 	TSharedPtr<FUICommandInfo> BeginAttributeEditorTool;
 	TSharedPtr<FUICommandInfo> BeginBakeMeshAttributeMapsTool;
 	TSharedPtr<FUICommandInfo> BeginUVSeamEditTool;
 
+	TSharedPtr<FUICommandInfo> BeginSkinWeightsPaintTool;
+	TSharedPtr<FUICommandInfo> BeginSkinWeightsBindingTool;
+	
 	TSharedPtr<FUICommandInfo> BeginGroomToMeshTool;
 	TSharedPtr<FUICommandInfo> BeginGenerateLODMeshesTool;
 	TSharedPtr<FUICommandInfo> BeginGroomCardsToMeshTool;
