MeshModelingTools: Add channel selection to displace tool.

#rnx

diff --git a/Engine/Plugins/Experimental/MeshModelingToolset/Source/MeshModelingTools/Private/DisplaceMeshTool.cpp b/Engine/Plugins/Experimental/MeshModelingToolset/Source/MeshModelingTools/Private/DisplaceMeshTool.cpp
--- a/Engine/Plugins/Experimental/MeshModelingToolset/Source/MeshModelingTools/Private/DisplaceMeshTool.cpp
+++ b/Engine/Plugins/Experimental/MeshModelingToolset/Source/MeshModelingTools/Private/DisplaceMeshTool.cpp
@@ -366,6 +366,7 @@ namespace DisplaceMeshToolLocals{
 		bool bRecalculateNormals = true;
 
 		// Used in texture map displacement
+		int32 DisplacementMapChannel = 0;
 		float DisplacementMapBaseValue = 128.0/255; // i.e., what constitutes no displacement
 		FVector2f UVScale = FVector2f(1,1);
 		FVector2f UVOffset = FVector2f(0, 0);
@@ -546,7 +547,7 @@ namespace DisplaceMeshToolLocals{
 		{
 			SetIntensity(DisplaceParametersIn.DisplaceIntensity);
 			SetRandomSeed(DisplaceParametersIn.RandomSeed);
-			SetDisplacementMap(DisplaceParametersIn.DisplacementMap); // Calls UpdateMap
+			SetDisplacementMap(DisplaceParametersIn.DisplacementMap, DisplaceParametersIn.DisplacementMapChannel); // Calls UpdateMap
 			SetFrequency(DisplaceParametersIn.SineWaveFrequency);
 			SetPhaseShift(DisplaceParametersIn.SineWavePhaseShift);
 			SetSineWaveDirection(DisplaceParametersIn.SineWaveDirection);
@@ -567,7 +568,7 @@ namespace DisplaceMeshToolLocals{
 		}
 		void SetIntensity(float IntensityIn);
 		void SetRandomSeed(int RandomSeedIn);
-		void SetDisplacementMap(UTexture2D* DisplacementMapIn);
+		void SetDisplacementMap(UTexture2D* DisplacementMapIn, int32 ChannelIn);
 		void SetDisplacementMapUVAdjustment(const FVector2f& UVScale, const FVector2f& UVOffset);
 		void SetDisplacementMapBaseValue(float DisplacementMapBaseValue);
 		void SetAdjustmentCurve(UCurveFloat* CurveFloat);
@@ -605,9 +606,10 @@ namespace DisplaceMeshToolLocals{
 		Parameters.RandomSeed = RandomSeedIn;
 	}
 
-	void FDisplaceMeshOpFactory::SetDisplacementMap(UTexture2D* DisplacementMapIn)
+	void FDisplaceMeshOpFactory::SetDisplacementMap(UTexture2D* DisplacementMapIn, int32 ChannelIn)
 	{
 		Parameters.DisplacementMap = DisplacementMapIn;
+		Parameters.DisplacementMapChannel = ChannelIn;
 
 		// Note that we do the update even if we got the same pointer, because the texture
 		// may have been changed in the editor.
@@ -687,7 +689,8 @@ namespace DisplaceMeshToolLocals{
 			{
 				for (int64 x = 0; x < TextureWidth; ++x)
 				{
-					Parameters.DisplaceField.GridValues[y * TextureWidth + x] = DisplacementMapValues.GetPixel(y * TextureWidth + x).X;
+					Parameters.DisplaceField.GridValues[y * TextureWidth + x] = 
+						DisplacementMapValues.GetPixel(y * TextureWidth + x)[Parameters.DisplacementMapChannel];
 				}
 			}
 		}
@@ -1021,7 +1024,11 @@ void UDisplaceMeshTool::OnPropertyModified(UObject* PropertySet, FProperty* Prop
 				GetToolManager()->DisplayMessage(FText::GetEmpty(), EToolMessageLevel::UserWarning);
 			}
 
-			DisplacerDownCast->SetDisplacementMap(TextureMapProperties->DisplacementMap);
+			DisplacerDownCast->SetDisplacementMap(TextureMapProperties->DisplacementMap, TextureMapProperties->Channel);
+		}
+		else if (PropName == GET_MEMBER_NAME_CHECKED(UDisplaceMeshTextureMapProperties, Channel))
+		{
+			DisplacerDownCast->SetDisplacementMap(TextureMapProperties->DisplacementMap, TextureMapProperties->Channel);
 		}
 		else if (PropName == GET_MEMBER_NAME_CHECKED(UDisplaceMeshTextureMapProperties, DisplacementMapBaseValue))
 		{
diff --git a/Engine/Plugins/Experimental/MeshModelingToolset/Source/MeshModelingTools/Public/DisplaceMeshTool.h b/Engine/Plugins/Experimental/MeshModelingToolset/Source/MeshModelingTools/Public/DisplaceMeshTool.h
--- a/Engine/Plugins/Experimental/MeshModelingToolset/Source/MeshModelingTools/Public/DisplaceMeshTool.h
+++ b/Engine/Plugins/Experimental/MeshModelingToolset/Source/MeshModelingTools/Public/DisplaceMeshTool.h
@@ -102,6 +102,11 @@ public:
 	UPROPERTY(EditAnywhere, Category = Options)
 	UTexture2D* DisplacementMap = nullptr;
 
+	//~ TODO: Consider making this an enum
+	/** Channel in the displacement map to use. 0-3 coresponding to the four channels (RGBA) */
+	UPROPERTY(EditAnywhere, Category = Options, meta = (ClampMin = "0", ClampMax = "3"))
+	int32 Channel = 0;
+
 	/** The value in the texture map that corresponds to no displacement. For instance, if set to 0, then all
 	 displacement will be positive. If set to 0.5, displacement below 0.5 will be negative, and above will be
 	 positive. Default is for 128/255 to be no displacement. */
