Move IDynamicMeshProvider and IDynamicMeshCommitter out of the dynamic mesh module so that it can remain devoid of UThings. Add comments to avoid this in the future.

#rnx

diff --git a/Engine/Plugins/Experimental/GeometryProcessing/Source/DynamicMesh/DynamicMesh.Build.cs b/Engine/Plugins/Experimental/GeometryProcessing/Source/DynamicMesh/DynamicMesh.Build.cs
--- a/Engine/Plugins/Experimental/GeometryProcessing/Source/DynamicMesh/DynamicMesh.Build.cs
+++ b/Engine/Plugins/Experimental/GeometryProcessing/Source/DynamicMesh/DynamicMesh.Build.cs
@@ -25,9 +25,12 @@ public class DynamicMesh : ModuleRules
 			}
 		);
 
+		// Note: The module purposefully doesn't have a dependency on CoreUObject.
+		// If possible, we would like avoid having UObjects in GeometryProcessing
+		// modules to keep the door open for writing standalone command-line programs
+		// (which won't have UObject garbage collection).
 		PrivateDependencyModuleNames.AddRange(
 			new string[] {
-				"CoreUObject",
 				"Eigen"
 			}
 		);
diff --git a/Engine/Plugins/Experimental/GeometryProcessing/Source/DynamicMesh/Public/TargetInterfaces/DynamicMeshCommitter.h b/Engine/Plugins/Experimental/GeometryProcessing/Source/DynamicMesh/Public/TargetInterfaces/DynamicMeshCommitter.h
deleted file mode 100644
diff --git a/Engine/Plugins/Experimental/GeometryProcessing/Source/DynamicMesh/Public/TargetInterfaces/DynamicMeshProvider.h b/Engine/Plugins/Experimental/GeometryProcessing/Source/DynamicMesh/Public/TargetInterfaces/DynamicMeshProvider.h
deleted file mode 100644
diff --git a/Engine/Plugins/Experimental/GeometryProcessing/Source/GeometricObjects/GeometricObjects.Build.cs b/Engine/Plugins/Experimental/GeometryProcessing/Source/GeometricObjects/GeometricObjects.Build.cs
--- a/Engine/Plugins/Experimental/GeometryProcessing/Source/GeometricObjects/GeometricObjects.Build.cs
+++ b/Engine/Plugins/Experimental/GeometryProcessing/Source/GeometricObjects/GeometricObjects.Build.cs
@@ -8,6 +8,10 @@ public class GeometricObjects : ModuleRules
 	{
         PCHUsage = ModuleRules.PCHUsageMode.UseExplicitOrSharedPCHs;
 
+		// Note: The module purposefully doesn't have a dependency on CoreUObject.
+		// If possible, we would like avoid having UObjects in GeometryProcessing
+		// modules to keep the door open for writing standalone command-line programs
+		// (which won't have UObject garbage collection).
 		PublicDependencyModuleNames.AddRange(
 			new string[] {
 				"Core"
diff --git a/Engine/Plugins/Experimental/GeometryProcessing/Source/GeometryAlgorithms/GeometryAlgorithms.Build.cs b/Engine/Plugins/Experimental/GeometryProcessing/Source/GeometryAlgorithms/GeometryAlgorithms.Build.cs
--- a/Engine/Plugins/Experimental/GeometryProcessing/Source/GeometryAlgorithms/GeometryAlgorithms.Build.cs
+++ b/Engine/Plugins/Experimental/GeometryProcessing/Source/GeometryAlgorithms/GeometryAlgorithms.Build.cs
@@ -8,6 +8,10 @@ public class GeometryAlgorithms : ModuleRules
 	{
         PCHUsage = ModuleRules.PCHUsageMode.UseExplicitOrSharedPCHs;
 
+		// Note: The module purposefully doesn't have a dependency on CoreUObject.
+		// If possible, we would like avoid having UObjects in GeometryProcessing
+		// modules to keep the door open for writing standalone command-line programs
+		// (which won't have UObject garbage collection).
         PublicDependencyModuleNames.AddRange(
 			new string[] {
 				"Core",
diff --git a/Engine/Plugins/Experimental/GeometryProcessing/Source/MeshConversion/MeshConversion.Build.cs b/Engine/Plugins/Experimental/GeometryProcessing/Source/MeshConversion/MeshConversion.Build.cs
--- a/Engine/Plugins/Experimental/GeometryProcessing/Source/MeshConversion/MeshConversion.Build.cs
+++ b/Engine/Plugins/Experimental/GeometryProcessing/Source/MeshConversion/MeshConversion.Build.cs
@@ -8,6 +8,10 @@ public class MeshConversion : ModuleRules
 	{
         PCHUsage = ModuleRules.PCHUsageMode.UseExplicitOrSharedPCHs;
 
+		// Note: The module purposefully doesn't have a dependency on CoreUObject.
+		// If possible, we would like avoid having UObjects in GeometryProcessing
+		// modules to keep the door open for writing standalone command-line programs
+		// (which won't have UObject garbage collection).
         PublicDependencyModuleNames.AddRange(
 			new string[] {
                 "Core",
@@ -19,12 +23,5 @@ public class MeshConversion : ModuleRules
 				"DynamicMesh"
             }
 		);
-
-        PrivateDependencyModuleNames.AddRange(
-            new string[]
-            {
-                "CoreUObject"
-			}
-		);
     }
 }
diff --git a/Engine/Plugins/Experimental/MeshModelingToolset/Source/MeshModelingTools/Public/TargetInterfaces/DynamicMeshCommitter.h b/Engine/Plugins/Experimental/MeshModelingToolset/Source/MeshModelingTools/Public/TargetInterfaces/DynamicMeshCommitter.h
new file mode 100644
--- /dev/null
+++ b/Engine/Plugins/Experimental/MeshModelingToolset/Source/MeshModelingTools/Public/TargetInterfaces/DynamicMeshCommitter.h
@@ -0,0 +1,66 @@
+// Copyright Epic Games, Inc. All Rights Reserved.
+
+#pragma once
+
+#include "CoreMinimal.h"
+#include "UObject/Interface.h"
+
+#include "GeometryBase.h"
+
+#include "DynamicMeshCommitter.generated.h"
+
+PREDECLARE_GEOMETRY(class FDynamicMesh3);
+
+
+UINTERFACE()
+class MESHMODELINGTOOLS_API UDynamicMeshCommitter : public UInterface
+{
+	GENERATED_BODY()
+};
+
+class MESHMODELINGTOOLS_API IDynamicMeshCommitter
+{
+	GENERATED_BODY()
+
+public:
+
+	/**
+	 * Extra information that can be passed to a CommitMesh call to potentially make
+	 * the commit faster. Note that setting any of these to false doesn't mean that
+	 * the corresponding data won't be updated, because a target may choose to always
+	 * update everything. But it may help some targets do faster updates by not
+	 * updating things that stayed the same.
+	 */
+	struct MESHMODELINGTOOLS_API FDynamicMeshCommitInfo
+	{
+		/** Initializes everything to bInitValue */
+		FDynamicMeshCommitInfo(bool bInitValue)
+		{
+			bPositionsChanged =
+			bTopologyChanged =
+			bPolygroupsChanged =
+			bNormalsChanged =
+			bTangentsChanged =
+			bUVsChanged =
+			bVertexColorsChanged = bInitValue;
+		}
+
+		/** Leaves everything initialized to default (true) */
+		FDynamicMeshCommitInfo() {}
+
+		bool bPositionsChanged = true;
+		bool bTopologyChanged = true;
+		bool bPolygroupsChanged = true;
+		bool bNormalsChanged = true;
+		bool bTangentsChanged = true;
+		bool bUVsChanged = true;
+		bool bVertexColorsChanged = true;
+	};
+
+	virtual void CommitDynamicMesh(const UE::Geometry::FDynamicMesh3& Mesh)
+	{
+		FDynamicMeshCommitInfo CommitInfo;
+		CommitDynamicMesh(Mesh, CommitInfo);
+	};
+	virtual void CommitDynamicMesh(const UE::Geometry::FDynamicMesh3& Mesh, const FDynamicMeshCommitInfo& CommitInfo) = 0;
+};
diff --git a/Engine/Plugins/Experimental/MeshModelingToolset/Source/MeshModelingTools/Public/TargetInterfaces/DynamicMeshProvider.h b/Engine/Plugins/Experimental/MeshModelingToolset/Source/MeshModelingTools/Public/TargetInterfaces/DynamicMeshProvider.h
new file mode 100644
--- /dev/null
+++ b/Engine/Plugins/Experimental/MeshModelingToolset/Source/MeshModelingTools/Public/TargetInterfaces/DynamicMeshProvider.h
@@ -0,0 +1,35 @@
+// Copyright Epic Games, Inc. All Rights Reserved.
+
+#pragma once
+
+#include "CoreMinimal.h"
+#include "UObject/Interface.h"
+
+#include "GeometryBase.h"
+
+#include "DynamicMeshProvider.generated.h"
+
+PREDECLARE_GEOMETRY(class FDynamicMesh3);
+
+UINTERFACE()
+class MESHMODELINGTOOLS_API UDynamicMeshProvider : public UInterface
+{
+	GENERATED_BODY()
+};
+
+class MESHMODELINGTOOLS_API IDynamicMeshProvider
+{
+	GENERATED_BODY()
+
+public:
+	/**
+	 * Gives back an editable dynamic mesh. Changes may or may not affect the target, but
+	 * target guarantees that they are safe to perform without crashing. In particular,
+	 * note that it's unspecified whether the next GetMesh() call after an edit will give
+	 * back the edited mesh or a copy of the original.
+	 * 
+	 * To guarantee changes having an effect, target must also be a IDynamicMeshCommitter 
+	 * and the mesh must be committed.
+	 */
+	virtual TSharedPtr<UE::Geometry::FDynamicMesh3> GetDynamicMesh() = 0;
+};
