ApproximateActors updates:
- FMeshSceneAdapter updates
   - if bThickenThinMeshes is enabled, then all instance meshes are converted to DynamicMesh and decomposed to find non-closed regions. Those regions are uniqued into separate meshes. The remaining closed geometry is separated out and kept as instanced.
   - various utility functions added to support above
   - previous code to explicitly 'thicken' open meshes into solids is removed. Instead unsigned distance is used to determine inside/outside and generate a 1/0 winding number
   - instances are now sorted into an octree to speed up winding number evaluations
   - If all meshes are closed (or effectively closed via above, ie will return integer winding number) then winding number query can early-out as soon as it finds a non-zero value
- Add hard-normals options in IGeometryProcessing_ApproximateActors (not exposed in MergeActors tool options yet)
- Added hard-normal callucation to ApproximateActorsImpl, and additional profiling
- Allow specification of actual  ECollisionTraceFlag options in CreateStaticMeshUtil, instead of just a boolean
- Add FSparseDynamicOctree3::ContainmentQueryCancellable() variant that can early-exit

#rnx

diff --git a/Engine/Plugins/Experimental/GeometryProcessing/Source/GeometricObjects/Private/Spatial/SparseDynamicOctree3.cpp b/Engine/Plugins/Experimental/GeometryProcessing/Source/GeometricObjects/Private/Spatial/SparseDynamicOctree3.cpp
--- a/Engine/Plugins/Experimental/GeometryProcessing/Source/GeometricObjects/Private/Spatial/SparseDynamicOctree3.cpp
+++ b/Engine/Plugins/Experimental/GeometryProcessing/Source/GeometricObjects/Private/Spatial/SparseDynamicOctree3.cpp
@@ -363,6 +363,73 @@ void FSparseDynamicOctree3::ContainmentQuery(
 }
 
 
+
+
+bool FSparseDynamicOctree3::ContainmentQueryCancellable(
+	const FVector3d& Point,
+	TFunctionRef<bool(int)> ObjectIDFunc) const
+{
+	// todo: this should take advantage of raster!
+
+	// always process spill objects
+	for (int ObjectID : SpillObjectSet)
+	{
+		if (ObjectIDFunc(ObjectID) == false)
+		{
+			return false;
+		}
+	}
+
+	TArray<const FSparseOctreeCell*, TInlineAllocator<32>> Queue;
+
+	// start at root cells
+	RootCells.AllocatedIteration([&](const uint32* RootCellID)
+	{
+		const FSparseOctreeCell* RootCell = &Cells[*RootCellID];
+		if ( GetCellBox(*RootCell, MaxExpandFactor).Contains(Point) )
+		{
+			Queue.Add(&Cells[*RootCellID]);
+		}
+	});
+
+
+	while (Queue.Num() > 0)
+	{
+		const FSparseOctreeCell* CurCell = Queue.Pop(false);
+
+		// process elements
+		bool bContinue = true;
+		CellObjectLists.Enumerate(CurCell->CellID, [&](int32 ObjectID)
+		{
+			if ( bContinue && ObjectIDFunc(ObjectID) == false )
+			{
+				bContinue = false;
+			}
+		});
+		if (!bContinue)
+		{
+			return false;
+		}
+
+		for (int k = 0; k < 8; ++k)
+		{
+			if (CurCell->HasChild(k))
+			{
+				const FSparseOctreeCell* ChildCell = &Cells[CurCell->GetChildCellID(k)];
+				if (GetCellBox(*ChildCell, MaxExpandFactor).Contains(Point))
+				{
+					Queue.Add(ChildCell);
+				}
+			}
+		}
+	}
+
+	return true;
+}
+
+
+
+
 void FSparseDynamicOctree3::RangeQuery(
 	const FAxisAlignedBox3d& Bounds,
 	TFunctionRef<void(int)> ObjectIDFunc) const
diff --git a/Engine/Plugins/Experimental/GeometryProcessing/Source/GeometricObjects/Public/Spatial/SparseDynamicOctree3.h b/Engine/Plugins/Experimental/GeometryProcessing/Source/GeometricObjects/Public/Spatial/SparseDynamicOctree3.h
--- a/Engine/Plugins/Experimental/GeometryProcessing/Source/GeometricObjects/Public/Spatial/SparseDynamicOctree3.h
+++ b/Engine/Plugins/Experimental/GeometryProcessing/Source/GeometricObjects/Public/Spatial/SparseDynamicOctree3.h
@@ -255,15 +255,24 @@ public:
 	/**
 	 * Process ObjectIDs from all the cells with bounding boxes that contain query point
 	 * @param Point query point
-	 * @param ObjectIDFunc this function is called
+	 * @param ObjectIDFunc this function is called for each ObjectID
 	 */
 	void ContainmentQuery(const FVector3d& Point,
 		TFunctionRef<void(int)> ObjectIDFunc) const;
 
+	/**
+	 * Process ObjectIDs from all the cells with bounding boxes that contain query point
+	 * @param Point query point
+	 * @param ObjectIDFunc this function is called for each ObjectID. Returns true to continue query and false to abort
+	 * @return true if query finished, false if it exited
+	 */
+	bool ContainmentQueryCancellable(const FVector3d& Point,
+		TFunctionRef<bool(int)> ObjectIDFunc) const;
+
 	/**
 	 * Process ObjectIDs from all the cells with bounding boxes that intersect Bounds
 	 * @param Bounds query box
-	 * @param ObjectIDFunc this function is called 
+	 * @param ObjectIDFunc this function is called for each ObjectID
 	 */
 	void RangeQuery(const FAxisAlignedBox3d& Bounds,
 		TFunctionRef<void(int)> ObjectIDFunc) const;
diff --git a/Engine/Plugins/Experimental/MeshModelingToolset/Source/GeometryProcessingAdapters/Private/ApproximateActorsImpl.cpp b/Engine/Plugins/Experimental/MeshModelingToolset/Source/GeometryProcessingAdapters/Private/ApproximateActorsImpl.cpp
--- a/Engine/Plugins/Experimental/MeshModelingToolset/Source/GeometryProcessingAdapters/Private/ApproximateActorsImpl.cpp
+++ b/Engine/Plugins/Experimental/MeshModelingToolset/Source/GeometryProcessingAdapters/Private/ApproximateActorsImpl.cpp
@@ -289,7 +289,7 @@ static TSharedPtr<FApproximationMeshData> GenerateApproximationMesh(
 
 	TSharedPtr<FApproximationMeshData> Result = MakeShared<FApproximationMeshData>();
 
-	// collect seed poitns
+	// collect seed points
 	TArray<FVector3d> SeedPoints;
 	{
 		TRACE_CPUPROFILER_EVENT_SCOPE(ApproximateActorsImpl_Generate_SeedPoints);
@@ -315,7 +315,7 @@ static TSharedPtr<FApproximationMeshData> GenerateApproximationMesh(
 	Progress.EnterProgressFrame(1.f, LOCTEXT("SolidifyMesh", "Approximating Mesh..."));
 
 	FWindingNumberBasedSolidify Solidify(
-		[&Scene](const FVector3d& Position) { return Scene.FastWindingNumber(Position); },
+		[&Scene](const FVector3d& Position) { return Scene.FastWindingNumber(Position, true); },
 		SceneBounds, SeedPoints);
 	Solidify.SetCellSizeAndExtendBounds(SceneBounds, 2.0 * ApproxAccuracy, VoxelDimTarget);
 	Solidify.WindingThreshold = Options.WindingThreshold;
@@ -351,6 +351,11 @@ static TSharedPtr<FApproximationMeshData> GenerateApproximationMesh(
 		MorphologyMesh = FDynamicMesh3(&ImplicitMorphology.Generate());
 		MorphologyMesh.DiscardAttributes();
 		CurResultMesh = &MorphologyMesh;
+
+		if (Options.bVerbose)
+		{
+			UE_LOG(LogTemp, Warning, TEXT("[ApproximateActors] Morphology mesh has %d triangles"), CurResultMesh->TriangleCount());
+		}
 	}
 
 	// if mesh has no triangles, something has gone wrong
@@ -369,18 +374,25 @@ static TSharedPtr<FApproximationMeshData> GenerateApproximationMesh(
 		Remover.InsideMode = EOcclusionCalculationMode::SimpleOcclusionTest;
 		Remover.TriangleSamplingMethod = EOcclusionTriangleSampling::Centroids;
 		Remover.AddTriangleSamples = 5;
-		FDynamicMeshAABBTree3 CurResultMeshSpatial(CurResultMesh, true);
+		FDynamicMeshAABBTree3 CurResultMeshSpatial(CurResultMesh, false);
+		{
+			TRACE_CPUPROFILER_EVENT_SCOPE(ApproximateActorsImpl_Generate_Occlusion_Spatial);
+			CurResultMeshSpatial.Build();
+		}
 
 		TArray<FTransform3d> NoTransforms;
 		NoTransforms.Add(FTransform3d::Identity());
 		TArray<FDynamicMeshAABBTree3*> Spatials;
 		Spatials.Add(&CurResultMeshSpatial);
-		Remover.Select(NoTransforms, Spatials, {}, NoTransforms);
+		{
+			TRACE_CPUPROFILER_EVENT_SCOPE(ApproximateActorsImpl_Generate_Occlusion_Compute);
+			Remover.Select(NoTransforms, Spatials, {}, NoTransforms);
+		}
 		if (Remover.RemovedT.Num() > 0)
 		{
 			FMeshFaceSelection Selection(CurResultMesh);
 			{
-				TRACE_CPUPROFILER_EVENT_SCOPE(ApproximateActorsImpl_Generate_Occlusion_Build);
+				TRACE_CPUPROFILER_EVENT_SCOPE(ApproximateActorsImpl_Generate_Occlusion_Clean);
 				Selection.Select(Remover.RemovedT);
 				Selection.ExpandToOneRingNeighbours(1);
 				Selection.ContractBorderByOneRingNeighbours(4);
@@ -391,6 +403,11 @@ static TSharedPtr<FApproximationMeshData> GenerateApproximationMesh(
 				Editor.RemoveTriangles(Selection.AsArray(), true);
 			}
 		}
+
+		if (Options.bVerbose)
+		{
+			UE_LOG(LogTemp, Warning, TEXT("[ApproximateActors] Occlusion-Filtered mesh has %d triangles"), CurResultMesh->TriangleCount());
+		}
 	}
 
 	Progress.EnterProgressFrame(1.f, LOCTEXT("SimplifyingMesh", "Simplifying Mesh..."));
@@ -417,11 +434,18 @@ static TSharedPtr<FApproximationMeshData> GenerateApproximationMesh(
 			double UseTargetTolerance = Options.SimplificationTargetMetric * 100.0;		// convert to cm (UE Units)
 
 			// first do fast collapse
-			Simplifier.FastCollapsePass(0.1 * UseTargetTolerance, 5);
+			// (this does not seem to help perf and probably makes the results slightly worse)
+			//{
+			//	TRACE_CPUPROFILER_EVENT_SCOPE(ApproximateActorsImpl_Generate_Simplification_PrePass);
+			//	Simplifier.FastCollapsePass(0.1 * UseTargetTolerance, 5);
+			//}
 
 			// now simplify down to a reasonable tri count, as geometric metric is (relatively) expensive
 			// (still, this is all incredibly cheap compared to the cost of the rest of this method in practice)
-			Simplifier.SimplifyToTriangleCount(50000);
+			{
+				TRACE_CPUPROFILER_EVENT_SCOPE(ApproximateActorsImpl_Generate_Simplification_Pass1);
+				Simplifier.SimplifyToTriangleCount(50000);
+			}
 
 			FDynamicMesh3 MeshCopy(*CurResultMesh);
 			FDynamicMeshAABBTree3 MeshCopySpatial(&MeshCopy, true);
@@ -429,7 +453,10 @@ static TSharedPtr<FApproximationMeshData> GenerateApproximationMesh(
 			Simplifier.SetProjectionTarget(&ProjectionTarget);
 			Simplifier.GeometricErrorConstraint = FVolPresMeshSimplification::EGeometricErrorCriteria::PredictedPointToProjectionTarget;
 			Simplifier.GeometricErrorTolerance = UseTargetTolerance;
-			Simplifier.SimplifyToTriangleCount(8);
+			{
+				TRACE_CPUPROFILER_EVENT_SCOPE(ApproximateActorsImpl_Generate_Simplification_Pass2);
+				Simplifier.SimplifyToTriangleCount(8);
+			}
 		}
 		else
 		{
@@ -451,7 +478,15 @@ static TSharedPtr<FApproximationMeshData> GenerateApproximationMesh(
 	// compute normals
 	{
 		TRACE_CPUPROFILER_EVENT_SCOPE(ApproximateActorsImpl_Generate_Normals);
-		FMeshNormals::InitializeOverlayToPerVertexNormals(CurResultMesh->Attributes()->PrimaryNormals());
+		if (Options.bCalculateHardNormals)
+		{
+			FMeshNormals::InitializeOverlayTopologyFromOpeningAngle(CurResultMesh, CurResultMesh->Attributes()->PrimaryNormals(), Options.HardNormalsAngleDeg);
+			FMeshNormals::QuickRecomputeOverlayNormals(*CurResultMesh);
+		}
+		else
+		{
+			FMeshNormals::InitializeOverlayToPerVertexNormals(CurResultMesh->Attributes()->PrimaryNormals());
+		}
 	}
 
 	// exit here if we are just generating a merged collision mesh
@@ -590,6 +625,9 @@ void FApproximateActorsImpl::GenerateApproximationForActorSet(const TArray<AActo
 		WriteDebugMesh = &DebugMesh;
 	}
 
+	// build spatial evaluation cache
+	Scene.BuildSpatialEvaluationCache();
+
 	// if we are only generating collision mesh, we are going to exit after mesh generation
 	if (Options.BasePolicy == IGeometryProcessing_ApproximateActors::EApproximationPolicy::CollisionMesh)
 	{
@@ -725,6 +763,10 @@ UStaticMesh* FApproximateActorsImpl::EmitGeneratedMeshAsset(
 	FDynamicMesh3* DebugMesh)
 {
 	FStaticMeshAssetOptions MeshAssetOptions;
+
+	MeshAssetOptions.CollisionType = ECollisionTraceFlag::CTF_UseSimpleAsComplex;
+	MeshAssetOptions.bEnableRecomputeTangents = false;
+
 	MeshAssetOptions.NewAssetPath = Options.BasePackagePath;
 	MeshAssetOptions.SourceMeshes.DynamicMeshes.Add(FinalMesh);
 	if (Material)
@@ -740,6 +782,8 @@ UStaticMesh* FApproximateActorsImpl::EmitGeneratedMeshAsset(
 	if (DebugMesh != nullptr)
 	{
 		FStaticMeshAssetOptions DebugMeshAssetOptions;
+		DebugMeshAssetOptions.CollisionType = ECollisionTraceFlag::CTF_UseSimpleAsComplex;
+		DebugMeshAssetOptions.bEnableRecomputeTangents = false;
 		DebugMeshAssetOptions.NewAssetPath = Options.BasePackagePath + TEXT("_DEBUG");
 		DebugMeshAssetOptions.SourceMeshes.DynamicMeshes.Add(DebugMesh);
 		FStaticMeshResults DebugMeshAssetOutputs;
diff --git a/Engine/Plugins/Experimental/MeshModelingToolset/Source/ModelingComponents/Private/Scene/MeshSceneAdapter.cpp b/Engine/Plugins/Experimental/MeshModelingToolset/Source/ModelingComponents/Private/Scene/MeshSceneAdapter.cpp
--- a/Engine/Plugins/Experimental/MeshModelingToolset/Source/ModelingComponents/Private/Scene/MeshSceneAdapter.cpp
+++ b/Engine/Plugins/Experimental/MeshModelingToolset/Source/ModelingComponents/Private/Scene/MeshSceneAdapter.cpp
@@ -16,6 +16,7 @@
 #include "DynamicMeshEditor.h"
 #include "ConvexHull2.h"
 #include "Generators/PlanarPolygonMeshGenerator.h"
+#include "Spatial/SparseDynamicOctree3.h"
 
 #include "GameFramework/Actor.h"
 #include "Components/StaticMeshComponent.h"
@@ -23,10 +24,15 @@
 #include "Engine/StaticMesh.h"
 
 #include "Async/ParallelFor.h"
+#include "Async/Async.h"
 
 using namespace UE::Geometry;
 
 
+/** 
+ * Compute the bounds of the vertices of Mesh, under 3D transformation TransformFunc
+ * @return computed bounding box
+ */
 template<typename MeshType>
 FAxisAlignedBox3d GetTransformedVertexBounds(const MeshType& Mesh, TFunctionRef<FVector3d(const FVector3d&)> TransformFunc)
 {
@@ -42,60 +48,48 @@ FAxisAlignedBox3d GetTransformedVertexBounds(const MeshType& Mesh, TFunctionRef<
 	return Bounds;
 }
 
-
+/**
+ * Collect a subset of vertices of the mesh as "seed points" for algorithms like marching-cubes/etc.
+ * Generally every vertex does not need to be used. This function will return at most 5000 point
+ * @param TransformFunc transformation applied to points, eg local-to-world mapping
+ * @param AccumPointsInOut points are added here
+ * @param MaxPoints at most this many vertices will be returned
+ */
 template<typename MeshType>
-void CollectSeedPointsFromMeshVertices(const MeshType& Mesh, TFunctionRef<FVector3d(const FVector3d&)> TransformFunc, TArray<FVector3d>& AccumPoints)
+void CollectSeedPointsFromMeshVertices(
+	const MeshType& Mesh, 
+	TFunctionRef<FVector3d(const FVector3d&)> TransformFunc, TArray<FVector3d>& AccumPointsInOut,
+	int32 MaxPoints = 500)
 {
 	int32 NumVertices = Mesh.VertexCount();
 	int32 LogNumVertices = FMath::Max(1, (int32)FMathd::Ceil(FMathd::Log(NumVertices)));
 	int32 SeedPointCount = (int)(10 * LogNumVertices);
-	SeedPointCount = FMath::Min(SeedPointCount, 5000);
+	SeedPointCount = FMath::Min(SeedPointCount, MaxPoints);
 	int32 Skip = FMath::Max(NumVertices / SeedPointCount, 2);
 	for (int32 k = 0; k < NumVertices; k += Skip)
 	{
-		AccumPoints.Add(TransformFunc(Mesh.GetVertex(k)));
-	}
-}
-
-template<typename MeshType>
-FDynamicMesh3 FastConvertToDynamicMesh(const MeshType& Mesh)
-{
-	FDynamicMesh3 NewMesh;
-	int32 VertexCount = Mesh.VertexCount();
-	TArray<int32> VertexIDMap;
-	VertexIDMap.SetNum(VertexCount);
-	for (int32 vid = 0; vid < VertexCount; ++vid)
-	{
-		if (Mesh.IsVertex(vid))
-		{
-			int32 newvid = NewMesh.AppendVertex(Mesh.GetVertex(vid));
-			VertexIDMap[vid] = newvid;
-		}
-	}
-
-	int32 TriCount = Mesh.TriangleCount();
-	for (int32 tid = 0; tid < TriCount; ++tid)
-	{
-		if (Mesh.IsTriangle(tid))
-		{
-			FIndex3i TriVerts = Mesh.GetTriangle(tid);
-			NewMesh.AppendTriangle(VertexIDMap[TriVerts.A], VertexIDMap[TriVerts.B], VertexIDMap[TriVerts.C]);
-		}
+		AccumPointsInOut.Add(TransformFunc(Mesh.GetVertex(k)));
 	}
-
-	return NewMesh;
 }
 
 
+/**
+ * Try to check if the subset of Triangles of Mesh represent a "thin" region, ie basically a planar patch (open or closed).
+ * The Normal of the largest-area triangle is taken as the plane normal, and then the "thickness" is measured relative to this plain
+ * @param ThinTolerance identify as Thin if the thickness extents is within this size
+ * @param ThinPlaneOut thin plane normal will be returned via this frame
+ * @return true if submesh identified as thin
+ */
 template<typename MeshType>
-bool IsThinPlanarMesh(const MeshType& Mesh, double ThinTolerance, FFrame3d& ThinPlane)
+bool IsThinPlanarSubMesh(const MeshType& Mesh, const TArray<int32>& Triangles, double ThinTolerance, FFrame3d& ThinPlaneOut)
 {
-	int32 TriCount = Mesh.TriangleCount();
+	int32 TriCount = Triangles.Num();
 	double MaxArea = 0;
 	FVector3d MaxAreaNormal;
 	FVector3d MaxAreaPoint;
-	for (int32 tid = 0; tid < TriCount; ++tid)
+	for (int32 i = 0; i < TriCount; ++i)
 	{
+		int32 tid = Triangles[i];
 		if (Mesh.IsTriangle(tid))
 		{
 			FVector3d A, B, C;
@@ -110,16 +104,21 @@ bool IsThinPlanarMesh(const MeshType& Mesh, double ThinTolerance, FFrame3d& Thin
 			}
 		}
 	}
-	ThinPlane = FFrame3d(MaxAreaPoint, MaxAreaNormal);
+	ThinPlaneOut = FFrame3d(MaxAreaPoint, MaxAreaNormal);
 	FAxisAlignedBox3d PlaneExtents = FAxisAlignedBox3d::Empty();
 	int32 VertexCount = Mesh.VertexCount();
-	for (int32 vid = 0; vid < VertexCount; ++vid)
+	for (int32 i = 0; i < TriCount; ++i)
 	{
-		if (Mesh.IsVertex(vid))
+		int32 tid = Triangles[i];
+		if (Mesh.IsTriangle(tid))
 		{
-			FVector3d Pos = Mesh.GetVertex(vid);
-			Pos = ThinPlane.ToFramePoint(Pos);
-			PlaneExtents.Contain(Pos);
+			FVector3d TriVerts[3];
+			Mesh.GetTriVertices(tid, TriVerts[0], TriVerts[1], TriVerts[2]);
+			for (int32 j = 0; j < 3; ++j)
+			{
+				TriVerts[j] = ThinPlaneOut.ToFramePoint(TriVerts[j]);
+				PlaneExtents.Contain(TriVerts[j]);
+			}
 		}
 	}
 	if (PlaneExtents.Depth() > ThinTolerance)
@@ -127,11 +126,15 @@ bool IsThinPlanarMesh(const MeshType& Mesh, double ThinTolerance, FFrame3d& Thin
 		return false;
 	}
 	FVector3d Center = PlaneExtents.Center();
-	ThinPlane.Origin += Center.X*ThinPlane.X() + Center.Y*ThinPlane.Y() + Center.Z*ThinPlane.Z();
+	ThinPlaneOut.Origin += Center.X*ThinPlaneOut.X() + Center.Y*ThinPlaneOut.Y() + Center.Z*ThinPlaneOut.Z();
 	return true;
 }
 
 
+
+/**
+ * @return false if any of Triangles in Mesh have open boundary edges
+ */
 static bool IsClosedRegion(const FDynamicMesh3& Mesh, const TArray<int32>& Triangles)
 {
 	for (int32 tid : Triangles)
@@ -148,159 +151,34 @@ static bool IsClosedRegion(const FDynamicMesh3& Mesh, const TArray<int32>& Trian
 }
 
 
-template<typename Func>
-void ApplyToTriangleVertexIndices(const FDynamicMesh3& Mesh, const TArray<int32>& Triangles, Func VertexFunc)
-{
-	for (int32 tid : Triangles)
-	{
-		FIndex3i TriVerts = Mesh.GetTriangle(tid);
-		VertexFunc(TriVerts.A);
-		VertexFunc(TriVerts.B);
-		VertexFunc(TriVerts.C);
-	}
-}
-
-template<typename Func>
-void ApplyToTriangleVertices(const FDynamicMesh3& Mesh, const TArray<int32>& Triangles, Func VertexFunc)
-{
-	FVector3d A, B, C;
-	for (int32 tid : Triangles)
-	{
-		Mesh.GetTriVertices(tid, A, B, C);
-		VertexFunc(A);
-		VertexFunc(B);
-		VertexFunc(C);
-	}
-}
-
-static void ThickenThinPlanarMesh(FDynamicMesh3& Mesh, double TargetThickness, const FFrame3d& ThinPlane)
-{
-	FVector3d PlaneZ = ThinPlane.Z();
-
-	TArray<bool> VertFlags;
-	VertFlags.Init(false, Mesh.VertexCount());
-
-	FMeshConnectedComponents Components(&Mesh);
-	Components.FindConnectedTriangles();
-	int32 NumComponents = Components.Num();
-	for (const FMeshConnectedComponents::FComponent& Component : Components)
-	{
-		const TArray<int32>& Triangles = Component.Indices;
-		bool bIsClosed = IsClosedRegion(Mesh, Triangles);
-
-		if (bIsClosed)
-		{
-			FInterval1d ThicknessRange = FInterval1d::Empty();
-			ApplyToTriangleVertices(Mesh, Triangles, [&](const FVector3d& V) {
-				ThicknessRange.Contain((V - ThinPlane.Origin).Dot(PlaneZ));
-			});
-			double CurThickness = ThicknessRange.Length();
-
-			double MoveDelta = (TargetThickness - CurThickness) * 0.5;
-
-			int32 NumTriangles = Triangles.Num();
-			TArray<FVector3d> TriNormals;
-			TriNormals.SetNum(NumTriangles);
-			for (int32 k = 0; k < NumTriangles; ++k)
-			{
-				TriNormals[k] = Mesh.GetTriNormal(Triangles[k]);
-			}
-
-			//for (int32 k = 0; k < NumTriangles; ++k)
-			//{
-			//	int32 tid = Triangles[k];
-			//	FIndex3i TriVerts = Mesh.GetTriangle(tid);
-			//	FVector3d MoveDir = TriNormals[k].Dot(PlaneZ) > 0 ? PlaneZ : -PlaneZ;
-			//	for (int32 j = 0; j < 3; ++j)
-			//	{
-			//		int32 vid = TriVerts[j];
-			//		if (VertFlags[vid] == false)
-			//		{
-			//			Mesh.SetVertex(vid, Mesh.GetVertex(vid) + MoveDir * MoveDelta);
-			//			VertFlags[vid] = true;
-			//		}
-			//	}
-			//}
-
-		}
-		else   // bIsClosed == false
-		{
-			FOffsetMeshRegion Offset(&Mesh);
-			Offset.Triangles = Triangles;
-			//Offset.bUseFaceNormals = true;
-			Offset.DefaultOffsetDistance = -TargetThickness;
-			Offset.bIsPositiveOffset = false;
-			Offset.Apply();
-		}
-	}
-}
-
-
-
-
-
-static bool CheckIfStaticMeshUsageRequiresDynamicMeshDecomposition(
-	UStaticMesh* StaticMesh, 
-	int32 LODIndex,
-	const FMeshSceneAdapterBuildOptions& BuildOptions,
-	const TArray<FActorChildMesh*>& ParentMeshes,
-	int32 NonUniformScaleCount )
-{
-	FMeshDescription* SourceMesh;
-	check(StaticMesh);
-#if WITH_EDITOR
-	SourceMesh = StaticMesh->GetMeshDescription(LODIndex);
-#else
-	checkf(false, TEXT("Not currently supported - to build at Runtime it is necessary to read from the StaticMesh RenderBuffers"));
-	return false;
-#endif
-
-	if (SourceMesh)
-	{
-		// does not need this - copypasta
-		FMeshDescriptionTriangleMeshAdapter Adapter(SourceMesh);
-
-		FVector3d BuildScale;
-#if WITH_EDITOR
-		// respect BuildScale build setting
-		const FMeshBuildSettings& LODBuildSettings = StaticMesh->GetSourceModel(LODIndex).BuildSettings;
-		Adapter.SetBuildScale((FVector3d)LODBuildSettings.BuildScale3D, false);
-#endif
-		if (BuildOptions.bThickenThinMeshes)
-		{
-			FFrame3d ThinMeshPlane;
-			bool IsThinMesh = IsThinPlanarMesh<FMeshDescriptionTriangleMeshAdapter>(Adapter, BuildOptions.DesiredMinThickness, ThinMeshPlane);
-			if (IsThinMesh)
-			{
-				if (Adapter.TriangleCount() > 100000 && ParentMeshes.Num() > 1)
-				{
-					return false;		// skip thickening of huge meshes that are frequently used
-				}
-
-				return true;
-			}
-		}
-	}
-	return false;
-}
-
-
-
 
 
 class FDynamicMeshSpatialWrapper : public IMeshSpatialWrapper
 {
 public:
 	FDynamicMesh3 Mesh;
+
+	// if true, Mesh is in world space (whatever that means)
 	bool bHasBakedTransform = false;
+	// if true, use unsigned distance to determine inside/outside instead of winding number
+	bool bUseDistanceShellForWinding = false;
+	// unsigned distance isovalue that defines inside
+	double WindingShellThickness = 0.0;
+
 	TUniquePtr<TMeshAABBTree3<FDynamicMesh3>> AABBTree;
 	TUniquePtr<TFastWindingTree<FDynamicMesh3>> FWNTree;
 
 	virtual bool Build(const FMeshSceneAdapterBuildOptions& BuildOptions) override
 	{
 		ensure(Mesh.TriangleCount() > 0);
-		AABBTree = MakeUnique<TMeshAABBTree3<FDynamicMesh3>>(&Mesh, true);
-		FWNTree = MakeUnique<TFastWindingTree<FDynamicMesh3>>(AABBTree.Get(), true);
+		if (BuildOptions.bBuildSpatialDataStructures)
+		{
+			AABBTree = MakeUnique<TMeshAABBTree3<FDynamicMesh3>>(&Mesh, true);
+			if (bUseDistanceShellForWinding == false)
+			{
+				FWNTree = MakeUnique<TFastWindingTree<FDynamicMesh3>>(AABBTree.Get(), true);
+			}
+		}
 		return true;
 	}
 
@@ -311,9 +189,14 @@ public:
 
 	virtual FAxisAlignedBox3d GetWorldBounds(TFunctionRef<FVector3d(const FVector3d&)> LocalToWorldFunc) override
 	{
-		return bHasBakedTransform ?
+		FAxisAlignedBox3d Bounds = bHasBakedTransform ?
 			GetTransformedVertexBounds<FDynamicMesh3>(Mesh, [&](const FVector3d& P) {return P;}) :
 			GetTransformedVertexBounds<FDynamicMesh3>(Mesh, LocalToWorldFunc);
+		if (bUseDistanceShellForWinding)
+		{
+			Bounds.Expand(WindingShellThickness);
+		}
+		return Bounds;
 	}
 
 	virtual void CollectSeedPoints(TArray<FVector3d>& WorldPoints, TFunctionRef<FVector3d(const FVector3d&)> LocalToWorldFunc) override
@@ -325,9 +208,36 @@ public:
 
 	virtual double FastWindingNumber(const FVector3d& P, const FTransformSequence3d& LocalToWorldTransform) override
 	{
-		return bHasBakedTransform ?
-			FWNTree->FastWindingNumber(P) :
-			FWNTree->FastWindingNumber(LocalToWorldTransform.InverseTransformPosition(P));
+		if (bUseDistanceShellForWinding)
+		{
+			if (bHasBakedTransform)
+			{
+				double NearestDistSqr;
+				int32 NearTriID = AABBTree->FindNearestTriangle(P, NearestDistSqr, IMeshSpatial::FQueryOptions(WindingShellThickness));
+				if (NearTriID != IndexConstants::InvalidID)
+				{
+					// Do we even need to do this? won't we return InvalidID if we don't find point within distance?
+					// (also technically we can early-out as soon as we find any point, not the nearest point - might be worth a custom query)
+					FDistPoint3Triangle3d Query = TMeshQueries<FDynamicMesh3>::TriangleDistance(Mesh, NearTriID, P);
+					if (Query.GetSquared() < WindingShellThickness * WindingShellThickness)
+					{
+						return 1.0;
+					}
+				}
+			}
+			else
+			{
+				// todo
+				check(false);
+			}
+			return 0.0;
+		}
+		else
+		{
+			return bHasBakedTransform ?
+				FWNTree->FastWindingNumber(P) :
+				FWNTree->FastWindingNumber(LocalToWorldTransform.InverseTransformPosition(P));
+		}
 	}
 
 	virtual void ProcessVerticesInWorld(TFunctionRef<void(const FVector3d&)> ProcessFunc, const FTransformSequence3d& LocalToWorldTransform) override
@@ -371,122 +281,6 @@ public:
 
 
 
-
-
-
-class FStaticMeshDynamicSpatialWrapper : public IMeshSpatialWrapper
-{
-public:
-	UStaticMesh* StaticMesh = nullptr;
-	int32 LODIndex = 0;
-	
-	FMeshDescription* SourceMesh = nullptr;
-
-	TUniquePtr<FDynamicMeshSpatialWrapper> DynamicWrapper;
-
-	virtual bool Build(const FMeshSceneAdapterBuildOptions& BuildOptions) override
-	{
-		return ExtendedBuild(BuildOptions, FTransformSequence3d(), false);
-	}
-
-	virtual bool ExtendedBuild(const FMeshSceneAdapterBuildOptions& BuildOptions, const FTransformSequence3d& TransformSeq, bool bBakeTransform)
-	{
-		check(StaticMesh);
-#if WITH_EDITOR
-		SourceMesh = StaticMesh->GetMeshDescription(LODIndex);
-#else
-		checkf(false, TEXT("Not currently supported - to build at Runtime it is necessary to read from the StaticMesh RenderBuffers"));
-		SourceMesh = nullptr;
-#endif
-		if (SourceMesh)
-		{
-			// does not need this - copypasta
-			TUniquePtr<FMeshDescriptionTriangleMeshAdapter> Adapter = MakeUnique<FMeshDescriptionTriangleMeshAdapter>(SourceMesh);
-
-			FVector3d BuildScale;
-#if WITH_EDITOR
-			// respect BuildScale build setting
-			const FMeshBuildSettings& LODBuildSettings = StaticMesh->GetSourceModel(LODIndex).BuildSettings;
-			BuildScale = (FVector3d)LODBuildSettings.BuildScale3D;
-			Adapter->SetBuildScale(BuildScale, false);
-#endif
-			FDynamicMesh3 ConvertedMesh = FastConvertToDynamicMesh<FMeshDescriptionTriangleMeshAdapter>(*Adapter);
-
-			if (bBakeTransform)
-			{
-				MeshTransforms::ApplyTransform(ConvertedMesh,
-					[&](const FVector3d& P) { return TransformSeq.TransformPosition(P); },
-					[&](const FVector3f& N) { return (FVector3f)TransformSeq.TransformNormal((FVector3d)N); });
-			}
-
-			if (BuildOptions.bThickenThinMeshes)
-			{
-				FFrame3d ThinMeshPlane;
-				bool IsThinMesh = IsThinPlanarMesh<FMeshDescriptionTriangleMeshAdapter>(*Adapter, BuildOptions.DesiredMinThickness, ThinMeshPlane);
-				if (IsThinMesh)
-				{
-					ThickenThinPlanarMesh(ConvertedMesh, BuildOptions.DesiredMinThickness, ThinMeshPlane);
-				}
-			}
-
-			DynamicWrapper = MakeUnique<FDynamicMeshSpatialWrapper>();
-			DynamicWrapper->Mesh = MoveTemp(ConvertedMesh);
-			DynamicWrapper->bHasBakedTransform = bBakeTransform;
-			DynamicWrapper->Build(BuildOptions);
-
-			return true;
-		}
-
-		SourceMesh = nullptr;
-		return false;
-	}
-
-	virtual int32 GetTriangleCount() override
-	{
-		return (SourceMesh) ? DynamicWrapper->GetTriangleCount() : 0;
-	}
-
-	virtual FAxisAlignedBox3d GetWorldBounds(TFunctionRef<FVector3d(const FVector3d&)> LocalToWorldFunc) override
-	{
-		return (SourceMesh) ? DynamicWrapper->GetWorldBounds(LocalToWorldFunc) : FAxisAlignedBox3d::Empty();
-	}
-
-	virtual void CollectSeedPoints(TArray<FVector3d>& WorldPoints, TFunctionRef<FVector3d(const FVector3d&)> LocalToWorldFunc) override
-	{
-		if (SourceMesh)
-		{
-			DynamicWrapper->CollectSeedPoints(WorldPoints, LocalToWorldFunc);
-		}
-	}
-
-	virtual double FastWindingNumber(const FVector3d& P, const FTransformSequence3d& LocalToWorldTransform) override
-	{
-		return (SourceMesh) ? DynamicWrapper->FastWindingNumber(P, LocalToWorldTransform) : 0.0;
-	}
-
-	virtual void ProcessVerticesInWorld(TFunctionRef<void(const FVector3d&)> ProcessFunc, const FTransformSequence3d& LocalToWorldTransform) override
-	{
-		if (SourceMesh)
-		{
-			DynamicWrapper->ProcessVerticesInWorld(ProcessFunc, LocalToWorldTransform);
-		}
-	}
-
-
-	virtual void AppendMesh(FDynamicMesh3& AppendTo, const FTransformSequence3d& TransformSeq) override
-	{
-		if (SourceMesh)
-		{
-			DynamicWrapper->AppendMesh(AppendTo, TransformSeq);
-		}
-	}
-
-};
-
-
-
-
-
 class FStaticMeshSpatialWrapper : public IMeshSpatialWrapper
 {
 public:
@@ -519,8 +313,11 @@ public:
 			BuildScale = (FVector3d)LODBuildSettings.BuildScale3D;
 			Adapter->SetBuildScale(BuildScale, false);
 #endif
-			AABBTree = MakeUnique<TMeshAABBTree3<FMeshDescriptionTriangleMeshAdapter>>(Adapter.Get(), true);
-			FWNTree = MakeUnique<TFastWindingTree<FMeshDescriptionTriangleMeshAdapter>>(AABBTree.Get(), true);
+			if (BuildOptions.bBuildSpatialDataStructures)
+			{
+				AABBTree = MakeUnique<TMeshAABBTree3<FMeshDescriptionTriangleMeshAdapter>>(Adapter.Get(), true);
+				FWNTree = MakeUnique<TFastWindingTree<FMeshDescriptionTriangleMeshAdapter>>(AABBTree.Get(), true);
+			}
 			return true;
 		}
 
@@ -728,6 +525,69 @@ void FMeshSceneAdapter::AddActors(const TArray<AActor*>& ActorsSetIn)
 
 
 void FMeshSceneAdapter::Build(const FMeshSceneAdapterBuildOptions& BuildOptions)
+{
+	if (BuildOptions.bThickenThinMeshes)
+	{
+		Build_FullDecompose(BuildOptions);
+	}
+	else
+	{
+		TArray<FSpatialWrapperInfo*> ToBuild;
+		for (TPair<void*, TSharedPtr<FSpatialWrapperInfo>> Pair : SpatialAdapters)
+		{
+			ToBuild.Add(Pair.Value.Get());
+		}
+
+		FCriticalSection ListsLock;
+
+		std::atomic<int32> DecomposedSourceMeshCount;
+		DecomposedSourceMeshCount = 0;
+		std::atomic<int32> DecomposedMeshesCount;
+		DecomposedMeshesCount = 0;
+		int32 AddedTrisCount = 0;
+
+		// parallel build of all the spatial data structures
+		ParallelFor(ToBuild.Num(), [&](int32 i)
+		{
+			FSpatialWrapperInfo* WrapperInfo = ToBuild[i];
+			TUniquePtr<IMeshSpatialWrapper>& Wrapper = WrapperInfo->SpatialWrapper;
+			bool bOK = Wrapper->Build(BuildOptions);
+			ensure(bOK);	// assumption is that the wrapper will handle failure gracefully
+		});
+
+		if (BuildOptions.bPrintDebugMessages)
+		{
+			UE_LOG(LogTemp, Warning, TEXT("[FMeshSceneAdapter] decomposed %d source meshes into %d unique meshes containing %d triangles"), DecomposedSourceMeshCount.load(), DecomposedMeshesCount.load(), AddedTrisCount)
+		}
+
+	}
+
+
+	// update bounding boxes
+	ParallelFor(SceneActors.Num(), [&](int32 i)
+	{
+		UpdateActorBounds(*SceneActors[i]);
+	});
+}
+
+
+void FMeshSceneAdapter::UpdateActorBounds(FActorAdapter& Actor)
+{
+	FAxisAlignedBox3d ActorBounds = FAxisAlignedBox3d::Empty();
+	for (const TUniquePtr<FActorChildMesh>& ChildMesh : Actor.ChildMeshes)
+	{
+		if (ChildMesh->MeshSpatial != nullptr)
+		{
+			FAxisAlignedBox3d ChildBounds = ChildMesh->MeshSpatial->GetWorldBounds(
+				[&](const FVector3d& P) { return ChildMesh->WorldTransform.TransformPosition(P); });
+			ActorBounds.Contain(ChildBounds);
+		}
+	}
+	Actor.WorldBounds = ActorBounds;
+}
+
+
+void FMeshSceneAdapter::Build_FullDecompose(const FMeshSceneAdapterBuildOptions& BuildOptions)
 {
 	TArray<FSpatialWrapperInfo*> ToBuild;
 	for (TPair<void*, TSharedPtr<FSpatialWrapperInfo>> Pair : SpatialAdapters)
@@ -743,56 +603,202 @@ void FMeshSceneAdapter::Build(const FMeshSceneAdapterBuildOptions& BuildOptions)
 	DecomposedMeshesCount = 0;
 	int32 AddedTrisCount = 0;
 
+	// use this to disable build of unnecessary AABBTree/etc
+	FMeshSceneAdapterBuildOptions TempBuildOptions = BuildOptions;
+	TempBuildOptions.bBuildSpatialDataStructures = false;
+
 	// parallel build of all the spatial data structures
 	ParallelFor(ToBuild.Num(), [&](int32 i)
 	{
+		TRACE_CPUPROFILER_EVENT_SCOPE(MeshScene_Build_ProcessMesh);
+
 		FSpatialWrapperInfo* WrapperInfo = ToBuild[i];
 
-		if (WrapperInfo->SourceContainer.MeshType == ESceneMeshType::StaticMeshAsset)
+		// convert this mesh to a dynamicmesh for processing
+		FDynamicMesh3 LocalMesh;
+		WrapperInfo->SpatialWrapper->Build(TempBuildOptions);
 		{
-			UStaticMesh* StaticMesh = WrapperInfo->SourceContainer.GetStaticMesh();
-			int32 LODIndex = 0;
-			bool bDecompose = CheckIfStaticMeshUsageRequiresDynamicMeshDecomposition(
-				StaticMesh, LODIndex,
-				BuildOptions,
-				WrapperInfo->ParentMeshes, WrapperInfo->NonUniformScaleCount);
+			TRACE_CPUPROFILER_EVENT_SCOPE(MeshScene_Build_ProcessMesh_1Copy);
+			WrapperInfo->SpatialWrapper->AppendMesh(LocalMesh, FTransformSequence3d());
+		}
+
+		// should we try to weld here??
 
-			if (bDecompose)
+		FMeshConnectedComponents Components(&LocalMesh);
+		{
+			TRACE_CPUPROFILER_EVENT_SCOPE(MeshScene_Build_ProcessMesh_2Components);
+			Components.FindConnectedTriangles();
+		}
+		int32 NumComponents = Components.Num();
+		TArray<bool> IsClosed;
+		IsClosed.Init(false, NumComponents);
+		std::atomic<int32> NumNonClosed = 0;
+		{
+			TRACE_CPUPROFILER_EVENT_SCOPE(MeshScene_Build_ProcessMesh_3Closed);
+			ParallelFor(NumComponents, [&](int32 ci)
 			{
-				DecomposedSourceMeshCount++;
+				const FMeshConnectedComponents::FComponent& Component = Components[ci];
+				const TArray<int32>& Triangles = Component.Indices;
+				IsClosed[ci] = IsClosedRegion(LocalMesh, Triangles);
+				if (IsClosed[ci] == false)
+				{
+					NumNonClosed++;
+				}
+			});
+		}
 
-				TArray<FActorChildMesh*> MeshesToDecompose = WrapperInfo->ParentMeshes;
-				for (FActorChildMesh* MeshInstance : MeshesToDecompose)
+		// if we have no open meshes, we can just use the SpatialWrapper we already have, but we have to
+		// rebuild it because we did not do a full build above
+		if (NumNonClosed == 0)
+		{
+			WrapperInfo->SpatialWrapper->Build(BuildOptions);
+			return;
+		}
+
+		// NOTE: because of thin-planar checks below, we might not actually decompose this mesh!!
+		DecomposedSourceMeshCount++;
+
+		// list of per-instance transforms
+		TArray<FActorChildMesh*> MeshesToDecompose = WrapperInfo->ParentMeshes;
+		TArray<FTransformSequence3d> ParentTransforms;
+		for (FActorChildMesh* MeshInstance : MeshesToDecompose)
+		{
+			ParentTransforms.Add(MeshInstance->WorldTransform);
+		}
+
+		bool bMakeClosedAssembly = (NumNonClosed != NumComponents);
+		FDynamicMesh3 LocalSpaceParts;
+		FDynamicMeshEditor LocalSpaceAccumulator(&LocalSpaceParts);
+		FDynamicMesh3 WorldSpaceParts;
+		FDynamicMeshEditor WorldSpaceAccumulator(&WorldSpaceParts);
+
+		int32 LocalCount = 0, DecomposedCount = 0;
+
+		// Accumulate closed parts in a local-space mesh that can stay instanced, and open parts
+		{
+			TRACE_CPUPROFILER_EVENT_SCOPE(MeshScene_Build_ProcessMesh_4Accumulate);
+			for (int32 ci = 0; ci < NumComponents; ++ci)
+			{
+				const TArray<int32>& Triangles = Components[ci].Indices;
+
+				bool bKeepLocal = true;
+				FFrame3d ThinMeshPlane;
+				if (IsClosed[ci] == false)
 				{
-					TSharedPtr<FSpatialWrapperInfo> NewWrapperInfo = MakeShared<FSpatialWrapperInfo>();
-					NewWrapperInfo->SourceContainer = MeshInstance->MeshContainer;
-
-					TUniquePtr<FStaticMeshDynamicSpatialWrapper> DynamicMeshWrapper = MakeUnique<FStaticMeshDynamicSpatialWrapper>();
-					DynamicMeshWrapper->SourceContainer = NewWrapperInfo->SourceContainer;
-					DynamicMeshWrapper->StaticMesh = StaticMesh;
-					DynamicMeshWrapper->ExtendedBuild(BuildOptions, MeshInstance->WorldTransform, true);
-					NewWrapperInfo->SpatialWrapper = MoveTemp(DynamicMeshWrapper);
-					MeshInstance->MeshSpatial = NewWrapperInfo->SpatialWrapper.Get();
-
-					ListsLock.Lock();
-					void* OldMeshKey = MeshInstance->MeshContainer.GetMeshKey();
-					SpatialAdapters.Remove(OldMeshKey);			// remove existing
-					void* NewMeshKey = (void*)MeshInstance;		// need to key on instance, not unique mesh
-					SpatialAdapters.Add(NewMeshKey, NewWrapperInfo);
-					AddedTrisCount += NewWrapperInfo->SpatialWrapper->GetTriangleCount();
-					ListsLock.Unlock();
-
-					DecomposedMeshesCount++;
+					bool IsThinMesh = IsThinPlanarSubMesh<FDynamicMesh3>(LocalMesh, Triangles, BuildOptions.DesiredMinThickness, ThinMeshPlane);
+					if (IsThinMesh && (Triangles.Num() < 10000 || WrapperInfo->ParentMeshes.Num() == 1))
+					{
+						bKeepLocal = false;
+					}
+				}
+
+				FMeshIndexMappings Mappings;
+				FDynamicMeshEditResult EditResult;
+				if (bKeepLocal)
+				{
+					LocalSpaceAccumulator.AppendTriangles(&LocalMesh, Triangles, Mappings, EditResult, false);
+					LocalCount++;
+				}
+				else
+				{
+					DecomposedCount++;
+					for (const FTransformSequence3d& Transform : ParentTransforms)
+					{
+						DecomposedMeshesCount++;
+						Mappings.Reset();
+						EditResult.Reset();
+						WorldSpaceAccumulator.AppendTriangles(&LocalMesh, Triangles, Mappings, EditResult, false);
+						for (int32 vid : EditResult.NewVertices)
+						{
+							FVector3d LocalPos = WorldSpaceParts.GetVertex(vid);
+							WorldSpaceParts.SetVertex(vid, Transform.TransformPosition(LocalPos));
+						}
+					}
 				}
-				return;
 			}
 		}
 
-		TUniquePtr<IMeshSpatialWrapper>& Wrapper = WrapperInfo->SpatialWrapper;
-		bool bOK = Wrapper->Build(BuildOptions);
-		ensure(bOK);	// assumption is that the wrapper will handle failure gracefully
+		// make new mesh containers
+
+		// we re-use the existing wrapper for the closed meshes, if there are any, otherwise
+		// we need to disable the parent meshes that refer to it
+		TFuture<void> PendingClosedAssemblyBuild;
+		if (bMakeClosedAssembly)
+		{
+			TUniquePtr<FDynamicMeshSpatialWrapper> LocalSpaceMeshWrapper = MakeUnique<FDynamicMeshSpatialWrapper>();
+			
+			LocalSpaceMeshWrapper->Mesh = MoveTemp(LocalSpaceParts);
+			FDynamicMeshSpatialWrapper* BuildWrapper = LocalSpaceMeshWrapper.Get();
+			PendingClosedAssemblyBuild = Async(EAsyncExecution::ThreadPool, [&BuildOptions, BuildWrapper]()
+			{
+				TRACE_CPUPROFILER_EVENT_SCOPE(MeshScene_Build_ProcessMesh_6SolidBuild);
+				BuildWrapper->Build(BuildOptions);
+			});
+
+			WrapperInfo->SpatialWrapper = MoveTemp(LocalSpaceMeshWrapper);
+
+			for (FActorChildMesh* MeshInstance : MeshesToDecompose)
+			{
+				MeshInstance->MeshSpatial = WrapperInfo->SpatialWrapper.Get();
+			}
+		}
+		else
+		{
+			// have to null out spatials for the child meshes so that they are ignored during computation
+			for (FActorChildMesh* MeshInstance : MeshesToDecompose)
+			{
+				MeshInstance->MeshSpatial = nullptr;
+			}
+		}
+
+		// currently happens if we fail to find open parts to expand/etc
+		if (WorldSpaceParts.TriangleCount() == 0)
+		{
+			return;
+		}
+
+		// world-space meshes are added to new fake adapter
+		TUniquePtr<FActorAdapter> Adapter = MakeUnique<FActorAdapter>();
+		Adapter->SourceActor = nullptr;
+
+		TUniquePtr<FActorChildMesh> ChildMesh = MakeUnique<FActorChildMesh>();
+		ChildMesh->SourceComponent = nullptr;
+		ChildMesh->ComponentType = EActorMeshComponentType::InternallyGeneratedComponent;
+		ChildMesh->ComponentIndex = 0;
+		ChildMesh->bIsNonUniformScaled = false;
+
+
+		TUniquePtr<FDynamicMeshSpatialWrapper> NewWorldSpaceSpatial = MakeUnique<FDynamicMeshSpatialWrapper>();
+		NewWorldSpaceSpatial->Mesh = MoveTemp(WorldSpaceParts);
+		NewWorldSpaceSpatial->bHasBakedTransform = true;
+		{
+			TRACE_CPUPROFILER_EVENT_SCOPE(MeshScene_Build_ProcessMesh_8WorldSpaceBuild);
+			NewWorldSpaceSpatial->Build(BuildOptions);
+		}
+		NewWorldSpaceSpatial->bUseDistanceShellForWinding = true;
+		NewWorldSpaceSpatial->WindingShellThickness = 0.5 * BuildOptions.DesiredMinThickness;
+
+		TSharedPtr<FSpatialWrapperInfo> NewWrapperInfo = MakeShared<FSpatialWrapperInfo>();
+		NewWrapperInfo->SpatialWrapper = MoveTemp(NewWorldSpaceSpatial);
+		NewWrapperInfo->ParentMeshes.Add(ChildMesh.Get());
+		ChildMesh->MeshSpatial = NewWrapperInfo->SpatialWrapper.Get();
+
+		void* UseKey = ChildMesh.Get();
+		Adapter->ChildMeshes.Add(MoveTemp(ChildMesh));
+
+		ListsLock.Lock();
+		SpatialAdapters.Add(UseKey, NewWrapperInfo);
+		SceneActors.Add(MoveTemp(Adapter));
+		AddedTrisCount += NewWrapperInfo->SpatialWrapper->GetTriangleCount();
+		ListsLock.Unlock();
+
+		// make sure this finishes
+		PendingClosedAssemblyBuild.Wait();
 	});
 
+	// currently true with the methods used above?
+	bSceneIsAllSolids = true;
+
 	if (BuildOptions.bPrintDebugMessages)
 	{
 		UE_LOG(LogTemp, Warning, TEXT("[FMeshSceneAdapter] decomposed %d source meshes into %d unique meshes containing %d triangles"), DecomposedSourceMeshCount.load(), DecomposedMeshesCount.load(), AddedTrisCount)
@@ -802,7 +808,6 @@ void FMeshSceneAdapter::Build(const FMeshSceneAdapterBuildOptions& BuildOptions)
 
 
 
-
 void FMeshSceneAdapter::GetGeometryStatistics(FStatistics& StatsOut)
 {
 	StatsOut.UniqueMeshCount = 0;
@@ -831,8 +836,15 @@ void FMeshSceneAdapter::GetGeometryStatistics(FStatistics& StatsOut)
 
 FAxisAlignedBox3d FMeshSceneAdapter::GetBoundingBox()
 {
-	FAxisAlignedBox3d SceneBounds = FAxisAlignedBox3d::Empty();
+	TRACE_CPUPROFILER_EVENT_SCOPE(MeshScene_GetBoundingBox);
 
+	if (bHaveSpatialEvaluationCache)
+	{
+		return CachedWorldBounds;
+	}
+
+	// this could be done in parallel...
+	FAxisAlignedBox3d SceneBounds = FAxisAlignedBox3d::Empty();
 	for (const TUniquePtr<FActorAdapter>& Actor : SceneActors)
 	{
 		for (const TUniquePtr<FActorChildMesh>& ChildMesh : Actor->ChildMeshes)
@@ -851,6 +863,8 @@ FAxisAlignedBox3d FMeshSceneAdapter::GetBoundingBox()
 
 void FMeshSceneAdapter::CollectMeshSeedPoints(TArray<FVector3d>& Points)
 {
+	TRACE_CPUPROFILER_EVENT_SCOPE(MeshScene_CollectMeshSeedPoints);
+
 	for (const TUniquePtr<FActorAdapter>& Actor : SceneActors)
 	{
 		for (const TUniquePtr<FActorChildMesh>& ChildMesh : Actor->ChildMeshes)
@@ -864,29 +878,105 @@ void FMeshSceneAdapter::CollectMeshSeedPoints(TArray<FVector3d>& Points)
 	}
 }
 
-double FMeshSceneAdapter::FastWindingNumber(const FVector3d& P)
+double FMeshSceneAdapter::FastWindingNumber(const FVector3d& P, bool bFastEarlyOutIfPossible)
 {
+	check(bHaveSpatialEvaluationCache);		// must call BuildSpatialEvaluationCache() to build Octree
+
 	double SumWinding = 0.0;
-	for (const TUniquePtr<FActorAdapter>& Actor : SceneActors)
+
+	// if all objects in scene are solids, then all winding queries will return integers so if any value
+	// is > 0, we are "inside"
+	if (bSceneIsAllSolids)
 	{
-		// TODO: cannot skip here because fast winding number for open region will extend...
-		if ( true /**Actor->WorldBounds.Contains(P)*/ )
+		if (bFastEarlyOutIfPossible)
 		{
-			for (const TUniquePtr<FActorChildMesh>& ChildMesh : Actor->ChildMeshes)
+			bool bFinished = Octree->ContainmentQueryCancellable(P, [&](int32 k)
 			{
-				if (ChildMesh->MeshSpatial != nullptr)
-				{
-					double MeshWinding = ChildMesh->MeshSpatial->FastWindingNumber(P, ChildMesh->WorldTransform);
-					SumWinding += MeshWinding;
-				}
-			}
+				double WindingNumber = SortedSpatials[k].Spatial->FastWindingNumber(P, SortedSpatials[k].ChildMesh->WorldTransform);
+				SumWinding += WindingNumber;
+				return (FMath::Abs(WindingNumber) < 0.99);		// if we see an "inside" winding number we can just exit
+			});
+		}
+		else
+		{
+			Octree->ContainmentQuery(P, [&](int32 k)
+			{
+				double WindingNumber = SortedSpatials[k].Spatial->FastWindingNumber(P, SortedSpatials[k].ChildMesh->WorldTransform);
+				SumWinding += WindingNumber;
+			});
 		}
 	}
+	else
+	{
+		for (const FSpatialCacheInfo& SpatialInfo : SortedSpatials)
+		{
+			double WindingNumber = SpatialInfo.Spatial->FastWindingNumber(P, SpatialInfo.ChildMesh->WorldTransform);
+			SumWinding += WindingNumber;
+		}
+	}
+
 	return SumWinding;
 }
 
 
 
+void FMeshSceneAdapter::BuildSpatialEvaluationCache()
+{
+	TRACE_CPUPROFILER_EVENT_SCOPE(MeshScene_BuildSpatialEvaluationCache);
+
+	// build list of unique meshes we need to evaluate for spatial queries
+	SortedSpatials.Reset();
+	for (const TUniquePtr<FActorAdapter>& Actor : SceneActors)
+	{
+		for (const TUniquePtr<FActorChildMesh>& ChildMesh : Actor->ChildMeshes)
+		{
+			if (ChildMesh->MeshSpatial != nullptr)
+			{
+				FSpatialCacheInfo Cache;
+				Cache.Actor = Actor.Get();
+				Cache.ChildMesh = ChildMesh.Get();
+				Cache.Spatial = ChildMesh->MeshSpatial;
+				Cache.Bounds = ChildMesh->MeshSpatial->GetWorldBounds(
+					[&](const FVector3d& P) { return ChildMesh->WorldTransform.TransformPosition(P); });
+				SortedSpatials.Add(Cache);
+
+
+			}
+		}
+	}
+
+	// sort the list (not really necessary but might improve cache coherency during linear queries)
+	SortedSpatials.Sort([&](const FSpatialCacheInfo& A, const FSpatialCacheInfo& B)
+	{
+		return A.Spatial < B.Spatial;
+	});
+
+	int32 NumSpatials = SortedSpatials.Num();
+	ParallelFor(SortedSpatials.Num(), [&](int32 k)
+	{
+		SortedSpatials[k].Bounds = SortedSpatials[k].Spatial->GetWorldBounds(
+			[&](const FVector3d& P) { return SortedSpatials[k].ChildMesh->WorldTransform.TransformPosition(P); });
+	});
+
+	CachedWorldBounds = FAxisAlignedBox3d::Empty();
+	for (const FSpatialCacheInfo& Cache : SortedSpatials )
+	{
+		CachedWorldBounds.Contain(Cache.Bounds);
+	}
+
+	// build an octree of the mesh objects
+	Octree = MakeShared<FSparseDynamicOctree3>();
+	Octree->RootDimension = CachedWorldBounds.MaxDim() / 4.0;
+	Octree->SetMaxTreeDepth(5);
+	for (int32 k = 0; k < NumSpatials; ++k)
+	{
+		Octree->InsertObject(k, SortedSpatials[k].Bounds);
+	}
+
+	bHaveSpatialEvaluationCache = true;
+}
+
+
 
 void FMeshSceneAdapter::GetAccumulatedMesh(FDynamicMesh3& AccumMesh)
 {
@@ -954,12 +1044,22 @@ void FMeshSceneAdapter::GenerateBaseClosingMesh(double BaseHeight, double Extrud
 	});
 
 	FConvexHull2d FinalHullSolver;
-	FinalHullSolver.Solve(WorldHullPoints);
+	bool bOK = FinalHullSolver.Solve(WorldHullPoints);
+	if (bOK == false)
+	{
+		UE_LOG(LogTemp, Warning, TEXT("[FMeshSceneAdapter::GenerateBaseClosingMesh] failed to solve convex hull"));
+		return;
+	}
 	FPolygon2d ConvexHullPoly;
 	for (int32 idx : FinalHullSolver.GetPolygonIndices())
 	{
 		ConvexHullPoly.AppendVertex(WorldHullPoints[idx]);
 	}
+	if (ConvexHullPoly.VertexCount() < 2)
+	{
+		UE_LOG(LogTemp, Warning, TEXT("[FMeshSceneAdapter::GenerateBaseClosingMesh] convex hull is degenerate"));
+		return;
+	}
 
 	FPlanarPolygonMeshGenerator MeshGen;
 	MeshGen.Polygon = ConvexHullPoly;
@@ -969,6 +1069,7 @@ void FMeshSceneAdapter::GenerateBaseClosingMesh(double BaseHeight, double Extrud
 	if (ExtrudeHeight == 0)
 	{
 		BasePolygonMesh.ReverseOrientation();		// flip so it points down
+		bSceneIsAllSolids = false;					// if scene was solids, it's not anymore
 	}
 	else
 	{
@@ -987,8 +1088,8 @@ void FMeshSceneAdapter::GenerateBaseClosingMesh(double BaseHeight, double Extrud
 	// append a fake actor/mesh
 	//
 
-	TUniquePtr<FActorAdapter> Adapter = MakeUnique<FActorAdapter>();
-	Adapter->SourceActor = nullptr;
+	TUniquePtr<FActorAdapter> ActorAdapter = MakeUnique<FActorAdapter>();
+	ActorAdapter->SourceActor = nullptr;
 
 	TUniquePtr<FActorChildMesh> ChildMesh = MakeUnique<FActorChildMesh>();
 	ChildMesh->SourceComponent = nullptr;
@@ -1011,8 +1112,9 @@ void FMeshSceneAdapter::GenerateBaseClosingMesh(double BaseHeight, double Extrud
 	NewWrapperInfo->SpatialWrapper = MoveTemp(DynamicMeshWrapper);
 	NewWrapperInfo->ParentMeshes.Add(ChildMesh.Get());
 	ChildMesh->MeshSpatial = NewWrapperInfo->SpatialWrapper.Get();
-	Adapter->ChildMeshes.Add(MoveTemp(ChildMesh));
-	SceneActors.Add(MoveTemp(Adapter));
+	ActorAdapter->ChildMeshes.Add(MoveTemp(ChildMesh));
+	UpdateActorBounds(*ActorAdapter);
+	SceneActors.Add(MoveTemp(ActorAdapter));
 
 }
 
diff --git a/Engine/Plugins/Experimental/MeshModelingToolset/Source/ModelingComponents/Public/Scene/MeshSceneAdapter.h b/Engine/Plugins/Experimental/MeshModelingToolset/Source/ModelingComponents/Public/Scene/MeshSceneAdapter.h
--- a/Engine/Plugins/Experimental/MeshModelingToolset/Source/ModelingComponents/Public/Scene/MeshSceneAdapter.h
+++ b/Engine/Plugins/Experimental/MeshModelingToolset/Source/ModelingComponents/Public/Scene/MeshSceneAdapter.h
@@ -18,6 +18,7 @@ namespace UE
 namespace Geometry
 {
 
+class FSparseDynamicOctree3;
 
 /**
  * ESceneMeshType is used to indicate which type of Mesh a FMeshTypeContainer contains.
@@ -69,6 +70,9 @@ struct FMeshSceneAdapterBuildOptions
 	bool bThickenThinMeshes = false;
 	/** Thickness used for bThickenThinMeshes processing */
 	double DesiredMinThickness = 0.1;
+
+
+	bool bBuildSpatialDataStructures = true;
 };
 
 
@@ -231,8 +235,11 @@ public:
 	/** @return a set of points on the surface of the meshes, can be used to initialize the MarchingCubes mesher */
 	virtual void CollectMeshSeedPoints(TArray<FVector3d>& PointsOut);
 
+	/** Precompute data structures that accelerate spatial evaluation queries. Precondition for calling FastWindingNumber() */
+	virtual void BuildSpatialEvaluationCache();
+
 	/** @return FastWindingNumber computed across all mesh Actors/Components */
-	virtual double FastWindingNumber(const FVector3d& P);
+	virtual double FastWindingNumber(const FVector3d& P, bool bFastEarlyOutIfPossible = true);
 
 	/** Append all instance triangles to a single mesh. May be very large. */
 	virtual void GetAccumulatedMesh(FDynamicMesh3& AccumMesh);
@@ -253,8 +260,24 @@ protected:
 	// Unique set of spatial data structure query interfaces, one for each Mesh object, which is identified by void* pointer
 	TMap<void*, TSharedPtr<FSpatialWrapperInfo>> SpatialAdapters;
 
-	bool bEnableClipPlane = false;
-	FFrame3d ClipPlane;
+	bool bSceneIsAllSolids = false;
+
+	void UpdateActorBounds(FActorAdapter& Actor);
+
+	void Build_FullDecompose(const FMeshSceneAdapterBuildOptions& BuildOptions);
+
+
+	struct FSpatialCacheInfo
+	{
+		FActorAdapter* Actor;
+		FActorChildMesh* ChildMesh;
+		IMeshSpatialWrapper* Spatial;
+		FAxisAlignedBox3d Bounds;
+	};
+	TArray<FSpatialCacheInfo> SortedSpatials;
+	TSharedPtr<FSparseDynamicOctree3> Octree;
+	FAxisAlignedBox3d CachedWorldBounds;
+	bool bHaveSpatialEvaluationCache = false;
 };
 
 
diff --git a/Engine/Plugins/Experimental/MeshModelingToolset/Source/ModelingComponentsEditorOnly/Private/AssetUtils/CreateStaticMeshUtil.cpp b/Engine/Plugins/Experimental/MeshModelingToolset/Source/ModelingComponentsEditorOnly/Private/AssetUtils/CreateStaticMeshUtil.cpp
--- a/Engine/Plugins/Experimental/MeshModelingToolset/Source/ModelingComponentsEditorOnly/Private/AssetUtils/CreateStaticMeshUtil.cpp
+++ b/Engine/Plugins/Experimental/MeshModelingToolset/Source/ModelingComponentsEditorOnly/Private/AssetUtils/CreateStaticMeshUtil.cpp
@@ -61,11 +61,7 @@ UE::AssetUtils::ECreateStaticMeshResult UE::AssetUtils::CreateStaticMeshAsset(
 	if (Options.bCreatePhysicsBody)
 	{
 		NewStaticMesh->CreateBodySetup();
-		if (Options.bEnableComplexAsSimple)
-		{
-			// enable complex as simple collision to use mesh directly
-			NewStaticMesh->GetBodySetup()->CollisionTraceFlag = ECollisionTraceFlag::CTF_UseComplexAsSimple;
-		}
+		NewStaticMesh->GetBodySetup()->CollisionTraceFlag = Options.CollisionType;
 	}
 
 	// add a material slot. Must always have one material slot.
diff --git a/Engine/Plugins/Experimental/MeshModelingToolset/Source/ModelingComponentsEditorOnly/Public/AssetUtils/CreateStaticMeshUtil.h b/Engine/Plugins/Experimental/MeshModelingToolset/Source/ModelingComponentsEditorOnly/Public/AssetUtils/CreateStaticMeshUtil.h
--- a/Engine/Plugins/Experimental/MeshModelingToolset/Source/ModelingComponentsEditorOnly/Public/AssetUtils/CreateStaticMeshUtil.h
+++ b/Engine/Plugins/Experimental/MeshModelingToolset/Source/ModelingComponentsEditorOnly/Public/AssetUtils/CreateStaticMeshUtil.h
@@ -4,6 +4,7 @@
 
 #include "CoreMinimal.h"
 #include "DynamicMesh3.h"
+#include "BodySetupEnums.h"
 
 class UPackage;
 class UStaticMesh;
@@ -63,8 +64,8 @@ namespace AssetUtils
 
 		// Controls whether the UBodySetup on the Asset will be created (generally should be true)
 		bool bCreatePhysicsBody = true;
-		// If true, Asset will be configured to UseComplexAsSimple mode
-		bool bEnableComplexAsSimple = true;
+		// set asset collision type
+		ECollisionTraceFlag CollisionType = ECollisionTraceFlag::CTF_UseDefault;
 		// TODO: option for simple collision geo
 
 		// Optional list of materials to initialize the Asset with. If defined, size must be the same as NumMaterialSlots
diff --git a/Engine/Source/Developer/GeometryProcessingInterfaces/Public/GeometryProcessingInterfaces/ApproximateActors.h b/Engine/Source/Developer/GeometryProcessingInterfaces/Public/GeometryProcessingInterfaces/ApproximateActors.h
--- a/Engine/Source/Developer/GeometryProcessingInterfaces/Public/GeometryProcessingInterfaces/ApproximateActors.h
+++ b/Engine/Source/Developer/GeometryProcessingInterfaces/Public/GeometryProcessingInterfaces/ApproximateActors.h
@@ -97,6 +97,9 @@ public:
 		int32 FixedTriangleCount = 5000;
 		double SimplificationTargetMetric = 0.1;		// interpretation varies depending on MeshSimplificationPolicy
 
+		bool bCalculateHardNormals = true;
+		double HardNormalsAngleDeg = 60.0;
+
 		//
 		// Material approximation settings
 		//
