PolyEdit: add Duplicate operation
#rnx

diff --git a/Engine/Plugins/Experimental/MeshModelingToolset/Source/MeshModelingTools/Private/EditMeshPolygonsTool.cpp b/Engine/Plugins/Experimental/MeshModelingToolset/Source/MeshModelingTools/Private/EditMeshPolygonsTool.cpp
--- a/Engine/Plugins/Experimental/MeshModelingToolset/Source/MeshModelingTools/Private/EditMeshPolygonsTool.cpp
+++ b/Engine/Plugins/Experimental/MeshModelingToolset/Source/MeshModelingTools/Private/EditMeshPolygonsTool.cpp
@@ -1012,6 +1012,10 @@ void UEditMeshPolygonsTool::OnTick(float DeltaTime)
 		{
 			ApplyDisconnect();
 		}
+		else if (PendingAction == EEditMeshPolygonsToolActions::Duplicate)
+		{
+			ApplyDuplicate();
+		}
 		else if (PendingAction == EEditMeshPolygonsToolActions::PokeSingleFace)
 		{
 			ApplyPokeSingleFace();
@@ -1896,6 +1900,42 @@ void UEditMeshPolygonsTool::ApplyDisconnect()
 
 
 
+void UEditMeshPolygonsTool::ApplyDuplicate()
+{
+	if (BeginMeshFaceEditChange() == false)
+	{
+		GetToolManager()->DisplayMessage(LOCTEXT("OnDuplicateFailed", "Cannot Duplicate Current Selection"), EToolMessageLevel::UserWarning);
+		return;
+	}
+
+	FDynamicMesh3* Mesh = DynamicMeshComponent->GetMesh();
+	FDynamicMeshChangeTracker ChangeTracker(Mesh);
+	ChangeTracker.BeginChange();
+	FGroupTopologySelection ActiveSelection = SelectionMechanic->GetActiveSelection();
+	TArray<int32> AllTriangles;
+	for (int32 GroupID : ActiveSelection.SelectedGroupIDs)
+	{
+		AllTriangles.Append(Topology->GetGroupTriangles(GroupID));
+	}
+	FDynamicMeshEditor Editor(Mesh);
+	FMeshIndexMappings Mappings;
+	FDynamicMeshEditResult EditResult;
+	Editor.DuplicateTriangles(AllTriangles, Mappings, EditResult);
+
+	FGroupTopologySelection NewSelection;
+	for ( int NewGroupID : EditResult.NewGroups )
+	{ 
+		NewSelection.SelectedGroupIDs.Add(NewGroupID);
+	}
+
+	TUniquePtr<FMeshChange> MeshChange = MakeUnique<FMeshChange>(ChangeTracker.EndChange());
+	CompleteMeshEditChange(LOCTEXT("PolyMeshDisconnectChange", "Disconnect"), MoveTemp(MeshChange), NewSelection);
+	CurrentToolMode = ECurrentToolMode::TransformSelection;
+}
+
+
+
+
 void UEditMeshPolygonsTool::ApplyCollapseEdge()
 {
 	// AAAHHH cannot do because of overlays!
diff --git a/Engine/Plugins/Experimental/MeshModelingToolset/Source/MeshModelingTools/Public/EditMeshPolygonsTool.h b/Engine/Plugins/Experimental/MeshModelingToolset/Source/MeshModelingTools/Public/EditMeshPolygonsTool.h
--- a/Engine/Plugins/Experimental/MeshModelingToolset/Source/MeshModelingTools/Public/EditMeshPolygonsTool.h
+++ b/Engine/Plugins/Experimental/MeshModelingToolset/Source/MeshModelingTools/Public/EditMeshPolygonsTool.h
@@ -104,6 +104,7 @@ enum class EEditMeshPolygonsToolActions
 	Retriangulate,
 	Decompose,
 	Disconnect,
+	Duplicate,
 
 	CollapseEdge,
 	WeldEdges,
@@ -230,6 +231,11 @@ public:
 	/** Separate the selected faces at their borders */
 	UFUNCTION(CallInEditor, Category = FaceEdits, meta = (DisplayName = "Disconnect", DisplayPriority = 11))
 	void Disconnect() { PostAction(EEditMeshPolygonsToolActions::Disconnect); }
+
+	/** Duplicate the selected faces at their borders */
+	UFUNCTION(CallInEditor, Category = FaceEdits, meta = (DisplayName = "Duplicate", DisplayPriority = 12))
+	void Duplicate() { PostAction(EEditMeshPolygonsToolActions::Duplicate); }
+
 };
 
 
@@ -275,8 +281,12 @@ public:
 	UFUNCTION(CallInEditor, Category = TriangleEdits, meta = (DisplayName = "Disconnect", DisplayPriority = 11))
 	void Disconnect() { PostAction(EEditMeshPolygonsToolActions::Disconnect); }
 
+	/** Duplicate the selected faces */
+	UFUNCTION(CallInEditor, Category = TriangleEdits, meta = (DisplayName = "Duplicate", DisplayPriority = 12))
+	void Duplicate() { PostAction(EEditMeshPolygonsToolActions::Duplicate); }
+
 	/** Poke each face at its center point */
-	UFUNCTION(CallInEditor, Category = TriangleEdits, meta = (DisplayName = "Poke", DisplayPriority = 12))
+	UFUNCTION(CallInEditor, Category = TriangleEdits, meta = (DisplayName = "Poke", DisplayPriority = 13))
 	void Poke() { PostAction(EEditMeshPolygonsToolActions::PokeSingleFace); }
 };
 
@@ -674,6 +684,7 @@ protected:
 	void ApplyRetriangulate();
 	void ApplyDecompose();
 	void ApplyDisconnect();
+	void ApplyDuplicate();
 	void ApplyPokeSingleFace();
 
 	void ApplyCollapseEdge();
