Added "ApproximateActors.RenderDocCapture" cvar
* Makes a scoped render doc capture that encompass the material baking processing of the new "Approximate Actors" merge method.

diff --git a/Engine/Plugins/Experimental/MeshModelingToolset/Source/GeometryProcessingAdapters/GeometryProcessingAdapters.Build.cs b/Engine/Plugins/Experimental/MeshModelingToolset/Source/GeometryProcessingAdapters/GeometryProcessingAdapters.Build.cs
--- a/Engine/Plugins/Experimental/MeshModelingToolset/Source/GeometryProcessingAdapters/GeometryProcessingAdapters.Build.cs
+++ b/Engine/Plugins/Experimental/MeshModelingToolset/Source/GeometryProcessingAdapters/GeometryProcessingAdapters.Build.cs
@@ -51,6 +51,7 @@ public class GeometryProcessingAdapters : ModuleRules
 			{
 				"CoreUObject",
 				"Engine",
+				"RenderCore",
 
 				"MeshBuilder",
 				"MeshUtilitiesCommon",  
diff --git a/Engine/Plugins/Experimental/MeshModelingToolset/Source/GeometryProcessingAdapters/Private/ApproximateActorsImpl.cpp b/Engine/Plugins/Experimental/MeshModelingToolset/Source/GeometryProcessingAdapters/Private/ApproximateActorsImpl.cpp
--- a/Engine/Plugins/Experimental/MeshModelingToolset/Source/GeometryProcessingAdapters/Private/ApproximateActorsImpl.cpp
+++ b/Engine/Plugins/Experimental/MeshModelingToolset/Source/GeometryProcessingAdapters/Private/ApproximateActorsImpl.cpp
@@ -34,12 +34,21 @@
 
 #include "Async/Async.h"
 #include "Misc/ScopedSlowTask.h"
+#include "RenderCaptureInterface.h"
 
 using namespace UE::Geometry;
 using namespace UE::AssetUtils;
 
 #define LOCTEXT_NAMESPACE "ApproximateActorsImpl"
 
+static TAutoConsoleVariable<int32> CVarApproximateActorsRDOCCapture(
+	TEXT("ApproximateActors.RenderCapture"),
+	0,
+	TEXT("Determines whether or not to trigger a render capture.\n")
+	TEXT("0: Turned Off\n")
+	TEXT("1: Turned On"),
+	ECVF_Default);
+
 struct FGeneratedResultTextures
 {
 	UTexture2D* BaseColorMap;
@@ -573,6 +582,8 @@ void FApproximateActorsImpl::GenerateApproximationForActorSet(const TArray<AActo
 {
 	TRACE_CPUPROFILER_EVENT_SCOPE(ApproximateActorsImpl_Generate);
 
+	RenderCaptureInterface::FScopedCapture RenderCapture(CVarApproximateActorsRDOCCapture.GetValueOnAnyThread() == 1, TEXT("ApproximateActors"));
+
 	//
 	// Future Optimizations
 	// 	   - can do most of the mesh processing at the same time as capturing the photo set (if that matters)
