always disable normal recomputation when modeling tools commit a static mesh, as we don't in general want normals re-computed coming out of our modeling tools

#rnx

diff --git a/Engine/Plugins/Experimental/MeshModelingToolset/Source/MeshModelingToolsEditorOnly/Private/ToolTargets/StaticMeshComponentToolTarget.cpp b/Engine/Plugins/Experimental/MeshModelingToolset/Source/MeshModelingToolsEditorOnly/Private/ToolTargets/StaticMeshComponentToolTarget.cpp
--- a/Engine/Plugins/Experimental/MeshModelingToolset/Source/MeshModelingToolsEditorOnly/Private/ToolTargets/StaticMeshComponentToolTarget.cpp
+++ b/Engine/Plugins/Experimental/MeshModelingToolset/Source/MeshModelingToolsEditorOnly/Private/ToolTargets/StaticMeshComponentToolTarget.cpp
@@ -8,6 +8,7 @@
 #include "Engine/StaticMesh.h"
 #include "MeshDescriptionToDynamicMesh.h"
 #include "RenderingThread.h"
+#include "AssetUtils/MeshDescriptionUtil.h"
 
 static void DisplayCriticalWarningMessage(const FString& Message)
 {
@@ -225,6 +226,12 @@ void UStaticMeshComponentToolTarget::CommitMeshDescription(const FCommitter& Com
 	StaticMesh->SetFlags(RF_Transactional);
 
 	verify(StaticMesh->Modify());
+
+	// disable auto-generated normals StaticMesh build setting
+	UE::MeshDescription::FStaticMeshBuildSettingChange SettingsChange;
+	SettingsChange.AutoGeneratedNormals = UE::MeshDescription::EBuildSettingBoolChange::Disable;
+	UE::MeshDescription::ConfigureBuildSettings(StaticMesh, 0, SettingsChange);
+
 	if (EditingLOD == EStaticMeshEditingLOD::HiResSource)
 	{
 		verify(StaticMesh->ModifyHiResMeshDescription());
