ModelingTools: Fix selection after fill hole op in TriEd.

#rnx

diff --git a/Engine/Plugins/Experimental/MeshModelingToolset/Source/MeshModelingTools/Private/EditMeshPolygonsTool.cpp b/Engine/Plugins/Experimental/MeshModelingToolset/Source/MeshModelingTools/Private/EditMeshPolygonsTool.cpp
--- a/Engine/Plugins/Experimental/MeshModelingToolset/Source/MeshModelingTools/Private/EditMeshPolygonsTool.cpp
+++ b/Engine/Plugins/Experimental/MeshModelingToolset/Source/MeshModelingTools/Private/EditMeshPolygonsTool.cpp
@@ -2084,7 +2084,14 @@ void UEditMeshPolygonsTool::ApplyFillHole()
 				Filler.FillType = FSimpleHoleFiller::EFillType::PolygonEarClipping;
 				int32 NewGroupID = Mesh->AllocateTriangleGroup();
 				Filler.Fill(NewGroupID);
-				NewSelection.SelectedGroupIDs.Add(NewGroupID);
+				if (!bTriangleMode)
+				{
+					NewSelection.SelectedGroupIDs.Add(NewGroupID);
+				}
+				else
+				{
+					NewSelection.SelectedGroupIDs.Append(Filler.NewTriangles);
+				}
 
 				// Compute normals and UVs
 				if (Mesh->HasAttributes())
