ModelingComponents: support Packed Metallic/Roughness/Specular (MRS) capture in WorldRenderCapture using custom postprocess render pass. Also support in FSceneCapturePhotoSet. Add Linear Color variant to FTexture2DBuilder.
#rnx

diff --git a/Engine/Plugins/Experimental/MeshModelingToolset/Source/ModelingComponents/ModelingComponents.Build.cs b/Engine/Plugins/Experimental/MeshModelingToolset/Source/ModelingComponents/ModelingComponents.Build.cs
--- a/Engine/Plugins/Experimental/MeshModelingToolset/Source/ModelingComponents/ModelingComponents.Build.cs
+++ b/Engine/Plugins/Experimental/MeshModelingToolset/Source/ModelingComponents/ModelingComponents.Build.cs
@@ -47,6 +47,7 @@ public class ModelingComponents : ModuleRules
 				"InputCore",
                 "RenderCore",
                 "RHI",
+				"ImageWriteQueue"
 
 				//"MeshUtilities",    // temp for saving mesh asset
 				//"UnrealEd"
diff --git a/Engine/Plugins/Experimental/MeshModelingToolset/Source/ModelingComponents/Private/Scene/SceneCapturePhotoSet.cpp b/Engine/Plugins/Experimental/MeshModelingToolset/Source/ModelingComponents/Private/Scene/SceneCapturePhotoSet.cpp
--- a/Engine/Plugins/Experimental/MeshModelingToolset/Source/ModelingComponents/Private/Scene/SceneCapturePhotoSet.cpp
+++ b/Engine/Plugins/Experimental/MeshModelingToolset/Source/ModelingComponents/Private/Scene/SceneCapturePhotoSet.cpp
@@ -35,6 +35,9 @@ void FSceneCapturePhotoSet::SetCaptureTypeEnabled(ERenderCaptureType CaptureType
 		case ERenderCaptureType::Emissive:
 			bEnableEmissive = bEnabled;
 			break;
+		case ERenderCaptureType::CombinedMRS:
+			bEnablePackedMRS = bEnabled;
+			break;
 		default:
 			check(false);
 	}
@@ -165,6 +168,10 @@ void FSceneCapturePhotoSet::AddExteriorCaptures(
 		{
 			CaptureImageTypeFunc_1f(ERenderCaptureType::Metallic, MetallicPhotoSet);
 		}
+		if (bEnablePackedMRS)
+		{
+			CaptureImageTypeFunc_3f(ERenderCaptureType::CombinedMRS, PackedMRSPhotoSet);
+		}
 		if (bEnableWorldNormal)
 		{
 			CaptureImageTypeFunc_3f(ERenderCaptureType::WorldNormal, WorldNormalPhotoSet);
@@ -242,6 +249,8 @@ FVector4f FSceneCapturePhotoSet::FSceneSample::GetValue4f(ERenderCaptureType Cap
 		return FVector4f(Metallic, Metallic, Metallic, 1.0f);
 	case ERenderCaptureType::Specular:
 		return FVector4f(Specular, Specular, Specular, 1.0f);
+	case ERenderCaptureType::CombinedMRS:
+		return FVector4f(Metallic, Roughness, Specular, 1.0f);
 	case ERenderCaptureType::Emissive:
 		return FVector4f(Emissive.X, Emissive.Y, Emissive.Z, 1.0f);
 	default:
@@ -286,6 +295,17 @@ bool FSceneCapturePhotoSet::ComputeSample(
 			MetallicPhotoSet.ComputeSample(Position, Normal, VisibilityFunction, DefaultsInResultsOut.Metallic);
 		DefaultsInResultsOut.HaveValues.bMetallic = true;
 	}
+	if (SampleChannels.bCombinedMRS)
+	{
+		FVector3f MRSValue(DefaultsInResultsOut.Metallic, DefaultsInResultsOut.Roughness, DefaultsInResultsOut.Specular);
+		MRSValue = PackedMRSPhotoSet.ComputeSample(Position, Normal, VisibilityFunction, MRSValue);
+		DefaultsInResultsOut.Metallic = MRSValue.X;
+		DefaultsInResultsOut.Roughness = MRSValue.Y;
+		DefaultsInResultsOut.Specular = MRSValue.Z;
+		DefaultsInResultsOut.HaveValues.bMetallic = true;
+		DefaultsInResultsOut.HaveValues.bRoughness = true;
+		DefaultsInResultsOut.HaveValues.bSpecular = true;
+	}
 	if (SampleChannels.bEmissive)
 	{
 		DefaultsInResultsOut.Emissive =
diff --git a/Engine/Plugins/Experimental/MeshModelingToolset/Source/ModelingComponents/Private/Scene/WorldRenderCapture.cpp b/Engine/Plugins/Experimental/MeshModelingToolset/Source/ModelingComponents/Private/Scene/WorldRenderCapture.cpp
--- a/Engine/Plugins/Experimental/MeshModelingToolset/Source/ModelingComponents/Private/Scene/WorldRenderCapture.cpp
+++ b/Engine/Plugins/Experimental/MeshModelingToolset/Source/ModelingComponents/Private/Scene/WorldRenderCapture.cpp
@@ -7,6 +7,8 @@
 #include "LegacyScreenPercentageDriver.h"
 #include "EngineModule.h"
 #include "Rendering/Texture2DResource.h"
+#include "ImagePixelData.h"
+#include "ImageWriteStream.h"
 
 #include "PreviewScene.h"
 #include "Components/DirectionalLightComponent.h"
@@ -65,6 +67,9 @@ void FRenderCaptureTypeFlags::SetEnabled(ERenderCaptureType CaptureType, bool bE
 	case ERenderCaptureType::Emissive:
 		bEmissive = bEnabled;
 		break;
+	case ERenderCaptureType::CombinedMRS:
+		bCombinedMRS = bEnabled;
+		break;
 	default:
 		check(false);
 	}
@@ -304,6 +309,11 @@ bool FWorldRenderCapture::CaptureEmissiveFromPosition(
 	ViewFamily.EngineShowFlags.SetDiffuse(false);
 	ViewFamily.EngineShowFlags.SetSpecular(false);
 
+	ViewFamily.EngineShowFlags.SetScreenSpaceReflections(false);
+	ViewFamily.EngineShowFlags.SetLumenReflections(false);
+	ViewFamily.EngineShowFlags.SetRefraction(false);
+	ViewFamily.EngineShowFlags.SetReflectionEnvironment(false);
+
 	ViewFamily.EngineShowFlags.SetDynamicShadows(false);
 	ViewFamily.EngineShowFlags.SetCapsuleShadows(false);
 	ViewFamily.EngineShowFlags.SetContactShadows(false);
@@ -377,6 +387,238 @@ bool FWorldRenderCapture::CaptureEmissiveFromPosition(
 
 
 
+namespace UE 
+{
+namespace Internal
+{
+
+static bool ReadPixelDataToImage(TUniquePtr<FImagePixelData>& PixelData, FImageAdapter& ResultImageOut, bool bLinear)
+{
+	int64 SizeInBytes;
+	const void* OutRawData = nullptr;
+	PixelData->GetRawData(OutRawData, SizeInBytes);
+	int32 Width = PixelData->GetSize().X;
+	int32 Height = PixelData->GetSize().Y;
+
+	ResultImageOut.SetDimensions(FImageDimensions(Width, Height));
+
+	switch (PixelData->GetType())
+	{
+		case EImagePixelType::Color:
+		{
+			const uint8* SourceBuffer = (uint8*)OutRawData;
+			for (int32 yi = 0; yi < Height; ++yi)
+			{
+				for (int32 xi = 0; xi < Width; ++xi)
+				{
+					const uint8* BufferPixel = &SourceBuffer[(yi * Width + xi) * 4];
+					FColor Color(BufferPixel[2], BufferPixel[1], BufferPixel[0], BufferPixel[3]);		// BGRA
+					FLinearColor PixelColorf = (bLinear) ? Color.ReinterpretAsLinear() : FLinearColor(Color);
+					ResultImageOut.SetPixel(FVector2i(xi, yi), FVector4f(PixelColorf));
+				}
+			}
+			return true;
+		}
+
+		case EImagePixelType::Float16:
+		{
+			ensure(bLinear);		// data must be linear
+			const FFloat16Color* SourceBuffer = (FFloat16Color*)(OutRawData);
+			for (int32 yi = 0; yi < Height; ++yi)
+			{
+				for (int32 xi = 0; xi < Width; ++xi)
+				{
+					const FFloat16Color* BufferPixel = &SourceBuffer[yi * Width + xi];
+					FLinearColor PixelColorf = BufferPixel->GetFloats();
+					ResultImageOut.SetPixel(FVector2i(xi, yi), FVector4f(PixelColorf));
+				}
+			}
+			return true;
+		}
+
+		case EImagePixelType::Float32:
+		{
+			ensure(bLinear);		// data must be linear
+			const FLinearColor* SourceBuffer = (FLinearColor*)(OutRawData);
+			for (int32 yi = 0; yi < Height; ++yi)
+			{
+				for (int32 xi = 0; xi < Width; ++xi)
+				{
+					FLinearColor PixelColorf = SourceBuffer[yi * Width + xi];
+					ResultImageOut.SetPixel(FVector2i(xi, yi), FVector4f(PixelColorf));
+				}
+			}
+			return true;
+		}
+
+		default:
+		{
+			ensure(false);
+			return false;
+		}
+	}
+	return false;
+}
+
+} // end namespace Internal
+} // end namespace UE
+
+
+bool FWorldRenderCapture::CaptureMRSFromPosition(
+	const FFrame3d& ViewFrame,
+	double HorzFOVDegrees,
+	double NearPlaneDist,
+	FImageAdapter& ResultImageOut)
+{
+	// this post-process material renders an image with R=Metallic, G=Roughness, B=Specular, A=AmbientOcclusion
+	FString MRSPostProcessMaterialAssetPath = TEXT("/MeshModelingToolset/Materials/PostProcess_PackedMRSA.PostProcess_PackedMRSA");
+	TSoftObjectPtr<UMaterialInterface> PostProcessMaterialPtr = TSoftObjectPtr<UMaterialInterface>(FSoftObjectPath(MRSPostProcessMaterialAssetPath));
+	UMaterialInterface* PostProcessMaterial = PostProcessMaterialPtr.LoadSynchronous();
+	if (!ensure(PostProcessMaterial))
+	{
+		return false;
+	}
+
+	check(this->World);
+	FSceneInterface* Scene = this->World->Scene;
+
+	UTextureRenderTarget2D* RenderTargetTexture = GetRenderTexture(true);
+	if (!ensure(RenderTargetTexture))
+	{
+		return false;
+	}
+	FRenderTarget* RenderTargetResource = RenderTargetTexture->GameThread_GetRenderTargetResource();
+
+	int32 Width = Dimensions.GetWidth();
+	int32 Height = Dimensions.GetHeight();
+
+	FQuat ViewOrientation = (FQuat)ViewFrame.Rotation;
+	FMatrix ViewRotationMatrix = FInverseRotationMatrix(ViewOrientation.Rotator());
+	FVector ViewOrigin = (FVector)ViewFrame.Origin;
+
+	// convert to rendering coordinate system
+	ViewRotationMatrix = ViewRotationMatrix * FMatrix(
+		FPlane(0, 0, 1, 0),
+		FPlane(1, 0, 0, 0),
+		FPlane(0, 1, 0, 0),
+		FPlane(0, 0, 0, 1));
+
+	const float HalfFOVRadians = FMath::DegreesToRadians<float>(HorzFOVDegrees) * 0.5f;
+	static_assert((int32)ERHIZBuffer::IsInverted != 0, "Check NearPlane and Projection Matrix");
+	FMatrix ProjectionMatrix = FReversedZPerspectiveMatrix(HalfFOVRadians, 1.0f, 1.0f, (float)NearPlaneDist);
+
+	EViewModeIndex ViewModeIndex = EViewModeIndex::VMI_Unlit;		// VMI_Lit, VMI_LightingOnly, VMI_VisualizeBuffer
+
+	FEngineShowFlags ShowFlags = FEngineShowFlags(EShowFlagInitMode::ESFIM_Game);
+	ApplyViewMode(ViewModeIndex, true, ShowFlags);
+
+	// unclear if these flags need to be set before creating ViewFamily
+	ShowFlags.SetAntiAliasing(false);
+	ShowFlags.SetDepthOfField(false);
+	ShowFlags.SetMotionBlur(false);
+	ShowFlags.SetBloom(false);
+	ShowFlags.SetSceneColorFringe(false);
+
+	FSceneViewFamilyContext ViewFamily(FSceneViewFamily::ConstructionValues(
+		RenderTargetResource, Scene, ShowFlags)
+		.SetWorldTimes(0, 0, 0)
+		.SetRealtimeUpdate(false));
+
+	// unclear whether these show flags are really necessary, since we are using
+	// a custom postprocess pass and reading it's output buffer directly
+
+	ViewFamily.EngineShowFlags.DisableAdvancedFeatures();
+	ViewFamily.EngineShowFlags.MotionBlur = 0;
+	ViewFamily.EngineShowFlags.LOD = 0;
+
+	ViewFamily.EngineShowFlags.SetTonemapper(false);
+	ViewFamily.EngineShowFlags.SetColorGrading(false);
+	ViewFamily.EngineShowFlags.SetToneCurve(false);
+
+	ViewFamily.EngineShowFlags.SetPostProcessing(false);
+	ViewFamily.EngineShowFlags.SetFog(false);
+	ViewFamily.EngineShowFlags.SetGlobalIllumination(false);
+	ViewFamily.EngineShowFlags.SetEyeAdaptation(false);
+	ViewFamily.EngineShowFlags.SetDirectionalLights(false);
+	ViewFamily.EngineShowFlags.SetPointLights(false);
+	ViewFamily.EngineShowFlags.SetSpotLights(false);
+	ViewFamily.EngineShowFlags.SetRectLights(false);
+
+	ViewFamily.EngineShowFlags.SetDiffuse(false);
+	ViewFamily.EngineShowFlags.SetSpecular(false);
+
+	ViewFamily.EngineShowFlags.SetDynamicShadows(false);
+	ViewFamily.EngineShowFlags.SetCapsuleShadows(false);
+	ViewFamily.EngineShowFlags.SetContactShadows(false);
+
+	//ViewFamily.EngineShowFlags.SetScreenPercentage(false);
+	ViewFamily.SetScreenPercentageInterface(new FLegacyScreenPercentageDriver(ViewFamily, 1.f, false));
+
+	// This is called in various other places, unclear if we should be doing this too
+	//EngineShowFlagOverride(EShowFlagInitMode::ESFIM_Game, ViewModeIndex, ViewFamily.EngineShowFlags, true);
+
+	FSceneViewInitOptions ViewInitOptions;
+	ViewInitOptions.SetViewRectangle(FIntRect(0, 0, Width, Height));
+	ViewInitOptions.ViewFamily = &ViewFamily;
+	if (VisiblePrimitives.Num() > 0)
+	{
+		ViewInitOptions.ShowOnlyPrimitives = VisiblePrimitives;
+	}
+	ViewInitOptions.ViewOrigin = ViewOrigin;
+	ViewInitOptions.ViewRotationMatrix = ViewRotationMatrix;
+	ViewInitOptions.ProjectionMatrix = ProjectionMatrix;
+	ViewInitOptions.FOV = HorzFOVDegrees;
+
+	FSceneView* NewView = new FSceneView(ViewInitOptions);
+	ViewFamily.Views.Add(NewView);
+
+	NewView->StartFinalPostprocessSettings(ViewInitOptions.ViewOrigin);
+
+	//
+	// Add custom PostProcessMaterial. This will be configured to write to a FImagePixelPipe,
+	// which will be filled by the renderer and available after we force the render
+	//
+	TUniquePtr<FImagePixelData> PostProcessPassPixelData;
+
+	// clear any existing PostProcess materials and add our own
+	NewView->FinalPostProcessSettings.BufferVisualizationOverviewMaterials.Empty();
+	NewView->FinalPostProcessSettings.BufferVisualizationOverviewMaterials.Add(PostProcessMaterial);
+
+	// create new FImagePixelPipe and add output lambda that takes ownership of the data
+	NewView->FinalPostProcessSettings.BufferVisualizationPipes.Empty();
+	TSharedPtr<FImagePixelPipe, ESPMode::ThreadSafe> ImagePixelPipe = MakeShared<FImagePixelPipe, ESPMode::ThreadSafe>();
+	ImagePixelPipe->AddEndpoint([&PostProcessPassPixelData](TUniquePtr<FImagePixelData>&& Data) {
+		PostProcessPassPixelData = MoveTemp(Data);
+	});
+	NewView->FinalPostProcessSettings.BufferVisualizationPipes.Add(PostProcessMaterial->GetFName(), ImagePixelPipe);
+
+	// enable buffer visualization writing
+	NewView->FinalPostProcessSettings.bBufferVisualizationDumpRequired = true;
+	NewView->EndFinalPostprocessSettings(ViewInitOptions);
+
+	// can we fully disable auto-exposure?
+	NewView->FinalPostProcessSettings.AutoExposureMethod = EAutoExposureMethod::AEM_Manual;
+
+	// do we actually need for force SM5 here? 
+	FCanvas Canvas = FCanvas(RenderTargetResource, nullptr, this->World, ERHIFeatureLevel::SM5, FCanvas::CDM_DeferDrawing, 1.0f);
+	Canvas.Clear(FLinearColor::Transparent);
+	GetRendererModule().BeginRenderingViewFamily(&Canvas, &ViewFamily);
+
+	// wait for render
+	FlushRenderingCommands();
+
+	// read back image
+	if (PostProcessPassPixelData.IsValid())
+	{
+		UE::Internal::ReadPixelDataToImage(PostProcessPassPixelData, ResultImageOut, true);
+		return true;
+	}
+
+	return false;
+}
+
+
+
 
 /**
  * internal Utility function to render the given Scene to a render target and capture
@@ -463,6 +705,10 @@ bool FWorldRenderCapture::CaptureFromPosition(
 	{
 		return CaptureEmissiveFromPosition(ViewFrame, HorzFOVDegrees, NearPlaneDist, ResultImageOut);
 	}
+	else if (CaptureType == ERenderCaptureType::CombinedMRS)
+	{
+		return CaptureMRSFromPosition(ViewFrame, HorzFOVDegrees, NearPlaneDist, ResultImageOut);
+	}
 
 	// Roughness visualization is rendered with gamma correction (unclear why)
 	bool bLinear = (CaptureType != ERenderCaptureType::Roughness);
diff --git a/Engine/Plugins/Experimental/MeshModelingToolset/Source/ModelingComponents/Public/AssetUtils/Texture2DBuilder.h b/Engine/Plugins/Experimental/MeshModelingToolset/Source/ModelingComponents/Public/AssetUtils/Texture2DBuilder.h
--- a/Engine/Plugins/Experimental/MeshModelingToolset/Source/ModelingComponents/Public/AssetUtils/Texture2DBuilder.h
+++ b/Engine/Plugins/Experimental/MeshModelingToolset/Source/ModelingComponents/Public/AssetUtils/Texture2DBuilder.h
@@ -34,6 +34,7 @@ public:
 	enum class ETextureType
 	{
 		Color,
+		ColorLinear,
 		Roughness,
 		Metallic,
 		Specular,
@@ -88,12 +89,12 @@ public:
 			return false;
 		}
 
-		if (BuildType == ETextureType::Roughness || BuildType == ETextureType::Metallic || BuildType == ETextureType::Specular)
+		if (BuildType == ETextureType::ColorLinear || BuildType == ETextureType::Roughness || BuildType == ETextureType::Metallic || BuildType == ETextureType::Specular)
 		{
 			RawTexture2D->SRGB = false;
 			RawTexture2D->UpdateResource();
 		}
-		if (BuildType == ETextureType::NormalMap)
+		else if (BuildType == ETextureType::NormalMap)
 		{
 			RawTexture2D->CompressionSettings = TC_Normalmap;
 			RawTexture2D->SRGB = false;
@@ -442,6 +443,7 @@ public:
 		{
 		default:
 		case ETextureType::Color:
+		case ETextureType::ColorLinear:
 			return DefaultColor;
 		case ETextureType::Roughness:
 			return DefaultRoughness;
diff --git a/Engine/Plugins/Experimental/MeshModelingToolset/Source/ModelingComponents/Public/Scene/SceneCapturePhotoSet.h b/Engine/Plugins/Experimental/MeshModelingToolset/Source/ModelingComponents/Public/Scene/SceneCapturePhotoSet.h
--- a/Engine/Plugins/Experimental/MeshModelingToolset/Source/ModelingComponents/Public/Scene/SceneCapturePhotoSet.h
+++ b/Engine/Plugins/Experimental/MeshModelingToolset/Source/ModelingComponents/Public/Scene/SceneCapturePhotoSet.h
@@ -22,6 +22,7 @@ namespace Geometry
  *		Roughness
  *		Metallic
  *		Specular
+ *		PackedMRS   (Metallic / Roughness / Specular)
  *		Emissive
  *		WorldNormal
  * There are various efficiences possible by doing these captures as a group, rather
@@ -112,14 +113,24 @@ public:
 		FSceneSample& DefaultsInResultsOut) const;
 
 
+	const FSpatialPhotoSet3f& GetBaseColorPhotoSet() { return BaseColorPhotoSet; }
+	const FSpatialPhotoSet1f& GetRoughnessPhotoSet() { return RoughnessPhotoSet; }
+	const FSpatialPhotoSet1f& GetSpecularPhotoSet() { return SpecularPhotoSet; }
+	const FSpatialPhotoSet1f& GetMetallicPhotoSet() { return MetallicPhotoSet; }
+	const FSpatialPhotoSet3f& GetPackedMRSPhotoSet() { return PackedMRSPhotoSet; }
+	const FSpatialPhotoSet3f& GetWorldNormalPhotoSet() { return WorldNormalPhotoSet; }
+	const FSpatialPhotoSet3f& GetEmissivePhotoSet() { return EmissivePhotoSet; }
+
+
 protected:
 	UWorld* TargetWorld = nullptr;
 	TArray<AActor*> VisibleActors;
 
 	bool bEnableBaseColor = true;
-	bool bEnableRoughness = true;
-	bool bEnableSpecular = true;
-	bool bEnableMetallic = true;
+	bool bEnableRoughness = false;
+	bool bEnableSpecular = false;
+	bool bEnableMetallic = false;
+	bool bEnablePackedMRS = true;
 	bool bEnableWorldNormal = true;
 	bool bEnableEmissive = true;
 
@@ -127,6 +138,7 @@ protected:
 	FSpatialPhotoSet1f RoughnessPhotoSet;
 	FSpatialPhotoSet1f SpecularPhotoSet;
 	FSpatialPhotoSet1f MetallicPhotoSet;
+	FSpatialPhotoSet3f PackedMRSPhotoSet;
 	FSpatialPhotoSet3f WorldNormalPhotoSet;
 	FSpatialPhotoSet3f EmissivePhotoSet;
 };
diff --git a/Engine/Plugins/Experimental/MeshModelingToolset/Source/ModelingComponents/Public/Scene/WorldRenderCapture.h b/Engine/Plugins/Experimental/MeshModelingToolset/Source/ModelingComponents/Public/Scene/WorldRenderCapture.h
--- a/Engine/Plugins/Experimental/MeshModelingToolset/Source/ModelingComponents/Public/Scene/WorldRenderCapture.h
+++ b/Engine/Plugins/Experimental/MeshModelingToolset/Source/ModelingComponents/Public/Scene/WorldRenderCapture.h
@@ -16,14 +16,15 @@ namespace Geometry
 /**
  * ERenderCaptureType defines which type of render buffer should be captured.
  */
-enum class ERenderCaptureType : uint8
+enum class ERenderCaptureType
 {
 	BaseColor = 1,
 	Roughness = 2,
 	Metallic = 4,
 	Specular = 8,
 	Emissive = 16,
-	WorldNormal = 32
+	WorldNormal = 32,
+	CombinedMRS = 128
 };
 
 /**
@@ -37,6 +38,7 @@ struct MODELINGCOMPONENTS_API FRenderCaptureTypeFlags
 	bool bSpecular = false;
 	bool bEmissive = false;
 	bool bWorldNormal = false;
+	bool bCombinedMRS = false;
 
 	/** @return FRenderCaptureTypeFlags with all types enabled/true */
 	static FRenderCaptureTypeFlags All();
@@ -132,6 +134,13 @@ protected:
 		double HorzFOVDegrees,
 		double NearPlaneDist,
 		FImageAdapter& ResultImageOut);
+
+	/** Combined Metallic/Roughness/Specular uses a custom postprocess material */
+	bool CaptureMRSFromPosition(
+		const FFrame3d& Frame,
+		double HorzFOVDegrees,
+		double NearPlaneDist,
+		FImageAdapter& ResultImageOut);
 };
 
 
