ApproximateActors: add support for combined Metallic/Roughness/Specular texture in ApproximateActors options and GeometryProcessingAdapters implementation. ApproximateActorsImpl now respects the baking enable/disable flags.
#rnx

diff --git a/Engine/Plugins/Experimental/MeshModelingToolset/Source/GeometryProcessingAdapters/Private/ApproximateActorsImpl.cpp b/Engine/Plugins/Experimental/MeshModelingToolset/Source/GeometryProcessingAdapters/Private/ApproximateActorsImpl.cpp
--- a/Engine/Plugins/Experimental/MeshModelingToolset/Source/GeometryProcessingAdapters/Private/ApproximateActorsImpl.cpp
+++ b/Engine/Plugins/Experimental/MeshModelingToolset/Source/GeometryProcessingAdapters/Private/ApproximateActorsImpl.cpp
@@ -51,12 +51,13 @@ static TAutoConsoleVariable<int32> CVarApproximateActorsRDOCCapture(
 
 struct FGeneratedResultTextures
 {
-	UTexture2D* BaseColorMap;
-	UTexture2D* RoughnessMap;
-	UTexture2D* MetallicMap;
-	UTexture2D* SpecularMap;
-	UTexture2D* EmissiveMap;
-	UTexture2D* NormalMap;
+	UTexture2D* BaseColorMap = nullptr;
+	UTexture2D* RoughnessMap = nullptr;
+	UTexture2D* MetallicMap = nullptr;
+	UTexture2D* SpecularMap = nullptr;
+	UTexture2D* PackedMRSMap = nullptr;
+	UTexture2D* EmissiveMap = nullptr;
+	UTexture2D* NormalMap = nullptr;
 };
 
 
@@ -73,6 +74,31 @@ static TUniquePtr<FSceneCapturePhotoSet> CapturePhotoSet(
 	FImageDimensions CaptureDimensions(Options.RenderCaptureImageSize, Options.RenderCaptureImageSize);
 
 	TUniquePtr<FSceneCapturePhotoSet> SceneCapture = MakeUnique<FSceneCapturePhotoSet>();
+
+
+	SceneCapture->SetCaptureTypeEnabled(ERenderCaptureType::BaseColor, Options.bBakeBaseColor);
+	SceneCapture->SetCaptureTypeEnabled(ERenderCaptureType::WorldNormal, Options.bBakeNormalMap);
+	SceneCapture->SetCaptureTypeEnabled(ERenderCaptureType::Emissive, Options.bBakeEmissive);
+
+	bool bMetallic = Options.bBakeMetallic;
+	bool bRoughness = Options.bBakeRoughness;
+	bool bSpecular  = Options.bBakeSpecular;
+	if (Options.bUsePackedMRS && (bMetallic || bRoughness || bSpecular ) )
+	{
+		SceneCapture->SetCaptureTypeEnabled(ERenderCaptureType::CombinedMRS, true);
+		SceneCapture->SetCaptureTypeEnabled(ERenderCaptureType::Roughness, false);
+		SceneCapture->SetCaptureTypeEnabled(ERenderCaptureType::Metallic, false);
+		SceneCapture->SetCaptureTypeEnabled(ERenderCaptureType::Specular, false);
+	}
+	else
+	{
+		SceneCapture->SetCaptureTypeEnabled(ERenderCaptureType::CombinedMRS, false);
+		SceneCapture->SetCaptureTypeEnabled(ERenderCaptureType::Roughness, true);
+		SceneCapture->SetCaptureTypeEnabled(ERenderCaptureType::Metallic, true);
+		SceneCapture->SetCaptureTypeEnabled(ERenderCaptureType::Specular, true);
+	}
+
+
 	SceneCapture->SetCaptureSceneActors(Actors[0]->GetWorld(), Actors);
 
 	SceneCapture->AddStandardExteriorCapturesFromBoundingBox(
@@ -210,48 +236,73 @@ static void BakeTexturesFromPhotoCapture(
 		return MoveTemp(Image);
 	};
 
-	TUniquePtr<TImageBuilder<FVector4f>> RoughnessImage, MetallicImage, SpecularImage, EmissiveImage;
+	bool bMetallic = Options.bBakeMetallic;
+	bool bRoughness = Options.bBakeRoughness;
+	bool bSpecular = Options.bBakeSpecular;
+	TUniquePtr<TImageBuilder<FVector4f>> RoughnessImage, MetallicImage, SpecularImage, PackedMRSImage, EmissiveImage;
 	{
 		TRACE_CPUPROFILER_EVENT_SCOPE(ApproximateActorsImpl_Textures_OtherChannels);
 
-		Progress.EnterProgressFrame(1.f, LOCTEXT("BakingRoughness", "Baking Roughness..."));
-		RoughnessImage = ProcessChannelFunc(ERenderCaptureType::Roughness);
-		Progress.EnterProgressFrame(1.f, LOCTEXT("BakingMetallic", "Baking Metallic..."));
-		MetallicImage = ProcessChannelFunc(ERenderCaptureType::Metallic);
-		Progress.EnterProgressFrame(1.f, LOCTEXT("BakingSpecular", "Baking Specular..."));
-		SpecularImage = ProcessChannelFunc(ERenderCaptureType::Specular);
-		Progress.EnterProgressFrame(1.f, LOCTEXT("BakingEmissive", "Baking Emissive..."));
-		EmissiveImage = ProcessChannelFunc(ERenderCaptureType::Emissive);
+		if (Options.bUsePackedMRS && (bMetallic || bRoughness || bSpecular))
+		{
+			PackedMRSImage = ProcessChannelFunc(ERenderCaptureType::CombinedMRS);
+		}
+		else
+		{
+			if (bRoughness)
+			{
+				RoughnessImage = ProcessChannelFunc(ERenderCaptureType::Roughness);
+			}
+			if (bMetallic)
+			{
+				MetallicImage = ProcessChannelFunc(ERenderCaptureType::Metallic);
+			}
+			if (bSpecular)
+			{
+				SpecularImage = ProcessChannelFunc(ERenderCaptureType::Specular);
+			}
+		}
+
+		if (Options.bBakeEmissive)
+		{
+			EmissiveImage = ProcessChannelFunc(ERenderCaptureType::Emissive);
+		}
 	}
 
+
+
 	Progress.EnterProgressFrame(1.f, LOCTEXT("BakingNormals", "Baking Normals..."));
 
-	// no infill on normal map for now, doesn't make sense to do after mapping to tangent space!
-	//  (should we build baked normal map in world space, and then resample to tangent space??)
-	FVector4f DefaultNormalValue(0, 0, 1, 1);
-	FMeshGenericWorldPositionNormalBaker NormalMapBaker;
-	NormalMapBaker.SetCache(&TempBakeCache);
-	NormalMapBaker.BaseMeshTangents = MeshTangents;
-	NormalMapBaker.NormalSampleFunction = [&](FVector3d Position, FVector3d Normal) {
-		FSceneCapturePhotoSet::FSceneSample Sample = DefaultSample;
-		SceneCapture->ComputeSample(FRenderCaptureTypeFlags::WorldNormal(),
-			Position, Normal, VisibilityFunction, Sample);
-		FVector3f NormalColor = Sample.WorldNormal;
-		float x = (NormalColor.X - 0.5f) * 2.0f;
-		float y = (NormalColor.Y - 0.5f) * 2.0f;
-		float z = (NormalColor.Z - 0.5f) * 2.0f;
-		return FVector3f(x, y, z);
-	};
+	TUniquePtr<TImageBuilder<FVector3f>> NormalImage;
+	if (Options.bBakeNormalMap)
 	{
 		TRACE_CPUPROFILER_EVENT_SCOPE(ApproximateActorsImpl_Textures_NormalMapBake);
+
+		// no infill on normal map for now, doesn't make sense to do after mapping to tangent space!
+		//  (should we build baked normal map in world space, and then resample to tangent space??)
+		FVector4f DefaultNormalValue(0, 0, 1, 1);
+		FMeshGenericWorldPositionNormalBaker NormalMapBaker;
+		NormalMapBaker.SetCache(&TempBakeCache);
+		NormalMapBaker.BaseMeshTangents = MeshTangents;
+		NormalMapBaker.NormalSampleFunction = [&](FVector3d Position, FVector3d Normal) {
+			FSceneCapturePhotoSet::FSceneSample Sample = DefaultSample;
+			SceneCapture->ComputeSample(FRenderCaptureTypeFlags::WorldNormal(),
+				Position, Normal, VisibilityFunction, Sample);
+			FVector3f NormalColor = Sample.WorldNormal;
+			float x = (NormalColor.X - 0.5f) * 2.0f;
+			float y = (NormalColor.Y - 0.5f) * 2.0f;
+			float z = (NormalColor.Z - 0.5f) * 2.0f;
+			return FVector3f(x, y, z);
+		};
+
 		NormalMapBaker.Bake();
-	}
-	TUniquePtr<TImageBuilder<FVector3f>> NormalImage = NormalMapBaker.TakeResult();
+		NormalImage = NormalMapBaker.TakeResult();
 
-	if (Supersample > 1)
-	{
-		TImageBuilder<FVector3f> Downsampled = NormalImage->FastDownsample(Supersample, FVector3f::Zero(), [](FVector3f V, int N) { return V / (float)N; });
-		*NormalImage = MoveTemp(Downsampled);
+		if (Supersample > 1)
+		{
+			TImageBuilder<FVector3f> Downsampled = NormalImage->FastDownsample(Supersample, FVector3f::Zero(), [](FVector3f V, int N) { return V / (float)N; });
+			*NormalImage = MoveTemp(Downsampled);
+		}
 	}
 
 	// build textures
@@ -261,18 +312,46 @@ static void BakeTexturesFromPhotoCapture(
 
 		FScopedSlowTask BuildTexProgress(6.f, LOCTEXT("BuildingTextures", "Building Textures..."));
 		BuildTexProgress.MakeDialog(true);
-		BuildTexProgress.EnterProgressFrame(1.f);
-		GeneratedTextures.BaseColorMap = FTexture2DBuilder::BuildTextureFromImage(*ColorImage, FTexture2DBuilder::ETextureType::Color, true, false);
-		BuildTexProgress.EnterProgressFrame(1.f);
-		GeneratedTextures.RoughnessMap = FTexture2DBuilder::BuildTextureFromImage(*RoughnessImage, FTexture2DBuilder::ETextureType::Roughness, false, false);
-		BuildTexProgress.EnterProgressFrame(1.f);
-		GeneratedTextures.MetallicMap = FTexture2DBuilder::BuildTextureFromImage(*MetallicImage, FTexture2DBuilder::ETextureType::Metallic, false, false);
-		BuildTexProgress.EnterProgressFrame(1.f);
-		GeneratedTextures.SpecularMap = FTexture2DBuilder::BuildTextureFromImage(*SpecularImage, FTexture2DBuilder::ETextureType::Specular, false, false);
-		BuildTexProgress.EnterProgressFrame(1.f);
-		GeneratedTextures.EmissiveMap = FTexture2DBuilder::BuildTextureFromImage(*EmissiveImage, FTexture2DBuilder::ETextureType::Color, true, false);
-		BuildTexProgress.EnterProgressFrame(1.f);
-		GeneratedTextures.NormalMap = FTexture2DBuilder::BuildTextureFromImage(*NormalImage, FTexture2DBuilder::ETextureType::NormalMap, false, false);
+		if (Options.bBakeBaseColor && ColorImage.IsValid())
+		{
+			BuildTexProgress.EnterProgressFrame(1.f);
+			GeneratedTextures.BaseColorMap = FTexture2DBuilder::BuildTextureFromImage(*ColorImage, FTexture2DBuilder::ETextureType::Color, true, false);
+		}
+		if (Options.bBakeEmissive && EmissiveImage.IsValid())
+		{
+			BuildTexProgress.EnterProgressFrame(1.f);
+			GeneratedTextures.EmissiveMap = FTexture2DBuilder::BuildTextureFromImage(*EmissiveImage, FTexture2DBuilder::ETextureType::ColorLinear, true, false);
+		}
+		if (Options.bBakeNormalMap && NormalImage.IsValid())
+		{
+			BuildTexProgress.EnterProgressFrame(1.f);
+			GeneratedTextures.NormalMap = FTexture2DBuilder::BuildTextureFromImage(*NormalImage, FTexture2DBuilder::ETextureType::NormalMap, false, false);
+		}
+
+		if ( (bRoughness || bMetallic || bSpecular) && PackedMRSImage.IsValid())
+		{
+			BuildTexProgress.EnterProgressFrame(1.f);
+			GeneratedTextures.PackedMRSMap = FTexture2DBuilder::BuildTextureFromImage(*PackedMRSImage, FTexture2DBuilder::ETextureType::ColorLinear, false, false);
+
+		}
+		else
+		{ 
+			if (bRoughness && RoughnessImage.IsValid())
+			{
+				BuildTexProgress.EnterProgressFrame(1.f);
+				GeneratedTextures.RoughnessMap = FTexture2DBuilder::BuildTextureFromImage(*RoughnessImage, FTexture2DBuilder::ETextureType::Roughness, false, false);
+			}
+			if (bMetallic && MetallicImage.IsValid())
+			{
+				BuildTexProgress.EnterProgressFrame(1.f);
+				GeneratedTextures.MetallicMap = FTexture2DBuilder::BuildTextureFromImage(*MetallicImage, FTexture2DBuilder::ETextureType::Metallic, false, false);
+			}
+			if (bSpecular && SpecularImage.IsValid())
+			{
+				BuildTexProgress.EnterProgressFrame(1.f);
+				GeneratedTextures.SpecularMap = FTexture2DBuilder::BuildTextureFromImage(*SpecularImage, FTexture2DBuilder::ETextureType::Specular, false, false);
+			}
+		}
 	}
 }
 
@@ -682,9 +761,9 @@ void FApproximateActorsImpl::GenerateApproximationForActorSet(const TArray<AActo
 
 	Progress.EnterProgressFrame(1.f, LOCTEXT("Writing Assets", "Writing Assets..."));
 
-	// Make material for textures by duplicating input material (hardcoded!!)
+	// Make material for textures by creating MIC of input material, or fall back to known material
 	UMaterialInterface* UseBaseMaterial = (Options.BakeMaterial != nullptr) ? 
-		Options.BakeMaterial : LoadObject<UMaterial>(nullptr, TEXT("/MeshModelingToolset/Materials/FullMaterialBakePreviewMaterial"));
+		Options.BakeMaterial : LoadObject<UMaterial>(nullptr, TEXT("/MeshModelingToolset/Materials/FullMaterialBakePreviewMaterial_PackedMRS"));
 	FMaterialAssetOptions MatOptions;
 	MatOptions.NewAssetPath = Options.BasePackagePath + TEXT("_Material");
 	FMaterialAssetResults MatResults;
@@ -735,6 +814,19 @@ void FApproximateActorsImpl::GenerateApproximationForActorSet(const TArray<AActo
 	{
 		WriteTextureLambda(GeneratedTextures.BaseColorMap, TEXT("_BaseColor"), FTexture2DBuilder::ETextureType::Color, Options.BaseColorTexParamName);
 	}
+	if (Options.bBakeEmissive && GeneratedTextures.EmissiveMap)
+	{
+		WriteTextureLambda(GeneratedTextures.EmissiveMap, TEXT("_Emissive"), FTexture2DBuilder::ETextureType::ColorLinear, Options.EmissiveTexParamName);
+	}
+	if (Options.bBakeNormalMap && GeneratedTextures.NormalMap)
+	{
+		WriteTextureLambda(GeneratedTextures.NormalMap, TEXT("_Normal"), FTexture2DBuilder::ETextureType::NormalMap, Options.NormalTexParamName);
+	}
+
+	if ((Options.bBakeRoughness || Options.bBakeMetallic || Options.bBakeSpecular) && Options.bUsePackedMRS && GeneratedTextures.PackedMRSMap)
+	{
+		WriteTextureLambda(GeneratedTextures.PackedMRSMap, TEXT("_PackedMRS"), FTexture2DBuilder::ETextureType::ColorLinear, Options.PackedMRSTexParamName);
+	}
 	if (Options.bBakeRoughness && GeneratedTextures.RoughnessMap)
 	{
 		WriteTextureLambda(GeneratedTextures.RoughnessMap, TEXT("_Roughness"), FTexture2DBuilder::ETextureType::Roughness, Options.RoughnessTexParamName);
@@ -747,14 +839,7 @@ void FApproximateActorsImpl::GenerateApproximationForActorSet(const TArray<AActo
 	{
 		WriteTextureLambda(GeneratedTextures.SpecularMap, TEXT("_Specular"), FTexture2DBuilder::ETextureType::Specular, Options.SpecularTexParamName);
 	}
-	if (Options.bBakeEmissive && GeneratedTextures.EmissiveMap)
-	{
-		WriteTextureLambda(GeneratedTextures.EmissiveMap, TEXT("_Emissive"), FTexture2DBuilder::ETextureType::Color, Options.EmissiveTexParamName);
-	}
-	if (Options.bBakeNormalMap && GeneratedTextures.NormalMap)
-	{
-		WriteTextureLambda(GeneratedTextures.NormalMap, TEXT("_Normal"), FTexture2DBuilder::ETextureType::NormalMap, Options.NormalTexParamName);
-	}
+
 
 	// force material update now that we have updated texture parameters
 	// (does this do that? Let calling code do it?)
diff --git a/Engine/Source/Developer/GeometryProcessingInterfaces/Public/GeometryProcessingInterfaces/ApproximateActors.h b/Engine/Source/Developer/GeometryProcessingInterfaces/Public/GeometryProcessingInterfaces/ApproximateActors.h
--- a/Engine/Source/Developer/GeometryProcessingInterfaces/Public/GeometryProcessingInterfaces/ApproximateActors.h
+++ b/Engine/Source/Developer/GeometryProcessingInterfaces/Public/GeometryProcessingInterfaces/ApproximateActors.h
@@ -115,7 +115,7 @@ public:
 		//
 
 		// A new MIC derived from this material will be created and assigned to the generated mesh
-		UMaterialInterface* BakeMaterial = nullptr;		// if null, will use /MeshModelingToolset/Materials/FullMaterialBakePreviewMaterial instead
+		UMaterialInterface* BakeMaterial = nullptr;		// if null, will use /MeshModelingToolset/Materials/FullMaterialBakePreviewMaterial_PackedMRS instead
 		FName BaseColorTexParamName = FName("BaseColor");
 		bool bBakeBaseColor = true;
 		FName RoughnessTexParamName = FName("Roughness");
@@ -128,6 +128,9 @@ public:
 		bool bBakeEmissive = true;
 		FName NormalTexParamName = FName("NormalMap");
 		bool bBakeNormalMap = true;
+		
+		bool bUsePackedMRS = true;
+		FName PackedMRSTexParamName = FName("PackedMRS");
 
 		// output texture options
 		int32 TextureImageSize = 1024;
