Work around a threading issue where the thread pool was saturated and we're blocked waiting on futures results

diff --git a/Engine/Plugins/Experimental/MeshModelingToolset/Source/ModelingComponents/Private/Scene/MeshSceneAdapter.cpp b/Engine/Plugins/Experimental/MeshModelingToolset/Source/ModelingComponents/Private/Scene/MeshSceneAdapter.cpp
--- a/Engine/Plugins/Experimental/MeshModelingToolset/Source/ModelingComponents/Private/Scene/MeshSceneAdapter.cpp
+++ b/Engine/Plugins/Experimental/MeshModelingToolset/Source/ModelingComponents/Private/Scene/MeshSceneAdapter.cpp
@@ -603,6 +603,9 @@ void FMeshSceneAdapter::Build_FullDecompose(const FMeshSceneAdapterBuildOptions&
 	DecomposedMeshesCount = 0;
 	int32 AddedTrisCount = 0;
 
+	FCriticalSection PendingFuturesLock;
+	TArray<TFuture<void>> PendingClosedAssemblyBuilds;
+
 	// use this to disable build of unnecessary AABBTree/etc
 	FMeshSceneAdapterBuildOptions TempBuildOptions = BuildOptions;
 	TempBuildOptions.bBuildSpatialDataStructures = false;
@@ -722,19 +725,22 @@ void FMeshSceneAdapter::Build_FullDecompose(const FMeshSceneAdapterBuildOptions&
 
 		// we re-use the existing wrapper for the closed meshes, if there are any, otherwise
 		// we need to disable the parent meshes that refer to it
-		TFuture<void> PendingClosedAssemblyBuild;
 		if (bMakeClosedAssembly)
 		{
 			TUniquePtr<FDynamicMeshSpatialWrapper> LocalSpaceMeshWrapper = MakeUnique<FDynamicMeshSpatialWrapper>();
 			
 			LocalSpaceMeshWrapper->Mesh = MoveTemp(LocalSpaceParts);
 			FDynamicMeshSpatialWrapper* BuildWrapper = LocalSpaceMeshWrapper.Get();
-			PendingClosedAssemblyBuild = Async(EAsyncExecution::ThreadPool, [&BuildOptions, BuildWrapper]()
+			TFuture<void> PendingClosedAssemblyBuild = Async(EAsyncExecution::ThreadPool, [&BuildOptions, BuildWrapper]()
 			{
 				TRACE_CPUPROFILER_EVENT_SCOPE(MeshScene_Build_ProcessMesh_6SolidBuild);
 				BuildWrapper->Build(BuildOptions);
 			});
 
+			PendingFuturesLock.Lock();
+			PendingClosedAssemblyBuilds.Add(MoveTemp(PendingClosedAssemblyBuild));
+			PendingFuturesLock.Unlock();
+
 			WrapperInfo->SpatialWrapper = MoveTemp(LocalSpaceMeshWrapper);
 
 			for (FActorChildMesh* MeshInstance : MeshesToDecompose)
@@ -791,10 +797,13 @@ void FMeshSceneAdapter::Build_FullDecompose(const FMeshSceneAdapterBuildOptions&
 		SceneActors.Add(MoveTemp(Adapter));
 		AddedTrisCount += NewWrapperInfo->SpatialWrapper->GetTriangleCount();
 		ListsLock.Unlock();
+	});
 
-		// make sure this finishes
+	// make sure these finishes
+	for (auto& PendingClosedAssemblyBuild : PendingClosedAssemblyBuilds)
+	{
 		PendingClosedAssemblyBuild.Wait();
-	});
+	}	
 
 	// currently true with the methods used above?
 	bSceneIsAllSolids = true;
