try to match the number of UV layers in the source mesh(es) when duplicating / combining meshes

#rnx

diff --git a/Engine/Plugins/Experimental/MeshModelingToolset/Source/MeshModelingTools/Private/CombineMeshesTool.cpp b/Engine/Plugins/Experimental/MeshModelingToolset/Source/MeshModelingTools/Private/CombineMeshesTool.cpp
--- a/Engine/Plugins/Experimental/MeshModelingToolset/Source/MeshModelingTools/Private/CombineMeshesTool.cpp
+++ b/Engine/Plugins/Experimental/MeshModelingToolset/Source/MeshModelingTools/Private/CombineMeshesTool.cpp
@@ -256,6 +256,11 @@ void UCombineMeshesTool::CreateNewAsset()
 			Converter.Convert(MeshDescription, ComponentDMesh);
 			bNeedColorAttr = bNeedColorAttr || (ComponentDMesh.HasAttributes() && ComponentDMesh.Attributes()->HasPrimaryColors());
 
+			if (ComponentDMesh.HasAttributes() && ComponentDMesh.Attributes()->NumUVLayers() > AccumulateDMesh.Attributes()->NumUVLayers())
+			{
+				AccumulateDMesh.Attributes()->SetNumUVLayers(ComponentDMesh.Attributes()->NumUVLayers());
+			}
+
 			UE::Geometry::FTransform3d XF = (UE::Geometry::FTransform3d)(TargetComponentInterface(ComponentIdx)->GetWorldTransform() * ToAccum);
 			if (XF.GetDeterminant() < 0)
 			{
