ModelingComponents: add support for HDR Emissive capture and texture handling
- FTexture2DBuilder has new EmissiveHDR texture type, uses PF_FloatRGBA/FFloat16Color texture format/buffers in this case. Moved most functions from header to cpp.
- new WorldRenderCapture::CaptureEmissiveFromPosition() implementation captures PreTonemapHDRColor buffer, with showflags config such that only Emissive (and possibly some GI - untested) will be included.
- ApproximateActorsImpl now stores Emissive bake in HDR texture. Update preview material to use Linear Texture Sampler for Emissive, required to use HDR texture
#rnx

diff --git a/Engine/Plugins/Experimental/MeshModelingToolset/Source/GeometryProcessingAdapters/Private/ApproximateActorsImpl.cpp b/Engine/Plugins/Experimental/MeshModelingToolset/Source/GeometryProcessingAdapters/Private/ApproximateActorsImpl.cpp
--- a/Engine/Plugins/Experimental/MeshModelingToolset/Source/GeometryProcessingAdapters/Private/ApproximateActorsImpl.cpp
+++ b/Engine/Plugins/Experimental/MeshModelingToolset/Source/GeometryProcessingAdapters/Private/ApproximateActorsImpl.cpp
@@ -320,7 +320,7 @@ static void BakeTexturesFromPhotoCapture(
 		if (Options.bBakeEmissive && EmissiveImage.IsValid())
 		{
 			BuildTexProgress.EnterProgressFrame(1.f);
-			GeneratedTextures.EmissiveMap = FTexture2DBuilder::BuildTextureFromImage(*EmissiveImage, FTexture2DBuilder::ETextureType::ColorLinear, true, false);
+			GeneratedTextures.EmissiveMap = FTexture2DBuilder::BuildTextureFromImage(*EmissiveImage, FTexture2DBuilder::ETextureType::EmissiveHDR, false, false);
 		}
 		if (Options.bBakeNormalMap && NormalImage.IsValid())
 		{
@@ -816,7 +816,7 @@ void FApproximateActorsImpl::GenerateApproximationForActorSet(const TArray<AActo
 	}
 	if (Options.bBakeEmissive && GeneratedTextures.EmissiveMap)
 	{
-		WriteTextureLambda(GeneratedTextures.EmissiveMap, TEXT("_Emissive"), FTexture2DBuilder::ETextureType::ColorLinear, Options.EmissiveTexParamName);
+		WriteTextureLambda(GeneratedTextures.EmissiveMap, TEXT("_Emissive"), FTexture2DBuilder::ETextureType::EmissiveHDR, Options.EmissiveTexParamName);
 	}
 	if (Options.bBakeNormalMap && GeneratedTextures.NormalMap)
 	{
diff --git a/Engine/Plugins/Experimental/MeshModelingToolset/Source/ModelingComponents/Private/AssetUtils/Texture2DBuilder.cpp b/Engine/Plugins/Experimental/MeshModelingToolset/Source/ModelingComponents/Private/AssetUtils/Texture2DBuilder.cpp
new file mode 100644
--- /dev/null
+++ b/Engine/Plugins/Experimental/MeshModelingToolset/Source/ModelingComponents/Private/AssetUtils/Texture2DBuilder.cpp
@@ -0,0 +1,397 @@
+// Copyright Epic Games, Inc. All Rights Reserved. 
+
+#include "AssetUtils/Texture2DBuilder.h"
+
+#include "ExplicitUseGeometryMathTypes.h"		// using UE::Geometry::(math types)
+using namespace UE::Geometry;
+
+
+
+bool FTexture2DBuilder::Initialize(ETextureType BuildTypeIn, FImageDimensions DimensionsIn)
+{
+	check(DimensionsIn.IsSquare());
+	BuildType = BuildTypeIn;
+	Dimensions = DimensionsIn;
+
+	// create new texture
+	EPixelFormat UsePixelFormat = (BuildTypeIn == ETextureType::EmissiveHDR) ? PF_FloatRGBA : PF_B8G8R8A8;
+	RawTexture2D = UTexture2D::CreateTransient((int32)Dimensions.GetWidth(), (int32)Dimensions.GetHeight(), UsePixelFormat);
+	if (RawTexture2D == nullptr)
+	{
+		return false;
+	}
+	CurrentPixelFormat = UsePixelFormat;
+
+	if (BuildType == ETextureType::ColorLinear 
+		|| BuildType == ETextureType::Roughness 
+		|| BuildType == ETextureType::Metallic 
+		|| BuildType == ETextureType::Specular)
+	{
+		RawTexture2D->SRGB = false;
+		RawTexture2D->UpdateResource();
+	}
+	else if (BuildType == ETextureType::EmissiveHDR)
+	{
+		RawTexture2D->SRGB = false;
+		RawTexture2D->CompressionSettings = TC_HDR;
+		RawTexture2D->UpdateResource();
+	}
+	else if (BuildType == ETextureType::NormalMap)
+	{
+		RawTexture2D->CompressionSettings = TC_Normalmap;
+		RawTexture2D->SRGB = false;
+		RawTexture2D->LODGroup = TEXTUREGROUP_WorldNormalMap;
+		//RawTexture2D->bFlipGreenChannel = true;
+#if WITH_EDITOR
+		RawTexture2D->MipGenSettings = TMGS_NoMipmaps;
+#endif
+		RawTexture2D->UpdateResource();
+	}
+
+	// lock
+	if (LockForEditing() == false)
+	{
+		return false;
+	}
+
+	if (IsEditable())
+	{
+		Clear();
+	}
+
+	return true;
+}
+
+
+bool FTexture2DBuilder::Initialize(UTexture2D* ExistingTexture, ETextureType BuildTypeIn, bool bLockForEditing)
+{
+	if (!ensure(ExistingTexture != nullptr)) return false;
+	const FTexturePlatformData* TexPlatformData = ExistingTexture->GetPlatformData();
+	if (!ensure(TexPlatformData != nullptr)) return false;
+	if (!ensure(TexPlatformData->Mips.Num() > 0)) return false;
+
+	EPixelFormat UsePixelFormat = (BuildTypeIn == ETextureType::EmissiveHDR) ? PF_FloatRGBA : PF_B8G8R8A8;
+	if (!ensure(TexPlatformData->PixelFormat == UsePixelFormat))
+	{
+		return false;
+	}
+	CurrentPixelFormat = UsePixelFormat;
+
+	int32 Width = TexPlatformData->Mips[0].SizeX;
+	int32 Height = TexPlatformData->Mips[0].SizeY;
+	Dimensions = FImageDimensions(Width, Height);
+	BuildType = BuildTypeIn;
+	RawTexture2D = ExistingTexture;
+
+	// lock for editing
+	if (bLockForEditing && LockForEditing() == false)
+	{
+		return false;
+	}
+
+	return true;
+}
+
+
+
+bool FTexture2DBuilder::LockForEditing()
+{
+	if (!ensure(RawTexture2D != nullptr && CurrentMipData == nullptr && CurrentMipDataFloat16 == nullptr))
+	{
+		return false;
+	}
+
+	if (IsByteTexture())
+	{
+		if (RawTexture2D && CurrentMipData == nullptr)
+		{
+			CurrentMipData = reinterpret_cast<FColor*>(RawTexture2D->GetPlatformData()->Mips[0].BulkData.Lock(LOCK_READ_WRITE));
+			ensure(CurrentMipData);
+		}
+	}
+	else
+	{
+		if (RawTexture2D && CurrentMipDataFloat16 == nullptr)
+		{
+			CurrentMipDataFloat16 = reinterpret_cast<FFloat16Color*>(RawTexture2D->GetPlatformData()->Mips[0].BulkData.Lock(LOCK_READ_WRITE));
+			ensure(CurrentMipDataFloat16);
+		}
+	}
+
+	return IsEditable();
+}
+
+
+void FTexture2DBuilder::Commit(bool bUpdateSourceData)
+{
+	if (ensure(RawTexture2D != nullptr && IsEditable()))
+	{
+		if (bUpdateSourceData)
+		{
+			UpdateSourceData();
+		}
+
+		RawTexture2D->GetPlatformData()->Mips[0].BulkData.Unlock();
+		RawTexture2D->UpdateResource();
+
+		CurrentMipData = nullptr;
+		CurrentMipDataFloat16 = nullptr;
+	}
+}
+
+PRAGMA_DISABLE_OPTIMIZATION
+void FTexture2DBuilder::UpdateSourceData()
+{
+	// source data only exists in Editor
+#if WITH_EDITOR
+	check(RawTexture2D);
+
+	bool bIsEditable = IsEditable();
+	bool bIsValid = RawTexture2D->Source.IsValid();		// what is this for??
+
+	if (IsByteTexture())
+	{
+		const FColor* SourceMipData = (bIsEditable) ? CurrentMipData :
+			reinterpret_cast<const FColor*>(RawTexture2D->GetPlatformData()->Mips[0].BulkData.Lock(LOCK_READ_ONLY));
+		RawTexture2D->Source.Init2DWithMipChain(Dimensions.GetWidth(), Dimensions.GetHeight(), TSF_BGRA8);
+		uint8* DestData = RawTexture2D->Source.LockMip(0);
+		FMemory::Memcpy(DestData, SourceMipData, Dimensions.GetWidth() * Dimensions.GetHeight() * 4 * sizeof(uint8));
+	}
+	else
+	{
+		const FFloat16Color* SourceMipData = (bIsEditable) ? CurrentMipDataFloat16 :
+			reinterpret_cast<const FFloat16Color*>(RawTexture2D->GetPlatformData()->Mips[0].BulkData.Lock(LOCK_READ_ONLY));
+		RawTexture2D->Source.Init2DWithMipChain(Dimensions.GetWidth(), Dimensions.GetHeight(), TSF_RGBA16F);
+		uint8* DestData = RawTexture2D->Source.LockMip(0);
+		FMemory::Memcpy((void*)DestData, (void*)SourceMipData, Dimensions.GetWidth() * Dimensions.GetHeight() * sizeof(FFloat16Color));
+	}
+
+	RawTexture2D->Source.UnlockMip(0);
+	if (bIsEditable == false)
+	{
+		RawTexture2D->GetPlatformData()->Mips[0].BulkData.Unlock();
+	}
+#endif
+}
+PRAGMA_ENABLE_OPTIMIZATION
+
+
+void FTexture2DBuilder::Cancel()
+{
+	bool bIsEditable = IsEditable();
+	if (bIsEditable)
+	{
+		RawTexture2D->GetPlatformData()->Mips[0].BulkData.Unlock();
+		CurrentMipData = nullptr;
+		CurrentMipDataFloat16 = nullptr;
+	}
+}
+
+
+void FTexture2DBuilder::Clear()
+{
+	if (IsByteTexture())
+	{
+		Clear(GetClearColor());
+	}
+	else
+	{
+		Clear(GetClearColorFloat16());
+	}
+}
+
+
+void FTexture2DBuilder::Clear(const FColor& ClearColor)
+{
+	if (ensure(IsEditable() && IsByteTexture()))
+	{
+		int32 Num = Dimensions.Num();
+		for (int64 k = 0; k < Num; ++k)
+		{
+			CurrentMipData[k] = ClearColor;
+		}
+	}
+}
+
+/**
+ * Clear all texels in the current Mip to the given ClearColor
+ */
+void FTexture2DBuilder::Clear(const FFloat16Color& ClearColor)
+{
+	if (ensure(IsEditable() && IsFloat16Texture()))
+	{
+		int32 Num = Dimensions.Num();
+		for (int64 k = 0; k < Num; ++k)
+		{
+			CurrentMipDataFloat16[k] = ClearColor;
+		}
+	}
+}
+
+
+bool FTexture2DBuilder::Copy(const TImageBuilder<FVector3f>& SourceImage, const bool bConvertToSRGB)
+{
+	if (ensure(SourceImage.GetDimensions() == Dimensions) == false)
+	{
+		return false;
+	}
+	if (IsFloat16Texture() && bConvertToSRGB)
+	{
+		ensure(false);		// not currently supported
+	}
+
+	int64 Num = Dimensions.Num();
+	for (int32 i = 0; i < Num; ++i)
+	{
+		FVector3f Pixel = SourceImage.GetPixel(i);
+		if (IsByteTexture())
+		{
+			Pixel.X = FMathf::Clamp(Pixel.X, 0.0, 1.0);
+			Pixel.Y = FMathf::Clamp(Pixel.Y, 0.0, 1.0);
+			Pixel.Z = FMathf::Clamp(Pixel.Z, 0.0, 1.0);
+			FColor Texel = ToLinearColor(Pixel).ToFColor(bConvertToSRGB);
+			SetTexel(i, Texel);
+		}
+		else
+		{
+			FFloat16Color Texel(ToLinearColor(Pixel));
+			SetTexel(i, Texel);
+		}
+	}
+	return true;
+}
+
+
+bool FTexture2DBuilder::Copy(const TImageBuilder<FVector4f>& SourceImage, const bool bConvertToSRGB)
+{
+	if (ensure(SourceImage.GetDimensions() == Dimensions) == false)
+	{
+		return false;
+	}
+	if (IsFloat16Texture() && bConvertToSRGB)
+	{
+		ensure(false);		// not currently supported
+	}
+
+	int64 Num = Dimensions.Num();
+	for (int32 i = 0; i < Num; ++i)
+	{
+		FVector4f Pixel = SourceImage.GetPixel(i);
+		if (IsByteTexture())
+		{
+			Pixel.X = FMathf::Clamp(Pixel.X, 0.0, 1.0);
+			Pixel.Y = FMathf::Clamp(Pixel.Y, 0.0, 1.0);
+			Pixel.Z = FMathf::Clamp(Pixel.Z, 0.0, 1.0);
+			Pixel.W = FMathf::Clamp(Pixel.W, 0.0, 1.0);
+			FColor Texel = ((FLinearColor)Pixel).ToFColor(bConvertToSRGB);
+			SetTexel(i, Texel);
+		}
+		else
+		{
+			FFloat16Color Texel((FLinearColor)Pixel);
+			SetTexel(i, Texel);
+		}
+	}
+	return true;
+}
+
+
+bool FTexture2DBuilder::CopyTo(TImageBuilder<FVector4f>& DestImage) const
+{
+	if (ensure(DestImage.GetDimensions() == Dimensions) == false)
+	{
+		return false;
+	}
+	int64 Num = Dimensions.Num();
+	for (int32 i = 0; i < Num; ++i)
+	{
+		if (IsByteTexture())
+		{
+			FColor ByteColor = GetTexel(i);
+			FLinearColor FloatColor(ByteColor);
+			DestImage.SetPixel(i, FVector4f(FloatColor));
+		}
+		else
+		{
+			FFloat16Color Float16Color = GetTexelFloat16(i);
+			FLinearColor FloatColor = Float16Color.GetFloats();
+			DestImage.SetPixel(i, FVector4f(FloatColor));
+		}
+	}
+	return true;
+}
+
+
+bool FTexture2DBuilder::CopyPlatformDataToSourceData(UTexture2D* Texture, ETextureType TextureType)
+{
+	FTexture2DBuilder Builder;
+	bool bOK = Builder.Initialize(Texture, TextureType, false);
+	if (bOK)
+	{
+		Builder.UpdateSourceData();
+	}
+	return bOK;
+}
+
+
+
+const FColor& FTexture2DBuilder::GetClearColor() const
+{
+	static const FColor DefaultColor = FColor::Black;
+	static const FColor DefaultRoughness(128, 128, 128);
+	static const FColor DefaultSpecular(100, 100, 100);
+	static const FColor DefaultMetallic(16, 16, 16);
+	static const FColor DefaultNormalColor(128, 128, 255);
+	static const FColor DefaultAOColor = FColor::White;
+
+	switch (BuildType)
+	{
+		default:
+		case ETextureType::Color:
+		case ETextureType::ColorLinear:
+		case ETextureType::EmissiveHDR:
+			return DefaultColor;
+		case ETextureType::Roughness:
+			return DefaultRoughness;
+		case ETextureType::Metallic:
+			return DefaultMetallic;
+		case ETextureType::Specular:
+			return DefaultSpecular;
+		case ETextureType::NormalMap:
+			return DefaultNormalColor;
+		case ETextureType::AmbientOcclusion:
+			return DefaultAOColor;
+	}
+}
+
+
+/**
+ * @return the default color for the current texture build type
+ */
+FFloat16Color FTexture2DBuilder::GetClearColorFloat16() const
+{
+	static const FLinearColor DefaultColor(0.0f, 0.0f, 0.0f);
+	static const FLinearColor DefaultRoughness(0.5f, 0.5f, 0.5f);
+	static const FLinearColor DefaultSpecular(0.4f, 0.4f, 0.4f);
+	static const FLinearColor DefaultMetallic(0.05f, 0.05f, 0.05f);
+	static const FLinearColor DefaultNormalColor(0.5f, 0.5f, 0.5f);
+	static const FLinearColor DefaultAOColor(1.0f, 1.0f, 1.0f);
+	switch (BuildType)
+	{
+		default:
+		case ETextureType::Color:
+		case ETextureType::ColorLinear:
+		case ETextureType::EmissiveHDR:
+			return FFloat16Color(DefaultColor);
+		case ETextureType::Roughness:
+			return FFloat16Color(DefaultRoughness);
+		case ETextureType::Metallic:
+			return FFloat16Color(DefaultMetallic);
+		case ETextureType::Specular:
+			return FFloat16Color(DefaultSpecular);
+		case ETextureType::NormalMap:
+			return FFloat16Color(DefaultNormalColor);
+		case ETextureType::AmbientOcclusion:
+			return FFloat16Color(DefaultAOColor);
+	}
+}
\ No newline at end of file
diff --git a/Engine/Plugins/Experimental/MeshModelingToolset/Source/ModelingComponents/Private/Scene/WorldRenderCapture.cpp b/Engine/Plugins/Experimental/MeshModelingToolset/Source/ModelingComponents/Private/Scene/WorldRenderCapture.cpp
--- a/Engine/Plugins/Experimental/MeshModelingToolset/Source/ModelingComponents/Private/Scene/WorldRenderCapture.cpp
+++ b/Engine/Plugins/Experimental/MeshModelingToolset/Source/ModelingComponents/Private/Scene/WorldRenderCapture.cpp
@@ -225,173 +225,15 @@
 
 
 
-bool FWorldRenderCapture::CaptureEmissiveFromPosition(
-	const FFrame3d& ViewFrame,
-	double HorzFOVDegrees,
-	double NearPlaneDist,
-	FImageAdapter& ResultImageOut)
-{
-	// This function is a combination of CaptureFromPosition() and RenderSceneVisualizationToTexture() below,
-	// that attempts to capture Emissive. The difficulty with capturing Emissive is there is no Visualization
-	// Buffer Mode specifically for it, like there is for BaseColor, Specular, WorldNormal, etc. So, the
-	// strategy we take here is basically to disable all scene lights but not disable lighting. This should
-	// result in only Emissive being rendered, ideally without any tone mapping, gamma, bloom, etc, so that we 
-	// directly capture the raw Emissive shader output.
-	//
-	// One known possible issue here is with baked lighting, which may also appear with the current setup.
-	// Currently untested.
 
-	check(this->World);
-	FSceneInterface* Scene = this->World->Scene;
-
-	UTextureRenderTarget2D* RenderTargetTexture = GetRenderTexture(true);
-	if (!ensure(RenderTargetTexture))
-	{
-		return false;
-	}
-	FRenderTarget* RenderTargetResource = RenderTargetTexture->GameThread_GetRenderTargetResource();
-
-	int32 Width = Dimensions.GetWidth();
-	int32 Height = Dimensions.GetHeight();
-
-	FQuat ViewOrientation = (FQuat)ViewFrame.Rotation;
-	FMatrix ViewRotationMatrix = FInverseRotationMatrix(ViewOrientation.Rotator());
-	FVector ViewOrigin = (FVector)ViewFrame.Origin;
-
-	// convert to rendering coordinate system
-	ViewRotationMatrix = ViewRotationMatrix * FMatrix(
-		FPlane(0, 0, 1, 0),
-		FPlane(1, 0, 0, 0),
-		FPlane(0, 1, 0, 0),
-		FPlane(0, 0, 0, 1));
-
-	const float HalfFOVRadians = FMath::DegreesToRadians<float>(HorzFOVDegrees) * 0.5f;
-	static_assert((int32)ERHIZBuffer::IsInverted != 0, "Check NearPlane and Projection Matrix");
-	FMatrix ProjectionMatrix = FReversedZPerspectiveMatrix(HalfFOVRadians, 1.0f, 1.0f, (float)NearPlaneDist);
-
-	EViewModeIndex ViewModeIndex = EViewModeIndex::VMI_Unlit;		// VMI_Lit, VMI_LightingOnly, VMI_VisualizeBuffer
-
-	FEngineShowFlags ShowFlags = FEngineShowFlags(EShowFlagInitMode::ESFIM_Game);
-	ApplyViewMode(ViewModeIndex, true, ShowFlags);
-
-	// unclear if these flags need to be set before creating ViewFamily
-	ShowFlags.SetAntiAliasing(false);
-	ShowFlags.SetDepthOfField(false);
-	ShowFlags.SetMotionBlur(false);
-	ShowFlags.SetBloom(false);
-	ShowFlags.SetSceneColorFringe(false);
-
-	FSceneViewFamilyContext ViewFamily(FSceneViewFamily::ConstructionValues(
-		RenderTargetResource, Scene, ShowFlags)
-		.SetWorldTimes(0, 0, 0)
-		.SetRealtimeUpdate(false));
-
-	// This set of flags currently seems to work for capturing Emissive. Possibly many are unnecessary or
-	// ignored due to ther rendering config, but it's quite hard to know without extensive A/B testing
-
-	ViewFamily.EngineShowFlags.DisableAdvancedFeatures();
-	ViewFamily.EngineShowFlags.MotionBlur = 0;
-	ViewFamily.EngineShowFlags.LOD = 0;
-
-	ViewFamily.EngineShowFlags.SetTonemapper(false);
-	ViewFamily.EngineShowFlags.SetColorGrading(false);
-	ViewFamily.EngineShowFlags.SetToneCurve(false);
-
-	ViewFamily.EngineShowFlags.SetPostProcessing(false);
-	ViewFamily.EngineShowFlags.SetFog(false);
-	ViewFamily.EngineShowFlags.SetGlobalIllumination(false);
-	ViewFamily.EngineShowFlags.SetEyeAdaptation(false);
-	ViewFamily.EngineShowFlags.SetDirectionalLights(false);
-	ViewFamily.EngineShowFlags.SetPointLights(false);
-	ViewFamily.EngineShowFlags.SetSpotLights(false);
-	ViewFamily.EngineShowFlags.SetRectLights(false);
-
-	ViewFamily.EngineShowFlags.SetDiffuse(false);
-	ViewFamily.EngineShowFlags.SetSpecular(false);
-
-	ViewFamily.EngineShowFlags.SetScreenSpaceReflections(false);
-	ViewFamily.EngineShowFlags.SetLumenReflections(false);
-	ViewFamily.EngineShowFlags.SetRefraction(false);
-	ViewFamily.EngineShowFlags.SetReflectionEnvironment(false);
-
-	ViewFamily.EngineShowFlags.SetDynamicShadows(false);
-	ViewFamily.EngineShowFlags.SetCapsuleShadows(false);
-	ViewFamily.EngineShowFlags.SetContactShadows(false);
-
-	// some of these settings may need to be configured in more complex cases, currently untested
-	//ViewFamily.SceneCaptureSource = SCS_FinalColorHDR;
-	//ViewFamily.bWorldIsPaused = true;
-	//ViewFamily.ViewMode = ViewModeIndex;
-	// since we are capturing actual render we are not using visualize-buffer mode (?)
-	//ViewFamily.EngineShowFlags.SetVisualizeBuffer(false);
-
-	//ViewFamily.EngineShowFlags.SetScreenPercentage(false);
-	ViewFamily.SetScreenPercentageInterface(new FLegacyScreenPercentageDriver(ViewFamily, 1.f, false));
-
-	// This is called in various other places, unclear if we should be doing this too
-	//EngineShowFlagOverride(EShowFlagInitMode::ESFIM_Game, ViewModeIndex, ViewFamily.EngineShowFlags, true);
-
-	FSceneViewInitOptions ViewInitOptions;
-	ViewInitOptions.SetViewRectangle(FIntRect(0, 0, Width, Height));
-	ViewInitOptions.ViewFamily = &ViewFamily;
-	if (VisiblePrimitives.Num() > 0)
-	{
-		ViewInitOptions.ShowOnlyPrimitives = VisiblePrimitives;
-	}
-	ViewInitOptions.ViewOrigin = ViewOrigin;
-	ViewInitOptions.ViewRotationMatrix = ViewRotationMatrix;
-	ViewInitOptions.ProjectionMatrix = ProjectionMatrix;
-	ViewInitOptions.FOV = HorzFOVDegrees;
-
-	FSceneView* NewView = new FSceneView(ViewInitOptions);
-	ViewFamily.Views.Add(NewView);
-
-	NewView->StartFinalPostprocessSettings(ViewInitOptions.ViewOrigin);
-	NewView->EndFinalPostprocessSettings(ViewInitOptions);
-
-	// other FSceneView settings that may need configuring to properly capture Emissive - needs testing
-	//NewView->SetupRayTracedRendering();
-	//NewView->bIsOfflineRender = true;
-	//NewView->bIsSceneCapture = false;
-	//NewView->AntiAliasingMethod = IsAntiAliasingSupported() ? InOutSampleState.AntiAliasingMethod : EAntiAliasingMethod::AAM_None;
-	//NewView->AntiAliasingMethod = EAntiAliasingMethod::AAM_None;
-
-	// can we fully disable auto-exposure?
-	NewView->FinalPostProcessSettings.AutoExposureMethod = EAutoExposureMethod::AEM_Manual;
-
-	// do we actually need for force SM5 here? the other FCanvas() constructor below does not pass these flags...
-	FCanvas Canvas = FCanvas(RenderTargetResource, nullptr, this->World, ERHIFeatureLevel::SM5, FCanvas::CDM_DeferDrawing, 1.0f);
-	Canvas.Clear(FLinearColor::Transparent);
-	GetRendererModule().BeginRenderingViewFamily(&Canvas, &ViewFamily);
-
-	ReadImageBuffer.SetNumUninitialized(Width * Height);
-	FReadSurfaceDataFlags ReadSurfaceDataFlags;
-	ReadSurfaceDataFlags.SetLinearToGamma(false);
-	RenderTargetResource->ReadLinearColorPixels(ReadImageBuffer, ReadSurfaceDataFlags, FIntRect(0, 0, Width, Height));
-
-	FlushRenderingCommands();
-
-	ResultImageOut.SetDimensions(Dimensions);
-	for (int32 yi = 0; yi < Height; ++yi)
-	{
-		for (int32 xi = 0; xi < Width; ++xi)
-		{
-			FLinearColor PixelColorf = ReadImageBuffer[yi * Width + xi];
-			PixelColorf.A = 1.0f;		// ?
-			ResultImageOut.SetPixel(FVector2i(xi, yi), FVector4f(PixelColorf));
-		}
-	}
-
-	return true;
-}
-
-
-
 namespace UE 
 {
 namespace Internal
 {
 
+/**
+ * Reads data from FImagePixelData and stores in an output image, with optional ColorTransformFunc
+ */
 static bool ReadPixelDataToImage(TUniquePtr<FImagePixelData>& PixelData, FImageAdapter& ResultImageOut, bool bLinear)
 {
 	int64 SizeInBytes;
@@ -610,7 +452,7 @@
 	// read back image
 	if (PostProcessPassPixelData.IsValid())
 	{
-		UE::Internal::ReadPixelDataToImage(PostProcessPassPixelData, ResultImageOut, true);
+		UE::Internal::ReadPixelDataToImage(PostProcessPassPixelData, ResultImageOut, true );
 		return true;
 	}
 
@@ -620,6 +462,125 @@
 
 
 
+
+bool FWorldRenderCapture::CaptureEmissiveFromPosition(
+	const FFrame3d& Frame,
+	double HorzFOVDegrees,
+	double NearPlaneDist,
+	FImageAdapter& ResultImageOut)
+{
+	UTextureRenderTarget2D* RenderTargetTexture = GetRenderTexture(true);
+	if (ensure(RenderTargetTexture) == false)
+	{
+		return false;
+	}
+	FTextureRenderTargetResource* RenderTargetResource = RenderTargetTexture->GameThread_GetRenderTargetResource();
+
+	int32 Width = Dimensions.GetWidth();
+	int32 Height = Dimensions.GetHeight();
+
+	const FRotator RotationOffsetToViewCenter(0.f, 90.f, 0.f);
+	FQuat ViewOrientation = (FQuat)Frame.Rotation;
+	FMatrix ViewRotationMatrix = FInverseRotationMatrix(ViewOrientation.Rotator());
+	FVector ViewOrigin = (FVector)Frame.Origin;
+
+	// convert to rendering coordinate system
+	ViewRotationMatrix = ViewRotationMatrix * FMatrix(
+		FPlane(0, 0, 1, 0),
+		FPlane(1, 0, 0, 0),
+		FPlane(0, 1, 0, 0),
+		FPlane(0, 0, 0, 1));
+
+	const float HalfFOVRadians = FMath::DegreesToRadians<float>(HorzFOVDegrees) * 0.5f;
+	static_assert((int32)ERHIZBuffer::IsInverted != 0, "Check NearPlane and Projection Matrix");
+	FMatrix ProjectionMatrix = FReversedZPerspectiveMatrix(HalfFOVRadians, 1.0f, 1.0f, (float)NearPlaneDist);
+
+	// unclear if these flags need to be set before creating ViewFamily
+	FEngineShowFlags ShowFlags(ESFIM_Game);
+	ShowFlags.SetAntiAliasing(false);
+	ShowFlags.SetDepthOfField(false);
+	ShowFlags.SetMotionBlur(false);
+	ShowFlags.SetBloom(false);
+	ShowFlags.SetSceneColorFringe(false);
+
+	FSceneViewFamilyContext ViewFamily(
+		FSceneViewFamily::ConstructionValues(RenderTargetResource, World->Scene, ShowFlags)
+		.SetWorldTimes(FApp::GetCurrentTime() - GStartTime, FApp::GetDeltaTime(), FApp::GetCurrentTime() - GStartTime)
+	);
+
+	// To enable visualization mode
+	ViewFamily.EngineShowFlags.SetPostProcessing(true);
+	ViewFamily.EngineShowFlags.SetVisualizeBuffer(true);
+	ViewFamily.EngineShowFlags.SetTonemapper(false);
+	ViewFamily.EngineShowFlags.SetScreenPercentage(false);
+
+	FSceneViewInitOptions ViewInitOptions;
+	ViewInitOptions.SetViewRectangle(FIntRect(0, 0, Width, Height));
+	ViewInitOptions.ViewFamily = &ViewFamily;
+	if (VisiblePrimitives.Num() > 0)
+	{
+		ViewInitOptions.ShowOnlyPrimitives = VisiblePrimitives;
+	}
+	ViewInitOptions.ViewOrigin = ViewOrigin;
+	ViewInitOptions.ViewRotationMatrix = ViewRotationMatrix;
+	ViewInitOptions.ProjectionMatrix = ProjectionMatrix;
+
+	FSceneView* NewView = new FSceneView(ViewInitOptions);
+	NewView->CurrentBufferVisualizationMode = FName("PreTonemapHDRColor");
+	ViewFamily.Views.Add(NewView);
+
+	ViewFamily.SetScreenPercentageInterface(new FLegacyScreenPercentageDriver(ViewFamily, 1.0f, false));
+
+	ViewFamily.EngineShowFlags.DisableAdvancedFeatures();
+	ViewFamily.EngineShowFlags.LOD = 0;
+
+	ViewFamily.EngineShowFlags.SetFog(false);
+	ViewFamily.EngineShowFlags.SetGlobalIllumination(false);
+	ViewFamily.EngineShowFlags.SetDirectionalLights(false);
+	ViewFamily.EngineShowFlags.SetPointLights(false);
+	ViewFamily.EngineShowFlags.SetSpotLights(false);
+	ViewFamily.EngineShowFlags.SetRectLights(false);
+	ViewFamily.EngineShowFlags.SetSkyLighting(false);
+	//ViewFamily.EngineShowFlags.SetLighting(false);		// this breaks things
+
+	ViewFamily.EngineShowFlags.SetDiffuse(false);
+	ViewFamily.EngineShowFlags.SetSpecular(false);
+
+	ViewFamily.EngineShowFlags.SetRefraction(false);
+	ViewFamily.EngineShowFlags.SetReflectionEnvironment(false);
+
+	ViewFamily.EngineShowFlags.SetDynamicShadows(false);
+	ViewFamily.EngineShowFlags.SetCapsuleShadows(false);
+	ViewFamily.EngineShowFlags.SetContactShadows(false);
+
+	FCanvas Canvas(RenderTargetResource, NULL, FApp::GetCurrentTime() - GStartTime, FApp::GetDeltaTime(), FApp::GetCurrentTime() - GStartTime, World->Scene->GetFeatureLevel());
+	Canvas.Clear(FLinearColor::Transparent);
+
+	GetRendererModule().BeginRenderingViewFamily(&Canvas, &ViewFamily);
+
+	// Copy the contents of the remote texture to system memory
+	ReadImageBuffer.Reset();
+	ReadImageBuffer.SetNumUninitialized(Width * Height);
+	FReadSurfaceDataFlags ReadSurfaceDataFlags(RCM_MinMax);		// don't normalize buffer values, we want full HDR range
+	ReadSurfaceDataFlags.SetLinearToGamma(false);
+	RenderTargetResource->ReadLinearColorPixels(ReadImageBuffer, ReadSurfaceDataFlags, FIntRect(0, 0, Width, Height));
+
+	FlushRenderingCommands();
+
+	ResultImageOut.SetDimensions(Dimensions);
+	for (int32 yi = 0; yi < Height; ++yi)
+	{
+		for (int32 xi = 0; xi < Width; ++xi)
+		{
+			FLinearColor PixelColorf = ReadImageBuffer[yi * Width + xi];
+			ResultImageOut.SetPixel(FVector2i(xi, yi), FVector4f(PixelColorf));
+		}
+	}
+
+	return true;
+}
+
+
 /**
  * internal Utility function to render the given Scene to a render target and capture
  * one of the render buffers, defined by VisualizationMode. Not clear where
@@ -683,6 +644,7 @@
 
 		// Copy the contents of the remote texture to system memory
 		OutSamples.SetNumUninitialized(Width * Height);
+		//FReadSurfaceDataFlags ReadSurfaceDataFlags(RCM_MinMax);		// should we use MinMax to avoid normalization?
 		FReadSurfaceDataFlags ReadSurfaceDataFlags;
 		ReadSurfaceDataFlags.SetLinearToGamma(false);
 		RenderTargetResource->ReadLinearColorPixels(OutSamples, ReadSurfaceDataFlags, FIntRect(0, 0, Width, Height));
@@ -745,6 +707,7 @@
 	case ERenderCaptureType::Metallic:		CaptureTypeName = FName("Metallic");	break;
 	case ERenderCaptureType::Specular:		CaptureTypeName = FName("Specular");	break;
 	case ERenderCaptureType::BaseColor:		CaptureTypeName = FName("BaseColor");	break;
+	case ERenderCaptureType::Emissive:		CaptureTypeName = FName("PreTonemapHDRColor");	break;
 	}
 
 	ReadImageBuffer.Reset();
diff --git a/Engine/Plugins/Experimental/MeshModelingToolset/Source/ModelingComponents/Public/AssetUtils/Texture2DBuilder.h b/Engine/Plugins/Experimental/MeshModelingToolset/Source/ModelingComponents/Public/AssetUtils/Texture2DBuilder.h
--- a/Engine/Plugins/Experimental/MeshModelingToolset/Source/ModelingComponents/Public/AssetUtils/Texture2DBuilder.h
+++ b/Engine/Plugins/Experimental/MeshModelingToolset/Source/ModelingComponents/Public/AssetUtils/Texture2DBuilder.h
@@ -26,20 +26,23 @@ namespace Geometry
  * If you have generated a UTexture2D by other means, you can use the static function ::CopyPlatformDataToSourceData() to populate the 
  * Source data from the PlatformData, which is required to save it as a UAsset. 
  */
-class FTexture2DBuilder
+class MODELINGCOMPONENTS_API FTexture2DBuilder
 {
 public:
 
-	/** Supported texture types */
+	/** 
+	 * Supported texture types.
+	 */
 	enum class ETextureType
 	{
-		Color,
-		ColorLinear,
-		Roughness,
-		Metallic,
-		Specular,
-		NormalMap,
-		AmbientOcclusion
+		Color,					// byte-BGRA SRGB
+		ColorLinear,			// byte-BGRA linear
+		Roughness,				// byte-BGRA linear
+		Metallic,				// byte-BGRA linear
+		Specular,				// byte-BGRA linear
+		NormalMap,				// byte-BGRA linear
+		AmbientOcclusion,		// byte-BGRA linear
+		EmissiveHDR				// float-RGBA linear
 	};
 
 protected:
@@ -48,6 +51,8 @@ protected:
 
 	UTexture2D* RawTexture2D = nullptr;
 	FColor* CurrentMipData = nullptr;
+	FFloat16Color* CurrentMipDataFloat16 = nullptr;
+	EPixelFormat CurrentPixelFormat = PF_B8G8R8A8;
 
 public:
 
@@ -66,6 +71,23 @@ public:
 		return BuildType;
 	}
 
+	const EPixelFormat GetTexturePixelFormat() const
+	{
+		return CurrentPixelFormat;
+	}
+
+	/** @return true if texture buffer is FColor/8-bit BGRA */
+	const bool IsByteTexture() const
+	{
+		return CurrentPixelFormat == PF_B8G8R8A8;
+	}
+
+	/** @return true if texture buffer is FFloat16Color/16-bit RGBA  */
+	const bool IsFloat16Texture() const
+	{
+		return CurrentPixelFormat == PF_FloatRGBA;
+	}
+
 	/** @return the internal texture */
 	UTexture2D* GetTexture2D() const
 	{
@@ -76,96 +98,25 @@ public:
 	/**
 	 * Create a new UTexture2D configured with the given BuildType and Dimensions
 	 */
-	bool Initialize(ETextureType BuildTypeIn, FImageDimensions DimensionsIn)
-	{
-		check(DimensionsIn.IsSquare());
-		BuildType = BuildTypeIn;
-		Dimensions = DimensionsIn;
-
-		// create new texture
-		RawTexture2D = UTexture2D::CreateTransient((int32)Dimensions.GetWidth(), (int32)Dimensions.GetHeight(), PF_B8G8R8A8);
-		if (RawTexture2D == nullptr)
-		{
-			return false;
-		}
-
-		if (BuildType == ETextureType::ColorLinear || BuildType == ETextureType::Roughness || BuildType == ETextureType::Metallic || BuildType == ETextureType::Specular)
-		{
-			RawTexture2D->SRGB = false;
-			RawTexture2D->UpdateResource();
-		}
-		else if (BuildType == ETextureType::NormalMap)
-		{
-			RawTexture2D->CompressionSettings = TC_Normalmap;
-			RawTexture2D->SRGB = false;
-			RawTexture2D->LODGroup = TEXTUREGROUP_WorldNormalMap;
-			//RawTexture2D->bFlipGreenChannel = true;
-#if WITH_EDITOR
-			RawTexture2D->MipGenSettings = TMGS_NoMipmaps;
-#endif
-			RawTexture2D->UpdateResource();
-		}
-
-		// lock
-		if (LockForEditing() == false)
-		{
-			return false;
-		}
-
-		if (IsEditable())
-		{
-			Clear();
-		}
-
-		return true;
-	}
+	bool Initialize(ETextureType BuildTypeIn, FImageDimensions DimensionsIn);
 
 
 	/**
 	 * Initialize the builder with an existing UTexture2D
 	 */
-	bool Initialize(UTexture2D* ExistingTexture, ETextureType BuildTypeIn, bool bLockForEditing = true)
-	{
-		if (! ensure(ExistingTexture != nullptr)) return false;
-		if (! ensure(ExistingTexture->PlatformData != nullptr)) return false;
-		if (! ensure(ExistingTexture->PlatformData->Mips.Num() > 0)) return false;
-
-		int32 Width = ExistingTexture->PlatformData->Mips[0].SizeX;
-		int32 Height = ExistingTexture->PlatformData->Mips[0].SizeY;
-		Dimensions = FImageDimensions(Width, Height);
-		BuildType = BuildTypeIn;
-		RawTexture2D = ExistingTexture;
-
-		// lock for editing
-		if (bLockForEditing && LockForEditing() == false)
-		{
-			return false;
-		}
-
-		return true;
-	}
+	bool Initialize(UTexture2D* ExistingTexture, ETextureType BuildTypeIn, bool bLockForEditing = true);
 
 
 	/**
 	 * Lock the Mip 0 buffer for editing
 	 * @return true if lock was successfull. Will return false if already locked.
 	 */
-	bool LockForEditing()
-	{
-		check(RawTexture2D);
-		check(CurrentMipData == nullptr);
-		if (RawTexture2D && CurrentMipData == nullptr)
-		{
-			CurrentMipData = reinterpret_cast<FColor*>(RawTexture2D->PlatformData->Mips[0].BulkData.Lock(LOCK_READ_WRITE));
-			check(CurrentMipData);
-		}
-		return IsEditable();
-	}
+	bool LockForEditing();
 
 	/** @return true if the texture data is currently locked and editable */
 	bool IsEditable() const
 	{
-		return CurrentMipData != nullptr;
+		return CurrentMipData != nullptr || CurrentMipDataFloat16 != nullptr;
 	}
 
 
@@ -174,136 +125,88 @@ public:
 	 * This does not PostEditChange() so any materials using this texture may not be updated, the caller must do that.
 	 * @param bUpdateSourceData if true, UpdateSourceData() is called to copy the current PlatformData to the SourceData
 	 */
-	void Commit(bool bUpdateSourceData = true)
-	{
-		check(RawTexture2D);
-		check(IsEditable());
-		if (RawTexture2D)
-		{
-			if (bUpdateSourceData)
-			{
-				UpdateSourceData();
-			}
-
-			RawTexture2D->PlatformData->Mips[0].BulkData.Unlock();
-			RawTexture2D->UpdateResource();
-
-			CurrentMipData = nullptr;
-		}
-	}
+	void Commit(bool bUpdateSourceData = true);
 
 
 	/**
 	 * Copy the current PlatformData to the UTexture2D Source Data.
 	 * This does not require the texture to be locked for editing, if it is not locked, a read-only lock will be acquired as needed
-	 * @warning currently assumes both buffers are BGRA
+	 * @warning currently assumes both buffers are the same format
 	 */
-	void UpdateSourceData()
-	{
-		// source data only exists in Editor
-#if WITH_EDITOR
-		check(RawTexture2D);
+	void UpdateSourceData();
 
-		bool bIsEditable = IsEditable();
-		const FColor* SourceMipData = nullptr;
-		if (bIsEditable)
-		{
-			SourceMipData = CurrentMipData;
-		}
-		else
-		{
-			SourceMipData = reinterpret_cast<const FColor*>(RawTexture2D->PlatformData->Mips[0].BulkData.Lock(LOCK_READ_ONLY));
-		}
-
-		bool bIsValid = RawTexture2D->Source.IsValid();
+	/**
+	 * Cancel active editing and unlock the texture data
+	 */
+	void Cancel();
 
-		RawTexture2D->Source.Init2DWithMipChain(Dimensions.GetWidth(), Dimensions.GetHeight(), TSF_BGRA8);
 
-		uint8* DestData = RawTexture2D->Source.LockMip(0);
-		FMemory::Memcpy(DestData, SourceMipData, Dimensions.GetWidth() * Dimensions.GetHeight() * 4 * sizeof(uint8));
-		RawTexture2D->Source.UnlockMip(0);
+	/**
+	 * Clear all texels in the current Mip to the clear/default color for the texture build type
+	 */
+	void Clear();
 
-		if (bIsEditable == false)
-		{
-			RawTexture2D->PlatformData->Mips[0].BulkData.Unlock();
-		}
-#endif
-	}
+	/**
+	 * Clear all texels in the current Mip to the given ClearColor
+	 */
+	void Clear(const FColor& ClearColor);
 
+	/**
+	 * Clear all texels in the current Mip to the given ClearColor
+	 */
+	void Clear(const FFloat16Color& ClearColor);
 
-	void Cancel()
-	{
-		bool bIsEditable = IsEditable();
-		if (bIsEditable)
-		{
-			RawTexture2D->PlatformData->Mips[0].BulkData.Unlock();
-			CurrentMipData = nullptr;
-		}
-	}
 
 
 	/**
-	 * Clear all texels in the current Mip to the clear/default color for the texture build type
+	 * Get the FColor texel at the given X/Y coordinates
 	 */
-	void Clear()
+	const FColor& GetTexel(const FVector2i& ImageCoords) const
 	{
-		check(IsEditable());
-		if (IsEditable())
-		{
-			FColor ClearColor = GetClearColor();
-			for (int64 k = 0; k < Dimensions.Num(); ++k)
-			{
-				CurrentMipData[k] = ClearColor;
-			}
-		}
+		checkSlow(IsEditable() && IsByteTexture());
+		int64 UseIndex = Dimensions.GetIndex(ImageCoords);
+		return CurrentMipData[UseIndex];
 	}
 
 
 	/**
-	 * Clear all texels in the current Mip to the given ClearColor
+	 * Get the FColor texel at the given linear index
 	 */
-	void Clear(const FColor& ClearColor)
+	const FColor& GetTexel(int64 LinearIndex) const
 	{
-		check(IsEditable());
-		if (IsEditable())
-		{
-			for (int64 k = 0; k < Dimensions.Num(); ++k)
-			{
-				CurrentMipData[k] = ClearColor;
-			}
-		}
+		checkSlow(IsEditable() && IsByteTexture());
+		return CurrentMipData[LinearIndex];
 	}
 
 
-
 	/**
-	 * Get the texel at the given X/Y coordinates
+	 * Get the FFloat16Color texel at the given X/Y coordinates
 	 */
-	const FColor& GetTexel(const FVector2i& ImageCoords) const
+	const FFloat16Color& GetTexelFloat16(const FVector2i& ImageCoords) const
 	{
-		checkSlow(IsEditable());
+		checkSlow(IsEditable() && IsFloat16Texture());
 		int64 UseIndex = Dimensions.GetIndex(ImageCoords);
-		return CurrentMipData[UseIndex];
+		return CurrentMipDataFloat16[UseIndex];
 	}
 
 
 	/**
-	 * Get the texel at the given linear index
+	 * Get the FFloat16Color texel at the given linear index
 	 */
-	const FColor& GetTexel(int64 LinearIndex) const
+	const FFloat16Color& GetTexelFloat16(int64 LinearIndex) const
 	{
-		checkSlow(IsEditable());
-		return CurrentMipData[LinearIndex];
+		checkSlow(IsEditable() && IsFloat16Texture());
+		return CurrentMipDataFloat16[LinearIndex];
 	}
 
 
+
 	/**
 	 * Set the texel at the given X/Y coordinates to the given FColor
 	 */
 	void SetTexel(const FVector2i& ImageCoords, const FColor& NewValue)
 	{
-		checkSlow(IsEditable());
-		//int64 UseIndex = Dimensions.GetIndexMirrored(ImageCoords, false, false);
+		checkSlow(IsEditable() && IsByteTexture());
 		int64 UseIndex = Dimensions.GetIndex(ImageCoords);
 		CurrentMipData[UseIndex] = NewValue;
 	}
@@ -314,166 +217,109 @@ public:
 	 */
 	void SetTexel(int64 LinearIndex, const FColor& NewValue)
 	{
-		checkSlow(IsEditable());
+		checkSlow(IsEditable() && IsByteTexture());
 		SetTexel(Dimensions.GetCoords(LinearIndex), NewValue);
 	}
 
 
 	/**
-	 * Set the texel at the given linear index to the clear/default color
+	 * Set the texel at the given X/Y coordinates to the given FFloat16Color
 	 */
-	void ClearTexel(int64 LinearIndex)
+	void SetTexel(const FVector2i& ImageCoords, const FFloat16Color& NewValue)
 	{
-		checkSlow(IsEditable());
-		SetTexel(Dimensions.GetCoords(LinearIndex), GetClearColor());
+		checkSlow(IsEditable() && IsFloat16Texture());
+		int64 UseIndex = Dimensions.GetIndex(ImageCoords);
+		CurrentMipDataFloat16[UseIndex] = NewValue;
 	}
 
 
 	/**
-	 * Copy texel value from one linear index to another
+	 * Set the texel at the given linear index to the given FFloat16Color
 	 */
-	void CopyTexel(int64 FromLinearIndex, int64 ToLinearIndex)
+	void SetTexel(int64 LinearIndex, const FFloat16Color& NewValue)
 	{
-		checkSlow(IsEditable());
-		CurrentMipData[ToLinearIndex] = CurrentMipData[FromLinearIndex];
+		checkSlow(IsEditable() && IsFloat16Texture());
+		SetTexel(Dimensions.GetCoords(LinearIndex), NewValue);
 	}
 
 
 	/**
-	 * populate texel values from floating-point SourceImage
+	 * Copy texel value from one linear index to another
 	 */
-	bool Copy(const TImageBuilder<FVector3f>& SourceImage, const bool bConvertToSRGB = false)
+	void CopyTexel(int64 FromLinearIndex, int64 ToLinearIndex)
 	{
-		if (ensure(SourceImage.GetDimensions() == Dimensions) == false)
+		checkSlow(IsEditable());
+		if (IsByteTexture())
 		{
-			return false;
+			CurrentMipData[ToLinearIndex] = CurrentMipData[FromLinearIndex];
 		}
-		int64 Num = Dimensions.Num();
-		for (int32 i = 0; i < Num; ++i)
+		else
 		{
-			FVector3f Pixel = SourceImage.GetPixel(i);
-			Pixel.X = FMathf::Clamp(Pixel.X, 0.0, 1.0);
-			Pixel.Y = FMathf::Clamp(Pixel.Y, 0.0, 1.0);
-			Pixel.Z = FMathf::Clamp(Pixel.Z, 0.0, 1.0);
-			FColor Texel = ToLinearColor(Pixel).ToFColor(bConvertToSRGB);
-			SetTexel(i, Texel);
+			CurrentMipDataFloat16[ToLinearIndex] = CurrentMipDataFloat16[FromLinearIndex];
 		}
-		return true;
 	}
 
+
 	/**
-	 * Populate texel values from floating-point SourceImage.
+	 * populate texel values from floating-point SourceImage
 	 * @param bConvertToSRGB if true, SourceImage is assumed to be Linear, and converted to SRGB for the Texture
 	 */
-	bool Copy(const TImageBuilder<FVector4f>& SourceImage, const bool bConvertToSRGB = false)
-	{
-		if (ensure(SourceImage.GetDimensions() == Dimensions) == false)
-		{
-			return false;
-		}
-		int64 Num = Dimensions.Num();
-		for (int32 i = 0; i < Num; ++i)
-		{
-			FVector4f Pixel = SourceImage.GetPixel(i);
-			Pixel.X = FMathf::Clamp(Pixel.X, 0.0, 1.0);
-			Pixel.Y = FMathf::Clamp(Pixel.Y, 0.0, 1.0);
-			Pixel.Z = FMathf::Clamp(Pixel.Z, 0.0, 1.0);
-			Pixel.W = FMathf::Clamp(Pixel.W, 0.0, 1.0);
-			FColor Texel = ((FLinearColor)Pixel).ToFColor(bConvertToSRGB);
-			SetTexel(i, Texel);
-		}
-		return true;
-	}
-
+	bool Copy(const TImageBuilder<FVector3f>& SourceImage, const bool bConvertToSRGB = false);
 
 	/**
-	 * copy existing texel values to floating-point DestImage. Assumed to be same color space (eg both Linear)
+	 * Populate texel values from floating-point SourceImage.
+	 * @param bConvertToSRGB if true, SourceImage is assumed to be Linear, and converted to SRGB for the Texture
 	 */
-	bool CopyTo(TImageBuilder<FVector4f>& DestImage) const
-	{
-		if (ensure(DestImage.GetDimensions() == Dimensions) == false)
-		{
-			return false;
-		}
-		int64 Num = Dimensions.Num();
-		for (int32 i = 0; i < Num; ++i)
-		{
-			FColor ByteColor = GetTexel(i);
-			FLinearColor FloatColor(ByteColor);
-			DestImage.SetPixel(i, FVector4f(FloatColor));
-		}
-		return true;
-	}
+	bool Copy(const TImageBuilder<FVector4f>& SourceImage, const bool bConvertToSRGB = false);
 
 
+	/**
+	 * copy existing texel values to floating-point DestImage. Assumed to be same color space (eg values are copied without any linear/SRGB conversion)
+	 */
+	bool CopyTo(TImageBuilder<FVector4f>& DestImage) const;
 
 	/**
-	 * @return current locked Mip data. Nullptr if IsEditable() == false.
+	 * @return current locked Mip data. Nullptr if IsEditable() == false. Use IsByteTexture() / IsFloat16Texture() to determine type.
 	 * @warning this point is invalid after the texture is Committed!
 	 */
-	const FColor* GetRawTexelBufferUnsafe() const
+	const void* GetRawTexelBufferUnsafe() const
 	{
 		return CurrentMipData;
 	}
 
-
 	/**
-	 * @return current locked Mip data. Nullptr if IsEditable() == false.
+	 * @return current locked Mip data. Nullptr if IsEditable() == false. Use IsByteTexture() / IsFloat16Texture() to determine type.
 	 * @warning this point is invalid after the texture is Committed!
 	 */
-	FColor* GetRawTexelBufferUnsafe()
+	void* GetRawTexelBufferUnsafe()
 	{
 		return CurrentMipData;
 	}
 
 
 	/**
-	 * @return the default color for the current texture build type
+	 * @return the default FColor for the current texture build type
 	 */
-	const FColor& GetClearColor() const
-	{
-		static const FColor DefaultColor = FColor::Black;
-		static const FColor DefaultRoughness(128, 128, 128);
-		static const FColor DefaultSpecular(100, 100, 100);
-		static const FColor DefaultMetallic(16, 16, 16);
-		static const FColor DefaultNormalColor(128, 128, 255);
-		static const FColor DefaultAOColor = FColor::White;
-
-		switch (BuildType)
-		{
-		default:
-		case ETextureType::Color:
-		case ETextureType::ColorLinear:
-			return DefaultColor;
-		case ETextureType::Roughness:
-			return DefaultRoughness;
-		case ETextureType::Metallic:
-			return DefaultMetallic;
-		case ETextureType::Specular:
-			return DefaultSpecular;
-		case ETextureType::NormalMap:
-			return DefaultNormalColor;
-		case ETextureType::AmbientOcclusion:
-			return DefaultAOColor;
-		}
-	}
+	const FColor& GetClearColor() const;
+
+
+	/**
+	 * @return the default FFloat16Color for the current texture build type
+	 */
+	FFloat16Color GetClearColorFloat16() const;
 
 
 	/**
 	 * Use a FTexture2DBuilder to copy the PlatformData to the UTexture2D Source data, so it can be saved as an Asset
 	 */
-	static bool CopyPlatformDataToSourceData(UTexture2D* Texture, ETextureType TextureType)
-	{
-		FTexture2DBuilder Builder;
-		bool bOK = Builder.Initialize(Texture, TextureType, false);
-		if (bOK)
-		{
-			Builder.UpdateSourceData();
-		}
-		return bOK;
-	}
+	static bool CopyPlatformDataToSourceData(UTexture2D* Texture, ETextureType TextureType);
 
 
+	/**
+	 * Create a new UTexture2D of the given TextureType from the given SourceImage
+	 * @param bConvertToSRGB if true, assumption is SourceImage is Linear and we want output UTexture2D to be SRGB. Only valid for 8-bit textures.
+	 * @param bPopulateSourceData if true, first the PlatformData is initialized, then it is copied to the SourceData. This is necessary if the UTexture2D is going to be stored as an Asset (but can also be done explicitly using Commit())
+	 */
 	template<typename PixelType>
 	static UTexture2D* BuildTextureFromImage(const TImageBuilder<PixelType>& SourceImage, FTexture2DBuilder::ETextureType TextureType, const bool bConvertToSRGB, bool bPopulateSourceData = true)
 	{
