ModelingComponents:
- add ModelingToolTargetUtil, contains helper static functions for interacting with ToolTargets without having to manually cast things/etc
- add StoredMeshSelectionUtil, contains helper functions for interacting with stored selections
- USingleSelectionMeshEditingTool now can automatically grab an available stored selection, if it exists
- Add GroupTopologyStorableSelection::GetGroupIDs()
#rnx

diff --git a/Engine/Plugins/Experimental/MeshModelingToolset/Source/ModelingComponents/Private/BaseTools/SingleSelectionMeshEditingTool.cpp b/Engine/Plugins/Experimental/MeshModelingToolset/Source/ModelingComponents/Private/BaseTools/SingleSelectionMeshEditingTool.cpp
--- a/Engine/Plugins/Experimental/MeshModelingToolset/Source/ModelingComponents/Private/BaseTools/SingleSelectionMeshEditingTool.cpp
+++ b/Engine/Plugins/Experimental/MeshModelingToolset/Source/ModelingComponents/Private/BaseTools/SingleSelectionMeshEditingTool.cpp
@@ -46,6 +46,11 @@ void USingleSelectionMeshEditingToolBuilder::InitializeNewTool(USingleSelectionM
 	{
 		NewTool->SetAssetAPI(AssetAPI);
 	}
+
+	if (WantsInputSelectionIfAvailable() && SceneState.StoredToolSelection != nullptr)
+	{
+		NewTool->SetInputSelection(SceneState.StoredToolSelection);
+	}
 }
 
 
diff --git a/Engine/Plugins/Experimental/MeshModelingToolset/Source/ModelingComponents/Private/ModelingToolTargetUtil.cpp b/Engine/Plugins/Experimental/MeshModelingToolset/Source/ModelingComponents/Private/ModelingToolTargetUtil.cpp
new file mode 100644
--- /dev/null
+++ b/Engine/Plugins/Experimental/MeshModelingToolset/Source/ModelingComponents/Private/ModelingToolTargetUtil.cpp
@@ -0,0 +1,137 @@
+// Copyright Epic Games, Inc. All Rights Reserved.
+
+#pragma once
+
+#include "ModelingToolTargetUtil.h"
+
+#include "ToolTargets/ToolTarget.h"
+#include "TargetInterfaces/MaterialProvider.h"
+#include "TargetInterfaces/MeshDescriptionCommitter.h"
+#include "TargetInterfaces/MeshDescriptionProvider.h"
+#include "TargetInterfaces/PrimitiveComponentBackedTarget.h"
+
+#include "Components/PrimitiveComponent.h"
+
+#include "MeshDescriptionToDynamicMesh.h"
+#include "DynamicMeshToMeshDescription.h"
+
+using namespace UE::Geometry;
+
+UPrimitiveComponent* UE::ToolTarget::GetTargetComponent(UToolTarget* Target)
+{
+	IPrimitiveComponentBackedTarget* TargetComponent = Cast<IPrimitiveComponentBackedTarget>(Target);
+	if (TargetComponent)
+	{
+		return TargetComponent->GetOwnerComponent();
+	}
+	ensure(false);
+	return nullptr;
+}
+
+bool UE::ToolTarget::HideSourceObject(UToolTarget* Target)
+{
+	IPrimitiveComponentBackedTarget* TargetComponent = Cast<IPrimitiveComponentBackedTarget>(Target);
+	if (TargetComponent)
+	{
+		TargetComponent->SetOwnerVisibility(false);
+		return true;
+	}
+	ensure(false);
+	return false;
+}
+
+bool UE::ToolTarget::ShowSourceObject(UToolTarget* Target)
+{
+	IPrimitiveComponentBackedTarget* TargetComponent = Cast<IPrimitiveComponentBackedTarget>(Target);
+	if (TargetComponent)
+	{
+		TargetComponent->SetOwnerVisibility(true);
+		return true;
+	}
+	ensure(false);
+	return false;
+}
+
+
+FTransform3d UE::ToolTarget::GetLocalToWorldTransform(UToolTarget* Target)
+{
+	IPrimitiveComponentBackedTarget* TargetComponent = Cast<IPrimitiveComponentBackedTarget>(Target);
+	if (TargetComponent)
+	{
+		return (FTransform3d)TargetComponent->GetWorldTransform();
+	}
+	ensure(false);
+	return FTransform3d();
+}
+
+FComponentMaterialSet UE::ToolTarget::GetMaterialSet(UToolTarget* Target, bool bPreferAssetMaterials)
+{
+	FComponentMaterialSet MaterialSet;
+	IMaterialProvider* MaterialProvider = Cast<IMaterialProvider>(Target);
+	if (ensure(MaterialProvider))
+	{
+		MaterialProvider->GetMaterialSet(MaterialSet, bPreferAssetMaterials);
+	}
+	return MaterialSet;
+}
+
+
+const FMeshDescription* UE::ToolTarget::GetMeshDescription(UToolTarget* Target)
+{
+	static FMeshDescription EmptyMeshDescription;
+
+	IMeshDescriptionProvider* MeshDescriptionProvider = Cast<IMeshDescriptionProvider>(Target);
+	if (MeshDescriptionProvider)
+	{
+		return MeshDescriptionProvider->GetMeshDescription();
+	}
+	ensure(false);
+	return &EmptyMeshDescription;
+}
+
+
+FDynamicMesh3 UE::ToolTarget::GetDynamicMeshCopy(UToolTarget* Target)
+{
+	IMeshDescriptionProvider* MeshDescriptionProvider = Cast<IMeshDescriptionProvider>(Target);
+	FDynamicMesh3 Mesh(EMeshComponents::FaceGroups);
+	Mesh.EnableAttributes();
+	if (MeshDescriptionProvider)
+	{
+		FMeshDescriptionToDynamicMesh Converter;
+		Converter.Convert(MeshDescriptionProvider->GetMeshDescription(), Mesh);
+		return Mesh;
+	}
+
+	ensure(false);
+	return Mesh;
+}
+
+
+
+UE::ToolTarget::EDynamicMeshUpdateResult UE::ToolTarget::CommitDynamicMeshUVUpdate(UToolTarget* Target, const UE::Geometry::FDynamicMesh3* UpdatedMesh)
+{
+	IMeshDescriptionCommitter* MeshDescriptionCommitter = Cast<IMeshDescriptionCommitter>(Target);
+
+	EDynamicMeshUpdateResult Result = EDynamicMeshUpdateResult::Failed;
+	MeshDescriptionCommitter->CommitMeshDescription([UpdatedMesh, &Result](const IMeshDescriptionCommitter::FCommitterParams& CommitParams)
+	{
+		FMeshDescription* MeshDescription = CommitParams.MeshDescriptionOut;
+
+		bool bVerticesOnly = false;
+		bool bAttributesOnly = true;
+		if (FDynamicMeshToMeshDescription::HaveMatchingElementCounts(UpdatedMesh, MeshDescription, bVerticesOnly, bAttributesOnly))
+		{
+			FDynamicMeshToMeshDescription Converter;
+			Converter.UpdateAttributes(UpdatedMesh, *MeshDescription, false, false, true/*update uvs*/);
+			Result = EDynamicMeshUpdateResult::Ok;
+		}
+		else
+		{
+			// must have been duplicate tris in the mesh description; we can't count on 1-to-1 mapping of TriangleIDs.  Just convert 
+			FDynamicMeshToMeshDescription Converter;
+			Converter.Convert(UpdatedMesh, *MeshDescription);
+			Result = EDynamicMeshUpdateResult::Ok_ForcedFullUpdate;
+		}
+	});
+	return Result;
+}
\ No newline at end of file
diff --git a/Engine/Plugins/Experimental/MeshModelingToolset/Source/ModelingComponents/Private/Selection/StoredMeshSelectionUtil.cpp b/Engine/Plugins/Experimental/MeshModelingToolset/Source/ModelingComponents/Private/Selection/StoredMeshSelectionUtil.cpp
new file mode 100644
--- /dev/null
+++ b/Engine/Plugins/Experimental/MeshModelingToolset/Source/ModelingComponents/Private/Selection/StoredMeshSelectionUtil.cpp
@@ -0,0 +1,85 @@
+// Copyright Epic Games, Inc. All Rights Reserved.
+
+#include "Selection/StoredMeshSelectionUtil.h"
+#include "Selection/GroupTopologyStorableSelection.h"
+#include "TargetInterfaces/PrimitiveComponentBackedTarget.h"
+#include "ToolTargetManager.h"
+
+
+const UInteractiveToolStorableSelection* UE::Geometry::GetCurrentToolInputSelection(const FToolBuilderState& SceneState, UToolTarget* Target)
+{
+	const UInteractiveToolStorableSelection* Selection = SceneState.StoredToolSelection;
+	if (Selection == nullptr)
+	{
+		return nullptr;
+	}
+
+	IPrimitiveComponentBackedTarget* TargetInterface = Cast<IPrimitiveComponentBackedTarget>(Target);
+	if (TargetInterface == nullptr)
+	{
+		return nullptr;
+	}
+
+	const UGroupTopologyStorableSelection* GroupTopoSelection = Cast<UGroupTopologyStorableSelection>(Selection);
+	if (GroupTopoSelection == nullptr || GroupTopoSelection->IdentifyingInfo.ComponentTarget != TargetInterface->GetOwnerComponent())
+	{
+		return nullptr;
+	}
+
+	return Selection;
+}
+
+
+
+bool UE::Geometry::GetStoredSelectionAsTriangles(
+	const UInteractiveToolStorableSelection* Selection,
+	const FDynamicMesh3& Mesh,
+	TArray<int32>& TrianglesOut)
+{
+	const UGroupTopologyStorableSelection* GroupTopoSelection = Cast<UGroupTopologyStorableSelection>(Selection);
+	if (GroupTopoSelection == nullptr)
+	{
+		return false;
+	}
+
+	// don't support UV selection currently - unclear how reproducible this is?
+	if (GroupTopoSelection->IdentifyingInfo.TopologyType == UGroupTopologyStorableSelection::ETopologyType::FUVGroupTopology)
+	{
+		return false;
+	}
+
+	if (GroupTopoSelection->GetGroupIDs().Num() == 0)
+	{
+		return false;
+	}
+
+	if (GroupTopoSelection->IdentifyingInfo.TopologyType == UGroupTopologyStorableSelection::ETopologyType::FGroupTopology)
+	{
+		TSet<int32> SelectedGroups;
+		for (int32 gid : GroupTopoSelection->GetGroupIDs())
+		{
+			SelectedGroups.Add(gid);
+		}
+
+		for (int32 tid : Mesh.TriangleIndicesItr())
+		{
+			int32 gid = Mesh.GetTriangleGroup(tid);
+			if (SelectedGroups.Contains(gid))
+			{
+				TrianglesOut.Add(tid);
+			}
+		}
+	}
+	else if (GroupTopoSelection->IdentifyingInfo.TopologyType == UGroupTopologyStorableSelection::ETopologyType::FTriangleGroupTopology)
+	{
+		for (int gid : GroupTopoSelection->GetGroupIDs())
+		{
+			if (Mesh.IsTriangle(gid))
+			{
+				TrianglesOut.Add(gid);
+			}
+		}
+	}
+
+	return true;
+}
\ No newline at end of file
diff --git a/Engine/Plugins/Experimental/MeshModelingToolset/Source/ModelingComponents/Public/BaseTools/SingleSelectionMeshEditingTool.h b/Engine/Plugins/Experimental/MeshModelingToolset/Source/ModelingComponents/Public/BaseTools/SingleSelectionMeshEditingTool.h
--- a/Engine/Plugins/Experimental/MeshModelingToolset/Source/ModelingComponents/Public/BaseTools/SingleSelectionMeshEditingTool.h
+++ b/Engine/Plugins/Experimental/MeshModelingToolset/Source/ModelingComponents/Public/BaseTools/SingleSelectionMeshEditingTool.h
@@ -4,6 +4,7 @@
 
 #include "SingleSelectionTool.h"
 #include "InteractiveToolBuilder.h"
+#include "InteractiveToolStorableSelection.h"
 #include "SingleSelectionMeshEditingTool.generated.h"
 
 class IAssetGenerationAPI;
@@ -35,9 +36,12 @@ public:
 	/** Called by BuildTool to configure the Tool with the input MeshSource based on the SceneState */
 	virtual void InitializeNewTool(USingleSelectionMeshEditingTool* Tool, const FToolBuilderState& SceneState) const;
 
-	/** @return true if this tool requires an AssetAPI */
+	/** @return true if this Tool requires an AssetAPI */
 	virtual bool RequiresAssetAPI() const { return false; }
 
+	/** @return true if this Tool would like access to an available Input Selection object */
+	virtual bool WantsInputSelectionIfAvailable() const { return false; }
+
 protected:
 	virtual const FToolTargetTypeRequirements& GetTargetRequirements() const override;
 };
@@ -58,4 +62,42 @@ public:
 protected:
 	UWorld* TargetWorld = nullptr;
 	IAssetGenerationAPI* AssetAPI = nullptr;
+
+
+
+	//
+	// Mesh Selection support
+	//
+
+public:
+	/**
+	 * Set a Selection for the Tool to use. This should be called before tool Setup() (ie in the ToolBuilder) 
+	 * to allow the Tool to configure it's behavior based on the Selection (which may or may-not exist).
+	 * If the Tool requires a Selection, this needs to be handled at the Builder level.
+	 */
+	virtual void SetInputSelection(const UInteractiveToolStorableSelection* StoredToolSelectionIn)
+	{
+		InputSelection = StoredToolSelectionIn;
+	}
+
+	/** @return true if an InputSelection is available */
+	virtual bool HasInputSelection() const
+	{
+		return InputSelection != nullptr;
+	}
+
+	/** @return the input Selection, or nullptr if one was not configured */
+	virtual const UInteractiveToolStorableSelection* GetInputSelection() const
+	{
+		return InputSelection;
+	}
+
+protected:
+
+	/**
+	 * (optional) Stored Selection provided on Tool Input. This should never be modified after the Tool's Setup() is called
+	 */
+	UPROPERTY()
+	const UInteractiveToolStorableSelection* InputSelection = nullptr;
+
 };
diff --git a/Engine/Plugins/Experimental/MeshModelingToolset/Source/ModelingComponents/Public/ModelingToolTargetUtil.h b/Engine/Plugins/Experimental/MeshModelingToolset/Source/ModelingComponents/Public/ModelingToolTargetUtil.h
new file mode 100644
--- /dev/null
+++ b/Engine/Plugins/Experimental/MeshModelingToolset/Source/ModelingComponents/Public/ModelingToolTargetUtil.h
@@ -0,0 +1,95 @@
+// Copyright Epic Games, Inc. All Rights Reserved.
+
+#pragma once
+
+#include "CoreMinimal.h"
+#include "DynamicMesh/DynamicMesh3.h"
+#include "TargetInterfaces/MaterialProvider.h" // FComponentMaterialSet
+
+class UToolTarget;
+class UPrimitiveComponent;
+struct FMeshDescription;
+
+
+//
+// UE::ToolTarget:: namespace contains utility/helper functions for interacting with UToolTargets.
+// Generally these are meant to be used by UInteractiveTools to handle standard tasks that would
+// otherwise require each Tool to figure out things like which ToolTargetInterface to cast to, etc.
+// Using these functions ideally avoids all the boilerplate inherent in the ToolTarget system.
+// 
+// However, the cost is that it is not necessarily the most efficient, as each one of these functions
+// may potentially do many repeated Cast<>'s internally. So, use sparingly, or cache the outputs.
+//
+namespace UE
+{
+namespace ToolTarget
+{
+
+/**
+ * @return the UPrimitiveComponent backing a ToolTarget, or nullptr if there is no such Component
+ */
+MODELINGCOMPONENTS_API UPrimitiveComponent* GetTargetComponent(UToolTarget* Target);
+
+/**
+ * Hide the "Source Object" (eg PrimitiveComponent, Actor, etc) backing a ToolTarget
+ * @return true on success
+ */
+MODELINGCOMPONENTS_API bool HideSourceObject(UToolTarget* Target);
+
+/**
+ * Show the "Source Object" (eg PrimitiveComponent, Actor, etc) backing a ToolTarget
+ * @return true on success
+ */
+MODELINGCOMPONENTS_API bool ShowSourceObject(UToolTarget* Target);
+
+/**
+ * @return the local-to-world Transform underlying a ToolTarget, eg the Component or Actor transform
+ */
+MODELINGCOMPONENTS_API UE::Geometry::FTransform3d GetLocalToWorldTransform(UToolTarget* Target);
+
+/**
+ * Fetch the Material Set on the object underlying a ToolTarget. In cases where there are (eg) 
+ * separate Component and Asset material sets, prefers the Component material set
+ * @param bPreferAssetMaterials if true, prefer an Asset material set, if available
+ * @return a valid MaterialSet
+ */
+MODELINGCOMPONENTS_API FComponentMaterialSet GetMaterialSet(UToolTarget* Target, bool bPreferAssetMaterials = false);
+
+
+/**
+ * @return the MeshDescription underlying a ToolTarget, if it has such a mesh. May be generated internally by the ToolTarget. May be nullptr if the Target does not have a mesh.
+ */
+MODELINGCOMPONENTS_API const FMeshDescription* GetMeshDescription(UToolTarget* Target);
+
+/**
+ * Fetch a DynamicMesh3 representing the given ToolTarget. This may be a conersion of the output of GetMeshDescription().
+ * This function returns a copy, so the caller can take ownership of this Mesh.
+ * @return a created DynamicMesh3, which may be empty if the Target doesn't have a mesh 
+ */
+MODELINGCOMPONENTS_API UE::Geometry::FDynamicMesh3 GetDynamicMeshCopy(UToolTarget* Target);
+
+
+/**
+ * EDynamicMeshUpdateResult is returned by functions below that update a ToolTarget with a new Mesh
+ */
+enum class EDynamicMeshUpdateResult
+{
+	/** Update was successful */
+	Ok = 0,
+	/** Update failed */
+	Failed = 1,
+	/** Update was successful, but required that the entire target mesh be replaced, instead of a (requested) partial update */
+	Ok_ForcedFullUpdate = 2
+};
+
+
+/**
+ * Update the UV sets of the ToolTarget's mesh (assuming it has one) based on the provided UpdatedMesh.
+ * @todo: support updating a specific UV set/index, rather than all sets
+ * @return EDynamicMeshUpdateResult::Ok on success, or Ok_ForcedFullUpdate if any dependent mesh topology was modified
+ */
+MODELINGCOMPONENTS_API EDynamicMeshUpdateResult CommitDynamicMeshUVUpdate(UToolTarget* Target, const UE::Geometry::FDynamicMesh3* UpdatedMesh);
+
+
+}  // end namespace ToolTarget
+}  // end namespace UE
\ No newline at end of file
diff --git a/Engine/Plugins/Experimental/MeshModelingToolset/Source/ModelingComponents/Public/Selection/GroupTopologyStorableSelection.h b/Engine/Plugins/Experimental/MeshModelingToolset/Source/ModelingComponents/Public/Selection/GroupTopologyStorableSelection.h
--- a/Engine/Plugins/Experimental/MeshModelingToolset/Source/ModelingComponents/Public/Selection/GroupTopologyStorableSelection.h
+++ b/Engine/Plugins/Experimental/MeshModelingToolset/Source/ModelingComponents/Public/Selection/GroupTopologyStorableSelection.h
@@ -73,6 +73,9 @@ public:
 	 */
 	void ExtractIntoSelectionObject(const FGroupTopology& TopologyIn, FGroupTopologySelection& SelectionOut) const;
 
+	/** @return selected GroupIDs */
+	const TArray<int32>& GetGroupIDs() const { return GroupIDs; }
+
 	bool IsEmpty() const
 	{
 		return CornerVids.IsEmpty() && GroupEdgeRepresentativeVerts.IsEmpty() && GroupIDs.IsEmpty();
diff --git a/Engine/Plugins/Experimental/MeshModelingToolset/Source/ModelingComponents/Public/Selection/StoredMeshSelectionUtil.h b/Engine/Plugins/Experimental/MeshModelingToolset/Source/ModelingComponents/Public/Selection/StoredMeshSelectionUtil.h
new file mode 100644
--- /dev/null
+++ b/Engine/Plugins/Experimental/MeshModelingToolset/Source/ModelingComponents/Public/Selection/StoredMeshSelectionUtil.h
@@ -0,0 +1,40 @@
+// Copyright Epic Games, Inc. All Rights Reserved.
+
+#pragma once
+
+#include "CoreMinimal.h"
+#include "DynamicMesh/DynamicMesh3.h"
+
+struct FToolBuilderState;
+class UInteractiveToolStorableSelection;
+class UToolTarget;
+
+
+//
+// Utility functions for Tool implementations to use to work with Stored Selections
+//
+namespace UE
+{
+namespace Geometry
+{
+
+	/**
+	 * Given a FToolBuilderState, determine if a Stored Selection exists that is applicable to the ToolTarget
+	 * @return Stored Selection if found, or nullptr if not
+	 */
+	MODELINGCOMPONENTS_API const UInteractiveToolStorableSelection* GetCurrentToolInputSelection(const FToolBuilderState& SceneState, UToolTarget* Target);
+
+	/**
+	 * Convert the given Seletion to a list of Triangles of the specified Mesh. 
+	 * This will check the known variants of UInteractiveToolStorableSelection and convert their various selection types
+	 * (eg group topology, UV island, etc) to triangle indices
+	 * @return true if a list of triangles could be created.
+	 */
+	MODELINGCOMPONENTS_API bool GetStoredSelectionAsTriangles(
+		const UInteractiveToolStorableSelection* Selection, 
+		const FDynamicMesh3& Mesh,
+		TArray<int32>& TrianglesOut );
+
+
+}
+}
\ No newline at end of file
