GeometryProcessing:
- extend FImageAdapter to support reading
ModelingComponents::
- Add AssetUtils::ForceVirtualTexturePrefetch(), ::SaveDebugImage() variants to Texture2DUtil
- Add support for debug image writing to WorldRenderCapture and SceneCapturePhotoSet
- FWorldRenderCapture now computes more accurate bounds (visible Components instead of Actors), and does an explicit VirtualTexture prefetch before doing render capture
- FStaticMeshAssetOptions now exposes Nanite settings to allow UStaticMeshes to be created with Nanite already enabled (currently NaniteSettings.PositionPrecision defaults to max)
ApproximateActors:
- IGeometryProcessing_ApproximateActors::FOptions now has settings for enabling Nanite on generated meshes
- ApproximateActorsImpl applies these settings, also enables VT on generated textures if necessary
- FMeshApproximationSettings/FMeshApproximationTool updated w/ new Nanite settings
#rnx

diff --git a/Engine/Plugins/Experimental/MeshModelingToolset/Source/GeometryProcessingAdapters/Private/ApproximateActorsImpl.cpp b/Engine/Plugins/Experimental/MeshModelingToolset/Source/GeometryProcessingAdapters/Private/ApproximateActorsImpl.cpp
--- a/Engine/Plugins/Experimental/MeshModelingToolset/Source/GeometryProcessingAdapters/Private/ApproximateActorsImpl.cpp
+++ b/Engine/Plugins/Experimental/MeshModelingToolset/Source/GeometryProcessingAdapters/Private/ApproximateActorsImpl.cpp
@@ -794,6 +794,17 @@ void FApproximateActorsImpl::GenerateApproximationForActorSet(const TArray<AActo
 			UE::AssetUtils::ConvertToSingleChannel(Texture);
 		}
 
+		// Make sure the texture is a VT if required by the material sampler
+		if (NewMaterial != nullptr)
+		{
+			UTexture* DefaultTexture = nullptr;
+			NewMaterial->GetTextureParameterValue(MaterialParamName, DefaultTexture);
+			if (ensure(DefaultTexture))
+			{
+				Texture->VirtualTextureStreaming = DefaultTexture->VirtualTextureStreaming;
+			}
+		}
+
 		FTexture2DAssetOptions TexOptions;
 		TexOptions.NewAssetPath = BaseTexturePath + TextureTypeSuffix;
 		FTexture2DAssetResults Results;
@@ -865,6 +876,10 @@ UStaticMesh* FApproximateActorsImpl::EmitGeneratedMeshAsset(
 
 	MeshAssetOptions.NewAssetPath = Options.BasePackagePath;
 	MeshAssetOptions.SourceMeshes.DynamicMeshes.Add(FinalMesh);
+
+	MeshAssetOptions.bGenerateNaniteEnabledMesh = Options.bGenerateNaniteEnabledMesh;
+	MeshAssetOptions.NaniteProxyTrianglePercent = Options.NaniteProxyTrianglePercent;
+
 	if (Material)
 	{
 		MeshAssetOptions.AssetMaterials.Add(Material);
@@ -877,11 +892,11 @@ UStaticMesh* FApproximateActorsImpl::EmitGeneratedMeshAsset(
 
 	if (DebugMesh != nullptr)
 	{
-		FStaticMeshAssetOptions DebugMeshAssetOptions;
-		DebugMeshAssetOptions.CollisionType = ECollisionTraceFlag::CTF_UseSimpleAsComplex;
-		DebugMeshAssetOptions.bEnableRecomputeTangents = false;
+		FStaticMeshAssetOptions DebugMeshAssetOptions = MeshAssetOptions;
 		DebugMeshAssetOptions.NewAssetPath = Options.BasePackagePath + TEXT("_DEBUG");
+		DebugMeshAssetOptions.SourceMeshes.DynamicMeshes.Reset(1);
 		DebugMeshAssetOptions.SourceMeshes.DynamicMeshes.Add(DebugMesh);
+
 		FStaticMeshResults DebugMeshAssetOutputs;
 		UE::AssetUtils::CreateStaticMeshAsset(DebugMeshAssetOptions, DebugMeshAssetOutputs);
 	}
diff --git a/Engine/Plugins/Experimental/MeshModelingToolset/Source/ModelingComponents/Private/AssetUtils/Texture2DUtil.cpp b/Engine/Plugins/Experimental/MeshModelingToolset/Source/ModelingComponents/Private/AssetUtils/Texture2DUtil.cpp
--- a/Engine/Plugins/Experimental/MeshModelingToolset/Source/ModelingComponents/Private/AssetUtils/Texture2DUtil.cpp
+++ b/Engine/Plugins/Experimental/MeshModelingToolset/Source/ModelingComponents/Private/AssetUtils/Texture2DUtil.cpp
@@ -1,6 +1,12 @@
 // Copyright Epic Games, Inc. All Rights Reserved. 
 
 #include "AssetUtils/Texture2DUtil.h"
+#include "EngineModule.h"
+#include "RendererInterface.h"
+#include "RenderUtils.h"
+
+#include "Misc/FileHelper.h"
+#include "Misc/Paths.h"
 
 #include "ExplicitUseGeometryMathTypes.h"		// using UE::Geometry::(math types)
 using namespace UE::Geometry;
@@ -205,3 +211,135 @@ bool UE::AssetUtils::ConvertToSingleChannel(UTexture2D* TextureMap)
 #endif
 
 }
+
+
+bool UE::AssetUtils::ForceVirtualTexturePrefetch(FImageDimensions ScreenSpaceDimensions, bool bWaitForPrefetchToComplete)
+{
+	// Prefetch all virtual textures so that we have content available
+	if (UseVirtualTexturing(GMaxRHIFeatureLevel))
+	{
+		const FVector2D ScreenSpaceSize(ScreenSpaceDimensions.GetWidth(), ScreenSpaceDimensions.GetHeight());
+
+		ENQUEUE_RENDER_COMMAND(AssetUtils_ForceVirtualTexturePrefetch)(
+			[ScreenSpaceSize](FRHICommandListImmediate& RHICmdList)
+		{
+			GetRendererModule().RequestVirtualTextureTiles(ScreenSpaceSize, -1);
+			GetRendererModule().LoadPendingVirtualTextureTiles(RHICmdList, GMaxRHIFeatureLevel);
+		});
+
+		if (bWaitForPrefetchToComplete)
+		{
+			FlushRenderingCommands();
+		}
+
+		return true;
+	}
+
+	return false;
+}
+
+
+
+bool UE::AssetUtils::SaveDebugImage(
+	const TArray<FColor>& Pixels,
+	FImageDimensions Dimensions,
+	FString DebugSubfolder,
+	FString FilenameBase,
+	int32 UseFileCounter)
+{
+	static int32 CaptureIndex = 0;
+#if WITH_EDITOR
+	// Save capture result to a file to ease debugging
+	TRACE_CPUPROFILER_EVENT_SCOPE(AssetUtils_SaveDebugImage);
+
+	FString DirectoryPath = FPaths::ConvertRelativePathToFull(FPaths::ProjectIntermediateDir());
+	if (DebugSubfolder.Len() > 0)
+	{
+		DirectoryPath = FPaths::Combine(DirectoryPath, DebugSubfolder);
+	}
+	int32 FileCounter = (UseFileCounter > 0) ? UseFileCounter : CaptureIndex++;
+	FString Filename = FString::Printf(TEXT("%s-%04d.bmp"), *FilenameBase, FileCounter);
+	FString FilePath = FPaths::Combine(DirectoryPath, Filename);
+
+	return FFileHelper::CreateBitmap(*FilePath, Dimensions.GetWidth(), Dimensions.GetHeight(), Pixels.GetData());
+#else
+	return false;
+#endif
+}
+
+
+
+bool UE::AssetUtils::SaveDebugImage(
+	const TArray<FLinearColor>& Pixels,
+	FImageDimensions Dimensions,
+	bool bConvertToSRGB,
+	FString DebugSubfolder,
+	FString FilenameBase,
+	int32 UseFileCounter )
+{
+	static int32 CaptureIndex = 0;
+#if WITH_EDITOR
+	// Save capture result to a file to ease debugging
+	TRACE_CPUPROFILER_EVENT_SCOPE(AssetUtils_SaveDebugImage);
+
+	FString DirectoryPath = FPaths::ConvertRelativePathToFull(FPaths::ProjectIntermediateDir()); 
+	if (DebugSubfolder.Len() > 0)
+	{
+		DirectoryPath = FPaths::Combine(DirectoryPath, DebugSubfolder);
+	}
+	int32 FileCounter = (UseFileCounter > 0) ? UseFileCounter : CaptureIndex++;
+	FString Filename = FString::Printf(TEXT("%s-%04d.bmp"), *FilenameBase, FileCounter);
+	FString FilePath = FPaths::Combine(DirectoryPath, Filename);
+
+	TArray<FColor> ConvertedColor;
+	ConvertedColor.Reserve(Pixels.Num());
+	for (const FLinearColor& LinearColor : Pixels)
+	{
+		ConvertedColor.Add(LinearColor.ToFColor(bConvertToSRGB));
+	}
+
+	return FFileHelper::CreateBitmap(*FilePath, Dimensions.GetWidth(), Dimensions.GetHeight(), ConvertedColor.GetData());
+
+#else
+	return false;
+#endif
+}
+
+
+
+bool UE::AssetUtils::SaveDebugImage(
+	const FImageAdapter& Image,
+	bool bConvertToSRGB,
+	FString DebugSubfolder,
+	FString FilenameBase,
+	int32 UseFileCounter)
+{
+	static int32 CaptureIndex = 0;
+#if WITH_EDITOR
+	// Save capture result to a file to ease debugging
+	TRACE_CPUPROFILER_EVENT_SCOPE(AssetUtils_SaveDebugImage);
+
+	FString DirectoryPath = FPaths::ConvertRelativePathToFull(FPaths::ProjectIntermediateDir());
+	if (DebugSubfolder.Len() > 0)
+	{
+		DirectoryPath = FPaths::Combine(DirectoryPath, DebugSubfolder);
+	}
+	int32 FileCounter = (UseFileCounter > 0) ? UseFileCounter : CaptureIndex++;
+	FString Filename = FString::Printf(TEXT("%s-%04d.bmp"), *FilenameBase, FileCounter);
+	FString FilePath = FPaths::Combine(DirectoryPath, Filename);
+
+	FImageDimensions Dimensions = Image.GetDimensions();
+	int64 N = Dimensions.Num();
+	TArray<FColor> ConvertedColor;
+	ConvertedColor.Reserve(N);
+	for ( int64 i = 0; i < N; ++i )
+	{
+		FLinearColor LinearColor = (FLinearColor)Image.GetPixel(i);
+		ConvertedColor.Add(LinearColor.ToFColor(bConvertToSRGB));
+	}
+
+	return FFileHelper::CreateBitmap(*FilePath, Dimensions.GetWidth(), Dimensions.GetHeight(), ConvertedColor.GetData());
+#else
+	return false;
+#endif
+}
\ No newline at end of file
diff --git a/Engine/Plugins/Experimental/MeshModelingToolset/Source/ModelingComponents/Private/Scene/SceneCapturePhotoSet.cpp b/Engine/Plugins/Experimental/MeshModelingToolset/Source/ModelingComponents/Private/Scene/SceneCapturePhotoSet.cpp
--- a/Engine/Plugins/Experimental/MeshModelingToolset/Source/ModelingComponents/Private/Scene/SceneCapturePhotoSet.cpp
+++ b/Engine/Plugins/Experimental/MeshModelingToolset/Source/ModelingComponents/Private/Scene/SceneCapturePhotoSet.cpp
@@ -108,6 +108,10 @@ void FSceneCapturePhotoSet::AddExteriorCaptures(
 	RenderCapture.SetWorld(TargetWorld);
 	RenderCapture.SetVisibleActors(VisibleActors);
 	RenderCapture.SetDimensions(PhotoDimensions);
+	if (bWriteDebugImages)
+	{
+		RenderCapture.SetEnableWriteDebugImage(true, 0, DebugImagesFolderName);
+	}
 
 	// this tells us origin and radius - could be view-dependent...
 	FSphere RenderSphere = RenderCapture.ComputeContainingRenderSphere(HorizontalFOVDegrees);
@@ -320,4 +324,13 @@ bool FSceneCapturePhotoSet::ComputeSample(
 	}
 
 	return true;
+}
+
+void FSceneCapturePhotoSet::SetEnableWriteDebugImages(bool bEnable, FString FolderName)
+{
+	bWriteDebugImages = bEnable;
+	if (FolderName.Len() > 0)
+	{
+		DebugImagesFolderName = FolderName;
+	}
 }
\ No newline at end of file
diff --git a/Engine/Plugins/Experimental/MeshModelingToolset/Source/ModelingComponents/Private/Scene/WorldRenderCapture.cpp b/Engine/Plugins/Experimental/MeshModelingToolset/Source/ModelingComponents/Private/Scene/WorldRenderCapture.cpp
--- a/Engine/Plugins/Experimental/MeshModelingToolset/Source/ModelingComponents/Private/Scene/WorldRenderCapture.cpp
+++ b/Engine/Plugins/Experimental/MeshModelingToolset/Source/ModelingComponents/Private/Scene/WorldRenderCapture.cpp
@@ -1,6 +1,7 @@
 // Copyright Epic Games, Inc. All Rights Reserved.
 
 #include "Scene/WorldRenderCapture.h"
+#include "AssetUtils/Texture2DUtil.h"
 
 #include "Engine/Canvas.h"
 #include "Engine/TextureRenderTarget2D.h"
@@ -124,12 +125,6 @@ void FWorldRenderCapture::SetVisibleActors(const TArray<AActor*>& Actors)
 	TArray<UActorComponent*> ComponentQueue;
 	for (AActor* Actor : Actors)
 	{
-		FVector ActorOrigin, ActorExtent;
-		Actor->GetActorBounds(false, ActorOrigin, ActorExtent, true);
-		FBoxSphereBounds ActorBounds(ActorOrigin, ActorExtent, ActorExtent.Size());
-		VisibleBounds = (bFirst) ? ActorBounds : (VisibleBounds + ActorBounds);
-		bFirst = false;
-
 		ComponentQueue.Reset();
 		for (UActorComponent* Component : Actor->GetComponents())
 		{
@@ -142,6 +137,11 @@ void FWorldRenderCapture::SetVisibleActors(const TArray<AActor*>& Actors)
 			{
 				UPrimitiveComponent* PrimitiveComponent = Cast<UPrimitiveComponent>(Component);
 				VisiblePrimitives.Add(PrimitiveComponent->ComponentId);
+
+				// Append bounds of visible components only
+				FBoxSphereBounds ComponentBounds = PrimitiveComponent->Bounds;
+				VisibleBounds = (bFirst) ? ComponentBounds : (VisibleBounds + ComponentBounds);
+				bFirst = false;
 			}
 			else if (Cast<UChildActorComponent>(Component) != nullptr)
 			{
@@ -324,6 +324,9 @@ bool FWorldRenderCapture::CaptureMRSFromPosition(
 	check(this->World);
 	FSceneInterface* Scene = this->World->Scene;
 
+	// Prefetch all virtual textures so that we have content available
+	UE::AssetUtils::ForceVirtualTexturePrefetch(Dimensions);
+
 	UTextureRenderTarget2D* RenderTargetTexture = GetRenderTexture(true);
 	if (!ensure(RenderTargetTexture))
 	{
@@ -469,6 +472,9 @@ bool FWorldRenderCapture::CaptureEmissiveFromPosition(
 	double NearPlaneDist,
 	FImageAdapter& ResultImageOut)
 {
+	// Prefetch all virtual textures so that we have content available
+	UE::AssetUtils::ForceVirtualTexturePrefetch(Dimensions);
+
 	UTextureRenderTarget2D* RenderTargetTexture = GetRenderTexture(true);
 	if (ensure(RenderTargetTexture) == false)
 	{
@@ -603,6 +609,9 @@ namespace Internal
 		TArray<FLinearColor>& OutSamples
 	)
 	{
+		// Prefetch all virtual textures so that we have content available
+		UE::AssetUtils::ForceVirtualTexturePrefetch(Dimensions);
+
 		int32 Width = Dimensions.GetWidth();
 		int32 Height = Dimensions.GetHeight();
 		FTextureRenderTargetResource* RenderTargetResource = RenderTargetTexture->GameThread_GetRenderTargetResource();
@@ -665,11 +674,21 @@ bool FWorldRenderCapture::CaptureFromPosition(
 {
 	if (CaptureType == ERenderCaptureType::Emissive)
 	{
-		return CaptureEmissiveFromPosition(ViewFrame, HorzFOVDegrees, NearPlaneDist, ResultImageOut);
+		bool bOK = CaptureEmissiveFromPosition(ViewFrame, HorzFOVDegrees, NearPlaneDist, ResultImageOut);
+		if (bWriteDebugImage)
+		{
+			WriteDebugImage(ResultImageOut, TEXT("Emissive"));
+		}
+		return bOK;
 	}
 	else if (CaptureType == ERenderCaptureType::CombinedMRS)
 	{
-		return CaptureMRSFromPosition(ViewFrame, HorzFOVDegrees, NearPlaneDist, ResultImageOut);
+		bool bOK = CaptureMRSFromPosition(ViewFrame, HorzFOVDegrees, NearPlaneDist, ResultImageOut);
+		if (bWriteDebugImage)
+		{
+			WriteDebugImage(ResultImageOut, TEXT("CombinedMRS"));
+		}
+		return bOK;
 	}
 
 	// Roughness visualization is rendered with gamma correction (unclear why)
@@ -725,13 +744,37 @@ bool FWorldRenderCapture::CaptureFromPosition(
 	{
 		for (int32 xi = 0; xi < Width; ++xi)
 		{
-			//FColor PixelColor = OutputImage[yi * Width + xi];
-			//FLinearColor PixelColorf = PixelColor.ReinterpretAsLinear();
 			FLinearColor PixelColorf = ReadImageBuffer[yi * Width + xi];
 			PixelColorf.A = 1.0f;		// ?
 			ResultImageOut.SetPixel(FVector2i(xi, yi), FVector4f(PixelColorf));
 		}
 	}
 
+	if (bWriteDebugImage)
+	{
+		WriteDebugImage(ResultImageOut, *CaptureTypeName.ToString());
+	}
+
 	return true;
 }
+
+
+void FWorldRenderCapture::SetEnableWriteDebugImage(bool bEnable, int32 ImageCounter, FString FolderName)
+{
+	bWriteDebugImage = bEnable;
+	if (ImageCounter > 0)
+	{
+		DebugImageCounter = ImageCounter;
+	}
+	if (FolderName.Len() > 0)
+	{
+		DebugImageFolderName = FolderName;
+	}
+}
+
+void FWorldRenderCapture::WriteDebugImage(const FImageAdapter& ResultImageOut, const FString& ImageTypeName)
+{
+	static int32 CaptureIndex = 0;
+	int32 UseCounter = (DebugImageCounter >= 0) ? DebugImageCounter : CaptureIndex++;
+	UE::AssetUtils::SaveDebugImage(ResultImageOut, false, DebugImageFolderName, *ImageTypeName, UseCounter);
+}
\ No newline at end of file
diff --git a/Engine/Plugins/Experimental/MeshModelingToolset/Source/ModelingComponents/Public/AssetUtils/Texture2DUtil.h b/Engine/Plugins/Experimental/MeshModelingToolset/Source/ModelingComponents/Public/AssetUtils/Texture2DUtil.h
--- a/Engine/Plugins/Experimental/MeshModelingToolset/Source/ModelingComponents/Public/AssetUtils/Texture2DUtil.h
+++ b/Engine/Plugins/Experimental/MeshModelingToolset/Source/ModelingComponents/Public/AssetUtils/Texture2DUtil.h
@@ -6,7 +6,7 @@
 #include "VectorTypes.h"
 #include "Image/ImageDimensions.h"
 #include "Image/ImageBuilder.h"
-#include "Engine/Classes/Engine/Texture2D.h"
+#include "Engine/Texture2D.h"
 
 
 
@@ -16,14 +16,64 @@ namespace AssetUtils
 {
 	using namespace UE::Geometry;
 
+	/**
+	 * Extract the Dimensions and image pixels from an input TextureMap
+	 * By default, the "Source" texture data is read. This is only available in-Editor. At runtime, only "Platform" data is available.
+	 * @param bPreferPlatformData if true, platform data will be returned even in-Editor.
+	 * @return true on success
+	 */
 	MODELINGCOMPONENTS_API bool ReadTexture(
 		UTexture2D* TextureMap,
 		FImageDimensions& DimensionsOut,
 		TImageBuilder<FVector4f>& DestImageOut,
 		bool bPreferPlatformData = false);
 
+	/**
+	 * Convert input UTexture2D to single-channel. Assumption is it has more than one channel. Red channel is used.
+	 * @return true on success
+	 */
 	MODELINGCOMPONENTS_API bool ConvertToSingleChannel(UTexture2D* TextureMap);
 
+	/**
+	 * Issue requests to the render thread to force virtual textures to load for the given screen dimensions.
+	 * @param bWaitForPrefetchToComplete if true, FlushRenderingCommands() is called to wait for the textures to finish loading
+	 * @return true on success
+	 */
+	MODELINGCOMPONENTS_API bool ForceVirtualTexturePrefetch(FImageDimensions ScreenSpaceDimensions, bool bWaitForPrefetchToComplete = true);
+
+	/**
+	 * Save image stored in Pixels, of given Dimensions to <Project>/Intermediate/DebugSubFolder/FilenameBase_<FileCounter>.bmp
+	 * If UseFileCounter is not specified, an internal static counter that is incremented each call is used.
+	 */
+	MODELINGCOMPONENTS_API bool SaveDebugImage(
+		const TArray<FColor>& Pixels,
+		FImageDimensions Dimensions,
+		FString DebugSubfolder,
+		FString FilenameBase,
+		int32 UseFileCounter = -1);
+
+	/**
+	 * Save image stored in Pixels, of given Dimensions to <Project>/Intermediate/DebugSubFolder/FilenameBase_<FileCounter>.bmp
+	 * If UseFileCounter is not specified, an internal static counter that is incremented each call is used.
+	 */
+	MODELINGCOMPONENTS_API bool SaveDebugImage(
+		const TArray<FLinearColor>& Pixels,
+		FImageDimensions Dimensions,
+		bool bConvertToSRGB,
+		FString DebugSubfolder,
+		FString FilenameBase,
+		int32 UseFileCounter = -1);
+
+	/**
+	 * Save Image to <Project>/Intermediate/DebugSubFolder/FilenameBase_<FileCounter>.bmp
+	 * If UseFileCounter is not specified, an internal static counter that is incremented each call is used.
+	 */
+	MODELINGCOMPONENTS_API bool SaveDebugImage(
+		const FImageAdapter& Image,
+		bool bConvertToSRGB,
+		FString DebugSubfolder,
+		FString FilenameBase,
+		int32 UseFileCounter = -1);
 }
 }
 
diff --git a/Engine/Plugins/Experimental/MeshModelingToolset/Source/ModelingComponents/Public/Scene/SceneCapturePhotoSet.h b/Engine/Plugins/Experimental/MeshModelingToolset/Source/ModelingComponents/Public/Scene/SceneCapturePhotoSet.h
--- a/Engine/Plugins/Experimental/MeshModelingToolset/Source/ModelingComponents/Public/Scene/SceneCapturePhotoSet.h
+++ b/Engine/Plugins/Experimental/MeshModelingToolset/Source/ModelingComponents/Public/Scene/SceneCapturePhotoSet.h
@@ -121,6 +121,13 @@ public:
 	const FSpatialPhotoSet3f& GetWorldNormalPhotoSet() { return WorldNormalPhotoSet; }
 	const FSpatialPhotoSet3f& GetEmissivePhotoSet() { return EmissivePhotoSet; }
 
+	/**
+	 * Enable debug image writing. All captured images will be written to <Project>/Intermediate/<FolderName>.
+	 * If FolderName is not specified, "SceneCapturePhotoSet" is used by default.
+	 * See FWorldRenderCapture::SetEnableWriteDebugImage() for more details
+	 */
+	void SetEnableWriteDebugImages(bool bEnable, FString FolderName = FString());
+
 
 protected:
 	UWorld* TargetWorld = nullptr;
@@ -141,6 +148,9 @@ protected:
 	FSpatialPhotoSet3f PackedMRSPhotoSet;
 	FSpatialPhotoSet3f WorldNormalPhotoSet;
 	FSpatialPhotoSet3f EmissivePhotoSet;
+
+	bool bWriteDebugImages;
+	FString DebugImagesFolderName = TEXT("SceneCapturePhotoSet");
 };
 
 
diff --git a/Engine/Plugins/Experimental/MeshModelingToolset/Source/ModelingComponents/Public/Scene/WorldRenderCapture.h b/Engine/Plugins/Experimental/MeshModelingToolset/Source/ModelingComponents/Public/Scene/WorldRenderCapture.h
--- a/Engine/Plugins/Experimental/MeshModelingToolset/Source/ModelingComponents/Public/Scene/WorldRenderCapture.h
+++ b/Engine/Plugins/Experimental/MeshModelingToolset/Source/ModelingComponents/Public/Scene/WorldRenderCapture.h
@@ -110,6 +110,13 @@ public:
 		double NearPlaneDist,
 		FImageAdapter& ResultImageOut);
 
+	/**
+	 * Enable debug image write. The captured image will be written to <Project>/Intermediate/<FolderName>/<CaptureType>_<ImageCounter>.bmp
+	 * If FolderName is not specified, "WorldRenderCapture" is used by default.
+	 * If an ImageCounter is not specified, an internal static counter is used that increments on every write
+	 */
+	void SetEnableWriteDebugImage(bool bEnable, int32 ImageCounter = -1, FString FolderName = FString());
+
 protected:
 	UWorld* World;
 
@@ -141,6 +148,11 @@ protected:
 		double HorzFOVDegrees,
 		double NearPlaneDist,
 		FImageAdapter& ResultImageOut);
+
+	bool bWriteDebugImage = false;
+	int32 DebugImageCounter = -1;
+	FString DebugImageFolderName = TEXT("WorldRenderCapture");
+	void WriteDebugImage(const FImageAdapter& ResultImageOut, const FString& ImageTypeName);
 };
 
 
diff --git a/Engine/Plugins/Experimental/MeshModelingToolset/Source/ModelingComponentsEditorOnly/Private/AssetUtils/CreateStaticMeshUtil.cpp b/Engine/Plugins/Experimental/MeshModelingToolset/Source/ModelingComponentsEditorOnly/Private/AssetUtils/CreateStaticMeshUtil.cpp
--- a/Engine/Plugins/Experimental/MeshModelingToolset/Source/ModelingComponentsEditorOnly/Private/AssetUtils/CreateStaticMeshUtil.cpp
+++ b/Engine/Plugins/Experimental/MeshModelingToolset/Source/ModelingComponentsEditorOnly/Private/AssetUtils/CreateStaticMeshUtil.cpp
@@ -119,6 +119,14 @@ UE::AssetUtils::ECreateStaticMeshResult UE::AssetUtils::CreateStaticMeshAsset(
 		}
 	}
 
+	// Nanite options
+	if (Options.bGenerateNaniteEnabledMesh)
+	{
+		NewStaticMesh->NaniteSettings.bEnabled = Options.bGenerateNaniteEnabledMesh;
+		NewStaticMesh->NaniteSettings.PercentTriangles = Options.NaniteProxyTrianglePercent * 0.01f;
+		NewStaticMesh->NaniteSettings.PositionPrecision = MIN_int32;
+	}
+
 	NewStaticMesh->MarkPackageDirty();
 	if (Options.bDeferPostEditChange == false)
 	{
diff --git a/Engine/Plugins/Experimental/MeshModelingToolset/Source/ModelingComponentsEditorOnly/Public/AssetUtils/CreateStaticMeshUtil.h b/Engine/Plugins/Experimental/MeshModelingToolset/Source/ModelingComponentsEditorOnly/Public/AssetUtils/CreateStaticMeshUtil.h
--- a/Engine/Plugins/Experimental/MeshModelingToolset/Source/ModelingComponentsEditorOnly/Public/AssetUtils/CreateStaticMeshUtil.h
+++ b/Engine/Plugins/Experimental/MeshModelingToolset/Source/ModelingComponentsEditorOnly/Public/AssetUtils/CreateStaticMeshUtil.h
@@ -62,6 +62,11 @@ namespace AssetUtils
 		// Controls whether or not the RecomputeTangents option will be enabled on the Asset
 		bool bEnableRecomputeTangents = true;
 
+		// Whether to generate a nanite-enabled mesh
+		bool bGenerateNaniteEnabledMesh = false;
+		// Percentage of triangles to reduce down to for generating a coarse proxy mesh from the Nanite mesh
+		float NaniteProxyTrianglePercent = 100;
+
 		// Controls whether the UBodySetup on the Asset will be created (generally should be true)
 		bool bCreatePhysicsBody = true;
 		// set asset collision type
diff --git a/Engine/Source/Developer/GeometryProcessingInterfaces/Public/GeometryProcessingInterfaces/ApproximateActors.h b/Engine/Source/Developer/GeometryProcessingInterfaces/Public/GeometryProcessingInterfaces/ApproximateActors.h
--- a/Engine/Source/Developer/GeometryProcessingInterfaces/Public/GeometryProcessingInterfaces/ApproximateActors.h
+++ b/Engine/Source/Developer/GeometryProcessingInterfaces/Public/GeometryProcessingInterfaces/ApproximateActors.h
@@ -149,6 +149,11 @@ public:
 		// create a flattened (ie non-instanced) mesh and save it with _DEBUG suffix. Warning often absolutely enormous!
 		bool bWriteDebugMesh = false;
 
+		//
+		// Nanite settings
+		//
+		bool bGenerateNaniteEnabledMesh = false;
+		float NaniteProxyTrianglePercent = 100;
 
 		// todo
 		//   - which texture types to generate (BaseColor, Specular, WorldNormal, etc)
diff --git a/Engine/Source/Editor/MergeActors/Private/MeshApproximationTool/MeshApproximationTool.cpp b/Engine/Source/Editor/MergeActors/Private/MeshApproximationTool/MeshApproximationTool.cpp
--- a/Engine/Source/Editor/MergeActors/Private/MeshApproximationTool/MeshApproximationTool.cpp
+++ b/Engine/Source/Editor/MergeActors/Private/MeshApproximationTool/MeshApproximationTool.cpp
@@ -220,6 +220,9 @@ bool FMeshApproximationTool::RunMerge(const FString& PackageName, const TArray<T
 	Options.bVerbose = UseSettings.bPrintDebugMessages;
 	Options.bWriteDebugMesh = UseSettings.bEmitFullDebugMesh;
 
+	Options.bGenerateNaniteEnabledMesh = UseSettings.bGenerateNaniteEnabledMesh;
+	Options.NaniteProxyTrianglePercent = UseSettings.NaniteProxyTrianglePercent;
+
 	// run actor approximation computation
 	IGeometryProcessing_ApproximateActors::FResults Results;
 	ApproxActorsAPI->ApproximateActors(Actors, Options, Results);
diff --git a/Engine/Source/Runtime/Engine/Classes/Engine/MeshMerging.h b/Engine/Source/Runtime/Engine/Classes/Engine/MeshMerging.h
--- a/Engine/Source/Runtime/Engine/Classes/Engine/MeshMerging.h
+++ b/Engine/Source/Runtime/Engine/Classes/Engine/MeshMerging.h
@@ -768,7 +768,13 @@ struct FMeshApproximationSettings
 	UPROPERTY(BlueprintReadWrite, EditAnywhere, Category = SimplifySettings, meta = (ClampMin = "0.0001", EditCondition = "SimplifyMethod == EMeshApproximationSimplificationPolicy::GeometricTolerance"))
 	float GeometricDeviation = 0.1f;
 
+	/** Whether to generate a nanite-enabled mesh */
+	UPROPERTY(BlueprintReadWrite, EditAnywhere, Category = NaniteSettings)
+	bool bGenerateNaniteEnabledMesh = false;
 
+	/** Percentage of triangles to reduce down to for generating a coarse proxy mesh from the Nanite mesh */
+	UPROPERTY(BlueprintReadWrite, EditAnywhere, AdvancedDisplay, Category = NaniteSettings, meta = (EditConditionHides, EditCondition = "bGenerateNaniteEnabledMesh", ClampMin = 0, ClampMax = 100))
+	float NaniteProxyTrianglePercent = 100;
 
 	/** If Value is > 1, Multisample output baked textures by this amount in each direction (eg 4 == 16x supersampling) */
 	UPROPERTY(BlueprintReadWrite, EditAnywhere, Category = MaterialSettings, meta = (ClampMin = "0", ClampMax = "8", UIMin = "0", UIMax = "4"))
@@ -805,18 +811,26 @@ struct FMeshApproximationSettings
 		return OutputType == Other.OutputType
 			&& ApproximationAccuracy == Other.ApproximationAccuracy
 			&& ClampVoxelDimension == Other.ClampVoxelDimension
+			&& bAttemptAutoThickening == Other.bAttemptAutoThickening
+			&& TargetMinThicknessMultiplier == Other.TargetMinThicknessMultiplier
+			&& BaseCapping == Other.BaseCapping
 			&& WindingThreshold == Other.WindingThreshold
 			&& bFillGaps == Other.bFillGaps
 			&& GapDistance == Other.GapDistance
+			&& OcclusionMethod == Other.OcclusionMethod
 			&& SimplifyMethod == Other.SimplifyMethod
 			&& TargetTriCount == Other.TargetTriCount
 			&& TrianglesPerM == Other.TrianglesPerM
-
-			&& MaterialSettings == Other.MaterialSettings
+			&& GeometricDeviation == Other.GeometricDeviation
+			&& bGenerateNaniteEnabledMesh == Other.bGenerateNaniteEnabledMesh
+			&& NaniteProxyTrianglePercent == Other.NaniteProxyTrianglePercent
 			&& MultiSamplingAA == Other.MultiSamplingAA
+			&& RenderCaptureResolution == Other.RenderCaptureResolution
+			&& MaterialSettings == Other.MaterialSettings
 			&& CaptureFieldOfView == Other.CaptureFieldOfView
 			&& NearPlaneDist == Other.NearPlaneDist
-			&& RenderCaptureResolution == Other.RenderCaptureResolution;
+			&& bPrintDebugMessages == Other.bPrintDebugMessages
+			&& bEmitFullDebugMesh == Other.bEmitFullDebugMesh;
 	}
 
 	/** Inequality. */
diff --git a/Engine/Source/Runtime/GeometryCore/Public/Image/ImageBuilder.h b/Engine/Source/Runtime/GeometryCore/Public/Image/ImageBuilder.h
--- a/Engine/Source/Runtime/GeometryCore/Public/Image/ImageBuilder.h
+++ b/Engine/Source/Runtime/GeometryCore/Public/Image/ImageBuilder.h
@@ -310,31 +310,86 @@ public:
 	{
 		switch (ImageType)
 		{
-		case EImageType::Float1:
-			Image1f->SetDimensions(Dimensions); break;
-		case EImageType::Float3:
-			Image3f->SetDimensions(Dimensions); break;
-		case EImageType::Float4:
-			Image4f->SetDimensions(Dimensions); break;
+			case EImageType::Float1:
+				Image1f->SetDimensions(Dimensions); break;
+			case EImageType::Float3:
+				Image3f->SetDimensions(Dimensions); break;
+			case EImageType::Float4:
+				Image4f->SetDimensions(Dimensions); break;
 		}
 	}
 
+	FImageDimensions GetDimensions() const
+	{
+		switch (ImageType)
+		{
+			case EImageType::Float1:
+				return Image1f->GetDimensions();
+			case EImageType::Float3:
+				return Image3f->GetDimensions();
+			case EImageType::Float4:
+				return Image4f->GetDimensions();
+		}
+		return FImageDimensions();
+	}
+
 	void SetPixel(const FVector2i& PixelCoords, const FVector4f& FloatPixel)
 	{
 		switch (ImageType)
 		{
-		case EImageType::Float1:
-			Image1f->SetPixel(PixelCoords, FloatPixel.X); 
-			break;
-		case EImageType::Float3:
-			Image3f->SetPixel(PixelCoords, FVector3f(FloatPixel.X, FloatPixel.Y, FloatPixel.Z)); 
-			break;
-		case EImageType::Float4:
-			Image4f->SetPixel(PixelCoords, FloatPixel); 
-			break;
+			case EImageType::Float1:
+				Image1f->SetPixel(PixelCoords, FloatPixel.X);
+				break;
+			case EImageType::Float3:
+				Image3f->SetPixel(PixelCoords, FVector3f(FloatPixel.X, FloatPixel.Y, FloatPixel.Z));
+				break;
+			case EImageType::Float4:
+				Image4f->SetPixel(PixelCoords, FloatPixel);
+				break;
 		}
 	}
 
+	FVector4f GetPixel(int64 LinearIndex) const
+	{
+		switch (ImageType)
+		{
+			case EImageType::Float1:
+			{
+				float Value = Image1f->GetPixel(LinearIndex);
+				return FVector4f(Value, Value, Value, 1.0f);
+			}
+			case EImageType::Float3:
+			{
+				FVector3f Value = Image3f->GetPixel(LinearIndex);
+				return FVector4f(Value.X, Value.Y, Value.Z, 1.0f);
+			}
+			case EImageType::Float4:
+				return Image4f->GetPixel(LinearIndex);
+		}
+		return FVector4f::One();
+	}
+
+	FVector4f GetPixel(const FVector2i& PixelCoords) const
+	{
+		switch (ImageType)
+		{
+			case EImageType::Float1:
+			{
+				float Value = Image1f->GetPixel(PixelCoords);
+				return FVector4f(Value, Value, Value, 1.0f);
+			}
+			case EImageType::Float3:
+			{
+				FVector3f Value = Image3f->GetPixel(PixelCoords);
+				return FVector4f(Value.X, Value.Y, Value.Z, 1.0f);
+			}
+			case EImageType::Float4:
+				return Image4f->GetPixel(PixelCoords);
+		}
+		return FVector4f::One();
+	}
+
+
 protected:
 	EImageType ImageType;
 
