ModelingTools: fix incorrect branch that was disabling Geometric Tolerance in SimplifyMeshOp
#rnx

diff --git a/Engine/Plugins/Experimental/MeshModelingToolset/Source/ModelingOperatorsEditorOnly/Private/CleaningOps/SimplifyMeshOp.cpp b/Engine/Plugins/Experimental/MeshModelingToolset/Source/ModelingOperatorsEditorOnly/Private/CleaningOps/SimplifyMeshOp.cpp
--- a/Engine/Plugins/Experimental/MeshModelingToolset/Source/ModelingOperatorsEditorOnly/Private/CleaningOps/SimplifyMeshOp.cpp
+++ b/Engine/Plugins/Experimental/MeshModelingToolset/Source/ModelingOperatorsEditorOnly/Private/CleaningOps/SimplifyMeshOp.cpp
@@ -70,12 +70,15 @@ void ComputeSimplify(FDynamicMesh3* TargetMesh, const bool bReproject,
 	Reducer.GroupBoundaryConstraint = GroupBoundaryConstraint;
 	Reducer.MaterialBoundaryConstraint = MaterialBoundaryConstraint;
 	
-	// MinimalPlanar is a special path that ignores the collapse mode (and quadric error ), so don't bother setting it.
-	if (TargetMode != ESimplifyTargetType::MinimalPlanar)
+	if (TargetMode == ESimplifyTargetType::MinimalPlanar)
 	{
-		Reducer.CollapseMode = CollapseMode;
+		Reducer.CollapseMode = SimplificationType::ESimplificationCollapseModes::AverageVertexPosition;
 		GeometricTolerance = 0;		// MinimalPlanar does not allow vertices to move off the input surface
 	}
+	else
+	{
+		Reducer.CollapseMode = CollapseMode;
+	}
 
 	// use projection target if we are reprojecting or doing geometric error checking
 	FMeshProjectionTarget ProjTarget(&OriginalMesh, &OriginalMeshSpatial);
