GeometryProcessing: Improve cylinder projection in FDynamicMeshUVEditor.
ModelingTools: Replace local code in UVProjectionOp with calls to FDynamicMeshUVEditor. Major rewrite of UVProjectionTool to support ExpMap projection, improve property handling, support world-space dimensions, support re-using previous projection to simplify world-space tiling, support input selections coming from PolyEdit/TriEdit, support ctrl+click and ctrl+shift+click interactions to reposition projection primitive.
#rnx

diff --git a/Engine/Plugins/Experimental/GeometryProcessing/Source/DynamicMesh/Private/Parameterization/DynamicMeshUVEditor.cpp b/Engine/Plugins/Experimental/GeometryProcessing/Source/DynamicMesh/Private/Parameterization/DynamicMeshUVEditor.cpp
--- a/Engine/Plugins/Experimental/GeometryProcessing/Source/DynamicMesh/Private/Parameterization/DynamicMeshUVEditor.cpp
+++ b/Engine/Plugins/Experimental/GeometryProcessing/Source/DynamicMesh/Private/Parameterization/DynamicMeshUVEditor.cpp
@@ -866,6 +866,7 @@ void FDynamicMeshUVEditor::SetTriangleUVsFromCylinderProjection(
 	TFunctionRef<FVector3d(const FVector3d&)> PointTransform, 
 	const FFrame3d& BoxFrame, 
 	const FVector3d& BoxDimensions, 
+	float CylinderSplitAngle,
 	FUVEditResult* Result)
 {
 	if (ensure(UVOverlay) == false) return;
@@ -890,11 +891,10 @@ void FDynamicMeshUVEditor::SetTriangleUVsFromCylinderProjection(
 	double ScaleZ = (FMathd::Abs(BoxDimensions.Z) > FMathf::ZeroTolerance) ? (1.0 / BoxDimensions.Z) : 1.0;
 	FVector3d Scale(ScaleX, ScaleY, ScaleZ);
 
-	double w = FMathd::Sqrt(BoxDimensions.X*BoxDimensions.X + BoxDimensions.Y*BoxDimensions.Y);
-	double h = BoxDimensions.Z;
-	double Angle = FMathd::Atan2(h, w);
-	double DotThresholdRejectFromPlane = FMathd::Cos(Angle);
+	double DotThresholdRejectFromPlane = FMathd::Cos(CylinderSplitAngle * FMathf::DegToRad);
 
+	// sort triangles into buckets based on normal. 1/0 is +/-Z, and 3/4 is negative/positive angle around the cylinder,
+	// where angles range from [-180,180]. Currently we split at 0 so the 3=[-180,0] and 4=[0,180] spans get their own UV islands
 	TArray<FVector3d> TriNormals;
 	TArray<FIndex2i> TriangleCylinderAssignments;
 	TriNormals.SetNum(NumTriangles);
@@ -960,16 +960,16 @@ void FDynamicMeshUVEditor::SetTriangleUVsFromCylinderProjection(
 				}
 				else
 				{
-					double VAngle = FMathd::Atan2Positive(BoxPos.Y, BoxPos.X);
-					if (Bucket == 4 && VAngle < FMathd::HalfPi)
+					double VAngle = FMathd::Atan2(BoxPos.Y, BoxPos.X);
+					if (Bucket == 4 && VAngle < -FMathd::HalfPi)				// 4 = [0, 180]
 					{
 						VAngle += FMathd::TwoPi;
 					}
-					else if (Bucket == 3 && VAngle > FMathd::HalfPi)
+					else if (Bucket == 3 && VAngle > FMathd::HalfPi)		// 3=[-180,0]
 					{
 						VAngle -= FMathd::TwoPi;
 					}
-					UV = FVector2f(-(float(VAngle) * FMathf::InvPi - 1.0f), -float(BoxPos.Z));
+					UV = FVector2f( -(float(VAngle) * FMathf::InvPi - 1.0f), -float(BoxPos.Z));
 				}
 
 				ElemTri[j] = UVOverlay->AppendElement(UV);
diff --git a/Engine/Plugins/Experimental/GeometryProcessing/Source/DynamicMesh/Public/Parameterization/DynamicMeshUVEditor.h b/Engine/Plugins/Experimental/GeometryProcessing/Source/DynamicMesh/Public/Parameterization/DynamicMeshUVEditor.h
--- a/Engine/Plugins/Experimental/GeometryProcessing/Source/DynamicMesh/Public/Parameterization/DynamicMeshUVEditor.h
+++ b/Engine/Plugins/Experimental/GeometryProcessing/Source/DynamicMesh/Public/Parameterization/DynamicMeshUVEditor.h
@@ -138,11 +138,11 @@ public:
 
 
 	/**
-	 * Set UVs by box projection. Triangles will be grouped to "best" box face
+	 * Set UVs by cylinder projection. Triangles will be grouped to cylinder endcaps or side based on CylinderSplitAngle.
 	 * PointTransform is applied to points before projectiong onto ProjectionFrame X/Y axes
 	 * Projected U/V coordinates are divided by Dimensions.X/Y
 	 */
-	void SetTriangleUVsFromCylinderProjection(const TArray<int32>& Triangles, TFunctionRef<FVector3d(const FVector3d&)> PointTransform, const FFrame3d& CylFrame, const FVector3d& CylDimensions, FUVEditResult* Result = nullptr);
+	void SetTriangleUVsFromCylinderProjection(const TArray<int32>& Triangles, TFunctionRef<FVector3d(const FVector3d&)> PointTransform, const FFrame3d& CylFrame, const FVector3d& CylDimensions, float CylinderSplitAngle, FUVEditResult* Result = nullptr);
 
 
 
diff --git a/Engine/Plugins/Experimental/MeshModelingToolset/Source/MeshModelingTools/Private/UVProjectionTool.cpp b/Engine/Plugins/Experimental/MeshModelingToolset/Source/MeshModelingTools/Private/UVProjectionTool.cpp
--- a/Engine/Plugins/Experimental/MeshModelingToolset/Source/MeshModelingTools/Private/UVProjectionTool.cpp
+++ b/Engine/Plugins/Experimental/MeshModelingToolset/Source/MeshModelingTools/Private/UVProjectionTool.cpp
@@ -2,30 +2,23 @@
 
 #include "UVProjectionTool.h"
 #include "InteractiveToolManager.h"
+#include "InteractiveGizmoManager.h"
+#include "ToolTargetManager.h"
 #include "ToolBuilderUtil.h"
-
 #include "ToolSetupUtil.h"
+#include "BaseBehaviors/SingleClickBehavior.h"
 
-#include "DynamicMesh3.h"
-#include "BaseBehaviors/MultiClickSequenceInputBehavior.h"
-#include "Selection/SelectClickedAction.h"
-
-#include "MeshDescriptionToDynamicMesh.h"
-#include "DynamicMeshToMeshDescription.h"
-#include "Materials/MaterialInstanceDynamic.h"
-
-#include "InteractiveGizmoManager.h"
-
-#include "AssetGenerationUtil.h"
-
-#include "BaseGizmos/GizmoComponents.h"
 #include "BaseGizmos/TransformGizmoUtil.h"
+#include "ModelingToolTargetUtil.h"
+#include "Selection/StoredMeshSelectionUtil.h"
 
-#include "TargetInterfaces/MaterialProvider.h"
-#include "TargetInterfaces/MeshDescriptionCommitter.h"
-#include "TargetInterfaces/MeshDescriptionProvider.h"
-#include "TargetInterfaces/PrimitiveComponentBackedTarget.h"
-#include "ToolTargetManager.h"
+#include "DynamicMesh/DynamicMesh3.h"
+#include "DynamicMesh/DynamicMeshAABBTree3.h"
+#include "DynamicMesh/MeshIndexUtil.h"
+#include "MeshRegionBoundaryLoops.h"
+#include "Parameterization/DynamicMeshUVEditor.h"
+#include "Operations/MeshConvexHull.h"
+#include "MinVolumeBox3.h"
 
 #include "ExplicitUseGeometryMathTypes.h"		// using UE::Geometry::(math types)
 using namespace UE::Geometry;
@@ -36,35 +29,21 @@ using namespace UE::Geometry;
  * ToolBuilder
  */
 
-
-const FToolTargetTypeRequirements& UUVProjectionToolBuilder::GetTargetRequirements() const
+USingleSelectionMeshEditingTool* UUVProjectionToolBuilder::CreateNewTool(const FToolBuilderState& SceneState) const
 {
-	static FToolTargetTypeRequirements TypeRequirements({
-		UMaterialProvider::StaticClass(),
-		UMeshDescriptionCommitter::StaticClass(),
-		UMeshDescriptionProvider::StaticClass(),
-		UPrimitiveComponentBackedTarget::StaticClass()
-		});
-	return TypeRequirements;
+	return NewObject<UUVProjectionTool>(SceneState.ToolManager);
 }
 
-bool UUVProjectionToolBuilder::CanBuildTool(const FToolBuilderState& SceneState) const
-{
-	// note most of the code is written to support working on any number > 0, but that seems maybe confusing UI-wise and is not fully tested, so I've limited it to acting on one for now
-	// TODO: if enable tool working on multiple components, figure out what to do if we have multiple component targets that point to the same underlying mesh data?
-	return SceneState.TargetManager->CountSelectedAndTargetable(SceneState, GetTargetRequirements()) == 1;
-}
+/*
+ * Propert Sets
+ */
 
-UInteractiveTool* UUVProjectionToolBuilder::BuildTool(const FToolBuilderState& SceneState) const
+void UUVProjectionToolEditActions::PostAction(EUVProjectionToolActions Action)
 {
-	UUVProjectionTool* NewTool = NewObject<UUVProjectionTool>(SceneState.ToolManager);
-
-	TArray<TObjectPtr<UToolTarget>> Targets = SceneState.TargetManager->BuildAllSelectedTargetable(SceneState, GetTargetRequirements());
-	NewTool->SetTargets(MoveTemp(Targets));
-	NewTool->SetWorld(SceneState.World, SceneState.GizmoManager);
-	NewTool->SetAssetAPI(AssetAPI);
-
-	return NewTool;
+	if (ParentTool.IsValid())
+	{
+		ParentTool->RequestAction(Action);
+	}
 }
 
 
@@ -74,210 +53,314 @@ UInteractiveTool* UUVProjectionToolBuilder::BuildTool(const FToolBuilderState& S
  * Tool
  */
 
-UUVProjectionToolProperties::UUVProjectionToolProperties()
-{
-	UVProjectionMethod = EUVProjectionMethod::Cube;
-	ProjectionPrimitiveScale = FVector::OneVector;
-	UVScale = FVector2D::UnitVector;
-	UVOffset = FVector2D::ZeroVector;
-}
-
-UUVProjectionAdvancedProperties::UUVProjectionAdvancedProperties()
+void UUVProjectionTool::Setup()
 {
-}
+	UInteractiveTool::Setup();
 
+	UE::ToolTarget::HideSourceObject(Target);
 
-UUVProjectionTool::UUVProjectionTool()
-{
-}
+	BasicProperties = NewObject<UUVProjectionToolProperties>(this);
+	BasicProperties->RestoreProperties(this);
+	BasicProperties->GetOnModified().AddLambda([this](UObject*, FProperty*)	{ Preview->InvalidateResult(); });
+	BasicProperties->WatchProperty(BasicProperties->ProjectionType, [this](EUVProjectionMethod) { OnProjectionTypeChanged(); });
 
-void UUVProjectionTool::SetWorld(UWorld* World, UInteractiveGizmoManager* GizmoManagerIn)
-{
-	this->TargetWorld = World;
-	this->GizmoManager = GizmoManagerIn;
-}
+	ExpMapProperties = NewObject<UUVProjectionToolExpMapProperties>(this);
+	ExpMapProperties->RestoreProperties(this);
+	ExpMapProperties->GetOnModified().AddLambda([this](UObject*, FProperty*) { Preview->InvalidateResult();	});
 
-void UUVProjectionTool::Setup()
-{
-	UInteractiveTool::Setup();
+	CylinderProperties = NewObject<UUVProjectionToolCylinderProperties>(this);
+	CylinderProperties->RestoreProperties(this);
+	CylinderProperties->GetOnModified().AddLambda([this](UObject*, FProperty*) { Preview->InvalidateResult();	});
 
-	// hide input StaticMeshComponent
-	for (int32 ComponentIdx = 0; ComponentIdx < Targets.Num(); ComponentIdx++)
-	{
-		TargetComponentInterface(ComponentIdx)->SetOwnerVisibility(false);
-	}
+	// initialize input mesh
+	InitializeMesh();
 
-	BasicProperties = NewObject<UUVProjectionToolProperties>(this, TEXT("UV Projection Settings"));
-	BasicProperties->RestoreProperties(this);
-	AdvancedProperties = NewObject<UUVProjectionAdvancedProperties>(this, TEXT("Advanced Settings"));
+	// initialize the PreviewMesh+BackgroundCompute object
+	UpdateNumPreviews();
 
-	// initialize our properties
-	AddToolPropertySource(BasicProperties);
-	AddToolPropertySource(AdvancedProperties);
+	UVChannelProperties = NewObject<UMeshUVChannelProperties>(this);
+	UVChannelProperties->RestoreProperties(this);
+	UVChannelProperties->Initialize(InputMesh.Get(), false);
+	UVChannelProperties->ValidateSelection(true);
+	UVChannelProperties->WatchProperty(UVChannelProperties->UVChannel, [this](const FString& NewValue)
+	{
+		MaterialSettings->UVChannel = UVChannelProperties->GetSelectedChannelIndex(true);
+		Preview->InvalidateResult();
+		OnMaterialSettingsChanged();
+	});
 
 	MaterialSettings = NewObject<UExistingMeshMaterialProperties>(this);
 	MaterialSettings->RestoreProperties(this);
+	MaterialSettings->GetOnModified().AddLambda([this](UObject*, FProperty*)
+	{
+		OnMaterialSettingsChanged();
+	});
 
-	AddToolPropertySource(MaterialSettings);
+	EditActions = NewObject<UUVProjectionToolEditActions>(this);
+	EditActions->Initialize(this);
 
-	// initialize the PreviewMesh+BackgroundCompute object
-	UpdateNumPreviews();
+	AddToolPropertySource(UVChannelProperties);
+	AddToolPropertySource(BasicProperties);
+	AddToolPropertySource(ExpMapProperties);
+	AddToolPropertySource(CylinderProperties);
+	AddToolPropertySource(EditActions);
+	AddToolPropertySource(MaterialSettings);
+	
+	OnMaterialSettingsChanged();
+	OnProjectionTypeChanged();
 
 	// set up visualizers
 	ProjectionShapeVisualizer.LineColor = FLinearColor::Red;
 	ProjectionShapeVisualizer.LineThickness = 2.0;
 
+	// initialize Gizmo
+	FTransform3d GizmoPositionWorld(WorldBounds.Center());
+	TransformProxy = NewObject<UTransformProxy>(this);
+	TransformProxy->SetTransform((FTransform)GizmoPositionWorld);
+	TransformProxy->OnTransformChanged.AddUObject(this, &UUVProjectionTool::TransformChanged);
+
+	TransformGizmo = UE::TransformGizmoUtil::CreateCustomTransformGizmo(
+		GetToolManager(), ETransformGizmoSubElements::StandardTranslateRotate, this);
+	TransformGizmo->SetActiveTarget(TransformProxy, GetToolManager());
+
+	InitialDimensions = BasicProperties->Dimensions;
+	InitialTransform = TransformGizmo->GetGizmoTransform();
 	
-	for (UMeshOpPreviewWithBackgroundCompute* Preview : Previews)
+	ApplyInitializationMode();
+
+	// start watching for dimensions changes
+	DimensionsWatcher = BasicProperties->WatchProperty(BasicProperties->Dimensions, [this](FVector) { bTransformModified = true; });
+	BasicProperties->WatchProperty(BasicProperties->Initialization, [this](EUVProjectionToolInitializationMode NewMode) { OnInitializationModeChanged(); });
+	bTransformModified = false;
+	BasicProperties->SilentUpdateWatched();
+
+	// click to set plane behavior
+	SetPlaneCtrlClickBehaviorTarget = MakeUnique<FSelectClickedAction>();
+	SetPlaneCtrlClickBehaviorTarget->World = this->TargetWorld;
+	SetPlaneCtrlClickBehaviorTarget->OnClickedPositionFunc = [this](const FHitResult& Hit)
 	{
-		Preview->InvalidateResult();
-	}
-	UpdateVisualization();
+		UpdatePlaneFromClick(FVector3d(Hit.ImpactPoint), FVector3d(Hit.ImpactNormal), SetPlaneCtrlClickBehaviorTarget->bShiftModifierToggle);
+	};
+	SetPlaneCtrlClickBehaviorTarget->InvisibleComponentsToHitTest.Add(UE::ToolTarget::GetTargetComponent(Target));
+	ClickToSetPlaneBehavior = NewObject<USingleClickInputBehavior>();
+	ClickToSetPlaneBehavior->ModifierCheckFunc = FInputDeviceState::IsCtrlKeyDown;
+	ClickToSetPlaneBehavior->Modifiers.RegisterModifier(FSelectClickedAction::ShiftModifier, FInputDeviceState::IsShiftKeyDown);
+	ClickToSetPlaneBehavior->Initialize(SetPlaneCtrlClickBehaviorTarget.Get());
+	AddInputBehavior(ClickToSetPlaneBehavior);
+
+	// probably already done
+	Preview->InvalidateResult();
 
 	SetToolDisplayName(LOCTEXT("ToolName", "UV Projection"));
 	GetToolManager()->DisplayMessage(
-		LOCTEXT("UVProjectionToolDescription", "Generate UVs for a Mesh by projecting onto simple geometric shapes."),
+		LOCTEXT("UVProjectionToolDescription", "Generate UVs for a Mesh by projecting onto simple geometric shapes. Ctrl+click to reposition shape. Face selections can be made in the PolyEdit and TriEdit Tools."),
 		EToolMessageLevel::UserNotification);
 }
 
 
 void UUVProjectionTool::UpdateNumPreviews()
 {
-	int32 CurrentNumPreview = Previews.Num();
-	int32 TargetNumPreview = Targets.Num();
-	if (TargetNumPreview < CurrentNumPreview)
-	{
-		for (int32 PreviewIdx = CurrentNumPreview - 1; PreviewIdx >= TargetNumPreview; PreviewIdx--)
-		{
-			Previews[PreviewIdx]->Cancel();
-			GizmoManager->DestroyGizmo(TransformGizmos[PreviewIdx]);
-		}
-		Previews.SetNum(TargetNumPreview);
-		TransformGizmos.SetNum(TargetNumPreview);
-		TransformProxies.SetNum(TargetNumPreview);
-		OriginalDynamicMeshes.SetNum(TargetNumPreview);
-	}
-	else
+	OperatorFactory = NewObject<UUVProjectionOperatorFactory>();
+	OperatorFactory->Tool = this;
+
+	Preview = NewObject<UMeshOpPreviewWithBackgroundCompute>(OperatorFactory);
+	Preview->Setup(this->TargetWorld, OperatorFactory);
+	Preview->OnMeshUpdated.AddUObject(this, &UUVProjectionTool::OnMeshUpdated);
+	Preview->PreviewMesh->SetTangentsMode(EDynamicMeshTangentCalcType::AutoCalculated);
+
+	FComponentMaterialSet MaterialSet = UE::ToolTarget::GetMaterialSet(Target, false);
+	Preview->ConfigureMaterials(MaterialSet.Materials,
+		ToolSetupUtil::GetDefaultWorkingMaterial(GetToolManager())
+	);
+	Preview->PreviewMesh->UpdatePreview(InputMesh.Get());
+	Preview->PreviewMesh->SetTransform((FTransform)WorldTransform);
+
+	Preview->SetVisibility(true);
+
+	EdgeRenderer = NewObject<UPreviewGeometry>(this);
+	EdgeRenderer->CreateInWorld(this->TargetWorld, (FTransform)WorldTransform);
+
+	// if we have an ROI, show its borders
+	if (TriangleROI->Num() > 0)
 	{
-		OriginalDynamicMeshes.SetNum(TargetNumPreview);
-		for (int32 PreviewIdx = CurrentNumPreview; PreviewIdx < TargetNumPreview; PreviewIdx++)
+		FDynamicMesh3* UseMesh = InputMesh.Get();
+		FMeshRegionBoundaryLoops BoundaryLoops(UseMesh, *TriangleROI, true);
+		TSet<int32> BorderEdges;
+		for (const FEdgeLoop& Loop : BoundaryLoops.GetLoops())
 		{
-			UUVProjectionOperatorFactory *OpFactory = NewObject<UUVProjectionOperatorFactory>();
-			OpFactory->Tool = this;
-			OpFactory->ComponentIndex = PreviewIdx;
-			OriginalDynamicMeshes[PreviewIdx] = MakeShared<FDynamicMesh3, ESPMode::ThreadSafe>();
-			FMeshDescriptionToDynamicMesh Converter;
-			Converter.Convert(TargetMeshProviderInterface(PreviewIdx)->GetMeshDescription(), *OriginalDynamicMeshes[PreviewIdx]);
-
-			FVector Center, Extents;
-			FBoxSphereBounds Bounds = TargetComponentInterface(PreviewIdx)->GetOwnerComponent()->CalcLocalBounds();
-			
-			FTransform LocalXF(Bounds.Origin);
-			LocalXF.SetScale3D(Bounds.BoxExtent);
-
-			UMeshOpPreviewWithBackgroundCompute* Preview = Previews.Add_GetRef(NewObject<UMeshOpPreviewWithBackgroundCompute>(OpFactory, "Preview"));
-			Preview->Setup(this->TargetWorld, OpFactory);
-			Preview->PreviewMesh->SetTangentsMode(EDynamicMeshTangentCalcType::AutoCalculated);
-
-			FComponentMaterialSet MaterialSet;
-			TargetMaterialInterface(PreviewIdx)->GetMaterialSet(MaterialSet);
-			Preview->ConfigureMaterials(MaterialSet.Materials,
-				ToolSetupUtil::GetDefaultWorkingMaterial(GetToolManager())
-			);
-			Preview->PreviewMesh->UpdatePreview(OriginalDynamicMeshes[PreviewIdx].Get());
-			Preview->PreviewMesh->SetTransform(TargetComponentInterface(PreviewIdx)->GetWorldTransform());
-
-			Preview->SetVisibility(true);
-
-			UTransformProxy* TransformProxy = TransformProxies.Add_GetRef(NewObject<UTransformProxy>(this));
-			TransformProxy->SetTransform(LocalXF * TargetComponentInterface(PreviewIdx)->GetWorldTransform());
-			TransformProxy->OnTransformChanged.AddUObject(this, &UUVProjectionTool::TransformChanged);
-
-			UTransformGizmo* TransformGizmo = TransformGizmos.Add_GetRef(UE::TransformGizmoUtil::Create3AxisTransformGizmo(GizmoManager, this));
-			TransformGizmo->SetActiveTarget(TransformProxy, GetToolManager());
+			BorderEdges.Append(Loop.Edges);
 		}
-		check(TransformProxies.Num() == TargetNumPreview);
-		check(TransformGizmos.Num() == TargetNumPreview);
+
+		const FColor ROIBorderColor(240, 15, 240);
+		const float ROIBorderThickness = 4.0f;
+		const float ROIBorderDepthBias = 0.1f * (float)(WorldBounds.DiagonalLength() * 0.01);
+		EdgeRenderer->CreateOrUpdateLineSet(TEXT("ROIBorders"), UseMesh->MaxEdgeID(), [&](int32 eid, TArray<FRenderableLine>& LinesOut) {
+			if (BorderEdges.Contains(eid))
+			{
+				FVector3d A, B;
+				UseMesh->GetEdgeV(eid, A, B);
+				LinesOut.Add(FRenderableLine((FVector)A, (FVector)B, ROIBorderColor, ROIBorderThickness, ROIBorderDepthBias));
+			}
+		}, 1);
 	}
 }
 
 
 void UUVProjectionTool::Shutdown(EToolShutdownType ShutdownType)
 {
+	UVChannelProperties->SaveProperties(this);
+	BasicProperties->SavedDimensions = BasicProperties->Dimensions;
+	BasicProperties->SavedTransform = TransformGizmo->GetGizmoTransform();
 	BasicProperties->SaveProperties(this);
+	ExpMapProperties->SaveProperties(this);
+	CylinderProperties->SaveProperties(this);
 	MaterialSettings->SaveProperties(this);
 
+	EdgeRenderer->Disconnect();
+
 	// Restore (unhide) the source meshes
-	for (int32 ComponentIdx = 0; ComponentIdx < Targets.Num(); ComponentIdx++)
-	{
-		TargetComponentInterface(ComponentIdx)->SetOwnerVisibility(true);
-	}
+	UE::ToolTarget::ShowSourceObject(Target);
 
-	TArray<FDynamicMeshOpResult> Results;
-	for (UMeshOpPreviewWithBackgroundCompute* Preview : Previews)
-	{
-		Results.Emplace(Preview->Shutdown());
-	}
+	FDynamicMeshOpResult Result = Preview->Shutdown();
 	if (ShutdownType == EToolShutdownType::Accept)
 	{
-		GenerateAsset(Results);
+		GetToolManager()->BeginUndoTransaction(LOCTEXT("UVProjectionToolTransactionName", "Project UVs"));
+		FDynamicMesh3* NewDynamicMesh = Result.Mesh.Get();
+		if (ensure(NewDynamicMesh))
+		{
+			UE::ToolTarget::CommitDynamicMeshUVUpdate(Target, NewDynamicMesh);
+		}
+		GetToolManager()->EndUndoTransaction();
 	}
 
-	GizmoManager->DestroyAllGizmosByOwner(this);
-	TransformGizmos.Empty();
+	GetToolManager()->GetPairedGizmoManager()->DestroyAllGizmosByOwner(this);
+	TransformGizmo = nullptr;
+	TransformProxy = nullptr;
 }
 
-void UUVProjectionTool::SetAssetAPI(IAssetGenerationAPI* AssetAPIIn)
+
+void UUVProjectionTool::RequestAction(EUVProjectionToolActions ActionType)
 {
-	this->AssetAPI = AssetAPIIn;
+	if (bHavePendingAction)
+	{
+		return;
+	}
+	PendingAction = ActionType;
+	bHavePendingAction = true;
 }
 
+
+
 TUniquePtr<FDynamicMeshOperator> UUVProjectionOperatorFactory::MakeNewOperator()
 {
 	TUniquePtr<FUVProjectionOp> Op = MakeUnique<FUVProjectionOp>();
-	Op->ProjectionMethod = Tool->BasicProperties->UVProjectionMethod;
+	Op->ProjectionMethod = Tool->BasicProperties->ProjectionType;
+
+	Op->MeshToProjectionSpace = Tool->WorldTransform;
+	Op->ProjectionBox = Tool->GetProjectionBox();
+	if (Tool->BasicProperties->DimensionMode == EUVProjectionToolDimensionMode::UseFirst)
+	{
+		double UseDim = Op->ProjectionBox.Extents.X;
+		Op->ProjectionBox.Extents = FVector3d(UseDim, UseDim, UseDim);
+	}
+	Op->CylinderSplitAngle = Tool->CylinderProperties->SplitAngle;
+	Op->BlendWeight = Tool->ExpMapProperties->NormalBlending;
+	Op->SmoothingRounds = Tool->ExpMapProperties->SmoothingRounds;
+	Op->SmoothingAlpha = Tool->ExpMapProperties->SmoothingAlpha;
 
-	// TODO: de-dupe this logic (it's also in Render, below)
-	FTransform LocalScale = FTransform::Identity;
-	LocalScale.SetScale3D(Tool->BasicProperties->ProjectionPrimitiveScale);
-	Op->ProjectionTransform = LocalScale * Tool->TransformProxies[ComponentIndex]->GetTransform();
-	Op->CylinderProjectToTopOrBottomAngleThreshold = Tool->BasicProperties->CylinderProjectToTopOrBottomAngleThreshold;
+	Op->UVRotationAngleDeg = Tool->BasicProperties->UVRotation;
 	Op->UVScale = (FVector2f)Tool->BasicProperties->UVScale;
-	Op->UVOffset = (FVector2f)Tool->BasicProperties->UVOffset;
-	Op->bWorldSpaceUVScale = Tool->BasicProperties->bWorldSpaceUVScale;
+	Op->UVTranslate = (FVector2f)Tool->BasicProperties->UVTranslate;
 
-	FTransform LocalToWorld = Tool->TargetComponentInterface(ComponentIndex)->GetWorldTransform();
-	Op->OriginalMesh = Tool->OriginalDynamicMeshes[ComponentIndex];
-	
-	Op->SetTransform(LocalToWorld);
+	Op->OriginalMesh = Tool->InputMesh;
+	Op->TriangleROI = Tool->TriangleROI;
+	Op->UseUVLayer = Tool->UVChannelProperties->GetSelectedChannelIndex(true);
+		
+	Op->SetResultTransform(Tool->WorldTransform);
 
 	return Op;
 }
 
 
 
+
+void UUVProjectionTool::InitializeMesh()
+{
+	InputMesh = MakeShared<FDynamicMesh3, ESPMode::ThreadSafe>();
+	*InputMesh = UE::ToolTarget::GetDynamicMeshCopy(Target);
+	WorldTransform = UE::ToolTarget::GetLocalToWorldTransform(Target);
+
+	// initialize triangle ROI if one exists
+	TriangleROI = MakeShared<TArray<int32>, ESPMode::ThreadSafe>();
+	if (HasInputSelection())
+	{
+		GetStoredSelectionAsTriangles(GetInputSelection(), *InputMesh, *TriangleROI);
+		TriangleROISet.Append(*TriangleROI);
+	}
+	VertexROI = MakeShared<TArray<int32>, ESPMode::ThreadSafe>();
+	UE::Geometry::TriangleToVertexIDs(InputMesh.Get(), *TriangleROI, *VertexROI);
+
+	InputMeshROISpatial = MakeShared<FDynamicMeshAABBTree3, ESPMode::ThreadSafe>(InputMesh.Get(), false);
+	if (TriangleROI->Num() > 0)
+	{
+		InputMeshROISpatial->Build(*TriangleROI);
+	}
+	else
+	{
+		InputMeshROISpatial->Build();
+	}
+
+	WorldBounds = (VertexROI->Num() > 0) ?
+		FAxisAlignedBox3d::MakeBoundsFromIndices(*VertexROI, [this](int32 Index) { return WorldTransform.TransformPosition(InputMesh->GetVertex(Index)); })
+		: FAxisAlignedBox3d::MakeBoundsFromIndices(InputMesh->VertexIndicesItr(), [this](int32 Index) { return WorldTransform.TransformPosition(InputMesh->GetVertex(Index)); });
+
+	BasicProperties->Dimensions = (FVector)WorldBounds.Diagonal();
+}
+
+
+
+FOrientedBox3d UUVProjectionTool::GetProjectionBox()
+{
+	FFrame3d BoxFrame(TransformProxy->GetTransform());
+	FVector3d BoxDimensions = 0.5 * (FVector3d)BasicProperties->Dimensions;
+	return FOrientedBox3d(BoxFrame, BoxDimensions);
+}
+
+
+
 void UUVProjectionTool::Render(IToolsContextRenderAPI* RenderAPI)
 {
 	ProjectionShapeVisualizer.bDepthTested = false;
 	ProjectionShapeVisualizer.BeginFrame(RenderAPI, CameraState);
-	FTransform LocalScale = FTransform::Identity;
-	LocalScale.SetScale3D(BasicProperties->ProjectionPrimitiveScale);
-	for (UTransformProxy* TransformProxy : TransformProxies)
+
+	FOrientedBox3d PrimBox = GetProjectionBox();
+	FVector MinCorner = (FVector)-PrimBox.Extents;
+	FVector MaxCorner = (FVector)PrimBox.Extents;
+	float Width = MaxCorner.X - MinCorner.X;
+	float Depth = MaxCorner.Y - MinCorner.Y;
+	float Height = MaxCorner.Z - MinCorner.Z;
+
+	FTransform UseTransform = PrimBox.Frame.ToFTransform();
+	if (BasicProperties->ProjectionType == EUVProjectionMethod::Cylinder)
 	{
-		ProjectionShapeVisualizer.SetTransform(LocalScale * TransformProxy->GetTransform());
+		UseTransform.SetScale3D(FVector(Width, Depth, 1.0f));
+	}
+	ProjectionShapeVisualizer.SetTransform(UseTransform);
 
-		switch (BasicProperties->UVProjectionMethod)
-		{
-		case EUVProjectionMethod::Cube:
-			ProjectionShapeVisualizer.DrawWireBox(FBox(FVector(-1, -1, -1), FVector(1, 1, 1)));
-			break;
-		case EUVProjectionMethod::Cylinder:
-			ProjectionShapeVisualizer.DrawWireCylinder(FVector(0, 0, -1), FVector(0, 0, 1), 1, 2, 20);
-			break;
-		case EUVProjectionMethod::Plane:
-			ProjectionShapeVisualizer.DrawSquare(FVector(0, 0, 0), FVector(2, 0, 0), FVector(0, 2, 0));
-			break;
-		}
+	switch (BasicProperties->ProjectionType)
+	{
+	case EUVProjectionMethod::Box:
+		ProjectionShapeVisualizer.DrawWireBox(FBox(MinCorner, MaxCorner));
+		break;
+	case EUVProjectionMethod::Cylinder:
+		ProjectionShapeVisualizer.DrawWireCylinder(FVector(0, 0, -Height*0.5f), FVector(0, 0, 1), 0.5f, Height, 16);
+		break;
+	case EUVProjectionMethod::Plane:
+		ProjectionShapeVisualizer.DrawSquare(FVector(0, 0, 0), FVector(Width, 0, 0), FVector(0, Depth, 0));
+		break;
+	case EUVProjectionMethod::ExpMap:
+		ProjectionShapeVisualizer.DrawSquare(FVector(0, 0, 0), FVector(Width, 0, 0), FVector(0, Depth, 0));
+		break;
 	}
 	
 	ProjectionShapeVisualizer.EndFrame();
@@ -285,104 +368,246 @@ void UUVProjectionTool::Render(IToolsContextRenderAPI* RenderAPI)
 
 void UUVProjectionTool::OnTick(float DeltaTime)
 {
-	for (UMeshOpPreviewWithBackgroundCompute* Preview : Previews)
+	if (bHavePendingAction)
 	{
-		Preview->Tick(DeltaTime);
+		ApplyAction(PendingAction);
+		bHavePendingAction = false;
+		PendingAction = EUVProjectionToolActions::NoAction;
 	}
+
+	Preview->Tick(DeltaTime);
 }
 
 
 
-#if WITH_EDITOR
-void UUVProjectionTool::PostEditChangeProperty(FPropertyChangedEvent& PropertyChangedEvent)
+void UUVProjectionTool::OnMaterialSettingsChanged()
 {
-	UpdateNumPreviews();
-	for (UMeshOpPreviewWithBackgroundCompute* Preview : Previews)
+	MaterialSettings->UpdateMaterials();
+
+	UMaterialInterface* OverrideMaterial = MaterialSettings->GetActiveOverrideMaterial();
+	if (OverrideMaterial != nullptr)
 	{
-		Preview->InvalidateResult();
+		Preview->OverrideMaterial = OverrideMaterial;
+	}
+	else
+	{
+		Preview->OverrideMaterial = nullptr;
 	}
 }
-#endif
 
-void UUVProjectionTool::UpdateVisualization()
+
+void UUVProjectionTool::OnProjectionTypeChanged()
 {
-	MaterialSettings->UpdateMaterials();
-	for (int PreviewIdx = 0; PreviewIdx < Previews.Num(); PreviewIdx++)
+	SetToolPropertySourceEnabled(ExpMapProperties, BasicProperties->ProjectionType == EUVProjectionMethod::ExpMap);
+	SetToolPropertySourceEnabled(CylinderProperties, BasicProperties->ProjectionType == EUVProjectionMethod::Cylinder);
+}
+
+
+void UUVProjectionTool::OnMeshUpdated(UMeshOpPreviewWithBackgroundCompute* PreviewCompute)
+{
+	const FColor UVSeamColor(15, 240, 15);
+	const float UVSeamThickness = 2.0f;
+	const float UVSeamDepthBias = 0.1f * (float)(WorldBounds.DiagonalLength() * 0.01);
+
+	const FDynamicMesh3* UseMesh = Preview->PreviewMesh->GetMesh();
+	const FDynamicMeshUVOverlay* UVOverlay = UseMesh->Attributes()->GetUVLayer(UVChannelProperties->GetSelectedChannelIndex(true));
+	auto AppendSeamEdge = [UseMesh, UVOverlay, UVSeamColor, UVSeamThickness, UVSeamDepthBias](int32 eid, TArray<FRenderableLine>& LinesOut) {
+		FVector3d A, B;
+		UseMesh->GetEdgeV(eid, A, B);
+		LinesOut.Add(FRenderableLine((FVector)A, (FVector)B, UVSeamColor, UVSeamThickness, UVSeamDepthBias));
+	};
+	if (TriangleROI->Num() > 0)
 	{
-		UMeshOpPreviewWithBackgroundCompute* Preview = Previews[PreviewIdx];
-		Preview->OverrideMaterial = MaterialSettings->GetActiveOverrideMaterial();
+		EdgeRenderer->CreateOrUpdateLineSet(TEXT("UVSeams"), UseMesh->MaxEdgeID(), [&](int32 eid, TArray<FRenderableLine>& LinesOut) {
+			FIndex2i EdgeT = UseMesh->GetEdgeT(eid);
+			if (TriangleROISet.Contains(EdgeT.A) && TriangleROISet.Contains(EdgeT.B) && UVOverlay->IsSeamEdge(eid))
+			{
+				AppendSeamEdge(eid, LinesOut);
+			}
+		}, 1);
 	}
-
-	UpdateNumPreviews();
-	for (UMeshOpPreviewWithBackgroundCompute* Preview : Previews)
+	else
 	{
-		Preview->InvalidateResult();
+		EdgeRenderer->CreateOrUpdateLineSet(TEXT("UVSeams"), UseMesh->MaxEdgeID(), [&](int32 eid, TArray<FRenderableLine>& LinesOut) {
+			if (UVOverlay->IsSeamEdge(eid))
+			{
+				AppendSeamEdge(eid, LinesOut);
+			}
+		}, 1);
 	}
 }
 
-void UUVProjectionTool::OnPropertyModified(UObject* PropertySet, FProperty* Property)
+
+void UUVProjectionTool::UpdatePlaneFromClick(const FVector3d& Position, const FVector3d& Normal, bool bTransationOnly)
 {
-	UpdateVisualization();
+	FFrame3d CurrentFrame(TransformProxy->GetTransform());
+	CurrentFrame.Origin = Position;
+	if (bTransationOnly == false)
+	{
+		CurrentFrame.AlignAxis(2, Normal);
+	}
+	TransformGizmo->SetNewGizmoTransform(CurrentFrame.ToFTransform());
+	Preview->InvalidateResult();
+	bTransformModified = true;
 }
 
 
 void UUVProjectionTool::TransformChanged(UTransformProxy* Proxy, FTransform Transform)
 {
-	// TODO: if multi-select is re-enabled, only invalidate the preview that actually needs it?
-	for (UMeshOpPreviewWithBackgroundCompute* Preview : Previews)
+	Preview->InvalidateResult();
+	bTransformModified = true;
+}
+
+
+bool UUVProjectionTool::CanAccept() const
+{
+	if (!Preview->HaveValidResult())
 	{
-		Preview->InvalidateResult();
+		return false;
 	}
+	return Super::CanAccept();
 }
 
 
-bool UUVProjectionTool::CanAccept() const
+
+void UUVProjectionTool::OnInitializationModeChanged()
+{
+	bool bWasTransformModified = bTransformModified;
+	if (!bTransformModified)
+	{
+		ApplyInitializationMode();
+		bTransformModified = bWasTransformModified;
+	}
+}
+
+
+void UUVProjectionTool::ApplyInitializationMode()
 {
-	for (UMeshOpPreviewWithBackgroundCompute* Preview : Previews)
+
+	switch (BasicProperties->Initialization)
 	{
-		if (!Preview->HaveValidResult())
+	case EUVProjectionToolInitializationMode::Default:
+		BasicProperties->Dimensions = InitialDimensions;
+		TransformGizmo->SetNewGizmoTransform(InitialTransform);
+		break;
+
+	case EUVProjectionToolInitializationMode::UsePrevious:
 		{
-			return false;
+			bool bHavePrevious = (BasicProperties->SavedDimensions != FVector::ZeroVector);
+			BasicProperties->Dimensions = (bHavePrevious) ? BasicProperties->SavedDimensions : InitialDimensions;
+			TransformGizmo->SetNewGizmoTransform( (bHavePrevious) ? BasicProperties->SavedTransform : InitialTransform);
 		}
+		break;
+
+	case EUVProjectionToolInitializationMode::AutoFit:
+		ApplyAction_AutoFit(false);
+		break;
+
+	case EUVProjectionToolInitializationMode::AutoFitAlign:
+		ApplyAction_AutoFit(true);
+		break;
 	}
-	return Super::CanAccept();
 }
 
 
-void UUVProjectionTool::GenerateAsset(const TArray<FDynamicMeshOpResult>& Results)
+
+void UUVProjectionTool::ApplyAction(EUVProjectionToolActions ActionType)
 {
-	GetToolManager()->BeginUndoTransaction(LOCTEXT("UVProjectionToolTransactionName", "UV Projection Tool"));
+	switch (ActionType)
+	{
+		case EUVProjectionToolActions::AutoFit:
+			ApplyAction_AutoFit(false);
+			break;
+		case EUVProjectionToolActions::AutoFitAlign:
+			ApplyAction_AutoFit(true);
+			break;
+		case EUVProjectionToolActions::Reset:
+			ApplyAction_Reset();
+			break;
+	}
+}
 
-	check(Results.Num() == Targets.Num());
-	
-	for (int32 ComponentIdx = 0; ComponentIdx < Targets.Num(); ComponentIdx++)
+
+void UUVProjectionTool::ApplyAction_AutoFit(bool bAlign)
+{
+	// get current transform
+	FFrame3d AlignFrame(TransformGizmo->GetGizmoTransform());
+
+	TArray<FVector3d> Points;
+	if (VertexROI->Num() > 0)
+	{
+		CollectVertexPositions(*InputMesh, *VertexROI, Points);
+	}
+	else
+	{
+		CollectVertexPositions(*InputMesh, InputMesh->VertexIndicesItr(), Points);
+	}
+
+	// auto-compute orientation if desired
+	if (bAlign)
 	{
-		check(Results[ComponentIdx].Mesh.Get() != nullptr);
-		TargetMeshCommitterInterface(ComponentIdx)->CommitMeshDescription([&Results, &ComponentIdx, this](const IMeshDescriptionCommitter::FCommitterParams& CommitParams)
+		if (BasicProperties->ProjectionType == EUVProjectionMethod::Plane || BasicProperties->ProjectionType == EUVProjectionMethod::Box || BasicProperties->ProjectionType == EUVProjectionMethod::Cylinder)
 		{
-			FDynamicMesh3* DynamicMesh = Results[ComponentIdx].Mesh.Get();
-			FMeshDescription* MeshDescription = CommitParams.MeshDescriptionOut;
+			// compute min-volume box
+			FOrientedBox3d MinBoundingBox;
+			FMinVolumeBox3d MinBoxCalc;
+			bool bMinBoxOK = MinBoxCalc.SolveSubsample(Points.Num(), 1000,
+				[&](int32 Index) { return WorldTransform.TransformPosition(Points[Index]); }, false, nullptr);
+			if (bMinBoxOK && MinBoxCalc.IsSolutionAvailable())
+			{
+				MinBoxCalc.GetResult(MinBoundingBox);
+			}
+			else
+			{
+				MinBoundingBox = FOrientedBox3d(WorldBounds);
+			}
 
-			bool bVerticesOnly = false;
-			bool bAttributesOnly = true;
-			if (FDynamicMeshToMeshDescription::HaveMatchingElementCounts(DynamicMesh, MeshDescription, bVerticesOnly, bAttributesOnly))
+			AlignFrame = MinBoundingBox.Frame;
+		}
+		else if (BasicProperties->ProjectionType == EUVProjectionMethod::ExpMap)
+		{
+			int32 VertexID;
+			if (TriangleROI->Num() > 0)
 			{
-				FDynamicMeshToMeshDescription Converter;
-				Converter.UpdateAttributes(DynamicMesh, *MeshDescription, false, false, true/*update uvs*/);
+				FDynamicMeshUVEditor::EstimateGeodesicCenterFrameVertex(*InputMesh, *TriangleROI, AlignFrame, VertexID);
+				AlignFrame.Transform(WorldTransform);
 			}
 			else
 			{
-				// must have been duplicate tris in the mesh description; we can't count on 1-to-1 mapping of TriangleIDs.  Just convert 
-				FDynamicMeshToMeshDescription Converter;
-				Converter.Convert(DynamicMesh, *MeshDescription);
+				FDynamicMeshUVEditor::EstimateGeodesicCenterFrameVertex(*InputMesh, AlignFrame, VertexID);
+				AlignFrame.Transform(WorldTransform);
 			}
-		});
+		}
 	}
 
-	GetToolManager()->EndUndoTransaction();
+	// fit bounds to current frame
+	FAxisAlignedBox3d FitBounds = FAxisAlignedBox3d::MakeBoundsFromIndices(Points.Num(), [this, &Points, &AlignFrame](int32 Index) 
+	{ 
+		return AlignFrame.ToFramePoint(WorldTransform.TransformPosition(Points[Index])); }
+	);
+	if (BasicProperties->ProjectionType == EUVProjectionMethod::Plane || BasicProperties->ProjectionType == EUVProjectionMethod::Box || BasicProperties->ProjectionType == EUVProjectionMethod::Cylinder)
+	{
+		FVector3d ShiftCenter = AlignFrame.FromFramePoint(FitBounds.Center());
+		AlignFrame.Origin = ShiftCenter;
+	}
+	BasicProperties->Dimensions = 2.0f * (FVector)FitBounds.Extents();
+	if (DimensionsWatcher >= 0)
+	{
+		BasicProperties->SilentUpdateWatcherAtIndex(DimensionsWatcher);
+	}
+
+	// update Gizmo
+	TransformGizmo->SetNewGizmoTransform(AlignFrame.ToFTransform());
 }
 
 
 
+void UUVProjectionTool::ApplyAction_Reset()
+{
+	bool bWasTransformModified = bTransformModified;
+	ApplyInitializationMode();
+	bTransformModified = bWasTransformModified;
+}
+
 
 #undef LOCTEXT_NAMESPACE
diff --git a/Engine/Plugins/Experimental/MeshModelingToolset/Source/MeshModelingTools/Public/UVProjectionTool.h b/Engine/Plugins/Experimental/MeshModelingToolset/Source/MeshModelingTools/Public/UVProjectionTool.h
--- a/Engine/Plugins/Experimental/MeshModelingToolset/Source/MeshModelingTools/Public/UVProjectionTool.h
+++ b/Engine/Plugins/Experimental/MeshModelingToolset/Source/MeshModelingTools/Public/UVProjectionTool.h
@@ -3,15 +3,14 @@
 #pragma once
 
 #include "CoreMinimal.h"
-#include "UObject/NoExportTypes.h"
-#include "MultiSelectionTool.h"
-#include "InteractiveToolBuilder.h"
+#include "BaseTools/SingleSelectionMeshEditingTool.h"
 #include "MeshOpPreviewHelpers.h"
 #include "ToolDataVisualizer.h"
 #include "ParameterizationOps/UVProjectionOp.h"
-#include "DynamicMesh3.h"
-#include "BaseTools/SingleClickTool.h"
 #include "Properties/MeshMaterialProperties.h"
+#include "Properties/MeshUVChannelProperties.h"
+#include "Drawing/PreviewGeometryActor.h"
+#include "Selection/SelectClickedAction.h"
 
 #include "UVProjectionTool.generated.h"
 
@@ -21,29 +20,95 @@ struct FMeshDescription;
 class USimpleDynamicMeshComponent;
 class UTransformGizmo;
 class UTransformProxy;
+class USingleClickInputBehavior;
 class IAssetGenerationAPI;
-PREDECLARE_USE_GEOMETRY_CLASS(FDynamicMesh3);
 
 /**
  *
  */
 UCLASS()
-class MESHMODELINGTOOLS_API UUVProjectionToolBuilder : public UInteractiveToolBuilder
+class MESHMODELINGTOOLS_API UUVProjectionToolBuilder : public USingleSelectionMeshEditingToolBuilder
 {
 	GENERATED_BODY()
 
 public:
-	IAssetGenerationAPI* AssetAPI = nullptr;
+	virtual USingleSelectionMeshEditingTool* CreateNewTool(const FToolBuilderState& SceneState) const override;
 
-	virtual bool CanBuildTool(const FToolBuilderState& SceneState) const override;
-	virtual UInteractiveTool* BuildTool(const FToolBuilderState& SceneState) const override;
+	virtual bool WantsInputSelectionIfAvailable() const override { return true; }
+};
 
-protected:
-	virtual const FToolTargetTypeRequirements& GetTargetRequirements() const override;
+
+
+
+UENUM()
+enum class EUVProjectionToolActions
+{
+	NoAction,
+	AutoFit,
+	AutoFitAlign,
+	Reset
+};
+
+
+UCLASS()
+class MESHMODELINGTOOLS_API UUVProjectionToolEditActions : public UInteractiveToolPropertySet
+{
+	GENERATED_BODY()
+
+public:
+	TWeakObjectPtr<UUVProjectionTool> ParentTool;
+
+	void Initialize(UUVProjectionTool* ParentToolIn) { ParentTool = ParentToolIn; }
+	void PostAction(EUVProjectionToolActions Action);
+
+	/** Automatically fit the Projection Dimensions based on the current projection orientation */
+	UFUNCTION(CallInEditor, Category = Actions, meta = (DisplayName = "AutoFit", DisplayPriority = 1))
+	void AutoFit()
+	{
+		PostAction(EUVProjectionToolActions::AutoFit);
+	}
+
+	/** Automatically align the projection orientation and then automatically fit the Dimensions */
+	UFUNCTION(CallInEditor, Category = Actions, meta = (DisplayName = "AutoFitAlign", DisplayPriority = 2))
+	void AutoFitAlign()
+	{
+		PostAction(EUVProjectionToolActions::AutoFitAlign);
+	}
+
+	/** Re-initialize the Projection based on the Initialization setting */
+	UFUNCTION(CallInEditor, Category = Actions, meta = (DisplayName = "Reset", DisplayPriority = 3))
+	void Reset()
+	{
+		PostAction(EUVProjectionToolActions::Reset);
+	}
 };
 
 
 
+UENUM()
+enum class EUVProjectionToolInitializationMode
+{
+	/** Initialize projection to bounding box center */
+	Default,
+	/** Initialize projection based on previous usage of the UV Projection Tool */
+	UsePrevious,
+	/** Initialize projection using AutoFitting, for the initial projection type */
+	AutoFit,
+	/** Initialize projection using AutoFitting with Alignment, for the initial projection type */
+	AutoFitAlign
+};
+
+
+UENUM()
+enum class EUVProjectionToolDimensionMode
+{
+	/** All Dimensions of the projection are used */
+	AllFree,
+	/** The Dimensions.X value is used to define all projection Dimensions */
+	UseFirst
+};
+
+
 
 
 /**
@@ -55,51 +120,91 @@ class MESHMODELINGTOOLS_API UUVProjectionToolProperties : public UInteractiveToo
 	GENERATED_BODY()
 
 public:
-	UUVProjectionToolProperties();
-
-
-
 	/** Choose the UV projection method (cube, cylinder, plane) */
 	UPROPERTY(EditAnywhere, Category = ProjectionSettings)
-	EUVProjectionMethod UVProjectionMethod;
+	EUVProjectionMethod ProjectionType = EUVProjectionMethod::Plane;
 
-	/** Per-axis scaling of projection primitive */
+	/** Per-axis Dimensions of the Projection. Z is height of box/cylinder. */
 	UPROPERTY(EditAnywhere, Category = ProjectionSettings)
-	FVector ProjectionPrimitiveScale;
+	FVector Dimensions = FVector(100.0f, 100.0f, 100.0f);
 
-	/** If triangle normal direction is within this threshold degrees of the cylinder top/bottom plane direction, project UVs to the top/bottom plane instead of the sides */
-	UPROPERTY(EditAnywhere, Category = ProjectionSettings, meta = (Default = "1", UIMin = "0", UIMax = "20", ClampMin = "0.0", ClampMax = "90.0", EditCondition = "UVProjectionMethod == EUVProjectionMethod::Cylinder"))
-	float CylinderProjectToTopOrBottomAngleThreshold;
-
-	/** Choose the UV scale factors */
+	/** How the Dimensions are interpreted for Box/Plane/ExpMap Projections */
 	UPROPERTY(EditAnywhere, Category = ProjectionSettings)
-	FVector2D UVScale;
+	EUVProjectionToolDimensionMode DimensionMode = EUVProjectionToolDimensionMode::AllFree;
 
-	/** Choose the UV offsets */
+	/** Determines how projection settings will be initialized. Only affects Tool until changes are made to Dimensions/Position. */
 	UPROPERTY(EditAnywhere, Category = ProjectionSettings)
-	FVector2D UVOffset;
+	EUVProjectionToolInitializationMode Initialization = EUVProjectionToolInitializationMode::Default;
+
+	//
+	// UV-space transform options
+	//
+
+	/** UV-Space Rotation Angle in Degrees, applied after computing projected UVs */
+	UPROPERTY(EditAnywhere, Category = "UV Transform", meta = (DisplayName = "UV Rotation") )
+	float UVRotation = 0.0;
+
+	/** UV-Space Scaling applied after computing projected UVs */
+	UPROPERTY(EditAnywhere, Category = "UV Transform", meta = (DisplayName = "UV Scale"))
+	FVector2D UVScale = FVector2D::UnitVector;
+
+	/** UV-Space Translation applied after computing projected UVs */
+	UPROPERTY(EditAnywhere, Category = "UV Transform", meta = (DisplayName = "UV Translation"))
+	FVector2D UVTranslate = FVector2D::ZeroVector;
 
-	/** If set, UV scales will be relative to world space so different objects created with the same UV scale should have the same average texel size */
-	UPROPERTY(EditAnywhere, Category = ProjectionSettings, meta = (DisplayName = "UV Scale Relative to World Space"))
-	bool bWorldSpaceUVScale = false;
+
+	//
+	// Saved State. These are used internally to support UsePrevious initialization mode
+	//
+	UPROPERTY()
+	FVector SavedDimensions = FVector::ZeroVector;
+
+	UPROPERTY()
+	FTransform SavedTransform;
 };
 
 
 
 
 /**
- * Advanced properties
+ * ExpMap Projection properties
  */
 UCLASS()
-class MESHMODELINGTOOLS_API UUVProjectionAdvancedProperties : public UInteractiveToolPropertySet
+class MESHMODELINGTOOLS_API UUVProjectionToolExpMapProperties : public UInteractiveToolPropertySet
 {
 	GENERATED_BODY()
+public:
+	/** Blend between surface normals and projection normal - ExpMap becomes Planar Projection when value is 1 */
+	UPROPERTY(EditAnywhere, Category = ExpMap, meta = (UIMin = "0", UIMax = "1"))
+	float NormalBlending = 0.0;
+
+	/** Rounds of Smoothing to apply to surface normals before computing ExpMap */
+	UPROPERTY(EditAnywhere, Category = ExpMap, AdvancedDisplay, meta = (UIMin = "0", UIMax = "100"))
+	int SmoothingRounds = 0;
+
+	/** Smoothing Strength for ExpMap Smoothing */
+	UPROPERTY(EditAnywhere, Category = ExpMap, AdvancedDisplay, meta = (UIMin = "0", UIMax = "1"))
+	float SmoothingAlpha = 0.25f;
+};
+
 
+
+/**
+ * Cylinder Projection properties
+ */
+UCLASS()
+class MESHMODELINGTOOLS_API UUVProjectionToolCylinderProperties : public UInteractiveToolPropertySet
+{
+	GENERATED_BODY()
 public:
-	UUVProjectionAdvancedProperties();
+	/** Angle used to determine whether faces should be assigned to cylinder or flat endcaps */
+	UPROPERTY(EditAnywhere, Category = Cylinder, meta = (UIMin = "0", UIMax = "90"))
+	float SplitAngle = 45.0f;
 };
 
 
+
+
 /**
  * Factory with enough info to spawn the background-thread Operator to do a chunk of work for the tool
  *  stores a pointer to the tool and enough info to know which specific operator it should spawn
@@ -108,23 +213,20 @@ UCLASS()
 class MESHMODELINGTOOLS_API UUVProjectionOperatorFactory : public UObject, public UE::Geometry::IDynamicMeshOperatorFactory
 {
 	GENERATED_BODY()
-
 public:
 	// IDynamicMeshOperatorFactory API
 	virtual TUniquePtr<UE::Geometry::FDynamicMeshOperator> MakeNewOperator() override;
 
 	UPROPERTY()
 	UUVProjectionTool *Tool;
-
-	int ComponentIndex;
-
 };
 
+
 /**
  * Simple Mesh Normal Updating Tool
  */
 UCLASS()
-class MESHMODELINGTOOLS_API UUVProjectionTool : public UMultiSelectionTool
+class MESHMODELINGTOOLS_API UUVProjectionTool : public USingleSelectionMeshEditingTool
 {
 	GENERATED_BODY()
 
@@ -132,14 +234,9 @@ public:
 
 	friend UUVProjectionOperatorFactory;
 
-	UUVProjectionTool();
-
 	virtual void Setup() override;
 	virtual void Shutdown(EToolShutdownType ShutdownType) override;
 
-	virtual void SetWorld(UWorld* World, UInteractiveGizmoManager* GizmoManagerIn);
-	virtual void SetAssetAPI(IAssetGenerationAPI* AssetAPI);
-
 	virtual void OnTick(float DeltaTime) override;
 	virtual void Render(IToolsContextRenderAPI* RenderAPI) override;
 
@@ -147,54 +244,99 @@ public:
 	virtual bool HasAccept() const override { return true; }
 	virtual bool CanAccept() const override;
 
-#if WITH_EDITOR
-	virtual void PostEditChangeProperty(FPropertyChangedEvent &PropertyChangedEvent) override;
-#endif
-
-	virtual void OnPropertyModified(UObject* PropertySet, FProperty* Property) override;
+	virtual void RequestAction(EUVProjectionToolActions ActionType);
 
 protected:
 
 	UPROPERTY()
-	UUVProjectionToolProperties* BasicProperties;
+	UMeshUVChannelProperties* UVChannelProperties = nullptr;
 
 	UPROPERTY()
-	UUVProjectionAdvancedProperties* AdvancedProperties;
+	UUVProjectionToolProperties* BasicProperties = nullptr;
 
 	UPROPERTY()
-	UExistingMeshMaterialProperties* MaterialSettings = nullptr;
+	UUVProjectionToolExpMapProperties* ExpMapProperties = nullptr;
+
+	UPROPERTY()
+	UUVProjectionToolCylinderProperties* CylinderProperties = nullptr;
 
+	UPROPERTY()
+	UUVProjectionToolEditActions* EditActions = nullptr;
 
+	UPROPERTY()
+	UExistingMeshMaterialProperties* MaterialSettings = nullptr;
 
 	UPROPERTY()
-	TArray<UMeshOpPreviewWithBackgroundCompute*> Previews;
+	UMeshOpPreviewWithBackgroundCompute* Preview = nullptr;
 
 
 	UPROPERTY()
 	UMaterialInstanceDynamic* CheckerMaterial = nullptr;
 
 	UPROPERTY()
-	TArray<UTransformGizmo*> TransformGizmos;
+	UTransformGizmo* TransformGizmo = nullptr;
 	
 	UPROPERTY()
-	TArray<UTransformProxy*> TransformProxies;
+	UTransformProxy* TransformProxy = nullptr;
 
-protected:
-	TArray<TSharedPtr<FDynamicMesh3, ESPMode::ThreadSafe>> OriginalDynamicMeshes;
-	TArray<FDynamicMesh3> ReferencePrimitives;
+	UPROPERTY()
+	UUVProjectionOperatorFactory* OperatorFactory = nullptr;
+
+	UPROPERTY()
+	UPreviewGeometry* EdgeRenderer = nullptr;
 
-	UWorld* TargetWorld;
-	IAssetGenerationAPI* AssetAPI;
-	UInteractiveGizmoManager* GizmoManager;
+protected:
+	TSharedPtr<UE::Geometry::FDynamicMesh3, ESPMode::ThreadSafe> InputMesh;
+	TSharedPtr<TArray<int32>, ESPMode::ThreadSafe> TriangleROI;
+	TSharedPtr<TArray<int32>, ESPMode::ThreadSafe> VertexROI;
+	TSharedPtr<UE::Geometry::FDynamicMeshAABBTree3, ESPMode::ThreadSafe> InputMeshROISpatial;
+	TSet<int32> TriangleROISet;
+
+	UE::Geometry::FTransform3d WorldTransform;
+	UE::Geometry::FAxisAlignedBox3d WorldBounds;
+
+	FVector InitialDimensions;
+	FTransform InitialTransform;
+	int32 DimensionsWatcher = -1;
+	bool bTransformModified = false;
+	void OnInitializationModeChanged();
+	void ApplyInitializationMode();
 
 	FViewCameraState CameraState;
 
 	FToolDataVisualizer ProjectionShapeVisualizer;
 
+	void InitializeMesh();
 	void UpdateNumPreviews();
-	void UpdateVisualization();
-
-	void GenerateAsset(const TArray<FDynamicMeshOpResult>& Results);
 
 	void TransformChanged(UTransformProxy* Proxy, FTransform Transform);
+
+	void OnProjectionTypeChanged();
+	void OnMaterialSettingsChanged();
+	void OnMeshUpdated(UMeshOpPreviewWithBackgroundCompute* PreviewCompute);
+
+	UE::Geometry::FOrientedBox3d GetProjectionBox();
+
+
+	//
+	// Support for ctrl+click to set plane from hit point
+	//
+
+	TUniquePtr<FSelectClickedAction> SetPlaneCtrlClickBehaviorTarget;
+
+	UPROPERTY()
+	USingleClickInputBehavior* ClickToSetPlaneBehavior;
+
+	void UpdatePlaneFromClick(const FVector3d& Position, const FVector3d& Normal, bool bTransationOnly);
+
+	//
+	// Support for Action Buttons
+	//
+
+	bool bHavePendingAction = false;
+	EUVProjectionToolActions PendingAction;
+	virtual void ApplyAction(EUVProjectionToolActions ActionType);
+	void ApplyAction_AutoFit(bool bAlign);
+	void ApplyAction_Reset();
+
 };
diff --git a/Engine/Plugins/Experimental/MeshModelingToolset/Source/ModelingOperators/Private/ParameterizationOps/UVProjectionOp.cpp b/Engine/Plugins/Experimental/MeshModelingToolset/Source/ModelingOperators/Private/ParameterizationOps/UVProjectionOp.cpp
--- a/Engine/Plugins/Experimental/MeshModelingToolset/Source/ModelingOperators/Private/ParameterizationOps/UVProjectionOp.cpp
+++ b/Engine/Plugins/Experimental/MeshModelingToolset/Source/ModelingOperators/Private/ParameterizationOps/UVProjectionOp.cpp
@@ -8,18 +8,14 @@
 #include "MeshSimplification.h"
 #include "MeshConstraintsUtil.h"
 #include "MatrixTypes.h"
-
 #include "Async/ParallelFor.h"
-
 #include "MeshNormals.h"
 #include "DynamicMeshEditor.h"
+#include "Parameterization/DynamicMeshUVEditor.h"
 
 #include "ExplicitUseGeometryMathTypes.h"		// using UE::Geometry::(math types)
 using namespace UE::Geometry;
 
-void FUVProjectionOp::SetTransform(const FTransform& Transform) {
-	ResultTransform = (FTransform3d)Transform;
-}
 
 void FUVProjectionOp::CalculateResult(FProgressCancel* Progress)
 {
@@ -27,222 +23,168 @@ void FUVProjectionOp::CalculateResult(FProgressCancel* Progress)
 	{
 		return;
 	}
-	bool bDiscardAttributes = false;
-	ResultMesh->Copy(*OriginalMesh, true, true, true, !bDiscardAttributes);
-
+	*ResultMesh = *OriginalMesh;
 	if (!ensureMsgf(ResultMesh->HasAttributes(), TEXT("Attributes not found on mesh? Conversion should always create them, so this operator should not need to do so.")))
 	{
 		ResultMesh->EnableAttributes();
 	}
+	
+	if (ProjectionMethod == EUVProjectionMethod::Plane)
+	{
+		CalculateResult_Plane(Progress);
+	}
+	if (ProjectionMethod == EUVProjectionMethod::ExpMap)
+	{
+		CalculateResult_ExpMap(Progress);
+	}
+	else if (ProjectionMethod == EUVProjectionMethod::Box)
+	{
+		CalculateResult_Box(Progress);
+	}
+	else if (ProjectionMethod == EUVProjectionMethod::Cylinder)
+	{
+		CalculateResult_Cylinder(Progress);
+	}
+	else
+	{
+		return;
+	}
+}
 
-	// TODO: get transform from mesh vertex space to projection primitive space
-	FTransform ToProjPrimXF = FTransform(ResultTransform) * ProjectionTransform.Inverse();
-	FTransform3d ToProjPrim(ToProjPrimXF);
 
-	TArray<FVector3d> TransformedVertices; TransformedVertices.SetNumUninitialized(ResultMesh->VertexCount());
-	ParallelFor(ResultMesh->VertexCount(), [this, &ToProjPrim, &TransformedVertices](int32 VertexID)
+static void ApplyUVTransforms(
+	FDynamicMeshUVEditor& UVEditor, const TArray<int32>& UVElements, 
+	FVector2f UVOrigin,
+	float UVRotationAngleDeg,
+	FVector2f UVScale,
+	FVector2f UVTranslate)
+{
+	FMatrix2f RotationMatrix = FMatrix2f::RotationDeg(UVRotationAngleDeg);
+	UVEditor.TransformUVElements(UVElements, [&](const FVector2f& UV)
 	{
-		if (ResultMesh->IsVertex(VertexID))
-		{
-			TransformedVertices[VertexID] = ToProjPrim.TransformPosition(ResultMesh->GetVertex(VertexID));
-		}
+		FVector2f UVScaleRotate = (RotationMatrix * UV) * UVScale;
+		return (UVScaleRotate + UVOrigin) + UVTranslate;
 	});
-	
-	
-	int NumUVLayers = ResultMesh->Attributes()->NumUVLayers();
-	int LayerIndex = 0; // TODO: support non-zero UV layer, make this a member of FUVProjectionOp and a setting in the UI (will require updating dmesh attribute code)
-	FDynamicMeshUVOverlay* UVLayer = ResultMesh->Attributes()->GetUVLayer(LayerIndex);
-	UVLayer->ClearElements();
+}
 
-	FVector2f Scale = FVector2f(.5,.5)*UVScale, Offset = FVector2f(.5,.5)+UVOffset;
-	// project to major axis
-	auto ProjAxis = [&Scale, &Offset](const FVector3d& P, int Ax1, int Ax2, float Ax1Scale, float Ax2Scale)
-	{
-		return FVector2f(float(P[Ax1]) * Ax1Scale * Scale.X + Offset.X, float(P[Ax2]) * Ax2Scale * Scale.Y + Offset.Y);
-	};
 
-	if (Progress && Progress->Cancelled())
-	{
-		return;
-	}
+void FUVProjectionOp::CalculateResult_Box(FProgressCancel* Progress)
+{
+	int LayerIndex = FMath::Min(UseUVLayer, 8);		// hardcoded max UV layers
+	FDynamicMeshUVEditor UVEditor(ResultMesh.Get(), UseUVLayer, true);
 
-	int MaxVertexID = ResultMesh->MaxVertexID();
-	int MaxTriangleID = ResultMesh->MaxTriangleID();
-	if (ProjectionMethod == EUVProjectionMethod::Plane)
+	FUVEditResult EditResult;
+	if (!TriangleROI || TriangleROI->Num() == 0)
 	{
-		UVLayer->BeginUnsafeElementsInsert();
-		// 1:1 with vertices; note we cannot just delete the UV layer and switch to vertex UVs because that would delete the normals layer too
-		for (int VID = MaxVertexID-1; VID >= 0; VID--)
-		{
-			if (ResultMesh->IsVertex(VID))
-			{
-				FVector2f UV = ProjAxis(TransformedVertices[VID], 0, 1, 1, 1);
-				UVLayer->InsertElement(VID, (float*)UV, true);
-			}
-		}
-		UVLayer->EndUnsafeElementsInsert();
-		for (int TID = 0; TID < MaxTriangleID; TID++)
+		TArray<int32> Triangles;
+		for (int32 tid : ResultMesh->TriangleIndicesItr())
 		{
-			if (ResultMesh->IsTriangle(TID))
-			{
-				UVLayer->SetTriangle(TID, ResultMesh->GetTriangle(TID));
-			}
+			Triangles.Add(tid);
 		}
+		UVEditor.ResetUVs();
+		UVEditor.SetTriangleUVsFromBoxProjection(Triangles, [&](const FVector3d P) { return MeshToProjectionSpace.TransformPosition(P); },
+			ProjectionBox.Frame, 2*ProjectionBox.Extents, &EditResult);
 	}
 	else
-	{	
-		// All cases in this branch require normals 
-		// compute normals on the transformed vertices
-		TArray<FVector3d> TransformedNormals; TransformedNormals.SetNumUninitialized(ResultMesh->MaxTriangleID());
-		ParallelFor(ResultMesh->MaxTriangleID(), [this, &TransformedVertices, &TransformedNormals](int32 TriangleID)
-		{
-			if (ResultMesh->IsTriangle(TriangleID))
-			{
-				FIndex3i Triangle = ResultMesh->GetTriangle(TriangleID);
-				TransformedNormals[TriangleID] = VectorUtil::Normal(TransformedVertices[Triangle.A], TransformedVertices[Triangle.B], TransformedVertices[Triangle.C]);
-			}
-		});
-		TMap<FIndex2i, int> VertIdxAndBucketIDToElementID;
-
-		if (ProjectionMethod == EUVProjectionMethod::Cube)
+	{
+		UVEditor.SetTriangleUVsFromBoxProjection(*TriangleROI, [&](const FVector3d P) { return MeshToProjectionSpace.TransformPosition(P); },
+			ProjectionBox.Frame, 2*ProjectionBox.Extents, &EditResult);
+		FCompactMaps CompactMaps;
+		//UVEditor.GetOverlay()->CompactInPlace(CompactMaps);
+	}
+
+	ApplyUVTransforms(UVEditor, EditResult.NewUVElements, UVOrigin, UVRotationAngleDeg, UVScale, UVTranslate);
+}
+
+
+void FUVProjectionOp::CalculateResult_Plane(FProgressCancel* Progress)
+{
+	int LayerIndex = FMath::Min(UseUVLayer, 8);		// hardcoded max UV layers
+	FFrame3d ProjectionFrame(ProjectionBox.Frame);
+	FVector2d Dimensions(2 * ProjectionBox.Extents.X, 2 * ProjectionBox.Extents.Y);
+
+	FDynamicMeshUVEditor UVEditor(ResultMesh.Get(), UseUVLayer, true);
+
+	FUVEditResult EditResult;
+	if (!TriangleROI || TriangleROI->Num() == 0)
+	{
+		TArray<int32> Triangles;
+		for (int32 tid : ResultMesh->TriangleIndicesItr())
 		{
-			int Minor1s[3] = { 1, 0, 0 };
-			int Minor2s[3] = { 2, 2, 1 };
-			int Minor1Flip[3] = { -1, 1, 1 };
-			int Minor2Flip[3] = { -1, -1, 1 };
-			for (int TID = 0; TID < MaxTriangleID; TID++)
-			{
-				if (!ResultMesh->IsTriangle(TID))
-				{
-					continue;
-				}
-				const FVector3d& N = TransformedNormals[TID];
-				FVector3d NAbs(FMathd::Abs(N.X), FMathd::Abs(N.Y), FMathd::Abs(N.Z));
-				int MajorAxis = NAbs[0] > NAbs[1] ? (NAbs[0] > NAbs[2] ? 0 : 2) : (NAbs[1] > NAbs[2] ? 1 : 2);
-				double MajorAxisSign = FMathd::Sign(N[MajorAxis]);
-				int Minor1 = Minor1s[MajorAxis];
-				int Minor2 = Minor2s[MajorAxis];
-				int Bucket = MajorAxis;
-				if (MajorAxisSign > 0)
-				{
-					Bucket += 3;
-				}
-				FIndex3i OverlayTri;
-				FIndex3i SourceTri = ResultMesh->GetTriangle(TID);
-				for (int SubIdx = 0; SubIdx < 3; SubIdx++)
-				{
-					FIndex2i ElementKey(SourceTri[SubIdx], Bucket);
-					const int* ElementIdx = VertIdxAndBucketIDToElementID.Find(ElementKey);
-					if (!ElementIdx)
-					{
-						FVector3d LocalV = TransformedVertices[SourceTri[SubIdx]];
-						int NewElementIdx = UVLayer->AppendElement(ProjAxis(LocalV, Minor1, Minor2, MajorAxisSign*Minor1Flip[MajorAxis], Minor2Flip[MajorAxis]));
-						VertIdxAndBucketIDToElementID.Add(ElementKey, NewElementIdx);
-						OverlayTri[SubIdx] = NewElementIdx;
-					}
-					else
-					{
-						OverlayTri[SubIdx] = *ElementIdx;
-					}
-				}
-				UVLayer->SetTriangle(TID, OverlayTri);
-			}
+			Triangles.Add(tid);
 		}
-		else if (ProjectionMethod == EUVProjectionMethod::Cylinder)
+		UVEditor.ResetUVs();
+		UVEditor.SetTriangleUVsFromPlanarProjection(Triangles, [&](const FVector3d P) { return MeshToProjectionSpace.TransformPosition(P); },
+			ProjectionFrame, Dimensions, &EditResult);
+	}
+	else
+	{
+		UVEditor.SetTriangleUVsFromPlanarProjection(*TriangleROI, [&](const FVector3d P) { return MeshToProjectionSpace.TransformPosition(P); },
+			ProjectionFrame, Dimensions, &EditResult);
+		FCompactMaps CompactMaps;
+		//UVEditor.GetOverlay()->CompactInPlace(CompactMaps);
+	}
+
+	ApplyUVTransforms(UVEditor, EditResult.NewUVElements, UVOrigin, UVRotationAngleDeg, UVScale, UVTranslate);
+}
+
+
+void FUVProjectionOp::CalculateResult_ExpMap(FProgressCancel* Progress)
+{
+	int LayerIndex = FMath::Min(UseUVLayer, 8);		// hardcoded max UV layers
+	FFrame3d ProjectionFrame(ProjectionBox.Frame);
+	FVector2d Dimensions(2 * ProjectionBox.Extents.X, 2 * ProjectionBox.Extents.Y);
+
+	FDynamicMeshUVEditor UVEditor(ResultMesh.Get(), UseUVLayer, true);
+
+	FUVEditResult EditResult;
+	if (!TriangleROI || TriangleROI->Num() == 0)
+	{
+		TArray<int32> Triangles;
+		for (int32 tid : ResultMesh->TriangleIndicesItr())
 		{
-			double DotThresholdRejectFromPlane = FMathd::Cos(FMathd::DegToRad * CylinderProjectToTopOrBottomAngleThreshold);
-			for (int TID = 0; TID < MaxTriangleID; TID++)
-			{
-				if (!ResultMesh->IsTriangle(TID))
-				{
-					continue;
-				}
-				const FVector3d& N = TransformedNormals[TID];
-				FIndex3i OverlayTri;
-				FIndex3i SourceTri = ResultMesh->GetTriangle(TID);
-				if (FMathd::Abs(N.Z) > DotThresholdRejectFromPlane)
-				{
-					int MajorAxis = 2;
-					double MajorAxisSign = FMathd::Sign(N[MajorAxis]);
-					// project to +/- Z
-					int Bucket = MajorAxisSign > 0 ? 1 : 0;
-					int Minor1 = 0;
-					int Minor2 = 1;
-
-					// TODO: extract common major axis plane projection sub-function from cube projection & this?
-					for (int SubIdx = 0; SubIdx < 3; SubIdx++)
-					{
-						FIndex2i ElementKey(SourceTri[SubIdx], Bucket);
-						const int* ElementIdx = VertIdxAndBucketIDToElementID.Find(ElementKey);
-						if (!ElementIdx)
-						{
-							FVector3d LocalV = TransformedVertices[SourceTri[SubIdx]];
-							int NewElementIdx = UVLayer->AppendElement(ProjAxis(LocalV, Minor1, Minor2, MajorAxisSign, 1));
-							VertIdxAndBucketIDToElementID.Add(ElementKey, NewElementIdx);
-							OverlayTri[SubIdx] = NewElementIdx;
-						}
-						else
-						{
-							OverlayTri[SubIdx] = *ElementIdx;
-						}
-					}
-				}
-				else
-				{
-					// project to cylinder
-					int Bucket = 3;
-					FVector3d Centroid = (TransformedVertices[SourceTri.A] + TransformedVertices[SourceTri.B] + TransformedVertices[SourceTri.C]) * (1.0 / 3.0);
-					double CentroidAngle = FMathd::Atan2Positive(Centroid.Y, Centroid.X);
-					for (int SubIdx = 0; SubIdx < 3; SubIdx++)
-					{
-						FVector3d LocalV = TransformedVertices[SourceTri[SubIdx]];
-						double VAngle = FMathd::Atan2Positive(LocalV.Y, LocalV.X);
-						double AngleDiff = VAngle - CentroidAngle;
-						if (FMathd::Abs(AngleDiff) > FMathd::Pi) {
-							if (AngleDiff < 0)
-							{
-								VAngle += FMathd::TwoPi;
-								Bucket--;
-							}
-							else
-							{
-								VAngle -= FMathd::TwoPi;
-								Bucket++;
-							}
-						}
-						FIndex2i ElementKey(SourceTri[SubIdx], Bucket);
-						const int* ElementIdx = VertIdxAndBucketIDToElementID.Find(ElementKey);
-						if (!ElementIdx)
-						{
-							FVector2f UV(-(float(VAngle)*FMathf::InvPi - 1.0f)*Scale.X + Offset.X, -float(LocalV.Z)*Scale.Y + Offset.Y);
-							int NewElementIdx = UVLayer->AppendElement(UV);
-							VertIdxAndBucketIDToElementID.Add(ElementKey, NewElementIdx);
-							OverlayTri[SubIdx] = NewElementIdx;
-						}
-						else
-						{
-							OverlayTri[SubIdx] = *ElementIdx;
-						}
-					}
-				}
-
-				UVLayer->SetTriangle(TID, OverlayTri);
-			}
+			Triangles.Add(tid);
 		}
+		UVEditor.ResetUVs();
+		UVEditor.SetTriangleUVsFromExpMap(Triangles, [&](const FVector3d P) { return MeshToProjectionSpace.TransformPosition(P); },
+			ProjectionFrame, Dimensions, SmoothingRounds, SmoothingAlpha, BlendWeight, &EditResult);
+	}
+	else
+	{
+		UVEditor.SetTriangleUVsFromExpMap(*TriangleROI, [&](const FVector3d P) { return MeshToProjectionSpace.TransformPosition(P); },
+			ProjectionFrame, Dimensions, SmoothingRounds, SmoothingAlpha, BlendWeight, &EditResult);
+		FCompactMaps CompactMaps;
+		//UVEditor.GetOverlay()->CompactInPlace(CompactMaps);
+	}
 
+	ApplyUVTransforms(UVEditor, EditResult.NewUVElements, UVOrigin, UVRotationAngleDeg, UVScale, UVTranslate);
+}
+
+
+void FUVProjectionOp::CalculateResult_Cylinder(FProgressCancel* Progress)
+{
+	int LayerIndex = FMath::Min(UseUVLayer, 8);		// hardcoded max UV layers
+	FDynamicMeshUVEditor UVEditor(ResultMesh.Get(), UseUVLayer, true);
 
-		// rescale UVs to an average world-space target length (and then re-multiply by the UVScale param)
-		if (bWorldSpaceUVScale)
+	FUVEditResult EditResult;
+	if (!TriangleROI || TriangleROI->Num() == 0)
+	{
+		TArray<int32> Triangles;
+		for (int32 tid : ResultMesh->TriangleIndicesItr())
 		{
-			FDynamicMeshEditor Editor(ResultMesh.Get());
-			// global rescale with a 'standard' global scale factor
-			Editor.RescaleAttributeUVs(.01, true, 0, ResultTransform);
-			// apply the 2D uvscale after
-			for (int ElID : UVLayer->ElementIndicesItr())
-			{
-				UVLayer->SetElement(ElID, UVLayer->GetElement(ElID) * UVScale);
-			}
+			Triangles.Add(tid);
 		}
+		UVEditor.ResetUVs();
+		UVEditor.SetTriangleUVsFromCylinderProjection(Triangles, [&](const FVector3d P) { return MeshToProjectionSpace.TransformPosition(P); },
+			ProjectionBox.Frame, 2 * ProjectionBox.Extents, CylinderSplitAngle, &EditResult);
+	}
+	else
+	{
+		UVEditor.SetTriangleUVsFromCylinderProjection(*TriangleROI, [&](const FVector3d P) { return MeshToProjectionSpace.TransformPosition(P); },
+			ProjectionBox.Frame, 2 * ProjectionBox.Extents, CylinderSplitAngle, &EditResult);
+		FCompactMaps CompactMaps;
+		//UVEditor.GetOverlay()->CompactInPlace(CompactMaps);
 	}
 
+	ApplyUVTransforms(UVEditor, EditResult.NewUVElements, UVOrigin, UVRotationAngleDeg, UVScale, UVTranslate);
 }
diff --git a/Engine/Plugins/Experimental/MeshModelingToolset/Source/ModelingOperators/Public/ParameterizationOps/UVProjectionOp.h b/Engine/Plugins/Experimental/MeshModelingToolset/Source/ModelingOperators/Public/ParameterizationOps/UVProjectionOp.h
--- a/Engine/Plugins/Experimental/MeshModelingToolset/Source/ModelingOperators/Public/ParameterizationOps/UVProjectionOp.h
+++ b/Engine/Plugins/Experimental/MeshModelingToolset/Source/ModelingOperators/Public/ParameterizationOps/UVProjectionOp.h
@@ -5,15 +5,21 @@
 #include "CoreMinimal.h"
 #include "Util/ProgressCancel.h"
 #include "ModelingOperators.h"
+#include "OrientedBoxTypes.h"
 
 #include "UVProjectionOp.generated.h"
 
 UENUM()
 enum class EUVProjectionMethod : uint8
 {
-	Cube,
+	/** Calculate UVs by assigning triangles to faces of a Box, and then apply per-box-face planar projection */
+	Box,
+	/** Calculate UVs by assigning triangles to Cylinder (radial projection) or Endcaps (planar projection) */
 	Cylinder,
-	Plane
+	/** Calculate UVs by projecting to a 3D plane */
+	Plane,
+	/** Calculate UVs by Exponential Map Projection centered at nearest surface point to input 3D plane */
+	ExpMap
 };
 
 namespace UE
@@ -28,20 +34,43 @@ public:
 
 	// inputs
 	TSharedPtr<FDynamicMesh3, ESPMode::ThreadSafe> OriginalMesh;
+	TSharedPtr<TArray<int32>, ESPMode::ThreadSafe> TriangleROI;
+	int32 UseUVLayer = 0;
+
+	// transform from Mesh into space of Projection Geometry
+	FTransform3d MeshToProjectionSpace;
 
 	EUVProjectionMethod ProjectionMethod;
+	FOrientedBox3d ProjectionBox;
+	float CylinderSplitAngle = 45.0f;
+
+	// ExpMap parameters for modifying normals
+	double BlendWeight = 0.0;
+	int32 SmoothingRounds = 0;
+	double SmoothingAlpha = 0.5;
+
 	float CylinderProjectToTopOrBottomAngleThreshold;
-	FTransform ProjectionTransform;
-	FVector2f UVScale, UVOffset;
-	bool bWorldSpaceUVScale;
 
-	void SetTransform(const FTransform& Transform);
+	// position to use as UV origin
+	FVector2f UVOrigin = FVector2f(0.5f,0.5f);
+	// rotation applied to computed UVs
+	float UVRotationAngleDeg = 0.0;
+	// scale applied to computed UVs
+	FVector2f UVScale = FVector2f::One();
+	// translation applied to computed UVs
+	FVector2f UVTranslate = FVector2f::Zero();
+
 
 	//
 	// FDynamicMeshOperator implementation
 	// 
-
 	virtual void CalculateResult(FProgressCancel* Progress) override;
+
+
+	virtual void CalculateResult_Box(FProgressCancel* Progress);
+	virtual void CalculateResult_Plane(FProgressCancel* Progress);
+	virtual void CalculateResult_ExpMap(FProgressCancel* Progress);
+	virtual void CalculateResult_Cylinder(FProgressCancel* Progress);
 };
 
 } // end namespace UE::Geometry
