Fix some property comments

#rnx

diff --git a/Engine/Plugins/Experimental/MeshModelingToolset/Source/MeshModelingTools/Public/WeldMeshEdgesTool.h b/Engine/Plugins/Experimental/MeshModelingToolset/Source/MeshModelingTools/Public/WeldMeshEdgesTool.h
--- a/Engine/Plugins/Experimental/MeshModelingToolset/Source/MeshModelingTools/Public/WeldMeshEdgesTool.h
+++ b/Engine/Plugins/Experimental/MeshModelingToolset/Source/MeshModelingTools/Public/WeldMeshEdgesTool.h
@@ -54,11 +54,11 @@ protected:
 	// need to update bResultValid if these are modified, so we don't publicly expose them. 
 	// @todo setters/getters for these
 
-	/** Target edge length */
+	/** Edges are considered matching if both pairs of endpoint vertices are closer than this distance */
 	UPROPERTY(EditAnywhere, Category = Options, meta = (UIMin = "0.000001", UIMax = "0.01", ClampMin = "0.00000001", ClampMax = "1000.0"))
 	float Tolerance;
 
-	/** Enable edge flips */
+	/** Only merge unambiguous pairs that have unique duplicate-edge matches */
 	UPROPERTY(EditAnywhere, Category = Options, AdvancedDisplay)
 	bool bOnlyUnique;
 
