Approximate Actors HLOD

* Moved duplicated code that was constructing a IGeometryProcessing_ApproximateActors::FOptions from a FMeshApproximationSettings to FApproximateActorsImpl::ConstructOptions()
* Added mesh generation options
** bAllowDistanceField - on by default
** bSupportRayTracing - on by default
** bGenerateLightmapUVs - off by default

* Disable creation of physics body for HLOD

* Scale texture size based on grid promotion (higher z grid index -> higher area covered) - Only temporary, will probably switch to a texel density option instead

diff --git a/Engine/Plugins/Experimental/MeshModelingToolset/Source/GeometryProcessingAdapters/Private/ApproximateActorsImpl.cpp b/Engine/Plugins/Experimental/MeshModelingToolset/Source/GeometryProcessingAdapters/Private/ApproximateActorsImpl.cpp
--- a/Engine/Plugins/Experimental/MeshModelingToolset/Source/GeometryProcessingAdapters/Private/ApproximateActorsImpl.cpp
+++ b/Engine/Plugins/Experimental/MeshModelingToolset/Source/GeometryProcessingAdapters/Private/ApproximateActorsImpl.cpp
@@ -32,6 +32,8 @@
 #include "Materials/Material.h"
 #include "Materials/MaterialInstanceConstant.h"
 
+#include "Engine/MeshMerging.h"
+
 #include "Async/Async.h"
 #include "Misc/ScopedSlowTask.h"
 #include "RenderCaptureInterface.h"
@@ -643,6 +645,77 @@ static TSharedPtr<FApproximationMeshData> GenerateApproximationMesh(
 }
 
 
+IGeometryProcessing_ApproximateActors::FOptions FApproximateActorsImpl::ConstructOptions(const FMeshApproximationSettings& UseSettings)
+{
+	//
+	// Construct options for ApproximateActors operation
+	//
+	FOptions Options;
+
+	Options.BasePolicy = (UseSettings.OutputType == EMeshApproximationType::MeshShapeOnly) ?
+		IGeometryProcessing_ApproximateActors::EApproximationPolicy::CollisionMesh :
+		IGeometryProcessing_ApproximateActors::EApproximationPolicy::MeshAndGeneratedMaterial;
+	Options.WorldSpaceApproximationAccuracyMeters = UseSettings.ApproximationAccuracy;
+
+	Options.bAutoThickenThinParts = UseSettings.bAttemptAutoThickening;
+	Options.AutoThickenThicknessMeters = UseSettings.TargetMinThicknessMultiplier * UseSettings.ApproximationAccuracy;
+
+	Options.BaseCappingPolicy = IGeometryProcessing_ApproximateActors::EBaseCappingPolicy::NoBaseCapping;
+	if (UseSettings.BaseCapping == EMeshApproximationBaseCappingType::ConvexPolygon)
+	{
+		Options.BaseCappingPolicy = IGeometryProcessing_ApproximateActors::EBaseCappingPolicy::ConvexPolygon;
+	}
+	else if (UseSettings.BaseCapping == EMeshApproximationBaseCappingType::ConvexSolid)
+	{
+		Options.BaseCappingPolicy = IGeometryProcessing_ApproximateActors::EBaseCappingPolicy::ConvexSolid;
+	}
+
+	Options.ClampVoxelDimension = UseSettings.ClampVoxelDimension;
+	Options.WindingThreshold = UseSettings.WindingThreshold;
+	Options.bApplyMorphology = UseSettings.bFillGaps;
+	Options.MorphologyDistanceMeters = UseSettings.GapDistance;
+
+	Options.OcclusionPolicy = (UseSettings.OcclusionMethod == EOccludedGeometryFilteringPolicy::VisibilityBasedFiltering) ?
+		IGeometryProcessing_ApproximateActors::EOcclusionPolicy::VisibilityBased : IGeometryProcessing_ApproximateActors::EOcclusionPolicy::None;
+	Options.FixedTriangleCount = UseSettings.TargetTriCount;
+	if (UseSettings.SimplifyMethod == EMeshApproximationSimplificationPolicy::TrianglesPerArea)
+	{
+		Options.MeshSimplificationPolicy = IGeometryProcessing_ApproximateActors::ESimplificationPolicy::TrianglesPerUnitSqMeter;
+		Options.SimplificationTargetMetric = UseSettings.TrianglesPerM;
+	}
+	else if (UseSettings.SimplifyMethod == EMeshApproximationSimplificationPolicy::GeometricTolerance)
+	{
+		Options.MeshSimplificationPolicy = IGeometryProcessing_ApproximateActors::ESimplificationPolicy::GeometricTolerance;
+		Options.SimplificationTargetMetric = UseSettings.GeometricDeviation;
+	}
+	else
+	{
+		Options.MeshSimplificationPolicy = IGeometryProcessing_ApproximateActors::ESimplificationPolicy::FixedTriangleCount;
+	}
+
+	Options.TextureImageSize = UseSettings.MaterialSettings.TextureSize.X;
+	Options.AntiAliasMultiSampling = FMath::Max(1, UseSettings.MultiSamplingAA);
+
+	Options.RenderCaptureImageSize = (UseSettings.RenderCaptureResolution == 0) ?
+		Options.TextureImageSize : UseSettings.RenderCaptureResolution;
+	Options.FieldOfViewDegrees = UseSettings.CaptureFieldOfView;
+	Options.NearPlaneDist = UseSettings.NearPlaneDist;
+
+	Options.bVerbose = UseSettings.bPrintDebugMessages;
+	Options.bWriteDebugMesh = UseSettings.bEmitFullDebugMesh;
+
+	// Nanite settings
+	Options.bGenerateNaniteEnabledMesh = UseSettings.bGenerateNaniteEnabledMesh;
+	Options.NaniteProxyTrianglePercent = UseSettings.NaniteProxyTrianglePercent;
+
+	// Distance field
+	Options.bAllowDistanceField = UseSettings.bAllowDistanceField;
+
+	// Ray tracing
+	Options.bSupportRayTracing = UseSettings.bSupportRayTracing;
+
+	return Options;
+}
 
 
 void FApproximateActorsImpl::ApproximateActors(const TArray<AActor*>& Actors, const FOptions& Options, FResults& ResultsOut)
@@ -880,6 +953,11 @@ UStaticMesh* FApproximateActorsImpl::EmitGeneratedMeshAsset(
 	MeshAssetOptions.bGenerateNaniteEnabledMesh = Options.bGenerateNaniteEnabledMesh;
 	MeshAssetOptions.NaniteProxyTrianglePercent = Options.NaniteProxyTrianglePercent;
 
+	MeshAssetOptions.bSupportRayTracing = Options.bSupportRayTracing;
+	MeshAssetOptions.bAllowDistanceField = Options.bAllowDistanceField;
+	MeshAssetOptions.bGenerateLightmapUVs = Options.bGenerateLightmapUVs;
+	MeshAssetOptions.bCreatePhysicsBody = Options.bCreatePhysicsBody;
+
 	if (Material)
 	{
 		MeshAssetOptions.AssetMaterials.Add(Material);
diff --git a/Engine/Plugins/Experimental/MeshModelingToolset/Source/GeometryProcessingAdapters/Public/GeometryProcessing/ApproximateActorsImpl.h b/Engine/Plugins/Experimental/MeshModelingToolset/Source/GeometryProcessingAdapters/Public/GeometryProcessing/ApproximateActorsImpl.h
--- a/Engine/Plugins/Experimental/MeshModelingToolset/Source/GeometryProcessingAdapters/Public/GeometryProcessing/ApproximateActorsImpl.h
+++ b/Engine/Plugins/Experimental/MeshModelingToolset/Source/GeometryProcessingAdapters/Public/GeometryProcessing/ApproximateActorsImpl.h
@@ -19,8 +19,9 @@ class FDynamicMesh3;
  */
 class GEOMETRYPROCESSINGADAPTERS_API FApproximateActorsImpl : public IGeometryProcessing_ApproximateActors
 {
-public:	
-	
+public:
+	virtual FOptions ConstructOptions(const FMeshApproximationSettings& MeshApproximationSettings) override;
+
 	virtual void ApproximateActors(const TArray<AActor*>& Actors, const FOptions& Options, FResults& ResultsOut) override;
 	
 
diff --git a/Engine/Plugins/Experimental/MeshModelingToolset/Source/ModelingComponentsEditorOnly/Private/AssetUtils/CreateStaticMeshUtil.cpp b/Engine/Plugins/Experimental/MeshModelingToolset/Source/ModelingComponentsEditorOnly/Private/AssetUtils/CreateStaticMeshUtil.cpp
--- a/Engine/Plugins/Experimental/MeshModelingToolset/Source/ModelingComponentsEditorOnly/Private/AssetUtils/CreateStaticMeshUtil.cpp
+++ b/Engine/Plugins/Experimental/MeshModelingToolset/Source/ModelingComponentsEditorOnly/Private/AssetUtils/CreateStaticMeshUtil.cpp
@@ -52,8 +52,17 @@ UE::AssetUtils::ECreateStaticMeshResult UE::AssetUtils::CreateStaticMeshAsset(
 	NewStaticMesh->SetNumSourceModels(UseNumSourceModels);
 	for (int32 k = 0; k < UseNumSourceModels; ++k)
 	{
-		NewStaticMesh->GetSourceModel(k).BuildSettings.bRecomputeNormals = Options.bEnableRecomputeNormals;
-		NewStaticMesh->GetSourceModel(k).BuildSettings.bRecomputeTangents = Options.bEnableRecomputeTangents;
+		FMeshBuildSettings& BuildSettings = NewStaticMesh->GetSourceModel(k).BuildSettings;
+
+		BuildSettings.bRecomputeNormals = Options.bEnableRecomputeNormals;
+		BuildSettings.bRecomputeTangents = Options.bEnableRecomputeTangents;
+		BuildSettings.bGenerateLightmapUVs = Options.bGenerateLightmapUVs;
+
+		if (!Options.bAllowDistanceField)
+		{
+			BuildSettings.DistanceFieldResolutionScale = 0.0f;
+		}
+
 		NewStaticMesh->CreateMeshDescription(k);
 	}
 
@@ -127,6 +136,12 @@ UE::AssetUtils::ECreateStaticMeshResult UE::AssetUtils::CreateStaticMeshAsset(
 		NewStaticMesh->NaniteSettings.PositionPrecision = MIN_int32;
 	}
 
+	// Ray tracing
+	NewStaticMesh->bSupportRayTracing = Options.bSupportRayTracing;
+
+	// Distance field
+	NewStaticMesh->bGenerateMeshDistanceField = Options.bAllowDistanceField;
+
 	NewStaticMesh->MarkPackageDirty();
 	if (Options.bDeferPostEditChange == false)
 	{
diff --git a/Engine/Plugins/Experimental/MeshModelingToolset/Source/ModelingComponentsEditorOnly/Public/AssetUtils/CreateStaticMeshUtil.h b/Engine/Plugins/Experimental/MeshModelingToolset/Source/ModelingComponentsEditorOnly/Public/AssetUtils/CreateStaticMeshUtil.h
--- a/Engine/Plugins/Experimental/MeshModelingToolset/Source/ModelingComponentsEditorOnly/Public/AssetUtils/CreateStaticMeshUtil.h
+++ b/Engine/Plugins/Experimental/MeshModelingToolset/Source/ModelingComponentsEditorOnly/Public/AssetUtils/CreateStaticMeshUtil.h
@@ -67,6 +67,15 @@ namespace AssetUtils
 		// Percentage of triangles to reduce down to for generating a coarse proxy mesh from the Nanite mesh
 		float NaniteProxyTrianglePercent = 100;
 
+		// Whether ray tracing will be supported on this mesh. Disable this to save memory if the generated mesh will only be rendered in the distance
+		bool bSupportRayTracing = true;
+
+		// Whether to allow distance field to be computed for this mesh. Disable this to save memory if the generated mesh will only be rendered in the distance
+		bool bAllowDistanceField = true;
+
+		// Whether to generate lightmap uvs for the generated mesh
+		bool bGenerateLightmapUVs = false;
+
 		// Controls whether the UBodySetup on the Asset will be created (generally should be true)
 		bool bCreatePhysicsBody = true;
 		// set asset collision type
diff --git a/Engine/Source/Developer/GeometryProcessingInterfaces/Public/GeometryProcessingInterfaces/ApproximateActors.h b/Engine/Source/Developer/GeometryProcessingInterfaces/Public/GeometryProcessingInterfaces/ApproximateActors.h
--- a/Engine/Source/Developer/GeometryProcessingInterfaces/Public/GeometryProcessingInterfaces/ApproximateActors.h
+++ b/Engine/Source/Developer/GeometryProcessingInterfaces/Public/GeometryProcessingInterfaces/ApproximateActors.h
@@ -10,6 +10,7 @@ class UStaticMesh;
 class UMaterial;
 class UMaterialInterface;
 class UTexture2D;
+struct FMeshApproximationSettings;
 
 
 /**
@@ -139,6 +140,29 @@ public:
 		int32 AntiAliasMultiSampling = 0;
 
 
+		//
+		// Mesh settings
+		//
+
+		// Whether to generate a nanite-enabled mesh
+		bool bGenerateNaniteEnabledMesh = false;
+
+		// Percentage of triangles to retain for the nanite coarse mesh
+		float NaniteProxyTrianglePercent = 100;
+
+		// Whether ray tracing will be supported on this mesh. Disable this to save memory if the generated mesh will only be rendered in the distance
+		bool bSupportRayTracing = true;
+
+		// Whether to allow distance field to be computed for this mesh. Disable this to save memory if the generated mesh will only be rendered in the distance
+		bool bAllowDistanceField = true;
+
+		// Whether to generate lightmap uvs for the merged mesh
+		bool bGenerateLightmapUVs = false;
+
+		// Whether the UBodySetup on the Asset will be created
+		bool bCreatePhysicsBody = true;
+
+
 		//
 		// Debug settings
 		//
@@ -149,12 +173,6 @@ public:
 		// create a flattened (ie non-instanced) mesh and save it with _DEBUG suffix. Warning often absolutely enormous!
 		bool bWriteDebugMesh = false;
 
-		//
-		// Nanite settings
-		//
-		bool bGenerateNaniteEnabledMesh = false;
-		float NaniteProxyTrianglePercent = 100;
-
 		// todo
 		//   - which texture types to generate (BaseColor, Specular, WorldNormal, etc)
 		//   - per-texture resolution, and channel/compression settings
@@ -185,6 +203,15 @@ public:
 		// todo
 	};
 
+	/**
+	 * Construct an FOptions from the provided FMeshApproximationSettings.
+	 */
+	virtual FOptions ConstructOptions(const FMeshApproximationSettings& MeshApproximationSettings)
+	{
+		check(false);		// not implemented in base class
+		return FOptions();
+	}
+
 	/**
 	 * Top-level driver function that clients call to generate the approximation for a set of input Actors.
 	 */
diff --git a/Engine/Source/Developer/WorldPartitionHLODUtilities/Private/HLODBuilderMeshApproximate.cpp b/Engine/Source/Developer/WorldPartitionHLODUtilities/Private/HLODBuilderMeshApproximate.cpp
--- a/Engine/Source/Developer/WorldPartitionHLODUtilities/Private/HLODBuilderMeshApproximate.cpp
+++ b/Engine/Source/Developer/WorldPartitionHLODUtilities/Private/HLODBuilderMeshApproximate.cpp
@@ -32,63 +32,10 @@ TArray<UPrimitiveComponent*> FHLODBuilder_MeshApproximate::CreateComponents(AWor
 
 	const FMeshApproximationSettings& UseSettings = InHLODLayer->GetMeshApproximationSettings();
 
-	IGeometryProcessing_ApproximateActors::FOptions Options;
+	IGeometryProcessing_ApproximateActors::FOptions Options = ApproxActorsAPI->ConstructOptions(UseSettings);
 	Options.BasePackagePath = InHLODActor->GetPackage()->GetName();
-	Options.BasePolicy = (UseSettings.OutputType == EMeshApproximationType::MeshShapeOnly) ?
-		IGeometryProcessing_ApproximateActors::EApproximationPolicy::CollisionMesh :
-		IGeometryProcessing_ApproximateActors::EApproximationPolicy::MeshAndGeneratedMaterial;
-	Options.WorldSpaceApproximationAccuracyMeters = UseSettings.ApproximationAccuracy;
-
-	Options.bAutoThickenThinParts = UseSettings.bAttemptAutoThickening;
-	Options.AutoThickenThicknessMeters = UseSettings.TargetMinThicknessMultiplier * UseSettings.ApproximationAccuracy;
-
-	Options.BaseCappingPolicy = IGeometryProcessing_ApproximateActors::EBaseCappingPolicy::NoBaseCapping;
-	if (UseSettings.BaseCapping == EMeshApproximationBaseCappingType::ConvexPolygon)
-	{
-		Options.BaseCappingPolicy = IGeometryProcessing_ApproximateActors::EBaseCappingPolicy::ConvexPolygon;
-	}
-	else if (UseSettings.BaseCapping == EMeshApproximationBaseCappingType::ConvexSolid)
-	{
-		Options.BaseCappingPolicy = IGeometryProcessing_ApproximateActors::EBaseCappingPolicy::ConvexSolid;
-	}
-
-	Options.ClampVoxelDimension = UseSettings.ClampVoxelDimension;
-	Options.WindingThreshold = UseSettings.WindingThreshold;
-	Options.bApplyMorphology = UseSettings.bFillGaps;
-	Options.MorphologyDistanceMeters = UseSettings.GapDistance;
-
-	Options.OcclusionPolicy = (UseSettings.OcclusionMethod == EOccludedGeometryFilteringPolicy::VisibilityBasedFiltering) ?
-		IGeometryProcessing_ApproximateActors::EOcclusionPolicy::VisibilityBased : IGeometryProcessing_ApproximateActors::EOcclusionPolicy::None;
-	Options.FixedTriangleCount = UseSettings.TargetTriCount;
-	if (UseSettings.SimplifyMethod == EMeshApproximationSimplificationPolicy::TrianglesPerArea)
-	{
-		Options.MeshSimplificationPolicy = IGeometryProcessing_ApproximateActors::ESimplificationPolicy::TrianglesPerUnitSqMeter;
-		Options.SimplificationTargetMetric = UseSettings.TrianglesPerM;
-	}
-	else if (UseSettings.SimplifyMethod == EMeshApproximationSimplificationPolicy::GeometricTolerance)
-	{
-		Options.MeshSimplificationPolicy = IGeometryProcessing_ApproximateActors::ESimplificationPolicy::GeometricTolerance;
-		Options.SimplificationTargetMetric = UseSettings.GeometricDeviation;
-	}
-	else
-	{
-		Options.MeshSimplificationPolicy = IGeometryProcessing_ApproximateActors::ESimplificationPolicy::FixedTriangleCount;
-	}
-
-	Options.TextureImageSize = UseSettings.MaterialSettings.TextureSize.X;
-	Options.AntiAliasMultiSampling = FMath::Max(1, UseSettings.MultiSamplingAA);
-
-	Options.RenderCaptureImageSize = (UseSettings.RenderCaptureResolution == 0) ?
-		Options.TextureImageSize : UseSettings.RenderCaptureResolution;
-	Options.FieldOfViewDegrees = UseSettings.CaptureFieldOfView;
-	Options.NearPlaneDist = UseSettings.NearPlaneDist;
-
-	Options.bVerbose = UseSettings.bPrintDebugMessages;
-	Options.bWriteDebugMesh = UseSettings.bEmitFullDebugMesh;
-
-	// Nanite settings
-	Options.bGenerateNaniteEnabledMesh = UseSettings.bGenerateNaniteEnabledMesh;
-	Options.NaniteProxyTrianglePercent = UseSettings.NaniteProxyTrianglePercent;
+	Options.bGenerateLightmapUVs = false;
+	Options.bCreatePhysicsBody = false;
 
 	// Material baking settings
 	Options.BakeMaterial = GEngine->DefaultHLODFlattenMaterial;
@@ -101,6 +48,15 @@ TArray<UPrimitiveComponent*> FHLODBuilder_MeshApproximate::CreateComponents(AWor
 	Options.bUsePackedMRS = true;
 	Options.PackedMRSTexParamName = FName("PackedTexture");
 
+	// Scale capture size based on grid promotion (higher z grid index -> higher area covered)
+	uint64 GridIndexX = 0;
+	uint64 GridIndexY = 0;
+	uint64 GridIndexZ = 0;
+	InHLODActor->GetGridIndices(GridIndexX, GridIndexY, GridIndexZ);
+	uint64 TexSizeMultiplier = (uint64)FMath::Pow(2.0f, GridIndexZ);
+	Options.TextureImageSize *= TexSizeMultiplier;
+	Options.RenderCaptureImageSize *= TexSizeMultiplier;
+
 	// run actor approximation computation
 	IGeometryProcessing_ApproximateActors::FResults Results;
 	ApproxActorsAPI->ApproximateActors(Actors.Array(), Options, Results);
diff --git a/Engine/Source/Editor/MergeActors/Private/MeshApproximationTool/MeshApproximationTool.cpp b/Engine/Source/Editor/MergeActors/Private/MeshApproximationTool/MeshApproximationTool.cpp
--- a/Engine/Source/Editor/MergeActors/Private/MeshApproximationTool/MeshApproximationTool.cpp
+++ b/Engine/Source/Editor/MergeActors/Private/MeshApproximationTool/MeshApproximationTool.cpp
@@ -166,63 +166,9 @@ bool FMeshApproximationTool::RunMerge(const FString& PackageName, const TArray<T
 	// Construct options for ApproximateActors operation
 	//
 
-	IGeometryProcessing_ApproximateActors::FOptions Options;
+	IGeometryProcessing_ApproximateActors::FOptions Options = ApproxActorsAPI->ConstructOptions(UseSettings);
 	Options.BasePackagePath = PackageName;
-	Options.BasePolicy = (UseSettings.OutputType == EMeshApproximationType::MeshShapeOnly) ?
-		IGeometryProcessing_ApproximateActors::EApproximationPolicy::CollisionMesh :
-		IGeometryProcessing_ApproximateActors::EApproximationPolicy::MeshAndGeneratedMaterial;
-	Options.WorldSpaceApproximationAccuracyMeters = UseSettings.ApproximationAccuracy;
-
-	Options.bAutoThickenThinParts = UseSettings.bAttemptAutoThickening;
-	Options.AutoThickenThicknessMeters = UseSettings.TargetMinThicknessMultiplier * UseSettings.ApproximationAccuracy;
-
-	Options.BaseCappingPolicy = IGeometryProcessing_ApproximateActors::EBaseCappingPolicy::NoBaseCapping;
-	if (UseSettings.BaseCapping == EMeshApproximationBaseCappingType::ConvexPolygon)
-	{
-		Options.BaseCappingPolicy = IGeometryProcessing_ApproximateActors::EBaseCappingPolicy::ConvexPolygon;
-	}
-	else if (UseSettings.BaseCapping == EMeshApproximationBaseCappingType::ConvexSolid)
-	{
-		Options.BaseCappingPolicy = IGeometryProcessing_ApproximateActors::EBaseCappingPolicy::ConvexSolid;
-	}
-
-	Options.ClampVoxelDimension = UseSettings.ClampVoxelDimension;
-	Options.WindingThreshold = UseSettings.WindingThreshold;
-	Options.bApplyMorphology = UseSettings.bFillGaps;
-	Options.MorphologyDistanceMeters = UseSettings.GapDistance;
-
-	Options.OcclusionPolicy = (UseSettings.OcclusionMethod == EOccludedGeometryFilteringPolicy::VisibilityBasedFiltering) ?
-		IGeometryProcessing_ApproximateActors::EOcclusionPolicy::VisibilityBased : IGeometryProcessing_ApproximateActors::EOcclusionPolicy::None;
-	Options.FixedTriangleCount = UseSettings.TargetTriCount;
-	if (UseSettings.SimplifyMethod == EMeshApproximationSimplificationPolicy::TrianglesPerArea)
-	{
-		Options.MeshSimplificationPolicy = IGeometryProcessing_ApproximateActors::ESimplificationPolicy::TrianglesPerUnitSqMeter;
-		Options.SimplificationTargetMetric = UseSettings.TrianglesPerM;
-	}
-	else if (UseSettings.SimplifyMethod == EMeshApproximationSimplificationPolicy::GeometricTolerance)
-	{
-		Options.MeshSimplificationPolicy = IGeometryProcessing_ApproximateActors::ESimplificationPolicy::GeometricTolerance;
-		Options.SimplificationTargetMetric = UseSettings.GeometricDeviation;
-	}
-	else
-	{
-		Options.MeshSimplificationPolicy = IGeometryProcessing_ApproximateActors::ESimplificationPolicy::FixedTriangleCount;
-	}
-
-	Options.TextureImageSize = UseSettings.MaterialSettings.TextureSize.X;
-	Options.AntiAliasMultiSampling = FMath::Max(1, UseSettings.MultiSamplingAA);
-
-	Options.RenderCaptureImageSize = (UseSettings.RenderCaptureResolution == 0) ?
-		Options.TextureImageSize : UseSettings.RenderCaptureResolution;
-	Options.FieldOfViewDegrees = UseSettings.CaptureFieldOfView;
-	Options.NearPlaneDist = UseSettings.NearPlaneDist;
-
-	Options.bVerbose = UseSettings.bPrintDebugMessages;
-	Options.bWriteDebugMesh = UseSettings.bEmitFullDebugMesh;
-
-	Options.bGenerateNaniteEnabledMesh = UseSettings.bGenerateNaniteEnabledMesh;
-	Options.NaniteProxyTrianglePercent = UseSettings.NaniteProxyTrianglePercent;
-
+	
 	// run actor approximation computation
 	IGeometryProcessing_ApproximateActors::FResults Results;
 	ApproxActorsAPI->ApproximateActors(Actors, Options, Results);
diff --git a/Engine/Source/Runtime/Engine/Classes/Engine/MeshMerging.h b/Engine/Source/Runtime/Engine/Classes/Engine/MeshMerging.h
--- a/Engine/Source/Runtime/Engine/Classes/Engine/MeshMerging.h
+++ b/Engine/Source/Runtime/Engine/Classes/Engine/MeshMerging.h
@@ -769,13 +769,21 @@ struct FMeshApproximationSettings
 	float GeometricDeviation = 0.1f;
 
 	/** Whether to generate a nanite-enabled mesh */
-	UPROPERTY(BlueprintReadWrite, EditAnywhere, Category = NaniteSettings)
+	UPROPERTY(BlueprintReadWrite, EditAnywhere, Category = MeshSettings)
 	bool bGenerateNaniteEnabledMesh = false;
 
 	/** Percentage of triangles to reduce down to for generating a coarse proxy mesh from the Nanite mesh */
-	UPROPERTY(BlueprintReadWrite, EditAnywhere, AdvancedDisplay, Category = NaniteSettings, meta = (EditConditionHides, EditCondition = "bGenerateNaniteEnabledMesh", ClampMin = 0, ClampMax = 100))
+	UPROPERTY(BlueprintReadWrite, EditAnywhere, AdvancedDisplay, Category = MeshSettings, meta = (EditConditionHides, EditCondition = "bGenerateNaniteEnabledMesh", ClampMin = 0, ClampMax = 100))
 	float NaniteProxyTrianglePercent = 100;
 
+	/** Whether ray tracing will be supported on this mesh. Disable this to save memory if the generated mesh will only be rendered in the distance. */
+	UPROPERTY(BlueprintReadWrite, EditAnywhere, Category = MeshSettings)
+	bool bSupportRayTracing = true;
+
+	/** Whether to allow distance field to be computed for this mesh. Disable this to save memory if the generated mesh will only be rendered in the distance. */
+	UPROPERTY(BlueprintReadWrite, EditAnywhere, Category = MeshSettings)
+	bool bAllowDistanceField = true;
+
 	/** If Value is > 1, Multisample output baked textures by this amount in each direction (eg 4 == 16x supersampling) */
 	UPROPERTY(BlueprintReadWrite, EditAnywhere, Category = MaterialSettings, meta = (ClampMin = "0", ClampMax = "8", UIMin = "0", UIMax = "4"))
 	int32 MultiSamplingAA = 0;
@@ -824,6 +832,8 @@ struct FMeshApproximationSettings
 			&& GeometricDeviation == Other.GeometricDeviation
 			&& bGenerateNaniteEnabledMesh == Other.bGenerateNaniteEnabledMesh
 			&& NaniteProxyTrianglePercent == Other.NaniteProxyTrianglePercent
+			&& bSupportRayTracing == Other.bSupportRayTracing
+			&& bAllowDistanceField == Other.bAllowDistanceField
 			&& MultiSamplingAA == Other.MultiSamplingAA
 			&& RenderCaptureResolution == Other.RenderCaptureResolution
 			&& MaterialSettings == Other.MaterialSettings
