ModelingTools: Update BakeMeshAttributeMapsTool to the new baking framework.

#rnx

diff --git a/Engine/Plugins/Experimental/MeshModelingToolset/Source/MeshModelingTools/Private/BakeMeshAttributeMapsTool.cpp b/Engine/Plugins/Experimental/MeshModelingToolset/Source/MeshModelingTools/Private/BakeMeshAttributeMapsTool.cpp
--- a/Engine/Plugins/Experimental/MeshModelingToolset/Source/MeshModelingTools/Private/BakeMeshAttributeMapsTool.cpp
+++ b/Engine/Plugins/Experimental/MeshModelingToolset/Source/MeshModelingTools/Private/BakeMeshAttributeMapsTool.cpp
@@ -40,6 +40,30 @@ using namespace UE::Geometry;
 
 #define LOCTEXT_NAMESPACE "UBakeMeshAttributeMapsTool"
 
+/*
+ * Static init
+ */
+
+// Only include the Occlusion bitmask rather than its components
+// (AmbientOcclusion | BentNormal). Since the Occlusion baker can
+// bake both types in a single pass, only iterating over the Occlusion
+// bitmask gives direct access to both types without the need to
+// externally track if we've handled the Occlusion evaluator in a prior
+// iteration loop.
+static constexpr EBakeMapType ALL_BAKE_MAP_TYPES[] =
+{
+	EBakeMapType::TangentSpaceNormalMap,
+	EBakeMapType::Occlusion, // (AmbientOcclusion | BentNormal)
+	EBakeMapType::Curvature,
+	EBakeMapType::Texture2DImage,
+	EBakeMapType::NormalImage,
+	EBakeMapType::FaceNormalImage,
+	EBakeMapType::PositionImage,
+	EBakeMapType::MaterialID,
+	EBakeMapType::MultiTexture
+};
+
+
 /*
  * ToolBuilder
  */
@@ -84,20 +108,32 @@ TArray<FString> UBakeMeshAttributeMapsToolProperties::GetUVLayerNamesFunc()
  * Operators
  */
 
-class FBakeMapBaseOp : public TGenericDataOperator<FMeshMapBaker>
+class FMeshMapBakerOp : public TGenericDataOperator<FMeshMapBaker>
 {
 public:
-	virtual ~FBakeMapBaseOp() {}
-
+	// General bake settings
 	TSharedPtr<UE::Geometry::FDynamicMesh3, ESPMode::ThreadSafe> DetailMesh;
 	TSharedPtr<UE::Geometry::FDynamicMeshAABBTree3, ESPMode::ThreadSafe> DetailSpatial;
 	UE::Geometry::FDynamicMesh3* BaseMesh;
 	TUniquePtr<UE::Geometry::FMeshMapBaker> Baker;
 	UBakeMeshAttributeMapsTool::FBakeCacheSettings BakeCacheSettings;
+	TSharedPtr<UE::Geometry::TMeshTangents<double>, ESPMode::ThreadSafe> BaseMeshTangents;
+
+	// Map Type settings
+	EBakeMapType Maps;
+	UBakeMeshAttributeMapsTool::FNormalMapSettings NormalSettings;
+	UBakeMeshAttributeMapsTool::FOcclusionMapSettings OcclusionSettings;
+	UBakeMeshAttributeMapsTool::FCurvatureMapSettings CurvatureSettings;
+	UBakeMeshAttributeMapsTool::FMeshPropertyMapSettings PropertySettings;
+	UBakeMeshAttributeMapsTool::FTexture2DImageSettings TextureSettings;
 
-	//
-	// TGenericDataOperator implementation
-	//
+	// Texture2DImage & MultiTexture settings
+	using ImagePtr = TSharedPtr<UE::Geometry::TImageBuilder<FVector4f>, ESPMode::ThreadSafe>;
+	const FDynamicMeshUVOverlay* UVOverlay = nullptr;
+	ImagePtr TextureImage;
+	TMap<int32, ImagePtr> MaterialToTextureImageMap;
+
+	// Begin TGenericDataOperator interface
 	virtual void CalculateResult(FProgressCancel* Progress) override
 	{
 		Baker = MakeUnique<FMeshMapBaker>();
@@ -110,234 +146,124 @@ public:
 		Baker->SetUVLayer(BakeCacheSettings.UVLayer);
 		Baker->SetThickness(BakeCacheSettings.Thickness);
 		Baker->SetMultisampling(BakeCacheSettings.Multisampling);
-	}
-};
-
-class FBakeNormalMapOp : public FBakeMapBaseOp
-{
-public:
-	typedef FBakeMapBaseOp Super;
-
-	// inputs
-	TSharedPtr<UE::Geometry::TMeshTangents<double>, ESPMode::ThreadSafe> BaseMeshTangents;
-	UBakeMeshAttributeMapsTool::FNormalMapSettings Settings;
-
-public:
-	virtual ~FBakeNormalMapOp() {}
-
-	//
-	// TGenericDataOperator implementation
-	// 
-	virtual void CalculateResult(FProgressCancel* Progress) override
-	{
-		Super::CalculateResult(Progress);
-
-		if (Progress && Progress->Cancelled())
-		{
-			return;
-		}
-
-		TSharedPtr<FMeshNormalMapEvaluator, ESPMode::ThreadSafe> NormalEval = MakeShared<FMeshNormalMapEvaluator, ESPMode::ThreadSafe>();
-		Baker->AddBaker(NormalEval);
 		Baker->SetTargetMeshTangents(BaseMeshTangents);
-		Baker->Bake();
-		SetResult(MoveTemp(Baker));
-	}
-};
-
-class FBakeOcclusionMapOp : public FBakeMapBaseOp
-{
-public:
-	typedef FBakeMapBaseOp Super;
-
-	// inputs
-	TSharedPtr<UE::Geometry::TMeshTangents<double>, ESPMode::ThreadSafe> BaseMeshTangents;
-	UBakeMeshAttributeMapsTool::FOcclusionMapSettings Settings;
-
-public:
-	virtual ~FBakeOcclusionMapOp() {}
-
-	//
-	// TGenericDataOperator implementation
-	// 
-	virtual void CalculateResult(FProgressCancel* Progress) override
-	{
-		Super::CalculateResult(Progress);
-		 
-		if (Progress && Progress->Cancelled())
-		{
-			return;
-		}
-
-		TSharedPtr<FMeshOcclusionMapEvaluator> OcclusionEval = MakeShared<FMeshOcclusionMapEvaluator>();
-		OcclusionEval->OcclusionType = EMeshOcclusionMapType::All;
-		OcclusionEval->NumOcclusionRays = Settings.OcclusionRays;
-		OcclusionEval->MaxDistance = Settings.MaxDistance;
-		OcclusionEval->SpreadAngle = Settings.SpreadAngle;
-		OcclusionEval->BiasAngleDeg = Settings.BiasAngle;
-
-		switch (Settings.Distribution)
-		{
-		case EOcclusionMapDistribution::Cosine:
-			OcclusionEval->Distribution = FMeshOcclusionMapEvaluator::EDistribution::Cosine;
-			break;
-		case EOcclusionMapDistribution::Uniform:
-			OcclusionEval->Distribution = FMeshOcclusionMapEvaluator::EDistribution::Uniform;
-			break;
-		}
-
-		switch (Settings.NormalSpace)
-		{
-		case ENormalMapSpace::Tangent:
-			OcclusionEval->NormalSpace = FMeshOcclusionMapEvaluator::ESpace::Tangent;
-			break;
-		case ENormalMapSpace::Object:
-			OcclusionEval->NormalSpace = FMeshOcclusionMapEvaluator::ESpace::Object;
-			break;
-		}
-
-		Baker->SetTargetMeshTangents(BaseMeshTangents);
-		Baker->AddBaker(OcclusionEval);
-		Baker->Bake();
-		SetResult(MoveTemp(Baker));
-	}
-};
-
-class FBakeCurvatureMapOp : public FBakeMapBaseOp
-{
-public:
-	typedef FBakeMapBaseOp Super;
-
-	// inputs
-	UBakeMeshAttributeMapsTool::FCurvatureMapSettings Settings;
-
-public:
-	virtual ~FBakeCurvatureMapOp() {}
-
-	//
-	// TGenericDataOperator implementation
-	// 
-	virtual void CalculateResult(FProgressCancel* Progress) override
-	{
-		Super::CalculateResult(Progress);
-
-		if (Progress && Progress->Cancelled())
-		{
-			return;
-		}
-
-		TSharedPtr<FMeshCurvatureMapEvaluator, ESPMode::ThreadSafe> CurvatureBaker = MakeShared<FMeshCurvatureMapEvaluator, ESPMode::ThreadSafe>();
-		CurvatureBaker->RangeScale = FMathd::Clamp(Settings.RangeMultiplier, 0.0001, 1000.0);
-		CurvatureBaker->MinRangeScale = FMathd::Clamp(Settings.MinRangeMultiplier, 0.0, 1.0);
-		CurvatureBaker->UseCurvatureType = (FMeshCurvatureMapEvaluator::ECurvatureType)Settings.CurvatureType;
-		CurvatureBaker->UseColorMode = (FMeshCurvatureMapEvaluator::EColorMode)Settings.ColorMode;
-		CurvatureBaker->UseClampMode = (FMeshCurvatureMapEvaluator::EClampMode)Settings.ClampMode;
-		Baker->AddBaker(CurvatureBaker);
-		Baker->Bake();
-		SetResult(MoveTemp(Baker));
-	}
-};
-
-class FBakeMeshPropertyMapOp : public FBakeMapBaseOp
-{
-public:
-	typedef FBakeMapBaseOp Super;
-
-	// inputs
-	UBakeMeshAttributeMapsTool::FMeshPropertyMapSettings Settings;
-
-public:
-	virtual ~FBakeMeshPropertyMapOp() {}
-
-	//
-	// TGenericDataOperator implementation
-	// 
-	virtual void CalculateResult(FProgressCancel* Progress) override
-	{
-		Super::CalculateResult(Progress);
-
-		if (Progress && Progress->Cancelled())
-		{
-			return;
-		}
-
-		TSharedPtr<FMeshPropertyMapEvaluator, ESPMode::ThreadSafe> PropertyBaker = MakeShared<FMeshPropertyMapEvaluator, ESPMode::ThreadSafe>();
-		PropertyBaker->Property = (EMeshPropertyMapType) Settings.PropertyTypeIndex;
-		Baker->AddBaker(PropertyBaker);
-		Baker->Bake();
-		SetResult(MoveTemp(Baker));
-	}
-};
 
-class FBakeTexture2DImageMapOp : public FBakeMapBaseOp
-{
-public:
-	typedef FBakeMapBaseOp Super;
-
-	// inputs
-	const FDynamicMeshUVOverlay* UVOverlay = nullptr;
-	TSharedPtr<UE::Geometry::TImageBuilder<FVector4f>, ESPMode::ThreadSafe> TextureImage;
-	UBakeMeshAttributeMapsTool::FTexture2DImageSettings Settings;
-
-public:
-	virtual ~FBakeTexture2DImageMapOp() {}
-
-	//
-	// TGenericDataOperator implementation
-	// 
-	virtual void CalculateResult(FProgressCancel* Progress) override
-	{
-		Super::CalculateResult(Progress);
-
-		if (Progress && Progress->Cancelled())
+		for (const EBakeMapType MapType : ALL_BAKE_MAP_TYPES)
 		{
-			return;
-		}
-
-		TSharedPtr<FMeshResampleImageEvaluator, ESPMode::ThreadSafe> ResampleBaker = MakeShared<FMeshResampleImageEvaluator, ESPMode::ThreadSafe>();
-		ResampleBaker->DetailUVOverlay = UVOverlay;
-		ResampleBaker->SampleFunction = [this](FVector2d UVCoord) {
-			return TextureImage->BilinearSampleUV<float>(UVCoord, FVector4f(0, 0, 0, 1));
-		};
-		Baker->AddBaker(ResampleBaker);
-		Baker->Bake();
-		SetResult(MoveTemp(Baker));
-	}
-};
-
-
-
-class FBakeMultiTextureOp : public FBakeMapBaseOp
-{
-public:
-	typedef FBakeMapBaseOp Super;
-	using ImagePtr = TSharedPtr<UE::Geometry::TImageBuilder<FVector4f>, ESPMode::ThreadSafe>;
-
-	// inputs
-	const FDynamicMeshUVOverlay* UVOverlay = nullptr;
-	TMap<int32, ImagePtr> MaterialToTextureImageMap;
-
-	UBakeMeshAttributeMapsTool::FTexture2DImageSettings Settings;
-
-public:
-
-	virtual ~FBakeMultiTextureOp() {}
-
-	virtual void CalculateResult(FProgressCancel* Progress) override
-	{
-		Super::CalculateResult(Progress);
-		if (Progress && Progress->Cancelled())
-		{
-			return;
+			switch (BakeCacheSettings.BakeMapTypes & MapType)
+			{
+			case EBakeMapType::TangentSpaceNormalMap:
+			{
+				TSharedPtr<FMeshNormalMapEvaluator, ESPMode::ThreadSafe> NormalEval = MakeShared<FMeshNormalMapEvaluator, ESPMode::ThreadSafe>();
+				Baker->AddBaker(NormalEval);
+				break;
+			}
+			case EBakeMapType::AmbientOcclusion:
+			case EBakeMapType::BentNormal:
+			case EBakeMapType::Occlusion:
+			{
+				TSharedPtr<FMeshOcclusionMapEvaluator> OcclusionEval = MakeShared<FMeshOcclusionMapEvaluator>();
+				OcclusionEval->OcclusionType = EMeshOcclusionMapType::None;
+				if ((bool)(BakeCacheSettings.BakeMapTypes & EBakeMapType::AmbientOcclusion))
+				{
+					OcclusionEval->OcclusionType |= EMeshOcclusionMapType::AmbientOcclusion;
+				}
+				if ((bool)(BakeCacheSettings.BakeMapTypes & EBakeMapType::BentNormal))
+				{
+					OcclusionEval->OcclusionType |= EMeshOcclusionMapType::BentNormal;
+				}
+				OcclusionEval->NumOcclusionRays = OcclusionSettings.OcclusionRays;
+				OcclusionEval->MaxDistance = OcclusionSettings.MaxDistance;
+				OcclusionEval->SpreadAngle = OcclusionSettings.SpreadAngle;
+				OcclusionEval->BiasAngleDeg = OcclusionSettings.BiasAngle;
+
+				switch (OcclusionSettings.Distribution)
+				{
+				case EOcclusionMapDistribution::Cosine:
+					OcclusionEval->Distribution = FMeshOcclusionMapEvaluator::EDistribution::Cosine;
+					break;
+				case EOcclusionMapDistribution::Uniform:
+					OcclusionEval->Distribution = FMeshOcclusionMapEvaluator::EDistribution::Uniform;
+					break;
+				}
+
+				switch (OcclusionSettings.NormalSpace)
+				{
+				case ENormalMapSpace::Tangent:
+					OcclusionEval->NormalSpace = FMeshOcclusionMapEvaluator::ESpace::Tangent;
+					break;
+				case ENormalMapSpace::Object:
+					OcclusionEval->NormalSpace = FMeshOcclusionMapEvaluator::ESpace::Object;
+					break;
+				}
+				Baker->AddBaker(OcclusionEval);
+				break;
+			}
+			case EBakeMapType::Curvature:
+			{
+				TSharedPtr<FMeshCurvatureMapEvaluator, ESPMode::ThreadSafe> CurvatureBaker = MakeShared<FMeshCurvatureMapEvaluator, ESPMode::ThreadSafe>();
+				CurvatureBaker->RangeScale = FMathd::Clamp(CurvatureSettings.RangeMultiplier, 0.0001, 1000.0);
+				CurvatureBaker->MinRangeScale = FMathd::Clamp(CurvatureSettings.MinRangeMultiplier, 0.0, 1.0);
+				CurvatureBaker->UseCurvatureType = (FMeshCurvatureMapEvaluator::ECurvatureType)CurvatureSettings.CurvatureType;
+				CurvatureBaker->UseColorMode = (FMeshCurvatureMapEvaluator::EColorMode)CurvatureSettings.ColorMode;
+				CurvatureBaker->UseClampMode = (FMeshCurvatureMapEvaluator::EClampMode)CurvatureSettings.ClampMode;
+				Baker->AddBaker(CurvatureBaker);
+				break;
+			}
+			case EBakeMapType::NormalImage:
+			{
+				TSharedPtr<FMeshPropertyMapEvaluator, ESPMode::ThreadSafe> PropertyBaker = MakeShared<FMeshPropertyMapEvaluator, ESPMode::ThreadSafe>();
+				PropertyBaker->Property = EMeshPropertyMapType::Normal;
+				Baker->AddBaker(PropertyBaker);
+				break;
+			}
+			case EBakeMapType::FaceNormalImage:
+			{
+				TSharedPtr<FMeshPropertyMapEvaluator, ESPMode::ThreadSafe> PropertyBaker = MakeShared<FMeshPropertyMapEvaluator, ESPMode::ThreadSafe>();
+				PropertyBaker->Property = EMeshPropertyMapType::FacetNormal;
+				Baker->AddBaker(PropertyBaker);
+				break;
+			}
+			case EBakeMapType::PositionImage:
+			{
+				TSharedPtr<FMeshPropertyMapEvaluator, ESPMode::ThreadSafe> PropertyBaker = MakeShared<FMeshPropertyMapEvaluator, ESPMode::ThreadSafe>();
+				PropertyBaker->Property = EMeshPropertyMapType::Position;
+				Baker->AddBaker(PropertyBaker);
+				break;
+			}
+			case EBakeMapType::MaterialID:
+			{
+				TSharedPtr<FMeshPropertyMapEvaluator, ESPMode::ThreadSafe> PropertyBaker = MakeShared<FMeshPropertyMapEvaluator, ESPMode::ThreadSafe>();
+				PropertyBaker->Property = EMeshPropertyMapType::MaterialID;
+				Baker->AddBaker(PropertyBaker);
+				break;
+			}
+			case EBakeMapType::Texture2DImage:
+			{
+				TSharedPtr<FMeshResampleImageEvaluator, ESPMode::ThreadSafe> ResampleBaker = MakeShared<FMeshResampleImageEvaluator, ESPMode::ThreadSafe>();
+				ResampleBaker->DetailUVOverlay = UVOverlay;
+				ResampleBaker->SampleFunction = [this](FVector2d UVCoord) {
+					return TextureImage->BilinearSampleUV<float>(UVCoord, FVector4f(0, 0, 0, 1));
+				};
+				Baker->AddBaker(ResampleBaker);
+				break;
+			}
+			case EBakeMapType::MultiTexture:
+			{
+				TSharedPtr<FMeshMultiResampleImageEvaluator, ESPMode::ThreadSafe> TextureBaker = MakeShared<FMeshMultiResampleImageEvaluator, ESPMode::ThreadSafe>();
+				TextureBaker->DetailUVOverlay = UVOverlay;
+				TextureBaker->MultiTextures = MaterialToTextureImageMap;
+				Baker->AddBaker(TextureBaker);
+				break;
+			}
+			default:
+				break;
+			}
 		}
-
-		TSharedPtr<FMeshMultiResampleImageEvaluator, ESPMode::ThreadSafe> TextureBaker = MakeShared<FMeshMultiResampleImageEvaluator, ESPMode::ThreadSafe>();
-		TextureBaker->DetailUVOverlay = UVOverlay;
-		TextureBaker->MultiTextures = MaterialToTextureImageMap;
-		Baker->AddBaker(TextureBaker);
 		Baker->Bake();
 		SetResult(MoveTemp(Baker));
 	}
+	// End TGenericDataOperator interface
 };
 
 /*
@@ -422,7 +348,8 @@ void UBakeMeshAttributeMapsTool::Setup()
 	}
 	AddToolPropertySource(Settings);
 
-	Settings->WatchProperty(Settings->MapType, [this](EBakeMapType) { bInputsDirty = true; UpdateOnModeChange(); });
+	Settings->WatchProperty(Settings->MapTypes, [this](int32) { bInputsDirty = true; UpdateOnModeChange(); });
+	Settings->WatchProperty(Settings->MapPreview, [this](int32) { UpdateVisualization(); GetToolManager()->PostInvalidation(); });
 	Settings->WatchProperty(Settings->Resolution, [this](EBakeTextureResolution) { bInputsDirty = true; });
 	Settings->WatchProperty(Settings->UVLayer, [this](FString) { bInputsDirty = true; });
 	Settings->WatchProperty(Settings->bUseWorldSpace, [this](bool) { bDetailMeshValid = false; bInputsDirty = true; });
@@ -443,7 +370,6 @@ void UBakeMeshAttributeMapsTool::Setup()
 	OcclusionMapProps->WatchProperty(OcclusionMapProps->MaxDistance, [this](float) { bInputsDirty = true; });
 	OcclusionMapProps->WatchProperty(OcclusionMapProps->SpreadAngle, [this](float) { bInputsDirty = true; });
 	OcclusionMapProps->WatchProperty(OcclusionMapProps->Distribution, [this](EOcclusionMapDistribution) { bInputsDirty = true; });
-	OcclusionMapProps->WatchProperty(OcclusionMapProps->Preview, [this](EOcclusionMapPreview) { UpdateVisualization(); GetToolManager()->PostInvalidation(); });
 	OcclusionMapProps->WatchProperty(OcclusionMapProps->BlurRadius, [this](float) { bInputsDirty = true; });
 	OcclusionMapProps->WatchProperty(OcclusionMapProps->bGaussianBlur, [this](float) { bInputsDirty = true; });
 	OcclusionMapProps->WatchProperty(OcclusionMapProps->BiasAngle, [this](float) { bInputsDirty = true; });
@@ -485,6 +411,7 @@ void UBakeMeshAttributeMapsTool::Setup()
 	AddToolPropertySource(VisualizationProps);
 
 	InitializeEmptyMaps();
+	UpdateOnModeChange();
 
 	bInputsDirty = true;
 	bDetailMeshValid = false;
@@ -503,10 +430,11 @@ bool UBakeMeshAttributeMapsTool::CanAccept() const
 	bool bCanAccept = Compute ? Compute->HaveValidResult() : false;
 	if (bCanAccept)
 	{
-		bCanAccept = Settings->Result.Num() > 0;
-		for (UTexture2D* Result : Settings->Result)
+		// Allow Accept if all non-None types have valid results.
+		int NumResults = Settings->Result.Num();
+		for (int ResultIdx = 0; ResultIdx < NumResults; ++ResultIdx)
 		{
-			bCanAccept = bCanAccept && Result;
+			bCanAccept = bCanAccept && Settings->Result[ResultIdx];
 		}
 	}
 	return bCanAccept;
@@ -515,70 +443,57 @@ bool UBakeMeshAttributeMapsTool::CanAccept() const
 
 TUniquePtr<UE::Geometry::TGenericDataOperator<FMeshMapBaker>> UBakeMeshAttributeMapsTool::MakeNewOperator()
 {
-	auto InitializeBaseOp = [this](FBakeMapBaseOp* Op)
-	{
-		Op->DetailMesh = DetailMesh;
-		Op->DetailSpatial = DetailSpatial;
-		Op->BaseMesh = &BaseMesh;
-		Op->BakeCacheSettings = CachedBakeCacheSettings;
-	};
+	TUniquePtr<FMeshMapBakerOp> Op = MakeUnique<FMeshMapBakerOp>();
+	Op->DetailMesh = DetailMesh;
+	Op->DetailSpatial = DetailSpatial;
+	Op->BaseMesh = &BaseMesh;
+	Op->BakeCacheSettings = CachedBakeCacheSettings;
 
-	switch (Settings->MapType)
+	const EBakeMapType RequiresTangents = EBakeMapType::TangentSpaceNormalMap | EBakeMapType::BentNormal;
+	if ((bool)(CachedBakeCacheSettings.BakeMapTypes & RequiresTangents))
 	{
-	case EBakeMapType::TangentSpaceNormalMap:
-	{
-		TUniquePtr<FBakeNormalMapOp> Op = MakeUnique<FBakeNormalMapOp>();
-		InitializeBaseOp(Op.Get());
 		Op->BaseMeshTangents = BaseMeshTangents;
-		Op->Settings = CachedNormalMapSettings;
-		return Op;
 	}
-	case EBakeMapType::Occlusion:
+
+	if ((bool)(CachedBakeCacheSettings.BakeMapTypes & EBakeMapType::TangentSpaceNormalMap))
 	{
-		TUniquePtr<FBakeOcclusionMapOp> Op = MakeUnique<FBakeOcclusionMapOp>();
-		InitializeBaseOp(Op.Get());
-		Op->BaseMeshTangents = BaseMeshTangents;
-		Op->Settings = CachedOcclusionMapSettings;
-		return Op;
+		Op->NormalSettings = CachedNormalMapSettings;
 	}
-	case EBakeMapType::Curvature:
+
+	if ((bool)(CachedBakeCacheSettings.BakeMapTypes & EBakeMapType::AmbientOcclusion) ||
+		(bool)(CachedBakeCacheSettings.BakeMapTypes & EBakeMapType::BentNormal))
 	{
-		TUniquePtr<FBakeCurvatureMapOp> Op = MakeUnique<FBakeCurvatureMapOp>();
-		InitializeBaseOp(Op.Get());
-		Op->Settings = CachedCurvatureMapSettings;
-		return Op;
+		Op->OcclusionSettings = CachedOcclusionMapSettings;
 	}
-	case EBakeMapType::PositionImage:
-	case EBakeMapType::NormalImage:
-	case EBakeMapType::FaceNormalImage:
-	case EBakeMapType::MaterialID:
+
+	if ((bool)(CachedBakeCacheSettings.BakeMapTypes & EBakeMapType::Curvature))
 	{
-		TUniquePtr<FBakeMeshPropertyMapOp> Op = MakeUnique<FBakeMeshPropertyMapOp>();
-		InitializeBaseOp(Op.Get());
-		Op->Settings = CachedMeshPropertyMapSettings;
-		return Op;
+		Op->CurvatureSettings = CachedCurvatureMapSettings;
 	}
-	case EBakeMapType::Texture2DImage:
+
+	if ((bool)(CachedBakeCacheSettings.BakeMapTypes & EBakeMapType::NormalImage) ||
+		(bool)(CachedBakeCacheSettings.BakeMapTypes & EBakeMapType::FaceNormalImage) ||
+		(bool)(CachedBakeCacheSettings.BakeMapTypes & EBakeMapType::PositionImage) ||
+		(bool)(CachedBakeCacheSettings.BakeMapTypes & EBakeMapType::MaterialID))
 	{
-		TUniquePtr<FBakeTexture2DImageMapOp> Op = MakeUnique<FBakeTexture2DImageMapOp>();
-		InitializeBaseOp(Op.Get());
-		Op->UVOverlay = DetailMesh->Attributes()->GetUVLayer(CachedTexture2DImageSettings.UVLayer);
-		Op->TextureImage = CachedTextureImage;
-		Op->Settings = CachedTexture2DImageSettings;
-		return Op;
+		Op->PropertySettings = CachedMeshPropertyMapSettings;
 	}
-	case EBakeMapType::MultiTexture:
+
+	if ((bool)(CachedBakeCacheSettings.BakeMapTypes & EBakeMapType::Texture2DImage))
 	{
-		TUniquePtr<FBakeMultiTextureOp> Op = MakeUnique<FBakeMultiTextureOp>();
-		InitializeBaseOp(Op.Get());
+		Op->TextureSettings = CachedTexture2DImageSettings;
+		Op->TextureImage = CachedTextureImage;
 		Op->UVOverlay = DetailMesh->Attributes()->GetUVLayer(CachedTexture2DImageSettings.UVLayer);
-		Op->MaterialToTextureImageMap = CachedMultiTextures;
-		Op->Settings = CachedTexture2DImageSettings;
-		return Op;
 	}
+
+	if ((bool)(CachedBakeCacheSettings.BakeMapTypes & EBakeMapType::MultiTexture))
+	{
+		Op->TextureSettings = CachedTexture2DImageSettings;
+		Op->MaterialToTextureImageMap = CachedMultiTextures;
+		Op->UVOverlay = DetailMesh->Attributes()->GetUVLayer(CachedTexture2DImageSettings.UVLayer);
 	}
-	check(false);
-	return nullptr;
+
+	return Op;
 }
 
 
@@ -609,68 +524,57 @@ void UBakeMeshAttributeMapsTool::Shutdown(EToolShutdownType ShutdownType)
 			check(StaticMeshAsset);
 			FString BaseName = TargetComponent->GetOwnerActor()->GetName();
 
-			bool bCreatedAssetOK = false;
-			switch (Settings->MapType)
+			bool bCreatedAssetOK = true;
+			int NumResults = Settings->Result.Num();
+			for (int ResultIdx = 0; ResultIdx < NumResults; ++ResultIdx)
 			{
-			default:
-				check(false);
-				break;
-
-			case EBakeMapType::TangentSpaceNormalMap:
-				FTexture2DBuilder::CopyPlatformDataToSourceData(Settings->Result[0], FTexture2DBuilder::ETextureType::NormalMap);
-				bCreatedAssetOK = UE::Modeling::CreateTextureObject(GetToolManager(),
-					FCreateTextureObjectParams{ 0, StaticMeshAsset->GetWorld(), StaticMeshAsset, FString::Printf(TEXT("%s_Normals"), *BaseName), Settings->Result[0] }).IsOK();
-				break;
-
-			case EBakeMapType::Occlusion:
-				FTexture2DBuilder::CopyPlatformDataToSourceData(Settings->Result[0], FTexture2DBuilder::ETextureType::AmbientOcclusion);
-				bCreatedAssetOK = UE::Modeling::CreateTextureObject(GetToolManager(),
-					FCreateTextureObjectParams{ 0, StaticMeshAsset->GetWorld(), StaticMeshAsset, FString::Printf(TEXT("%s_Occlusion"), *BaseName), Settings->Result[0] }).IsOK();
-				FTexture2DBuilder::CopyPlatformDataToSourceData(Settings->Result[1], FTexture2DBuilder::ETextureType::NormalMap);
-				bCreatedAssetOK = UE::Modeling::CreateTextureObject(GetToolManager(),
-					FCreateTextureObjectParams{ 0, StaticMeshAsset->GetWorld(), StaticMeshAsset, FString::Printf(TEXT("%s_BentNormal"), *BaseName), Settings->Result[1] }).IsOK();
-				break;
-
-			case EBakeMapType::Curvature:
-				FTexture2DBuilder::CopyPlatformDataToSourceData(Settings->Result[0], FTexture2DBuilder::ETextureType::Color);
-				bCreatedAssetOK = UE::Modeling::CreateTextureObject(GetToolManager(),
-					FCreateTextureObjectParams{ 0, StaticMeshAsset->GetWorld(), StaticMeshAsset, FString::Printf(TEXT("%s_Curvature"), *BaseName), Settings->Result[0] }).IsOK();
-				break;
-
-			case EBakeMapType::NormalImage:
-				FTexture2DBuilder::CopyPlatformDataToSourceData(Settings->Result[0], FTexture2DBuilder::ETextureType::Color);
-				bCreatedAssetOK = UE::Modeling::CreateTextureObject(GetToolManager(),
-					FCreateTextureObjectParams{ 0, StaticMeshAsset->GetWorld(), StaticMeshAsset, FString::Printf(TEXT("%s_NormalImg"), *BaseName), Settings->Result[0] }).IsOK();
-				break;
-
-			case EBakeMapType::FaceNormalImage:
-				FTexture2DBuilder::CopyPlatformDataToSourceData(Settings->Result[0], FTexture2DBuilder::ETextureType::Color);
-				bCreatedAssetOK = UE::Modeling::CreateTextureObject(GetToolManager(),
-					FCreateTextureObjectParams{ 0, StaticMeshAsset->GetWorld(), StaticMeshAsset, FString::Printf(TEXT("%s_FaceNormalImg"), *BaseName), Settings->Result[0] }).IsOK();
-				break;
-			case EBakeMapType::MaterialID:
-				FTexture2DBuilder::CopyPlatformDataToSourceData(Settings->Result[0], FTexture2DBuilder::ETextureType::Color);
-				bCreatedAssetOK = UE::Modeling::CreateTextureObject(GetToolManager(),
-					FCreateTextureObjectParams{ 0, StaticMeshAsset->GetWorld(), StaticMeshAsset, FString::Printf(TEXT("%s_MaterialIDImg"), *BaseName), Settings->Result[0] }).IsOK();
-				break;
-			case EBakeMapType::PositionImage:
-				FTexture2DBuilder::CopyPlatformDataToSourceData(Settings->Result[0], FTexture2DBuilder::ETextureType::Color);
-				bCreatedAssetOK = UE::Modeling::CreateTextureObject(GetToolManager(),
-					FCreateTextureObjectParams{ 0, StaticMeshAsset->GetWorld(), StaticMeshAsset, FString::Printf(TEXT("%s_PositionImg"), *BaseName), Settings->Result[0] }).IsOK();
-				break;
-
-			case EBakeMapType::Texture2DImage:
-				FTexture2DBuilder::CopyPlatformDataToSourceData(Settings->Result[0], FTexture2DBuilder::ETextureType::Color);
-				bCreatedAssetOK = UE::Modeling::CreateTextureObject(GetToolManager(),
-					FCreateTextureObjectParams{ 0, StaticMeshAsset->GetWorld(), StaticMeshAsset, FString::Printf(TEXT("%s_TextureImg"), *BaseName), Settings->Result[0] }).IsOK();
-				break;
-
-			case EBakeMapType::MultiTexture:
-				FTexture2DBuilder::CopyPlatformDataToSourceData(Settings->Result[0], FTexture2DBuilder::ETextureType::Color);
-				bCreatedAssetOK = UE::Modeling::CreateTextureObject(GetToolManager(),
-					FCreateTextureObjectParams{ 0, StaticMeshAsset->GetWorld(), StaticMeshAsset, FString::Printf(TEXT("%s_MultiTextureImg"), *BaseName), Settings->Result[0] }).IsOK();
-				break;
-
+				FTexture2DBuilder::ETextureType TexType = FTexture2DBuilder::ETextureType::Color;
+				FString TexName;
+				switch (ResultTypes[ResultIdx])
+				{
+				default:
+					// Should never reach this case.
+					check(false);
+					continue;
+				case EBakeMapType::TangentSpaceNormalMap:
+					TexName = FString::Printf(TEXT("%s_Normals"), *BaseName);
+					TexType = FTexture2DBuilder::ETextureType::NormalMap;
+					break;
+				case EBakeMapType::AmbientOcclusion:
+					TexName = FString::Printf(TEXT("%s_Occlusion"), *BaseName);
+					TexType = FTexture2DBuilder::ETextureType::AmbientOcclusion;
+					break;
+				case EBakeMapType::BentNormal:
+					TexName = FString::Printf(TEXT("%s_BentNormal"), *BaseName);
+					TexType = FTexture2DBuilder::ETextureType::NormalMap;
+					break;
+				case EBakeMapType::Curvature:
+					TexName = FString::Printf(TEXT("%s_Curvature"), *BaseName);
+					break;
+				case EBakeMapType::NormalImage:
+					TexName = FString::Printf(TEXT("%s_NormalImg"), *BaseName);
+					break;
+				case EBakeMapType::FaceNormalImage:
+					TexName = FString::Printf(TEXT("%s_FaceNormalImg"), *BaseName);
+					break;
+				case EBakeMapType::MaterialID:
+					TexName = FString::Printf(TEXT("%s_MaterialIDImg"), *BaseName);
+					break;
+				case EBakeMapType::PositionImage:
+					TexName = FString::Printf(TEXT("%s_PositionImg"), *BaseName);
+					TexType = FTexture2DBuilder::ETextureType::Color;
+					break;
+				case EBakeMapType::Texture2DImage:
+					TexName = FString::Printf(TEXT("%s_TextureImg"), *BaseName);
+					TexType = FTexture2DBuilder::ETextureType::Color;
+					break;
+				case EBakeMapType::MultiTexture:
+					TexName = FString::Printf(TEXT("%s_MultiTextureImg"), *BaseName);
+					TexType = FTexture2DBuilder::ETextureType::Color;
+					break;
+				}
+				FTexture2DBuilder::CopyPlatformDataToSourceData(Settings->Result[ResultIdx], TexType);
+				bCreatedAssetOK = bCreatedAssetOK && UE::Modeling::CreateTextureObject(GetToolManager(), FCreateTextureObjectParams{ 0, StaticMeshAsset->GetWorld(), StaticMeshAsset, TexName, Settings->Result[ResultIdx] }).IsOK();
 			}
 			ensure(bCreatedAssetOK);
 		}
@@ -877,6 +781,9 @@ void UBakeMeshAttributeMapsTool::UpdateResult()
 	BakeCacheSettings.Thickness = Settings->Thickness;
 	BakeCacheSettings.Multisampling = (int32)Settings->Multisampling;
 
+	// process the raw bitfield before caching which may add additional targets.
+	BakeCacheSettings.BakeMapTypes = GetMapTypes(Settings->MapTypes);
+
 	// update bake cache settings
 	if (!(CachedBakeCacheSettings == BakeCacheSettings))
 	{
@@ -889,46 +796,48 @@ void UBakeMeshAttributeMapsTool::UpdateResult()
 		CachedTexture2DImageSettings = FTexture2DImageSettings();
 	}
 
-	// update map type settings
-	switch (Settings->MapType)
-	{
-		default:
-		case EBakeMapType::TangentSpaceNormalMap:
-			OpState = UpdateResult_Normal();
-			break;
-		case EBakeMapType::Occlusion:
-			OpState = UpdateResult_Occlusion();
-			break;
-		case EBakeMapType::Curvature:
-			OpState = UpdateResult_Curvature();
-			break;
-
-		case EBakeMapType::NormalImage:
-		case EBakeMapType::FaceNormalImage:
-		case EBakeMapType::PositionImage:
-		case EBakeMapType::MaterialID:
-			OpState = UpdateResult_MeshProperty();
-			break;
-
-		case EBakeMapType::Texture2DImage:
-			OpState = UpdateResult_Texture2DImage();
-			break;
-
-		case EBakeMapType::MultiTexture:
-			OpState = UpdateResult_MultiTexture();
-			break;
+	// Clear our invalid bitflag to check again for valid inputs.
+	OpState = EBakeOpState::Evaluate;
 
+	// Update map type settings
+	if ((bool)(CachedBakeCacheSettings.BakeMapTypes & EBakeMapType::TangentSpaceNormalMap))
+	{
+		OpState |= UpdateResult_Normal();
+	}
+	if ((bool)(CachedBakeCacheSettings.BakeMapTypes & EBakeMapType::AmbientOcclusion) ||
+		(bool)(CachedBakeCacheSettings.BakeMapTypes & EBakeMapType::BentNormal))
+	{
+		OpState |= UpdateResult_Occlusion();
+	}
+	if ((bool)(CachedBakeCacheSettings.BakeMapTypes & EBakeMapType::Curvature))
+	{
+		OpState |= UpdateResult_Curvature();
+	}
+	if ((bool)(CachedBakeCacheSettings.BakeMapTypes & EBakeMapType::NormalImage) ||
+		(bool)(CachedBakeCacheSettings.BakeMapTypes & EBakeMapType::FaceNormalImage) ||
+		(bool)(CachedBakeCacheSettings.BakeMapTypes & EBakeMapType::PositionImage) ||
+		(bool)(CachedBakeCacheSettings.BakeMapTypes & EBakeMapType::MaterialID))
+	{
+		OpState |= UpdateResult_MeshProperty();
+	}
+	if ((bool)(CachedBakeCacheSettings.BakeMapTypes & EBakeMapType::Texture2DImage))
+	{
+		OpState |= UpdateResult_Texture2DImage();
+	}
+	if ((bool)(CachedBakeCacheSettings.BakeMapTypes & EBakeMapType::MultiTexture))
+	{
+		OpState |= UpdateResult_MultiTexture();
 	}
 
 	// Early exit if op input parameters are invalid.
-	if (OpState == EOpState::Invalid)
+	if ((bool)(OpState & EBakeOpState::Invalid))
 	{
 		return;
 	}
 
 	// This should be the only point of compute invalidation to
 	// minimize synchronization issues.
-	bool bInvalidate = bInputsDirty || (OpState == EOpState::Evaluate);
+	bool bInvalidate = bInputsDirty || (bool)(OpState & EBakeOpState::Evaluate);
 	if (!Compute)
 	{
 		Compute = MakeUnique<TGenericDataBackgroundCompute<FMeshMapBaker>>();
@@ -945,9 +854,9 @@ void UBakeMeshAttributeMapsTool::UpdateResult()
 
 
 
-UBakeMeshAttributeMapsTool::EOpState UBakeMeshAttributeMapsTool::UpdateResult_Normal()
+EBakeOpState UBakeMeshAttributeMapsTool::UpdateResult_Normal()
 {
-	EOpState ResultState = EOpState::Complete;
+	EBakeOpState ResultState = EBakeOpState::Complete;
 
 	int32 ImageSize = (int32)Settings->Resolution;
 	FImageDimensions Dimensions(ImageSize, ImageSize);
@@ -958,15 +867,15 @@ UBakeMeshAttributeMapsTool::EOpState UBakeMeshAttributeMapsTool::UpdateResult_No
 	if (!(CachedNormalMapSettings == NormalMapSettings))
 	{
 		CachedNormalMapSettings = NormalMapSettings;
-		ResultState = EOpState::Evaluate;
+		ResultState = EBakeOpState::Evaluate;
 	}
 	return ResultState;
 }
 
 
-UBakeMeshAttributeMapsTool::EOpState UBakeMeshAttributeMapsTool::UpdateResult_Occlusion()
+EBakeOpState UBakeMeshAttributeMapsTool::UpdateResult_Occlusion()
 {
-	EOpState ResultState = EOpState::Complete;
+	EBakeOpState ResultState = EBakeOpState::Complete;
 
 	int32 ImageSize = (int32)Settings->Resolution;
 	FImageDimensions Dimensions(ImageSize, ImageSize);
@@ -984,16 +893,16 @@ UBakeMeshAttributeMapsTool::EOpState UBakeMeshAttributeMapsTool::UpdateResult_Oc
 	if ( !(CachedOcclusionMapSettings == OcclusionMapSettings) )
 	{
 		CachedOcclusionMapSettings = OcclusionMapSettings;
-		ResultState = EOpState::Evaluate;
+		ResultState = EBakeOpState::Evaluate;
 	}
 	return ResultState;
 }
 
 
 
-UBakeMeshAttributeMapsTool::EOpState UBakeMeshAttributeMapsTool::UpdateResult_Curvature()
+EBakeOpState UBakeMeshAttributeMapsTool::UpdateResult_Curvature()
 {
-	EOpState ResultState = EOpState::Complete;
+	EBakeOpState ResultState = EBakeOpState::Complete;
 
 	int32 ImageSize = (int32)Settings->Resolution;
 	FImageDimensions Dimensions(ImageSize, ImageSize);
@@ -1049,44 +958,27 @@ UBakeMeshAttributeMapsTool::EOpState UBakeMeshAttributeMapsTool::UpdateResult_Cu
 	if (!(CachedCurvatureMapSettings == CurvatureMapSettings))
 	{
 		CachedCurvatureMapSettings = CurvatureMapSettings;
-		ResultState = EOpState::Evaluate;
+		ResultState = EBakeOpState::Evaluate;
 	}
 	return ResultState;
 }
 
 
 
-UBakeMeshAttributeMapsTool::EOpState UBakeMeshAttributeMapsTool::UpdateResult_MeshProperty()
+EBakeOpState UBakeMeshAttributeMapsTool::UpdateResult_MeshProperty()
 {
-	EOpState ResultState = EOpState::Complete;
+	EBakeOpState ResultState = EBakeOpState::Complete;
 
 	int32 ImageSize = (int32)Settings->Resolution;
 	FImageDimensions Dimensions(ImageSize, ImageSize);
 
 	FMeshPropertyMapSettings MeshPropertyMapSettings;
 	MeshPropertyMapSettings.Dimensions = Dimensions;
-	switch (Settings->MapType)
-	{
-		case EBakeMapType::NormalImage:
-			MeshPropertyMapSettings.PropertyTypeIndex = (int32)EMeshPropertyMapType::Normal;
-			break;
-		case EBakeMapType::FaceNormalImage:
-			MeshPropertyMapSettings.PropertyTypeIndex = (int32)EMeshPropertyMapType::FacetNormal;
-			break;
-		case EBakeMapType::PositionImage:
-			MeshPropertyMapSettings.PropertyTypeIndex = (int32)EMeshPropertyMapType::Position;
-			break;
-		case EBakeMapType::MaterialID:
-			MeshPropertyMapSettings.PropertyTypeIndex = (int32)EMeshPropertyMapType::MaterialID;
-			break;
-		default:
-			check(false);		// should not be possible!
-	}
 
 	if (!(CachedMeshPropertyMapSettings == MeshPropertyMapSettings))
 	{
 		CachedMeshPropertyMapSettings = MeshPropertyMapSettings;
-		ResultState = EOpState::Evaluate;
+		ResultState = EBakeOpState::Evaluate;
 	}
 	return ResultState;
 }
@@ -1188,9 +1080,9 @@ private:
 
 
 
-UBakeMeshAttributeMapsTool::EOpState UBakeMeshAttributeMapsTool::UpdateResult_Texture2DImage()
+EBakeOpState UBakeMeshAttributeMapsTool::UpdateResult_Texture2DImage()
 {
-	EOpState ResultState = EOpState::Complete;
+	EBakeOpState ResultState = EBakeOpState::Complete;
 
 	int32 ImageSize = (int32)Settings->Resolution;
 	FImageDimensions Dimensions(ImageSize, ImageSize);
@@ -1203,13 +1095,13 @@ UBakeMeshAttributeMapsTool::EOpState UBakeMeshAttributeMapsTool::UpdateResult_Te
 	if (UVOverlay == nullptr)
 	{
 		GetToolManager()->DisplayMessage(LOCTEXT("InvalidUVWarning", "The Source Mesh does not have the selected UV layer"), EToolMessageLevel::UserWarning);
-		return EOpState::Invalid;
+		return EBakeOpState::Invalid;
 	}
 	
 	if (Texture2DProps->SourceTexture == nullptr)
 	{
 		GetToolManager()->DisplayMessage(LOCTEXT("InvalidTextureWarning", "The Source Texture is not valid"), EToolMessageLevel::UserWarning);
-		return EOpState::Invalid;
+		return EBakeOpState::Invalid;
 	}
 
 
@@ -1220,22 +1112,22 @@ UBakeMeshAttributeMapsTool::EOpState UBakeMeshAttributeMapsTool::UpdateResult_Te
 		if (!TextureAccess.CopyTo(*CachedTextureImage))
 		{
 			GetToolManager()->DisplayMessage(LOCTEXT("CannotReadTextureWarning", "Cannot read from the source texture"), EToolMessageLevel::UserWarning);
-			return EOpState::Invalid;
+			return EBakeOpState::Invalid;
 		}
 	}
 
 	if (!(CachedTexture2DImageSettings == NewSettings))
 	{
 		CachedTexture2DImageSettings = NewSettings;
-		ResultState = EOpState::Evaluate;
+		ResultState = EBakeOpState::Evaluate;
 	}
 	return ResultState;
 }
 
 
-UBakeMeshAttributeMapsTool::EOpState UBakeMeshAttributeMapsTool::UpdateResult_MultiTexture()
+EBakeOpState UBakeMeshAttributeMapsTool::UpdateResult_MultiTexture()
 {
-	EOpState ResultState = EOpState::Complete;
+	EBakeOpState ResultState = EBakeOpState::Complete;
 
 	int32 ImageSize = (int32)Settings->Resolution;
 	FImageDimensions Dimensions(ImageSize, ImageSize);
@@ -1248,7 +1140,7 @@ UBakeMeshAttributeMapsTool::EOpState UBakeMeshAttributeMapsTool::UpdateResult_Mu
 	if (UVOverlay == nullptr)
 	{
 		GetToolManager()->DisplayMessage(LOCTEXT("InvalidUVWarning", "The Source Mesh does not have the selected UV layer"), EToolMessageLevel::UserWarning);
-		return EOpState::Invalid;
+		return EBakeOpState::Invalid;
 	}
 
 	for (TPair<int32, UTexture2D*>& InputTexture : MultiTextureProps->MaterialIDSourceTextureMap)
@@ -1256,7 +1148,7 @@ UBakeMeshAttributeMapsTool::EOpState UBakeMeshAttributeMapsTool::UpdateResult_Mu
 		if (InputTexture.Value == nullptr)
 		{
 			GetToolManager()->DisplayMessage(LOCTEXT("InvalidTextureWarning", "The Source Texture is not valid"), EToolMessageLevel::UserWarning);
-			return EOpState::Invalid;
+			return EBakeOpState::Invalid;
 		}
 	}
 
@@ -1269,7 +1161,7 @@ UBakeMeshAttributeMapsTool::EOpState UBakeMeshAttributeMapsTool::UpdateResult_Mu
 		if (!ensure(Texture != nullptr))
 		{
 			GetToolManager()->DisplayMessage(LOCTEXT("InvalidTextureWarning", "The Source Texture is not valid"), EToolMessageLevel::UserWarning);
-			return EOpState::Invalid;
+			return EBakeOpState::Invalid;
 		}
 
 		int32 MaterialID = InputTexture.Key;
@@ -1280,19 +1172,19 @@ UBakeMeshAttributeMapsTool::EOpState UBakeMeshAttributeMapsTool::UpdateResult_Mu
 		if (!TextureAccess.CopyTo(*CachedMultiTextures[MaterialID]))
 		{
 			GetToolManager()->DisplayMessage(LOCTEXT("CannotReadTextureWarning", "Cannot read from the source texture"), EToolMessageLevel::UserWarning);
-			return EOpState::Invalid;
+			return EBakeOpState::Invalid;
 		}
 	}
 	if (CachedMultiTextures.Num() == 0)
 	{
 		GetToolManager()->DisplayMessage(LOCTEXT("InvalidTextureWarning", "The Source Texture is not valid"), EToolMessageLevel::UserWarning);
-		return EOpState::Invalid;
+		return EBakeOpState::Invalid;
 	}
 
 	if (!(CachedTexture2DImageSettings == NewSettings))
 	{
 		CachedTexture2DImageSettings = NewSettings;
-		ResultState = EOpState::Evaluate;
+		ResultState = EBakeOpState::Evaluate;
 	}
 	return ResultState;
 }
@@ -1304,63 +1196,65 @@ UBakeMeshAttributeMapsTool::EOpState UBakeMeshAttributeMapsTool::UpdateResult_Mu
 void UBakeMeshAttributeMapsTool::UpdateVisualization()
 {
 	DynamicMeshComponent->SetOverrideRenderMaterial(PreviewMaterial);
-	switch (Settings->MapType)
+
+	// Map CachedMaps to Settings->Result
+	int NumResults = Settings->Result.Num();
+	for (int ResultIdx = 0; ResultIdx < NumResults; ResultIdx++)
 	{
-	default:
-		Settings->Result[0] = nullptr;
-		PreviewMaterial->SetTextureParameterValue(TEXT("NormalMap"), EmptyNormalMap);
-		PreviewMaterial->SetTextureParameterValue(TEXT("OcclusionMap"), EmptyColorMapWhite);
-		PreviewMaterial->SetTextureParameterValue(TEXT("ColorMap"), EmptyColorMapWhite);
-		break;
-	case EBakeMapType::TangentSpaceNormalMap:
-		Settings->Result[0] = CachedNormalMap;
-		PreviewMaterial->SetTextureParameterValue(TEXT("NormalMap"), CachedNormalMap);
-		PreviewMaterial->SetTextureParameterValue(TEXT("OcclusionMap"), EmptyColorMapWhite);
-		PreviewMaterial->SetTextureParameterValue(TEXT("ColorMap"), EmptyColorMapWhite);
-		break;
-	case EBakeMapType::Occlusion:
-		Settings->Result[0] = CachedOcclusionMap;
-		Settings->Result[1] = CachedBentNormalMap;
-		switch (OcclusionMapProps->Preview)
+		Settings->Result[ResultIdx] = CachedMaps[CachedMapIndices[ResultTypes[ResultIdx]]];
+	}
+
+	// Set the preview material according to the preview index.
+	if (Settings->MapPreview >= 0 && Settings->MapPreview < Settings->Result.Num())
+	{
+		const EBakeMapType& PreviewMapType = ResultTypes[Settings->MapPreview];
+		if (PreviewMapType != EBakeMapType::None)
 		{
-		case EOcclusionMapPreview::AmbientOcclusion:
-			PreviewMaterial->SetTextureParameterValue(TEXT("NormalMap"), EmptyNormalMap);
-			PreviewMaterial->SetTextureParameterValue(TEXT("OcclusionMap"), CachedOcclusionMap);
-			PreviewMaterial->SetTextureParameterValue(TEXT("ColorMap"), EmptyColorMapWhite);
-			break;
-		case EOcclusionMapPreview::BentNormal:
-			BentNormalPreviewMaterial->SetTextureParameterValue(TEXT("NormalMap"), EmptyNormalMap);
-			BentNormalPreviewMaterial->SetTextureParameterValue(TEXT("OcclusionMap"), CachedOcclusionMap);
-			BentNormalPreviewMaterial->SetTextureParameterValue(TEXT("ColorMap"), EmptyColorMapWhite);
-			BentNormalPreviewMaterial->SetTextureParameterValue(TEXT("BentNormalMap"), CachedBentNormalMap);
-			DynamicMeshComponent->SetOverrideRenderMaterial(BentNormalPreviewMaterial);
-			break;
+			UTexture2D* PreviewMap = CachedMaps[CachedMapIndices[PreviewMapType]];
+			switch (PreviewMapType)
+			{
+			default:
+				PreviewMaterial->SetTextureParameterValue(TEXT("NormalMap"), EmptyNormalMap);
+				PreviewMaterial->SetTextureParameterValue(TEXT("OcclusionMap"), EmptyColorMapWhite);
+				PreviewMaterial->SetTextureParameterValue(TEXT("ColorMap"), EmptyColorMapWhite);
+				break;
+			case EBakeMapType::TangentSpaceNormalMap:
+				PreviewMaterial->SetTextureParameterValue(TEXT("NormalMap"), PreviewMap);
+				PreviewMaterial->SetTextureParameterValue(TEXT("OcclusionMap"), EmptyColorMapWhite);
+				PreviewMaterial->SetTextureParameterValue(TEXT("ColorMap"), EmptyColorMapWhite);
+				break;
+			case EBakeMapType::AmbientOcclusion:
+				PreviewMaterial->SetTextureParameterValue(TEXT("NormalMap"), EmptyNormalMap);
+				PreviewMaterial->SetTextureParameterValue(TEXT("OcclusionMap"), PreviewMap);
+				PreviewMaterial->SetTextureParameterValue(TEXT("ColorMap"), EmptyColorMapWhite);
+				break;
+			case EBakeMapType::BentNormal:
+				BentNormalPreviewMaterial->SetTextureParameterValue(TEXT("NormalMap"), EmptyNormalMap);
+				if (CachedMapIndices.Contains(EBakeMapType::AmbientOcclusion))
+				{
+					BentNormalPreviewMaterial->SetTextureParameterValue(TEXT("OcclusionMap"), CachedMaps[CachedMapIndices[EBakeMapType::AmbientOcclusion]]);
+				}
+				else
+				{
+					BentNormalPreviewMaterial->SetTextureParameterValue(TEXT("OcclusionMap"), EmptyColorMapWhite);
+				}
+				BentNormalPreviewMaterial->SetTextureParameterValue(TEXT("ColorMap"), EmptyColorMapWhite);
+				BentNormalPreviewMaterial->SetTextureParameterValue(TEXT("BentNormalMap"), PreviewMap);
+				DynamicMeshComponent->SetOverrideRenderMaterial(BentNormalPreviewMaterial);
+				break;
+			case EBakeMapType::Curvature:
+			case EBakeMapType::NormalImage:
+			case EBakeMapType::FaceNormalImage:
+			case EBakeMapType::PositionImage:
+			case EBakeMapType::MaterialID:
+			case EBakeMapType::Texture2DImage:
+			case EBakeMapType::MultiTexture:
+				PreviewMaterial->SetTextureParameterValue(TEXT("NormalMap"), EmptyNormalMap);
+				PreviewMaterial->SetTextureParameterValue(TEXT("OcclusionMap"), EmptyColorMapWhite);
+				PreviewMaterial->SetTextureParameterValue(TEXT("ColorMap"), PreviewMap);
+				break;
+			}
 		}
-		break;
-	case EBakeMapType::Curvature:
-		Settings->Result[0] = CachedCurvatureMap;
-		PreviewMaterial->SetTextureParameterValue(TEXT("NormalMap"), EmptyNormalMap);
-		PreviewMaterial->SetTextureParameterValue(TEXT("OcclusionMap"), EmptyColorMapWhite);
-		PreviewMaterial->SetTextureParameterValue(TEXT("ColorMap"), CachedCurvatureMap);
-		break;
-
-	case EBakeMapType::NormalImage:
-	case EBakeMapType::FaceNormalImage:
-	case EBakeMapType::PositionImage:
-	case EBakeMapType::MaterialID:
-		Settings->Result[0] = CachedMeshPropertyMap;
-		PreviewMaterial->SetTextureParameterValue(TEXT("NormalMap"), EmptyNormalMap);
-		PreviewMaterial->SetTextureParameterValue(TEXT("OcclusionMap"), EmptyColorMapWhite);
-		PreviewMaterial->SetTextureParameterValue(TEXT("ColorMap"), CachedMeshPropertyMap);
-		break;
-
-	case EBakeMapType::Texture2DImage:
-	case EBakeMapType::MultiTexture:
-		Settings->Result[0] = CachedTexture2DImageMap;
-		PreviewMaterial->SetTextureParameterValue(TEXT("NormalMap"), EmptyNormalMap);
-		PreviewMaterial->SetTextureParameterValue(TEXT("OcclusionMap"), EmptyColorMapWhite);
-		PreviewMaterial->SetTextureParameterValue(TEXT("ColorMap"), CachedTexture2DImageMap);
-		break;
 	}
 }
 
@@ -1374,102 +1268,176 @@ void UBakeMeshAttributeMapsTool::UpdateOnModeChange()
 	SetToolPropertySourceEnabled(Texture2DProps, false);
 	SetToolPropertySourceEnabled(MultiTextureProps, false);
 
-	Settings->Result.Empty();
-	Settings->Result.Add(nullptr);
-	switch (Settings->MapType)
+	for (const EBakeMapType MapType : ALL_BAKE_MAP_TYPES)
 	{
-	case EBakeMapType::TangentSpaceNormalMap:
-		SetToolPropertySourceEnabled(NormalMapProps, true);
-		break;
-	case EBakeMapType::Occlusion:
-		SetToolPropertySourceEnabled(OcclusionMapProps, true);
-		Settings->Result.Add(nullptr); // Extra slot for BentNormalMap
-		break;
-	case EBakeMapType::Curvature:
-		SetToolPropertySourceEnabled(CurvatureMapProps, true);
-		break;
-	case EBakeMapType::Texture2DImage:
-		SetToolPropertySourceEnabled(Texture2DProps, true);
-		break;
-	case EBakeMapType::MultiTexture:
-		SetToolPropertySourceEnabled(MultiTextureProps, true);
-		break;
+		switch ((EBakeMapType)Settings->MapTypes & MapType)
+		{
+		case EBakeMapType::TangentSpaceNormalMap:
+			SetToolPropertySourceEnabled(NormalMapProps, true);
+			break;
+		case EBakeMapType::AmbientOcclusion:
+		case EBakeMapType::BentNormal:
+		case EBakeMapType::Occlusion:
+			SetToolPropertySourceEnabled(OcclusionMapProps, true);
+			break;
+		case EBakeMapType::Curvature:
+			SetToolPropertySourceEnabled(CurvatureMapProps, true);
+			break;
+		case EBakeMapType::NormalImage:
+		case EBakeMapType::FaceNormalImage:
+		case EBakeMapType::PositionImage:
+		case EBakeMapType::MaterialID:
+			break;
+		case EBakeMapType::Texture2DImage:
+			SetToolPropertySourceEnabled(Texture2DProps, true);
+			break;
+		case EBakeMapType::MultiTexture:
+			SetToolPropertySourceEnabled(MultiTextureProps, true);
+			break;
+		default:
+			break;
+		}
 	}
 
+	ResultTypes = GetMapTypesArray(Settings->MapTypes);
+	Settings->Result.Empty();
+	Settings->Result.SetNum(ResultTypes.Num());
+
+	// Generate a map between EBakeMapType and CachedMaps
+	CachedMapIndices.Empty();
+	int32 CachedMapIdx = 0;
+
+	// Use the processed bitfield which may contain additional targets
+	// (ex. AO if BentNormal was requested).
+	const EBakeMapType BakeMapTypes = GetMapTypes(Settings->MapTypes);
+	for (EBakeMapType MapType : ALL_BAKE_MAP_TYPES)
+	{
+		if (MapType == EBakeMapType::Occlusion)
+		{
+			if ((bool)(BakeMapTypes & EBakeMapType::AmbientOcclusion))
+			{
+				CachedMapIndices.Add(EBakeMapType::AmbientOcclusion, CachedMapIdx++);
+			}
+			if ((bool)(BakeMapTypes & EBakeMapType::BentNormal))
+			{
+				CachedMapIndices.Add(EBakeMapType::BentNormal, CachedMapIdx++);
+			}
+		}
+		else if( (bool)(BakeMapTypes & MapType) )
+		{
+			CachedMapIndices.Add(MapType, CachedMapIdx++);
+		}
+	}
+	CachedMaps.Empty();
+	CachedMaps.SetNum(CachedMapIndices.Num());
 }
 
 
 void UBakeMeshAttributeMapsTool::OnMapsUpdated(const TUniquePtr<FMeshMapBaker>& NewResult)
 {
-	switch (Settings->MapType)
-	{
-	case EBakeMapType::TangentSpaceNormalMap:
-	{
-		FTexture2DBuilder TextureBuilder;
-		TextureBuilder.Initialize(FTexture2DBuilder::ETextureType::NormalMap, CachedNormalMapSettings.Dimensions);
-		TextureBuilder.Copy(*NewResult->GetBakeResults(0)[0]);
-		TextureBuilder.Commit(false);
-		CachedNormalMap = TextureBuilder.GetTexture2D();
-		break;
-	}
-	case EBakeMapType::Occlusion:
-	{
-		FTexture2DBuilder TextureBuilder;
-		TextureBuilder.Initialize(FTexture2DBuilder::ETextureType::AmbientOcclusion, CachedOcclusionMapSettings.Dimensions);
-		TextureBuilder.Copy(*NewResult->GetBakeResults(0)[0]);
-		TextureBuilder.Commit(false);
-		CachedOcclusionMap = TextureBuilder.GetTexture2D();
-
-		FTexture2DBuilder TextureNormalBuilder;
-		TextureNormalBuilder.Initialize(FTexture2DBuilder::ETextureType::NormalMap, CachedOcclusionMapSettings.Dimensions);
-		TextureNormalBuilder.Copy(*NewResult->GetBakeResults(0)[1]);
-		TextureNormalBuilder.Commit(false);
-		CachedBentNormalMap = TextureNormalBuilder.GetTexture2D();
-		break;
-	}
-	case EBakeMapType::Curvature:
-	{
-		FTexture2DBuilder TextureBuilder;
-		TextureBuilder.Initialize(FTexture2DBuilder::ETextureType::Color, CachedCurvatureMapSettings.Dimensions);
-		TextureBuilder.Copy(*NewResult->GetBakeResults(0)[0]);
-		TextureBuilder.Commit(false);
-		CachedCurvatureMap = TextureBuilder.GetTexture2D();
-		break;
-	}
-	case EBakeMapType::PositionImage:
-	case EBakeMapType::NormalImage:
-	case EBakeMapType::FaceNormalImage:
-	case EBakeMapType::MaterialID:
+	// This method assumes that the bake evaluators were instantiated in the order
+	// defined by ALL_BAKE_MAP_TYPES.
+	const EBakeMapType& BakeMapTypes = CachedBakeCacheSettings.BakeMapTypes;
+	int32 BakerIdx = 0;
+	for (EBakeMapType MapType : ALL_BAKE_MAP_TYPES)
 	{
-		FTexture2DBuilder TextureBuilder;
-		TextureBuilder.Initialize(FTexture2DBuilder::ETextureType::Color, CachedMeshPropertyMapSettings.Dimensions);
-		TextureBuilder.Copy(*NewResult->GetBakeResults(0)[0]);
-		TextureBuilder.Commit(false);
-		CachedMeshPropertyMap = TextureBuilder.GetTexture2D();
-		break;
+		switch (BakeMapTypes & MapType)
+		{
+		case EBakeMapType::TangentSpaceNormalMap:
+		{
+			FTexture2DBuilder TextureBuilder;
+			TextureBuilder.Initialize(FTexture2DBuilder::ETextureType::NormalMap, CachedNormalMapSettings.Dimensions);
+			TextureBuilder.Copy(*NewResult->GetBakeResults(BakerIdx++)[0]);
+			TextureBuilder.Commit(false);
+			CachedMaps[CachedMapIndices[EBakeMapType::TangentSpaceNormalMap]] = TextureBuilder.GetTexture2D();
+			break;
+		}
+		case EBakeMapType::AmbientOcclusion:
+		case EBakeMapType::BentNormal:
+		case EBakeMapType::Occlusion:
+		{
+			int32 OcclusionIdx = 0;
+			if ((bool)(BakeMapTypes & EBakeMapType::AmbientOcclusion))
+			{
+				FTexture2DBuilder TextureBuilder;
+				TextureBuilder.Initialize(FTexture2DBuilder::ETextureType::AmbientOcclusion, CachedOcclusionMapSettings.Dimensions);
+				TextureBuilder.Copy(*NewResult->GetBakeResults(BakerIdx)[OcclusionIdx++]);
+				TextureBuilder.Commit(false);
+				CachedMaps[CachedMapIndices[EBakeMapType::AmbientOcclusion]] = TextureBuilder.GetTexture2D();
+			}
+			if ((bool)(BakeMapTypes & EBakeMapType::BentNormal))
+			{
+				FTexture2DBuilder TextureBuilder;
+				TextureBuilder.Initialize(FTexture2DBuilder::ETextureType::NormalMap, CachedOcclusionMapSettings.Dimensions);
+				TextureBuilder.Copy(*NewResult->GetBakeResults(BakerIdx)[OcclusionIdx++]);
+				TextureBuilder.Commit(false);
+				CachedMaps[CachedMapIndices[EBakeMapType::BentNormal]] = TextureBuilder.GetTexture2D();
+			}
+			++BakerIdx;
+			break;
+		}
+		case EBakeMapType::Curvature:
+		{
+			FTexture2DBuilder TextureBuilder;
+			TextureBuilder.Initialize(FTexture2DBuilder::ETextureType::Color, CachedCurvatureMapSettings.Dimensions);
+			TextureBuilder.Copy(*NewResult->GetBakeResults(BakerIdx++)[0]);
+			TextureBuilder.Commit(false);
+			CachedMaps[CachedMapIndices[BakeMapTypes & MapType]] = TextureBuilder.GetTexture2D();
+			break;
+		}
+		case EBakeMapType::NormalImage:
+		case EBakeMapType::FaceNormalImage:
+		case EBakeMapType::PositionImage:
+		case EBakeMapType::MaterialID:
+		{
+			FTexture2DBuilder TextureBuilder;
+			TextureBuilder.Initialize(FTexture2DBuilder::ETextureType::Color, CachedMeshPropertyMapSettings.Dimensions);
+			TextureBuilder.Copy(*NewResult->GetBakeResults(BakerIdx++)[0]);
+			TextureBuilder.Commit(false);
+			CachedMaps[CachedMapIndices[BakeMapTypes & MapType]] = TextureBuilder.GetTexture2D();
+			break;
+		}
+		case EBakeMapType::Texture2DImage:
+		case EBakeMapType::MultiTexture:
+		{
+			FTexture2DBuilder TextureBuilder;
+			TextureBuilder.Initialize(FTexture2DBuilder::ETextureType::Color, CachedTexture2DImageSettings.Dimensions);
+			TextureBuilder.Copy(*NewResult->GetBakeResults(BakerIdx++)[0]);
+			TextureBuilder.Commit(false);
+			CachedMaps[CachedMapIndices[BakeMapTypes & MapType]] = TextureBuilder.GetTexture2D();
+			break;
+		}
+		}
 	}
-	case EBakeMapType::Texture2DImage:
+
+	UpdateVisualization();
+	GetToolManager()->PostInvalidation();
+}
+
+
+EBakeMapType UBakeMeshAttributeMapsTool::GetMapTypes(const int32& MapTypes) const
+{
+	EBakeMapType OutMapTypes = (EBakeMapType)MapTypes & EBakeMapType::All;
+	// Force AO bake for BentNormal preview
+	if ((bool)(OutMapTypes & EBakeMapType::BentNormal))
 	{
-		FTexture2DBuilder TextureBuilder;
-		TextureBuilder.Initialize(FTexture2DBuilder::ETextureType::Color, CachedTexture2DImageSettings.Dimensions);
-		TextureBuilder.Copy(*NewResult->GetBakeResults(0)[0], true);
-		TextureBuilder.Commit(false);
-		CachedTexture2DImageMap = TextureBuilder.GetTexture2D();
-		break;
+		OutMapTypes |= EBakeMapType::AmbientOcclusion;
 	}
-	case EBakeMapType::MultiTexture:
+	return OutMapTypes;
+}
+
+TArray<EBakeMapType> UBakeMeshAttributeMapsTool::GetMapTypesArray(const int32& MapTypes) const
+{
+	TArray<EBakeMapType> OutMapTypes;
+	int32 Bitfield = MapTypes & (int32)EBakeMapType::All;
+	for (int32 BitIdx = 0; Bitfield; Bitfield >>= 1, ++BitIdx)
 	{
-		FTexture2DBuilder TextureBuilder;
-		TextureBuilder.Initialize(FTexture2DBuilder::ETextureType::Color, CachedTexture2DImageSettings.Dimensions);
-		TextureBuilder.Copy(*NewResult->GetBakeResults(0)[0], true);
-		TextureBuilder.Commit(false);
-		CachedTexture2DImageMap = TextureBuilder.GetTexture2D();
-		break;
-	}
+		if (Bitfield & 1)
+		{
+			OutMapTypes.Add((EBakeMapType)(1 << BitIdx));
+		}
 	}
-
-	UpdateVisualization();
-	GetToolManager()->PostInvalidation();
+	return OutMapTypes;
 }
 
 
diff --git a/Engine/Plugins/Experimental/MeshModelingToolset/Source/MeshModelingTools/Public/BakeMeshAttributeMapsTool.h b/Engine/Plugins/Experimental/MeshModelingToolset/Source/MeshModelingTools/Public/BakeMeshAttributeMapsTool.h
--- a/Engine/Plugins/Experimental/MeshModelingToolset/Source/MeshModelingTools/Public/BakeMeshAttributeMapsTool.h
+++ b/Engine/Plugins/Experimental/MeshModelingToolset/Source/MeshModelingTools/Public/BakeMeshAttributeMapsTool.h
@@ -51,19 +51,24 @@ protected:
 
 
 
-UENUM()
+UENUM(meta=(Bitflags, UseEnumValuesAsMaskValuesInEditor="true"))
 enum class EBakeMapType
 {
-	TangentSpaceNormalMap,
-	Occlusion,
-	Curvature,
-	Texture2DImage,
-	NormalImage,
-	FaceNormalImage,
-	PositionImage,
-	MaterialID,
-	MultiTexture
+	None                   = 0 UMETA(Hidden),
+	TangentSpaceNormalMap  = 1 << 0,
+	AmbientOcclusion       = 1 << 1,
+	BentNormal             = 1 << 2,
+	Curvature              = 1 << 3,
+	Texture2DImage         = 1 << 4,
+	NormalImage            = 1 << 5,
+	FaceNormalImage        = 1 << 6,
+	PositionImage          = 1 << 7,
+	MaterialID             = 1 << 8,
+	MultiTexture           = 1 << 9,
+	Occlusion              = (AmbientOcclusion | BentNormal) UMETA(Hidden),
+	All                    = 0x3FF UMETA(Hidden)
 };
+ENUM_CLASS_FLAGS(EBakeMapType);
 
 
 UENUM()
@@ -99,9 +104,13 @@ class MESHMODELINGTOOLS_API UBakeMeshAttributeMapsToolProperties : public UInter
 	GENERATED_BODY()
 
 public:
-	/** The type of map to generate */
-	UPROPERTY(EditAnywhere, Category = MapSettings)
-	EBakeMapType MapType = EBakeMapType::TangentSpaceNormalMap;
+	/** The map types to generate */
+	UPROPERTY(EditAnywhere, Category = MapSettings, meta=(Bitmask, BitmaskEnum=EBakeMapType))
+	int32 MapTypes = (int32) EBakeMapType::TangentSpaceNormalMap;
+
+	/** The map type index to preview */
+	UPROPERTY(EditAnywhere, Category = MapSettings, meta=(ArrayClamp="Result"))
+	int MapPreview = 0;
 
 	/** The pixel resolution of the generated map */
 	UPROPERTY(EditAnywhere, Category = MapSettings, meta = (TransientToolProperty))
@@ -161,25 +170,12 @@ enum class EOcclusionMapDistribution
 	Cosine UMETA(DisplayName = "Cosine")
 };
 
-UENUM()
-enum class EOcclusionMapPreview
-{
-	/** Ambient Occlusion */
-	AmbientOcclusion UMETA(DisplayName = "Ambient Occlusion"),
-	/** Bent Normal */
-	BentNormal UMETA(DisplayName = "Bent Normal")
-};
-
 
 UCLASS()
 class MESHMODELINGTOOLS_API UBakedOcclusionMapToolProperties : public UInteractiveToolPropertySet
 {
 	GENERATED_BODY()
 public:
-	/** Occlusion map output to preview */
-	UPROPERTY(EditAnywhere, Category = OcclusionMap)
-	EOcclusionMapPreview Preview = EOcclusionMapPreview::AmbientOcclusion;
-
 	/** Number of occlusion rays */
 	UPROPERTY(EditAnywhere, Category = OcclusionMap, meta = (UIMin = "1", UIMax = "1024", ClampMin = "0", ClampMax = "50000"))
 	int32 OcclusionRays = 16;
@@ -340,6 +336,13 @@ public:
 };
 
 
+enum class EBakeOpState
+{
+	Complete = 0,		// Inputs valid & Result is valid - no-op.
+	Evaluate = 1 << 0,	// Inputs valid & Result is invalid - re-evaluate.
+	Invalid	 = 1 << 1	// Inputs invalid - pause eval.
+};
+ENUM_CLASS_FLAGS(EBakeOpState);
 
 /**
  * Detail Map Baking Tool
@@ -391,7 +394,6 @@ protected:
 	UBakedOcclusionMapVisualizationProperties* VisualizationProps;
 
 
-
 protected:
 	friend class FBakeMapBaseOp;
 	friend class FBakeNormalMapOp;
@@ -400,7 +402,6 @@ protected:
 	friend class FBakeTexture2DImageMapOp;
 	friend class FBakeMeshPropertyMapOp;
 	friend class FBakeMultiTextureOp;
-
 	friend class FMeshMapBakerOp;
 
 	USimpleDynamicMeshComponent* DynamicMeshComponent;
@@ -437,8 +438,13 @@ protected:
 	TUniquePtr<TGenericDataBackgroundCompute<UE::Geometry::FMeshMapBaker>> Compute = nullptr;
 	void OnMapsUpdated(const TUniquePtr<UE::Geometry::FMeshMapBaker>& NewResult);
 
+	/** @return A single bitfield of map types from an array of map types. */
+	EBakeMapType GetMapTypes(const int32& MapTypes) const;
+	TArray<EBakeMapType> GetMapTypesArray(const int32& MapTypes) const;
+
 	struct FBakeCacheSettings
 	{
+		EBakeMapType BakeMapTypes = EBakeMapType::None;
 		FImageDimensions Dimensions;
 		int32 UVLayer;
 		int32 DetailTimestamp;
@@ -447,18 +453,21 @@ protected:
 
 		bool operator==(const FBakeCacheSettings& Other) const
 		{
-			return Dimensions == Other.Dimensions && UVLayer == Other.UVLayer && DetailTimestamp == Other.DetailTimestamp && Thickness == Other.Thickness && Multisampling == Other.Multisampling;
+			return BakeMapTypes == Other.BakeMapTypes && Dimensions == Other.Dimensions &&
+				UVLayer == Other.UVLayer && DetailTimestamp == Other.DetailTimestamp &&
+				Thickness == Other.Thickness && Multisampling == Other.Multisampling;
 		}
 	};
 	FBakeCacheSettings CachedBakeCacheSettings;
+	TArray<EBakeMapType> ResultTypes;
+
+	EBakeOpState OpState = EBakeOpState::Evaluate;
+
+	UPROPERTY()
+	TArray<UTexture2D*> CachedMaps;
+	using CachedMapIndex = TMap<EBakeMapType, int32>;
+	CachedMapIndex CachedMapIndices;
 
-	enum class EOpState
-	{
-		Complete,	// Inputs valid & Result is valid - no-op.
-		Evaluate,	// Inputs valid & Result is invalid - re-evaluate.
-		Invalid		// Inputs invalid - pause eval.
-	};
-	EOpState OpState = EOpState::Evaluate;
 
 	struct FNormalMapSettings
 	{
@@ -470,11 +479,7 @@ protected:
 		}
 	};
 	FNormalMapSettings CachedNormalMapSettings;
-	UPROPERTY()
-	UTexture2D* CachedNormalMap;
-
-	EOpState UpdateResult_Normal();
-
+	EBakeOpState UpdateResult_Normal();
 
 
 	struct FOcclusionMapSettings
@@ -501,13 +506,7 @@ protected:
 		}
 	};
 	FOcclusionMapSettings CachedOcclusionMapSettings;
-	UPROPERTY()
-	UTexture2D* CachedOcclusionMap;
-	UPROPERTY()
-	UTexture2D* CachedBentNormalMap;
-
-	EOpState UpdateResult_Occlusion();
-
+	EBakeOpState UpdateResult_Occlusion();
 
 
 	struct FCurvatureMapSettings
@@ -528,31 +527,20 @@ protected:
 		}
 	};
 	FCurvatureMapSettings CachedCurvatureMapSettings;
-	UPROPERTY()
-	UTexture2D* CachedCurvatureMap;
-
-	EOpState UpdateResult_Curvature();
-
-
+	EBakeOpState UpdateResult_Curvature();
 
 
 	struct FMeshPropertyMapSettings
 	{
 		FImageDimensions Dimensions;
-		int32 PropertyTypeIndex;
 
 		bool operator==(const FMeshPropertyMapSettings& Other) const
 		{
-			return Dimensions == Other.Dimensions && PropertyTypeIndex == Other.PropertyTypeIndex;
+			return Dimensions == Other.Dimensions;
 		}
 	};
 	FMeshPropertyMapSettings CachedMeshPropertyMapSettings;
-	UPROPERTY()
-	UTexture2D* CachedMeshPropertyMap;
-
-	EOpState UpdateResult_MeshProperty();
-
-
+	EBakeOpState UpdateResult_MeshProperty();
 
 
 	struct FTexture2DImageSettings
@@ -567,19 +555,14 @@ protected:
 	};
 	TSharedPtr<UE::Geometry::TImageBuilder<FVector4f>, ESPMode::ThreadSafe> CachedTextureImage;
 	FTexture2DImageSettings CachedTexture2DImageSettings;
-	UPROPERTY()
-	UTexture2D* CachedTexture2DImageMap;
-
-	EOpState UpdateResult_Texture2DImage();
+	EBakeOpState UpdateResult_Texture2DImage();
 
 
 	TMap<int32, TSharedPtr<UE::Geometry::TImageBuilder<FVector4f>, ESPMode::ThreadSafe>> CachedMultiTextures;
-	EOpState UpdateResult_MultiTexture();
-
+	EBakeOpState UpdateResult_MultiTexture();
 
 
 	// empty maps are shown when nothing is computed
-
 	UPROPERTY()
 	UTexture2D* EmptyNormalMap;
 
