ApproximateActors - Ensure debug view modes are available when generating materials.
Use LogApproximateActors instead of LogTemp

diff --git a/Engine/Plugins/Experimental/MeshModelingToolset/Source/GeometryProcessingAdapters/Private/ApproximateActorsImpl.cpp b/Engine/Plugins/Experimental/MeshModelingToolset/Source/GeometryProcessingAdapters/Private/ApproximateActorsImpl.cpp
--- a/Engine/Plugins/Experimental/MeshModelingToolset/Source/GeometryProcessingAdapters/Private/ApproximateActorsImpl.cpp
+++ b/Engine/Plugins/Experimental/MeshModelingToolset/Source/GeometryProcessingAdapters/Private/ApproximateActorsImpl.cpp
@@ -41,6 +41,8 @@
 using namespace UE::Geometry;
 using namespace UE::AssetUtils;
 
+DEFINE_LOG_CATEGORY_STATIC(LogApproximateActors, Log, All);
+
 #define LOCTEXT_NAMESPACE "ApproximateActorsImpl"
 
 static TAutoConsoleVariable<int32> CVarApproximateActorsRDOCCapture(
@@ -124,7 +126,7 @@ static void BakeTexturesFromPhotoCapture(
 	int32 Supersample = FMath::Max(1, Options.AntiAliasMultiSampling);
 	if ( (Options.TextureImageSize * Supersample) > 16384)
 	{
-		UE_LOG(LogTemp, Warning, TEXT("[ApproximateActors] Ignoring requested supersampling rate %d because it would require image buffers with resolution %d, please try lower value."), Supersample, Options.TextureImageSize * Supersample);
+		UE_LOG(LogApproximateActors, Warning, TEXT("Ignoring requested supersampling rate %d because it would require image buffers with resolution %d, please try lower value."), Supersample, Options.TextureImageSize * Supersample);
 		Supersample = 1;
 	}
 
@@ -398,7 +400,7 @@ static TSharedPtr<FApproximationMeshData> GenerateApproximationMesh(
 	// avoid insane memory usage
 	if (VoxelDimTarget > Options.ClampVoxelDimension)
 	{
-		UE_LOG(LogTemp, Warning, TEXT("FApproximateActorsImpl - very large voxel size %d clamped to %d"), VoxelDimTarget, Options.ClampVoxelDimension);
+		UE_LOG(LogApproximateActors, Warning, TEXT("FApproximateActorsImpl - very large voxel size %d clamped to %d"), VoxelDimTarget, Options.ClampVoxelDimension);
 		VoxelDimTarget = Options.ClampVoxelDimension;
 	}
 
@@ -420,7 +422,7 @@ static TSharedPtr<FApproximationMeshData> GenerateApproximationMesh(
 
 	if (Options.bVerbose)
 	{
-		UE_LOG(LogTemp, Warning, TEXT("[ApproximateActors] Solidify mesh has %d triangles"), CurResultMesh->TriangleCount());
+		UE_LOG(LogApproximateActors, Warning, TEXT("Solidify mesh has %d triangles"), CurResultMesh->TriangleCount());
 	}
 
 	Progress.EnterProgressFrame(1.f, LOCTEXT("ClosingMesh", "Topological Operations..."));
@@ -444,7 +446,7 @@ static TSharedPtr<FApproximationMeshData> GenerateApproximationMesh(
 
 		if (Options.bVerbose)
 		{
-			UE_LOG(LogTemp, Warning, TEXT("[ApproximateActors] Morphology mesh has %d triangles"), CurResultMesh->TriangleCount());
+			UE_LOG(LogApproximateActors, Warning, TEXT("Morphology mesh has %d triangles"), CurResultMesh->TriangleCount());
 		}
 	}
 
@@ -496,7 +498,7 @@ static TSharedPtr<FApproximationMeshData> GenerateApproximationMesh(
 
 		if (Options.bVerbose)
 		{
-			UE_LOG(LogTemp, Warning, TEXT("[ApproximateActors] Occlusion-Filtered mesh has %d triangles"), CurResultMesh->TriangleCount());
+			UE_LOG(LogApproximateActors, Warning, TEXT("Occlusion-Filtered mesh has %d triangles"), CurResultMesh->TriangleCount());
 		}
 	}
 
@@ -556,7 +558,7 @@ static TSharedPtr<FApproximationMeshData> GenerateApproximationMesh(
 		int32 AfterCount = CurResultMesh->TriangleCount();
 		if (Options.bVerbose)
 		{
-			UE_LOG(LogTemp, Warning, TEXT("[ApproximateActors] Simplified mesh from %d to %d triangles"), BeforeCount, AfterCount);
+			UE_LOG(LogApproximateActors, Warning, TEXT("Simplified mesh from %d to %d triangles"), BeforeCount, AfterCount);
 		}
 	}
 
@@ -736,6 +738,17 @@ void FApproximateActorsImpl::GenerateApproximationForActorSet(const TArray<AActo
 
 	RenderCaptureInterface::FScopedCapture RenderCapture(CVarApproximateActorsRDOCCapture.GetValueOnAnyThread() == 1, TEXT("ApproximateActors"));
 
+	if (Options.BasePolicy == IGeometryProcessing_ApproximateActors::EApproximationPolicy::MeshAndGeneratedMaterial)
+	{
+		// The scene capture photoset part of this process relies on debug view modes being available.
+		// If it ain't the case, fail immediatelly
+		if (!AllowDebugViewmodes())
+		{
+			UE_LOG(LogApproximateActors, Error, TEXT("Debug view modes not are available - unable to generate material"));
+			ResultsOut.ResultCode = EResultCode::MaterialGenerationFailed;
+			return;
+		}
+	}
 	//
 	// Future Optimizations
 	// 	   - can do most of the mesh processing at the same time as capturing the photo set (if that matters)
@@ -764,7 +777,7 @@ void FApproximateActorsImpl::GenerateApproximationForActorSet(const TArray<AActo
 	{
 		FMeshSceneAdapter::FStatistics Stats;
 		Scene.GetGeometryStatistics(Stats);
-		UE_LOG(LogTemp, Warning, TEXT("[ApproximateActors] %d triangles in %d unique meshes, total %d triangles in %d instances"),
+		UE_LOG(LogApproximateActors, Warning, TEXT("%d triangles in %d unique meshes, total %d triangles in %d instances"),
 			Stats.UniqueMeshTriangleCount, Stats.UniqueMeshCount, Stats.InstanceMeshTriangleCount, Stats.InstanceMeshCount);
 	}
 
diff --git a/Engine/Source/Developer/GeometryProcessingInterfaces/Public/GeometryProcessingInterfaces/ApproximateActors.h b/Engine/Source/Developer/GeometryProcessingInterfaces/Public/GeometryProcessingInterfaces/ApproximateActors.h
--- a/Engine/Source/Developer/GeometryProcessingInterfaces/Public/GeometryProcessingInterfaces/ApproximateActors.h
+++ b/Engine/Source/Developer/GeometryProcessingInterfaces/Public/GeometryProcessingInterfaces/ApproximateActors.h
@@ -184,6 +184,7 @@ public:
 	{
 		Success,
 		MeshGenerationFailed,
+		MaterialGenerationFailed,
 		UnknownError
 	};
 
