nonunity fixes
#rnx

diff --git a/Engine/Plugins/Experimental/MeshModelingToolset/Source/ModelingComponentsEditorOnly/Private/EditorModelingObjectsCreationAPI.cpp b/Engine/Plugins/Experimental/MeshModelingToolset/Source/ModelingComponentsEditorOnly/Private/EditorModelingObjectsCreationAPI.cpp
--- a/Engine/Plugins/Experimental/MeshModelingToolset/Source/ModelingComponentsEditorOnly/Private/EditorModelingObjectsCreationAPI.cpp
+++ b/Engine/Plugins/Experimental/MeshModelingToolset/Source/ModelingComponentsEditorOnly/Private/EditorModelingObjectsCreationAPI.cpp
@@ -14,6 +14,7 @@
 #include "Engine/StaticMesh.h"
 #include "Components/StaticMeshComponent.h"
 #include "Engine/StaticMeshActor.h"
+#include "Materials/Material.h"
 
 #include "ToolTargets/VolumeDynamicMeshToolTarget.h"  // for CVarModelingMaxVolumeTriangleCount
 
diff --git a/Engine/Plugins/Experimental/MeshModelingToolset/Source/ModelingComponentsEditorOnly/Private/ToolTargets/StaticMeshComponentToolTarget.cpp b/Engine/Plugins/Experimental/MeshModelingToolset/Source/ModelingComponentsEditorOnly/Private/ToolTargets/StaticMeshComponentToolTarget.cpp
--- a/Engine/Plugins/Experimental/MeshModelingToolset/Source/ModelingComponentsEditorOnly/Private/ToolTargets/StaticMeshComponentToolTarget.cpp
+++ b/Engine/Plugins/Experimental/MeshModelingToolset/Source/ModelingComponentsEditorOnly/Private/ToolTargets/StaticMeshComponentToolTarget.cpp
@@ -6,6 +6,7 @@
 #include "Components/StaticMeshComponent.h"
 #include "DynamicMeshToMeshDescription.h"
 #include "Engine/StaticMesh.h"
+#include "Materials/Material.h"
 #include "MeshDescriptionToDynamicMesh.h"
 #include "RenderingThread.h"
 #include "AssetUtils/MeshDescriptionUtil.h"
