ModelingComponents: Move SubdividePoly to ModelingComponentsEditorOnly. Need this to be able to apply SubD in other contexts.
#rnx

diff --git a/Engine/Plugins/Experimental/MeshModelingToolset/Source/MeshModelingToolsEditorOnly/MeshModelingToolsEditorOnly.Build.cs b/Engine/Plugins/Experimental/MeshModelingToolset/Source/MeshModelingToolsEditorOnly/MeshModelingToolsEditorOnly.Build.cs
--- a/Engine/Plugins/Experimental/MeshModelingToolset/Source/MeshModelingToolsEditorOnly/MeshModelingToolsEditorOnly.Build.cs
+++ b/Engine/Plugins/Experimental/MeshModelingToolset/Source/MeshModelingToolsEditorOnly/MeshModelingToolsEditorOnly.Build.cs
@@ -62,7 +62,6 @@ public class MeshModelingToolsEditorOnly : ModuleRules
 				"MeshUtilitiesCommon",  
 				"MeshReductionInterface", // for UE standard simplification 
                 "ProxyLODMeshReduction", // for mesh merging voxel-based csg
-				"OpenSubdiv",
 				//"Slate",
 				//"SlateCore",
 				// ... add private dependencies that you statically link with here ...	
diff --git a/Engine/Plugins/Experimental/MeshModelingToolset/Source/MeshModelingToolsEditorOnly/Private/SubdividePoly.cpp b/Engine/Plugins/Experimental/MeshModelingToolset/Source/MeshModelingToolsEditorOnly/Private/SubdividePoly.cpp
deleted file mode 100644
diff --git a/Engine/Plugins/Experimental/MeshModelingToolset/Source/MeshModelingToolsEditorOnly/Public/SubdividePoly.h b/Engine/Plugins/Experimental/MeshModelingToolset/Source/MeshModelingToolsEditorOnly/Public/SubdividePoly.h
deleted file mode 100644
diff --git a/Engine/Plugins/Experimental/MeshModelingToolset/Source/MeshModelingToolsEditorOnly/Public/SubdividePolyTool.h b/Engine/Plugins/Experimental/MeshModelingToolset/Source/MeshModelingToolsEditorOnly/Public/SubdividePolyTool.h
--- a/Engine/Plugins/Experimental/MeshModelingToolset/Source/MeshModelingToolsEditorOnly/Public/SubdividePolyTool.h
+++ b/Engine/Plugins/Experimental/MeshModelingToolset/Source/MeshModelingToolsEditorOnly/Public/SubdividePolyTool.h
@@ -2,7 +2,7 @@
 
 #pragma once
 
-#include "SubdividePoly.h"
+#include "Operations/SubdividePoly.h"
 #include "CoreMinimal.h"
 #include "InteractiveTool.h"
 #include "InteractiveToolBuilder.h"
diff --git a/Engine/Plugins/Experimental/MeshModelingToolset/Source/ModelingComponentsEditorOnly/ModelingComponentsEditorOnly.Build.cs b/Engine/Plugins/Experimental/MeshModelingToolset/Source/ModelingComponentsEditorOnly/ModelingComponentsEditorOnly.Build.cs
--- a/Engine/Plugins/Experimental/MeshModelingToolset/Source/ModelingComponentsEditorOnly/ModelingComponentsEditorOnly.Build.cs
+++ b/Engine/Plugins/Experimental/MeshModelingToolset/Source/ModelingComponentsEditorOnly/ModelingComponentsEditorOnly.Build.cs
@@ -50,6 +50,8 @@ public class ModelingComponentsEditorOnly : ModuleRules
 				"PhysicsCore",
                 "RHI",
 
+				"OpenSubdiv",
+
 				"AssetTools",
 				"UnrealEd"			// required for asset factories
 
diff --git a/Engine/Plugins/Experimental/MeshModelingToolset/Source/ModelingComponentsEditorOnly/Private/Operations/SubdividePoly.cpp b/Engine/Plugins/Experimental/MeshModelingToolset/Source/ModelingComponentsEditorOnly/Private/Operations/SubdividePoly.cpp
new file mode 100644
--- /dev/null
+++ b/Engine/Plugins/Experimental/MeshModelingToolset/Source/ModelingComponentsEditorOnly/Private/Operations/SubdividePoly.cpp
@@ -0,0 +1,662 @@
+// Copyright Epic Games, Inc. All Rights Reserved.
+
+#include "Operations/SubdividePoly.h"
+#include "GroupTopology.h"
+#include "MeshNormals.h"
+
+// OpenSubdiv needs M_PI defined
+#ifndef M_PI
+#define M_PI PI
+#define LOCAL_M_PI 1
+#endif
+
+#pragma warning(push, 0)     
+#include "far/topologyRefiner.h"
+#include "far/topologyDescriptor.h"
+#include "far/primvarRefiner.h"
+#pragma warning(pop)     
+
+#if LOCAL_M_PI
+#undef M_PI
+#endif
+
+#include "ExplicitUseGeometryMathTypes.h"		// using UE::Geometry::(math types)
+using namespace UE::Geometry;
+
+namespace SubdividePolyLocal
+{
+	// Vertices with positions and normals. Used for interpolating from polymesh to subdivision surface mesh.
+	struct SubdVertex
+	{
+		FVertexInfo VertexInfo;
+
+		void Clear()
+		{
+			VertexInfo = FVertexInfo();
+		}
+
+		void AddWithWeight(SubdVertex const& Src, float Weight)
+		{
+			VertexInfo.Position += Weight * Src.VertexInfo.Position;
+			VertexInfo.Normal += Weight * Src.VertexInfo.Normal;
+			VertexInfo.bHaveN = Src.VertexInfo.bHaveN;
+			// TODO: add support for vertex color
+		}
+	};
+
+	// Vertex with only UV data. This is used for "face-varying" interpolation, i.e. vertices can have multiple UV 
+	// values if they are incident on triangles in different UV islands.
+	struct SubdUVVertex
+	{
+		FVector2f VertexUV;
+
+		SubdUVVertex(const FVector2f& InVertexUV) :
+			VertexUV{ InVertexUV }
+		{}
+
+		void Clear()
+		{
+			VertexUV = FVector2f();
+		}
+
+		void AddWithWeight(SubdUVVertex const& Src, float Weight)
+		{
+			VertexUV += Weight * Src.VertexUV;
+		}
+	};
+
+
+	// Compute the average of all normals corresponding to the given VertexID
+	FVector3f GetAverageVertexNormalFromOverlay(const FDynamicMeshNormalOverlay* NormalOverlay,
+												int VertexID)
+	{
+		TArray<int> NormalElements;
+		NormalOverlay->GetVertexElements(VertexID, NormalElements);
+
+		FVector3f SumElements{ 0.0f, 0.0f, 0.0f };
+		float ElementCount = 0.0f;
+		for (int ElementID : NormalElements)
+		{
+			SumElements += NormalOverlay->GetElement(ElementID);
+			ElementCount += 1.0f;
+		}
+
+		return SumElements / ElementCount;
+	}
+
+	// Get the indices of GroupTopology "Corners" for a particular group boundary.
+	void GetBoundaryCorners(const FGroupTopology::FGroupBoundary& Boundary,
+							const FGroupTopology& Topology,
+							TArray<int>& Corners)
+	{
+		int FirstEdgeIndex = Boundary.GroupEdges[0];
+		Corners.Add(Topology.Edges[FirstEdgeIndex].EndpointCorners[0]);
+		Corners.Add(Topology.Edges[FirstEdgeIndex].EndpointCorners[1]);
+
+		int NextEdgeIndex = Boundary.GroupEdges[1];
+		FIndex2i NextEdgeCorners = Topology.Edges[NextEdgeIndex].EndpointCorners;
+		if (Corners[1] != NextEdgeCorners[0] && Corners[1] != NextEdgeCorners[1])
+		{
+			Swap(Corners[0], Corners[1]);
+			check(Corners[1] == NextEdgeCorners[0] || Corners[1] == NextEdgeCorners[1]);
+		}
+
+		for (int i = 1; i < Boundary.GroupEdges.Num() - 1; ++i)
+		{
+			int EdgeIndex = Boundary.GroupEdges[i];
+			FIndex2i CurrEdgeCorners = Topology.Edges[EdgeIndex].EndpointCorners;
+			if (Corners.Last() == CurrEdgeCorners[0])
+			{
+				Corners.Add(CurrEdgeCorners[1]);
+			}
+			else
+			{
+				check(Corners.Last() == CurrEdgeCorners[1]);
+				Corners.Add(CurrEdgeCorners[0]);
+			}
+		}
+	}
+
+	SubdVertex GetVertexInfo(int32 VertexID, const FDynamicMesh3& Mesh, bool bGetNormals)
+	{
+		SubdVertex Vertex;
+		constexpr bool bWantColors = false;
+		constexpr bool bWantVertexUVs = false;
+		constexpr bool bWantNormals = false;
+
+		Mesh.GetVertex(VertexID, Vertex.VertexInfo, bWantNormals, bWantColors, bWantVertexUVs);
+
+		if (bGetNormals)
+		{
+			if (Mesh.HasAttributes() && (Mesh.Attributes()->NumNormalLayers() > 0))
+			{
+				Vertex.VertexInfo.Normal = GetAverageVertexNormalFromOverlay(Mesh.Attributes()->PrimaryNormals(),
+																			 VertexID);
+				Vertex.VertexInfo.bHaveN = true;
+			}
+		}
+
+		return Vertex;
+	}
+
+	// Treat the given FGroupTopology as a polygonal mesh, and get its vertices
+	void GetGroupPolyMeshVertices(const FDynamicMesh3& Mesh,
+								  const FGroupTopology& Topology,
+								  bool bGetNormals,
+								  TArray<SubdVertex>& OutVertices )
+	{
+		for (const FGroupTopology::FCorner& Corner : Topology.Corners)
+		{
+			OutVertices.Add(GetVertexInfo(Corner.VertexID, Mesh, bGetNormals));
+		}
+	}
+
+	void GetAllMeshVertices(const FDynamicMesh3& Mesh,
+							bool bGetNormals,
+							TArray<SubdVertex>& OutVertices)
+	{
+		for (int32 VertexID : Mesh.VertexIndicesItr())
+		{
+			OutVertices.Add(GetVertexInfo(VertexID, Mesh, bGetNormals));
+		}
+	}
+
+	// Given a group and a vertex ID, return:
+	// - a triangle in the group with one of its corners equal to vertex ID
+	// - the (0-2) triangle corner index corresponding to the given vertex ID
+	bool FindTriangleVertex(const FGroupTopology::FGroup& Group,
+							int VertexID,
+							const FDynamicMesh3& Mesh,
+							TTuple<int, int>& OutTriangleVertex)
+	{
+		for (int Tri : Group.Triangles)
+		{
+			FIndex3i TriVertices = Mesh.GetTriangle(Tri);
+			for (int i = 0; i < 3; ++i)
+			{
+				if (TriVertices[i] == VertexID)
+				{
+					OutTriangleVertex = TTuple<int, int>{ Tri, i };
+					return true;
+				}
+			}
+		}
+		return false;
+	}
+
+	// Treating the GroupTopology as a polygonal mesh, get its "face-varying" UV coordinates. Assumes that each 
+	// polygonal face has no UV seams cutting through it, but that UV seams might exist at polygon boundaries.
+	bool GetGroupPolyMeshUVs(const FGroupTopology& Topology,
+							 const FDynamicMesh3& Mesh,
+							 const FDynamicMeshUVOverlay* UVOverlay,
+							 TArray<SubdUVVertex>& OutUVs)		// for each group, all the uvs on the group corners
+	{
+		for (const FGroupTopology::FGroup& Group : Topology.Groups)
+		{
+			check(Group.Boundaries.Num() == 1);
+			check(Group.Triangles.Num() > 0);
+			const FGroupTopology::FGroupBoundary& Bdry = Group.Boundaries[0];
+
+			TArray<int> Corners;
+			GetBoundaryCorners(Bdry, Topology, Corners);
+
+			for (int CornerID : Corners)
+			{
+				int CornerVertexID = Topology.Corners[CornerID].VertexID;
+
+				// Find a triangle in the group that has a vertex ID equal to the given corner
+				TTuple<int, int> TriangleVertex;
+				if (FindTriangleVertex(Group, CornerVertexID, Mesh, TriangleVertex))
+				{
+					int TriangleID = TriangleVertex.Get<0>();
+					int TriVertexIndex = TriangleVertex.Get<1>();	// The (0,1,2) index of the polygon corner wrt the triangle
+
+					TStaticArray<FVector2f, 3> TriangleUVs;
+					UVOverlay->GetTriElements(TriangleID, TriangleUVs[0], TriangleUVs[1], TriangleUVs[2]);
+					FVector2f CornerUV = TriangleUVs[TriVertexIndex];
+
+					OutUVs.Add(SubdUVVertex{ CornerUV });
+				}
+				else
+				{
+					return false;
+				}
+			}
+		}
+
+		return true;
+	}
+
+
+	bool GetMeshUVs(const FDynamicMesh3& Mesh,
+					const FDynamicMeshUVOverlay* UVOverlay,
+					TArray<SubdUVVertex>& OutUVs)		// for each triangle, all the uvs
+	{
+		for (int32 TriangleID : Mesh.TriangleIndicesItr())
+		{
+			TStaticArray<FVector2f, 3> TriangleUVs;
+			UVOverlay->GetTriElements(TriangleID, TriangleUVs[0], TriangleUVs[1], TriangleUVs[2]);
+			OutUVs.Add(SubdUVVertex{ TriangleUVs[0] });
+			OutUVs.Add(SubdUVVertex{ TriangleUVs[1] });
+			OutUVs.Add(SubdUVVertex{ TriangleUVs[2] });
+		}
+
+		return true;
+	}
+
+	void InitializeOverlayToFaceVertexUVs(FDynamicMeshUVOverlay* UVOverlay, const TArray<FVector2f>& UVs)
+	{
+		const FDynamicMesh3* Mesh = UVOverlay->GetParentMesh();
+		check(Mesh->IsCompact());
+
+		int NumTriangles = Mesh->TriangleCount();
+		check(NumTriangles == Mesh->MaxTriangleID());
+		check(UVs.Num() == 3 * NumTriangles);
+
+		UVOverlay->ClearElements();
+		UVOverlay->InitializeTriangles(NumTriangles);
+
+		for (int TriangleIndex = 0; TriangleIndex < NumTriangles; ++TriangleIndex)
+		{
+			FIndex3i Tri{ 3 * TriangleIndex + 0, 3 * TriangleIndex + 1,	3 * TriangleIndex + 2 };
+			UVOverlay->AppendElement(UVs[Tri[0]]);
+			UVOverlay->AppendElement(UVs[Tri[1]]);
+			UVOverlay->AppendElement(UVs[Tri[2]]);
+			UVOverlay->SetTriangle(TriangleIndex, Tri);
+		}
+	}
+
+
+}	// namespace SubdivisionSurfaceLocal
+
+
+class FSubdividePoly::RefinerImpl
+{
+public:
+	OpenSubdiv::Far::TopologyRefiner* TopologyRefiner = nullptr;
+};
+
+FSubdividePoly::FSubdividePoly(const FGroupTopology& InTopology,
+							   const FDynamicMesh3& InOriginalMesh,
+							   int InLevel) :
+	GroupTopology(InTopology)
+	, OriginalMesh(InOriginalMesh)
+	, Level(InLevel)
+{
+	Refiner = MakeUnique<RefinerImpl>();
+}
+
+FSubdividePoly::~FSubdividePoly()
+{
+	if (Refiner && Refiner->TopologyRefiner)
+	{
+		// This was created by TopologyRefinerFactory; looks like we are responsible for cleaning it up.
+		delete Refiner->TopologyRefiner;
+	}
+	Refiner = nullptr;
+}
+
+
+bool FSubdividePoly::ComputeTopologySubdivision()
+{
+	if (Level < 1)
+	{
+		return false;
+	}
+
+	TArray<int> BoundaryVertsPerFace;
+	TArray<int> NumVertsPerFace;
+	OpenSubdiv::Far::TopologyDescriptor::FVarChannel UVChannel;
+	TArray<int> UVBuffer;
+
+	auto DescriptorFromTriangleMesh = [this, &NumVertsPerFace, &BoundaryVertsPerFace, &UVChannel, &UVBuffer]
+	(OpenSubdiv::Far::TopologyDescriptor& Descriptor)
+	{
+		for (auto TriangleID : OriginalMesh.TriangleIndicesItr())
+		{
+			FIndex3i TriangleVertices = OriginalMesh.GetTriangle(TriangleID);
+			NumVertsPerFace.Add(3);
+			BoundaryVertsPerFace.Add(TriangleVertices[0]);
+			BoundaryVertsPerFace.Add(TriangleVertices[1]);
+			BoundaryVertsPerFace.Add(TriangleVertices[2]);
+		}
+
+		// TODO: We should probably create a compact mesh descriptor for subdivision to operate on. UETOOL-2944
+		Descriptor.numVertices = OriginalMesh.MaxVertexID();
+		Descriptor.numFaces = OriginalMesh.TriangleCount();
+		Descriptor.numVertsPerFace = NumVertsPerFace.GetData();
+		Descriptor.vertIndicesPerFace = BoundaryVertsPerFace.GetData();
+
+		if (UVComputationMethod == ESubdivisionOutputUVs::Interpolated)
+		{
+			UVChannel.numValues = BoundaryVertsPerFace.Num();
+			UVBuffer.SetNum(UVChannel.numValues);
+			for (int i = 0; i < UVChannel.numValues; ++i)
+			{
+				UVBuffer[i] = i;
+			}
+			UVChannel.valueIndices = UVBuffer.GetData();
+			Descriptor.numFVarChannels = 1;
+			Descriptor.fvarChannels = &UVChannel;
+		}
+		else
+		{
+			Descriptor.numFVarChannels = 0;
+		}
+
+		return true;
+	};
+
+	auto DescriptorFromGroupTopology = [this, &NumVertsPerFace, &BoundaryVertsPerFace, &UVChannel, &UVBuffer]
+	(OpenSubdiv::Far::TopologyDescriptor& Descriptor)
+	{
+		int TotalNumVertices = 0;
+
+		for (const FGroupTopology::FGroup& Group : GroupTopology.Groups)
+		{
+			if (Group.Boundaries.Num() != 1)
+			{
+				return false;
+			}
+
+			const FGroupTopology::FGroupBoundary& Boundary = Group.Boundaries[0];
+			if (Boundary.GroupEdges.Num() < 2)
+			{
+				return false;
+			}
+
+			TArray<int> Corners;
+			SubdividePolyLocal::GetBoundaryCorners(Boundary, GroupTopology, Corners);
+
+			NumVertsPerFace.Add(Corners.Num());
+			TotalNumVertices += Corners.Num();
+			BoundaryVertsPerFace.Append(Corners);
+		}
+
+		Descriptor.numVertices = TotalNumVertices;
+		Descriptor.numFaces = GroupTopology.Groups.Num();
+		Descriptor.numVertsPerFace = NumVertsPerFace.GetData();
+		Descriptor.vertIndicesPerFace = BoundaryVertsPerFace.GetData();
+
+		if (UVComputationMethod == ESubdivisionOutputUVs::Interpolated)
+		{
+			UVChannel.numValues = BoundaryVertsPerFace.Num();
+			UVBuffer.SetNum(UVChannel.numValues);
+			for (int i = 0; i < UVChannel.numValues; ++i)
+			{
+				UVBuffer[i] = i;
+			}
+			UVChannel.valueIndices = UVBuffer.GetData();
+			Descriptor.numFVarChannels = 1;
+			Descriptor.fvarChannels = &UVChannel;
+		}
+		else
+		{
+			Descriptor.numFVarChannels = 0;
+		}
+
+		return true;
+	};
+
+	OpenSubdiv::Far::TopologyDescriptor Descriptor;
+
+	if (SubdivisionScheme == ESubdivisionScheme::Loop)
+	{
+		if (!DescriptorFromTriangleMesh(Descriptor))
+		{
+			return false;
+		}
+	}
+	else
+	{
+		if (!DescriptorFromGroupTopology(Descriptor))
+		{
+			return false;
+		}
+	}
+
+	using RefinerFactory = OpenSubdiv::Far::TopologyRefinerFactory<OpenSubdiv::Far::TopologyDescriptor>;
+	RefinerFactory::Options RefinerOptions;
+
+	RefinerOptions.schemeOptions.SetVtxBoundaryInterpolation(OpenSubdiv::Sdc::Options::VTX_BOUNDARY_EDGE_AND_CORNER);
+
+	switch (SubdivisionScheme)
+	{
+	case ESubdivisionScheme::Bilinear:
+		RefinerOptions.schemeType = OpenSubdiv::Sdc::SchemeType::SCHEME_BILINEAR;
+		break;
+	case ESubdivisionScheme::CatmullClark:
+		RefinerOptions.schemeType = OpenSubdiv::Sdc::SchemeType::SCHEME_CATMARK;
+		break;
+	case ESubdivisionScheme::Loop:
+		RefinerOptions.schemeType = OpenSubdiv::Sdc::SchemeType::SCHEME_LOOP;
+		break;
+	}
+
+	Refiner->TopologyRefiner = RefinerFactory::Create(Descriptor, RefinerOptions);
+
+	if (Refiner->TopologyRefiner == nullptr)
+	{
+		return false;
+	}
+
+	Refiner->TopologyRefiner->RefineUniform(OpenSubdiv::Far::TopologyRefiner::UniformOptions(Level));
+
+	return true;
+}
+
+
+bool FSubdividePoly::ComputeSubdividedMesh(FDynamicMesh3& OutMesh)
+{
+	if (Level < 1)
+	{
+		return false;
+	}
+
+	if (!Refiner || !(Refiner->TopologyRefiner))
+	{
+		return false;
+	}
+
+	OpenSubdiv::Far::PrimvarRefiner Interpolator(*Refiner->TopologyRefiner);
+
+	//
+	// Interpolate vertex positions from initial Group poly mesh down to refinement level
+	// 
+	TArray<SubdividePolyLocal::SubdVertex> SourceVertices;
+	bool bInterpolateVertexNormals = (NormalComputationMethod == ESubdivisionOutputNormals::Interpolated);
+	if (SubdivisionScheme == ESubdivisionScheme::Loop)
+	{
+		GetAllMeshVertices(OriginalMesh, bInterpolateVertexNormals, SourceVertices);
+	}
+	else
+	{
+		GetGroupPolyMeshVertices(OriginalMesh, GroupTopology, bInterpolateVertexNormals, SourceVertices);
+	}
+
+	TArray<SubdividePolyLocal::SubdVertex> RefinedVertices;
+	for (int CurrentLevel = 1; CurrentLevel <= Level; ++CurrentLevel)
+	{
+		// TODO: Don't keep resizing -- preallocate one big buffer and move through it
+		int NumVertices = Refiner->TopologyRefiner->GetLevel(CurrentLevel).GetNumVertices();
+		RefinedVertices.SetNumUninitialized(NumVertices);
+		SubdividePolyLocal::SubdVertex* s = SourceVertices.GetData();
+		SubdividePolyLocal::SubdVertex* d = RefinedVertices.GetData();
+		Interpolator.Interpolate(CurrentLevel, s, d);
+		SourceVertices = RefinedVertices;
+	}
+
+	//
+	// Interpolate face group IDs
+	// 
+	TArray<int> SourceGroupIDs;
+	if (SubdivisionScheme == ESubdivisionScheme::Loop)
+	{
+		for ( int TriangleID : OriginalMesh.TriangleIndicesItr() )
+		{
+			SourceGroupIDs.Add(OriginalMesh.GetTriangleGroup(TriangleID));
+		}
+	}
+	else
+	{
+		for (const FGroupTopology::FGroup& Group : GroupTopology.Groups)
+		{
+			SourceGroupIDs.Add(Group.GroupID);
+		}
+	}
+	check(SourceGroupIDs.Num() == Refiner->TopologyRefiner->GetLevel(0).GetNumFaces());
+
+	TArray<int> RefinedGroupIDs;
+	for (int CurrentLevel = 1; CurrentLevel <= Level; ++CurrentLevel)
+	{
+		// TODO: Don't keep resizing -- preallocate one big buffer and move through it
+		RefinedGroupIDs.SetNumUninitialized(Refiner->TopologyRefiner->GetLevel(CurrentLevel).GetNumFaces());
+		int* s = SourceGroupIDs.GetData();
+		int* d = RefinedGroupIDs.GetData();
+		Interpolator.InterpolateFaceUniform(CurrentLevel, s, d);
+		SourceGroupIDs = RefinedGroupIDs;
+	}
+
+
+	//
+	// Interpolate UVs
+	// 
+	TArray<SubdividePolyLocal::SubdUVVertex> RefinedUVs;
+
+	if (UVComputationMethod == ESubdivisionOutputUVs::Interpolated)
+	{
+		if (!OriginalMesh.HasAttributes() || !OriginalMesh.Attributes()->PrimaryUV())
+		{
+			return false;
+		}
+
+		TArray<SubdividePolyLocal::SubdUVVertex> SourceUVs;
+		bool bGetUVsOK = false;
+
+		if (SubdivisionScheme == ESubdivisionScheme::Loop)
+		{
+			bGetUVsOK = GetMeshUVs(OriginalMesh, OriginalMesh.Attributes()->PrimaryUV(), SourceUVs);
+		}
+		else
+		{
+			bGetUVsOK = GetGroupPolyMeshUVs(GroupTopology,
+											OriginalMesh,
+											OriginalMesh.Attributes()->PrimaryUV(),
+											SourceUVs);
+		}
+
+		if (!bGetUVsOK)
+		{
+			return false;
+		}
+
+		check(SourceUVs.Num() == Refiner->TopologyRefiner->GetLevel(0).GetNumFaceVertices());
+
+		for (int CurrentLevel = 1; CurrentLevel <= Level; ++CurrentLevel)
+		{
+			// TODO: Don't keep resizing -- preallocate one big buffer and move through it
+			RefinedUVs.SetNumUninitialized(Refiner->TopologyRefiner->GetLevel(CurrentLevel).GetNumFVarValues());
+			SubdividePolyLocal::SubdUVVertex* s = SourceUVs.GetData();
+			SubdividePolyLocal::SubdUVVertex* d = RefinedUVs.GetData();
+			Interpolator.InterpolateFaceVarying(CurrentLevel, s, d);
+			SourceUVs = RefinedUVs;
+		}
+	}
+
+	// Now transfer to output mesh
+	OutMesh.Clear();
+
+	OutMesh.EnableTriangleGroups();
+	if (bInterpolateVertexNormals)
+	{
+		OutMesh.EnableVertexNormals(FVector3f{ 0,0,0 });
+	}
+	if ((NormalComputationMethod != ESubdivisionOutputNormals::None) || (UVComputationMethod != ESubdivisionOutputUVs::None))
+	{
+		OutMesh.EnableAttributes();
+	}
+
+	// Add the vertices
+	// TODO: Can we resize the Mesh vertex array up front and then populate it?
+	for (const SubdividePolyLocal::SubdVertex& V : RefinedVertices)
+	{
+		OutMesh.AppendVertex(V.VertexInfo);
+	}
+
+	const OpenSubdiv::Far::TopologyLevel& FinalLevel = Refiner->TopologyRefiner->GetLevel(Level);
+
+	check(UVComputationMethod != ESubdivisionOutputUVs::Interpolated || FinalLevel.GetNumFVarValues() == RefinedUVs.Num());
+	check(FinalLevel.GetNumFaces() == RefinedGroupIDs.Num());
+
+	// Add the faces (manually triangulate the output here)
+
+	int NumFaceVertices = 0;
+	TArray<FVector2f> TriangleUVs;
+
+	for (int FaceID = 0; FaceID < FinalLevel.GetNumFaces(); ++FaceID)
+	{
+		int GroupID = RefinedGroupIDs[FaceID];
+
+		OpenSubdiv::Far::ConstIndexArray Face = FinalLevel.GetFaceVertices(FaceID);
+		NumFaceVertices += Face.size();
+
+		if (Face.size() == 4)
+		{
+			FIndex3i TriA{ Face[0], Face[1], Face[3] };
+			FIndex3i TriB{ Face[2], Face[3], Face[1] };
+			OutMesh.AppendTriangle(TriA, GroupID);
+			OutMesh.AppendTriangle(TriB, GroupID);
+
+			if (UVComputationMethod == ESubdivisionOutputUVs::Interpolated)
+			{
+				OpenSubdiv::Far::ConstIndexArray FaceUVIndices = FinalLevel.GetFaceFVarValues(FaceID);
+
+				FIndex3i UVTriA{ FaceUVIndices[0], FaceUVIndices[1], FaceUVIndices[3] };
+				TriangleUVs.Add(RefinedUVs[UVTriA[0]].VertexUV);
+				TriangleUVs.Add(RefinedUVs[UVTriA[1]].VertexUV);
+				TriangleUVs.Add(RefinedUVs[UVTriA[2]].VertexUV);
+
+				FIndex3i UVTriB{ FaceUVIndices[2], FaceUVIndices[3], FaceUVIndices[1] };
+				TriangleUVs.Add(RefinedUVs[UVTriB[0]].VertexUV);
+				TriangleUVs.Add(RefinedUVs[UVTriB[1]].VertexUV);
+				TriangleUVs.Add(RefinedUVs[UVTriB[2]].VertexUV);
+			}
+		}
+		else
+		{
+			check(Face.size() == 3);
+			OutMesh.AppendTriangle(FIndex3i{ Face[0], Face[1], Face[2] }, GroupID);
+
+			if (UVComputationMethod == ESubdivisionOutputUVs::Interpolated)
+			{
+				OpenSubdiv::Far::ConstIndexArray FaceUVIndices = FinalLevel.GetFaceFVarValues(FaceID);
+				TriangleUVs.Add(RefinedUVs[FaceUVIndices[0]].VertexUV);
+				TriangleUVs.Add(RefinedUVs[FaceUVIndices[1]].VertexUV);
+				TriangleUVs.Add(RefinedUVs[FaceUVIndices[2]].VertexUV);
+			}
+		}
+	}
+
+	if (NormalComputationMethod != ESubdivisionOutputNormals::None)
+	{
+		const bool bUseExistingMeshVertexNormals = (NormalComputationMethod == ESubdivisionOutputNormals::Interpolated);
+		FMeshNormals::InitializeOverlayToPerVertexNormals(OutMesh.Attributes()->PrimaryNormals(), bUseExistingMeshVertexNormals);
+	}
+
+	if (UVComputationMethod == ESubdivisionOutputUVs::Interpolated)
+	{
+		SubdividePolyLocal::InitializeOverlayToFaceVertexUVs(OutMesh.Attributes()->PrimaryUV(), TriangleUVs);
+	}
+
+	// Remove any vertices that are not referenced by a face
+	for (int Vid = 0; Vid < OutMesh.MaxVertexID(); ++Vid)
+	{
+		if (OutMesh.IsVertex(Vid) && !OutMesh.IsReferencedVertex(Vid))
+		{
+			OutMesh.RemoveVertex(Vid, false, false);
+		}
+	}
+
+	return true;
+}
+
diff --git a/Engine/Plugins/Experimental/MeshModelingToolset/Source/ModelingComponentsEditorOnly/Public/Operations/SubdividePoly.h b/Engine/Plugins/Experimental/MeshModelingToolset/Source/ModelingComponentsEditorOnly/Public/Operations/SubdividePoly.h
new file mode 100644
--- /dev/null
+++ b/Engine/Plugins/Experimental/MeshModelingToolset/Source/ModelingComponentsEditorOnly/Public/Operations/SubdividePoly.h
@@ -0,0 +1,92 @@
+// Copyright Epic Games, Inc. All Rights Reserved.
+
+#pragma once
+
+#include "CoreMinimal.h"
+#include "ModelingOperators.h"
+#include "SubdividePoly.generated.h"
+
+PREDECLARE_USE_GEOMETRY_CLASS(FGroupTopology);
+PREDECLARE_USE_GEOMETRY_CLASS(FDynamicMesh3);
+class FProgressCancel;
+
+/**
+ * Subdivision scheme
+ */
+UENUM()
+enum class ESubdivisionScheme : uint8
+{
+	Bilinear,
+	CatmullClark,
+	Loop
+};
+
+
+/**
+ * Method for computing normals on the output mesh
+ */
+UENUM()
+enum class ESubdivisionOutputNormals : uint8
+{
+	// Copy normals to poly mesh then interpolate to new triangles as we subdivide.
+	Interpolated,
+
+	// Recompute normals after subdivision
+	Generated,
+
+	// Leave normals uninitialized on output mesh
+	None
+};
+
+/**
+ * Method for computing UVs on the output mesh
+ */
+UENUM()
+enum class ESubdivisionOutputUVs : uint8
+{
+	// Copy UVs to poly mesh vertices then interpolate as we subdivide.
+	Interpolated,
+
+	// Leave UVs uninitialized on output mesh
+	None
+};
+
+
+/**
+ * Interprets an FGroupTopology object as a poly mesh and refines it to a given level. Topological subdivision 
+ * can be computed separately from the interpolation of vertex positions. A typical workflow might be to first 
+ * precompute the refined topology, then compute the interpolated positions (and other vertex values) as the original 
+ * polymesh is deformed.
+ */
+class MODELINGCOMPONENTSEDITORONLY_API FSubdividePoly
+{
+public:
+
+	FSubdividePoly(const FGroupTopology& InTopology,
+				   const FDynamicMesh3& InOriginalMesh,
+				   int InLevel);
+
+	~FSubdividePoly();
+
+	bool ComputeTopologySubdivision();
+
+	bool ComputeSubdividedMesh(FDynamicMesh3& OutMesh);
+
+	ESubdivisionScheme SubdivisionScheme = ESubdivisionScheme::CatmullClark;
+
+	ESubdivisionOutputNormals NormalComputationMethod = ESubdivisionOutputNormals::Generated;
+
+	ESubdivisionOutputUVs UVComputationMethod = ESubdivisionOutputUVs::Interpolated;
+
+	// wrapper around Opensubdiv::TopologyRefiner to avoid including OSD headers here. See cpp file.
+	class RefinerImpl;
+
+private:
+
+	const FGroupTopology& GroupTopology;
+	const FDynamicMesh3& OriginalMesh;
+	int Level;
+
+	TUniquePtr<RefinerImpl> Refiner;
+};
+
