ApproximateActors: Implement support for detection of thin-slabs (ie thin planar meshes that are not single-sided) in FMeshSceneAdapter. Add support for sub-mesh instancing, to reduce number of times that repeated meshes are unique'd when used with the same scaling transform. Reorganize computation so that spatial builds are processed in decreasing-mesh-size order, and do other parallel improvements, such that time to process a large assembly (total ~2.2b triangles from 37k instances of 415 unique meshes w/~44m unique triangles) reduces from 58s to 22s (with ~10s of this time being loading source mesh data) on 64-core threadripper.
#rnx

diff --git a/Engine/Plugins/Experimental/MeshModelingToolset/Source/GeometryProcessingAdapters/Private/ApproximateActorsImpl.cpp b/Engine/Plugins/Experimental/MeshModelingToolset/Source/GeometryProcessingAdapters/Private/ApproximateActorsImpl.cpp
--- a/Engine/Plugins/Experimental/MeshModelingToolset/Source/GeometryProcessingAdapters/Private/ApproximateActorsImpl.cpp
+++ b/Engine/Plugins/Experimental/MeshModelingToolset/Source/GeometryProcessingAdapters/Private/ApproximateActorsImpl.cpp
@@ -777,7 +777,7 @@ void FApproximateActorsImpl::GenerateApproximationForActorSet(const TArray<AActo
 	{
 		FMeshSceneAdapter::FStatistics Stats;
 		Scene.GetGeometryStatistics(Stats);
-		UE_LOG(LogApproximateActors, Warning, TEXT("%d triangles in %d unique meshes, total %d triangles in %d instances"),
+		UE_LOG(LogApproximateActors, Warning, TEXT("%lld triangles in %lld unique meshes, total %lld triangles in %lld instances"),
 			Stats.UniqueMeshTriangleCount, Stats.UniqueMeshCount, Stats.InstanceMeshTriangleCount, Stats.InstanceMeshCount);
 	}
 
diff --git a/Engine/Plugins/Experimental/MeshModelingToolset/Source/ModelingComponents/Private/Scene/MeshSceneAdapter.cpp b/Engine/Plugins/Experimental/MeshModelingToolset/Source/ModelingComponents/Private/Scene/MeshSceneAdapter.cpp
--- a/Engine/Plugins/Experimental/MeshModelingToolset/Source/ModelingComponents/Private/Scene/MeshSceneAdapter.cpp
+++ b/Engine/Plugins/Experimental/MeshModelingToolset/Source/ModelingComponents/Private/Scene/MeshSceneAdapter.cpp
@@ -73,9 +73,77 @@ void CollectSeedPointsFromMeshVertices(
 }
 
 
+/**
+ * Try to check if a Mesh is "thin", ie basically a planar patch (open or closed).
+ * The Normal of the largest-area triangle is taken as the plane normal, and then the "thickness" is measured relative to this plane
+ * @param ThinTolerance identify as Thin if the thickness extents is within this size
+ * @param ThinPlaneOut thin plane normal will be returned via this frame
+ * @param ThicknessOut measured thickness relative to plane will be returned here
+ * @return true if mesh is identified as thin under ThinTolerance
+ */
+template<typename MeshType>
+bool IsThinPlanarSubMesh(const MeshType& Mesh, double ThinTolerance, FFrame3d& ThinPlaneOut, double& ThicknessOut)
+{
+	int32 TriCount = Mesh.TriangleCount();
+
+	// Find triangle with largest area and use it's normal as the plane normal
+	// (this is not ideal and we should probably do a normals histogram
+	double MaxArea = 0;
+	FVector3d MaxAreaNormal;
+	FVector3d MaxAreaPoint;
+	for (int32 tid = 0; tid < TriCount; ++tid)
+	{
+		if (Mesh.IsTriangle(tid))
+		{
+			FVector3d A, B, C;
+			Mesh.GetTriVertices(tid, A, B, C);
+			double TriArea;
+			FVector3d TriNormal = VectorUtil::NormalArea(A, B, C, TriArea);
+			if (TriArea > MaxArea)
+			{
+				MaxArea = TriArea;
+				MaxAreaNormal = TriNormal;
+				MaxAreaPoint = A;
+			}
+		}
+	}
+
+	// Now compute the bounding box in the local space of this plane
+	ThinPlaneOut = FFrame3d(MaxAreaPoint, MaxAreaNormal);
+	FAxisAlignedBox3d PlaneExtents = FAxisAlignedBox3d::Empty();
+	int32 VertexCount = Mesh.VertexCount();
+	for (int32 tid = 0; tid < TriCount; ++tid)
+	{
+		if (Mesh.IsTriangle(tid))
+		{
+			FVector3d TriVerts[3];
+			Mesh.GetTriVertices(tid, TriVerts[0], TriVerts[1], TriVerts[2]);
+			for (int32 j = 0; j < 3; ++j)
+			{
+				TriVerts[j] = ThinPlaneOut.ToFramePoint(TriVerts[j]);
+				PlaneExtents.Contain(TriVerts[j]);
+			}
+		}
+
+		// early-out if we exceed tolerance
+		if (PlaneExtents.Depth() > ThinTolerance)
+		{
+			return false;
+		}
+	}
+
+	// shift plane to center
+	FVector3d Center = PlaneExtents.Center();
+	ThinPlaneOut.Origin += Center.X*ThinPlaneOut.X() + Center.Y*ThinPlaneOut.Y() + Center.Z*ThinPlaneOut.Z();
+	ThicknessOut = PlaneExtents.Depth();
+	return true;
+}
+
+
+
 /**
  * Try to check if the subset of Triangles of Mesh represent a "thin" region, ie basically a planar patch (open or closed).
- * The Normal of the largest-area triangle is taken as the plane normal, and then the "thickness" is measured relative to this plain
+ * The Normal of the largest-area triangle is taken as the plane normal, and then the "thickness" is measured relative to this plane
  * @param ThinTolerance identify as Thin if the thickness extents is within this size
  * @param ThinPlaneOut thin plane normal will be returned via this frame
  * @return true if submesh identified as thin
@@ -84,6 +152,9 @@ template<typename MeshType>
 bool IsThinPlanarSubMesh(const MeshType& Mesh, const TArray<int32>& Triangles, double ThinTolerance, FFrame3d& ThinPlaneOut)
 {
 	int32 TriCount = Triangles.Num();
+
+	// Find triangle with largest area and use it's normal as the plane normal
+	// (this is not ideal and we should probably do a normals histogram
 	double MaxArea = 0;
 	FVector3d MaxAreaNormal;
 	FVector3d MaxAreaPoint;
@@ -104,6 +175,8 @@ bool IsThinPlanarSubMesh(const MeshType& Mesh, const TArray<int32>& Triangles, d
 			}
 		}
 	}
+
+	// Now compute the bounding box in the local space of this plane
 	ThinPlaneOut = FFrame3d(MaxAreaPoint, MaxAreaNormal);
 	FAxisAlignedBox3d PlaneExtents = FAxisAlignedBox3d::Empty();
 	int32 VertexCount = Mesh.VertexCount();
@@ -120,13 +193,18 @@ bool IsThinPlanarSubMesh(const MeshType& Mesh, const TArray<int32>& Triangles, d
 				PlaneExtents.Contain(TriVerts[j]);
 			}
 		}
+
+		// early-out if we exceed tolerance
+		if (PlaneExtents.Depth() > ThinTolerance)
+		{
+			return false;
+		}
 	}
-	if (PlaneExtents.Depth() > ThinTolerance)
-	{
-		return false;
-	}
+
+	// shift plane to center
 	FVector3d Center = PlaneExtents.Center();
 	ThinPlaneOut.Origin += Center.X*ThinPlaneOut.X() + Center.Y*ThinPlaneOut.Y() + Center.Z*ThinPlaneOut.Z();
+
 	return true;
 }
 
@@ -160,6 +238,8 @@ public:
 
 	// if true, Mesh is in world space (whatever that means)
 	bool bHasBakedTransform = false;
+	// if true, Mesh is only translated and rotated (allows some assumptions to be made)
+	bool bHasBakedScale = false;
 	// if true, use unsigned distance to determine inside/outside instead of winding number
 	bool bUseDistanceShellForWinding = false;
 	// unsigned distance isovalue that defines inside
@@ -173,10 +253,16 @@ public:
 		ensure(Mesh.TriangleCount() > 0);
 		if (BuildOptions.bBuildSpatialDataStructures)
 		{
-			AABBTree = MakeUnique<TMeshAABBTree3<FDynamicMesh3>>(&Mesh, true);
+			{
+				TRACE_CPUPROFILER_EVENT_SCOPE(MeshScene_WrapperBuild_DMesh_AABBTree);
+				AABBTree = MakeUnique<TMeshAABBTree3<FDynamicMesh3>>(&Mesh, true);
+			}
 			if (bUseDistanceShellForWinding == false)
 			{
-				FWNTree = MakeUnique<TFastWindingTree<FDynamicMesh3>>(AABBTree.Get(), true);
+				{
+					TRACE_CPUPROFILER_EVENT_SCOPE(MeshScene_WrapperBuild_DMesh_FWNTree);
+					FWNTree = MakeUnique<TFastWindingTree<FDynamicMesh3>>(AABBTree.Get(), true);
+				}
 			}
 		}
 		return true;
@@ -210,15 +296,16 @@ public:
 	{
 		if (bUseDistanceShellForWinding)
 		{
-			if (bHasBakedTransform)
+			if (bHasBakedTransform || bHasBakedScale)
 			{
+				FVector3d UseP = (bHasBakedTransform) ? P : LocalToWorldTransform.InverseTransformPosition(P);
 				double NearestDistSqr;
-				int32 NearTriID = AABBTree->FindNearestTriangle(P, NearestDistSqr, IMeshSpatial::FQueryOptions(WindingShellThickness));
+				int32 NearTriID = AABBTree->FindNearestTriangle(UseP, NearestDistSqr, IMeshSpatial::FQueryOptions(WindingShellThickness));
 				if (NearTriID != IndexConstants::InvalidID)
 				{
 					// Do we even need to do this? won't we return InvalidID if we don't find point within distance?
 					// (also technically we can early-out as soon as we find any point, not the nearest point - might be worth a custom query)
-					FDistPoint3Triangle3d Query = TMeshQueries<FDynamicMesh3>::TriangleDistance(Mesh, NearTriID, P);
+					FDistPoint3Triangle3d Query = TMeshQueries<FDynamicMesh3>::TriangleDistance(Mesh, NearTriID, UseP);
 					if (Query.GetSquared() < WindingShellThickness * WindingShellThickness)
 					{
 						return 1.0;
@@ -227,8 +314,7 @@ public:
 			}
 			else
 			{
-				// todo
-				check(false);
+				ensure(false);		// not supported yet
 			}
 			return 0.0;
 		}
@@ -286,6 +372,7 @@ class FStaticMeshSpatialWrapper : public IMeshSpatialWrapper
 public:
 	UStaticMesh* StaticMesh = nullptr;
 	int32 LODIndex = 0;
+	FVector3d BuildScale = FVector3d::One();
 	
 	FMeshDescription* SourceMesh = nullptr;
 
@@ -306,7 +393,7 @@ public:
 		{
 			Adapter = MakeUnique<FMeshDescriptionTriangleMeshAdapter>(SourceMesh);
 
-			FVector3d BuildScale;
+			BuildScale = FVector3d::One();
 #if WITH_EDITOR
 			// respect BuildScale build setting
 			const FMeshBuildSettings& LODBuildSettings = StaticMesh->GetSourceModel(LODIndex).BuildSettings;
@@ -315,8 +402,14 @@ public:
 #endif
 			if (BuildOptions.bBuildSpatialDataStructures)
 			{
-				AABBTree = MakeUnique<TMeshAABBTree3<FMeshDescriptionTriangleMeshAdapter>>(Adapter.Get(), true);
-				FWNTree = MakeUnique<TFastWindingTree<FMeshDescriptionTriangleMeshAdapter>>(AABBTree.Get(), true);
+				{
+					TRACE_CPUPROFILER_EVENT_SCOPE(MeshScene_WrapperBuild_StaticMesh_AABBTree);
+					AABBTree = MakeUnique<TMeshAABBTree3<FMeshDescriptionTriangleMeshAdapter>>(Adapter.Get(), true);
+				}
+				{
+					TRACE_CPUPROFILER_EVENT_SCOPE(MeshScene_WrapperBuild_StaticMesh_FWNTree);
+					FWNTree = MakeUnique<TFastWindingTree<FMeshDescriptionTriangleMeshAdapter>>(AABBTree.Get(), true);
+				}
 			}
 			return true;
 		}
@@ -364,6 +457,21 @@ public:
 	virtual void AppendMesh(FDynamicMesh3& AppendTo, const FTransformSequence3d& TransformSeq) override
 	{
 		if (!SourceMesh) return;
+
+#if WITH_EDITOR
+		if (AppendTo.TriangleCount() == 0 && TransformSeq.Num() == 0)
+		{
+			// this is somewhat faster in profiling
+			FMeshDescription* UseMeshDescription = StaticMesh->GetMeshDescription(LODIndex);
+			FMeshDescriptionToDynamicMesh Converter;
+			Converter.bEnableOutputGroups = false; Converter.bCalculateMaps = false;
+			Converter.bDisableAttributes = true;
+			Converter.Convert(UseMeshDescription, AppendTo);
+			MeshTransforms::Scale(AppendTo, BuildScale, FVector3d::Zero());
+			return;
+		}
+#endif
+
 		FDynamicMeshEditor Editor(&AppendTo);
 		FMeshIndexMappings Mappings;
 		FMeshDescriptionMeshAdapterd AdapterWrapper(*Adapter);
@@ -453,6 +561,8 @@ static void CollectActorChildMeshes(AActor* Actor, UActorComponent* Component, F
 
 void FMeshSceneAdapter::AddActors(const TArray<AActor*>& ActorsSetIn)
 {
+	TRACE_CPUPROFILER_EVENT_SCOPE(MeshScene_AddActors);
+
 	// build an FActorAdapter for each Actor, that contains all mesh Components we know 
 	// how to process, including those contained in ChildActorComponents
 	TArray<AActor*> ChildActors;
@@ -512,12 +622,6 @@ void FMeshSceneAdapter::AddActors(const TArray<AActor*>& ActorsSetIn)
 		}
 	}
 
-
-	// TODO: compute per-actor world bounds
-	//FVector ActorOrigin, ActorExtent;
-	//Actor->GetActorBounds(false, ActorOrigin, ActorExtent, true);
-	//Adapter->WorldBounds = FAxisAlignedBox3d(
-	//	(FVector3d)ActorOrigin - (FVector3d)ActorExtent, (FVector3d)ActorOrigin + (FVector3d)ActorExtent);
 }
 
 
@@ -526,6 +630,8 @@ void FMeshSceneAdapter::AddActors(const TArray<AActor*>& ActorsSetIn)
 
 void FMeshSceneAdapter::Build(const FMeshSceneAdapterBuildOptions& BuildOptions)
 {
+	TRACE_CPUPROFILER_EVENT_SCOPE(MeshScene_Build);
+
 	if (BuildOptions.bThickenThinMeshes)
 	{
 		Build_FullDecompose(BuildOptions);
@@ -564,62 +670,176 @@ void FMeshSceneAdapter::Build(const FMeshSceneAdapterBuildOptions& BuildOptions)
 
 
 	// update bounding boxes
-	ParallelFor(SceneActors.Num(), [&](int32 i)
 	{
-		UpdateActorBounds(*SceneActors[i]);
-	});
+		TRACE_CPUPROFILER_EVENT_SCOPE(MeshScene_Build_ActorBounds);
+		ParallelFor(SceneActors.Num(), [&](int32 i)
+		{
+			UpdateActorBounds(*SceneActors[i]);
+		});
+	}
 }
 
 
 void FMeshSceneAdapter::UpdateActorBounds(FActorAdapter& Actor)
 {
-	FAxisAlignedBox3d ActorBounds = FAxisAlignedBox3d::Empty();
-	for (const TUniquePtr<FActorChildMesh>& ChildMesh : Actor.ChildMeshes)
+	int32 NumChildren = Actor.ChildMeshes.Num();
+	TArray<FAxisAlignedBox3d> ChildBounds;
+	ChildBounds.Init(FAxisAlignedBox3d::Empty(), NumChildren);
+	ParallelFor(NumChildren, [&](int32 k)
 	{
+		const TUniquePtr<FActorChildMesh>& ChildMesh = Actor.ChildMeshes[k];
 		if (ChildMesh->MeshSpatial != nullptr)
 		{
-			FAxisAlignedBox3d ChildBounds = ChildMesh->MeshSpatial->GetWorldBounds(
+			ChildBounds[k] = ChildMesh->MeshSpatial->GetWorldBounds(
 				[&](const FVector3d& P) { return ChildMesh->WorldTransform.TransformPosition(P); });
-			ActorBounds.Contain(ChildBounds);
 		}
+	});
+
+	Actor.WorldBounds = FAxisAlignedBox3d::Empty();
+	for (FAxisAlignedBox3d ChildBound : ChildBounds)
+	{
+		Actor.WorldBounds.Contain(ChildBound);
 	}
-	Actor.WorldBounds = ActorBounds;
 }
 
 
+
+/**
+ * This function is used to group the input set of transforms into subsets
+ * that have the same scale. Each of those subsets can be represented by
+ * a single scaled mesh with different rotate/translate-only transforms.
+ * We use this to reduce the number of times a mesh has to be duplicated when
+ * breaking it up into parts that require further processing that is incompatible
+ * with (nonuniform) scaling.
+ * TODO: Currently cannot differentiate between uniform and nonuniform scaling
+ */
+void ConstructUniqueScalesMapping(
+	const TArray<FTransformSequence3d>& TransformSet, 
+	TArray<TArray<int32>>& UniqueScaleSetsOut,
+	double ScaleComponentTolerance = 0.01)
+{
+	// two transforms are "the same up to scaling" if this returns true
+	auto CompareScales = [ScaleComponentTolerance](const FTransform3d& T1, const FTransform3d& T2)
+	{
+		return (T1.GetScale() - T2.GetScale()).GetAbsMax() < ScaleComponentTolerance;
+	};
+
+	TArray<FTransformSequence3d> UniqueScaleTransforms;		// accumulate transform-with-unique-scale's here
+	int32 N = TransformSet.Num();
+	TArray<int32> UniqueScaleMap;
+	UniqueScaleMap.SetNum(N);
+	for (int32 k = 0; k < N; ++k)
+	{
+		FTransformSequence3d CurTransform = TransformSet[k];
+		int32 FoundIndex = -1;
+		for (int32 j = 0; j < UniqueScaleTransforms.Num(); ++j)
+		{
+			if (CurTransform.IsEquivalent(UniqueScaleTransforms[j], CompareScales))
+			{
+				FoundIndex = j;
+				break;
+			}
+		}
+		if (FoundIndex >= 0)
+		{
+			UniqueScaleMap[k] = FoundIndex;
+		}
+		else
+		{
+			UniqueScaleMap[k] = UniqueScaleTransforms.Num();
+			UniqueScaleTransforms.Add(CurTransform);
+		}
+	}
+
+	// build clusters
+	int32 NumUniqueScales = UniqueScaleTransforms.Num();
+	UniqueScaleSetsOut.SetNum(NumUniqueScales);
+	for (int32 k = 0; k < N; ++k)
+	{
+		UniqueScaleSetsOut[UniqueScaleMap[k]].Add(k);
+	}
+
+}
+
+
+
 void FMeshSceneAdapter::Build_FullDecompose(const FMeshSceneAdapterBuildOptions& BuildOptions)
 {
+	// initial list of spatial wrappers that need to be built
 	TArray<FSpatialWrapperInfo*> ToBuild;
 	for (TPair<void*, TSharedPtr<FSpatialWrapperInfo>> Pair : SpatialAdapters)
 	{
 		ToBuild.Add(Pair.Value.Get());
 	}
 
-	FCriticalSection ListsLock;
+	// Initialize the initial set of wrappers. Must do this here so that meshes are loaded and TriangleCount() below is valid
+	FMeshSceneAdapterBuildOptions TempBuildOptions = BuildOptions;
+	TempBuildOptions.bBuildSpatialDataStructures = false;
+	ParallelFor(ToBuild.Num(), [&](int32 i)
+	{
+		FSpatialWrapperInfo* WrapperInfo = ToBuild[i];
+		WrapperInfo->SpatialWrapper->Build(TempBuildOptions);
+	});
 
-	std::atomic<int32> DecomposedSourceMeshCount;
-	DecomposedSourceMeshCount = 0;
-	std::atomic<int32> DecomposedMeshesCount;
-	DecomposedMeshesCount = 0;
-	int32 AddedTrisCount = 0;
+	// sort build list by increasing triangle count
+	ToBuild.Sort([&](const FSpatialWrapperInfo& A, const FSpatialWrapperInfo& B)
+	{
+		return const_cast<FSpatialWrapperInfo&>(A).SpatialWrapper->GetTriangleCount() < const_cast<FSpatialWrapperInfo&>(B).SpatialWrapper->GetTriangleCount();
+	});
 
-	FCriticalSection PendingFuturesLock;
-	TArray<TFuture<void>> PendingClosedAssemblyBuilds;
+	// stats we will collect during execution
+	int32 NumInitialSources = ToBuild.Num();
+	std::atomic<int32> DecomposedSourceMeshCount = 0;
+	std::atomic<int32> DecomposedMeshesCount = 0;
+	std::atomic<int32> SourceInstancesCount = 0;
+	std::atomic<int32> NewInstancesCount = 0;
+	std::atomic<int32> SkippedDecompositionCount = 0;
+	int32 AddedUniqueTrisCount = 0;
+	int32 InstancedTrisCount = 0;
+
+	// these locks are used below to control access
+	FCriticalSection ToBuildQueueLock;
+	FCriticalSection InternalListsLock;
+
+	// The loop below will emit new IMeshSpatialWrapper's that need to have Build() called.
+	// Since larger meshes take longer, it is a better strategy to collect up these jobs and
+	// then call Build() in decreasing-size order
+	struct FBuildJob
+	{
+		int TriangleCount;
+		IMeshSpatialWrapper* BuildWrapper;
+	};
+	TArray<FBuildJob> PendingBuildJobs;
+	FCriticalSection PendingBuildsLock;
+	// this lambda is used below to append to the PendingBuildJobs list above
+	auto AddBuildJob = [&PendingBuildJobs, &PendingBuildsLock](IMeshSpatialWrapper* ToBuild, int TriangleCount)
+	{
+		PendingBuildsLock.Lock();
+		PendingBuildJobs.Add(FBuildJob{ TriangleCount, ToBuild });
+		PendingBuildsLock.Unlock();
+	};
 
-	// use this to disable build of unnecessary AABBTree/etc
-	FMeshSceneAdapterBuildOptions TempBuildOptions = BuildOptions;
-	TempBuildOptions.bBuildSpatialDataStructures = false;
 
-	// parallel build of all the spatial data structures
+	// Parallel-process all the ToBuild spatial wrappers. If the mesh is closed and all the pieces are good,
+	// this will just emit a Build job. Otherwise it will pull the mesh apart into pieces, move all the closed non-thin
+	// pieces into a new instance to be referenced by the original FActorChildMesh, and then make new meshes/wrappers
+	// for anything that needs geometric changes (eg to bake in scale, thicken mesh, etc), and in those cases, generate
+	// new instances as FActorAdapter/FActorChildMesh's. And emit BuildJob's for those different spatial wrappers.
 	ParallelFor(ToBuild.Num(), [&](int32 i)
 	{
 		TRACE_CPUPROFILER_EVENT_SCOPE(MeshScene_Build_ProcessMesh);
+		
+		// ParallelFor will not respect the sorting by triangle-count we did above, so we have to treat the list as a queue and pop from the back
+		ToBuildQueueLock.Lock();
+		check(ToBuild.Num() > 0);
+		FSpatialWrapperInfo* WrapperInfo = ToBuild.Pop(false);
+		ToBuildQueueLock.Unlock();
 
-		FSpatialWrapperInfo* WrapperInfo = ToBuild[i];
+		// get name for debugging purposes
+		FString AssetName = WrapperInfo->SourceContainer.GetStaticMesh() ? WrapperInfo->SourceContainer.GetStaticMesh()->GetName() : TEXT("Unknown");
 
 		// convert this mesh to a dynamicmesh for processing
 		FDynamicMesh3 LocalMesh;
-		WrapperInfo->SpatialWrapper->Build(TempBuildOptions);
 		{
 			TRACE_CPUPROFILER_EVENT_SCOPE(MeshScene_Build_ProcessMesh_1Copy);
 			WrapperInfo->SpatialWrapper->AppendMesh(LocalMesh, FTransformSequence3d());
@@ -627,15 +847,20 @@ void FMeshSceneAdapter::Build_FullDecompose(const FMeshSceneAdapterBuildOptions&
 
 		// should we try to weld here??
 
+		// find separate submeshes of the mesh
 		FMeshConnectedComponents Components(&LocalMesh);
 		{
 			TRACE_CPUPROFILER_EVENT_SCOPE(MeshScene_Build_ProcessMesh_2Components);
 			Components.FindConnectedTriangles();
 		}
 		int32 NumComponents = Components.Num();
-		TArray<bool> IsClosed;
+
+		// for each submesh/component, determine if it is closed, and if it is 'thin'
+		TArray<bool> IsClosed, IsThin;
 		IsClosed.Init(false, NumComponents);
 		std::atomic<int32> NumNonClosed = 0;
+		IsThin.Init(false, NumComponents);
+		std::atomic<int32> NumThin = 0;
 		{
 			TRACE_CPUPROFILER_EVENT_SCOPE(MeshScene_Build_ProcessMesh_3Closed);
 			ParallelFor(NumComponents, [&](int32 ci)
@@ -647,102 +872,134 @@ void FMeshSceneAdapter::Build_FullDecompose(const FMeshSceneAdapterBuildOptions&
 				{
 					NumNonClosed++;
 				}
+
+				FFrame3d TempPlane;
+				IsThin[ci] = IsThinPlanarSubMesh<FDynamicMesh3>(LocalMesh, Triangles, BuildOptions.DesiredMinThickness, TempPlane);
+				if (IsThin[ci])
+				{
+					NumThin++;
+				}
 			});
 		}
 
-		// if we have no open meshes, we can just use the SpatialWrapper we already have, but we have to
-		// rebuild it because we did not do a full build above
-		if (NumNonClosed == 0)
+		// if we have no open meshes and no thin meshes, we can just use the SpatialWrapper we already have, 
+		// but we have to rebuild it because we did not do a full build above
+		// note: possibly some other cases where we can do this, if the StaticMesh wrapper supported unsigned/offset mode
+		if (NumNonClosed == 0 && NumThin == 0)
 		{
-			WrapperInfo->SpatialWrapper->Build(BuildOptions);
+			AddBuildJob(WrapperInfo->SpatialWrapper.Get(), LocalMesh.TriangleCount());
 			return;
 		}
 
-		// NOTE: because of thin-planar checks below, we might not actually decompose this mesh!!
-		DecomposedSourceMeshCount++;
-
-		// list of per-instance transforms
+		// construct list of per-instance transforms that reference this mesh
 		TArray<FActorChildMesh*> MeshesToDecompose = WrapperInfo->ParentMeshes;
 		TArray<FTransformSequence3d> ParentTransforms;
 		for (FActorChildMesh* MeshInstance : MeshesToDecompose)
 		{
 			ParentTransforms.Add(MeshInstance->WorldTransform);
+			SourceInstancesCount++;
 		}
 
-		bool bMakeClosedAssembly = (NumNonClosed != NumComponents);
+		// Decompose the per-instance transforms into subsets that share the same total scaling ("unique scale").
+		// If we apply these different scales to copies of the mesh, we can generate new instances for the copies,
+		// which can avoid uniquing a lot of geometry
+		TArray<TArray<int32>> UniqueScaleTransformSets;
+		ConstructUniqueScalesMapping(ParentTransforms, UniqueScaleTransformSets);
+		int32 NumUniqueScales = UniqueScaleTransformSets.Num();
+
+		// Accumulate submesh/components that do *not* need further processing here, that accumulated mesh
+		// (if non-empty) can be shared among all the original FActorChildMesh instances
 		FDynamicMesh3 LocalSpaceParts;
 		FDynamicMeshEditor LocalSpaceAccumulator(&LocalSpaceParts);
-		FDynamicMesh3 WorldSpaceParts;
-		FDynamicMeshEditor WorldSpaceAccumulator(&WorldSpaceParts);
-
-		int32 LocalCount = 0, DecomposedCount = 0;
 
-		// Accumulate closed parts in a local-space mesh that can stay instanced, and open parts
+		// a new copy of one of the submeshes that has been scaled/processed such that it can only be
+		// represented with some of the original instance transforms (NewTransforms). 
+		struct FInstancedSubmesh
+		{
+			TSharedPtr<FDynamicMesh3> Mesh;
+			TArray<FTransformSequence3d> NewTransforms;
+			double ComputedThickness = 0;
+		};
+		TArray<FInstancedSubmesh> NewSubmeshes;
+
+		// Split all the submeshes/components into the LocalSpaceParts mesh (for closed and non-thin) and
+		// a set of new FInstancedSubmesh's
 		{
+			FMeshIndexMappings Mappings;		// these are re-used between calls
+			FDynamicMeshEditResult EditResult;
+
 			TRACE_CPUPROFILER_EVENT_SCOPE(MeshScene_Build_ProcessMesh_4Accumulate);
 			for (int32 ci = 0; ci < NumComponents; ++ci)
 			{
 				const TArray<int32>& Triangles = Components[ci].Indices;
 
-				bool bKeepLocal = true;
-				FFrame3d ThinMeshPlane;
-				if (IsClosed[ci] == false)
-				{
-					bool IsThinMesh = IsThinPlanarSubMesh<FDynamicMesh3>(LocalMesh, Triangles, BuildOptions.DesiredMinThickness, ThinMeshPlane);
-					if (IsThinMesh && (Triangles.Num() < 10000 || WrapperInfo->ParentMeshes.Num() == 1))
-					{
-						bKeepLocal = false;
-					}
-				}
-
-				FMeshIndexMappings Mappings;
-				FDynamicMeshEditResult EditResult;
-				if (bKeepLocal)
+				// We will make unscaled copies of a mesh if (1) it is "thin" and (2) it has a moderate number of triangles *or* a single usage
+				// TODO: should we always unique a mesh with a single usage? We can just make it unsigned...
+				bool bIsClosed = IsClosed[ci];
+				if (IsThin[ci] == false || (Triangles.Num() > 10000 && WrapperInfo->ParentMeshes.Num() > 1))
 				{
+					Mappings.Reset(); EditResult.Reset();
 					LocalSpaceAccumulator.AppendTriangles(&LocalMesh, Triangles, Mappings, EditResult, false);
-					LocalCount++;
+					continue;
 				}
-				else
+
+				// if we go this far, we need to unique this mesh once for each "unique scale", and then
+				// make a new set of instance transforms for it
+				for (int32 k = 0; k < NumUniqueScales; ++k)
 				{
-					DecomposedCount++;
-					for (const FTransformSequence3d& Transform : ParentTransforms)
+					FInstancedSubmesh NewSubmesh;
+					const TArray<int32>& InstanceIndices = UniqueScaleTransformSets[k];
+					// make unique copy of submesh
+					NewSubmesh.Mesh = MakeShared<FDynamicMesh3>();
+					FDynamicMeshEditor Editor(NewSubmesh.Mesh.Get());
+					Mappings.Reset(); EditResult.Reset();
+					Editor.AppendTriangles(&LocalMesh, Triangles, Mappings, EditResult, false);
+					// bake in the scaling
+					FVector3d Scale = ParentTransforms[InstanceIndices[0]].GetAccumulatedScale();
+					for (int32 vid : EditResult.NewVertices)
 					{
-						DecomposedMeshesCount++;
-						Mappings.Reset();
-						EditResult.Reset();
-						WorldSpaceAccumulator.AppendTriangles(&LocalMesh, Triangles, Mappings, EditResult, false);
-						for (int32 vid : EditResult.NewVertices)
-						{
-							FVector3d LocalPos = WorldSpaceParts.GetVertex(vid);
-							WorldSpaceParts.SetVertex(vid, Transform.TransformPosition(LocalPos));
-						}
+						FVector3d LocalPos = NewSubmesh.Mesh->GetVertex(vid);
+						NewSubmesh.Mesh->SetVertex(vid, LocalPos * Scale);
+					}
+
+					// Recompute thickness of scaled mesh and store it. Note that we might fail to be
+					// considered "thin" anymore, in that case we will fall back to using winding number
+					// for this mesh (So, it was a waste to do this separation, but messy to turn back now)
+					FFrame3d TempPlane; double Thickness;
+					if (bIsClosed == false)
+					{
+						NewSubmesh.ComputedThickness = 0;
+					}
+					else if (bIsClosed && IsThinPlanarSubMesh<FDynamicMesh3>(*NewSubmesh.Mesh, BuildOptions.DesiredMinThickness, TempPlane, Thickness))
+					{
+						NewSubmesh.ComputedThickness = Thickness;
 					}
+					else
+					{
+						NewSubmesh.ComputedThickness = BuildOptions.DesiredMinThickness;
+					}
+
+					// make new set of instances
+					for (int32 j : InstanceIndices)
+					{
+						FTransformSequence3d InstanceTransform = ParentTransforms[j];
+						InstanceTransform.ClearScales();
+						NewSubmesh.NewTransforms.Add(InstanceTransform);
+					}
+					NewSubmeshes.Add(MoveTemp(NewSubmesh));
 				}
 			}
 		}
+		// At this point we have processed all the Submeshes/Components. Now we generate new MeshSpatialWrapper's
+		// and any necessary new FActorAdapter's/FActorChildMesh's
 
-		// make new mesh containers
-
-		// we re-use the existing wrapper for the closed meshes, if there are any, otherwise
-		// we need to disable the parent meshes that refer to it
-		if (bMakeClosedAssembly)
+		// First handle the LocalSpaceParts mesh, which can still be shared between the original FActorChildMesh instances
+		if (LocalSpaceParts.TriangleCount() > 0)
 		{
 			TUniquePtr<FDynamicMeshSpatialWrapper> LocalSpaceMeshWrapper = MakeUnique<FDynamicMeshSpatialWrapper>();
-			
 			LocalSpaceMeshWrapper->Mesh = MoveTemp(LocalSpaceParts);
-			FDynamicMeshSpatialWrapper* BuildWrapper = LocalSpaceMeshWrapper.Get();
-			TFuture<void> PendingClosedAssemblyBuild = Async(EAsyncExecution::ThreadPool, [&BuildOptions, BuildWrapper]()
-			{
-				TRACE_CPUPROFILER_EVENT_SCOPE(MeshScene_Build_ProcessMesh_6SolidBuild);
-				BuildWrapper->Build(BuildOptions);
-			});
-
-			PendingFuturesLock.Lock();
-			PendingClosedAssemblyBuilds.Add(MoveTemp(PendingClosedAssemblyBuild));
-			PendingFuturesLock.Unlock();
-
+			AddBuildJob(LocalSpaceMeshWrapper.Get(), LocalSpaceMeshWrapper->Mesh.TriangleCount());
 			WrapperInfo->SpatialWrapper = MoveTemp(LocalSpaceMeshWrapper);
-
 			for (FActorChildMesh* MeshInstance : MeshesToDecompose)
 			{
 				MeshInstance->MeshSpatial = WrapperInfo->SpatialWrapper.Get();
@@ -757,60 +1014,103 @@ void FMeshSceneAdapter::Build_FullDecompose(const FMeshSceneAdapterBuildOptions&
 			}
 		}
 
-		// currently happens if we fail to find open parts to expand/etc
-		if (WorldSpaceParts.TriangleCount() == 0)
+		// Exit if we don't have any more work to do. This happens if we ended up skipping all the possible decompositions
+		// Note: in this case we could just re-use the existing actor and skip the LocalSpaceParts mesh entirely?
+		if (NewSubmeshes.Num() == 0)
 		{
+			SkippedDecompositionCount++;
 			return;
 		}
 
-		// world-space meshes are added to new fake adapter
-		TUniquePtr<FActorAdapter> Adapter = MakeUnique<FActorAdapter>();
-		Adapter->SourceActor = nullptr;
+		// definitely decomposing this mesh
+		DecomposedSourceMeshCount++;
+
+		// Now we create a new FActorAdapter for each new InstancedSubmesh, and then an FActorChildMesh
+		// for each Instance (ie rotate/translate transform of that instance). This is somewhat arbitrary,
+		// eg it could all be done in a single Actor, or split up further. At evaluation time we will have
+		// pulled these back out of the Actor so it doesn't really matter.
+		for (FInstancedSubmesh& Submesh : NewSubmeshes)
+		{
+			TUniquePtr<FActorAdapter> NewActor = MakeUnique<FActorAdapter>();
+			NewActor->SourceActor = nullptr;		// not a "real" actor
+
+			// make new spatialwrapper for this instanced mesh
+			TUniquePtr<FDynamicMeshSpatialWrapper> NewInstancedMesh = MakeUnique<FDynamicMeshSpatialWrapper>();
+			int32 TriangleCount = Submesh.Mesh->TriangleCount();
+			NewInstancedMesh->Mesh = MoveTemp(*Submesh.Mesh);
+			NewInstancedMesh->bHasBakedTransform = false;
+			NewInstancedMesh->bHasBakedScale = true;
+			// queue up build job
+			AddBuildJob(NewInstancedMesh.Get(), TriangleCount);
+			// if mesh is too thin, configure the extra shell offset based on 'missing' thickness
+			if (Submesh.ComputedThickness < BuildOptions.DesiredMinThickness)
+			{
+				NewInstancedMesh->bUseDistanceShellForWinding = true;
+				NewInstancedMesh->WindingShellThickness = 0.5 * (BuildOptions.DesiredMinThickness - Submesh.ComputedThickness);
+			}
 
-		TUniquePtr<FActorChildMesh> ChildMesh = MakeUnique<FActorChildMesh>();
-		ChildMesh->SourceComponent = nullptr;
-		ChildMesh->ComponentType = EActorMeshComponentType::InternallyGeneratedComponent;
-		ChildMesh->ComponentIndex = 0;
-		ChildMesh->bIsNonUniformScaled = false;
+			TSharedPtr<FSpatialWrapperInfo> NewWrapperInfo = MakeShared<FSpatialWrapperInfo>();
+			NewWrapperInfo->SpatialWrapper = MoveTemp(NewInstancedMesh);
 
+			// add to internal lists
+			InternalListsLock.Lock();
+			AddedUniqueTrisCount += TriangleCount;
+			InstancedTrisCount += TriangleCount * Submesh.NewTransforms.Num();
+			void* UseKey = (void*)NewWrapperInfo->SpatialWrapper.Get();
+			SpatialAdapters.Add(UseKey, NewWrapperInfo);
+			InternalListsLock.Unlock();
 
-		TUniquePtr<FDynamicMeshSpatialWrapper> NewWorldSpaceSpatial = MakeUnique<FDynamicMeshSpatialWrapper>();
-		NewWorldSpaceSpatial->Mesh = MoveTemp(WorldSpaceParts);
-		NewWorldSpaceSpatial->bHasBakedTransform = true;
-		{
-			TRACE_CPUPROFILER_EVENT_SCOPE(MeshScene_Build_ProcessMesh_8WorldSpaceBuild);
-			NewWorldSpaceSpatial->Build(BuildOptions);
+			// create the new transform instances
+			for (FTransformSequence3d InstanceTransform : Submesh.NewTransforms)
+			{
+				TUniquePtr<FActorChildMesh> ChildMesh = MakeUnique<FActorChildMesh>();
+				ChildMesh->SourceComponent = nullptr;
+				ChildMesh->ComponentType = EActorMeshComponentType::InternallyGeneratedComponent;
+				ChildMesh->ComponentIndex = 0;
+				ChildMesh->WorldTransform = InstanceTransform;
+				ChildMesh->bIsNonUniformScaled = false;
+
+				NewWrapperInfo->ParentMeshes.Add(ChildMesh.Get());
+				ChildMesh->MeshSpatial = NewWrapperInfo->SpatialWrapper.Get();
+
+				NewActor->ChildMeshes.Add(MoveTemp(ChildMesh));
+				NewInstancesCount++;
+			}
+
+			// add actor our actor set
+			InternalListsLock.Lock();
+			SceneActors.Add(MoveTemp(NewActor));
+			InternalListsLock.Unlock();
 		}
-		NewWorldSpaceSpatial->bUseDistanceShellForWinding = true;
-		NewWorldSpaceSpatial->WindingShellThickness = 0.5 * BuildOptions.DesiredMinThickness;
+	});		// end outer ParallelFor over ToBuild meshes
 
-		TSharedPtr<FSpatialWrapperInfo> NewWrapperInfo = MakeShared<FSpatialWrapperInfo>();
-		NewWrapperInfo->SpatialWrapper = MoveTemp(NewWorldSpaceSpatial);
-		NewWrapperInfo->ParentMeshes.Add(ChildMesh.Get());
-		ChildMesh->MeshSpatial = NewWrapperInfo->SpatialWrapper.Get();
+	check(ToBuild.Num() == 0);
 
-		void* UseKey = ChildMesh.Get();
-		Adapter->ChildMeshes.Add(MoveTemp(ChildMesh));
+	// Now all that is left is to actually Build() all the different spatial wrappers that exist at this point
 
-		ListsLock.Lock();
-		SpatialAdapters.Add(UseKey, NewWrapperInfo);
-		SceneActors.Add(MoveTemp(Adapter));
-		AddedTrisCount += NewWrapperInfo->SpatialWrapper->GetTriangleCount();
-		ListsLock.Unlock();
+	// sort by increasing triangle size. 
+	PendingBuildJobs.Sort([](const FBuildJob& A, const FBuildJob& B)
+	{
+		return A.TriangleCount < B.TriangleCount;
 	});
-
-	// make sure these finishes
-	for (auto& PendingClosedAssemblyBuild : PendingClosedAssemblyBuilds)
+	ParallelFor(PendingBuildJobs.Num(), [&](int32 i)
 	{
-		PendingClosedAssemblyBuild.Wait();
-	}	
+		// ParallelFor will not respect our sort order if we just use the index directly (because it splits into chunks internally), so
+		// we have to treat the list like a queue to get it to be processed in our desired order
+		ToBuildQueueLock.Lock();
+		check(PendingBuildJobs.Num() > 0);
+		FBuildJob BuildJob = PendingBuildJobs.Pop(false);
+		ToBuildQueueLock.Unlock();
+		BuildJob.BuildWrapper->Build(BuildOptions);
+	});
+	check(PendingBuildJobs.Num() == 0);
 
 	// currently true with the methods used above?
 	bSceneIsAllSolids = true;
 
 	if (BuildOptions.bPrintDebugMessages)
 	{
-		UE_LOG(LogTemp, Warning, TEXT("[FMeshSceneAdapter] decomposed %d source meshes into %d unique meshes containing %d triangles"), DecomposedSourceMeshCount.load(), DecomposedMeshesCount.load(), AddedTrisCount)
+		UE_LOG(LogTemp, Warning, TEXT("[FMeshSceneAdapter] decomposed %d source meshes used in %d instances (of %d total source meshes), into %d new instances containing %ld unique triangles (%ld total instanced). Skipped %d decompositions."), DecomposedSourceMeshCount.load(), SourceInstancesCount.load(), NumInitialSources, NewInstancesCount.load(), AddedUniqueTrisCount, InstancedTrisCount, SkippedDecompositionCount.load())
 	}
 }
 
@@ -824,7 +1124,7 @@ void FMeshSceneAdapter::GetGeometryStatistics(FStatistics& StatsOut)
 	for (TPair<void*, TSharedPtr<FSpatialWrapperInfo>> Pair : SpatialAdapters)
 	{
 		StatsOut.UniqueMeshCount++;
-		StatsOut.UniqueMeshTriangleCount += Pair.Value->SpatialWrapper->GetTriangleCount();
+		StatsOut.UniqueMeshTriangleCount += (int64)Pair.Value->SpatialWrapper->GetTriangleCount();
 	}
 
 	StatsOut.InstanceMeshCount = 0;
@@ -836,7 +1136,7 @@ void FMeshSceneAdapter::GetGeometryStatistics(FStatistics& StatsOut)
 			StatsOut.InstanceMeshCount++;
 			if (ChildMesh->MeshSpatial != nullptr)
 			{
-				StatsOut.InstanceMeshTriangleCount += ChildMesh->MeshSpatial->GetTriangleCount();
+				StatsOut.InstanceMeshTriangleCount += (int64)ChildMesh->MeshSpatial->GetTriangleCount();
 			}
 		}
 	}
@@ -945,11 +1245,7 @@ void FMeshSceneAdapter::BuildSpatialEvaluationCache()
 				Cache.Actor = Actor.Get();
 				Cache.ChildMesh = ChildMesh.Get();
 				Cache.Spatial = ChildMesh->MeshSpatial;
-				Cache.Bounds = ChildMesh->MeshSpatial->GetWorldBounds(
-					[&](const FVector3d& P) { return ChildMesh->WorldTransform.TransformPosition(P); });
 				SortedSpatials.Add(Cache);
-
-
 			}
 		}
 	}
@@ -961,25 +1257,32 @@ void FMeshSceneAdapter::BuildSpatialEvaluationCache()
 	});
 
 	int32 NumSpatials = SortedSpatials.Num();
-	ParallelFor(SortedSpatials.Num(), [&](int32 k)
-	{
-		SortedSpatials[k].Bounds = SortedSpatials[k].Spatial->GetWorldBounds(
-			[&](const FVector3d& P) { return SortedSpatials[k].ChildMesh->WorldTransform.TransformPosition(P); });
-	});
-
 	CachedWorldBounds = FAxisAlignedBox3d::Empty();
-	for (const FSpatialCacheInfo& Cache : SortedSpatials )
 	{
-		CachedWorldBounds.Contain(Cache.Bounds);
+		TRACE_CPUPROFILER_EVENT_SCOPE(MeshScene_BuildSpatialEvaluationCache_Bounds);
+		ParallelFor(SortedSpatials.Num(), [&](int32 k)
+		{
+			SortedSpatials[k].Bounds = SortedSpatials[k].Spatial->GetWorldBounds(
+				[&](const FVector3d& P) { return SortedSpatials[k].ChildMesh->WorldTransform.TransformPosition(P); });
+		});
+
+		for (const FSpatialCacheInfo& Cache : SortedSpatials)
+		{
+			CachedWorldBounds.Contain(Cache.Bounds);
+		}
 	}
 
+
 	// build an octree of the mesh objects
 	Octree = MakeShared<FSparseDynamicOctree3>();
 	Octree->RootDimension = CachedWorldBounds.MaxDim() / 4.0;
 	Octree->SetMaxTreeDepth(5);
-	for (int32 k = 0; k < NumSpatials; ++k)
 	{
-		Octree->InsertObject(k, SortedSpatials[k].Bounds);
+		TRACE_CPUPROFILER_EVENT_SCOPE(MeshScene_BuildSpatialEvaluationCache_OctreeInserts);
+		for (int32 k = 0; k < NumSpatials; ++k)
+		{
+			Octree->InsertObject(k, SortedSpatials[k].Bounds);
+		}
 	}
 
 	bHaveSpatialEvaluationCache = true;
@@ -1113,7 +1416,7 @@ void FMeshSceneAdapter::GenerateBaseClosingMesh(double BaseHeight, double Extrud
 	TUniquePtr<FDynamicMeshSpatialWrapper> DynamicMeshWrapper = MakeUnique<FDynamicMeshSpatialWrapper>();
 	//DynamicMeshWrapper->SourceContainer = NewWrapperInfo->SourceContainer;
 	DynamicMeshWrapper->Mesh = MoveTemp(BasePolygonMesh);
-	DynamicMeshWrapper->bHasBakedTransform = true;
+	DynamicMeshWrapper->bHasBakedTransform = DynamicMeshWrapper->bHasBakedScale = true;
 	FMeshSceneAdapterBuildOptions UseBuildOptions;
 	DynamicMeshWrapper->Build(UseBuildOptions);
 
diff --git a/Engine/Plugins/Experimental/MeshModelingToolset/Source/ModelingComponents/Public/Scene/MeshSceneAdapter.h b/Engine/Plugins/Experimental/MeshModelingToolset/Source/ModelingComponents/Public/Scene/MeshSceneAdapter.h
--- a/Engine/Plugins/Experimental/MeshModelingToolset/Source/ModelingComponents/Public/Scene/MeshSceneAdapter.h
+++ b/Engine/Plugins/Experimental/MeshModelingToolset/Source/ModelingComponents/Public/Scene/MeshSceneAdapter.h
@@ -219,11 +219,11 @@ public:
 	 */
 	struct FStatistics
 	{
-		int32 UniqueMeshCount = 0;
-		int32 UniqueMeshTriangleCount = 0;
+		int64 UniqueMeshCount = 0;
+		int64 UniqueMeshTriangleCount = 0;
 
-		int32 InstanceMeshCount = 0;
-		int32 InstanceMeshTriangleCount = 0;
+		int64 InstanceMeshCount = 0;
+		int64 InstanceMeshTriangleCount = 0;
 	};
 
 	/** @return bounding box for the Actor set */
diff --git a/Engine/Source/Runtime/GeometryCore/Public/TransformSequence.h b/Engine/Source/Runtime/GeometryCore/Public/TransformSequence.h
--- a/Engine/Source/Runtime/GeometryCore/Public/TransformSequence.h
+++ b/Engine/Source/Runtime/GeometryCore/Public/TransformSequence.h
@@ -64,6 +64,30 @@ public:
 		return false;
 	}
 
+	/**
+	 * @return Cumulative scale across Transforms.
+	 */
+	FVector3<RealType> GetAccumulatedScale() const
+	{
+		FVector3<RealType> FinalScale = FVector3<RealType>::One();
+		for (const TTransform3<RealType>& Transform : Transforms)
+		{
+			FinalScale = FinalScale * Transform.GetScale();
+		}
+		return FinalScale;
+	}
+
+	/**
+	 * Set scales of all transforms to (1,1,1)
+	 */
+	void ClearScales()
+	{
+		for (TTransform3<RealType>& Transform : Transforms)
+		{
+			Transform.SetScale(FVector3<RealType>::One());
+		}
+	}
+
 	/**
 	 * @return point P with transform sequence applied
 	 */
@@ -114,7 +138,24 @@ public:
 		return Normal;
 	}
 
-
+	/**
+	 * @return true if each Transform in this sequence is equivalent to each transform of another sequence under the given test
+	 */
+	template<typename TransformsEquivalentFunc>
+	bool IsEquivalent(const TTransformSequence3<RealType>& OtherSeq, TransformsEquivalentFunc TransformsTest) const
+	{
+		int32 N = Transforms.Num();
+		if (N == OtherSeq.Transforms.Num())
+		{
+			bool bAllTransformsEqual = true;
+			for (int32 k = 0; k < N && bAllTransformsEqual; ++k)
+			{
+				bAllTransformsEqual = bAllTransformsEqual && TransformsTest(Transforms[k], OtherSeq.Transforms[k]);
+			}
+			return bAllTransformsEqual;
+		}
+		return false;
+	}
 };
 
 typedef TTransformSequence3<float> FTransformSequence3f;
