ModelingTools: Add basic Transfer Mesh Tool that can copy a mesh from one object to another. Add UE::ToolTarget::CommitMeshDescriptionUpdate() helper function.
#rnx

diff --git a/Engine/Plugins/Experimental/MeshModelingToolset/Source/MeshModelingTools/Private/TransferMeshTool.cpp b/Engine/Plugins/Experimental/MeshModelingToolset/Source/MeshModelingTools/Private/TransferMeshTool.cpp
new file mode 100644
--- /dev/null
+++ b/Engine/Plugins/Experimental/MeshModelingToolset/Source/MeshModelingTools/Private/TransferMeshTool.cpp
@@ -0,0 +1,102 @@
+// Copyright Epic Games, Inc. All Rights Reserved.
+
+#include "TransferMeshTool.h"
+#include "ComponentSourceInterfaces.h"
+#include "InteractiveToolManager.h"
+#include "ToolTargetManager.h"
+#include "ToolBuilderUtil.h"
+#include "ToolSetupUtil.h"
+#include "ModelingToolTargetUtil.h"
+
+#include "ExplicitUseGeometryMathTypes.h"		// using UE::Geometry::(math types)
+using namespace UE::Geometry;
+
+#define LOCTEXT_NAMESPACE "UTransferMeshTool"
+
+/*
+ * ToolBuilder
+ */
+const FToolTargetTypeRequirements& UTransferMeshToolBuilder::GetTargetRequirements() const
+{
+	static FToolTargetTypeRequirements TypeRequirements({
+		UMaterialProvider::StaticClass(),
+		UMeshDescriptionCommitter::StaticClass(),
+		UMeshDescriptionProvider::StaticClass(),
+		UPrimitiveComponentBackedTarget::StaticClass()
+		});
+	return TypeRequirements;
+}
+
+bool UTransferMeshToolBuilder::CanBuildTool(const FToolBuilderState& SceneState) const
+{
+	return SceneState.TargetManager->CountSelectedAndTargetable(SceneState, GetTargetRequirements()) == 2;
+}
+
+UInteractiveTool* UTransferMeshToolBuilder::BuildTool(const FToolBuilderState& SceneState) const
+{
+	UTransferMeshTool* NewTool = NewObject<UTransferMeshTool>(SceneState.ToolManager);
+	TArray<TObjectPtr<UToolTarget>> Targets = SceneState.TargetManager->BuildAllSelectedTargetable(SceneState, GetTargetRequirements());
+	NewTool->SetTargets(MoveTemp(Targets));
+	NewTool->SetWorld(SceneState.World);
+	return NewTool;
+}
+
+/*
+ * Tool
+ */
+
+UTransferMeshTool::UTransferMeshTool(const FObjectInitializer&)
+{
+}
+
+void UTransferMeshTool::SetWorld(UWorld* World)
+{
+	this->TargetWorld = World;
+}
+
+void UTransferMeshTool::Setup()
+{
+	UInteractiveTool::Setup();
+
+	BasicProperties = NewObject<UTransferMeshToolProperties>(this);
+	BasicProperties->RestoreProperties(this);
+	AddToolPropertySource(BasicProperties);
+
+	SetToolDisplayName(LOCTEXT("ToolName", "Transfer"));
+	GetToolManager()->DisplayMessage(
+		LOCTEXT("OnStartTool", "Copy mesh from Source object to Target object"),
+		EToolMessageLevel::UserNotification);
+}
+
+bool UTransferMeshTool::CanAccept() const
+{
+	return Super::CanAccept();
+}
+
+void UTransferMeshTool::Shutdown(EToolShutdownType ShutdownType)
+{
+	BasicProperties->SaveProperties(this);
+
+	if (ShutdownType == EToolShutdownType::Accept)
+	{
+		GetToolManager()->BeginUndoTransaction(LOCTEXT("TransferMeshToolTransactionName", "Transfer Mesh"));
+
+		const FMeshDescription* SourceMesh = UE::ToolTarget::GetMeshDescription(Targets[0]);
+
+		FComponentMaterialSet Materials = UE::ToolTarget::GetMaterialSet(Targets[0]);
+		const FComponentMaterialSet* TransferMaterials = (BasicProperties->bTransferMaterials) ? &Materials : nullptr;
+
+		if (ensure(SourceMesh))
+		{
+			UE::ToolTarget::CommitMeshDescriptionUpdate(Targets[1], SourceMesh, TransferMaterials);
+		}
+
+		GetToolManager()->EndUndoTransaction();
+	}
+}
+
+
+
+
+
+#undef LOCTEXT_NAMESPACE
diff --git a/Engine/Plugins/Experimental/MeshModelingToolset/Source/MeshModelingTools/Public/TransferMeshTool.h b/Engine/Plugins/Experimental/MeshModelingToolset/Source/MeshModelingTools/Public/TransferMeshTool.h
new file mode 100644
--- /dev/null
+++ b/Engine/Plugins/Experimental/MeshModelingToolset/Source/MeshModelingTools/Public/TransferMeshTool.h
@@ -0,0 +1,62 @@
+// Copyright Epic Games, Inc. All Rights Reserved.
+
+#pragma once
+
+#include "CoreMinimal.h"
+#include "MultiSelectionTool.h"
+#include "InteractiveToolBuilder.h"
+#include "TransferMeshTool.generated.h"
+
+/**
+ *
+ */
+UCLASS()
+class MESHMODELINGTOOLS_API UTransferMeshToolBuilder : public UInteractiveToolBuilder
+{
+	GENERATED_BODY()
+
+public:
+	virtual bool CanBuildTool(const FToolBuilderState& SceneState) const override;
+	virtual UInteractiveTool* BuildTool(const FToolBuilderState& SceneState) const override;
+
+protected:
+	virtual const FToolTargetTypeRequirements& GetTargetRequirements() const override;
+};
+
+/**
+ * Standard properties of the Transfer operation
+ */
+UCLASS()
+class MESHMODELINGTOOLS_API UTransferMeshToolProperties : public UInteractiveToolPropertySet
+{
+	GENERATED_BODY()
+public:
+	UPROPERTY(EditAnywhere, Category = Options)
+	bool bTransferMaterials = true;
+};
+
+
+
+UCLASS()
+class MESHMODELINGTOOLS_API UTransferMeshTool : public UMultiSelectionTool
+{
+	GENERATED_BODY()
+
+public:
+	UTransferMeshTool(const FObjectInitializer&);
+
+	virtual void SetWorld(UWorld* World);
+
+	virtual void Setup() override;
+	virtual void Shutdown(EToolShutdownType ShutdownType) override;
+
+	virtual bool HasCancel() const override { return true; }
+	virtual bool HasAccept() const override { return true; }
+	virtual bool CanAccept() const override;
+
+	UPROPERTY()
+	UTransferMeshToolProperties* BasicProperties;
+
+protected:
+	UWorld* TargetWorld;
+};
diff --git a/Engine/Plugins/Experimental/MeshModelingToolset/Source/ModelingComponents/Private/ModelingToolTargetUtil.cpp b/Engine/Plugins/Experimental/MeshModelingToolset/Source/ModelingComponents/Private/ModelingToolTargetUtil.cpp
--- a/Engine/Plugins/Experimental/MeshModelingToolset/Source/ModelingComponents/Private/ModelingToolTargetUtil.cpp
+++ b/Engine/Plugins/Experimental/MeshModelingToolset/Source/ModelingComponents/Private/ModelingToolTargetUtil.cpp
@@ -13,6 +13,7 @@
 #include "Components/PrimitiveComponent.h"
 #include "Components/StaticMeshComponent.h"
 #include "GameFramework/Volume.h"
+#include "SimpleDynamicMeshComponent.h"
 
 #include "MeshDescriptionToDynamicMesh.h"
 #include "DynamicMeshToMeshDescription.h"
@@ -120,10 +121,41 @@ FDynamicMesh3 UE::ToolTarget::GetDynamicMeshCopy(UToolTarget* Target)
 }
 
 
+UE::ToolTarget::EDynamicMeshUpdateResult UE::ToolTarget::CommitMeshDescriptionUpdate(UToolTarget* Target, const FMeshDescription* UpdatedMesh, const FComponentMaterialSet* UpdatedMaterials)
+{
+	IMeshDescriptionCommitter* MeshDescriptionCommitter = Cast<IMeshDescriptionCommitter>(Target);
+	if (!ensure(MeshDescriptionCommitter))
+	{
+		return EDynamicMeshUpdateResult::Failed;
+	}
+
+	if (UpdatedMaterials != nullptr)
+	{
+		IMaterialProvider* MaterialProvider = Cast<IMaterialProvider>(Target);
+		if (MaterialProvider)
+		{
+			MaterialProvider->CommitMaterialSetUpdate(*UpdatedMaterials, true);
+		}
+	}
+
+	EDynamicMeshUpdateResult Result = EDynamicMeshUpdateResult::Failed;
+	MeshDescriptionCommitter->CommitMeshDescription([UpdatedMesh, &Result](const IMeshDescriptionCommitter::FCommitterParams& CommitParams)
+	{
+		*CommitParams.MeshDescriptionOut = *UpdatedMesh;
+		Result = EDynamicMeshUpdateResult::Ok;
+	});
+	return Result;
+}
+
+
 
 UE::ToolTarget::EDynamicMeshUpdateResult UE::ToolTarget::CommitDynamicMeshUVUpdate(UToolTarget* Target, const UE::Geometry::FDynamicMesh3* UpdatedMesh)
 {
 	IMeshDescriptionCommitter* MeshDescriptionCommitter = Cast<IMeshDescriptionCommitter>(Target);
+	if (!ensure(MeshDescriptionCommitter))
+	{
+		return EDynamicMeshUpdateResult::Failed;
+	}
 
 	EDynamicMeshUpdateResult Result = EDynamicMeshUpdateResult::Failed;
 	MeshDescriptionCommitter->CommitMeshDescription([UpdatedMesh, &Result](const IMeshDescriptionCommitter::FCommitterParams& CommitParams)
@@ -162,6 +194,12 @@ bool UE::ToolTarget::ConfigureCreateMeshObjectParams(UToolTarget* SourceTarget,
 			return true;
 		}
 
+		if (Cast<USimpleDynamicMeshComponent>(ComponentTarget->GetOwnerComponent()) != nullptr)
+		{
+			DerivedParamsOut.TypeHint = ECreateObjectTypeHint::DynamicMeshActor;
+			return true;
+		}
+
 		AVolume* VolumeActor = Cast<AVolume>(ComponentTarget->GetOwnerActor());
 		if (VolumeActor != nullptr)
 		{
diff --git a/Engine/Plugins/Experimental/MeshModelingToolset/Source/ModelingComponents/Public/ModelingObjectsCreationAPI.h b/Engine/Plugins/Experimental/MeshModelingToolset/Source/ModelingComponents/Public/ModelingObjectsCreationAPI.h
--- a/Engine/Plugins/Experimental/MeshModelingToolset/Source/ModelingComponents/Public/ModelingObjectsCreationAPI.h
+++ b/Engine/Plugins/Experimental/MeshModelingToolset/Source/ModelingComponents/Public/ModelingObjectsCreationAPI.h
@@ -55,8 +55,9 @@ UENUM(BlueprintType)
 enum class ECreateObjectTypeHint : uint8
 {
 	Undefined = 0,
-	Volume = 1,
-	StaticMesh = 2
+	StaticMesh = 1,
+	Volume = 2,
+	DynamicMeshActor = 3
 };
 
 
diff --git a/Engine/Plugins/Experimental/MeshModelingToolset/Source/ModelingComponents/Public/ModelingToolTargetUtil.h b/Engine/Plugins/Experimental/MeshModelingToolset/Source/ModelingComponents/Public/ModelingToolTargetUtil.h
--- a/Engine/Plugins/Experimental/MeshModelingToolset/Source/ModelingComponents/Public/ModelingToolTargetUtil.h
+++ b/Engine/Plugins/Experimental/MeshModelingToolset/Source/ModelingComponents/Public/ModelingToolTargetUtil.h
@@ -89,6 +89,16 @@ enum class EDynamicMeshUpdateResult
 };
 
 
+/**
+ * Update the Mesh in a ToolTarget based on the provided MeshDescription, and optional material set
+ * @return EDynamicMeshUpdateResult::Ok on success
+ */
+MODELINGCOMPONENTS_API EDynamicMeshUpdateResult CommitMeshDescriptionUpdate(
+	UToolTarget* Target, 
+	const FMeshDescription* UpdatedMesh, 
+	const FComponentMaterialSet* UpdatedMaterials = nullptr);
+
+
 /**
  * Update the UV sets of the ToolTarget's mesh (assuming it has one) based on the provided UpdatedMesh.
  * @todo: support updating a specific UV set/index, rather than all sets
diff --git a/Engine/Plugins/Experimental/ModelingToolsEditorMode/Source/ModelingToolsEditorMode/Private/ModelingToolsEditorMode.cpp b/Engine/Plugins/Experimental/ModelingToolsEditorMode/Source/ModelingToolsEditorMode/Private/ModelingToolsEditorMode.cpp
--- a/Engine/Plugins/Experimental/ModelingToolsEditorMode/Source/ModelingToolsEditorMode/Private/ModelingToolsEditorMode.cpp
+++ b/Engine/Plugins/Experimental/ModelingToolsEditorMode/Source/ModelingToolsEditorMode/Private/ModelingToolsEditorMode.cpp
@@ -10,6 +10,7 @@
 #include "ToolTargets/StaticMeshComponentToolTarget.h"
 #include "ToolTargets/VolumeDynamicMeshToolTarget.h"
 #include "ToolTargets/VolumeMeshDescriptionToolTarget.h"
+//#include "ToolTargets/DynamicMeshComponentToolTarget.h"
 #include "Framework/Commands/UICommandList.h"
 #include "Framework/Application/SlateApplication.h"
 #include "Framework/MultiBox/MultiBoxBuilder.h"
@@ -82,6 +83,7 @@
 #include "LatticeDeformerTool.h"
 #include "SeamSculptTool.h"
 #include "MeshGroupPaintTool.h"
+#include "TransferMeshTool.h"
 
 #include "Physics/PhysicsInspectorTool.h"
 #include "Physics/SetCollisionGeometryTool.h"
@@ -300,6 +302,7 @@ void UModelingToolsEditorMode::Enter()
 	ToolsContext->TargetManager->AddTargetFactory(NewObject<UStaticMeshComponentToolTargetFactory>(ToolsContext->TargetManager));
 	ToolsContext->TargetManager->AddTargetFactory(NewObject<UVolumeDynamicMeshToolTargetFactory>(ToolsContext->TargetManager));
 	ToolsContext->TargetManager->AddTargetFactory(NewObject<UVolumeMeshDescriptionToolTargetFactory>(ToolsContext->TargetManager));
+	//ToolsContext->TargetManager->AddTargetFactory(NewObject<UDynamicMeshComponentToolTargetFactory>(ToolsContext->TargetManager));
 
 	// register stylus event handler
 	StylusStateTracker = MakeUnique<FStylusStateTracker>();
@@ -421,6 +424,7 @@ void UModelingToolsEditorMode::Enter()
 	RegisterTool(ToolManagerCommands.BeginTransformMeshesTool, TEXT("BeginTransformMeshesTool"), NewObject<UTransformMeshesToolBuilder>());
 	RegisterTool(ToolManagerCommands.BeginEditPivotTool, TEXT("BeginEditPivotTool"), NewObject<UEditPivotToolBuilder>());
 	RegisterTool(ToolManagerCommands.BeginAlignObjectsTool, TEXT("BeginAlignObjectsTool"), NewObject<UAlignObjectsToolBuilder>());
+	RegisterTool(ToolManagerCommands.BeginTransferMeshTool, TEXT("BeginTransferMeshTool"), NewObject<UTransferMeshToolBuilder>());
 	RegisterTool(ToolManagerCommands.BeginBakeTransformTool, TEXT("BeginBakeTransformTool"), NewObject<UBakeTransformToolBuilder>());
 	RegisterTool(ToolManagerCommands.BeginTransformUVIslandsTool, TEXT("BeginTransformUVIslandsTool"), NewObject<UEditUVIslandsToolBuilder>());
 	RegisterTool(ToolManagerCommands.BeginLatticeDeformerTool, TEXT("BeginLatticeDeformerTool"), NewObject<ULatticeDeformerToolBuilder>());
diff --git a/Engine/Plugins/Experimental/ModelingToolsEditorMode/Source/ModelingToolsEditorMode/Private/ModelingToolsEditorModeStyle.cpp b/Engine/Plugins/Experimental/ModelingToolsEditorMode/Source/ModelingToolsEditorMode/Private/ModelingToolsEditorModeStyle.cpp
--- a/Engine/Plugins/Experimental/ModelingToolsEditorMode/Source/ModelingToolsEditorMode/Private/ModelingToolsEditorModeStyle.cpp
+++ b/Engine/Plugins/Experimental/ModelingToolsEditorMode/Source/ModelingToolsEditorMode/Private/ModelingToolsEditorModeStyle.cpp
@@ -202,6 +202,8 @@ void FModelingToolsEditorModeStyle::Initialize()
 
 		StyleSet->Set("ModelingToolsManagerCommands.BeginAlignObjectsTool",                  new FSlateImageBrush(StyleSet->RootToCoreContentDir(TEXT("../Editor/Slate/Icons/GeneralTools/Align_40x.png")), Icon20x20));
 		StyleSet->Set("ModelingToolsManagerCommands.BeginAlignObjectsTool.Small",            new FSlateImageBrush(StyleSet->RootToCoreContentDir(TEXT("../Editor/Slate/Icons/GeneralTools/Align_40x.png")), Icon20x20));
+		StyleSet->Set("ModelingToolsManagerCommands.BeginTransferMeshTool",                  new FSlateImageBrush(StyleSet->RootToCoreContentDir(TEXT("../Editor/Slate/Icons/GeneralTools/Next_40x.png")), Icon20x20));
+		StyleSet->Set("ModelingToolsManagerCommands.BeginTransferMeshTool.Small",            new FSlateImageBrush(StyleSet->RootToCoreContentDir(TEXT("../Editor/Slate/Icons/GeneralTools/Next_40x.png")), Icon20x20));
 		StyleSet->Set("ModelingToolsManagerCommands.BeginGlobalUVGenerateTool",              new IMAGE_PLUGIN_BRUSH("Icons/AutoUnwrap_40x",       Icon20x20));      
 		StyleSet->Set("ModelingToolsManagerCommands.BeginGlobalUVGenerateTool.Small",        new IMAGE_PLUGIN_BRUSH("Icons/AutoUnwrap_40x",       Icon20x20));      
 		StyleSet->Set("ModelingToolsManagerCommands.BeginBakeTransformTool",                 new IMAGE_PLUGIN_BRUSH("Icons/BakeXForm_40x",        Icon20x20));      
diff --git a/Engine/Plugins/Experimental/ModelingToolsEditorMode/Source/ModelingToolsEditorMode/Private/ModelingToolsEditorModeToolkit.cpp b/Engine/Plugins/Experimental/ModelingToolsEditorMode/Source/ModelingToolsEditorMode/Private/ModelingToolsEditorModeToolkit.cpp
--- a/Engine/Plugins/Experimental/ModelingToolsEditorMode/Source/ModelingToolsEditorMode/Private/ModelingToolsEditorModeToolkit.cpp
+++ b/Engine/Plugins/Experimental/ModelingToolsEditorMode/Source/ModelingToolsEditorMode/Private/ModelingToolsEditorModeToolkit.cpp
@@ -551,6 +551,7 @@ void FModelingToolsEditorModeToolkit::BuildToolPalette_Experimental(FName Palett
 		ToolbarBuilder.AddToolBarButton(Commands.BeginAlignObjectsTool);
 		ToolbarBuilder.AddToolBarButton(Commands.BeginEditPivotTool);
 		ToolbarBuilder.AddToolBarButton(Commands.BeginBakeTransformTool);
+		ToolbarBuilder.AddToolBarButton(Commands.BeginTransferMeshTool);
 	}
 	else if (PaletteIndex == DeformTabName)
 	{
diff --git a/Engine/Plugins/Experimental/ModelingToolsEditorMode/Source/ModelingToolsEditorMode/Private/ModelingToolsManagerActions.cpp b/Engine/Plugins/Experimental/ModelingToolsEditorMode/Source/ModelingToolsEditorMode/Private/ModelingToolsManagerActions.cpp
--- a/Engine/Plugins/Experimental/ModelingToolsEditorMode/Source/ModelingToolsEditorMode/Private/ModelingToolsManagerActions.cpp
+++ b/Engine/Plugins/Experimental/ModelingToolsEditorMode/Source/ModelingToolsEditorMode/Private/ModelingToolsManagerActions.cpp
@@ -55,6 +55,7 @@ void FModelingToolsManagerCommands::RegisterCommands()
 	UI_COMMAND(BeginCombineMeshesTool, "Append", "Combine Selection into new StaticMesh Asset", EUserInterfaceActionType::ToggleButton, FInputChord());
 	UI_COMMAND(BeginDuplicateMeshesTool, "Dupe", "Duplicate Selection into new StaticMesh Asset", EUserInterfaceActionType::ToggleButton, FInputChord());
 	UI_COMMAND(BeginAlignObjectsTool, "Align", "Align Objects", EUserInterfaceActionType::ToggleButton, FInputChord());
+	UI_COMMAND(BeginTransferMeshTool, "Transfer", "Transfer Mesh", EUserInterfaceActionType::ToggleButton, FInputChord());
 
 	UI_COMMAND(BeginRemeshSculptMeshTool, "DSclpt", "Dynamic Mesh Sculpting", EUserInterfaceActionType::ToggleButton, FInputChord());
 	UI_COMMAND(BeginRemeshMeshTool, "Remesh", "Retriangulate Mesh", EUserInterfaceActionType::ToggleButton, FInputChord());
diff --git a/Engine/Plugins/Experimental/ModelingToolsEditorMode/Source/ModelingToolsEditorMode/Public/ModelingToolsManagerActions.h b/Engine/Plugins/Experimental/ModelingToolsEditorMode/Source/ModelingToolsEditorMode/Public/ModelingToolsManagerActions.h
--- a/Engine/Plugins/Experimental/ModelingToolsEditorMode/Source/ModelingToolsEditorMode/Public/ModelingToolsManagerActions.h
+++ b/Engine/Plugins/Experimental/ModelingToolsEditorMode/Source/ModelingToolsEditorMode/Public/ModelingToolsManagerActions.h
@@ -49,6 +49,7 @@ public:
 	TSharedPtr<FUICommandInfo> BeginCombineMeshesTool;
 	TSharedPtr<FUICommandInfo> BeginDuplicateMeshesTool;
 	TSharedPtr<FUICommandInfo> BeginAlignObjectsTool;
+	TSharedPtr<FUICommandInfo> BeginTransferMeshTool;
 
 	TSharedPtr<FUICommandInfo> BeginRemeshSculptMeshTool;
 	TSharedPtr<FUICommandInfo> BeginRemeshMeshTool;
