MeshModelingTools: Tweaking UMeshOpPreviewWithBackgroundCompute to allow calculations to be cancelled without shutting down the preview entirely, and to allow the op generator to be swapped out. This will be useful when sharing a preview object between tools or tool mechanics.

#rnx

diff --git a/Engine/Plugins/Experimental/MeshModelingToolset/Source/ModelingComponents/Private/MeshOpPreviewHelpers.cpp b/Engine/Plugins/Experimental/MeshModelingToolset/Source/ModelingComponents/Private/MeshOpPreviewHelpers.cpp
--- a/Engine/Plugins/Experimental/MeshModelingToolset/Source/ModelingComponents/Private/MeshOpPreviewHelpers.cpp
+++ b/Engine/Plugins/Experimental/MeshModelingToolset/Source/ModelingComponents/Private/MeshOpPreviewHelpers.cpp
@@ -4,21 +4,42 @@
 
 using namespace UE::Geometry;
 
-void UMeshOpPreviewWithBackgroundCompute::Setup(UWorld* InWorld, IDynamicMeshOperatorFactory* OpGenerator)
+void UMeshOpPreviewWithBackgroundCompute::Setup(UWorld* InWorld)
 {
 	PreviewMesh = NewObject<UPreviewMesh>(this, TEXT("PreviewMesh"));
 	PreviewMesh->CreateInWorld(InWorld, FTransform::Identity);
+	PreviewWorld = InWorld;
+	bResultValid = false;
+	bMeshInitialized = false;
+}
+
+void UMeshOpPreviewWithBackgroundCompute::Setup(UWorld* InWorld, IDynamicMeshOperatorFactory* OpGenerator)
+{
+	Setup(InWorld);
+	BackgroundCompute = MakeUnique<FBackgroundDynamicMeshComputeSource>(OpGenerator);
+}
 
+void UMeshOpPreviewWithBackgroundCompute::ChangeOpFactory(IDynamicMeshOperatorFactory* OpGenerator)
+{
+	CancelCompute();
 	BackgroundCompute = MakeUnique<FBackgroundDynamicMeshComputeSource>(OpGenerator);
 	bResultValid = false;
 	bMeshInitialized = false;
 }
 
-FDynamicMeshOpResult UMeshOpPreviewWithBackgroundCompute::Shutdown()
+void UMeshOpPreviewWithBackgroundCompute::ClearOpFactory()
 {
-	BackgroundCompute->CancelActiveCompute();
+	CancelCompute();
+	BackgroundCompute = nullptr;
+	bResultValid = false;
+	bMeshInitialized = false;
+}
 
 
+FDynamicMeshOpResult UMeshOpPreviewWithBackgroundCompute::Shutdown()
+{
+	CancelCompute();
+
 	FDynamicMeshOpResult Result{};
 	Result.Mesh = PreviewMesh->ExtractPreviewMesh();
 	Result.Transform = UE::Geometry::FTransform3d(PreviewMesh->GetTransform());
@@ -31,9 +52,17 @@ FDynamicMeshOpResult UMeshOpPreviewWithBackgroundCompute::Shutdown()
 }
 
 
+void UMeshOpPreviewWithBackgroundCompute::CancelCompute()
+{
+	if (BackgroundCompute)
+	{
+		BackgroundCompute->CancelActiveCompute();
+	}
+}
+
 void UMeshOpPreviewWithBackgroundCompute::Cancel()
 {
-	BackgroundCompute->CancelActiveCompute();
+	CancelCompute();
 
 	PreviewMesh->SetVisible(false);
 	PreviewMesh->Disconnect();
@@ -46,10 +75,9 @@ void UMeshOpPreviewWithBackgroundCompute::Tick(float DeltaTime)
 	if (BackgroundCompute)
 	{
 		BackgroundCompute->Tick(DeltaTime);
+		UpdateResults();
 	}
 
-	UpdateResults();
-
 	if (!IsUsingWorkingMaterial())
 	{
 		if (OverrideMaterial != nullptr)
@@ -80,7 +108,9 @@ void UMeshOpPreviewWithBackgroundCompute::Tick(float DeltaTime)
 
 void UMeshOpPreviewWithBackgroundCompute::UpdateResults()
 {
-	EBackgroundComputeTaskStatus Status = BackgroundCompute->CheckStatus();
+	EBackgroundComputeTaskStatus Status = BackgroundCompute ? BackgroundCompute->CheckStatus()
+		: EBackgroundComputeTaskStatus::NotComputing;
+
 	if (Status == EBackgroundComputeTaskStatus::NewResultAvailable)
 	{
 		TUniquePtr<FDynamicMeshOperator> MeshOp = BackgroundCompute->ExtractResult();
@@ -106,7 +136,10 @@ void UMeshOpPreviewWithBackgroundCompute::UpdateResults()
 
 void UMeshOpPreviewWithBackgroundCompute::InvalidateResult()
 {
-	BackgroundCompute->NotifyActiveComputeInvalidated();
+	if (BackgroundCompute)
+	{
+		BackgroundCompute->NotifyActiveComputeInvalidated();
+	}
 	bResultValid = false;
 }
 
diff --git a/Engine/Plugins/Experimental/MeshModelingToolset/Source/ModelingComponents/Public/MeshOpPreviewHelpers.h b/Engine/Plugins/Experimental/MeshModelingToolset/Source/ModelingComponents/Public/MeshOpPreviewHelpers.h
--- a/Engine/Plugins/Experimental/MeshModelingToolset/Source/ModelingComponents/Public/MeshOpPreviewHelpers.h
+++ b/Engine/Plugins/Experimental/MeshModelingToolset/Source/ModelingComponents/Public/MeshOpPreviewHelpers.h
@@ -61,13 +61,30 @@ public:
 	 */
 	void Setup(UWorld* InWorld, UE::Geometry::IDynamicMeshOperatorFactory* OpGenerator);
 
+	void Setup(UWorld* InWorld);
+
 	/**
 	 * Terminate any active computation and return the current Preview Mesh/Transform
 	 */
 	FDynamicMeshOpResult Shutdown();
 
 	/**
-	* Terminate any active computation without returning anything
+	 * Stops any running computes and swaps in a different op generator. Does not 
+	 * update the preview mesh or start a new compute.
+	 */
+	void ChangeOpFactory(UE::Geometry::IDynamicMeshOperatorFactory* OpGenerator);
+
+
+	void ClearOpFactory();
+
+	/**
+	 * Cancel the active computation without returning anything. Doesn't destroy the mesh.
+	 */
+	void CancelCompute();
+
+	/**
+	* Terminate any active computation without returning anything. Destroys the preview
+	* mesh.
 	*/
 	void Cancel();
 
@@ -112,6 +129,7 @@ public:
 	 */
 	bool GetCurrentResultCopy(FDynamicMesh3& MeshOut, bool bOnlyIfValid = true);
 
+	virtual UWorld* GetWorld() const override { return PreviewWorld; }
 
 	//
 	// Optional configuration
@@ -188,6 +206,7 @@ public:
 	UPROPERTY()
 	UMaterialInterface* SecondaryMaterial = nullptr;
 
+	UWorld* PreviewWorld;
 protected:
 	// state flag, if true then we have valid result
 	bool bResultValid = false;
