ModelingTools: Fix synchronization issues in BakeTexture when modifying bake targets during a background compute.

#rnx

diff --git a/Engine/Plugins/Experimental/GeometryProcessing/Source/DynamicMesh/Private/Sampling/MeshMapBaker.cpp b/Engine/Plugins/Experimental/GeometryProcessing/Source/DynamicMesh/Private/Sampling/MeshMapBaker.cpp
--- a/Engine/Plugins/Experimental/GeometryProcessing/Source/DynamicMesh/Private/Sampling/MeshMapBaker.cpp
+++ b/Engine/Plugins/Experimental/GeometryProcessing/Source/DynamicMesh/Private/Sampling/MeshMapBaker.cpp
@@ -515,12 +515,22 @@ int32 FMeshMapBaker::AddBaker(TSharedPtr<FMeshMapEvaluator> Sampler)
 	return Bakers.Add(Sampler);
 }
 
+FMeshMapEvaluator* FMeshMapBaker::GetBaker(int32 BakerIdx)
+{
+	return Bakers[BakerIdx].Get();
+}
+
 void FMeshMapBaker::Reset()
 {
 	Bakers.Empty();
 	BakeResults.Empty();
 }
 
+int32 FMeshMapBaker::NumBakers()
+{
+	return Bakers.Num();
+}
+
 const TArrayView<TUniquePtr<TImageBuilder<FVector4f>>> FMeshMapBaker::GetBakeResults(int32 BakerIdx)
 {
 	int32 ResultIdx = BakeOffsets[BakerIdx];
diff --git a/Engine/Plugins/Experimental/GeometryProcessing/Source/DynamicMesh/Public/Sampling/MeshCurvatureMapEvaluator.h b/Engine/Plugins/Experimental/GeometryProcessing/Source/DynamicMesh/Public/Sampling/MeshCurvatureMapEvaluator.h
--- a/Engine/Plugins/Experimental/GeometryProcessing/Source/DynamicMesh/Public/Sampling/MeshCurvatureMapEvaluator.h
+++ b/Engine/Plugins/Experimental/GeometryProcessing/Source/DynamicMesh/Public/Sampling/MeshCurvatureMapEvaluator.h
@@ -50,9 +50,12 @@ public:
 	TSharedPtr<FMeshVertexCurvatureCache> Curvatures;
 
 public:
-	/** Invoked at start of bake to initialize baker. */
+	// Begin FMeshMapEvaluator interface
 	virtual void Setup(const FMeshMapBaker& Baker, FEvaluationContext& Context) override;
 
+	virtual EMeshMapEvaluatorType Type() const override { return EMeshMapEvaluatorType::Curvature; }
+	// End FMeshMapEvaluator interface
+
 	static void EvaluateSample(float*& Out, const FCorrespondenceSample& Sample, void* EvalData);
 
 	static void EvaluateDefault(float*& Out, void* EvalData);
diff --git a/Engine/Plugins/Experimental/GeometryProcessing/Source/DynamicMesh/Public/Sampling/MeshMapBaker.h b/Engine/Plugins/Experimental/GeometryProcessing/Source/DynamicMesh/Public/Sampling/MeshMapBaker.h
--- a/Engine/Plugins/Experimental/GeometryProcessing/Source/DynamicMesh/Public/Sampling/MeshMapBaker.h
+++ b/Engine/Plugins/Experimental/GeometryProcessing/Source/DynamicMesh/Public/Sampling/MeshMapBaker.h
@@ -66,6 +66,12 @@ public:
 	/** Add a baker to be processed. */
 	int32 AddBaker(TSharedPtr<FMeshMapEvaluator, ESPMode::ThreadSafe> Target);
 
+	/** @return the evaluator at the given index. */
+	FMeshMapEvaluator* GetBaker(int32 BakerIdx);
+
+	/** @return the number of bake evaluators on this baker. */
+	int32 NumBakers();
+
 	/** Reset the list of bakers. */
 	void Reset();
 
diff --git a/Engine/Plugins/Experimental/GeometryProcessing/Source/DynamicMesh/Public/Sampling/MeshMapEvaluator.h b/Engine/Plugins/Experimental/GeometryProcessing/Source/DynamicMesh/Public/Sampling/MeshMapEvaluator.h
--- a/Engine/Plugins/Experimental/GeometryProcessing/Source/DynamicMesh/Public/Sampling/MeshMapEvaluator.h
+++ b/Engine/Plugins/Experimental/GeometryProcessing/Source/DynamicMesh/Public/Sampling/MeshMapEvaluator.h
@@ -11,6 +11,16 @@ namespace Geometry
 
 class FMeshMapBaker;
 
+enum class EMeshMapEvaluatorType
+{
+	Normal,
+	Occlusion,
+	Curvature,
+	Property,
+	ResampleImage,
+	MultiResampleImage
+};
+
 class FMeshMapEvaluator
 {
 public:
@@ -113,6 +123,9 @@ public:
 	 */
 	virtual void Setup(const FMeshMapBaker& Baker, FEvaluationContext& Context) = 0;
 
+	/** @return the type of evaluator. */
+	virtual EMeshMapEvaluatorType Type() const = 0;
+
 protected:
 	/**
 	 * Write float data to a float buffer and increment.
diff --git a/Engine/Plugins/Experimental/GeometryProcessing/Source/DynamicMesh/Public/Sampling/MeshNormalMapEvaluator.h b/Engine/Plugins/Experimental/GeometryProcessing/Source/DynamicMesh/Public/Sampling/MeshNormalMapEvaluator.h
--- a/Engine/Plugins/Experimental/GeometryProcessing/Source/DynamicMesh/Public/Sampling/MeshNormalMapEvaluator.h
+++ b/Engine/Plugins/Experimental/GeometryProcessing/Source/DynamicMesh/Public/Sampling/MeshNormalMapEvaluator.h
@@ -13,9 +13,12 @@ namespace Geometry
 class DYNAMICMESH_API FMeshNormalMapEvaluator : public FMeshMapEvaluator
 {
 public:
-	/** Invoked at start of bake to initialize baker. */
+	// Begin FMeshMapEvaluator interface
 	virtual void Setup(const FMeshMapBaker& Baker, FEvaluationContext& Context) override;
 
+	virtual EMeshMapEvaluatorType Type() const override { return EMeshMapEvaluatorType::Normal; }
+	// End FMeshMapEvaluator interface
+
 	static void EvaluateSample(float*& Out, const FCorrespondenceSample& Sample, void* EvalData);
 
 	static void EvaluateDefault(float*& Out, void* EvalData);
diff --git a/Engine/Plugins/Experimental/GeometryProcessing/Source/DynamicMesh/Public/Sampling/MeshOcclusionMapEvaluator.h b/Engine/Plugins/Experimental/GeometryProcessing/Source/DynamicMesh/Public/Sampling/MeshOcclusionMapEvaluator.h
--- a/Engine/Plugins/Experimental/GeometryProcessing/Source/DynamicMesh/Public/Sampling/MeshOcclusionMapEvaluator.h
+++ b/Engine/Plugins/Experimental/GeometryProcessing/Source/DynamicMesh/Public/Sampling/MeshOcclusionMapEvaluator.h
@@ -52,8 +52,11 @@ public:
 	using FOcclusionTuple = TTuple<float, FVector3f>;
 
 public:
-	/** Invoked at start of bake to initialize baker. */
-	virtual void Setup(const FMeshMapBaker & Baker, FEvaluationContext & Context) override;
+	// Begin FMeshMapEvaluator interface
+	virtual void Setup(const FMeshMapBaker& Baker, FEvaluationContext& Context) override;
+
+	virtual EMeshMapEvaluatorType Type() const override { return EMeshMapEvaluatorType::Occlusion; }
+	// End FMeshMapEvaluator interface
 
 	template <EMeshOcclusionMapType OcclusionType, ESpace NormalSpace>
 	static void EvaluateSample(float*& Out, const FCorrespondenceSample& Sample, void* EvalData);
diff --git a/Engine/Plugins/Experimental/GeometryProcessing/Source/DynamicMesh/Public/Sampling/MeshPropertyMapEvaluator.h b/Engine/Plugins/Experimental/GeometryProcessing/Source/DynamicMesh/Public/Sampling/MeshPropertyMapEvaluator.h
--- a/Engine/Plugins/Experimental/GeometryProcessing/Source/DynamicMesh/Public/Sampling/MeshPropertyMapEvaluator.h
+++ b/Engine/Plugins/Experimental/GeometryProcessing/Source/DynamicMesh/Public/Sampling/MeshPropertyMapEvaluator.h
@@ -24,9 +24,12 @@ public:
 	EMeshPropertyMapType Property = EMeshPropertyMapType::Normal;
 
 public:
-	/** Invoked at start of bake to initialize baker. */
+	// Begin FMeshMapEvaluator interface
 	virtual void Setup(const FMeshMapBaker& Baker, FEvaluationContext& Context) override;
 
+	virtual EMeshMapEvaluatorType Type() const override { return EMeshMapEvaluatorType::Property; }
+	// End FMeshMapEvaluator interface
+
 	static void EvaluateSample(float*& Out, const FCorrespondenceSample& Sample, void* EvalData);
 
 	static void EvaluateDefault(float*& Out, void* EvalData);
diff --git a/Engine/Plugins/Experimental/GeometryProcessing/Source/DynamicMesh/Public/Sampling/MeshResampleImageEvaluator.h b/Engine/Plugins/Experimental/GeometryProcessing/Source/DynamicMesh/Public/Sampling/MeshResampleImageEvaluator.h
--- a/Engine/Plugins/Experimental/GeometryProcessing/Source/DynamicMesh/Public/Sampling/MeshResampleImageEvaluator.h
+++ b/Engine/Plugins/Experimental/GeometryProcessing/Source/DynamicMesh/Public/Sampling/MeshResampleImageEvaluator.h
@@ -20,9 +20,12 @@ public:
 	FVector4f DefaultColor = FVector4f(0, 0, 0, 1);
 
 public:
-	/** Invoked at start of bake to initialize baker. */
+	// Begin FMeshMapEvaluator interface
 	virtual void Setup(const FMeshMapBaker& Baker, FEvaluationContext& Context) override;
 
+	virtual EMeshMapEvaluatorType Type() const override { return EMeshMapEvaluatorType::ResampleImage; }
+	// End FMeshMapEvaluator interface
+
 	static void EvaluateSample(float*& Out, const FCorrespondenceSample& Sample, void* EvalData);
 
 	static void EvaluateDefault(float*& Out, void* EvalData);
@@ -42,10 +45,11 @@ public:
 	TMap<int32, TSharedPtr<UE::Geometry::TImageBuilder<FVector4f>, ESPMode::ThreadSafe>> MultiTextures;
 
 public:
-	virtual ~FMeshMultiResampleImageEvaluator() = default;
+	// Begin FMeshMapEvaluator interface
+	virtual void Setup(const FMeshMapBaker& Baker, FEvaluationContext& Context) override;
 
-	/** Invoked at start of bake to initialize baker. */
-	virtual void Setup(const FMeshMapBaker & Baker, FEvaluationContext & Context) override;
+	virtual EMeshMapEvaluatorType Type() const override { return EMeshMapEvaluatorType::MultiResampleImage; }
+	// End FMeshMapEvaluator interface
 
 	static void EvaluateSampleMulti(float*& Out, const FCorrespondenceSample& Sample, void* EvalData);
 
diff --git a/Engine/Plugins/Experimental/MeshModelingToolset/Source/MeshModelingTools/Private/BakeMeshAttributeMapsTool.cpp b/Engine/Plugins/Experimental/MeshModelingToolset/Source/MeshModelingTools/Private/BakeMeshAttributeMapsTool.cpp
--- a/Engine/Plugins/Experimental/MeshModelingToolset/Source/MeshModelingTools/Private/BakeMeshAttributeMapsTool.cpp
+++ b/Engine/Plugins/Experimental/MeshModelingToolset/Source/MeshModelingTools/Private/BakeMeshAttributeMapsTool.cpp
@@ -1335,76 +1335,88 @@ void UBakeMeshAttributeMapsTool::UpdateOnModeChange()
 
 void UBakeMeshAttributeMapsTool::OnMapsUpdated(const TUniquePtr<FMeshMapBaker>& NewResult)
 {
-	// This method assumes that the bake evaluators were instantiated in the order
-	// defined by ALL_BAKE_MAP_TYPES.
-	const EBakeMapType& BakeMapTypes = CachedBakeCacheSettings.BakeMapTypes;
-	int32 BakerIdx = 0;
-	for (EBakeMapType MapType : ALL_BAKE_MAP_TYPES)
+	FImageDimensions BakeDimensions = NewResult->GetDimensions();
+	int32 NumBakers = NewResult->NumBakers();
+	for (int32 BakerIdx = 0; BakerIdx < NumBakers; ++BakerIdx)
 	{
-		switch (BakeMapTypes & MapType)
-		{
-		case EBakeMapType::TangentSpaceNormalMap:
+		FMeshMapEvaluator* Baker = NewResult->GetBaker(BakerIdx);
+
+		auto UpdateCachedMap = [this, &NewResult, &BakerIdx, &BakeDimensions](const EBakeMapType BakeMapType, const FTexture2DBuilder::ETextureType TexType, const int32 ResultIdx) -> void
 		{
 			FTexture2DBuilder TextureBuilder;
-			TextureBuilder.Initialize(FTexture2DBuilder::ETextureType::NormalMap, CachedNormalMapSettings.Dimensions);
-			TextureBuilder.Copy(*NewResult->GetBakeResults(BakerIdx++)[0]);
+			TextureBuilder.Initialize(TexType, BakeDimensions);
+			TextureBuilder.Copy(*NewResult->GetBakeResults(BakerIdx)[ResultIdx]);
 			TextureBuilder.Commit(false);
-			CachedMaps[CachedMapIndices[EBakeMapType::TangentSpaceNormalMap]] = TextureBuilder.GetTexture2D();
+
+			// The CachedMap & CachedMapIndices can be thrown out of sync if updated during
+			// a background compute. Validate the computed type against our cached maps.
+			if (CachedMapIndices.Contains(BakeMapType))
+			{
+				CachedMaps[CachedMapIndices[BakeMapType]] = TextureBuilder.GetTexture2D();
+			}
+		};
+
+		switch (Baker->Type())
+		{
+		case EMeshMapEvaluatorType::Normal:
+		{
+			UpdateCachedMap(EBakeMapType::TangentSpaceNormalMap, FTexture2DBuilder::ETextureType::NormalMap, 0);
 			break;
 		}
-		case EBakeMapType::AmbientOcclusion:
-		case EBakeMapType::BentNormal:
-		case EBakeMapType::Occlusion:
+		case EMeshMapEvaluatorType::Occlusion:
 		{
+			// Occlusion Evaluator always outputs AmbientOcclusion then BentNormal.
+			FMeshOcclusionMapEvaluator* OcclusionBaker = static_cast<FMeshOcclusionMapEvaluator*>(Baker);
 			int32 OcclusionIdx = 0;
-			if ((bool)(BakeMapTypes & EBakeMapType::AmbientOcclusion))
+			if ((bool)(OcclusionBaker->OcclusionType & EMeshOcclusionMapType::AmbientOcclusion))
 			{
-				FTexture2DBuilder TextureBuilder;
-				TextureBuilder.Initialize(FTexture2DBuilder::ETextureType::AmbientOcclusion, CachedOcclusionMapSettings.Dimensions);
-				TextureBuilder.Copy(*NewResult->GetBakeResults(BakerIdx)[OcclusionIdx++]);
-				TextureBuilder.Commit(false);
-				CachedMaps[CachedMapIndices[EBakeMapType::AmbientOcclusion]] = TextureBuilder.GetTexture2D();
+				UpdateCachedMap(EBakeMapType::AmbientOcclusion, FTexture2DBuilder::ETextureType::AmbientOcclusion, OcclusionIdx++);
 			}
-			if ((bool)(BakeMapTypes & EBakeMapType::BentNormal))
+			if ((bool)(OcclusionBaker->OcclusionType & EMeshOcclusionMapType::BentNormal))
 			{
-				FTexture2DBuilder TextureBuilder;
-				TextureBuilder.Initialize(FTexture2DBuilder::ETextureType::NormalMap, CachedOcclusionMapSettings.Dimensions);
-				TextureBuilder.Copy(*NewResult->GetBakeResults(BakerIdx)[OcclusionIdx++]);
-				TextureBuilder.Commit(false);
-				CachedMaps[CachedMapIndices[EBakeMapType::BentNormal]] = TextureBuilder.GetTexture2D();
+				UpdateCachedMap(EBakeMapType::BentNormal, FTexture2DBuilder::ETextureType::NormalMap, OcclusionIdx++);
 			}
-			++BakerIdx;
 			break;
 		}
-		case EBakeMapType::Curvature:
+		case EMeshMapEvaluatorType::Curvature:
 		{
-			FTexture2DBuilder TextureBuilder;
-			TextureBuilder.Initialize(FTexture2DBuilder::ETextureType::Color, CachedCurvatureMapSettings.Dimensions);
-			TextureBuilder.Copy(*NewResult->GetBakeResults(BakerIdx++)[0]);
-			TextureBuilder.Commit(false);
-			CachedMaps[CachedMapIndices[BakeMapTypes & MapType]] = TextureBuilder.GetTexture2D();
+			UpdateCachedMap(EBakeMapType::Curvature, FTexture2DBuilder::ETextureType::Color, 0);
 			break;
 		}
-		case EBakeMapType::NormalImage:
-		case EBakeMapType::FaceNormalImage:
-		case EBakeMapType::PositionImage:
-		case EBakeMapType::MaterialID:
+		case EMeshMapEvaluatorType::Property:
 		{
-			FTexture2DBuilder TextureBuilder;
-			TextureBuilder.Initialize(FTexture2DBuilder::ETextureType::Color, CachedMeshPropertyMapSettings.Dimensions);
-			TextureBuilder.Copy(*NewResult->GetBakeResults(BakerIdx++)[0]);
-			TextureBuilder.Commit(false);
-			CachedMaps[CachedMapIndices[BakeMapTypes & MapType]] = TextureBuilder.GetTexture2D();
+			FMeshPropertyMapEvaluator* PropertyBaker = static_cast<FMeshPropertyMapEvaluator*>(Baker);
+			EBakeMapType MapType = EBakeMapType::None;
+			switch (PropertyBaker->Property)
+			{
+			case EMeshPropertyMapType::Normal:
+				MapType = EBakeMapType::NormalImage;
+				break;
+			case EMeshPropertyMapType::FacetNormal:
+				MapType = EBakeMapType::FaceNormalImage;
+				break;
+			case EMeshPropertyMapType::Position:
+				MapType = EBakeMapType::PositionImage;
+				break;
+			case EMeshPropertyMapType::MaterialID:
+				MapType = EBakeMapType::MaterialID;
+				break;
+			case EMeshPropertyMapType::UVPosition:
+			default:
+				break;
+			}
+
+			UpdateCachedMap(MapType, FTexture2DBuilder::ETextureType::Color, 0);
 			break;
 		}
-		case EBakeMapType::Texture2DImage:
-		case EBakeMapType::MultiTexture:
+		case EMeshMapEvaluatorType::ResampleImage:
 		{
-			FTexture2DBuilder TextureBuilder;
-			TextureBuilder.Initialize(FTexture2DBuilder::ETextureType::Color, CachedTexture2DImageSettings.Dimensions);
-			TextureBuilder.Copy(*NewResult->GetBakeResults(BakerIdx++)[0]);
-			TextureBuilder.Commit(false);
-			CachedMaps[CachedMapIndices[BakeMapTypes & MapType]] = TextureBuilder.GetTexture2D();
+			UpdateCachedMap(EBakeMapType::Texture2DImage, FTexture2DBuilder::ETextureType::Color, 0);
+			break;
+		}
+		case EMeshMapEvaluatorType::MultiResampleImage:
+		{
+			UpdateCachedMap(EBakeMapType::MultiTexture, FTexture2DBuilder::ETextureType::Color, 0);
 			break;
 		}
 		}
