Default to using UVAtlas in mesh parameterization op

#rnx

diff --git a/Engine/Plugins/Experimental/MeshModelingToolset/Source/ModelingOperatorsEditorOnly/Public/ParameterizationOps/ParameterizeMeshOp.h b/Engine/Plugins/Experimental/MeshModelingToolset/Source/ModelingOperatorsEditorOnly/Public/ParameterizationOps/ParameterizeMeshOp.h
--- a/Engine/Plugins/Experimental/MeshModelingToolset/Source/ModelingOperatorsEditorOnly/Public/ParameterizationOps/ParameterizeMeshOp.h
+++ b/Engine/Plugins/Experimental/MeshModelingToolset/Source/ModelingOperatorsEditorOnly/Public/ParameterizationOps/ParameterizeMeshOp.h
@@ -66,7 +66,7 @@ public:
 
 	EParamOpUnwrapType UnwrapType = EParamOpUnwrapType::ExpMap;
 	EParamOpIslandMode IslandMode = EParamOpIslandMode::PolyGroups;
-	EParamOpBackend Method = EParamOpBackend::XAtlas;
+	EParamOpBackend Method = EParamOpBackend::UVAltas;
 
 	// set ability on protected transform.
 	void SetTransform(FTransform3d& XForm)
