MeshModelingTools: Added tool targets that operate on static and skeletal mesh assets without an underlying component, and refactored to use the same code. Also changed the dynamic mesh shared pointers in the tool targets to be thread safe ones.

#rnx

diff --git a/Engine/Plugins/Experimental/MeshModelingToolset/Source/ModelingComponents/Private/ConversionUtils/DynamicMeshViaMeshDescriptionUtil.cpp b/Engine/Plugins/Experimental/MeshModelingToolset/Source/ModelingComponents/Private/ConversionUtils/DynamicMeshViaMeshDescriptionUtil.cpp
new file mode 100644
--- /dev/null
+++ b/Engine/Plugins/Experimental/MeshModelingToolset/Source/ModelingComponents/Private/ConversionUtils/DynamicMeshViaMeshDescriptionUtil.cpp
@@ -0,0 +1,46 @@
+// Copyright Epic Games, Inc. All Rights Reserved.
+
+#include "ConversionUtils/DynamicMeshViaMeshDescriptionUtil.h"
+
+#include "TargetInterfaces/MeshDescriptionCommitter.h"
+#include "TargetInterfaces/MeshDescriptionProvider.h"
+#include "DynamicMeshToMeshDescription.h"
+#include "MeshConversionOptions.h" //FConversionToMeshDescriptionOptions
+#include "MeshDescriptionToDynamicMesh.h"
+
+TSharedPtr<FDynamicMesh3, ESPMode::ThreadSafe> UE::Geometry::GetDynamicMeshViaMeshDescription(
+	IMeshDescriptionProvider& MeshDescriptionProvider)
+{
+	TSharedPtr<FDynamicMesh3> DynamicMesh = MakeShared<FDynamicMesh3>();
+	FMeshDescriptionToDynamicMesh Converter;
+	Converter.Convert(MeshDescriptionProvider.GetMeshDescription(), *DynamicMesh);
+	return DynamicMesh;
+}
+
+void UE::Geometry::CommitDynamicMeshViaMeshDescription(
+	IMeshDescriptionCommitter& MeshDescriptionCommitter, 
+	const FDynamicMesh3& Mesh, const IDynamicMeshCommitter::FDynamicMeshCommitInfo& CommitInfo)
+{
+	FConversionToMeshDescriptionOptions ConversionOptions;
+	ConversionOptions.bSetPolyGroups = CommitInfo.bPolygroupsChanged;
+	ConversionOptions.bUpdatePositions = CommitInfo.bPositionsChanged;
+	ConversionOptions.bUpdateNormals = CommitInfo.bNormalsChanged;
+	ConversionOptions.bUpdateTangents = CommitInfo.bTangentsChanged;
+	ConversionOptions.bUpdateUVs = CommitInfo.bUVsChanged;
+	ConversionOptions.bUpdateVtxColors = CommitInfo.bVertexColorsChanged;
+
+	MeshDescriptionCommitter.CommitMeshDescription([&CommitInfo, &ConversionOptions, &Mesh](const IMeshDescriptionCommitter::FCommitterParams& CommitParams)
+	{
+		FDynamicMeshToMeshDescription Converter(ConversionOptions);
+
+		if (!CommitInfo.bTopologyChanged)
+		{
+			Converter.UpdateUsingConversionOptions(&Mesh, *CommitParams.MeshDescriptionOut);
+		}
+		else
+		{
+			// Do a full conversion.
+			Converter.Convert(&Mesh, *CommitParams.MeshDescriptionOut);
+		}
+	});
+}
\ No newline at end of file
diff --git a/Engine/Plugins/Experimental/MeshModelingToolset/Source/ModelingComponents/Public/ConversionUtils/DynamicMeshViaMeshDescriptionUtil.h b/Engine/Plugins/Experimental/MeshModelingToolset/Source/ModelingComponents/Public/ConversionUtils/DynamicMeshViaMeshDescriptionUtil.h
new file mode 100644
--- /dev/null
+++ b/Engine/Plugins/Experimental/MeshModelingToolset/Source/ModelingComponents/Public/ConversionUtils/DynamicMeshViaMeshDescriptionUtil.h
@@ -0,0 +1,24 @@
+// Copyright Epic Games, Inc. All Rights Reserved.
+
+#pragma once
+
+#include "CoreMinimal.h"
+
+#include "TargetInterfaces/DynamicMeshCommitter.h"
+
+#include "GeometryBase.h"
+
+PREDECLARE_GEOMETRY(class FDynamicMesh3);
+class IMeshDescriptionCommitter;
+class IMeshDescriptionProvider;
+
+namespace UE {
+namespace Geometry {
+
+	MODELINGCOMPONENTS_API TSharedPtr<FDynamicMesh3, ESPMode::ThreadSafe> GetDynamicMeshViaMeshDescription(
+		IMeshDescriptionProvider& MeshDescriptionProvider);
+
+	MODELINGCOMPONENTS_API void CommitDynamicMeshViaMeshDescription(IMeshDescriptionCommitter& MeshDescriptionCommitter,
+		const UE::Geometry::FDynamicMesh3& Mesh, const IDynamicMeshCommitter::FDynamicMeshCommitInfo& CommitInfo);
+
+}}
\ No newline at end of file
diff --git a/Engine/Plugins/Experimental/MeshModelingToolset/Source/ModelingComponents/Public/TargetInterfaces/DynamicMeshProvider.h b/Engine/Plugins/Experimental/MeshModelingToolset/Source/ModelingComponents/Public/TargetInterfaces/DynamicMeshProvider.h
--- a/Engine/Plugins/Experimental/MeshModelingToolset/Source/ModelingComponents/Public/TargetInterfaces/DynamicMeshProvider.h
+++ b/Engine/Plugins/Experimental/MeshModelingToolset/Source/ModelingComponents/Public/TargetInterfaces/DynamicMeshProvider.h
@@ -31,5 +31,5 @@ public:
 	 * To guarantee changes having an effect, target must also be a IDynamicMeshCommitter 
 	 * and the mesh must be committed.
 	 */
-	virtual TSharedPtr<UE::Geometry::FDynamicMesh3> GetDynamicMesh() = 0;
+	virtual TSharedPtr<UE::Geometry::FDynamicMesh3, ESPMode::ThreadSafe> GetDynamicMesh() = 0;
 };
diff --git a/Engine/Plugins/Experimental/MeshModelingToolset/Source/ModelingComponentsEditorOnly/Private/ToolTargets/SkeletalMeshComponentToolTarget.cpp b/Engine/Plugins/Experimental/MeshModelingToolset/Source/ModelingComponentsEditorOnly/Private/ToolTargets/SkeletalMeshComponentToolTarget.cpp
--- a/Engine/Plugins/Experimental/MeshModelingToolset/Source/ModelingComponentsEditorOnly/Private/ToolTargets/SkeletalMeshComponentToolTarget.cpp
+++ b/Engine/Plugins/Experimental/MeshModelingToolset/Source/ModelingComponentsEditorOnly/Private/ToolTargets/SkeletalMeshComponentToolTarget.cpp
@@ -3,10 +3,13 @@
 #include "ToolTargets/SkeletalMeshComponentToolTarget.h"
 
 #include "Components/SkinnedMeshComponent.h"
+#include "ConversionUtils/DynamicMeshViaMeshDescriptionUtil.h"
 #include "DynamicMeshToMeshDescription.h"
 #include "MeshDescriptionToDynamicMesh.h"
 #include "Rendering/SkeletalMeshModel.h"
+#include "ToolTargets/SkeletalMeshToolTarget.h"
 
+using namespace UE::Geometry;
 
 namespace USkeletalMeshComponentToolTargetLocals
 {
@@ -31,12 +34,7 @@ void USkeletalMeshComponentToolTarget::GetMaterialSet(FComponentMaterialSet& Mat
 	if (bPreferAssetMaterials)
 	{
 		const USkeletalMesh* SkeletalMesh = Cast<USkinnedMeshComponent>(Component)->SkeletalMesh;
-		const TArray<FSkeletalMaterial>& Materials = SkeletalMesh->GetMaterials(); 
-		MaterialSetOut.Materials.SetNum(Materials.Num());
-		for (int32 k = 0; k < Materials.Num(); ++k)
-		{
-			MaterialSetOut.Materials[k] = Materials[k].MaterialInterface;
-		}
+		USkeletalMeshToolTarget::GetMaterialSet(SkeletalMesh, MaterialSetOut, bPreferAssetMaterials);
 	}
 	else
 	{
@@ -57,42 +55,10 @@ bool USkeletalMeshComponentToolTarget::CommitMaterialSetUpdate(const FComponentM
 	{
 		USkeletalMesh* SkeletalMesh = Cast<USkinnedMeshComponent>(Component)->SkeletalMesh;
 
-		if (SkeletalMesh->GetPathName().StartsWith(TEXT("/Engine/")))
-		{
-			UE_LOG(LogTemp, Warning, TEXT("CANNOT MODIFY BUILT-IN ENGINE ASSET %s"), *SkeletalMesh->GetPathName());
-			return false;
-		}
-
-		// flush any pending rendering commands, which might touch this component while we are rebuilding its mesh
-		FlushRenderingCommands();
-
 		// unregister the component while we update it's static mesh
 		TUniquePtr<FComponentReregisterContext> ComponentReregisterContext = MakeUnique<FComponentReregisterContext>(Component);
 
-		// make sure transactional flag is on
-		SkeletalMesh->SetFlags(RF_Transactional);
-
-		SkeletalMesh->Modify();
-
-		const int NewNumMaterials = MaterialSet.Materials.Num();
-		TArray<FSkeletalMaterial> &SkeletalMaterials = SkeletalMesh->GetMaterials(); 
-		if (NewNumMaterials != SkeletalMaterials.Num())
-		{
-			SkeletalMaterials.SetNum(NewNumMaterials);
-		}
-		for (int k = 0; k < NewNumMaterials; ++k)
-		{
-			if (SkeletalMaterials[k].MaterialInterface != MaterialSet.Materials[k])
-			{
-				SkeletalMaterials[k].MaterialInterface = MaterialSet.Materials[k];
-				if (SkeletalMaterials[k].MaterialSlotName.IsNone())
-				{
-					SkeletalMaterials[k].MaterialSlotName = MaterialSet.Materials[k]->GetFName();
-				}
-			}
-		}
-
-		SkeletalMesh->PostEditChange();
+		return USkeletalMeshToolTarget::CommitMaterialSetUpdate(SkeletalMesh, MaterialSet, bApplyToAsset);
 	}
 	else
 	{
@@ -126,29 +92,9 @@ FMeshDescription* USkeletalMeshComponentToolTarget::GetMeshDescription()
 
 	if (!CachedMeshDescription.IsValid())
 	{
-		using namespace USkeletalMeshComponentToolTargetLocals;
-
 		CachedMeshDescription = MakeUnique<FMeshDescription>();
 		const USkeletalMesh* SkeletalMesh = Cast<USkinnedMeshComponent>(Component)->SkeletalMesh;
-
-		// Check first if we have bulk data available and non-empty.
-		if (SkeletalMesh->IsLODImportedDataBuildAvailable(LODIndex) && !SkeletalMesh->IsLODImportedDataEmpty(LODIndex))
-		{
-			FSkeletalMeshImportData SkeletalMeshImportData;
-			SkeletalMesh->LoadLODImportedData(LODIndex, SkeletalMeshImportData);
-			SkeletalMeshImportData.GetMeshDescription(*CachedMeshDescription);
-		}
-		else
-		{
-			// Fall back on the LOD model directly if no bulk data exists. When we commit
-			// the mesh description, we override using the bulk data. This can happen for older
-			// skeletal meshes, from UE 4.24 and earlier.
-			const FSkeletalMeshModel* SkeletalMeshModel = SkeletalMesh->GetImportedModel();
-			if (SkeletalMeshModel && SkeletalMeshModel->LODModels.IsValidIndex(LODIndex))
-			{
-				SkeletalMeshModel->LODModels[LODIndex].GetMeshDescription(*CachedMeshDescription, SkeletalMesh);
-			}			
-		}
+		USkeletalMeshToolTarget::GetMeshDescription(SkeletalMesh, *CachedMeshDescription);
 	}
 	
 	return CachedMeshDescription.Get();
@@ -156,87 +102,29 @@ FMeshDescription* USkeletalMeshComponentToolTarget::GetMeshDescription()
 
 void USkeletalMeshComponentToolTarget::CommitMeshDescription(const FCommitter& Committer)
 {
-	using namespace USkeletalMeshComponentToolTargetLocals;
 	if (ensure(IsValid()) == false) return;
 
 	USkeletalMesh* SkeletalMesh = Cast<USkinnedMeshComponent>(Component)->SkeletalMesh;
 
-	if (SkeletalMesh->GetPathName().StartsWith(TEXT("/Engine/")))
-	{
-		const FString DebugMessage = FString::Printf(TEXT("CANNOT MODIFY BUILT-IN ENGINE ASSET %s"), *SkeletalMesh->GetPathName());
-		if (GAreScreenMessagesEnabled)
-		{
-			GEngine->AddOnScreenDebugMessage(INDEX_NONE, 10.0f, FColor::Red, DebugMessage);
-		}
-		UE_LOG(LogTemp, Warning, TEXT("%s"), *DebugMessage);
-		return;
-	}
-
-	// flush any pending rendering commands, which might touch this component while we are rebuilding it's mesh
-	FlushRenderingCommands();
-
 	// unregister the component while we update its skeletal mesh
 	FComponentReregisterContext ComponentReregisterContext(Component);
 
-	// make sure transactional flag is on for this asset
-	SkeletalMesh->SetFlags(RF_Transactional);
-
-	verify(SkeletalMesh->Modify());
-
-	FCommitterParams CommitterParams;
-
-	CommitterParams.MeshDescriptionOut = GetMeshDescription();
-
-	Committer(CommitterParams);
-
-	FSkeletalMeshImportData SkeletalMeshImportData = FSkeletalMeshImportData::CreateFromMeshDescription(*CommitterParams.MeshDescriptionOut);
-	SkeletalMesh->SaveLODImportedData(LODIndex, SkeletalMeshImportData);
-
-	// Make sure the mesh builder knows it's the latest variety, so that the render data gets
-	// properly rebuilt.
-	SkeletalMesh->SetLODImportedDataVersions(LODIndex, ESkeletalMeshGeoImportVersions::LatestVersion, ESkeletalMeshSkinningImportVersions::LatestVersion);
-	SkeletalMesh->SetUseLegacyMeshDerivedDataKey(false);
-
-	SkeletalMesh->PostEditChange();
+	USkeletalMeshToolTarget::CommitMeshDescription(SkeletalMesh, GetMeshDescription(), Committer);
 
 	// this rebuilds physics, but it doesn't undo!
 	Component->RecreatePhysicsState();
-	
+
 	CachedMeshDescription.Reset();
 }
 
-TSharedPtr<FDynamicMesh3> USkeletalMeshComponentToolTarget::GetDynamicMesh()
+TSharedPtr<FDynamicMesh3, ESPMode::ThreadSafe> USkeletalMeshComponentToolTarget::GetDynamicMesh()
 {
-	TSharedPtr<FDynamicMesh3> DynamicMesh = MakeShared<FDynamicMesh3>();
-	FMeshDescriptionToDynamicMesh Converter;
-	Converter.Convert(GetMeshDescription(), *DynamicMesh);
-	return DynamicMesh;
+	return GetDynamicMeshViaMeshDescription(*this);
 }
 
 void USkeletalMeshComponentToolTarget::CommitDynamicMesh(const FDynamicMesh3& Mesh, const FDynamicMeshCommitInfo& CommitInfo)
 {
-	FConversionToMeshDescriptionOptions ConversionOptions;
-	ConversionOptions.bSetPolyGroups = CommitInfo.bPolygroupsChanged;
-	ConversionOptions.bUpdatePositions = CommitInfo.bPositionsChanged;
-	ConversionOptions.bUpdateNormals = CommitInfo.bNormalsChanged;
-	ConversionOptions.bUpdateTangents = CommitInfo.bTangentsChanged;
-	ConversionOptions.bUpdateUVs = CommitInfo.bUVsChanged;
-	ConversionOptions.bUpdateVtxColors = CommitInfo.bVertexColorsChanged;
-
-	CommitMeshDescription([&CommitInfo, &ConversionOptions, &Mesh](const IMeshDescriptionCommitter::FCommitterParams& CommitParams)
-		{
-			FDynamicMeshToMeshDescription Converter(ConversionOptions);
-
-			if (!CommitInfo.bTopologyChanged)
-			{
-				Converter.UpdateUsingConversionOptions(&Mesh, *CommitParams.MeshDescriptionOut);
-			}
-			else
-			{
-				// Do a full conversion.
-				Converter.Convert(&Mesh, *CommitParams.MeshDescriptionOut);
-			}
-		});
+	CommitDynamicMeshViaMeshDescription(*this, Mesh, CommitInfo);
 }
 
 USkeletalMesh* USkeletalMeshComponentToolTarget::GetSkeletalMesh() const
diff --git a/Engine/Plugins/Experimental/MeshModelingToolset/Source/ModelingComponentsEditorOnly/Private/ToolTargets/SkeletalMeshToolTarget.cpp b/Engine/Plugins/Experimental/MeshModelingToolset/Source/ModelingComponentsEditorOnly/Private/ToolTargets/SkeletalMeshToolTarget.cpp
new file mode 100644
--- /dev/null
+++ b/Engine/Plugins/Experimental/MeshModelingToolset/Source/ModelingComponentsEditorOnly/Private/ToolTargets/SkeletalMeshToolTarget.cpp
@@ -0,0 +1,219 @@
+﻿// Copyright Epic Games, Inc. All Rights Reserved.
+
+#include "ToolTargets/SkeletalMeshToolTarget.h"
+
+#include "Components/SkinnedMeshComponent.h"
+#include "ConversionUtils/DynamicMeshViaMeshDescriptionUtil.h"
+#include "Rendering/SkeletalMeshModel.h"
+
+using namespace UE::Geometry;
+
+namespace USkeletalMeshToolTargetLocals
+{
+	int32 LODIndex = 0;
+}
+
+bool USkeletalMeshToolTarget::IsValid() const
+{
+	return SkeletalMesh && !SkeletalMesh->IsPendingKillOrUnreachable() && SkeletalMesh->IsValidLowLevel();
+}
+
+int32 USkeletalMeshToolTarget::GetNumMaterials() const
+{
+	return ensure(IsValid()) ? SkeletalMesh->GetMaterials().Num() : 0;
+}
+
+UMaterialInterface* USkeletalMeshToolTarget::GetMaterial(int32 MaterialIndex) const
+{
+	return ensure(IsValid() && MaterialIndex < SkeletalMesh->GetMaterials().Num()) ? 
+		SkeletalMesh->GetMaterials()[MaterialIndex].MaterialInterface : nullptr;
+}
+
+void USkeletalMeshToolTarget::GetMaterialSet(FComponentMaterialSet& MaterialSetOut, bool bPreferAssetMaterials) const
+{
+	if (!ensure(IsValid())) return;
+	GetMaterialSet(SkeletalMesh, MaterialSetOut, bPreferAssetMaterials);
+}
+
+void USkeletalMeshToolTarget::GetMaterialSet(const USkeletalMesh* SkeletalMeshIn, FComponentMaterialSet& MaterialSetOut,
+	bool bPreferAssetMaterials)
+{
+	const TArray<FSkeletalMaterial>& Materials = SkeletalMeshIn->GetMaterials(); 
+	MaterialSetOut.Materials.SetNum(Materials.Num());
+	for (int32 k = 0; k < Materials.Num(); ++k)
+	{
+		MaterialSetOut.Materials[k] = Materials[k].MaterialInterface;
+	}
+}
+
+bool USkeletalMeshToolTarget::CommitMaterialSetUpdate(const FComponentMaterialSet& MaterialSet, bool bApplyToAsset)
+{
+	if (!ensure(IsValid())) return false;
+	return CommitMaterialSetUpdate(SkeletalMesh, MaterialSet, bApplyToAsset);
+}
+
+bool USkeletalMeshToolTarget::CommitMaterialSetUpdate(USkeletalMesh* SkeletalMeshIn, 
+	const FComponentMaterialSet& MaterialSet, bool bApplyToAsset)
+{
+	if (!bApplyToAsset)
+	{
+		return false;
+	}
+
+	if (SkeletalMeshIn->GetPathName().StartsWith(TEXT("/Engine/")))
+	{
+		UE_LOG(LogTemp, Warning, TEXT("CANNOT MODIFY BUILT-IN ENGINE ASSET %s"), *SkeletalMeshIn->GetPathName());
+		return false;
+	}
+
+	// flush any pending rendering commands, which might touch a component while we are rebuilding its mesh
+	FlushRenderingCommands();
+
+	// make sure transactional flag is on
+	SkeletalMeshIn->SetFlags(RF_Transactional);
+
+	SkeletalMeshIn->Modify();
+
+	const int NewNumMaterials = MaterialSet.Materials.Num();
+	TArray<FSkeletalMaterial> &SkeletalMaterials = SkeletalMeshIn->GetMaterials(); 
+	if (NewNumMaterials != SkeletalMaterials.Num())
+	{
+		SkeletalMaterials.SetNum(NewNumMaterials);
+	}
+	for (int k = 0; k < NewNumMaterials; ++k)
+	{
+		if (SkeletalMaterials[k].MaterialInterface != MaterialSet.Materials[k])
+		{
+			SkeletalMaterials[k].MaterialInterface = MaterialSet.Materials[k];
+			if (SkeletalMaterials[k].MaterialSlotName.IsNone())
+			{
+				SkeletalMaterials[k].MaterialSlotName = MaterialSet.Materials[k]->GetFName();
+			}
+		}
+	}
+
+	SkeletalMeshIn->PostEditChange();
+
+	return true;
+}
+
+FMeshDescription* USkeletalMeshToolTarget::GetMeshDescription()
+{
+	if (!ensure(IsValid()))
+	{
+		return nullptr;
+	}
+
+	if (!CachedMeshDescription.IsValid())
+	{
+		CachedMeshDescription = MakeUnique<FMeshDescription>();
+		GetMeshDescription(SkeletalMesh, *CachedMeshDescription);
+	}
+
+	return CachedMeshDescription.Get();
+}
+
+void USkeletalMeshToolTarget::GetMeshDescription(const USkeletalMesh* SkeletalMeshIn, FMeshDescription& MeshDescription)
+{
+	using namespace USkeletalMeshToolTargetLocals;
+
+	// Check first if we have bulk data available and non-empty.
+	if (SkeletalMeshIn->IsLODImportedDataBuildAvailable(LODIndex) && !SkeletalMeshIn->IsLODImportedDataEmpty(LODIndex))
+	{
+		FSkeletalMeshImportData SkeletalMeshImportData;
+		SkeletalMeshIn->LoadLODImportedData(LODIndex, SkeletalMeshImportData);
+		SkeletalMeshImportData.GetMeshDescription(MeshDescription);
+	}
+	else
+	{
+		// Fall back on the LOD model directly if no bulk data exists. When we commit
+		// the mesh description, we override using the bulk data. This can happen for older
+		// skeletal meshes, from UE 4.24 and earlier.
+		const FSkeletalMeshModel* SkeletalMeshModel = SkeletalMeshIn->GetImportedModel();
+		if (SkeletalMeshModel && SkeletalMeshModel->LODModels.IsValidIndex(LODIndex))
+		{
+			SkeletalMeshModel->LODModels[LODIndex].GetMeshDescription(MeshDescription, SkeletalMeshIn);
+		}			
+	}
+}
+
+void USkeletalMeshToolTarget::CommitMeshDescription(const FCommitter& Committer)
+{
+	if (ensure(IsValid()) == false) return;
+	CommitMeshDescription(SkeletalMesh, GetMeshDescription(), Committer);
+}
+
+void USkeletalMeshToolTarget::CommitMeshDescription(USkeletalMesh* SkeletalMeshIn,
+	FMeshDescription* MeshDescription, const FCommitter& Committer)
+{
+	using namespace USkeletalMeshToolTargetLocals;
+
+	if (SkeletalMeshIn->GetPathName().StartsWith(TEXT("/Engine/")))
+	{
+		const FString DebugMessage = FString::Printf(TEXT("CANNOT MODIFY BUILT-IN ENGINE ASSET %s"), *SkeletalMeshIn->GetPathName());
+		if (GAreScreenMessagesEnabled)
+		{
+			GEngine->AddOnScreenDebugMessage(INDEX_NONE, 10.0f, FColor::Red, DebugMessage);
+		}
+		UE_LOG(LogTemp, Warning, TEXT("%s"), *DebugMessage);
+		return;
+	}
+
+	// flush any pending rendering commands, which might touch a component while we are rebuilding it's mesh
+	FlushRenderingCommands();
+
+	// make sure transactional flag is on for this asset
+	SkeletalMeshIn->SetFlags(RF_Transactional);
+
+	verify(SkeletalMeshIn->Modify());
+
+	FCommitterParams CommitterParams;
+
+	CommitterParams.MeshDescriptionOut = MeshDescription;
+
+	Committer(CommitterParams);
+
+	FSkeletalMeshImportData SkeletalMeshImportData = 
+		FSkeletalMeshImportData::CreateFromMeshDescription(*CommitterParams.MeshDescriptionOut);
+	SkeletalMeshIn->SaveLODImportedData(LODIndex, SkeletalMeshImportData);
+
+	// Make sure the mesh builder knows it's the latest variety, so that the render data gets
+	// properly rebuilt.
+	SkeletalMeshIn->SetLODImportedDataVersions(LODIndex, ESkeletalMeshGeoImportVersions::LatestVersion, ESkeletalMeshSkinningImportVersions::LatestVersion);
+	SkeletalMeshIn->SetUseLegacyMeshDerivedDataKey(false);
+
+	SkeletalMeshIn->PostEditChange();
+}
+
+TSharedPtr<FDynamicMesh3, ESPMode::ThreadSafe> USkeletalMeshToolTarget::GetDynamicMesh()
+{
+	return GetDynamicMeshViaMeshDescription(*this);
+}
+
+void USkeletalMeshToolTarget::CommitDynamicMesh(const FDynamicMesh3& Mesh, const FDynamicMeshCommitInfo& CommitInfo)
+{
+	CommitDynamicMeshViaMeshDescription(*this, Mesh, CommitInfo);
+}
+
+USkeletalMesh* USkeletalMeshToolTarget::GetSkeletalMesh() const
+{
+	return IsValid() ? SkeletalMesh : nullptr;
+}
+
+
+// Factory
+
+bool USkeletalMeshToolTargetFactory::CanBuildTarget(UObject* SourceObject, const FToolTargetTypeRequirements& Requirements) const
+{
+	return Cast<USkeletalMesh>(SourceObject) 
+		&& Requirements.AreSatisfiedBy(USkeletalMeshToolTarget::StaticClass());
+}
+
+UToolTarget* USkeletalMeshToolTargetFactory::BuildTarget(UObject* SourceObject, const FToolTargetTypeRequirements& Requirements)
+{
+	USkeletalMeshToolTarget* Target = NewObject<USkeletalMeshToolTarget>();
+	Target->SkeletalMesh = Cast<USkeletalMesh>(SourceObject);
+	check(Target->SkeletalMesh && Requirements.AreSatisfiedBy(Target));
+
+	return Target;
+}
diff --git a/Engine/Plugins/Experimental/MeshModelingToolset/Source/ModelingComponentsEditorOnly/Private/ToolTargets/StaticMeshComponentToolTarget.cpp b/Engine/Plugins/Experimental/MeshModelingToolset/Source/ModelingComponentsEditorOnly/Private/ToolTargets/StaticMeshComponentToolTarget.cpp
--- a/Engine/Plugins/Experimental/MeshModelingToolset/Source/ModelingComponentsEditorOnly/Private/ToolTargets/StaticMeshComponentToolTarget.cpp
+++ b/Engine/Plugins/Experimental/MeshModelingToolset/Source/ModelingComponentsEditorOnly/Private/ToolTargets/StaticMeshComponentToolTarget.cpp
@@ -4,22 +4,13 @@
 
 #include "ComponentReregisterContext.h"
 #include "Components/StaticMeshComponent.h"
-#include "DynamicMeshToMeshDescription.h"
+#include "ConversionUtils/DynamicMeshViaMeshDescriptionUtil.h"
 #include "Engine/StaticMesh.h"
 #include "Materials/Material.h"
-#include "MeshDescriptionToDynamicMesh.h"
 #include "RenderingThread.h"
-#include "AssetUtils/MeshDescriptionUtil.h"
-
-static void DisplayCriticalWarningMessage(const FString& Message)
-{
-	if (GAreScreenMessagesEnabled)
-	{
-		GEngine->AddOnScreenDebugMessage(INDEX_NONE, 10.0f, FColor::Red, Message);
-	}
-	UE_LOG(LogTemp, Warning, TEXT("%s"), *Message);
-}
+#include "ToolTargets/StaticMeshToolTarget.h"
 
+using namespace UE::Geometry;
 
 void UStaticMeshComponentToolTarget::SetEditingLOD(EStaticMeshEditingLOD RequestedEditingLOD)
 {
@@ -29,31 +20,7 @@ void UStaticMeshComponentToolTarget::SetEditingLOD(EStaticMeshEditingLOD Request
 	if (ensure(StaticMeshComponent != nullptr))
 	{
 		UStaticMesh* StaticMeshAsset = StaticMeshComponent->GetStaticMesh();
-		if (ensure(StaticMeshAsset != nullptr))
-		{
-			if (RequestedEditingLOD == EStaticMeshEditingLOD::MaxQuality)
-			{
-				ValidEditingLOD = StaticMeshAsset->IsHiResMeshDescriptionValid() ? EStaticMeshEditingLOD::HiResSource : EStaticMeshEditingLOD::LOD0;
-			}
-			else if (RequestedEditingLOD == EStaticMeshEditingLOD::HiResSource)
-			{
-				ValidEditingLOD = StaticMeshAsset->IsHiResMeshDescriptionValid() ? EStaticMeshEditingLOD::HiResSource : EStaticMeshEditingLOD::LOD0;
-				if (ValidEditingLOD != EStaticMeshEditingLOD::HiResSource)
-				{
-					DisplayCriticalWarningMessage(FString(TEXT("HiRes Source selected but not available - Falling Back to LOD0")));
-				}
-			}
-			else
-			{
-				ValidEditingLOD = RequestedEditingLOD;
-				int32 MaxExistingLOD = StaticMeshAsset->GetNumSourceModels() - 1;
-				if ((int32)ValidEditingLOD > MaxExistingLOD)
-				{
-					DisplayCriticalWarningMessage(FString::Printf(TEXT("LOD%d Requested but not available - Falling Back to LOD%d"), (int32)ValidEditingLOD, MaxExistingLOD));
-					ValidEditingLOD = (EStaticMeshEditingLOD)MaxExistingLOD;
-				}
-			}
-		}
+		ValidEditingLOD = UStaticMeshToolTarget::GetValidEditingLOD(StaticMeshAsset, RequestedEditingLOD);
 	}
 
 	EditingLOD = ValidEditingLOD;
@@ -72,24 +39,8 @@ bool UStaticMeshComponentToolTarget::IsValid() const
 		return false;
 	}
 	UStaticMesh* StaticMesh = StaticMeshComponent->GetStaticMesh();
-	if (StaticMesh == nullptr)
-	{
-		return false;
-	}
-
-	if (EditingLOD == EStaticMeshEditingLOD::HiResSource)
-	{
-		if (StaticMesh->IsHiResMeshDescriptionValid() == false)
-		{
-			return false;
-		}
-	}
-	else if ((int32)EditingLOD >= StaticMesh->GetNumSourceModels())
-	{
-		return false;
-	}
-
-	return true;
+	
+	return UStaticMeshToolTarget::IsValid(StaticMesh, EditingLOD);
 }
 
 
@@ -111,12 +62,7 @@ void UStaticMeshComponentToolTarget::GetMaterialSet(FComponentMaterialSet& Mater
 	if (bPreferAssetMaterials)
 	{
 		UStaticMesh* StaticMesh = Cast<UStaticMeshComponent>(Component)->GetStaticMesh();
-		int32 NumMaterials = Component->GetNumMaterials();
-		MaterialSetOut.Materials.SetNum(NumMaterials);
-		for (int32 k = 0; k < NumMaterials; ++k)
-		{
-			MaterialSetOut.Materials[k] = StaticMesh->GetMaterial(k);
-		}
+		UStaticMeshToolTarget::GetMaterialSet(StaticMesh, MaterialSetOut, bPreferAssetMaterials);
 	}
 	else
 	{
@@ -133,55 +79,28 @@ bool UStaticMeshComponentToolTarget::CommitMaterialSetUpdate(const FComponentMat
 {
 	if (!ensure(IsValid())) return false;
 
-	// filter out any Engine materials that we don't want to be permanently assigning
-	TArray<UMaterialInterface*> FilteredMaterials = MaterialSet.Materials;
-	for (int32 k = 0; k < FilteredMaterials.Num(); ++k)
-	{
-		FString AssetPath = FilteredMaterials[k]->GetPathName();
-		if (AssetPath.StartsWith(TEXT("/MeshModelingToolset/")))
-		{
-			FilteredMaterials[k] = UMaterial::GetDefaultMaterial(MD_Surface);
-		}
-	}
-
 	if (bApplyToAsset)
 	{
 		UStaticMesh* StaticMesh = Cast<UStaticMeshComponent>(Component)->GetStaticMesh();
 
-		if (StaticMesh->GetPathName().StartsWith(TEXT("/Engine/")))
-		{
-			UE_LOG(LogTemp, Warning, TEXT("CANNOT MODIFY BUILT-IN ENGINE ASSET %s"), *StaticMesh->GetPathName());
-			return false;
-		}
-
-		// flush any pending rendering commands, which might touch this component while we are rebuilding its mesh
-		FlushRenderingCommands();
-
 		// unregister the component while we update it's static mesh
 		TUniquePtr<FComponentReregisterContext> ComponentReregisterContext = MakeUnique<FComponentReregisterContext>(Component);
 
-		// make sure transactional flag is on
-		StaticMesh->SetFlags(RF_Transactional);
-
-		StaticMesh->Modify();
-
-		int NewNumMaterials = FilteredMaterials.Num();
-		if (NewNumMaterials != StaticMesh->GetStaticMaterials().Num())
-		{
-			StaticMesh->GetStaticMaterials().SetNum(NewNumMaterials);
-		}
-		for (int k = 0; k < NewNumMaterials; ++k)
+		return UStaticMeshToolTarget::CommitMaterialSetUpdate(StaticMesh, MaterialSet, bApplyToAsset);
+	}
+	else
+	{
+		// filter out any Engine materials that we don't want to be permanently assigning
+		TArray<UMaterialInterface*> FilteredMaterials = MaterialSet.Materials;
+		for (int32 k = 0; k < FilteredMaterials.Num(); ++k)
 		{
-			if (StaticMesh->GetMaterial(k) != FilteredMaterials[k])
+			FString AssetPath = FilteredMaterials[k]->GetPathName();
+			if (AssetPath.StartsWith(TEXT("/MeshModelingToolset/")))
 			{
-				StaticMesh->SetMaterial(k, FilteredMaterials[k]);
+				FilteredMaterials[k] = UMaterial::GetDefaultMaterial(MD_Surface);
 			}
 		}
 
-		StaticMesh->PostEditChange();
-	}
-	else
-	{
 		int32 NumMaterialsNeeded = Component->GetNumMaterials();
 		int32 NumMaterialsGiven = FilteredMaterials.Num();
 
@@ -222,89 +141,23 @@ void UStaticMeshComponentToolTarget::CommitMeshDescription(const FCommitter& Com
 
 	UStaticMesh* StaticMesh = Cast<UStaticMeshComponent>(Component)->GetStaticMesh();
 
-	if (StaticMesh->GetPathName().StartsWith(TEXT("/Engine/")))
-	{
-		DisplayCriticalWarningMessage(FString::Printf(TEXT("CANNOT MODIFY BUILT-IN ENGINE ASSET %s"), *StaticMesh->GetPathName()));
-		return;
-	}
-
-	// flush any pending rendering commands, which might touch this component while we are rebuilding it's mesh
-	FlushRenderingCommands();
-
 	// unregister the component while we update its static mesh
 	FComponentReregisterContext ComponentReregisterContext(Component);
 
-	// make sure transactional flag is on for this asset
-	StaticMesh->SetFlags(RF_Transactional);
-
-	verify(StaticMesh->Modify());
-
-	// disable auto-generated normals StaticMesh build setting
-	UE::MeshDescription::FStaticMeshBuildSettingChange SettingsChange;
-	SettingsChange.AutoGeneratedNormals = UE::MeshDescription::EBuildSettingBoolChange::Disable;
-	UE::MeshDescription::ConfigureBuildSettings(StaticMesh, 0, SettingsChange);
-
-	if (EditingLOD == EStaticMeshEditingLOD::HiResSource)
-	{
-		verify(StaticMesh->ModifyHiResMeshDescription());
-	}
-	else
-	{
-		verify(StaticMesh->ModifyMeshDescription((int32)EditingLOD));
-	}
-
-	FCommitterParams CommitterParams;
-	CommitterParams.MeshDescriptionOut = GetMeshDescription();
-
-	Committer(CommitterParams);
-
-	if (EditingLOD == EStaticMeshEditingLOD::HiResSource)
-	{
-		StaticMesh->CommitHiResMeshDescription();
-	}
-	else
-	{
-		StaticMesh->CommitMeshDescription((int32)EditingLOD);
-	}
-
-	StaticMesh->PostEditChange();
+	UStaticMeshToolTarget::CommitMeshDescription(StaticMesh, GetMeshDescription(), Committer, EditingLOD);
 
 	// this rebuilds physics, but it doesn't undo!
 	Component->RecreatePhysicsState();
 }
 
-TSharedPtr<FDynamicMesh3> UStaticMeshComponentToolTarget::GetDynamicMesh()
+TSharedPtr<FDynamicMesh3, ESPMode::ThreadSafe> UStaticMeshComponentToolTarget::GetDynamicMesh()
 {
-	TSharedPtr<FDynamicMesh3> DynamicMesh = MakeShared<FDynamicMesh3>();
-	FMeshDescriptionToDynamicMesh Converter;
-	Converter.Convert(GetMeshDescription(), *DynamicMesh);
-	return DynamicMesh;
+	return GetDynamicMeshViaMeshDescription(*this);
 }
 
 void UStaticMeshComponentToolTarget::CommitDynamicMesh(const FDynamicMesh3& Mesh, const FDynamicMeshCommitInfo& CommitInfo)
 {
-	FConversionToMeshDescriptionOptions ConversionOptions;
-	ConversionOptions.bSetPolyGroups = CommitInfo.bPolygroupsChanged;
-	ConversionOptions.bUpdatePositions = CommitInfo.bPositionsChanged;
-	ConversionOptions.bUpdateNormals = CommitInfo.bNormalsChanged;
-	ConversionOptions.bUpdateTangents = CommitInfo.bTangentsChanged;
-	ConversionOptions.bUpdateUVs = CommitInfo.bUVsChanged;
-	ConversionOptions.bUpdateVtxColors = CommitInfo.bVertexColorsChanged;
-
-	CommitMeshDescription([&CommitInfo, &ConversionOptions, &Mesh](const IMeshDescriptionCommitter::FCommitterParams& CommitParams)
-	{
-		FDynamicMeshToMeshDescription Converter(ConversionOptions);
-
-		if (!CommitInfo.bTopologyChanged)
-		{
-			Converter.UpdateUsingConversionOptions(&Mesh, *CommitParams.MeshDescriptionOut);
-		}
-		else
-		{
-			// Do a full conversion.
-			Converter.Convert(&Mesh, *CommitParams.MeshDescriptionOut);
-		}
-	});
+	CommitDynamicMeshViaMeshDescription(*this, Mesh, CommitInfo);
 }
 
 UStaticMesh* UStaticMeshComponentToolTarget::GetStaticMesh() const
diff --git a/Engine/Plugins/Experimental/MeshModelingToolset/Source/ModelingComponentsEditorOnly/Private/ToolTargets/StaticMeshToolTarget.cpp b/Engine/Plugins/Experimental/MeshModelingToolset/Source/ModelingComponentsEditorOnly/Private/ToolTargets/StaticMeshToolTarget.cpp
new file mode 100644
--- /dev/null
+++ b/Engine/Plugins/Experimental/MeshModelingToolset/Source/ModelingComponentsEditorOnly/Private/ToolTargets/StaticMeshToolTarget.cpp
@@ -0,0 +1,289 @@
+// Copyright Epic Games, Inc. All Rights Reserved.
+
+#include "ToolTargets/StaticMeshToolTarget.h"
+
+#include "AssetUtils/MeshDescriptionUtil.h"
+#include "ComponentReregisterContext.h"
+#include "ConversionUtils/DynamicMeshViaMeshDescriptionUtil.h"
+#include "Engine/StaticMesh.h"
+#include "Materials/Material.h"
+#include "RenderingThread.h"
+
+using namespace UE::Geometry;
+
+namespace StaticMeshToolTargetLocals
+{
+	static void DisplayCriticalWarningMessage(const FString& Message)
+	{
+		if (GAreScreenMessagesEnabled)
+		{
+			GEngine->AddOnScreenDebugMessage(INDEX_NONE, 10.0f, FColor::Red, Message);
+		}
+		UE_LOG(LogTemp, Warning, TEXT("%s"), *Message);
+	}
+}
+
+void UStaticMeshToolTarget::SetEditingLOD(EStaticMeshEditingLOD RequestedEditingLOD)
+{
+	EditingLOD = GetValidEditingLOD(StaticMesh, RequestedEditingLOD);
+}
+
+EStaticMeshEditingLOD UStaticMeshToolTarget::GetValidEditingLOD(const UStaticMesh* StaticMeshIn, 
+	EStaticMeshEditingLOD RequestedEditingLOD)
+{
+	using namespace StaticMeshToolTargetLocals;
+
+	EStaticMeshEditingLOD ValidEditingLOD = EStaticMeshEditingLOD::LOD0;
+
+	if (ensure(StaticMeshIn != nullptr))
+	{
+		if (RequestedEditingLOD == EStaticMeshEditingLOD::MaxQuality)
+		{
+			ValidEditingLOD = StaticMeshIn->IsHiResMeshDescriptionValid() ? EStaticMeshEditingLOD::HiResSource : EStaticMeshEditingLOD::LOD0;
+		}
+		else if (RequestedEditingLOD == EStaticMeshEditingLOD::HiResSource)
+		{
+			ValidEditingLOD = StaticMeshIn->IsHiResMeshDescriptionValid() ? EStaticMeshEditingLOD::HiResSource : EStaticMeshEditingLOD::LOD0;
+			if (ValidEditingLOD != EStaticMeshEditingLOD::HiResSource)
+			{
+				DisplayCriticalWarningMessage(FString(TEXT("HiRes Source selected but not available - Falling Back to LOD0")));
+			}
+		}
+		else
+		{
+			ValidEditingLOD = RequestedEditingLOD;
+			int32 MaxExistingLOD = StaticMeshIn->GetNumSourceModels() - 1;
+			if ((int32)ValidEditingLOD > MaxExistingLOD)
+			{
+				DisplayCriticalWarningMessage(FString::Printf(TEXT("LOD%d Requested but not available - Falling Back to LOD%d"), (int32)ValidEditingLOD, MaxExistingLOD));
+				ValidEditingLOD = (EStaticMeshEditingLOD)MaxExistingLOD;
+			}
+		}
+	}
+
+	return ValidEditingLOD;
+}
+
+
+bool UStaticMeshToolTarget::IsValid() const
+{
+	return IsValid(StaticMesh, EditingLOD);
+}
+
+bool UStaticMeshToolTarget::IsValid(const UStaticMesh* StaticMeshIn, EStaticMeshEditingLOD EditingLODIn)
+{
+	if (!StaticMeshIn || StaticMeshIn->IsPendingKillOrUnreachable() || !StaticMeshIn->IsValidLowLevel())
+	{
+		return false;
+	}
+
+	if (EditingLODIn == EStaticMeshEditingLOD::HiResSource)
+	{
+		if (StaticMeshIn->IsHiResMeshDescriptionValid() == false)
+		{
+			return false;
+		}
+	}
+	else if ((int32)EditingLODIn >= StaticMeshIn->GetNumSourceModels())
+	{
+		return false;
+	}
+
+	return true;
+}
+
+
+int32 UStaticMeshToolTarget::GetNumMaterials() const
+{
+	return ensure(IsValid()) ? StaticMesh->GetStaticMaterials().Num() : 0;
+}
+
+UMaterialInterface* UStaticMeshToolTarget::GetMaterial(int32 MaterialIndex) const
+{
+	return ensure(IsValid()) ? StaticMesh->GetMaterial(MaterialIndex) : nullptr;
+}
+
+void UStaticMeshToolTarget::GetMaterialSet(FComponentMaterialSet& MaterialSetOut, bool bPreferAssetMaterials) const
+{
+	if (!ensure(IsValid())) return;
+
+	GetMaterialSet(StaticMesh, MaterialSetOut, bPreferAssetMaterials);
+}
+
+void UStaticMeshToolTarget::GetMaterialSet(const UStaticMesh* StaticMeshIn, 
+	FComponentMaterialSet& MaterialSetOut, bool bPreferAssetMaterials)
+{
+	int32 NumMaterials = StaticMeshIn->GetStaticMaterials().Num();
+	MaterialSetOut.Materials.SetNum(NumMaterials);
+	for (int32 k = 0; k < NumMaterials; ++k)
+	{
+		MaterialSetOut.Materials[k] = StaticMeshIn->GetMaterial(k);
+	}
+}
+
+bool UStaticMeshToolTarget::CommitMaterialSetUpdate(const FComponentMaterialSet& MaterialSet, bool bApplyToAsset)
+{
+	if (!ensure(IsValid())) return false;
+
+	return CommitMaterialSetUpdate(StaticMesh, MaterialSet, bApplyToAsset);
+}
+
+bool UStaticMeshToolTarget::CommitMaterialSetUpdate(UStaticMesh* StaticMeshIn,
+	const FComponentMaterialSet& MaterialSet, bool bApplyToAsset)
+{
+	if (!bApplyToAsset)
+	{
+		return false;
+	}
+
+	if (StaticMeshIn->GetPathName().StartsWith(TEXT("/Engine/")))
+	{
+		UE_LOG(LogTemp, Warning, TEXT("CANNOT MODIFY BUILT-IN ENGINE ASSET %s"), *StaticMeshIn->GetPathName());
+		return false;
+	}
+
+	// filter out any Engine materials that we don't want to be permanently assigning
+	TArray<UMaterialInterface*> FilteredMaterials = MaterialSet.Materials;
+	for (int32 k = 0; k < FilteredMaterials.Num(); ++k)
+	{
+		FString AssetPath = FilteredMaterials[k]->GetPathName();
+		if (AssetPath.StartsWith(TEXT("/MeshModelingToolset/")))
+		{
+			FilteredMaterials[k] = UMaterial::GetDefaultMaterial(MD_Surface);
+		}
+	}
+
+	// flush any pending rendering commands, which might touch this component while we are rebuilding its mesh
+	FlushRenderingCommands();
+
+	// make sure transactional flag is on
+	StaticMeshIn->SetFlags(RF_Transactional);
+
+	StaticMeshIn->Modify();
+
+	int NewNumMaterials = FilteredMaterials.Num();
+	if (NewNumMaterials != StaticMeshIn->GetStaticMaterials().Num())
+	{
+		StaticMeshIn->GetStaticMaterials().SetNum(NewNumMaterials);
+	}
+	for (int k = 0; k < NewNumMaterials; ++k)
+	{
+		if (StaticMeshIn->GetMaterial(k) != FilteredMaterials[k])
+		{
+			StaticMeshIn->SetMaterial(k, FilteredMaterials[k]);
+		}
+	}
+
+	StaticMeshIn->PostEditChange();
+
+	return true;
+}
+
+FMeshDescription* UStaticMeshToolTarget::GetMeshDescription()
+{
+	if (ensure(IsValid()))
+	{
+		return (EditingLOD == EStaticMeshEditingLOD::HiResSource) ?
+			StaticMesh->GetHiResMeshDescription() : StaticMesh->GetMeshDescription((int32)EditingLOD);
+	}
+	return nullptr;
+}
+
+void UStaticMeshToolTarget::CommitMeshDescription(const FCommitter& Committer)
+{
+	if (ensure(IsValid()) == false) return;
+
+	CommitMeshDescription(StaticMesh, GetMeshDescription(), Committer, EditingLOD);
+}
+
+void UStaticMeshToolTarget::CommitMeshDescription(UStaticMesh* StaticMeshIn, FMeshDescription* MeshDescription,
+	const FCommitter& Committer, EStaticMeshEditingLOD EditingLODIn)
+{
+	using namespace StaticMeshToolTargetLocals;
+
+	if (StaticMeshIn->GetPathName().StartsWith(TEXT("/Engine/")))
+	{
+		DisplayCriticalWarningMessage(FString::Printf(TEXT("CANNOT MODIFY BUILT-IN ENGINE ASSET %s"), *StaticMeshIn->GetPathName()));
+		return;
+	}
+
+	// flush any pending rendering commands, which might touch this component while we are rebuilding it's mesh
+	FlushRenderingCommands();
+
+	// make sure transactional flag is on for this asset
+	StaticMeshIn->SetFlags(RF_Transactional);
+
+	verify(StaticMeshIn->Modify());
+
+	// disable auto-generated normals StaticMesh build setting
+	UE::MeshDescription::FStaticMeshBuildSettingChange SettingsChange;
+	SettingsChange.AutoGeneratedNormals = UE::MeshDescription::EBuildSettingBoolChange::Disable;
+	UE::MeshDescription::ConfigureBuildSettings(StaticMeshIn, 0, SettingsChange);
+
+	if (EditingLODIn == EStaticMeshEditingLOD::HiResSource)
+	{
+		verify(StaticMeshIn->ModifyHiResMeshDescription());
+	}
+	else
+	{
+		verify(StaticMeshIn->ModifyMeshDescription((int32)EditingLODIn));
+	}
+
+	FCommitterParams CommitterParams;
+	CommitterParams.MeshDescriptionOut = MeshDescription;
+
+	Committer(CommitterParams);
+
+	if (EditingLODIn == EStaticMeshEditingLOD::HiResSource)
+	{
+		StaticMeshIn->CommitHiResMeshDescription();
+	}
+	else
+	{
+		StaticMeshIn->CommitMeshDescription((int32)EditingLODIn);
+	}
+
+	StaticMeshIn->PostEditChange();
+}
+
+TSharedPtr<FDynamicMesh3, ESPMode::ThreadSafe> UStaticMeshToolTarget::GetDynamicMesh()
+{
+	return GetDynamicMeshViaMeshDescription(*this);
+}
+
+void UStaticMeshToolTarget::CommitDynamicMesh(const FDynamicMesh3& Mesh, const FDynamicMeshCommitInfo& CommitInfo)
+{
+	CommitDynamicMeshViaMeshDescription(*this, Mesh, CommitInfo);
+}
+
+UStaticMesh* UStaticMeshToolTarget::GetStaticMesh() const
+{
+	return IsValid() ? StaticMesh : nullptr;
+}
+
+
+// Factory
+
+bool UStaticMeshToolTargetFactory::CanBuildTarget(UObject* SourceObject, const FToolTargetTypeRequirements& Requirements) const
+{
+	const UStaticMesh* StaticMesh = Cast<UStaticMesh>(SourceObject);
+	return StaticMesh && !StaticMesh->IsPendingKillOrUnreachable() && StaticMesh->IsValidLowLevel()
+		&& (StaticMesh->GetNumSourceModels() > 0)
+		&& Requirements.AreSatisfiedBy(UStaticMeshToolTarget::StaticClass());
+}
+
+UToolTarget* UStaticMeshToolTargetFactory::BuildTarget(UObject* SourceObject, const FToolTargetTypeRequirements& Requirements)
+{
+	UStaticMeshToolTarget* Target = NewObject<UStaticMeshToolTarget>();// TODO: Should we set an outer here?
+	Target->StaticMesh = Cast<UStaticMesh>(SourceObject);
+	Target->SetEditingLOD(EditingLOD);
+	check(Target->StaticMesh && Requirements.AreSatisfiedBy(Target));
+
+	return Target;
+}
+
+
+void UStaticMeshToolTargetFactory::SetActiveEditingLOD(EStaticMeshEditingLOD NewEditingLOD)
+{
+	EditingLOD = NewEditingLOD;
+}
\ No newline at end of file
diff --git a/Engine/Plugins/Experimental/MeshModelingToolset/Source/ModelingComponentsEditorOnly/Private/ToolTargets/VolumeDynamicMeshToolTarget.cpp b/Engine/Plugins/Experimental/MeshModelingToolset/Source/ModelingComponentsEditorOnly/Private/ToolTargets/VolumeDynamicMeshToolTarget.cpp
--- a/Engine/Plugins/Experimental/MeshModelingToolset/Source/ModelingComponentsEditorOnly/Private/ToolTargets/VolumeDynamicMeshToolTarget.cpp
+++ b/Engine/Plugins/Experimental/MeshModelingToolset/Source/ModelingComponentsEditorOnly/Private/ToolTargets/VolumeDynamicMeshToolTarget.cpp
@@ -70,7 +70,7 @@ void UVolumeDynamicMeshToolTarget::GetMaterialSet(FComponentMaterialSet& Materia
 	}
 }
 
-TSharedPtr<UE::Geometry::FDynamicMesh3> UVolumeDynamicMeshToolTarget::GetDynamicMesh()
+TSharedPtr<UE::Geometry::FDynamicMesh3, ESPMode::ThreadSafe> UVolumeDynamicMeshToolTarget::GetDynamicMesh()
 {
 	UBrushComponent* BrushComponent = Cast<UBrushComponent>(Component);
 	if (!BrushComponent)
@@ -83,7 +83,7 @@ TSharedPtr<UE::Geometry::FDynamicMesh3> UVolumeDynamicMeshToolTarget::GetDynamic
 		return nullptr;
 	}
 
-	TSharedPtr<FDynamicMesh3> DynamicMesh = MakeShared<FDynamicMesh3>();
+	TSharedPtr<FDynamicMesh3, ESPMode::ThreadSafe> DynamicMesh = MakeShared<FDynamicMesh3, ESPMode::ThreadSafe>();
 	UE::Conversion::VolumeToDynamicMesh(Volume, *DynamicMesh, GetVolumeToMeshOptions());
 	FMeshNormals::InitializeMeshToPerTriangleNormals(DynamicMesh.Get());
 
diff --git a/Engine/Plugins/Experimental/MeshModelingToolset/Source/ModelingComponentsEditorOnly/Private/ToolTargets/VolumeMeshDescriptionToolTarget.cpp b/Engine/Plugins/Experimental/MeshModelingToolset/Source/ModelingComponentsEditorOnly/Private/ToolTargets/VolumeMeshDescriptionToolTarget.cpp
--- a/Engine/Plugins/Experimental/MeshModelingToolset/Source/ModelingComponentsEditorOnly/Private/ToolTargets/VolumeMeshDescriptionToolTarget.cpp
+++ b/Engine/Plugins/Experimental/MeshModelingToolset/Source/ModelingComponentsEditorOnly/Private/ToolTargets/VolumeMeshDescriptionToolTarget.cpp
@@ -26,13 +26,13 @@ FMeshDescription* UVolumeMeshDescriptionToolTarget::GetMeshDescription()
 		// Editor.h. However, that path doesn't assign polygroups to the result, which we
 		// typically want when using this target, hence the path through a dynamic mesh.
 
-		TSharedPtr<FDynamicMesh3> DynamicMesh = GetDynamicMesh();
+		TSharedPtr<FDynamicMesh3, ESPMode::ThreadSafe> DynamicMesh = GetDynamicMesh();
 		if (!DynamicMesh)
 		{
 			return nullptr;
 		}
 
-		ConvertedMeshDescription = MakeShared<FMeshDescription>();
+		ConvertedMeshDescription = MakeShared<FMeshDescription, ESPMode::ThreadSafe>();
 		FStaticMeshAttributes Attributes(*ConvertedMeshDescription);
 		Attributes.Register();
 		FDynamicMeshToMeshDescription Converter;
@@ -48,7 +48,7 @@ void UVolumeMeshDescriptionToolTarget::CommitMeshDescription(const FCommitter& C
 	// Let the user fill our mesh description with the Committer
 	if (!ConvertedMeshDescription.IsValid())
 	{
-		ConvertedMeshDescription = MakeShared<FMeshDescription>();
+		ConvertedMeshDescription = MakeShared<FMeshDescription, ESPMode::ThreadSafe>();
 		FStaticMeshAttributes Attributes(*ConvertedMeshDescription);
 		Attributes.Register();
 	}
diff --git a/Engine/Plugins/Experimental/MeshModelingToolset/Source/ModelingComponentsEditorOnly/Public/ToolTargets/SkeletalMeshComponentToolTarget.h b/Engine/Plugins/Experimental/MeshModelingToolset/Source/ModelingComponentsEditorOnly/Public/ToolTargets/SkeletalMeshComponentToolTarget.h
--- a/Engine/Plugins/Experimental/MeshModelingToolset/Source/ModelingComponentsEditorOnly/Public/ToolTargets/SkeletalMeshComponentToolTarget.h
+++ b/Engine/Plugins/Experimental/MeshModelingToolset/Source/ModelingComponentsEditorOnly/Public/ToolTargets/SkeletalMeshComponentToolTarget.h
@@ -48,7 +48,7 @@ public:
 	bool CommitMaterialSetUpdate(const FComponentMaterialSet& MaterialSet, bool bApplyToAsset) override;	
 
 	// IDynamicMeshProvider
-	virtual TSharedPtr<UE::Geometry::FDynamicMesh3> GetDynamicMesh() override;
+	virtual TSharedPtr<UE::Geometry::FDynamicMesh3, ESPMode::ThreadSafe> GetDynamicMesh() override;
 
 	// IDynamicMeshCommitter
 	virtual void CommitDynamicMesh(const UE::Geometry::FDynamicMesh3& Mesh, const FDynamicMeshCommitInfo& CommitInfo) override;
@@ -64,7 +64,7 @@ protected:
 private:
 	// Until USkeletalMesh stores its internal representation as FMeshDescription, we need to
 	// retain the storage here to cover the lifetime of the pointer returned by GetMeshDescription(). 
-	TUniquePtr<FMeshDescription> CachedMeshDescription;	
+	TUniquePtr<FMeshDescription> CachedMeshDescription;
 };
 
 /** Factory for USkeletalMeshComponentToolTarget to be used by the target manager. */
diff --git a/Engine/Plugins/Experimental/MeshModelingToolset/Source/ModelingComponentsEditorOnly/Public/ToolTargets/SkeletalMeshToolTarget.h b/Engine/Plugins/Experimental/MeshModelingToolset/Source/ModelingComponentsEditorOnly/Public/ToolTargets/SkeletalMeshToolTarget.h
new file mode 100644
--- /dev/null
+++ b/Engine/Plugins/Experimental/MeshModelingToolset/Source/ModelingComponentsEditorOnly/Public/ToolTargets/SkeletalMeshToolTarget.h
@@ -0,0 +1,97 @@
+// Copyright Epic Games, Inc. All Rights Reserved.
+
+#pragma once
+
+#include "CoreMinimal.h"
+
+#include "TargetInterfaces/DynamicMeshCommitter.h"
+#include "TargetInterfaces/DynamicMeshProvider.h"
+#include "TargetInterfaces/MaterialProvider.h"
+#include "TargetInterfaces/MeshDescriptionCommitter.h"
+#include "TargetInterfaces/MeshDescriptionProvider.h"
+#include "TargetInterfaces/SkeletalMeshBackedTarget.h"
+#include "ToolTargets/PrimitiveComponentToolTarget.h"
+
+#include "SkeletalMeshToolTarget.generated.h"
+
+
+class USkeletalMesh;
+
+/**
+ * A tool target backed by a skeletal mesh.
+ */
+UCLASS(Transient)
+class MODELINGCOMPONENTSEDITORONLY_API USkeletalMeshToolTarget :
+	public UToolTarget,
+	public IMeshDescriptionCommitter,
+	public IMeshDescriptionProvider,
+	public IDynamicMeshProvider, 
+	public IDynamicMeshCommitter,
+	public IMaterialProvider,
+	public ISkeletalMeshBackedTarget
+{
+	GENERATED_BODY()
+
+public:
+	// UToolTarget
+	virtual bool IsValid() const override;
+
+	// IMeshDescriptionProvider implementation
+	FMeshDescription* GetMeshDescription() override;
+
+	// IMeshDescritpionCommitter implementation
+	void CommitMeshDescription(const FCommitter& Committer) override;
+	using IMeshDescriptionCommitter::CommitMeshDescription; // unhide the other overload
+
+	// IMaterialProvider implementation
+	int32 GetNumMaterials() const override;
+	UMaterialInterface* GetMaterial(int32 MaterialIndex) const override;
+	void GetMaterialSet(FComponentMaterialSet& MaterialSetOut, bool bPreferAssetMaterials) const override;
+	bool CommitMaterialSetUpdate(const FComponentMaterialSet& MaterialSet, bool bApplyToAsset) override;	
+
+	// IDynamicMeshProvider
+	virtual TSharedPtr<UE::Geometry::FDynamicMesh3, ESPMode::ThreadSafe> GetDynamicMesh() override;
+
+	// IDynamicMeshCommitter
+	virtual void CommitDynamicMesh(const UE::Geometry::FDynamicMesh3& Mesh, 
+		const FDynamicMeshCommitInfo& CommitInfo) override;
+	using IDynamicMeshCommitter::CommitDynamicMesh; // unhide the other overload
+
+	// ISkeletalMeshBackedTarget implementation
+	USkeletalMesh* GetSkeletalMesh() const override;
+
+protected:
+	USkeletalMesh* SkeletalMesh = nullptr;
+
+	// So that the tool target factory can poke into Component.
+	friend class USkeletalMeshToolTargetFactory;
+
+	
+	friend class USkeletalMeshComponentToolTarget;
+
+	static void GetMeshDescription(const USkeletalMesh* SkeletalMesh, FMeshDescription& MeshDescriptionOut);
+	static void CommitMeshDescription(USkeletalMesh* SkeletalMesh, 
+		FMeshDescription* MeshDescription, const FCommitter& Committer);
+	static void GetMaterialSet(const USkeletalMesh* SkeletalMesh, FComponentMaterialSet& MaterialSetOut,
+		bool bPreferAssetMaterials);
+	static bool CommitMaterialSetUpdate(USkeletalMesh* SkeletalMesh,
+		const FComponentMaterialSet& MaterialSet, bool bApplyToAsset);
+	
+private:
+	// Until USkeletalMesh stores its internal representation as FMeshDescription, we need to
+	// retain the storage here to cover the lifetime of the pointer returned by GetMeshDescription(). 
+	TUniquePtr<FMeshDescription> CachedMeshDescription;	
+};
+
+/** Factory for USkeletalMeshToolTarget to be used by the target manager. */
+UCLASS(Transient)
+class MODELINGCOMPONENTSEDITORONLY_API USkeletalMeshToolTargetFactory : public UToolTargetFactory
+{
+	GENERATED_BODY()
+
+public:
+
+	bool CanBuildTarget(UObject* SourceObject, const FToolTargetTypeRequirements& TargetTypeInfo) const override;
+
+	UToolTarget* BuildTarget(UObject* SourceObject, const FToolTargetTypeRequirements& TargetTypeInfo) override;
+};
diff --git a/Engine/Plugins/Experimental/MeshModelingToolset/Source/ModelingComponentsEditorOnly/Public/ToolTargets/StaticMeshComponentToolTarget.h b/Engine/Plugins/Experimental/MeshModelingToolset/Source/ModelingComponentsEditorOnly/Public/ToolTargets/StaticMeshComponentToolTarget.h
--- a/Engine/Plugins/Experimental/MeshModelingToolset/Source/ModelingComponentsEditorOnly/Public/ToolTargets/StaticMeshComponentToolTarget.h
+++ b/Engine/Plugins/Experimental/MeshModelingToolset/Source/ModelingComponentsEditorOnly/Public/ToolTargets/StaticMeshComponentToolTarget.h
@@ -51,7 +51,7 @@ public:
 	using IMeshDescriptionCommitter::CommitMeshDescription; // unhide the other overload
 
 	// IDynamicMeshProvider
-	virtual TSharedPtr<UE::Geometry::FDynamicMesh3> GetDynamicMesh() override;
+	virtual TSharedPtr<UE::Geometry::FDynamicMesh3, ESPMode::ThreadSafe> GetDynamicMesh() override;
 
 	// IDynamicMeshCommitter
 	virtual void CommitDynamicMesh(const UE::Geometry::FDynamicMesh3& Mesh, const FDynamicMeshCommitInfo& CommitInfo) override;
diff --git a/Engine/Plugins/Experimental/MeshModelingToolset/Source/ModelingComponentsEditorOnly/Public/ToolTargets/StaticMeshToolTarget.h b/Engine/Plugins/Experimental/MeshModelingToolset/Source/ModelingComponentsEditorOnly/Public/ToolTargets/StaticMeshToolTarget.h
new file mode 100644
--- /dev/null
+++ b/Engine/Plugins/Experimental/MeshModelingToolset/Source/ModelingComponentsEditorOnly/Public/ToolTargets/StaticMeshToolTarget.h
@@ -0,0 +1,116 @@
+// Copyright Epic Games, Inc. All Rights Reserved.
+
+#pragma once
+
+#include "CoreMinimal.h"
+
+#include "TargetInterfaces/DynamicMeshCommitter.h"
+#include "TargetInterfaces/DynamicMeshProvider.h"
+#include "TargetInterfaces/MaterialProvider.h"
+#include "TargetInterfaces/MeshDescriptionCommitter.h"
+#include "TargetInterfaces/MeshDescriptionProvider.h"
+#include "TargetInterfaces/StaticMeshBackedTarget.h"
+#include "ToolTargets/PrimitiveComponentToolTarget.h"
+#include "ComponentSourceInterfaces.h"  // for EStaticMeshEditingLOD
+
+#include "StaticMeshToolTarget.generated.h"
+
+class UStaticMesh;
+
+/**
+ * A tool target backed by a static mesh component that can provide and take a mesh
+ * description.
+ */
+UCLASS(Transient)
+class MODELINGCOMPONENTSEDITORONLY_API UStaticMeshToolTarget : 
+	public UToolTarget,
+	public IMeshDescriptionCommitter, 
+	public IMeshDescriptionProvider, 
+	public IMaterialProvider, 
+	public IStaticMeshBackedTarget,
+	public IDynamicMeshProvider, 
+	public IDynamicMeshCommitter
+{
+	GENERATED_BODY()
+
+public:
+	/**
+	 * Configure active LOD to edit. Can only call this after Component is configured in base UPrimitiveComponentToolTarget.
+	 * If requested LOD does not exist, fall back to one that does.
+	 */
+	virtual void SetEditingLOD(EStaticMeshEditingLOD RequestedEditingLOD);
+
+	/** @return current editing LOD */
+	virtual EStaticMeshEditingLOD GetEditingLOD() const { return EditingLOD; }
+
+	// UToolTarget
+	virtual bool IsValid() const override;
+
+	// IMeshDescriptionProvider implementation
+	FMeshDescription* GetMeshDescription() override;
+
+	// IMeshDescritpionCommitter implementation
+	virtual void CommitMeshDescription(const FCommitter& Committer) override;
+	using IMeshDescriptionCommitter::CommitMeshDescription; // unhide the other overload
+
+	// IDynamicMeshProvider
+	virtual TSharedPtr<UE::Geometry::FDynamicMesh3, ESPMode::ThreadSafe> GetDynamicMesh() override;
+
+	// IDynamicMeshCommitter
+	virtual void CommitDynamicMesh(const UE::Geometry::FDynamicMesh3& Mesh, const FDynamicMeshCommitInfo& CommitInfo) override;
+	using IDynamicMeshCommitter::CommitDynamicMesh; // unhide the other overload
+
+	// IMaterialProvider implementation
+	int32 GetNumMaterials() const override;
+	UMaterialInterface* GetMaterial(int32 MaterialIndex) const override;
+	void GetMaterialSet(FComponentMaterialSet& MaterialSetOut, bool bPreferAssetMaterials) const override;
+	virtual bool CommitMaterialSetUpdate(const FComponentMaterialSet& MaterialSet, bool bApplyToAsset) override;	
+
+	// IStaticMeshBackedTarget
+	UStaticMesh* GetStaticMesh() const override;
+
+	// Rest provided by parent class
+
+protected:
+	UStaticMesh* StaticMesh = nullptr;
+
+	EStaticMeshEditingLOD EditingLOD = EStaticMeshEditingLOD::LOD0;
+
+	friend class UStaticMeshToolTargetFactory;
+
+	friend class UStaticMeshComponentToolTarget;
+
+	static bool IsValid(const UStaticMesh* StaticMesh, EStaticMeshEditingLOD EditingLOD);
+	static EStaticMeshEditingLOD GetValidEditingLOD(const UStaticMesh* StaticMesh, 
+		EStaticMeshEditingLOD RequestedEditingLOD);
+	static void CommitMeshDescription(UStaticMesh* SkeletalMesh, FMeshDescription* MeshDescription,
+		const FCommitter& Committer, EStaticMeshEditingLOD EditingLODIn);
+	static void GetMaterialSet(const UStaticMesh* SkeletalMesh, 
+		FComponentMaterialSet& MaterialSetOut, bool bPreferAssetMaterials);
+	static bool CommitMaterialSetUpdate(UStaticMesh* SkeletalMesh, 
+		const FComponentMaterialSet& MaterialSet, bool bApplyToAsset);
+};
+
+
+/** Factory for UStaticMeshToolTarget to be used by the target manager. */
+UCLASS(Transient)
+class MODELINGCOMPONENTSEDITORONLY_API UStaticMeshToolTargetFactory : public UToolTargetFactory
+{
+	GENERATED_BODY()
+
+public:
+
+	virtual bool CanBuildTarget(UObject* SourceObject, const FToolTargetTypeRequirements& TargetTypeInfo) const override;
+
+	virtual UToolTarget* BuildTarget(UObject* SourceObject, const FToolTargetTypeRequirements& TargetTypeInfo) override;
+
+
+
+public:
+	virtual EStaticMeshEditingLOD GetActiveEditingLOD() const { return EditingLOD; }
+	virtual void SetActiveEditingLOD(EStaticMeshEditingLOD NewEditingLOD);
+
+protected:
+	// LOD to edit, default is to edit LOD0
+	EStaticMeshEditingLOD EditingLOD = EStaticMeshEditingLOD::LOD0;
+};
\ No newline at end of file
diff --git a/Engine/Plugins/Experimental/MeshModelingToolset/Source/ModelingComponentsEditorOnly/Public/ToolTargets/VolumeDynamicMeshToolTarget.h b/Engine/Plugins/Experimental/MeshModelingToolset/Source/ModelingComponentsEditorOnly/Public/ToolTargets/VolumeDynamicMeshToolTarget.h
--- a/Engine/Plugins/Experimental/MeshModelingToolset/Source/ModelingComponentsEditorOnly/Public/ToolTargets/VolumeDynamicMeshToolTarget.h
+++ b/Engine/Plugins/Experimental/MeshModelingToolset/Source/ModelingComponentsEditorOnly/Public/ToolTargets/VolumeDynamicMeshToolTarget.h
@@ -39,7 +39,7 @@ public:
 	const UE::Conversion::FVolumeToMeshOptions& GetVolumeToMeshOptions() { return VolumeToMeshOptions; }
 
 	// IDynamicMeshProvider implementation
-	virtual TSharedPtr<UE::Geometry::FDynamicMesh3> GetDynamicMesh() override;
+	virtual TSharedPtr<UE::Geometry::FDynamicMesh3, ESPMode::ThreadSafe> GetDynamicMesh() override;
 
 	// IDynamicMeshCommitter implementation
 	virtual void CommitDynamicMesh(const UE::Geometry::FDynamicMesh3& Mesh, const FDynamicMeshCommitInfo&) override;
@@ -57,7 +57,7 @@ public:
 	// Rest provided by parent class
 
 protected:
-	TSharedPtr<FMeshDescription> ConvertedMeshDescription;
+	TSharedPtr<FMeshDescription, ESPMode::ThreadSafe> ConvertedMeshDescription;
 	UE::Conversion::FVolumeToMeshOptions VolumeToMeshOptions;
 
 	friend class UVolumeDynamicMeshToolTargetFactory;
diff --git a/Engine/Plugins/Experimental/MeshModelingToolset/Source/ModelingComponentsEditorOnly/Public/ToolTargets/VolumeMeshDescriptionToolTarget.h b/Engine/Plugins/Experimental/MeshModelingToolset/Source/ModelingComponentsEditorOnly/Public/ToolTargets/VolumeMeshDescriptionToolTarget.h
--- a/Engine/Plugins/Experimental/MeshModelingToolset/Source/ModelingComponentsEditorOnly/Public/ToolTargets/VolumeMeshDescriptionToolTarget.h
+++ b/Engine/Plugins/Experimental/MeshModelingToolset/Source/ModelingComponentsEditorOnly/Public/ToolTargets/VolumeMeshDescriptionToolTarget.h
@@ -31,7 +31,7 @@ public:
 protected:
 	// This isn't for caching- we have to take ownership of the mesh description because it is
 	// expected for things like a static mesh.
-	TSharedPtr<FMeshDescription> ConvertedMeshDescription;
+	TSharedPtr<FMeshDescription, ESPMode::ThreadSafe> ConvertedMeshDescription;
 
 	friend class UVolumeMeshDescriptionToolTargetFactory;
 };
