ModelingComponents: Add UDynamicMesh class, wraps FDynamicMesh3 in a UObject. Use as mesh storage inside SimpleDynamicMeshComponent and OctreeDynamicMeshComponent. Clean up SimpleDynamicMeshComponent and BaseDynamicMeshComponent API and headers.
#rnx

diff --git a/Engine/Plugins/Experimental/MeshModelingToolset/Source/MeshModelingTools/Private/MeshGroupPaintTool.cpp b/Engine/Plugins/Experimental/MeshModelingToolset/Source/MeshModelingTools/Private/MeshGroupPaintTool.cpp
--- a/Engine/Plugins/Experimental/MeshModelingToolset/Source/MeshModelingTools/Private/MeshGroupPaintTool.cpp
+++ b/Engine/Plugins/Experimental/MeshModelingToolset/Source/MeshModelingTools/Private/MeshGroupPaintTool.cpp
@@ -102,7 +102,7 @@ void UMeshGroupPaintTool::Setup()
 		DynamicMeshComponent->SetMaterial(k, MaterialSet.Materials[k]);
 	}
 
-	DynamicMeshComponent->bInvalidateProxyOnChange = false;
+	DynamicMeshComponent->SetInvalidateProxyOnChangeEnabled(false);
 	OnDynamicMeshComponentChangedHandle = DynamicMeshComponent->OnMeshVerticesChanged.AddUObject(this, &UMeshGroupPaintTool::OnDynamicMeshComponentChanged);
 
 	FDynamicMesh3* Mesh = GetSculptMesh();
@@ -223,11 +223,10 @@ void UMeshGroupPaintTool::Setup()
 	MeshElementsDisplay->SetMeshAccessFunction([&](void) { return GetSculptMesh(); });
 
 	// force colors update... ?
-	DynamicMeshComponent->TriangleColorFunc = [this](const FDynamicMesh3* Mesh, int TriangleID)
+	DynamicMeshComponent->SetTriangleColorFunction([this](const FDynamicMesh3* Mesh, int TriangleID)
 	{
 		return GetColorForGroup(ActiveGroupSet->GetGroup(TriangleID));
-	};
-	DynamicMeshComponent->FastNotifyColorsUpdated();
+	});
 
 	// disable view properties
 	SetViewPropertiesEnabled(false);
diff --git a/Engine/Plugins/Experimental/MeshModelingToolset/Source/MeshModelingTools/Private/MeshVertexSculptTool.cpp b/Engine/Plugins/Experimental/MeshModelingToolset/Source/MeshModelingTools/Private/MeshVertexSculptTool.cpp
--- a/Engine/Plugins/Experimental/MeshModelingToolset/Source/MeshModelingTools/Private/MeshVertexSculptTool.cpp
+++ b/Engine/Plugins/Experimental/MeshModelingToolset/Source/MeshModelingTools/Private/MeshVertexSculptTool.cpp
@@ -75,7 +75,7 @@ void UMeshVertexSculptTool::Setup()
 		DynamicMeshComponent->SetMaterial(k, MaterialSet.Materials[k]);
 	}
 
-	DynamicMeshComponent->bInvalidateProxyOnChange = false;
+	DynamicMeshComponent->SetInvalidateProxyOnChangeEnabled(false);
 	OnDynamicMeshComponentChangedHandle = DynamicMeshComponent->OnMeshVerticesChanged.AddUObject(this, &UMeshVertexSculptTool::OnDynamicMeshComponentChanged);
 
 	FDynamicMesh3* SculptMesh = GetSculptMesh();
diff --git a/Engine/Plugins/Experimental/MeshModelingToolset/Source/ModelingComponents/Private/BaseDynamicMeshSceneProxy.h b/Engine/Plugins/Experimental/MeshModelingToolset/Source/ModelingComponents/Private/BaseDynamicMeshSceneProxy.h
--- a/Engine/Plugins/Experimental/MeshModelingToolset/Source/ModelingComponents/Private/BaseDynamicMeshSceneProxy.h
+++ b/Engine/Plugins/Experimental/MeshModelingToolset/Source/ModelingComponents/Private/BaseDynamicMeshSceneProxy.h
@@ -809,7 +809,7 @@ public:
 		QUICK_SCOPE_CYCLE_COUNTER(STAT_OctreeDynamicMeshSceneProxy_GetDynamicMeshElements);
 
 		const bool bWireframe = (AllowDebugViewmodes() && ViewFamily.EngineShowFlags.Wireframe)
-			|| ParentBaseComponent->EnableWireframeRenderPass();
+			|| ParentBaseComponent->GetEnableWireframeRenderPass();
 
 		// set up wireframe material. Probably bad to reference GEngine here...also this material is very bad?
 		FMaterialRenderProxy* WireframeMaterialProxy = nullptr;
diff --git a/Engine/Plugins/Experimental/MeshModelingToolset/Source/ModelingComponents/Private/OctreeDynamicMeshComponent.cpp b/Engine/Plugins/Experimental/MeshModelingToolset/Source/ModelingComponents/Private/OctreeDynamicMeshComponent.cpp
--- a/Engine/Plugins/Experimental/MeshModelingToolset/Source/ModelingComponents/Private/OctreeDynamicMeshComponent.cpp
+++ b/Engine/Plugins/Experimental/MeshModelingToolset/Source/ModelingComponents/Private/OctreeDynamicMeshComponent.cpp
@@ -43,18 +43,25 @@ UOctreeDynamicMeshComponent::UOctreeDynamicMeshComponent(const FObjectInitialize
 
 	SetCollisionProfileName(UCollisionProfile::NoCollision_ProfileName);
 
-	InitializeNewMesh();
+	MeshObject = CreateDefaultSubobject<UDynamicMesh>(TEXT("DynamicMesh"));
+	MeshObjectChangedHandle = MeshObject->OnMeshChanged().AddUObject(this, &UOctreeDynamicMeshComponent::OnMeshObjectChanged);
+	PreMeshChangeHandle = MeshObject->OnPreMeshChanged().AddUObject(this, &UOctreeDynamicMeshComponent::OnPreMeshObjectChanged);
+	
+
+	Octree = MakeUnique<FDynamicMeshOctree3>();
+	Octree->Initialize(GetMesh());
+	OctreeCut = MakeUnique<FDynamicMeshOctree3::FTreeCutSet>();
 }
 
 
 
-void UOctreeDynamicMeshComponent::InitializeMesh(FMeshDescription* MeshDescription)
+void UOctreeDynamicMeshComponent::InitializeMesh(const FMeshDescription* MeshDescription)
 {
 	FMeshDescriptionToDynamicMesh Converter;
-	Mesh->Clear();
-	Converter.Convert(MeshDescription, *Mesh);
+	GetMesh()->Clear();
+	Converter.Convert(MeshDescription, *GetMesh());
 
-	FAxisAlignedBox3d MeshBounds = Mesh->GetCachedBounds();
+	FAxisAlignedBox3d MeshBounds = GetMesh()->GetCachedBounds();
 	Octree = MakeUnique<FDynamicMeshOctree3>();
 	Octree->RootDimension = MeshBounds.MaxDim() * 0.25;
 
@@ -65,34 +72,12 @@ void UOctreeDynamicMeshComponent::InitializeMesh(FMeshDescription* MeshDescripti
 	//UE_LOG(LogTemp, Warning, TEXT("OctreeStats %s"), *Stats.ToString());
 
 	OctreeCut = MakeUnique<FDynamicMeshOctree3::FTreeCutSet>();
+	CutCellSetMap.Reset();
 
 	NotifyMeshUpdated();
 }
 
 
-TUniquePtr<FDynamicMesh3> UOctreeDynamicMeshComponent::ExtractMesh(bool bNotifyUpdate)
-{
-	TUniquePtr<FDynamicMesh3> CurMesh = MoveTemp(Mesh);
-	InitializeNewMesh();
-	if (bNotifyUpdate)
-	{
-		NotifyMeshUpdated();
-	}
-	return CurMesh;
-}
-
-
-void UOctreeDynamicMeshComponent::InitializeNewMesh()
-{
-	Mesh = MakeUnique<FDynamicMesh3>();
-	// discard any attributes/etc initialized by default
-	Mesh->Clear();
-	Octree = MakeUnique<FDynamicMeshOctree3>();
-	Octree->Initialize(GetMesh());
-	OctreeCut = MakeUnique<FDynamicMeshOctree3::FTreeCutSet>();
-}
-
-
 
 void UOctreeDynamicMeshComponent::ApplyTransform(const UE::Geometry::FTransform3d& Transform, bool bInvert)
 {
@@ -114,11 +99,12 @@ void UOctreeDynamicMeshComponent::ApplyTransform(const UE::Geometry::FTransform3
 	}
 	else
 	{
-		FAxisAlignedBox3d MeshBounds = Mesh->GetCachedBounds();
+		FAxisAlignedBox3d MeshBounds = GetMesh()->GetCachedBounds();
 		Octree = MakeUnique<FDynamicMeshOctree3>();
 		Octree->RootDimension = MeshBounds.MaxDim() * 0.25;
 		Octree->Initialize(GetMesh());
 		OctreeCut = MakeUnique<FDynamicMeshOctree3::FTreeCutSet>();
+		CutCellSetMap.Reset();
 	}
 
 }
@@ -127,13 +113,13 @@ void UOctreeDynamicMeshComponent::ApplyTransform(const UE::Geometry::FTransform3
 void UOctreeDynamicMeshComponent::Bake(FMeshDescription* MeshDescription, bool bHaveModifiedTopology, const FConversionToMeshDescriptionOptions& ConversionOptions)
 {
 	FDynamicMeshToMeshDescription Converter(ConversionOptions);
-	if (bHaveModifiedTopology == false && Converter.HaveMatchingElementCounts(Mesh.Get(), MeshDescription))
+	if (bHaveModifiedTopology == false && Converter.HaveMatchingElementCounts(GetMesh(), MeshDescription))
 	{
-		Converter.Update(Mesh.Get(), *MeshDescription);
+		Converter.Update(GetMesh(), *MeshDescription);
 	}
 	else
 	{
-		Converter.Convert(Mesh.Get(), *MeshDescription);
+		Converter.Convert(GetMesh(), *MeshDescription);
 
 		//UE_LOG(LogTemp, Warning, TEXT("MeshDescription has %d instances"), MeshDescription->VertexInstances().Num());
 	}
@@ -265,7 +251,7 @@ FPrimitiveSceneProxy* UOctreeDynamicMeshComponent::CreateSceneProxy()
 	check(GetCurrentSceneProxy() == nullptr);
 
 	FOctreeDynamicMeshSceneProxy* NewProxy = nullptr;
-	if (Mesh->TriangleCount() > 0)
+	if (GetMesh()->TriangleCount() > 0)
 	{
 		NewProxy = new FOctreeDynamicMeshSceneProxy(this);
 
@@ -275,15 +261,14 @@ FPrimitiveSceneProxy* UOctreeDynamicMeshComponent::CreateSceneProxy()
 			NewProxy->PerTriangleColorFunc = [this](const FDynamicMesh3* MeshIn, int TriangleID) { return GetTriangleColor(TriangleID); };
 		}
 
-
 		OctreeCut = MakeUnique<FDynamicMeshOctree3::FTreeCutSet>();
-		InitializeOctreeCutSet(*Mesh, *Octree, OctreeCut);
+		CutCellSetMap.Reset();
+		InitializeOctreeCutSet(*GetMesh(), *Octree, OctreeCut);
 
 		TriangleDecomposition = FArrayIndexSetsDecomposition();
 
 		SpillDecompSetID = TriangleDecomposition.CreateNewIndexSet();
 
-		CutCellSetMap.Reset();
 		for (auto CellRef : OctreeCut->CutCells)
 		{
 			int32 IndexSetID = TriangleDecomposition.CreateNewIndexSet();
@@ -331,7 +316,7 @@ FColor UOctreeDynamicMeshComponent::GetTriangleColor(int TriangleID)
 {
 	if (TriangleColorFunc != nullptr)
 	{
-		return TriangleColorFunc(Mesh.Get(), TriangleID);
+		return TriangleColorFunc(GetMesh(), TriangleID);
 	}
 	else
 	{
@@ -345,7 +330,7 @@ FBoxSphereBounds UOctreeDynamicMeshComponent::CalcBounds(const FTransform& Local
 {
 	// Bounds are tighter if the box is generated from transformed vertices.
 	FBox BoundingBox(ForceInit);
-	for ( FVector3d Vertex : Mesh->VerticesItr() ) 
+	for ( FVector3d Vertex : GetMesh()->VerticesItr() ) 
 	{
 		BoundingBox += LocalToWorld.TransformPosition((FVector)Vertex);
 	}
@@ -356,7 +341,15 @@ FBoxSphereBounds UOctreeDynamicMeshComponent::CalcBounds(const FTransform& Local
 
 void UOctreeDynamicMeshComponent::ApplyChange(const FMeshVertexChange* Change, bool bRevert)
 {
+	// This function is not currently used in the codebase, and needs to be restructured
+	// to work properly with passing the change down to the UDynamicMesh
+	check(false);
+
+/*
 	Octree->ResetModifiedBounds();
+
+	const FDynamicMesh3* Mesh = GetMesh();
+
 	TSet<int> TrianglesToUpdate;
 	auto NotifyTriVerticesAffected = [&](int32 vid) 
 	{
@@ -401,6 +394,7 @@ void UOctreeDynamicMeshComponent::ApplyChange(const FMeshVertexChange* Change, b
 	//NotifyMeshUpdated();
 
 	OnMeshChanged.Broadcast();
+*/
 }
 
 
@@ -408,17 +402,7 @@ void UOctreeDynamicMeshComponent::ApplyChange(const FMeshVertexChange* Change, b
 
 void UOctreeDynamicMeshComponent::ApplyChange(const FMeshChange* Change, bool bRevert)
 {
-	TArray<int> RemoveTriangles, AddTriangles;
-	bool bRemoveOld = ! bRevert;
-	Change->DynamicMeshChange->GetSavedTriangleList(RemoveTriangles, bRemoveOld);
-	Change->DynamicMeshChange->GetSavedTriangleList(AddTriangles, !bRemoveOld);
-
-	Octree->ResetModifiedBounds();
-	Octree->RemoveTriangles(RemoveTriangles);
-
-	Change->DynamicMeshChange->Apply(Mesh.Get(), bRevert);
-
-	Octree->InsertTriangles(AddTriangles);
+	MeshObject->ApplyChange(Change, bRevert);
 
 	//NotifyMeshUpdated();
 
@@ -429,15 +413,72 @@ void UOctreeDynamicMeshComponent::ApplyChange(const FMeshChange* Change, bool bR
 void UOctreeDynamicMeshComponent::ApplyChange(const FMeshReplacementChange* Change, bool bRevert)
 {
 	// full clear and reset
-	Mesh->Clear();
-	Mesh->Copy(*Change->GetMesh(bRevert));
+	MeshObject->ApplyChange(Change, bRevert);
+
 	Octree = MakeUnique<FDynamicMeshOctree3>();
 	Octree->Initialize(GetMesh());
 	
 	OctreeCut = MakeUnique<FDynamicMeshOctree3::FTreeCutSet>();
+	CutCellSetMap.Reset();
 
 	// need to force proxy recreation because we changed the Octree pointer...
 	MarkRenderStateDirty();
 
 	OnMeshChanged.Broadcast();
 }
+
+
+void UOctreeDynamicMeshComponent::OnPreMeshObjectChanged(UDynamicMesh* ChangedMeshObject, FDynamicMeshChangeInfo ChangeInfo)
+{
+	TArray<int> RemoveTriangles;
+	if (ChangeInfo.Type == EDynamicMeshChangeType::MeshChange && ensure(ChangeInfo.MeshChange != nullptr))
+	{
+		ChangeInfo.MeshChange->DynamicMeshChange->GetSavedTriangleList(RemoveTriangles, !ChangeInfo.bIsRevertChange);
+	}
+	if (RemoveTriangles.Num() > 0)
+	{
+		Octree->ResetModifiedBounds();
+		Octree->RemoveTriangles(RemoveTriangles);
+	}
+}
+
+
+void UOctreeDynamicMeshComponent::OnMeshObjectChanged(UDynamicMesh* ChangedMeshObject, FDynamicMeshChangeInfo ChangeInfo)
+{
+	if (ChangeInfo.Type == EDynamicMeshChangeType::MeshChange && ensure(ChangeInfo.MeshChange != nullptr))
+	{
+		TArray<int> AddTriangles;
+		ChangeInfo.MeshChange->DynamicMeshChange->GetSavedTriangleList(AddTriangles, ChangeInfo.bIsRevertChange);
+		Octree->InsertTriangles(AddTriangles);
+	}
+
+	NotifyMeshUpdated();
+	OnMeshChanged.Broadcast();
+}
+
+
+void UOctreeDynamicMeshComponent::SetDynamicMesh(UDynamicMesh* NewMesh)
+{
+	if (ensure(NewMesh) == false)
+	{
+		return;
+	}
+
+	if (ensure(MeshObject))
+	{
+		MeshObject->OnMeshChanged().Remove(MeshObjectChangedHandle);
+		MeshObject->OnPreMeshChanged().Remove(PreMeshChangeHandle);
+	}
+
+	MeshObject = NewMesh;
+	MeshObjectChangedHandle = MeshObject->OnMeshChanged().AddUObject(this, &UOctreeDynamicMeshComponent::OnMeshObjectChanged);
+	PreMeshChangeHandle = MeshObject->OnPreMeshChanged().AddUObject(this, &UOctreeDynamicMeshComponent::OnPreMeshObjectChanged);
+
+	Octree = MakeUnique<FDynamicMeshOctree3>();
+	Octree->Initialize(GetMesh());
+	OctreeCut = MakeUnique<FDynamicMeshOctree3::FTreeCutSet>();
+	CutCellSetMap.Reset();
+
+	NotifyMeshUpdated();
+	OnMeshChanged.Broadcast();
+}
\ No newline at end of file
diff --git a/Engine/Plugins/Experimental/MeshModelingToolset/Source/ModelingComponents/Private/PreviewMesh.cpp b/Engine/Plugins/Experimental/MeshModelingToolset/Source/ModelingComponents/Private/PreviewMesh.cpp
--- a/Engine/Plugins/Experimental/MeshModelingToolset/Source/ModelingComponents/Private/PreviewMesh.cpp
+++ b/Engine/Plugins/Experimental/MeshModelingToolset/Source/ModelingComponents/Private/PreviewMesh.cpp
@@ -21,7 +21,6 @@ APreviewMeshActor::APreviewMeshActor()
 UPreviewMesh::UPreviewMesh()
 {
 	bBuildSpatialDataStructure = false;
-	bDrawOnTop = false;
 }
 
 UPreviewMesh::~UPreviewMesh()
@@ -242,8 +241,6 @@ void UPreviewMesh::ClearPreview()
 void UPreviewMesh::UpdatePreview(const FDynamicMesh3* Mesh, ERenderUpdateMode UpdateMode,
 	EMeshRenderAttributeFlags ModifiedAttribs)
 {
-	DynamicMeshComponent->SetDrawOnTop(this->bDrawOnTop);
-
 	DynamicMeshComponent->GetMesh()->Copy(*Mesh);
 
 	NotifyDeferredEditCompleted(UpdateMode, ModifiedAttribs, bBuildSpatialDataStructure);
@@ -252,8 +249,6 @@ void UPreviewMesh::UpdatePreview(const FDynamicMesh3* Mesh, ERenderUpdateMode Up
 void UPreviewMesh::UpdatePreview(FDynamicMesh3&& Mesh, ERenderUpdateMode UpdateMode, 
 	EMeshRenderAttributeFlags ModifiedAttribs)
 {
-	DynamicMeshComponent->SetDrawOnTop(this->bDrawOnTop);
-
 	*(DynamicMeshComponent->GetMesh()) = MoveTemp(Mesh);
 
 	NotifyDeferredEditCompleted(UpdateMode, ModifiedAttribs, bBuildSpatialDataStructure);
@@ -288,7 +283,7 @@ TUniquePtr<FDynamicMesh3> UPreviewMesh::ExtractPreviewMesh() const
 {
 	if (DynamicMeshComponent != nullptr)
 	{
-		return DynamicMeshComponent->ExtractMesh(true);
+		return DynamicMeshComponent->GetDynamicMesh()->ExtractMesh();
 	}
 	return nullptr;
 }
@@ -545,31 +540,12 @@ void UPreviewMesh::Bake(FMeshDescription* MeshDescription, bool bHaveModifiedTop
 
 void UPreviewMesh::SetTriangleColorFunction(TFunction<FColor(const FDynamicMesh3*, int)> TriangleColorFunc, ERenderUpdateMode UpdateMode)
 {
-	DynamicMeshComponent->TriangleColorFunc = TriangleColorFunc;
-	if (UpdateMode == ERenderUpdateMode::FastUpdate)
-	{
-		DynamicMeshComponent->FastNotifyColorsUpdated();
-	}
-	else if (UpdateMode == ERenderUpdateMode::FullUpdate)
-	{
-		DynamicMeshComponent->NotifyMeshUpdated();
-	}
+	DynamicMeshComponent->SetTriangleColorFunction(TriangleColorFunc,  (EDynamicMeshComponentRenderUpdateMode)(int32)UpdateMode );
 }
 
 void UPreviewMesh::ClearTriangleColorFunction(ERenderUpdateMode UpdateMode)
 {
-	if (DynamicMeshComponent->TriangleColorFunc)
-	{
-		DynamicMeshComponent->TriangleColorFunc = nullptr;
-		if (UpdateMode == ERenderUpdateMode::FastUpdate)
-		{
-			DynamicMeshComponent->FastNotifyColorsUpdated();
-		}
-		else if (UpdateMode == ERenderUpdateMode::FullUpdate)
-		{
-			DynamicMeshComponent->NotifyMeshUpdated();
-		}
-	}
+	DynamicMeshComponent->ClearTriangleColorFunction((EDynamicMeshComponentRenderUpdateMode)(int32)UpdateMode);
 }
 
 
diff --git a/Engine/Plugins/Experimental/MeshModelingToolset/Source/ModelingComponents/Private/SimpleDynamicMeshComponent.cpp b/Engine/Plugins/Experimental/MeshModelingToolset/Source/ModelingComponents/Private/SimpleDynamicMeshComponent.cpp
--- a/Engine/Plugins/Experimental/MeshModelingToolset/Source/ModelingComponents/Private/SimpleDynamicMeshComponent.cpp
+++ b/Engine/Plugins/Experimental/MeshModelingToolset/Source/ModelingComponents/Private/SimpleDynamicMeshComponent.cpp
@@ -54,7 +54,25 @@ USimpleDynamicMeshComponent::USimpleDynamicMeshComponent(const FObjectInitialize
 
 	SetCollisionProfileName(UCollisionProfile::NoCollision_ProfileName);
 
-	InitializeNewMesh();
+	MeshObject = CreateDefaultSubobject<UDynamicMesh>(TEXT("DynamicMesh"));
+	//MeshObject->SetFlags(RF_Transactional);
+
+	MeshObjectChangedHandle = MeshObject->OnMeshChanged().AddUObject(this, &USimpleDynamicMeshComponent::OnMeshObjectChanged);
+
+	Tangents.SetMesh(GetMesh());
+}
+
+
+void USimpleDynamicMeshComponent::PostLoad()
+{
+	Super::PostLoad();
+
+	check(MeshObject != nullptr);
+	MeshObjectChangedHandle = MeshObject->OnMeshChanged().AddUObject(this, &USimpleDynamicMeshComponent::OnMeshObjectChanged);
+
+	Tangents.SetMesh(GetMesh());
+
+	ResetProxy();
 }
 
 
@@ -66,7 +84,7 @@ FDynamicMesh3* USimpleDynamicMeshComponent::GetRenderMesh()
 	}
 	else
 	{
-		return Mesh.Get();
+		return GetMesh();
 	}
 }
 
@@ -78,18 +96,18 @@ const FDynamicMesh3* USimpleDynamicMeshComponent::GetRenderMesh() const
 	}
 	else
 	{
-		return Mesh.Get();
+		return GetMesh();
 	}
 }
 
-void USimpleDynamicMeshComponent::InitializeMesh(FMeshDescription* MeshDescription)
+void USimpleDynamicMeshComponent::InitializeMesh(const FMeshDescription* MeshDescription)
 {
 	FMeshDescriptionToDynamicMesh Converter;
-	Mesh->Clear();
-	Converter.Convert(MeshDescription, *Mesh);
+	MeshObject->Reset();
+	Converter.Convert(MeshDescription, *GetMesh());
 	if (TangentsType == EDynamicMeshTangentCalcType::ExternallyCalculated)
 	{
-		Converter.CopyTangents(MeshDescription, Mesh.Get(), &Tangents);
+		Converter.CopyTangents(MeshDescription, GetMesh(), &Tangents);
 	}
 
 	NotifyMeshUpdated();
@@ -102,7 +120,7 @@ void USimpleDynamicMeshComponent::SetRenderMeshPostProcessor(TUniquePtr<IRenderM
 	{
 		if (!RenderMesh)
 		{
-			RenderMesh = MakeUnique<FDynamicMesh3>(*Mesh);
+			RenderMesh = MakeUnique<FDynamicMesh3>(*GetMesh());
 		}
 	}
 	else
@@ -140,28 +158,6 @@ void USimpleDynamicMeshComponent::UpdateTangents(const FMeshTangentsd* ExternalT
 
 
 
-TUniquePtr<FDynamicMesh3> USimpleDynamicMeshComponent::ExtractMesh(bool bNotifyUpdate)
-{
-	TUniquePtr<FDynamicMesh3> CurMesh = MoveTemp(Mesh);
-	InitializeNewMesh();
-	if (bNotifyUpdate)
-	{
-		NotifyMeshUpdated();
-	}
-	return CurMesh;
-}
-
-
-void USimpleDynamicMeshComponent::InitializeNewMesh()
-{
-	Mesh = MakeUnique<FDynamicMesh3>();
-	// discard any attributes/etc initialized by default
-	Mesh->Clear();
-
-	Tangents.SetMesh(Mesh.Get());
-}
-
-
 void USimpleDynamicMeshComponent::ApplyTransform(const UE::Geometry::FTransform3d& Transform, bool bInvert)
 {
 	if (bInvert)
@@ -182,11 +178,11 @@ void USimpleDynamicMeshComponent::Bake(FMeshDescription* MeshDescription, bool b
 	FDynamicMeshToMeshDescription Converter(ConversionOptions);
 	if (!bHaveModifiedTopology)
 	{
-		Converter.UpdateUsingConversionOptions(Mesh.Get(), *MeshDescription);
+		Converter.UpdateUsingConversionOptions(GetMesh(), *MeshDescription);
 	}
 	else
 	{
-		Converter.Convert(Mesh.Get(), *MeshDescription);
+		Converter.Convert(GetMesh(), *MeshDescription);
 	}
 }
 
@@ -204,10 +200,10 @@ const FMeshTangentsf* USimpleDynamicMeshComponent::GetTangents()
 	
 	if (TangentsType == EDynamicMeshTangentCalcType::AutoCalculated)
 	{
-		if (bTangentsValid == false && Mesh->HasAttributes())
+		if (bTangentsValid == false && GetMesh()->HasAttributes())
 		{
-			FDynamicMeshUVOverlay* UVOverlay = Mesh->Attributes()->PrimaryUV();
-			FDynamicMeshNormalOverlay* NormalOverlay = Mesh->Attributes()->PrimaryNormals();
+			FDynamicMeshUVOverlay* UVOverlay = GetMesh()->Attributes()->PrimaryUV();
+			FDynamicMeshNormalOverlay* NormalOverlay = GetMesh()->Attributes()->PrimaryNormals();
 			if (UVOverlay != nullptr && NormalOverlay != nullptr)
 			{
 				Tangents.ComputeTriVertexTangents(NormalOverlay, UVOverlay, FComputeTangentsOptions());
@@ -230,32 +226,49 @@ const FMeshTangentsf* USimpleDynamicMeshComponent::GetTangents()
 
 
 
-void USimpleDynamicMeshComponent::SetDrawOnTop(bool bSet)
+void USimpleDynamicMeshComponent::UpdateLocalBounds()
 {
-	bDrawOnTop = bSet;
-	bUseEditorCompositing = bSet;
+	LocalBounds = GetMesh()->GetBounds(true);
+	if (LocalBounds.MaxDim() <= 0)
+	{
+		LocalBounds = FAxisAlignedBox3d(FVector3d::Zero(), FMathf::ZeroTolerance);
+	}
+}
+
+FSimpleDynamicMeshSceneProxy* USimpleDynamicMeshComponent::GetCurrentSceneProxy() 
+{ 
+	if (bProxyValid)
+	{
+		return (FSimpleDynamicMeshSceneProxy*)SceneProxy;
+	}
+	return nullptr;
 }
 
+
 void USimpleDynamicMeshComponent::ResetProxy()
 {
 	bProxyValid = false;
 
 	// Need to recreate scene proxy to send it over
 	MarkRenderStateDirty();
-	LocalBounds = Mesh->GetBounds(true);
+	UpdateLocalBounds();
 	UpdateBounds();
 
 	if (TangentsType != EDynamicMeshTangentCalcType::ExternallyCalculated)
 	{
 		bTangentsValid = false;
 	}
+
+	// this seems speculative, ie we may not actually have a mesh update, but we currently ResetProxy() in lots
+	// of places where that is what it means
+	GetDynamicMesh()->PostRealtimeUpdate();
 }
 
 void USimpleDynamicMeshComponent::NotifyMeshUpdated()
 {
 	if (RenderMeshPostProcessor)
 	{
-		RenderMeshPostProcessor->ProcessMesh(*Mesh, *RenderMesh);
+		RenderMeshPostProcessor->ProcessMesh(*GetMesh(), *RenderMesh);
 	}
 
 	ResetProxy();
@@ -267,7 +280,7 @@ void USimpleDynamicMeshComponent::FastNotifyColorsUpdated()
 	// should not be using fast paths if we have to run mesh postprocessor
 	if (ensure(!RenderMeshPostProcessor) == false)
 	{
-		RenderMeshPostProcessor->ProcessMesh(*Mesh, *RenderMesh);
+		RenderMeshPostProcessor->ProcessMesh(*GetMesh(), *RenderMesh);
 		ResetProxy();
 		return;
 	}
@@ -275,12 +288,12 @@ void USimpleDynamicMeshComponent::FastNotifyColorsUpdated()
 	FSimpleDynamicMeshSceneProxy* Proxy = GetCurrentSceneProxy();
 	if (Proxy)
 	{
-		if (TriangleColorFunc != nullptr &&  Proxy->bUsePerTriangleColor == false )
+		if (HasTriangleColorFunction() && Proxy->bUsePerTriangleColor == false )
 		{
 			Proxy->bUsePerTriangleColor = true;
 			Proxy->PerTriangleColorFunc = [this](const FDynamicMesh3* MeshIn, int TriangleID) { return GetTriangleColor(MeshIn, TriangleID); };
 		} 
-		else if (TriangleColorFunc == nullptr && Proxy->bUsePerTriangleColor == true)
+		else if ( !HasTriangleColorFunction() && Proxy->bUsePerTriangleColor == true)
 		{
 			Proxy->bUsePerTriangleColor = false;
 			Proxy->PerTriangleColorFunc = nullptr;
@@ -302,7 +315,7 @@ void USimpleDynamicMeshComponent::FastNotifyPositionsUpdated(bool bNormals, bool
 	// should not be using fast paths if we have to run mesh postprocessor
 	if (ensure(!RenderMeshPostProcessor) == false)
 	{
-		RenderMeshPostProcessor->ProcessMesh(*Mesh, *RenderMesh);
+		RenderMeshPostProcessor->ProcessMesh(*GetMesh(), *RenderMesh);
 		ResetProxy();
 		return;
 	}
@@ -315,14 +328,17 @@ void USimpleDynamicMeshComponent::FastNotifyPositionsUpdated(bool bNormals, bool
 		UpdateBoundsCalc = Async(SimpleDynamicMeshComponentAsyncExecTarget, [this]()
 		{
 			TRACE_CPUPROFILER_EVENT_SCOPE(SimpleDynamicMeshComponent_FastPositionsUpdate_AsyncBoundsUpdate);
-			LocalBounds = Mesh->GetBounds(true);
+			UpdateLocalBounds();
 		});
 
 		GetCurrentSceneProxy()->FastUpdateVertices(true, bNormals, bColors, bUVs);
+
 		//MarkRenderDynamicDataDirty();
 		MarkRenderTransformDirty();
 		UpdateBoundsCalc.Wait();
 		UpdateBounds();
+
+		GetDynamicMesh()->PostRealtimeUpdate();
 	}
 	else
 	{
@@ -336,7 +352,7 @@ void USimpleDynamicMeshComponent::FastNotifyVertexAttributesUpdated(bool bNormal
 	// should not be using fast paths if we have to run mesh postprocessor
 	if (ensure(!RenderMeshPostProcessor) == false)
 	{
-		RenderMeshPostProcessor->ProcessMesh(*Mesh, *RenderMesh);
+		RenderMeshPostProcessor->ProcessMesh(*GetMesh(), *RenderMesh);
 		ResetProxy();
 		return;
 	}
@@ -347,6 +363,8 @@ void USimpleDynamicMeshComponent::FastNotifyVertexAttributesUpdated(bool bNormal
 		GetCurrentSceneProxy()->FastUpdateVertices(false, bNormals, bColors, bUVs);
 		//MarkRenderDynamicDataDirty();
 		//MarkRenderTransformDirty();
+
+		GetDynamicMesh()->PostRealtimeUpdate();
 	}
 	else
 	{
@@ -360,7 +378,7 @@ void USimpleDynamicMeshComponent::FastNotifyVertexAttributesUpdated(EMeshRenderA
 	// should not be using fast paths if we have to run mesh postprocessor
 	if (ensure(!RenderMeshPostProcessor) == false)
 	{
-		RenderMeshPostProcessor->ProcessMesh(*Mesh, *RenderMesh);
+		RenderMeshPostProcessor->ProcessMesh(*GetMesh(), *RenderMesh);
 		ResetProxy();
 		return;
 	}
@@ -377,7 +395,7 @@ void USimpleDynamicMeshComponent::FastNotifyVertexAttributesUpdated(EMeshRenderA
 			UpdateBoundsCalc = Async(SimpleDynamicMeshComponentAsyncExecTarget, [this]()
 			{
 				TRACE_CPUPROFILER_EVENT_SCOPE(SimpleDynamicMeshComponent_FastVertexAttribUpdate_AsyncBoundsUpdate);
-				LocalBounds = Mesh->GetBounds(true);
+				UpdateLocalBounds();
 			});
 		}
 
@@ -392,6 +410,8 @@ void USimpleDynamicMeshComponent::FastNotifyVertexAttributesUpdated(EMeshRenderA
 			UpdateBoundsCalc.Wait();
 			UpdateBounds();
 		}
+
+		GetDynamicMesh()->PostRealtimeUpdate();
 	}
 	else
 	{
@@ -412,7 +432,7 @@ void USimpleDynamicMeshComponent::FastNotifySecondaryTrianglesChanged()
 	// should not be using fast paths if we have to run mesh postprocessor
 	if (ensure(!RenderMeshPostProcessor) == false)
 	{
-		RenderMeshPostProcessor->ProcessMesh(*Mesh, *RenderMesh);
+		RenderMeshPostProcessor->ProcessMesh(*GetMesh(), *RenderMesh);
 		ResetProxy();
 		return;
 	}
@@ -421,6 +441,7 @@ void USimpleDynamicMeshComponent::FastNotifySecondaryTrianglesChanged()
 	if (Proxy)
 	{
 		GetCurrentSceneProxy()->FastUpdateAllIndexBuffers();
+		GetDynamicMesh()->PostRealtimeUpdate();
 	}
 	else
 	{
@@ -434,7 +455,7 @@ void USimpleDynamicMeshComponent::FastNotifyTriangleVerticesUpdated(const TArray
 	// should not be using fast paths if we have to run mesh postprocessor
 	if (ensure(!RenderMeshPostProcessor) == false)
 	{
-		RenderMeshPostProcessor->ProcessMesh(*Mesh, *RenderMesh);
+		RenderMeshPostProcessor->ProcessMesh(*GetMesh(), *RenderMesh);
 		ResetProxy();
 		return;
 	}
@@ -454,6 +475,7 @@ void USimpleDynamicMeshComponent::FastNotifyTriangleVerticesUpdated(const TArray
 		{
 			Proxy->FastUpdateAllIndexBuffers();
 		}
+		GetDynamicMesh()->PostRealtimeUpdate();
 	}
 	else
 	{
@@ -477,7 +499,7 @@ void USimpleDynamicMeshComponent::FastNotifyTriangleVerticesUpdated(const TArray
 			UpdateBoundsCalc = Async(SimpleDynamicMeshComponentAsyncExecTarget, [this]()
 			{
 				TRACE_CPUPROFILER_EVENT_SCOPE(SimpleDynamicMeshComponent_FastVertexUpdate_AsyncBoundsUpdate);
-				LocalBounds = Mesh->GetBounds(true);
+				UpdateLocalBounds();
 			});
 		}
 
@@ -502,6 +524,8 @@ void USimpleDynamicMeshComponent::FastNotifyTriangleVerticesUpdated(const TArray
 			UpdateBoundsCalc.Wait();
 			UpdateBounds();
 		}
+
+		GetDynamicMesh()->PostRealtimeUpdate();
 	}
 }
 
@@ -513,7 +537,7 @@ void USimpleDynamicMeshComponent::FastNotifyTriangleVerticesUpdated(const TSet<i
 	// should not be using fast paths if we have to run mesh postprocessor
 	if (ensure(!RenderMeshPostProcessor) == false)
 	{
-		RenderMeshPostProcessor->ProcessMesh(*Mesh, *RenderMesh);
+		RenderMeshPostProcessor->ProcessMesh(*GetMesh(), *RenderMesh);
 		ResetProxy();
 		return;
 	}
@@ -533,6 +557,7 @@ void USimpleDynamicMeshComponent::FastNotifyTriangleVerticesUpdated(const TSet<i
 		{
 			Proxy->FastUpdateAllIndexBuffers();
 		}
+		GetDynamicMesh()->PostRealtimeUpdate();
 	}
 	else
 	{
@@ -556,7 +581,7 @@ void USimpleDynamicMeshComponent::FastNotifyTriangleVerticesUpdated(const TSet<i
 			UpdateBoundsCalc = Async(SimpleDynamicMeshComponentAsyncExecTarget, [this]()
 			{
 				TRACE_CPUPROFILER_EVENT_SCOPE(SimpleDynamicMeshComponent_FastVertexUpdate_AsyncBoundsUpdate);
-				LocalBounds = Mesh->GetBounds(true);
+				UpdateLocalBounds();
 			});
 		}
 
@@ -585,6 +610,8 @@ void USimpleDynamicMeshComponent::FastNotifyTriangleVerticesUpdated(const TSet<i
 			UpdateBoundsCalc.Wait();
 			UpdateBounds();
 		}
+
+		GetDynamicMesh()->PostRealtimeUpdate();
 	}
 }
 
@@ -643,7 +670,7 @@ TFuture<bool> USimpleDynamicMeshComponent::FastNotifyTriangleVerticesUpdated_Try
 		TFuture<void> ComputeBounds = Async(SimpleDynamicMeshComponentAsyncExecTarget, [this, &BoundsOut, &Triangles]()
 		{
 			TRACE_CPUPROFILER_EVENT_SCOPE(SimpleDynamicMeshComponent_FastVertexUpdatePrecomp_CalcBounds);
-			BoundsOut = ParallelComputeROIBounds(*Mesh, Triangles);
+			BoundsOut = ParallelComputeROIBounds(*GetMesh(), Triangles);
 		});
 
 		TFuture<void> ComputeSets = Async(SimpleDynamicMeshComponentAsyncExecTarget, [this, &UpdateSetsOut, &Triangles]()
@@ -674,6 +701,7 @@ TFuture<bool> USimpleDynamicMeshComponent::FastNotifyTriangleVerticesUpdated_Try
 
 		ComputeSets.Wait();
 		ComputeBounds.Wait();
+
 		return true;
 	});
 }
@@ -723,6 +751,8 @@ void USimpleDynamicMeshComponent::FastNotifyTriangleVerticesUpdated_ApplyPrecomp
 		LocalBounds.Contain(UpdateSetBounds);
 		UpdateBounds();
 	}
+
+	GetDynamicMesh()->PostRealtimeUpdate();
 }
 
 
@@ -735,7 +765,7 @@ FPrimitiveSceneProxy* USimpleDynamicMeshComponent::CreateSceneProxy()
 	ensure(GetCurrentSceneProxy() == nullptr);
 
 	FSimpleDynamicMeshSceneProxy* NewProxy = nullptr;
-	if (Mesh->TriangleCount() > 0)
+	if (GetMesh()->TriangleCount() > 0)
 	{
 		NewProxy = new FSimpleDynamicMeshSceneProxy(this);
 
@@ -779,6 +809,45 @@ void USimpleDynamicMeshComponent::NotifyMaterialSetUpdated()
 }
 
 
+void USimpleDynamicMeshComponent::SetTriangleColorFunction(
+	TUniqueFunction<FColor(const FDynamicMesh3*, int)> TriangleColorFuncIn,
+	EDynamicMeshComponentRenderUpdateMode UpdateMode)
+{
+	TriangleColorFunc = MoveTemp(TriangleColorFuncIn);
+
+	if (UpdateMode == EDynamicMeshComponentRenderUpdateMode::FastUpdate)
+	{
+		FastNotifyColorsUpdated();
+	}
+	else if (UpdateMode == EDynamicMeshComponentRenderUpdateMode::FullUpdate)
+	{
+		NotifyMeshUpdated();
+	}
+}
+
+void USimpleDynamicMeshComponent::ClearTriangleColorFunction(EDynamicMeshComponentRenderUpdateMode UpdateMode)
+{
+	if (TriangleColorFunc)
+	{
+		TriangleColorFunc = nullptr;
+
+		if (UpdateMode == EDynamicMeshComponentRenderUpdateMode::FastUpdate)
+		{
+			FastNotifyColorsUpdated();
+		}
+		else if (UpdateMode == EDynamicMeshComponentRenderUpdateMode::FullUpdate)
+		{
+			NotifyMeshUpdated();
+		}
+	}
+}
+
+bool USimpleDynamicMeshComponent::HasTriangleColorFunction()
+{
+	return !!TriangleColorFunc;
+}
+
+
 
 
 void USimpleDynamicMeshComponent::EnableSecondaryTriangleBuffers(TUniqueFunction<bool(const FDynamicMesh3*, int32)> SecondaryTriFilterFuncIn)
@@ -827,76 +896,89 @@ FBoxSphereBounds USimpleDynamicMeshComponent::CalcBounds(const FTransform& Local
 }
 
 
+
+
+void USimpleDynamicMeshComponent::SetInvalidateProxyOnChangeEnabled(bool bEnabled)
+{
+	bInvalidateProxyOnChange = bEnabled;
+}
+
+
 void USimpleDynamicMeshComponent::ApplyChange(const FMeshVertexChange* Change, bool bRevert)
 {
-	bool bHavePositions = Change->bHaveVertexPositions;
-	bool bHaveColors = Change->bHaveVertexColors && Mesh->HasVertexColors();
+	// will fire UDynamicMesh::MeshChangedEvent, which will call OnMeshObjectChanged() below to invalidate proxy, fire change events, etc
+	MeshObject->ApplyChange(Change, bRevert);
+}
+
+void USimpleDynamicMeshComponent::ApplyChange(const FMeshChange* Change, bool bRevert)
+{
+	// will fire UDynamicMesh::MeshChangedEvent, which will call OnMeshObjectChanged() below to invalidate proxy, fire change events, etc
+	MeshObject->ApplyChange(Change, bRevert);
+}
+
+void USimpleDynamicMeshComponent::ApplyChange(const FMeshReplacementChange* Change, bool bRevert)
+{
+	// will fire UDynamicMesh::MeshChangedEvent, which will call OnMeshObjectChanged() below to invalidate proxy, fire change events, etc
+	MeshObject->ApplyChange(Change, bRevert);
+}
+
+void USimpleDynamicMeshComponent::OnMeshObjectChanged(UDynamicMesh* ChangedMeshObject, FDynamicMeshChangeInfo ChangeInfo)
+{
+	bool bIsFChange = (
+		ChangeInfo.Type == EDynamicMeshChangeType::MeshChange
+		|| ChangeInfo.Type == EDynamicMeshChangeType::MeshVertexChange
+		|| ChangeInfo.Type == EDynamicMeshChangeType::MeshReplacementChange);
 
-	int32 NV = Change->Vertices.Num();
-	const TArray<FVector3d>& Positions = (bRevert) ? Change->OldPositions : Change->NewPositions;
-	const TArray<FVector3f>& Colors = (bRevert) ? Change->OldColors : Change->NewColors;
-	for (int32 k = 0; k < NV; ++k)
+	if (bIsFChange)
 	{
-		int32 vid = Change->Vertices[k];
-		if (Mesh->IsVertex(vid))
+		if (bInvalidateProxyOnChange)
 		{
-			if (bHavePositions)
-			{
-				Mesh->SetVertex(vid, Positions[k]);
-			}
-			if (bHaveColors)
-			{
-				Mesh->SetVertexColor(vid, Colors[k]);
-			}
+			NotifyMeshUpdated();
 		}
-	}
 
-	if (Change->bHaveOverlayNormals && Mesh->HasAttributes() && Mesh->Attributes()->PrimaryNormals() )
-	{
-		FDynamicMeshNormalOverlay* Overlay = Mesh->Attributes()->PrimaryNormals();
-		int32 NumNormals = Change->Normals.Num();
-		const TArray<FVector3f>& UseNormals = (bRevert) ? Change->OldNormals : Change->NewNormals;
-		for (int32 k = 0; k < NumNormals; ++k)
+		OnMeshChanged.Broadcast();
+
+		if (ChangeInfo.Type == EDynamicMeshChangeType::MeshVertexChange)
 		{
-			int32 elemid = Change->Normals[k];
-			if (Overlay->IsElement(elemid))
-			{
-				Overlay->SetElement(elemid, UseNormals[k]);
-			}
+			OnMeshVerticesChanged.Broadcast(this, ChangeInfo.VertexChange, ChangeInfo.bIsRevertChange);
 		}
 	}
-
-	if (bInvalidateProxyOnChange)
+	else
 	{
 		NotifyMeshUpdated();
+		OnMeshChanged.Broadcast();
 	}
-	OnMeshChanged.Broadcast();
-	OnMeshVerticesChanged.Broadcast(this, Change, bRevert);
 }
 
 
-
-
-void USimpleDynamicMeshComponent::ApplyChange(const FMeshChange* Change, bool bRevert)
+void USimpleDynamicMeshComponent::SetDynamicMesh(UDynamicMesh* NewMesh)
 {
-	Change->ApplyChangeToMesh(Mesh.Get(), bRevert);
+	if (ensure(NewMesh) == false)
+	{
+		return;
+	}
 
-	if (bInvalidateProxyOnChange)
+	if (ensure(MeshObject))
 	{
-		NotifyMeshUpdated();
+		MeshObject->OnMeshChanged().Remove(MeshObjectChangedHandle);
 	}
+
+	MeshObject = NewMesh;
+	MeshObjectChangedHandle = MeshObject->OnMeshChanged().AddUObject(this, &USimpleDynamicMeshComponent::OnMeshObjectChanged);
+
+	NotifyMeshUpdated();
 	OnMeshChanged.Broadcast();
 }
 
 
-void USimpleDynamicMeshComponent::ApplyChange(const FMeshReplacementChange* Change, bool bRevert)
-{
-	Mesh->Copy(*Change->GetMesh(bRevert));
 
-	if (bInvalidateProxyOnChange)
-	{
-		NotifyMeshUpdated();
-	}
-	OnMeshChanged.Broadcast();
+void USimpleDynamicMeshComponent::OnChildAttached(USceneComponent* ChildComponent)
+{
+	Super::OnChildAttached(ChildComponent);
+	OnChildAttachmentModified.Broadcast(ChildComponent, true);
+}
+void USimpleDynamicMeshComponent::OnChildDetached(USceneComponent* ChildComponent)
+{
+	Super::OnChildDetached(ChildComponent);
+	OnChildAttachmentModified.Broadcast(ChildComponent, false);
 }
-
diff --git a/Engine/Plugins/Experimental/MeshModelingToolset/Source/ModelingComponents/Private/SimpleDynamicMeshSceneProxy.h b/Engine/Plugins/Experimental/MeshModelingToolset/Source/ModelingComponents/Private/SimpleDynamicMeshSceneProxy.h
--- a/Engine/Plugins/Experimental/MeshModelingToolset/Source/ModelingComponents/Private/SimpleDynamicMeshSceneProxy.h
+++ b/Engine/Plugins/Experimental/MeshModelingToolset/Source/ModelingComponents/Private/SimpleDynamicMeshSceneProxy.h
@@ -615,27 +615,16 @@ public:
 	{
 		FPrimitiveViewRelevance Result;
 
-		if (ParentComponent->bDrawOnTop)
-		{
-			Result.bDrawRelevance = IsShown(View);
-			Result.bDynamicRelevance = true;
-			Result.bShadowRelevance = false;
-			Result.bEditorPrimitiveRelevance = UseEditorCompositing(View);
-			Result.bEditorNoDepthTestPrimitiveRelevance = true;
-		}
-		else
-		{
-			Result.bDrawRelevance = IsShown(View);
-			Result.bShadowRelevance = IsShadowCast(View);
-			Result.bDynamicRelevance = true;
-			Result.bRenderInMainPass = ShouldRenderInMainPass();
-			Result.bUsesLightingChannels = GetLightingChannelMask() != GetDefaultLightingChannelMask();
-			Result.bTranslucentSelfShadow = bCastVolumetricTranslucentShadow;
-			Result.bRenderCustomDepth = ShouldRenderCustomDepth();
-			// Note that this is actually a getter. One may argue that it is poorly named.
-			MaterialRelevance.SetPrimitiveViewRelevance(Result);
-			Result.bVelocityRelevance = DrawsVelocity() && Result.bOpaque && Result.bRenderInMainPass;
-		}
+		Result.bDrawRelevance = IsShown(View);
+		Result.bShadowRelevance = IsShadowCast(View);
+		Result.bDynamicRelevance = true;
+		Result.bRenderInMainPass = ShouldRenderInMainPass();
+		Result.bUsesLightingChannels = GetLightingChannelMask() != GetDefaultLightingChannelMask();
+		Result.bTranslucentSelfShadow = bCastVolumetricTranslucentShadow;
+		Result.bRenderCustomDepth = ShouldRenderCustomDepth();
+		// Note that this is actually a getter. One may argue that it is poorly named.
+		MaterialRelevance.SetPrimitiveViewRelevance(Result);
+		Result.bVelocityRelevance = DrawsVelocity() && Result.bOpaque && Result.bRenderInMainPass;
 
 		return Result;
 	}
@@ -651,20 +640,9 @@ public:
 
 	virtual void GetLightRelevance(const FLightSceneProxy* LightSceneProxy, bool& bDynamic, bool& bRelevant, bool& bLightMapped, bool& bShadowMapped) const override
 	{
-		if (ParentComponent->bDrawOnTop)
-		{
-			bDynamic = false;
-			bRelevant = false;
-			bLightMapped = false;
-			bShadowMapped = false;
-		}
-		else
-		{
-			FPrimitiveSceneProxy::GetLightRelevance(LightSceneProxy, bDynamic, bRelevant, bLightMapped, bShadowMapped);
-		}
+		FPrimitiveSceneProxy::GetLightRelevance(LightSceneProxy, bDynamic, bRelevant, bLightMapped, bShadowMapped);
 	}
 
-
 	virtual bool CanBeOccluded() const override
 	{
 		return !MaterialRelevance.bDisableDepthTest;
@@ -674,15 +652,12 @@ public:
 
 	uint32 GetAllocatedSize(void) const { return(FPrimitiveSceneProxy::GetAllocatedSize()); }
 
-
 	SIZE_T GetTypeHash() const override
 	{
 		static size_t UniquePointer;
 		return reinterpret_cast<size_t>(&UniquePointer);
 	}
 
-
-
 };
 
 
diff --git a/Engine/Plugins/Experimental/MeshModelingToolset/Source/ModelingComponents/Private/UDynamicMesh.cpp b/Engine/Plugins/Experimental/MeshModelingToolset/Source/ModelingComponents/Private/UDynamicMesh.cpp
new file mode 100644
--- /dev/null
+++ b/Engine/Plugins/Experimental/MeshModelingToolset/Source/ModelingComponents/Private/UDynamicMesh.cpp
@@ -0,0 +1,570 @@
+// Copyright Epic Games, Inc. All Rights Reserved.
+
+#include "UDynamicMesh.h"
+
+#include "Changes/MeshVertexChange.h"
+#include "Changes/MeshChange.h"
+#include "Changes/MeshReplacementChange.h"
+#include "Misc/Base64.h"
+#include "Serialization/MemoryWriter.h"
+#include "Serialization/MemoryReader.h"
+#include "UObject/UE5MainStreamObjectVersion.h"
+#include "HAL/IConsoleManager.h"
+#include "Engine/Engine.h"
+#include "Generators/MinimalBoxMeshGenerator.h"
+
+#include "ExplicitUseGeometryMathTypes.h"		// using UE::Geometry::(math types)
+using namespace UE::Geometry;
+
+
+// these cvars are used to support T3D encoding of the internal FDynamicMesh3, see ::ExportCustomProperties() and ::ImportCustomProperties()
+static TAutoConsoleVariable<int32> CVarDynamicMeshTextBasedDupeTriangleCountThreshold(
+	TEXT("geometry.DynamicMesh.TextBasedDupeTriThreshold"),
+	1000,
+	TEXT("Triangle count threshold for text-based UDynamicMesh duplication using Base64. Large values are quite slow."));
+static TAutoConsoleVariable<int32> CVarDynamicMeshDupeHelperTimeout(
+	TEXT("geometry.DynamicMesh.DupeStashTimeout"),
+	5*60,
+	TEXT("Timeout in seconds for references held by internal UDynamicMesh duplication helper system. See FDynamicMeshCopyHelper."));
+
+
+
+UDynamicMesh::UDynamicMesh(const FObjectInitializer& ObjectInitializer)
+	: Super(ObjectInitializer)
+{
+	InitializeNewMesh();
+}
+
+
+void UDynamicMesh::InitializeNewMesh()
+{
+	Mesh = MakeUnique<FDynamicMesh3>(EMeshComponents::FaceGroups);
+	InitializeMesh();
+}
+
+void UDynamicMesh::Reset()
+{
+	FDynamicMeshChangeInfo ChangeInfo;
+	ChangeInfo.Type = EDynamicMeshChangeType::GeneralEdit;
+	EditMeshInternal([this](FDynamicMesh3& EditMesh)
+	{
+		check(&EditMesh == Mesh.Get());		// assuming that EditMesh is internal mesh here...
+		InitializeMesh();
+	}, ChangeInfo);
+}
+
+void UDynamicMesh::ResetToCube()
+{
+	FDynamicMeshChangeInfo ChangeInfo;
+	ChangeInfo.Type = EDynamicMeshChangeType::GeneralEdit;
+	EditMeshInternal([this](FDynamicMesh3& EditMesh)
+	{
+		EditMesh.Clear();
+		FMinimalBoxMeshGenerator BoxGen;
+		BoxGen.Box = FOrientedBox3d(FVector3d::Zero(), 50.0 * FVector3d::One());
+		EditMesh = FDynamicMesh3(&BoxGen.Generate());
+		EditMesh.EnableTriangleGroups();
+		EditMesh.Attributes()->EnableMaterialID();
+	}, ChangeInfo);
+}
+
+
+void UDynamicMesh::InitializeMesh()
+{
+	Mesh->Clear();
+	Mesh->EnableTriangleGroups();
+	Mesh->EnableAttributes();
+	Mesh->Attributes()->EnableMaterialID();
+
+	if (MeshGenerator != nullptr && bEnableMeshGenerator)
+	{
+		MeshGenerator->Generate(*Mesh);
+	}
+}
+
+
+void UDynamicMesh::SetMesh(const UE::Geometry::FDynamicMesh3& MoveMesh)
+{
+	FDynamicMeshChangeInfo ChangeInfo;
+	ChangeInfo.Type = EDynamicMeshChangeType::GeneralEdit;
+	EditMeshInternal([&](FDynamicMesh3& EditMesh) {
+		EditMesh = MoveMesh;
+	}, ChangeInfo);
+}
+
+void UDynamicMesh::SetMesh(UE::Geometry::FDynamicMesh3&& MoveMesh)
+{
+	FDynamicMeshChangeInfo ChangeInfo;
+	ChangeInfo.Type = EDynamicMeshChangeType::GeneralEdit;
+	EditMeshInternal([&](FDynamicMesh3& EditMesh) {
+		EditMesh = MoveTemp(MoveMesh);
+	}, ChangeInfo);
+}
+
+void UDynamicMesh::ProcessMesh(TFunctionRef<void(const UE::Geometry::FDynamicMesh3&)> ProcessFunc) const
+{
+	ProcessFunc(GetMeshRef());
+}
+
+void UDynamicMesh::EditMesh(TFunctionRef<void(FDynamicMesh3&)> EditFunc)
+{
+	FDynamicMeshChangeInfo ChangeInfo;
+	ChangeInfo.Type = EDynamicMeshChangeType::GeneralEdit;
+	EditMeshInternal(EditFunc, ChangeInfo);
+}
+
+void UDynamicMesh::EditMeshInternal(TFunctionRef<void(FDynamicMesh3&)> EditFunc, const FDynamicMeshChangeInfo& ChangeInfo)
+{
+	PreMeshChangedEvent.Broadcast(this, ChangeInfo);
+	EditFunc(GetMeshRef());
+	MeshChangedEvent.Broadcast(this, ChangeInfo);
+	MeshModifiedBPEvent.Broadcast(this);
+}
+
+
+TUniquePtr<FDynamicMesh3> UDynamicMesh::ExtractMesh()
+{
+	TUniquePtr<FDynamicMesh3> ReturnMesh = MoveTemp(Mesh);
+	InitializeNewMesh();
+	return ReturnMesh;
+}
+
+
+
+void UDynamicMesh::SetMeshGenerator(TObjectPtr<UDynamicMeshGenerator> NewGenerator)
+{
+	MeshGenerator = NewGenerator;
+}
+
+void UDynamicMesh::ClearMeshGenerator()
+{
+	MeshGenerator = nullptr;
+}
+
+void UDynamicMesh::Regenerate()
+{
+	if (MeshGenerator != nullptr && bEnableMeshGenerator)
+	{
+		Reset();
+	}
+}
+
+
+void UDynamicMesh::PostRealtimeUpdate() 
+{ 
+	MeshRealtimeUpdateEvent.Broadcast(this); 
+}
+
+
+void UDynamicMesh::ApplyChange(const FMeshVertexChange* Change, bool bRevert)
+{
+	FDynamicMeshChangeInfo ChangeInfo;
+	ChangeInfo.Type = EDynamicMeshChangeType::MeshVertexChange;
+	ChangeInfo.VertexChange = Change;
+	ChangeInfo.bIsRevertChange = bRevert;
+
+	EditMeshInternal([&](FDynamicMesh3& EditMesh)
+	{
+		bool bHavePositions = Change->bHaveVertexPositions;
+		bool bHaveColors = Change->bHaveVertexColors && EditMesh.HasVertexColors();
+
+		int32 NV = Change->Vertices.Num();
+		const TArray<FVector3d>& Positions = (bRevert) ? Change->OldPositions : Change->NewPositions;
+		const TArray<FVector3f>& Colors = (bRevert) ? Change->OldColors : Change->NewColors;
+		for (int32 k = 0; k < NV; ++k)
+		{
+			int32 vid = Change->Vertices[k];
+			if (EditMesh.IsVertex(vid))
+			{
+				if (bHavePositions)
+				{
+					EditMesh.SetVertex(vid, Positions[k]);
+				}
+				if (bHaveColors)
+				{
+					EditMesh.SetVertexColor(vid, Colors[k]);
+				}
+			}
+		}
+
+		if (Change->bHaveOverlayNormals && EditMesh.HasAttributes() && EditMesh.Attributes()->PrimaryNormals())
+		{
+			FDynamicMeshNormalOverlay* Overlay = EditMesh.Attributes()->PrimaryNormals();
+			int32 NumNormals = Change->Normals.Num();
+			const TArray<FVector3f>& UseNormals = (bRevert) ? Change->OldNormals : Change->NewNormals;
+			for (int32 k = 0; k < NumNormals; ++k)
+			{
+				int32 elemid = Change->Normals[k];
+				if (Overlay->IsElement(elemid))
+				{
+					Overlay->SetElement(elemid, UseNormals[k]);
+				}
+			}
+		}
+	}, ChangeInfo);
+}
+
+void UDynamicMesh::ApplyChange(const FMeshChange* Change, bool bRevert)
+{
+	FDynamicMeshChangeInfo ChangeInfo;
+	ChangeInfo.Type = EDynamicMeshChangeType::MeshChange;
+	ChangeInfo.MeshChange = Change;
+	ChangeInfo.bIsRevertChange = bRevert;
+
+	EditMeshInternal([&](FDynamicMesh3& EditMesh)
+	{
+		Change->ApplyChangeToMesh(&EditMesh, bRevert);
+	}, ChangeInfo);
+}
+
+void UDynamicMesh::ApplyChange(const FMeshReplacementChange* Change, bool bRevert)
+{
+	FDynamicMeshChangeInfo ChangeInfo;
+	ChangeInfo.Type = EDynamicMeshChangeType::MeshReplacementChange;
+	ChangeInfo.ReplaceChange = Change;
+	ChangeInfo.bIsRevertChange = bRevert;
+
+	EditMeshInternal([&](FDynamicMesh3& EditMesh)
+	{
+		EditMesh.Copy(*Change->GetMesh(bRevert));
+	}, ChangeInfo);
+}
+
+
+void UDynamicMesh::Serialize(FArchive& Ar)
+{
+	TRACE_CPUPROFILER_EVENT_SCOPE(UDynamicMesh::Serialize);
+
+	Super::Serialize(Ar);
+
+	// do not serialize mesh for transactions
+	if (Ar.IsTransacting())
+	{
+		return;
+	}
+
+	Ar.UsingCustomVersion(FUE5MainStreamObjectVersion::GUID);
+
+	if (Ar.IsLoading())
+	{
+		InitializeNewMesh();
+	}
+	Ar << *Mesh;
+}
+
+
+
+
+
+
+/**
+ * This is an internal singleton used to support UDynamicMesh::ExportCustomProperties() and UDynamicMesh::ImportCustomProperties() below.
+ * 
+ */
+namespace UE 
+{
+namespace Local
+{
+
+class FDynamicMeshCopyHelper
+{
+protected:
+	struct FStashedMesh
+	{
+		TSoftObjectPtr<UDynamicMesh> SourceMesh;
+		FDateTime Timestamp;
+	};
+
+	static TMap<int32, FStashedMesh> StashedMeshes;
+	static FRandomStream KeyGenerator;
+public:
+
+	static void Initialize()
+	{
+		static bool bInitialized = false;
+		if (!bInitialized)
+		{
+			KeyGenerator = FRandomStream( (int32)FDateTime::UtcNow().ToUnixTimestamp() );
+			bInitialized = true;
+		}
+	}
+
+	static void DiscardExpiredMeshes()
+	{
+		const int32 ExpiryTimeoutInSeconds = CVarDynamicMeshDupeHelperTimeout.GetValueOnGameThread();
+
+		TArray<int32> ToRemove;
+		for (TPair<int32, FStashedMesh>& Pair : StashedMeshes)
+		{
+			if ((FDateTime::Now() - Pair.Value.Timestamp).GetTotalSeconds() > (double)ExpiryTimeoutInSeconds)
+			{
+				ToRemove.Add(Pair.Key);
+			}
+		}
+		for (int32 id : ToRemove)
+		{
+			StashedMeshes.Remove(id);
+		}
+	}
+
+	static int32 StashMeshReference(UDynamicMesh* SourceMesh)
+	{
+		Initialize();
+		int32 MeshKey = FMath::Abs((int32)KeyGenerator.GetUnsignedInt());
+
+		FStashedMesh NewMesh;
+		NewMesh.Timestamp = FDateTime::Now();
+		NewMesh.SourceMesh = SourceMesh;
+		StashedMeshes.Add(MeshKey, MoveTemp(NewMesh));
+
+		DiscardExpiredMeshes();
+		return MeshKey;
+	}
+
+	static TUniquePtr<FDynamicMesh3> ExtractStashedMesh(int32 MeshKey)
+	{
+		TUniquePtr<FDynamicMesh3> ReturnMesh;
+		FStashedMesh* Found = StashedMeshes.Find(MeshKey);
+		if (Found != nullptr && Found->SourceMesh.IsValid())
+		{
+			ReturnMesh = MakeUnique<FDynamicMesh3>();
+			UDynamicMesh* SourceMesh = Found->SourceMesh.Get();
+			if (SourceMesh)
+			{
+				// should we Cast<UDynamicMesh> here to make sure this is still the same UClass? (due to recyling)
+				SourceMesh->ProcessMesh([&](const FDynamicMesh3& EditMesh)
+				{
+					*ReturnMesh = EditMesh;
+				});
+			}
+		}
+
+		DiscardExpiredMeshes();
+		return ReturnMesh;
+	}
+
+};
+TMap<int32, FDynamicMeshCopyHelper::FStashedMesh> FDynamicMeshCopyHelper::StashedMeshes;
+FRandomStream FDynamicMeshCopyHelper::KeyGenerator;
+
+} } // end namespace UE::Local
+
+
+void UDynamicMesh::ExportCustomProperties(FOutputDevice& Out, uint32 Indent)
+{
+	// ignore empty meshes
+	if (Mesh->TriangleCount() == 0 && Mesh->VertexCount() == 0)
+	{
+		return;
+	}
+
+	// In the Editor, Copy/Paste of Actors and Components is not based on serialization,
+	// but rather on T3D structured text records. This is what allows things to be copied
+	// between Editor sessions, for example. Duplicate functionality in the Editor is implemented
+	// as Copy and Paste, so for Duplicate to work, any non-UProperty data must be serialized via text.
+	//
+	// This poses a problem for large data, such as large meshes. Serializing a mesh with millions of
+	// triangles as structured text is very expensive. An alternative is to serialize it to binary and 
+	// then use Base64 encoding. This is implemented below, however even that is quite slow. 
+	// 
+	// So, we employ a second strategy (read: hack), of just passing the source UObject pointer 
+	// via text. Instead of directly sending the pointer, we store it in FDynamicMeshCopyHelper, and pass
+	// an integer key instead. This allows for somewhat better handling, eg the copy helper could for example
+	// keep the UDynamicMeshes alive (so far this has not been necessary) and the keys are randomized so
+	// even in the unlikely even that two Editor instances end up with the same keys, it would just
+	// result in an correct pasted object, rather than accessing garbage pointers.
+	// 
+	// FDynamicMeshCopyHelper does attempt to discard "old" references, the CVar geometry.DynamicMesh.DupeStashTimeout 
+	// controls the definition of old (currenly 5 minutes). One effect this can have is if one does a
+	// copy and then a paste after the timeout, the mesh will not be found. This could also be problematic
+	// for large full-scene copies of hundreds of objects that takes longer than the timeout (resolvable via the CVar)
+	// 
+	// Obviously this does not work between Editor sessions. So we also optionally do a Base64 binary encoding if
+	// the mesh triangle count is below a CVar threshold (geometry.DynamicMesh.TextBasedDupeTriThreshold) 
+	// defaulting to 1000. Larger meshes can be supported by increasing the CVar value if need be.
+	// 
+	// If the mesh is not found in the FDynamicMeshCopyHelper, and is too large to text-copy, then instead
+	// of leaving an empty mesh, we emit a cube, as empty meshes can be problematic. A warning is also
+	// printed to the Output Log, pointing the user to the CVars.
+	// 
+	// Possible todos: UObject Reuse/Recycling
+	//
+
+	Super::ExportCustomProperties(Out, Indent);
+	Out.Logf(TEXT("%sCustomProperties "), FCString::Spc(Indent));
+	Out.Logf(TEXT("MeshData "));
+
+	// stash copy to circumvent expensive T3D generation/parsing
+	int32 StashedKey = UE::Local::FDynamicMeshCopyHelper::StashMeshReference(this);
+	Out.Logf(TEXT("MESHKEY=%d "), StashedKey);
+
+	if (Mesh->TriangleCount() < CVarDynamicMeshTextBasedDupeTriangleCountThreshold.GetValueOnGameThread())
+	{
+		// serialize our mesh
+		TArray<uint8> MeshWriteBuffer;
+		FMemoryWriter MemWriter(MeshWriteBuffer);
+		Mesh->Serialize(MemWriter);
+
+		FString Base64String = FBase64::Encode(MeshWriteBuffer);
+		//TArray<uint8> MeshReadBuffer;
+		//ensure(FBase64::Decode(Base64String, MeshReadBuffer));		// test decode
+
+		// Base64 encoding uses the '/' character, but T3D interprets '//' as some kind of
+		// terminator (?). If it occurs then the string passed to ImportCustomProperties() will
+		// come back as full of nullptrs. So we will swap in '-' here, and swap back to '/' in ImportCustomProperties()
+		for (int32 k = 0; k < Base64String.Len(); ++k)
+		{
+			if (Base64String[k] == '/')
+			{
+				Base64String[k] = '-';
+			}
+		}
+
+		Out.Logf(TEXT("MESHDATALEN=%d MESHDATA=%s"), Base64String.Len(), *Base64String);
+	}
+
+	Out.Logf(TEXT("\r\n"));
+}
+
+
+void UDynamicMesh::ImportCustomProperties(const TCHAR* SourceText, FFeedbackContext* Warn)
+{
+	Super::ImportCustomProperties(SourceText, Warn);
+
+	if (FParse::Command(&SourceText, TEXT("MeshData")))	
+	{
+		static const TCHAR MeshKeyToken[] = TEXT("MESHKEY=");
+		const TCHAR* FoundMeshKeyStart = FCString::Strifind(SourceText, MeshKeyToken);
+		if (FoundMeshKeyStart)
+		{
+			SourceText = FoundMeshKeyStart + FCString::Strlen(MeshKeyToken);
+			int32 MeshKey = FCString::Atoi(SourceText);
+			TUniquePtr<FDynamicMesh3> FoundMesh = UE::Local::FDynamicMeshCopyHelper::ExtractStashedMesh(MeshKey);
+			if (FoundMesh.IsValid() && FoundMesh->TriangleCount() > 0)
+			{
+				InitializeNewMesh();
+				SetMesh(MoveTemp(*FoundMesh));
+				return;
+			}
+		}
+
+		static const TCHAR MeshDataLenToken[] = TEXT("MESHDATALEN=");
+		const TCHAR* FoundMeshDataLenStart = FCString::Strifind(SourceText, MeshDataLenToken);
+		if (FoundMeshDataLenStart)
+		{
+			SourceText = FoundMeshDataLenStart + FCString::Strlen(MeshDataLenToken);
+			int32 MeshDataLen = FCString::Atoi(SourceText);
+
+			static const TCHAR MeshDataToken[] = TEXT("MESHDATA=");
+			const TCHAR* FoundMeshDataStart = FCString::Strifind(SourceText, MeshDataToken);
+			if (FoundMeshDataStart)
+			{
+				SourceText = FoundMeshDataStart + FCString::Strlen(MeshDataToken);
+				FString MeshData(MeshDataLen, SourceText);
+
+				// fix-up the hack applied to the Base64-encoded string in ExportCustomProperties()
+				for (int32 k = 0; k < MeshData.Len(); ++k)
+				{
+					if (MeshData[k] == '-')
+					{
+						MeshData[k] = '/';
+					}
+				}
+
+				TArray<uint8> MeshReadBuffer;
+				bool bDecoded = FBase64::Decode(MeshData, MeshReadBuffer);
+				if (bDecoded)
+				{
+					FMemoryReader MemReader(MeshReadBuffer);
+					FDynamicMesh3 NewMesh;
+					NewMesh.Serialize(MemReader);
+
+					InitializeNewMesh();
+					SetMesh(MoveTemp(NewMesh));
+					return;
+				}
+			}
+		}
+
+		// if we got here we failed. Rather than produce an empty mesh, we generate a small cube
+		UE_LOG(LogTemp, Warning, TEXT("UDynamicMesh text-based property serialization incomplete, generating box as placeholder. Try increasing geometry.DynamicMesh.TextBasedDupeTriThreshold, or geometry.DynamicMesh.DupeStashTimeout."))
+		FMinimalBoxMeshGenerator BoxGen;
+		BoxGen.Box = FOrientedBox3d(FVector3d::Zero(), 50.0 * FVector3d::One());
+		FDynamicMesh3 GenMesh(&BoxGen.Generate());
+		GenMesh.Attributes()->EnableMaterialID();
+		SetMesh(MoveTemp(GenMesh));
+	}
+}
+
+
+
+
+
+//
+// Pool support for blueprints
+//
+
+static TAutoConsoleVariable<int32> CVarDynamicMeshPoolMaxPoolSizeThreshold(
+	TEXT("geometry.DynamicMesh.MaxPoolSize"),
+	1000,
+	TEXT("Maximum number of meshes a UDynamicMeshPool will allow to be in the pool before running garbage collection"));
+
+
+
+UDynamicMesh* UDynamicMeshPool::RequestMesh()
+{
+	if (CachedMeshes.Num() > 0)
+	{
+		return CachedMeshes.Pop(false);
+	}
+	UDynamicMesh* NewMesh = NewObject<UDynamicMesh>();
+
+	// If we have allocated more meshes than our safety threshold, drop our holds on the existing meshes.
+	// This will allow them to be garbage-collected (eventually)
+	if (!ensure(AllCreatedMeshes.Num() < CVarDynamicMeshPoolMaxPoolSizeThreshold.GetValueOnGameThread()))
+	{
+		UE_LOG(LogTemp, Warning, TEXT("UDynamicMeshPool Threshold of %d Allocated Meshes exceeded! Releasing references to all current meshes and forcing a garbage collection."), CVarDynamicMeshPoolMaxPoolSizeThreshold.GetValueOnGameThread());
+		AllCreatedMeshes.Reset();
+		GEngine->ForceGarbageCollection(true);
+	}
+
+	AllCreatedMeshes.Add(NewMesh);
+	return NewMesh;
+}
+
+
+
+void UDynamicMeshPool::ReturnMesh(UDynamicMesh* Mesh)
+{
+	if (ensure(Mesh))
+	{
+		Mesh->Reset();
+		if (ensure(CachedMeshes.Contains(Mesh) == false))
+		{
+			CachedMeshes.Add(Mesh);
+		}
+	}
+}
+
+
+void UDynamicMeshPool::ReturnAllMeshes()
+{
+	CachedMeshes = AllCreatedMeshes;
+	for (UDynamicMesh* Mesh : CachedMeshes)
+	{
+		if (Mesh)
+		{
+			Mesh->Reset();
+		}
+	}
+
+	// TODO: this may be vestigial code, unclear how it could be hit
+	int32 Removed = CachedMeshes.RemoveAll([](UDynamicMesh* Mesh) { return Mesh == nullptr; });
+	ensure(Removed == 0);
+}
+
+void UDynamicMeshPool::FreeAllMeshes()
+{
+	CachedMeshes.Reset();
+	AllCreatedMeshes.Reset();
+}
\ No newline at end of file
diff --git a/Engine/Plugins/Experimental/MeshModelingToolset/Source/ModelingComponents/Public/BaseDynamicMeshComponent.h b/Engine/Plugins/Experimental/MeshModelingToolset/Source/ModelingComponents/Public/BaseDynamicMeshComponent.h
--- a/Engine/Plugins/Experimental/MeshModelingToolset/Source/ModelingComponents/Public/BaseDynamicMeshComponent.h
+++ b/Engine/Plugins/Experimental/MeshModelingToolset/Source/ModelingComponents/Public/BaseDynamicMeshComponent.h
@@ -3,14 +3,14 @@
 #pragma once
 
 #include "CoreMinimal.h"
-#include "UObject/ObjectMacros.h"
 #include "Components/MeshComponent.h"
 #include "InteractiveToolObjects.h"
 #include "Changes/MeshVertexChange.h"
 #include "Changes/MeshChange.h"
 #include "Changes/MeshReplacementChange.h"
 #include "MeshConversionOptions.h"
-#include "DynamicMesh3.h"
+#include "DynamicMesh3.h"		// todo replace with predeclaration (lots of fallout)
+#include "UDynamicMesh.h"
 
 #include "BaseDynamicMeshComponent.generated.h"
 
@@ -20,7 +20,6 @@ class FMeshVertexChange;
 class FMeshChange;
 using UE::Geometry::FDynamicMesh3;
 
-
 /**
  * EMeshRenderAttributeFlags is used to identify different mesh rendering attributes, for things
  * like fast-update functions
@@ -56,26 +55,35 @@ enum class EDynamicMeshTangentCalcType : uint8
 
 
 /**
- * UBaseDynamicMeshComponent is a base interface for a UMeshComponent based on a FDynamicMesh.
- * Currently no functionality lives here, only some interface functions are defined that various subclasses implement.
+ * UBaseDynamicMeshComponent is a base interface for a UMeshComponent based on a UDynamicMesh.
  */
-UCLASS(hidecategories = (LOD, Physics, Collision), editinlinenew, ClassGroup = Rendering)
-class MODELINGCOMPONENTS_API UBaseDynamicMeshComponent : public UMeshComponent, public IToolFrameworkComponent, public IMeshVertexCommandChangeTarget, public IMeshCommandChangeTarget, public IMeshReplacementCommandChangeTarget
+UCLASS(Abstract, hidecategories = (LOD, Physics, Collision), ClassGroup = Rendering)
+class MODELINGCOMPONENTS_API UBaseDynamicMeshComponent : 
+	public UMeshComponent, 
+	public IToolFrameworkComponent, 
+	public IMeshVertexCommandChangeTarget, 
+	public IMeshCommandChangeTarget, 
+	public IMeshReplacementCommandChangeTarget
 {
 	GENERATED_UCLASS_BODY()
 
-public:
 
+	//===============================================================================================================
+	// UBaseDynamicMeshComponent API. Subclasses must implement these functions
+	//
+public:
 	/**
 	 * initialize the internal mesh from a MeshDescription
+	 * @warning avoid usage of this function, access via GetDynamicMesh() instead
 	 */
-	virtual void InitializeMesh(FMeshDescription* MeshDescription)
+	virtual void InitializeMesh(const FMeshDescription* MeshDescription)
 	{
 		unimplemented();
 	}
 
 	/**
 	 * @return pointer to internal mesh
+	 * @warning avoid usage of this function, access via GetDynamicMesh() instead
 	 */
 	virtual FDynamicMesh3* GetMesh()
 	{
@@ -92,6 +100,16 @@ public:
 		return nullptr;
 	}
 
+	/**
+	 * @return the child UDynamicMesh
+	 */
+	UFUNCTION(BlueprintCallable, Category = "Dynamic Mesh Component")
+	virtual UDynamicMesh* GetDynamicMesh() 
+	{ 
+		unimplemented();
+		return nullptr; 
+	}
+
 	/**
 	 * Call this if you update the mesh via GetMesh()
 	 * @todo should provide a function that calls a lambda to modify the mesh, and only return const mesh pointer
@@ -125,12 +143,9 @@ public:
 		unimplemented();
 	}
 
-
-	//
-	// Modification support
-	//
-public:
-
+	/**
+	 * Apply a transform to the mesh
+	 */
 	virtual void ApplyTransform(const UE::Geometry::FTransform3d& Transform, bool bInvert)
 	{
 		unimplemented();
@@ -146,8 +161,6 @@ public:
 		unimplemented();
 	}
 
-
-
 protected:
 	/**
 	 * Subclass must implement this to notify allocated proxies of updated materials
@@ -159,36 +172,56 @@ protected:
 
 
 
+
+	//===============================================================================================================
+	// Built-in Wireframe-on-Shaded Rendering support. The wireframe looks terrible but this is a convenient 
+	// way to enable/disable it.
+	//
 public:
+	/**
+	 * if true, we always show the wireframe on top of the shaded mesh, even when not in wireframe mode
+	 * @todo: this should not be public, access via Set/Get once all usage is cleaned up
+	 */
+	UPROPERTY(EditAnywhere, BlueprintReadWrite, Category = "Dynamic Mesh Component")
+	bool bExplicitShowWireframe = false;
 
 	/**
 	 * Configure whether wireframe rendering is enabled or not
 	 */
-	virtual void SetEnableWireframeRenderPass(bool bEnable) { check(false); }
+	UFUNCTION(BlueprintCallable, Category = "Dynamic Mesh Component")
+	virtual void SetEnableWireframeRenderPass(bool bEnable) { bExplicitShowWireframe = bEnable; }
 
 	/**
-	 * @return true if wireframe rendering pass is enabled (default false)
+	 * @return true if wireframe rendering pass is enabled
 	 */
-	virtual bool EnableWireframeRenderPass() const { return false; }
+	UFUNCTION(BlueprintCallable, Category = "Dynamic Mesh Component")
+	virtual bool GetEnableWireframeRenderPass() const { return bExplicitShowWireframe; }
+
 
 
-	//
-	// Override rendering material support
-	//
 
+
+	//===============================================================================================================
+	// Override rendering material support. If an Override material is set, then it
+	// will be used during drawing of all mesh buffers except Secondary buffers.
+	//
+public:
 	/**
 	 * Set an active override render material. This should replace all materials during rendering.
 	 */
+	UFUNCTION(BlueprintCallable, Category = "Dynamic Mesh Component")
 	virtual void SetOverrideRenderMaterial(UMaterialInterface* Material);
 
 	/**
 	 * Clear any active override render material
 	 */
+	UFUNCTION(BlueprintCallable, Category = "Dynamic Mesh Component")
 	virtual void ClearOverrideRenderMaterial();
 	
 	/**
 	 * @return true if an override render material is currently enabled for the given MaterialIndex
 	 */
+	UFUNCTION(BlueprintCallable, Category = "Dynamic Mesh Component")
 	virtual bool HasOverrideRenderMaterial(int k) const
 	{
 		return OverrideRenderMaterial != nullptr;
@@ -197,32 +230,40 @@ public:
 	/**
 	 * @return active override render material for the given MaterialIndex
 	 */
+	UFUNCTION(BlueprintCallable, Category = "Dynamic Mesh Component")
 	virtual UMaterialInterface* GetOverrideRenderMaterial(int MaterialIndex) const
 	{
 		return OverrideRenderMaterial;
 	}
 
+protected:
+	UPROPERTY()
+	TObjectPtr<UMaterialInterface> OverrideRenderMaterial = nullptr;
 
 
 
-public:
 
-	//
-	// secondary buffer support
-	//
 
+	//===============================================================================================================
+	// Secondary Render Buffers support. This requires implementation in subclasses. It allows
+	// a subset of the mesh triangles to be moved to a separate set of render buffers, which
+	// can then have a separate material (eg to highlight faces), or be shown/hidden independently.
+	// 
+public:
 	/**
-	 * Set an active override render material. This should replace all materials during rendering.
+	 * Set an active secondary render material. 
 	 */
+	UFUNCTION(BlueprintCallable, Category = "Dynamic Mesh Component")
 	virtual void SetSecondaryRenderMaterial(UMaterialInterface* Material);
 
 	/**
-	 * Clear any active override render material
+	 * Clear any active secondary render material
 	 */
+	UFUNCTION(BlueprintCallable, Category = "Dynamic Mesh Component")
 	virtual void ClearSecondaryRenderMaterial();
 
 	/**
-	 * @return true if an override render material is currently enabled for the given MaterialIndex
+	 * @return true if a secondary render material is set
 	 */
 	virtual bool HasSecondaryRenderMaterial() const
 	{
@@ -230,47 +271,48 @@ public:
 	}
 
 	/**
-	 * @return active override render material for the given MaterialIndex
+	 * @return active secondary render material
 	 */
+	UFUNCTION(BlueprintCallable, Category = "Dynamic Mesh Component")
 	virtual UMaterialInterface* GetSecondaryRenderMaterial() const
 	{
 		return SecondaryRenderMaterial;
 	}
 
-
 	/**
 	 * Show/Hide the secondary triangle buffers. Does not invalidate SceneProxy.
 	 */
-	virtual void SetSecondaryBuffersVisibility(bool bVisible);
+	UFUNCTION(BlueprintCallable, Category = "Dynamic Mesh Component")
+	virtual void SetSecondaryBuffersVisibility(bool bSetVisible);
 
 	/**
 	 * @return true if secondary buffers are currently set to be visible
 	 */
+	UFUNCTION(BlueprintCallable, Category = "Dynamic Mesh Component")
 	virtual bool GetSecondaryBuffersVisibility() const;
 
-
 protected:
-
 	UPROPERTY()
-	TArray<UMaterialInterface*> BaseMaterials;
+	TObjectPtr<UMaterialInterface> SecondaryRenderMaterial = nullptr;
+
+	bool bDrawSecondaryBuffers = true;
 
-	UPROPERTY()
-	UMaterialInterface* OverrideRenderMaterial = nullptr;
 
-	UPROPERTY()
-	UMaterialInterface* SecondaryRenderMaterial = nullptr;
 
-	bool bDrawSecondaryBuffers = true;
 
-public:
 
+	//===============================================================================================================
+	// Standard Component interfaces
+	//
+public:
 
-	//~ Begin UMeshComponent Interface.
+	// UMeshComponent Interface.
 	virtual int32 GetNumMaterials() const override;
 	virtual UMaterialInterface* GetMaterial(int32 ElementIndex) const override;
 	virtual FMaterialRelevance GetMaterialRelevance(ERHIFeatureLevel::Type InFeatureLevel) const override;
 	virtual void SetMaterial(int32 ElementIndex, UMaterialInterface* Material) override;
 	virtual void GetUsedMaterials(TArray<UMaterialInterface*>& OutMaterials, bool bGetDebugMaterials = false) const override;
-	//~ End UMeshComponent Interface.
 
+	UPROPERTY()
+	TArray<TObjectPtr<UMaterialInterface>> BaseMaterials;
 };
diff --git a/Engine/Plugins/Experimental/MeshModelingToolset/Source/ModelingComponents/Public/OctreeDynamicMeshComponent.h b/Engine/Plugins/Experimental/MeshModelingToolset/Source/ModelingComponents/Public/OctreeDynamicMeshComponent.h
--- a/Engine/Plugins/Experimental/MeshModelingToolset/Source/ModelingComponents/Public/OctreeDynamicMeshComponent.h
+++ b/Engine/Plugins/Experimental/MeshModelingToolset/Source/ModelingComponents/Public/OctreeDynamicMeshComponent.h
@@ -5,7 +5,7 @@
 #include "CoreMinimal.h"
 #include "BaseDynamicMeshComponent.h"
 #include "MeshConversionOptions.h"
-
+#include "UDynamicMesh.h"
 #include "DynamicMeshOctree3.h"
 #include "Util/IndexSetDecompositions.h"
 
@@ -41,26 +41,21 @@ public:
 	/**
 	 * initialize the internal mesh from a MeshDescription
 	 */
-	virtual void InitializeMesh(FMeshDescription* MeshDescription) override;
+	virtual void InitializeMesh(const FMeshDescription* MeshDescription) override;
 
 	/**
 	 * @return pointer to internal mesh
 	 */
-	virtual FDynamicMesh3* GetMesh() override { return Mesh.Get(); }
+	virtual FDynamicMesh3* GetMesh() override { return MeshObject->GetMeshPtr(); }
 
 	/**
 	 * @return pointer to internal mesh
 	 */
-	virtual const FDynamicMesh3* GetMesh() const override { return Mesh.Get(); }
+	virtual const FDynamicMesh3* GetMesh() const override { return MeshObject->GetMeshPtr(); }
 
 
 	UE::Geometry::FDynamicMeshOctree3* GetOctree() { return Octree.Get(); }
 
-	/**
-	 * @return the current internal mesh, which is replaced with an empty mesh
-	 */
-	TUniquePtr<FDynamicMesh3> ExtractMesh(bool bNotifyUpdate);
-
 	/**
 	 * Write the internal mesh to a MeshDescription
 	 * @param bHaveModifiedTopology if false, we only update the vertex positions in the MeshDescription, otherwise it is Empty()'d and regenerated entirely
@@ -116,22 +111,6 @@ public:
 	 */
 	FSimpleMulticastDelegate OnMeshChanged;
 
-	/**
-	 * if true, we always show the wireframe on top of the shaded mesh, even when not in wireframe mode
-	 */
-	UPROPERTY()
-	bool bExplicitShowWireframe = false;
-
-	/**
-	 * Configure whether wireframe rendering is enabled or not
-	 */
-	virtual void SetEnableWireframeRenderPass(bool bEnable) override { bExplicitShowWireframe = bEnable; }
-
-	/**
-	 * @return true if wireframe rendering pass is enabled
-	 */
-	virtual bool EnableWireframeRenderPass() const override { return bExplicitShowWireframe; }
-
 	/**
 	 * If this function is set, we will use these colors instead of vertex colors
 	 */
@@ -156,7 +135,15 @@ private:
 	virtual FBoxSphereBounds CalcBounds(const FTransform& LocalToWorld) const override;
 	//~ Begin USceneComponent Interface.
 
-	TUniquePtr<FDynamicMesh3> Mesh;
+	UPROPERTY()
+	UDynamicMesh* MeshObject;
+
+	FDelegateHandle MeshObjectChangedHandle;
+	void OnMeshObjectChanged(UDynamicMesh* ChangedMeshObject, FDynamicMeshChangeInfo ChangeInfo);
+
+	FDelegateHandle PreMeshChangeHandle;
+	void OnPreMeshObjectChanged(UDynamicMesh* ChangedMeshObject, FDynamicMeshChangeInfo ChangeInfo);
+
 	TUniquePtr<UE::Geometry::FDynamicMeshOctree3> Octree;
 	TUniquePtr<UE::Geometry::FDynamicMeshOctree3::FTreeCutSet> OctreeCut;
 	UE::Geometry::FArrayIndexSetsDecomposition TriangleDecomposition;
@@ -168,11 +155,16 @@ private:
 	TArray<FCutCellIndexSet> CutCellSetMap;
 	int32 SpillDecompSetID;
 
-	void InitializeNewMesh();
-
-
 
 	FColor GetTriangleColor(int TriangleID);
 
 	//friend class FCustomMeshSceneProxy;
+
+
+public:
+	//UFUNCTION(BlueprintCallable, Category = "DynamicMesh")
+	virtual UDynamicMesh* GetDynamicMesh() override { return MeshObject; }
+
+	UFUNCTION(BlueprintCallable, Category = "DynamicMesh")
+	void SetDynamicMesh(UDynamicMesh* NewMesh);
 };
diff --git a/Engine/Plugins/Experimental/MeshModelingToolset/Source/ModelingComponents/Public/PreviewMesh.h b/Engine/Plugins/Experimental/MeshModelingToolset/Source/ModelingComponents/Public/PreviewMesh.h
--- a/Engine/Plugins/Experimental/MeshModelingToolset/Source/ModelingComponents/Public/PreviewMesh.h
+++ b/Engine/Plugins/Experimental/MeshModelingToolset/Source/ModelingComponents/Public/PreviewMesh.h
@@ -207,15 +207,15 @@ public:
 	void SetVisible(bool bVisible);
 
 
-	/** Render data update hint */
+	/** Render data update hint (values mirror EDynamicMeshComponentRenderUpdateMode) */
 	enum class ERenderUpdateMode
 	{
 		/** Do not update render data */
-		NoUpdate,	
+		NoUpdate = 0,	
 		/** Invalidate overlay of internal component, rebuilding all render data */
-		FullUpdate,
+		FullUpdate = 1,
 		/** Attempt to do partial update of render data if possible */
-		FastUpdate
+		FastUpdate = 2
 	};
 
 	/**
@@ -431,10 +431,6 @@ public:
 	UPROPERTY()
 	bool bBuildSpatialDataStructure;
 
-	// results in component drawing w/o z-testing and with editor compositing. Do not recommend using this flag, will be deprecated/removed.
-	UPROPERTY()
-	bool bDrawOnTop;
-
 protected:
 	/** The temporary actor we create internally to own the preview mesh component */
 	APreviewMeshActor* TemporaryParentActor = nullptr;
diff --git a/Engine/Plugins/Experimental/MeshModelingToolset/Source/ModelingComponents/Public/SimpleDynamicMeshComponent.h b/Engine/Plugins/Experimental/MeshModelingToolset/Source/ModelingComponents/Public/SimpleDynamicMeshComponent.h
--- a/Engine/Plugins/Experimental/MeshModelingToolset/Source/ModelingComponents/Public/SimpleDynamicMeshComponent.h
+++ b/Engine/Plugins/Experimental/MeshModelingToolset/Source/ModelingComponents/Public/SimpleDynamicMeshComponent.h
@@ -9,6 +9,7 @@
 #include "MeshTangents.h"
 #include "TransformTypes.h"
 #include "Async/Future.h"
+#include "UDynamicMesh.h"
 
 #include "SimpleDynamicMeshComponent.generated.h"
 
@@ -33,108 +34,84 @@ public:
 };
 
 
+/** Render data update hint */
+UENUM()
+enum class EDynamicMeshComponentRenderUpdateMode
+{
+	/** Do not update render data */
+	NoUpdate = 0,
+	/** Invalidate overlay of internal component, rebuilding all render data */
+	FullUpdate = 1,
+	/** Attempt to do partial update of render data if possible */
+	FastUpdate = 2
+};
+
+
 /** 
  * USimpleDynamicMeshComponent is a mesh component similar to UProceduralMeshComponent,
- * except it bases the renderable geometry off an internal FDynamicMesh3 instance.
+ * except it bases the renderable geometry off an internal UDynamicMesh instance (which
+ * encapsulates a FDynamicMesh3). 
  * 
- * There is some support for undo/redo on the component (@todo is this the right place?)
+ * There is extensive support for partial updates to render buffers, customizing colors,
+ * internally decomposing the mesh into separate chunks for more efficient render updates,
+ * and support for attaching a 'Postprocessor' to generate a render mesh on-the-fly
+ * See comment sections below for details.
  * 
- * This component draws wireframe-on-shaded when Wireframe is enabled, or when bExplicitShowWireframe = true
- *
  */
-UCLASS(hidecategories = (LOD, Physics, Collision), editinlinenew, ClassGroup = Rendering)
+UCLASS(hidecategories = (LOD, Physics, Collision), meta = (BlueprintSpawnableComponent), ClassGroup = Rendering)
 class MODELINGCOMPONENTS_API USimpleDynamicMeshComponent : public UBaseDynamicMeshComponent
 {
 	GENERATED_UCLASS_BODY()
 
 
-public:
-	/** How should Tangents be calculated/handled */
-	UPROPERTY()
-	EDynamicMeshTangentCalcType TangentsType = EDynamicMeshTangentCalcType::NoTangents;
 
+	//===============================================================================================================
+	// Mesh Access. Usage via GetDynamicMesh() is preferred, the GetMesh() pass-throughs exist largely to support
+	// existing code from before UDynamicMesh was added.
 public:
-	/**
-	 * initialize the internal mesh from a MeshDescription
-	 */
-	virtual void InitializeMesh(FMeshDescription* MeshDescription) override;
-
 	/**
 	 * @return pointer to internal mesh
+	 * @warning avoid usage of this function, access via GetDynamicMesh() instead
 	 */
-	virtual FDynamicMesh3* GetMesh() override { return Mesh.Get(); }
+	virtual FDynamicMesh3* GetMesh() override { return MeshObject->GetMeshPtr(); }
 
 	/**
 	 * @return pointer to internal mesh
+	 * @warning avoid usage of this function, access via GetDynamicMesh() instead
 	 */
-	virtual const FDynamicMesh3* GetMesh() const override { return Mesh.Get(); }
-
-	/*
-	* The SceneProxy should call these functions to get the post-processed RenderMesh. (See IRenderMeshPostProcessor.)
-	*/
-	virtual FDynamicMesh3* GetRenderMesh();
-
-	/*
-	* The SceneProxy should call these functions to get the post-processed RenderMesh. (See IRenderMeshPostProcessor.)
-	*/
-	virtual const FDynamicMesh3* GetRenderMesh() const;
-
-	/**
-	 * @return the current internal mesh, which is replaced with an empty mesh
-	 */
-	TUniquePtr<FDynamicMesh3> ExtractMesh(bool bNotifyUpdate);
-
-	/**
-	 * Copy externally-calculated tangents into the internal tangets buffer.
-	 * @param bFastUpdateIfPossible if true, will try to do a fast normals/tangets update of the SceneProxy, instead of full invalidatiohn
-	 */
-	void UpdateTangents(const UE::Geometry::FMeshTangentsf* ExternalTangents, bool bFastUpdateIfPossible);
+	virtual const FDynamicMesh3* GetMesh() const override { return MeshObject->GetMeshPtr(); }
 
 	/**
-	 * Copy externally-calculated tangents into the internal tangets buffer.
-	 * @param bFastUpdateIfPossible if true, will try to do a fast normals/tangets update of the SceneProxy, instead of full invalidatiohn
+	 * @return the child UDynamicMesh
 	 */
-	void UpdateTangents(const UE::Geometry::FMeshTangentsd* ExternalTangents, bool bFastUpdateIfPossible);
-
+	//UFUNCTION(BlueprintCallable, Category = "Dynamic Mesh Component")
+	virtual UDynamicMesh* GetDynamicMesh() override { return MeshObject; }
 
 	/**
-	 * @return pointer to internal tangents object. 
-	 * @warning calling this with TangentsType = AutoCalculated will result in possibly-expensive Tangents calculation
-	 * @warning this is only currently safe to call on the Game Thread!!
+	 * Set the child UDynamicMesh. This can be used to 'share' a UDynamicMesh between Component instances.
+	 * @warning Currently this is somewhat risky, it is on the caller/clients to make sure that the actual mesh is not being simultaneously modified on multiple threads
 	 */
-	const UE::Geometry::FMeshTangentsf* GetTangents();
-
-
-
+	UFUNCTION(BlueprintCallable, Category = "Dynamic Mesh Component")
+	void SetDynamicMesh(UDynamicMesh* NewMesh);
 
+protected:
 	/**
-	 * Write the internal mesh to a MeshDescription
-	 * @param bHaveModifiedTopology if false, we only update the vertex positions in the MeshDescription, otherwise it is Empty()'d and regenerated entirely
-	 * @param ConversionOptions struct of additional options for the conversion
+	 * Internal FDynamicMesh is stored inside a UDynamicMesh container, which allows it to be
+	 * used from BP, shared with other UObjects, and so on
 	 */
-	virtual void Bake(FMeshDescription* MeshDescription, bool bHaveModifiedTopology, const FConversionToMeshDescriptionOptions& ConversionOptions) override;
+	UPROPERTY(Instanced)
+	TObjectPtr<UDynamicMesh> MeshObject;
 
-	/**
-	* Write the internal mesh to a MeshDescription with default conversion options
-	* @param bHaveModifiedTopology if false, we only update the vertex positions in the MeshDescription, otherwise it is Empty()'d and regenerated entirely
-	*/
-	void Bake(FMeshDescription* MeshDescription, bool bHaveModifiedTopology)
-	{
-		FConversionToMeshDescriptionOptions ConversionOptions;
-		Bake(MeshDescription, bHaveModifiedTopology, ConversionOptions);
-	}
 
-	/**
-	 * Apply transform to internal mesh. Updates Octree and RenderProxy if available.
-	 * @param bInvert if true, inverse tranform is applied instead of forward transform
-	 */
-	virtual void ApplyTransform(const UE::Geometry::FTransform3d& Transform, bool bInvert) override;
 
 
+	//===============================================================================================================
+	// RenderBuffer Update API. These functions can be used by external code (and internally in some places)
+	// to tell the Component that the Mesh data has been modified in some way, and that the RenderBuffers in the RenderProxy
+	// need to be updated (or rebuilt entirely). On large meshes a full rebuild is expensive, so there are quite a few
+	// variants that can be used to minimize the amount of data updated in different situations.
 	//
-	// change tracking/etc
-	//
-
+public:
 	/**
 	 * Call this if you update the mesh via GetMesh(). This will destroy the existing RenderProxy and create a new one.
 	 * @todo should provide a function that calls a lambda to modify the mesh, and only return const mesh pointer
@@ -189,15 +166,14 @@ public:
 	 */
 	void FastNotifyTriangleVerticesUpdated(const TSet<int32>& Triangles, EMeshRenderAttributeFlags UpdatedAttributes);
 
-
 	/**
 	 * If a Decomposition is set on this Component, and everything is currently valid (proxy/etc), precompute the set of
 	 * buffers that will be modified, as well as the bounds of the modified region. These are both computed in parallel.
+	 * Use FastNotifyTriangleVerticesUpdated_ApplyPrecompute() with the returned future to apply this precomputation.
 	 * @return a future that will (eventually) return true if the precompute is OK, and (immediately) false if it is not
 	 */
 	TFuture<bool> FastNotifyTriangleVerticesUpdated_TryPrecompute(const TArray<int32>& Triangles, TArray<int32>& UpdateSetsOut, UE::Geometry::FAxisAlignedBox3d& BoundsOut);
 
-
 	/**
 	 * This function updates vertex positions/attributes of existing SceneProxy render buffers if possible, for the given triangles.
 	 * The assumption is that FastNotifyTriangleVerticesUpdated_TryPrecompute() was used to get the Precompute future, this function
@@ -210,63 +186,102 @@ public:
 
 
 
-	/** If false, we don't completely invalidate the RenderProxy when ApplyChange() is called (assumption is it will be handled elsewhere) */
-	UPROPERTY()
-	bool bInvalidateProxyOnChange = true;
 
+
+	//===============================================================================================================
+	// Change Support. These changes are primarily used for Undo/Redo, however there is no strict assumption
+	// about this internally, objects of these change types could also be used to perform more structured editing.
+	// (Note that these functions simply forward the change events to the child UDynamicMesh, which will
+	// 	post a mesh-change event that 
+	//
+public:
 	/**
-	 * Apply a vertex deformation change to the internal mesh
+	 * Apply a vertex deformation change to the mesh
 	 */
 	virtual void ApplyChange(const FMeshVertexChange* Change, bool bRevert) override;
 
 	/**
-	 * Apply a general mesh change to the internal mesh
+	 * Apply a general mesh change to the mesh
 	 */
 	virtual void ApplyChange(const FMeshChange* Change, bool bRevert) override;
 
 	/**
-	* Apply a general mesh replacement change to the internal mesh
+	* Apply a mesh replacement change to mesh
 	*/
 	virtual void ApplyChange(const FMeshReplacementChange* Change, bool bRevert) override;
 
-
 	/**
-	 * This delegate fires when a FCommandChange is applied to this component, so that
-	 * parent objects know the mesh has changed.
+	 * This delegate fires when the mesh has been changed
 	 */
 	FSimpleMulticastDelegate OnMeshChanged;
 
-
 	/**
-	 * This delegate fires when ApplyChange(FMeshVertexChange) executes
+	 * This delegate fires when the mesh vertices have been changed via an FMeshVertexChange
 	 */
 	DECLARE_MULTICAST_DELEGATE_ThreeParams(FMeshVerticesModified, USimpleDynamicMeshComponent*, const FMeshVertexChange*, bool);
 	FMeshVerticesModified OnMeshVerticesChanged;
 
-
 	/**
-	 * if true, we always show the wireframe on top of the shaded mesh, even when not in wireframe mode
+	 * When a FMeshChange or FMeshVertexChange is applied, by default we currently fully invalidate the render proxy. However in certain
+	 * realtime situations (eg like Sculpting tools) it can be critical to undo/redo performance to do more optimized render data updates
+	 * (eg using one of the FastXYZ functions above). To allow for that, the full proxy invalidation on change can be (temporarily!) disabled 
+	 * using this function. 
 	 */
-	UPROPERTY()
-	bool bExplicitShowWireframe = false;
+	void SetInvalidateProxyOnChangeEnabled(bool bEnabled);
 
-	/**
-	 * Configure whether wireframe rendering is enabled or not
-	 */
-	virtual void SetEnableWireframeRenderPass(bool bEnable) override { bExplicitShowWireframe = bEnable; }
+	/** @return true if InvalidateProxyOnChange is enabled (default) */
+	bool GetInvalidateProxyOnChangeEnabled() const { return bInvalidateProxyOnChange; }
 
-	/**
-	 * @return true if wireframe rendering pass is enabled
-	 */
-	virtual bool EnableWireframeRenderPass() const override { return bExplicitShowWireframe; }
+protected:
+	/** If false, we don't completely invalidate the RenderProxy when ApplyChange() is called (assumption is it will be handled elsewhere) */
+	bool bInvalidateProxyOnChange = true;
 
+	/** Handle for OnMeshObjectChanged which is registered with MeshObject::OnMeshChanged delegate */
+	FDelegateHandle MeshObjectChangedHandle;
+
+	/** Called whenever internal MeshObject is modified, fires OnMeshChanged and OnMeshVerticesChanged above */
+	void OnMeshObjectChanged(UDynamicMesh* ChangedMeshObject, FDynamicMeshChangeInfo ChangeInfo);
 
-	/**
-	 * If this function is set, we will use these colors instead of vertex colors
-	 */
-	TFunction<FColor(const FDynamicMesh3*, int)> TriangleColorFunc = nullptr;
 
 
+
+	//===============================================================================================================
+	// Support for specifying per-triangle colors as vertex colors. This allows external code to dynamically override 
+	// the vertex colors on the rendered mesh. The lambda that is passed is held for the lifetime of the Component and
+	// must remain valid. A Material that uses the vertex colors must be applied, otherwise setting this override will
+	// have no visible effect. If the colors change externally, FastNotifyColorsUpdated() can be used to do the 
+	// minimal vertex buffer updates necessary in the RenderProxy
+	//
+public:
+	/** Clear an active triangle color function if one exists, and update the mesh */
+	virtual void SetTriangleColorFunction(TUniqueFunction<FColor(const FDynamicMesh3*, int)> TriangleColorFuncIn,
+										  EDynamicMeshComponentRenderUpdateMode UpdateMode = EDynamicMeshComponentRenderUpdateMode::FastUpdate);
+
+	/** Clear an active triangle color function if one exists, and update the mesh */
+	virtual void ClearTriangleColorFunction(EDynamicMeshComponentRenderUpdateMode UpdateMode = EDynamicMeshComponentRenderUpdateMode::FastUpdate);
+
+	/** @return true if a triangle color function is configured */
+	virtual bool HasTriangleColorFunction();
+
+protected:
+	/** If this function is set, we will use these colors instead of vertex colors */
+	TUniqueFunction<FColor(const FDynamicMesh3*, int)> TriangleColorFunc = nullptr;
+
+	/** This function is passed via lambda to the RenderProxy to be able to access TriangleColorFunc */
+	FColor GetTriangleColor(const FDynamicMesh3* Mesh, int TriangleID);
+
+
+
+
+
+	//===============================================================================================================
+	// Support for Secondary triangle index buffers. When this is configured, then triangles identified
+	// by the filtering predicate function will be placed in a second set of RenderBuffers at the SceneProxy level.
+	// This can be combined with the SecondaryRenderMaterial support in UBaseDynamicMeshComponent to draw
+	// that triangle set with a different material, to efficiently accomplish UI features like highlighting a
+	// subset of mesh triangles.
+	//
+public:
 	/**
 	 * If Secondary triangle buffers are enabled, then we will filter triangles that pass the given predicate
 	 * function into a second index buffer. These triangles will be drawn with the Secondary render material
@@ -279,36 +294,96 @@ public:
 	 */
 	virtual void DisableSecondaryTriangleBuffers();
 
+protected:
+	TUniqueFunction<bool(const FDynamicMesh3*, int32)> SecondaryTriFilterFunc = nullptr;
+
+
+
+
+	//===============================================================================================================
+	// Support for a Render Decomposition, which is basically a segmentation of the mesh triangles into
+	// subsets which will be turned into separate RenderBuffers in the Render Proxy. If this is configured,
+	// then various of the FastNotifyXYZUpdated() functions above will only need to rebuild the RenderBuffers
+	// that include affected triangles. The FMeshRenderDecomposition implementation has various options for
+	// building decompositions based on material, spatial clustering, etc.
+	//
+public:
 	/**
-	 * Configure a decomposition of the mesh, which will result in separate render buffers for each decomposition triangle group.
-	 * Invalidates existing SceneProxy.
+	 * Configure a decomposition of the mesh, which will result in separate render buffers for each 
+	 * decomposition triangle group. Invalidates existing SceneProxy.
 	 */
 	virtual void SetExternalDecomposition(TUniquePtr<FMeshRenderDecomposition> Decomposition);
 
+protected:
+	TUniquePtr<FMeshRenderDecomposition> Decomposition;
+
+
+
+
+	//===============================================================================================================
+	// IRenderMeshPostProcessor Support. If a RenderMesh Postprocessor is configured, then instead of directly 
+	// passing the internal mesh to the RenderProxy, IRenderMeshPostProcessor::PostProcess is applied to populate
+	// the internal RenderMesh which is passed instead. This allows things like Displacement or Subdivision to be 
+	// done on-the-fly at the rendering level (which is potentially more efficient).
+	//
+public:
 	/**
 	 * Add a render mesh processor, to be called before the mesh is sent for rendering.
 	 */
 	virtual void SetRenderMeshPostProcessor(TUniquePtr<IRenderMeshPostProcessor> Processor);
 
+	/**
+	 * The SceneProxy should call these functions to get the post-processed RenderMesh. (See IRenderMeshPostProcessor.)
+	 */
+	virtual FDynamicMesh3* GetRenderMesh();
+
+	/**
+	 * The SceneProxy should call these functions to get the post-processed RenderMesh. (See IRenderMeshPostProcessor.)
+	 */
+	virtual const FDynamicMesh3* GetRenderMesh() const;
+
+protected:
+	TUniquePtr<IRenderMeshPostProcessor> RenderMeshPostProcessor;
+	TUniquePtr<FDynamicMesh3> RenderMesh;
+
+
+
+
+	//===============================================================================================================
+	// Support for Component attachment change notifications via delegates. Standard UE
+	// Actor/Component hierarchy does not generally provide these capabilities, but in some use
+	// cases (eg procedural mesh Actors) we need to know things like when the Component set inside
+	// an Actor is modified.
 public:
+	DECLARE_MULTICAST_DELEGATE_TwoParams(FComponentChildrenChangedDelegate, USceneComponent*, bool);
+
+	/**
+	 * The OnChildAttached() and OnChildDetached() implementations (from USceneComponent API) broadcast this delegate. This
+	 * allows Actors that have USimpleDynamicMeshComponent's to respond to changes in their Component hierarchy.
+	 */
+	FComponentChildrenChangedDelegate OnChildAttachmentModified;
 
-	// do not use this
-	UPROPERTY()
-	bool bDrawOnTop = false;
 
-	// do not use this
-	void SetDrawOnTop(bool bSet);
 
 
 protected:
+
+	//
+	// standard Component internals, for computing bounds and managing the SceneProxy
+	//
+
+	/** Current local-space bounding box of Mesh */
+	UE::Geometry::FAxisAlignedBox3d LocalBounds;
+
+	/** Recompute LocalBounds from the current Mesh */
+	void UpdateLocalBounds();
+
 	/**
 	 * This is called to tell our RenderProxy about modifications to the material set.
 	 * We need to pass this on for things like material validation in the Editor.
 	 */
 	virtual void NotifyMaterialSetUpdated();
 
-private:
-
 	/**
 	 * If the render proxy is invalidated (eg by MarkRenderStateDirty()), it will be destroyed at the end of
 	 * the frame, but the base SceneProxy pointer is not nulled out immediately. As a result if we call various
@@ -317,35 +392,101 @@ private:
 	 */
 	bool bProxyValid = false;
 
-	/** @return current render proxy, if valid, otherwise nullptr */
-	FSimpleDynamicMeshSceneProxy* GetCurrentSceneProxy() { return (bProxyValid) ? (FSimpleDynamicMeshSceneProxy*)SceneProxy : nullptr; }
+	/** 
+	 * @return current render proxy, if valid, otherwise nullptr 
+	 */
+	FSimpleDynamicMeshSceneProxy* GetCurrentSceneProxy();
 
-	// Called from NotifyMeshUpdated, as well as the FastNotify functions if needed
+	/**
+	 * Fully invalidate all rendering data for this Component. Current Proxy will be discarded, Bounds and possibly Tangents recomputed, etc
+	 */
 	void ResetProxy();
 
 	//~ Begin UPrimitiveComponent Interface.
 	virtual FPrimitiveSceneProxy* CreateSceneProxy() override;
-	//~ End UPrimitiveComponent Interface.
 
-	//~ Begin USceneComponent Interface.
+	//~ USceneComponent Interface.
 	virtual FBoxSphereBounds CalcBounds(const FTransform& LocalToWorld) const override;
-	//~ Begin USceneComponent Interface.
+	virtual void OnChildAttached(USceneComponent* ChildComponent) override;
+	virtual void OnChildDetached(USceneComponent* ChildComponent) override;
 
-	TUniquePtr<FDynamicMesh3> Mesh;
-	void InitializeNewMesh();
+	//~ UObject Interface.
+	virtual void PostLoad() override;
 
-	// local-space bounding of Mesh
-	UE::Geometry::FAxisAlignedBox3d LocalBounds;
 
-	bool bTangentsValid = false;
-	UE::Geometry::FMeshTangentsf Tangents;
-	
-	FColor GetTriangleColor(const FDynamicMesh3* Mesh, int TriangleID);
 
-	TUniqueFunction<bool(const FDynamicMesh3*, int32)> SecondaryTriFilterFunc = nullptr;
 
-	TUniquePtr<FMeshRenderDecomposition> Decomposition;
 
-	TUniquePtr<IRenderMeshPostProcessor> RenderMeshPostProcessor;
-	TUniquePtr<FDynamicMesh3> RenderMesh;
+
+
+
+	///===============================================================================================================
+	// The functions below should be considered DEPRECATED
+	// They are slated for removal (but require some 
+	// refactoring for that to be possible)
+	//
+public:
+
+	/**
+	 * initialize the internal mesh from a MeshDescription
+	 */
+	virtual void InitializeMesh(const FMeshDescription* MeshDescription) override;
+
+
+	/** How should Tangents be calculated/handled */
+	UPROPERTY()
+	EDynamicMeshTangentCalcType TangentsType = EDynamicMeshTangentCalcType::NoTangents;
+
+	//
+	// Separate Tangents support (todo: remove this)
+	//
+
+	/**
+	 * Copy externally-calculated tangents into the internal tangets buffer.
+	 * @param bFastUpdateIfPossible if true, will try to do a fast normals/tangets update of the SceneProxy, instead of full invalidatiohn
+	 */
+	void UpdateTangents(const UE::Geometry::FMeshTangentsf* ExternalTangents, bool bFastUpdateIfPossible);
+
+	/**
+	 * Copy externally-calculated tangents into the internal tangets buffer.
+	 * @param bFastUpdateIfPossible if true, will try to do a fast normals/tangets update of the SceneProxy, instead of full invalidatiohn
+	 */
+	void UpdateTangents(const UE::Geometry::FMeshTangentsd* ExternalTangents, bool bFastUpdateIfPossible);
+
+	/**
+	 * @return pointer to internal tangents object.
+	 * @warning calling this with TangentsType = AutoCalculated will result in possibly-expensive Tangents calculation
+	 * @warning this is only currently safe to call on the Game Thread!!
+	 */
+	const UE::Geometry::FMeshTangentsf* GetTangents();
+
+
+	/**
+	 * Write the internal mesh to a MeshDescription
+	 * @param bHaveModifiedTopology if false, we only update the vertex positions in the MeshDescription, otherwise it is Empty()'d and regenerated entirely
+	 * @param ConversionOptions struct of additional options for the conversion
+	 */
+	virtual void Bake(FMeshDescription* MeshDescription, bool bHaveModifiedTopology, const FConversionToMeshDescriptionOptions& ConversionOptions) override;
+
+	/**
+	* Write the internal mesh to a MeshDescription with default conversion options
+	* @param bHaveModifiedTopology if false, we only update the vertex positions in the MeshDescription, otherwise it is Empty()'d and regenerated entirely
+	*/
+	void Bake(FMeshDescription* MeshDescription, bool bHaveModifiedTopology)
+	{
+		FConversionToMeshDescriptionOptions ConversionOptions;
+		Bake(MeshDescription, bHaveModifiedTopology, ConversionOptions);
+	}
+
+	/**
+	 * Apply transform to internal mesh. Updates Octree and RenderProxy if available.
+	 * @param bInvert if true, inverse tranform is applied instead of forward transform
+	 */
+	virtual void ApplyTransform(const UE::Geometry::FTransform3d& Transform, bool bInvert) override;
+
+
+protected:
+	/** Separate tangents storage (should be replaced with DynamicMesh Tangents attribute overlays) */
+	bool bTangentsValid = false;
+	UE::Geometry::FMeshTangentsf Tangents;
 };
diff --git a/Engine/Plugins/Experimental/MeshModelingToolset/Source/ModelingComponents/Public/UDynamicMesh.h b/Engine/Plugins/Experimental/MeshModelingToolset/Source/ModelingComponents/Public/UDynamicMesh.h
new file mode 100644
--- /dev/null
+++ b/Engine/Plugins/Experimental/MeshModelingToolset/Source/ModelingComponents/Public/UDynamicMesh.h
@@ -0,0 +1,361 @@
+// Copyright Epic Games, Inc. All Rights Reserved.
+
+#pragma once
+
+#include "CoreMinimal.h"
+#include "DynamicMesh3.h"
+
+#include "Changes/MeshVertexChange.h"
+#include "Changes/MeshChange.h"
+#include "Changes/MeshReplacementChange.h"
+
+#include "UDynamicMesh.generated.h"
+
+
+
+/**
+ * UDynamicMeshGenerator is an abstract base class for an implementation that can
+ * mutate an input mesh into an output mesh. A subclass of this class can be attached to
+ * a UDynamicMesh to allow for arbitrarily-complex procedural generation
+ */
+UCLASS(Abstract)
+class MODELINGCOMPONENTS_API UDynamicMeshGenerator : public UObject
+{
+	GENERATED_BODY()
+public:
+	/** Apply a change to MeshInOut */
+	virtual void Generate(FDynamicMesh3& MeshInOut) { };
+};
+
+
+
+
+
+/**
+ * EDynamicMeshChangeType is used by FDynamicMeshChangeInfo to indicate a "type" of mesh change
+ */
+UENUM(BlueprintType)
+enum class EDynamicMeshChangeType : uint8
+{
+	GeneralEdit = 0,
+	MeshChange = 1,
+	MeshReplacementChange = 2,
+	MeshVertexChange = 3
+};
+
+
+/**
+ * FDynamicMeshChangeInfo stores information about a change to a UDynamicMesh.
+ * This struct is emitted by the UDynamicMesh OnPreMeshChanged() and OnMeshChanged() delegates.
+ */
+USTRUCT(BlueprintType)
+struct MODELINGCOMPONENTS_API FDynamicMeshChangeInfo
+{
+	GENERATED_BODY()
+
+	UPROPERTY(BlueprintReadOnly, EditAnywhere, Category = "DynamicMeshChangeInfo")
+	EDynamicMeshChangeType Type = EDynamicMeshChangeType::GeneralEdit;
+
+	// for changes that are an FChange, indicates whether this is an 'Apply' or 'Revert' of the FChange
+	UPROPERTY(BlueprintReadOnly, EditAnywhere, Category = "DynamicMeshChangeInfo")
+	bool bIsRevertChange = false;
+
+	//
+	// internals
+	//
+	const FMeshChange* MeshChange = nullptr;
+	const FMeshReplacementChange* ReplaceChange = nullptr;
+	const FMeshVertexChange* VertexChange = nullptr;
+};
+
+// These delegates are used by UDynamicMesh
+DECLARE_MULTICAST_DELEGATE_TwoParams(FOnDynamicMeshChanged, UDynamicMesh*, FDynamicMeshChangeInfo);
+DECLARE_DYNAMIC_MULTICAST_DELEGATE_OneParam(FOnDynamicMeshModifiedBP, UDynamicMesh*, Mesh);
+
+
+
+/**
+ * UDynamicMesh is a UObject container for a FDynamicMesh3. 
+ */
+UCLASS(BlueprintType)
+class MODELINGCOMPONENTS_API UDynamicMesh : public UObject, 
+	public IMeshVertexCommandChangeTarget, 
+	public IMeshCommandChangeTarget, 
+	public IMeshReplacementCommandChangeTarget
+{
+	GENERATED_UCLASS_BODY()
+public:
+
+	/**
+	 * Clear the internal mesh to an empty mesh.
+	 * This *does not* allocate a new mesh, so any existing mesh pointers/refs are still valid
+	 */
+	UFUNCTION(BlueprintCallable, Category = "Dynamic Mesh")
+	void Reset();
+
+	/**
+	 * Clear the internal mesh to a 100x100x100 cube with base at the origin.
+	 * This this instead of Reset() if an initially-empty mesh is undesirable (eg for a Component)
+	 */
+	UFUNCTION(BlueprintCallable, Category = "Dynamic Mesh")
+	void ResetToCube();
+
+	//
+	// Native access/modification functions
+	//
+
+	/**
+	 * Reset the internal mesh data and then optionally run the MeshGenerator
+	 */
+	virtual void InitializeMesh();
+
+	/** 
+	 * @return Const reference to the internal FDynamicMesh3.
+	 * @warning Calling ProcessMesh() is preferred! This interface may be  removed in the future
+	 */
+	const UE::Geometry::FDynamicMesh3& GetMeshRef() const { return *Mesh; }
+	/**
+	 * @return Const pointer to the internal FDynamicMesh3.
+	 * @warning Calling ProcessMesh() is preferred! This interface may be  removed in the future
+	 */
+	const UE::Geometry::FDynamicMesh3* GetMeshPtr() const { return Mesh.Get(); }
+
+	/**
+	 * @return Writable reference to the internal FDynamicMesh3. 
+	 * @warning Calling EditMesh() is preferred! This interface may be  removed in the future
+	 */
+	UE::Geometry::FDynamicMesh3& GetMeshRef() { return *Mesh; }
+	/**
+	 * @return Writable pointer to the internal FDynamicMesh3.
+	 * @warning Calling EditMesh() is preferred! This interface may be  removed in the future
+	 */
+	UE::Geometry::FDynamicMesh3* GetMeshPtr() { return Mesh.Get(); }
+
+
+	/** Replace the internal mesh with a copy of MoveMesh */
+	void SetMesh(const UE::Geometry::FDynamicMesh3& MoveMesh);
+
+	/** Replace the internal mesh with the data in MoveMesh */
+	void SetMesh(UE::Geometry::FDynamicMesh3&& MoveMesh);
+
+	/**
+	 * Apply ProcessFunc to the internal Mesh
+	 */
+	void ProcessMesh(TFunctionRef<void(const UE::Geometry::FDynamicMesh3&)> ProcessFunc) const;
+
+	/**
+	 * Apply EditFunc to the internal mesh.
+	 * This will broadcast PreMeshChangedEvent, then call EditFunc(), then broadcast MeshChangedEvent and MeshModifiedBPEvent
+	 */
+	void EditMesh(TFunctionRef<void(UE::Geometry::FDynamicMesh3&)> EditFunc);
+
+	/**
+	 * Take ownership of the internal Mesh, and have it replaced with a new mesh
+	 */
+	TUniquePtr<UE::Geometry::FDynamicMesh3> ExtractMesh();
+
+
+	//
+	// Change support
+	//
+
+
+	// IMeshVertexCommandChangeTarget implementation, allows a FVertexChange to be applied to the mesh
+	void ApplyChange(const FMeshVertexChange* Change, bool bRevert);
+
+	// IMeshCommandChangeTarget implementation, allows a FMeshChange to be applied to the mesh
+	void ApplyChange(const FMeshChange* Change, bool bRevert);
+
+	// IMeshReplacementCommandChangeTarget implementation, allows a FMeshReplacementChange to be applied to the mesh
+	void ApplyChange(const FMeshReplacementChange* Change, bool bRevert);
+
+
+	//
+	// Event support
+	//
+protected:
+	/** Broadcast before the internal mesh is modified by Reset funcs, SetMesh(), EditMesh(), and ApplyChange()'s above */
+	FOnDynamicMeshChanged PreMeshChangedEvent;
+	/** Broadcast after the internal mesh is modified, in the same cases as PreMeshChangedEvent */
+	FOnDynamicMeshChanged MeshChangedEvent;
+
+public:
+	/** Broadcast before the internal mesh is modified by Reset funcs, SetMesh(), EditMesh(), and ApplyChange()'s above */
+	FOnDynamicMeshChanged& OnPreMeshChanged() { return PreMeshChangedEvent; }
+	/** Broadcast after the internal mesh is modified, in the same cases as OnPreMeshChanged */
+	FOnDynamicMeshChanged& OnMeshChanged() { return MeshChangedEvent; }
+
+public:
+	/** Blueprintable event called when mesh is modified, in the same cases as OnMeshChanged */
+	UPROPERTY(BlueprintAssignable, meta = (DisplayName = "MeshModified"))
+	FOnDynamicMeshModifiedBP MeshModifiedBPEvent;
+
+
+	//
+	// Realtime Update support.
+	// This is intended to be used in situations where the internal mesh is being
+	// modified directly, ie instead of via EditMesh(), and we would like to
+	// notify listeners of these changes. Generally EditMesh() is preferred but in certain
+	// cases (eg like 3D sculpting) it is too complex to refactor all the mesh updates
+	// into EditMesh() calls (eg some are done async/etc). So, code that does those
+	// kinds of modifications can call PostRealtimeUpdate() to let any interested
+	// parties know that the mesh is actively changing.
+	//
+public:
+	DECLARE_MULTICAST_DELEGATE_OneParam(FOnMeshRealtimeUpdate, UDynamicMesh*);
+
+	/**
+	 * Multicast delegate that is broadcast whenever PostRealtimeUpdate() is called
+	 */
+	FOnMeshRealtimeUpdate& OnMeshRealtimeUpdate() { return MeshRealtimeUpdateEvent; }
+
+	/**
+	 * Broadcasts FOnMeshRealtimeUpdate
+	 */
+	virtual void PostRealtimeUpdate();
+
+protected:
+	FOnMeshRealtimeUpdate MeshRealtimeUpdateEvent;
+
+
+
+protected:
+	/**
+	 * Mesh data, owned by this object.
+	 * By default will have Attributes enabled and MaterialID enabled.
+	 */
+	TUniquePtr<UE::Geometry::FDynamicMesh3> Mesh;
+
+	/**
+	 * Allocate a new Mesh (ie pointer will change) and then call InitializeMesh()
+	 */
+	void InitializeNewMesh();
+
+	/**
+	 * Internal function that edits the Mesh, but broadcasts PreMeshChangedEvent and MeshChangedEvent, and then MeshModifiedBPEvent
+	 */
+	void EditMeshInternal(TFunctionRef<void(UE::Geometry::FDynamicMesh3&)> EditFunc, const FDynamicMeshChangeInfo& ChangeInfo);
+
+
+public:
+	// serialize Mesh to an Archive
+	virtual void Serialize(FArchive& Archive) override;
+
+	// serialize Mesh to/from T3D
+	virtual void ExportCustomProperties(FOutputDevice& Out, uint32 Indent) override;
+	virtual void ImportCustomProperties(const TCHAR* SourceText, FFeedbackContext* Warn) override;
+
+
+
+	//
+	// (Preliminary) Procedural Generator support. If a Generator is configured, then each
+	// time this mesh is Reset(), it will call MeshGenerator->Generate(). The idea is that 
+	// generator will be set to something that emits a new mesh based on external data,
+	// for example a procedural primitive Actor/Component could configure a Generator that
+	// generates the primitive mesh based on Actor settings.
+	//
+
+protected:
+
+	/**
+	 * Active mesh generator. If configured, and bEnableMeshGenerator is true, then MeshGenerator->Generate() 
+	 * will be called when this mesh is Reset(). The Regenerate() function below can be used to force regeneration.
+	 */
+	UPROPERTY(Instanced)
+	TObjectPtr<UDynamicMeshGenerator> MeshGenerator;
+
+public:
+
+	/**
+	 * Controls whether the active Generator (if configured) will be applied when rebuilding the mesh
+	 */
+	UPROPERTY(EditAnywhere, Category = "Dynamic Mesh")
+	bool bEnableMeshGenerator = false;
+
+	/**
+	 * Set the active mesh generator. Clears if nullptr
+	 */
+	virtual void SetMeshGenerator(TObjectPtr<UDynamicMeshGenerator> NewGenerator);
+
+	/**
+	 * Set the active mesh generator. Clears if nullptr
+	 */
+	virtual void ClearMeshGenerator();
+
+	/**
+	 * Reset() the mesh, which will re-run the active MeshGenerator, if bEnableMeshGenerator
+	 */
+	virtual void Regenerate();
+
+};
+
+
+
+
+
+
+/**
+ * UDynamicMeshPool manages a Pool of UDynamicMesh objects. This allows
+ * the meshes to be re-used instead of being garbage-collected. This
+ * is intended to be used by Blueprints that need to do procedural geometry
+ * operations that generate temporary meshes, as these will commonly run their
+ * construction scripts many times as the user (eg) manipulates parameters,
+ * and constantly spawning new UDynamicMesh instances will result in enormous
+ * memory usage hanging around until GC runs.
+ *
+ * Usage is to call RequestMesh() to take ownership of an available UDynamicMesh (which
+ * will allocate a new one if the pool is empty) and ReturnMesh() to return it to the pool.
+ *
+ * ReturnAllMeshes() can be called to return all allocated meshes.
+ *
+ * In both cases, there is nothing preventing you from still holding on to the mesh.
+ * So, be careful.
+ *
+ * FreeAllMeshes() calls ReturnAllMeshes() and then releases the pool's references to
+ * the allocated meshes, so they can be Garbage Collected
+ *
+ * If you Request() more meshes than you Return(), the Pool will still be holding on to
+ * references to those meshes, and they will never be Garbage Collected (ie memory leak).
+ * As a failsafe, if the number of allocated meshes exceeds geometry.DynamicMesh.MaxPoolSize,
+ * the Pool will release all it's references and run garbage collection on the next call to RequestMesh().
+ * (Do not rely on this as a memory management strategy)
+ *
+ * An alternate strategy that could be employed here is for the Pool to not hold
+ * references to meshes it has provided, only those that have been explicitly returned.
+ * Then non-returned meshes would simply be garbage-collected, however it allows
+ * potentially a large amount of memory to be consumed until that occurs.
+ *
+ * UDynamicMesh::Reset() is called on the object returned to the Pool, which clears
+ * the internal FDynamicMesh3 (which uses normal C++ memory management, so no garbage collection involved)
+ * So the Pool does not re-use mesh memory, only the UObject containers.
+ */
+UCLASS(Transient)
+class MODELINGCOMPONENTS_API UDynamicMeshPool : public UObject
+{
+	GENERATED_BODY()
+public:
+	/** @return an available UDynamicMesh from the pool (possibly allocating a new mesh) */
+	UFUNCTION(BlueprintCallable, Category="Dynamic Mesh")
+	UDynamicMesh* RequestMesh();
+
+	/** Release a UDynamicMesh returned by RequestMesh() back to the pool */
+	UFUNCTION(BlueprintCallable, Category = "Dynamic Mesh")
+	void ReturnMesh(UDynamicMesh* Mesh);
+
+	/** Release all GeneratedMeshes back to the pool */
+	UFUNCTION(BlueprintCallable, Category = "Dynamic Mesh")
+	void ReturnAllMeshes();
+
+	/** Release all GeneratedMeshes back to the pool and allow them to be garbage collected */
+	UFUNCTION(BlueprintCallable, Category = "Dynamic Mesh")
+	void FreeAllMeshes();
+
+protected:
+	/** Meshes in the pool that are available */
+	UPROPERTY()
+	TArray<UDynamicMesh*> CachedMeshes;
+
+	/** All meshes the pool has allocated */
+	UPROPERTY()
+	TArray<UDynamicMesh*> AllCreatedMeshes;
+};
\ No newline at end of file
