MeshModelingTools: Fix materials not being shown in Warp tool.

#rnx

diff --git a/Engine/Plugins/Experimental/MeshModelingToolset/Source/MeshModelingTools/Private/MeshSpaceDeformerTool.cpp b/Engine/Plugins/Experimental/MeshModelingToolset/Source/MeshModelingTools/Private/MeshSpaceDeformerTool.cpp
--- a/Engine/Plugins/Experimental/MeshModelingToolset/Source/MeshModelingTools/Private/MeshSpaceDeformerTool.cpp
+++ b/Engine/Plugins/Experimental/MeshModelingToolset/Source/MeshModelingTools/Private/MeshSpaceDeformerTool.cpp
@@ -22,6 +22,7 @@
 #include "BaseGizmos/TransformGizmoUtil.h"
 #include "BaseGizmos/IntervalGizmo.h"
 #include "MeshDescriptionToDynamicMesh.h"
+#include "ModelingToolTargetUtil.h"
 #include "DynamicMeshToMeshDescription.h"
 #include "CoreMinimal.h"
 #include "Math/Matrix.h"
@@ -275,9 +276,9 @@ void UMeshSpaceDeformerTool::Setup()
 		Preview->PreviewMesh->UpdatePreview(OriginalDynamicMesh.Get());
 		Preview->PreviewMesh->SetTransform(MeshTransform);
 
-		Preview->ConfigureMaterials(StoredMaterialSet.Materials,
-			ToolSetupUtil::GetDefaultWorkingMaterial(GetToolManager())
-		);
+		FComponentMaterialSet MaterialSet;
+		MaterialSet = UE::ToolTarget::GetMaterialSet(Target);
+		Preview->ConfigureMaterials(MaterialSet.Materials, ToolSetupUtil::GetDefaultWorkingMaterial(GetToolManager()));
 
 		// show the preview mesh
 		Preview->SetVisibility(true);
diff --git a/Engine/Plugins/Experimental/MeshModelingToolset/Source/MeshModelingTools/Public/MeshSpaceDeformerTool.h b/Engine/Plugins/Experimental/MeshModelingToolset/Source/MeshModelingTools/Public/MeshSpaceDeformerTool.h
--- a/Engine/Plugins/Experimental/MeshModelingToolset/Source/MeshModelingTools/Public/MeshSpaceDeformerTool.h
+++ b/Engine/Plugins/Experimental/MeshModelingToolset/Source/MeshModelingTools/Public/MeshSpaceDeformerTool.h
@@ -238,9 +238,6 @@ protected:
 	UPROPERTY()
 	UPreviewMesh* OriginalMeshPreview;
 
-	// Stores the materials of the original object, used to make the object semi-transparent.
-	FComponentMaterialSet StoredMaterialSet;
-
 	UPROPERTY()
 	UIntervalGizmo* IntervalGizmo;
 
