UVEditor: Adding multiple selection to MeshSelectionMechanic

#rnx

diff --git a/Engine/Plugins/Experimental/MeshModelingToolset/Source/ModelingComponents/Private/Selection/MeshSelectionMechanic.cpp b/Engine/Plugins/Experimental/MeshModelingToolset/Source/ModelingComponents/Private/Selection/MeshSelectionMechanic.cpp
--- a/Engine/Plugins/Experimental/MeshModelingToolset/Source/ModelingComponents/Private/Selection/MeshSelectionMechanic.cpp
+++ b/Engine/Plugins/Experimental/MeshModelingToolset/Source/ModelingComponents/Private/Selection/MeshSelectionMechanic.cpp
@@ -24,6 +24,7 @@ void UMeshSelectionMechanic::Setup(UInteractiveTool* ParentToolIn)
 
 	// TODO: Add shift/ctrl modifiers
 	USingleClickInputBehavior* ClickBehavior = NewObject<USingleClickInputBehavior>();
+	ClickBehavior->Modifiers.RegisterModifier(ShiftModifierID, FInputDeviceState::IsShiftKeyDown);
 	ClickBehavior->Initialize(this);
 	ParentTool->AddInputBehavior(ClickBehavior);
 
@@ -210,13 +211,24 @@ void UMeshSelectionMechanic::OnClicked(const FInputDeviceRay& ClickPos)
 {
 	FDynamicMeshSelection OriginalSelection = CurrentSelection;
 
-	CurrentSelection.SelectedIDs.Empty();
-	CurrentSelection.Mesh = nullptr;
+	if (!InMultiSelectMode() ||
+		(SelectionMode == EMeshSelectionMechanicMode::Component && CurrentSelection.Type != FDynamicMeshSelection::EType::Triangle) ||
+		(SelectionMode == EMeshSelectionMechanicMode::Edge && CurrentSelection.Type != FDynamicMeshSelection::EType::Edge) ) 
+	{ // If we're not enabling multiselect, or our desired mode doesn't match our current selection type, clear the existing selection
+		CurrentSelection.SelectedIDs.Empty();
+		CurrentSelection.Mesh = nullptr;
+	}
 
 	int32 HitTid = IndexConstants::InvalidID;
 	double RayT = 0; 
 	for (int32 i = 0; i < MeshSpatials.Num(); ++i)
 	{
+		//Short circuit the loop if we're in multiselect mode, to save us from testing things we aren't allowed to select
+		if (InMultiSelectMode() && !CurrentSelection.SelectedIDs.IsEmpty() && i != CurrentSelectionIndex) 
+		{
+			continue;
+		}
+
 		// Minor note: Using FRay3d in this statement instead of FRay causes an internal compiler error in VS
 		// in Development Editor builds, for no good reason.
 		FRay LocalRay(
@@ -225,6 +237,7 @@ void UMeshSelectionMechanic::OnClicked(const FInputDeviceRay& ClickPos)
 
 		if (MeshSpatials[i]->FindNearestHitTriangle(LocalRay, RayT, HitTid))
 		{
+			// It should be okay to reassign this. If we're in multiselect, and it misses a hit, this just won't happen and selection won't be altered.
 			CurrentSelection.Mesh = MeshSpatials[i]->GetMesh();
 			CurrentSelectionIndex = i;
 			break;
@@ -235,10 +248,14 @@ void UMeshSelectionMechanic::OnClicked(const FInputDeviceRay& ClickPos)
 	{
 		if (SelectionMode == EMeshSelectionMechanicMode::Component)
 		{
-			CurrentSelection.SelectedIDs.Add(HitTid);
 			FMeshConnectedComponents MeshSelectedComponent(CurrentSelection.Mesh);
-			MeshSelectedComponent.FindTrianglesConnectedToSeeds(CurrentSelection.SelectedIDs.Array());
-			CurrentSelection.SelectedIDs = TSet(MeshSelectedComponent.Components[0].Indices);
+			TArray<int32> SeedTriangles;
+			SeedTriangles.Add(HitTid);
+			MeshSelectedComponent.FindTrianglesConnectedToSeeds(SeedTriangles);
+			// Since we've already handled the situation of not holding triangles in a multiselect above, this should be safe.
+			// But just to be sure, we'll toss in a check to validate everything.
+			ensure((InMultiSelectMode() && CurrentSelection.Type == FDynamicMeshSelection::EType::Triangle) || CurrentSelection.SelectedIDs.IsEmpty());
+			CurrentSelection.SelectedIDs.Append( TSet(MeshSelectedComponent.Components[0].Indices) );
 			CurrentSelection.Type = FDynamicMeshSelection::EType::Triangle;
 		}
 		// TODO: We'll need the ability to hit occluded triangles to see if there is a better edge
@@ -262,6 +279,9 @@ void UMeshSelectionMechanic::OnClicked(const FInputDeviceRay& ClickPos)
 						Position1, Position2,
 						ToolSceneQueriesUtil::GetDefaultVisualAngleSnapThreshD()); }))
 			{
+				// Since we've already handled multiselect and not holding edges above, this should be safe.	
+				// But just to be sure, we'll toss in a check to validate everything.
+				ensure((InMultiSelectMode() && CurrentSelection.Type == FDynamicMeshSelection::EType::Edge) || CurrentSelection.SelectedIDs.IsEmpty());
 				CurrentSelection.SelectedIDs.Add(Result.ID);
 				CurrentSelection.Type = FDynamicMeshSelection::EType::Edge;
 			}
@@ -276,4 +296,15 @@ void UMeshSelectionMechanic::OnClicked(const FInputDeviceRay& ClickPos)
 	}
 }
 
+void UMeshSelectionMechanic::OnUpdateModifierState(int ModifierID, bool bIsOn)
+{
+	switch (ModifierID)
+	{
+	case ShiftModifierID:
+		bShiftToggle = bIsOn;
+		break;
+	// Add more modifiers here, if needed.
+	}
+}
+
 #undef LOCTEXT_NAMESPACE
diff --git a/Engine/Plugins/Experimental/MeshModelingToolset/Source/ModelingComponents/Public/Selection/MeshSelectionMechanic.h b/Engine/Plugins/Experimental/MeshModelingToolset/Source/ModelingComponents/Public/Selection/MeshSelectionMechanic.h
--- a/Engine/Plugins/Experimental/MeshModelingToolset/Source/ModelingComponents/Public/Selection/MeshSelectionMechanic.h
+++ b/Engine/Plugins/Experimental/MeshModelingToolset/Source/ModelingComponents/Public/Selection/MeshSelectionMechanic.h
@@ -78,12 +78,17 @@ public:
 	virtual FInputRayHit IsHitByClick(const FInputDeviceRay& ClickPos) override;
 	virtual void OnClicked(const FInputDeviceRay& ClickPos) override;
 
+	// IModifierToggleBehaviorTarget implementation
+	virtual void OnUpdateModifierState(int ModifierID, bool bIsOn) override;
+
 	FSimpleMulticastDelegate OnSelectionChanged;
 
 	EMeshSelectionMechanicMode SelectionMode;
 
 protected:
 
+	bool InMultiSelectMode() const { return bShiftToggle; }
+
 	UPROPERTY()
 	APreviewGeometryActor* PreviewGeometryActor = nullptr;
 
@@ -98,6 +103,10 @@ protected:
 	FViewCameraState CameraState;
 
 	FColor LineColor = FColor::Yellow;
+
+	bool bShiftToggle = false;
+	static const int32 ShiftModifierID = 1;
+
 	float LineThickness = 3;
 	float DepthBias = 0.3;
 
