ModelingComponents: Clean up DynamicMeshComponent API. Update Component and Proxy handling of Tangents to use Attribute Overlay if available. Update affected Tools and also convert most of the affected Tools to use UE::ToolTarget helper functions.

- Add UE::ToolTarget::CommitMaterialSetUpdate() and ::CommitDynamicMeshUpdate(). ::GetDynamicMeshCopy() can now return tangents if requested.
- Add IMeshDescriptionProvider::CalculateAutoGeneratedAttributes(). Default implementation does nothing, UStaticMeshComponentToolTarget implementation initializes auto-generated MeshDescription attributes. Used in ::GetDynamicMeshCopy() to get tangents (but requires a MeshDescription copy).
- Clean up handling of Tangents in Simple/OctreeDynamicMeshComponent. Add local MakeTangentsFunc() to generate the Tangents lambda, handle different cases and no-tangents fallbacks consistently.
- UDynamicMesh: add optional info arguments to EditMesh() and ChangeInfo struct. Add support for deferring change events from Edit funcs.
- Remove UBaseDynamicMeshComponent::InitializeMesh(), ::Bake() APIs, and add ::SetMesh(). Implement in Simple/Octree implementations, update all Tools that used those APIs.
- Add USimpleDynamicMeshComponent::ProcessMesh(), EditMesh(). These are now the preferred ways to read/write mesh.
- Update USimpleDynamicMeshComponent tangents handling. Externally-computed tangents are now taken directly from the FDynamicMesh3 attribute set. Autogenerated tangents are still computed and stored in an internal FMeshTangentsf, but this is no longer exposed for external updates.
- Remove UPreviewMesh pass-through functions for Tangents access, InitializeMesh() and Bake(). Add ProcessMesh()
- Update all affected Tools. In most cases these Tools have also been converted to use ModelingToolTargetUtil functions, instead of direct ToolTarget interface casting.

#rnx

diff --git a/Engine/Plugins/Experimental/MeshLODToolset/Source/MeshLODToolset/Private/Tools/LODManagerTool.cpp b/Engine/Plugins/Experimental/MeshLODToolset/Source/MeshLODToolset/Private/Tools/LODManagerTool.cpp
--- a/Engine/Plugins/Experimental/MeshLODToolset/Source/MeshLODToolset/Private/Tools/LODManagerTool.cpp
+++ b/Engine/Plugins/Experimental/MeshLODToolset/Source/MeshLODToolset/Private/Tools/LODManagerTool.cpp
@@ -9,6 +9,7 @@
 #include "AssetUtils/MeshDescriptionUtil.h"
 #include "MeshDescriptionToDynamicMesh.h"
 #include "IndexedMeshToDynamicMesh.h"
+#include "ModelingToolTargetUtil.h"
 
 #include "MeshDescription.h"
 #include "StaticMeshAttributes.h"
@@ -88,39 +89,46 @@ void ULODManagerTool::Setup()
 {
 	UInteractiveTool::Setup();
 
-	HiResSourceModelActions = NewObject<ULODManagerHiResSourceModelActions>(this);
-	HiResSourceModelActions->Initialize(this);
-	AddToolPropertySource(HiResSourceModelActions);
-
-	MaterialActions = NewObject<ULODManagerMaterialActions>(this);
-	MaterialActions->Initialize(this);
-	AddToolPropertySource(MaterialActions);
 
 	if (Targets.Num() == 1)
 	{
-		LODInfoProperties = NewObject<ULODManagerLODProperties>(this);
-		AddToolPropertySource(LODInfoProperties);
-
 		LODPreviewProperties = NewObject<ULODManagerPreviewLODProperties>(this);
 		AddToolPropertySource(LODPreviewProperties);
 		LODPreviewProperties->VisibleLOD = this->DefaultLODName;
 		LODPreviewProperties->WatchProperty(LODPreviewProperties->VisibleLOD, [this](FString NewLOD) { bPreviewLODValid = false; });
+		LODPreviewProperties->WatchProperty(LODPreviewProperties->bShowSeams, [this](bool bNewValue) { bPreviewLODValid = false; });
 
 		LODPreview = NewObject<UPreviewMesh>(this);
-		LODPreview->CreateInWorld(TargetWorld, TargetComponentInterface(0)->GetWorldTransform());
-		LODPreview->SetTangentsMode(EDynamicMeshTangentCalcType::ExternallyCalculated);
+		LODPreview->CreateInWorld(TargetWorld, (FTransform)UE::ToolTarget::GetLocalToWorldTransform(Targets[0]));
+		LODPreview->SetTangentsMode(EDynamicMeshComponentTangentsMode::ExternallyProvided);
 		LODPreview->SetVisible(false);
 
 		LODPreviewLines = NewObject<UPreviewGeometry>(this);
-		LODPreviewLines->CreateInWorld(TargetWorld, TargetComponentInterface(0)->GetWorldTransform());
+		LODPreviewLines->CreateInWorld(TargetWorld, (FTransform)UE::ToolTarget::GetLocalToWorldTransform(Targets[0]));
 
-		FComponentMaterialSet MaterialSet;
-		TargetMaterialInterface(0)->GetMaterialSet(MaterialSet);
+		FComponentMaterialSet MaterialSet = UE::ToolTarget::GetMaterialSet(Targets[0]);
 		LODPreview->SetMaterials(MaterialSet.Materials);
 
 		bLODInfoValid = false;
 	}
 
+
+	HiResSourceModelActions = NewObject<ULODManagerHiResSourceModelActions>(this);
+	HiResSourceModelActions->Initialize(this);
+	AddToolPropertySource(HiResSourceModelActions);
+
+	MaterialActions = NewObject<ULODManagerMaterialActions>(this);
+	MaterialActions->Initialize(this);
+	AddToolPropertySource(MaterialActions);
+
+
+	if (Targets.Num() == 1)
+	{
+		LODInfoProperties = NewObject<ULODManagerLODProperties>(this);
+		AddToolPropertySource(LODInfoProperties);
+		bLODInfoValid = false;
+	}
+
 	SetToolDisplayName(LOCTEXT("ToolName", "Manage LODs"));
 	GetToolManager()->DisplayMessage(
 		LOCTEXT("OnStartTool", "Inspect and Modify LODs of a StaticMesh Asset"),
@@ -192,8 +200,7 @@ UStaticMesh* ULODManagerTool::GetSingleStaticMesh()
 	{
 		return nullptr;
 	}
-	IPrimitiveComponentBackedTarget* Target = TargetComponentInterface(0);
-	UStaticMeshComponent* StaticMeshComponent = Cast<UStaticMeshComponent>(Target->GetOwnerComponent());
+	UStaticMeshComponent* StaticMeshComponent = Cast<UStaticMeshComponent>(UE::ToolTarget::GetTargetComponent(Targets[0]));
 	if (StaticMeshComponent == nullptr || StaticMeshComponent->GetStaticMesh() == nullptr)
 	{
 		return nullptr;
@@ -363,6 +370,15 @@ void ULODManagerTool::UpdatePreviewLines(FLODMeshInfo& MeshInfo)
 }
 
 
+void ULODManagerTool::ClearPreviewLines()
+{
+	ULineSetComponent* BoundaryLines = LODPreviewLines->FindLineSet(TEXT("BoundaryLines"));
+	if (BoundaryLines)
+	{
+		BoundaryLines->Clear();
+	}
+}
+
 
 void ULODManagerTool::UpdatePreviewLOD()
 {
@@ -377,9 +393,11 @@ void ULODManagerTool::UpdatePreviewLOD()
 	const FLODName* FoundName = ActiveLODNames.Find(SelectedLOD);
 	if (SelectedLOD.IsEmpty() || FoundName == nullptr || FoundName->IsDefault() )
 	{
-		// badness
+		// currently just showing the source object when Default is selected
 		LODPreview->SetVisible(false);
-		TargetComponentInterface(0)->SetOwnerVisibility(true);
+		ClearPreviewLines();
+		LODPreviewLines->SetAllVisible(false);
+		UE::ToolTarget::ShowSourceObject(Targets[0]);
 		return;
 	}
 
@@ -392,18 +410,19 @@ void ULODManagerTool::UpdatePreviewLOD()
 	
 	if (Found)
 	{
-		TargetComponentInterface(0)->SetOwnerVisibility(false);
-		LODPreview->ReplaceMesh( (*Found)->Mesh );
-		LODPreview->UpdateTangents<float>( &(*Found)->Tangents, false);
+		UE::ToolTarget::HideSourceObject(Targets[0]);
+
+		LODPreview->ReplaceMesh((*Found)->Mesh);
 		LODPreview->SetVisible(true);
 		UpdatePreviewLines(**Found);
-		LODPreviewLines->SetAllVisible(true);
+		LODPreviewLines->SetAllVisible(LODPreviewProperties->bShowSeams);
 	}
 	else
 	{
 		LODPreview->SetVisible(false);
+		ClearPreviewLines();
 		LODPreviewLines->SetAllVisible(false);
-		TargetComponentInterface(0)->SetOwnerVisibility(true);
+		UE::ToolTarget::ShowSourceObject(Targets[0]);
 	}
 }
 
@@ -428,8 +447,7 @@ bool ULODManagerTool::CacheLODMesh(const FString& Name, FLODName LODName)
 			UE::MeshDescription::InitializeAutoGeneratedAttributes(TmpMeshDescription, &HiResBuildSettings);
 
 			TUniquePtr<FLODMeshInfo> NewMeshInfo = MakeUnique<FLODMeshInfo>();
-			Converter.Convert(&TmpMeshDescription, NewMeshInfo->Mesh);
-			Converter.CopyTangents(&TmpMeshDescription, &NewMeshInfo->Mesh, &NewMeshInfo->Tangents);
+			Converter.Convert(&TmpMeshDescription, NewMeshInfo->Mesh, true);
 			LODMeshCache.Add(Name, MoveTemp(NewMeshInfo));
 			return true;
 		}
@@ -446,8 +464,7 @@ bool ULODManagerTool::CacheLODMesh(const FString& Name, FLODName LODName)
 			UE::MeshDescription::InitializeAutoGeneratedAttributes(TmpMeshDescription, &LODBuildSettings);
 
 			TUniquePtr<FLODMeshInfo> NewMeshInfo = MakeUnique<FLODMeshInfo>();
-			Converter.Convert(&TmpMeshDescription, NewMeshInfo->Mesh);
-			Converter.CopyTangents(&TmpMeshDescription, &NewMeshInfo->Mesh, &NewMeshInfo->Tangents);
+			Converter.Convert(&TmpMeshDescription, NewMeshInfo->Mesh, true);
 			LODMeshCache.Add(Name, MoveTemp(NewMeshInfo));
 			return true;
 		}
@@ -462,8 +479,6 @@ bool ULODManagerTool::CacheLODMesh(const FString& Name, FLODName LODName)
 				const FStaticMeshLODResources& LODResource = RenderData->LODResources[LODName.RenderDataIndex];
 				TUniquePtr<FLODMeshInfo> NewMeshInfo = MakeUnique<FLODMeshInfo>();
 				UE::Conversion::RenderBuffersToDynamicMesh(LODResource.VertexBuffers, LODResource.IndexBuffer, LODResource.Sections, NewMeshInfo->Mesh);
-				NewMeshInfo->Tangents.SetMesh(&NewMeshInfo->Mesh);
-				NewMeshInfo->Tangents.CopyTriVertexTangents(NewMeshInfo->Mesh);
 				LODMeshCache.Add(Name, MoveTemp(NewMeshInfo));
 				return true;
 			}
diff --git a/Engine/Plugins/Experimental/MeshLODToolset/Source/MeshLODToolset/Public/Tools/LODManagerTool.h b/Engine/Plugins/Experimental/MeshLODToolset/Source/MeshLODToolset/Public/Tools/LODManagerTool.h
--- a/Engine/Plugins/Experimental/MeshLODToolset/Source/MeshLODToolset/Public/Tools/LODManagerTool.h
+++ b/Engine/Plugins/Experimental/MeshLODToolset/Source/MeshLODToolset/Public/Tools/LODManagerTool.h
@@ -78,7 +78,7 @@ class MESHLODTOOLSET_API ULODManagerPreviewLODProperties : public UInteractiveTo
 {
 	GENERATED_BODY()
 public:
-
+	/** LOD to visualise. Default option is equivalent to disabling the Tool, RenderData is the mesh used for rendering derived from the SourceModel (possibly simplified) */
 	UPROPERTY(EditAnywhere, Category = LODPreview, meta = (DisplayName = "Show LOD", GetOptions = GetLODNamesFunc))
 	FString VisibleLOD;
 
@@ -87,6 +87,10 @@ public:
 
 	UPROPERTY(meta = (TransientToolProperty))
 	TArray<FString> LODNamesList;
+
+	/** Control whether mesh borders are displayed */
+	UPROPERTY(EditAnywhere, Category = LODPreview)
+	bool bShowSeams = true;
 };
 
 
@@ -232,7 +236,6 @@ protected:
 	struct FLODMeshInfo
 	{
 		UE::Geometry::FDynamicMesh3 Mesh;
-		UE::Geometry::FMeshTangentsf Tangents;		// todo: store tangents on the Mesh. Requires support from SimpleDynamicMeshComponent.
 
 		bool bInfoCached = false;
 		TArray<int> BoundaryEdges;
@@ -249,5 +252,6 @@ protected:
 	bool bPreviewLODValid = false;
 	void UpdatePreviewLOD();
 	void UpdatePreviewLines(FLODMeshInfo& LODMeshInfo);
+	void ClearPreviewLines();
 	bool CacheLODMesh(const FString& Name, FLODName LODName);
 };
diff --git a/Engine/Plugins/Experimental/MeshModelingToolset/Source/MeshModelingTools/Private/BakeMeshAttributeMapsTool.cpp b/Engine/Plugins/Experimental/MeshModelingToolset/Source/MeshModelingTools/Private/BakeMeshAttributeMapsTool.cpp
--- a/Engine/Plugins/Experimental/MeshModelingToolset/Source/MeshModelingTools/Private/BakeMeshAttributeMapsTool.cpp
+++ b/Engine/Plugins/Experimental/MeshModelingToolset/Source/MeshModelingTools/Private/BakeMeshAttributeMapsTool.cpp
@@ -22,7 +22,6 @@
 #include "ImageUtils.h"
 
 #include "AssetUtils/Texture2DBuilder.h"
-#include "AssetUtils/MeshDescriptionUtil.h"
 #include "ModelingObjectsCreationAPI.h"
 
 #include "TargetInterfaces/MaterialProvider.h"
@@ -30,6 +29,7 @@
 #include "TargetInterfaces/PrimitiveComponentBackedTarget.h"
 #include "TargetInterfaces/StaticMeshBackedTarget.h"
 #include "ToolTargetManager.h"
+#include "ModelingToolTargetUtil.h"
 
 // required to pass UStaticMesh asset so we can save at same location
 #include "Engine/Classes/Components/StaticMeshComponent.h"
@@ -274,35 +274,32 @@ void UBakeMeshAttributeMapsTool::Setup()
 {
 	UInteractiveTool::Setup();
 
-	IPrimitiveComponentBackedTarget* TargetComponent = TargetComponentInterface(0);
-	IMeshDescriptionProvider* MeshProvider = TargetMeshProviderInterface(0);
-	IMaterialProvider* MaterialProvider = TargetMaterialInterface(0);
-
-	// copy input MeshDescription and make sure it has initialized normals/tangents
-	BaseMeshDescription = MakeShared<FMeshDescription, ESPMode::ThreadSafe>(*MeshProvider->GetMeshDescription());
-	UE::MeshDescription::InitializeAutoGeneratedAttributes(*BaseMeshDescription, TargetComponent->GetOwnerComponent(), 0);
-
 	// create dynamic mesh component to use for live preview
-	DynamicMeshComponent = NewObject<USimpleDynamicMeshComponent>(TargetComponent->GetOwnerActor(), "DynamicMesh");
-	DynamicMeshComponent->SetupAttachment(TargetComponent->GetOwnerActor()->GetRootComponent());
+	// TODO: convert to UPreviewMesh
+	AActor* ParentActor = UE::ToolTarget::GetTargetActor(Targets[0]);
+	DynamicMeshComponent = NewObject<USimpleDynamicMeshComponent>(ParentActor);
+	DynamicMeshComponent->SetupAttachment(ParentActor->GetRootComponent());
 	DynamicMeshComponent->RegisterComponent();
-	DynamicMeshComponent->SetWorldTransform(TargetComponent->GetWorldTransform());
+	DynamicMeshComponent->SetWorldTransform((FTransform)UE::ToolTarget::GetLocalToWorldTransform(Targets[0]));
 
 	// transfer materials
-	FComponentMaterialSet MaterialSet;
-	MaterialProvider->GetMaterialSet(MaterialSet);
+	FComponentMaterialSet MaterialSet = UE::ToolTarget::GetMaterialSet(Targets[0]);
 	for (int k = 0; k < MaterialSet.Materials.Num(); ++k)
 	{
 		DynamicMeshComponent->SetMaterial(k, MaterialSet.Materials[k]);
 	}
 
-	DynamicMeshComponent->TangentsType = EDynamicMeshTangentCalcType::ExternallyCalculated;
-	DynamicMeshComponent->InitializeMesh(BaseMeshDescription.Get());
+	FDynamicMesh3 InputMeshWithTangents = UE::ToolTarget::GetDynamicMeshCopy(Targets[0], true);
+	DynamicMeshComponent->SetTangentsType(EDynamicMeshComponentTangentsMode::ExternallyProvided);
+	DynamicMeshComponent->SetMesh(MoveTemp(InputMeshWithTangents));
 	
-	BaseMesh.Copy(*DynamicMeshComponent->GetMesh());
+	DynamicMeshComponent->ProcessMesh([&](const FDynamicMesh3& ReadMesh)
+	{
+		BaseMesh.Copy(ReadMesh);
+		BaseMeshTangents = MakeShared<FMeshTangentsd, ESPMode::ThreadSafe>(&BaseMesh);
+		BaseMeshTangents->CopyTriVertexTangents(ReadMesh);
+	});
 	BaseSpatial.SetMesh(&BaseMesh, true);
-	BaseMeshTangents = MakeShared<FMeshTangentsd, ESPMode::ThreadSafe>(&BaseMesh);
-	BaseMeshTangents->CopyTriVertexTangents(*DynamicMeshComponent->GetTangents());
 
 	UMaterial* Material = LoadObject<UMaterial>(nullptr, TEXT("/MeshModelingToolset/Materials/BakePreviewMaterial"));
 	check(Material);
@@ -326,9 +323,7 @@ void UBakeMeshAttributeMapsTool::Setup()
 
 	bIsBakeToSelf = (Targets.Num() == 1);
 
-	// hide input StaticMeshComponent
-	TargetComponent->SetOwnerVisibility(false);
-
+	UE::ToolTarget::HideSourceObject(Targets[0]);
 
 	Settings = NewObject<UBakeMeshAttributeMapsToolProperties>(this);
 	Settings->RestoreProperties(this);
@@ -499,8 +494,6 @@ TUniquePtr<UE::Geometry::TGenericDataOperator<FMeshMapBaker>> UBakeMeshAttribute
 
 void UBakeMeshAttributeMapsTool::Shutdown(EToolShutdownType ShutdownType)
 {
-	IPrimitiveComponentBackedTarget* TargetComponent = TargetComponentInterface(0);
-
 	Settings->SaveProperties(this);
 	OcclusionMapProps->SaveProperties(this);
 	NormalMapProps->SaveProperties(this);
@@ -515,14 +508,14 @@ void UBakeMeshAttributeMapsTool::Shutdown(EToolShutdownType ShutdownType)
 	}
 	if (DynamicMeshComponent != nullptr)
 	{
-		TargetComponent->SetOwnerVisibility(true);
+		UE::ToolTarget::ShowSourceObject(Targets[0]);
 
 		if (ShutdownType == EToolShutdownType::Accept)
 		{
-			UStaticMeshComponent* StaticMeshComponent = CastChecked<UStaticMeshComponent>(TargetComponent->GetOwnerComponent());
+			UStaticMeshComponent* StaticMeshComponent = CastChecked<UStaticMeshComponent>(UE::ToolTarget::GetTargetComponent(Targets[0]));
 			UStaticMesh* StaticMeshAsset = StaticMeshComponent->GetStaticMesh();
 			check(StaticMeshAsset);
-			FString BaseName = TargetComponent->GetOwnerActor()->GetName();
+			FString BaseName = UE::ToolTarget::GetTargetActor(Targets[0])->GetName();
 
 			bool bCreatedAssetOK = true;
 			int NumResults = Settings->Result.Num();
@@ -605,15 +598,10 @@ void UBakeMeshAttributeMapsTool::Render(IToolsContextRenderAPI* RenderAPI)
 {
 	UpdateResult();
 
-	IPrimitiveComponentBackedTarget* TargetComponent = TargetComponentInterface(0);
-
 	float GrayLevel = VisualizationProps->BaseGrayLevel;
 	PreviewMaterial->SetVectorParameterValue(TEXT("BaseColor"), FVector(GrayLevel, GrayLevel, GrayLevel) );
 	float AOWeight = VisualizationProps->OcclusionMultiplier;
 	PreviewMaterial->SetScalarParameterValue(TEXT("AOWeight"), AOWeight );
-
-	FPrimitiveDrawInterface* PDI = RenderAPI->GetPrimitiveDrawInterface();
-	FTransform Transform = TargetComponent->GetWorldTransform();
 }
 
 
@@ -665,7 +653,7 @@ int SelectTextureToBake(const TArray<UTexture*>& Textures)
 }
 
 
-void UBakeMeshAttributeMapsTool::GetTexturesFromDetailMesh(const IPrimitiveComponentBackedTarget* DetailComponent)
+void UBakeMeshAttributeMapsTool::GetTexturesFromDetailMesh(const UPrimitiveComponent* DetailComponent)
 {
 	constexpr bool bGuessAtTextures = true;
 
@@ -673,7 +661,7 @@ void UBakeMeshAttributeMapsTool::GetTexturesFromDetailMesh(const IPrimitiveCompo
 	MultiTextureProps->MaterialIDSourceTextureMap.Reset();
 
 	TArray<UMaterialInterface*> Materials;
-	DetailComponent->GetOwnerComponent()->GetUsedMaterials(Materials);
+	DetailComponent->GetUsedMaterials(Materials);
 	
 	for (int32 MaterialID = 0; MaterialID < Materials.Num(); ++MaterialID)	// TODO: This won't match MaterialIDs on the FDynamicMesh3 in general, will it?
 	{
@@ -717,25 +705,22 @@ void UBakeMeshAttributeMapsTool::GetTexturesFromDetailMesh(const IPrimitiveCompo
 
 void UBakeMeshAttributeMapsTool::UpdateDetailMesh()
 {
-	IPrimitiveComponentBackedTarget* TargetComponent = TargetComponentInterface(0);
-	IPrimitiveComponentBackedTarget* DetailComponent = TargetComponentInterface(bIsBakeToSelf ? 0 : 1);
-	IMeshDescriptionProvider* DetailMeshProvider = TargetMeshProviderInterface(bIsBakeToSelf ? 0 : 1);
+	UToolTarget* DetailTarget = Targets[bIsBakeToSelf ? 0 : 1];
+
+	DetailMesh = MakeShared<FDynamicMesh3, ESPMode::ThreadSafe>(UE::ToolTarget::GetDynamicMeshCopy(DetailTarget));
 
-	DetailMesh = MakeShared<FDynamicMesh3, ESPMode::ThreadSafe>();
-	FMeshDescriptionToDynamicMesh Converter;
-	Converter.Convert(DetailMeshProvider->GetMeshDescription(), *DetailMesh);
 	if (Settings->bUseWorldSpace && bIsBakeToSelf == false)
 	{
-		UE::Geometry::FTransform3d DetailToWorld(DetailComponent->GetWorldTransform());
+		FTransform3d DetailToWorld = UE::ToolTarget::GetLocalToWorldTransform(DetailTarget);
 		MeshTransforms::ApplyTransform(*DetailMesh, DetailToWorld);
-		UE::Geometry::FTransform3d WorldToBase(TargetComponent->GetWorldTransform());
+		FTransform3d WorldToBase = UE::ToolTarget::GetLocalToWorldTransform(Targets[0]);
 		MeshTransforms::ApplyTransform(*DetailMesh, WorldToBase.Inverse());
 	}
 	
 	DetailSpatial = MakeShared<FDynamicMeshAABBTree3, ESPMode::ThreadSafe>();
 	DetailSpatial->SetMesh(DetailMesh.Get(), true);
 
-	GetTexturesFromDetailMesh(DetailComponent);
+	GetTexturesFromDetailMesh(UE::ToolTarget::GetTargetComponent(DetailTarget));
 
 	bInputsDirty = true;
 	DetailMeshTimestamp++;
diff --git a/Engine/Plugins/Experimental/MeshModelingToolset/Source/MeshModelingTools/Private/ConvertToPolygonsTool.cpp b/Engine/Plugins/Experimental/MeshModelingToolset/Source/MeshModelingTools/Private/ConvertToPolygonsTool.cpp
--- a/Engine/Plugins/Experimental/MeshModelingToolset/Source/MeshModelingTools/Private/ConvertToPolygonsTool.cpp
+++ b/Engine/Plugins/Experimental/MeshModelingToolset/Source/MeshModelingTools/Private/ConvertToPolygonsTool.cpp
@@ -187,7 +187,7 @@ void UConvertToPolygonsTool::Setup()
 
 		// Give the preview something to display
 		PreviewWithBackgroundCompute->PreviewMesh->SetTransform(TargetComponent->GetWorldTransform());
-		PreviewWithBackgroundCompute->PreviewMesh->SetTangentsMode(EDynamicMeshTangentCalcType::AutoCalculated);
+		PreviewWithBackgroundCompute->PreviewMesh->SetTangentsMode(EDynamicMeshComponentTangentsMode::AutoCalculated);
 		PreviewWithBackgroundCompute->PreviewMesh->UpdatePreview(OriginalDynamicMesh.Get());
 		
 		PreviewWithBackgroundCompute->ConfigureMaterials(MaterialSet.Materials,
diff --git a/Engine/Plugins/Experimental/MeshModelingToolset/Source/MeshModelingTools/Private/DeformMeshPolygonsTool.cpp b/Engine/Plugins/Experimental/MeshModelingToolset/Source/MeshModelingTools/Private/DeformMeshPolygonsTool.cpp
--- a/Engine/Plugins/Experimental/MeshModelingToolset/Source/MeshModelingTools/Private/DeformMeshPolygonsTool.cpp
+++ b/Engine/Plugins/Experimental/MeshModelingToolset/Source/MeshModelingTools/Private/DeformMeshPolygonsTool.cpp
@@ -14,11 +14,7 @@
 #include "Solvers/ConstrainedMeshDeformer.h"
 #include "ToolBuilderUtil.h"
 #include "ToolSceneQueriesUtil.h"
-
-#include "TargetInterfaces/MaterialProvider.h"
-#include "TargetInterfaces/MeshDescriptionCommitter.h"
-#include "TargetInterfaces/MeshDescriptionProvider.h"
-#include "TargetInterfaces/PrimitiveComponentBackedTarget.h"
+#include "ModelingToolTargetUtil.h"
 
 #include "ExplicitUseGeometryMathTypes.h"		// using UE::Geometry::(math types)
 using namespace UE::Geometry;
@@ -853,23 +849,22 @@ void UDeformMeshPolygonsTool::Setup()
 	LaplacianDeformer = MakePimpl<FGroupTopologyLaplacianDeformer>();
 
 	// create dynamic mesh component to use for live preview
-	IPrimitiveComponentBackedTarget* TargetComponent = Cast<IPrimitiveComponentBackedTarget>(Target);
-	DynamicMeshComponent = NewObject<USimpleDynamicMeshComponent>(TargetComponent->GetOwnerActor(), "DynamicMesh");
-	DynamicMeshComponent->SetupAttachment(TargetComponent->GetOwnerActor()->GetRootComponent());
+	DynamicMeshComponent = NewObject<USimpleDynamicMeshComponent>(UE::ToolTarget::GetTargetActor(Target));
+	DynamicMeshComponent->SetupAttachment(UE::ToolTarget::GetTargetActor(Target)->GetRootComponent());
 	DynamicMeshComponent->RegisterComponent();
-	DynamicMeshComponent->SetWorldTransform(TargetComponent->GetWorldTransform());
+	WorldTransform = UE::ToolTarget::GetLocalToWorldTransform(Target);
+	DynamicMeshComponent->SetWorldTransform((FTransform)WorldTransform);
 
 	// set materials
-	FComponentMaterialSet MaterialSet;
-	Cast<IMaterialProvider>(Target)->GetMaterialSet(MaterialSet);
+	FComponentMaterialSet MaterialSet = UE::ToolTarget::GetMaterialSet(Target);
 	for (int k = 0; k < MaterialSet.Materials.Num(); ++k)
 	{
 		DynamicMeshComponent->SetMaterial(k, MaterialSet.Materials[k]);
 	}
 
 	// dynamic mesh configuration settings
-	DynamicMeshComponent->TangentsType = EDynamicMeshTangentCalcType::AutoCalculated;
-	DynamicMeshComponent->InitializeMesh(Cast<IMeshDescriptionProvider>(Target)->GetMeshDescription());
+	DynamicMeshComponent->SetTangentsType(EDynamicMeshComponentTangentsMode::AutoCalculated);
+	DynamicMeshComponent->SetMesh(UE::ToolTarget::GetDynamicMeshCopy(Target));
 	OnDynamicMeshComponentChangedHandle =
 	    DynamicMeshComponent->OnMeshChanged.Add(FSimpleMulticastDelegate::FDelegate::CreateUObject(
 	        this, &UDeformMeshPolygonsTool::OnDynamicMeshComponentChanged));
@@ -886,14 +881,13 @@ void UDeformMeshPolygonsTool::Setup()
 	//initialize topology selector
 	TopoSelector.Initialize(DynamicMeshComponent->GetMesh(), &Topology);
 	TopoSelector.SetSpatialSource([this]() { return &GetSpatial(); });
-	TopoSelector.PointsWithinToleranceTest = [this, TargetComponent](const FVector3d& Position1, const FVector3d& Position2, double TolScale) {
-		UE::Geometry::FTransform3d Transform(TargetComponent->GetWorldTransform());
-		return ToolSceneQueriesUtil::PointSnapQuery(CameraState, Transform.TransformPosition(Position1), Transform.TransformPosition(Position2),
+	TopoSelector.PointsWithinToleranceTest = [this](const FVector3d& Position1, const FVector3d& Position2, double TolScale) {
+		return ToolSceneQueriesUtil::PointSnapQuery(CameraState, WorldTransform.TransformPosition(Position1), WorldTransform.TransformPosition(Position2),
 			ToolSceneQueriesUtil::GetDefaultVisualAngleSnapThreshD() * TolScale);
 	};
 
 	// hide input StaticMeshComponent
-	TargetComponent->SetOwnerVisibility(false);
+	UE::ToolTarget::HideSourceObject(Target);
 
 	// init state flags flags
 	bInDrag = false;
@@ -952,17 +946,18 @@ void UDeformMeshPolygonsTool::Shutdown(EToolShutdownType ShutdownType)
 	{
 		DynamicMeshComponent->OnMeshChanged.Remove(OnDynamicMeshComponentChangedHandle);
 
-		Cast<IPrimitiveComponentBackedTarget>(Target)->SetOwnerVisibility(true);
+		UE::ToolTarget::ShowSourceObject(Target);
 
 		if (ShutdownType == EToolShutdownType::Accept)
 		{
 			// this block bakes the modified DynamicMeshComponent back into the StaticMeshComponent inside an undo transaction
 			GetToolManager()->BeginUndoTransaction(LOCTEXT("DeformMeshPolygonsToolTransactionName", "Deform Mesh"));
-			Cast<IMeshDescriptionCommitter>(Target)->CommitMeshDescription([=](const IMeshDescriptionCommitter::FCommitterParams& CommitParams) {
+			DynamicMeshComponent->ProcessMesh([&](const FDynamicMesh3& ReadMesh)
+			{
 				FConversionToMeshDescriptionOptions ConversionOptions;
 				ConversionOptions.bSetPolyGroups =
-				    false; // don't save polygroups, as we may change these temporarily in this tool just to get a different edit effect
-				DynamicMeshComponent->Bake(CommitParams.MeshDescriptionOut, false, ConversionOptions);
+					false; // don't save polygroups, as we may change these temporarily in this tool just to get a different edit effect
+				UE::ToolTarget::CommitDynamicMeshUpdate(Target, ReadMesh, false, ConversionOptions);
 			});
 			GetToolManager()->EndUndoTransaction();
 		}
@@ -1039,9 +1034,8 @@ FDynamicMeshAABBTree3& UDeformMeshPolygonsTool::GetSpatial()
 
 bool UDeformMeshPolygonsTool::HitTest(const FRay& WorldRay, FHitResult& OutHit)
 {
-	UE::Geometry::FTransform3d Transform(Cast<IPrimitiveComponentBackedTarget>(Target)->GetWorldTransform());
-	FRay3d LocalRay(Transform.InverseTransformPosition((FVector3d)WorldRay.Origin),
-	                Transform.InverseTransformVector((FVector3d)WorldRay.Direction));
+	FRay3d LocalRay(WorldTransform.InverseTransformPosition((FVector3d)WorldRay.Origin),
+					WorldTransform.InverseTransformVector((FVector3d)WorldRay.Direction));
 	UE::Geometry::Normalize(LocalRay.Direction);
 
 	FGroupTopologySelection Selection;
@@ -1056,13 +1050,13 @@ bool UDeformMeshPolygonsTool::HitTest(const FRay& WorldRay, FHitResult& OutHit)
 	{
 		OutHit.FaceIndex   = Selection.GetASelectedCornerID();
 		OutHit.Distance    = LocalRay.Project(LocalPosition);
-		OutHit.ImpactPoint = (FVector)Transform.TransformPosition(LocalRay.PointAt(OutHit.Distance));
+		OutHit.ImpactPoint = (FVector)WorldTransform.TransformPosition(LocalRay.PointAt(OutHit.Distance));
 	}
 	else if (Selection.SelectedEdgeIDs.Num() > 0)
 	{
 		OutHit.FaceIndex   = Selection.GetASelectedEdgeID();
 		OutHit.Distance    = LocalRay.Project(LocalPosition);
-		OutHit.ImpactPoint = (FVector)Transform.TransformPosition(LocalRay.PointAt(OutHit.Distance));
+		OutHit.ImpactPoint = (FVector)WorldTransform.TransformPosition(LocalRay.PointAt(OutHit.Distance));
 	}
 	else
 	{
@@ -1075,8 +1069,8 @@ bool UDeformMeshPolygonsTool::HitTest(const FRay& WorldRay, FHitResult& OutHit)
 			Query.Find();
 			OutHit.FaceIndex = HitTID;
 			OutHit.Distance  = Query.RayParameter;
-			OutHit.Normal    = (FVector)Transform.TransformVectorNoScale(GetSpatial().GetMesh()->GetTriNormal(HitTID));
-			OutHit.ImpactPoint = (FVector)Transform.TransformPosition(LocalRay.PointAt(Query.RayParameter));
+			OutHit.Normal    = (FVector)WorldTransform.TransformVectorNoScale(GetSpatial().GetMesh()->GetTriNormal(HitTID));
+			OutHit.ImpactPoint = (FVector)WorldTransform.TransformPosition(LocalRay.PointAt(Query.RayParameter));
 		}
 	}
 	return true;
@@ -1085,9 +1079,8 @@ bool UDeformMeshPolygonsTool::HitTest(const FRay& WorldRay, FHitResult& OutHit)
 
 void UDeformMeshPolygonsTool::OnBeginDrag(const FRay& WorldRay)
 {
-	UE::Geometry::FTransform3d Transform(Cast<IPrimitiveComponentBackedTarget>(Target)->GetWorldTransform());
-	FRay3d LocalRay(Transform.InverseTransformPosition((FVector3d)WorldRay.Origin),
-	                Transform.InverseTransformVector((FVector3d)WorldRay.Direction));
+	FRay3d LocalRay(WorldTransform.InverseTransformPosition((FVector3d)WorldRay.Origin),
+					WorldTransform.InverseTransformVector((FVector3d)WorldRay.Direction));
 	UE::Geometry::Normalize(LocalRay.Direction);
 
 	HilightSelection.Clear();
@@ -1105,8 +1098,8 @@ void UDeformMeshPolygonsTool::OnBeginDrag(const FRay& WorldRay)
 
 	HilightSelection = Selection;
 
-	FVector3d WorldHitPos    = Transform.TransformPosition(LocalPosition);
-	FVector3d WorldHitNormal = Transform.TransformVector(LocalNormal);
+	FVector3d WorldHitPos    = WorldTransform.TransformPosition(LocalPosition);
+	FVector3d WorldHitNormal = WorldTransform.TransformVector(LocalNormal);
 
 	bInDrag             = true;
 	StartHitPosWorld    = (FVector)WorldHitPos;
@@ -1117,7 +1110,7 @@ void UDeformMeshPolygonsTool::OnBeginDrag(const FRay& WorldRay)
 	UpdateActiveSurfaceFrame(HilightSelection);
 	UpdateQuickTransformer();
 
-	LastBrushPosLocal  = (FVector)Transform.InverseTransformPosition((FVector3d)LastHitPosWorld);
+	LastBrushPosLocal  = (FVector)WorldTransform.InverseTransformPosition((FVector3d)LastHitPosWorld);
 	StartBrushPosLocal = LastBrushPosLocal;
 
 	// Record the requested deformation strategy - NB: will be forced to linear if there aren't any free points to solve.
@@ -1194,8 +1187,6 @@ void UDeformMeshPolygonsTool::OnBeginDrag(const FRay& WorldRay)
 
 void UDeformMeshPolygonsTool::UpdateActiveSurfaceFrame(FGroupTopologySelection& Selection)
 {
-	UE::Geometry::FTransform3d Transform(Cast<IPrimitiveComponentBackedTarget>(Target)->GetWorldTransform());
-
 	// update surface frame
 	ActiveSurfaceFrame.Origin = (FVector3d)StartHitPosWorld;
 	if (HilightSelection.SelectedCornerIDs.Num() == 1)
@@ -1210,7 +1201,7 @@ void UDeformMeshPolygonsTool::UpdateActiveSurfaceFrame(FGroupTopologySelection&
 			FVector3d Tangent;
 			if (Topology.GetGroupEdgeTangent(HilightSelection.GetASelectedEdgeID(), Tangent))
 			{
-				Tangent = Transform.TransformVector(Tangent);
+				Tangent = WorldTransform.TransformVector(Tangent);
 				ActiveSurfaceFrame.ConstrainedAlignAxis(0, Tangent, ActiveSurfaceFrame.Z());
 			}
 		}
@@ -1300,9 +1291,8 @@ bool UDeformMeshPolygonsTool::OnUpdateHover(const FInputDeviceRay& DevicePos)
 	//if (!bNeedEmitEndChange)
 	if (ActiveVertexChange == nullptr)
 	{
-		UE::Geometry::FTransform3d Transform(Cast<IPrimitiveComponentBackedTarget>(Target)->GetWorldTransform());
-		FRay3d LocalRay(Transform.InverseTransformPosition((FVector3d)DevicePos.WorldRay.Origin),
-		                Transform.InverseTransformVector((FVector3d)DevicePos.WorldRay.Direction));
+		FRay3d LocalRay(WorldTransform.InverseTransformPosition((FVector3d)DevicePos.WorldRay.Origin),
+						WorldTransform.InverseTransformVector((FVector3d)DevicePos.WorldRay.Direction));
 		UE::Geometry::Normalize(LocalRay.Direction);
 
 		HilightSelection.Clear();
@@ -1312,8 +1302,8 @@ bool UDeformMeshPolygonsTool::OnUpdateHover(const FInputDeviceRay& DevicePos)
 
 		if (bHit)
 		{
-			StartHitPosWorld    = (FVector)Transform.TransformPosition(LocalPosition);
-			StartHitNormalWorld = (FVector)Transform.TransformVector(LocalNormal);
+			StartHitPosWorld    = (FVector)WorldTransform.TransformPosition(LocalPosition);
+			StartHitNormalWorld = (FVector)WorldTransform.TransformVector(LocalNormal);
 
 			UpdateActiveSurfaceFrame(HilightSelection);
 			UpdateQuickTransformer();
@@ -1398,7 +1388,6 @@ void UDeformMeshPolygonsTool::ComputeUpdate_Rotate()
 	FGroupTopologyDeformer& SelectedDeformer = (bIsLaplacian) ? *LaplacianDeformer : LinearDeformer;
 
 	FDynamicMesh3* Mesh    = DynamicMeshComponent->GetMesh();
-	UE::Geometry::FTransform3d Transform(Cast<IPrimitiveComponentBackedTarget>(Target)->GetWorldTransform());
 	FVector NewHitPosWorld = LastHitPosWorld;
 
 	FVector3d SnappedPoint;
@@ -1454,8 +1443,8 @@ void UDeformMeshPolygonsTool::ComputeUpdate_Rotate()
 	FVector2d RotateToVec = RotationStartFrame.ToPlaneUV((FVector3d)NewHitPosWorld, 2);
 	UE::Geometry::Normalize(RotateToVec);
 	double AngleRad = UE::Geometry::SignedAngleR(RotateStartVec, RotateToVec);
-	FQuaterniond Rotation(Transform.InverseTransformVectorNoScale(RotationStartFrame.Z()), AngleRad, false);
-	FVector3d LocalOrigin = Transform.InverseTransformPosition(RotationStartFrame.Origin);
+	FQuaterniond Rotation(WorldTransform.InverseTransformVectorNoScale(RotationStartFrame.Z()), AngleRad, false);
+	FVector3d LocalOrigin = WorldTransform.InverseTransformPosition(RotationStartFrame.Origin);
 
 	// Linear Deformer: Update Mesh the rotation,
 	// Laplacian Deformer:  Update handles constraints with the rotation and set bDeformerNeedsToRun = true;.
@@ -1492,7 +1481,6 @@ void UDeformMeshPolygonsTool::ComputeUpdate_Translate()
 		};
 	}
 
-	FTransform Transform   = Cast<IPrimitiveComponentBackedTarget>(Target)->GetWorldTransform();
 	FVector NewHitPosWorld = LastHitPosWorld;
 	FVector3d SnappedPoint;
 	if (QuickAxisTranslater.UpdateSnap(FRay3d(UpdateRay), SnappedPoint, PointConstraintFunc))
@@ -1504,8 +1492,8 @@ void UDeformMeshPolygonsTool::ComputeUpdate_Translate()
 		return;
 	}
 
-	FVector NewBrushPosLocal = Transform.InverseTransformPosition(NewHitPosWorld);
-	FVector3d NewMoveDelta   = (FVector3d)NewBrushPosLocal - (FVector3d)StartBrushPosLocal;
+	FVector3d NewBrushPosLocal = WorldTransform.InverseTransformPosition(NewHitPosWorld);
+	FVector3d NewMoveDelta   = NewBrushPosLocal - (FVector3d)StartBrushPosLocal;
 
 	FDynamicMesh3* Mesh = DynamicMeshComponent->GetMesh();
 	if (LastMoveDelta.SquaredLength() > 0.)
@@ -1571,7 +1559,7 @@ void UDeformMeshPolygonsTool::Render(IToolsContextRenderAPI* RenderAPI)
 	FDynamicMesh3* TargetMesh                    = DynamicMeshComponent->GetMesh();
 
 	PolyEdgesRenderer.BeginFrame(RenderAPI, CameraState);
-	PolyEdgesRenderer.SetTransform(Cast<IPrimitiveComponentBackedTarget>(Target)->GetWorldTransform());
+	PolyEdgesRenderer.SetTransform((FTransform)WorldTransform);
 
 
 	for (FGroupTopology::FGroupEdge& Edge : Topology.Edges)
@@ -1588,7 +1576,7 @@ void UDeformMeshPolygonsTool::Render(IToolsContextRenderAPI* RenderAPI)
 
 
 	HilightRenderer.BeginFrame(RenderAPI, CameraState);
-	HilightRenderer.SetTransform(Cast<IPrimitiveComponentBackedTarget>(Target)->GetWorldTransform());
+	HilightRenderer.SetTransform((FTransform)WorldTransform);
 
 #ifdef DEBUG_ROI_WEIGHTS
 	FDynamicMesh3* MeshPtr = DynamicMeshComponent->GetMesh();
diff --git a/Engine/Plugins/Experimental/MeshModelingToolset/Source/MeshModelingTools/Private/DisplaceMeshTool.cpp b/Engine/Plugins/Experimental/MeshModelingToolset/Source/MeshModelingTools/Private/DisplaceMeshTool.cpp
--- a/Engine/Plugins/Experimental/MeshModelingToolset/Source/MeshModelingTools/Private/DisplaceMeshTool.cpp
+++ b/Engine/Plugins/Experimental/MeshModelingToolset/Source/MeshModelingTools/Private/DisplaceMeshTool.cpp
@@ -15,16 +15,13 @@
 #include "Async/ParallelFor.h"
 #include "ProfilingDebugging/ScopedTimers.h"
 #include "MeshDescription.h"
+#include "ModelingToolTargetUtil.h"
 
 // needed to disable normals recalculation on the underlying asset
 #include "AssetUtils/MeshDescriptionUtil.h"
 #include "Engine/Classes/Engine/StaticMesh.h"
 #include "Engine/Classes/Components/StaticMeshComponent.h"
 
-#include "TargetInterfaces/MaterialProvider.h"
-#include "TargetInterfaces/MeshDescriptionCommitter.h"
-#include "TargetInterfaces/MeshDescriptionProvider.h"
-#include "TargetInterfaces/PrimitiveComponentBackedTarget.h"
 
 #include "ExplicitUseGeometryMathTypes.h"		// using UE::Geometry::(math types)
 using namespace UE::Geometry;
@@ -773,9 +770,9 @@ void UDisplaceMeshTool::Setup()
 	}
 	
 	// populate weight maps list
-	IMeshDescriptionProvider* TargetMeshProvider = Cast<IMeshDescriptionProvider>(Target);
+	const FMeshDescription* MeshDescription = UE::ToolTarget::GetMeshDescription(Target);
 	TArray<FName> WeightMaps;
-	UE::WeightMaps::FindVertexWeightMaps(TargetMeshProvider->GetMeshDescription(), WeightMaps);
+	UE::WeightMaps::FindVertexWeightMaps(MeshDescription, WeightMaps);
 	CommonProperties->WeightMapsList.Add(TEXT("None"));
 	for (FName Name : WeightMaps)
 	{
@@ -789,23 +786,21 @@ void UDisplaceMeshTool::Setup()
 
 
 	// create dynamic mesh component to use for live preview
-	IPrimitiveComponentBackedTarget* TargetComponent = Cast<IPrimitiveComponentBackedTarget>(Target);
-	DynamicMeshComponent = NewObject<USimpleDynamicMeshComponent>(TargetComponent->GetOwnerActor(), "DynamicMesh");
-	DynamicMeshComponent->SetupAttachment(TargetComponent->GetOwnerActor()->GetRootComponent());
+	DynamicMeshComponent = NewObject<USimpleDynamicMeshComponent>(UE::ToolTarget::GetTargetActor(Target));
+	DynamicMeshComponent->SetupAttachment(UE::ToolTarget::GetTargetActor(Target)->GetRootComponent());
 	DynamicMeshComponent->RegisterComponent();
-	DynamicMeshComponent->SetWorldTransform(TargetComponent->GetWorldTransform());
+	DynamicMeshComponent->SetWorldTransform((FTransform)UE::ToolTarget::GetLocalToWorldTransform(Target));
 	DynamicMeshComponent->bExplicitShowWireframe = CommonProperties->bShowWireframe;
 
 	// transfer materials
-	FComponentMaterialSet MaterialSet;
-	Cast<IMaterialProvider>(Target)->GetMaterialSet(MaterialSet);
+	FComponentMaterialSet MaterialSet = UE::ToolTarget::GetMaterialSet(Target);
 	for (int k = 0; k < MaterialSet.Materials.Num(); ++k)
 	{
 		DynamicMeshComponent->SetMaterial(k, MaterialSet.Materials[k]);
 	}
 
-	DynamicMeshComponent->TangentsType = EDynamicMeshTangentCalcType::AutoCalculated;
-	DynamicMeshComponent->InitializeMesh(TargetMeshProvider->GetMeshDescription());
+	DynamicMeshComponent->SetTangentsType(EDynamicMeshComponentTangentsMode::AutoCalculated);
+	DynamicMeshComponent->SetMesh(UE::ToolTarget::GetDynamicMeshCopy(Target));
 	OriginalMesh.Copy(*DynamicMeshComponent->GetMesh());
 	OriginalMeshSpatial.SetMesh(&OriginalMesh, true);
 
@@ -827,7 +822,7 @@ void UDisplaceMeshTool::Setup()
 	Displacer = MakeUnique<FDisplaceMeshOpFactory>(SubdividedMesh, Parameters, CommonProperties->DisplacementType);
 		
 	// hide input StaticMeshComponent
-	TargetComponent->SetOwnerVisibility(false);
+	UE::ToolTarget::HideSourceObject(Target);
 
 	// initialize our properties
 	ToolPropertyObjects.Add(this);
@@ -880,8 +875,7 @@ void UDisplaceMeshTool::Shutdown(EToolShutdownType ShutdownType)
 
 	if (DynamicMeshComponent != nullptr)
 	{
-		IPrimitiveComponentBackedTarget* TargetComponent = Cast<IPrimitiveComponentBackedTarget>(Target);
-		TargetComponent->SetOwnerVisibility(true);
+		UE::ToolTarget::ShowSourceObject(Target);
 
 		if (ShutdownType == EToolShutdownType::Accept)
 		{
@@ -892,7 +886,7 @@ void UDisplaceMeshTool::Shutdown(EToolShutdownType ShutdownType)
 			// on the underlying StaticMesh Asset, or it will run on the Bake() below and the output result will not be the same as the preview
 			if (CommonProperties->DisplacementType == EDisplaceMeshToolDisplaceType::DisplacementMap && TextureMapProperties->bRecalcNormals == false)
 			{
-				UStaticMeshComponent* StaticMeshComponent = Cast<UStaticMeshComponent>(TargetComponent->GetOwnerComponent());
+				UStaticMeshComponent* StaticMeshComponent = Cast<UStaticMeshComponent>(UE::ToolTarget::GetTargetComponent(Target));
 				if (StaticMeshComponent)
 				{
 					if (UStaticMesh* StaticMesh = StaticMeshComponent->GetStaticMesh())
@@ -907,10 +901,9 @@ void UDisplaceMeshTool::Shutdown(EToolShutdownType ShutdownType)
 				}
 			}
 
-
-			Cast<IMeshDescriptionCommitter>(Target)->CommitMeshDescription([this](const IMeshDescriptionCommitter::FCommitterParams& CommitParams)
+			DynamicMeshComponent->ProcessMesh([&](const FDynamicMesh3& ReadMesh)
 			{
-				DynamicMeshComponent->Bake(CommitParams.MeshDescriptionOut, CommonProperties->Subdivisions > 0);
+				UE::ToolTarget::CommitDynamicMeshUpdate(Target, ReadMesh, CommonProperties->Subdivisions > 0);
 			});
 			GetToolManager()->EndUndoTransaction();
 		}
@@ -1212,7 +1205,8 @@ void UDisplaceMeshTool::UpdateActiveWeightMap()
 	else
 	{
 		TSharedPtr<FIndexedWeightMap, ESPMode::ThreadSafe> NewWeightMap = MakeShared<FIndexedWeightMap, ESPMode::ThreadSafe>();
-		UE::WeightMaps::GetVertexWeightMap(Cast<IMeshDescriptionProvider>(Target)->GetMeshDescription(), CommonProperties->WeightMap, *NewWeightMap, 1.0f);
+		const FMeshDescription* MeshDescription = UE::ToolTarget::GetMeshDescription(Target);
+		UE::WeightMaps::GetVertexWeightMap(MeshDescription, CommonProperties->WeightMap, *NewWeightMap, 1.0f);
 		if (CommonProperties->bInvertWeightMap)
 		{
 			NewWeightMap->InvertWeightMap();
diff --git a/Engine/Plugins/Experimental/MeshModelingToolset/Source/MeshModelingTools/Private/DrawAndRevolveTool.cpp b/Engine/Plugins/Experimental/MeshModelingToolset/Source/MeshModelingTools/Private/DrawAndRevolveTool.cpp
--- a/Engine/Plugins/Experimental/MeshModelingToolset/Source/MeshModelingTools/Private/DrawAndRevolveTool.cpp
+++ b/Engine/Plugins/Experimental/MeshModelingToolset/Source/MeshModelingTools/Private/DrawAndRevolveTool.cpp
@@ -270,7 +270,7 @@ void UDrawAndRevolveTool::StartPreview()
 	Preview = NewObject<UMeshOpPreviewWithBackgroundCompute>(RevolveOpCreator, "RevolveToolPreview");
 
 	Preview->Setup(TargetWorld, RevolveOpCreator);
-	Preview->PreviewMesh->SetTangentsMode(EDynamicMeshTangentCalcType::AutoCalculated);
+	Preview->PreviewMesh->SetTangentsMode(EDynamicMeshComponentTangentsMode::AutoCalculated);
 
 	Preview->ConfigureMaterials(MaterialProperties->Material.Get(),
 		ToolSetupUtil::GetDefaultWorkingMaterial(GetToolManager()));
diff --git a/Engine/Plugins/Experimental/MeshModelingToolset/Source/MeshModelingTools/Private/DynamicMeshBrushTool.cpp b/Engine/Plugins/Experimental/MeshModelingToolset/Source/MeshModelingTools/Private/DynamicMeshBrushTool.cpp
--- a/Engine/Plugins/Experimental/MeshModelingToolset/Source/MeshModelingTools/Private/DynamicMeshBrushTool.cpp
+++ b/Engine/Plugins/Experimental/MeshModelingToolset/Source/MeshModelingTools/Private/DynamicMeshBrushTool.cpp
@@ -3,10 +3,7 @@
 #include "DynamicMeshBrushTool.h"
 #include "InteractiveToolManager.h"
 #include "ToolBuilderUtil.h"
-
-#include "TargetInterfaces/MaterialProvider.h"
-#include "TargetInterfaces/MeshDescriptionProvider.h"
-#include "TargetInterfaces/PrimitiveComponentBackedTarget.h"
+#include "ModelingToolTargetUtil.h"
 
 #include "ExplicitUseGeometryMathTypes.h"		// using UE::Geometry::(math types)
 using namespace UE::Geometry;
@@ -25,29 +22,27 @@ UDynamicMeshBrushTool::UDynamicMeshBrushTool()
 
 void UDynamicMeshBrushTool::Setup()
 {
-	IPrimitiveComponentBackedTarget* TargetComponent = Cast<IPrimitiveComponentBackedTarget>(Target);
 	PreviewMesh = NewObject<UPreviewMesh>(this);
 	PreviewMesh->bBuildSpatialDataStructure = true;
-	PreviewMesh->CreateInWorld(TargetComponent->GetOwnerActor()->GetWorld(), FTransform::Identity);
-	PreviewMesh->SetTransform(TargetComponent->GetWorldTransform());
+	PreviewMesh->CreateInWorld(UE::ToolTarget::GetTargetActor(Target)->GetWorld(), FTransform::Identity);
+	FTransform3d LocalToWorldTransform = UE::ToolTarget::GetLocalToWorldTransform(Target);
+	PreviewMesh->SetTransform((FTransform)LocalToWorldTransform);
 
-	FComponentMaterialSet MaterialSet;
-	Cast<IMaterialProvider>(Target)->GetMaterialSet(MaterialSet);
+	FComponentMaterialSet MaterialSet = UE::ToolTarget::GetMaterialSet(Target);
 	PreviewMesh->SetMaterials(MaterialSet.Materials);
 
-	// initialize from LOD-0 MeshDescription
-	PreviewMesh->InitializeMesh(Cast<IMeshDescriptionProvider>(Target)->GetMeshDescription());
+	PreviewMesh->ReplaceMesh(UE::ToolTarget::GetDynamicMeshCopy(Target));
+
 	OnBaseMeshComponentChangedHandle = PreviewMesh->GetOnMeshChanged().Add(
 		FSimpleMulticastDelegate::FDelegate::CreateUObject(this, &UDynamicMeshBrushTool::OnBaseMeshComponentChanged));
 
 	// call this here so that base tool can estimate target dimension
 	InputMeshBoundsLocal = PreviewMesh->GetPreviewDynamicMesh()->GetBounds();
-	double ScaledDim = TargetComponent->GetWorldTransform().TransformVector(FVector::OneVector).Size();
+	double ScaledDim = LocalToWorldTransform.TransformVector(FVector::OneVector).Size();
 	this->WorldToLocalScale = FMathd::Sqrt3 / FMathd::Max(FMathf::ZeroTolerance, ScaledDim);
 	UBaseBrushTool::Setup();
 
-	// hide input StaticMeshComponent
-	TargetComponent->SetOwnerVisibility(false);
+	UE::ToolTarget::HideSourceObject(Target);
 }
 
 
@@ -62,7 +57,7 @@ void UDynamicMeshBrushTool::Shutdown(EToolShutdownType ShutdownType)
 {
 	UBaseBrushTool::Shutdown(ShutdownType);
 
-	Cast<IPrimitiveComponentBackedTarget>(Target)->SetOwnerVisibility(true);
+	UE::ToolTarget::ShowSourceObject(Target);
 
 	if (PreviewMesh != nullptr)
 	{
diff --git a/Engine/Plugins/Experimental/MeshModelingToolset/Source/MeshModelingTools/Private/DynamicMeshSculptTool.cpp b/Engine/Plugins/Experimental/MeshModelingToolset/Source/MeshModelingTools/Private/DynamicMeshSculptTool.cpp
--- a/Engine/Plugins/Experimental/MeshModelingToolset/Source/MeshModelingTools/Private/DynamicMeshSculptTool.cpp
+++ b/Engine/Plugins/Experimental/MeshModelingToolset/Source/MeshModelingTools/Private/DynamicMeshSculptTool.cpp
@@ -18,6 +18,7 @@
 #include "PreviewMesh.h"
 #include "ToolSetupUtil.h"
 #include "ToolSceneQueriesUtil.h"
+#include "ModelingToolTargetUtil.h"
 
 #include "Changes/MeshVertexChange.h"
 #include "Changes/MeshChange.h"
@@ -36,11 +37,6 @@
 #include "Materials/Material.h"
 #include "Materials/MaterialInstanceDynamic.h"
 
-#include "TargetInterfaces/MaterialProvider.h"
-#include "TargetInterfaces/MeshDescriptionCommitter.h"
-#include "TargetInterfaces/MeshDescriptionProvider.h"
-#include "TargetInterfaces/PrimitiveComponentBackedTarget.h"
-
 #include "ExplicitUseGeometryMathTypes.h"		// using UE::Geometry::(math types)
 using namespace UE::Geometry;
 
@@ -104,17 +100,16 @@ void UDynamicMeshSculptTool::Setup()
 	SetToolDisplayName(LOCTEXT("ToolName", "DynaSculpt"));
 
 	// create dynamic mesh component to use for live preview
-	IPrimitiveComponentBackedTarget* TargetComponent = Cast<IPrimitiveComponentBackedTarget>(Target);
-	DynamicMeshComponent = NewObject<UOctreeDynamicMeshComponent>(TargetComponent->GetOwnerActor(), "DynamicMeshSculptToolMesh");
-	DynamicMeshComponent->SetupAttachment(TargetComponent->GetOwnerActor()->GetRootComponent());
+	DynamicMeshComponent = NewObject<UOctreeDynamicMeshComponent>(UE::ToolTarget::GetTargetActor(Target));
+	DynamicMeshComponent->SetupAttachment(UE::ToolTarget::GetTargetActor(Target)->GetRootComponent());
 	DynamicMeshComponent->RegisterComponent();
 
 	// initialize from LOD-0 MeshDescription
-	DynamicMeshComponent->InitializeMesh(Cast<IMeshDescriptionProvider>(Target)->GetMeshDescription());
+	DynamicMeshComponent->SetMesh(UE::ToolTarget::GetDynamicMeshCopy(Target));
 
 	// transform mesh to world space because handling scaling inside brush is a mess
 	// Note: this transform does not include translation ( so only the 3x3 transform)
-	InitialTargetTransform = UE::Geometry::FTransform3d(TargetComponent->GetWorldTransform());
+	InitialTargetTransform = UE::ToolTarget::GetLocalToWorldTransform(Target);
 	// clamp scaling because if we allow zero-scale we cannot invert this transform on Accept
 	InitialTargetTransform.ClampMinimumScale(0.01);
 	FVector3d Translation = InitialTargetTransform.GetTranslation();
@@ -125,10 +120,10 @@ void UDynamicMeshSculptTool::Setup()
 	DynamicMeshComponent->SetWorldTransform((FTransform)CurTargetTransform);
 
 	// copy material if there is one
-	UMaterialInterface* Material = Cast<IMaterialProvider>(Target)->GetMaterial(0);
-	if (Material != nullptr)
+	FComponentMaterialSet MaterialSet = UE::ToolTarget::GetMaterialSet(Target);
+	if (MaterialSet.Materials.Num() > 0)
 	{
-		DynamicMeshComponent->SetMaterial(0, Material);
+		DynamicMeshComponent->SetMaterial(0, MaterialSet.Materials[0]);
 	}
 
 	OnDynamicMeshComponentChangedHandle = DynamicMeshComponent->OnMeshChanged.Add(
@@ -166,7 +161,7 @@ void UDynamicMeshSculptTool::Setup()
 	InitialEdgeLength = EstimateIntialSafeTargetLength(*Mesh, 5000);
 
 	// hide input Component
-	TargetComponent->SetOwnerVisibility(false);
+	UE::ToolTarget::HideSourceObject(Target);
 
 	// init state flags flags
 	bInDrag = false;
@@ -307,7 +302,7 @@ void UDynamicMeshSculptTool::Shutdown(EToolShutdownType ShutdownType)
 	{
 		DynamicMeshComponent->OnMeshChanged.Remove(OnDynamicMeshComponentChangedHandle);
 
-		Cast<IPrimitiveComponentBackedTarget>(Target)->SetOwnerVisibility(true);
+		UE::ToolTarget::ShowSourceObject(Target);
 
 		if (ShutdownType == EToolShutdownType::Accept)
 		{
@@ -316,9 +311,9 @@ void UDynamicMeshSculptTool::Shutdown(EToolShutdownType ShutdownType)
 
 			// this block bakes the modified DynamicMeshComponent back into the StaticMeshComponent inside an undo transaction
 			GetToolManager()->BeginUndoTransaction(LOCTEXT("SculptMeshToolTransactionName", "Sculpt Mesh"));
-			Cast<IMeshDescriptionCommitter>(Target)->CommitMeshDescription([=](const IMeshDescriptionCommitter::FCommitterParams& CommitParams)
+			DynamicMeshComponent->ProcessMesh([&](const FDynamicMesh3& ReadMesh)
 			{
-				DynamicMeshComponent->Bake(CommitParams.MeshDescriptionOut, bHaveRemeshed);
+				UE::ToolTarget::CommitDynamicMeshUpdate(Target, ReadMesh, bHaveRemeshed);
 			});
 			GetToolManager()->EndUndoTransaction();
 		}
@@ -2453,7 +2448,7 @@ void UDynamicMeshSculptTool::UpdateMaterialMode(EMeshEditingMaterialModes Materi
 	if (MaterialMode == EMeshEditingMaterialModes::ExistingMaterial)
 	{
 		DynamicMeshComponent->ClearOverrideRenderMaterial();
-		DynamicMeshComponent->bCastDynamicShadow = Cast<IPrimitiveComponentBackedTarget>(Target)->GetOwnerComponent()->bCastDynamicShadow;
+		DynamicMeshComponent->bCastDynamicShadow = UE::ToolTarget::GetTargetComponent(Target)->bCastDynamicShadow;
 		ActiveOverrideMaterial = nullptr;
 	}
 	else 
diff --git a/Engine/Plugins/Experimental/MeshModelingToolset/Source/MeshModelingTools/Private/EdgeLoopInsertionTool.cpp b/Engine/Plugins/Experimental/MeshModelingToolset/Source/MeshModelingTools/Private/EdgeLoopInsertionTool.cpp
--- a/Engine/Plugins/Experimental/MeshModelingToolset/Source/MeshModelingTools/Private/EdgeLoopInsertionTool.cpp
+++ b/Engine/Plugins/Experimental/MeshModelingToolset/Source/MeshModelingTools/Private/EdgeLoopInsertionTool.cpp
@@ -157,7 +157,7 @@ void UEdgeLoopInsertionTool::SetupPreview()
 
 	Preview = NewObject<UMeshOpPreviewWithBackgroundCompute>(OpFactory);
 	Preview->Setup(TargetWorld, OpFactory);
-	Preview->PreviewMesh->SetTangentsMode(EDynamicMeshTangentCalcType::AutoCalculated);
+	Preview->PreviewMesh->SetTangentsMode(EDynamicMeshComponentTangentsMode::AutoCalculated);
 
 	FComponentMaterialSet MaterialSet;
 	Cast<IMaterialProvider>(Target)->GetMaterialSet(MaterialSet);
diff --git a/Engine/Plugins/Experimental/MeshModelingToolset/Source/MeshModelingTools/Private/EditMeshPolygonsTool.cpp b/Engine/Plugins/Experimental/MeshModelingToolset/Source/MeshModelingTools/Private/EditMeshPolygonsTool.cpp
--- a/Engine/Plugins/Experimental/MeshModelingToolset/Source/MeshModelingTools/Private/EditMeshPolygonsTool.cpp
+++ b/Engine/Plugins/Experimental/MeshModelingToolset/Source/MeshModelingTools/Private/EditMeshPolygonsTool.cpp
@@ -28,6 +28,7 @@
 #include "Changes/MeshChange.h"
 #include "DynamicMesh/MeshIndexUtil.h"
 #include "MeshRegionBoundaryLoops.h"
+#include "ModelingToolTargetUtil.h"
 
 #include "Operations/OffsetMeshRegion.h"
 #include "Operations/InsetMeshRegion.h"
@@ -40,11 +41,6 @@
 #include "Materials/MaterialInstanceDynamic.h"
 #include "Components/BrushComponent.h"
 
-#include "TargetInterfaces/MaterialProvider.h"
-#include "TargetInterfaces/MeshDescriptionCommitter.h"
-#include "TargetInterfaces/MeshDescriptionProvider.h"
-#include "TargetInterfaces/PrimitiveComponentBackedTarget.h"
-
 #include "ExplicitUseGeometryMathTypes.h"		// using UE::Geometry::(math types)
 using namespace UE::Geometry;
 
@@ -170,16 +166,14 @@ void UEditMeshPolygonsTool::Setup()
 	AddInputBehavior(ClickBehavior);
 
 	// create dynamic mesh component to use for live preview
-	IPrimitiveComponentBackedTarget* TargetComponent = Cast<IPrimitiveComponentBackedTarget>(Target);
-	DynamicMeshComponent = NewObject<USimpleDynamicMeshComponent>(TargetComponent->GetOwnerActor(), "DynamicMesh");
-	DynamicMeshComponent->SetupAttachment(TargetComponent->GetOwnerActor()->GetRootComponent());
+	DynamicMeshComponent = NewObject<USimpleDynamicMeshComponent>(UE::ToolTarget::GetTargetActor(Target), "DynamicMesh");
+	DynamicMeshComponent->SetupAttachment(UE::ToolTarget::GetTargetActor(Target)->GetRootComponent());
 	DynamicMeshComponent->RegisterComponent();
-	DynamicMeshComponent->SetWorldTransform(TargetComponent->GetWorldTransform());
-	WorldTransform = FTransform3d(DynamicMeshComponent->GetComponentTransform());
+	WorldTransform = UE::ToolTarget::GetLocalToWorldTransform(Target);
+	DynamicMeshComponent->SetWorldTransform((FTransform)WorldTransform);
 
 	// set materials
-	FComponentMaterialSet MaterialSet;
-	Cast<IMaterialProvider>(Target)->GetMaterialSet(MaterialSet);
+	FComponentMaterialSet MaterialSet = UE::ToolTarget::GetMaterialSet(Target);
 	for (int k = 0; k < MaterialSet.Materials.Num(); ++k)
 	{
 		DynamicMeshComponent->SetMaterial(k, MaterialSet.Materials[k]);
@@ -201,8 +195,8 @@ void UEditMeshPolygonsTool::Setup()
 
 
 	// dynamic mesh configuration settings
-	DynamicMeshComponent->TangentsType = EDynamicMeshTangentCalcType::AutoCalculated;
-	DynamicMeshComponent->InitializeMesh(Cast<IMeshDescriptionProvider>(Target)->GetMeshDescription());
+	DynamicMeshComponent->SetTangentsType(EDynamicMeshComponentTangentsMode::AutoCalculated);
+	DynamicMeshComponent->SetMesh(UE::ToolTarget::GetDynamicMeshCopy(Target));
 	FMeshNormals::QuickComputeVertexNormals(*DynamicMeshComponent->GetMesh());
 	OnDynamicMeshComponentChangedHandle = DynamicMeshComponent->OnMeshChanged.Add(
 		FSimpleMulticastDelegate::FDelegate::CreateUObject(this, &UEditMeshPolygonsTool::OnDynamicMeshComponentChanged));
@@ -255,7 +249,7 @@ void UEditMeshPolygonsTool::Setup()
 	}
 
 	// hide input StaticMeshComponent
-	TargetComponent->SetOwnerVisibility(false);
+	UE::ToolTarget::HideSourceObject(Target);
 
 	// init state flags flags
 	bInDrag = false;
@@ -382,13 +376,13 @@ bool UEditMeshPolygonsTool::IsStoredToolSelectionUsable(const UGroupTopologyStor
 	// same vids unrecoverable. Once we have persistence of dynamic meshes, this will
 	// hopefully not become a problem, and this function (along with stored selection
 	// identifying info) will change.
-	return !Cast<UBrushComponent>(Cast<IPrimitiveComponentBackedTarget>(Target)->GetOwnerComponent())
+	return !Cast<UBrushComponent>(UE::ToolTarget::GetTargetComponent(Target))
 
 		&& StoredSelection
 		&& StoredSelection->IdentifyingInfo.TopologyType == (bTriangleMode ? 
 			UGroupTopologyStorableSelection::ETopologyType::FTriangleGroupTopology
 			: UGroupTopologyStorableSelection::ETopologyType::FGroupTopology)
-		&& StoredSelection->IdentifyingInfo.ComponentTarget == Cast<IPrimitiveComponentBackedTarget>(Target)->GetOwnerComponent()
+		&& StoredSelection->IdentifyingInfo.ComponentTarget == UE::ToolTarget::GetTargetComponent(Target)
 		&& !StoredSelection->IsEmpty();
 }
 
@@ -416,8 +410,7 @@ void UEditMeshPolygonsTool::Shutdown(EToolShutdownType ShutdownType)
 	{
 		DynamicMeshComponent->OnMeshChanged.Remove(OnDynamicMeshComponentChangedHandle);
 
-		IPrimitiveComponentBackedTarget* TargetComponent = Cast<IPrimitiveComponentBackedTarget>(Target);
-		TargetComponent->SetOwnerVisibility(true);
+		UE::ToolTarget::ShowSourceObject(Target);
 
 		if (ShutdownType == EToolShutdownType::Accept)
 		{
@@ -427,10 +420,10 @@ void UEditMeshPolygonsTool::Shutdown(EToolShutdownType ShutdownType)
 			// Prep if we have a selection to store. We don't support storing selections for volumes
 			// because the conversion will change vids.
 			if (!SelectionMechanic->GetActiveSelection().IsEmpty()
-				&& !Cast<UBrushComponent>(TargetComponent->GetOwnerComponent()))
+				&& !Cast<UBrushComponent>(UE::ToolTarget::GetTargetComponent(Target)))
 			{
 				NewStoredToolSelection = NewObject<UGroupTopologyStorableSelection>();
-				NewStoredToolSelection->IdentifyingInfo.ComponentTarget = TargetComponent->GetOwnerComponent();
+				NewStoredToolSelection->IdentifyingInfo.ComponentTarget = UE::ToolTarget::GetTargetComponent(Target);
 				NewStoredToolSelection->IdentifyingInfo.TopologyType = (bTriangleMode ?
 					UGroupTopologyStorableSelection::ETopologyType::FTriangleGroupTopology
 					: UGroupTopologyStorableSelection::ETopologyType::FGroupTopology);
@@ -451,12 +444,12 @@ void UEditMeshPolygonsTool::Shutdown(EToolShutdownType ShutdownType)
 
 			// this block bakes the modified DynamicMeshComponent back into the StaticMeshComponent inside an undo transaction
 			GetToolManager()->BeginUndoTransaction(LOCTEXT("EditMeshPolygonsToolTransactionName", "Deform Mesh"));
-			Cast<IMeshDescriptionCommitter>(Target)->CommitMeshDescription([=](const IMeshDescriptionCommitter::FCommitterParams& CommitParams)
+			DynamicMeshComponent->ProcessMesh([&](const FDynamicMesh3& ReadMesh)
 			{
 				FConversionToMeshDescriptionOptions ConversionOptions;
 				bool bModifiedTopology = (ModifiedTopologyCounter > 0);
 				ConversionOptions.bSetPolyGroups = bModifiedTopology;
-				DynamicMeshComponent->Bake(CommitParams.MeshDescriptionOut, bModifiedTopology, ConversionOptions);
+				UE::ToolTarget::CommitDynamicMeshUpdate(Target, ReadMesh, true, ConversionOptions);
 			});
 
 			// The stored selection change should go into this transaction as well.
@@ -2321,7 +2314,7 @@ bool UEditMeshPolygonsTool::BeginMeshFaceEditChangeWithPreview()
 	if (bOK)
 	{
 		EditPreview = NewObject<UPolyEditPreviewMesh>(this);
-		EditPreview->CreateInWorld(Cast<IPrimitiveComponentBackedTarget>(Target)->GetOwnerActor()->GetWorld(), FTransform::Identity);
+		EditPreview->CreateInWorld(UE::ToolTarget::GetTargetActor(Target)->GetWorld(), FTransform::Identity);
 		UpdateEditPreviewMaterials(EPreviewMaterialType::PreviewMaterial);
 		EditPreview->EnableWireframe(true);
 
diff --git a/Engine/Plugins/Experimental/MeshModelingToolset/Source/MeshModelingTools/Private/EditNormalsTool.cpp b/Engine/Plugins/Experimental/MeshModelingToolset/Source/MeshModelingTools/Private/EditNormalsTool.cpp
--- a/Engine/Plugins/Experimental/MeshModelingToolset/Source/MeshModelingTools/Private/EditNormalsTool.cpp
+++ b/Engine/Plugins/Experimental/MeshModelingToolset/Source/MeshModelingTools/Private/EditNormalsTool.cpp
@@ -155,7 +155,7 @@ void UEditNormalsTool::UpdateNumPreviews()
 
 			UMeshOpPreviewWithBackgroundCompute* Preview = Previews.Add_GetRef(NewObject<UMeshOpPreviewWithBackgroundCompute>(OpFactory, "Preview"));
 			Preview->Setup(this->TargetWorld, OpFactory);
-			Preview->PreviewMesh->SetTangentsMode(EDynamicMeshTangentCalcType::AutoCalculated);
+			Preview->PreviewMesh->SetTangentsMode(EDynamicMeshComponentTangentsMode::AutoCalculated);
 
 			FComponentMaterialSet MaterialSet;
 			TargetMaterialInterface(PreviewIdx)->GetMaterialSet(MaterialSet);
diff --git a/Engine/Plugins/Experimental/MeshModelingToolset/Source/MeshModelingTools/Private/EditUVIslandsTool.cpp b/Engine/Plugins/Experimental/MeshModelingToolset/Source/MeshModelingTools/Private/EditUVIslandsTool.cpp
--- a/Engine/Plugins/Experimental/MeshModelingToolset/Source/MeshModelingTools/Private/EditUVIslandsTool.cpp
+++ b/Engine/Plugins/Experimental/MeshModelingToolset/Source/MeshModelingTools/Private/EditUVIslandsTool.cpp
@@ -11,14 +11,10 @@
 #include "Transforms/MultiTransformer.h"
 #include "BaseBehaviors/SingleClickBehavior.h"
 #include "ToolSetupUtil.h"
+#include "ModelingToolTargetUtil.h"
 
 #include "Materials/MaterialInstanceDynamic.h"
 
-#include "TargetInterfaces/MaterialProvider.h"
-#include "TargetInterfaces/MeshDescriptionCommitter.h"
-#include "TargetInterfaces/MeshDescriptionProvider.h"
-#include "TargetInterfaces/PrimitiveComponentBackedTarget.h"
-
 #include "ExplicitUseGeometryMathTypes.h"		// using UE::Geometry::(math types)
 using namespace UE::Geometry;
 
@@ -47,16 +43,14 @@ void UEditUVIslandsTool::Setup()
 	UMeshSurfacePointTool::Setup();
 
 	// create dynamic mesh component to use for live preview
-	IPrimitiveComponentBackedTarget* TargetComponent = Cast<IPrimitiveComponentBackedTarget>(Target);
-	DynamicMeshComponent = NewObject<USimpleDynamicMeshComponent>(TargetComponent->GetOwnerActor(), "DynamicMesh");
-	DynamicMeshComponent->SetupAttachment(TargetComponent->GetOwnerActor()->GetRootComponent());
+	DynamicMeshComponent = NewObject<USimpleDynamicMeshComponent>(UE::ToolTarget::GetTargetActor(Target));
+	DynamicMeshComponent->SetupAttachment(UE::ToolTarget::GetTargetActor(Target)->GetRootComponent());
 	DynamicMeshComponent->RegisterComponent();
-	DynamicMeshComponent->SetWorldTransform(TargetComponent->GetWorldTransform());
+	DynamicMeshComponent->SetWorldTransform((FTransform)UE::ToolTarget::GetLocalToWorldTransform(Target));
 	WorldTransform = UE::Geometry::FTransform3d(DynamicMeshComponent->GetComponentTransform());
 
 	// set materials
-	FComponentMaterialSet MaterialSet;
-	Cast<IMaterialProvider>(Target)->GetMaterialSet(MaterialSet);
+	FComponentMaterialSet MaterialSet = UE::ToolTarget::GetMaterialSet(Target);
 	for (int k = 0; k < MaterialSet.Materials.Num(); ++k)
 	{
 		DynamicMeshComponent->SetMaterial(k, MaterialSet.Materials[k]);
@@ -70,8 +64,8 @@ void UEditUVIslandsTool::Setup()
 	});
 
 	// dynamic mesh configuration settings
-	DynamicMeshComponent->TangentsType = EDynamicMeshTangentCalcType::AutoCalculated;
-	DynamicMeshComponent->InitializeMesh(Cast<IMeshDescriptionProvider>(Target)->GetMeshDescription());
+	DynamicMeshComponent->SetTangentsType(EDynamicMeshComponentTangentsMode::AutoCalculated);
+	DynamicMeshComponent->SetMesh(UE::ToolTarget::GetDynamicMeshCopy(Target));
 	FMeshNormals::QuickComputeVertexNormals(*DynamicMeshComponent->GetMesh());
 	OnDynamicMeshComponentChangedHandle = DynamicMeshComponent->OnMeshChanged.Add(
 		FSimpleMulticastDelegate::FDelegate::CreateUObject(this, &UEditUVIslandsTool::OnDynamicMeshComponentChanged));
@@ -92,7 +86,7 @@ void UEditUVIslandsTool::Setup()
 	UVTranslateScale = 1.0 / DynamicMeshComponent->GetMesh()->GetBounds().MaxDim();
 
 	// hide input StaticMeshComponent
-	TargetComponent->SetOwnerVisibility(false);
+	UE::ToolTarget::HideSourceObject(Target);
 
 	// init state flags flags
 	bInDrag = false;
@@ -136,18 +130,18 @@ void UEditUVIslandsTool::Shutdown(EToolShutdownType ShutdownType)
 	{
 		DynamicMeshComponent->OnMeshChanged.Remove(OnDynamicMeshComponentChangedHandle);
 
-		Cast<IPrimitiveComponentBackedTarget>(Target)->SetOwnerVisibility(true);
+		UE::ToolTarget::ShowSourceObject(Target);
 
 		if (ShutdownType == EToolShutdownType::Accept)
 		{
 			// this block bakes the modified DynamicMeshComponent back into the StaticMeshComponent inside an undo transaction
 			GetToolManager()->BeginUndoTransaction(LOCTEXT("EditUVIslandsToolTransactionName", "Edit UVs"));
-			Cast<IMeshDescriptionCommitter>(Target)->CommitMeshDescription([=](const IMeshDescriptionCommitter::FCommitterParams& CommitParams)
+			DynamicMeshComponent->ProcessMesh([&](const FDynamicMesh3& ReadMesh)
 			{
 				FConversionToMeshDescriptionOptions ConversionOptions;
 				ConversionOptions.bUpdateNormals = ConversionOptions.bUpdatePositions = false;
 				ConversionOptions.bUpdateUVs = true;
-				DynamicMeshComponent->Bake(CommitParams.MeshDescriptionOut, false, ConversionOptions);
+				UE::ToolTarget::CommitDynamicMeshUpdate(Target, ReadMesh, true, ConversionOptions);
 			});
 			GetToolManager()->EndUndoTransaction();
 		}
diff --git a/Engine/Plugins/Experimental/MeshModelingToolset/Source/MeshModelingTools/Private/GroupEdgeInsertionTool.cpp b/Engine/Plugins/Experimental/MeshModelingToolset/Source/MeshModelingTools/Private/GroupEdgeInsertionTool.cpp
--- a/Engine/Plugins/Experimental/MeshModelingToolset/Source/MeshModelingTools/Private/GroupEdgeInsertionTool.cpp
+++ b/Engine/Plugins/Experimental/MeshModelingToolset/Source/MeshModelingTools/Private/GroupEdgeInsertionTool.cpp
@@ -133,7 +133,7 @@ void UGroupEdgeInsertionTool::SetupPreview()
 
 	Preview = NewObject<UMeshOpPreviewWithBackgroundCompute>(OpFactory);
 	Preview->Setup(TargetWorld, OpFactory);
-	Preview->PreviewMesh->SetTangentsMode(EDynamicMeshTangentCalcType::AutoCalculated);
+	Preview->PreviewMesh->SetTangentsMode(EDynamicMeshComponentTangentsMode::AutoCalculated);
 
 	FComponentMaterialSet MaterialSet;
 	Cast<IMaterialProvider>(Target)->GetMaterialSet(MaterialSet);
diff --git a/Engine/Plugins/Experimental/MeshModelingToolset/Source/MeshModelingTools/Private/LatticeDeformerTool.cpp b/Engine/Plugins/Experimental/MeshModelingToolset/Source/MeshModelingTools/Private/LatticeDeformerTool.cpp
--- a/Engine/Plugins/Experimental/MeshModelingToolset/Source/MeshModelingTools/Private/LatticeDeformerTool.cpp
+++ b/Engine/Plugins/Experimental/MeshModelingToolset/Source/MeshModelingTools/Private/LatticeDeformerTool.cpp
@@ -212,7 +212,7 @@ void ULatticeDeformerTool::StartPreview()
 		Preview->PreviewMesh->SetSecondaryRenderMaterial(SelectionMaterial);
 	}
 
-	Preview->PreviewMesh->SetTangentsMode(EDynamicMeshTangentCalcType::NoTangents);
+	Preview->PreviewMesh->SetTangentsMode(EDynamicMeshComponentTangentsMode::NoTangents);
 	Preview->SetVisibility(true);
 	Preview->InvalidateResult();
 
diff --git a/Engine/Plugins/Experimental/MeshModelingToolset/Source/MeshModelingTools/Private/MeshAttributePaintTool.cpp b/Engine/Plugins/Experimental/MeshModelingToolset/Source/MeshModelingTools/Private/MeshAttributePaintTool.cpp
--- a/Engine/Plugins/Experimental/MeshModelingToolset/Source/MeshModelingTools/Private/MeshAttributePaintTool.cpp
+++ b/Engine/Plugins/Experimental/MeshModelingToolset/Source/MeshModelingTools/Private/MeshAttributePaintTool.cpp
@@ -163,7 +163,7 @@ void UMeshAttributePaintTool::Setup()
 	//AddToolPropertySource(AttributeEditActions);
 
 	// configure preview mesh
-	PreviewMesh->SetTangentsMode(EDynamicMeshTangentCalcType::AutoCalculated);
+	PreviewMesh->SetTangentsMode(EDynamicMeshComponentTangentsMode::AutoCalculated);
 	//PreviewMesh->EnableWireframe(SelectionProps->bShowWireframe);
 	PreviewMesh->GetRootComponent()->bCastDynamicShadow = false;
 
diff --git a/Engine/Plugins/Experimental/MeshModelingToolset/Source/MeshModelingTools/Private/MeshGroupPaintTool.cpp b/Engine/Plugins/Experimental/MeshModelingToolset/Source/MeshModelingTools/Private/MeshGroupPaintTool.cpp
--- a/Engine/Plugins/Experimental/MeshModelingToolset/Source/MeshModelingTools/Private/MeshGroupPaintTool.cpp
+++ b/Engine/Plugins/Experimental/MeshModelingToolset/Source/MeshModelingTools/Private/MeshGroupPaintTool.cpp
@@ -6,7 +6,8 @@
 #include "Drawing/MeshElementsVisualizer.h"
 #include "Async/ParallelFor.h"
 #include "Async/Async.h"
-
+#include "ToolSetupUtil.h"
+#include "ModelingToolTargetUtil.h"
 #include "MeshWeights.h"
 #include "DynamicMesh/MeshNormals.h"
 #include "DynamicMesh/MeshIndexUtil.h"
@@ -28,15 +29,9 @@
 #include "Sculpting/StampFalloffs.h"
 #include "Sculpting/MeshSculptUtil.h"
 
-
 #include "CanvasTypes.h"
 #include "CanvasItem.h"
 
-#include "ToolSetupUtil.h"
-
-#include "TargetInterfaces/MaterialProvider.h"
-#include "TargetInterfaces/MeshDescriptionCommitter.h"
-#include "TargetInterfaces/PrimitiveComponentBackedTarget.h"
 
 #include "ExplicitUseGeometryMathTypes.h"		// using UE::Geometry::(math types)
 using namespace UE::Geometry;
@@ -91,12 +86,11 @@ void UMeshGroupPaintTool::Setup()
 	SetToolDisplayName(LOCTEXT("ToolName", "Paint PolyGroups"));
 
 	// create dynamic mesh component to use for live preview
-	DynamicMeshComponent = NewObject<USimpleDynamicMeshComponent>(Cast<IPrimitiveComponentBackedTarget>(Target)->GetOwnerActor());
+	DynamicMeshComponent = NewObject<USimpleDynamicMeshComponent>(UE::ToolTarget::GetTargetActor(Target));
 	InitializeSculptMeshComponent(DynamicMeshComponent);
 
 	// assign materials
-	FComponentMaterialSet MaterialSet;
-	Cast<IMaterialProvider>(Target)->GetMaterialSet(MaterialSet);
+	FComponentMaterialSet MaterialSet = UE::ToolTarget::GetMaterialSet(Target);
 	for (int k = 0; k < MaterialSet.Materials.Num(); ++k)
 	{
 		DynamicMeshComponent->SetMaterial(k, MaterialSet.Materials[k]);
@@ -267,10 +261,9 @@ void UMeshGroupPaintTool::Shutdown(EToolShutdownType ShutdownType)
 	if (ShutdownType == EToolShutdownType::Accept)
 	{
 		GetToolManager()->BeginUndoTransaction(LOCTEXT("GroupPaintMeshToolTransactionName", "Paint Groups"));
-		Cast<IMeshDescriptionCommitter>(Target)->CommitMeshDescription([=](const IMeshDescriptionCommitter::FCommitterParams& CommitParams)
+		DynamicMeshComponent->ProcessMesh([&](const FDynamicMesh3& ReadMesh)
 		{
-			FConversionToMeshDescriptionOptions ConversionOptions;
-			DynamicMeshComponent->Bake(CommitParams.MeshDescriptionOut, true, ConversionOptions);
+			UE::ToolTarget::CommitDynamicMeshUpdate(Target, ReadMesh, true);
 		});
 		GetToolManager()->EndUndoTransaction();
 	}
@@ -882,7 +875,8 @@ void UMeshGroupPaintTool::ApplyVisibilityFilter(const TArray<int32>& Triangles,
 
 	FViewCameraState StateOut;
 	GetToolManager()->GetContextQueriesAPI()->GetCurrentViewState(StateOut);
-	FVector3d LocalEyePosition(Cast<IPrimitiveComponentBackedTarget>(Target)->GetWorldTransform().InverseTransformPosition(StateOut.Position));
+	FTransform3d LocalToWorld = UE::ToolTarget::GetLocalToWorldTransform(Target);
+	FVector3d LocalEyePosition(LocalToWorld.InverseTransformPosition(StateOut.Position));
 
 	const FDynamicMesh3* Mesh = GetSculptMesh();
 
diff --git a/Engine/Plugins/Experimental/MeshModelingToolset/Source/MeshModelingTools/Private/MeshInspectorTool.cpp b/Engine/Plugins/Experimental/MeshModelingToolset/Source/MeshModelingTools/Private/MeshInspectorTool.cpp
--- a/Engine/Plugins/Experimental/MeshModelingToolset/Source/MeshModelingTools/Private/MeshInspectorTool.cpp
+++ b/Engine/Plugins/Experimental/MeshModelingToolset/Source/MeshModelingTools/Private/MeshInspectorTool.cpp
@@ -21,10 +21,7 @@
 
 #include "ToolSetupUtil.h"
 #include "AssetUtils/MeshDescriptionUtil.h"
-
-#include "TargetInterfaces/MaterialProvider.h"
-#include "TargetInterfaces/MeshDescriptionProvider.h"
-#include "TargetInterfaces/PrimitiveComponentBackedTarget.h"
+#include "ModelingToolTargetUtil.h"
 
 #include "ExplicitUseGeometryMathTypes.h"		// using UE::Geometry::(math types)
 using namespace UE::Geometry;
@@ -52,23 +49,18 @@ void UMeshInspectorTool::Setup()
 {
 	UInteractiveTool::Setup();
 
-	// copy input MeshDescription and make sure it has initialized normals/tangents
-	IPrimitiveComponentBackedTarget* TargetComponent = Cast<IPrimitiveComponentBackedTarget>(Target);
-	InputMeshDescription = MakeShared<FMeshDescription>(*Cast<IMeshDescriptionProvider>(Target)->GetMeshDescription());
-	UE::MeshDescription::InitializeAutoGeneratedAttributes(*InputMeshDescription, TargetComponent->GetOwnerComponent(), 0);
-
 	PreviewMesh = NewObject<UPreviewMesh>(this);
 	PreviewMesh->bBuildSpatialDataStructure = false;
-	PreviewMesh->CreateInWorld(TargetComponent->GetOwnerActor()->GetWorld(), FTransform::Identity);
-	PreviewMesh->SetTransform(TargetComponent->GetWorldTransform());
+	PreviewMesh->CreateInWorld(UE::ToolTarget::GetTargetActor(Target)->GetWorld(), FTransform::Identity);
+	PreviewMesh->SetTransform((FTransform)UE::ToolTarget::GetLocalToWorldTransform(Target));
 
-	FComponentMaterialSet MaterialSet;
-	Cast<IMaterialProvider>(Target)->GetMaterialSet(MaterialSet);
+	FComponentMaterialSet MaterialSet = UE::ToolTarget::GetMaterialSet(Target);
 	PreviewMesh->SetMaterials(MaterialSet.Materials);
 	DefaultMaterial = PreviewMesh->GetMaterial(0);
 
-	PreviewMesh->SetTangentsMode(EDynamicMeshTangentCalcType::ExternallyCalculated);
-	PreviewMesh->InitializeMesh(InputMeshDescription.Get());
+	PreviewMesh->SetTangentsMode(EDynamicMeshComponentTangentsMode::ExternallyProvided);
+	FDynamicMesh3 InputMeshWithTangents = UE::ToolTarget::GetDynamicMeshCopy(Target, true);
+	PreviewMesh->ReplaceMesh(MoveTemp(InputMeshWithTangents));
 
 	DrawnLineSet = NewObject<ULineSetComponent>(PreviewMesh->GetRootComponent(), "MeshInspectorToolLineSet");
 
@@ -79,8 +71,7 @@ void UMeshInspectorTool::Setup()
 
 	Precompute();
 
-	// hide input StaticMeshComponent
-	TargetComponent->SetOwnerVisibility(false);
+	UE::ToolTarget::HideSourceObject(Target);
 
 	// initialize our properties
 	Settings = NewObject<UMeshInspectorProperties>(this);
@@ -98,7 +89,7 @@ void UMeshInspectorTool::Setup()
 	AddToolPropertySource(Statistics);
 
 	UMeshAnalysisProperties* MeshAnalysis = NewObject<UMeshAnalysisProperties>(this);
-	MeshAnalysis->Update(*PreviewMesh->GetPreviewDynamicMesh(), TargetComponent->GetWorldTransform());
+	MeshAnalysis->Update(*PreviewMesh->GetPreviewDynamicMesh(), (FTransform)UE::ToolTarget::GetLocalToWorldTransform(Target));
 	AddToolPropertySource(MeshAnalysis);
 
 	UpdateVisualization();
@@ -112,7 +103,7 @@ void UMeshInspectorTool::Setup()
 
 void UMeshInspectorTool::Shutdown(EToolShutdownType ShutdownType)
 {
-	Cast<IPrimitiveComponentBackedTarget>(Target)->SetOwnerVisibility(true);
+	UE::ToolTarget::ShowSourceObject(Target);
 
 	if (PreviewMesh != nullptr)
 	{
@@ -320,12 +311,12 @@ void UMeshInspectorTool::UpdateVisualization()
 		}
 	}
 
-	if (Settings->bTangentVectors && PreviewMesh->GetTangents() != nullptr)
+	if (Settings->bTangentVectors && PreviewMesh->GetMesh()->Attributes()->HasTangentSpace() )
 	{
 		// See note in normals about scaling
 		FVector3f TangentScaling = Settings->TangentLength * InverseScale;
 
-		const FMeshTangentsf* Tangents = PreviewMesh->GetTangents();
+		FDynamicMeshTangents Tangents(PreviewMesh->GetMesh());
 		for (int TID : TargetMesh->TriangleIndicesItr())
 		{
 			FVector3d TriV[3];
@@ -334,7 +325,7 @@ void UMeshInspectorTool::UpdateVisualization()
 			{
 				FVector3f Vert(TriV[SubIdx]);
 				FVector3f Tangent, Bitangent;
-				Tangents->GetPerTriangleTangent(TID, SubIdx, Tangent, Bitangent);
+				Tangents.GetTangentVectors(TID, SubIdx, Tangent, Bitangent);
 				DrawnLineSet->AddLine((FVector)Vert, (FVector)(Vert + TangentScaling * Tangent),
 					TangentColor, TangentThickness, TangentDepthBias);
 				DrawnLineSet->AddLine((FVector)Vert, (FVector)(Vert + TangentScaling * Bitangent),
diff --git a/Engine/Plugins/Experimental/MeshModelingToolset/Source/MeshModelingTools/Private/MeshSelectionTool.cpp b/Engine/Plugins/Experimental/MeshModelingToolset/Source/MeshModelingTools/Private/MeshSelectionTool.cpp
--- a/Engine/Plugins/Experimental/MeshModelingToolset/Source/MeshModelingTools/Private/MeshSelectionTool.cpp
+++ b/Engine/Plugins/Experimental/MeshModelingToolset/Source/MeshModelingTools/Private/MeshSelectionTool.cpp
@@ -16,13 +16,10 @@
 #include "ToolSetupUtil.h"
 #include "Selections/MeshConnectedComponents.h"
 #include "Selections/MeshFaceSelection.h"
+#include "ModelingToolTargetUtil.h"
 
 #include "Algo/MaxElement.h"
 
-#include "TargetInterfaces/MaterialProvider.h"
-#include "TargetInterfaces/MeshDescriptionCommitter.h"
-#include "TargetInterfaces/PrimitiveComponentBackedTarget.h"
-
 #include "ExplicitUseGeometryMathTypes.h"		// using UE::Geometry::(math types)
 using namespace UE::Geometry;
 
@@ -83,7 +80,7 @@ void UMeshSelectionTool::Setup()
 	AddToolPropertySource(EditActions);
 
 	// set autocalculated tangents
-	PreviewMesh->SetTangentsMode(EDynamicMeshTangentCalcType::AutoCalculated);
+	PreviewMesh->SetTangentsMode(EDynamicMeshComponentTangentsMode::AutoCalculated);
 
 	// enable wireframe on component
 	PreviewMesh->EnableWireframe(SelectionProps->bShowWireframe);
@@ -164,11 +161,7 @@ void UMeshSelectionTool::OnShutdown(EToolShutdownType ShutdownType)
 	{
 		// this block bakes the modified DynamicMeshComponent back into the StaticMeshComponent inside an undo transaction
 		GetToolManager()->BeginUndoTransaction(LOCTEXT("MeshSelectionToolTransactionName", "Edit Mesh"));
-
-		Cast<IMeshDescriptionCommitter>(Target)->CommitMeshDescription([=](const IMeshDescriptionCommitter::FCommitterParams& CommitParams)
-		{
-			PreviewMesh->Bake(CommitParams.MeshDescriptionOut, true);
-		});
+		UE::ToolTarget::CommitDynamicMeshUpdate(Target, *PreviewMesh->GetMesh(), true);
 		GetToolManager()->EndUndoTransaction();
 	}
 	else if (ShutdownType == EToolShutdownType::Cancel)
@@ -323,7 +316,8 @@ bool UMeshSelectionTool::HitTest(const FRay& Ray, FHitResult& OutHit)
 		SourceMesh->GetTriInfo(OutHit.FaceIndex, Normal, Area, Centroid);
 		FViewCameraState StateOut;
 		GetToolManager()->GetContextQueriesAPI()->GetCurrentViewState(StateOut);
-		FVector3d LocalEyePosition(Cast<IPrimitiveComponentBackedTarget>(Target)->GetWorldTransform().InverseTransformPosition(StateOut.Position));
+		FTransform3d LocalToWorld = UE::ToolTarget::GetLocalToWorldTransform(Target);
+		FVector3d LocalEyePosition(LocalToWorld.InverseTransformPosition(StateOut.Position));
 
 		if (Normal.Dot((Centroid - LocalEyePosition)) > 0)
 		{
@@ -378,7 +372,7 @@ TUniquePtr<FDynamicMeshOctree3>& UMeshSelectionTool::GetOctree()
 
 void UMeshSelectionTool::CalculateVertexROI(const FBrushStampData& Stamp, TArray<int>& VertexROI)
 {
-	FTransform3d Transform(Cast<IPrimitiveComponentBackedTarget>(Target)->GetWorldTransform());
+	FTransform3d Transform = UE::ToolTarget::GetLocalToWorldTransform(Target);
 	FVector3d StampPosLocal = Transform.InverseTransformPosition((FVector3d)Stamp.WorldPosition);
 
 	// TODO: need dynamic vertex hash table!
@@ -400,7 +394,7 @@ void UMeshSelectionTool::CalculateVertexROI(const FBrushStampData& Stamp, TArray
 
 void UMeshSelectionTool::CalculateTriangleROI(const FBrushStampData& Stamp, TArray<int>& TriangleROI)
 {
-	FTransform3d Transform(Cast<IPrimitiveComponentBackedTarget>(Target)->GetWorldTransform());
+	FTransform3d Transform = UE::ToolTarget::GetLocalToWorldTransform(Target);
 	FVector3d StampPosLocal = Transform.InverseTransformPosition((FVector3d)Stamp.WorldPosition);
 
 	// always select first triangle
@@ -553,7 +547,8 @@ void UMeshSelectionTool::UpdateFaceSelection(const FBrushStampData& Stamp, const
 	{
 		FViewCameraState StateOut;
 		GetToolManager()->GetContextQueriesAPI()->GetCurrentViewState(StateOut);
-		FVector3d LocalEyePosition(Cast<IPrimitiveComponentBackedTarget>(Target)->GetWorldTransform().InverseTransformPosition(StateOut.Position));
+		FTransform3d LocalToWorld = UE::ToolTarget::GetLocalToWorldTransform(Target);
+		FVector3d LocalEyePosition(LocalToWorld.InverseTransformPosition(StateOut.Position));
 
 		for (int tid : TriangleROI)
 		{
@@ -784,7 +779,7 @@ void UMeshSelectionTool::Render(IToolsContextRenderAPI* RenderAPI)
 {
 	UDynamicMeshBrushTool::Render(RenderAPI);
 
-	FTransform WorldTransform = Cast<IPrimitiveComponentBackedTarget>(Target)->GetWorldTransform();
+	FTransform WorldTransform = (FTransform)UE::ToolTarget::GetLocalToWorldTransform(Target);
 	const FDynamicMesh3* Mesh = PreviewMesh->GetMesh();
 
 	if (SelectionProps->bShowPoints)
@@ -1380,12 +1375,8 @@ void UMeshSelectionTool::SeparateSelectedTriangles()
 	GetToolManager()->BeginUndoTransaction(LOCTEXT("MeshSelectionToolSeparate", "Separate"));
 
 	// build array of materials from the original
-	TArray<UMaterialInterface*> Materials;
-	IMaterialProvider* TargetMaterial = Cast<IMaterialProvider>(Target);
-	for (int MaterialIdx = 0, NumMaterials = TargetMaterial->GetNumMaterials(); MaterialIdx < NumMaterials; MaterialIdx++)
-	{
-		Materials.Add(TargetMaterial->GetMaterial(MaterialIdx));
-	}
+	FComponentMaterialSet MaterialSet = UE::ToolTarget::GetMaterialSet(Target);
+	TArray<UMaterialInterface*> Materials = MaterialSet.Materials;
 
 	FCreateMeshObjectParams NewMeshObjectParams;
 	NewMeshObjectParams.TargetWorld = TargetWorld;
diff --git a/Engine/Plugins/Experimental/MeshModelingToolset/Source/MeshModelingTools/Private/MeshSpaceDeformerTool.cpp b/Engine/Plugins/Experimental/MeshModelingToolset/Source/MeshModelingTools/Private/MeshSpaceDeformerTool.cpp
--- a/Engine/Plugins/Experimental/MeshModelingToolset/Source/MeshModelingTools/Private/MeshSpaceDeformerTool.cpp
+++ b/Engine/Plugins/Experimental/MeshModelingToolset/Source/MeshModelingTools/Private/MeshSpaceDeformerTool.cpp
@@ -268,7 +268,7 @@ void UMeshSpaceDeformerTool::Setup()
 
 		Preview = NewObject<UMeshOpPreviewWithBackgroundCompute>(DeformerOperatorFactory, "Preview");
 		Preview->Setup(this->TargetWorld, DeformerOperatorFactory);
-		Preview->PreviewMesh->SetTangentsMode(EDynamicMeshTangentCalcType::AutoCalculated);
+		Preview->PreviewMesh->SetTangentsMode(EDynamicMeshComponentTangentsMode::AutoCalculated);
 
 		Preview->SetIsMeshTopologyConstant(true, EMeshRenderAttributeFlags::Positions | EMeshRenderAttributeFlags::VertexNormals);
 
diff --git a/Engine/Plugins/Experimental/MeshModelingToolset/Source/MeshModelingTools/Private/MirrorTool.cpp b/Engine/Plugins/Experimental/MeshModelingToolset/Source/MeshModelingTools/Private/MirrorTool.cpp
--- a/Engine/Plugins/Experimental/MeshModelingToolset/Source/MeshModelingTools/Private/MirrorTool.cpp
+++ b/Engine/Plugins/Experimental/MeshModelingToolset/Source/MeshModelingTools/Private/MirrorTool.cpp
@@ -303,7 +303,7 @@ void UMirrorTool::SetupPreviews()
 		UMeshOpPreviewWithBackgroundCompute* Preview = Previews.Add_GetRef(
 			NewObject<UMeshOpPreviewWithBackgroundCompute>(MirrorOpCreator, "Preview"));
 		Preview->Setup(TargetWorld, MirrorOpCreator);
-		Preview->PreviewMesh->SetTangentsMode(EDynamicMeshTangentCalcType::AutoCalculated);
+		Preview->PreviewMesh->SetTangentsMode(EDynamicMeshComponentTangentsMode::AutoCalculated);
 
 		FComponentMaterialSet MaterialSet;
 		TargetMaterialInterface(PreviewIndex)->GetMaterialSet(MaterialSet);
diff --git a/Engine/Plugins/Experimental/MeshModelingToolset/Source/MeshModelingTools/Private/PlaneCutTool.cpp b/Engine/Plugins/Experimental/MeshModelingToolset/Source/MeshModelingTools/Private/PlaneCutTool.cpp
--- a/Engine/Plugins/Experimental/MeshModelingToolset/Source/MeshModelingTools/Private/PlaneCutTool.cpp
+++ b/Engine/Plugins/Experimental/MeshModelingToolset/Source/MeshModelingTools/Private/PlaneCutTool.cpp
@@ -218,7 +218,7 @@ void UPlaneCutTool::SetupPreviews()
 		CutSide->ComponentIndex = PreviewIdx;
 		UMeshOpPreviewWithBackgroundCompute* Preview = Previews.Add_GetRef(NewObject<UMeshOpPreviewWithBackgroundCompute>(CutSide, "Preview"));
 		Preview->Setup(this->TargetWorld, CutSide);
-		Preview->PreviewMesh->SetTangentsMode(EDynamicMeshTangentCalcType::AutoCalculated);
+		Preview->PreviewMesh->SetTangentsMode(EDynamicMeshComponentTangentsMode::AutoCalculated);
 
 		FComponentMaterialSet MaterialSet;
 		TargetMaterialInterface(PreviewIdx)->GetMaterialSet(MaterialSet);
diff --git a/Engine/Plugins/Experimental/MeshModelingToolset/Source/MeshModelingTools/Private/RemeshMeshTool.cpp b/Engine/Plugins/Experimental/MeshModelingToolset/Source/MeshModelingTools/Private/RemeshMeshTool.cpp
--- a/Engine/Plugins/Experimental/MeshModelingToolset/Source/MeshModelingTools/Private/RemeshMeshTool.cpp
+++ b/Engine/Plugins/Experimental/MeshModelingToolset/Source/MeshModelingTools/Private/RemeshMeshTool.cpp
@@ -133,7 +133,7 @@ void URemeshMeshTool::Setup()
 	Converter.Convert(TargetMeshProvider->GetMeshDescription(), *OriginalMesh);
 
 	Preview->PreviewMesh->SetTransform(TargetComponent->GetWorldTransform());
-	Preview->PreviewMesh->SetTangentsMode(EDynamicMeshTangentCalcType::AutoCalculated);
+	Preview->PreviewMesh->SetTangentsMode(EDynamicMeshComponentTangentsMode::AutoCalculated);
 	Preview->PreviewMesh->UpdatePreview(OriginalMesh.Get());
 
 	OriginalMeshSpatial = MakeShared<FDynamicMeshAABBTree3, ESPMode::ThreadSafe>(OriginalMesh.Get(), true);
diff --git a/Engine/Plugins/Experimental/MeshModelingToolset/Source/MeshModelingTools/Private/RemoveOccludedTrianglesTool.cpp b/Engine/Plugins/Experimental/MeshModelingToolset/Source/MeshModelingTools/Private/RemoveOccludedTrianglesTool.cpp
--- a/Engine/Plugins/Experimental/MeshModelingToolset/Source/MeshModelingTools/Private/RemoveOccludedTrianglesTool.cpp
+++ b/Engine/Plugins/Experimental/MeshModelingToolset/Source/MeshModelingTools/Private/RemoveOccludedTrianglesTool.cpp
@@ -258,7 +258,7 @@ void URemoveOccludedTrianglesTool::SetupPreviews()
 
 			UMeshOpPreviewWithBackgroundCompute* Preview = Previews.Add_GetRef(NewObject<UMeshOpPreviewWithBackgroundCompute>(OpFactory, "Preview"));
 			Preview->Setup(this->TargetWorld, OpFactory);
-			Preview->PreviewMesh->SetTangentsMode(EDynamicMeshTangentCalcType::AutoCalculated);
+			Preview->PreviewMesh->SetTangentsMode(EDynamicMeshComponentTangentsMode::AutoCalculated);
 
 			FComponentMaterialSet MaterialSet;
 			TargetMaterialInterface(TargetIdx)->GetMaterialSet(MaterialSet);
diff --git a/Engine/Plugins/Experimental/MeshModelingToolset/Source/MeshModelingTools/Private/RevolveBoundaryTool.cpp b/Engine/Plugins/Experimental/MeshModelingToolset/Source/MeshModelingTools/Private/RevolveBoundaryTool.cpp
--- a/Engine/Plugins/Experimental/MeshModelingToolset/Source/MeshModelingTools/Private/RevolveBoundaryTool.cpp
+++ b/Engine/Plugins/Experimental/MeshModelingToolset/Source/MeshModelingTools/Private/RevolveBoundaryTool.cpp
@@ -231,7 +231,7 @@ void URevolveBoundaryTool::StartPreview()
 
 	Preview = NewObject<UMeshOpPreviewWithBackgroundCompute>(RevolveBoundaryOpCreator);
 	Preview->Setup(TargetWorld, RevolveBoundaryOpCreator);
-	Preview->PreviewMesh->SetTangentsMode(EDynamicMeshTangentCalcType::AutoCalculated);
+	Preview->PreviewMesh->SetTangentsMode(EDynamicMeshComponentTangentsMode::AutoCalculated);
 
 	Preview->ConfigureMaterials(MaterialProperties->Material.Get(),
 		ToolSetupUtil::GetDefaultWorkingMaterial(GetToolManager()));
diff --git a/Engine/Plugins/Experimental/MeshModelingToolset/Source/MeshModelingTools/Private/Sculpting/MeshSculptToolBase.cpp b/Engine/Plugins/Experimental/MeshModelingToolset/Source/MeshModelingTools/Private/Sculpting/MeshSculptToolBase.cpp
--- a/Engine/Plugins/Experimental/MeshModelingToolset/Source/MeshModelingTools/Private/Sculpting/MeshSculptToolBase.cpp
+++ b/Engine/Plugins/Experimental/MeshModelingToolset/Source/MeshModelingTools/Private/Sculpting/MeshSculptToolBase.cpp
@@ -15,9 +15,7 @@
 
 #include "Generators/SphereGenerator.h"
 
-#include "TargetInterfaces/MeshDescriptionCommitter.h"
-#include "TargetInterfaces/MeshDescriptionProvider.h"
-#include "TargetInterfaces/PrimitiveComponentBackedTarget.h"
+#include "ModelingToolTargetUtil.h"
 
 #include "ExplicitUseGeometryMathTypes.h"		// using UE::Geometry::(math types)
 using namespace UE::Geometry;
@@ -141,7 +139,7 @@ void UMeshSculptToolBase::Shutdown(EToolShutdownType ShutdownType)
 	UBaseDynamicMeshComponent* DynamicMeshComponent = GetSculptMeshComponent();
 	if (DynamicMeshComponent != nullptr)
 	{
-		Cast<IPrimitiveComponentBackedTarget>(Target)->SetOwnerVisibility(true);
+		UE::ToolTarget::ShowSourceObject(Target);
 
 		if (ShutdownType == EToolShutdownType::Accept)
 		{
@@ -150,7 +148,7 @@ void UMeshSculptToolBase::Shutdown(EToolShutdownType ShutdownType)
 
 			// this block bakes the modified DynamicMeshComponent back into the StaticMeshComponent inside an undo transaction
 			GetToolManager()->BeginUndoTransaction(LOCTEXT("SculptMeshToolTransactionName", "Sculpt Mesh"));
-			CommitResult(DynamicMeshComponent);
+			CommitResult(DynamicMeshComponent, false);
 			GetToolManager()->EndUndoTransaction();
 		}
 
@@ -163,12 +161,11 @@ void UMeshSculptToolBase::Shutdown(EToolShutdownType ShutdownType)
 
 
 
-void UMeshSculptToolBase::CommitResult(UBaseDynamicMeshComponent* Component)
+void UMeshSculptToolBase::CommitResult(UBaseDynamicMeshComponent* Component, bool bModifiedTopology)
 {
-	Cast<IMeshDescriptionCommitter>(Target)->CommitMeshDescription([=](const IMeshDescriptionCommitter::FCommitterParams& CommitParams)
+	Component->ProcessMesh([&](const FDynamicMesh3& CurMesh)
 	{
-		FConversionToMeshDescriptionOptions ConversionOptions;
-		Component->Bake(CommitParams.MeshDescriptionOut, false, ConversionOptions);
+		UE::ToolTarget::CommitDynamicMeshUpdate(Target, CurMesh, bModifiedTopology);
 	});
 }
 
@@ -238,17 +235,16 @@ void UMeshSculptToolBase::Render(IToolsContextRenderAPI* RenderAPI)
 
 void UMeshSculptToolBase::InitializeSculptMeshComponent(UBaseDynamicMeshComponent* Component)
 {
-	IPrimitiveComponentBackedTarget* TargetComponent = Cast<IPrimitiveComponentBackedTarget>(Target);
-	Component->SetupAttachment(TargetComponent->GetOwnerActor()->GetRootComponent());
+	Component->SetupAttachment(UE::ToolTarget::GetTargetActor(Target)->GetRootComponent());
 	Component->RegisterComponent();
 
 	// initialize from LOD-0 MeshDescription
-	Component->InitializeMesh(Cast<IMeshDescriptionProvider>(Target)->GetMeshDescription());
+	Component->SetMesh(UE::ToolTarget::GetDynamicMeshCopy(Target));
 	double MaxDimension = Component->GetMesh()->GetCachedBounds().MaxDim();
 
 	// bake rotation and scaling into mesh because handling these inside sculpting is a mess
 	// Note: this transform does not include translation ( so only the 3x3 transform)
-	InitialTargetTransform = FTransform3d(TargetComponent->GetWorldTransform());
+	InitialTargetTransform = UE::ToolTarget::GetLocalToWorldTransform(Target);
 	// clamp scaling because if we allow zero-scale we cannot invert this transform on Accept
 	InitialTargetTransform.ClampMinimumScale(0.01);
 	FVector3d Translation = InitialTargetTransform.GetTranslation();
@@ -258,7 +254,7 @@ void UMeshSculptToolBase::InitializeSculptMeshComponent(UBaseDynamicMeshComponen
 	Component->SetWorldTransform((FTransform)CurTargetTransform);
 
 	// hide input Component
-	TargetComponent->SetOwnerVisibility(false);
+	UE::ToolTarget::HideSourceObject(Target);
 }
 
 
@@ -961,7 +957,7 @@ void UMeshSculptToolBase::UpdateMaterialMode(EMeshEditingMaterialModes MaterialM
 	if (MaterialMode == EMeshEditingMaterialModes::ExistingMaterial)
 	{
 		GetSculptMeshComponent()->ClearOverrideRenderMaterial();
-		GetSculptMeshComponent()->bCastDynamicShadow = Cast<IPrimitiveComponentBackedTarget>(Target)->GetOwnerComponent()->bCastDynamicShadow;
+		GetSculptMeshComponent()->bCastDynamicShadow = UE::ToolTarget::GetTargetComponent(Target)->bCastDynamicShadow;
 		ActiveOverrideMaterial = nullptr;
 	}
 	else
diff --git a/Engine/Plugins/Experimental/MeshModelingToolset/Source/MeshModelingTools/Private/SkinWeightsBindingTool.cpp b/Engine/Plugins/Experimental/MeshModelingToolset/Source/MeshModelingTools/Private/SkinWeightsBindingTool.cpp
--- a/Engine/Plugins/Experimental/MeshModelingToolset/Source/MeshModelingTools/Private/SkinWeightsBindingTool.cpp
+++ b/Engine/Plugins/Experimental/MeshModelingToolset/Source/MeshModelingTools/Private/SkinWeightsBindingTool.cpp
@@ -857,7 +857,7 @@ void USkinWeightsBindingTool::Setup()
 	OriginalMesh->Attributes()->PrimaryColors()->CreateFromPredicate([](int ParentVID, int TriIDA, int TriIDB){return true;}, 0.f);
 
 	Preview->PreviewMesh->SetTransform(TargetComponent->GetWorldTransform());
-	Preview->PreviewMesh->SetTangentsMode(EDynamicMeshTangentCalcType::AutoCalculated);
+	Preview->PreviewMesh->SetTangentsMode(EDynamicMeshComponentTangentsMode::AutoCalculated);
 	Preview->PreviewMesh->GetRootComponent()->bCastDynamicShadow = false;
 	Preview->PreviewMesh->UpdatePreview(OriginalMesh.Get());
 
diff --git a/Engine/Plugins/Experimental/MeshModelingToolset/Source/MeshModelingTools/Private/SkinWeightsPaintTool.cpp b/Engine/Plugins/Experimental/MeshModelingToolset/Source/MeshModelingTools/Private/SkinWeightsPaintTool.cpp
--- a/Engine/Plugins/Experimental/MeshModelingToolset/Source/MeshModelingTools/Private/SkinWeightsPaintTool.cpp
+++ b/Engine/Plugins/Experimental/MeshModelingToolset/Source/MeshModelingTools/Private/SkinWeightsPaintTool.cpp
@@ -98,7 +98,7 @@ void USkinWeightsPaintTool::Setup()
 	UpdateBonePositionInfos(MaxDrawRadius);
 
 	// configure preview mesh
-	PreviewMesh->SetTangentsMode(EDynamicMeshTangentCalcType::AutoCalculated);
+	PreviewMesh->SetTangentsMode(EDynamicMeshComponentTangentsMode::AutoCalculated);
 	//PreviewMesh->EnableWireframe(SelectionProps->bShowWireframe);
 	PreviewMesh->GetRootComponent()->bCastDynamicShadow = false;
 
diff --git a/Engine/Plugins/Experimental/MeshModelingToolset/Source/MeshModelingTools/Private/UVProjectionTool.cpp b/Engine/Plugins/Experimental/MeshModelingToolset/Source/MeshModelingTools/Private/UVProjectionTool.cpp
--- a/Engine/Plugins/Experimental/MeshModelingToolset/Source/MeshModelingTools/Private/UVProjectionTool.cpp
+++ b/Engine/Plugins/Experimental/MeshModelingToolset/Source/MeshModelingTools/Private/UVProjectionTool.cpp
@@ -165,7 +165,7 @@ void UUVProjectionTool::UpdateNumPreviews()
 	Preview = NewObject<UMeshOpPreviewWithBackgroundCompute>(OperatorFactory);
 	Preview->Setup(this->TargetWorld, OperatorFactory);
 	Preview->OnMeshUpdated.AddUObject(this, &UUVProjectionTool::OnMeshUpdated);
-	Preview->PreviewMesh->SetTangentsMode(EDynamicMeshTangentCalcType::AutoCalculated);
+	Preview->PreviewMesh->SetTangentsMode(EDynamicMeshComponentTangentsMode::AutoCalculated);
 
 	FComponentMaterialSet MaterialSet = UE::ToolTarget::GetMaterialSet(Target, false);
 	Preview->ConfigureMaterials(MaterialSet.Materials,
diff --git a/Engine/Plugins/Experimental/MeshModelingToolset/Source/MeshModelingTools/Private/WeldMeshEdgesTool.cpp b/Engine/Plugins/Experimental/MeshModelingToolset/Source/MeshModelingTools/Private/WeldMeshEdgesTool.cpp
--- a/Engine/Plugins/Experimental/MeshModelingToolset/Source/MeshModelingTools/Private/WeldMeshEdgesTool.cpp
+++ b/Engine/Plugins/Experimental/MeshModelingToolset/Source/MeshModelingTools/Private/WeldMeshEdgesTool.cpp
@@ -9,15 +9,11 @@
 #include "DynamicMesh/Operations/MergeCoincidentMeshEdges.h"
 
 #include "SimpleDynamicMeshComponent.h"
+#include "ModelingToolTargetUtil.h"
 
 #include "SceneManagement.h" // for FPrimitiveDrawInterface
 #include "Async/ParallelFor.h"
 
-#include "TargetInterfaces/MaterialProvider.h"
-#include "TargetInterfaces/MeshDescriptionCommitter.h"
-#include "TargetInterfaces/MeshDescriptionProvider.h"
-#include "TargetInterfaces/PrimitiveComponentBackedTarget.h"
-
 #include "ExplicitUseGeometryMathTypes.h"		// using UE::Geometry::(math types)
 using namespace UE::Geometry;
 
@@ -49,26 +45,24 @@ void UWeldMeshEdgesTool::Setup()
 	UInteractiveTool::Setup();
 
 	// create dynamic mesh component to use for live preview
-	IPrimitiveComponentBackedTarget* TargetComponent = Cast<IPrimitiveComponentBackedTarget>(Target);
-	DynamicMeshComponent = NewObject<USimpleDynamicMeshComponent>(TargetComponent->GetOwnerActor(), "DynamicMesh");
-	DynamicMeshComponent->SetupAttachment(TargetComponent->GetOwnerActor()->GetRootComponent());
+	DynamicMeshComponent = NewObject<USimpleDynamicMeshComponent>(UE::ToolTarget::GetTargetActor(Target));
+	DynamicMeshComponent->SetupAttachment(UE::ToolTarget::GetTargetComponent(Target));
 	DynamicMeshComponent->RegisterComponent();
-	DynamicMeshComponent->SetWorldTransform(TargetComponent->GetWorldTransform());
+	DynamicMeshComponent->SetWorldTransform((FTransform)UE::ToolTarget::GetLocalToWorldTransform(Target));
 
 	// transfer materials
-	FComponentMaterialSet MaterialSet;
-	Cast<IMaterialProvider>(Target)->GetMaterialSet(MaterialSet);
+	FComponentMaterialSet MaterialSet = UE::ToolTarget::GetMaterialSet(Target);
 	for (int k = 0; k < MaterialSet.Materials.Num(); ++k)
 	{
 		DynamicMeshComponent->SetMaterial(k, MaterialSet.Materials[k]);
 	}
 
-	DynamicMeshComponent->TangentsType = EDynamicMeshTangentCalcType::AutoCalculated;
-	DynamicMeshComponent->InitializeMesh(Cast<IMeshDescriptionProvider>(Target)->GetMeshDescription());
-	OriginalMesh.Copy(*DynamicMeshComponent->GetMesh());
+	DynamicMeshComponent->SetTangentsType(EDynamicMeshComponentTangentsMode::AutoCalculated);
+	DynamicMeshComponent->SetMesh(UE::ToolTarget::GetDynamicMeshCopy(Target));
+	DynamicMeshComponent->ProcessMesh([&](const FDynamicMesh3& ReadMesh) { OriginalMesh = ReadMesh; });
 
 	// hide input StaticMeshComponent
-	TargetComponent->SetOwnerVisibility(false);
+	UE::ToolTarget::HideSourceObject(Target);
 
 	// initialize our properties
 	ToolPropertyObjects.Add(this);
@@ -86,15 +80,15 @@ void UWeldMeshEdgesTool::Shutdown(EToolShutdownType ShutdownType)
 {
 	if (DynamicMeshComponent != nullptr)
 	{
-		Cast<IPrimitiveComponentBackedTarget>(Target)->SetOwnerVisibility(true);
+		UE::ToolTarget::ShowSourceObject(Target);
 
 		if (ShutdownType == EToolShutdownType::Accept)
 		{
 			// this block bakes the modified DynamicMeshComponent back into the StaticMeshComponent inside an undo transaction
 			GetToolManager()->BeginUndoTransaction(LOCTEXT("WeldMeshEdgesToolTransactionName", "Remesh Mesh"));
-			Cast<IMeshDescriptionCommitter>(Target)->CommitMeshDescription([this](const IMeshDescriptionCommitter::FCommitterParams& CommitParams)
+			DynamicMeshComponent->ProcessMesh([&](const FDynamicMesh3& CurMesh)
 			{
-				DynamicMeshComponent->Bake(CommitParams.MeshDescriptionOut, true);
+				UE::ToolTarget::CommitDynamicMeshUpdate(Target, CurMesh, true);
 			});
 			GetToolManager()->EndUndoTransaction();
 		}
@@ -111,7 +105,7 @@ void UWeldMeshEdgesTool::Render(IToolsContextRenderAPI* RenderAPI)
 	UpdateResult();
 
 	FPrimitiveDrawInterface* PDI = RenderAPI->GetPrimitiveDrawInterface();
-	FTransform Transform = Cast<IPrimitiveComponentBackedTarget>(Target)->GetWorldTransform(); //Actor->GetTransform();
+	FTransform Transform = (FTransform)UE::ToolTarget::GetLocalToWorldTransform(Target);
 
 	FColor LineColor(200, 200, 200);
 	float PDIScale = RenderAPI->GetCameraState().GetPDIScalingFactor();
@@ -155,29 +149,31 @@ void UWeldMeshEdgesTool::UpdateResult()
 		return;
 	}
 
-	FDynamicMesh3* TargetMesh = DynamicMeshComponent->GetMesh();
-	TargetMesh->Copy(OriginalMesh);
+	ComputeMesh = OriginalMesh;
 
-	FMergeCoincidentMeshEdges Merger(TargetMesh);
+	FMergeCoincidentMeshEdges Merger(&ComputeMesh);
 	Merger.MergeVertexTolerance = this->Tolerance;
 	Merger.MergeSearchTolerance = 2*Merger.MergeVertexTolerance;
 	Merger.OnlyUniquePairs = this->bOnlyUnique;
+	FDynamicMesh3* SetMesh = &ComputeMesh;
 		
 	if (Merger.Apply() == false)
 	{
 		UE_LOG(LogTemp, Warning, TEXT("WeldMeshEdgesTool : Merger.Apply() returned false"));
-		TargetMesh->Copy(OriginalMesh);
+		SetMesh = &OriginalMesh;
 	}
-
-	if (TargetMesh->CheckValidity(true, EValidityCheckFailMode::ReturnOnly) == false)
+	else if (ComputeMesh.CheckValidity(true, EValidityCheckFailMode::ReturnOnly) == false)
 	{
 		UE_LOG(LogTemp, Warning, TEXT("WeldMeshEdgesTool : returned mesh is invalid"));
-		TargetMesh->Copy(OriginalMesh);
+		SetMesh = &OriginalMesh;
 	}
 
-	DynamicMeshComponent->NotifyMeshUpdated();
-	GetToolManager()->PostInvalidation();
+	DynamicMeshComponent->EditMesh([SetMesh](FDynamicMesh3& EditMesh)
+	{
+		EditMesh = *SetMesh;
+	});
 
+	GetToolManager()->PostInvalidation();
 	bResultValid = true;
 }
 
diff --git a/Engine/Plugins/Experimental/MeshModelingToolset/Source/MeshModelingTools/Public/BakeMeshAttributeMapsTool.h b/Engine/Plugins/Experimental/MeshModelingToolset/Source/MeshModelingTools/Public/BakeMeshAttributeMapsTool.h
--- a/Engine/Plugins/Experimental/MeshModelingToolset/Source/MeshModelingTools/Public/BakeMeshAttributeMapsTool.h
+++ b/Engine/Plugins/Experimental/MeshModelingToolset/Source/MeshModelingTools/Public/BakeMeshAttributeMapsTool.h
@@ -416,7 +416,6 @@ protected:
 	UMaterialInstanceDynamic* WorkingPreviewMaterial;
 	float SecondsBeforeWorkingMaterial = 0.75;
 
-	TSharedPtr<FMeshDescription, ESPMode::ThreadSafe> BaseMeshDescription;
 	TSharedPtr<UE::Geometry::TMeshTangents<double>, ESPMode::ThreadSafe> BaseMeshTangents;
 	UE::Geometry::FDynamicMesh3 BaseMesh;
 	UE::Geometry::FDynamicMeshAABBTree3 BaseSpatial;
@@ -574,6 +573,6 @@ protected:
 
 	void InitializeEmptyMaps();
 
-	void GetTexturesFromDetailMesh(const IPrimitiveComponentBackedTarget* DetailComponent);
+	void GetTexturesFromDetailMesh(const UPrimitiveComponent* DetailComponent);
 
 };
diff --git a/Engine/Plugins/Experimental/MeshModelingToolset/Source/MeshModelingTools/Public/DeformMeshPolygonsTool.h b/Engine/Plugins/Experimental/MeshModelingToolset/Source/MeshModelingTools/Public/DeformMeshPolygonsTool.h
--- a/Engine/Plugins/Experimental/MeshModelingToolset/Source/MeshModelingTools/Public/DeformMeshPolygonsTool.h
+++ b/Engine/Plugins/Experimental/MeshModelingToolset/Source/MeshModelingTools/Public/DeformMeshPolygonsTool.h
@@ -191,6 +191,7 @@ protected:
 	
 	// camera state at last render
 	FViewCameraState CameraState;
+	UE::Geometry::FTransform3d WorldTransform;
 
 	FToolDataVisualizer PolyEdgesRenderer;
 
diff --git a/Engine/Plugins/Experimental/MeshModelingToolset/Source/MeshModelingTools/Public/MeshInspectorTool.h b/Engine/Plugins/Experimental/MeshModelingToolset/Source/MeshModelingTools/Public/MeshInspectorTool.h
--- a/Engine/Plugins/Experimental/MeshModelingToolset/Source/MeshModelingTools/Public/MeshInspectorTool.h
+++ b/Engine/Plugins/Experimental/MeshModelingToolset/Source/MeshModelingTools/Public/MeshInspectorTool.h
@@ -136,9 +136,6 @@ protected:
 	UPROPERTY()
 	UMaterialInterface* DefaultMaterial = nullptr;
 
-	// copy of input mesh description with tangents/etc populated
-	TSharedPtr<FMeshDescription> InputMeshDescription;
-
 	TArray<int> BoundaryEdges;
 	TArray<int> BoundaryBowties;
 	TArray<int> UVSeamEdges;
diff --git a/Engine/Plugins/Experimental/MeshModelingToolset/Source/MeshModelingTools/Public/Sculpting/MeshSculptToolBase.h b/Engine/Plugins/Experimental/MeshModelingToolset/Source/MeshModelingTools/Public/Sculpting/MeshSculptToolBase.h
--- a/Engine/Plugins/Experimental/MeshModelingToolset/Source/MeshModelingTools/Public/Sculpting/MeshSculptToolBase.h
+++ b/Engine/Plugins/Experimental/MeshModelingToolset/Source/MeshModelingTools/Public/Sculpting/MeshSculptToolBase.h
@@ -239,7 +239,7 @@ protected:
 	 * Subclass can override this to change what results are written.
 	 * Default is to apply a default vertex positions update to the target object.
 	 */
-	virtual void CommitResult(UBaseDynamicMeshComponent* Component);
+	virtual void CommitResult(UBaseDynamicMeshComponent* Component, bool bModifiedTopology);
 
 
 	//
diff --git a/Engine/Plugins/Experimental/MeshModelingToolset/Source/MeshModelingTools/Public/WeldMeshEdgesTool.h b/Engine/Plugins/Experimental/MeshModelingToolset/Source/MeshModelingTools/Public/WeldMeshEdgesTool.h
--- a/Engine/Plugins/Experimental/MeshModelingToolset/Source/MeshModelingTools/Public/WeldMeshEdgesTool.h
+++ b/Engine/Plugins/Experimental/MeshModelingToolset/Source/MeshModelingTools/Public/WeldMeshEdgesTool.h
@@ -65,6 +65,7 @@ protected:
 protected:
 	USimpleDynamicMeshComponent* DynamicMeshComponent;
 	UE::Geometry::FDynamicMesh3 OriginalMesh;
+	UE::Geometry::FDynamicMesh3 ComputeMesh;
 	bool bResultValid;
 	void UpdateResult();
 };
diff --git a/Engine/Plugins/Experimental/MeshModelingToolset/Source/MeshModelingToolsEditorOnly/Private/GenerateLODMeshesTool.cpp b/Engine/Plugins/Experimental/MeshModelingToolset/Source/MeshModelingToolsEditorOnly/Private/GenerateLODMeshesTool.cpp
--- a/Engine/Plugins/Experimental/MeshModelingToolset/Source/MeshModelingToolsEditorOnly/Private/GenerateLODMeshesTool.cpp
+++ b/Engine/Plugins/Experimental/MeshModelingToolset/Source/MeshModelingToolsEditorOnly/Private/GenerateLODMeshesTool.cpp
@@ -203,7 +203,7 @@ void UGenerateLODMeshesTool::UpdateNumPreviews()
 			ToolSetupUtil::GetDefaultWorkingMaterial(GetToolManager())
 		);
 
-		NewPreview->PreviewMesh->SetTangentsMode(EDynamicMeshTangentCalcType::AutoCalculated);
+		NewPreview->PreviewMesh->SetTangentsMode(EDynamicMeshComponentTangentsMode::AutoCalculated);
 		NewPreview->PreviewMesh->UpdatePreview(OriginalMesh.Get());
 
 		FTransform UseTransform = OrigTransform;
diff --git a/Engine/Plugins/Experimental/MeshModelingToolset/Source/MeshModelingToolsEditorOnly/Private/Hair/GroomCardsEditorTool.cpp b/Engine/Plugins/Experimental/MeshModelingToolset/Source/MeshModelingToolsEditorOnly/Private/Hair/GroomCardsEditorTool.cpp
--- a/Engine/Plugins/Experimental/MeshModelingToolset/Source/MeshModelingToolsEditorOnly/Private/Hair/GroomCardsEditorTool.cpp
+++ b/Engine/Plugins/Experimental/MeshModelingToolset/Source/MeshModelingToolsEditorOnly/Private/Hair/GroomCardsEditorTool.cpp
@@ -19,6 +19,8 @@
 #include "DynamicMeshChangeTracker.h"
 #include "Parameterization/DynamicMeshUVEditor.h"
 #include "Parameterization/MeshUVTransforms.h"
+#include "ModelingToolTargetUtil.h"
+#include "DynamicMeshToMeshDescription.h"
 
 #include "Changes/MeshVertexChange.h"
 
@@ -27,11 +29,7 @@
 #include "Engine/Classes/Engine/StaticMesh.h"
 #include "Engine/Classes/Components/StaticMeshComponent.h"
 #include "Engine/Classes/Engine/StaticMeshActor.h"
-#include "MeshDescriptionToDynamicMesh.h"
 
-#include "TargetInterfaces/MeshDescriptionCommitter.h"
-#include "TargetInterfaces/MeshDescriptionProvider.h"
-#include "TargetInterfaces/PrimitiveComponentBackedTarget.h"
 
 #define LOCTEXT_NAMESPACE "UGroomCardsEditorTool"
 
@@ -75,11 +73,10 @@ void UGroomCardsEditorTool::Setup()
 {
 	UMeshSurfacePointTool::Setup();
 
-	IPrimitiveComponentBackedTarget* TargetComponent = Cast<IPrimitiveComponentBackedTarget>(Target);
 	PreviewMesh = NewObject<UPreviewMesh>(this);
 	PreviewMesh->bBuildSpatialDataStructure = true;
 	PreviewMesh->CreateInWorld(TargetWorld, FTransform::Identity);
-	PreviewMesh->SetTransform(TargetComponent->GetWorldTransform());
+	PreviewMesh->SetTransform((FTransform)UE::ToolTarget::GetLocalToWorldTransform(Target));
 
 	MeshMaterial = ToolSetupUtil::GetDefaultSculptMaterial(GetToolManager());
 	UVMaterial = ToolSetupUtil::GetUVCheckerboardMaterial(50.0);
@@ -112,11 +109,9 @@ void UGroomCardsEditorTool::Setup()
 
 	PreviewGeom = NewObject<UPreviewGeometry>(this);
 	PreviewGeom->CreateInWorld(TargetWorld, FTransform::Identity);
-	PreviewGeom->GetActor()->SetActorTransform(TargetComponent->GetWorldTransform());
-
-	TargetComponent->SetOwnerVisibility(false);
-
+	PreviewGeom->GetActor()->SetActorTransform((FTransform)UE::ToolTarget::GetLocalToWorldTransform(Target));
 
+	UE::ToolTarget::HideSourceObject(Target);
 
 	ControlPointsMechanic = NewObject<USpaceCurveDeformationMechanic>(this);
 	ControlPointsMechanic->Setup(this);
@@ -184,9 +179,9 @@ void UGroomCardsEditorTool::Shutdown(EToolShutdownType ShutdownType)
 	if (ShutdownType == EToolShutdownType::Accept)
 	{
 		GetToolManager()->BeginUndoTransaction(LOCTEXT("CardsMeshApplyEditChange", "Update Cards Mesh"));
-		Cast<IMeshDescriptionCommitter>(Target)->CommitMeshDescription([this](const IMeshDescriptionCommitter::FCommitterParams& CommitParams)
+		PreviewMesh->ProcessMesh([&](const FDynamicMesh3& ReadMesh)
 		{
-			PreviewMesh->Bake(CommitParams.MeshDescriptionOut, false);
+			UE::ToolTarget::CommitDynamicMeshUpdate(Target, ReadMesh, false);
 		});
 		GetToolManager()->EndUndoTransaction();
 	}
@@ -197,7 +192,7 @@ void UGroomCardsEditorTool::Shutdown(EToolShutdownType ShutdownType)
 	PreviewGeom->Disconnect();
 	PreviewGeom = nullptr;
 
-	Cast<IPrimitiveComponentBackedTarget>(Target)->SetOwnerVisibility(true);
+	UE::ToolTarget::ShowSourceObject(Target);
 }
 
 
@@ -762,7 +757,7 @@ void UGroomCardsEditorTool::UpdateLineSet()
 void UGroomCardsEditorTool::InitializeMesh()
 {
 	EditableCardSet = MakePimpl<FEditableGroomCardSet>();
-	EditableCardSet->Initialize(Cast<IMeshDescriptionProvider>(Target)->GetMeshDescription());
+	EditableCardSet->Initialize(UE::ToolTarget::GetMeshDescription(Target));
 	const FDynamicMesh3* CardsMesh = EditableCardSet->FullCardMesh.Get();
 
 	PreviewMesh->UpdatePreview(CardsMesh);
diff --git a/Engine/Plugins/Experimental/MeshModelingToolset/Source/MeshModelingToolsEditorOnly/Private/MeshTangentsTool.cpp b/Engine/Plugins/Experimental/MeshModelingToolset/Source/MeshModelingToolsEditorOnly/Private/MeshTangentsTool.cpp
--- a/Engine/Plugins/Experimental/MeshModelingToolset/Source/MeshModelingToolsEditorOnly/Private/MeshTangentsTool.cpp
+++ b/Engine/Plugins/Experimental/MeshModelingToolset/Source/MeshModelingToolsEditorOnly/Private/MeshTangentsTool.cpp
@@ -20,6 +20,7 @@
 #include "TargetInterfaces/MeshDescriptionProvider.h"
 #include "TargetInterfaces/PrimitiveComponentBackedTarget.h"
 #include "TargetInterfaces/StaticMeshBackedTarget.h"
+#include "ModelingToolTargetUtil.h"
 
 #include "ExplicitUseGeometryMathTypes.h"		// using UE::Geometry::(math types)
 using namespace UE::Geometry;
@@ -60,31 +61,29 @@ void UMeshTangentsTool::Setup()
 {
 	UInteractiveTool::Setup();
 
-	// copy input MeshDescription and make sure it has initialized normals/tangents
-	IPrimitiveComponentBackedTarget* TargetComponent = Cast<IPrimitiveComponentBackedTarget>(Target);
-	InputMeshDescription = MakeShared<FMeshDescription, ESPMode::ThreadSafe>(*Cast<IMeshDescriptionProvider>(Target)->GetMeshDescription());
-	UE::MeshDescription::InitializeAutoGeneratedAttributes(*InputMeshDescription, TargetComponent->GetOwnerComponent(), 0);
-
-	// hide input StaticMeshComponent
-	TargetComponent->SetOwnerVisibility(false);
+	UE::ToolTarget::HideSourceObject(Target);
 
 	// make our preview mesh
+	AActor* TargetActor = UE::ToolTarget::GetTargetActor(Target);
 	PreviewMesh = NewObject<UPreviewMesh>(this);
 	PreviewMesh->bBuildSpatialDataStructure = false;
-	PreviewMesh->CreateInWorld(TargetComponent->GetOwnerActor()->GetWorld(), FTransform::Identity);
-	PreviewMesh->SetTransform(TargetComponent->GetWorldTransform());
+	PreviewMesh->CreateInWorld(TargetActor->GetWorld(), FTransform::Identity);
+	PreviewMesh->SetTransform((FTransform)UE::ToolTarget::GetLocalToWorldTransform(Target));
 	// configure materials
-	FComponentMaterialSet MaterialSet;
-	Cast<IMaterialProvider>(Target)->GetMaterialSet(MaterialSet);
+	FComponentMaterialSet MaterialSet = UE::ToolTarget::GetMaterialSet(Target);
 	PreviewMesh->SetMaterials(MaterialSet.Materials);
-	// initialize with meshdescription created above, that has valid tangents
-	PreviewMesh->SetTangentsMode(EDynamicMeshTangentCalcType::ExternallyCalculated);
-	PreviewMesh->InitializeMesh(InputMeshDescription.Get());
+	// configure mesh
+	PreviewMesh->SetTangentsMode(EDynamicMeshComponentTangentsMode::ExternallyProvided);
+	FDynamicMesh3 InputMeshWithTangents = UE::ToolTarget::GetDynamicMeshCopy(Target, true);
+	PreviewMesh->ReplaceMesh(MoveTemp(InputMeshWithTangents));
 
 	// make a copy of initialized mesh and tangents
-	InputMesh = MakeShared<FDynamicMesh3, ESPMode::ThreadSafe>(*PreviewMesh->GetMesh());
-	InitialTangents = MakeShared<FMeshTangentsf, ESPMode::ThreadSafe>(InputMesh.Get());
-	InitialTangents->CopyTriVertexTangents(*PreviewMesh->GetTangents());
+	PreviewMesh->ProcessMesh([&](const FDynamicMesh3& ReadMesh)
+	{
+		InputMesh = MakeShared<FDynamicMesh3, ESPMode::ThreadSafe>(ReadMesh);
+		InitialTangents = MakeShared<FMeshTangentsf, ESPMode::ThreadSafe>(InputMesh.Get());
+		InitialTangents->CopyTriVertexTangents(ReadMesh);
+	});
 
 	// initialize our properties
 	Settings = NewObject<UMeshTangentsToolProperties>(this);
@@ -98,7 +97,7 @@ void UMeshTangentsTool::Setup()
 	Settings->WatchProperty(Settings->bShowNormals, [this](float) { bVisibilityChanged = true; });
 
 	PreviewGeometry = NewObject<UPreviewGeometry>(this);
-	PreviewGeometry->CreateInWorld(TargetComponent->GetOwnerActor()->GetWorld(), PreviewMesh->GetTransform());
+	PreviewGeometry->CreateInWorld(TargetActor->GetWorld(), PreviewMesh->GetTransform());
 
 	Compute = MakeUnique<TGenericDataBackgroundCompute<FMeshTangentsd>>();
 	Compute->Setup(this);
@@ -121,15 +120,14 @@ void UMeshTangentsTool::Shutdown(EToolShutdownType ShutdownType)
 	Settings->SaveProperties(this);
 
 	// Restore (unhide) the source meshes
-	IPrimitiveComponentBackedTarget* TargetComponent = Cast<IPrimitiveComponentBackedTarget>(Target);
-	TargetComponent->SetOwnerVisibility(true);
+	UE::ToolTarget::ShowSourceObject(Target);
 
 	TUniquePtr<FMeshTangentsd> Tangents = Compute->Shutdown();
 	if (ShutdownType == EToolShutdownType::Accept)
 	{
 		GetToolManager()->BeginUndoTransaction(LOCTEXT("UpdateTangents", "Update Tangents"));
 
-		UStaticMeshComponent* StaticMeshComponent = Cast<UStaticMeshComponent>(TargetComponent->GetOwnerComponent());
+		UStaticMeshComponent* StaticMeshComponent = Cast<UStaticMeshComponent>(UE::ToolTarget::GetTargetComponent(Target));
 		UStaticMesh* StaticMesh = StaticMeshComponent->GetStaticMesh();
 		StaticMesh->Modify();
 
@@ -139,17 +137,13 @@ void UMeshTangentsTool::Shutdown(EToolShutdownType ShutdownType)
 		SettingsChange.AutoGeneratedTangents = UE::MeshDescription::EBuildSettingBoolChange::Disable;
 		UE::MeshDescription::ConfigureBuildSettings(StaticMesh, 0, SettingsChange);
 
-		Cast<IMeshDescriptionCommitter>(Target)->CommitMeshDescription([&](const IMeshDescriptionCommitter::FCommitterParams& CommitParams)
-		{
-			FDynamicMeshToMeshDescription Converter;
-			{
-				const bool bCopyNormals = true; 
-				const bool bCopyOverlayTangents = false;
-				const bool bCopyOverlayUVs = false;
-				Converter.UpdateAttributes(InputMesh.Get(), *CommitParams.MeshDescriptionOut, bCopyNormals, bCopyOverlayTangents, bCopyOverlayUVs);
-			}
-			Converter.UpdateTangents(InputMesh.Get(), *CommitParams.MeshDescriptionOut, Tangents.Get());
-		});
+		Tangents->CopyToOverlays(*InputMesh);
+
+		FConversionToMeshDescriptionOptions Options;
+		Options.bUpdatePositions = false;
+		Options.bUpdateNormals = true;
+		Options.bUpdateTangents = true;
+		UE::ToolTarget::CommitDynamicMeshUpdate(Target, *InputMesh, false, Options);
 
 		GetToolManager()->EndUndoTransaction();
 	}
@@ -331,8 +325,11 @@ void UMeshTangentsTool::OnTangentsUpdated(const TUniquePtr<FMeshTangentsd>& NewR
 		}
 	}
 
-	// update the tangents on the preview mesh
-	PreviewMesh->UpdateTangents(NewResult.Get(), true);
+	PreviewMesh->DeferredEditMesh([&](FDynamicMesh3& EditMesh)
+	{
+		NewResult.Get()->CopyToOverlays(EditMesh);
+	}, false);
+	PreviewMesh->NotifyDeferredEditCompleted(UPreviewMesh::ERenderUpdateMode::FastUpdate, EMeshRenderAttributeFlags::VertexNormals, false);
 
 	UpdateVisualization(false, false);
 }
diff --git a/Engine/Plugins/Experimental/MeshModelingToolset/Source/MeshModelingToolsEditorOnly/Private/MeshToVolumeTool.cpp b/Engine/Plugins/Experimental/MeshModelingToolset/Source/MeshModelingToolsEditorOnly/Private/MeshToVolumeTool.cpp
--- a/Engine/Plugins/Experimental/MeshModelingToolset/Source/MeshModelingToolsEditorOnly/Private/MeshToVolumeTool.cpp
+++ b/Engine/Plugins/Experimental/MeshModelingToolset/Source/MeshModelingToolsEditorOnly/Private/MeshToVolumeTool.cpp
@@ -12,10 +12,9 @@
 #include "ToolSetupUtil.h"
 #include "Selections/MeshConnectedComponents.h"
 #include "Selection/ToolSelectionUtil.h"
-#include "TargetInterfaces/MaterialProvider.h"
-#include "TargetInterfaces/MeshDescriptionProvider.h"
-#include "TargetInterfaces/PrimitiveComponentBackedTarget.h"
 #include "ToolTargetManager.h"
+#include "ModelingToolTargetUtil.h"
+
 
 #include "Engine/Classes/Engine/BlockingVolume.h"
 #include "Engine/Classes/Components/BrushComponent.h"
@@ -56,26 +55,16 @@ void UMeshToVolumeTool::Setup()
 {
 	UInteractiveTool::Setup();
 
-	IPrimitiveComponentBackedTarget* TargetComponent = Cast<IPrimitiveComponentBackedTarget>(Target);
-	check(TargetComponent);
-
 	PreviewMesh = NewObject<UPreviewMesh>(this);
 	PreviewMesh->bBuildSpatialDataStructure = false;
-	PreviewMesh->CreateInWorld(TargetComponent->GetOwnerActor()->GetWorld(), FTransform::Identity);
-	PreviewMesh->SetTransform(TargetComponent->GetWorldTransform());
+	PreviewMesh->CreateInWorld(UE::ToolTarget::GetTargetActor(Target)->GetWorld(), FTransform::Identity);
+	PreviewMesh->SetTransform((FTransform)UE::ToolTarget::GetLocalToWorldTransform(Target));
 
-	FComponentMaterialSet MaterialSet;
-	check(Cast<IMaterialProvider>(Target));
-	Cast<IMaterialProvider>(Target)->GetMaterialSet(MaterialSet);
-	PreviewMesh->SetMaterials(MaterialSet.Materials);
+	PreviewMesh->SetTangentsMode(EDynamicMeshComponentTangentsMode::AutoCalculated);
+	PreviewMesh->ReplaceMesh(UE::ToolTarget::GetDynamicMeshCopy(Target));
 
-	PreviewMesh->SetTangentsMode(EDynamicMeshTangentCalcType::ExternallyCalculated);
-
-	// TODO: We really just need a dynamic mesh, not a mesh description (since we're converting the mesh
-	// description to a dynamic mesh), but it requires refactoring PreviewMesh to accept a dynamic mesh. 
-	// Once we do, however, the requirement for this tool will change to IDynamicMeshProvider, and we'll
-	// use that instead.
-	PreviewMesh->InitializeMesh(Cast<IMeshDescriptionProvider>(Target)->GetMeshDescription());
+	FComponentMaterialSet MaterialSet = UE::ToolTarget::GetMaterialSet(Target);
+	PreviewMesh->SetMaterials(MaterialSet.Materials);
 
 	InputMesh.Copy(*PreviewMesh->GetMesh());
 
@@ -84,8 +73,7 @@ void UMeshToVolumeTool::Setup()
 	VolumeEdgesSet->SetLineMaterial(ToolSetupUtil::GetDefaultLineComponentMaterial(GetToolManager()));
 	VolumeEdgesSet->RegisterComponent();
 
-	// hide input StaticMeshComponent
-	TargetComponent->SetOwnerVisibility(false);
+	UE::ToolTarget::HideSourceObject(Target);
 
 	Settings = NewObject<UMeshToVolumeToolProperties>(this);
 	Settings->RestoreProperties(this);
@@ -117,21 +105,18 @@ void UMeshToVolumeTool::Shutdown(EToolShutdownType ShutdownType)
 	PreviewMesh->Disconnect();
 	PreviewMesh = nullptr;
 
-	IPrimitiveComponentBackedTarget* TargetComponent = Cast<IPrimitiveComponentBackedTarget>(Target);
-	TargetComponent->SetOwnerVisibility(true);
-
-	UWorld* TargetOwnerWorld = TargetComponent->GetOwnerActor()->GetWorld();
-
+	UE::ToolTarget::ShowSourceObject(Target);
 
 	if (ShutdownType == EToolShutdownType::Accept)
 	{
 		GetToolManager()->BeginUndoTransaction(LOCTEXT("MeshToVolumeToolTransactionName", "Create Volume"));
 
-		FTransform TargetTransform = TargetComponent->GetWorldTransform();
+		AActor* TargetOwnerActor = UE::ToolTarget::GetTargetActor(Target);
+		UWorld* TargetOwnerWorld = TargetOwnerActor->GetWorld();
+		FTransform SetTransform = (FTransform)UE::ToolTarget::GetLocalToWorldTransform(Target);
 
 		AVolume* TargetVolume = nullptr;
 
-		FTransform SetTransform = TargetTransform;
 		if (Settings->TargetVolume.IsValid() == false)
 		{
 			FRotator Rotation(0.0f, 0.0f, 0.0f);
@@ -166,12 +151,10 @@ void UMeshToVolumeTool::Shutdown(EToolShutdownType ShutdownType)
 		ToolSelectionUtil::SetNewActorSelection(GetToolManager(), TargetVolume);
 
 		TArray<AActor*> Actors;
-		Actors.Add(TargetComponent->GetOwnerActor());
+		Actors.Add(TargetOwnerActor);
 		HandleSourcesProperties->ApplyMethod(Actors, GetToolManager());
 
-
 		GetToolManager()->EndUndoTransaction();
-
 	}
 }
 
diff --git a/Engine/Plugins/Experimental/MeshModelingToolset/Source/MeshModelingToolsEditorOnly/Private/ParameterizeMeshTool.cpp b/Engine/Plugins/Experimental/MeshModelingToolset/Source/MeshModelingToolsEditorOnly/Private/ParameterizeMeshTool.cpp
--- a/Engine/Plugins/Experimental/MeshModelingToolset/Source/MeshModelingToolsEditorOnly/Private/ParameterizeMeshTool.cpp
+++ b/Engine/Plugins/Experimental/MeshModelingToolset/Source/MeshModelingToolsEditorOnly/Private/ParameterizeMeshTool.cpp
@@ -15,7 +15,7 @@
 
 #include "ParameterizationOps/ParameterizeMeshOp.h"
 #include "ToolSetupUtil.h"
-
+#include "ModelingToolTargetUtil.h"
 #include "TargetInterfaces/MaterialProvider.h"
 #include "TargetInterfaces/MeshDescriptionCommitter.h"
 #include "TargetInterfaces/MeshDescriptionProvider.h"
@@ -72,9 +72,8 @@ void UParameterizeMeshTool::Setup()
 	IMeshDescriptionProvider* TargetMeshProvider = Cast<IMeshDescriptionProvider>(Target);
 	Preview = NewObject<UMeshOpPreviewWithBackgroundCompute>(this, "Preview");
 	Preview->Setup(this->TargetWorld, this);
-	Preview->PreviewMesh->SetTangentsMode(EDynamicMeshTangentCalcType::AutoCalculated);
-
-	Preview->PreviewMesh->InitializeMesh(TargetMeshProvider->GetMeshDescription());
+	Preview->PreviewMesh->SetTangentsMode(EDynamicMeshComponentTangentsMode::AutoCalculated);
+	Preview->PreviewMesh->ReplaceMesh(UE::ToolTarget::GetDynamicMeshCopy(Target));
 
 	Preview->OnMeshUpdated.AddLambda([this](UMeshOpPreviewWithBackgroundCompute* Op)
 	{
diff --git a/Engine/Plugins/Experimental/MeshModelingToolset/Source/MeshModelingToolsEditorOnly/Private/PolygonOnMeshTool.cpp b/Engine/Plugins/Experimental/MeshModelingToolset/Source/MeshModelingToolsEditorOnly/Private/PolygonOnMeshTool.cpp
--- a/Engine/Plugins/Experimental/MeshModelingToolset/Source/MeshModelingToolsEditorOnly/Private/PolygonOnMeshTool.cpp
+++ b/Engine/Plugins/Experimental/MeshModelingToolset/Source/MeshModelingToolsEditorOnly/Private/PolygonOnMeshTool.cpp
@@ -216,7 +216,7 @@ void UPolygonOnMeshTool::SetupPreview()
 {
 	Preview = NewObject<UMeshOpPreviewWithBackgroundCompute>(this);
 	Preview->Setup(this->TargetWorld, this);
-	Preview->PreviewMesh->SetTangentsMode(EDynamicMeshTangentCalcType::AutoCalculated);
+	Preview->PreviewMesh->SetTangentsMode(EDynamicMeshComponentTangentsMode::AutoCalculated);
 
 	FComponentMaterialSet MaterialSet;
 	Cast<IMaterialProvider>(Target)->GetMaterialSet(MaterialSet);
diff --git a/Engine/Plugins/Experimental/MeshModelingToolset/Source/MeshModelingToolsEditorOnly/Private/SimplifyMeshTool.cpp b/Engine/Plugins/Experimental/MeshModelingToolset/Source/MeshModelingToolsEditorOnly/Private/SimplifyMeshTool.cpp
--- a/Engine/Plugins/Experimental/MeshModelingToolset/Source/MeshModelingToolsEditorOnly/Private/SimplifyMeshTool.cpp
+++ b/Engine/Plugins/Experimental/MeshModelingToolset/Source/MeshModelingToolsEditorOnly/Private/SimplifyMeshTool.cpp
@@ -124,7 +124,7 @@ void USimplifyMeshTool::Setup()
 	}
 
 	Preview->PreviewMesh->SetTransform(TargetComponent->GetWorldTransform());
-	Preview->PreviewMesh->SetTangentsMode(EDynamicMeshTangentCalcType::AutoCalculated);
+	Preview->PreviewMesh->SetTangentsMode(EDynamicMeshComponentTangentsMode::AutoCalculated);
 	Preview->PreviewMesh->UpdatePreview(OriginalMesh.Get());
 
 	// initialize our properties
diff --git a/Engine/Plugins/Experimental/MeshModelingToolset/Source/ModelingComponents/Private/BaseDynamicMeshSceneProxy.h b/Engine/Plugins/Experimental/MeshModelingToolset/Source/ModelingComponents/Private/BaseDynamicMeshSceneProxy.h
--- a/Engine/Plugins/Experimental/MeshModelingToolset/Source/ModelingComponents/Private/BaseDynamicMeshSceneProxy.h
+++ b/Engine/Plugins/Experimental/MeshModelingToolset/Source/ModelingComponents/Private/BaseDynamicMeshSceneProxy.h
@@ -414,7 +414,7 @@ public:
 		const FDynamicMeshUVOverlay* UVOverlay,
 		const FDynamicMeshNormalOverlay* NormalOverlay,
 		const FDynamicMeshColorOverlay* ColorOverlay,
-		TFunction<void(int, int, int, FVector3f&, FVector3f&)> TangentsFunc = nullptr,
+		TFunctionRef<void(int, int, int, const FVector3f&, FVector3f&, FVector3f&)> TangentsFunc,
 		bool bTrackTriangles = false)
 	{
 		TArray<const FDynamicMeshUVOverlay*> UVOverlays;
@@ -437,7 +437,7 @@ public:
 		const TArray<const FDynamicMeshUVOverlay*>& UVOverlays,
 		const FDynamicMeshNormalOverlay* NormalOverlay,
 		const FDynamicMeshColorOverlay* ColorOverlay,
-		TFunction<void(int, int, int, FVector3f&, FVector3f&)> TangentsFunc = nullptr,
+		TFunctionRef<void(int, int, int, const FVector3f&, FVector3f&, FVector3f&)> TangentsFunc,
 		bool bTrackTriangles = false)
 	{
 		RenderBuffers->TriangleCount = NumTriangles;
@@ -493,16 +493,10 @@ public:
 
 				FVector3f Normal = (NormalOverlay != nullptr && TriNormal[j] != FDynamicMesh3::InvalidID) ?
 					NormalOverlay->GetElement(TriNormal[j]) : Mesh->GetVertexNormal(Tri[j]);
-				
-				// either request known tangent, or calculate a nonsense one
-				if (TangentsFunc != nullptr)
-				{
-					TangentsFunc(Tri[j], TriangleID, j, TangentX, TangentY);
-				}
-				else
-				{
-					VectorUtil::MakePerpVectors(Normal, TangentX, TangentY);
-				}
+
+				// get tangents
+				TangentsFunc(Tri[j], TriangleID, j, Normal, TangentX, TangentY);
+
 				RenderBuffers->StaticMeshVertexBuffer.SetVertexTangents(VertIdx, (FVector)TangentX, (FVector)TangentY, (FVector)Normal);
 
 				for (int32 k = 0; k < NumTexCoords; ++k)
@@ -640,8 +634,8 @@ public:
 		FMeshRenderBufferSet* RenderBuffers,
 		const FDynamicMesh3* Mesh,
 		int NumTriangles, TriangleEnumerable Enumerable,
-		FDynamicMeshNormalOverlay* NormalOverlay,
-		FDynamicMeshColorOverlay* ColorOverlay,
+		const FDynamicMeshNormalOverlay* NormalOverlay,
+		const FDynamicMeshColorOverlay* ColorOverlay,
 		TFunctionRef<void(int, int, int, const FVector3f&, FVector3f&, FVector3f&)> TangentsFunc,
 		bool bUpdatePositions = true,
 		bool bUpdateNormals = false,
@@ -718,7 +712,7 @@ public:
 		FMeshRenderBufferSet* RenderBuffers,
 		const FDynamicMesh3* Mesh,
 		int32 NumTriangles, TriangleEnumerable Enumerable,
-		FDynamicMeshUVOverlay* UVOverlay, int32 UVIndex)
+		const FDynamicMeshUVOverlay* UVOverlay, int32 UVIndex)
 	{
 		if (RenderBuffers->TriangleCount == 0)
 		{
diff --git a/Engine/Plugins/Experimental/MeshModelingToolset/Source/ModelingComponents/Private/BaseTools/BaseMeshProcessingTool.cpp b/Engine/Plugins/Experimental/MeshModelingToolset/Source/ModelingComponents/Private/BaseTools/BaseMeshProcessingTool.cpp
--- a/Engine/Plugins/Experimental/MeshModelingToolset/Source/ModelingComponents/Private/BaseTools/BaseMeshProcessingTool.cpp
+++ b/Engine/Plugins/Experimental/MeshModelingToolset/Source/ModelingComponents/Private/BaseTools/BaseMeshProcessingTool.cpp
@@ -142,7 +142,7 @@ void UBaseMeshProcessingTool::Setup()
 	// Construct the preview object and set the material on it.
 	Preview = NewObject<UMeshOpPreviewWithBackgroundCompute>(this, "Preview");
 	Preview->Setup(this->TargetWorld, this); // Adds the actual functional tool in the Preview object
-	Preview->PreviewMesh->SetTangentsMode(EDynamicMeshTangentCalcType::AutoCalculated);
+	Preview->PreviewMesh->SetTangentsMode(EDynamicMeshComponentTangentsMode::AutoCalculated);
 
 	FComponentMaterialSet MaterialSet;
 	Cast<IMaterialProvider>(Target)->GetMaterialSet(MaterialSet);
diff --git a/Engine/Plugins/Experimental/MeshModelingToolset/Source/ModelingComponents/Private/ModelingToolTargetUtil.cpp b/Engine/Plugins/Experimental/MeshModelingToolset/Source/ModelingComponents/Private/ModelingToolTargetUtil.cpp
--- a/Engine/Plugins/Experimental/MeshModelingToolset/Source/ModelingComponents/Private/ModelingToolTargetUtil.cpp
+++ b/Engine/Plugins/Experimental/MeshModelingToolset/Source/ModelingComponents/Private/ModelingToolTargetUtil.cpp
@@ -15,6 +15,7 @@
 #include "GameFramework/Volume.h"
 #include "SimpleDynamicMeshComponent.h"
 
+#include "MeshConversionOptions.h"
 #include "MeshDescriptionToDynamicMesh.h"
 #include "DynamicMeshToMeshDescription.h"
 
@@ -90,6 +91,21 @@ FComponentMaterialSet UE::ToolTarget::GetMaterialSet(UToolTarget* Target, bool b
 }
 
 
+bool UE::ToolTarget::CommitMaterialSetUpdate(
+	UToolTarget* Target,
+	const FComponentMaterialSet& UpdatedMaterials,
+	bool bApplyToAsset)
+{
+	IMaterialProvider* MaterialProvider = Cast<IMaterialProvider>(Target);
+	if (MaterialProvider)
+	{
+		return MaterialProvider->CommitMaterialSetUpdate(UpdatedMaterials, bApplyToAsset);
+	}
+	return false;
+}
+
+
+
 const FMeshDescription* UE::ToolTarget::GetMeshDescription(UToolTarget* Target)
 {
 	static FMeshDescription EmptyMeshDescription;
@@ -104,7 +120,7 @@ const FMeshDescription* UE::ToolTarget::GetMeshDescription(UToolTarget* Target)
 }
 
 
-FDynamicMesh3 UE::ToolTarget::GetDynamicMeshCopy(UToolTarget* Target)
+FDynamicMesh3 UE::ToolTarget::GetDynamicMeshCopy(UToolTarget* Target, bool bWantMeshTangents)
 {
 	IMeshDescriptionProvider* MeshDescriptionProvider = Cast<IMeshDescriptionProvider>(Target);
 	FDynamicMesh3 Mesh(EMeshComponents::FaceGroups);
@@ -112,7 +128,19 @@ FDynamicMesh3 UE::ToolTarget::GetDynamicMeshCopy(UToolTarget* Target)
 	if (MeshDescriptionProvider)
 	{
 		FMeshDescriptionToDynamicMesh Converter;
-		Converter.Convert(MeshDescriptionProvider->GetMeshDescription(), Mesh);
+		if (bWantMeshTangents)
+		{
+			// Currently to auto-calculate attributes we must make a copy, because they are written
+			// directly to the input mesh
+			FMeshDescription MeshDescriptionCopy( *MeshDescriptionProvider->GetMeshDescription() );
+			MeshDescriptionProvider->CalculateAutoGeneratedAttributes(MeshDescriptionCopy);
+			Converter.Convert(&MeshDescriptionCopy, Mesh, bWantMeshTangents);
+		}
+		else
+		{
+			Converter.Convert(MeshDescriptionProvider->GetMeshDescription(), Mesh, bWantMeshTangents);
+		}
+
 		return Mesh;
 	}
 
@@ -131,11 +159,7 @@ UE::ToolTarget::EDynamicMeshUpdateResult UE::ToolTarget::CommitMeshDescriptionUp
 
 	if (UpdatedMaterials != nullptr)
 	{
-		IMaterialProvider* MaterialProvider = Cast<IMaterialProvider>(Target);
-		if (MaterialProvider)
-		{
-			MaterialProvider->CommitMaterialSetUpdate(*UpdatedMaterials, true);
-		}
+		CommitMaterialSetUpdate(Target, *UpdatedMaterials, true);
 	}
 
 	EDynamicMeshUpdateResult Result = EDynamicMeshUpdateResult::Failed;
@@ -149,6 +173,49 @@ UE::ToolTarget::EDynamicMeshUpdateResult UE::ToolTarget::CommitMeshDescriptionUp
 
 
 
+
+UE::ToolTarget::EDynamicMeshUpdateResult UE::ToolTarget::CommitDynamicMeshUpdate(
+	UToolTarget* Target, const FDynamicMesh3& UpdatedMesh, 
+	bool bHaveModifiedTopology,
+	const FConversionToMeshDescriptionOptions& ConversionOptions,
+	const FComponentMaterialSet* UpdatedMaterials)
+{
+	IMeshDescriptionCommitter* MeshDescriptionCommitter = Cast<IMeshDescriptionCommitter>(Target);
+	if (MeshDescriptionCommitter)
+	{
+		if (UpdatedMaterials != nullptr)
+		{
+			CommitMaterialSetUpdate(Target, *UpdatedMaterials, true);
+		}
+
+		EDynamicMeshUpdateResult Result = EDynamicMeshUpdateResult::Failed;
+		MeshDescriptionCommitter->CommitMeshDescription([&](const IMeshDescriptionCommitter::FCommitterParams& CommitParams)
+		{
+			FMeshDescription* MeshDescription = CommitParams.MeshDescriptionOut;
+
+			FDynamicMeshToMeshDescription Converter(ConversionOptions);
+			if (!bHaveModifiedTopology)
+			{
+				Converter.UpdateUsingConversionOptions(&UpdatedMesh, *MeshDescription);
+				Result = EDynamicMeshUpdateResult::Ok;
+			}
+			else
+			{
+				Converter.Convert(&UpdatedMesh, *MeshDescription);
+				Result = EDynamicMeshUpdateResult::Ok;
+			}
+		});
+		return Result;
+	}
+
+	ensure(false);
+	return EDynamicMeshUpdateResult::Failed;
+}
+
+
+
+
+
 UE::ToolTarget::EDynamicMeshUpdateResult UE::ToolTarget::CommitDynamicMeshUVUpdate(UToolTarget* Target, const UE::Geometry::FDynamicMesh3* UpdatedMesh)
 {
 	IMeshDescriptionCommitter* MeshDescriptionCommitter = Cast<IMeshDescriptionCommitter>(Target);
diff --git a/Engine/Plugins/Experimental/MeshModelingToolset/Source/ModelingComponents/Private/OctreeDynamicMeshComponent.cpp b/Engine/Plugins/Experimental/MeshModelingToolset/Source/ModelingComponents/Private/OctreeDynamicMeshComponent.cpp
--- a/Engine/Plugins/Experimental/MeshModelingToolset/Source/ModelingComponents/Private/OctreeDynamicMeshComponent.cpp
+++ b/Engine/Plugins/Experimental/MeshModelingToolset/Source/ModelingComponents/Private/OctreeDynamicMeshComponent.cpp
@@ -21,7 +21,6 @@
 #include "MeshNormals.h"
 #include "MeshTransforms.h"
 #include "MeshDescriptionToDynamicMesh.h"
-#include "DynamicMeshToMeshDescription.h"
 
 #include "Changes/MeshVertexChange.h"
 #include "Changes/MeshChange.h"
@@ -55,18 +54,16 @@ UOctreeDynamicMeshComponent::UOctreeDynamicMeshComponent(const FObjectInitialize
 
 
 
-void UOctreeDynamicMeshComponent::InitializeMesh(const FMeshDescription* MeshDescription)
+void UOctreeDynamicMeshComponent::SetMesh(UE::Geometry::FDynamicMesh3&& MoveMesh)
 {
-	FMeshDescriptionToDynamicMesh Converter;
-	GetMesh()->Clear();
-	Converter.Convert(MeshDescription, *GetMesh());
+	MeshObject->SetMesh(MoveTemp(MoveMesh));
 
 	FAxisAlignedBox3d MeshBounds = GetMesh()->GetCachedBounds();
 	Octree = MakeUnique<FDynamicMeshOctree3>();
 	Octree->RootDimension = MeshBounds.MaxDim() * 0.25;
 
 	Octree->Initialize(GetMesh());
-	
+
 	FDynamicMeshOctree3::FStatistics Stats;
 	Octree->ComputeStatistics(Stats);
 	//UE_LOG(LogTemp, Warning, TEXT("OctreeStats %s"), *Stats.ToString());
@@ -93,8 +90,8 @@ void UOctreeDynamicMeshComponent::ApplyTransform(const UE::Geometry::FTransform3
 	if (GetCurrentSceneProxy() != nullptr)
 	{
 		Octree->ModifiedBounds = FAxisAlignedBox3d(
-			-TNumericLimits<float>::Max()*FVector3d::One(),
-			TNumericLimits<float>::Max()*FVector3d::One());
+			-TNumericLimits<float>::Max() * FVector3d::One(),
+			TNumericLimits<float>::Max() * FVector3d::One());
 		NotifyMeshUpdated();
 	}
 	else
@@ -110,22 +107,6 @@ void UOctreeDynamicMeshComponent::ApplyTransform(const UE::Geometry::FTransform3
 }
 
 
-void UOctreeDynamicMeshComponent::Bake(FMeshDescription* MeshDescription, bool bHaveModifiedTopology, const FConversionToMeshDescriptionOptions& ConversionOptions)
-{
-	FDynamicMeshToMeshDescription Converter(ConversionOptions);
-	if (bHaveModifiedTopology == false && Converter.HaveMatchingElementCounts(GetMesh(), MeshDescription))
-	{
-		Converter.Update(GetMesh(), *MeshDescription);
-	}
-	else
-	{
-		Converter.Convert(GetMesh(), *MeshDescription);
-
-		//UE_LOG(LogTemp, Warning, TEXT("MeshDescription has %d instances"), MeshDescription->VertexInstances().Num());
-	}
-}
-
-
 
 void UOctreeDynamicMeshComponent::NotifyMeshUpdated()
 {
diff --git a/Engine/Plugins/Experimental/MeshModelingToolset/Source/ModelingComponents/Private/OctreeDynamicMeshSceneProxy.h b/Engine/Plugins/Experimental/MeshModelingToolset/Source/ModelingComponents/Private/OctreeDynamicMeshSceneProxy.h
--- a/Engine/Plugins/Experimental/MeshModelingToolset/Source/ModelingComponents/Private/OctreeDynamicMeshSceneProxy.h
+++ b/Engine/Plugins/Experimental/MeshModelingToolset/Source/ModelingComponents/Private/OctreeDynamicMeshSceneProxy.h
@@ -65,6 +65,15 @@ public:
 
 
 
+	TUniqueFunction<void(int, int, int, const FVector3f&, FVector3f&, FVector3f&)> MakeTangentsFunc()
+	{
+		return [](int VertexID, int TriangleID, int TriVtxIdx, const FVector3f& Normal, FVector3f& TangentX, FVector3f& TangentY) -> void
+		{
+			VectorUtil::MakePerpVectors(Normal, TangentX, TangentY);
+		};
+	}
+
+
 
 	void InitializeSingleBuffer()
 	{
@@ -88,7 +97,7 @@ public:
 
 		InitializeBuffersFromOverlays(RenderBuffers, Mesh,
 			Mesh->TriangleCount(), Mesh->TriangleIndicesItr(),
-			UVOverlay, NormalOverlay, ColorOverlay);
+			UVOverlay, NormalOverlay, ColorOverlay, MakeTangentsFunc());
 
 		ENQUEUE_RENDER_COMMAND(FOctreeDynamicMeshSceneProxyInitializeSingle)(
 			[this, RenderBuffers](FRHICommandListImmediate& RHICmdList)
@@ -116,6 +125,7 @@ public:
 			NormalOverlay = Mesh->Attributes()->PrimaryNormals();
 			ColorOverlay = Mesh->Attributes()->PrimaryColors();
 		}
+		auto TangentsFunc = MakeTangentsFunc();
 
 		const TArray<int32>& SetIDs = Decomposition.GetIndexSetIDs();
 		for (int32 SetID : SetIDs)
@@ -127,7 +137,7 @@ public:
 
 			InitializeBuffersFromOverlays(RenderBuffers, Mesh,
 				Tris.Num(), Tris,
-				UVOverlay, NormalOverlay, ColorOverlay);
+				UVOverlay, NormalOverlay, ColorOverlay, TangentsFunc);
 
 			ENQUEUE_RENDER_COMMAND(FOctreeDynamicMeshSceneProxyInitializeFromDecomposition)(
 				[this, SetID, RenderBuffers](FRHICommandListImmediate& RHICmdList)
@@ -175,6 +185,7 @@ public:
 			NormalOverlay = Mesh->Attributes()->PrimaryNormals();
 			ColorOverlay = Mesh->Attributes()->PrimaryColors();
 		}
+		auto TangentsFunc = MakeTangentsFunc();
 
 		{
 			SCOPE_CYCLE_COUNTER(STAT_SculptToolOctree_UpdateDecompCreate);
@@ -189,7 +200,7 @@ public:
 
 				InitializeBuffersFromOverlays(RenderBuffers, Mesh,
 					Tris.Num(), Tris,
-					UVOverlay, NormalOverlay, ColorOverlay);
+					UVOverlay, NormalOverlay, ColorOverlay, TangentsFunc);
 
 				ENQUEUE_RENDER_COMMAND(FOctreeDynamicMeshSceneProxyUpdateAddOne)(
 					[this, SetID, RenderBuffers](FRHICommandListImmediate& RHICmdList)
diff --git a/Engine/Plugins/Experimental/MeshModelingToolset/Source/ModelingComponents/Private/PreviewMesh.cpp b/Engine/Plugins/Experimental/MeshModelingToolset/Source/ModelingComponents/Private/PreviewMesh.cpp
--- a/Engine/Plugins/Experimental/MeshModelingToolset/Source/ModelingComponents/Private/PreviewMesh.cpp
+++ b/Engine/Plugins/Experimental/MeshModelingToolset/Source/ModelingComponents/Private/PreviewMesh.cpp
@@ -6,6 +6,9 @@
 #include "Engine/Classes/Materials/Material.h"
 #include "TargetInterfaces/MaterialProvider.h" //FComponentMaterialSet
 
+#include "DynamicMeshToMeshDescription.h"
+#include "MeshDescriptionToDynamicMesh.h"
+
 #include "ExplicitUseGeometryMathTypes.h"		// using UE::Geometry::(math types)
 using namespace UE::Geometry;
 
@@ -170,16 +173,13 @@ void UPreviewMesh::DisableSecondaryTriangleBuffers()
 }
 
 
-void UPreviewMesh::SetTangentsMode(EDynamicMeshTangentCalcType TangentsType)
+void UPreviewMesh::SetTangentsMode(EDynamicMeshComponentTangentsMode TangentsType)
 {
 	check(DynamicMeshComponent);
-	DynamicMeshComponent->TangentsType = TangentsType;
+	DynamicMeshComponent->SetTangentsType(TangentsType);
 }
 
-const TMeshTangents<float>* UPreviewMesh::GetTangents() const
-{
-	return DynamicMeshComponent->GetTangents();
-}
+
 
 void UPreviewMesh::EnableWireframe(bool bEnable)
 {
@@ -264,6 +264,17 @@ const FDynamicMesh3* UPreviewMesh::GetMesh() const
 	return nullptr;
 }
 
+
+void UPreviewMesh::ProcessMesh(TFunctionRef<void(const UE::Geometry::FDynamicMesh3&)> ProcessFunc) const
+{
+	if (ensure(DynamicMeshComponent))
+	{
+		DynamicMeshComponent->ProcessMesh(ProcessFunc);
+	}
+}
+
+
+
 FDynamicMeshAABBTree3* UPreviewMesh::GetSpatial()
 {
 	if (DynamicMeshComponent != nullptr && bBuildSpatialDataStructure)
@@ -354,22 +365,6 @@ FVector3d UPreviewMesh::FindNearestPoint(const FVector3d& WorldPoint, bool bLine
 
 
 
-void UPreviewMesh::InitializeMesh(FMeshDescription* MeshDescription)
-{
-	check(DynamicMeshComponent != nullptr);
-	DynamicMeshComponent->InitializeMesh(MeshDescription);
-
-	if (bBuildSpatialDataStructure)
-	{
-		MeshAABBTree.SetMesh(DynamicMeshComponent->GetMesh(), true);
-	}
-	if (bDecompositionEnabled)
-	{
-		UpdateRenderMeshDecomposition();
-	}
-}
-
-
 void UPreviewMesh::ReplaceMesh(const FDynamicMesh3& NewMesh)
 {
 	ReplaceMesh(FDynamicMesh3(NewMesh));
@@ -377,10 +372,7 @@ void UPreviewMesh::ReplaceMesh(const FDynamicMesh3& NewMesh)
 
 void UPreviewMesh::ReplaceMesh(FDynamicMesh3&& NewMesh)
 {
-	FDynamicMesh3* Mesh = DynamicMeshComponent->GetMesh();
-	*Mesh = MoveTemp(NewMesh);
-
-	DynamicMeshComponent->NotifyMeshUpdated();
+	DynamicMeshComponent->SetMesh(MoveTemp(NewMesh));
 
 	if (bBuildSpatialDataStructure)
 	{
@@ -532,11 +524,6 @@ FSimpleMulticastDelegate& UPreviewMesh::GetOnMeshChanged()
 }
 
 
-void UPreviewMesh::Bake(FMeshDescription* MeshDescription, bool bHaveModifiedToplogy)
-{
-	check(DynamicMeshComponent != nullptr);
-	DynamicMeshComponent->Bake(MeshDescription, bHaveModifiedToplogy);
-}
 
 void UPreviewMesh::SetTriangleColorFunction(TFunction<FColor(const FDynamicMesh3*, int)> TriangleColorFunc, ERenderUpdateMode UpdateMode)
 {
diff --git a/Engine/Plugins/Experimental/MeshModelingToolset/Source/ModelingComponents/Private/SimpleDynamicMeshComponent.cpp b/Engine/Plugins/Experimental/MeshModelingToolset/Source/ModelingComponents/Private/SimpleDynamicMeshComponent.cpp
--- a/Engine/Plugins/Experimental/MeshModelingToolset/Source/ModelingComponents/Private/SimpleDynamicMeshComponent.cpp
+++ b/Engine/Plugins/Experimental/MeshModelingToolset/Source/ModelingComponents/Private/SimpleDynamicMeshComponent.cpp
@@ -21,7 +21,6 @@
 #include "DynamicMeshAttributeSet.h"
 #include "MeshNormals.h"
 #include "MeshDescriptionToDynamicMesh.h"
-#include "DynamicMeshToMeshDescription.h"
 
 #include "Changes/MeshVertexChange.h"
 #include "Changes/MeshChange.h"
@@ -59,7 +58,7 @@ USimpleDynamicMeshComponent::USimpleDynamicMeshComponent(const FObjectInitialize
 
 	MeshObjectChangedHandle = MeshObject->OnMeshChanged().AddUObject(this, &USimpleDynamicMeshComponent::OnMeshObjectChanged);
 
-	Tangents.SetMesh(GetMesh());
+	ResetProxy();
 }
 
 
@@ -70,49 +69,35 @@ void USimpleDynamicMeshComponent::PostLoad()
 	check(MeshObject != nullptr);
 	MeshObjectChangedHandle = MeshObject->OnMeshChanged().AddUObject(this, &USimpleDynamicMeshComponent::OnMeshObjectChanged);
 
-	Tangents.SetMesh(GetMesh());
-
 	ResetProxy();
 }
 
 
-FDynamicMesh3* USimpleDynamicMeshComponent::GetRenderMesh()
+
+void USimpleDynamicMeshComponent::SetMesh(UE::Geometry::FDynamicMesh3&& MoveMesh)
 {
-	if (RenderMeshPostProcessor && RenderMesh)
-	{
-		return RenderMesh.Get();
-	}
-	else
-	{
-		return GetMesh();
-	}
+	MeshObject->SetMesh(MoveTemp(MoveMesh));
 }
 
-const FDynamicMesh3* USimpleDynamicMeshComponent::GetRenderMesh() const
+
+void USimpleDynamicMeshComponent::ProcessMesh(
+	TFunctionRef<void(const UE::Geometry::FDynamicMesh3&)> ProcessFunc ) const
 {
-	if (RenderMeshPostProcessor && RenderMesh)
-	{
-		return RenderMesh.Get();
-	}
-	else
-	{
-		return GetMesh();
-	}
+	MeshObject->ProcessMesh(ProcessFunc);
 }
 
-void USimpleDynamicMeshComponent::InitializeMesh(const FMeshDescription* MeshDescription)
+
+void USimpleDynamicMeshComponent::EditMesh(TFunctionRef<void(UE::Geometry::FDynamicMesh3&)> EditFunc,
+										   EDynamicMeshComponentRenderUpdateMode UpdateMode )
 {
-	FMeshDescriptionToDynamicMesh Converter;
-	MeshObject->Reset();
-	Converter.Convert(MeshDescription, *GetMesh());
-	if (TangentsType == EDynamicMeshTangentCalcType::ExternallyCalculated)
+	MeshObject->EditMesh(EditFunc);
+	if (UpdateMode != EDynamicMeshComponentRenderUpdateMode::NoUpdate)
 	{
-		Converter.CopyTangents(MeshDescription, GetMesh(), &Tangents);
+		NotifyMeshUpdated();
 	}
-
-	NotifyMeshUpdated();
 }
 
+
 void USimpleDynamicMeshComponent::SetRenderMeshPostProcessor(TUniquePtr<IRenderMeshPostProcessor> Processor)
 {
 	RenderMeshPostProcessor = MoveTemp(Processor);
@@ -130,102 +115,94 @@ void USimpleDynamicMeshComponent::SetRenderMeshPostProcessor(TUniquePtr<IRenderM
 	}
 }
 
-void USimpleDynamicMeshComponent::UpdateTangents(const FMeshTangentsf* ExternalTangents, bool bFastUpdateIfPossible)
+FDynamicMesh3* USimpleDynamicMeshComponent::GetRenderMesh()
 {
-	Tangents.CopyTriVertexTangents(*ExternalTangents);
-	if (bFastUpdateIfPossible)
+	if (RenderMeshPostProcessor && RenderMesh)
 	{
-		FastNotifyVertexAttributesUpdated(EMeshRenderAttributeFlags::VertexNormals);
+		return RenderMesh.Get();
 	}
 	else
 	{
-		NotifyMeshUpdated();
+		return GetMesh();
 	}
 }
 
-void USimpleDynamicMeshComponent::UpdateTangents(const FMeshTangentsd* ExternalTangents, bool bFastUpdateIfPossible)
+const FDynamicMesh3* USimpleDynamicMeshComponent::GetRenderMesh() const
 {
-	Tangents.CopyTriVertexTangents(*ExternalTangents);
-	if (bFastUpdateIfPossible)
+	if (RenderMeshPostProcessor && RenderMesh)
 	{
-		FastNotifyVertexAttributesUpdated(EMeshRenderAttributeFlags::VertexNormals);
+		return RenderMesh.Get();
 	}
 	else
 	{
-		NotifyMeshUpdated();
+		return GetMesh();
 	}
 }
 
 
 
+
 void USimpleDynamicMeshComponent::ApplyTransform(const UE::Geometry::FTransform3d& Transform, bool bInvert)
 {
-	if (bInvert)
+	MeshObject->EditMesh([&](FDynamicMesh3& EditMesh)
 	{
-		MeshTransforms::ApplyTransformInverse(*GetMesh(), Transform);
-	}
-	else
-	{
-		MeshTransforms::ApplyTransform(*GetMesh(), Transform);
-	}
-
-	NotifyMeshUpdated();
+		if (bInvert)
+		{
+			MeshTransforms::ApplyTransformInverse(EditMesh, Transform);
+		}
+		else
+		{
+			MeshTransforms::ApplyTransform(EditMesh, Transform);
+		}
+	}, EDynamicMeshChangeType::DeformationEdit);
 }
 
 
-void USimpleDynamicMeshComponent::Bake(FMeshDescription* MeshDescription, bool bHaveModifiedTopology, const FConversionToMeshDescriptionOptions& ConversionOptions)
+
+void USimpleDynamicMeshComponent::SetTangentsType(EDynamicMeshComponentTangentsMode NewTangentsType)
 {
-	FDynamicMeshToMeshDescription Converter(ConversionOptions);
-	if (!bHaveModifiedTopology)
-	{
-		Converter.UpdateUsingConversionOptions(GetMesh(), *MeshDescription);
-	}
-	else
+	if (NewTangentsType != TangentsType)
 	{
-		Converter.Convert(GetMesh(), *MeshDescription);
+		TangentsType = NewTangentsType;
+		InvalidateAutoCalculatedTangents();
 	}
 }
 
+void USimpleDynamicMeshComponent::InvalidateAutoCalculatedTangents() 
+{ 
+	bAutoCalculatedTangentsValid = false; 
+}
 
-
-
-
-
-const FMeshTangentsf* USimpleDynamicMeshComponent::GetTangents()
-{
-	if (TangentsType == EDynamicMeshTangentCalcType::NoTangents)
+const UE::Geometry::FMeshTangentsf* USimpleDynamicMeshComponent::GetAutoCalculatedTangents() 
+{ 
+	if (TangentsType == EDynamicMeshComponentTangentsMode::AutoCalculated)
 	{
-		return nullptr;
+		UpdateAutoCalculatedTangents();
+		return (bAutoCalculatedTangentsValid) ? &AutoCalculatedTangents : nullptr;
 	}
-	
-	if (TangentsType == EDynamicMeshTangentCalcType::AutoCalculated)
+	return nullptr;
+}
+
+void USimpleDynamicMeshComponent::UpdateAutoCalculatedTangents()
+{
+	if (TangentsType == EDynamicMeshComponentTangentsMode::AutoCalculated && bAutoCalculatedTangentsValid == false)
 	{
-		if (bTangentsValid == false && GetMesh()->HasAttributes())
+		GetDynamicMesh()->ProcessMesh([&](const FDynamicMesh3& Mesh)
 		{
-			FDynamicMeshUVOverlay* UVOverlay = GetMesh()->Attributes()->PrimaryUV();
-			FDynamicMeshNormalOverlay* NormalOverlay = GetMesh()->Attributes()->PrimaryNormals();
-			if (UVOverlay != nullptr && NormalOverlay != nullptr)
-			{
-				Tangents.ComputeTriVertexTangents(NormalOverlay, UVOverlay, FComputeTangentsOptions());
-				bTangentsValid = true;
-			}
-		}
-		return (bTangentsValid) ? &Tangents : nullptr;
-	}
+			const FDynamicMeshUVOverlay* UVOverlay = Mesh.Attributes()->PrimaryUV();
+			const FDynamicMeshNormalOverlay* NormalOverlay = Mesh.Attributes()->PrimaryNormals();
+			AutoCalculatedTangents.SetMesh(&Mesh);
+			AutoCalculatedTangents.ComputeTriVertexTangents(NormalOverlay, UVOverlay, FComputeTangentsOptions());
+			AutoCalculatedTangents.SetMesh(nullptr);
+		});
 
-	// in this mode we assume the tangents are valid
-	ensure(TangentsType == EDynamicMeshTangentCalcType::ExternallyCalculated);
-	if (TangentsType == EDynamicMeshTangentCalcType::ExternallyCalculated)
-	{
-		// if you hit this, you did not request ExternallyCalculated tangents before initializing this PreviewMesh
-		ensure(Tangents.GetTangents().Num() > 0);
+		bAutoCalculatedTangentsValid = true;
 	}
-
-	return &Tangents;
 }
 
 
 
+
 void USimpleDynamicMeshComponent::UpdateLocalBounds()
 {
 	LocalBounds = GetMesh()->GetBounds(true);
@@ -248,17 +225,13 @@ FSimpleDynamicMeshSceneProxy* USimpleDynamicMeshComponent::GetCurrentSceneProxy(
 void USimpleDynamicMeshComponent::ResetProxy()
 {
 	bProxyValid = false;
+	InvalidateAutoCalculatedTangents();
 
 	// Need to recreate scene proxy to send it over
 	MarkRenderStateDirty();
 	UpdateLocalBounds();
 	UpdateBounds();
 
-	if (TangentsType != EDynamicMeshTangentCalcType::ExternallyCalculated)
-	{
-		bTangentsValid = false;
-	}
-
 	// this seems speculative, ie we may not actually have a mesh update, but we currently ResetProxy() in lots
 	// of places where that is what it means
 	GetDynamicMesh()->PostRealtimeUpdate();
diff --git a/Engine/Plugins/Experimental/MeshModelingToolset/Source/ModelingComponents/Private/SimpleDynamicMeshSceneProxy.h b/Engine/Plugins/Experimental/MeshModelingToolset/Source/ModelingComponents/Private/SimpleDynamicMeshSceneProxy.h
--- a/Engine/Plugins/Experimental/MeshModelingToolset/Source/ModelingComponents/Private/SimpleDynamicMeshSceneProxy.h
+++ b/Engine/Plugins/Experimental/MeshModelingToolset/Source/ModelingComponents/Private/SimpleDynamicMeshSceneProxy.h
@@ -4,6 +4,7 @@
 
 #include "SimpleDynamicMeshComponent.h"
 #include "BaseDynamicMeshSceneProxy.h"
+#include "DynamicMesh/MeshTangents.h"
 #include "Async/ParallelFor.h"
 
 /**
@@ -84,6 +85,37 @@ public:
 
 
 
+	TUniqueFunction<void(int, int, int, const FVector3f&, FVector3f&, FVector3f&)> MakeTangentsFunc(bool bSkipAutoCompute = false)
+	{
+		EDynamicMeshComponentTangentsMode TangentsType = ParentComponent->GetTangentsType();
+		if (TangentsType == EDynamicMeshComponentTangentsMode::ExternallyProvided)
+		{
+			UE::Geometry::FDynamicMeshTangents Tangents(ParentComponent->GetRenderMesh());
+			return [Tangents](int VertexID, int TriangleID, int TriVtxIdx, const FVector3f& Normal, FVector3f& TangentX, FVector3f& TangentY) -> void
+			{
+				Tangents.GetTangentVectors(TriangleID, TriVtxIdx, Normal, TangentX, TangentY);
+			};
+		}
+		else if (TangentsType == EDynamicMeshComponentTangentsMode::AutoCalculated && bSkipAutoCompute == false)
+		{
+			const UE::Geometry::FMeshTangentsf* Tangents = ParentComponent->GetAutoCalculatedTangents();
+			if (Tangents != nullptr )
+			{
+				return [Tangents](int VertexID, int TriangleID, int TriVtxIdx, const FVector3f& Normal, FVector3f& TangentX, FVector3f& TangentY) -> void
+				{
+					Tangents->GetTriangleVertexTangentVectors(TriangleID, TriVtxIdx, TangentX, TangentY);
+				};
+			}
+		}
+
+		// fallback to orthogonal basis
+		return [](int VertexID, int TriangleID, int TriVtxIdx, const FVector3f& Normal, FVector3f& TangentX, FVector3f& TangentY) -> void
+		{
+			VectorUtil::MakePerpVectors(Normal, TangentX, TangentY);
+		};
+	}
+
+
 	/**
 	 * Initialize multiple buffers for the mesh based on the given Decomposition
 	 */
@@ -104,21 +136,12 @@ public:
 
 		bIsSingleBuffer = false;
 
-		FDynamicMesh3* Mesh = ParentComponent->GetRenderMesh();
+		const FDynamicMesh3* Mesh = ParentComponent->GetRenderMesh();
 		// find suitable overlays
-		FDynamicMeshUVOverlay* UVOverlay = Mesh->Attributes()->PrimaryUV();
-		FDynamicMeshNormalOverlay* NormalOverlay = Mesh->Attributes()->PrimaryNormals();
-		FDynamicMeshColorOverlay* ColorOverlay = Mesh->Attributes()->PrimaryColors();
-
-		TFunction<void(int, int, int, FVector3f&, FVector3f&)> TangentsFunc = nullptr;
-		const UE::Geometry::FMeshTangentsf* Tangents = ParentComponent->GetTangents();
-		if (Tangents != nullptr)
-		{
-			TangentsFunc = [Tangents](int VertexID, int TriangleID, int TriVtxIdx, FVector3f& TangentX, FVector3f& TangentY)
-			{
-				return Tangents->GetPerTriangleTangent(TriangleID, TriVtxIdx, TangentX, TangentY);
-			};
-		}
+		const FDynamicMeshUVOverlay* UVOverlay = Mesh->Attributes()->PrimaryUV();
+		const FDynamicMeshNormalOverlay* NormalOverlay = Mesh->Attributes()->PrimaryNormals();
+		const FDynamicMeshColorOverlay* ColorOverlay = Mesh->Attributes()->PrimaryColors();
+		TUniqueFunction<void(int, int, int, const FVector3f&, FVector3f&, FVector3f&)> TangentsFunc = MakeTangentsFunc();
 
 		// init renderbuffers for each set
 		ParallelFor(NumSets, [&](int32 SetIndex)
@@ -171,16 +194,7 @@ public:
 			}
 			ColorOverlay = Attributes->PrimaryColors();
 		}
-
-		TFunction<void(int, int, int, FVector3f&, FVector3f&)> TangentsFunc = nullptr;
-		const UE::Geometry::FMeshTangentsf* Tangents = ParentComponent->GetTangents();
-		if (Tangents != nullptr)
-		{
-			TangentsFunc = [Tangents](int VertexID, int TriangleID, int TriVtxIdx, FVector3f& TangentX, FVector3f& TangentY)
-			{
-				return Tangents->GetPerTriangleTangent(TriangleID, TriVtxIdx, TangentX, TangentY);
-			};
-		}
+		TUniqueFunction<void(int, int, int, const FVector3f&, FVector3f&, FVector3f&)> TangentsFunc = MakeTangentsFunc();
 
 		InitializeBuffersFromOverlays(RenderBuffers, Mesh,
 			Mesh->TriangleCount(), Mesh->TriangleIndicesItr(),
@@ -220,15 +234,7 @@ public:
 			UVOverlays[k] = Attributes->GetUVLayer(k);
 		}
 
-		TFunction<void(int, int, int, FVector3f&, FVector3f&)> TangentsFunc = nullptr;
-		const UE::Geometry::FMeshTangentsf* Tangents = ParentComponent->GetTangents();
-		if (Tangents != nullptr)
-		{
-			TangentsFunc = [Tangents](int VertexID, int TriangleID, int TriVtxIdx, FVector3f& TangentX, FVector3f& TangentY)
-			{
-				return Tangents->GetPerTriangleTangent(TriangleID, TriVtxIdx, TangentX, TangentY);
-			};
-		}
+		TUniqueFunction<void(int, int, int, const FVector3f&, FVector3f&, FVector3f&)> TangentsFunc = MakeTangentsFunc();
 
 		// count number of triangles for each material (could do in parallel?)
 		int NumMaterials = BufferSets.Num();
@@ -371,20 +377,9 @@ public:
 
 		// find suitable overlays and attributes
 		FDynamicMeshNormalOverlay* NormalOverlay = nullptr;
-		TFunction<void(int, int, int, const FVector3f&, FVector3f&, FVector3f&)> TangentsFunc =
-			[](int, int, int, const FVector3f& Normal, FVector3f& TangentX, FVector3f& TangentY) { VectorUtil::MakePerpVectors(Normal, TangentX, TangentY); };
 		if (bNormals && ensure(Mesh->HasAttributes()) )
 		{
 			NormalOverlay = Mesh->Attributes()->PrimaryNormals();
-
-			const UE::Geometry::FMeshTangentsf* Tangents = ParentComponent->GetTangents();
-			if (Tangents != nullptr)
-			{
-				TangentsFunc = [Tangents](int VertexID, int TriangleID, int TriVtxIdx, const FVector3f& Normal, FVector3f& TangentX, FVector3f& TangentY)
-				{
-					return Tangents->GetPerTriangleTangent(TriangleID, TriVtxIdx, TangentX, TangentY);
-				};
-			}
 		}
 		FDynamicMeshUVOverlay* UVOVerlay = nullptr;
 		if (bUVs && ensure(Mesh->HasAttributes()) )
@@ -395,8 +390,11 @@ public:
 		if (bColors && ensure(Mesh->HasAttributes()) )
 		{
 			ColorOverlay = Mesh->Attributes()->PrimaryColors();
-					
 		}
+
+		// Currently deferring tangents auto-computation if doing fast vertex update (maybe needs to be exposed?)
+		TUniqueFunction<void(int, int, int, const FVector3f&, FVector3f&, FVector3f&)> TangentsFunc = MakeTangentsFunc(true);
+
 		if (bIsSingleBuffer)
 		{
 			if (ensure(RenderBufferSets.Num() == 1))
@@ -475,36 +473,28 @@ public:
 		// This needs to be rewritten for split-by-material buffers.
 		// Could store triangle set with each buffer, and then rebuild vtx buffer(s) as needed?
 
-		FDynamicMesh3* Mesh = ParentComponent->GetRenderMesh();
+		const FDynamicMesh3* Mesh = ParentComponent->GetRenderMesh();
 
 		// find suitable overlays
-		FDynamicMeshNormalOverlay* NormalOverlay = nullptr;
-		TFunction<void(int, int, int, const FVector3f&, FVector3f&, FVector3f&)> TangentsFunc =
-			[](int, int, int, const FVector3f& Normal, FVector3f& TangentX, FVector3f& TangentY) { VectorUtil::MakePerpVectors(Normal, TangentX, TangentY); };
+		const FDynamicMeshNormalOverlay* NormalOverlay = nullptr;
 		if (bNormals && ensure(Mesh->HasAttributes()))
 		{
 			NormalOverlay = Mesh->Attributes()->PrimaryNormals();
-
-			const UE::Geometry::FMeshTangentsf* Tangents = ParentComponent->GetTangents();
-			if (Tangents != nullptr)
-			{
-				TangentsFunc = [Tangents](int VertexID, int TriangleID, int TriVtxIdx, const FVector3f& Normal, FVector3f& TangentX, FVector3f& TangentY)
-				{
-					return Tangents->GetPerTriangleTangent(TriangleID, TriVtxIdx, TangentX, TangentY);
-				};
-			}
 		}
-		FDynamicMeshUVOverlay* UVOVerlay = nullptr;
+		const FDynamicMeshUVOverlay* UVOVerlay = nullptr;
 		if (bUVs && ensure(Mesh->HasAttributes()))
 		{
 			UVOVerlay = Mesh->Attributes()->PrimaryUV();
 		}
-		FDynamicMeshColorOverlay* ColorOverlay = nullptr;
+		const FDynamicMeshColorOverlay* ColorOverlay = nullptr;
 		if (bColors && ensure(Mesh->HasAttributes()))
 		{
 			ColorOverlay = Mesh->Attributes()->PrimaryColors();
-
 		}
+
+		// Currently deferring tangents auto-computation if doing fast vertex update (maybe needs to be exposed?)
+		TUniqueFunction<void(int, int, int, const FVector3f&, FVector3f&, FVector3f&)> TangentsFunc = MakeTangentsFunc(true);
+
 		ParallelFor(WhichBuffers.Num(), [&](int idx)
 		{
 			int32 BufferIndex = WhichBuffers[idx];
diff --git a/Engine/Plugins/Experimental/MeshModelingToolset/Source/ModelingComponents/Private/UDynamicMesh.cpp b/Engine/Plugins/Experimental/MeshModelingToolset/Source/ModelingComponents/Private/UDynamicMesh.cpp
--- a/Engine/Plugins/Experimental/MeshModelingToolset/Source/ModelingComponents/Private/UDynamicMesh.cpp
+++ b/Engine/Plugins/Experimental/MeshModelingToolset/Source/ModelingComponents/Private/UDynamicMesh.cpp
@@ -106,19 +106,29 @@ void UDynamicMesh::ProcessMesh(TFunctionRef<void(const UE::Geometry::FDynamicMes
 	ProcessFunc(GetMeshRef());
 }
 
-void UDynamicMesh::EditMesh(TFunctionRef<void(FDynamicMesh3&)> EditFunc)
+void UDynamicMesh::EditMesh(TFunctionRef<void(FDynamicMesh3&)> EditFunc,
+							EDynamicMeshChangeType ChangeType,
+							EDynamicMeshAttributeChangeFlags ChangeFlags,
+							bool bDeferChangeEvents)
 {
 	FDynamicMeshChangeInfo ChangeInfo;
-	ChangeInfo.Type = EDynamicMeshChangeType::GeneralEdit;
-	EditMeshInternal(EditFunc, ChangeInfo);
+	ChangeInfo.Type = ChangeType;
+	ChangeInfo.Flags = ChangeFlags;
+	EditMeshInternal(EditFunc, ChangeInfo, bDeferChangeEvents);
 }
 
-void UDynamicMesh::EditMeshInternal(TFunctionRef<void(FDynamicMesh3&)> EditFunc, const FDynamicMeshChangeInfo& ChangeInfo)
+void UDynamicMesh::EditMeshInternal(TFunctionRef<void(FDynamicMesh3&)> EditFunc, const FDynamicMeshChangeInfo& ChangeInfo, bool bDeferChangeEvents)
 {
-	PreMeshChangedEvent.Broadcast(this, ChangeInfo);
+	if (!bDeferChangeEvents)
+	{
+		PreMeshChangedEvent.Broadcast(this, ChangeInfo);
+	}
 	EditFunc(GetMeshRef());
-	MeshChangedEvent.Broadcast(this, ChangeInfo);
-	MeshModifiedBPEvent.Broadcast(this);
+	if (!bDeferChangeEvents)
+	{
+		MeshChangedEvent.Broadcast(this, ChangeInfo);
+		MeshModifiedBPEvent.Broadcast(this);
+	}
 }
 
 
diff --git a/Engine/Plugins/Experimental/MeshModelingToolset/Source/ModelingComponents/Public/BaseDynamicMeshComponent.h b/Engine/Plugins/Experimental/MeshModelingToolset/Source/ModelingComponents/Public/BaseDynamicMeshComponent.h
--- a/Engine/Plugins/Experimental/MeshModelingToolset/Source/ModelingComponents/Public/BaseDynamicMeshComponent.h
+++ b/Engine/Plugins/Experimental/MeshModelingToolset/Source/ModelingComponents/Public/BaseDynamicMeshComponent.h
@@ -43,14 +43,14 @@ ENUM_CLASS_FLAGS(EMeshRenderAttributeFlags);
  * Tangent calculation modes
  */
 UENUM()
-enum class EDynamicMeshTangentCalcType : uint8
+enum class EDynamicMeshComponentTangentsMode : uint8
 {
 	/** Tangents are not used/available, proceed accordingly (eg generate arbitrary orthogonal basis) */
 	NoTangents,
-	/** Tangents should be automatically calculated on demand */
+	/** Tangents will be automatically calculated on demand. Note that mesh changes due to tangents calculation will *not* be broadcast via MeshChange events! */
 	AutoCalculated,
-	/** Tangents are externally calculated (behavior undefined if they are not actually externally calculated!) */
-	ExternallyCalculated
+	/** Tangents are externally provided via the FDynamicMesh3 AttributeSet */
+	ExternallyProvided
 };
 
 
@@ -73,10 +73,9 @@ class MODELINGCOMPONENTS_API UBaseDynamicMeshComponent :
 	//
 public:
 	/**
-	 * initialize the internal mesh from a MeshDescription
-	 * @warning avoid usage of this function, access via GetDynamicMesh() instead
+	 * initialize the internal mesh from a DynamicMesh
 	 */
-	virtual void InitializeMesh(const FMeshDescription* MeshDescription)
+	virtual void SetMesh(UE::Geometry::FDynamicMesh3&& MoveMesh)
 	{
 		unimplemented();
 	}
@@ -100,6 +99,14 @@ public:
 		return nullptr;
 	}
 
+	/**
+	 * Allow external code to read the internal mesh.
+	 */
+	virtual void ProcessMesh(TFunctionRef<void(const UE::Geometry::FDynamicMesh3&)> ProcessFunc) const
+	{
+		unimplemented();
+	}
+
 	/**
 	 * @return the child UDynamicMesh
 	 */
@@ -151,15 +158,6 @@ public:
 		unimplemented();
 	}
 
-	/**
-	 * Write the internal mesh to a MeshDescription
-	 * @param bHaveModifiedTopology if false, we only update the vertex positions in the MeshDescription, otherwise it is Empty()'d and regenerated entirely
-	 * @param ConversionOptions struct of additional options for the conversion
-	 */
-	virtual void Bake(FMeshDescription* MeshDescription, bool bHaveModifiedTopology, const FConversionToMeshDescriptionOptions& ConversionOptions)
-	{
-		unimplemented();
-	}
 
 protected:
 	/**
diff --git a/Engine/Plugins/Experimental/MeshModelingToolset/Source/ModelingComponents/Public/ModelingToolTargetUtil.h b/Engine/Plugins/Experimental/MeshModelingToolset/Source/ModelingComponents/Public/ModelingToolTargetUtil.h
--- a/Engine/Plugins/Experimental/MeshModelingToolset/Source/ModelingComponents/Public/ModelingToolTargetUtil.h
+++ b/Engine/Plugins/Experimental/MeshModelingToolset/Source/ModelingComponents/Public/ModelingToolTargetUtil.h
@@ -5,6 +5,7 @@
 #include "CoreMinimal.h"
 #include "DynamicMesh/DynamicMesh3.h"
 #include "TargetInterfaces/MaterialProvider.h" // FComponentMaterialSet
+#include "MeshConversionOptions.h"
 
 class UToolTarget;
 class UPrimitiveComponent;
@@ -62,6 +63,16 @@ MODELINGCOMPONENTS_API UE::Geometry::FTransform3d GetLocalToWorldTransform(UTool
 MODELINGCOMPONENTS_API FComponentMaterialSet GetMaterialSet(UToolTarget* Target, bool bPreferAssetMaterials = false);
 
 
+/**
+ * Update the material set of the Target
+ * @param bApplyToAsset In situations where the Target has both Component-level and Asset-level materials, this specifies which should be updated (this flag is passed to the IMaterialProvider, which may or may not respect it)
+ */
+MODELINGCOMPONENTS_API bool CommitMaterialSetUpdate(
+	UToolTarget* Target,
+	const FComponentMaterialSet& UpdatedMaterials,
+	bool bApplyToAsset = true);
+
+
 /**
  * @return the MeshDescription underlying a ToolTarget, if it has such a mesh. May be generated internally by the ToolTarget. May be nullptr if the Target does not have a mesh.
  */
@@ -70,9 +81,10 @@ MODELINGCOMPONENTS_API const FMeshDescription* GetMeshDescription(UToolTarget* T
 /**
  * Fetch a DynamicMesh3 representing the given ToolTarget. This may be a conersion of the output of GetMeshDescription().
  * This function returns a copy, so the caller can take ownership of this Mesh.
+ * @param bWantMeshTangents if true, tangents will be returned if the target has them available. This may require that they be auto-calculated in some cases (which may be expensive)
  * @return a created DynamicMesh3, which may be empty if the Target doesn't have a mesh 
  */
-MODELINGCOMPONENTS_API UE::Geometry::FDynamicMesh3 GetDynamicMeshCopy(UToolTarget* Target);
+MODELINGCOMPONENTS_API UE::Geometry::FDynamicMesh3 GetDynamicMeshCopy(UToolTarget* Target, bool bWantMeshTangents = false);
 
 
 /**
@@ -98,6 +110,19 @@ MODELINGCOMPONENTS_API EDynamicMeshUpdateResult CommitMeshDescriptionUpdate(
 	const FMeshDescription* UpdatedMesh, 
 	const FComponentMaterialSet* UpdatedMaterials = nullptr);
 
+/**
+ * Update the Mesh in a ToolTarget based on the provided DynamicMesh, and optional material set
+ * @param bHaveModifiedTopology If the update only changes vertex or attribute values (but not counts), then in some cases a more efficient and/or less destructive update can be applied to the Target
+ * @param ConversionOptions if the commit to the Target involves conversion to MeshDescription, these options can configure that conversion
+ * @param UpdatedMaterials optional new material set that will be applied to the updated Target, and the Target Asset if available. If more control is needed use CommitMaterialSetUpdate()
+ * @return EDynamicMeshUpdateResult::Ok on success
+ */
+MODELINGCOMPONENTS_API EDynamicMeshUpdateResult CommitDynamicMeshUpdate(
+	UToolTarget* Target,
+	const UE::Geometry::FDynamicMesh3& UpdatedMesh,
+	bool bHaveModifiedTopology,
+	const FConversionToMeshDescriptionOptions& ConversionOptions = FConversionToMeshDescriptionOptions(),
+	const FComponentMaterialSet* UpdatedMaterials = nullptr);
 
 /**
  * Update the UV sets of the ToolTarget's mesh (assuming it has one) based on the provided UpdatedMesh.
diff --git a/Engine/Plugins/Experimental/MeshModelingToolset/Source/ModelingComponents/Public/OctreeDynamicMeshComponent.h b/Engine/Plugins/Experimental/MeshModelingToolset/Source/ModelingComponents/Public/OctreeDynamicMeshComponent.h
--- a/Engine/Plugins/Experimental/MeshModelingToolset/Source/ModelingComponents/Public/OctreeDynamicMeshComponent.h
+++ b/Engine/Plugins/Experimental/MeshModelingToolset/Source/ModelingComponents/Public/OctreeDynamicMeshComponent.h
@@ -39,9 +39,9 @@ class MODELINGCOMPONENTS_API UOctreeDynamicMeshComponent : public UBaseDynamicMe
 
 public:
 	/**
-	 * initialize the internal mesh from a MeshDescription
+	 * initialize the internal mesh from a DynamicMesh
 	 */
-	virtual void InitializeMesh(const FMeshDescription* MeshDescription) override;
+	virtual void SetMesh(UE::Geometry::FDynamicMesh3&& MoveMesh) override;
 
 	/**
 	 * @return pointer to internal mesh
@@ -54,24 +54,15 @@ public:
 	virtual const FDynamicMesh3* GetMesh() const override { return MeshObject->GetMeshPtr(); }
 
 
+	virtual void ProcessMesh(TFunctionRef<void(const UE::Geometry::FDynamicMesh3&)> ProcessFunc) const
+	{
+		MeshObject->ProcessMesh(ProcessFunc);
+	}
+
+
 	UE::Geometry::FDynamicMeshOctree3* GetOctree() { return Octree.Get(); }
 
-	/**
-	 * Write the internal mesh to a MeshDescription
-	 * @param bHaveModifiedTopology if false, we only update the vertex positions in the MeshDescription, otherwise it is Empty()'d and regenerated entirely
-	 * @param ConversionOptions struct of additional options for the conversion
-	 */
-	virtual void Bake(FMeshDescription* MeshDescription, bool bHaveModifiedTopology, const FConversionToMeshDescriptionOptions& ConversionOptions) override;
 
-	/**
-	* Write the internal mesh to a MeshDescription with default conversion options
-	* @param bHaveModifiedTopology if false, we only update the vertex positions in the MeshDescription, otherwise it is Empty()'d and regenerated entirely
-	*/
-	void Bake(FMeshDescription* MeshDescription, bool bHaveModifiedTopology)
-	{
-		FConversionToMeshDescriptionOptions ConversionOptions;
-		Bake(MeshDescription, bHaveModifiedTopology, ConversionOptions);
-	}
 
 	/**
 	 * Apply transform to internal mesh. Invalidates RenderProxy.
diff --git a/Engine/Plugins/Experimental/MeshModelingToolset/Source/ModelingComponents/Public/PreviewMesh.h b/Engine/Plugins/Experimental/MeshModelingToolset/Source/ModelingComponents/Public/PreviewMesh.h
--- a/Engine/Plugins/Experimental/MeshModelingToolset/Source/ModelingComponents/Public/PreviewMesh.h
+++ b/Engine/Plugins/Experimental/MeshModelingToolset/Source/ModelingComponents/Public/PreviewMesh.h
@@ -171,19 +171,7 @@ public:
 	 * Set the tangents mode for the underlying component, if available. 
 	 * Note that this function may need to be called before the mesh is initialized.
 	 */
-	void SetTangentsMode(EDynamicMeshTangentCalcType TangentsType);
-
-	/**
-	 * @return a MeshTangents data structure for the underlying component, if available, otherwise nullptr
-	 */
-	const UE::Geometry::TMeshTangents<float>* GetTangents() const;
-
-	/**
-	 * Update tangents on internal Mesh Component
-	 */
-	template<typename RealType>
-	void UpdateTangents(const UE::Geometry::TMeshTangents<RealType>* ExternalTangents, bool bFastUpdateIfPossible);
-
+	void SetTangentsMode(EDynamicMeshComponentTangentsMode TangentsType);
 
 
 	/**
@@ -293,21 +281,24 @@ public:
 	void UpdatePreview(FDynamicMesh3&& Mesh, ERenderUpdateMode UpdateMode = ERenderUpdateMode::FullUpdate,
 		EMeshRenderAttributeFlags ModifiedAttribs = EMeshRenderAttributeFlags::AllVertexAttribs);
 
-	/**
-	 * Initialize the internal mesh based on the given MeshDescription
-	 */
-	void InitializeMesh(FMeshDescription* MeshDescription);
-
 	/**
 	* @return pointer to the current FDynamicMesh used for preview  @todo deprecate this function, use GetMesh() instead
 	*/
 	const FDynamicMesh3* GetPreviewDynamicMesh() const { return GetMesh(); }
 
 	/**
+	* Read access to the internal mesh. This function will be deprecated/removed, use ProcessMesh() instead.
 	* @return pointer to the current FDynamicMesh used for preview
 	*/
 	const FDynamicMesh3* GetMesh() const;
 
+
+	/**
+	 * Give external code direct read access to the internal FDynamicMesh3. 
+	 * This should be used preferentially over GetMesh() / GetPreviewDynamicMesh()
+	 */
+	virtual void ProcessMesh(TFunctionRef<void(const UE::Geometry::FDynamicMesh3&)> ProcessFunc) const;
+
 	/**
 	 * @return point to the current AABBTree used for preview spatial mesh, or nullptr if not available
 	 * @warning this has to return non-const because of current FDynamicMeshAABBTree3 API, but you should not modify this!
@@ -319,12 +310,6 @@ public:
 	 */
 	TUniquePtr<FDynamicMesh3> ExtractPreviewMesh() const;
 
-	/**
-	 * Write the internal mesh to a MeshDescription
-	 * @param bHaveModifiedToplogy if false, we only update the vertex positions in the MeshDescription, otherwise it is Empty()'d and regenerated entirely
-	 */
-	void Bake(FMeshDescription* MeshDescription, bool bHaveModifiedToplogy);
-
 
 
 
@@ -450,12 +435,3 @@ protected:
 };
 
 
-
-template<typename RealType>
-void UPreviewMesh::UpdateTangents(const UE::Geometry::TMeshTangents<RealType>* ExternalTangents, bool bFastUpdateIfPossible)
-{
-	if (ensure(DynamicMeshComponent))
-	{
-		DynamicMeshComponent->UpdateTangents(ExternalTangents, bFastUpdateIfPossible);
-	}
-}
\ No newline at end of file
diff --git a/Engine/Plugins/Experimental/MeshModelingToolset/Source/ModelingComponents/Public/SimpleDynamicMeshComponent.h b/Engine/Plugins/Experimental/MeshModelingToolset/Source/ModelingComponents/Public/SimpleDynamicMeshComponent.h
--- a/Engine/Plugins/Experimental/MeshModelingToolset/Source/ModelingComponents/Public/SimpleDynamicMeshComponent.h
+++ b/Engine/Plugins/Experimental/MeshModelingToolset/Source/ModelingComponents/Public/SimpleDynamicMeshComponent.h
@@ -66,8 +66,8 @@ class MODELINGCOMPONENTS_API USimpleDynamicMeshComponent : public UBaseDynamicMe
 
 
 	//===============================================================================================================
-	// Mesh Access. Usage via GetDynamicMesh() is preferred, the GetMesh() pass-throughs exist largely to support
-	// existing code from before UDynamicMesh was added.
+	// Mesh Access. Usage via GetDynamicMesh() or SetMesh()/ProcessMesh()/EditMesh() is preferred, the GetMesh() 
+	// pointer access exist largely to support existing code from before UDynamicMesh was added.
 public:
 	/**
 	 * @return pointer to internal mesh
@@ -94,6 +94,30 @@ public:
 	UFUNCTION(BlueprintCallable, Category = "Dynamic Mesh Component")
 	void SetDynamicMesh(UDynamicMesh* NewMesh);
 
+	/**
+	 * initialize the internal mesh from a DynamicMesh
+	 */
+	virtual void SetMesh(UE::Geometry::FDynamicMesh3&& MoveMesh) override;
+
+	/**
+	 * Allow external code to read the internal mesh.
+	 */
+	virtual void ProcessMesh(TFunctionRef<void(const UE::Geometry::FDynamicMesh3&)> ProcessFunc) const;
+
+	/**
+	 * Allow external code to to edit the internal mesh.
+	 */
+	virtual void EditMesh(TFunctionRef<void(UE::Geometry::FDynamicMesh3&)> EditFunc,
+						  EDynamicMeshComponentRenderUpdateMode UpdateMode = EDynamicMeshComponentRenderUpdateMode::FullUpdate);
+
+	/**
+	 * Apply transform to internal mesh. In some cases this can be more efficient than a general edit.
+	 * @param bInvert if true, inverse tranform is applied instead of forward transform
+	 */
+	virtual void ApplyTransform(const UE::Geometry::FTransform3d& Transform, bool bInvert) override;
+
+
+
 protected:
 	/**
 	 * Internal FDynamicMesh is stored inside a UDynamicMesh container, which allows it to be
@@ -416,77 +440,40 @@ protected:
 
 
 
-
-
-
-
-	///===============================================================================================================
-	// The functions below should be considered DEPRECATED
-	// They are slated for removal (but require some 
-	// refactoring for that to be possible)
+	//===============================================================================================================
+	// Triangle-Vertex Tangents support. The default behavior is to not use Tangents, this will lead to incorrect
+	// rendering for any material with Normal Maps and some other shaders.
+	// If TangentsType == EDynamicMeshComponentTangentsMode::ExternallyProvided, the Tangent and Bitangent attributes of
+	//   the FDynamicMesh3 AttributeSet are used at the SceneProxy level, the Component is not involved
+	// If TangentsType == EDynamicMeshComponentTangentsMode::AutoCalculated, the Tangents are computed internally using
+	//   a fast MikkT approximation via FMeshTangentsf. They will be recomputed when the mesh is modified, however
+	//   they are *not* recomputed when using the Fast Update functions above (in that case InvalidateAutoCalculatedTangents() 
+	//   can be used to force recomputation)
 	//
 public:
+	UFUNCTION(BlueprintCallable, Category = "Dynamic Mesh Component")
+	void SetTangentsType(EDynamicMeshComponentTangentsMode NewTangentsType);
 
-	/**
-	 * initialize the internal mesh from a MeshDescription
-	 */
-	virtual void InitializeMesh(const FMeshDescription* MeshDescription) override;
+	UFUNCTION(BlueprintCallable, Category = "Dynamic Mesh Component")
+	EDynamicMeshComponentTangentsMode GetTangentsType() const { return TangentsType; }
 
+	/** This function marks the auto tangents as dirty, they will be recomputed before they are used again */
+	virtual void InvalidateAutoCalculatedTangents();
 
+	/** @return AutoCalculated Tangent Set, which may require that they be recomputed, or nullptr if not enabled/available */
+	const UE::Geometry::FMeshTangentsf* GetAutoCalculatedTangents();
+
+protected:
 	/** How should Tangents be calculated/handled */
 	UPROPERTY()
-	EDynamicMeshTangentCalcType TangentsType = EDynamicMeshTangentCalcType::NoTangents;
+	EDynamicMeshComponentTangentsMode TangentsType = EDynamicMeshComponentTangentsMode::NoTangents;
 
-	//
-	// Separate Tangents support (todo: remove this)
-	//
+	/** true if AutoCalculatedTangents has been computed for current mesh */
+	bool bAutoCalculatedTangentsValid = false;
 
-	/**
-	 * Copy externally-calculated tangents into the internal tangets buffer.
-	 * @param bFastUpdateIfPossible if true, will try to do a fast normals/tangets update of the SceneProxy, instead of full invalidatiohn
-	 */
-	void UpdateTangents(const UE::Geometry::FMeshTangentsf* ExternalTangents, bool bFastUpdateIfPossible);
-
-	/**
-	 * Copy externally-calculated tangents into the internal tangets buffer.
-	 * @param bFastUpdateIfPossible if true, will try to do a fast normals/tangets update of the SceneProxy, instead of full invalidatiohn
-	 */
-	void UpdateTangents(const UE::Geometry::FMeshTangentsd* ExternalTangents, bool bFastUpdateIfPossible);
-
-	/**
-	 * @return pointer to internal tangents object.
-	 * @warning calling this with TangentsType = AutoCalculated will result in possibly-expensive Tangents calculation
-	 * @warning this is only currently safe to call on the Game Thread!!
-	 */
-	const UE::Geometry::FMeshTangentsf* GetTangents();
+	/** Set of per-triangle-vertex tangents computed for the current mesh. Only valid if bAutoCalculatedTangentsValid == true */
+	UE::Geometry::FMeshTangentsf AutoCalculatedTangents;
 
+	void UpdateAutoCalculatedTangents();
 
-	/**
-	 * Write the internal mesh to a MeshDescription
-	 * @param bHaveModifiedTopology if false, we only update the vertex positions in the MeshDescription, otherwise it is Empty()'d and regenerated entirely
-	 * @param ConversionOptions struct of additional options for the conversion
-	 */
-	virtual void Bake(FMeshDescription* MeshDescription, bool bHaveModifiedTopology, const FConversionToMeshDescriptionOptions& ConversionOptions) override;
-
-	/**
-	* Write the internal mesh to a MeshDescription with default conversion options
-	* @param bHaveModifiedTopology if false, we only update the vertex positions in the MeshDescription, otherwise it is Empty()'d and regenerated entirely
-	*/
-	void Bake(FMeshDescription* MeshDescription, bool bHaveModifiedTopology)
-	{
-		FConversionToMeshDescriptionOptions ConversionOptions;
-		Bake(MeshDescription, bHaveModifiedTopology, ConversionOptions);
-	}
-
-	/**
-	 * Apply transform to internal mesh. Updates Octree and RenderProxy if available.
-	 * @param bInvert if true, inverse tranform is applied instead of forward transform
-	 */
-	virtual void ApplyTransform(const UE::Geometry::FTransform3d& Transform, bool bInvert) override;
-
-
-protected:
-	/** Separate tangents storage (should be replaced with DynamicMesh Tangents attribute overlays) */
-	bool bTangentsValid = false;
-	UE::Geometry::FMeshTangentsf Tangents;
 };
diff --git a/Engine/Plugins/Experimental/MeshModelingToolset/Source/ModelingComponents/Public/UDynamicMesh.h b/Engine/Plugins/Experimental/MeshModelingToolset/Source/ModelingComponents/Public/UDynamicMesh.h
--- a/Engine/Plugins/Experimental/MeshModelingToolset/Source/ModelingComponents/Public/UDynamicMesh.h
+++ b/Engine/Plugins/Experimental/MeshModelingToolset/Source/ModelingComponents/Public/UDynamicMesh.h
@@ -40,9 +40,25 @@ enum class EDynamicMeshChangeType : uint8
 	GeneralEdit = 0,
 	MeshChange = 1,
 	MeshReplacementChange = 2,
-	MeshVertexChange = 3
+	MeshVertexChange = 3,
+	
+	DeformationEdit = 4,
+	AttributeEdit = 5
 };
 
+UENUM(BlueprintType)
+enum class EDynamicMeshAttributeChangeFlags : uint8
+{
+	Unknown = 0,
+	MeshTopology = 1 << 0,
+
+	VertexPositions = 1 << 1,
+	NormalsTangents = 1 << 2,
+	VertexColors = 1 << 3,
+	UVs = 1 << 4,
+	TriangleGroups = 1 << 5
+};
+ENUM_CLASS_FLAGS(EDynamicMeshAttributeChangeFlags)
 
 /**
  * FDynamicMeshChangeInfo stores information about a change to a UDynamicMesh.
@@ -56,6 +72,9 @@ struct MODELINGCOMPONENTS_API FDynamicMeshChangeInfo
 	UPROPERTY(BlueprintReadOnly, EditAnywhere, Category = "DynamicMeshChangeInfo")
 	EDynamicMeshChangeType Type = EDynamicMeshChangeType::GeneralEdit;
 
+	UPROPERTY(BlueprintReadOnly, EditAnywhere, Category = "DynamicMeshChangeInfo")
+	EDynamicMeshAttributeChangeFlags Flags = EDynamicMeshAttributeChangeFlags::Unknown;
+
 	// for changes that are an FChange, indicates whether this is an 'Apply' or 'Revert' of the FChange
 	UPROPERTY(BlueprintReadOnly, EditAnywhere, Category = "DynamicMeshChangeInfo")
 	bool bIsRevertChange = false;
@@ -147,7 +166,10 @@ public:
 	 * Apply EditFunc to the internal mesh.
 	 * This will broadcast PreMeshChangedEvent, then call EditFunc(), then broadcast MeshChangedEvent and MeshModifiedBPEvent
 	 */
-	void EditMesh(TFunctionRef<void(UE::Geometry::FDynamicMesh3&)> EditFunc);
+	void EditMesh(TFunctionRef<void(UE::Geometry::FDynamicMesh3&)> EditFunc,
+				  EDynamicMeshChangeType ChangeType = EDynamicMeshChangeType::GeneralEdit,
+				  EDynamicMeshAttributeChangeFlags ChangeFlags = EDynamicMeshAttributeChangeFlags::Unknown,
+				  bool bDeferChangeEvents = false);
 
 	/**
 	 * Take ownership of the internal Mesh, and have it replaced with a new mesh
@@ -234,7 +256,7 @@ protected:
 	/**
 	 * Internal function that edits the Mesh, but broadcasts PreMeshChangedEvent and MeshChangedEvent, and then MeshModifiedBPEvent
 	 */
-	void EditMeshInternal(TFunctionRef<void(UE::Geometry::FDynamicMesh3&)> EditFunc, const FDynamicMeshChangeInfo& ChangeInfo);
+	void EditMeshInternal(TFunctionRef<void(UE::Geometry::FDynamicMesh3&)> EditFunc, const FDynamicMeshChangeInfo& ChangeInfo, bool bDeferChangeEvents = false);
 
 
 public:
diff --git a/Engine/Plugins/Experimental/MeshModelingToolset/Source/ModelingComponentsEditorOnly/Private/ToolTargets/StaticMeshComponentToolTarget.cpp b/Engine/Plugins/Experimental/MeshModelingToolset/Source/ModelingComponentsEditorOnly/Private/ToolTargets/StaticMeshComponentToolTarget.cpp
--- a/Engine/Plugins/Experimental/MeshModelingToolset/Source/ModelingComponentsEditorOnly/Private/ToolTargets/StaticMeshComponentToolTarget.cpp
+++ b/Engine/Plugins/Experimental/MeshModelingToolset/Source/ModelingComponentsEditorOnly/Private/ToolTargets/StaticMeshComponentToolTarget.cpp
@@ -9,6 +9,7 @@
 #include "Materials/Material.h"
 #include "RenderingThread.h"
 #include "ToolTargets/StaticMeshToolTarget.h"
+#include "AssetUtils/MeshDescriptionUtil.h"
 
 using namespace UE::Geometry;
 
@@ -135,6 +136,24 @@ FMeshDescription* UStaticMeshComponentToolTarget::GetMeshDescription()
 }
 
 
+void UStaticMeshComponentToolTarget::CalculateAutoGeneratedAttributes(FMeshDescription& MeshDescription)
+{
+	if (ensure(IsValid()))
+	{
+		UStaticMesh* StaticMesh = Cast<UStaticMeshComponent>(Component)->GetStaticMesh();
+		if (EditingLOD == EStaticMeshEditingLOD::HiResSource)
+		{
+			const FStaticMeshSourceModel& SourceModel = StaticMesh->GetHiResSourceModel();
+			UE::MeshDescription::InitializeAutoGeneratedAttributes(MeshDescription, &SourceModel.BuildSettings);
+		}
+		else
+		{
+			UE::MeshDescription::InitializeAutoGeneratedAttributes(MeshDescription, StaticMesh, (int32)EditingLOD);
+		}
+	}
+}
+
+
 void UStaticMeshComponentToolTarget::CommitMeshDescription(const FCommitter& Committer)
 {
 	if (ensure(IsValid()) == false) return;
diff --git a/Engine/Plugins/Experimental/MeshModelingToolset/Source/ModelingComponentsEditorOnly/Public/ToolTargets/StaticMeshComponentToolTarget.h b/Engine/Plugins/Experimental/MeshModelingToolset/Source/ModelingComponentsEditorOnly/Public/ToolTargets/StaticMeshComponentToolTarget.h
--- a/Engine/Plugins/Experimental/MeshModelingToolset/Source/ModelingComponentsEditorOnly/Public/ToolTargets/StaticMeshComponentToolTarget.h
+++ b/Engine/Plugins/Experimental/MeshModelingToolset/Source/ModelingComponentsEditorOnly/Public/ToolTargets/StaticMeshComponentToolTarget.h
@@ -45,6 +45,7 @@ public:
 public:
 	// IMeshDescriptionProvider implementation
 	FMeshDescription* GetMeshDescription() override;
+	virtual void CalculateAutoGeneratedAttributes(FMeshDescription& MeshDescription) override;
 
 	// IMeshDescritpionCommitter implementation
 	virtual void CommitMeshDescription(const FCommitter& Committer) override;
diff --git a/Engine/Source/Runtime/Experimental/InteractiveToolsFramework/Public/TargetInterfaces/MeshDescriptionProvider.h b/Engine/Source/Runtime/Experimental/InteractiveToolsFramework/Public/TargetInterfaces/MeshDescriptionProvider.h
--- a/Engine/Source/Runtime/Experimental/InteractiveToolsFramework/Public/TargetInterfaces/MeshDescriptionProvider.h
+++ b/Engine/Source/Runtime/Experimental/InteractiveToolsFramework/Public/TargetInterfaces/MeshDescriptionProvider.h
@@ -20,9 +20,22 @@ class INTERACTIVETOOLSFRAMEWORK_API IMeshDescriptionProvider
 	GENERATED_BODY()
 
 public:
+	/**
+	 * Access the MeshDescriptionavailable through this Provider. Note that this MeshDescription may or may not 
+	 * be owned by the provider and should not be modified directly. Use IMeshDescriptionCommitter for writes.
+	 * @return pointer to MeshDescription 
+	 * @todo can this be const?
+	 */
 	virtual FMeshDescription* GetMeshDescription() = 0;
 
-	//TODO: Should this be providing a shared pointer instead? That seems like generally
-	// the better idea, but some refactoring may be necessary to avoid having to copy
-	// when getting from a static mesh.
+	/**
+	 * Compute any auto-generated attributes on the input MeshDescription. This is required to support
+	 * auto-computed Normals and Tangents on Static Mesh Assets, as those normals/tangents are computed 
+	 * on-the-fly when generating rendering data, and will be zero in the MeshDescription returned by
+	 * GetMeshDescription(). A no-op default implementation is provided.
+	 * @warning Do not pass the result from GetMeshDescription() here. Must be done on a copy.
+	 */
+	virtual void CalculateAutoGeneratedAttributes(FMeshDescription& MeshDescription)
+	{
+	}
 };
\ No newline at end of file
