Rename USimpleDynamicMeshComponent to UDynamicMeshComponent. Move ModelingComponents Components and SceneProxys to /Components subdirectory.
#rnx

diff --git a/Engine/Plugins/Experimental/MeshLODToolset/Source/MeshLODToolset/Public/Tools/GenerateStaticMeshLODAssetTool.h b/Engine/Plugins/Experimental/MeshLODToolset/Source/MeshLODToolset/Public/Tools/GenerateStaticMeshLODAssetTool.h
--- a/Engine/Plugins/Experimental/MeshLODToolset/Source/MeshLODToolset/Public/Tools/GenerateStaticMeshLODAssetTool.h
+++ b/Engine/Plugins/Experimental/MeshLODToolset/Source/MeshLODToolset/Public/Tools/GenerateStaticMeshLODAssetTool.h
@@ -17,7 +17,7 @@
 
 // predeclarations
 struct FMeshDescription;
-class USimpleDynamicMeshComponent;
+class UDynamicMeshComponent;
 class UMeshOpPreviewWithBackgroundCompute;
 class UGenerateStaticMeshLODAssetTool;
 namespace GenerateStaticMeshLODAssetLocals
diff --git a/Engine/Plugins/Experimental/MeshModelingToolset/Source/MeshModelingTools/Private/BakeMeshAttributeMapsTool.cpp b/Engine/Plugins/Experimental/MeshModelingToolset/Source/MeshModelingTools/Private/BakeMeshAttributeMapsTool.cpp
--- a/Engine/Plugins/Experimental/MeshModelingToolset/Source/MeshModelingTools/Private/BakeMeshAttributeMapsTool.cpp
+++ b/Engine/Plugins/Experimental/MeshModelingToolset/Source/MeshModelingTools/Private/BakeMeshAttributeMapsTool.cpp
@@ -15,7 +15,7 @@
 #include "Sampling/MeshResampleImageEvaluator.h"
 #include "Util/IndexUtil.h"
 
-#include "SimpleDynamicMeshComponent.h"
+#include "Components/DynamicMeshComponent.h"
 
 #include "Materials/Material.h"
 #include "Materials/MaterialInstanceDynamic.h"
@@ -277,7 +277,7 @@ void UBakeMeshAttributeMapsTool::Setup()
 	// create dynamic mesh component to use for live preview
 	// TODO: convert to UPreviewMesh
 	AActor* ParentActor = UE::ToolTarget::GetTargetActor(Targets[0]);
-	DynamicMeshComponent = NewObject<USimpleDynamicMeshComponent>(ParentActor);
+	DynamicMeshComponent = NewObject<UDynamicMeshComponent>(ParentActor);
 	DynamicMeshComponent->SetupAttachment(ParentActor->GetRootComponent());
 	DynamicMeshComponent->RegisterComponent();
 	DynamicMeshComponent->SetWorldTransform((FTransform)UE::ToolTarget::GetLocalToWorldTransform(Targets[0]));
diff --git a/Engine/Plugins/Experimental/MeshModelingToolset/Source/MeshModelingTools/Private/DeformMeshPolygonsTool.cpp b/Engine/Plugins/Experimental/MeshModelingToolset/Source/MeshModelingTools/Private/DeformMeshPolygonsTool.cpp
--- a/Engine/Plugins/Experimental/MeshModelingToolset/Source/MeshModelingTools/Private/DeformMeshPolygonsTool.cpp
+++ b/Engine/Plugins/Experimental/MeshModelingToolset/Source/MeshModelingTools/Private/DeformMeshPolygonsTool.cpp
@@ -849,7 +849,7 @@ void UDeformMeshPolygonsTool::Setup()
 	LaplacianDeformer = MakePimpl<FGroupTopologyLaplacianDeformer>();
 
 	// create dynamic mesh component to use for live preview
-	DynamicMeshComponent = NewObject<USimpleDynamicMeshComponent>(UE::ToolTarget::GetTargetActor(Target));
+	DynamicMeshComponent = NewObject<UDynamicMeshComponent>(UE::ToolTarget::GetTargetActor(Target));
 	DynamicMeshComponent->SetupAttachment(UE::ToolTarget::GetTargetActor(Target)->GetRootComponent());
 	DynamicMeshComponent->RegisterComponent();
 	WorldTransform = UE::ToolTarget::GetLocalToWorldTransform(Target);
diff --git a/Engine/Plugins/Experimental/MeshModelingToolset/Source/MeshModelingTools/Private/DisplaceMeshTool.cpp b/Engine/Plugins/Experimental/MeshModelingToolset/Source/MeshModelingTools/Private/DisplaceMeshTool.cpp
--- a/Engine/Plugins/Experimental/MeshModelingToolset/Source/MeshModelingTools/Private/DisplaceMeshTool.cpp
+++ b/Engine/Plugins/Experimental/MeshModelingToolset/Source/MeshModelingTools/Private/DisplaceMeshTool.cpp
@@ -9,7 +9,7 @@
 #include "ToolBuilderUtil.h"
 #include "ToolSetupUtil.h"
 #include "DynamicMesh3.h"
-#include "SimpleDynamicMeshComponent.h"
+#include "Components/DynamicMeshComponent.h"
 #include "MeshNormals.h"
 #include "ModelingOperators.h"
 #include "Async/ParallelFor.h"
@@ -786,7 +786,7 @@ void UDisplaceMeshTool::Setup()
 
 
 	// create dynamic mesh component to use for live preview
-	DynamicMeshComponent = NewObject<USimpleDynamicMeshComponent>(UE::ToolTarget::GetTargetActor(Target));
+	DynamicMeshComponent = NewObject<UDynamicMeshComponent>(UE::ToolTarget::GetTargetActor(Target));
 	DynamicMeshComponent->SetupAttachment(UE::ToolTarget::GetTargetActor(Target)->GetRootComponent());
 	DynamicMeshComponent->RegisterComponent();
 	DynamicMeshComponent->SetWorldTransform((FTransform)UE::ToolTarget::GetLocalToWorldTransform(Target));
diff --git a/Engine/Plugins/Experimental/MeshModelingToolset/Source/MeshModelingTools/Private/EditMeshPolygonsTool.cpp b/Engine/Plugins/Experimental/MeshModelingToolset/Source/MeshModelingTools/Private/EditMeshPolygonsTool.cpp
--- a/Engine/Plugins/Experimental/MeshModelingToolset/Source/MeshModelingTools/Private/EditMeshPolygonsTool.cpp
+++ b/Engine/Plugins/Experimental/MeshModelingToolset/Source/MeshModelingTools/Private/EditMeshPolygonsTool.cpp
@@ -166,7 +166,7 @@ void UEditMeshPolygonsTool::Setup()
 	AddInputBehavior(ClickBehavior);
 
 	// create dynamic mesh component to use for live preview
-	DynamicMeshComponent = NewObject<USimpleDynamicMeshComponent>(UE::ToolTarget::GetTargetActor(Target), "DynamicMesh");
+	DynamicMeshComponent = NewObject<UDynamicMeshComponent>(UE::ToolTarget::GetTargetActor(Target), "DynamicMesh");
 	DynamicMeshComponent->SetupAttachment(UE::ToolTarget::GetTargetActor(Target)->GetRootComponent());
 	DynamicMeshComponent->RegisterComponent();
 	WorldTransform = UE::ToolTarget::GetLocalToWorldTransform(Target);
diff --git a/Engine/Plugins/Experimental/MeshModelingToolset/Source/MeshModelingTools/Private/EditUVIslandsTool.cpp b/Engine/Plugins/Experimental/MeshModelingToolset/Source/MeshModelingTools/Private/EditUVIslandsTool.cpp
--- a/Engine/Plugins/Experimental/MeshModelingToolset/Source/MeshModelingTools/Private/EditUVIslandsTool.cpp
+++ b/Engine/Plugins/Experimental/MeshModelingToolset/Source/MeshModelingTools/Private/EditUVIslandsTool.cpp
@@ -43,7 +43,7 @@ void UEditUVIslandsTool::Setup()
 	UMeshSurfacePointTool::Setup();
 
 	// create dynamic mesh component to use for live preview
-	DynamicMeshComponent = NewObject<USimpleDynamicMeshComponent>(UE::ToolTarget::GetTargetActor(Target));
+	DynamicMeshComponent = NewObject<UDynamicMeshComponent>(UE::ToolTarget::GetTargetActor(Target));
 	DynamicMeshComponent->SetupAttachment(UE::ToolTarget::GetTargetActor(Target)->GetRootComponent());
 	DynamicMeshComponent->RegisterComponent();
 	DynamicMeshComponent->SetWorldTransform((FTransform)UE::ToolTarget::GetLocalToWorldTransform(Target));
diff --git a/Engine/Plugins/Experimental/MeshModelingToolset/Source/MeshModelingTools/Private/MeshGroupPaintTool.cpp b/Engine/Plugins/Experimental/MeshModelingToolset/Source/MeshModelingTools/Private/MeshGroupPaintTool.cpp
--- a/Engine/Plugins/Experimental/MeshModelingToolset/Source/MeshModelingTools/Private/MeshGroupPaintTool.cpp
+++ b/Engine/Plugins/Experimental/MeshModelingToolset/Source/MeshModelingTools/Private/MeshGroupPaintTool.cpp
@@ -86,7 +86,7 @@ void UMeshGroupPaintTool::Setup()
 	SetToolDisplayName(LOCTEXT("ToolName", "Paint PolyGroups"));
 
 	// create dynamic mesh component to use for live preview
-	DynamicMeshComponent = NewObject<USimpleDynamicMeshComponent>(UE::ToolTarget::GetTargetActor(Target));
+	DynamicMeshComponent = NewObject<UDynamicMeshComponent>(UE::ToolTarget::GetTargetActor(Target));
 	InitializeSculptMeshComponent(DynamicMeshComponent);
 
 	// assign materials
@@ -1383,7 +1383,7 @@ void UMeshGroupPaintTool::WaitForPendingUndoRedo()
 	}
 }
 
-void UMeshGroupPaintTool::OnDynamicMeshComponentChanged(USimpleDynamicMeshComponent* Component, const FMeshVertexChange* Change, bool bRevert)
+void UMeshGroupPaintTool::OnDynamicMeshComponentChanged(UDynamicMeshComponent* Component, const FMeshVertexChange* Change, bool bRevert)
 {
 	// update octree
 	FDynamicMesh3* Mesh = GetSculptMesh();
diff --git a/Engine/Plugins/Experimental/MeshModelingToolset/Source/MeshModelingTools/Private/MeshInspectorTool.cpp b/Engine/Plugins/Experimental/MeshModelingToolset/Source/MeshModelingTools/Private/MeshInspectorTool.cpp
--- a/Engine/Plugins/Experimental/MeshModelingToolset/Source/MeshModelingTools/Private/MeshInspectorTool.cpp
+++ b/Engine/Plugins/Experimental/MeshModelingToolset/Source/MeshModelingTools/Private/MeshInspectorTool.cpp
@@ -8,7 +8,7 @@
 #include "DynamicMeshAttributeSet.h"
 #include "MeshNormals.h"
 
-#include "SimpleDynamicMeshComponent.h"
+#include "Components/DynamicMeshComponent.h"
 #include "Materials/MaterialInstanceDynamic.h"
 
 #include "Properties/MeshStatisticsProperties.h"
diff --git a/Engine/Plugins/Experimental/MeshModelingToolset/Source/MeshModelingTools/Private/MeshVertexSculptTool.cpp b/Engine/Plugins/Experimental/MeshModelingToolset/Source/MeshModelingTools/Private/MeshVertexSculptTool.cpp
--- a/Engine/Plugins/Experimental/MeshModelingToolset/Source/MeshModelingTools/Private/MeshVertexSculptTool.cpp
+++ b/Engine/Plugins/Experimental/MeshModelingToolset/Source/MeshModelingTools/Private/MeshVertexSculptTool.cpp
@@ -64,7 +64,7 @@ void UMeshVertexSculptTool::Setup()
 	SetToolDisplayName(LOCTEXT("ToolName", "Sculpt"));
 
 	// create dynamic mesh component to use for live preview
-	DynamicMeshComponent = NewObject<USimpleDynamicMeshComponent>(Cast<IPrimitiveComponentBackedTarget>(Target)->GetOwnerActor());
+	DynamicMeshComponent = NewObject<UDynamicMeshComponent>(Cast<IPrimitiveComponentBackedTarget>(Target)->GetOwnerActor());
 	InitializeSculptMeshComponent(DynamicMeshComponent);
 
 	// assign materials
@@ -1028,7 +1028,7 @@ void UMeshVertexSculptTool::WaitForPendingUndoRedo()
 	}
 }
 
-void UMeshVertexSculptTool::OnDynamicMeshComponentChanged(USimpleDynamicMeshComponent* Component, const FMeshVertexChange* Change, bool bRevert)
+void UMeshVertexSculptTool::OnDynamicMeshComponentChanged(UDynamicMeshComponent* Component, const FMeshVertexChange* Change, bool bRevert)
 {
 	// have to wait for any outstanding stamp update to finish...
 	WaitForPendingStampUpdate();
diff --git a/Engine/Plugins/Experimental/MeshModelingToolset/Source/MeshModelingTools/Private/Properties/MeshMaterialProperties.cpp b/Engine/Plugins/Experimental/MeshModelingToolset/Source/MeshModelingTools/Private/Properties/MeshMaterialProperties.cpp
--- a/Engine/Plugins/Experimental/MeshModelingToolset/Source/MeshModelingTools/Private/Properties/MeshMaterialProperties.cpp
+++ b/Engine/Plugins/Experimental/MeshModelingToolset/Source/MeshModelingTools/Private/Properties/MeshMaterialProperties.cpp
@@ -6,7 +6,7 @@
 #include "DynamicMeshAttributeSet.h"
 #include "MeshNormals.h"
 
-#include "SimpleDynamicMeshComponent.h"
+#include "Components/DynamicMeshComponent.h"
 
 #include "Materials/MaterialInstanceDynamic.h"
 
diff --git a/Engine/Plugins/Experimental/MeshModelingToolset/Source/MeshModelingTools/Private/Properties/MeshStatisticsProperties.cpp b/Engine/Plugins/Experimental/MeshModelingToolset/Source/MeshModelingTools/Private/Properties/MeshStatisticsProperties.cpp
--- a/Engine/Plugins/Experimental/MeshModelingToolset/Source/MeshModelingTools/Private/Properties/MeshStatisticsProperties.cpp
+++ b/Engine/Plugins/Experimental/MeshModelingToolset/Source/MeshModelingTools/Private/Properties/MeshStatisticsProperties.cpp
@@ -6,7 +6,7 @@
 #include "DynamicMeshAttributeSet.h"
 #include "MeshNormals.h"
 
-#include "SimpleDynamicMeshComponent.h"
+#include "Components/DynamicMeshComponent.h"
 
 using namespace UE::Geometry;
 
diff --git a/Engine/Plugins/Experimental/MeshModelingToolset/Source/MeshModelingTools/Private/WeldMeshEdgesTool.cpp b/Engine/Plugins/Experimental/MeshModelingToolset/Source/MeshModelingTools/Private/WeldMeshEdgesTool.cpp
--- a/Engine/Plugins/Experimental/MeshModelingToolset/Source/MeshModelingTools/Private/WeldMeshEdgesTool.cpp
+++ b/Engine/Plugins/Experimental/MeshModelingToolset/Source/MeshModelingTools/Private/WeldMeshEdgesTool.cpp
@@ -8,7 +8,7 @@
 #include "DynamicMesh/DynamicMeshAttributeSet.h"
 #include "DynamicMesh/Operations/MergeCoincidentMeshEdges.h"
 
-#include "SimpleDynamicMeshComponent.h"
+#include "Components/DynamicMeshComponent.h"
 #include "ModelingToolTargetUtil.h"
 
 #include "SceneManagement.h" // for FPrimitiveDrawInterface
@@ -45,7 +45,7 @@ void UWeldMeshEdgesTool::Setup()
 	UInteractiveTool::Setup();
 
 	// create dynamic mesh component to use for live preview
-	DynamicMeshComponent = NewObject<USimpleDynamicMeshComponent>(UE::ToolTarget::GetTargetActor(Target));
+	DynamicMeshComponent = NewObject<UDynamicMeshComponent>(UE::ToolTarget::GetTargetActor(Target));
 	DynamicMeshComponent->SetupAttachment(UE::ToolTarget::GetTargetComponent(Target));
 	DynamicMeshComponent->RegisterComponent();
 	DynamicMeshComponent->SetWorldTransform((FTransform)UE::ToolTarget::GetLocalToWorldTransform(Target));
diff --git a/Engine/Plugins/Experimental/MeshModelingToolset/Source/MeshModelingTools/Public/BakeMeshAttributeMapsTool.h b/Engine/Plugins/Experimental/MeshModelingToolset/Source/MeshModelingTools/Public/BakeMeshAttributeMapsTool.h
--- a/Engine/Plugins/Experimental/MeshModelingToolset/Source/MeshModelingTools/Public/BakeMeshAttributeMapsTool.h
+++ b/Engine/Plugins/Experimental/MeshModelingToolset/Source/MeshModelingTools/Public/BakeMeshAttributeMapsTool.h
@@ -18,7 +18,7 @@
 
 // predeclarations
 struct FMeshDescription;
-class USimpleDynamicMeshComponent;
+class UDynamicMeshComponent;
 class UMaterialInstanceDynamic;
 class UTexture2D;
 PREDECLARE_GEOMETRY(template<typename RealType> class TMeshTangents);
@@ -404,7 +404,7 @@ protected:
 	friend class FBakeMultiTextureOp;
 	friend class FMeshMapBakerOp;
 
-	USimpleDynamicMeshComponent* DynamicMeshComponent;
+	UDynamicMeshComponent* DynamicMeshComponent;
 
 	UPROPERTY()
 	UMaterialInstanceDynamic* PreviewMaterial;
diff --git a/Engine/Plugins/Experimental/MeshModelingToolset/Source/MeshModelingTools/Public/BakeTransformTool.h b/Engine/Plugins/Experimental/MeshModelingToolset/Source/MeshModelingTools/Public/BakeTransformTool.h
--- a/Engine/Plugins/Experimental/MeshModelingToolset/Source/MeshModelingTools/Public/BakeTransformTool.h
+++ b/Engine/Plugins/Experimental/MeshModelingToolset/Source/MeshModelingTools/Public/BakeTransformTool.h
@@ -14,7 +14,7 @@
 
 // predeclarations
 struct FMeshDescription;
-class USimpleDynamicMeshComponent;
+class UDynamicMeshComponent;
 
 
 /**
diff --git a/Engine/Plugins/Experimental/MeshModelingToolset/Source/MeshModelingTools/Public/CombineMeshesTool.h b/Engine/Plugins/Experimental/MeshModelingToolset/Source/MeshModelingTools/Public/CombineMeshesTool.h
--- a/Engine/Plugins/Experimental/MeshModelingToolset/Source/MeshModelingTools/Public/CombineMeshesTool.h
+++ b/Engine/Plugins/Experimental/MeshModelingToolset/Source/MeshModelingTools/Public/CombineMeshesTool.h
@@ -15,7 +15,7 @@
 
 // predeclarations
 struct FMeshDescription;
-class USimpleDynamicMeshComponent;
+class UDynamicMeshComponent;
 class IPrimitiveComponentBackedTarget;
 class IMeshDescriptionCommitter;
 class IMeshDescriptionProvider;
diff --git a/Engine/Plugins/Experimental/MeshModelingToolset/Source/MeshModelingTools/Public/DeformMeshPolygonsTool.h b/Engine/Plugins/Experimental/MeshModelingToolset/Source/MeshModelingTools/Public/DeformMeshPolygonsTool.h
--- a/Engine/Plugins/Experimental/MeshModelingToolset/Source/MeshModelingTools/Public/DeformMeshPolygonsTool.h
+++ b/Engine/Plugins/Experimental/MeshModelingToolset/Source/MeshModelingTools/Public/DeformMeshPolygonsTool.h
@@ -6,7 +6,7 @@
 #include "Templates/PimplPtr.h"
 #include "UObject/NoExportTypes.h"
 #include "BaseTools/MeshSurfacePointTool.h"
-#include "SimpleDynamicMeshComponent.h"
+#include "Components/DynamicMeshComponent.h"
 #include "DynamicMeshAABBTree3.h"
 #include "ToolDataVisualizer.h"
 #include "Transforms/QuickAxisTranslater.h"
@@ -178,7 +178,7 @@ public:
 
 protected:
 	UPROPERTY()
-	USimpleDynamicMeshComponent* DynamicMeshComponent;
+	UDynamicMeshComponent* DynamicMeshComponent;
 
 	UPROPERTY()
 	UDeformMeshPolygonsTransformProperties* TransformProps;
diff --git a/Engine/Plugins/Experimental/MeshModelingToolset/Source/MeshModelingTools/Public/DisplaceMeshTool.h b/Engine/Plugins/Experimental/MeshModelingToolset/Source/MeshModelingTools/Public/DisplaceMeshTool.h
--- a/Engine/Plugins/Experimental/MeshModelingToolset/Source/MeshModelingTools/Public/DisplaceMeshTool.h
+++ b/Engine/Plugins/Experimental/MeshModelingToolset/Source/MeshModelingTools/Public/DisplaceMeshTool.h
@@ -17,7 +17,7 @@
 
 struct FMeshDescription;
 class UCurveFloat;
-class USimpleDynamicMeshComponent;
+class UDynamicMeshComponent;
 
 UENUM()
 enum class EDisplaceMeshToolDisplaceType : uint8
@@ -298,7 +298,7 @@ private:
 	float WeightMapQuery(const FVector3d& Position, const UE::Geometry::FIndexedWeightMap& WeightMap) const;
 
 	TSharedPtr<UE::Geometry::FDynamicMesh3, ESPMode::ThreadSafe> SubdividedMesh = nullptr;
-	USimpleDynamicMeshComponent* DynamicMeshComponent = nullptr;
+	UDynamicMeshComponent* DynamicMeshComponent = nullptr;
 
 	TUniquePtr<UE::Geometry::IDynamicMeshOperatorFactory> Subdivider = nullptr;
 	TUniquePtr<UE::Geometry::IDynamicMeshOperatorFactory> Displacer = nullptr;
diff --git a/Engine/Plugins/Experimental/MeshModelingToolset/Source/MeshModelingTools/Public/DynamicMeshBrushTool.h b/Engine/Plugins/Experimental/MeshModelingToolset/Source/MeshModelingTools/Public/DynamicMeshBrushTool.h
--- a/Engine/Plugins/Experimental/MeshModelingToolset/Source/MeshModelingTools/Public/DynamicMeshBrushTool.h
+++ b/Engine/Plugins/Experimental/MeshModelingToolset/Source/MeshModelingTools/Public/DynamicMeshBrushTool.h
@@ -7,7 +7,7 @@
 #include "InteractiveToolBuilder.h"
 #include "BaseTools/BaseBrushTool.h"
 #include "DynamicMeshAABBTree3.h"
-#include "SimpleDynamicMeshComponent.h"
+#include "Components/DynamicMeshComponent.h"
 #include "PreviewMesh.h"
 #include "DynamicMeshBrushTool.generated.h"
 
diff --git a/Engine/Plugins/Experimental/MeshModelingToolset/Source/MeshModelingTools/Public/DynamicMeshSculptTool.h b/Engine/Plugins/Experimental/MeshModelingToolset/Source/MeshModelingTools/Public/DynamicMeshSculptTool.h
--- a/Engine/Plugins/Experimental/MeshModelingToolset/Source/MeshModelingTools/Public/DynamicMeshSculptTool.h
+++ b/Engine/Plugins/Experimental/MeshModelingToolset/Source/MeshModelingTools/Public/DynamicMeshSculptTool.h
@@ -7,7 +7,7 @@
 #include "SubRegionRemesher.h"
 #include "UObject/NoExportTypes.h"
 #include "BaseTools/MeshSurfacePointTool.h"
-#include "OctreeDynamicMeshComponent.h"
+#include "Components/OctreeDynamicMeshComponent.h"
 #include "DynamicMeshAABBTree3.h"
 #include "DynamicMeshOctree3.h"
 #include "MeshNormals.h"
diff --git a/Engine/Plugins/Experimental/MeshModelingToolset/Source/MeshModelingTools/Public/EditMeshPolygonsTool.h b/Engine/Plugins/Experimental/MeshModelingToolset/Source/MeshModelingTools/Public/EditMeshPolygonsTool.h
--- a/Engine/Plugins/Experimental/MeshModelingToolset/Source/MeshModelingTools/Public/EditMeshPolygonsTool.h
+++ b/Engine/Plugins/Experimental/MeshModelingToolset/Source/MeshModelingTools/Public/EditMeshPolygonsTool.h
@@ -5,7 +5,7 @@
 #include "CoreMinimal.h"
 #include "UObject/NoExportTypes.h"
 #include "BaseTools/MeshSurfacePointTool.h"
-#include "SimpleDynamicMeshComponent.h"
+#include "Components/DynamicMeshComponent.h"
 #include "DynamicMeshAABBTree3.h"
 #include "ToolDataVisualizer.h"
 #include "Transforms/QuickAxisTranslater.h"
@@ -553,7 +553,7 @@ protected:
 	bool bTriangleMode;		
 
 	UPROPERTY()
-	USimpleDynamicMeshComponent* DynamicMeshComponent = nullptr;
+	UDynamicMeshComponent* DynamicMeshComponent = nullptr;
 
 	UPROPERTY()
 	UPolyEditCommonProperties* CommonProps;
diff --git a/Engine/Plugins/Experimental/MeshModelingToolset/Source/MeshModelingTools/Public/EditNormalsTool.h b/Engine/Plugins/Experimental/MeshModelingToolset/Source/MeshModelingTools/Public/EditNormalsTool.h
--- a/Engine/Plugins/Experimental/MeshModelingToolset/Source/MeshModelingTools/Public/EditNormalsTool.h
+++ b/Engine/Plugins/Experimental/MeshModelingToolset/Source/MeshModelingTools/Public/EditNormalsTool.h
@@ -15,7 +15,7 @@
 
 // predeclarations
 struct FMeshDescription;
-class USimpleDynamicMeshComponent;
+class UDynamicMeshComponent;
 
 /**
  *
diff --git a/Engine/Plugins/Experimental/MeshModelingToolset/Source/MeshModelingTools/Public/EditUVIslandsTool.h b/Engine/Plugins/Experimental/MeshModelingToolset/Source/MeshModelingTools/Public/EditUVIslandsTool.h
--- a/Engine/Plugins/Experimental/MeshModelingToolset/Source/MeshModelingTools/Public/EditUVIslandsTool.h
+++ b/Engine/Plugins/Experimental/MeshModelingToolset/Source/MeshModelingTools/Public/EditUVIslandsTool.h
@@ -5,7 +5,7 @@
 #include "CoreMinimal.h"
 #include "BaseTools/MeshSurfacePointTool.h"
 #include "Properties/MeshMaterialProperties.h"
-#include "SimpleDynamicMeshComponent.h"
+#include "Components/DynamicMeshComponent.h"
 #include "DynamicMeshAABBTree3.h"
 #include "ToolDataVisualizer.h"
 #include "Changes/MeshVertexChange.h"
@@ -96,7 +96,7 @@ public:
 
 protected:
 	UPROPERTY()
-	USimpleDynamicMeshComponent* DynamicMeshComponent;
+	UDynamicMeshComponent* DynamicMeshComponent;
 
 	UPROPERTY()
 	UPolygonSelectionMechanic* SelectionMechanic;
diff --git a/Engine/Plugins/Experimental/MeshModelingToolset/Source/MeshModelingTools/Public/HoleFillTool.h b/Engine/Plugins/Experimental/MeshModelingToolset/Source/MeshModelingTools/Public/HoleFillTool.h
--- a/Engine/Plugins/Experimental/MeshModelingToolset/Source/MeshModelingTools/Public/HoleFillTool.h
+++ b/Engine/Plugins/Experimental/MeshModelingToolset/Source/MeshModelingTools/Public/HoleFillTool.h
@@ -18,7 +18,7 @@
 class UPolygonSelectionMechanic;
 class UDynamicMeshReplacementChangeTarget;
 class UMeshOpPreviewWithBackgroundCompute;
-class USimpleDynamicMeshComponent;
+class UDynamicMeshComponent;
 struct FDynamicMeshOpResult;
 class UHoleFillTool;
 
diff --git a/Engine/Plugins/Experimental/MeshModelingToolset/Source/MeshModelingTools/Public/MeshGroupPaintTool.h b/Engine/Plugins/Experimental/MeshModelingToolset/Source/MeshModelingTools/Public/MeshGroupPaintTool.h
--- a/Engine/Plugins/Experimental/MeshModelingToolset/Source/MeshModelingTools/Public/MeshGroupPaintTool.h
+++ b/Engine/Plugins/Experimental/MeshModelingToolset/Source/MeshModelingTools/Public/MeshGroupPaintTool.h
@@ -6,7 +6,7 @@
 #include "HAL/Platform.h"
 #include "UObject/NoExportTypes.h"
 #include "BaseTools/BaseBrushTool.h"
-#include "SimpleDynamicMeshComponent.h"
+#include "Components/DynamicMeshComponent.h"
 #include "PropertySets/PolygroupLayersProperties.h"
 #include "Mechanics/PolyLassoMarqueeMechanic.h"
 
@@ -383,13 +383,13 @@ protected:
 
 protected:
 	UPROPERTY()
-	USimpleDynamicMeshComponent* DynamicMeshComponent;
+	UDynamicMeshComponent* DynamicMeshComponent;
 
 	UPROPERTY()
 	UMeshElementsVisualizer* MeshElementsDisplay;
 
 	// realtime visualization
-	void OnDynamicMeshComponentChanged(USimpleDynamicMeshComponent* Component, const FMeshVertexChange* Change, bool bRevert);
+	void OnDynamicMeshComponentChanged(UDynamicMeshComponent* Component, const FMeshVertexChange* Change, bool bRevert);
 	FDelegateHandle OnDynamicMeshComponentChangedHandle;
 
 	TUniquePtr<UE::Geometry::FPolygroupSet> ActiveGroupSet;
diff --git a/Engine/Plugins/Experimental/MeshModelingToolset/Source/MeshModelingTools/Public/MeshSpaceDeformerTool.h b/Engine/Plugins/Experimental/MeshModelingToolset/Source/MeshModelingTools/Public/MeshSpaceDeformerTool.h
--- a/Engine/Plugins/Experimental/MeshModelingToolset/Source/MeshModelingTools/Public/MeshSpaceDeformerTool.h
+++ b/Engine/Plugins/Experimental/MeshModelingToolset/Source/MeshModelingTools/Public/MeshSpaceDeformerTool.h
@@ -10,7 +10,7 @@
 #include "SpaceDeformerOps/BendMeshOp.h"
 #include "SpaceDeformerOps/TwistMeshOp.h"
 #include "SpaceDeformerOps/FlareMeshOp.h"
-#include "SimpleDynamicMeshComponent.h"
+#include "Components/DynamicMeshComponent.h"
 #include "ToolDataVisualizer.h"
 #include "BaseGizmos/GizmoInterfaces.h"
 #include "BaseTools/SingleSelectionMeshEditingTool.h"
diff --git a/Engine/Plugins/Experimental/MeshModelingToolset/Source/MeshModelingTools/Public/MeshVertexSculptTool.h b/Engine/Plugins/Experimental/MeshModelingToolset/Source/MeshModelingTools/Public/MeshVertexSculptTool.h
--- a/Engine/Plugins/Experimental/MeshModelingToolset/Source/MeshModelingTools/Public/MeshVertexSculptTool.h
+++ b/Engine/Plugins/Experimental/MeshModelingToolset/Source/MeshModelingTools/Public/MeshVertexSculptTool.h
@@ -5,7 +5,7 @@
 #include "CoreMinimal.h"
 #include "HAL/Platform.h"
 #include "UObject/NoExportTypes.h"
-#include "SimpleDynamicMeshComponent.h"
+#include "Components/DynamicMeshComponent.h"
 #include "DynamicMeshAABBTree3.h"
 #include "DynamicMeshOctree3.h"
 #include "MeshNormals.h"
@@ -229,10 +229,10 @@ protected:
 
 protected:
 	UPROPERTY()
-	USimpleDynamicMeshComponent* DynamicMeshComponent;
+	UDynamicMeshComponent* DynamicMeshComponent;
 
 	// realtime visualization
-	void OnDynamicMeshComponentChanged(USimpleDynamicMeshComponent* Component, const FMeshVertexChange* Change, bool bRevert);
+	void OnDynamicMeshComponentChanged(UDynamicMeshComponent* Component, const FMeshVertexChange* Change, bool bRevert);
 	FDelegateHandle OnDynamicMeshComponentChangedHandle;
 
 	void UpdateBrushType(EMeshVertexSculptBrushType BrushType);
diff --git a/Engine/Plugins/Experimental/MeshModelingToolset/Source/MeshModelingTools/Public/Sculpting/MeshSculptToolBase.h b/Engine/Plugins/Experimental/MeshModelingToolset/Source/MeshModelingTools/Public/Sculpting/MeshSculptToolBase.h
--- a/Engine/Plugins/Experimental/MeshModelingToolset/Source/MeshModelingTools/Public/Sculpting/MeshSculptToolBase.h
+++ b/Engine/Plugins/Experimental/MeshModelingToolset/Source/MeshModelingTools/Public/Sculpting/MeshSculptToolBase.h
@@ -9,7 +9,7 @@
 #include "BoxTypes.h"
 #include "Properties/MeshMaterialProperties.h"
 #include "Changes/ValueWatcher.h"
-#include "BaseDynamicMeshComponent.h"
+#include "Components/BaseDynamicMeshComponent.h"
 #include "MeshSculptToolBase.generated.h"
 
 
diff --git a/Engine/Plugins/Experimental/MeshModelingToolset/Source/MeshModelingTools/Public/UVProjectionTool.h b/Engine/Plugins/Experimental/MeshModelingToolset/Source/MeshModelingTools/Public/UVProjectionTool.h
--- a/Engine/Plugins/Experimental/MeshModelingToolset/Source/MeshModelingTools/Public/UVProjectionTool.h
+++ b/Engine/Plugins/Experimental/MeshModelingToolset/Source/MeshModelingTools/Public/UVProjectionTool.h
@@ -18,7 +18,7 @@
 
 // predeclarations
 struct FMeshDescription;
-class USimpleDynamicMeshComponent;
+class UDynamicMeshComponent;
 class UTransformGizmo;
 class UTransformProxy;
 class USingleClickInputBehavior;
diff --git a/Engine/Plugins/Experimental/MeshModelingToolset/Source/MeshModelingTools/Public/WeldMeshEdgesTool.h b/Engine/Plugins/Experimental/MeshModelingToolset/Source/MeshModelingTools/Public/WeldMeshEdgesTool.h
--- a/Engine/Plugins/Experimental/MeshModelingToolset/Source/MeshModelingTools/Public/WeldMeshEdgesTool.h
+++ b/Engine/Plugins/Experimental/MeshModelingToolset/Source/MeshModelingTools/Public/WeldMeshEdgesTool.h
@@ -14,7 +14,7 @@
 
 // predeclarations
 struct FMeshDescription;
-class USimpleDynamicMeshComponent;
+class UDynamicMeshComponent;
 
 /**
  *
@@ -63,7 +63,7 @@ protected:
 	bool bOnlyUnique;
 
 protected:
-	USimpleDynamicMeshComponent* DynamicMeshComponent;
+	UDynamicMeshComponent* DynamicMeshComponent;
 	UE::Geometry::FDynamicMesh3 OriginalMesh;
 	UE::Geometry::FDynamicMesh3 ComputeMesh;
 	bool bResultValid;
diff --git a/Engine/Plugins/Experimental/MeshModelingToolset/Source/MeshModelingToolsEditorOnly/Private/Hair/GroomCardsEditorTool.cpp b/Engine/Plugins/Experimental/MeshModelingToolset/Source/MeshModelingToolsEditorOnly/Private/Hair/GroomCardsEditorTool.cpp
--- a/Engine/Plugins/Experimental/MeshModelingToolset/Source/MeshModelingToolsEditorOnly/Private/Hair/GroomCardsEditorTool.cpp
+++ b/Engine/Plugins/Experimental/MeshModelingToolset/Source/MeshModelingToolsEditorOnly/Private/Hair/GroomCardsEditorTool.cpp
@@ -258,7 +258,7 @@ void UGroomCardsEditorTool::OnTick(float DeltaTime)
 
 	if (bSelectionStateDirty)
 	{
-		USimpleDynamicMeshComponent* DynamicMeshComponent = Cast<USimpleDynamicMeshComponent>(PreviewMesh->GetRootComponent());		// ugh
+		UDynamicMeshComponent* DynamicMeshComponent = Cast<UDynamicMeshComponent>(PreviewMesh->GetRootComponent());		// ugh
 		if (DynamicMeshComponent)
 		{
 			DynamicMeshComponent->FastNotifySecondaryTrianglesChanged();
diff --git a/Engine/Plugins/Experimental/MeshModelingToolset/Source/MeshModelingToolsEditorOnly/Private/ParameterizeMeshTool.cpp b/Engine/Plugins/Experimental/MeshModelingToolset/Source/MeshModelingToolsEditorOnly/Private/ParameterizeMeshTool.cpp
--- a/Engine/Plugins/Experimental/MeshModelingToolset/Source/MeshModelingToolsEditorOnly/Private/ParameterizeMeshTool.cpp
+++ b/Engine/Plugins/Experimental/MeshModelingToolset/Source/MeshModelingToolsEditorOnly/Private/ParameterizeMeshTool.cpp
@@ -8,7 +8,7 @@
 #include "DynamicMeshToMeshDescription.h"
 #include "FaceGroupUtil.h"
 
-#include "SimpleDynamicMeshComponent.h"
+#include "Components/DynamicMeshComponent.h"
 #include "Materials/MaterialInstanceDynamic.h"
 #include "MeshDescription.h"
 #include "MeshDescriptionToDynamicMesh.h"
diff --git a/Engine/Plugins/Experimental/MeshModelingToolset/Source/MeshModelingToolsEditorOnly/Private/ShapeSprayTool.cpp b/Engine/Plugins/Experimental/MeshModelingToolset/Source/MeshModelingToolsEditorOnly/Private/ShapeSprayTool.cpp
--- a/Engine/Plugins/Experimental/MeshModelingToolset/Source/MeshModelingToolsEditorOnly/Private/ShapeSprayTool.cpp
+++ b/Engine/Plugins/Experimental/MeshModelingToolset/Source/MeshModelingToolsEditorOnly/Private/ShapeSprayTool.cpp
@@ -69,7 +69,7 @@ void UShapeSprayTool::Setup()
 
 	// create dynamic mesh component to use for live preview
 	IPrimitiveComponentBackedTarget* TargetComponent = Cast<IPrimitiveComponentBackedTarget>(Target);
-	AccumMeshComponent = NewObject<USimpleDynamicMeshComponent>(TargetComponent->GetOwnerActor(), "SprayMesh");
+	AccumMeshComponent = NewObject<UDynamicMeshComponent>(TargetComponent->GetOwnerActor(), "SprayMesh");
 		AccumMeshComponent->SetupAttachment(TargetComponent->GetOwnerActor()->GetRootComponent());
 	AccumMeshComponent->RegisterComponent();
 
diff --git a/Engine/Plugins/Experimental/MeshModelingToolset/Source/MeshModelingToolsEditorOnly/Private/SubdividePolyTool.cpp b/Engine/Plugins/Experimental/MeshModelingToolset/Source/MeshModelingToolsEditorOnly/Private/SubdividePolyTool.cpp
--- a/Engine/Plugins/Experimental/MeshModelingToolset/Source/MeshModelingToolsEditorOnly/Private/SubdividePolyTool.cpp
+++ b/Engine/Plugins/Experimental/MeshModelingToolset/Source/MeshModelingToolsEditorOnly/Private/SubdividePolyTool.cpp
@@ -9,7 +9,7 @@
 #include "ToolSetupUtil.h"
 #include "Util/ColorConstants.h"
 #include "MeshNormals.h"
-#include "SimpleDynamicMeshComponent.h"
+#include "Components/DynamicMeshComponent.h"
 #include "Drawing/PreviewGeometryActor.h"
 
 #include "TargetInterfaces/MaterialProvider.h"
@@ -204,7 +204,7 @@ void USubdividePolyTool::Setup()
 	PreviewMesh->SetTransform(TargetComponent->GetWorldTransform());
 	PreviewMesh->UpdatePreview(OriginalMesh.Get());
 
-	USimpleDynamicMeshComponent* PreviewDynamicMeshComponent = (USimpleDynamicMeshComponent*)PreviewMesh->GetRootComponent();
+	UDynamicMeshComponent* PreviewDynamicMeshComponent = (UDynamicMeshComponent*)PreviewMesh->GetRootComponent();
 	if (PreviewDynamicMeshComponent == nullptr)
 	{
 		return;
@@ -236,7 +236,7 @@ void USubdividePolyTool::Setup()
 	// dynamic mesh configuration settings
 	auto RebuildMeshPostProcessor = [this]()
 	{
-		USimpleDynamicMeshComponent* PreviewDynamicMeshComponent = (USimpleDynamicMeshComponent*)PreviewMesh->GetRootComponent();
+		UDynamicMeshComponent* PreviewDynamicMeshComponent = (UDynamicMeshComponent*)PreviewMesh->GetRootComponent();
 		PreviewDynamicMeshComponent->SetRenderMeshPostProcessor(MakeUnique<SubdivPostProcessor>(Properties->SubdivisionLevel,
 																								Properties->SubdivisionScheme,
 																								Properties->NormalComputationMethod,
@@ -390,7 +390,7 @@ void USubdividePolyTool::Shutdown(EToolShutdownType ShutdownType)
 		{
 			GetToolManager()->BeginUndoTransaction(LOCTEXT("USubdividePolyTool", "Subdivide Mesh"));
 
-			USimpleDynamicMeshComponent* PreviewDynamicMeshComponent = (USimpleDynamicMeshComponent*)PreviewMesh->GetRootComponent();
+			UDynamicMeshComponent* PreviewDynamicMeshComponent = (UDynamicMeshComponent*)PreviewMesh->GetRootComponent();
 			FDynamicMesh3* DynamicMeshResult = PreviewDynamicMeshComponent->GetRenderMesh();
 
 			Cast<IMeshDescriptionCommitter>(Target)->CommitMeshDescription([DynamicMeshResult](const IMeshDescriptionCommitter::FCommitterParams& CommitParams)
diff --git a/Engine/Plugins/Experimental/MeshModelingToolset/Source/MeshModelingToolsEditorOnly/Public/MeshTangentsTool.h b/Engine/Plugins/Experimental/MeshModelingToolset/Source/MeshModelingToolsEditorOnly/Public/MeshTangentsTool.h
--- a/Engine/Plugins/Experimental/MeshModelingToolset/Source/MeshModelingToolsEditorOnly/Public/MeshTangentsTool.h
+++ b/Engine/Plugins/Experimental/MeshModelingToolset/Source/MeshModelingToolsEditorOnly/Public/MeshTangentsTool.h
@@ -17,7 +17,7 @@
 
 // predeclarations
 struct FMeshDescription;
-class USimpleDynamicMeshComponent;
+class UDynamicMeshComponent;
 class UMaterialInterface;
 class UMaterialInstanceDynamic;
 
diff --git a/Engine/Plugins/Experimental/MeshModelingToolset/Source/MeshModelingToolsEditorOnly/Public/ParameterizeMeshTool.h b/Engine/Plugins/Experimental/MeshModelingToolset/Source/MeshModelingToolsEditorOnly/Public/ParameterizeMeshTool.h
--- a/Engine/Plugins/Experimental/MeshModelingToolset/Source/MeshModelingToolsEditorOnly/Public/ParameterizeMeshTool.h
+++ b/Engine/Plugins/Experimental/MeshModelingToolset/Source/MeshModelingToolsEditorOnly/Public/ParameterizeMeshTool.h
@@ -16,7 +16,7 @@
 
 
 // predeclarations
-class USimpleDynamicMeshComponent;
+class UDynamicMeshComponent;
 class UMaterialInterface;
 class UMaterialInstanceDynamic;
 
diff --git a/Engine/Plugins/Experimental/MeshModelingToolset/Source/MeshModelingToolsEditorOnly/Public/PolygonOnMeshTool.h b/Engine/Plugins/Experimental/MeshModelingToolset/Source/MeshModelingToolsEditorOnly/Public/PolygonOnMeshTool.h
--- a/Engine/Plugins/Experimental/MeshModelingToolset/Source/MeshModelingToolsEditorOnly/Public/PolygonOnMeshTool.h
+++ b/Engine/Plugins/Experimental/MeshModelingToolset/Source/MeshModelingToolsEditorOnly/Public/PolygonOnMeshTool.h
@@ -20,7 +20,7 @@
 
 // predeclarations
 struct FMeshDescription;
-class USimpleDynamicMeshComponent;
+class UDynamicMeshComponent;
 class UTransformGizmo;
 class UTransformProxy;
 class ULineSetComponent;
diff --git a/Engine/Plugins/Experimental/MeshModelingToolset/Source/MeshModelingToolsEditorOnly/Public/ShapeSprayTool.h b/Engine/Plugins/Experimental/MeshModelingToolset/Source/MeshModelingToolsEditorOnly/Public/ShapeSprayTool.h
--- a/Engine/Plugins/Experimental/MeshModelingToolset/Source/MeshModelingToolsEditorOnly/Public/ShapeSprayTool.h
+++ b/Engine/Plugins/Experimental/MeshModelingToolset/Source/MeshModelingToolsEditorOnly/Public/ShapeSprayTool.h
@@ -89,7 +89,7 @@ protected:
 
 	// small meshes are accumulated here
 	UPROPERTY()
-	USimpleDynamicMeshComponent* AccumMeshComponent;
+	UDynamicMeshComponent* AccumMeshComponent;
 
 
 protected:
diff --git a/Engine/Plugins/Experimental/MeshModelingToolset/Source/MeshModelingToolsEditorOnly/Public/SubdividePolyTool.h b/Engine/Plugins/Experimental/MeshModelingToolset/Source/MeshModelingToolsEditorOnly/Public/SubdividePolyTool.h
--- a/Engine/Plugins/Experimental/MeshModelingToolset/Source/MeshModelingToolsEditorOnly/Public/SubdividePolyTool.h
+++ b/Engine/Plugins/Experimental/MeshModelingToolset/Source/MeshModelingToolsEditorOnly/Public/SubdividePolyTool.h
@@ -8,7 +8,7 @@
 #include "InteractiveToolBuilder.h"
 #include "ModelingOperators.h"
 #include "SingleSelectionTool.h"
-#include "SimpleDynamicMeshComponent.h"
+#include "Components/DynamicMeshComponent.h"
 #include "MeshOpPreviewHelpers.h"
 #include "BaseTools/SingleSelectionMeshEditingTool.h"
 #include "SubdividePolyTool.generated.h"
diff --git a/Engine/Plugins/Experimental/MeshModelingToolset/Source/MeshModelingToolsEditorOnly/Public/UVLayoutTool.h b/Engine/Plugins/Experimental/MeshModelingToolset/Source/MeshModelingToolsEditorOnly/Public/UVLayoutTool.h
--- a/Engine/Plugins/Experimental/MeshModelingToolset/Source/MeshModelingToolsEditorOnly/Public/UVLayoutTool.h
+++ b/Engine/Plugins/Experimental/MeshModelingToolset/Source/MeshModelingToolsEditorOnly/Public/UVLayoutTool.h
@@ -19,7 +19,7 @@
 
 // predeclarations
 struct FMeshDescription;
-class USimpleDynamicMeshComponent;
+class UDynamicMeshComponent;
 
 /**
  *
diff --git a/Engine/Plugins/Experimental/MeshModelingToolset/Source/ModelingComponents/Private/BaseDynamicMeshComponent.cpp b/Engine/Plugins/Experimental/MeshModelingToolset/Source/ModelingComponents/Private/BaseDynamicMeshComponent.cpp
deleted file mode 100644
diff --git a/Engine/Plugins/Experimental/MeshModelingToolset/Source/ModelingComponents/Private/BaseDynamicMeshSceneProxy.h b/Engine/Plugins/Experimental/MeshModelingToolset/Source/ModelingComponents/Private/BaseDynamicMeshSceneProxy.h
deleted file mode 100644
diff --git a/Engine/Plugins/Experimental/MeshModelingToolset/Source/ModelingComponents/Private/BaseTools/BaseCreateFromSelectedTool.cpp b/Engine/Plugins/Experimental/MeshModelingToolset/Source/ModelingComponents/Private/BaseTools/BaseCreateFromSelectedTool.cpp
--- a/Engine/Plugins/Experimental/MeshModelingToolset/Source/ModelingComponents/Private/BaseTools/BaseCreateFromSelectedTool.cpp
+++ b/Engine/Plugins/Experimental/MeshModelingToolset/Source/ModelingComponents/Private/BaseTools/BaseCreateFromSelectedTool.cpp
@@ -4,7 +4,7 @@
 #include "InteractiveToolManager.h"
 #include "ToolBuilderUtil.h"
 #include "ToolSetupUtil.h"
-#include "SimpleDynamicMeshComponent.h"
+#include "Components/DynamicMeshComponent.h"
 #include "BaseGizmos/TransformGizmoUtil.h"
 #include "Async/Async.h"
 
diff --git a/Engine/Plugins/Experimental/MeshModelingToolset/Source/ModelingComponents/Private/BaseTools/BaseMeshProcessingTool.cpp b/Engine/Plugins/Experimental/MeshModelingToolset/Source/ModelingComponents/Private/BaseTools/BaseMeshProcessingTool.cpp
--- a/Engine/Plugins/Experimental/MeshModelingToolset/Source/ModelingComponents/Private/BaseTools/BaseMeshProcessingTool.cpp
+++ b/Engine/Plugins/Experimental/MeshModelingToolset/Source/ModelingComponents/Private/BaseTools/BaseMeshProcessingTool.cpp
@@ -4,7 +4,7 @@
 #include "InteractiveToolManager.h"
 #include "ToolBuilderUtil.h"
 #include "ToolSetupUtil.h"
-#include "SimpleDynamicMeshComponent.h"
+#include "Components/DynamicMeshComponent.h"
 #include "Async/Async.h"
 
 #include "MeshNormals.h"
diff --git a/Engine/Plugins/Experimental/MeshModelingToolset/Source/ModelingComponents/Private/BaseTools/BaseVoxelTool.cpp b/Engine/Plugins/Experimental/MeshModelingToolset/Source/ModelingComponents/Private/BaseTools/BaseVoxelTool.cpp
--- a/Engine/Plugins/Experimental/MeshModelingToolset/Source/ModelingComponents/Private/BaseTools/BaseVoxelTool.cpp
+++ b/Engine/Plugins/Experimental/MeshModelingToolset/Source/ModelingComponents/Private/BaseTools/BaseVoxelTool.cpp
@@ -4,7 +4,7 @@
 #include "InteractiveToolManager.h"
 #include "ToolBuilderUtil.h"
 #include "ToolSetupUtil.h"
-#include "SimpleDynamicMeshComponent.h"
+#include "Components/DynamicMeshComponent.h"
 #include "DynamicMesh3.h"
 
 #include "MeshDescriptionToDynamicMesh.h"
diff --git a/Engine/Plugins/Experimental/MeshModelingToolset/Source/ModelingComponents/Private/Changes/MeshVertexChange.cpp b/Engine/Plugins/Experimental/MeshModelingToolset/Source/ModelingComponents/Private/Changes/MeshVertexChange.cpp
--- a/Engine/Plugins/Experimental/MeshModelingToolset/Source/ModelingComponents/Private/Changes/MeshVertexChange.cpp
+++ b/Engine/Plugins/Experimental/MeshModelingToolset/Source/ModelingComponents/Private/Changes/MeshVertexChange.cpp
@@ -3,7 +3,7 @@
 #include "Changes/MeshVertexChange.h"
 #include "DynamicMesh3.h"
 
-#include "BaseDynamicMeshComponent.h"
+#include "Components/BaseDynamicMeshComponent.h"
 
 using namespace UE::Geometry;
 
diff --git a/Engine/Plugins/Experimental/MeshModelingToolset/Source/ModelingComponents/Private/Components/BaseDynamicMeshComponent.cpp b/Engine/Plugins/Experimental/MeshModelingToolset/Source/ModelingComponents/Private/Components/BaseDynamicMeshComponent.cpp
new file mode 100644
--- /dev/null
+++ b/Engine/Plugins/Experimental/MeshModelingToolset/Source/ModelingComponents/Private/Components/BaseDynamicMeshComponent.cpp
@@ -0,0 +1,116 @@
+// Copyright Epic Games, Inc. All Rights Reserved. 
+
+#include "Components/BaseDynamicMeshComponent.h"
+
+using namespace UE::Geometry;
+
+UBaseDynamicMeshComponent::UBaseDynamicMeshComponent(const FObjectInitializer& ObjectInitializer)
+	: Super(ObjectInitializer)
+{
+}
+
+
+
+void UBaseDynamicMeshComponent::SetOverrideRenderMaterial(UMaterialInterface* Material)
+{
+	if (OverrideRenderMaterial != Material)
+	{
+		OverrideRenderMaterial = Material;
+		NotifyMaterialSetUpdated();
+	}
+}
+
+void UBaseDynamicMeshComponent::ClearOverrideRenderMaterial()
+{
+	if (OverrideRenderMaterial != nullptr)
+	{
+		OverrideRenderMaterial = nullptr;
+		NotifyMaterialSetUpdated();
+	}
+}
+
+
+
+
+
+void UBaseDynamicMeshComponent::SetSecondaryRenderMaterial(UMaterialInterface* Material)
+{
+	if (SecondaryRenderMaterial != Material)
+	{
+		SecondaryRenderMaterial = Material;
+		NotifyMaterialSetUpdated();
+	}
+}
+
+void UBaseDynamicMeshComponent::ClearSecondaryRenderMaterial()
+{
+	if (SecondaryRenderMaterial != nullptr)
+	{
+		SecondaryRenderMaterial = nullptr;
+		NotifyMaterialSetUpdated();
+	}
+}
+
+
+
+void UBaseDynamicMeshComponent::SetSecondaryBuffersVisibility(bool bSecondaryVisibility)
+{
+	bDrawSecondaryBuffers = bSecondaryVisibility;
+}
+
+bool UBaseDynamicMeshComponent::GetSecondaryBuffersVisibility() const
+{
+	return bDrawSecondaryBuffers;
+}
+
+
+
+
+int32 UBaseDynamicMeshComponent::GetNumMaterials() const
+{
+	return BaseMaterials.Num();
+}
+
+UMaterialInterface* UBaseDynamicMeshComponent::GetMaterial(int32 ElementIndex) const 
+{
+	return (ElementIndex >= 0 && ElementIndex < BaseMaterials.Num()) ? BaseMaterials[ElementIndex] : nullptr;
+}
+
+FMaterialRelevance UBaseDynamicMeshComponent::GetMaterialRelevance(ERHIFeatureLevel::Type InFeatureLevel) const
+{
+	FMaterialRelevance Result = UMeshComponent::GetMaterialRelevance(InFeatureLevel);
+	if (OverrideRenderMaterial)
+	{
+		Result |= OverrideRenderMaterial->GetRelevance_Concurrent(InFeatureLevel);
+	}
+	if (SecondaryRenderMaterial)
+	{
+		Result |= SecondaryRenderMaterial->GetRelevance_Concurrent(InFeatureLevel);
+	}
+	return Result;
+}
+
+void UBaseDynamicMeshComponent::SetMaterial(int32 ElementIndex, UMaterialInterface* Material)
+{
+	check(ElementIndex >= 0);
+	if (ElementIndex >= BaseMaterials.Num())
+	{
+		BaseMaterials.SetNum(ElementIndex + 1, false);
+	}
+	BaseMaterials[ElementIndex] = Material;
+}
+
+
+void UBaseDynamicMeshComponent::GetUsedMaterials(TArray<UMaterialInterface*>& OutMaterials, bool bGetDebugMaterials) const
+{
+	UMeshComponent::GetUsedMaterials(OutMaterials, bGetDebugMaterials);
+	if (OverrideRenderMaterial != nullptr)
+	{
+		OutMaterials.Add(OverrideRenderMaterial);
+	}
+	if (SecondaryRenderMaterial != nullptr)
+	{
+		OutMaterials.Add(SecondaryRenderMaterial);
+	}
+}
+
diff --git a/Engine/Plugins/Experimental/MeshModelingToolset/Source/ModelingComponents/Private/Components/BaseDynamicMeshSceneProxy.h b/Engine/Plugins/Experimental/MeshModelingToolset/Source/ModelingComponents/Private/Components/BaseDynamicMeshSceneProxy.h
new file mode 100644
--- /dev/null
+++ b/Engine/Plugins/Experimental/MeshModelingToolset/Source/ModelingComponents/Private/Components/BaseDynamicMeshSceneProxy.h
@@ -0,0 +1,933 @@
+// Copyright Epic Games, Inc. All Rights Reserved. 
+
+#pragma once
+
+#include "CoreMinimal.h"
+#include "PrimitiveSceneProxy.h"
+#include "VertexFactory.h"
+#include "LocalVertexFactory.h"
+#include "PrimitiveViewRelevance.h"
+#include "Components/BaseDynamicMeshComponent.h"
+
+using UE::Geometry::FDynamicMesh3;
+using UE::Geometry::FDynamicMeshAttributeSet;
+using UE::Geometry::FDynamicMeshUVOverlay;
+using UE::Geometry::FDynamicMeshNormalOverlay;
+using UE::Geometry::FDynamicMeshColorOverlay;
+using UE::Geometry::FDynamicMeshMaterialAttribute;
+
+/**
+ * FMeshRenderBufferSet stores a set of RenderBuffers for a mesh
+ */
+class FMeshRenderBufferSet
+{
+public:
+	/** Number of triangles in this renderbuffer set. Note that triangles may be split between IndexBuffer and SecondaryIndexBuffer. */
+	int TriangleCount = 0;
+
+	/** The buffer containing vertex data. */
+	FStaticMeshVertexBuffer StaticMeshVertexBuffer;
+	/** The buffer containing the position vertex data. */
+	FPositionVertexBuffer PositionVertexBuffer;
+	/** The buffer containing the vertex color data. */
+	FColorVertexBuffer ColorVertexBuffer;
+
+	/** triangle indices */
+	FDynamicMeshIndexBuffer32 IndexBuffer;
+
+	/** vertex factory */
+	FLocalVertexFactory VertexFactory;
+
+	/** Material to draw this mesh with */
+	UMaterialInterface* Material = nullptr;
+
+	/**
+	 * Optional list of triangles stored in this buffer. Storing this allows us
+	 * to rebuild the buffers if vertex data changes.
+	 */
+	TOptional<TArray<int>> Triangles;
+
+	/**
+	 * If secondary index buffer is enabled, we populate this index buffer with additional triangles indexing into the same vertex buffers
+	 */
+	bool bEnableSecondaryIndexBuffer = false;
+
+	/**
+	 * partition or subset of IndexBuffer that indexes into same vertex buffers
+	 */
+	FDynamicMeshIndexBuffer32 SecondaryIndexBuffer;
+
+
+	/**
+	 * In situations where we want to *update* the existing Vertex or Index buffers, we need to synchronize
+	 * access between the Game and Render threads. We use this lock to do that.
+	 */
+	FCriticalSection BuffersLock;
+
+
+	FMeshRenderBufferSet(ERHIFeatureLevel::Type FeatureLevelType)
+		: VertexFactory(FeatureLevelType, "FMeshRenderBufferSet")
+	{
+		StaticMeshVertexBuffer.SetUseFullPrecisionUVs(true);
+		StaticMeshVertexBuffer.SetUseHighPrecisionTangentBasis(true);
+	}
+
+
+	virtual ~FMeshRenderBufferSet()
+	{
+		check(IsInRenderingThread());
+
+		if (TriangleCount > 0)
+		{
+			PositionVertexBuffer.ReleaseResource();
+			StaticMeshVertexBuffer.ReleaseResource();
+			ColorVertexBuffer.ReleaseResource();
+			VertexFactory.ReleaseResource();
+			if (IndexBuffer.IsInitialized())
+			{
+				IndexBuffer.ReleaseResource();
+			}
+			if (SecondaryIndexBuffer.IsInitialized())
+			{
+				SecondaryIndexBuffer.ReleaseResource();
+			}
+		}
+	}
+
+
+	/**
+	 * Upload initialized mesh buffers. 
+	 * @warning This can only be called on the Rendering Thread.
+	 */
+	void Upload()
+	{
+		check(IsInRenderingThread());
+
+		if (TriangleCount == 0)
+		{
+			return;
+		}
+
+		InitOrUpdateResource(&this->PositionVertexBuffer);
+		InitOrUpdateResource(&this->StaticMeshVertexBuffer);
+		InitOrUpdateResource(&this->ColorVertexBuffer);
+
+		FLocalVertexFactory::FDataType Data;
+		this->PositionVertexBuffer.BindPositionVertexBuffer(&this->VertexFactory, Data);
+		this->StaticMeshVertexBuffer.BindTangentVertexBuffer(&this->VertexFactory, Data);
+		this->StaticMeshVertexBuffer.BindPackedTexCoordVertexBuffer(&this->VertexFactory, Data);
+		// currently no lightmaps support
+		//this->StaticMeshVertexBuffer.BindLightMapVertexBuffer(&this->VertexFactory, Data, LightMapIndex);
+		this->ColorVertexBuffer.BindColorVertexBuffer(&this->VertexFactory, Data);
+		this->VertexFactory.SetData(Data);
+
+		InitOrUpdateResource(&this->VertexFactory);
+		PositionVertexBuffer.InitResource();
+		StaticMeshVertexBuffer.InitResource();
+		ColorVertexBuffer.InitResource();
+		VertexFactory.InitResource();
+
+		if (IndexBuffer.Indices.Num() > 0)
+		{
+			IndexBuffer.InitResource();
+		}
+		if (bEnableSecondaryIndexBuffer && SecondaryIndexBuffer.Indices.Num() > 0)
+		{
+			SecondaryIndexBuffer.InitResource();
+		}
+	}
+
+
+	void UploadIndexBufferUpdate()
+	{
+		check(IsInRenderingThread());
+		if (IndexBuffer.Indices.Num() > 0)
+		{
+			InitOrUpdateResource(&IndexBuffer);
+		}
+		if (bEnableSecondaryIndexBuffer && SecondaryIndexBuffer.Indices.Num() > 0)
+		{
+			InitOrUpdateResource(&SecondaryIndexBuffer);
+		}
+	}
+
+
+	void UploadVertexUpdate(bool bPositions, bool bMeshAttribs, bool bColors)
+	{
+		check(IsInRenderingThread());
+
+		if (TriangleCount == 0)
+		{
+			return;
+		}
+
+		if (bPositions)
+		{
+			InitOrUpdateResource(&this->PositionVertexBuffer);
+		}
+		if (bMeshAttribs)
+		{
+			InitOrUpdateResource(&this->StaticMeshVertexBuffer);
+		}
+		if (bColors)
+		{
+			InitOrUpdateResource(&this->ColorVertexBuffer);
+		}
+
+		FLocalVertexFactory::FDataType Data;
+		this->PositionVertexBuffer.BindPositionVertexBuffer(&this->VertexFactory, Data);
+		this->StaticMeshVertexBuffer.BindTangentVertexBuffer(&this->VertexFactory, Data);
+		this->StaticMeshVertexBuffer.BindPackedTexCoordVertexBuffer(&this->VertexFactory, Data);
+		// currently no lightmaps support
+		//this->StaticMeshVertexBuffer.BindLightMapVertexBuffer(&this->VertexFactory, Data, LightMapIndex);
+		this->ColorVertexBuffer.BindColorVertexBuffer(&this->VertexFactory, Data);
+		this->VertexFactory.SetData(Data);
+
+		InitOrUpdateResource(&this->VertexFactory);
+	}
+
+
+	void TransferVertexUpdateToGPU(bool bPositions, bool bNormals, bool bTexCoords, bool bColors)
+	{
+		check(IsInRenderingThread());
+		if (TriangleCount == 0)
+		{
+			return;
+		}
+
+		if (bPositions)
+		{
+			auto& VertexBuffer = this->PositionVertexBuffer;
+			void* VertexBufferData = RHILockBuffer(VertexBuffer.VertexBufferRHI, 0, VertexBuffer.GetNumVertices() * VertexBuffer.GetStride(), RLM_WriteOnly);
+			FMemory::Memcpy(VertexBufferData, VertexBuffer.GetVertexData(), VertexBuffer.GetNumVertices() * VertexBuffer.GetStride());
+			RHIUnlockBuffer(VertexBuffer.VertexBufferRHI);
+		}
+		if (bNormals)
+		{
+			auto& VertexBuffer = this->StaticMeshVertexBuffer;
+			void* VertexBufferData = RHILockBuffer(VertexBuffer.TangentsVertexBuffer.VertexBufferRHI, 0, VertexBuffer.GetTangentSize(), RLM_WriteOnly);
+			FMemory::Memcpy(VertexBufferData, VertexBuffer.GetTangentData(), VertexBuffer.GetTangentSize());
+			RHIUnlockBuffer(VertexBuffer.TangentsVertexBuffer.VertexBufferRHI);
+		}
+		if (bColors)
+		{
+			auto& VertexBuffer = this->ColorVertexBuffer;
+			void* VertexBufferData = RHILockBuffer(VertexBuffer.VertexBufferRHI, 0, VertexBuffer.GetNumVertices() * VertexBuffer.GetStride(), RLM_WriteOnly);
+			FMemory::Memcpy(VertexBufferData, VertexBuffer.GetVertexData(), VertexBuffer.GetNumVertices() * VertexBuffer.GetStride());
+			RHIUnlockBuffer(VertexBuffer.VertexBufferRHI);
+		}
+		if (bTexCoords)
+		{
+			auto& VertexBuffer = this->StaticMeshVertexBuffer;
+			void* VertexBufferData = RHILockBuffer(VertexBuffer.TexCoordVertexBuffer.VertexBufferRHI, 0, VertexBuffer.GetTexCoordSize(), RLM_WriteOnly);
+			FMemory::Memcpy(VertexBufferData, VertexBuffer.GetTexCoordData(), VertexBuffer.GetTexCoordSize());
+			RHIUnlockBuffer(VertexBuffer.TexCoordVertexBuffer.VertexBufferRHI);
+		}
+	}
+
+
+
+
+
+
+	/**
+	 * Initializes a render resource, or update it if already initialized.
+	 * @wearning This function can only be called on the Render Thread
+	 */
+	void InitOrUpdateResource(FRenderResource* Resource)
+	{
+		check(IsInRenderingThread());
+
+		if (!Resource->IsInitialized())
+		{
+			Resource->InitResource();
+		}
+		else
+		{
+			Resource->UpdateRHI();
+		}
+	}
+
+
+
+
+protected:
+	friend class FBaseDynamicMeshSceneProxy;
+
+	/**
+	 * Enqueue a command on the Render Thread to destroy the passed in buffer set.
+	 * At this point the buffer set should be considered invalid.
+	 */
+	static void DestroyRenderBufferSet(FMeshRenderBufferSet* BufferSet)
+	{
+		if (BufferSet->TriangleCount == 0)
+		{
+			return;
+		}
+
+		ENQUEUE_RENDER_COMMAND(FMeshRenderBufferSetDestroy)(
+			[BufferSet](FRHICommandListImmediate& RHICmdList)
+		{
+			delete BufferSet;
+		});
+	}
+
+
+};
+
+
+
+
+/**
+ * FBaseDynamicMeshSceneProxy is an abstract base class for a Render Proxy
+ * for a UBaseDynamicMeshComponent, where the assumption is that mesh data
+ * will be stored in FMeshRenderBufferSet instances
+ */
+class FBaseDynamicMeshSceneProxy : public FPrimitiveSceneProxy
+{
+	using FIndex2i = UE::Geometry::FIndex2i;
+	using FIndex3i = UE::Geometry::FIndex3i;
+public:
+	UBaseDynamicMeshComponent* ParentBaseComponent;
+
+	/**
+	 * Constant color assigned to vertices if no other vertex color is specified
+	 */
+	FColor ConstantVertexColor = FColor::White;
+
+	/**
+	 * If true, vertex colors on the FDynamicMesh will be ignored
+	 */
+	bool bIgnoreVertexColors = false;
+
+	/**
+	 * If true, a per-triangle color is used to set vertex colors
+	 */
+	bool bUsePerTriangleColor = false;
+
+	/**
+	 * Per-triangle color function. Only called if bUsePerTriangleColor=true
+	 */
+	TFunction<FColor(const FDynamicMesh3*, int)> PerTriangleColorFunc = nullptr;
+
+
+	/**
+	 * If true, populate secondary buffers using SecondaryTriFilterFunc
+	 */
+	bool bUseSecondaryTriBuffers = false;
+
+	/**
+	 * Filter predicate for secondary triangle index buffer. Only called if bUseSecondaryTriBuffers=true
+	 */
+	TUniqueFunction<bool(const FDynamicMesh3*, int32)> SecondaryTriFilterFunc = nullptr;
+
+protected:
+	// Set of currently-allocated RenderBuffers. We own these pointers and must clean them up.
+	// Must guard access with AllocatedSetsLock!!
+	TSet<FMeshRenderBufferSet*> AllocatedBufferSets;
+
+	// use to control access to AllocatedBufferSets 
+	FCriticalSection AllocatedSetsLock;
+
+public:
+	FBaseDynamicMeshSceneProxy(UBaseDynamicMeshComponent* Component)
+		: FPrimitiveSceneProxy(Component),
+		ParentBaseComponent(Component)
+	{
+	}
+
+
+	virtual ~FBaseDynamicMeshSceneProxy()
+	{
+		// we are assuming in code below that this is always called from the rendering thread
+		check(IsInRenderingThread());
+
+		// destroy all existing renderbuffers
+		for (FMeshRenderBufferSet* BufferSet : AllocatedBufferSets)
+		{
+			FMeshRenderBufferSet::DestroyRenderBufferSet(BufferSet);
+		}
+	}
+
+
+	//
+	// FBaseDynamicMeshSceneProxy API - subclasses must implement these functions
+	//
+
+
+	/**
+	 * Return set of active renderbuffers. Must be implemented by subclass.
+	 * This is the set of render buffers that will be drawn by GetDynamicMeshElements
+	 */
+	virtual void GetActiveRenderBufferSets(TArray<FMeshRenderBufferSet*>& Buffers) const = 0;
+
+
+
+	//
+	// RenderBuffer management
+	//
+
+
+	/**
+	 * Allocates a set of render buffers. FPrimitiveSceneProxy will keep track of these
+	 * buffers and destroy them on destruction.
+	 */
+	virtual FMeshRenderBufferSet* AllocateNewRenderBufferSet()
+	{
+		// should we hang onto these and destroy them in constructor? leaving to subclass seems risky?
+		FMeshRenderBufferSet* RenderBufferSet = new FMeshRenderBufferSet(GetScene().GetFeatureLevel());
+
+		RenderBufferSet->Material = UMaterial::GetDefaultMaterial(MD_Surface);
+
+		AllocatedSetsLock.Lock();
+		AllocatedBufferSets.Add(RenderBufferSet);
+		AllocatedSetsLock.Unlock();
+
+		return RenderBufferSet;
+	}
+
+	/**
+	 * Explicitly release a set of RenderBuffers
+	 */
+	virtual void ReleaseRenderBufferSet(FMeshRenderBufferSet* BufferSet)
+	{
+		FScopeLock Lock(&AllocatedSetsLock);
+		if (ensure(AllocatedBufferSets.Contains(BufferSet)))
+		{
+			AllocatedBufferSets.Remove(BufferSet);
+			Lock.Unlock();
+
+			FMeshRenderBufferSet::DestroyRenderBufferSet(BufferSet);
+		}
+	}
+
+
+	/**
+	 * Initialize rendering buffers from given attribute overlays.
+	 * Creates three vertices per triangle, IE no shared vertices in buffers.
+	 */
+	template<typename TriangleEnumerable>
+	void InitializeBuffersFromOverlays(
+		FMeshRenderBufferSet* RenderBuffers,
+		const FDynamicMesh3* Mesh,
+		int NumTriangles, TriangleEnumerable Enumerable,
+		const FDynamicMeshUVOverlay* UVOverlay,
+		const FDynamicMeshNormalOverlay* NormalOverlay,
+		const FDynamicMeshColorOverlay* ColorOverlay,
+		TFunctionRef<void(int, int, int, const FVector3f&, FVector3f&, FVector3f&)> TangentsFunc,
+		bool bTrackTriangles = false)
+	{
+		TArray<const FDynamicMeshUVOverlay*> UVOverlays;
+		UVOverlays.Add(UVOverlay);
+		InitializeBuffersFromOverlays(RenderBuffers, Mesh, NumTriangles, Enumerable,
+			UVOverlays, NormalOverlay, ColorOverlay, TangentsFunc, bTrackTriangles);
+	}
+
+
+
+	/**
+	 * Initialize rendering buffers from given attribute overlays.
+	 * Creates three vertices per triangle, IE no shared vertices in buffers.
+	 */
+	template<typename TriangleEnumerable>
+	void InitializeBuffersFromOverlays(
+		FMeshRenderBufferSet* RenderBuffers,
+		const FDynamicMesh3* Mesh,
+		int NumTriangles, TriangleEnumerable Enumerable,
+		const TArray<const FDynamicMeshUVOverlay*>& UVOverlays,
+		const FDynamicMeshNormalOverlay* NormalOverlay,
+		const FDynamicMeshColorOverlay* ColorOverlay,
+		TFunctionRef<void(int, int, int, const FVector3f&, FVector3f&, FVector3f&)> TangentsFunc,
+		bool bTrackTriangles = false)
+	{
+		RenderBuffers->TriangleCount = NumTriangles;
+		if (NumTriangles == 0)
+		{
+			return;
+		}
+
+		bool bHaveColors = (ColorOverlay != nullptr) && (bIgnoreVertexColors == false);
+
+		int NumVertices = NumTriangles * 3;
+		int NumUVOverlays = UVOverlays.Num();
+		int NumTexCoords = FMath::Max(1, NumUVOverlays);		// must have at least one tex coord
+		TArray<FIndex3i, TFixedAllocator<MAX_STATIC_TEXCOORDS>> UVTriangles;
+		UVTriangles.SetNum(NumTexCoords);
+
+		{
+			RenderBuffers->PositionVertexBuffer.Init(NumVertices);
+			RenderBuffers->StaticMeshVertexBuffer.Init(NumVertices, NumTexCoords );
+			RenderBuffers->ColorVertexBuffer.Init(NumVertices);
+			RenderBuffers->IndexBuffer.Indices.AddUninitialized(NumTriangles * 3);
+		}
+
+		// build triangle list if requested, or if we are using secondary buffers in which case we need it to filter later
+		bool bBuildTriangleList = bTrackTriangles || bUseSecondaryTriBuffers;
+		if (bBuildTriangleList)
+		{
+			RenderBuffers->Triangles = TArray<int32>();
+		}
+
+		int TriIdx = 0, VertIdx = 0;
+		FVector3f TangentX, TangentY;
+		for (int TriangleID : Enumerable)
+		{
+			FIndex3i Tri = Mesh->GetTriangle(TriangleID);
+			for (int32 k = 0; k < NumTexCoords; ++k)
+			{
+				UVTriangles[k] = (k < NumUVOverlays && UVOverlays[k] != nullptr) ? UVOverlays[k]->GetTriangle(TriangleID) : FIndex3i::Invalid();
+			}
+			FIndex3i TriNormal = (NormalOverlay != nullptr) ? NormalOverlay->GetTriangle(TriangleID) : FIndex3i::Zero();
+			FIndex3i TriColor = (ColorOverlay != nullptr) ? ColorOverlay->GetTriangle(TriangleID) : FIndex3i::Zero();
+
+			FColor UniformTriColor = ConstantVertexColor;
+			if (bUsePerTriangleColor && PerTriangleColorFunc != nullptr)
+			{
+				UniformTriColor = PerTriangleColorFunc(Mesh, TriangleID);
+				bHaveColors = false;
+			}
+
+			for (int j = 0; j < 3; ++j)
+			{
+				RenderBuffers->PositionVertexBuffer.VertexPosition(VertIdx) = (FVector)Mesh->GetVertex(Tri[j]);
+
+				FVector3f Normal = (NormalOverlay != nullptr && TriNormal[j] != FDynamicMesh3::InvalidID) ?
+					NormalOverlay->GetElement(TriNormal[j]) : Mesh->GetVertexNormal(Tri[j]);
+
+				// get tangents
+				TangentsFunc(Tri[j], TriangleID, j, Normal, TangentX, TangentY);
+
+				RenderBuffers->StaticMeshVertexBuffer.SetVertexTangents(VertIdx, (FVector)TangentX, (FVector)TangentY, (FVector)Normal);
+
+				for (int32 k = 0; k < NumTexCoords; ++k)
+				{
+					FVector2f UV = (UVTriangles[k][j] != FDynamicMesh3::InvalidID) ?
+						UVOverlays[k]->GetElement(UVTriangles[k][j]) : FVector2f::Zero();
+					RenderBuffers->StaticMeshVertexBuffer.SetVertexUV(VertIdx, k, (FVector2D)UV);
+				}
+
+				FColor VertexFColor = (bHaveColors && TriColor[j] != FDynamicMesh3::InvalidID) ?
+					FColor(ColorOverlay->GetElement(TriColor[j])) : UniformTriColor;
+
+				RenderBuffers->ColorVertexBuffer.VertexColor(VertIdx) = VertexFColor;
+
+				RenderBuffers->IndexBuffer.Indices[TriIdx++] = VertIdx;		// currently TriIdx == VertIdx so we don't really need both...
+				VertIdx++;
+			}
+
+			if (bBuildTriangleList)
+			{
+				RenderBuffers->Triangles->Add(TriangleID);
+			}
+		}
+
+		// split triangles into secondary buffer (at bit redudant since we just built IndexBuffer, but we may optionally duplicate triangles in the future
+		if (bUseSecondaryTriBuffers)
+		{
+			RenderBuffers->bEnableSecondaryIndexBuffer = true;
+			UpdateSecondaryTriangleBuffer(RenderBuffers, Mesh, false);
+		}
+	}
+
+
+
+
+	/**
+	 * Filter the triangles in a FMeshRenderBufferSet into the SecondaryIndexBuffer.
+	 * Requires that RenderBuffers->Triangles has been initialized.
+	 * @param bDuplicate if set, then primary IndexBuffer is unmodified and SecondaryIndexBuffer contains duplicates. Otherwise triangles are sorted via predicate into either primary or secondary.
+	 */
+	void UpdateSecondaryTriangleBuffer(
+		FMeshRenderBufferSet* RenderBuffers,
+		const FDynamicMesh3* Mesh,
+		bool bDuplicate)
+	{
+		if (ensure(bUseSecondaryTriBuffers == true && RenderBuffers->Triangles.IsSet()) == false)
+		{
+			return;
+		}
+
+		const TArray<int32>& TriangleIDs = RenderBuffers->Triangles.GetValue();
+		int NumTris = TriangleIDs.Num();
+		TArray<uint32>& Indices = RenderBuffers->IndexBuffer.Indices;
+		TArray<uint32>& SecondaryIndices = RenderBuffers->SecondaryIndexBuffer.Indices;
+
+		RenderBuffers->SecondaryIndexBuffer.Indices.Reset();
+		if (bDuplicate == false)
+		{
+			RenderBuffers->IndexBuffer.Indices.Reset();
+		}
+		for ( int k = 0; k < NumTris; ++k)
+		{
+			int TriangleID = TriangleIDs[k];
+			bool bInclude = SecondaryTriFilterFunc(Mesh, TriangleID);
+			if (bInclude)
+			{
+				SecondaryIndices.Add(3*k);
+				SecondaryIndices.Add(3*k + 1);
+				SecondaryIndices.Add(3*k + 2);
+			} 
+			else if (bDuplicate == false)
+			{
+				Indices.Add(3*k);
+				Indices.Add(3*k + 1);
+				Indices.Add(3*k + 2);
+			}
+		}
+	}
+
+
+	/**
+	 * RecomputeRenderBufferTriangleIndexSets re-sorts the existing set of triangles in a FMeshRenderBufferSet
+	 * into primary and secondary index buffers. Note that UploadIndexBufferUpdate() must be called
+	 * after this function!
+	 */
+	void RecomputeRenderBufferTriangleIndexSets(
+		FMeshRenderBufferSet* RenderBuffers, 
+		const FDynamicMesh3* Mesh)
+	{
+		if (RenderBuffers->TriangleCount == 0)
+		{
+			return;
+		}
+		if (ensure(RenderBuffers->Triangles.IsSet() && RenderBuffers->Triangles->Num() > 0) == false)
+		{
+			return;
+		}
+
+		//bool bDuplicate = false;		// flag for future use, in case we want to draw all triangles in primary and duplicates in secondary...
+		RenderBuffers->IndexBuffer.Indices.Reset();
+		RenderBuffers->SecondaryIndexBuffer.Indices.Reset();
+		
+		TArray<uint32>& Indices = RenderBuffers->IndexBuffer.Indices;
+		TArray<uint32>& SecondaryIndices = RenderBuffers->SecondaryIndexBuffer.Indices;
+		const TArray<int32>& TriangleIDs = RenderBuffers->Triangles.GetValue();
+
+		int NumTris = TriangleIDs.Num();
+		for (int k = 0; k < NumTris; ++k)
+		{
+			int TriangleID = TriangleIDs[k];
+			bool bInclude = SecondaryTriFilterFunc(Mesh, TriangleID);
+			if (bInclude)
+			{
+				SecondaryIndices.Add(3 * k);
+				SecondaryIndices.Add(3 * k + 1);
+				SecondaryIndices.Add(3 * k + 2);
+			}
+			else // if (bDuplicate == false)
+			{
+				Indices.Add(3 * k);
+				Indices.Add(3 * k + 1);
+				Indices.Add(3 * k + 2);
+			}
+		}
+	}
+
+
+
+	/**
+	 * Update vertex positions/normals/colors of an existing set of render buffers.
+	 * Assumes that buffers were created with unshared vertices, ie three vertices per triangle, eg by InitializeBuffersFromOverlays()
+	 */
+	template<typename TriangleEnumerable>
+	void UpdateVertexBuffersFromOverlays(
+		FMeshRenderBufferSet* RenderBuffers,
+		const FDynamicMesh3* Mesh,
+		int NumTriangles, TriangleEnumerable Enumerable,
+		const FDynamicMeshNormalOverlay* NormalOverlay,
+		const FDynamicMeshColorOverlay* ColorOverlay,
+		TFunctionRef<void(int, int, int, const FVector3f&, FVector3f&, FVector3f&)> TangentsFunc,
+		bool bUpdatePositions = true,
+		bool bUpdateNormals = false,
+		bool bUpdateColors = false)
+	{
+		if (RenderBuffers->TriangleCount == 0)
+		{
+			return;
+		}
+
+		bool bHaveColors = (ColorOverlay != nullptr) && (bIgnoreVertexColors == false);
+
+		int NumVertices = NumTriangles * 3;
+		if ( (bUpdatePositions && ensure(RenderBuffers->PositionVertexBuffer.GetNumVertices() == NumVertices) == false )
+			|| (bUpdateNormals && ensure(RenderBuffers->StaticMeshVertexBuffer.GetNumVertices() == NumVertices) == false )
+			|| (bUpdateColors && ensure(RenderBuffers->ColorVertexBuffer.GetNumVertices() == NumVertices) == false ) )
+		{
+			return;
+		}
+
+		int VertIdx = 0;
+		FVector3f TangentX, TangentY;
+		for (int TriangleID : Enumerable)
+		{
+			FIndex3i Tri = Mesh->GetTriangle(TriangleID);
+
+			FIndex3i TriNormal = (bUpdateNormals && NormalOverlay != nullptr) ? NormalOverlay->GetTriangle(TriangleID) : FIndex3i::Zero();
+			FIndex3i TriColor = (bUpdateColors && ColorOverlay != nullptr) ? ColorOverlay->GetTriangle(TriangleID) : FIndex3i::Zero();
+
+			FColor UniformTriColor = ConstantVertexColor;
+			if (bUpdateColors && bUsePerTriangleColor && PerTriangleColorFunc != nullptr)
+			{
+				UniformTriColor = PerTriangleColorFunc(Mesh, TriangleID);
+				bHaveColors = false;
+			}
+
+			for (int j = 0; j < 3; ++j)
+			{
+				if (bUpdatePositions)
+				{
+					RenderBuffers->PositionVertexBuffer.VertexPosition(VertIdx) = (FVector)Mesh->GetVertex(Tri[j]);
+				}
+
+				if (bUpdateNormals)
+				{
+					// get normal and tangent
+					FVector3f Normal = (NormalOverlay != nullptr && TriNormal[j] != FDynamicMesh3::InvalidID) ?
+						NormalOverlay->GetElement(TriNormal[j]) : Mesh->GetVertexNormal(Tri[j]);
+					TangentsFunc(Tri[j], TriangleID, j, Normal, TangentX, TangentY);
+
+					RenderBuffers->StaticMeshVertexBuffer.SetVertexTangents(VertIdx, (FVector)TangentX, (FVector)TangentY, (FVector)Normal);
+				}
+
+				if (bUpdateColors)
+				{
+					FColor VertexFColor = (bHaveColors  && TriColor[j] != FDynamicMesh3::InvalidID) ?
+						FColor(ColorOverlay->GetElement(TriColor[j])) : UniformTriColor;
+					RenderBuffers->ColorVertexBuffer.VertexColor(VertIdx) = VertexFColor;
+				}
+
+				VertIdx++;
+			}
+		}
+	}
+
+
+
+	/**
+	 * Update vertex uvs of an existing set of render buffers.
+	 * Assumes that buffers were created with unshared vertices, ie three vertices per triangle, eg by InitializeBuffersFromOverlays()
+	 */
+	template<typename TriangleEnumerable>
+	void UpdateVertexUVBufferFromOverlays(
+		FMeshRenderBufferSet* RenderBuffers,
+		const FDynamicMesh3* Mesh,
+		int32 NumTriangles, TriangleEnumerable Enumerable,
+		const FDynamicMeshUVOverlay* UVOverlay, int32 UVIndex)
+	{
+		if (RenderBuffers->TriangleCount == 0)
+		{
+			return;
+		}
+		int NumVertices = NumTriangles * 3;
+		if ( ensure(RenderBuffers->StaticMeshVertexBuffer.GetNumVertices() == NumVertices) == false)
+		{
+			return;
+		}
+
+		int VertIdx = 0;
+		for (int TriangleID : Enumerable)
+		{
+			FIndex3i UVTri = UVOverlay->GetTriangle(TriangleID);
+			for (int j = 0; j < 3; ++j)
+			{
+				FVector2f UV = UVOverlay->GetElement(UVTri[j]);
+				RenderBuffers->StaticMeshVertexBuffer.SetVertexUV(VertIdx, UVIndex, (FVector2D)UV);
+				VertIdx++;
+			}
+		}
+	}
+
+
+
+
+	/**
+	 * @return number of active materials
+	 */
+	virtual int32 GetNumMaterials() const
+	{
+		return ParentBaseComponent->GetNumMaterials();
+	}
+
+	/**
+	 * Safe GetMaterial function that will never return nullptr
+	 */
+	virtual UMaterialInterface* GetMaterial(int32 k) const
+	{
+		UMaterialInterface* Material = ParentBaseComponent->GetMaterial(k);
+		return (Material != nullptr) ? Material : UMaterial::GetDefaultMaterial(MD_Surface);
+	}
+
+
+	/**
+	 * This needs to be called if the set of active materials changes, otherwise
+	 * the check in FPrimitiveSceneProxy::VerifyUsedMaterial() will fail if an override
+	 * material is set, if materials change, etc, etc
+	 */
+	virtual void UpdatedReferencedMaterials()
+	{
+#if WITH_EDITOR
+		TArray<UMaterialInterface*> Materials;
+		ParentBaseComponent->GetUsedMaterials(Materials, true);
+
+		// Temporarily disable material verification while the enqueued render command is in flight.
+		// The original value for bVerifyUsedMaterials gets restored when the command is executed.
+		// If we do not do this, material verification might spuriously fail in cases where the render command for changing
+		// the verfifcation material is still in flight but the render thread is already trying to render the mesh.
+		const uint8 bRestoreVerifyUsedMaterials = bVerifyUsedMaterials;
+		bVerifyUsedMaterials = false;
+
+		ENQUEUE_RENDER_COMMAND(FMeshRenderBufferSetDestroy)(
+			[this, Materials, bRestoreVerifyUsedMaterials](FRHICommandListImmediate& RHICmdList)
+		{
+			this->SetUsedMaterialForVerification(Materials);
+			this->bVerifyUsedMaterials = bRestoreVerifyUsedMaterials;
+		});
+#endif
+	}
+
+
+	//
+	// FBaseDynamicMeshSceneProxy implementation
+	//
+
+
+	/**
+	 * Render set of active RenderBuffers returned by GetActiveRenderBufferSets
+	 */
+	virtual void GetDynamicMeshElements(
+		const TArray<const FSceneView*>& Views, 
+		const FSceneViewFamily& ViewFamily, 
+		uint32 VisibilityMap, 
+		FMeshElementCollector& Collector) const override
+	{
+		QUICK_SCOPE_CYCLE_COUNTER(STAT_OctreeDynamicMeshSceneProxy_GetDynamicMeshElements);
+
+		const bool bWireframe = (AllowDebugViewmodes() && ViewFamily.EngineShowFlags.Wireframe)
+			|| ParentBaseComponent->GetEnableWireframeRenderPass();
+
+		// set up wireframe material. Probably bad to reference GEngine here...also this material is very bad?
+		FMaterialRenderProxy* WireframeMaterialProxy = nullptr;
+		if (bWireframe)
+		{
+			FColoredMaterialRenderProxy* WireframeMaterialInstance = new FColoredMaterialRenderProxy(
+				GEngine->WireframeMaterial ? GEngine->WireframeMaterial->GetRenderProxy() : nullptr,
+				FLinearColor(0, 0.5f, 1.f)
+			);
+			Collector.RegisterOneFrameMaterialProxy(WireframeMaterialInstance);
+			WireframeMaterialProxy = WireframeMaterialInstance;
+		}
+
+		ESceneDepthPriorityGroup DepthPriority = SDPG_World;
+
+		TArray<FMeshRenderBufferSet*> Buffers;
+		GetActiveRenderBufferSets(Buffers);
+
+		FMaterialRenderProxy* SecondaryMaterialProxy =
+			ParentBaseComponent->HasSecondaryRenderMaterial() ? ParentBaseComponent->GetSecondaryRenderMaterial()->GetRenderProxy() : nullptr;
+		bool bDrawSecondaryBuffers = ParentBaseComponent->GetSecondaryBuffersVisibility();
+
+		for (int32 ViewIndex = 0; ViewIndex < Views.Num(); ViewIndex++)
+		{
+			if (VisibilityMap & (1 << ViewIndex))
+			{
+				const FSceneView* View = Views[ViewIndex];
+
+				bool bHasPrecomputedVolumetricLightmap;
+				FMatrix PreviousLocalToWorld;
+				int32 SingleCaptureIndex;
+				bool bOutputVelocity;
+				GetScene().GetPrimitiveUniformShaderParameters_RenderThread(GetPrimitiveSceneInfo(), bHasPrecomputedVolumetricLightmap, PreviousLocalToWorld, SingleCaptureIndex, bOutputVelocity);
+
+				// Draw the mesh.
+				for (FMeshRenderBufferSet* BufferSet : Buffers)
+				{
+					UMaterialInterface* UseMaterial = BufferSet->Material;
+					if (ParentBaseComponent->HasOverrideRenderMaterial(0))
+					{
+						UseMaterial = ParentBaseComponent->GetOverrideRenderMaterial(0);
+					}
+					FMaterialRenderProxy* MaterialProxy = UseMaterial->GetRenderProxy();
+
+					if (BufferSet->TriangleCount == 0)
+					{
+						continue;
+					}
+
+					// lock buffers so that they aren't modified while we are submitting them
+					FScopeLock BuffersLock(&BufferSet->BuffersLock);
+
+					// do we need separate one of these for each MeshRenderBufferSet?
+					FDynamicPrimitiveUniformBuffer& DynamicPrimitiveUniformBuffer = Collector.AllocateOneFrameResource<FDynamicPrimitiveUniformBuffer>();
+					DynamicPrimitiveUniformBuffer.Set(GetLocalToWorld(), PreviousLocalToWorld, GetBounds(), GetLocalBounds(), true, bHasPrecomputedVolumetricLightmap, DrawsVelocity(), bOutputVelocity);
+
+					if (BufferSet->IndexBuffer.Indices.Num() > 0)
+					{
+						// Unlike most meshes, which just use the wireframe material in wireframe mode, we draw the wireframe on top of the normal material if needed,
+						// as this is easier to interpret. However, we do not do this in ortho viewports, where it frequently causes the our edit gizmo to be hidden 
+						// beneath the material. So, only draw the base material if we are in perspective mode, or we're in ortho but not in wireframe.
+						if (View->IsPerspectiveProjection() || !(AllowDebugViewmodes() && ViewFamily.EngineShowFlags.Wireframe))
+						{
+							DrawBatch(Collector, *BufferSet, BufferSet->IndexBuffer, MaterialProxy, false, DepthPriority, ViewIndex, DynamicPrimitiveUniformBuffer);
+						}
+						if (bWireframe)
+						{
+							DrawBatch(Collector, *BufferSet, BufferSet->IndexBuffer, WireframeMaterialProxy, true, DepthPriority, ViewIndex, DynamicPrimitiveUniformBuffer);
+						}
+					}
+
+					// draw secondary buffer if we have it, falling back to base material if we don't have the Secondary material
+					FMaterialRenderProxy* UseSecondaryMaterialProxy = (SecondaryMaterialProxy != nullptr) ? SecondaryMaterialProxy : MaterialProxy;
+					if (bDrawSecondaryBuffers && BufferSet->SecondaryIndexBuffer.Indices.Num() > 0 && UseSecondaryMaterialProxy != nullptr)
+					{
+						DrawBatch(Collector, *BufferSet, BufferSet->SecondaryIndexBuffer, UseSecondaryMaterialProxy, false, DepthPriority, ViewIndex, DynamicPrimitiveUniformBuffer);
+						if (bWireframe)
+						{
+							DrawBatch(Collector, *BufferSet, BufferSet->SecondaryIndexBuffer, UseSecondaryMaterialProxy, true, DepthPriority, ViewIndex, DynamicPrimitiveUniformBuffer);
+						}
+					}
+				}
+			}
+		}
+	}
+
+
+
+	/**
+	 * Draw a single-frame FMeshBatch for a FMeshRenderBufferSet
+	 */
+	virtual void DrawBatch(FMeshElementCollector& Collector,
+		const FMeshRenderBufferSet& RenderBuffers,
+		const FDynamicMeshIndexBuffer32& IndexBuffer,
+		FMaterialRenderProxy* UseMaterial,
+		bool bWireframe,
+		ESceneDepthPriorityGroup DepthPriority,
+		int ViewIndex,
+		FDynamicPrimitiveUniformBuffer& DynamicPrimitiveUniformBuffer) const
+	{
+		FMeshBatch& Mesh = Collector.AllocateMesh();
+		FMeshBatchElement& BatchElement = Mesh.Elements[0];
+		BatchElement.IndexBuffer = &IndexBuffer;
+		Mesh.bWireframe = bWireframe;
+		Mesh.VertexFactory = &RenderBuffers.VertexFactory;
+		Mesh.MaterialRenderProxy = UseMaterial;
+
+		BatchElement.PrimitiveUniformBufferResource = &DynamicPrimitiveUniformBuffer.UniformBuffer;
+
+		BatchElement.FirstIndex = 0;
+		BatchElement.NumPrimitives = IndexBuffer.Indices.Num() / 3;
+		BatchElement.MinVertexIndex = 0;
+		BatchElement.MaxVertexIndex = RenderBuffers.PositionVertexBuffer.GetNumVertices() - 1;
+		Mesh.ReverseCulling = IsLocalToWorldDeterminantNegative();
+		Mesh.Type = PT_TriangleList;
+		Mesh.DepthPriorityGroup = DepthPriority;
+		Mesh.bCanApplyViewModeOverrides = false;
+		Collector.AddMesh(ViewIndex, Mesh);
+	}
+
+
+
+
+
+
+
+};
\ No newline at end of file
diff --git a/Engine/Plugins/Experimental/MeshModelingToolset/Source/ModelingComponents/Private/Components/DynamicMeshComponent.cpp b/Engine/Plugins/Experimental/MeshModelingToolset/Source/ModelingComponents/Private/Components/DynamicMeshComponent.cpp
new file mode 100644
--- /dev/null
+++ b/Engine/Plugins/Experimental/MeshModelingToolset/Source/ModelingComponents/Private/Components/DynamicMeshComponent.cpp
@@ -0,0 +1,957 @@
+// Copyright Epic Games, Inc. All Rights Reserved. 
+
+#include "Components/DynamicMeshComponent.h"
+#include "RenderingThread.h"
+#include "RenderResource.h"
+#include "PrimitiveViewRelevance.h"
+#include "PrimitiveSceneProxy.h"
+#include "VertexFactory.h"
+#include "MaterialShared.h"
+#include "Engine/CollisionProfile.h"
+#include "Materials/Material.h"
+#include "LocalVertexFactory.h"
+#include "SceneManagement.h"
+#include "DynamicMeshBuilder.h"
+#include "EngineGlobals.h"
+#include "Engine/Engine.h"
+#include "StaticMeshResources.h"
+#include "StaticMeshAttributes.h"
+#include "Async/Async.h"
+
+#include "DynamicMeshAttributeSet.h"
+#include "MeshNormals.h"
+#include "MeshDescriptionToDynamicMesh.h"
+
+#include "Changes/MeshVertexChange.h"
+#include "Changes/MeshChange.h"
+#include "DynamicMeshChangeTracker.h"
+#include "MeshTransforms.h"
+
+// default proxy for this component
+#include "Components/DynamicMeshSceneProxy.h"
+
+#include "ExplicitUseGeometryMathTypes.h"		// using UE::Geometry::(math types)
+using namespace UE::Geometry;
+
+namespace
+{
+	// probably should be something defined for the whole tool framework...
+#if WITH_EDITOR
+	static EAsyncExecution DynamicMeshComponentAsyncExecTarget = EAsyncExecution::LargeThreadPool;
+#else
+	static EAsyncExecution DynamicMeshComponentAsyncExecTarget = EAsyncExecution::ThreadPool;
+#endif
+}
+
+
+
+
+UDynamicMeshComponent::UDynamicMeshComponent(const FObjectInitializer& ObjectInitializer)
+	: Super(ObjectInitializer)
+{
+	PrimaryComponentTick.bCanEverTick = false;
+
+	SetCollisionProfileName(UCollisionProfile::NoCollision_ProfileName);
+
+	MeshObject = CreateDefaultSubobject<UDynamicMesh>(TEXT("DynamicMesh"));
+	//MeshObject->SetFlags(RF_Transactional);
+
+	MeshObjectChangedHandle = MeshObject->OnMeshChanged().AddUObject(this, &UDynamicMeshComponent::OnMeshObjectChanged);
+
+	ResetProxy();
+}
+
+
+void UDynamicMeshComponent::PostLoad()
+{
+	Super::PostLoad();
+
+	check(MeshObject != nullptr);
+	MeshObjectChangedHandle = MeshObject->OnMeshChanged().AddUObject(this, &UDynamicMeshComponent::OnMeshObjectChanged);
+
+	ResetProxy();
+}
+
+
+
+void UDynamicMeshComponent::SetMesh(UE::Geometry::FDynamicMesh3&& MoveMesh)
+{
+	MeshObject->SetMesh(MoveTemp(MoveMesh));
+}
+
+
+void UDynamicMeshComponent::ProcessMesh(
+	TFunctionRef<void(const UE::Geometry::FDynamicMesh3&)> ProcessFunc ) const
+{
+	MeshObject->ProcessMesh(ProcessFunc);
+}
+
+
+void UDynamicMeshComponent::EditMesh(TFunctionRef<void(UE::Geometry::FDynamicMesh3&)> EditFunc,
+										   EDynamicMeshComponentRenderUpdateMode UpdateMode )
+{
+	MeshObject->EditMesh(EditFunc);
+	if (UpdateMode != EDynamicMeshComponentRenderUpdateMode::NoUpdate)
+	{
+		NotifyMeshUpdated();
+	}
+}
+
+
+void UDynamicMeshComponent::SetRenderMeshPostProcessor(TUniquePtr<IRenderMeshPostProcessor> Processor)
+{
+	RenderMeshPostProcessor = MoveTemp(Processor);
+	if (RenderMeshPostProcessor)
+	{
+		if (!RenderMesh)
+		{
+			RenderMesh = MakeUnique<FDynamicMesh3>(*GetMesh());
+		}
+	}
+	else
+	{
+		// No post processor, no render mesh
+		RenderMesh = nullptr;
+	}
+}
+
+FDynamicMesh3* UDynamicMeshComponent::GetRenderMesh()
+{
+	if (RenderMeshPostProcessor && RenderMesh)
+	{
+		return RenderMesh.Get();
+	}
+	else
+	{
+		return GetMesh();
+	}
+}
+
+const FDynamicMesh3* UDynamicMeshComponent::GetRenderMesh() const
+{
+	if (RenderMeshPostProcessor && RenderMesh)
+	{
+		return RenderMesh.Get();
+	}
+	else
+	{
+		return GetMesh();
+	}
+}
+
+
+
+
+void UDynamicMeshComponent::ApplyTransform(const UE::Geometry::FTransform3d& Transform, bool bInvert)
+{
+	MeshObject->EditMesh([&](FDynamicMesh3& EditMesh)
+	{
+		if (bInvert)
+		{
+			MeshTransforms::ApplyTransformInverse(EditMesh, Transform);
+		}
+		else
+		{
+			MeshTransforms::ApplyTransform(EditMesh, Transform);
+		}
+	}, EDynamicMeshChangeType::DeformationEdit);
+}
+
+
+
+void UDynamicMeshComponent::SetTangentsType(EDynamicMeshComponentTangentsMode NewTangentsType)
+{
+	if (NewTangentsType != TangentsType)
+	{
+		TangentsType = NewTangentsType;
+		InvalidateAutoCalculatedTangents();
+	}
+}
+
+void UDynamicMeshComponent::InvalidateAutoCalculatedTangents() 
+{ 
+	bAutoCalculatedTangentsValid = false; 
+}
+
+const UE::Geometry::FMeshTangentsf* UDynamicMeshComponent::GetAutoCalculatedTangents() 
+{ 
+	if (TangentsType == EDynamicMeshComponentTangentsMode::AutoCalculated)
+	{
+		UpdateAutoCalculatedTangents();
+		return (bAutoCalculatedTangentsValid) ? &AutoCalculatedTangents : nullptr;
+	}
+	return nullptr;
+}
+
+void UDynamicMeshComponent::UpdateAutoCalculatedTangents()
+{
+	if (TangentsType == EDynamicMeshComponentTangentsMode::AutoCalculated && bAutoCalculatedTangentsValid == false)
+	{
+		GetDynamicMesh()->ProcessMesh([&](const FDynamicMesh3& Mesh)
+		{
+			const FDynamicMeshUVOverlay* UVOverlay = Mesh.Attributes()->PrimaryUV();
+			const FDynamicMeshNormalOverlay* NormalOverlay = Mesh.Attributes()->PrimaryNormals();
+			AutoCalculatedTangents.SetMesh(&Mesh);
+			AutoCalculatedTangents.ComputeTriVertexTangents(NormalOverlay, UVOverlay, FComputeTangentsOptions());
+			AutoCalculatedTangents.SetMesh(nullptr);
+		});
+
+		bAutoCalculatedTangentsValid = true;
+	}
+}
+
+
+
+
+void UDynamicMeshComponent::UpdateLocalBounds()
+{
+	LocalBounds = GetMesh()->GetBounds(true);
+	if (LocalBounds.MaxDim() <= 0)
+	{
+		LocalBounds = FAxisAlignedBox3d(FVector3d::Zero(), FMathf::ZeroTolerance);
+	}
+}
+
+FDynamicMeshSceneProxy* UDynamicMeshComponent::GetCurrentSceneProxy() 
+{ 
+	if (bProxyValid)
+	{
+		return (FDynamicMeshSceneProxy*)SceneProxy;
+	}
+	return nullptr;
+}
+
+
+void UDynamicMeshComponent::ResetProxy()
+{
+	bProxyValid = false;
+	InvalidateAutoCalculatedTangents();
+
+	// Need to recreate scene proxy to send it over
+	MarkRenderStateDirty();
+	UpdateLocalBounds();
+	UpdateBounds();
+
+	// this seems speculative, ie we may not actually have a mesh update, but we currently ResetProxy() in lots
+	// of places where that is what it means
+	GetDynamicMesh()->PostRealtimeUpdate();
+}
+
+void UDynamicMeshComponent::NotifyMeshUpdated()
+{
+	if (RenderMeshPostProcessor)
+	{
+		RenderMeshPostProcessor->ProcessMesh(*GetMesh(), *RenderMesh);
+	}
+
+	ResetProxy();
+}
+
+
+void UDynamicMeshComponent::FastNotifyColorsUpdated()
+{
+	// should not be using fast paths if we have to run mesh postprocessor
+	if (ensure(!RenderMeshPostProcessor) == false)
+	{
+		RenderMeshPostProcessor->ProcessMesh(*GetMesh(), *RenderMesh);
+		ResetProxy();
+		return;
+	}
+
+	FDynamicMeshSceneProxy* Proxy = GetCurrentSceneProxy();
+	if (Proxy)
+	{
+		if (HasTriangleColorFunction() && Proxy->bUsePerTriangleColor == false )
+		{
+			Proxy->bUsePerTriangleColor = true;
+			Proxy->PerTriangleColorFunc = [this](const FDynamicMesh3* MeshIn, int TriangleID) { return GetTriangleColor(MeshIn, TriangleID); };
+		} 
+		else if ( !HasTriangleColorFunction() && Proxy->bUsePerTriangleColor == true)
+		{
+			Proxy->bUsePerTriangleColor = false;
+			Proxy->PerTriangleColorFunc = nullptr;
+		}
+
+		Proxy->FastUpdateVertices(false, false, true, false);
+		//MarkRenderDynamicDataDirty();
+	}
+	else
+	{
+		ResetProxy();
+	}
+}
+
+
+
+void UDynamicMeshComponent::FastNotifyPositionsUpdated(bool bNormals, bool bColors, bool bUVs)
+{
+	// should not be using fast paths if we have to run mesh postprocessor
+	if (ensure(!RenderMeshPostProcessor) == false)
+	{
+		RenderMeshPostProcessor->ProcessMesh(*GetMesh(), *RenderMesh);
+		ResetProxy();
+		return;
+	}
+
+	FDynamicMeshSceneProxy* Proxy = GetCurrentSceneProxy();
+	if (Proxy)
+	{
+		// calculate bounds while we are updating vertices
+		TFuture<void> UpdateBoundsCalc;
+		UpdateBoundsCalc = Async(DynamicMeshComponentAsyncExecTarget, [this]()
+		{
+			TRACE_CPUPROFILER_EVENT_SCOPE(SimpleDynamicMeshComponent_FastPositionsUpdate_AsyncBoundsUpdate);
+			UpdateLocalBounds();
+		});
+
+		GetCurrentSceneProxy()->FastUpdateVertices(true, bNormals, bColors, bUVs);
+
+		//MarkRenderDynamicDataDirty();
+		MarkRenderTransformDirty();
+		UpdateBoundsCalc.Wait();
+		UpdateBounds();
+
+		GetDynamicMesh()->PostRealtimeUpdate();
+	}
+	else
+	{
+		ResetProxy();
+	}
+}
+
+
+void UDynamicMeshComponent::FastNotifyVertexAttributesUpdated(bool bNormals, bool bColors, bool bUVs)
+{
+	// should not be using fast paths if we have to run mesh postprocessor
+	if (ensure(!RenderMeshPostProcessor) == false)
+	{
+		RenderMeshPostProcessor->ProcessMesh(*GetMesh(), *RenderMesh);
+		ResetProxy();
+		return;
+	}
+
+	FDynamicMeshSceneProxy* Proxy = GetCurrentSceneProxy();
+	if (Proxy && ensure(bNormals || bColors || bUVs) )
+	{
+		GetCurrentSceneProxy()->FastUpdateVertices(false, bNormals, bColors, bUVs);
+		//MarkRenderDynamicDataDirty();
+		//MarkRenderTransformDirty();
+
+		GetDynamicMesh()->PostRealtimeUpdate();
+	}
+	else
+	{
+		ResetProxy();
+	}
+}
+
+
+void UDynamicMeshComponent::FastNotifyVertexAttributesUpdated(EMeshRenderAttributeFlags UpdatedAttributes)
+{
+	// should not be using fast paths if we have to run mesh postprocessor
+	if (ensure(!RenderMeshPostProcessor) == false)
+	{
+		RenderMeshPostProcessor->ProcessMesh(*GetMesh(), *RenderMesh);
+		ResetProxy();
+		return;
+	}
+
+	FDynamicMeshSceneProxy* Proxy = GetCurrentSceneProxy();
+	if (Proxy && ensure(UpdatedAttributes != EMeshRenderAttributeFlags::None))
+	{
+		bool bPositions = (UpdatedAttributes & EMeshRenderAttributeFlags::Positions) != EMeshRenderAttributeFlags::None;
+
+		// calculate bounds while we are updating vertices
+		TFuture<void> UpdateBoundsCalc;
+		if (bPositions)
+		{
+			UpdateBoundsCalc = Async(DynamicMeshComponentAsyncExecTarget, [this]()
+			{
+				TRACE_CPUPROFILER_EVENT_SCOPE(SimpleDynamicMeshComponent_FastVertexAttribUpdate_AsyncBoundsUpdate);
+				UpdateLocalBounds();
+			});
+		}
+
+		GetCurrentSceneProxy()->FastUpdateVertices(bPositions,
+			(UpdatedAttributes & EMeshRenderAttributeFlags::VertexNormals) != EMeshRenderAttributeFlags::None,
+			(UpdatedAttributes & EMeshRenderAttributeFlags::VertexColors) != EMeshRenderAttributeFlags::None,
+			(UpdatedAttributes & EMeshRenderAttributeFlags::VertexUVs) != EMeshRenderAttributeFlags::None);
+
+		if (bPositions)
+		{
+			MarkRenderTransformDirty();
+			UpdateBoundsCalc.Wait();
+			UpdateBounds();
+		}
+
+		GetDynamicMesh()->PostRealtimeUpdate();
+	}
+	else
+	{
+		ResetProxy();
+	}
+}
+
+
+void UDynamicMeshComponent::FastNotifyUVsUpdated()
+{
+	FastNotifyVertexAttributesUpdated(EMeshRenderAttributeFlags::VertexUVs);
+}
+
+
+
+void UDynamicMeshComponent::FastNotifySecondaryTrianglesChanged()
+{
+	// should not be using fast paths if we have to run mesh postprocessor
+	if (ensure(!RenderMeshPostProcessor) == false)
+	{
+		RenderMeshPostProcessor->ProcessMesh(*GetMesh(), *RenderMesh);
+		ResetProxy();
+		return;
+	}
+
+	FDynamicMeshSceneProxy* Proxy = GetCurrentSceneProxy();
+	if (Proxy)
+	{
+		GetCurrentSceneProxy()->FastUpdateAllIndexBuffers();
+		GetDynamicMesh()->PostRealtimeUpdate();
+	}
+	else
+	{
+		ResetProxy();
+	}
+}
+
+
+void UDynamicMeshComponent::FastNotifyTriangleVerticesUpdated(const TArray<int32>& Triangles, EMeshRenderAttributeFlags UpdatedAttributes)
+{
+	// should not be using fast paths if we have to run mesh postprocessor
+	if (ensure(!RenderMeshPostProcessor) == false)
+	{
+		RenderMeshPostProcessor->ProcessMesh(*GetMesh(), *RenderMesh);
+		ResetProxy();
+		return;
+	}
+
+	bool bUpdateSecondarySort = (SecondaryTriFilterFunc) &&
+		((UpdatedAttributes & EMeshRenderAttributeFlags::SecondaryIndexBuffers) != EMeshRenderAttributeFlags::None);
+
+	FDynamicMeshSceneProxy* Proxy = GetCurrentSceneProxy();
+	if (!Proxy)
+	{
+		ResetProxy();
+	}
+	else if ( ! Decomposition )
+	{
+		FastNotifyVertexAttributesUpdated(UpdatedAttributes);
+		if (bUpdateSecondarySort)
+		{
+			Proxy->FastUpdateAllIndexBuffers();
+		}
+		GetDynamicMesh()->PostRealtimeUpdate();
+	}
+	else
+	{
+		// compute list of sets to update
+		TArray<int32> UpdatedSets;
+		{
+			TRACE_CPUPROFILER_EVENT_SCOPE(SimpleDynamicMeshComponent_FastVertexUpdate_FindSets);
+			for (int32 tid : Triangles)
+			{
+				int32 SetID = Decomposition->GetGroupForTriangle(tid);
+				UpdatedSets.AddUnique(SetID);
+			}
+		}
+
+		bool bPositions = (UpdatedAttributes & EMeshRenderAttributeFlags::Positions) != EMeshRenderAttributeFlags::None;
+
+		// calculate bounds while we are updating vertices
+		TFuture<void> UpdateBoundsCalc;
+		if (bPositions)
+		{
+			UpdateBoundsCalc = Async(DynamicMeshComponentAsyncExecTarget, [this]()
+			{
+				TRACE_CPUPROFILER_EVENT_SCOPE(SimpleDynamicMeshComponent_FastVertexUpdate_AsyncBoundsUpdate);
+				UpdateLocalBounds();
+			});
+		}
+
+		// update the render buffers
+		{
+			TRACE_CPUPROFILER_EVENT_SCOPE(SimpleDynamicMeshComponent_FastVertexUpdate_ApplyUpdate);
+			Proxy->FastUpdateVertices(UpdatedSets, bPositions,
+				(UpdatedAttributes & EMeshRenderAttributeFlags::VertexNormals) != EMeshRenderAttributeFlags::None,
+				(UpdatedAttributes & EMeshRenderAttributeFlags::VertexColors) != EMeshRenderAttributeFlags::None,
+				(UpdatedAttributes & EMeshRenderAttributeFlags::VertexUVs) != EMeshRenderAttributeFlags::None);
+		}
+
+		if (bUpdateSecondarySort)
+		{
+			Proxy->FastUpdateIndexBuffers(UpdatedSets);
+		}
+
+		if (bPositions)
+		{
+			TRACE_CPUPROFILER_EVENT_SCOPE(SimpleDynamicMeshComponent_FastVertexUpdate_FinalPositionsUpdate);
+			MarkRenderTransformDirty();
+			UpdateBoundsCalc.Wait();
+			UpdateBounds();
+		}
+
+		GetDynamicMesh()->PostRealtimeUpdate();
+	}
+}
+
+
+
+
+void UDynamicMeshComponent::FastNotifyTriangleVerticesUpdated(const TSet<int32>& Triangles, EMeshRenderAttributeFlags UpdatedAttributes)
+{
+	// should not be using fast paths if we have to run mesh postprocessor
+	if (ensure(!RenderMeshPostProcessor) == false)
+	{
+		RenderMeshPostProcessor->ProcessMesh(*GetMesh(), *RenderMesh);
+		ResetProxy();
+		return;
+	}
+
+	bool bUpdateSecondarySort = (SecondaryTriFilterFunc) &&
+		((UpdatedAttributes & EMeshRenderAttributeFlags::SecondaryIndexBuffers) != EMeshRenderAttributeFlags::None);
+
+	FDynamicMeshSceneProxy* Proxy = GetCurrentSceneProxy();
+	if (!Proxy)
+	{
+		ResetProxy();
+	}
+	else if (!Decomposition)
+	{
+		FastNotifyVertexAttributesUpdated(UpdatedAttributes);
+		if (bUpdateSecondarySort)
+		{
+			Proxy->FastUpdateAllIndexBuffers();
+		}
+		GetDynamicMesh()->PostRealtimeUpdate();
+	}
+	else
+	{
+		// compute list of sets to update
+		TArray<int32> UpdatedSets;
+		{
+			TRACE_CPUPROFILER_EVENT_SCOPE(SimpleDynamicMeshComponent_FastVertexUpdate_FindSets);
+			for (int32 tid : Triangles)
+			{
+				int32 SetID = Decomposition->GetGroupForTriangle(tid);
+				UpdatedSets.AddUnique(SetID);
+			}
+		}
+
+		bool bPositions = (UpdatedAttributes & EMeshRenderAttributeFlags::Positions) != EMeshRenderAttributeFlags::None;
+
+		// calculate bounds while we are updating vertices
+		TFuture<void> UpdateBoundsCalc;
+		if (bPositions)
+		{
+			UpdateBoundsCalc = Async(DynamicMeshComponentAsyncExecTarget, [this]()
+			{
+				TRACE_CPUPROFILER_EVENT_SCOPE(SimpleDynamicMeshComponent_FastVertexUpdate_AsyncBoundsUpdate);
+				UpdateLocalBounds();
+			});
+		}
+
+		// update the render buffers
+		{
+			TRACE_CPUPROFILER_EVENT_SCOPE(SimpleDynamicMeshComponent_FastVertexUpdate_ApplyUpdate);
+			Proxy->FastUpdateVertices(UpdatedSets, bPositions,
+				(UpdatedAttributes & EMeshRenderAttributeFlags::VertexNormals) != EMeshRenderAttributeFlags::None,
+				(UpdatedAttributes & EMeshRenderAttributeFlags::VertexColors) != EMeshRenderAttributeFlags::None,
+				(UpdatedAttributes & EMeshRenderAttributeFlags::VertexUVs) != EMeshRenderAttributeFlags::None);
+		}
+
+		{
+			TRACE_CPUPROFILER_EVENT_SCOPE(SimpleDynamicMeshComponent_FastVertexUpdate_UpdateIndexBuffers);
+			if (bUpdateSecondarySort)
+			{
+				Proxy->FastUpdateIndexBuffers(UpdatedSets);
+			}
+		}
+
+		// finish up, have to wait for background bounds recalculation here
+		if (bPositions)
+		{
+			TRACE_CPUPROFILER_EVENT_SCOPE(SimpleDynamicMeshComponent_FastVertexUpdate_FinalPositionsUpdate);
+			MarkRenderTransformDirty();
+			UpdateBoundsCalc.Wait();
+			UpdateBounds();
+		}
+
+		GetDynamicMesh()->PostRealtimeUpdate();
+	}
+}
+
+
+
+/**
+ * Compute the combined bounding-box of the Triangles array in parallel, by computing
+ * partial boxes for subsets of this array, and then combining those boxes.
+ * TODO: this should move to a pulbic utility function, and possibly the block-based ParallelFor
+ * should be refactored out into something more general, as this pattern is useful in many places...
+ */
+static FAxisAlignedBox3d ParallelComputeROIBounds(const FDynamicMesh3& Mesh, const TArray<int32>& Triangles)
+{
+	FAxisAlignedBox3d FinalBounds = FAxisAlignedBox3d::Empty();
+	FCriticalSection FinalBoundsLock;
+	int32 N = Triangles.Num();
+	constexpr int32 BlockSize = 4096;
+	int32 Blocks = (N / BlockSize) + 1;
+	ParallelFor(Blocks, [&](int bi)
+	{
+		FAxisAlignedBox3d BlockBounds = FAxisAlignedBox3d::Empty();
+		for (int32 k = 0; k < BlockSize; ++k)
+		{
+			int32 i = bi * BlockSize + k;
+			if (i < N)
+			{
+				int32 tid = Triangles[i];
+				const FIndex3i& TriV = Mesh.GetTriangleRef(tid);
+				BlockBounds.Contain(Mesh.GetVertexRef(TriV.A));
+				BlockBounds.Contain(Mesh.GetVertexRef(TriV.B));
+				BlockBounds.Contain(Mesh.GetVertexRef(TriV.C));
+			}
+		}
+		FinalBoundsLock.Lock();
+		FinalBounds.Contain(BlockBounds);
+		FinalBoundsLock.Unlock();
+	});
+	return FinalBounds;
+}
+
+
+
+TFuture<bool> UDynamicMeshComponent::FastNotifyTriangleVerticesUpdated_TryPrecompute(
+	const TArray<int32>& Triangles,
+	TArray<int32>& UpdateSetsOut,
+	FAxisAlignedBox3d& BoundsOut)
+{
+	if ((!!RenderMeshPostProcessor) || (GetCurrentSceneProxy() == nullptr) || (!Decomposition))
+	{
+		// is there a simpler way to do this? cannot seem to just make a TFuture<bool>...
+		return Async(DynamicMeshComponentAsyncExecTarget, []() { return false; });
+	}
+
+	return Async(DynamicMeshComponentAsyncExecTarget, [this, &Triangles, &UpdateSetsOut, &BoundsOut]()
+	{
+		TFuture<void> ComputeBounds = Async(DynamicMeshComponentAsyncExecTarget, [this, &BoundsOut, &Triangles]()
+		{
+			TRACE_CPUPROFILER_EVENT_SCOPE(SimpleDynamicMeshComponent_FastVertexUpdatePrecomp_CalcBounds);
+			BoundsOut = ParallelComputeROIBounds(*GetMesh(), Triangles);
+		});
+
+		TFuture<void> ComputeSets = Async(DynamicMeshComponentAsyncExecTarget, [this, &UpdateSetsOut, &Triangles]()
+		{
+			TRACE_CPUPROFILER_EVENT_SCOPE(SimpleDynamicMeshComponent_FastVertexUpdatePrecomp_FindSets);
+			int32 NumBuffers = Decomposition->Num();
+			TArray<std::atomic<bool>> BufferFlags;
+			BufferFlags.SetNum(NumBuffers);
+			for (int32 k = 0; k < NumBuffers; ++k)
+			{
+				BufferFlags[k] = false;
+			}
+			ParallelFor(Triangles.Num(), [&](int32 k)
+			{
+				int32 SetID = Decomposition->GetGroupForTriangle(Triangles[k]);
+				BufferFlags[SetID] = true;
+			});
+			UpdateSetsOut.Reset();
+			for (int32 k = 0; k < NumBuffers; ++k)
+			{
+				if (BufferFlags[k])
+				{
+					UpdateSetsOut.Add(k);
+				}
+			}
+
+		});
+
+		ComputeSets.Wait();
+		ComputeBounds.Wait();
+
+		return true;
+	});
+}
+
+
+void UDynamicMeshComponent::FastNotifyTriangleVerticesUpdated_ApplyPrecompute(
+	const TArray<int32>& Triangles,
+	EMeshRenderAttributeFlags UpdatedAttributes, 
+	TFuture<bool>& Precompute, 
+	const TArray<int32>& UpdateSets, 
+	const FAxisAlignedBox3d& UpdateSetBounds)
+{
+	Precompute.Wait();
+
+	bool bPrecomputeOK = Precompute.Get();
+	if (bPrecomputeOK == false || GetCurrentSceneProxy() == nullptr )
+	{
+		FastNotifyTriangleVerticesUpdated(Triangles, UpdatedAttributes);
+		return;
+	}
+
+	FDynamicMeshSceneProxy* Proxy = GetCurrentSceneProxy();
+	bool bPositions = (UpdatedAttributes & EMeshRenderAttributeFlags::Positions) != EMeshRenderAttributeFlags::None;
+	bool bUpdateSecondarySort = (SecondaryTriFilterFunc) &&
+		((UpdatedAttributes & EMeshRenderAttributeFlags::SecondaryIndexBuffers) != EMeshRenderAttributeFlags::None);
+
+	{
+		TRACE_CPUPROFILER_EVENT_SCOPE(SimpleDynamicMeshComponent_FastVertexUpdate_ApplyUpdate);
+		Proxy->FastUpdateVertices(UpdateSets, bPositions,
+			(UpdatedAttributes & EMeshRenderAttributeFlags::VertexNormals) != EMeshRenderAttributeFlags::None,
+			(UpdatedAttributes & EMeshRenderAttributeFlags::VertexColors) != EMeshRenderAttributeFlags::None,
+			(UpdatedAttributes & EMeshRenderAttributeFlags::VertexUVs) != EMeshRenderAttributeFlags::None);
+	}
+
+	{
+		TRACE_CPUPROFILER_EVENT_SCOPE(SimpleDynamicMeshComponent_FastVertexUpdate_UpdateIndexBuffers);
+		if (bUpdateSecondarySort)
+		{
+			Proxy->FastUpdateIndexBuffers(UpdateSets);
+		}
+	}
+
+	if (bPositions)
+	{
+		TRACE_CPUPROFILER_EVENT_SCOPE(SimpleDynamicMeshComponent_FastVertexUpdate_FinalPositionsUpdate);
+		MarkRenderTransformDirty();
+		LocalBounds.Contain(UpdateSetBounds);
+		UpdateBounds();
+	}
+
+	GetDynamicMesh()->PostRealtimeUpdate();
+}
+
+
+
+
+
+FPrimitiveSceneProxy* UDynamicMeshComponent::CreateSceneProxy()
+{
+	// if this is not always the case, we have made incorrect assumptions
+	ensure(GetCurrentSceneProxy() == nullptr);
+
+	FDynamicMeshSceneProxy* NewProxy = nullptr;
+	if (GetMesh()->TriangleCount() > 0)
+	{
+		NewProxy = new FDynamicMeshSceneProxy(this);
+
+		if (TriangleColorFunc)
+		{
+			NewProxy->bUsePerTriangleColor = true;
+			NewProxy->PerTriangleColorFunc = [this](const FDynamicMesh3* MeshIn, int TriangleID) { return GetTriangleColor(MeshIn, TriangleID); };
+		}
+
+		if (SecondaryTriFilterFunc)
+		{
+			NewProxy->bUseSecondaryTriBuffers = true;
+			NewProxy->SecondaryTriFilterFunc = [this](const FDynamicMesh3* MeshIn, int32 TriangleID) 
+			{ 
+				return (SecondaryTriFilterFunc) ? SecondaryTriFilterFunc(MeshIn, TriangleID) : false;
+			};
+		}
+
+		if (Decomposition)
+		{
+			NewProxy->InitializeFromDecomposition(Decomposition);
+		}
+		else
+		{
+			NewProxy->Initialize();
+		}
+	}
+
+	bProxyValid = true;
+	return NewProxy;
+}
+
+
+
+void UDynamicMeshComponent::NotifyMaterialSetUpdated()
+{
+	if (GetCurrentSceneProxy() != nullptr)
+	{
+		GetCurrentSceneProxy()->UpdatedReferencedMaterials();
+	}
+}
+
+
+void UDynamicMeshComponent::SetTriangleColorFunction(
+	TUniqueFunction<FColor(const FDynamicMesh3*, int)> TriangleColorFuncIn,
+	EDynamicMeshComponentRenderUpdateMode UpdateMode)
+{
+	TriangleColorFunc = MoveTemp(TriangleColorFuncIn);
+
+	if (UpdateMode == EDynamicMeshComponentRenderUpdateMode::FastUpdate)
+	{
+		FastNotifyColorsUpdated();
+	}
+	else if (UpdateMode == EDynamicMeshComponentRenderUpdateMode::FullUpdate)
+	{
+		NotifyMeshUpdated();
+	}
+}
+
+void UDynamicMeshComponent::ClearTriangleColorFunction(EDynamicMeshComponentRenderUpdateMode UpdateMode)
+{
+	if (TriangleColorFunc)
+	{
+		TriangleColorFunc = nullptr;
+
+		if (UpdateMode == EDynamicMeshComponentRenderUpdateMode::FastUpdate)
+		{
+			FastNotifyColorsUpdated();
+		}
+		else if (UpdateMode == EDynamicMeshComponentRenderUpdateMode::FullUpdate)
+		{
+			NotifyMeshUpdated();
+		}
+	}
+}
+
+bool UDynamicMeshComponent::HasTriangleColorFunction()
+{
+	return !!TriangleColorFunc;
+}
+
+
+
+
+void UDynamicMeshComponent::EnableSecondaryTriangleBuffers(TUniqueFunction<bool(const FDynamicMesh3*, int32)> SecondaryTriFilterFuncIn)
+{
+	SecondaryTriFilterFunc = MoveTemp(SecondaryTriFilterFuncIn);
+	NotifyMeshUpdated();
+}
+
+void UDynamicMeshComponent::DisableSecondaryTriangleBuffers()
+{
+	SecondaryTriFilterFunc = nullptr;
+	NotifyMeshUpdated();
+}
+
+
+void UDynamicMeshComponent::SetExternalDecomposition(TUniquePtr<FMeshRenderDecomposition> DecompositionIn)
+{
+	Decomposition = MoveTemp(DecompositionIn);
+	NotifyMeshUpdated();
+}
+
+
+
+FColor UDynamicMeshComponent::GetTriangleColor(const FDynamicMesh3* MeshIn, int TriangleID)
+{
+	if (TriangleColorFunc)
+	{
+		return TriangleColorFunc(MeshIn, TriangleID);
+	}
+	else
+	{
+		return (TriangleID % 2 == 0) ? FColor::Red : FColor::White;
+	}
+}
+
+
+
+FBoxSphereBounds UDynamicMeshComponent::CalcBounds(const FTransform& LocalToWorld) const
+{
+	// can get a tighter box by calculating in world space, but we care more about performance
+	FBox LocalBoundingBox = (FBox)LocalBounds;
+	FBoxSphereBounds Ret(LocalBoundingBox.TransformBy(LocalToWorld));
+	Ret.BoxExtent *= BoundsScale;
+	Ret.SphereRadius *= BoundsScale;
+	return Ret;
+}
+
+
+
+
+void UDynamicMeshComponent::SetInvalidateProxyOnChangeEnabled(bool bEnabled)
+{
+	bInvalidateProxyOnChange = bEnabled;
+}
+
+
+void UDynamicMeshComponent::ApplyChange(const FMeshVertexChange* Change, bool bRevert)
+{
+	// will fire UDynamicMesh::MeshChangedEvent, which will call OnMeshObjectChanged() below to invalidate proxy, fire change events, etc
+	MeshObject->ApplyChange(Change, bRevert);
+}
+
+void UDynamicMeshComponent::ApplyChange(const FMeshChange* Change, bool bRevert)
+{
+	// will fire UDynamicMesh::MeshChangedEvent, which will call OnMeshObjectChanged() below to invalidate proxy, fire change events, etc
+	MeshObject->ApplyChange(Change, bRevert);
+}
+
+void UDynamicMeshComponent::ApplyChange(const FMeshReplacementChange* Change, bool bRevert)
+{
+	// will fire UDynamicMesh::MeshChangedEvent, which will call OnMeshObjectChanged() below to invalidate proxy, fire change events, etc
+	MeshObject->ApplyChange(Change, bRevert);
+}
+
+void UDynamicMeshComponent::OnMeshObjectChanged(UDynamicMesh* ChangedMeshObject, FDynamicMeshChangeInfo ChangeInfo)
+{
+	bool bIsFChange = (
+		ChangeInfo.Type == EDynamicMeshChangeType::MeshChange
+		|| ChangeInfo.Type == EDynamicMeshChangeType::MeshVertexChange
+		|| ChangeInfo.Type == EDynamicMeshChangeType::MeshReplacementChange);
+
+	if (bIsFChange)
+	{
+		if (bInvalidateProxyOnChange)
+		{
+			NotifyMeshUpdated();
+		}
+
+		OnMeshChanged.Broadcast();
+
+		if (ChangeInfo.Type == EDynamicMeshChangeType::MeshVertexChange)
+		{
+			OnMeshVerticesChanged.Broadcast(this, ChangeInfo.VertexChange, ChangeInfo.bIsRevertChange);
+		}
+	}
+	else
+	{
+		NotifyMeshUpdated();
+		OnMeshChanged.Broadcast();
+	}
+}
+
+
+void UDynamicMeshComponent::SetDynamicMesh(UDynamicMesh* NewMesh)
+{
+	if (ensure(NewMesh) == false)
+	{
+		return;
+	}
+
+	if (ensure(MeshObject))
+	{
+		MeshObject->OnMeshChanged().Remove(MeshObjectChangedHandle);
+	}
+
+	MeshObject = NewMesh;
+	MeshObjectChangedHandle = MeshObject->OnMeshChanged().AddUObject(this, &UDynamicMeshComponent::OnMeshObjectChanged);
+
+	NotifyMeshUpdated();
+	OnMeshChanged.Broadcast();
+}
+
+
+
+void UDynamicMeshComponent::OnChildAttached(USceneComponent* ChildComponent)
+{
+	Super::OnChildAttached(ChildComponent);
+	OnChildAttachmentModified.Broadcast(ChildComponent, true);
+}
+void UDynamicMeshComponent::OnChildDetached(USceneComponent* ChildComponent)
+{
+	Super::OnChildDetached(ChildComponent);
+	OnChildAttachmentModified.Broadcast(ChildComponent, false);
+}
diff --git a/Engine/Plugins/Experimental/MeshModelingToolset/Source/ModelingComponents/Private/Components/DynamicMeshSceneProxy.h b/Engine/Plugins/Experimental/MeshModelingToolset/Source/ModelingComponents/Private/Components/DynamicMeshSceneProxy.h
new file mode 100644
--- /dev/null
+++ b/Engine/Plugins/Experimental/MeshModelingToolset/Source/ModelingComponents/Private/Components/DynamicMeshSceneProxy.h
@@ -0,0 +1,658 @@
+// Copyright Epic Games, Inc. All Rights Reserved. 
+
+#pragma once
+
+#include "Components/DynamicMeshComponent.h"
+#include "Components/BaseDynamicMeshSceneProxy.h"
+#include "DynamicMesh/MeshTangents.h"
+#include "Async/ParallelFor.h"
+
+/**
+ * Scene Proxy for UDynamicMeshComponent.
+ * 
+ * Based on FProceduralMeshSceneProxy but simplified in various ways.
+ * 
+ * Supports wireframe-on-shaded rendering.
+ * 
+ */
+class FDynamicMeshSceneProxy final : public FBaseDynamicMeshSceneProxy
+{
+	using FIndex2i = UE::Geometry::FIndex2i;
+	using FIndex3i = UE::Geometry::FIndex3i;
+private:
+	FMaterialRelevance MaterialRelevance;
+
+	// note: FBaseDynamicMeshSceneProxy owns and will destroy these
+	TArray<FMeshRenderBufferSet*> RenderBufferSets;
+
+	// if true, we store entire mesh in single RenderBuffers and we can do some optimizations
+	bool bIsSingleBuffer = false;
+
+public:
+	/** Component that created this proxy (is there a way to look this up?) */
+	UDynamicMeshComponent* ParentComponent;
+
+
+	FDynamicMeshSceneProxy(UDynamicMeshComponent* Component)
+		: FBaseDynamicMeshSceneProxy(Component)
+		, MaterialRelevance(Component->GetMaterialRelevance(GetScene().GetFeatureLevel()))
+	{
+		ParentComponent = Component;
+	}
+
+
+	virtual void GetActiveRenderBufferSets(TArray<FMeshRenderBufferSet*>& Buffers) const override
+	{
+		Buffers = RenderBufferSets;
+	}
+
+
+
+	void Initialize()
+	{
+		// allocate buffer sets based on materials
+		ensure(RenderBufferSets.Num() == 0);
+		int32 NumMaterials = GetNumMaterials();
+		if (NumMaterials == 0)
+		{
+			RenderBufferSets.SetNum(1);
+			RenderBufferSets[0] = AllocateNewRenderBufferSet();
+			RenderBufferSets[0]->Material = UMaterial::GetDefaultMaterial(MD_Surface);
+		}
+		else
+		{
+			RenderBufferSets.SetNum(NumMaterials);
+			for (int32 k = 0; k < NumMaterials; ++k)
+			{
+				RenderBufferSets[k] = AllocateNewRenderBufferSet();
+				RenderBufferSets[k]->Material = GetMaterial(k);
+			}
+		}
+
+
+		FDynamicMesh3* Mesh = ParentComponent->GetRenderMesh();
+		if (Mesh->HasAttributes() && Mesh->Attributes()->HasMaterialID() && NumMaterials > 1)
+		{
+			bIsSingleBuffer = false;
+			InitializeByMaterial(RenderBufferSets);
+		}
+		else
+		{
+			bIsSingleBuffer = true;
+			InitializeSingleBufferSet(RenderBufferSets[0]);
+		}
+	}
+
+
+
+	TUniqueFunction<void(int, int, int, const FVector3f&, FVector3f&, FVector3f&)> MakeTangentsFunc(bool bSkipAutoCompute = false)
+	{
+		EDynamicMeshComponentTangentsMode TangentsType = ParentComponent->GetTangentsType();
+		if (TangentsType == EDynamicMeshComponentTangentsMode::ExternallyProvided)
+		{
+			UE::Geometry::FDynamicMeshTangents Tangents(ParentComponent->GetRenderMesh());
+			return [Tangents](int VertexID, int TriangleID, int TriVtxIdx, const FVector3f& Normal, FVector3f& TangentX, FVector3f& TangentY) -> void
+			{
+				Tangents.GetTangentVectors(TriangleID, TriVtxIdx, Normal, TangentX, TangentY);
+			};
+		}
+		else if (TangentsType == EDynamicMeshComponentTangentsMode::AutoCalculated && bSkipAutoCompute == false)
+		{
+			const UE::Geometry::FMeshTangentsf* Tangents = ParentComponent->GetAutoCalculatedTangents();
+			if (Tangents != nullptr )
+			{
+				return [Tangents](int VertexID, int TriangleID, int TriVtxIdx, const FVector3f& Normal, FVector3f& TangentX, FVector3f& TangentY) -> void
+				{
+					Tangents->GetTriangleVertexTangentVectors(TriangleID, TriVtxIdx, TangentX, TangentY);
+				};
+			}
+		}
+
+		// fallback to orthogonal basis
+		return [](int VertexID, int TriangleID, int TriVtxIdx, const FVector3f& Normal, FVector3f& TangentX, FVector3f& TangentY) -> void
+		{
+			VectorUtil::MakePerpVectors(Normal, TangentX, TangentY);
+		};
+	}
+
+
+	/**
+	 * Initialize multiple buffers for the mesh based on the given Decomposition
+	 */
+	void InitializeFromDecomposition(TUniquePtr<FMeshRenderDecomposition>& Decomposition)
+	{
+		ensure(RenderBufferSets.Num() == 0);
+		int32 NumSets = Decomposition->Num();
+		RenderBufferSets.SetNum(NumSets);
+		for (int32 k = 0; k < NumSets; ++k)
+		{
+			RenderBufferSets[k] = AllocateNewRenderBufferSet();
+			RenderBufferSets[k]->Material = Decomposition->GetGroup(k).Material;
+			if (RenderBufferSets[k]->Material == nullptr)
+			{
+				RenderBufferSets[k]->Material = UMaterial::GetDefaultMaterial(MD_Surface);
+			}
+		}
+
+		bIsSingleBuffer = false;
+
+		const FDynamicMesh3* Mesh = ParentComponent->GetRenderMesh();
+		// find suitable overlays
+		const FDynamicMeshUVOverlay* UVOverlay = Mesh->Attributes()->PrimaryUV();
+		const FDynamicMeshNormalOverlay* NormalOverlay = Mesh->Attributes()->PrimaryNormals();
+		const FDynamicMeshColorOverlay* ColorOverlay = Mesh->Attributes()->PrimaryColors();
+		TUniqueFunction<void(int, int, int, const FVector3f&, FVector3f&, FVector3f&)> TangentsFunc = MakeTangentsFunc();
+
+		// init renderbuffers for each set
+		ParallelFor(NumSets, [&](int32 SetIndex)
+		{
+			const FMeshRenderDecomposition::FGroup& Group = Decomposition->GetGroup(SetIndex);
+			const TArray<int32>& Triangles = Group.Triangles;
+			if (Triangles.Num() > 0)
+			{
+				FMeshRenderBufferSet* RenderBuffers = RenderBufferSets[SetIndex];
+				RenderBuffers->Triangles = Triangles;
+				InitializeBuffersFromOverlays(RenderBuffers, Mesh,
+					Triangles.Num(), Triangles,
+					UVOverlay, NormalOverlay, ColorOverlay, TangentsFunc);
+
+				ENQUEUE_RENDER_COMMAND(FDynamicMeshSceneProxyInitializeFromDecomposition)(
+					[RenderBuffers](FRHICommandListImmediate& RHICmdList)
+				{
+					RenderBuffers->Upload();
+				});
+			}
+		});
+	}
+
+
+
+
+
+
+
+
+	/**
+	 * Initialize a single set of mesh buffers for the entire mesh
+	 */
+	void InitializeSingleBufferSet(FMeshRenderBufferSet* RenderBuffers)
+	{
+		const FDynamicMesh3* Mesh = ParentComponent->GetRenderMesh();
+
+		// find suitable overlays
+		TArray<const FDynamicMeshUVOverlay*> UVOverlays;
+		const FDynamicMeshNormalOverlay* NormalOverlay = nullptr;
+		const FDynamicMeshColorOverlay* ColorOverlay = nullptr;
+		if (Mesh->HasAttributes())
+		{
+			const FDynamicMeshAttributeSet* Attributes = Mesh->Attributes();
+			NormalOverlay = Attributes->PrimaryNormals();
+			UVOverlays.SetNum(Attributes->NumUVLayers());
+			for (int32 k = 0; k < UVOverlays.Num(); ++k)
+			{
+				UVOverlays[k] = Attributes->GetUVLayer(k);
+			}
+			ColorOverlay = Attributes->PrimaryColors();
+		}
+		TUniqueFunction<void(int, int, int, const FVector3f&, FVector3f&, FVector3f&)> TangentsFunc = MakeTangentsFunc();
+
+		InitializeBuffersFromOverlays(RenderBuffers, Mesh,
+			Mesh->TriangleCount(), Mesh->TriangleIndicesItr(),
+			UVOverlays, NormalOverlay, ColorOverlay, TangentsFunc);
+
+		ENQUEUE_RENDER_COMMAND(FDynamicMeshSceneProxyInitializeSingle)(
+			[RenderBuffers](FRHICommandListImmediate& RHICmdList)
+		{
+			RenderBuffers->Upload();
+		});
+	}
+
+
+
+	/**
+	 * Initialize the mesh buffers, one per material
+	 */
+	void InitializeByMaterial(TArray<FMeshRenderBufferSet*>& BufferSets)
+	{
+		const FDynamicMesh3* Mesh = ParentComponent->GetRenderMesh();
+		if (ensure(Mesh->HasAttributes() && Mesh->Attributes()->HasMaterialID()) == false)
+		{
+			return;
+		}
+
+		const FDynamicMeshAttributeSet* Attributes = Mesh->Attributes();
+
+		// find suitable overlays
+		const FDynamicMeshMaterialAttribute* MaterialID = Attributes->GetMaterialID();
+		const FDynamicMeshNormalOverlay* NormalOverlay = Mesh->Attributes()->PrimaryNormals();
+		const FDynamicMeshColorOverlay* ColorOverlay = Mesh->Attributes()->PrimaryColors();
+
+		TArray<const FDynamicMeshUVOverlay*> UVOverlays;
+		UVOverlays.SetNum(Attributes->NumUVLayers());
+		for (int32 k = 0; k < UVOverlays.Num(); ++k)
+		{
+			UVOverlays[k] = Attributes->GetUVLayer(k);
+		}
+
+		TUniqueFunction<void(int, int, int, const FVector3f&, FVector3f&, FVector3f&)> TangentsFunc = MakeTangentsFunc();
+
+		// count number of triangles for each material (could do in parallel?)
+		int NumMaterials = BufferSets.Num();
+		TArray<FThreadSafeCounter> Counts;
+		Counts.SetNum(NumMaterials);
+		for (int k = 0; k < NumMaterials; ++k)
+		{
+			Counts[k].Reset();
+		}
+		ParallelFor(Mesh->MaxTriangleID(), [&](int tid)
+		{
+			int MatIdx;
+			MaterialID->GetValue(tid, &MatIdx);
+			if (MatIdx >= 0 && MatIdx < NumMaterials)
+			{
+				Counts[MatIdx].Increment();
+			}
+		});
+
+		// find max count
+		int32 MaxCount = 0;
+		for (FThreadSafeCounter& Count : Counts)
+		{
+			MaxCount = FMath::Max(Count.GetValue(), MaxCount);
+		}
+
+		// init renderbuffers for each material
+		// could do this in parallel but then we need to allocate separate triangle arrays...is it worth it?
+		TArray<int> Triangles;
+		Triangles.Reserve(MaxCount);
+		for (int k = 0; k < NumMaterials; ++k)
+		{
+			if (Counts[k].GetValue() > 0)
+			{
+				FMeshRenderBufferSet* RenderBuffers = BufferSets[k];
+
+				Triangles.Reset();
+				for (int tid : Mesh->TriangleIndicesItr())
+				{
+					int MatIdx;
+					MaterialID->GetValue(tid, &MatIdx);
+					if (MatIdx == k)
+					{
+						Triangles.Add(tid);
+					}
+				}
+
+				InitializeBuffersFromOverlays(RenderBuffers, Mesh,
+					Triangles.Num(), Triangles,
+					UVOverlays, NormalOverlay, ColorOverlay, TangentsFunc);
+
+				RenderBuffers->Triangles = Triangles;
+
+				ENQUEUE_RENDER_COMMAND(FDynamicMeshSceneProxyInitializeByMaterial)(
+					[RenderBuffers](FRHICommandListImmediate& RHICmdList)
+				{
+					RenderBuffers->Upload();
+				});
+			}
+		}
+	}
+
+
+	bool RenderMeshLayoutMatchesRenderBuffers() const
+	{
+		const FDynamicMesh3* Mesh = ParentComponent->GetRenderMesh();
+
+		auto CheckBufferSet = [](const FDynamicMesh3* Mesh, const FMeshRenderBufferSet* BufferSet, int NumTriangles) -> bool
+		{
+			if (BufferSet->Triangles)
+			{
+				for (int TriangleID : BufferSet->Triangles.GetValue())
+				{
+					if (!Mesh->IsTriangle(TriangleID))
+					{
+						return false;
+					}
+				}
+			}
+
+			int NumVertices = NumTriangles * 3;
+			if (BufferSet->PositionVertexBuffer.GetNumVertices() != NumVertices ||
+				BufferSet->StaticMeshVertexBuffer.GetNumVertices() != NumVertices ||
+				BufferSet->ColorVertexBuffer.GetNumVertices() != NumVertices)
+			{
+				return false;
+			}
+
+			return true;
+		};
+
+		if (bIsSingleBuffer)
+		{
+			if (ensure(RenderBufferSets.Num() == 1))
+			{
+				FMeshRenderBufferSet* BufferSet = RenderBufferSets[0];
+				if (BufferSet->TriangleCount != Mesh->TriangleCount())
+				{
+					return false;
+				}
+
+				int NumTriangles = Mesh->TriangleCount();
+
+				if (!CheckBufferSet(Mesh, BufferSet, NumTriangles))
+				{
+					return false;
+				}
+			}
+		}
+		else
+		{
+			for (FMeshRenderBufferSet* BufferSet : RenderBufferSets)
+			{
+				if (ensure(BufferSet->Triangles))
+				{
+					int NumTriangles = BufferSet->Triangles->Num();
+
+					if (!CheckBufferSet(Mesh, BufferSet, NumTriangles))
+					{
+						return false;
+					}
+				}
+			}
+		}
+
+		return true;
+	}
+
+
+
+	/**
+	 * Update the vertex position/normal/color buffers
+	 */
+	void FastUpdateVertices(bool bPositions, bool bNormals, bool bColors, bool bUVs)
+	{
+		// This needs to be rewritten for split-by-material buffers.
+		// Could store triangle set with each buffer, and then rebuild vtx buffer(s) as needed?
+
+		FDynamicMesh3* Mesh = ParentComponent->GetRenderMesh();
+
+		// find suitable overlays and attributes
+		FDynamicMeshNormalOverlay* NormalOverlay = nullptr;
+		if (bNormals && ensure(Mesh->HasAttributes()) )
+		{
+			NormalOverlay = Mesh->Attributes()->PrimaryNormals();
+		}
+		FDynamicMeshUVOverlay* UVOVerlay = nullptr;
+		if (bUVs && ensure(Mesh->HasAttributes()) )
+		{
+			UVOVerlay = Mesh->Attributes()->PrimaryUV();
+		}
+		FDynamicMeshColorOverlay* ColorOverlay = nullptr;
+		if (bColors && ensure(Mesh->HasAttributes()) )
+		{
+			ColorOverlay = Mesh->Attributes()->PrimaryColors();
+		}
+
+		// Currently deferring tangents auto-computation if doing fast vertex update (maybe needs to be exposed?)
+		TUniqueFunction<void(int, int, int, const FVector3f&, FVector3f&, FVector3f&)> TangentsFunc = MakeTangentsFunc(true);
+
+		if (bIsSingleBuffer)
+		{
+			if (ensure(RenderBufferSets.Num() == 1))
+			{
+				FMeshRenderBufferSet* Buffers = RenderBufferSets[0];
+				if (bPositions || bNormals || bColors)
+				{
+					UpdateVertexBuffersFromOverlays(Buffers, Mesh,
+						Mesh->TriangleCount(), Mesh->TriangleIndicesItr(),
+						NormalOverlay, ColorOverlay, TangentsFunc,
+						bPositions, bNormals, bColors);
+				}
+				if (bUVs)
+				{
+					UpdateVertexUVBufferFromOverlays(Buffers, Mesh,
+						Mesh->TriangleCount(), Mesh->TriangleIndicesItr(), UVOVerlay, 0);
+				}
+
+
+				ENQUEUE_RENDER_COMMAND(FDynamicMeshSceneProxyFastUpdateVertices)(
+					[Buffers, bPositions, bNormals, bColors, bUVs](FRHICommandListImmediate& RHICmdList)
+				{
+					Buffers->UploadVertexUpdate(bPositions, bNormals || bUVs, bColors);
+				});
+			}
+		}
+		else
+		{
+			ParallelFor(RenderBufferSets.Num(), [&](int i)
+			{
+				FMeshRenderBufferSet* Buffers = RenderBufferSets[i];
+				if (Buffers->TriangleCount == 0)
+				{
+					return;
+				}
+				if (ensure(Buffers->Triangles.IsSet()))
+				{
+					if (bPositions || bNormals || bColors)
+					{
+						UpdateVertexBuffersFromOverlays(Buffers, Mesh,
+							Buffers->Triangles->Num(), Buffers->Triangles.GetValue(),
+							NormalOverlay, ColorOverlay, TangentsFunc,
+							bPositions, bNormals, bColors);
+					}
+					if (bUVs)
+					{
+						UpdateVertexUVBufferFromOverlays(Buffers, Mesh,
+							Buffers->Triangles->Num(), Buffers->Triangles.GetValue(), UVOVerlay, 0);
+					}
+
+					ENQUEUE_RENDER_COMMAND(FDynamicMeshSceneProxyFastUpdateVertices)(
+						[Buffers, bPositions, bNormals, bColors, bUVs](FRHICommandListImmediate& RHICmdList)
+					{
+						Buffers->UploadVertexUpdate(bPositions, bNormals || bUVs, bColors);
+					});
+				}
+			});
+		}
+	}
+
+
+
+	/**
+	 * Update the vertex position/normal/color buffers
+	 */
+	void FastUpdateVertices(const TArray<int32>& WhichBuffers, bool bPositions, bool bNormals, bool bColors, bool bUVs)
+	{
+		TRACE_CPUPROFILER_EVENT_SCOPE(DynamicMeshProxy_FastUpdateVertices);
+
+		// skip if we have no updates
+		if (bPositions == false && bNormals == false && bColors == false && bUVs == false)
+		{
+			return;
+		}
+
+		// This needs to be rewritten for split-by-material buffers.
+		// Could store triangle set with each buffer, and then rebuild vtx buffer(s) as needed?
+
+		const FDynamicMesh3* Mesh = ParentComponent->GetRenderMesh();
+
+		// find suitable overlays
+		const FDynamicMeshNormalOverlay* NormalOverlay = nullptr;
+		if (bNormals && ensure(Mesh->HasAttributes()))
+		{
+			NormalOverlay = Mesh->Attributes()->PrimaryNormals();
+		}
+		const FDynamicMeshUVOverlay* UVOVerlay = nullptr;
+		if (bUVs && ensure(Mesh->HasAttributes()))
+		{
+			UVOVerlay = Mesh->Attributes()->PrimaryUV();
+		}
+		const FDynamicMeshColorOverlay* ColorOverlay = nullptr;
+		if (bColors && ensure(Mesh->HasAttributes()))
+		{
+			ColorOverlay = Mesh->Attributes()->PrimaryColors();
+		}
+
+		// Currently deferring tangents auto-computation if doing fast vertex update (maybe needs to be exposed?)
+		TUniqueFunction<void(int, int, int, const FVector3f&, FVector3f&, FVector3f&)> TangentsFunc = MakeTangentsFunc(true);
+
+		ParallelFor(WhichBuffers.Num(), [&](int idx)
+		{
+			int32 BufferIndex = WhichBuffers[idx];
+			if ( RenderBufferSets.IsValidIndex(BufferIndex) == false || RenderBufferSets[BufferIndex]->TriangleCount == 0)
+			{
+				return;
+			}
+			FMeshRenderBufferSet* Buffers = RenderBufferSets[BufferIndex];
+			if (ensure(Buffers->Triangles.IsSet()))
+			{
+				if (bPositions || bNormals || bColors)
+				{
+					UpdateVertexBuffersFromOverlays(Buffers, Mesh,
+						Buffers->Triangles->Num(), Buffers->Triangles.GetValue(),
+						NormalOverlay, ColorOverlay, TangentsFunc,
+						bPositions, bNormals, bColors);
+				}
+				if (bUVs)
+				{
+					UpdateVertexUVBufferFromOverlays(Buffers, Mesh,
+						Buffers->Triangles->Num(), Buffers->Triangles.GetValue(), UVOVerlay, 0);
+				}
+
+				ENQUEUE_RENDER_COMMAND(FDynamicMeshSceneProxyFastUpdateVerticesBufferList)(
+					[Buffers, bPositions, bNormals, bColors, bUVs](FRHICommandListImmediate& RHICmdList)
+				{
+					Buffers->TransferVertexUpdateToGPU(bPositions, bNormals, bUVs, bColors);
+				});
+			}
+		});
+	}
+
+
+
+
+	/**
+	 * Update index buffers inside each RenderBuffer set
+	 */
+	void FastUpdateAllIndexBuffers()
+	{
+		FDynamicMesh3* Mesh = ParentComponent->GetRenderMesh();
+
+		// have to wait for all outstanding rendering to finish because the index buffers we are about to edit might be in-use
+		FlushRenderingCommands();
+
+		ParallelFor(RenderBufferSets.Num(), [this, &Mesh](int i)
+		{
+			FMeshRenderBufferSet* Buffers = RenderBufferSets[i];
+
+			FScopeLock BuffersLock(&Buffers->BuffersLock);
+
+			if (Buffers->TriangleCount > 0)
+			{
+				RecomputeRenderBufferTriangleIndexSets(Buffers, Mesh);
+			}
+
+			ENQUEUE_RENDER_COMMAND(FDynamicMeshSceneProxyFastUpdateAllIndexBuffers)(
+				[Buffers](FRHICommandListImmediate& RHICmdList)
+			{
+				Buffers->UploadIndexBufferUpdate();
+			});
+			
+		});
+	}
+
+
+
+	/**
+	 * Update index buffers inside each RenderBuffer set
+	 */
+	void FastUpdateIndexBuffers(const TArray<int32>& WhichBuffers)
+	{
+		FDynamicMesh3* Mesh = ParentComponent->GetRenderMesh();
+
+		// have to wait for all outstanding rendering to finish because the index buffers we are about to edit might be in-use
+		FlushRenderingCommands();
+
+		ParallelFor(WhichBuffers.Num(), [this, &WhichBuffers, &Mesh](int i)
+		{
+			int32 BufferIndex = WhichBuffers[i];
+			if (RenderBufferSets.IsValidIndex(BufferIndex) == false)
+			{
+				return;
+			}
+
+			FMeshRenderBufferSet* Buffers = RenderBufferSets[BufferIndex];
+			FScopeLock BuffersLock(&Buffers->BuffersLock);
+			if (Buffers->TriangleCount > 0)
+			{
+				RecomputeRenderBufferTriangleIndexSets(Buffers, Mesh);
+			}
+
+			ENQUEUE_RENDER_COMMAND(FDynamicMeshSceneProxyFastUpdateSomeIndexBuffers)(
+				[Buffers](FRHICommandListImmediate& RHICmdList)
+			{
+				Buffers->UploadIndexBufferUpdate();
+			});
+		});
+	}
+
+
+
+public:
+
+
+
+	virtual FPrimitiveViewRelevance GetViewRelevance(const FSceneView* View) const override
+	{
+		FPrimitiveViewRelevance Result;
+
+		Result.bDrawRelevance = IsShown(View);
+		Result.bShadowRelevance = IsShadowCast(View);
+		Result.bDynamicRelevance = true;
+		Result.bRenderInMainPass = ShouldRenderInMainPass();
+		Result.bUsesLightingChannels = GetLightingChannelMask() != GetDefaultLightingChannelMask();
+		Result.bTranslucentSelfShadow = bCastVolumetricTranslucentShadow;
+		Result.bRenderCustomDepth = ShouldRenderCustomDepth();
+		// Note that this is actually a getter. One may argue that it is poorly named.
+		MaterialRelevance.SetPrimitiveViewRelevance(Result);
+		Result.bVelocityRelevance = DrawsVelocity() && Result.bOpaque && Result.bRenderInMainPass;
+
+		return Result;
+	}
+
+	virtual void UpdatedReferencedMaterials() override
+	{
+		FBaseDynamicMeshSceneProxy::UpdatedReferencedMaterials();
+
+		// The material relevance may need updating.
+		MaterialRelevance = ParentComponent->GetMaterialRelevance(GetScene().GetFeatureLevel());
+	}
+
+
+	virtual void GetLightRelevance(const FLightSceneProxy* LightSceneProxy, bool& bDynamic, bool& bRelevant, bool& bLightMapped, bool& bShadowMapped) const override
+	{
+		FPrimitiveSceneProxy::GetLightRelevance(LightSceneProxy, bDynamic, bRelevant, bLightMapped, bShadowMapped);
+	}
+
+	virtual bool CanBeOccluded() const override
+	{
+		return !MaterialRelevance.bDisableDepthTest;
+	}
+
+	virtual uint32 GetMemoryFootprint(void) const override { return(sizeof(*this) + GetAllocatedSize()); }
+
+	uint32 GetAllocatedSize(void) const { return(FPrimitiveSceneProxy::GetAllocatedSize()); }
+
+	SIZE_T GetTypeHash() const override
+	{
+		static size_t UniquePointer;
+		return reinterpret_cast<size_t>(&UniquePointer);
+	}
+
+};
+
+
+
+
+
+
+
diff --git a/Engine/Plugins/Experimental/MeshModelingToolset/Source/ModelingComponents/Private/Components/OctreeDynamicMeshComponent.cpp b/Engine/Plugins/Experimental/MeshModelingToolset/Source/ModelingComponents/Private/Components/OctreeDynamicMeshComponent.cpp
new file mode 100644
--- /dev/null
+++ b/Engine/Plugins/Experimental/MeshModelingToolset/Source/ModelingComponents/Private/Components/OctreeDynamicMeshComponent.cpp
@@ -0,0 +1,465 @@
+// Copyright Epic Games, Inc. All Rights Reserved. 
+
+#include "Components/OctreeDynamicMeshComponent.h"
+#include "RenderingThread.h"
+#include "RenderResource.h"
+#include "PrimitiveViewRelevance.h"
+#include "PrimitiveSceneProxy.h"
+#include "VertexFactory.h"
+#include "MaterialShared.h"
+#include "Engine/CollisionProfile.h"
+#include "Materials/Material.h"
+#include "LocalVertexFactory.h"
+#include "SceneManagement.h"
+#include "DynamicMeshBuilder.h"
+#include "EngineGlobals.h"
+#include "Engine/Engine.h"
+#include "StaticMeshResources.h"
+#include "StaticMeshAttributes.h"
+
+#include "DynamicMeshAttributeSet.h"
+#include "MeshNormals.h"
+#include "MeshTransforms.h"
+#include "MeshDescriptionToDynamicMesh.h"
+
+#include "Changes/MeshVertexChange.h"
+#include "Changes/MeshChange.h"
+#include "DynamicMeshChangeTracker.h"
+
+
+// default proxy for this component
+#include "OctreeDynamicMeshSceneProxy.h"
+
+#include "ExplicitUseGeometryMathTypes.h"		// using UE::Geometry::(math types)
+using namespace UE::Geometry;
+
+
+
+UOctreeDynamicMeshComponent::UOctreeDynamicMeshComponent(const FObjectInitializer& ObjectInitializer)
+	: Super(ObjectInitializer)
+{
+	PrimaryComponentTick.bCanEverTick = false;
+
+	SetCollisionProfileName(UCollisionProfile::NoCollision_ProfileName);
+
+	MeshObject = CreateDefaultSubobject<UDynamicMesh>(TEXT("DynamicMesh"));
+	MeshObjectChangedHandle = MeshObject->OnMeshChanged().AddUObject(this, &UOctreeDynamicMeshComponent::OnMeshObjectChanged);
+	PreMeshChangeHandle = MeshObject->OnPreMeshChanged().AddUObject(this, &UOctreeDynamicMeshComponent::OnPreMeshObjectChanged);
+	
+
+	Octree = MakeUnique<FDynamicMeshOctree3>();
+	Octree->Initialize(GetMesh());
+	OctreeCut = MakeUnique<FDynamicMeshOctree3::FTreeCutSet>();
+}
+
+
+
+void UOctreeDynamicMeshComponent::SetMesh(UE::Geometry::FDynamicMesh3&& MoveMesh)
+{
+	MeshObject->SetMesh(MoveTemp(MoveMesh));
+
+	FAxisAlignedBox3d MeshBounds = GetMesh()->GetCachedBounds();
+	Octree = MakeUnique<FDynamicMeshOctree3>();
+	Octree->RootDimension = MeshBounds.MaxDim() * 0.25;
+
+	Octree->Initialize(GetMesh());
+
+	FDynamicMeshOctree3::FStatistics Stats;
+	Octree->ComputeStatistics(Stats);
+	//UE_LOG(LogTemp, Warning, TEXT("OctreeStats %s"), *Stats.ToString());
+
+	OctreeCut = MakeUnique<FDynamicMeshOctree3::FTreeCutSet>();
+	CutCellSetMap.Reset();
+
+	NotifyMeshUpdated();
+}
+
+
+
+void UOctreeDynamicMeshComponent::ApplyTransform(const UE::Geometry::FTransform3d& Transform, bool bInvert)
+{
+	if (bInvert)
+	{
+		MeshTransforms::ApplyTransformInverse(*GetMesh(), Transform);
+	}
+	else
+	{
+		MeshTransforms::ApplyTransform(*GetMesh(), Transform);
+	}
+
+	if (GetCurrentSceneProxy() != nullptr)
+	{
+		Octree->ModifiedBounds = FAxisAlignedBox3d(
+			-TNumericLimits<float>::Max() * FVector3d::One(),
+			TNumericLimits<float>::Max() * FVector3d::One());
+		NotifyMeshUpdated();
+	}
+	else
+	{
+		FAxisAlignedBox3d MeshBounds = GetMesh()->GetCachedBounds();
+		Octree = MakeUnique<FDynamicMeshOctree3>();
+		Octree->RootDimension = MeshBounds.MaxDim() * 0.25;
+		Octree->Initialize(GetMesh());
+		OctreeCut = MakeUnique<FDynamicMeshOctree3::FTreeCutSet>();
+		CutCellSetMap.Reset();
+	}
+
+}
+
+
+
+void UOctreeDynamicMeshComponent::NotifyMeshUpdated()
+{
+	if (GetCurrentSceneProxy() != nullptr)
+	{
+		FAxisAlignedBox3d DirtyBox = Octree->ModifiedBounds;
+		Octree->ResetModifiedBounds();
+
+		// update existing cells
+		int TotalTriangleCount = 0;
+		int SpillTriangleCount = 0;
+		TArray<int32> SetsToUpdate;
+
+		{
+			SCOPE_CYCLE_COUNTER(STAT_SculptToolOctree_UpdateExisting);
+			FCriticalSection SetsLock;
+			int NumCutCells = CutCellSetMap.Num();
+			ParallelFor(NumCutCells, [&](int i) 
+			{
+				const FCutCellIndexSet& CutCellSet = CutCellSetMap[i];
+
+				if (Octree->TestCellIntersection(CutCellSet.CellRef, DirtyBox) == false)
+				{
+					return;
+				}
+
+				TArray<int32>& TriangleSet = TriangleDecomposition.GetIndexSetArray(CutCellSet.DecompSetID);
+				TriangleSet.Reset();
+				Octree->CollectTriangles(CutCellSet.CellRef, [&TriangleSet](int TriangleID) {
+					TriangleSet.Add(TriangleID);
+				});
+
+				SetsLock.Lock();
+				TotalTriangleCount += TriangleSet.Num();
+				SetsToUpdate.Add(CutCellSet.DecompSetID);
+				SetsLock.Unlock();
+			}, false);
+		}
+
+		// update cut set to find new cells
+		TArray<FDynamicMeshOctree3::FCellReference> NewCutCells;
+		{
+			SCOPE_CYCLE_COUNTER(STAT_SculptToolOctree_UpdateCutSet);
+			Octree->UpdateLevelCutSet(*OctreeCut, NewCutCells);
+		}
+
+		// add new ones
+		{
+			SCOPE_CYCLE_COUNTER(STAT_SculptToolOctree_CreateNew);
+			for (const FDynamicMeshOctree3::FCellReference& CellRef : NewCutCells)
+			{
+
+				int32 IndexSetID = TriangleDecomposition.CreateNewIndexSet();
+				TArray<int32>& TriangleSet = TriangleDecomposition.GetIndexSetArray(IndexSetID);
+				Octree->CollectTriangles(CellRef, [&TriangleSet](int TriangleID) {
+					TriangleSet.Add(TriangleID);
+				});
+				TotalTriangleCount += TriangleSet.Num();
+				FCutCellIndexSet SetMapEntry = { CellRef, IndexSetID };
+				CutCellSetMap.Add(SetMapEntry);
+				SetsToUpdate.Add(IndexSetID);
+			}
+		}
+
+		// rebuild spill set (always for now...should track bounds? and do separately for each root cell?)
+		{
+			SCOPE_CYCLE_COUNTER(STAT_SculptToolOctree_UpdateSpill);
+			TArray<int32>& SpillTriangleSet = TriangleDecomposition.GetIndexSetArray(SpillDecompSetID);
+			SpillTriangleSet.Reset();
+			Octree->CollectRootTriangles(*OctreeCut,
+				[&SpillTriangleSet](int TriangleID) {
+				SpillTriangleSet.Add(TriangleID);
+			});
+			Octree->CollectSpillTriangles(
+				[&SpillTriangleSet](int TriangleID) {
+				SpillTriangleSet.Add(TriangleID);
+			});
+			TotalTriangleCount += SpillTriangleSet.Num();
+			SpillTriangleCount += SpillTriangleSet.Num();
+			SetsToUpdate.Add(SpillDecompSetID);
+		}
+
+
+		//UE_LOG(LogTemp, Warning, TEXT("Updating %d of %d decomposition sets, %d tris total, spillcount %d"), SetsToUpdate.Num(), CutCellSetMap.Num(), TotalTriangleCount, SpillTriangleCount);
+		{
+			SCOPE_CYCLE_COUNTER(STAT_SculptToolOctree_UpdateFromDecomp);
+			GetCurrentSceneProxy()->UpdateFromDecomposition(TriangleDecomposition, SetsToUpdate);
+		}
+
+	}
+}
+
+
+
+
+static void InitializeOctreeCutSet(const FDynamicMesh3& Mesh, const FDynamicMeshOctree3& Octree, TUniquePtr<FDynamicMeshOctree3::FTreeCutSet>& CutSet)
+{
+	int TriangleCount = Mesh.TriangleCount();
+	if (TriangleCount < 50000)
+	{
+		*CutSet = Octree.BuildLevelCutSet(1);
+		return;
+	}
+
+	FDynamicMeshOctree3::FStatistics Stats;
+	Octree.ComputeStatistics(Stats);
+	int MaxLevel = Stats.Levels;
+
+	int CutLevel = 0;
+	while ( CutLevel < MaxLevel-1 && Stats.LevelBoxCounts[CutLevel] < 200 && Stats.LevelBoxCounts[CutLevel+1] < 300)
+	{
+		CutLevel++;
+	}
+	*CutSet = Octree.BuildLevelCutSet(CutLevel);
+	//UE_LOG(LogTemp, Warning, TEXT("InitializeOctreeCutSet - level %d cells %d"), CutLevel, CutSet->CutCells.Num());
+}
+
+
+
+
+FPrimitiveSceneProxy* UOctreeDynamicMeshComponent::CreateSceneProxy()
+{
+	check(GetCurrentSceneProxy() == nullptr);
+
+	FOctreeDynamicMeshSceneProxy* NewProxy = nullptr;
+	if (GetMesh()->TriangleCount() > 0)
+	{
+		NewProxy = new FOctreeDynamicMeshSceneProxy(this);
+
+		if (TriangleColorFunc != nullptr)
+		{
+			NewProxy->bUsePerTriangleColor = true;
+			NewProxy->PerTriangleColorFunc = [this](const FDynamicMesh3* MeshIn, int TriangleID) { return GetTriangleColor(TriangleID); };
+		}
+
+		OctreeCut = MakeUnique<FDynamicMeshOctree3::FTreeCutSet>();
+		CutCellSetMap.Reset();
+		InitializeOctreeCutSet(*GetMesh(), *Octree, OctreeCut);
+
+		TriangleDecomposition = FArrayIndexSetsDecomposition();
+
+		SpillDecompSetID = TriangleDecomposition.CreateNewIndexSet();
+
+		for (auto CellRef : OctreeCut->CutCells)
+		{
+			int32 IndexSetID = TriangleDecomposition.CreateNewIndexSet();
+			TArray<int32>& TriangleSet = TriangleDecomposition.GetIndexSetArray(IndexSetID);
+			Octree->CollectTriangles(CellRef, [&TriangleSet](int TriangleID) {
+				TriangleSet.Add(TriangleID);
+			});
+
+			FCutCellIndexSet SetMapEntry = { CellRef, IndexSetID };
+			CutCellSetMap.Add(SetMapEntry);
+		}
+
+		// collect spill triangles
+		{
+			TArray<int32>& SpillTriangleSet = TriangleDecomposition.GetIndexSetArray(SpillDecompSetID);
+			Octree->CollectRootTriangles(*OctreeCut, 
+				[&SpillTriangleSet](int TriangleID) {
+					SpillTriangleSet.Add(TriangleID);
+			});
+			Octree->CollectSpillTriangles(
+				[&SpillTriangleSet](int TriangleID) {
+				SpillTriangleSet.Add(TriangleID);
+			});
+		}
+
+		NewProxy->InitializeFromDecomposition(TriangleDecomposition);
+	}
+	return NewProxy;
+}
+
+
+
+void UOctreeDynamicMeshComponent::NotifyMaterialSetUpdated()
+{
+	if (GetCurrentSceneProxy() != nullptr)
+	{
+		GetCurrentSceneProxy()->UpdatedReferencedMaterials();
+	}
+}
+
+
+
+
+FColor UOctreeDynamicMeshComponent::GetTriangleColor(int TriangleID)
+{
+	if (TriangleColorFunc != nullptr)
+	{
+		return TriangleColorFunc(GetMesh(), TriangleID);
+	}
+	else
+	{
+		return (TriangleID % 2 == 0) ? FColor::Red : FColor::White;
+	}
+}
+
+
+
+FBoxSphereBounds UOctreeDynamicMeshComponent::CalcBounds(const FTransform& LocalToWorld) const
+{
+	// Bounds are tighter if the box is generated from transformed vertices.
+	FBox BoundingBox(ForceInit);
+	for ( FVector3d Vertex : GetMesh()->VerticesItr() ) 
+	{
+		BoundingBox += LocalToWorld.TransformPosition((FVector)Vertex);
+	}
+	return FBoxSphereBounds(BoundingBox);
+}
+
+
+
+void UOctreeDynamicMeshComponent::ApplyChange(const FMeshVertexChange* Change, bool bRevert)
+{
+	// This function is not currently used in the codebase, and needs to be restructured
+	// to work properly with passing the change down to the UDynamicMesh
+	check(false);
+
+/*
+	Octree->ResetModifiedBounds();
+
+	const FDynamicMesh3* Mesh = GetMesh();
+
+	TSet<int> TrianglesToUpdate;
+	auto NotifyTriVerticesAffected = [&](int32 vid) 
+	{
+		for (int tid : Mesh->VtxTrianglesItr(vid))
+		{
+			if (TrianglesToUpdate.Contains(tid) == false)
+			{
+				Octree->NotifyPendingModification(tid);
+				TrianglesToUpdate.Add(tid);
+			}
+		}
+	};
+
+	int NV = Change->Vertices.Num();
+	const TArray<FVector3d>& Positions = (bRevert) ? Change->OldPositions : Change->NewPositions;
+	for (int k = 0; k < NV; ++k)
+	{
+		int vid = Change->Vertices[k];
+		NotifyTriVerticesAffected(vid);
+		Mesh->SetVertex(vid, Positions[k]);
+	}
+
+	if (Change->bHaveOverlayNormals && Mesh->HasAttributes() && Mesh->Attributes()->PrimaryNormals())
+	{
+		FDynamicMeshNormalOverlay* Overlay = Mesh->Attributes()->PrimaryNormals();
+		int32 NumNormals = Change->Normals.Num();
+		const TArray<FVector3f>& UseNormals = (bRevert) ? Change->OldNormals : Change->NewNormals;
+		for (int32 k = 0; k < NumNormals; ++k)
+		{
+			int32 elemid = Change->Normals[k];
+			if (Overlay->IsElement(elemid))
+			{
+				Overlay->SetElement(elemid, UseNormals[k]);
+				int32 ParentVID = Overlay->GetParentVertex(elemid);
+				NotifyTriVerticesAffected(ParentVID);
+			}
+		}
+	}
+
+	Octree->ReinsertTriangles(TrianglesToUpdate);
+
+	//NotifyMeshUpdated();
+
+	OnMeshChanged.Broadcast();
+*/
+}
+
+
+
+
+void UOctreeDynamicMeshComponent::ApplyChange(const FMeshChange* Change, bool bRevert)
+{
+	MeshObject->ApplyChange(Change, bRevert);
+
+	//NotifyMeshUpdated();
+
+	OnMeshChanged.Broadcast();
+}
+
+
+void UOctreeDynamicMeshComponent::ApplyChange(const FMeshReplacementChange* Change, bool bRevert)
+{
+	// full clear and reset
+	MeshObject->ApplyChange(Change, bRevert);
+
+	Octree = MakeUnique<FDynamicMeshOctree3>();
+	Octree->Initialize(GetMesh());
+	
+	OctreeCut = MakeUnique<FDynamicMeshOctree3::FTreeCutSet>();
+	CutCellSetMap.Reset();
+
+	// need to force proxy recreation because we changed the Octree pointer...
+	MarkRenderStateDirty();
+
+	OnMeshChanged.Broadcast();
+}
+
+
+void UOctreeDynamicMeshComponent::OnPreMeshObjectChanged(UDynamicMesh* ChangedMeshObject, FDynamicMeshChangeInfo ChangeInfo)
+{
+	TArray<int> RemoveTriangles;
+	if (ChangeInfo.Type == EDynamicMeshChangeType::MeshChange && ensure(ChangeInfo.MeshChange != nullptr))
+	{
+		ChangeInfo.MeshChange->DynamicMeshChange->GetSavedTriangleList(RemoveTriangles, !ChangeInfo.bIsRevertChange);
+	}
+	if (RemoveTriangles.Num() > 0)
+	{
+		Octree->ResetModifiedBounds();
+		Octree->RemoveTriangles(RemoveTriangles);
+	}
+}
+
+
+void UOctreeDynamicMeshComponent::OnMeshObjectChanged(UDynamicMesh* ChangedMeshObject, FDynamicMeshChangeInfo ChangeInfo)
+{
+	if (ChangeInfo.Type == EDynamicMeshChangeType::MeshChange && ensure(ChangeInfo.MeshChange != nullptr))
+	{
+		TArray<int> AddTriangles;
+		ChangeInfo.MeshChange->DynamicMeshChange->GetSavedTriangleList(AddTriangles, ChangeInfo.bIsRevertChange);
+		Octree->InsertTriangles(AddTriangles);
+	}
+
+	NotifyMeshUpdated();
+	OnMeshChanged.Broadcast();
+}
+
+
+void UOctreeDynamicMeshComponent::SetDynamicMesh(UDynamicMesh* NewMesh)
+{
+	if (ensure(NewMesh) == false)
+	{
+		return;
+	}
+
+	if (ensure(MeshObject))
+	{
+		MeshObject->OnMeshChanged().Remove(MeshObjectChangedHandle);
+		MeshObject->OnPreMeshChanged().Remove(PreMeshChangeHandle);
+	}
+
+	MeshObject = NewMesh;
+	MeshObjectChangedHandle = MeshObject->OnMeshChanged().AddUObject(this, &UOctreeDynamicMeshComponent::OnMeshObjectChanged);
+	PreMeshChangeHandle = MeshObject->OnPreMeshChanged().AddUObject(this, &UOctreeDynamicMeshComponent::OnPreMeshObjectChanged);
+
+	Octree = MakeUnique<FDynamicMeshOctree3>();
+	Octree->Initialize(GetMesh());
+	OctreeCut = MakeUnique<FDynamicMeshOctree3::FTreeCutSet>();
+	CutCellSetMap.Reset();
+
+	NotifyMeshUpdated();
+	OnMeshChanged.Broadcast();
+}
\ No newline at end of file
diff --git a/Engine/Plugins/Experimental/MeshModelingToolset/Source/ModelingComponents/Private/Components/OctreeDynamicMeshSceneProxy.h b/Engine/Plugins/Experimental/MeshModelingToolset/Source/ModelingComponents/Private/Components/OctreeDynamicMeshSceneProxy.h
new file mode 100644
--- /dev/null
+++ b/Engine/Plugins/Experimental/MeshModelingToolset/Source/ModelingComponents/Private/Components/OctreeDynamicMeshSceneProxy.h
@@ -0,0 +1,281 @@
+// Copyright Epic Games, Inc. All Rights Reserved. 
+
+#pragma once
+
+#include "Components/OctreeDynamicMeshComponent.h"
+#include "Components/BaseDynamicMeshSceneProxy.h"
+#include "Util/IndexSetDecompositions.h"
+
+
+DECLARE_STATS_GROUP(TEXT("SculptToolOctree"), STATGROUP_SculptToolOctree, STATCAT_Advanced);
+DECLARE_CYCLE_STAT(TEXT("SculptToolOctree_UpdateExisting"), STAT_SculptToolOctree_UpdateExisting, STATGROUP_SculptToolOctree);
+DECLARE_CYCLE_STAT(TEXT("SculptToolOctree_UpdateCutSet"), STAT_SculptToolOctree_UpdateCutSet, STATGROUP_SculptToolOctree);
+DECLARE_CYCLE_STAT(TEXT("SculptToolOctree_CreateNew"), STAT_SculptToolOctree_CreateNew, STATGROUP_SculptToolOctree);
+DECLARE_CYCLE_STAT(TEXT("SculptToolOctree_UpdateSpill"), STAT_SculptToolOctree_UpdateSpill, STATGROUP_SculptToolOctree);
+DECLARE_CYCLE_STAT(TEXT("SculptToolOctree_UpdateFromDecomp"), STAT_SculptToolOctree_UpdateFromDecomp, STATGROUP_SculptToolOctree);
+DECLARE_CYCLE_STAT(TEXT("SculptToolOctree_UpdateDecompDestroy"), STAT_SculptToolOctree_UpdateDecompDestroy, STATGROUP_SculptToolOctree);
+DECLARE_CYCLE_STAT(TEXT("SculptToolOctree_UpdateDecompCreate"), STAT_SculptToolOctree_UpdateDecompCreate, STATGROUP_SculptToolOctree);
+DECLARE_CYCLE_STAT(TEXT("SculptToolOctree_InitializeBufferFromOverlay"), STAT_SculptToolOctree_InitializeBufferFromOverlay, STATGROUP_SculptToolOctree);
+DECLARE_CYCLE_STAT(TEXT("SculptToolOctree_BufferUpload"), STAT_SculptToolOctree_BufferUpload, STATGROUP_SculptToolOctree);
+
+
+/**
+ * Scene Proxy for a mesh buffer.
+ * 
+ * Based on FProceduralMeshSceneProxy but simplified in various ways.
+ * 
+ * Supports wireframe-on-shaded rendering.
+ * 
+ */
+class FOctreeDynamicMeshSceneProxy final : public FBaseDynamicMeshSceneProxy
+{
+	using FIndex2i = UE::Geometry::FIndex2i;
+	using FIndex3i = UE::Geometry::FIndex3i;
+private:
+	FMaterialRelevance MaterialRelevance;
+
+	// note: FBaseDynamicMeshSceneProxy owns and will destroy these
+	TMap<int32, FMeshRenderBufferSet*> RenderBufferSets;
+
+public:
+	/** Component that created this proxy (is there a way to look this up?) */
+	UOctreeDynamicMeshComponent* ParentComponent;
+
+
+	FOctreeDynamicMeshSceneProxy(UOctreeDynamicMeshComponent* Component)
+		: FBaseDynamicMeshSceneProxy(Component)
+		, MaterialRelevance(Component->GetMaterialRelevance(GetScene().GetFeatureLevel()))
+	{
+		// This is an assumption we are currently making. We do not necessarily require this
+		// but if this check is hit then possibly an assumption is wrong
+		check(IsInGameThread());
+
+		ParentComponent = Component;
+	}
+
+
+
+	virtual void GetActiveRenderBufferSets(TArray<FMeshRenderBufferSet*>& Buffers) const override
+	{
+		for (auto MapPair : RenderBufferSets)
+		{
+			Buffers.Add(MapPair.Value);
+		}
+	}
+
+
+
+	TUniqueFunction<void(int, int, int, const FVector3f&, FVector3f&, FVector3f&)> MakeTangentsFunc()
+	{
+		return [](int VertexID, int TriangleID, int TriVtxIdx, const FVector3f& Normal, FVector3f& TangentX, FVector3f& TangentY) -> void
+		{
+			VectorUtil::MakePerpVectors(Normal, TangentX, TangentY);
+		};
+	}
+
+
+
+	void InitializeSingleBuffer()
+	{
+		check(RenderBufferSets.Num() == 0);
+
+		FDynamicMesh3* Mesh = ParentComponent->GetMesh();
+
+		FMeshRenderBufferSet* RenderBuffers = AllocateNewRenderBufferSet();
+		RenderBuffers->Material = GetMaterial(0);
+
+		// find suitable overlays
+		FDynamicMeshUVOverlay* UVOverlay = nullptr;
+		FDynamicMeshNormalOverlay* NormalOverlay = nullptr;
+		FDynamicMeshColorOverlay* ColorOverlay = nullptr;
+		if (Mesh->HasAttributes())
+		{
+			UVOverlay = Mesh->Attributes()->PrimaryUV();
+			NormalOverlay = Mesh->Attributes()->PrimaryNormals();
+			ColorOverlay = Mesh->Attributes()->PrimaryColors();
+		}
+
+		InitializeBuffersFromOverlays(RenderBuffers, Mesh,
+			Mesh->TriangleCount(), Mesh->TriangleIndicesItr(),
+			UVOverlay, NormalOverlay, ColorOverlay, MakeTangentsFunc());
+
+		ENQUEUE_RENDER_COMMAND(FOctreeDynamicMeshSceneProxyInitializeSingle)(
+			[this, RenderBuffers](FRHICommandListImmediate& RHICmdList)
+		{
+			RenderBuffers->Upload();
+			RenderBufferSets.Add(0, RenderBuffers);
+		});
+	}
+
+
+
+	void InitializeFromDecomposition(const UE::Geometry::FArrayIndexSetsDecomposition& Decomposition)
+	{
+		check(RenderBufferSets.Num() == 0);
+
+		FDynamicMesh3* Mesh = ParentComponent->GetMesh();
+
+		// find suitable overlays
+		FDynamicMeshUVOverlay* UVOverlay = nullptr;
+		FDynamicMeshNormalOverlay* NormalOverlay = nullptr;
+		FDynamicMeshColorOverlay* ColorOverlay = nullptr;
+		if (Mesh->HasAttributes())
+		{
+			UVOverlay = Mesh->Attributes()->PrimaryUV();
+			NormalOverlay = Mesh->Attributes()->PrimaryNormals();
+			ColorOverlay = Mesh->Attributes()->PrimaryColors();
+		}
+		auto TangentsFunc = MakeTangentsFunc();
+
+		const TArray<int32>& SetIDs = Decomposition.GetIndexSetIDs();
+		for (int32 SetID : SetIDs)
+		{
+			const TArray<int32>& Tris = Decomposition.GetIndexSetArray(SetID);
+			
+			FMeshRenderBufferSet* RenderBuffers = AllocateNewRenderBufferSet();
+			RenderBuffers->Material = GetMaterial(0);
+
+			InitializeBuffersFromOverlays(RenderBuffers, Mesh,
+				Tris.Num(), Tris,
+				UVOverlay, NormalOverlay, ColorOverlay, TangentsFunc);
+
+			ENQUEUE_RENDER_COMMAND(FOctreeDynamicMeshSceneProxyInitializeFromDecomposition)(
+				[this, SetID, RenderBuffers](FRHICommandListImmediate& RHICmdList)
+			{
+				RenderBuffers->Upload();
+				RenderBufferSets.Add(SetID, RenderBuffers);
+			});
+		}
+	}
+
+
+
+	void UpdateFromDecomposition(const UE::Geometry::FArrayIndexSetsDecomposition& Decomposition, const TArray<int32>& SetsToUpdate )
+	{
+		// CAN WE REUSE EXISTING BUFFER SETS??
+		//   - could have timestamp for each decomposition set array...if tris don't change we only have to update vertices
+		//   - can re-use allocated memory if new data is smaller
+
+		SCOPE_CYCLE_COUNTER(STAT_SculptToolOctree_UpdateFromDecomp);
+
+		// remove sets to update
+		ENQUEUE_RENDER_COMMAND(FOctreeDynamicMeshSceneProxyUpdatePreClean)(
+			[this, SetsToUpdate](FRHICommandListImmediate& RHICmdList)
+		{
+			SCOPE_CYCLE_COUNTER(STAT_SculptToolOctree_UpdateDecompDestroy);
+			for (int32 SetID : SetsToUpdate)
+			{
+				if (RenderBufferSets.Contains(SetID))
+				{
+					FMeshRenderBufferSet* BufferSet = RenderBufferSets.FindAndRemoveChecked(SetID);
+					ReleaseRenderBufferSet(BufferSet);
+				}
+			}
+		});
+
+		FDynamicMesh3* Mesh = ParentComponent->GetMesh();
+
+		// find suitable overlays
+		FDynamicMeshUVOverlay* UVOverlay = nullptr;
+		FDynamicMeshNormalOverlay* NormalOverlay = nullptr;
+		FDynamicMeshColorOverlay* ColorOverlay = nullptr;
+		if (Mesh->HasAttributes())
+		{
+			UVOverlay = Mesh->Attributes()->PrimaryUV();
+			NormalOverlay = Mesh->Attributes()->PrimaryNormals();
+			ColorOverlay = Mesh->Attributes()->PrimaryColors();
+		}
+		auto TangentsFunc = MakeTangentsFunc();
+
+		{
+			SCOPE_CYCLE_COUNTER(STAT_SculptToolOctree_UpdateDecompCreate);
+			int NumSets = SetsToUpdate.Num();
+			ParallelFor(NumSets, [&](int k)
+			{
+				int32 SetID = SetsToUpdate[k];
+				const TArray<int32>& Tris = Decomposition.GetIndexSetArray(SetID);
+
+				FMeshRenderBufferSet* RenderBuffers = AllocateNewRenderBufferSet();
+				RenderBuffers->Material = GetMaterial(0);
+
+				InitializeBuffersFromOverlays(RenderBuffers, Mesh,
+					Tris.Num(), Tris,
+					UVOverlay, NormalOverlay, ColorOverlay, TangentsFunc);
+
+				ENQUEUE_RENDER_COMMAND(FOctreeDynamicMeshSceneProxyUpdateAddOne)(
+					[this, SetID, RenderBuffers](FRHICommandListImmediate& RHICmdList)
+				{
+					SCOPE_CYCLE_COUNTER(STAT_SculptToolOctree_BufferUpload);
+					RenderBuffers->Upload();
+					RenderBufferSets.Add(SetID, RenderBuffers);
+				});
+			}, false);
+
+
+		}
+
+		//UE_LOG(LogTemp, Warning, TEXT("Have %d renderbuffers"), RenderBufferSets.Num());
+	}
+
+
+
+public:
+
+
+
+
+	virtual FPrimitiveViewRelevance GetViewRelevance(const FSceneView* View) const override
+	{
+		FPrimitiveViewRelevance Result;
+
+		Result.bDrawRelevance = IsShown(View);
+		Result.bShadowRelevance = IsShadowCast(View);
+		Result.bDynamicRelevance = true;
+		Result.bRenderInMainPass = ShouldRenderInMainPass();
+		Result.bUsesLightingChannels = GetLightingChannelMask() != GetDefaultLightingChannelMask();
+		Result.bTranslucentSelfShadow = bCastVolumetricTranslucentShadow;
+		Result.bRenderCustomDepth = ShouldRenderCustomDepth();
+		MaterialRelevance.SetPrimitiveViewRelevance(Result);
+		Result.bVelocityRelevance = DrawsVelocity() && Result.bOpaque && Result.bRenderInMainPass;
+
+		return Result;
+	}
+
+	virtual void UpdatedReferencedMaterials() override
+	{
+		FBaseDynamicMeshSceneProxy::UpdatedReferencedMaterials();
+
+		// The material relevance may need updating.
+		MaterialRelevance = ParentComponent->GetMaterialRelevance(GetScene().GetFeatureLevel());
+	}
+
+	virtual void GetLightRelevance(const FLightSceneProxy* LightSceneProxy, bool& bDynamic, bool& bRelevant, bool& bLightMapped, bool& bShadowMapped) const override
+	{
+		FPrimitiveSceneProxy::GetLightRelevance(LightSceneProxy, bDynamic, bRelevant, bLightMapped, bShadowMapped);
+	}
+
+
+	virtual bool CanBeOccluded() const override
+	{
+		return !MaterialRelevance.bDisableDepthTest;
+	}
+
+	virtual uint32 GetMemoryFootprint(void) const override { return(sizeof(*this) + GetAllocatedSize()); }
+
+	uint32 GetAllocatedSize(void) const { return(FPrimitiveSceneProxy::GetAllocatedSize()); }
+
+
+
+	SIZE_T GetTypeHash() const override
+	{
+		static size_t UniquePointer;
+		return reinterpret_cast<size_t>(&UniquePointer);
+	}
+};
+
+
+
+
+
+
+
diff --git a/Engine/Plugins/Experimental/MeshModelingToolset/Source/ModelingComponents/Private/ModelingToolTargetUtil.cpp b/Engine/Plugins/Experimental/MeshModelingToolset/Source/ModelingComponents/Private/ModelingToolTargetUtil.cpp
--- a/Engine/Plugins/Experimental/MeshModelingToolset/Source/ModelingComponents/Private/ModelingToolTargetUtil.cpp
+++ b/Engine/Plugins/Experimental/MeshModelingToolset/Source/ModelingComponents/Private/ModelingToolTargetUtil.cpp
@@ -13,7 +13,7 @@
 #include "Components/PrimitiveComponent.h"
 #include "Components/StaticMeshComponent.h"
 #include "GameFramework/Volume.h"
-#include "SimpleDynamicMeshComponent.h"
+#include "Components/DynamicMeshComponent.h"
 
 #include "MeshConversionOptions.h"
 #include "MeshDescriptionToDynamicMesh.h"
@@ -261,7 +261,7 @@ bool UE::ToolTarget::ConfigureCreateMeshObjectParams(UToolTarget* SourceTarget,
 			return true;
 		}
 
-		if (Cast<USimpleDynamicMeshComponent>(ComponentTarget->GetOwnerComponent()) != nullptr)
+		if (Cast<UDynamicMeshComponent>(ComponentTarget->GetOwnerComponent()) != nullptr)
 		{
 			DerivedParamsOut.TypeHint = ECreateObjectTypeHint::DynamicMeshActor;
 			return true;
diff --git a/Engine/Plugins/Experimental/MeshModelingToolset/Source/ModelingComponents/Private/OctreeDynamicMeshComponent.cpp b/Engine/Plugins/Experimental/MeshModelingToolset/Source/ModelingComponents/Private/OctreeDynamicMeshComponent.cpp
deleted file mode 100644
diff --git a/Engine/Plugins/Experimental/MeshModelingToolset/Source/ModelingComponents/Private/OctreeDynamicMeshSceneProxy.h b/Engine/Plugins/Experimental/MeshModelingToolset/Source/ModelingComponents/Private/OctreeDynamicMeshSceneProxy.h
deleted file mode 100644
diff --git a/Engine/Plugins/Experimental/MeshModelingToolset/Source/ModelingComponents/Private/PreviewMesh.cpp b/Engine/Plugins/Experimental/MeshModelingToolset/Source/ModelingComponents/Private/PreviewMesh.cpp
--- a/Engine/Plugins/Experimental/MeshModelingToolset/Source/ModelingComponents/Private/PreviewMesh.cpp
+++ b/Engine/Plugins/Experimental/MeshModelingToolset/Source/ModelingComponents/Private/PreviewMesh.cpp
@@ -39,7 +39,7 @@ void UPreviewMesh::CreateInWorld(UWorld* World, const FTransform& WithTransform)
 	FActorSpawnParameters SpawnInfo;
 	TemporaryParentActor = World->SpawnActor<APreviewMeshActor>(FVector::ZeroVector, Rotation, SpawnInfo);
 
-	DynamicMeshComponent = NewObject<USimpleDynamicMeshComponent>(TemporaryParentActor);
+	DynamicMeshComponent = NewObject<UDynamicMeshComponent>(TemporaryParentActor);
 	TemporaryParentActor->SetRootComponent(DynamicMeshComponent);
 	//DynamicMeshComponent->SetupAttachment(TemporaryParentActor->GetRootComponent());
 	DynamicMeshComponent->RegisterComponent();
diff --git a/Engine/Plugins/Experimental/MeshModelingToolset/Source/ModelingComponents/Private/Selection/PolygonSelectionMechanic.cpp b/Engine/Plugins/Experimental/MeshModelingToolset/Source/ModelingComponents/Private/Selection/PolygonSelectionMechanic.cpp
--- a/Engine/Plugins/Experimental/MeshModelingToolset/Source/ModelingComponents/Private/Selection/PolygonSelectionMechanic.cpp
+++ b/Engine/Plugins/Experimental/MeshModelingToolset/Source/ModelingComponents/Private/Selection/PolygonSelectionMechanic.cpp
@@ -151,7 +151,7 @@ void UPolygonSelectionMechanic::Initialize(
 }
 
 void UPolygonSelectionMechanic::Initialize(
-	USimpleDynamicMeshComponent* MeshComponentIn,
+	UDynamicMeshComponent* MeshComponentIn,
 	const FGroupTopology* TopologyIn,
 	TFunction<FDynamicMeshAABBTree3 * ()> GetSpatialSourceFuncIn)
 {
diff --git a/Engine/Plugins/Experimental/MeshModelingToolset/Source/ModelingComponents/Private/SimpleDynamicMeshComponent.cpp b/Engine/Plugins/Experimental/MeshModelingToolset/Source/ModelingComponents/Private/SimpleDynamicMeshComponent.cpp
deleted file mode 100644
diff --git a/Engine/Plugins/Experimental/MeshModelingToolset/Source/ModelingComponents/Private/SimpleDynamicMeshSceneProxy.h b/Engine/Plugins/Experimental/MeshModelingToolset/Source/ModelingComponents/Private/SimpleDynamicMeshSceneProxy.h
deleted file mode 100644
diff --git a/Engine/Plugins/Experimental/MeshModelingToolset/Source/ModelingComponents/Public/BaseDynamicMeshComponent.h b/Engine/Plugins/Experimental/MeshModelingToolset/Source/ModelingComponents/Public/BaseDynamicMeshComponent.h
deleted file mode 100644
diff --git a/Engine/Plugins/Experimental/MeshModelingToolset/Source/ModelingComponents/Public/BaseTools/BaseMeshProcessingTool.h b/Engine/Plugins/Experimental/MeshModelingToolset/Source/ModelingComponents/Public/BaseTools/BaseMeshProcessingTool.h
--- a/Engine/Plugins/Experimental/MeshModelingToolset/Source/ModelingComponents/Public/BaseTools/BaseMeshProcessingTool.h
+++ b/Engine/Plugins/Experimental/MeshModelingToolset/Source/ModelingComponents/Public/BaseTools/BaseMeshProcessingTool.h
@@ -14,7 +14,7 @@
 
 // predeclarations
 struct FMeshDescription;
-class USimpleDynamicMeshComponent;
+class UDynamicMeshComponent;
 class UPreviewMesh;
 PREDECLARE_USE_GEOMETRY_CLASS(FMeshBoundaryLoops);
 PREDECLARE_USE_GEOMETRY_CLASS(FMeshNormals);
diff --git a/Engine/Plugins/Experimental/MeshModelingToolset/Source/ModelingComponents/Public/Changes/MeshChange.h b/Engine/Plugins/Experimental/MeshModelingToolset/Source/ModelingComponents/Public/Changes/MeshChange.h
--- a/Engine/Plugins/Experimental/MeshModelingToolset/Source/ModelingComponents/Public/Changes/MeshChange.h
+++ b/Engine/Plugins/Experimental/MeshModelingToolset/Source/ModelingComponents/Public/Changes/MeshChange.h
@@ -16,7 +16,7 @@
 
 /**
  * FMeshChange represents an undoable change to a FDynamicMesh3.
- * Currently only valid to call Apply/Revert when the Object is a one of several components backed by FDynamicMesh: USimpleDynamicMeshComponent, UOctreeDynamicMeshComponent, UPreviewMesh
+ * Currently only valid to call Apply/Revert when the Object is a one of several components backed by FDynamicMesh: UDynamicMeshComponent, UOctreeDynamicMeshComponent, UPreviewMesh
  */
 class MODELINGCOMPONENTS_API FMeshChange : public FToolCommandChange
 {
diff --git a/Engine/Plugins/Experimental/MeshModelingToolset/Source/ModelingComponents/Public/Changes/MeshReplacementChange.h b/Engine/Plugins/Experimental/MeshModelingToolset/Source/ModelingComponents/Public/Changes/MeshReplacementChange.h
--- a/Engine/Plugins/Experimental/MeshModelingToolset/Source/ModelingComponents/Public/Changes/MeshReplacementChange.h
+++ b/Engine/Plugins/Experimental/MeshModelingToolset/Source/ModelingComponents/Public/Changes/MeshReplacementChange.h
@@ -14,7 +14,7 @@ PREDECLARE_USE_GEOMETRY_CLASS(FDynamicMesh3);
 
 /**
  * FMeshReplacementChange represents an undoable *complete* change to a FDynamicMesh3.
- * Currently only valid to call Apply/Revert when the Object is a USimpleDynamicMeshComponent
+ * Currently only valid to call Apply/Revert when the Object is a UDynamicMeshComponent
  */
 class MODELINGCOMPONENTS_API FMeshReplacementChange : public FToolCommandChange
 {
diff --git a/Engine/Plugins/Experimental/MeshModelingToolset/Source/ModelingComponents/Public/Changes/MeshVertexChange.h b/Engine/Plugins/Experimental/MeshModelingToolset/Source/ModelingComponents/Public/Changes/MeshVertexChange.h
--- a/Engine/Plugins/Experimental/MeshModelingToolset/Source/ModelingComponents/Public/Changes/MeshVertexChange.h
+++ b/Engine/Plugins/Experimental/MeshModelingToolset/Source/ModelingComponents/Public/Changes/MeshVertexChange.h
@@ -12,7 +12,7 @@ PREDECLARE_USE_GEOMETRY_CLASS(FDynamicMesh3);
 
 /**
  * FMeshVertexChange represents an reversible change to a set of vertex positions.
- * Currently only a USimpleDynamicMeshComponent target is supported.
+ * Currently only a UDynamicMeshComponent target is supported.
  * 
  * @todo support optionally storing old/new normals and tangents
  * @todo support applying to a StaticMeshComponent/MeshDescription ?
diff --git a/Engine/Plugins/Experimental/MeshModelingToolset/Source/ModelingComponents/Public/Components/BaseDynamicMeshComponent.h b/Engine/Plugins/Experimental/MeshModelingToolset/Source/ModelingComponents/Public/Components/BaseDynamicMeshComponent.h
new file mode 100644
--- /dev/null
+++ b/Engine/Plugins/Experimental/MeshModelingToolset/Source/ModelingComponents/Public/Components/BaseDynamicMeshComponent.h
@@ -0,0 +1,316 @@
+// Copyright Epic Games, Inc. All Rights Reserved.
+
+#pragma once
+
+#include "CoreMinimal.h"
+#include "Components/MeshComponent.h"
+#include "InteractiveToolObjects.h"
+#include "Changes/MeshVertexChange.h"
+#include "Changes/MeshChange.h"
+#include "Changes/MeshReplacementChange.h"
+#include "MeshConversionOptions.h"
+#include "DynamicMesh3.h"		// todo replace with predeclaration (lots of fallout)
+#include "UDynamicMesh.h"
+
+#include "BaseDynamicMeshComponent.generated.h"
+
+// predecl
+struct FMeshDescription;
+class FMeshVertexChange;
+class FMeshChange;
+using UE::Geometry::FDynamicMesh3;
+
+/**
+ * EMeshRenderAttributeFlags is used to identify different mesh rendering attributes, for things
+ * like fast-update functions
+ */
+enum class EMeshRenderAttributeFlags : uint8
+{
+	None = 0,
+	Positions = 1,
+	VertexColors = 1<<1,
+	VertexNormals = 1<<2,
+	VertexUVs = 1<<3,
+	SecondaryIndexBuffers = 1<<4,
+
+	AllVertexAttribs = Positions | VertexColors | VertexNormals | VertexUVs
+};
+ENUM_CLASS_FLAGS(EMeshRenderAttributeFlags);
+
+
+
+/**
+ * Tangent calculation modes
+ */
+UENUM()
+enum class EDynamicMeshComponentTangentsMode : uint8
+{
+	/** Tangents are not used/available, proceed accordingly (eg generate arbitrary orthogonal basis) */
+	NoTangents,
+	/** Tangents will be automatically calculated on demand. Note that mesh changes due to tangents calculation will *not* be broadcast via MeshChange events! */
+	AutoCalculated,
+	/** Tangents are externally provided via the FDynamicMesh3 AttributeSet */
+	ExternallyProvided
+};
+
+
+/**
+ * UBaseDynamicMeshComponent is a base interface for a UMeshComponent based on a UDynamicMesh.
+ */
+UCLASS(Abstract, hidecategories = (LOD, Physics, Collision), ClassGroup = Rendering)
+class MODELINGCOMPONENTS_API UBaseDynamicMeshComponent : 
+	public UMeshComponent, 
+	public IToolFrameworkComponent, 
+	public IMeshVertexCommandChangeTarget, 
+	public IMeshCommandChangeTarget, 
+	public IMeshReplacementCommandChangeTarget
+{
+	GENERATED_UCLASS_BODY()
+
+
+	//===============================================================================================================
+	// UBaseDynamicMeshComponent API. Subclasses must implement these functions
+	//
+public:
+	/**
+	 * initialize the internal mesh from a DynamicMesh
+	 */
+	virtual void SetMesh(UE::Geometry::FDynamicMesh3&& MoveMesh)
+	{
+		unimplemented();
+	}
+
+	/**
+	 * @return pointer to internal mesh
+	 * @warning avoid usage of this function, access via GetDynamicMesh() instead
+	 */
+	virtual FDynamicMesh3* GetMesh()
+	{
+		unimplemented();
+		return nullptr;
+	}
+
+	/**
+	 * @return pointer to internal mesh
+	 */
+	virtual const FDynamicMesh3* GetMesh() const
+	{
+		unimplemented();
+		return nullptr;
+	}
+
+	/**
+	 * Allow external code to read the internal mesh.
+	 */
+	virtual void ProcessMesh(TFunctionRef<void(const UE::Geometry::FDynamicMesh3&)> ProcessFunc) const
+	{
+		unimplemented();
+	}
+
+	/**
+	 * @return the child UDynamicMesh
+	 */
+	UFUNCTION(BlueprintCallable, Category = "Dynamic Mesh Component")
+	virtual UDynamicMesh* GetDynamicMesh() 
+	{ 
+		unimplemented();
+		return nullptr; 
+	}
+
+	/**
+	 * Call this if you update the mesh via GetMesh()
+	 * @todo should provide a function that calls a lambda to modify the mesh, and only return const mesh pointer
+	 */
+	virtual void NotifyMeshUpdated()
+	{
+		unimplemented();
+	}
+
+	/**
+	 * Apply a vertex deformation change to the internal mesh  (implements IMeshVertexCommandChangeTarget)
+	 */
+	virtual void ApplyChange(const FMeshVertexChange* Change, bool bRevert) override
+	{
+		unimplemented();
+	}
+
+	/**
+	 * Apply a general mesh change to the internal mesh  (implements IMeshCommandChangeTarget)
+	 */
+	virtual void ApplyChange(const FMeshChange* Change, bool bRevert) override
+	{
+		unimplemented();
+	}
+
+	/**
+	* Apply a full mesh replacement change to the internal mesh  (implements IMeshReplacementCommandChangeTarget)
+	*/
+	virtual void ApplyChange(const FMeshReplacementChange* Change, bool bRevert) override
+	{
+		unimplemented();
+	}
+
+	/**
+	 * Apply a transform to the mesh
+	 */
+	virtual void ApplyTransform(const UE::Geometry::FTransform3d& Transform, bool bInvert)
+	{
+		unimplemented();
+	}
+
+
+protected:
+	/**
+	 * Subclass must implement this to notify allocated proxies of updated materials
+	 */
+	virtual void NotifyMaterialSetUpdated()
+	{
+		unimplemented();
+	}
+
+
+
+
+	//===============================================================================================================
+	// Built-in Wireframe-on-Shaded Rendering support. The wireframe looks terrible but this is a convenient 
+	// way to enable/disable it.
+	//
+public:
+	/**
+	 * if true, we always show the wireframe on top of the shaded mesh, even when not in wireframe mode
+	 * @todo: this should not be public, access via Set/Get once all usage is cleaned up
+	 */
+	UPROPERTY(EditAnywhere, BlueprintReadWrite, Category = "Dynamic Mesh Component")
+	bool bExplicitShowWireframe = false;
+
+	/**
+	 * Configure whether wireframe rendering is enabled or not
+	 */
+	UFUNCTION(BlueprintCallable, Category = "Dynamic Mesh Component")
+	virtual void SetEnableWireframeRenderPass(bool bEnable) { bExplicitShowWireframe = bEnable; }
+
+	/**
+	 * @return true if wireframe rendering pass is enabled
+	 */
+	UFUNCTION(BlueprintCallable, Category = "Dynamic Mesh Component")
+	virtual bool GetEnableWireframeRenderPass() const { return bExplicitShowWireframe; }
+
+
+
+
+
+	//===============================================================================================================
+	// Override rendering material support. If an Override material is set, then it
+	// will be used during drawing of all mesh buffers except Secondary buffers.
+	//
+public:
+	/**
+	 * Set an active override render material. This should replace all materials during rendering.
+	 */
+	UFUNCTION(BlueprintCallable, Category = "Dynamic Mesh Component")
+	virtual void SetOverrideRenderMaterial(UMaterialInterface* Material);
+
+	/**
+	 * Clear any active override render material
+	 */
+	UFUNCTION(BlueprintCallable, Category = "Dynamic Mesh Component")
+	virtual void ClearOverrideRenderMaterial();
+	
+	/**
+	 * @return true if an override render material is currently enabled for the given MaterialIndex
+	 */
+	UFUNCTION(BlueprintCallable, Category = "Dynamic Mesh Component")
+	virtual bool HasOverrideRenderMaterial(int k) const
+	{
+		return OverrideRenderMaterial != nullptr;
+	}
+
+	/**
+	 * @return active override render material for the given MaterialIndex
+	 */
+	UFUNCTION(BlueprintCallable, Category = "Dynamic Mesh Component")
+	virtual UMaterialInterface* GetOverrideRenderMaterial(int MaterialIndex) const
+	{
+		return OverrideRenderMaterial;
+	}
+
+protected:
+	UPROPERTY()
+	TObjectPtr<UMaterialInterface> OverrideRenderMaterial = nullptr;
+
+
+
+
+
+	//===============================================================================================================
+	// Secondary Render Buffers support. This requires implementation in subclasses. It allows
+	// a subset of the mesh triangles to be moved to a separate set of render buffers, which
+	// can then have a separate material (eg to highlight faces), or be shown/hidden independently.
+	// 
+public:
+	/**
+	 * Set an active secondary render material. 
+	 */
+	UFUNCTION(BlueprintCallable, Category = "Dynamic Mesh Component")
+	virtual void SetSecondaryRenderMaterial(UMaterialInterface* Material);
+
+	/**
+	 * Clear any active secondary render material
+	 */
+	UFUNCTION(BlueprintCallable, Category = "Dynamic Mesh Component")
+	virtual void ClearSecondaryRenderMaterial();
+
+	/**
+	 * @return true if a secondary render material is set
+	 */
+	virtual bool HasSecondaryRenderMaterial() const
+	{
+		return SecondaryRenderMaterial != nullptr;
+	}
+
+	/**
+	 * @return active secondary render material
+	 */
+	UFUNCTION(BlueprintCallable, Category = "Dynamic Mesh Component")
+	virtual UMaterialInterface* GetSecondaryRenderMaterial() const
+	{
+		return SecondaryRenderMaterial;
+	}
+
+	/**
+	 * Show/Hide the secondary triangle buffers. Does not invalidate SceneProxy.
+	 */
+	UFUNCTION(BlueprintCallable, Category = "Dynamic Mesh Component")
+	virtual void SetSecondaryBuffersVisibility(bool bSetVisible);
+
+	/**
+	 * @return true if secondary buffers are currently set to be visible
+	 */
+	UFUNCTION(BlueprintCallable, Category = "Dynamic Mesh Component")
+	virtual bool GetSecondaryBuffersVisibility() const;
+
+protected:
+	UPROPERTY()
+	TObjectPtr<UMaterialInterface> SecondaryRenderMaterial = nullptr;
+
+	bool bDrawSecondaryBuffers = true;
+
+
+
+
+
+	//===============================================================================================================
+	// Standard Component interfaces
+	//
+public:
+
+	// UMeshComponent Interface.
+	virtual int32 GetNumMaterials() const override;
+	virtual UMaterialInterface* GetMaterial(int32 ElementIndex) const override;
+	virtual FMaterialRelevance GetMaterialRelevance(ERHIFeatureLevel::Type InFeatureLevel) const override;
+	virtual void SetMaterial(int32 ElementIndex, UMaterialInterface* Material) override;
+	virtual void GetUsedMaterials(TArray<UMaterialInterface*>& OutMaterials, bool bGetDebugMaterials = false) const override;
+
+	UPROPERTY()
+	TArray<TObjectPtr<UMaterialInterface>> BaseMaterials;
+};
diff --git a/Engine/Plugins/Experimental/MeshModelingToolset/Source/ModelingComponents/Public/Components/DynamicMeshComponent.h b/Engine/Plugins/Experimental/MeshModelingToolset/Source/ModelingComponents/Public/Components/DynamicMeshComponent.h
new file mode 100644
--- /dev/null
+++ b/Engine/Plugins/Experimental/MeshModelingToolset/Source/ModelingComponents/Public/Components/DynamicMeshComponent.h
@@ -0,0 +1,479 @@
+// Copyright Epic Games, Inc. All Rights Reserved.
+
+#pragma once
+
+#include "CoreMinimal.h"
+#include "BaseDynamicMeshComponent.h"
+#include "MeshConversionOptions.h"
+#include "Drawing/MeshRenderDecomposition.h"
+#include "MeshTangents.h"
+#include "TransformTypes.h"
+#include "Async/Future.h"
+#include "UDynamicMesh.h"
+
+#include "DynamicMeshComponent.generated.h"
+
+// predecl
+struct FMeshDescription;
+
+/** internal FPrimitiveSceneProxy defined in DynamicMeshSceneProxy.h */
+class FDynamicMeshSceneProxy;
+
+
+/**
+ * Interface for a render mesh processor. Use this to process the Mesh stored in UDynamicMeshComponent before
+ * sending it off for rendering.
+ * NOTE: This is called whenever the Mesh is updated and before rendering, so performance matters.
+ */
+class MODELINGCOMPONENTS_API IRenderMeshPostProcessor
+{
+public:
+	virtual ~IRenderMeshPostProcessor() = default;
+
+	virtual void ProcessMesh(const FDynamicMesh3& Mesh, FDynamicMesh3& OutRenderMesh) = 0;
+};
+
+
+/** Render data update hint */
+UENUM()
+enum class EDynamicMeshComponentRenderUpdateMode
+{
+	/** Do not update render data */
+	NoUpdate = 0,
+	/** Invalidate overlay of internal component, rebuilding all render data */
+	FullUpdate = 1,
+	/** Attempt to do partial update of render data if possible */
+	FastUpdate = 2
+};
+
+
+/** 
+ * UDynamicMeshComponent is a mesh component similar to UProceduralMeshComponent,
+ * except it bases the renderable geometry off an internal UDynamicMesh instance (which
+ * encapsulates a FDynamicMesh3). 
+ * 
+ * There is extensive support for partial updates to render buffers, customizing colors,
+ * internally decomposing the mesh into separate chunks for more efficient render updates,
+ * and support for attaching a 'Postprocessor' to generate a render mesh on-the-fly
+ * See comment sections below for details.
+ * 
+ */
+UCLASS(hidecategories = (LOD, Physics, Collision), meta = (BlueprintSpawnableComponent), ClassGroup = Rendering)
+class MODELINGCOMPONENTS_API UDynamicMeshComponent : public UBaseDynamicMeshComponent
+{
+	GENERATED_UCLASS_BODY()
+
+
+
+	//===============================================================================================================
+	// Mesh Access. Usage via GetDynamicMesh() or SetMesh()/ProcessMesh()/EditMesh() is preferred, the GetMesh() 
+	// pointer access exist largely to support existing code from before UDynamicMesh was added.
+public:
+	/**
+	 * @return pointer to internal mesh
+	 * @warning avoid usage of this function, access via GetDynamicMesh() instead
+	 */
+	virtual FDynamicMesh3* GetMesh() override { return MeshObject->GetMeshPtr(); }
+
+	/**
+	 * @return pointer to internal mesh
+	 * @warning avoid usage of this function, access via GetDynamicMesh() instead
+	 */
+	virtual const FDynamicMesh3* GetMesh() const override { return MeshObject->GetMeshPtr(); }
+
+	/**
+	 * @return the child UDynamicMesh
+	 */
+	//UFUNCTION(BlueprintCallable, Category = "Dynamic Mesh Component")
+	virtual UDynamicMesh* GetDynamicMesh() override { return MeshObject; }
+
+	/**
+	 * Set the child UDynamicMesh. This can be used to 'share' a UDynamicMesh between Component instances.
+	 * @warning Currently this is somewhat risky, it is on the caller/clients to make sure that the actual mesh is not being simultaneously modified on multiple threads
+	 */
+	UFUNCTION(BlueprintCallable, Category = "Dynamic Mesh Component")
+	void SetDynamicMesh(UDynamicMesh* NewMesh);
+
+	/**
+	 * initialize the internal mesh from a DynamicMesh
+	 */
+	virtual void SetMesh(UE::Geometry::FDynamicMesh3&& MoveMesh) override;
+
+	/**
+	 * Allow external code to read the internal mesh.
+	 */
+	virtual void ProcessMesh(TFunctionRef<void(const UE::Geometry::FDynamicMesh3&)> ProcessFunc) const;
+
+	/**
+	 * Allow external code to to edit the internal mesh.
+	 */
+	virtual void EditMesh(TFunctionRef<void(UE::Geometry::FDynamicMesh3&)> EditFunc,
+						  EDynamicMeshComponentRenderUpdateMode UpdateMode = EDynamicMeshComponentRenderUpdateMode::FullUpdate);
+
+	/**
+	 * Apply transform to internal mesh. In some cases this can be more efficient than a general edit.
+	 * @param bInvert if true, inverse tranform is applied instead of forward transform
+	 */
+	virtual void ApplyTransform(const UE::Geometry::FTransform3d& Transform, bool bInvert) override;
+
+
+
+protected:
+	/**
+	 * Internal FDynamicMesh is stored inside a UDynamicMesh container, which allows it to be
+	 * used from BP, shared with other UObjects, and so on
+	 */
+	UPROPERTY(Instanced)
+	TObjectPtr<UDynamicMesh> MeshObject;
+
+
+
+
+	//===============================================================================================================
+	// RenderBuffer Update API. These functions can be used by external code (and internally in some places)
+	// to tell the Component that the Mesh data has been modified in some way, and that the RenderBuffers in the RenderProxy
+	// need to be updated (or rebuilt entirely). On large meshes a full rebuild is expensive, so there are quite a few
+	// variants that can be used to minimize the amount of data updated in different situations.
+	//
+public:
+	/**
+	 * Call this if you update the mesh via GetMesh(). This will destroy the existing RenderProxy and create a new one.
+	 * @todo should provide a function that calls a lambda to modify the mesh, and only return const mesh pointer
+	 */
+	virtual void NotifyMeshUpdated() override;
+
+	/**
+	 * Call this instead of NotifyMeshUpdated() if you have only updated the vertex colors (or triangle color function).
+	 * This function will update the existing RenderProxy buffers if possible
+	 */
+	void FastNotifyColorsUpdated();
+
+	/**
+	 * Call this instead of NotifyMeshUpdated() if you have only updated the vertex positions (and possibly some attributes).
+	 * This function will update the existing RenderProxy buffers if possible
+	 */
+	void FastNotifyPositionsUpdated(bool bNormals = false, bool bColors = false, bool bUVs = false);
+
+	/**
+	 * Call this instead of NotifyMeshUpdated() if you have only updated the vertex attributes (but not positions).
+	 * This function will update the existing RenderProxy buffers if possible, rather than create new ones.
+	 */
+	void FastNotifyVertexAttributesUpdated(bool bNormals, bool bColors, bool bUVs);
+
+	/**
+	 * Call this instead of NotifyMeshUpdated() if you have only updated the vertex positions/attributes
+	 * This function will update the existing RenderProxy buffers if possible, rather than create new ones.
+	 */
+	void FastNotifyVertexAttributesUpdated(EMeshRenderAttributeFlags UpdatedAttributes);
+
+	/**
+	 * Call this instead of NotifyMeshUpdated() if you have only updated the vertex uvs.
+	 * This function will update the existing RenderProxy buffers if possible
+	 */
+	void FastNotifyUVsUpdated();
+
+	/**
+	 * Call this instead of NotifyMeshUpdated() if you have only updated secondary triangle sorting.
+	 * This function will update the existing buffers if possible, without rebuilding entire RenderProxy.
+	 */
+	void FastNotifySecondaryTrianglesChanged();
+
+	/**
+	 * This function updates vertex positions/attributes of existing SceneProxy render buffers if possible, for the given triangles.
+	 * If a FMeshRenderDecomposition has not been explicitly set, call is forwarded to FastNotifyVertexAttributesUpdated()
+	 */
+	void FastNotifyTriangleVerticesUpdated(const TArray<int32>& Triangles, EMeshRenderAttributeFlags UpdatedAttributes);
+
+	/**
+	 * This function updates vertex positions/attributes of existing SceneProxy render buffers if possible, for the given triangles.
+	 * If a FMeshRenderDecomposition has not been explicitly set, call is forwarded to FastNotifyVertexAttributesUpdated()
+	 */
+	void FastNotifyTriangleVerticesUpdated(const TSet<int32>& Triangles, EMeshRenderAttributeFlags UpdatedAttributes);
+
+	/**
+	 * If a Decomposition is set on this Component, and everything is currently valid (proxy/etc), precompute the set of
+	 * buffers that will be modified, as well as the bounds of the modified region. These are both computed in parallel.
+	 * Use FastNotifyTriangleVerticesUpdated_ApplyPrecompute() with the returned future to apply this precomputation.
+	 * @return a future that will (eventually) return true if the precompute is OK, and (immediately) false if it is not
+	 */
+	TFuture<bool> FastNotifyTriangleVerticesUpdated_TryPrecompute(const TArray<int32>& Triangles, TArray<int32>& UpdateSetsOut, UE::Geometry::FAxisAlignedBox3d& BoundsOut);
+
+	/**
+	 * This function updates vertex positions/attributes of existing SceneProxy render buffers if possible, for the given triangles.
+	 * The assumption is that FastNotifyTriangleVerticesUpdated_TryPrecompute() was used to get the Precompute future, this function
+	 * will Wait() until it is done and then use the UpdateSets and UpdateSetBounds that were computed (must be the same variables
+	 * passed to FastNotifyTriangleVerticesUpdated_TryPrecompute). 
+	 * If the Precompute future returns false, then we forward the call to FastNotifyTriangleVerticesUpdated(), which will do more work.
+	 */
+	void FastNotifyTriangleVerticesUpdated_ApplyPrecompute(const TArray<int32>& Triangles, EMeshRenderAttributeFlags UpdatedAttributes,
+		TFuture<bool>& Precompute, const TArray<int32>& UpdateSets, const UE::Geometry::FAxisAlignedBox3d& UpdateSetBounds);
+
+
+
+
+
+	//===============================================================================================================
+	// Change Support. These changes are primarily used for Undo/Redo, however there is no strict assumption
+	// about this internally, objects of these change types could also be used to perform more structured editing.
+	// (Note that these functions simply forward the change events to the child UDynamicMesh, which will
+	// 	post a mesh-change event that 
+	//
+public:
+	/**
+	 * Apply a vertex deformation change to the mesh
+	 */
+	virtual void ApplyChange(const FMeshVertexChange* Change, bool bRevert) override;
+
+	/**
+	 * Apply a general mesh change to the mesh
+	 */
+	virtual void ApplyChange(const FMeshChange* Change, bool bRevert) override;
+
+	/**
+	* Apply a mesh replacement change to mesh
+	*/
+	virtual void ApplyChange(const FMeshReplacementChange* Change, bool bRevert) override;
+
+	/**
+	 * This delegate fires when the mesh has been changed
+	 */
+	FSimpleMulticastDelegate OnMeshChanged;
+
+	/**
+	 * This delegate fires when the mesh vertices have been changed via an FMeshVertexChange
+	 */
+	DECLARE_MULTICAST_DELEGATE_ThreeParams(FMeshVerticesModified, UDynamicMeshComponent*, const FMeshVertexChange*, bool);
+	FMeshVerticesModified OnMeshVerticesChanged;
+
+	/**
+	 * When a FMeshChange or FMeshVertexChange is applied, by default we currently fully invalidate the render proxy. However in certain
+	 * realtime situations (eg like Sculpting tools) it can be critical to undo/redo performance to do more optimized render data updates
+	 * (eg using one of the FastXYZ functions above). To allow for that, the full proxy invalidation on change can be (temporarily!) disabled 
+	 * using this function. 
+	 */
+	void SetInvalidateProxyOnChangeEnabled(bool bEnabled);
+
+	/** @return true if InvalidateProxyOnChange is enabled (default) */
+	bool GetInvalidateProxyOnChangeEnabled() const { return bInvalidateProxyOnChange; }
+
+protected:
+	/** If false, we don't completely invalidate the RenderProxy when ApplyChange() is called (assumption is it will be handled elsewhere) */
+	bool bInvalidateProxyOnChange = true;
+
+	/** Handle for OnMeshObjectChanged which is registered with MeshObject::OnMeshChanged delegate */
+	FDelegateHandle MeshObjectChangedHandle;
+
+	/** Called whenever internal MeshObject is modified, fires OnMeshChanged and OnMeshVerticesChanged above */
+	void OnMeshObjectChanged(UDynamicMesh* ChangedMeshObject, FDynamicMeshChangeInfo ChangeInfo);
+
+
+
+
+	//===============================================================================================================
+	// Support for specifying per-triangle colors as vertex colors. This allows external code to dynamically override 
+	// the vertex colors on the rendered mesh. The lambda that is passed is held for the lifetime of the Component and
+	// must remain valid. A Material that uses the vertex colors must be applied, otherwise setting this override will
+	// have no visible effect. If the colors change externally, FastNotifyColorsUpdated() can be used to do the 
+	// minimal vertex buffer updates necessary in the RenderProxy
+	//
+public:
+	/** Clear an active triangle color function if one exists, and update the mesh */
+	virtual void SetTriangleColorFunction(TUniqueFunction<FColor(const FDynamicMesh3*, int)> TriangleColorFuncIn,
+										  EDynamicMeshComponentRenderUpdateMode UpdateMode = EDynamicMeshComponentRenderUpdateMode::FastUpdate);
+
+	/** Clear an active triangle color function if one exists, and update the mesh */
+	virtual void ClearTriangleColorFunction(EDynamicMeshComponentRenderUpdateMode UpdateMode = EDynamicMeshComponentRenderUpdateMode::FastUpdate);
+
+	/** @return true if a triangle color function is configured */
+	virtual bool HasTriangleColorFunction();
+
+protected:
+	/** If this function is set, we will use these colors instead of vertex colors */
+	TUniqueFunction<FColor(const FDynamicMesh3*, int)> TriangleColorFunc = nullptr;
+
+	/** This function is passed via lambda to the RenderProxy to be able to access TriangleColorFunc */
+	FColor GetTriangleColor(const FDynamicMesh3* Mesh, int TriangleID);
+
+
+
+
+
+	//===============================================================================================================
+	// Support for Secondary triangle index buffers. When this is configured, then triangles identified
+	// by the filtering predicate function will be placed in a second set of RenderBuffers at the SceneProxy level.
+	// This can be combined with the SecondaryRenderMaterial support in UBaseDynamicMeshComponent to draw
+	// that triangle set with a different material, to efficiently accomplish UI features like highlighting a
+	// subset of mesh triangles.
+	//
+public:
+	/**
+	 * If Secondary triangle buffers are enabled, then we will filter triangles that pass the given predicate
+	 * function into a second index buffer. These triangles will be drawn with the Secondary render material
+	 * that is set in the BaseDynamicMeshComponent. Calling this function invalidates the SceneProxy.
+	 */
+	virtual void EnableSecondaryTriangleBuffers(TUniqueFunction<bool(const FDynamicMesh3*, int32)> SecondaryTriFilterFunc);
+
+	/**
+	 * Disable secondary triangle buffers. This invalidates the SceneProxy.
+	 */
+	virtual void DisableSecondaryTriangleBuffers();
+
+protected:
+	TUniqueFunction<bool(const FDynamicMesh3*, int32)> SecondaryTriFilterFunc = nullptr;
+
+
+
+
+	//===============================================================================================================
+	// Support for a Render Decomposition, which is basically a segmentation of the mesh triangles into
+	// subsets which will be turned into separate RenderBuffers in the Render Proxy. If this is configured,
+	// then various of the FastNotifyXYZUpdated() functions above will only need to rebuild the RenderBuffers
+	// that include affected triangles. The FMeshRenderDecomposition implementation has various options for
+	// building decompositions based on material, spatial clustering, etc.
+	//
+public:
+	/**
+	 * Configure a decomposition of the mesh, which will result in separate render buffers for each 
+	 * decomposition triangle group. Invalidates existing SceneProxy.
+	 */
+	virtual void SetExternalDecomposition(TUniquePtr<FMeshRenderDecomposition> Decomposition);
+
+protected:
+	TUniquePtr<FMeshRenderDecomposition> Decomposition;
+
+
+
+
+	//===============================================================================================================
+	// IRenderMeshPostProcessor Support. If a RenderMesh Postprocessor is configured, then instead of directly 
+	// passing the internal mesh to the RenderProxy, IRenderMeshPostProcessor::PostProcess is applied to populate
+	// the internal RenderMesh which is passed instead. This allows things like Displacement or Subdivision to be 
+	// done on-the-fly at the rendering level (which is potentially more efficient).
+	//
+public:
+	/**
+	 * Add a render mesh processor, to be called before the mesh is sent for rendering.
+	 */
+	virtual void SetRenderMeshPostProcessor(TUniquePtr<IRenderMeshPostProcessor> Processor);
+
+	/**
+	 * The SceneProxy should call these functions to get the post-processed RenderMesh. (See IRenderMeshPostProcessor.)
+	 */
+	virtual FDynamicMesh3* GetRenderMesh();
+
+	/**
+	 * The SceneProxy should call these functions to get the post-processed RenderMesh. (See IRenderMeshPostProcessor.)
+	 */
+	virtual const FDynamicMesh3* GetRenderMesh() const;
+
+protected:
+	TUniquePtr<IRenderMeshPostProcessor> RenderMeshPostProcessor;
+	TUniquePtr<FDynamicMesh3> RenderMesh;
+
+
+
+
+	//===============================================================================================================
+	// Support for Component attachment change notifications via delegates. Standard UE
+	// Actor/Component hierarchy does not generally provide these capabilities, but in some use
+	// cases (eg procedural mesh Actors) we need to know things like when the Component set inside
+	// an Actor is modified.
+public:
+	DECLARE_MULTICAST_DELEGATE_TwoParams(FComponentChildrenChangedDelegate, USceneComponent*, bool);
+
+	/**
+	 * The OnChildAttached() and OnChildDetached() implementations (from USceneComponent API) broadcast this delegate. This
+	 * allows Actors that have UDynamicMeshComponent's to respond to changes in their Component hierarchy.
+	 */
+	FComponentChildrenChangedDelegate OnChildAttachmentModified;
+
+
+
+
+protected:
+
+	//
+	// standard Component internals, for computing bounds and managing the SceneProxy
+	//
+
+	/** Current local-space bounding box of Mesh */
+	UE::Geometry::FAxisAlignedBox3d LocalBounds;
+
+	/** Recompute LocalBounds from the current Mesh */
+	void UpdateLocalBounds();
+
+	/**
+	 * This is called to tell our RenderProxy about modifications to the material set.
+	 * We need to pass this on for things like material validation in the Editor.
+	 */
+	virtual void NotifyMaterialSetUpdated();
+
+	/**
+	 * If the render proxy is invalidated (eg by MarkRenderStateDirty()), it will be destroyed at the end of
+	 * the frame, but the base SceneProxy pointer is not nulled out immediately. As a result if we call various
+	 * partial-update functions after invalidating the proxy, they may be operating on an invalid proxy.
+	 * So we have to keep track of proxy-valid state ourselves.
+	 */
+	bool bProxyValid = false;
+
+	/** 
+	 * @return current render proxy, if valid, otherwise nullptr 
+	 */
+	FDynamicMeshSceneProxy* GetCurrentSceneProxy();
+
+	/**
+	 * Fully invalidate all rendering data for this Component. Current Proxy will be discarded, Bounds and possibly Tangents recomputed, etc
+	 */
+	void ResetProxy();
+
+	//~ Begin UPrimitiveComponent Interface.
+	virtual FPrimitiveSceneProxy* CreateSceneProxy() override;
+
+	//~ USceneComponent Interface.
+	virtual FBoxSphereBounds CalcBounds(const FTransform& LocalToWorld) const override;
+	virtual void OnChildAttached(USceneComponent* ChildComponent) override;
+	virtual void OnChildDetached(USceneComponent* ChildComponent) override;
+
+	//~ UObject Interface.
+	virtual void PostLoad() override;
+
+
+
+
+	//===============================================================================================================
+	// Triangle-Vertex Tangents support. The default behavior is to not use Tangents, this will lead to incorrect
+	// rendering for any material with Normal Maps and some other shaders.
+	// If TangentsType == EDynamicMeshComponentTangentsMode::ExternallyProvided, the Tangent and Bitangent attributes of
+	//   the FDynamicMesh3 AttributeSet are used at the SceneProxy level, the Component is not involved
+	// If TangentsType == EDynamicMeshComponentTangentsMode::AutoCalculated, the Tangents are computed internally using
+	//   a fast MikkT approximation via FMeshTangentsf. They will be recomputed when the mesh is modified, however
+	//   they are *not* recomputed when using the Fast Update functions above (in that case InvalidateAutoCalculatedTangents() 
+	//   can be used to force recomputation)
+	//
+public:
+	UFUNCTION(BlueprintCallable, Category = "Dynamic Mesh Component")
+	void SetTangentsType(EDynamicMeshComponentTangentsMode NewTangentsType);
+
+	UFUNCTION(BlueprintCallable, Category = "Dynamic Mesh Component")
+	EDynamicMeshComponentTangentsMode GetTangentsType() const { return TangentsType; }
+
+	/** This function marks the auto tangents as dirty, they will be recomputed before they are used again */
+	virtual void InvalidateAutoCalculatedTangents();
+
+	/** @return AutoCalculated Tangent Set, which may require that they be recomputed, or nullptr if not enabled/available */
+	const UE::Geometry::FMeshTangentsf* GetAutoCalculatedTangents();
+
+protected:
+	/** How should Tangents be calculated/handled */
+	UPROPERTY()
+	EDynamicMeshComponentTangentsMode TangentsType = EDynamicMeshComponentTangentsMode::NoTangents;
+
+	/** true if AutoCalculatedTangents has been computed for current mesh */
+	bool bAutoCalculatedTangentsValid = false;
+
+	/** Set of per-triangle-vertex tangents computed for the current mesh. Only valid if bAutoCalculatedTangentsValid == true */
+	UE::Geometry::FMeshTangentsf AutoCalculatedTangents;
+
+	void UpdateAutoCalculatedTangents();
+
+};
diff --git a/Engine/Plugins/Experimental/MeshModelingToolset/Source/ModelingComponents/Public/Components/OctreeDynamicMeshComponent.h b/Engine/Plugins/Experimental/MeshModelingToolset/Source/ModelingComponents/Public/Components/OctreeDynamicMeshComponent.h
new file mode 100644
--- /dev/null
+++ b/Engine/Plugins/Experimental/MeshModelingToolset/Source/ModelingComponents/Public/Components/OctreeDynamicMeshComponent.h
@@ -0,0 +1,161 @@
+// Copyright Epic Games, Inc. All Rights Reserved.
+
+#pragma once
+
+#include "CoreMinimal.h"
+#include "BaseDynamicMeshComponent.h"
+#include "MeshConversionOptions.h"
+#include "UDynamicMesh.h"
+#include "DynamicMeshOctree3.h"
+#include "Util/IndexSetDecompositions.h"
+
+#include "OctreeDynamicMeshComponent.generated.h"
+
+// predecl
+struct FMeshDescription;
+
+/** internal FPrimitiveSceneProxy defined in OctreeDynamicMeshSceneProxy.h */
+class FOctreeDynamicMeshSceneProxy;
+
+
+/** 
+ * UOctreeDynamicMeshComponent is a mesh component similar to UProceduralMeshComponent,
+ * except it bases the renderable geometry off an internal FDynamicMesh3 instance.
+ * The class generally has the same capabilities as UDynamicMeshComponent.
+ * 
+ * A FDynamicMeshOctree3 is available to dynamically track the triangles of the mesh
+ * (however the client is responsible for updating this octree).
+ * Based on the Octree, the mesh is partitioned into chunks that are stored in separate
+ * RenderBuffers in the FOctreeDynamicMeshSceneProxy.
+ * Calling NotifyMeshUpdated() will result in only the "dirty" chunks being updated,
+ * rather than the entire mesh.
+ * 
+ * (So, if you don't need this capability, and don't want to update an Octree, use UDynamicMeshComponent!)
+ */
+UCLASS(hidecategories = (LOD, Physics, Collision), editinlinenew, ClassGroup = Rendering)
+class MODELINGCOMPONENTS_API UOctreeDynamicMeshComponent : public UBaseDynamicMeshComponent
+{
+	GENERATED_UCLASS_BODY()
+
+public:
+	/**
+	 * initialize the internal mesh from a DynamicMesh
+	 */
+	virtual void SetMesh(UE::Geometry::FDynamicMesh3&& MoveMesh) override;
+
+	/**
+	 * @return pointer to internal mesh
+	 */
+	virtual FDynamicMesh3* GetMesh() override { return MeshObject->GetMeshPtr(); }
+
+	/**
+	 * @return pointer to internal mesh
+	 */
+	virtual const FDynamicMesh3* GetMesh() const override { return MeshObject->GetMeshPtr(); }
+
+
+	virtual void ProcessMesh(TFunctionRef<void(const UE::Geometry::FDynamicMesh3&)> ProcessFunc) const
+	{
+		MeshObject->ProcessMesh(ProcessFunc);
+	}
+
+
+	UE::Geometry::FDynamicMeshOctree3* GetOctree() { return Octree.Get(); }
+
+
+
+	/**
+	 * Apply transform to internal mesh. Invalidates RenderProxy.
+	 * @param bInvert if true, inverse tranform is applied instead of forward transform
+	 */
+	virtual void ApplyTransform(const UE::Geometry::FTransform3d& Transform, bool bInvert) override;
+
+	//
+	// change tracking/etc
+	//
+
+	/**
+	 * Call this if you update the mesh via GetMesh()
+	 * @todo should provide a function that calls a lambda to modify the mesh, and only return const mesh pointer
+	 */
+	virtual void NotifyMeshUpdated() override;
+
+	/**
+	 * Apply a vertex deformation change to the internal mesh
+	 */
+	virtual void ApplyChange(const FMeshVertexChange* Change, bool bRevert) override;
+
+	/**
+	 * Apply a general mesh change to the internal mesh
+	 */
+	virtual void ApplyChange(const FMeshChange* Change, bool bRevert) override;
+
+	/**
+	* Apply a general mesh replacement change to the internal mesh
+	*/
+	virtual void ApplyChange(const FMeshReplacementChange* Change, bool bRevert) override;
+
+
+	/**
+	 * This delegate fires when a FCommandChange is applied to this component, so that
+	 * parent objects know the mesh has changed.
+	 */
+	FSimpleMulticastDelegate OnMeshChanged;
+
+	/**
+	 * If this function is set, we will use these colors instead of vertex colors
+	 */
+	TFunction<FColor(const FDynamicMesh3*, int)> TriangleColorFunc = nullptr;
+
+protected:
+	/**
+	 * This is called to tell our RenderProxy about modifications to the material set.
+	 * We need to pass this on for things like material validation in the Editor.
+	 */
+	virtual void NotifyMaterialSetUpdated();
+
+private:
+
+	FOctreeDynamicMeshSceneProxy* GetCurrentSceneProxy() { return (FOctreeDynamicMeshSceneProxy*)SceneProxy; }
+
+	//~ Begin UPrimitiveComponent Interface.
+	virtual FPrimitiveSceneProxy* CreateSceneProxy() override;
+	//~ End UPrimitiveComponent Interface.
+
+	//~ Begin USceneComponent Interface.
+	virtual FBoxSphereBounds CalcBounds(const FTransform& LocalToWorld) const override;
+	//~ Begin USceneComponent Interface.
+
+	UPROPERTY()
+	UDynamicMesh* MeshObject;
+
+	FDelegateHandle MeshObjectChangedHandle;
+	void OnMeshObjectChanged(UDynamicMesh* ChangedMeshObject, FDynamicMeshChangeInfo ChangeInfo);
+
+	FDelegateHandle PreMeshChangeHandle;
+	void OnPreMeshObjectChanged(UDynamicMesh* ChangedMeshObject, FDynamicMeshChangeInfo ChangeInfo);
+
+	TUniquePtr<UE::Geometry::FDynamicMeshOctree3> Octree;
+	TUniquePtr<UE::Geometry::FDynamicMeshOctree3::FTreeCutSet> OctreeCut;
+	UE::Geometry::FArrayIndexSetsDecomposition TriangleDecomposition;
+	struct FCutCellIndexSet
+	{
+		UE::Geometry::FDynamicMeshOctree3::FCellReference CellRef;
+		int32 DecompSetID;
+	};
+	TArray<FCutCellIndexSet> CutCellSetMap;
+	int32 SpillDecompSetID;
+
+
+	FColor GetTriangleColor(int TriangleID);
+
+	//friend class FCustomMeshSceneProxy;
+
+
+public:
+	//UFUNCTION(BlueprintCallable, Category = "DynamicMesh")
+	virtual UDynamicMesh* GetDynamicMesh() override { return MeshObject; }
+
+	UFUNCTION(BlueprintCallable, Category = "DynamicMesh")
+	void SetDynamicMesh(UDynamicMesh* NewMesh);
+};
diff --git a/Engine/Plugins/Experimental/MeshModelingToolset/Source/ModelingComponents/Public/OctreeDynamicMeshComponent.h b/Engine/Plugins/Experimental/MeshModelingToolset/Source/ModelingComponents/Public/OctreeDynamicMeshComponent.h
deleted file mode 100644
diff --git a/Engine/Plugins/Experimental/MeshModelingToolset/Source/ModelingComponents/Public/PreviewMesh.h b/Engine/Plugins/Experimental/MeshModelingToolset/Source/ModelingComponents/Public/PreviewMesh.h
--- a/Engine/Plugins/Experimental/MeshModelingToolset/Source/ModelingComponents/Public/PreviewMesh.h
+++ b/Engine/Plugins/Experimental/MeshModelingToolset/Source/ModelingComponents/Public/PreviewMesh.h
@@ -3,7 +3,7 @@
 #pragma once
 
 #include "CoreMinimal.h"
-#include "SimpleDynamicMeshComponent.h"
+#include "Components/DynamicMeshComponent.h"
 #include "DynamicMeshAABBTree3.h"
 #include "InteractiveToolObjects.h"
 #include "Changes/MeshVertexChange.h"
@@ -39,7 +39,7 @@ public:
  * This can be used to show live preview geometry during modeling operations.
  * Call CreateInWorld() to set it up, and Disconnect() to shut it down.
  * 
- * Currently implemented via an internal Actor that has a USimpleDynamicMeshComponent root component,
+ * Currently implemented via an internal Actor that has a UDynamicMeshComponent root component,
  * with an AABBTree created/updated if FProperty bBuildSpatialDataStructure=true.
  * The Actor is destroyed on Disconnect().
  * 
@@ -422,7 +422,7 @@ protected:
 
 	/** This component is set as the root component of TemporaryParentActor */
 	UPROPERTY()
-	USimpleDynamicMeshComponent* DynamicMeshComponent = nullptr;
+	UDynamicMeshComponent* DynamicMeshComponent = nullptr;
 
 	/** Spatial data structure that is initialized if bBuildSpatialDataStructure = true when UpdatePreview() is called */
 	UE::Geometry::FDynamicMeshAABBTree3 MeshAABBTree;
diff --git a/Engine/Plugins/Experimental/MeshModelingToolset/Source/ModelingComponents/Public/Selection/PolygonSelectionMechanic.h b/Engine/Plugins/Experimental/MeshModelingToolset/Source/ModelingComponents/Public/Selection/PolygonSelectionMechanic.h
--- a/Engine/Plugins/Experimental/MeshModelingToolset/Source/ModelingComponents/Public/Selection/PolygonSelectionMechanic.h
+++ b/Engine/Plugins/Experimental/MeshModelingToolset/Source/ModelingComponents/Public/Selection/PolygonSelectionMechanic.h
@@ -8,7 +8,7 @@
 #include "InputBehavior.h"
 #include "InteractiveTool.h"
 #include "Mechanics/RectangleMarqueeMechanic.h"
-#include "SimpleDynamicMeshComponent.h"
+#include "Components/DynamicMeshComponent.h"
 #include "Selection/GroupTopologySelector.h"
 #include "TransformTypes.h"
 #include "ToolDataVisualizer.h"
@@ -75,7 +75,7 @@ public:
 
 /**
  * UPolygonSelectionMechanic implements the interaction for selecting a set of faces/vertices/edges
- * from a FGroupTopology on a USimpleDynamicMeshComponent. 
+ * from a FGroupTopology on a UDynamicMeshComponent. 
  */
 UCLASS()
 class MODELINGCOMPONENTS_API UPolygonSelectionMechanic : public UInteractionMechanic, public IClickBehaviorTarget, public IHoverBehaviorTarget
@@ -113,7 +113,7 @@ public:
 		TFunction<FDynamicMeshAABBTree3*()> GetSpatialSourceFunc
 		);
 
-	void Initialize(USimpleDynamicMeshComponent* MeshComponent, const FGroupTopology* Topology,
+	void Initialize(UDynamicMeshComponent* MeshComponent, const FGroupTopology* Topology,
 		TFunction<FDynamicMeshAABBTree3 * ()> GetSpatialSourceFunc
 	);
 
diff --git a/Engine/Plugins/Experimental/MeshModelingToolset/Source/ModelingComponents/Public/SimpleDynamicMeshComponent.h b/Engine/Plugins/Experimental/MeshModelingToolset/Source/ModelingComponents/Public/SimpleDynamicMeshComponent.h
deleted file mode 100644
