ModelingComponents: make FMeshDescription returned by IMeshDescriptionProvider const, this pointer should not be modified
#rnx

diff --git a/Engine/Plugins/Experimental/MeshModelingToolset/Source/MeshModelingTools/Private/EditPivotTool.cpp b/Engine/Plugins/Experimental/MeshModelingToolset/Source/MeshModelingTools/Private/EditPivotTool.cpp
--- a/Engine/Plugins/Experimental/MeshModelingToolset/Source/MeshModelingTools/Private/EditPivotTool.cpp
+++ b/Engine/Plugins/Experimental/MeshModelingToolset/Source/MeshModelingTools/Private/EditPivotTool.cpp
@@ -175,7 +175,7 @@ void UEditPivotTool::Precompute()
 	{
 		Transform = UE::Geometry::FTransform3d(TargetComponentInterface(0)->GetWorldTransform());
 
-		FMeshDescription* Mesh = TargetMeshProviderInterface(0)->GetMeshDescription();
+		const FMeshDescription* Mesh = TargetMeshProviderInterface(0)->GetMeshDescription();
 		VertexIteration(Mesh, [&](int32 VertexID, const FVector& Position) {
 			ObjectBounds.Contain((FVector3d)Position);
 			WorldBounds.Contain(Transform.TransformPosition((FVector3d)Position));
@@ -189,7 +189,7 @@ void UEditPivotTool::Precompute()
 			IPrimitiveComponentBackedTarget* TargetComponent = TargetComponentInterface(k);
 			IMeshDescriptionProvider* TargetMeshProvider = TargetMeshProviderInterface(k);
 			UE::Geometry::FTransform3d CurTransform(TargetComponent->GetWorldTransform());
-			FMeshDescription* Mesh = TargetMeshProvider->GetMeshDescription();
+			const FMeshDescription* Mesh = TargetMeshProvider->GetMeshDescription();
 			VertexIteration(Mesh, [&](int32 VertexID, const FVector& Position) {
 				ObjectBounds.Contain(CurTransform.TransformPosition((FVector3d)Position));
 				WorldBounds.Contain(CurTransform.TransformPosition((FVector3d)Position));
diff --git a/Engine/Plugins/Experimental/MeshModelingToolset/Source/MeshModelingToolsEditorOnly/Private/AttributeEditorTool.cpp b/Engine/Plugins/Experimental/MeshModelingToolset/Source/MeshModelingToolsEditorOnly/Private/AttributeEditorTool.cpp
--- a/Engine/Plugins/Experimental/MeshModelingToolset/Source/MeshModelingToolsEditorOnly/Private/AttributeEditorTool.cpp
+++ b/Engine/Plugins/Experimental/MeshModelingToolset/Source/MeshModelingToolsEditorOnly/Private/AttributeEditorTool.cpp
@@ -153,7 +153,7 @@ void UAttributeEditorTool::Setup()
 
 
 template<typename AttribSetType>
-void ExtractAttribList(FMeshDescription* Mesh, AttribSetType& AttribSet, EAttributeEditorElementType ElemType, TArray<FAttributeEditorAttribInfo>& AttribList, TArray<FString>& StringList)
+void ExtractAttribList(const FMeshDescription* Mesh, AttribSetType& AttribSet, EAttributeEditorElementType ElemType, TArray<FAttributeEditorAttribInfo>& AttribList, TArray<FString>& StringList)
 {
 	AttribList.Reset();
 	StringList.Reset();
@@ -209,6 +209,26 @@ void ExtractAttribList(FMeshDescription* Mesh, AttribSetType& AttribSet, EAttrib
 
 
 
+static const FAttributesSetBase* GetAttributeSetByType(const FMeshDescription* Mesh, EAttributeEditorElementType ElemType)
+{
+	switch (ElemType)
+	{
+	case EAttributeEditorElementType::Vertex:
+		return &Mesh->VertexAttributes();
+	case EAttributeEditorElementType::VertexInstance:
+		return &Mesh->VertexInstanceAttributes();
+	case EAttributeEditorElementType::Triangle:
+		return &Mesh->TriangleAttributes();
+	case EAttributeEditorElementType::Polygon:
+		return &Mesh->PolygonAttributes();
+	case EAttributeEditorElementType::Edge:
+		return &Mesh->EdgeAttributes();
+	case EAttributeEditorElementType::PolygonGroup:
+		return &Mesh->PolygonGroupAttributes();
+	}
+	check(false);
+	return nullptr;
+}
 static FAttributesSetBase* GetAttributeSetByType(FMeshDescription* Mesh, EAttributeEditorElementType ElemType)
 {
 	switch (ElemType)
@@ -230,9 +250,10 @@ static FAttributesSetBase* GetAttributeSetByType(FMeshDescription* Mesh, EAttrib
 	return nullptr;
 }
 
-static bool HasAttribute(FMeshDescription* Mesh, EAttributeEditorElementType ElemType, FName AttributeName)
+
+static bool HasAttribute(const FMeshDescription* Mesh, EAttributeEditorElementType ElemType, FName AttributeName)
 {
-	FAttributesSetBase* AttribSetBase = GetAttributeSetByType(Mesh, ElemType);
+	const FAttributesSetBase* AttribSetBase = GetAttributeSetByType(Mesh, ElemType);
 	return (AttribSetBase) ? AttribSetBase->HasAttribute(AttributeName) : false;
 }
 
@@ -283,10 +304,9 @@ static bool RemoveAttribute(FMeshDescription* Mesh, EAttributeEditorElementType
 
 void UAttributeEditorTool::InitializeAttributeLists()
 {
-	FMeshDescription* Mesh = TargetMeshProviderInterface(0)->GetMeshDescription();
-
+	const FMeshDescription* Mesh = TargetMeshProviderInterface(0)->GetMeshDescription();
 
-	TVertexInstanceAttributesRef<FVector2D> InstanceUVs =
+	TVertexInstanceAttributesConstRef<FVector2D> InstanceUVs =
 		Mesh->VertexInstanceAttributes().GetAttributesRef<FVector2D>(MeshAttribute::VertexInstance::TextureCoordinate);
 
 	UVActions->UVLayerNamesList.Reset();
@@ -662,7 +682,7 @@ void UAttributeEditorTool::AddNewAttribute(EAttributeEditorElementType ElemType,
 		return;
 	}
 
-	FMeshDescription* CurMesh = TargetMeshProviderInterface(0)->GetMeshDescription();
+	const FMeshDescription* CurMesh = TargetMeshProviderInterface(0)->GetMeshDescription();
 	if (HasAttribute(CurMesh, ElemType, AttributeName))
 	{
 		GetToolManager()->DisplayMessage(LOCTEXT("ErrorAddingDuplicateNameMessage", "Attribute with this name already exists"), EToolMessageLevel::UserWarning);
@@ -715,7 +735,7 @@ void UAttributeEditorTool::ClearAttribute()
 
 void UAttributeEditorTool::DeleteAttribute()
 {
-	FMeshDescription* CurMesh = TargetMeshProviderInterface(0)->GetMeshDescription();
+	const FMeshDescription* CurMesh = TargetMeshProviderInterface(0)->GetMeshDescription();
 	FName SelectedName(ModifyAttributeProps->Attribute);
 
 	// We check on the skeletal mesh attributes because it is a superset of the static mesh
@@ -791,7 +811,7 @@ void UAttributeEditorTool::ResetLightmapUVsChannels()
 	GetToolManager()->BeginUndoTransaction(LOCTEXT("ResetLightmapUVs", "Reset Lightmap UVs"));
 	for (int32 ComponentIdx = 0; ComponentIdx < Targets.Num(); ComponentIdx++)
 	{
-		TVertexInstanceAttributesRef<FVector2D> InstanceUVs = 
+		TVertexInstanceAttributesConstRef<FVector2D> InstanceUVs = 
 			TargetMeshProviderInterface(ComponentIdx)->GetMeshDescription()->VertexInstanceAttributes().GetAttributesRef<FVector2D>(MeshAttribute::VertexInstance::TextureCoordinate);
 		int32 SetChannel = FMath::Max(InstanceUVs.GetNumChannels(), 1);
 
diff --git a/Engine/Plugins/Experimental/MeshModelingToolset/Source/ModelingComponentsEditorOnly/Private/ToolTargets/SkeletalMeshComponentToolTarget.cpp b/Engine/Plugins/Experimental/MeshModelingToolset/Source/ModelingComponentsEditorOnly/Private/ToolTargets/SkeletalMeshComponentToolTarget.cpp
--- a/Engine/Plugins/Experimental/MeshModelingToolset/Source/ModelingComponentsEditorOnly/Private/ToolTargets/SkeletalMeshComponentToolTarget.cpp
+++ b/Engine/Plugins/Experimental/MeshModelingToolset/Source/ModelingComponentsEditorOnly/Private/ToolTargets/SkeletalMeshComponentToolTarget.cpp
@@ -83,7 +83,7 @@ bool USkeletalMeshComponentToolTarget::CommitMaterialSetUpdate(const FComponentM
 	return true;
 }
 
-FMeshDescription* USkeletalMeshComponentToolTarget::GetMeshDescription()
+const FMeshDescription* USkeletalMeshComponentToolTarget::GetMeshDescription()
 {
 	if (!ensure(IsValid()))
 	{
@@ -105,11 +105,16 @@ void USkeletalMeshComponentToolTarget::CommitMeshDescription(const FCommitter& C
 	if (ensure(IsValid()) == false) return;
 
 	USkeletalMesh* SkeletalMesh = Cast<USkinnedMeshComponent>(Component)->SkeletalMesh;
+	if (!CachedMeshDescription.IsValid())
+	{
+		CachedMeshDescription = MakeUnique<FMeshDescription>();
+		USkeletalMeshToolTarget::GetMeshDescription(SkeletalMesh, *CachedMeshDescription);
+	}
 
 	// unregister the component while we update its skeletal mesh
 	FComponentReregisterContext ComponentReregisterContext(Component);
 
-	USkeletalMeshToolTarget::CommitMeshDescription(SkeletalMesh, GetMeshDescription(), Committer);
+	USkeletalMeshToolTarget::CommitMeshDescription(SkeletalMesh, CachedMeshDescription.Get(), Committer);
 
 	// this rebuilds physics, but it doesn't undo!
 	Component->RecreatePhysicsState();
diff --git a/Engine/Plugins/Experimental/MeshModelingToolset/Source/ModelingComponentsEditorOnly/Private/ToolTargets/SkeletalMeshToolTarget.cpp b/Engine/Plugins/Experimental/MeshModelingToolset/Source/ModelingComponentsEditorOnly/Private/ToolTargets/SkeletalMeshToolTarget.cpp
--- a/Engine/Plugins/Experimental/MeshModelingToolset/Source/ModelingComponentsEditorOnly/Private/ToolTargets/SkeletalMeshToolTarget.cpp
+++ b/Engine/Plugins/Experimental/MeshModelingToolset/Source/ModelingComponentsEditorOnly/Private/ToolTargets/SkeletalMeshToolTarget.cpp
@@ -97,7 +97,7 @@ bool USkeletalMeshToolTarget::CommitMaterialSetUpdate(USkeletalMesh* SkeletalMes
 	return true;
 }
 
-FMeshDescription* USkeletalMeshToolTarget::GetMeshDescription()
+const FMeshDescription* USkeletalMeshToolTarget::GetMeshDescription()
 {
 	if (!ensure(IsValid()))
 	{
@@ -140,7 +140,13 @@ void USkeletalMeshToolTarget::GetMeshDescription(const USkeletalMesh* SkeletalMe
 void USkeletalMeshToolTarget::CommitMeshDescription(const FCommitter& Committer)
 {
 	if (ensure(IsValid()) == false) return;
-	CommitMeshDescription(SkeletalMesh, GetMeshDescription(), Committer);
+
+	if (!CachedMeshDescription.IsValid())
+	{
+		CachedMeshDescription = MakeUnique<FMeshDescription>();
+		GetMeshDescription(SkeletalMesh, *CachedMeshDescription);
+	}
+	CommitMeshDescription(SkeletalMesh, CachedMeshDescription.Get(), Committer);
 }
 
 void USkeletalMeshToolTarget::CommitMeshDescription(USkeletalMesh* SkeletalMeshIn,
diff --git a/Engine/Plugins/Experimental/MeshModelingToolset/Source/ModelingComponentsEditorOnly/Private/ToolTargets/StaticMeshComponentToolTarget.cpp b/Engine/Plugins/Experimental/MeshModelingToolset/Source/ModelingComponentsEditorOnly/Private/ToolTargets/StaticMeshComponentToolTarget.cpp
--- a/Engine/Plugins/Experimental/MeshModelingToolset/Source/ModelingComponentsEditorOnly/Private/ToolTargets/StaticMeshComponentToolTarget.cpp
+++ b/Engine/Plugins/Experimental/MeshModelingToolset/Source/ModelingComponentsEditorOnly/Private/ToolTargets/StaticMeshComponentToolTarget.cpp
@@ -124,7 +124,7 @@ bool UStaticMeshComponentToolTarget::CommitMaterialSetUpdate(const FComponentMat
 }
 
 
-FMeshDescription* UStaticMeshComponentToolTarget::GetMeshDescription()
+const FMeshDescription* UStaticMeshComponentToolTarget::GetMeshDescription()
 {
 	if (ensure(IsValid()))
 	{
@@ -159,11 +159,13 @@ void UStaticMeshComponentToolTarget::CommitMeshDescription(const FCommitter& Com
 	if (ensure(IsValid()) == false) return;
 
 	UStaticMesh* StaticMesh = Cast<UStaticMeshComponent>(Component)->GetStaticMesh();
+	FMeshDescription* UpdateMeshDescription = (EditingLOD == EStaticMeshEditingLOD::HiResSource) ?
+		StaticMesh->GetHiResMeshDescription() : StaticMesh->GetMeshDescription((int32)EditingLOD);
 
 	// unregister the component while we update its static mesh
 	FComponentReregisterContext ComponentReregisterContext(Component);
 
-	UStaticMeshToolTarget::CommitMeshDescription(StaticMesh, GetMeshDescription(), Committer, EditingLOD);
+	UStaticMeshToolTarget::CommitMeshDescription(StaticMesh, UpdateMeshDescription, Committer, EditingLOD);
 
 	// this rebuilds physics, but it doesn't undo!
 	Component->RecreatePhysicsState();
diff --git a/Engine/Plugins/Experimental/MeshModelingToolset/Source/ModelingComponentsEditorOnly/Private/ToolTargets/StaticMeshToolTarget.cpp b/Engine/Plugins/Experimental/MeshModelingToolset/Source/ModelingComponentsEditorOnly/Private/ToolTargets/StaticMeshToolTarget.cpp
--- a/Engine/Plugins/Experimental/MeshModelingToolset/Source/ModelingComponentsEditorOnly/Private/ToolTargets/StaticMeshToolTarget.cpp
+++ b/Engine/Plugins/Experimental/MeshModelingToolset/Source/ModelingComponentsEditorOnly/Private/ToolTargets/StaticMeshToolTarget.cpp
@@ -179,7 +179,7 @@ bool UStaticMeshToolTarget::CommitMaterialSetUpdate(UStaticMesh* StaticMeshIn,
 	return true;
 }
 
-FMeshDescription* UStaticMeshToolTarget::GetMeshDescription()
+const FMeshDescription* UStaticMeshToolTarget::GetMeshDescription()
 {
 	if (ensure(IsValid()))
 	{
@@ -193,7 +193,10 @@ void UStaticMeshToolTarget::CommitMeshDescription(const FCommitter& Committer)
 {
 	if (ensure(IsValid()) == false) return;
 
-	CommitMeshDescription(StaticMesh, GetMeshDescription(), Committer, EditingLOD);
+	FMeshDescription* UpdateMeshDescription = (EditingLOD == EStaticMeshEditingLOD::HiResSource) ?
+		StaticMesh->GetHiResMeshDescription() : StaticMesh->GetMeshDescription((int32)EditingLOD);
+
+	CommitMeshDescription(StaticMesh, UpdateMeshDescription, Committer, EditingLOD);
 }
 
 void UStaticMeshToolTarget::CommitMeshDescription(UStaticMesh* StaticMeshIn, FMeshDescription* MeshDescription,
diff --git a/Engine/Plugins/Experimental/MeshModelingToolset/Source/ModelingComponentsEditorOnly/Private/ToolTargets/VolumeMeshDescriptionToolTarget.cpp b/Engine/Plugins/Experimental/MeshModelingToolset/Source/ModelingComponentsEditorOnly/Private/ToolTargets/VolumeMeshDescriptionToolTarget.cpp
--- a/Engine/Plugins/Experimental/MeshModelingToolset/Source/ModelingComponentsEditorOnly/Private/ToolTargets/VolumeMeshDescriptionToolTarget.cpp
+++ b/Engine/Plugins/Experimental/MeshModelingToolset/Source/ModelingComponentsEditorOnly/Private/ToolTargets/VolumeMeshDescriptionToolTarget.cpp
@@ -18,7 +18,7 @@
 #include "ExplicitUseGeometryMathTypes.h"		// using UE::Geometry::(math types)
 using namespace UE::Geometry;
 
-FMeshDescription* UVolumeMeshDescriptionToolTarget::GetMeshDescription()
+const FMeshDescription* UVolumeMeshDescriptionToolTarget::GetMeshDescription()
 {
 	if (!ConvertedMeshDescription.IsValid())
 	{
diff --git a/Engine/Plugins/Experimental/MeshModelingToolset/Source/ModelingComponentsEditorOnly/Public/ToolTargets/SkeletalMeshComponentToolTarget.h b/Engine/Plugins/Experimental/MeshModelingToolset/Source/ModelingComponentsEditorOnly/Public/ToolTargets/SkeletalMeshComponentToolTarget.h
--- a/Engine/Plugins/Experimental/MeshModelingToolset/Source/ModelingComponentsEditorOnly/Public/ToolTargets/SkeletalMeshComponentToolTarget.h
+++ b/Engine/Plugins/Experimental/MeshModelingToolset/Source/ModelingComponentsEditorOnly/Public/ToolTargets/SkeletalMeshComponentToolTarget.h
@@ -35,7 +35,7 @@ class MODELINGCOMPONENTSEDITORONLY_API USkeletalMeshComponentToolTarget :
 
 public:
 	// IMeshDescriptionProvider implementation
-	FMeshDescription* GetMeshDescription() override;
+	const FMeshDescription* GetMeshDescription() override;
 
 	// IMeshDescritpionCommitter implementation
 	void CommitMeshDescription(const FCommitter& Committer) override;
diff --git a/Engine/Plugins/Experimental/MeshModelingToolset/Source/ModelingComponentsEditorOnly/Public/ToolTargets/SkeletalMeshToolTarget.h b/Engine/Plugins/Experimental/MeshModelingToolset/Source/ModelingComponentsEditorOnly/Public/ToolTargets/SkeletalMeshToolTarget.h
--- a/Engine/Plugins/Experimental/MeshModelingToolset/Source/ModelingComponentsEditorOnly/Public/ToolTargets/SkeletalMeshToolTarget.h
+++ b/Engine/Plugins/Experimental/MeshModelingToolset/Source/ModelingComponentsEditorOnly/Public/ToolTargets/SkeletalMeshToolTarget.h
@@ -37,7 +37,7 @@ public:
 	virtual bool IsValid() const override;
 
 	// IMeshDescriptionProvider implementation
-	FMeshDescription* GetMeshDescription() override;
+	const FMeshDescription* GetMeshDescription() override;
 
 	// IMeshDescritpionCommitter implementation
 	void CommitMeshDescription(const FCommitter& Committer) override;
diff --git a/Engine/Plugins/Experimental/MeshModelingToolset/Source/ModelingComponentsEditorOnly/Public/ToolTargets/StaticMeshComponentToolTarget.h b/Engine/Plugins/Experimental/MeshModelingToolset/Source/ModelingComponentsEditorOnly/Public/ToolTargets/StaticMeshComponentToolTarget.h
--- a/Engine/Plugins/Experimental/MeshModelingToolset/Source/ModelingComponentsEditorOnly/Public/ToolTargets/StaticMeshComponentToolTarget.h
+++ b/Engine/Plugins/Experimental/MeshModelingToolset/Source/ModelingComponentsEditorOnly/Public/ToolTargets/StaticMeshComponentToolTarget.h
@@ -44,7 +44,7 @@ public:
 
 public:
 	// IMeshDescriptionProvider implementation
-	FMeshDescription* GetMeshDescription() override;
+	const FMeshDescription* GetMeshDescription() override;
 	virtual void CalculateAutoGeneratedAttributes(FMeshDescription& MeshDescription) override;
 
 	// IMeshDescritpionCommitter implementation
diff --git a/Engine/Plugins/Experimental/MeshModelingToolset/Source/ModelingComponentsEditorOnly/Public/ToolTargets/StaticMeshToolTarget.h b/Engine/Plugins/Experimental/MeshModelingToolset/Source/ModelingComponentsEditorOnly/Public/ToolTargets/StaticMeshToolTarget.h
--- a/Engine/Plugins/Experimental/MeshModelingToolset/Source/ModelingComponentsEditorOnly/Public/ToolTargets/StaticMeshToolTarget.h
+++ b/Engine/Plugins/Experimental/MeshModelingToolset/Source/ModelingComponentsEditorOnly/Public/ToolTargets/StaticMeshToolTarget.h
@@ -47,7 +47,7 @@ public:
 	virtual bool IsValid() const override;
 
 	// IMeshDescriptionProvider implementation
-	FMeshDescription* GetMeshDescription() override;
+	const FMeshDescription* GetMeshDescription() override;
 
 	// IMeshDescritpionCommitter implementation
 	virtual void CommitMeshDescription(const FCommitter& Committer) override;
diff --git a/Engine/Plugins/Experimental/MeshModelingToolset/Source/ModelingComponentsEditorOnly/Public/ToolTargets/VolumeMeshDescriptionToolTarget.h b/Engine/Plugins/Experimental/MeshModelingToolset/Source/ModelingComponentsEditorOnly/Public/ToolTargets/VolumeMeshDescriptionToolTarget.h
--- a/Engine/Plugins/Experimental/MeshModelingToolset/Source/ModelingComponentsEditorOnly/Public/ToolTargets/VolumeMeshDescriptionToolTarget.h
+++ b/Engine/Plugins/Experimental/MeshModelingToolset/Source/ModelingComponentsEditorOnly/Public/ToolTargets/VolumeMeshDescriptionToolTarget.h
@@ -22,7 +22,7 @@ class MODELINGCOMPONENTSEDITORONLY_API UVolumeMeshDescriptionToolTarget : public
 public:
 
 	// IMeshDescriptionProvider implementation
-	FMeshDescription* GetMeshDescription() override;
+	const FMeshDescription* GetMeshDescription() override;
 
 	// IMeshDescritpionCommitter implementation
 	virtual void CommitMeshDescription(const FCommitter& Committer) override;
diff --git a/Engine/Plugins/Runtime/ProceduralMeshComponent/Source/ProceduralMeshComponent/Private/ProceduralMeshComponentToolTarget.cpp b/Engine/Plugins/Runtime/ProceduralMeshComponent/Source/ProceduralMeshComponent/Private/ProceduralMeshComponentToolTarget.cpp
--- a/Engine/Plugins/Runtime/ProceduralMeshComponent/Source/ProceduralMeshComponent/Private/ProceduralMeshComponentToolTarget.cpp
+++ b/Engine/Plugins/Runtime/ProceduralMeshComponent/Source/ProceduralMeshComponent/Private/ProceduralMeshComponentToolTarget.cpp
@@ -6,7 +6,7 @@
 #include "ProceduralMeshComponent.h"
 #include "ProceduralMeshConversion.h"
 
-FMeshDescription* UProceduralMeshComponentToolTarget::GetMeshDescription()
+const FMeshDescription* UProceduralMeshComponentToolTarget::GetMeshDescription()
 {
 	if (!MeshDescription.IsValid())
 	{
diff --git a/Engine/Plugins/Runtime/ProceduralMeshComponent/Source/ProceduralMeshComponent/Public/ProceduralMeshComponentToolTarget.h b/Engine/Plugins/Runtime/ProceduralMeshComponent/Source/ProceduralMeshComponent/Public/ProceduralMeshComponentToolTarget.h
--- a/Engine/Plugins/Runtime/ProceduralMeshComponent/Source/ProceduralMeshComponent/Public/ProceduralMeshComponentToolTarget.h
+++ b/Engine/Plugins/Runtime/ProceduralMeshComponent/Source/ProceduralMeshComponent/Public/ProceduralMeshComponentToolTarget.h
@@ -16,7 +16,7 @@ class PROCEDURALMESHCOMPONENT_API UProceduralMeshComponentToolTarget : public UP
 	GENERATED_BODY()
 
 public:
-	virtual FMeshDescription* GetMeshDescription() override;
+	virtual const FMeshDescription* GetMeshDescription() override;
 	virtual void CommitMeshDescription(const FCommitter&) override;
 private:
 	TSharedPtr<FMeshDescription> MeshDescription;
diff --git a/Engine/Source/Runtime/Experimental/InteractiveToolsFramework/Public/TargetInterfaces/MeshDescriptionProvider.h b/Engine/Source/Runtime/Experimental/InteractiveToolsFramework/Public/TargetInterfaces/MeshDescriptionProvider.h
--- a/Engine/Source/Runtime/Experimental/InteractiveToolsFramework/Public/TargetInterfaces/MeshDescriptionProvider.h
+++ b/Engine/Source/Runtime/Experimental/InteractiveToolsFramework/Public/TargetInterfaces/MeshDescriptionProvider.h
@@ -24,16 +24,14 @@ public:
 	 * Access the MeshDescriptionavailable through this Provider. Note that this MeshDescription may or may not 
 	 * be owned by the provider and should not be modified directly. Use IMeshDescriptionCommitter for writes.
 	 * @return pointer to MeshDescription 
-	 * @todo can this be const?
 	 */
-	virtual FMeshDescription* GetMeshDescription() = 0;
+	virtual const FMeshDescription* GetMeshDescription() = 0;
 
 	/**
 	 * Compute any auto-generated attributes on the input MeshDescription. This is required to support
 	 * auto-computed Normals and Tangents on Static Mesh Assets, as those normals/tangents are computed 
 	 * on-the-fly when generating rendering data, and will be zero in the MeshDescription returned by
 	 * GetMeshDescription(). A no-op default implementation is provided.
-	 * @warning Do not pass the result from GetMeshDescription() here. Must be done on a copy.
 	 */
 	virtual void CalculateAutoGeneratedAttributes(FMeshDescription& MeshDescription)
 	{
