ModelingComponents: handle case in FMeshRenderDecomposition where input mesh has no materials
#rnx

diff --git a/Engine/Plugins/Experimental/MeshModelingToolset/Source/ModelingComponents/Private/Drawing/MeshRenderDecomposition.cpp b/Engine/Plugins/Experimental/MeshModelingToolset/Source/ModelingComponents/Private/Drawing/MeshRenderDecomposition.cpp
--- a/Engine/Plugins/Experimental/MeshModelingToolset/Source/ModelingComponents/Private/Drawing/MeshRenderDecomposition.cpp
+++ b/Engine/Plugins/Experimental/MeshModelingToolset/Source/ModelingComponents/Private/Drawing/MeshRenderDecomposition.cpp
@@ -1,9 +1,9 @@
 // Copyright Epic Games, Inc. All Rights Reserved. 
 
 #include "Drawing/MeshRenderDecomposition.h"
-#include "DynamicMesh3.h"
-#include "DynamicMeshAttributeSet.h"
-#include "DynamicMeshAABBTree3.h"
+#include "DynamicMesh/DynamicMesh3.h"
+#include "DynamicMesh/DynamicMeshAttributeSet.h"
+#include "DynamicMesh/DynamicMeshAABBTree3.h"
 #include "Async/ParallelFor.h"
 #include "ComponentSourceInterfaces.h"
 
@@ -31,9 +31,13 @@ void FMeshRenderDecomposition::BuildAssociations(const FDynamicMesh3* Mesh)
 
 void FMeshRenderDecomposition::BuildMaterialDecomposition(const FDynamicMesh3* Mesh, const FComponentMaterialSet* MaterialSet, FMeshRenderDecomposition& Decomp)
 {
-	check(MaterialSet);
+	// always have at least one material
+	FComponentMaterialSet LocalMaterials;
+	LocalMaterials.Materials.Add(nullptr);
+	const FComponentMaterialSet* UseMaterials = (ensure(MaterialSet != nullptr)) ? MaterialSet : &LocalMaterials;
+
 	const FDynamicMeshMaterialAttribute* MaterialID = Mesh->Attributes()->GetMaterialID();
-	check(MaterialID);
+	ensure(MaterialID != nullptr);
 
 	int32 NumMaterials = MaterialSet->Materials.Num();
 	Decomp.Initialize(NumMaterials);
@@ -45,8 +49,11 @@ void FMeshRenderDecomposition::BuildMaterialDecomposition(const FDynamicMesh3* M
 
 	for (int32 tid : Mesh->TriangleIndicesItr())
 	{
-		int32 MatIdx;
-		MaterialID->GetValue(tid, &MatIdx);
+		int32 MatIdx = 0;
+		if (MaterialID)
+		{
+			MaterialID->GetValue(tid, &MatIdx);
+		}
 		if (MatIdx < NumMaterials)
 		{
 			FMeshRenderDecomposition::FGroup& Group = Decomp.GetGroup(MatIdx);
