GeometryProcessing: remove forwarding headers used in GeometryCore transition, and update all affected includes.
#rnx

diff --git a/Engine/Plugins/Experimental/Enterprise/DataprepGeometryOperations/Source/Private/DataprepGeometryOperations.cpp b/Engine/Plugins/Experimental/Enterprise/DataprepGeometryOperations/Source/Private/DataprepGeometryOperations.cpp
--- a/Engine/Plugins/Experimental/Enterprise/DataprepGeometryOperations/Source/Private/DataprepGeometryOperations.cpp
+++ b/Engine/Plugins/Experimental/Enterprise/DataprepGeometryOperations/Source/Private/DataprepGeometryOperations.cpp
@@ -4,7 +4,7 @@
 
 #include "DataprepOperationsLibraryUtil.h"
 #include "DataprepAssetUserData.h"
-#include "DynamicMeshAABBTree3.h"
+#include "DynamicMesh/DynamicMeshAABBTree3.h"
 #include "DynamicMeshToMeshDescription.h"
 #include "DynamicMeshEditor.h"
 #include "Components/StaticMeshComponent.h"
diff --git a/Engine/Plugins/Experimental/Enterprise/DataprepGeometryOperations/Source/Private/DataprepGeometryOperations.h b/Engine/Plugins/Experimental/Enterprise/DataprepGeometryOperations/Source/Private/DataprepGeometryOperations.h
--- a/Engine/Plugins/Experimental/Enterprise/DataprepGeometryOperations/Source/Private/DataprepGeometryOperations.h
+++ b/Engine/Plugins/Experimental/Enterprise/DataprepGeometryOperations/Source/Private/DataprepGeometryOperations.h
@@ -9,7 +9,7 @@
 #include "Properties/RemeshProperties.h"
 #include "CleaningOps/RemeshMeshOp.h"
 #include "CleaningOps/SimplifyMeshOp.h"
-#include "DynamicMesh3.h"
+#include "DynamicMesh/DynamicMesh3.h"
 #include "ModelingOperators.h"
 #include "BakeTransformTool.h"
 
diff --git a/Engine/Plugins/Experimental/GeometryFlow/Source/GeometryFlowMeshProcessing/Private/MeshProcessingNodes/MeshNormalFlowNode.cpp b/Engine/Plugins/Experimental/GeometryFlow/Source/GeometryFlowMeshProcessing/Private/MeshProcessingNodes/MeshNormalFlowNode.cpp
--- a/Engine/Plugins/Experimental/GeometryFlow/Source/GeometryFlowMeshProcessing/Private/MeshProcessingNodes/MeshNormalFlowNode.cpp
+++ b/Engine/Plugins/Experimental/GeometryFlow/Source/GeometryFlowMeshProcessing/Private/MeshProcessingNodes/MeshNormalFlowNode.cpp
@@ -3,7 +3,7 @@
 #include "MeshProcessingNodes/MeshNormalFlowNode.h"
 #include "NormalFlowRemesher.h"
 #include "ProjectionTargets.h"
-#include "MeshNormals.h"
+#include "DynamicMesh/MeshNormals.h"
 
 using namespace UE::Geometry;
 using namespace UE::GeometryFlow;
diff --git a/Engine/Plugins/Experimental/GeometryFlow/Source/GeometryFlowMeshProcessing/Private/MeshProcessingNodes/MeshThickenNode.cpp b/Engine/Plugins/Experimental/GeometryFlow/Source/GeometryFlowMeshProcessing/Private/MeshProcessingNodes/MeshThickenNode.cpp
--- a/Engine/Plugins/Experimental/GeometryFlow/Source/GeometryFlowMeshProcessing/Private/MeshProcessingNodes/MeshThickenNode.cpp
+++ b/Engine/Plugins/Experimental/GeometryFlow/Source/GeometryFlowMeshProcessing/Private/MeshProcessingNodes/MeshThickenNode.cpp
@@ -1,9 +1,9 @@
 // Copyright Epic Games, Inc. All Rights Reserved.
 
 #include "MeshProcessingNodes/MeshThickenNode.h"
-#include "DynamicMesh3.h"
+#include "DynamicMesh/DynamicMesh3.h"
 #include "Operations/DisplaceMesh.h"
-#include "MeshNormals.h"
+#include "DynamicMesh/MeshNormals.h"
 
 using namespace UE::Geometry;
 using namespace UE::GeometryFlow;
diff --git a/Engine/Plugins/Experimental/GeometryFlow/Source/GeometryFlowMeshProcessing/Public/DataTypes/DynamicMeshData.h b/Engine/Plugins/Experimental/GeometryFlow/Source/GeometryFlowMeshProcessing/Public/DataTypes/DynamicMeshData.h
--- a/Engine/Plugins/Experimental/GeometryFlow/Source/GeometryFlowMeshProcessing/Public/DataTypes/DynamicMeshData.h
+++ b/Engine/Plugins/Experimental/GeometryFlow/Source/GeometryFlowMeshProcessing/Public/DataTypes/DynamicMeshData.h
@@ -5,7 +5,7 @@
 #include "GeometryFlowCoreNodes.h"
 #include "BaseNodes/TransferNode.h"
 #include "MeshProcessingNodes/MeshProcessingDataTypes.h"
-#include "DynamicMesh3.h"
+#include "DynamicMesh/DynamicMesh3.h"
 
 namespace UE
 {
diff --git a/Engine/Plugins/Experimental/GeometryFlow/Source/GeometryFlowMeshProcessing/Public/DataTypes/MeshImageBakingData.h b/Engine/Plugins/Experimental/GeometryFlow/Source/GeometryFlowMeshProcessing/Public/DataTypes/MeshImageBakingData.h
--- a/Engine/Plugins/Experimental/GeometryFlow/Source/GeometryFlowMeshProcessing/Public/DataTypes/MeshImageBakingData.h
+++ b/Engine/Plugins/Experimental/GeometryFlow/Source/GeometryFlowMeshProcessing/Public/DataTypes/MeshImageBakingData.h
@@ -6,8 +6,8 @@
 #include "GeometryFlowImmutableData.h"
 #include "DataTypes/DynamicMeshData.h"
 
-#include "DynamicMesh3.h"
-#include "DynamicMeshAABBTree3.h"
+#include "DynamicMesh/DynamicMesh3.h"
+#include "DynamicMesh/DynamicMeshAABBTree3.h"
 #include "Sampling/MeshImageBakingCache.h"
 
 namespace UE
diff --git a/Engine/Plugins/Experimental/GeometryFlow/Source/GeometryFlowMeshProcessing/Public/MeshProcessingNodes/MeshNormalsNodes.h b/Engine/Plugins/Experimental/GeometryFlow/Source/GeometryFlowMeshProcessing/Public/MeshProcessingNodes/MeshNormalsNodes.h
--- a/Engine/Plugins/Experimental/GeometryFlow/Source/GeometryFlowMeshProcessing/Public/MeshProcessingNodes/MeshNormalsNodes.h
+++ b/Engine/Plugins/Experimental/GeometryFlow/Source/GeometryFlowMeshProcessing/Public/MeshProcessingNodes/MeshNormalsNodes.h
@@ -5,7 +5,7 @@
 #include "MeshProcessingNodes/MeshProcessingBaseNodes.h"
 #include "MeshProcessingNodes/MeshProcessingDataTypes.h"
 
-#include "MeshNormals.h"
+#include "DynamicMesh/MeshNormals.h"
 
 
 namespace UE
diff --git a/Engine/Plugins/Experimental/GeometryFlow/Source/GeometryFlowMeshProcessing/Public/MeshProcessingNodes/MeshSolidifyNode.h b/Engine/Plugins/Experimental/GeometryFlow/Source/GeometryFlowMeshProcessing/Public/MeshProcessingNodes/MeshSolidifyNode.h
--- a/Engine/Plugins/Experimental/GeometryFlow/Source/GeometryFlowMeshProcessing/Public/MeshProcessingNodes/MeshSolidifyNode.h
+++ b/Engine/Plugins/Experimental/GeometryFlow/Source/GeometryFlowMeshProcessing/Public/MeshProcessingNodes/MeshSolidifyNode.h
@@ -4,7 +4,7 @@
 
 #include "MeshProcessingNodes/MeshProcessingBaseNodes.h"
 #include "MeshProcessingNodes/MeshProcessingDataTypes.h"
-#include "DynamicMeshAABBTree3.h"
+#include "DynamicMesh/DynamicMeshAABBTree3.h"
 #include "Implicit/Solidify.h"
 
 
diff --git a/Engine/Plugins/Experimental/GeometryFlow/Source/GeometryFlowMeshProcessing/Public/MeshProcessingNodes/MeshTangentsNodes.h b/Engine/Plugins/Experimental/GeometryFlow/Source/GeometryFlowMeshProcessing/Public/MeshProcessingNodes/MeshTangentsNodes.h
--- a/Engine/Plugins/Experimental/GeometryFlow/Source/GeometryFlowMeshProcessing/Public/MeshProcessingNodes/MeshTangentsNodes.h
+++ b/Engine/Plugins/Experimental/GeometryFlow/Source/GeometryFlowMeshProcessing/Public/MeshProcessingNodes/MeshTangentsNodes.h
@@ -6,7 +6,7 @@
 #include "MeshProcessingNodes/MeshProcessingBaseNodes.h"
 #include "MeshProcessingNodes/MeshProcessingDataTypes.h"
 
-#include "MeshTangents.h"
+#include "DynamicMesh/MeshTangents.h"
 
 
 namespace UE
diff --git a/Engine/Plugins/Experimental/GeometryProcessing/Source/DynamicMesh/Private/DynamicMeshEditor.cpp b/Engine/Plugins/Experimental/GeometryProcessing/Source/DynamicMesh/Private/DynamicMeshEditor.cpp
--- a/Engine/Plugins/Experimental/GeometryProcessing/Source/DynamicMesh/Private/DynamicMeshEditor.cpp
+++ b/Engine/Plugins/Experimental/GeometryProcessing/Source/DynamicMesh/Private/DynamicMeshEditor.cpp
@@ -2,12 +2,12 @@
 
 
 #include "DynamicMeshEditor.h"
-#include "DynamicMeshAttributeSet.h"
+#include "DynamicMesh/DynamicMeshAttributeSet.h"
 #include "Util/BufferUtil.h"
 #include "MeshRegionBoundaryLoops.h"
 #include "DynamicSubmesh3.h"
-#include "DynamicVertexSkinWeightsAttribute.h"
-#include "MeshNormals.h"
+#include "DynamicMesh/DynamicVertexSkinWeightsAttribute.h"
+#include "DynamicMesh/MeshNormals.h"
 #include "MeshQueries.h"
 #include "Selections/MeshConnectedComponents.h"
 
diff --git a/Engine/Plugins/Experimental/GeometryProcessing/Source/DynamicMesh/Private/MeshRefinerBase.cpp b/Engine/Plugins/Experimental/GeometryProcessing/Source/DynamicMesh/Private/MeshRefinerBase.cpp
--- a/Engine/Plugins/Experimental/GeometryProcessing/Source/DynamicMesh/Private/MeshRefinerBase.cpp
+++ b/Engine/Plugins/Experimental/GeometryProcessing/Source/DynamicMesh/Private/MeshRefinerBase.cpp
@@ -1,8 +1,8 @@
 // Copyright Epic Games, Inc. All Rights Reserved.
 
 #include "MeshRefinerBase.h"
-#include "DynamicMeshAttributeSet.h"
-#include "DynamicMeshChangeTracker.h"
+#include "DynamicMesh/DynamicMeshAttributeSet.h"
+#include "DynamicMesh/DynamicMeshChangeTracker.h"
 
 #include "ExplicitUseGeometryMathTypes.h"		// using UE::Geometry::(math types)
 using namespace UE::Geometry;
diff --git a/Engine/Plugins/Experimental/GeometryProcessing/Source/DynamicMesh/Private/MeshRegionBoundaryLoops.cpp b/Engine/Plugins/Experimental/GeometryProcessing/Source/DynamicMesh/Private/MeshRegionBoundaryLoops.cpp
--- a/Engine/Plugins/Experimental/GeometryProcessing/Source/DynamicMesh/Private/MeshRegionBoundaryLoops.cpp
+++ b/Engine/Plugins/Experimental/GeometryProcessing/Source/DynamicMesh/Private/MeshRegionBoundaryLoops.cpp
@@ -4,7 +4,7 @@
 #include "MeshRegionBoundaryLoops.h"
 
 #include "Algo/ForEach.h"
-#include "DynamicMeshAttributeSet.h"
+#include "DynamicMesh/DynamicMeshAttributeSet.h"
 #include "MeshBoundaryLoops.h"   // has a set of internal static functions we re-use
 #include "VectorUtil.h"
 #include "Util/SparseIndexCollectionTypes.h"
diff --git a/Engine/Plugins/Experimental/GeometryProcessing/Source/DynamicMesh/Private/MeshSimplification.cpp b/Engine/Plugins/Experimental/GeometryProcessing/Source/DynamicMesh/Private/MeshSimplification.cpp
--- a/Engine/Plugins/Experimental/GeometryProcessing/Source/DynamicMesh/Private/MeshSimplification.cpp
+++ b/Engine/Plugins/Experimental/GeometryProcessing/Source/DynamicMesh/Private/MeshSimplification.cpp
@@ -2,7 +2,7 @@
 
 #include "MeshSimplification.h"
 #include "MeshConstraintsUtil.h"
-#include "DynamicMeshAttributeSet.h"
+#include "DynamicMesh/DynamicMeshAttributeSet.h"
 #include "Util/IndexUtil.h"
 #include "Async/ParallelFor.h"
 #include "Templates/UnrealTypeTraits.h"
diff --git a/Engine/Plugins/Experimental/GeometryProcessing/Source/DynamicMesh/Private/NormalFlowRemesher.cpp b/Engine/Plugins/Experimental/GeometryProcessing/Source/DynamicMesh/Private/NormalFlowRemesher.cpp
--- a/Engine/Plugins/Experimental/GeometryProcessing/Source/DynamicMesh/Private/NormalFlowRemesher.cpp
+++ b/Engine/Plugins/Experimental/GeometryProcessing/Source/DynamicMesh/Private/NormalFlowRemesher.cpp
@@ -2,7 +2,7 @@
 
 #include "NormalFlowRemesher.h"
 #include "Async/ParallelFor.h"
-#include "InfoTypes.h"
+#include "DynamicMesh/InfoTypes.h"
 #include "Async/ParallelTransformReduce.h"
 
 #include "ExplicitUseGeometryMathTypes.h"		// using UE::Geometry::(math types)
diff --git a/Engine/Plugins/Experimental/GeometryProcessing/Source/DynamicMesh/Private/Operations/GroupEdgeInserter.cpp b/Engine/Plugins/Experimental/GeometryProcessing/Source/DynamicMesh/Private/Operations/GroupEdgeInserter.cpp
--- a/Engine/Plugins/Experimental/GeometryProcessing/Source/DynamicMesh/Private/Operations/GroupEdgeInserter.cpp
+++ b/Engine/Plugins/Experimental/GeometryProcessing/Source/DynamicMesh/Private/Operations/GroupEdgeInserter.cpp
@@ -5,7 +5,7 @@
 #include "Algo/ForEach.h"
 #include "CompGeom/PolygonTriangulation.h"
 #include "ConstrainedDelaunay2.h"
-#include "DynamicMeshChangeTracker.h"
+#include "DynamicMesh/DynamicMeshChangeTracker.h"
 #include "DynamicMeshEditor.h"
 #include "FrameTypes.h"
 #include "DynamicMesh/MeshIndexUtil.h"
diff --git a/Engine/Plugins/Experimental/GeometryProcessing/Source/DynamicMesh/Private/Operations/GroupTopologyDeformer.cpp b/Engine/Plugins/Experimental/GeometryProcessing/Source/DynamicMesh/Private/Operations/GroupTopologyDeformer.cpp
--- a/Engine/Plugins/Experimental/GeometryProcessing/Source/DynamicMesh/Private/Operations/GroupTopologyDeformer.cpp
+++ b/Engine/Plugins/Experimental/GeometryProcessing/Source/DynamicMesh/Private/Operations/GroupTopologyDeformer.cpp
@@ -1,12 +1,12 @@
 // Copyright Epic Games, Inc. All Rights Reserved.
 
 #include "Operations/GroupTopologyDeformer.h"
-#include "DynamicMesh3.h"
+#include "DynamicMesh/DynamicMesh3.h"
 #include "GroupTopology.h"
 #include "SegmentTypes.h"
 #include "Async/ParallelFor.h"
 #include "Containers/BitArray.h"
-#include "MeshNormals.h"
+#include "DynamicMesh/MeshNormals.h"
 
 #include "ExplicitUseGeometryMathTypes.h"		// using UE::Geometry::(math types)
 using namespace UE::Geometry;
diff --git a/Engine/Plugins/Experimental/GeometryProcessing/Source/DynamicMesh/Private/Operations/InsetMeshRegion.cpp b/Engine/Plugins/Experimental/GeometryProcessing/Source/DynamicMesh/Private/Operations/InsetMeshRegion.cpp
--- a/Engine/Plugins/Experimental/GeometryProcessing/Source/DynamicMesh/Private/Operations/InsetMeshRegion.cpp
+++ b/Engine/Plugins/Experimental/GeometryProcessing/Source/DynamicMesh/Private/Operations/InsetMeshRegion.cpp
@@ -1,16 +1,16 @@
 // Copyright Epic Games, Inc. All Rights Reserved.
 
 #include "Operations/InsetMeshRegion.h"
-#include "MeshNormals.h"
+#include "DynamicMesh/MeshNormals.h"
 #include "DynamicMeshEditor.h"
 #include "Selections/MeshVertexSelection.h"
-#include "DynamicMeshChangeTracker.h"
+#include "DynamicMesh/DynamicMeshChangeTracker.h"
 #include "Selections/MeshConnectedComponents.h"
 #include "Distance/DistLine3Line3.h"
 #include "DynamicSubmesh3.h"
 #include "Solvers/ConstrainedMeshDeformer.h"
-#include "DynamicMeshAABBTree3.h"
-#include "MeshTransforms.h"
+#include "DynamicMesh/DynamicMeshAABBTree3.h"
+#include "DynamicMesh/MeshTransforms.h"
 
 #include "ExplicitUseGeometryMathTypes.h"		// using UE::Geometry::(math types)
 using namespace UE::Geometry;
diff --git a/Engine/Plugins/Experimental/GeometryProcessing/Source/DynamicMesh/Private/Operations/MeshAttributeTransfer.cpp b/Engine/Plugins/Experimental/GeometryProcessing/Source/DynamicMesh/Private/Operations/MeshAttributeTransfer.cpp
--- a/Engine/Plugins/Experimental/GeometryProcessing/Source/DynamicMesh/Private/Operations/MeshAttributeTransfer.cpp
+++ b/Engine/Plugins/Experimental/GeometryProcessing/Source/DynamicMesh/Private/Operations/MeshAttributeTransfer.cpp
@@ -1,8 +1,8 @@
 // Copyright Epic Games, Inc. All Rights Reserved.
 
 #include "Operations/MeshAttributeTransfer.h"
-#include "DynamicMesh3.h"
-#include "DynamicMeshAttributeSet.h"
+#include "DynamicMesh/DynamicMesh3.h"
+#include "DynamicMesh/DynamicMeshAttributeSet.h"
 #include "TriangleTypes.h"
 #include "Util/IndexUtil.h"
 #include "Async/ParallelFor.h"
diff --git a/Engine/Plugins/Experimental/GeometryProcessing/Source/DynamicMesh/Private/Operations/MeshBoolean.cpp b/Engine/Plugins/Experimental/GeometryProcessing/Source/DynamicMesh/Private/Operations/MeshBoolean.cpp
--- a/Engine/Plugins/Experimental/GeometryProcessing/Source/DynamicMesh/Private/Operations/MeshBoolean.cpp
+++ b/Engine/Plugins/Experimental/GeometryProcessing/Source/DynamicMesh/Private/Operations/MeshBoolean.cpp
@@ -7,12 +7,12 @@
 #include "Selections/MeshConnectedComponents.h"
 
 #include "Async/ParallelFor.h"
-#include "MeshTransforms.h"
+#include "DynamicMesh/MeshTransforms.h"
 #include "Spatial/SparseDynamicOctree3.h"
 
 #include "Algo/RemoveIf.h"
 
-#include "DynamicMeshAABBTree3.h"
+#include "DynamicMesh/DynamicMeshAABBTree3.h"
 
 #include "ExplicitUseGeometryMathTypes.h"		// using UE::Geometry::(math types)
 using namespace UE::Geometry;
diff --git a/Engine/Plugins/Experimental/GeometryProcessing/Source/DynamicMesh/Private/Operations/MeshPlaneCut.cpp b/Engine/Plugins/Experimental/GeometryProcessing/Source/DynamicMesh/Private/Operations/MeshPlaneCut.cpp
--- a/Engine/Plugins/Experimental/GeometryProcessing/Source/DynamicMesh/Private/Operations/MeshPlaneCut.cpp
+++ b/Engine/Plugins/Experimental/GeometryProcessing/Source/DynamicMesh/Private/Operations/MeshPlaneCut.cpp
@@ -2,13 +2,13 @@
 
 #include "Operations/MeshPlaneCut.h"
 
-#include "DynamicMesh3.h"
-#include "DynamicMeshTriangleAttribute.h"
+#include "DynamicMesh/DynamicMesh3.h"
+#include "DynamicMesh/DynamicMeshTriangleAttribute.h"
 
 #include "Operations/SimpleHoleFiller.h"
 #include "Operations/PlanarHoleFiller.h"
 #include "Operations/MinimalHoleFiller.h"
-#include "MeshNormals.h"
+#include "DynamicMesh/MeshNormals.h"
 #include "DynamicMeshEditor.h"
 #include "MathUtil.h"
 #include "Selections/MeshConnectedComponents.h"
diff --git a/Engine/Plugins/Experimental/GeometryProcessing/Source/DynamicMesh/Private/Operations/MeshSelfUnion.cpp b/Engine/Plugins/Experimental/GeometryProcessing/Source/DynamicMesh/Private/Operations/MeshSelfUnion.cpp
--- a/Engine/Plugins/Experimental/GeometryProcessing/Source/DynamicMesh/Private/Operations/MeshSelfUnion.cpp
+++ b/Engine/Plugins/Experimental/GeometryProcessing/Source/DynamicMesh/Private/Operations/MeshSelfUnion.cpp
@@ -7,15 +7,15 @@
 
 #include "Selections/MeshConnectedComponents.h"
 
-#include "MeshNormals.h"
+#include "DynamicMesh/MeshNormals.h"
 
 #include "Async/ParallelFor.h"
-#include "MeshTransforms.h"
+#include "DynamicMesh/MeshTransforms.h"
 #include "Spatial/SparseDynamicOctree3.h"
 
 #include "Algo/RemoveIf.h"
 
-#include "DynamicMeshAABBTree3.h"
+#include "DynamicMesh/DynamicMeshAABBTree3.h"
 
 #include "ExplicitUseGeometryMathTypes.h"		// using UE::Geometry::(math types)
 using namespace UE::Geometry;
diff --git a/Engine/Plugins/Experimental/GeometryProcessing/Source/DynamicMesh/Private/Operations/OffsetMeshRegion.cpp b/Engine/Plugins/Experimental/GeometryProcessing/Source/DynamicMesh/Private/Operations/OffsetMeshRegion.cpp
--- a/Engine/Plugins/Experimental/GeometryProcessing/Source/DynamicMesh/Private/Operations/OffsetMeshRegion.cpp
+++ b/Engine/Plugins/Experimental/GeometryProcessing/Source/DynamicMesh/Private/Operations/OffsetMeshRegion.cpp
@@ -1,10 +1,10 @@
 // Copyright Epic Games, Inc. All Rights Reserved.
 
 #include "Operations/OffsetMeshRegion.h"
-#include "MeshNormals.h"
+#include "DynamicMesh/MeshNormals.h"
 #include "DynamicMeshEditor.h"
 #include "Selections/MeshVertexSelection.h"
-#include "DynamicMeshChangeTracker.h"
+#include "DynamicMesh/DynamicMeshChangeTracker.h"
 #include "Selections/MeshConnectedComponents.h"
 #include "Operations/ExtrudeMesh.h"
 #include "DynamicSubmesh3.h"
diff --git a/Engine/Plugins/Experimental/GeometryProcessing/Source/DynamicMesh/Private/Operations/SmoothHoleFiller.cpp b/Engine/Plugins/Experimental/GeometryProcessing/Source/DynamicMesh/Private/Operations/SmoothHoleFiller.cpp
--- a/Engine/Plugins/Experimental/GeometryProcessing/Source/DynamicMesh/Private/Operations/SmoothHoleFiller.cpp
+++ b/Engine/Plugins/Experimental/GeometryProcessing/Source/DynamicMesh/Private/Operations/SmoothHoleFiller.cpp
@@ -9,7 +9,7 @@
 #include "Operations/OffsetMeshRegion.h"
 #include "Operations/MeshRegionOperator.h"
 #include "Solvers/ConstrainedMeshSmoother.h"
-#include "DynamicMeshChangeTracker.h"
+#include "DynamicMesh/DynamicMeshChangeTracker.h"
 #include "SubRegionRemesher.h"
 #include "ProjectionTargets.h"
 #include "MeshConstraintsUtil.h"
diff --git a/Engine/Plugins/Experimental/GeometryProcessing/Source/DynamicMesh/Private/ShapeApproximation/MeshSimpleShapeApproximation.cpp b/Engine/Plugins/Experimental/GeometryProcessing/Source/DynamicMesh/Private/ShapeApproximation/MeshSimpleShapeApproximation.cpp
--- a/Engine/Plugins/Experimental/GeometryProcessing/Source/DynamicMesh/Private/ShapeApproximation/MeshSimpleShapeApproximation.cpp
+++ b/Engine/Plugins/Experimental/GeometryProcessing/Source/DynamicMesh/Private/ShapeApproximation/MeshSimpleShapeApproximation.cpp
@@ -6,7 +6,7 @@
 #include "MinVolumeSphere3.h"
 #include "MinVolumeBox3.h"
 #include "FitCapsule3.h"
-//#include "DynamicMeshAABBTree3.h"
+//#include "DynamicMesh/DynamicMeshAABBTree3.h"
 
 #include "ShapeApproximation/ShapeDetection3.h"
 #include "MeshQueries.h"
diff --git a/Engine/Plugins/Experimental/GeometryProcessing/Source/DynamicMesh/Private/ShapeApproximation/SimpleShapeSet3.cpp b/Engine/Plugins/Experimental/GeometryProcessing/Source/DynamicMesh/Private/ShapeApproximation/SimpleShapeSet3.cpp
--- a/Engine/Plugins/Experimental/GeometryProcessing/Source/DynamicMesh/Private/ShapeApproximation/SimpleShapeSet3.cpp
+++ b/Engine/Plugins/Experimental/GeometryProcessing/Source/DynamicMesh/Private/ShapeApproximation/SimpleShapeSet3.cpp
@@ -3,7 +3,7 @@
 #include "ShapeApproximation/SimpleShapeSet3.h"
 
 #include "MeshQueries.h"
-#include "MeshTransforms.h"
+#include "DynamicMesh/MeshTransforms.h"
 #include "Intersection/ContainmentQueries3.h"
 
 #include "ExplicitUseGeometryMathTypes.h"		// using UE::Geometry::(math types)
diff --git a/Engine/Plugins/Experimental/GeometryProcessing/Source/DynamicMesh/Public/DynamicMeshEditor.h b/Engine/Plugins/Experimental/GeometryProcessing/Source/DynamicMesh/Public/DynamicMeshEditor.h
--- a/Engine/Plugins/Experimental/GeometryProcessing/Source/DynamicMesh/Public/DynamicMeshEditor.h
+++ b/Engine/Plugins/Experimental/GeometryProcessing/Source/DynamicMesh/Public/DynamicMeshEditor.h
@@ -4,11 +4,11 @@
 
 #pragma once
 
-#include "DynamicMesh3.h"
-#include "DynamicMeshAttributeSet.h"
+#include "DynamicMesh/DynamicMesh3.h"
+#include "DynamicMesh/DynamicMeshAttributeSet.h"
 #include "EdgeLoop.h"
 #include "Util/SparseIndexCollectionTypes.h"
-#include "MeshIndexMappings.h"
+#include "DynamicMesh/MeshIndexMappings.h"
 #include "MeshAdapter.h"
 
 namespace UE
diff --git a/Engine/Plugins/Experimental/GeometryProcessing/Source/DynamicMesh/Public/GroupTopology.h b/Engine/Plugins/Experimental/GeometryProcessing/Source/DynamicMesh/Public/GroupTopology.h
--- a/Engine/Plugins/Experimental/GeometryProcessing/Source/DynamicMesh/Public/GroupTopology.h
+++ b/Engine/Plugins/Experimental/GeometryProcessing/Source/DynamicMesh/Public/GroupTopology.h
@@ -2,8 +2,8 @@
 
 #pragma once
 
-#include "DynamicMesh3.h"
-#include "DynamicMeshAttributeSet.h"
+#include "DynamicMesh/DynamicMesh3.h"
+#include "DynamicMesh/DynamicMeshAttributeSet.h"
 #include "EdgeSpan.h"
 #include "Util/IndexUtil.h"
 #include "Containers/BitArray.h"
diff --git a/Engine/Plugins/Experimental/GeometryProcessing/Source/DynamicMesh/Public/MeshConstraintsUtil.h b/Engine/Plugins/Experimental/GeometryProcessing/Source/DynamicMesh/Public/MeshConstraintsUtil.h
--- a/Engine/Plugins/Experimental/GeometryProcessing/Source/DynamicMesh/Public/MeshConstraintsUtil.h
+++ b/Engine/Plugins/Experimental/GeometryProcessing/Source/DynamicMesh/Public/MeshConstraintsUtil.h
@@ -4,9 +4,9 @@
 
 #pragma once
 
-#include "DynamicMesh3.h"
+#include "DynamicMesh/DynamicMesh3.h"
 #include "MeshConstraints.h"
-#include "DynamicMeshAttributeSet.h"
+#include "DynamicMesh/DynamicMeshAttributeSet.h"
 
 
 namespace UE
diff --git a/Engine/Plugins/Experimental/GeometryProcessing/Source/DynamicMesh/Public/MeshRegionBoundaryLoops.h b/Engine/Plugins/Experimental/GeometryProcessing/Source/DynamicMesh/Public/MeshRegionBoundaryLoops.h
--- a/Engine/Plugins/Experimental/GeometryProcessing/Source/DynamicMesh/Public/MeshRegionBoundaryLoops.h
+++ b/Engine/Plugins/Experimental/GeometryProcessing/Source/DynamicMesh/Public/MeshRegionBoundaryLoops.h
@@ -4,8 +4,8 @@
 
 #pragma once
 
-#include "DynamicMesh3.h"
-#include "DynamicMeshOverlay.h"
+#include "DynamicMesh/DynamicMesh3.h"
+#include "DynamicMesh/DynamicMeshOverlay.h"
 #include "EdgeLoop.h"
 #include "Util/SparseIndexCollectionTypes.h"
 
diff --git a/Engine/Plugins/Experimental/GeometryProcessing/Source/DynamicMesh/Public/Operations/RepairOrientation.h b/Engine/Plugins/Experimental/GeometryProcessing/Source/DynamicMesh/Public/Operations/RepairOrientation.h
--- a/Engine/Plugins/Experimental/GeometryProcessing/Source/DynamicMesh/Public/Operations/RepairOrientation.h
+++ b/Engine/Plugins/Experimental/GeometryProcessing/Source/DynamicMesh/Public/Operations/RepairOrientation.h
@@ -6,8 +6,8 @@
 
 #include "MathUtil.h"
 #include "VectorTypes.h"
-#include "DynamicMesh3.h"
-#include "DynamicMeshAABBTree3.h"
+#include "DynamicMesh/DynamicMesh3.h"
+#include "DynamicMesh/DynamicMeshAABBTree3.h"
 
 namespace UE
 {
diff --git a/Engine/Plugins/Experimental/GeometryProcessing/Source/DynamicMesh/Public/Sampling/MeshGenericWorldPositionBaker.h b/Engine/Plugins/Experimental/GeometryProcessing/Source/DynamicMesh/Public/Sampling/MeshGenericWorldPositionBaker.h
--- a/Engine/Plugins/Experimental/GeometryProcessing/Source/DynamicMesh/Public/Sampling/MeshGenericWorldPositionBaker.h
+++ b/Engine/Plugins/Experimental/GeometryProcessing/Source/DynamicMesh/Public/Sampling/MeshGenericWorldPositionBaker.h
@@ -6,7 +6,7 @@
 #include "Sampling/MeshImageBaker.h"
 #include "Sampling/MeshImageBakingCache.h"
 #include "Image/ImageBuilder.h"
-#include "MeshTangents.h"
+#include "DynamicMesh/MeshTangents.h"
 
 namespace UE
 {
diff --git a/Engine/Plugins/Experimental/GeometryProcessing/Source/DynamicMesh/Public/Sampling/MeshImageBakingCache.h b/Engine/Plugins/Experimental/GeometryProcessing/Source/DynamicMesh/Public/Sampling/MeshImageBakingCache.h
--- a/Engine/Plugins/Experimental/GeometryProcessing/Source/DynamicMesh/Public/Sampling/MeshImageBakingCache.h
+++ b/Engine/Plugins/Experimental/GeometryProcessing/Source/DynamicMesh/Public/Sampling/MeshImageBakingCache.h
@@ -2,9 +2,9 @@
 
 #pragma once
 
-#include "DynamicMesh3.h"
-#include "DynamicMeshAttributeSet.h"
-#include "DynamicMeshAABBTree3.h"
+#include "DynamicMesh/DynamicMesh3.h"
+#include "DynamicMesh/DynamicMeshAttributeSet.h"
+#include "DynamicMesh/DynamicMeshAABBTree3.h"
 #include "Sampling/MeshSurfaceSampler.h"
 #include "Spatial/DenseGrid2.h"
 #include "Image/ImageDimensions.h"
diff --git a/Engine/Plugins/Experimental/GeometryProcessing/Source/DynamicMesh/Public/Sampling/MeshMapBaker.h b/Engine/Plugins/Experimental/GeometryProcessing/Source/DynamicMesh/Public/Sampling/MeshMapBaker.h
--- a/Engine/Plugins/Experimental/GeometryProcessing/Source/DynamicMesh/Public/Sampling/MeshMapBaker.h
+++ b/Engine/Plugins/Experimental/GeometryProcessing/Source/DynamicMesh/Public/Sampling/MeshMapBaker.h
@@ -2,10 +2,10 @@
 
 #pragma once
 
-#include "DynamicMesh3.h"
-#include "DynamicMeshAttributeSet.h"
-#include "DynamicMeshAABBTree3.h"
-#include "MeshTangents.h"
+#include "DynamicMesh/DynamicMesh3.h"
+#include "DynamicMesh/DynamicMeshAttributeSet.h"
+#include "DynamicMesh/DynamicMeshAABBTree3.h"
+#include "DynamicMesh/MeshTangents.h"
 #include "Sampling/MeshMapEvaluator.h"
 #include "Sampling/MeshSurfaceSampler.h"
 #include "Spatial/DenseGrid2.h"
diff --git a/Engine/Plugins/Experimental/GeometryProcessing/Source/DynamicMesh/Public/Sampling/MeshSurfaceSampler.h b/Engine/Plugins/Experimental/GeometryProcessing/Source/DynamicMesh/Public/Sampling/MeshSurfaceSampler.h
--- a/Engine/Plugins/Experimental/GeometryProcessing/Source/DynamicMesh/Public/Sampling/MeshSurfaceSampler.h
+++ b/Engine/Plugins/Experimental/GeometryProcessing/Source/DynamicMesh/Public/Sampling/MeshSurfaceSampler.h
@@ -2,9 +2,9 @@
 
 #pragma once
 
-#include "DynamicMesh3.h"
-#include "DynamicMeshAttributeSet.h"
-#include "DynamicMeshAABBTree3.h"
+#include "DynamicMesh/DynamicMesh3.h"
+#include "DynamicMesh/DynamicMeshAttributeSet.h"
+#include "DynamicMesh/DynamicMeshAABBTree3.h"
 
 namespace UE
 {
diff --git a/Engine/Plugins/Experimental/MeshLODToolset/Source/MeshLODToolset/Private/Tools/GenerateStaticMeshLODAssetTool.cpp b/Engine/Plugins/Experimental/MeshLODToolset/Source/MeshLODToolset/Private/Tools/GenerateStaticMeshLODAssetTool.cpp
--- a/Engine/Plugins/Experimental/MeshLODToolset/Source/MeshLODToolset/Private/Tools/GenerateStaticMeshLODAssetTool.cpp
+++ b/Engine/Plugins/Experimental/MeshLODToolset/Source/MeshLODToolset/Private/Tools/GenerateStaticMeshLODAssetTool.cpp
@@ -8,12 +8,12 @@
 
 #include "Generators/SphereGenerator.h"
 
-#include "DynamicMesh3.h"
+#include "DynamicMesh/DynamicMesh3.h"
 #include "BaseBehaviors/MultiClickSequenceInputBehavior.h"
 #include "Selection/SelectClickedAction.h"
 #include "DynamicMeshEditor.h"
-#include "MeshTransforms.h"
-#include "MeshTangents.h"
+#include "DynamicMesh/MeshTransforms.h"
+#include "DynamicMesh/MeshTangents.h"
 
 #include "MeshDescriptionToDynamicMesh.h"
 #include "DynamicMeshToMeshDescription.h"
diff --git a/Engine/Plugins/Experimental/MeshLODToolset/Source/MeshLODToolset/Public/Graphs/GenerateMeshLODGraph.h b/Engine/Plugins/Experimental/MeshLODToolset/Source/MeshLODToolset/Public/Graphs/GenerateMeshLODGraph.h
--- a/Engine/Plugins/Experimental/MeshLODToolset/Source/MeshLODToolset/Public/Graphs/GenerateMeshLODGraph.h
+++ b/Engine/Plugins/Experimental/MeshLODToolset/Source/MeshLODToolset/Public/Graphs/GenerateMeshLODGraph.h
@@ -7,8 +7,8 @@
 #include "DataTypes/NormalMapData.h"
 #include "DataTypes/TextureImageData.h"
 #include "DataTypes/CollisionGeometryData.h"
-#include "DynamicMesh3.h"
-#include "MeshTangents.h"
+#include "DynamicMesh/DynamicMesh3.h"
+#include "DynamicMesh/MeshTangents.h"
 
 #include "MeshProcessingNodes/MeshSolidifyNode.h"
 #include "MeshProcessingNodes/MeshVoxMorphologyNode.h"
diff --git a/Engine/Plugins/Experimental/MeshModelingToolset/Source/GeometryProcessingAdapters/Private/ApproximateActorsImpl.cpp b/Engine/Plugins/Experimental/MeshModelingToolset/Source/GeometryProcessingAdapters/Private/ApproximateActorsImpl.cpp
--- a/Engine/Plugins/Experimental/MeshModelingToolset/Source/GeometryProcessingAdapters/Private/ApproximateActorsImpl.cpp
+++ b/Engine/Plugins/Experimental/MeshModelingToolset/Source/GeometryProcessingAdapters/Private/ApproximateActorsImpl.cpp
@@ -3,9 +3,9 @@
 #include "GeometryProcessing/ApproximateActorsImpl.h"
 
 #include "Scene/MeshSceneAdapter.h"
-#include "DynamicMesh3.h"
-#include "MeshNormals.h"
-#include "MeshTangents.h"
+#include "DynamicMesh/DynamicMesh3.h"
+#include "DynamicMesh/MeshNormals.h"
+#include "DynamicMesh/MeshTangents.h"
 #include "MeshSimplification.h"
 #include "MeshConstraintsUtil.h"
 #include "Implicit/Solidify.h"
diff --git a/Engine/Plugins/Experimental/MeshModelingToolset/Source/MeshModelingTools/Private/AddPatchTool.cpp b/Engine/Plugins/Experimental/MeshModelingToolset/Source/MeshModelingTools/Private/AddPatchTool.cpp
--- a/Engine/Plugins/Experimental/MeshModelingToolset/Source/MeshModelingTools/Private/AddPatchTool.cpp
+++ b/Engine/Plugins/Experimental/MeshModelingToolset/Source/MeshModelingTools/Private/AddPatchTool.cpp
@@ -10,9 +10,9 @@
 
 #include "MeshDescriptionBuilder.h"
 #include "Generators/RectangleMeshGenerator.h"
-#include "DynamicMeshAttributeSet.h"
+#include "DynamicMesh/DynamicMeshAttributeSet.h"
 #include "Drawing/MeshDebugDrawing.h"
-#include "MeshNormals.h"
+#include "DynamicMesh/MeshNormals.h"
 
 #include "Async/ParallelFor.h"
 
diff --git a/Engine/Plugins/Experimental/MeshModelingToolset/Source/MeshModelingTools/Private/AlignObjectsTool.cpp b/Engine/Plugins/Experimental/MeshModelingToolset/Source/MeshModelingTools/Private/AlignObjectsTool.cpp
--- a/Engine/Plugins/Experimental/MeshModelingToolset/Source/MeshModelingTools/Private/AlignObjectsTool.cpp
+++ b/Engine/Plugins/Experimental/MeshModelingToolset/Source/MeshModelingTools/Private/AlignObjectsTool.cpp
@@ -5,9 +5,9 @@
 #include "InteractiveGizmoManager.h"
 #include "ToolBuilderUtil.h"
 #include "ToolSetupUtil.h"
-#include "DynamicMesh3.h"
+#include "DynamicMesh/DynamicMesh3.h"
 #include "MeshAdapterTransforms.h"
-#include "MeshTransforms.h"
+#include "DynamicMesh/MeshTransforms.h"
 #include "BaseBehaviors/ClickDragBehavior.h"
 #include "TargetInterfaces/PrimitiveComponentBackedTarget.h"
 #include "ToolTargetManager.h"
diff --git a/Engine/Plugins/Experimental/MeshModelingToolset/Source/MeshModelingTools/Private/BakeMeshAttributeMapsTool.cpp b/Engine/Plugins/Experimental/MeshModelingToolset/Source/MeshModelingTools/Private/BakeMeshAttributeMapsTool.cpp
--- a/Engine/Plugins/Experimental/MeshModelingToolset/Source/MeshModelingTools/Private/BakeMeshAttributeMapsTool.cpp
+++ b/Engine/Plugins/Experimental/MeshModelingToolset/Source/MeshModelingTools/Private/BakeMeshAttributeMapsTool.cpp
@@ -4,9 +4,9 @@
 #include "InteractiveToolManager.h"
 #include "ToolBuilderUtil.h"
 
-#include "DynamicMesh3.h"
-#include "DynamicMeshAttributeSet.h"
-#include "MeshTransforms.h"
+#include "DynamicMesh/DynamicMesh3.h"
+#include "DynamicMesh/DynamicMeshAttributeSet.h"
+#include "DynamicMesh/MeshTransforms.h"
 #include "MeshDescriptionToDynamicMesh.h"
 #include "Sampling/MeshNormalMapEvaluator.h"
 #include "Sampling/MeshOcclusionMapEvaluator.h"
diff --git a/Engine/Plugins/Experimental/MeshModelingToolset/Source/MeshModelingTools/Private/CombineMeshesTool.cpp b/Engine/Plugins/Experimental/MeshModelingToolset/Source/MeshModelingTools/Private/CombineMeshesTool.cpp
--- a/Engine/Plugins/Experimental/MeshModelingToolset/Source/MeshModelingTools/Private/CombineMeshesTool.cpp
+++ b/Engine/Plugins/Experimental/MeshModelingToolset/Source/MeshModelingTools/Private/CombineMeshesTool.cpp
@@ -6,11 +6,11 @@
 
 #include "ToolSetupUtil.h"
 
-#include "DynamicMesh3.h"
+#include "DynamicMesh/DynamicMesh3.h"
 #include "BaseBehaviors/MultiClickSequenceInputBehavior.h"
 #include "Selection/SelectClickedAction.h"
 #include "DynamicMeshEditor.h"
-#include "MeshTransforms.h"
+#include "DynamicMesh/MeshTransforms.h"
 
 #include "MeshDescriptionToDynamicMesh.h"
 #include "DynamicMeshToMeshDescription.h"
diff --git a/Engine/Plugins/Experimental/MeshModelingToolset/Source/MeshModelingTools/Private/ConvertToPolygonsTool.cpp b/Engine/Plugins/Experimental/MeshModelingToolset/Source/MeshModelingTools/Private/ConvertToPolygonsTool.cpp
--- a/Engine/Plugins/Experimental/MeshModelingToolset/Source/MeshModelingTools/Private/ConvertToPolygonsTool.cpp
+++ b/Engine/Plugins/Experimental/MeshModelingToolset/Source/MeshModelingTools/Private/ConvertToPolygonsTool.cpp
@@ -4,8 +4,8 @@
 #include "InteractiveToolManager.h"
 #include "ToolBuilderUtil.h"
 
-#include "DynamicMesh3.h"
-#include "MeshNormals.h"
+#include "DynamicMesh/DynamicMesh3.h"
+#include "DynamicMesh/MeshNormals.h"
 #include "DynamicMeshEditor.h"
 #include "MeshRegionBoundaryLoops.h"
 #include "MeshDescriptionToDynamicMesh.h"
diff --git a/Engine/Plugins/Experimental/MeshModelingToolset/Source/MeshModelingTools/Private/CutMeshWithMeshTool.cpp b/Engine/Plugins/Experimental/MeshModelingToolset/Source/MeshModelingTools/Private/CutMeshWithMeshTool.cpp
--- a/Engine/Plugins/Experimental/MeshModelingToolset/Source/MeshModelingTools/Private/CutMeshWithMeshTool.cpp
+++ b/Engine/Plugins/Experimental/MeshModelingToolset/Source/MeshModelingTools/Private/CutMeshWithMeshTool.cpp
@@ -6,8 +6,8 @@
 #include "BaseGizmos/TransformGizmoUtil.h"
 #include "Selection/ToolSelectionUtil.h"
 #include "ModelingObjectsCreationAPI.h"
-#include "DynamicMesh3.h"
-#include "MeshTransforms.h"
+#include "DynamicMesh/DynamicMesh3.h"
+#include "DynamicMesh/MeshTransforms.h"
 #include "MeshDescriptionToDynamicMesh.h"
 #include "DynamicMeshToMeshDescription.h"
 #include "Async/Async.h"
diff --git a/Engine/Plugins/Experimental/MeshModelingToolset/Source/MeshModelingTools/Private/DeformMeshPolygonsTool.cpp b/Engine/Plugins/Experimental/MeshModelingToolset/Source/MeshModelingTools/Private/DeformMeshPolygonsTool.cpp
--- a/Engine/Plugins/Experimental/MeshModelingToolset/Source/MeshModelingTools/Private/DeformMeshPolygonsTool.cpp
+++ b/Engine/Plugins/Experimental/MeshModelingToolset/Source/MeshModelingTools/Private/DeformMeshPolygonsTool.cpp
@@ -4,11 +4,11 @@
 
 #include "Async/ParallelFor.h"
 #include "Containers/BitArray.h"
-#include "DynamicMeshAttributeSet.h"
+#include "DynamicMesh/DynamicMeshAttributeSet.h"
 #include "FindPolygonsAlgorithm.h"
 #include "InteractiveToolManager.h"
 #include "Intersection/IntersectionUtil.h"
-#include "MeshNormals.h"
+#include "DynamicMesh/MeshNormals.h"
 #include "ModelingOperators/Public/ModelingTaskTypes.h"
 #include "SegmentTypes.h"
 #include "Solvers/ConstrainedMeshDeformer.h"
diff --git a/Engine/Plugins/Experimental/MeshModelingToolset/Source/MeshModelingTools/Private/DrawPolyPathTool.cpp b/Engine/Plugins/Experimental/MeshModelingToolset/Source/MeshModelingTools/Private/DrawPolyPathTool.cpp
--- a/Engine/Plugins/Experimental/MeshModelingToolset/Source/MeshModelingTools/Private/DrawPolyPathTool.cpp
+++ b/Engine/Plugins/Experimental/MeshModelingToolset/Source/MeshModelingTools/Private/DrawPolyPathTool.cpp
@@ -15,10 +15,10 @@
 #include "Generators/RectangleMeshGenerator.h"
 #include "Distance/DistLine3Line3.h"
 #include "ModelingObjectsCreationAPI.h"
-#include "MeshTransforms.h"
+#include "DynamicMesh/MeshTransforms.h"
 #include "Selection/ToolSelectionUtil.h"
 #include "Operations/ExtrudeMesh.h"
-#include "MeshNormals.h"
+#include "DynamicMesh/MeshNormals.h"
 
 #include "ExplicitUseGeometryMathTypes.h"		// using UE::Geometry::(math types)
 using namespace UE::Geometry;
diff --git a/Engine/Plugins/Experimental/MeshModelingToolset/Source/MeshModelingTools/Private/EditMeshPolygonsTool.cpp b/Engine/Plugins/Experimental/MeshModelingToolset/Source/MeshModelingTools/Private/EditMeshPolygonsTool.cpp
--- a/Engine/Plugins/Experimental/MeshModelingToolset/Source/MeshModelingTools/Private/EditMeshPolygonsTool.cpp
+++ b/Engine/Plugins/Experimental/MeshModelingToolset/Source/MeshModelingTools/Private/EditMeshPolygonsTool.cpp
@@ -6,7 +6,7 @@
 
 #include "CompGeom/PolygonTriangulation.h"
 #include "SegmentTypes.h"
-#include "DynamicMeshAttributeSet.h"
+#include "DynamicMesh/DynamicMeshAttributeSet.h"
 #include "Mechanics/DragAlignmentMechanic.h"
 #include "DynamicMesh/MeshNormals.h"
 #include "ToolSceneQueriesUtil.h"
@@ -24,7 +24,7 @@
 #include "Selections/MeshConnectedComponents.h"
 #include "FaceGroupUtil.h"
 #include "DynamicMeshEditor.h"
-#include "DynamicMeshChangeTracker.h"
+#include "DynamicMesh/DynamicMeshChangeTracker.h"
 #include "Changes/MeshChange.h"
 #include "DynamicMesh/MeshIndexUtil.h"
 #include "MeshRegionBoundaryLoops.h"
@@ -33,7 +33,7 @@
 #include "Operations/OffsetMeshRegion.h"
 #include "Operations/InsetMeshRegion.h"
 #include "Operations/SimpleHoleFiller.h"
-#include "MeshTransforms.h"
+#include "DynamicMesh/MeshTransforms.h"
 
 #include "Algo/ForEach.h"
 #include "Async/ParallelFor.h"
diff --git a/Engine/Plugins/Experimental/MeshModelingToolset/Source/MeshModelingTools/Private/EditUVIslandsTool.cpp b/Engine/Plugins/Experimental/MeshModelingToolset/Source/MeshModelingTools/Private/EditUVIslandsTool.cpp
--- a/Engine/Plugins/Experimental/MeshModelingToolset/Source/MeshModelingTools/Private/EditUVIslandsTool.cpp
+++ b/Engine/Plugins/Experimental/MeshModelingToolset/Source/MeshModelingTools/Private/EditUVIslandsTool.cpp
@@ -5,8 +5,8 @@
 #include "ToolBuilderUtil.h"
 
 #include "SegmentTypes.h"
-#include "DynamicMeshAttributeSet.h"
-#include "MeshNormals.h"
+#include "DynamicMesh/DynamicMeshAttributeSet.h"
+#include "DynamicMesh/MeshNormals.h"
 #include "Selections/MeshConnectedComponents.h"
 #include "Transforms/MultiTransformer.h"
 #include "BaseBehaviors/SingleClickBehavior.h"
diff --git a/Engine/Plugins/Experimental/MeshModelingToolset/Source/MeshModelingTools/Private/LatticeDeformerTool.cpp b/Engine/Plugins/Experimental/MeshModelingToolset/Source/MeshModelingTools/Private/LatticeDeformerTool.cpp
--- a/Engine/Plugins/Experimental/MeshModelingToolset/Source/MeshModelingTools/Private/LatticeDeformerTool.cpp
+++ b/Engine/Plugins/Experimental/MeshModelingToolset/Source/MeshModelingTools/Private/LatticeDeformerTool.cpp
@@ -15,7 +15,7 @@
 #include "ToolBuilderUtil.h"
 #include "ToolSetupUtil.h"
 #include "DynamicMeshToMeshDescription.h"
-#include "MeshTransforms.h"
+#include "DynamicMesh/MeshTransforms.h"
 #include "Algo/ForEach.h"
 #include "Operations/FFDLattice.h"
 
diff --git a/Engine/Plugins/Experimental/MeshModelingToolset/Source/MeshModelingTools/Private/MeshInspectorTool.cpp b/Engine/Plugins/Experimental/MeshModelingToolset/Source/MeshModelingTools/Private/MeshInspectorTool.cpp
--- a/Engine/Plugins/Experimental/MeshModelingToolset/Source/MeshModelingTools/Private/MeshInspectorTool.cpp
+++ b/Engine/Plugins/Experimental/MeshModelingToolset/Source/MeshModelingTools/Private/MeshInspectorTool.cpp
@@ -4,9 +4,9 @@
 #include "InteractiveToolManager.h"
 #include "ToolBuilderUtil.h"
 
-#include "DynamicMesh3.h"
-#include "DynamicMeshAttributeSet.h"
-#include "MeshNormals.h"
+#include "DynamicMesh/DynamicMesh3.h"
+#include "DynamicMesh/DynamicMeshAttributeSet.h"
+#include "DynamicMesh/MeshNormals.h"
 
 #include "Components/DynamicMeshComponent.h"
 #include "Materials/MaterialInstanceDynamic.h"
diff --git a/Engine/Plugins/Experimental/MeshModelingToolset/Source/MeshModelingTools/Private/Physics/ExtractCollisionGeometryTool.cpp b/Engine/Plugins/Experimental/MeshModelingToolset/Source/MeshModelingTools/Private/Physics/ExtractCollisionGeometryTool.cpp
--- a/Engine/Plugins/Experimental/MeshModelingToolset/Source/MeshModelingTools/Private/Physics/ExtractCollisionGeometryTool.cpp
+++ b/Engine/Plugins/Experimental/MeshModelingToolset/Source/MeshModelingTools/Private/Physics/ExtractCollisionGeometryTool.cpp
@@ -11,11 +11,11 @@
 #include "PreviewMesh.h"
 
 #include "DynamicMeshEditor.h"
-#include "MeshNormals.h"
+#include "DynamicMesh/MeshNormals.h"
 #include "Generators/SphereGenerator.h"
 #include "Generators/MinimalBoxMeshGenerator.h"
 #include "Generators/CapsuleGenerator.h"
-#include "MeshTransforms.h"
+#include "DynamicMesh/MeshTransforms.h"
 #include "Parameterization/DynamicMeshUVEditor.h"
 
 #include "Physics/PhysicsDataCollection.h"
diff --git a/Engine/Plugins/Experimental/MeshModelingToolset/Source/MeshModelingTools/Private/PlaneCutTool.cpp b/Engine/Plugins/Experimental/MeshModelingToolset/Source/MeshModelingTools/Private/PlaneCutTool.cpp
--- a/Engine/Plugins/Experimental/MeshModelingToolset/Source/MeshModelingTools/Private/PlaneCutTool.cpp
+++ b/Engine/Plugins/Experimental/MeshModelingToolset/Source/MeshModelingTools/Private/PlaneCutTool.cpp
@@ -6,8 +6,8 @@
 
 #include "ToolSetupUtil.h"
 
-#include "DynamicMesh3.h"
-#include "DynamicMeshTriangleAttribute.h"
+#include "DynamicMesh/DynamicMesh3.h"
+#include "DynamicMesh/DynamicMeshTriangleAttribute.h"
 #include "DynamicMeshEditor.h"
 #include "BaseBehaviors/MultiClickSequenceInputBehavior.h"
 #include "BaseBehaviors/KeyAsModifierInputBehavior.h"
diff --git a/Engine/Plugins/Experimental/MeshModelingToolset/Source/MeshModelingTools/Private/Properties/MeshMaterialProperties.cpp b/Engine/Plugins/Experimental/MeshModelingToolset/Source/MeshModelingTools/Private/Properties/MeshMaterialProperties.cpp
--- a/Engine/Plugins/Experimental/MeshModelingToolset/Source/MeshModelingTools/Private/Properties/MeshMaterialProperties.cpp
+++ b/Engine/Plugins/Experimental/MeshModelingToolset/Source/MeshModelingTools/Private/Properties/MeshMaterialProperties.cpp
@@ -2,9 +2,9 @@
 
 #include "Properties/MeshMaterialProperties.h"
 
-#include "DynamicMesh3.h"
-#include "DynamicMeshAttributeSet.h"
-#include "MeshNormals.h"
+#include "DynamicMesh/DynamicMesh3.h"
+#include "DynamicMesh/DynamicMeshAttributeSet.h"
+#include "DynamicMesh/MeshNormals.h"
 
 #include "Components/DynamicMeshComponent.h"
 
diff --git a/Engine/Plugins/Experimental/MeshModelingToolset/Source/MeshModelingTools/Private/Sculpting/MeshSculptUtil.h b/Engine/Plugins/Experimental/MeshModelingToolset/Source/MeshModelingTools/Private/Sculpting/MeshSculptUtil.h
--- a/Engine/Plugins/Experimental/MeshModelingToolset/Source/MeshModelingTools/Private/Sculpting/MeshSculptUtil.h
+++ b/Engine/Plugins/Experimental/MeshModelingToolset/Source/MeshModelingTools/Private/Sculpting/MeshSculptUtil.h
@@ -3,9 +3,9 @@
 #pragma once
 
 #include "CoreMinimal.h"
-#include "DynamicMeshAABBTree3.h"
-#include "DynamicMeshOctree3.h"
-#include "MeshNormals.h"
+#include "DynamicMesh/DynamicMeshAABBTree3.h"
+#include "DynamicMesh/DynamicMeshOctree3.h"
+#include "DynamicMesh/MeshNormals.h"
 #include "Util/UniqueIndexSet.h"
 
 
diff --git a/Engine/Plugins/Experimental/MeshModelingToolset/Source/MeshModelingTools/Private/SeamSculptTool.cpp b/Engine/Plugins/Experimental/MeshModelingToolset/Source/MeshModelingTools/Private/SeamSculptTool.cpp
--- a/Engine/Plugins/Experimental/MeshModelingToolset/Source/MeshModelingTools/Private/SeamSculptTool.cpp
+++ b/Engine/Plugins/Experimental/MeshModelingToolset/Source/MeshModelingTools/Private/SeamSculptTool.cpp
@@ -10,8 +10,8 @@
 #include "Changes/ToolCommandChangeSequence.h"
 #include "Changes/MeshChange.h"
 
-#include "DynamicMesh3.h"
-#include "DynamicMeshOverlay.h"
+#include "DynamicMesh/DynamicMesh3.h"
+#include "DynamicMesh/DynamicMeshOverlay.h"
 #include "DynamicMeshEditor.h"
 #include "Parameterization/DynamicMeshUVEditor.h"
 #include "Parameterization/MeshUVTransforms.h"
@@ -19,7 +19,7 @@
 #include "Selections/MeshConnectedComponents.h"
 #include "DynamicMesh/MeshNormals.h"
 #include "DynamicMesh/MeshIndexUtil.h"
-#include "DynamicMeshChangeTracker.h"
+#include "DynamicMesh/DynamicMeshChangeTracker.h"
 #include "DynamicMeshToMeshDescription.h"
 
 #include "TargetInterfaces/MeshDescriptionCommitter.h"
diff --git a/Engine/Plugins/Experimental/MeshModelingToolset/Source/MeshModelingTools/Private/SkinWeightsBindingTool.cpp b/Engine/Plugins/Experimental/MeshModelingToolset/Source/MeshModelingTools/Private/SkinWeightsBindingTool.cpp
--- a/Engine/Plugins/Experimental/MeshModelingToolset/Source/MeshModelingTools/Private/SkinWeightsBindingTool.cpp
+++ b/Engine/Plugins/Experimental/MeshModelingToolset/Source/MeshModelingTools/Private/SkinWeightsBindingTool.cpp
@@ -4,7 +4,7 @@
 
 #include "BoneWeights.h"
 #include "DynamicMeshToMeshDescription.h"
-#include "DynamicVertexSkinWeightsAttribute.h"
+#include "DynamicMesh/DynamicVertexSkinWeightsAttribute.h"
 #include "InteractiveToolManager.h"
 #include "MeshDescriptionToDynamicMesh.h"
 #include "MeshOpPreviewHelpers.h"
diff --git a/Engine/Plugins/Experimental/MeshModelingToolset/Source/MeshModelingTools/Private/SmoothMeshTool.cpp b/Engine/Plugins/Experimental/MeshModelingToolset/Source/MeshModelingTools/Private/SmoothMeshTool.cpp
--- a/Engine/Plugins/Experimental/MeshModelingToolset/Source/MeshModelingTools/Private/SmoothMeshTool.cpp
+++ b/Engine/Plugins/Experimental/MeshModelingToolset/Source/MeshModelingTools/Private/SmoothMeshTool.cpp
@@ -3,8 +3,8 @@
 #include "SmoothMeshTool.h"
 #include "InteractiveToolManager.h"
 #include "ToolBuilderUtil.h"
-#include "DynamicMesh3.h"
-#include "MeshNormals.h"
+#include "DynamicMesh/DynamicMesh3.h"
+#include "DynamicMesh/MeshNormals.h"
 
 // Smoothing operators
 #include "SmoothingOps/IterativeSmoothingOp.h"
diff --git a/Engine/Plugins/Experimental/MeshModelingToolset/Source/MeshModelingTools/Private/VoxelBlendMeshesTool.cpp b/Engine/Plugins/Experimental/MeshModelingToolset/Source/MeshModelingTools/Private/VoxelBlendMeshesTool.cpp
--- a/Engine/Plugins/Experimental/MeshModelingToolset/Source/MeshModelingTools/Private/VoxelBlendMeshesTool.cpp
+++ b/Engine/Plugins/Experimental/MeshModelingToolset/Source/MeshModelingTools/Private/VoxelBlendMeshesTool.cpp
@@ -9,8 +9,8 @@
 
 #include "Selection/ToolSelectionUtil.h"
 
-#include "DynamicMesh3.h"
-#include "MeshTransforms.h"
+#include "DynamicMesh/DynamicMesh3.h"
+#include "DynamicMesh/MeshTransforms.h"
 
 #include "MeshDescriptionToDynamicMesh.h"
 #include "DynamicMeshToMeshDescription.h"
diff --git a/Engine/Plugins/Experimental/MeshModelingToolset/Source/MeshModelingTools/Public/DynamicMeshSculptTool.h b/Engine/Plugins/Experimental/MeshModelingToolset/Source/MeshModelingTools/Public/DynamicMeshSculptTool.h
--- a/Engine/Plugins/Experimental/MeshModelingToolset/Source/MeshModelingTools/Public/DynamicMeshSculptTool.h
+++ b/Engine/Plugins/Experimental/MeshModelingToolset/Source/MeshModelingTools/Public/DynamicMeshSculptTool.h
@@ -8,9 +8,9 @@
 #include "UObject/NoExportTypes.h"
 #include "BaseTools/MeshSurfacePointTool.h"
 #include "Components/OctreeDynamicMeshComponent.h"
-#include "DynamicMeshAABBTree3.h"
-#include "DynamicMeshOctree3.h"
-#include "MeshNormals.h"
+#include "DynamicMesh/DynamicMeshAABBTree3.h"
+#include "DynamicMesh/DynamicMeshOctree3.h"
+#include "DynamicMesh/MeshNormals.h"
 #include "BaseTools/BaseBrushTool.h"
 #include "ToolDataVisualizer.h"
 #include "Changes/ValueWatcher.h"
diff --git a/Engine/Plugins/Experimental/MeshModelingToolset/Source/MeshModelingTools/Public/MeshAttributePaintTool.h b/Engine/Plugins/Experimental/MeshModelingToolset/Source/MeshModelingTools/Public/MeshAttributePaintTool.h
--- a/Engine/Plugins/Experimental/MeshModelingToolset/Source/MeshModelingTools/Public/MeshAttributePaintTool.h
+++ b/Engine/Plugins/Experimental/MeshModelingToolset/Source/MeshModelingTools/Public/MeshAttributePaintTool.h
@@ -5,7 +5,7 @@
 #include "DynamicMeshBrushTool.h"
 #include "Changes/ValueWatcher.h"
 #include "Changes/IndexedAttributeChange.h"
-#include "DynamicVerticesOctree3.h"
+#include "DynamicMesh/DynamicVerticesOctree3.h"
 #include "MeshDescription.h"
 #include "MeshAttributePaintTool.generated.h"
 
diff --git a/Engine/Plugins/Experimental/MeshModelingToolset/Source/MeshModelingTools/Public/MeshSelectionTool.h b/Engine/Plugins/Experimental/MeshModelingToolset/Source/MeshModelingTools/Public/MeshSelectionTool.h
--- a/Engine/Plugins/Experimental/MeshModelingToolset/Source/MeshModelingTools/Public/MeshSelectionTool.h
+++ b/Engine/Plugins/Experimental/MeshModelingToolset/Source/MeshModelingTools/Public/MeshSelectionTool.h
@@ -6,7 +6,7 @@
 #include "DynamicMeshBrushTool.h"
 #include "SelectionSet.h"
 #include "Changes/MeshSelectionChange.h"
-#include "DynamicMeshOctree3.h"
+#include "DynamicMesh/DynamicMeshOctree3.h"
 #include "MeshSelectionTool.generated.h"
 
 /**
diff --git a/Engine/Plugins/Experimental/MeshModelingToolset/Source/MeshModelingTools/Public/SkinWeightsPaintTool.h b/Engine/Plugins/Experimental/MeshModelingToolset/Source/MeshModelingTools/Public/SkinWeightsPaintTool.h
--- a/Engine/Plugins/Experimental/MeshModelingToolset/Source/MeshModelingTools/Public/SkinWeightsPaintTool.h
+++ b/Engine/Plugins/Experimental/MeshModelingToolset/Source/MeshModelingTools/Public/SkinWeightsPaintTool.h
@@ -5,7 +5,7 @@
 #include "DynamicMeshBrushTool.h"
 #include "Changes/ValueWatcher.h"
 #include "MeshDescription.h"
-#include "DynamicVerticesOctree3.h"
+#include "DynamicMesh/DynamicVerticesOctree3.h"
 #include "BoneContainer.h"
 #include "Interfaces/Interface_BoneReferenceSkeletonProvider.h"
 #include "Misc/Optional.h"
diff --git a/Engine/Plugins/Experimental/MeshModelingToolset/Source/MeshModelingToolsEditorOnly/Private/Hair/GroomCardsEditorTool.cpp b/Engine/Plugins/Experimental/MeshModelingToolset/Source/MeshModelingToolsEditorOnly/Private/Hair/GroomCardsEditorTool.cpp
--- a/Engine/Plugins/Experimental/MeshModelingToolset/Source/MeshModelingToolsEditorOnly/Private/Hair/GroomCardsEditorTool.cpp
+++ b/Engine/Plugins/Experimental/MeshModelingToolset/Source/MeshModelingToolsEditorOnly/Private/Hair/GroomCardsEditorTool.cpp
@@ -7,16 +7,16 @@
 #include "BaseBehaviors/SingleClickBehavior.h"
 #include "BaseBehaviors/MouseHoverBehavior.h"
 
-#include "DynamicMesh3.h"
-#include "DynamicMeshAABBTree3.h"
-#include "MeshNormals.h"
+#include "DynamicMesh/DynamicMesh3.h"
+#include "DynamicMesh/DynamicMeshAABBTree3.h"
+#include "DynamicMesh/MeshNormals.h"
 #include "DynamicMeshEditor.h"
 #include "Util/ColorConstants.h"
 #include "ToolSetupUtil.h"
 #include "Selection/ToolSelectionUtil.h"
-#include "MeshTransforms.h"
+#include "DynamicMesh/MeshTransforms.h"
 #include "DynamicMeshEditor.h"
-#include "DynamicMeshChangeTracker.h"
+#include "DynamicMesh/DynamicMeshChangeTracker.h"
 #include "Parameterization/DynamicMeshUVEditor.h"
 #include "Parameterization/MeshUVTransforms.h"
 #include "ModelingToolTargetUtil.h"
diff --git a/Engine/Plugins/Experimental/MeshModelingToolset/Source/MeshModelingToolsEditorOnly/Private/Hair/GroomToMeshTool.cpp b/Engine/Plugins/Experimental/MeshModelingToolset/Source/MeshModelingToolsEditorOnly/Private/Hair/GroomToMeshTool.cpp
--- a/Engine/Plugins/Experimental/MeshModelingToolset/Source/MeshModelingToolsEditorOnly/Private/Hair/GroomToMeshTool.cpp
+++ b/Engine/Plugins/Experimental/MeshModelingToolset/Source/MeshModelingToolsEditorOnly/Private/Hair/GroomToMeshTool.cpp
@@ -4,16 +4,16 @@
 #include "InteractiveToolManager.h"
 #include "ToolBuilderUtil.h"
 
-#include "DynamicMesh3.h"
-#include "DynamicMeshAABBTree3.h"
+#include "DynamicMesh/DynamicMesh3.h"
+#include "DynamicMesh/DynamicMeshAABBTree3.h"
 #include "Spatial/FastWinding.h"
-#include "MeshNormals.h"
+#include "DynamicMesh/MeshNormals.h"
 #include "DynamicMeshEditor.h"
 #include "Util/ColorConstants.h"
 #include "ToolSetupUtil.h"
 #include "Selection/ToolSelectionUtil.h"
 #include "ModelingObjectsCreationAPI.h"
-#include "MeshTransforms.h"
+#include "DynamicMesh/MeshTransforms.h"
 #include "Selections/MeshFaceSelection.h"
 #include "Selections/MeshVertexSelection.h"
 #include "Selections/MeshConnectedComponents.h"
@@ -22,10 +22,10 @@
 #include "Parameterization/DynamicMeshUVEditor.h"
 #include "Parameterization/MeshUVTransforms.h"
 
-#include "DynamicPointSet3.h"
-#include "DynamicVerticesOctree3.h"
-#include "DynamicAttribute.h"
-#include "DynamicVertexAttribute.h"
+#include "DynamicMesh/DynamicPointSet3.h"
+#include "DynamicMesh/DynamicVerticesOctree3.h"
+#include "DynamicMesh/DynamicAttribute.h"
+#include "DynamicMesh/DynamicVertexAttribute.h"
 
 #include "Implicit/CachingMeshSDF.h"
 #include "Implicit/GridInterpolant.h"
diff --git a/Engine/Plugins/Experimental/MeshModelingToolset/Source/MeshModelingToolsEditorOnly/Public/MeshTangentsTool.h b/Engine/Plugins/Experimental/MeshModelingToolset/Source/MeshModelingToolsEditorOnly/Public/MeshTangentsTool.h
--- a/Engine/Plugins/Experimental/MeshModelingToolset/Source/MeshModelingToolsEditorOnly/Public/MeshTangentsTool.h
+++ b/Engine/Plugins/Experimental/MeshModelingToolset/Source/MeshModelingToolsEditorOnly/Public/MeshTangentsTool.h
@@ -8,8 +8,8 @@
 #include "PreviewMesh.h"
 #include "Drawing/PreviewGeometryActor.h"
 #include "MeshOpPreviewHelpers.h"
-#include "DynamicMesh3.h"
-#include "MeshTangents.h"
+#include "DynamicMesh/DynamicMesh3.h"
+#include "DynamicMesh/MeshTangents.h"
 #include "ParameterizationOps/CalculateTangentsOp.h"
 #include "BaseTools/SingleSelectionMeshEditingTool.h"
 #include "MeshTangentsTool.generated.h"
diff --git a/Engine/Plugins/Experimental/MeshModelingToolset/Source/ModelingComponents/Private/BaseTools/BaseCreateFromSelectedTool.cpp b/Engine/Plugins/Experimental/MeshModelingToolset/Source/ModelingComponents/Private/BaseTools/BaseCreateFromSelectedTool.cpp
--- a/Engine/Plugins/Experimental/MeshModelingToolset/Source/ModelingComponents/Private/BaseTools/BaseCreateFromSelectedTool.cpp
+++ b/Engine/Plugins/Experimental/MeshModelingToolset/Source/ModelingComponents/Private/BaseTools/BaseCreateFromSelectedTool.cpp
@@ -8,8 +8,8 @@
 #include "BaseGizmos/TransformGizmoUtil.h"
 #include "Async/Async.h"
 
-#include "MeshNormals.h"
-#include "MeshTransforms.h"
+#include "DynamicMesh/MeshNormals.h"
+#include "DynamicMesh/MeshTransforms.h"
 #include "DynamicMeshToMeshDescription.h"
 #include "MeshDescriptionToDynamicMesh.h"
 
diff --git a/Engine/Plugins/Experimental/MeshModelingToolset/Source/ModelingComponents/Private/BaseTools/BaseMeshProcessingTool.cpp b/Engine/Plugins/Experimental/MeshModelingToolset/Source/ModelingComponents/Private/BaseTools/BaseMeshProcessingTool.cpp
--- a/Engine/Plugins/Experimental/MeshModelingToolset/Source/ModelingComponents/Private/BaseTools/BaseMeshProcessingTool.cpp
+++ b/Engine/Plugins/Experimental/MeshModelingToolset/Source/ModelingComponents/Private/BaseTools/BaseMeshProcessingTool.cpp
@@ -7,9 +7,9 @@
 #include "Components/DynamicMeshComponent.h"
 #include "Async/Async.h"
 
-#include "MeshNormals.h"
+#include "DynamicMesh/MeshNormals.h"
 #include "MeshBoundaryLoops.h"
-#include "MeshTransforms.h"
+#include "DynamicMesh/MeshTransforms.h"
 #include "WeightMapUtil.h"
 #include "DynamicMeshToMeshDescription.h"
 #include "MeshDescriptionToDynamicMesh.h"
diff --git a/Engine/Plugins/Experimental/MeshModelingToolset/Source/ModelingComponents/Private/Components/DynamicMeshComponent.cpp b/Engine/Plugins/Experimental/MeshModelingToolset/Source/ModelingComponents/Private/Components/DynamicMeshComponent.cpp
--- a/Engine/Plugins/Experimental/MeshModelingToolset/Source/ModelingComponents/Private/Components/DynamicMeshComponent.cpp
+++ b/Engine/Plugins/Experimental/MeshModelingToolset/Source/ModelingComponents/Private/Components/DynamicMeshComponent.cpp
@@ -18,14 +18,14 @@
 #include "StaticMeshAttributes.h"
 #include "Async/Async.h"
 
-#include "DynamicMeshAttributeSet.h"
-#include "MeshNormals.h"
+#include "DynamicMesh/DynamicMeshAttributeSet.h"
+#include "DynamicMesh/MeshNormals.h"
 #include "MeshDescriptionToDynamicMesh.h"
 
 #include "Changes/MeshVertexChange.h"
 #include "Changes/MeshChange.h"
-#include "DynamicMeshChangeTracker.h"
-#include "MeshTransforms.h"
+#include "DynamicMesh/DynamicMeshChangeTracker.h"
+#include "DynamicMesh/MeshTransforms.h"
 
 // default proxy for this component
 #include "Components/DynamicMeshSceneProxy.h"
diff --git a/Engine/Plugins/Experimental/MeshModelingToolset/Source/ModelingComponents/Private/Components/OctreeDynamicMeshComponent.cpp b/Engine/Plugins/Experimental/MeshModelingToolset/Source/ModelingComponents/Private/Components/OctreeDynamicMeshComponent.cpp
--- a/Engine/Plugins/Experimental/MeshModelingToolset/Source/ModelingComponents/Private/Components/OctreeDynamicMeshComponent.cpp
+++ b/Engine/Plugins/Experimental/MeshModelingToolset/Source/ModelingComponents/Private/Components/OctreeDynamicMeshComponent.cpp
@@ -17,14 +17,14 @@
 #include "StaticMeshResources.h"
 #include "StaticMeshAttributes.h"
 
-#include "DynamicMeshAttributeSet.h"
-#include "MeshNormals.h"
-#include "MeshTransforms.h"
+#include "DynamicMesh/DynamicMeshAttributeSet.h"
+#include "DynamicMesh/MeshNormals.h"
+#include "DynamicMesh/MeshTransforms.h"
 #include "MeshDescriptionToDynamicMesh.h"
 
 #include "Changes/MeshVertexChange.h"
 #include "Changes/MeshChange.h"
-#include "DynamicMeshChangeTracker.h"
+#include "DynamicMesh/DynamicMeshChangeTracker.h"
 
 
 // default proxy for this component
diff --git a/Engine/Plugins/Experimental/MeshModelingToolset/Source/ModelingComponents/Private/Drawing/PolyEditPreviewMesh.cpp b/Engine/Plugins/Experimental/MeshModelingToolset/Source/ModelingComponents/Private/Drawing/PolyEditPreviewMesh.cpp
--- a/Engine/Plugins/Experimental/MeshModelingToolset/Source/ModelingComponents/Private/Drawing/PolyEditPreviewMesh.cpp
+++ b/Engine/Plugins/Experimental/MeshModelingToolset/Source/ModelingComponents/Private/Drawing/PolyEditPreviewMesh.cpp
@@ -2,12 +2,12 @@
 
 #include "Drawing/PolyEditPreviewMesh.h"
 #include "DynamicSubmesh3.h"
-#include "MeshTransforms.h"
+#include "DynamicMesh/MeshTransforms.h"
 #include "Operations/ExtrudeMesh.h"
 #include "Operations/InsetMeshRegion.h"
 #include "Selections/MeshVertexSelection.h"
 #include "MeshBoundaryLoops.h"
-#include "MeshNormals.h"
+#include "DynamicMesh/MeshNormals.h"
 #include "DynamicMeshEditor.h"
 
 using namespace UE::Geometry;
diff --git a/Engine/Plugins/Experimental/MeshModelingToolset/Source/ModelingComponents/Public/Components/BaseDynamicMeshComponent.h b/Engine/Plugins/Experimental/MeshModelingToolset/Source/ModelingComponents/Public/Components/BaseDynamicMeshComponent.h
--- a/Engine/Plugins/Experimental/MeshModelingToolset/Source/ModelingComponents/Public/Components/BaseDynamicMeshComponent.h
+++ b/Engine/Plugins/Experimental/MeshModelingToolset/Source/ModelingComponents/Public/Components/BaseDynamicMeshComponent.h
@@ -9,7 +9,7 @@
 #include "Changes/MeshChange.h"
 #include "Changes/MeshReplacementChange.h"
 #include "MeshConversionOptions.h"
-#include "DynamicMesh3.h"		// todo replace with predeclaration (lots of fallout)
+#include "DynamicMesh/DynamicMesh3.h"		// todo replace with predeclaration (lots of fallout)
 #include "UDynamicMesh.h"
 
 #include "BaseDynamicMeshComponent.generated.h"
diff --git a/Engine/Plugins/Experimental/MeshModelingToolset/Source/ModelingComponents/Public/Components/OctreeDynamicMeshComponent.h b/Engine/Plugins/Experimental/MeshModelingToolset/Source/ModelingComponents/Public/Components/OctreeDynamicMeshComponent.h
--- a/Engine/Plugins/Experimental/MeshModelingToolset/Source/ModelingComponents/Public/Components/OctreeDynamicMeshComponent.h
+++ b/Engine/Plugins/Experimental/MeshModelingToolset/Source/ModelingComponents/Public/Components/OctreeDynamicMeshComponent.h
@@ -6,7 +6,7 @@
 #include "BaseDynamicMeshComponent.h"
 #include "MeshConversionOptions.h"
 #include "UDynamicMesh.h"
-#include "DynamicMeshOctree3.h"
+#include "DynamicMesh/DynamicMeshOctree3.h"
 #include "Util/IndexSetDecompositions.h"
 
 #include "OctreeDynamicMeshComponent.generated.h"
diff --git a/Engine/Plugins/Experimental/MeshModelingToolset/Source/ModelingOperators/Private/CleaningOps/EditNormalsOp.cpp b/Engine/Plugins/Experimental/MeshModelingToolset/Source/ModelingOperators/Private/CleaningOps/EditNormalsOp.cpp
--- a/Engine/Plugins/Experimental/MeshModelingToolset/Source/ModelingOperators/Private/CleaningOps/EditNormalsOp.cpp
+++ b/Engine/Plugins/Experimental/MeshModelingToolset/Source/ModelingOperators/Private/CleaningOps/EditNormalsOp.cpp
@@ -8,9 +8,9 @@
 #include "MeshSimplification.h"
 #include "MeshConstraintsUtil.h"
 
-#include "MeshNormals.h"
+#include "DynamicMesh/MeshNormals.h"
 #include "Operations/RepairOrientation.h"
-#include "DynamicMeshAABBTree3.h"
+#include "DynamicMesh/DynamicMeshAABBTree3.h"
 
 #include "ExplicitUseGeometryMathTypes.h"		// using UE::Geometry::(math types)
 using namespace UE::Geometry;
diff --git a/Engine/Plugins/Experimental/MeshModelingToolset/Source/ModelingOperators/Private/CleaningOps/RemeshMeshOp.cpp b/Engine/Plugins/Experimental/MeshModelingToolset/Source/ModelingOperators/Private/CleaningOps/RemeshMeshOp.cpp
--- a/Engine/Plugins/Experimental/MeshModelingToolset/Source/ModelingOperators/Private/CleaningOps/RemeshMeshOp.cpp
+++ b/Engine/Plugins/Experimental/MeshModelingToolset/Source/ModelingOperators/Private/CleaningOps/RemeshMeshOp.cpp
@@ -2,14 +2,14 @@
 
 #include "CleaningOps/RemeshMeshOp.h"
 
-#include "DynamicMesh3.h"
-#include "DynamicMeshAABBTree3.h"
-#include "DynamicMeshAttributeSet.h"
+#include "DynamicMesh/DynamicMesh3.h"
+#include "DynamicMesh/DynamicMeshAABBTree3.h"
+#include "DynamicMesh/DynamicMeshAttributeSet.h"
 #include "Remesher.h"
 #include "QueueRemesher.h"
 #include "MeshConstraintsUtil.h"
 #include "ProjectionTargets.h"
-#include "MeshNormals.h"
+#include "DynamicMesh/MeshNormals.h"
 #include "NormalFlowRemesher.h"
 
 #include "ExplicitUseGeometryMathTypes.h"		// using UE::Geometry::(math types)
diff --git a/Engine/Plugins/Experimental/MeshModelingToolset/Source/ModelingOperators/Private/CompositionOps/CurveSweepOp.cpp b/Engine/Plugins/Experimental/MeshModelingToolset/Source/ModelingOperators/Private/CompositionOps/CurveSweepOp.cpp
--- a/Engine/Plugins/Experimental/MeshModelingToolset/Source/ModelingOperators/Private/CompositionOps/CurveSweepOp.cpp
+++ b/Engine/Plugins/Experimental/MeshModelingToolset/Source/ModelingOperators/Private/CompositionOps/CurveSweepOp.cpp
@@ -5,8 +5,8 @@
 #include "ConstrainedDelaunay2.h"
 #include "DynamicMeshEditor.h"
 #include "MeshBoundaryLoops.h"
-#include "MeshNormals.h"
-#include "MeshTransforms.h"
+#include "DynamicMesh/MeshNormals.h"
+#include "DynamicMesh/MeshTransforms.h"
 #include "Operations/SimpleHoleFiller.h"
 #include "Operations/PlanarHoleFiller.h"
 
diff --git a/Engine/Plugins/Experimental/MeshModelingToolset/Source/ModelingOperators/Private/CompositionOps/VoxelBlendMeshesOp.cpp b/Engine/Plugins/Experimental/MeshModelingToolset/Source/ModelingOperators/Private/CompositionOps/VoxelBlendMeshesOp.cpp
--- a/Engine/Plugins/Experimental/MeshModelingToolset/Source/ModelingOperators/Private/CompositionOps/VoxelBlendMeshesOp.cpp
+++ b/Engine/Plugins/Experimental/MeshModelingToolset/Source/ModelingOperators/Private/CompositionOps/VoxelBlendMeshesOp.cpp
@@ -3,11 +3,11 @@
 #include "CompositionOps/VoxelBlendMeshesOp.h"
 #include "CleaningOps/EditNormalsOp.h"
 
-#include "DynamicMeshAABBTree3.h"
+#include "DynamicMesh/DynamicMeshAABBTree3.h"
 #include "DynamicMeshEditor.h"
-#include "MeshTransforms.h"
+#include "DynamicMesh/MeshTransforms.h"
 #include "MeshSimplification.h"
-#include "MeshNormals.h"
+#include "DynamicMesh/MeshNormals.h"
 #include "Operations/ExtrudeMesh.h"
 #include "Operations/RemoveOccludedTriangles.h"
 
diff --git a/Engine/Plugins/Experimental/MeshModelingToolset/Source/ModelingOperators/Private/CompositionOps/VoxelMorphologyMeshesOp.cpp b/Engine/Plugins/Experimental/MeshModelingToolset/Source/ModelingOperators/Private/CompositionOps/VoxelMorphologyMeshesOp.cpp
--- a/Engine/Plugins/Experimental/MeshModelingToolset/Source/ModelingOperators/Private/CompositionOps/VoxelMorphologyMeshesOp.cpp
+++ b/Engine/Plugins/Experimental/MeshModelingToolset/Source/ModelingOperators/Private/CompositionOps/VoxelMorphologyMeshesOp.cpp
@@ -4,8 +4,8 @@
 #include "CleaningOps/EditNormalsOp.h"
 
 #include "DynamicMeshEditor.h"
-#include "DynamicMeshAABBTree3.h"
-#include "MeshTransforms.h"
+#include "DynamicMesh/DynamicMeshAABBTree3.h"
+#include "DynamicMesh/MeshTransforms.h"
 #include "MeshSimplification.h"
 #include "Operations/RemoveOccludedTriangles.h"
 #include "Operations/ExtrudeMesh.h"
diff --git a/Engine/Plugins/Experimental/MeshModelingToolset/Source/ModelingOperators/Private/CompositionOps/VoxelSolidifyMeshesOp.cpp b/Engine/Plugins/Experimental/MeshModelingToolset/Source/ModelingOperators/Private/CompositionOps/VoxelSolidifyMeshesOp.cpp
--- a/Engine/Plugins/Experimental/MeshModelingToolset/Source/ModelingOperators/Private/CompositionOps/VoxelSolidifyMeshesOp.cpp
+++ b/Engine/Plugins/Experimental/MeshModelingToolset/Source/ModelingOperators/Private/CompositionOps/VoxelSolidifyMeshesOp.cpp
@@ -2,14 +2,14 @@
 
 #include "CompositionOps/VoxelSolidifyMeshesOp.h"
 
-#include "DynamicMeshAABBTree3.h"
+#include "DynamicMesh/DynamicMeshAABBTree3.h"
 #include "DynamicMeshEditor.h"
-#include "MeshTransforms.h"
+#include "DynamicMesh/MeshTransforms.h"
 #include "MeshSimplification.h"
 #include "Operations/ExtrudeMesh.h"
 #include "Spatial/FastWinding.h"
 #include "Generators/MarchingCubes.h"
-#include "MeshNormals.h"
+#include "DynamicMesh/MeshNormals.h"
 
 #include "Implicit/Solidify.h"
 
diff --git a/Engine/Plugins/Experimental/MeshModelingToolset/Source/ModelingOperatorsEditorOnly/Private/ParameterizationOps/ParameterizeMeshOp.cpp b/Engine/Plugins/Experimental/MeshModelingToolset/Source/ModelingOperatorsEditorOnly/Private/ParameterizationOps/ParameterizeMeshOp.cpp
--- a/Engine/Plugins/Experimental/MeshModelingToolset/Source/ModelingOperatorsEditorOnly/Private/ParameterizationOps/ParameterizeMeshOp.cpp
+++ b/Engine/Plugins/Experimental/MeshModelingToolset/Source/ModelingOperatorsEditorOnly/Private/ParameterizationOps/ParameterizeMeshOp.cpp
@@ -2,14 +2,14 @@
 
 #include "ParameterizationOps/ParameterizeMeshOp.h"
 
-#include "DynamicMeshAttributeSet.h"
+#include "DynamicMesh/DynamicMeshAttributeSet.h"
 #include "ProxyLODParameterization.h"
 #include "Selections/MeshConnectedComponents.h"
 
 #include "Parameterization/MeshLocalParam.h"
 #include "Parameterization/DynamicMeshUVEditor.h"
-#include "DynamicMeshAABBTree3.h"
-#include "MeshNormals.h"
+#include "DynamicMesh/DynamicMeshAABBTree3.h"
+#include "DynamicMesh/MeshNormals.h"
 #include "DynamicSubmesh3.h"
 #include "XAtlasWrapper.h"
 
diff --git a/Engine/Plugins/Experimental/PlanarCutPlugin/Source/PlanarCut/Private/GeometryMeshConversion.cpp b/Engine/Plugins/Experimental/PlanarCutPlugin/Source/PlanarCut/Private/GeometryMeshConversion.cpp
--- a/Engine/Plugins/Experimental/PlanarCutPlugin/Source/PlanarCut/Private/GeometryMeshConversion.cpp
+++ b/Engine/Plugins/Experimental/PlanarCutPlugin/Source/PlanarCut/Private/GeometryMeshConversion.cpp
@@ -25,13 +25,13 @@
 #include "DynamicMeshEditor.h"
 #include "DynamicMesh/DynamicMeshAABBTree3.h"
 #include "Selections/MeshConnectedComponents.h"
-#include "MeshTransforms.h"
+#include "DynamicMesh/MeshTransforms.h"
 #include "Operations/MeshBoolean.h"
 #include "Operations/MeshSelfUnion.h"
 #include "DynamicMesh/Operations/MergeCoincidentMeshEdges.h"
 #include "MeshBoundaryLoops.h"
 #include "QueueRemesher.h"
-#include "DynamicVertexAttribute.h"
+#include "DynamicMesh/DynamicVertexAttribute.h"
 #include "DynamicMesh/MeshNormals.h"
 #include "DynamicMesh/MeshTangents.h"
 #include "ConstrainedDelaunay2.h"
diff --git a/Engine/Plugins/Experimental/PlanarCutPlugin/Source/PlanarCut/Private/GeometryMeshConversion.h b/Engine/Plugins/Experimental/PlanarCutPlugin/Source/PlanarCut/Private/GeometryMeshConversion.h
--- a/Engine/Plugins/Experimental/PlanarCutPlugin/Source/PlanarCut/Private/GeometryMeshConversion.h
+++ b/Engine/Plugins/Experimental/PlanarCutPlugin/Source/PlanarCut/Private/GeometryMeshConversion.h
@@ -2,11 +2,11 @@
 #pragma once
 #include "CoreMinimal.h"
 
-#include "DynamicMesh3.h"
+#include "DynamicMesh/DynamicMesh3.h"
 #include "DynamicMeshEditor.h"
-#include "DynamicMeshAABBTree3.h"
+#include "DynamicMesh/DynamicMeshAABBTree3.h"
 #include "Selections/MeshConnectedComponents.h"
-#include "MeshTransforms.h"
+#include "DynamicMesh/MeshTransforms.h"
 #include "MeshConstraints.h"
 #include "GeometryCollection/GeometryCollection.h"
 #include "Spatial/PointHashGrid3.h"
diff --git a/Engine/Plugins/Experimental/PlanarCutPlugin/Source/PlanarCut/Private/PlanarCut.cpp b/Engine/Plugins/Experimental/PlanarCutPlugin/Source/PlanarCut/Private/PlanarCut.cpp
--- a/Engine/Plugins/Experimental/PlanarCutPlugin/Source/PlanarCut/Private/PlanarCut.cpp
+++ b/Engine/Plugins/Experimental/PlanarCutPlugin/Source/PlanarCut/Private/PlanarCut.cpp
@@ -22,13 +22,13 @@
 #include "DynamicMeshEditor.h"
 #include "DynamicMesh/DynamicMeshAABBTree3.h"
 #include "Selections/MeshConnectedComponents.h"
-#include "MeshTransforms.h"
+#include "DynamicMesh/MeshTransforms.h"
 #include "Operations/MeshBoolean.h"
 #include "Operations/MeshSelfUnion.h"
 #include "DynamicMesh/Operations/MergeCoincidentMeshEdges.h"
 #include "MeshBoundaryLoops.h"
 #include "QueueRemesher.h"
-#include "DynamicVertexAttribute.h"
+#include "DynamicMesh/DynamicVertexAttribute.h"
 #include "DynamicMesh/MeshNormals.h"
 #include "DynamicMesh/MeshTangents.h"
 #include "ConstrainedDelaunay2.h"
