ModelingComponents: static analysis fix
#rnx

diff --git a/Engine/Plugins/Experimental/MeshModelingToolset/Source/ModelingComponents/Private/Drawing/MeshRenderDecomposition.cpp b/Engine/Plugins/Experimental/MeshModelingToolset/Source/ModelingComponents/Private/Drawing/MeshRenderDecomposition.cpp
--- a/Engine/Plugins/Experimental/MeshModelingToolset/Source/ModelingComponents/Private/Drawing/MeshRenderDecomposition.cpp
+++ b/Engine/Plugins/Experimental/MeshModelingToolset/Source/ModelingComponents/Private/Drawing/MeshRenderDecomposition.cpp
@@ -39,12 +39,12 @@ void FMeshRenderDecomposition::BuildMaterialDecomposition(const FDynamicMesh3* M
 	const FDynamicMeshMaterialAttribute* MaterialID = Mesh->Attributes()->GetMaterialID();
 	ensure(MaterialID != nullptr);
 
-	int32 NumMaterials = MaterialSet->Materials.Num();
+	int32 NumMaterials = UseMaterials->Materials.Num();
 	Decomp.Initialize(NumMaterials);
 	for (int32 k = 0; k < NumMaterials; ++k)
 	{
 		FMeshRenderDecomposition::FGroup& Group = Decomp.GetGroup(k);
-		Group.Material = MaterialSet->Materials[k];
+		Group.Material = UseMaterials->Materials[k];
 	}
 
 	for (int32 tid : Mesh->TriangleIndicesItr())
