ModelingTools: fix incorrect handling of polygroups on Accept in EditMeshPolygonsTool
#rnx

diff --git a/Engine/Plugins/Experimental/MeshModelingToolset/Source/MeshModelingTools/Private/EditMeshPolygonsTool.cpp b/Engine/Plugins/Experimental/MeshModelingToolset/Source/MeshModelingTools/Private/EditMeshPolygonsTool.cpp
--- a/Engine/Plugins/Experimental/MeshModelingToolset/Source/MeshModelingTools/Private/EditMeshPolygonsTool.cpp
+++ b/Engine/Plugins/Experimental/MeshModelingToolset/Source/MeshModelingTools/Private/EditMeshPolygonsTool.cpp
@@ -448,8 +448,7 @@ void UEditMeshPolygonsTool::Shutdown(EToolShutdownType ShutdownType)
 			{
 				FConversionToMeshDescriptionOptions ConversionOptions;
 				bool bModifiedTopology = (ModifiedTopologyCounter > 0);
-				ConversionOptions.bSetPolyGroups = bModifiedTopology;
-				UE::ToolTarget::CommitDynamicMeshUpdate(Target, ReadMesh, true, ConversionOptions);
+				UE::ToolTarget::CommitDynamicMeshUpdate(Target, ReadMesh, bModifiedTopology, ConversionOptions);
 			});
 
 			// The stored selection change should go into this transaction as well.
