ToolsFramework: Add UTransformProxy::AddComponentCustom(), allows adding a target Component with custom transform get/set functions. Can be used to (eg) wrap an ISMC instance.
ModelingTools: add support for ISMC instances in EditPivotTool and TransformMeshesTool. In both cases the Tool will operate on the Instances by default, rather than the Component.  Clean up TransformMeshesTool properties panel, add grid and rotation snapping.

#rnx

diff --git a/Engine/Plugins/Experimental/GizmoEdMode/Source/LightGizmos/Private/SubTransformProxy.cpp b/Engine/Plugins/Experimental/GizmoEdMode/Source/LightGizmos/Private/SubTransformProxy.cpp
--- a/Engine/Plugins/Experimental/GizmoEdMode/Source/LightGizmos/Private/SubTransformProxy.cpp
+++ b/Engine/Plugins/Experimental/GizmoEdMode/Source/LightGizmos/Private/SubTransformProxy.cpp
@@ -16,12 +16,11 @@ FTransform USubTransformProxy::AddSubTransformProxy(USubTransformProxy* InProxy,
 	// Add each object to the proxy
 	for (FRelativeObject& Object : InProxy->Objects)
 	{
-		FRelativeObject NewObj;
+		FRelativeObject& NewObj = Objects.Emplace_GetRef();
 		NewObj.Component = Object.Component;
 		NewObj.bModifyComponentOnTransform = Object.bModifyComponentOnTransform;
 		NewObj.StartTransform = Object.Component->GetComponentToWorld();
 		NewObj.RelativeTransform = FTransform::Identity;
-		Objects.Add(NewObj);
 	}
 
 	UpdateSharedTransform();
diff --git a/Engine/Plugins/Experimental/MeshModelingToolset/Source/MeshModelingTools/Private/EditPivotTool.cpp b/Engine/Plugins/Experimental/MeshModelingToolset/Source/MeshModelingTools/Private/EditPivotTool.cpp
--- a/Engine/Plugins/Experimental/MeshModelingToolset/Source/MeshModelingTools/Private/EditPivotTool.cpp
+++ b/Engine/Plugins/Experimental/MeshModelingToolset/Source/MeshModelingTools/Private/EditPivotTool.cpp
@@ -6,9 +6,7 @@
 #include "ToolBuilderUtil.h"
 #include "ToolSetupUtil.h"
 #include "DynamicMesh/DynamicMesh3.h"
-#include "DynamicMeshToMeshDescription.h"
 #include "Mechanics/DragAlignmentMechanic.h"
-#include "MeshDescriptionToDynamicMesh.h"
 #include "MeshAdapterTransforms.h"
 #include "MeshDescriptionAdapter.h"
 #include "DynamicMesh/MeshTransforms.h"
@@ -20,6 +18,7 @@
 #include "BaseGizmos/TransformGizmoUtil.h"
 
 #include "Components/PrimitiveComponent.h"
+#include "Components/InstancedStaticMeshComponent.h"
 #include "Engine/World.h"
 
 #include "TargetInterfaces/MeshDescriptionCommitter.h"
@@ -125,7 +124,20 @@ void UEditPivotTool::Setup()
 	bool bSharesSources = GetMapToSharedSourceData(MapToFirstOccurrences);
 	if (bSharesSources)
 	{
-		AllTheWarnings = FText::Format(FTextFormat::FromString("{0}\n\n{1}"), AllTheWarnings, LOCTEXT("EditPivotSharedAssetsWarning", "WARNING: Multiple meshes in your selection use the same source asset!  This is not supported -- each asset can only have one baked pivot."));
+		AllTheWarnings = FText::Format(FTextFormat::FromString("{0}\n\n{1}"), AllTheWarnings, LOCTEXT("EditPivotSharedAssetsWarning", "WARNING: Multiple selected meshes share the same source Asset! Each Asset can only have one baked pivot, some results will be incorrect."));
+	}
+
+	bool bHasISMCs = false;
+	for (int32 k = 0; k < Targets.Num(); ++k)
+	{
+		if (Cast<UInstancedStaticMeshComponent>(TargetComponentInterface(k)->GetOwnerComponent()) != nullptr)
+		{
+			bHasISMCs = true;
+		}
+	}
+	if (bHasISMCs)
+	{
+		AllTheWarnings = FText::Format(FTextFormat::FromString("{0}\n\n{1}"), AllTheWarnings, LOCTEXT("EditPivotISMCWarning", "WARNING: Some selected objects are Instanced Components. Pivot of Instances will be modified, instead of Asset."));
 	}
 
 	GetToolManager()->DisplayMessage(AllTheWarnings, EToolMessageLevel::UserWarning);
@@ -430,12 +442,37 @@ void UEditPivotTool::UpdateAssets(const FFrame3d& NewPivotWorldFrame)
 	for (int32 ComponentIdx = 0; ComponentIdx < Targets.Num(); ComponentIdx++)
 	{
 		IPrimitiveComponentBackedTarget* TargetComponent = TargetComponentInterface(ComponentIdx);
-		if (MapToFirstOccurrences[ComponentIdx] == ComponentIdx)
+		UPrimitiveComponent* Component = TargetComponent->GetOwnerComponent();
+		Component->Modify();
+
+		UInstancedStaticMeshComponent* InstancedComponent = Cast<UInstancedStaticMeshComponent>(TargetComponent->GetOwnerComponent());
+		if (InstancedComponent != nullptr)
 		{
-			UE::Geometry::FTransform3d ToBake(OriginalTransforms[ComponentIdx] * NewWorldInverse);
+			// For ISMC, we will not bake in a mesh transform, instead we will update the instance transforms relative to the new pivot
+			// TODO: this could be optional, and another alternative would be to bake the pivot to the mesh and then update
+			//   all the instance transforms so they stay in the same position?
+
+			// save world transforms
+			int32 NumInstances = InstancedComponent->GetInstanceCount();
+			TArray<FTransform> WorldTransforms;
+			WorldTransforms.SetNum(NumInstances);
+			for (int32 k = 0; k < NumInstances; ++k)
+			{
+				InstancedComponent->GetInstanceTransform(k, WorldTransforms[k], true);
+			}
 
-			UPrimitiveComponent* Component = TargetComponent->GetOwnerComponent();
-			Component->Modify();
+			// update position to new pivot
+			InstancedComponent->SetWorldTransform(NewWorldTransform);
+
+			// restore world transforms, which will compute new local transforms such that the instances do not move in the world
+			for (int32 k = 0; k < NumInstances; ++k)
+			{
+				InstancedComponent->UpdateInstanceTransform(k, WorldTransforms[k], true, true, false);
+			}
+		}
+		else if (MapToFirstOccurrences[ComponentIdx] == ComponentIdx)
+		{
+			UE::Geometry::FTransform3d ToBake(OriginalTransforms[ComponentIdx] * NewWorldInverse);
 
 			// transform simple collision geometry
 			UE::Geometry::TransformSimpleCollision(Component, ToBake);
@@ -450,13 +487,12 @@ void UEditPivotTool::UpdateAssets(const FFrame3d& NewPivotWorldFrame)
 		}
 		else
 		{
-			UPrimitiveComponent* Component = TargetComponent->GetOwnerComponent();
-			Component->Modify();
 			// try to invert baked transform
 			FTransform Baked = OriginalTransforms[MapToFirstOccurrences[ComponentIdx]] * NewWorldInverse;
 			Component->SetWorldTransform(Baked.Inverse() * OriginalTransforms[ComponentIdx]);
 		}
 		TargetComponent->GetOwnerActor()->MarkComponentsRenderStateDirty();
+		TargetComponent->GetOwnerActor()->UpdateComponentTransforms();
 	}
 
 	// hack to ensure user sees the updated pivot immediately: request re-select of the original selection
diff --git a/Engine/Plugins/Experimental/MeshModelingToolset/Source/MeshModelingTools/Private/TransformMeshesTool.cpp b/Engine/Plugins/Experimental/MeshModelingToolset/Source/MeshModelingTools/Private/TransformMeshesTool.cpp
--- a/Engine/Plugins/Experimental/MeshModelingToolset/Source/MeshModelingTools/Private/TransformMeshesTool.cpp
+++ b/Engine/Plugins/Experimental/MeshModelingToolset/Source/MeshModelingTools/Private/TransformMeshesTool.cpp
@@ -9,11 +9,13 @@
 #include "DynamicMesh/DynamicMesh3.h"
 #include "BaseBehaviors/ClickDragBehavior.h"
 #include "ToolSceneQueriesUtil.h"
+#include "ModelingToolTargetUtil.h"
 
 #include "BaseGizmos/GizmoComponents.h"
 #include "BaseGizmos/TransformGizmoUtil.h"
 
 #include "Components/PrimitiveComponent.h"
+#include "Components/InstancedStaticMeshComponent.h"
 #include "Engine/World.h"
 
 #include "TargetInterfaces/PrimitiveComponentBackedTarget.h"
@@ -86,12 +88,28 @@ void UTransformMeshesTool::Setup()
 
 	TransformProps = NewObject<UTransformMeshesToolProperties>();
 	AddToolPropertySource(TransformProps);
+	TransformProps->RestoreProperties(this);
+	TransformProps->WatchProperty(TransformProps->TransformMode, [this](ETransformMeshesTransformMode NewMode) { UpdateTransformMode(NewMode); });
+	TransformProps->WatchProperty(TransformProps->bApplyToInstances, [this](bool bNewValue) { UpdateTransformMode(TransformProps->TransformMode); });
+	TransformProps->WatchProperty(TransformProps->bSetPivotMode, [this](bool bNewValue) { UpdateSetPivotModes(bNewValue); });
+	TransformProps->WatchProperty(TransformProps->bSnapToGrid, [this](bool bNewValue) { OnSnappingSettingsUpdated(); });
+	TransformProps->WatchProperty(TransformProps->bSnapToRotation, [this](bool bNewValue) { OnSnappingSettingsUpdated(); });
+	
+	// determine if we have any ISMCs
+	TransformProps->bHaveInstances = false;
+	for (int32 ComponentIdx = 0; ComponentIdx < Targets.Num(); ComponentIdx++)
+	{
+		if (Cast<UInstancedStaticMeshComponent>(UE::ToolTarget::GetTargetComponent(Targets[ComponentIdx])) != nullptr)
+		{
+			TransformProps->bHaveInstances = true;
+		}
+	}
 
 	UpdateTransformMode(TransformProps->TransformMode);
 
 	SetToolDisplayName(LOCTEXT("ToolName", "Transform"));
 	GetToolManager()->DisplayMessage(
-		LOCTEXT("OnStartTransformMeshesTool", "Translate/Rotate/Scale the selected objects. Use middle mouse button on gizmo to reposition it. Hold Ctrl while translating or (in local mode) rotating to align. [A] cycles through Transform modes. [S] toggles Set Pivot Mode. [D] Toggles Snap Drag Mode. [W] and [E] cycle through Snap Drag Source and Rotation types."),
+		LOCTEXT("OnStartTransformMeshesTool", "Transform the selected objects. Middle-mouse-drag on gizmo to reposition it. Hold Ctrl while dragging to snap/align. [A] cycles through Transform modes. [S] toggles Set Pivot Mode. [D] Toggles Snap Drag Mode. [W] and [E] cycle through Snap Drag Source and Rotation types."),
 		EToolMessageLevel::UserNotification);
 }
 
@@ -103,6 +121,8 @@ void UTransformMeshesTool::Setup()
 
 void UTransformMeshesTool::Shutdown(EToolShutdownType ShutdownType)
 {
+	TransformProps->SaveProperties(this);
+
 	DragAlignmentMechanic->Shutdown();
 
 	GizmoManager->DestroyAllGizmosByOwner(this);
@@ -110,32 +130,12 @@ void UTransformMeshesTool::Shutdown(EToolShutdownType ShutdownType)
 
 
 
-void UTransformMeshesTool::OnTick(float DeltaTime)
-{
-	// make sure state is up-to-date
-
-	if (CurTransformMode != TransformProps->TransformMode)
-	{
-		UpdateTransformMode(TransformProps->TransformMode);
-	}
-
-	if (bCurSetPivotMode != TransformProps->bSetPivot)
-	{
-		bool bEnableSetPivot = TransformProps->bSetPivot;
-		UpdateSetPivotModes(bEnableSetPivot);
-		bCurSetPivotMode = TransformProps->bSetPivot;
-	}
-}
-
 
 void UTransformMeshesTool::Render(IToolsContextRenderAPI* RenderAPI)
 {
 	DragAlignmentMechanic->Render(RenderAPI);
 }
 
-void UTransformMeshesTool::OnPropertyModified(UObject* PropertySet, FProperty* Property)
-{
-}
 
 
 void UTransformMeshesTool::UpdateSetPivotModes(bool bEnableSetPivot)
@@ -155,7 +155,7 @@ void UTransformMeshesTool::RegisterActions(FInteractiveToolActionSet& ActionSet)
 		LOCTEXT("TransformToggleSetPivot", "Toggle Set Pivot"),
 		LOCTEXT("TransformToggleSetPivotTooltip", "Toggle Set Pivot on and off"),
 		EModifierKey::None, EKeys::S,
-		[this]() { this->TransformProps->bSetPivot = !this->TransformProps->bSetPivot; });
+		[this]() { this->TransformProps->bSetPivotMode = !this->TransformProps->bSetPivotMode; });
 
 	ActionSet.RegisterAction(this, (int32)EStandardToolActions::BaseClientDefinedActionID + 2,
 		TEXT("ToggleSnapDrag"),
@@ -213,11 +213,36 @@ void UTransformMeshesTool::UpdateTransformMode(ETransformMeshesTransformMode New
 			break;
 	}
 
+	OnSnappingSettingsUpdated();
+
 	CurTransformMode = NewMode;
 }
 
 
 
+namespace UE {
+namespace Local {
+
+static void AddInstancedComponentInstance(UInstancedStaticMeshComponent* ISMC, int32 Index, UTransformProxy* TransformProxy, bool bModifyOnTransform)
+{
+	TransformProxy->AddComponentCustom(ISMC,
+		[ISMC, Index]() {
+			FTransform Tmp;
+			ISMC->GetInstanceTransform(Index, Tmp, true);
+			return Tmp;
+		},
+		[ISMC, Index](FTransform NewTransform) {
+			ISMC->UpdateInstanceTransform(Index, NewTransform, true, true, true);
+		},
+		Index, bModifyOnTransform
+	);
+}
+
+}
+}
+
+
+
 void UTransformMeshesTool::SetActiveGizmos_Single(bool bLocalRotations)
 {
 	check(ActiveGizmos.Num() == 0);
@@ -230,9 +255,23 @@ void UTransformMeshesTool::SetActiveGizmos_Single(bool bLocalRotations)
 
 	for (int32 ComponentIdx = 0; ComponentIdx < Targets.Num(); ComponentIdx++)
 	{
-		IPrimitiveComponentBackedTarget* TargetComponent = TargetComponentInterface(ComponentIdx);
-		Transformable.TransformProxy->AddComponent(TargetComponent->GetOwnerComponent());
-		ComponentsToIgnoreInAlignment.Add(TargetComponent->GetOwnerComponent());
+		UPrimitiveComponent* Component = UE::ToolTarget::GetTargetComponent(Targets[ComponentIdx]);
+		UInstancedStaticMeshComponent* InstancedComponent = Cast<UInstancedStaticMeshComponent>(Component);
+		if (InstancedComponent != nullptr && TransformProps->bApplyToInstances)
+		{
+			int32 NumInstances = InstancedComponent->GetInstanceCount();
+			for (int32 k = 0; k < NumInstances; ++k)
+			{
+				if (InstancedComponent->IsValidInstance(k) == false) return;
+				UE::Local::AddInstancedComponentInstance(InstancedComponent, k, Transformable.TransformProxy, true);
+			}
+		}
+		else
+		{
+			Transformable.TransformProxy->AddComponent(Component);
+		}
+
+		ComponentsToIgnoreInAlignment.Add(Component);
 	}
 
 	// leave out nonuniform scale if we have multiple objects in non-local mode
@@ -254,22 +293,49 @@ void UTransformMeshesTool::SetActiveGizmos_PerObject()
 	TArray<const UPrimitiveComponent*> ComponentsToIgnoreInAlignment;
 	for (int32 ComponentIdx = 0; ComponentIdx < Targets.Num(); ComponentIdx++)
 	{
-		IPrimitiveComponentBackedTarget* TargetComponent = TargetComponentInterface(ComponentIdx);
+		UPrimitiveComponent* Component = UE::ToolTarget::GetTargetComponent(Targets[ComponentIdx]);
+		UInstancedStaticMeshComponent* InstancedComponent = Cast<UInstancedStaticMeshComponent>(Component);
+		if (InstancedComponent != nullptr && TransformProps->bApplyToInstances)
+		{
+			int32 NumInstances = InstancedComponent->GetInstanceCount();
+			for (int32 k = 0; k < NumInstances; ++k)
+			{
+				if (InstancedComponent->IsValidInstance(k) == false) return;
 
-		ComponentsToIgnoreInAlignment.Reset();
 
-		FTransformMeshesTarget Transformable;
-		Transformable.TransformProxy = NewObject<UTransformProxy>(this);
-		Transformable.TransformProxy->AddComponent(TargetComponent->GetOwnerComponent());
-		ComponentsToIgnoreInAlignment.Add(TargetComponent->GetOwnerComponent());
+				FTransformMeshesTarget Transformable;
+				Transformable.TransformProxy = NewObject<UTransformProxy>(this);
+				UE::Local::AddInstancedComponentInstance(InstancedComponent, k, Transformable.TransformProxy, true);
 
-		ETransformGizmoSubElements GizmoElements = ETransformGizmoSubElements::FullTranslateRotateScale;
-		Transformable.TransformGizmo = UE::TransformGizmoUtil::CreateCustomRepositionableTransformGizmo(GetToolManager(), GizmoElements, this);
-		Transformable.TransformGizmo->SetActiveTarget(Transformable.TransformProxy);
+				ETransformGizmoSubElements GizmoElements = ETransformGizmoSubElements::FullTranslateRotateScale;
+				Transformable.TransformGizmo = UE::TransformGizmoUtil::CreateCustomRepositionableTransformGizmo(GetToolManager(), GizmoElements, this);
+				Transformable.TransformGizmo->SetActiveTarget(Transformable.TransformProxy);
 
-		DragAlignmentMechanic->AddToGizmo(Transformable.TransformGizmo, &ComponentsToIgnoreInAlignment);
+				ComponentsToIgnoreInAlignment.Reset();
+				ComponentsToIgnoreInAlignment.Add(Component);
+				DragAlignmentMechanic->AddToGizmo(Transformable.TransformGizmo, &ComponentsToIgnoreInAlignment);
 
-		ActiveGizmos.Add(Transformable);
+				ActiveGizmos.Add(Transformable);
+			}
+		}
+		else
+		{
+
+			FTransformMeshesTarget Transformable;
+			Transformable.TransformProxy = NewObject<UTransformProxy>(this);
+			Transformable.TransformProxy->AddComponent(Component);
+
+			ETransformGizmoSubElements GizmoElements = ETransformGizmoSubElements::FullTranslateRotateScale;
+			Transformable.TransformGizmo = UE::TransformGizmoUtil::CreateCustomRepositionableTransformGizmo(GetToolManager(), GizmoElements, this);
+			Transformable.TransformGizmo->SetActiveTarget(Transformable.TransformProxy);
+
+			ComponentsToIgnoreInAlignment.Reset();
+			ComponentsToIgnoreInAlignment.Add(Component);
+			DragAlignmentMechanic->AddToGizmo(Transformable.TransformGizmo, &ComponentsToIgnoreInAlignment);
+
+			ActiveGizmos.Add(Transformable);
+
+		}
 	}
 }
 
@@ -336,7 +402,7 @@ void UTransformMeshesTool::OnClickPress(const FInputDeviceRay& PressPos)
 
 void UTransformMeshesTool::OnClickDrag(const FInputDeviceRay& DragPos)
 {
-	bool bApplyToPivot = TransformProps->bSetPivot;
+	bool bApplyToPivot = TransformProps->bSetPivotMode;
 
 	TArray<const UPrimitiveComponent*> IgnoreComponents;
 	if (bApplyToPivot == false)
@@ -347,7 +413,7 @@ void UTransformMeshesTool::OnClickDrag(const FInputDeviceRay& DragPos)
 		{
 			if (IgnoreIndex == -1 || k == IgnoreIndex)
 			{
-				IgnoreComponents.Add(TargetComponentInterface(k)->GetOwnerComponent());
+				IgnoreComponents.Add(UE::ToolTarget::GetTargetComponent(Targets[k]));
 			}
 		}
 	}
@@ -436,7 +502,14 @@ void UTransformMeshesTool::OnTerminateDragSequence()
 }
 
 
-
+void UTransformMeshesTool::OnSnappingSettingsUpdated()
+{
+	for (FTransformMeshesTarget& Target : ActiveGizmos)
+	{
+		Target.TransformGizmo->bSnapToWorldGrid = TransformProps->bSnapToGrid;
+		Target.TransformGizmo->bSnapToWorldRotGrid = TransformProps->bSnapToRotation;
+	}
+}
 
 
 #undef LOCTEXT_NAMESPACE
diff --git a/Engine/Plugins/Experimental/MeshModelingToolset/Source/MeshModelingTools/Public/TransformMeshesTool.h b/Engine/Plugins/Experimental/MeshModelingToolset/Source/MeshModelingTools/Public/TransformMeshesTool.h
--- a/Engine/Plugins/Experimental/MeshModelingToolset/Source/MeshModelingTools/Public/TransformMeshesTool.h
+++ b/Engine/Plugins/Experimental/MeshModelingToolset/Source/MeshModelingTools/Public/TransformMeshesTool.h
@@ -16,6 +16,7 @@ class UDragAlignmentMechanic;
 class UTransformGizmo;
 class UTransformProxy;
 
+
 /**
  *
  */
@@ -96,25 +97,47 @@ class MESHMODELINGTOOLS_API UTransformMeshesToolProperties : public UInteractive
 	GENERATED_BODY()
 public:
 
+	/** Transformation Mode controls the overall behavior of the Gizmos in the Tool */
 	UPROPERTY(EditAnywhere, Category = Options)
 	ETransformMeshesTransformMode TransformMode = ETransformMeshesTransformMode::SharedGizmo;
 
+	/** When true, transformations are applied to the Instances of any Instanced Components (eg InstancedStaticMeshComponent) instead of to the Components */
+	UPROPERTY(EditAnywhere, Category = Options, meta = (HideEditConditionToggle, EditCondition = "bHaveInstances"))
+	bool bApplyToInstances = true;
 
-	UPROPERTY(EditAnywhere, Category = Options)
-	bool bSetPivot = false;
 
+	/** Snap to World Space grid during translation with World Axes */
+	UPROPERTY(EditAnywhere, Category = Snapping)
+	bool bSnapToGrid = false;
+
+	/** Snap to fixed rotation angles during axis rotation */
+	UPROPERTY(EditAnywhere, Category = Snapping)
+	bool bSnapToRotation = false;
+
+
+	/** When true, the Gizmo can be moved independently without affecting objects. This allows the Gizmo to be repositioned before transforming. */
+	UPROPERTY(EditAnywhere, Category = Pivot, meta = (TransientToolProperty, EditCondition = "TransformMode != ETransformMeshesTransformMode::PerObjectGizmo") )
+	bool bSetPivotMode = false;
 
-	/** Click-drag starting on the target objects to reposition them on the rest of the scene */
-	UPROPERTY(EditAnywhere, Category = Options)
-	bool bEnableSnapDragging = false;
 
+	/** When Snap-Dragging is enabled, you can Click-drag starting on the target objects to reposition them relative to the rest of the scene */
+	UPROPERTY(EditAnywhere, Category = SnapDragging, meta = (TransientToolProperty, DisplayName = "Enable"))
+	bool bEnableSnapDragging = false;
 
-	UPROPERTY(EditAnywhere, Category = Options, meta = (EditCondition = "bEnableSnapDragging == true"))
+	/** Which point on the object being Snap-Dragged to use as the "Source" point */
+	UPROPERTY(EditAnywhere, Category = SnapDragging, meta = (EditCondition = "bEnableSnapDragging == true"))
 	ETransformMeshesSnapDragSource SnapDragSource = ETransformMeshesSnapDragSource::ClickPoint;
 
-	/** When Snap-Dragging, align source and target normals */
-	UPROPERTY(EditAnywhere, Category = Options, meta = (EditCondition = "bEnableSnapDragging == true"))
+	/** How the object being Snap-Dragged should be rotated relative to the Source point location and Hit Surface normal  */
+	UPROPERTY(EditAnywhere, Category = SnapDragging, meta = (EditCondition = "bEnableSnapDragging == true"))
 	ETransformMeshesSnapDragRotationMode RotationMode = ETransformMeshesSnapDragRotationMode::AlignFlipped;
+
+
+public:
+
+	// internal, used to control visibility of Instance settings
+	UPROPERTY(meta = (TransientToolProperty))
+	bool bHaveInstances = false;
 };
 
 
@@ -149,15 +172,12 @@ public:
 	virtual void Setup() override;
 	virtual void Shutdown(EToolShutdownType ShutdownType) override;
 
-	virtual void OnTick(float DeltaTime) override;
 	virtual void Render(IToolsContextRenderAPI* RenderAPI) override;
 
 	virtual bool HasCancel() const override { return false; }
 	virtual bool HasAccept() const override { return false; }
 	virtual bool CanAccept() const override { return false; }
 
-	virtual void OnPropertyModified(UObject* PropertySet, FProperty* Property) override;
-
 
 	// ICLickDragBehaviorTarget interface
 	virtual FInputRayHit CanBeginClickDragSequence(const FInputDeviceRay& PressPos) override;
@@ -182,7 +202,6 @@ protected:
 
 	ETransformMeshesTransformMode CurTransformMode;
 	void UpdateTransformMode(ETransformMeshesTransformMode NewMode);
-	bool bCurSetPivotMode;
 	void UpdateSetPivotModes(bool bEnableSetPivot);
 
 	void SetActiveGizmos_Single(bool bLocalRotations);
@@ -195,4 +214,6 @@ protected:
 	int ActiveSnapDragIndex = -1;
 
 	void OnParametersUpdated();
+
+	void OnSnappingSettingsUpdated();
 };
diff --git a/Engine/Source/Runtime/Experimental/InteractiveToolsFramework/Private/BaseGizmos/TransformProxy.cpp b/Engine/Source/Runtime/Experimental/InteractiveToolsFramework/Private/BaseGizmos/TransformProxy.cpp
--- a/Engine/Source/Runtime/Experimental/InteractiveToolsFramework/Private/BaseGizmos/TransformProxy.cpp
+++ b/Engine/Source/Runtime/Experimental/InteractiveToolsFramework/Private/BaseGizmos/TransformProxy.cpp
@@ -11,12 +11,36 @@ void UTransformProxy::AddComponent(USceneComponent* Component, bool bModifyCompo
 {
 	check(Component);
 
-	FRelativeObject NewObj;
+	FRelativeObject& NewObj = Objects.Emplace_GetRef();
 	NewObj.Component = Component;
 	NewObj.bModifyComponentOnTransform = bModifyComponentOnTransform;
-	NewObj.StartTransform = Component->GetComponentToWorld();
+	NewObj.GetTransformFunc = [Component]() { return Component->GetComponentToWorld(); };
+	NewObj.SetTransformFunc = [Component](FTransform NewTransform) { return Component->SetWorldTransform(NewTransform); };
+	NewObj.UserDefinedIndex = 0;
+	NewObj.StartTransform = NewObj.GetTransformFunc();
+	NewObj.RelativeTransform = FTransform::Identity;
+
+	UpdateSharedTransform();
+	OnPivotChanged.Broadcast(this, SharedTransform);
+}
+
+void UTransformProxy::AddComponentCustom(
+	USceneComponent* Component,
+	TUniqueFunction<FTransform(void)> GetTransformFunc,
+	TUniqueFunction<void(FTransform)> SetTransformFunc,
+	int64 UserDefinedIndex,
+	bool bModifyComponentOnTransform)
+{
+	check(Component);
+
+	FRelativeObject& NewObj = Objects.Emplace_GetRef();
+	NewObj.Component = Component;
+	NewObj.bModifyComponentOnTransform = bModifyComponentOnTransform;
+	NewObj.GetTransformFunc = MoveTemp(GetTransformFunc);
+	NewObj.SetTransformFunc = MoveTemp(SetTransformFunc);
+	NewObj.UserDefinedIndex = UserDefinedIndex;
+	NewObj.StartTransform = NewObj.GetTransformFunc();
 	NewObj.RelativeTransform = FTransform::Identity;
-	Objects.Add(NewObj);
 
 	UpdateSharedTransform();
 	OnPivotChanged.Broadcast(this, SharedTransform);
@@ -99,7 +123,7 @@ void UTransformProxy::UpdateObjects()
 				Obj.Component->Modify();
 			}
 
-			Obj.Component->SetWorldTransform(CombinedTransform);
+			Obj.SetTransformFunc(CombinedTransform);
 		}
 	}
 }
@@ -148,7 +172,7 @@ void UTransformProxy::UpdateObjectTransforms()
 	{
 		if (Obj.Component != nullptr)
 		{
-			Obj.StartTransform = Obj.Component->GetComponentToWorld();
+			Obj.StartTransform = Obj.GetTransformFunc();
 		}
 		Obj.RelativeTransform = Obj.StartTransform;
 		Obj.RelativeTransform.SetToRelativeTransform(SharedTransform);
diff --git a/Engine/Source/Runtime/Experimental/InteractiveToolsFramework/Public/BaseGizmos/TransformProxy.h b/Engine/Source/Runtime/Experimental/InteractiveToolsFramework/Public/BaseGizmos/TransformProxy.h
--- a/Engine/Source/Runtime/Experimental/InteractiveToolsFramework/Public/BaseGizmos/TransformProxy.h
+++ b/Engine/Source/Runtime/Experimental/InteractiveToolsFramework/Public/BaseGizmos/TransformProxy.h
@@ -34,6 +34,23 @@ public:
 	 */
 	virtual void AddComponent(USceneComponent* Component, bool bModifyComponentOnTransform = true);
 
+	/**
+	 * Add a component sub-object to the proxy set with custom transform access functions. This can be used
+	 * to do things like add an ISM Instance to the TransformProxy.
+	 * @param GetTransformFunc return current transform
+	 * @param SetTransformFunc set current transform
+	 * @param UserDefinedIndex an arbitrary integer that can be provided, not currently used but may be useful in subclasses
+	 * @param bModifyComponentOnTransform if true, Component->Modify() is called before the Component transform is updated
+	 * @warning The internal shared transform is regenerated each time a component is added
+	 */
+	virtual void AddComponentCustom(
+		USceneComponent* Component, 
+		TUniqueFunction<FTransform(void)> GetTransformFunc,
+		TUniqueFunction<void(FTransform)> SetTransformFunc,
+		int64 UserDefinedIndex = 0,
+		bool bModifyComponentOnTransform = true);
+
+
 	/**
 	 * @return the shared transform for all the sub-objects
 	 */
@@ -122,6 +139,11 @@ protected:
 		TWeakObjectPtr<USceneComponent> Component;
 		bool bModifyComponentOnTransform;
 
+		int64 UserDefinedIndex;		// no current purpose, reserved for subclass/future use
+
+		TUniqueFunction<FTransform(void)> GetTransformFunc;
+		TUniqueFunction<void(FTransform)> SetTransformFunc;
+
 		/** The initial transform of the object, set during UpdateSharedTransform() */
 		FTransform StartTransform;
 		/** The transform of the object relative to */
