ModelingComponents: improve handling of thin objects in FMeshSceneAdapter. Add Weld step before connected components. Per-Instance transforms are now extracted outside of Component-level processing, and 'thin' detection is now tested under this transform set, for each connected component. This improves handling of (eg) extreme scaling applied to create thin instances. Check AABB dimensions in IsThinPlanarSubMesh() before using normals, this handles many thin parts that are not "box-like" (will still not work if part is rotated). Seed points now generated on offset surfaces, if required, to avoid failures in surface rootfinding. Add support for filtering tiny objects (enabled in ApproximateActorsImpl, for 10% of voxel dimension)

#rnx

diff --git a/Engine/Plugins/Experimental/MeshModelingToolset/Source/GeometryProcessingAdapters/Private/ApproximateActorsImpl.cpp b/Engine/Plugins/Experimental/MeshModelingToolset/Source/GeometryProcessingAdapters/Private/ApproximateActorsImpl.cpp
--- a/Engine/Plugins/Experimental/MeshModelingToolset/Source/GeometryProcessingAdapters/Private/ApproximateActorsImpl.cpp
+++ b/Engine/Plugins/Experimental/MeshModelingToolset/Source/GeometryProcessingAdapters/Private/ApproximateActorsImpl.cpp
@@ -812,6 +812,9 @@ void FApproximateActorsImpl::GenerateApproximationForActorSet(const TArray<AActo
 	FMeshSceneAdapterBuildOptions SceneBuildOptions;
 	SceneBuildOptions.bThickenThinMeshes = Options.bAutoThickenThinParts;
 	SceneBuildOptions.DesiredMinThickness = Options.AutoThickenThicknessMeters * 100.0;		// convert to cm (UE Units)
+	// filter out objects smaller than 10% of voxel size
+	SceneBuildOptions.bFilterTinyObjects = true;
+	SceneBuildOptions.TinyObjectBoxMaxDimension = ApproxAccuracy * 0.1;
 	SceneBuildOptions.bPrintDebugMessages = Options.bVerbose;
 	{
 		TRACE_CPUPROFILER_EVENT_SCOPE(ApproximateActorsImpl_Generate_BuildScene);
diff --git a/Engine/Plugins/Experimental/MeshModelingToolset/Source/ModelingComponents/Private/Scene/MeshSceneAdapter.cpp b/Engine/Plugins/Experimental/MeshModelingToolset/Source/ModelingComponents/Private/Scene/MeshSceneAdapter.cpp
--- a/Engine/Plugins/Experimental/MeshModelingToolset/Source/ModelingComponents/Private/Scene/MeshSceneAdapter.cpp
+++ b/Engine/Plugins/Experimental/MeshModelingToolset/Source/ModelingComponents/Private/Scene/MeshSceneAdapter.cpp
@@ -12,6 +12,7 @@
 #include "BoxTypes.h"
 #include "FrameTypes.h"
 #include "Selections/MeshConnectedComponents.h"
+#include "DynamicMesh/Operations/MergeCoincidentMeshEdges.h"
 #include "Operations/OffsetMeshRegion.h"
 #include "DynamicMeshEditor.h"
 #include "CompGeom/ConvexHull2.h"
@@ -71,72 +72,48 @@ void CollectSeedPointsFromMeshVertices(
 		AccumPointsInOut.Add(TransformFunc(Mesh.GetVertex(k)));
 	}
 }
-
-
-/**
- * Try to check if a Mesh is "thin", ie basically a planar patch (open or closed).
- * The Normal of the largest-area triangle is taken as the plane normal, and then the "thickness" is measured relative to this plane
- * @param ThinTolerance identify as Thin if the thickness extents is within this size
- * @param ThinPlaneOut thin plane normal will be returned via this frame
- * @param ThicknessOut measured thickness relative to plane will be returned here
- * @return true if mesh is identified as thin under ThinTolerance
- */
-template<typename MeshType>
-bool IsThinPlanarSubMesh(const MeshType& Mesh, double ThinTolerance, FFrame3d& ThinPlaneOut, double& ThicknessOut)
+void CollectSeedPointsFromMeshVertices(
+	const FDynamicMesh3& Mesh,
+	TFunctionRef<FVector3d(const FVector3d&)> TransformFunc, TArray<FVector3d>& AccumPointsInOut,
+	double NormalOffset = 0,
+	int32 MaxPoints = 500)
 {
-	int32 TriCount = Mesh.TriangleCount();
-
-	// Find triangle with largest area and use it's normal as the plane normal
-	// (this is not ideal and we should probably do a normals histogram
-	double MaxArea = 0;
-	FVector3d MaxAreaNormal;
-	FVector3d MaxAreaPoint;
-	for (int32 tid = 0; tid < TriCount; ++tid)
+	int32 NumVertices = Mesh.VertexCount();
+	int32 LogNumVertices = FMath::Max(1, (int32)FMathd::Ceil(FMathd::Log(NumVertices)));
+	int32 SeedPointCount = (int)(10 * LogNumVertices);
+	SeedPointCount = FMath::Min(SeedPointCount, MaxPoints);
+	int32 Skip = FMath::Max(NumVertices / SeedPointCount, 2);
+	for (int32 k = 0; k < NumVertices; k += Skip)
 	{
-		if (Mesh.IsTriangle(tid))
+		FVector3d Pos = Mesh.GetVertex(k);
+		if (NormalOffset > 0)
 		{
-			FVector3d A, B, C;
-			Mesh.GetTriVertices(tid, A, B, C);
-			double TriArea;
-			FVector3d TriNormal = VectorUtil::NormalArea(A, B, C, TriArea);
-			if (TriArea > MaxArea)
-			{
-				MaxArea = TriArea;
-				MaxAreaNormal = TriNormal;
-				MaxAreaPoint = A;
-			}
+			int32 tid = *Mesh.VtxTrianglesItr(k).begin();
+			FVector3d TriNormal = Mesh.GetTriNormal(tid);
+			Pos += NormalOffset * TriNormal;
 		}
+		AccumPointsInOut.Add(TransformFunc(Pos));
 	}
+}
 
-	// Now compute the bounding box in the local space of this plane
-	ThinPlaneOut = FFrame3d(MaxAreaPoint, MaxAreaNormal);
+/**
+ * Try to check if a Mesh is "thin", ie basically a planar patch (open or closed), relative to a given plane
+ * 
+ */
+template<typename MeshType>
+double MeasureThickness(const MeshType& Mesh, const FFrame3d& Plane)
+{
 	FAxisAlignedBox3d PlaneExtents = FAxisAlignedBox3d::Empty();
 	int32 VertexCount = Mesh.VertexCount();
-	for (int32 tid = 0; tid < TriCount; ++tid)
+	for (int32 k = 0; k < VertexCount; ++k)
 	{
-		if (Mesh.IsTriangle(tid))
-		{
-			FVector3d TriVerts[3];
-			Mesh.GetTriVertices(tid, TriVerts[0], TriVerts[1], TriVerts[2]);
-			for (int32 j = 0; j < 3; ++j)
-			{
-				TriVerts[j] = ThinPlaneOut.ToFramePoint(TriVerts[j]);
-				PlaneExtents.Contain(TriVerts[j]);
-			}
-		}
-
-		// early-out if we exceed tolerance
-		if (PlaneExtents.Depth() > ThinTolerance)
+		if (Mesh.IsVertex(k))
 		{
-			return false;
+			PlaneExtents.Contain(Plane.ToFramePoint(Mesh.GetVertex(k)));
 		}
 	}
 
-	// shift plane to center
-	FVector3d Center = PlaneExtents.Center();
-	ThinPlaneOut.Origin += Center.X*ThinPlaneOut.X() + Center.Y*ThinPlaneOut.Y() + Center.Z*ThinPlaneOut.Z();
-	ThicknessOut = PlaneExtents.Depth();
-	return true;
+	return PlaneExtents.Depth();
 }
 
 
@@ -149,9 +126,15 @@ bool IsThinPlanarSubMesh(const MeshType& Mesh, double ThinTolerance, FFrame3d& T
  * @return true if submesh identified as thin
  */
 template<typename MeshType>
-bool IsThinPlanarSubMesh(const MeshType& Mesh, const TArray<int32>& Triangles, double ThinTolerance, FFrame3d& ThinPlaneOut)
+bool IsThinPlanarSubMesh(
+	const MeshType& Mesh, const TArray<int32>& Triangles, 
+	FTransformSequence3d& Transform, 
+	double ThinTolerance, 
+	FFrame3d& ThinPlaneOut)
 {
+	FVector3d Scale = Transform.GetAccumulatedScale();
 	int32 TriCount = Triangles.Num();
+	FAxisAlignedBox3d MeshBounds = FAxisAlignedBox3d::Empty();
 
 	// Find triangle with largest area and use it's normal as the plane normal
 	// (this is not ideal and we should probably do a normals histogram
@@ -165,6 +148,10 @@ bool IsThinPlanarSubMesh(const MeshType& Mesh, const TArray<int32>& Triangles, d
 		{
 			FVector3d A, B, C;
 			Mesh.GetTriVertices(tid, A, B, C);
+			A *= Scale; B *= Scale; C *= Scale;
+			MeshBounds.Contain(A);
+			MeshBounds.Contain(B);
+			MeshBounds.Contain(C);
 			double TriArea;
 			FVector3d TriNormal = VectorUtil::NormalArea(A, B, C, TriArea);
 			if (TriArea > MaxArea)
@@ -176,6 +163,17 @@ bool IsThinPlanarSubMesh(const MeshType& Mesh, const TArray<int32>& Triangles, d
 		}
 	}
 
+	// if one of the AABB dimensions is below the thin tolerance, just use it
+	FVector3d BoundsDimensions = MeshBounds.Diagonal();
+	if (BoundsDimensions.GetAbsMin() < ThinTolerance)
+	{
+		int32 Index = UE::Geometry::MinAbsElementIndex(BoundsDimensions);
+		FVector3d BoxNormal = FVector3d::Zero();
+		BoxNormal[Index] = 1.0;
+		ThinPlaneOut = FFrame3d(MeshBounds.Center(), BoxNormal);
+		return true;
+	}
+
 	// Now compute the bounding box in the local space of this plane
 	ThinPlaneOut = FFrame3d(MaxAreaPoint, MaxAreaNormal);
 	FAxisAlignedBox3d PlaneExtents = FAxisAlignedBox3d::Empty();
@@ -189,7 +187,7 @@ bool IsThinPlanarSubMesh(const MeshType& Mesh, const TArray<int32>& Triangles, d
 			Mesh.GetTriVertices(tid, TriVerts[0], TriVerts[1], TriVerts[2]);
 			for (int32 j = 0; j < 3; ++j)
 			{
-				TriVerts[j] = ThinPlaneOut.ToFramePoint(TriVerts[j]);
+				TriVerts[j] = ThinPlaneOut.ToFramePoint(Scale * TriVerts[j]);
 				PlaneExtents.Contain(TriVerts[j]);
 			}
 		}
@@ -210,6 +208,7 @@ bool IsThinPlanarSubMesh(const MeshType& Mesh, const TArray<int32>& Triangles, d
 
 
 
+
 /**
  * @return false if any of Triangles in Mesh have open boundary edges
  */
@@ -288,8 +287,8 @@ public:
 	virtual void CollectSeedPoints(TArray<FVector3d>& WorldPoints, TFunctionRef<FVector3d(const FVector3d&)> LocalToWorldFunc) override
 	{
 		return bHasBakedTransform ?
-			CollectSeedPointsFromMeshVertices<FDynamicMesh3>(Mesh, [&](const FVector3d& P) {return P;}, WorldPoints) :
-			CollectSeedPointsFromMeshVertices<FDynamicMesh3>(Mesh, LocalToWorldFunc, WorldPoints);
+			CollectSeedPointsFromMeshVertices(Mesh, [&](const FVector3d& P) {return P;}, WorldPoints, WindingShellThickness) :
+			CollectSeedPointsFromMeshVertices(Mesh, LocalToWorldFunc, WorldPoints, WindingShellThickness);
 	}
 
 	virtual double FastWindingNumber(const FVector3d& P, const FTransformSequence3d& LocalToWorldTransform) override
@@ -559,6 +558,7 @@ static void CollectActorChildMeshes(AActor* Actor, UActorComponent* Component, F
 }
 
 
+
 void FMeshSceneAdapter::AddActors(const TArray<AActor*>& ActorsSetIn)
 {
 	TRACE_CPUPROFILER_EVENT_SCOPE(MeshScene_AddActors);
@@ -716,6 +716,7 @@ void FMeshSceneAdapter::UpdateActorBounds(FActorAdapter& Actor)
 void ConstructUniqueScalesMapping(
 	const TArray<FTransformSequence3d>& TransformSet, 
 	TArray<TArray<int32>>& UniqueScaleSetsOut,
+	TArray<FTransformSequence3d>& UniqueScaleTransformsOut,
 	double ScaleComponentTolerance = 0.01)
 {
 	// two transforms are "the same up to scaling" if this returns true
@@ -724,7 +725,7 @@ void ConstructUniqueScalesMapping(
 		return (T1.GetScale() - T2.GetScale()).GetAbsMax() < ScaleComponentTolerance;
 	};
 
-	TArray<FTransformSequence3d> UniqueScaleTransforms;		// accumulate transform-with-unique-scale's here
+	UniqueScaleTransformsOut.Reset();
 	int32 N = TransformSet.Num();
 	TArray<int32> UniqueScaleMap;
 	UniqueScaleMap.SetNum(N);
@@ -732,9 +733,9 @@ void ConstructUniqueScalesMapping(
 	{
 		FTransformSequence3d CurTransform = TransformSet[k];
 		int32 FoundIndex = -1;
-		for (int32 j = 0; j < UniqueScaleTransforms.Num(); ++j)
+		for (int32 j = 0; j < UniqueScaleTransformsOut.Num(); ++j)
 		{
-			if (CurTransform.IsEquivalent(UniqueScaleTransforms[j], CompareScales))
+			if (CurTransform.IsEquivalent(UniqueScaleTransformsOut[j], CompareScales))
 			{
 				FoundIndex = j;
 				break;
@@ -746,13 +747,13 @@ void ConstructUniqueScalesMapping(
 		}
 		else
 		{
-			UniqueScaleMap[k] = UniqueScaleTransforms.Num();
-			UniqueScaleTransforms.Add(CurTransform);
+			UniqueScaleMap[k] = UniqueScaleTransformsOut.Num();
+			UniqueScaleTransformsOut.Add(CurTransform);
 		}
 	}
 
 	// build clusters
-	int32 NumUniqueScales = UniqueScaleTransforms.Num();
+	int32 NumUniqueScales = UniqueScaleTransformsOut.Num();
 	UniqueScaleSetsOut.SetNum(NumUniqueScales);
 	for (int32 k = 0; k < N; ++k)
 	{
@@ -762,9 +763,11 @@ void ConstructUniqueScalesMapping(
 }
 
 
-
 void FMeshSceneAdapter::Build_FullDecompose(const FMeshSceneAdapterBuildOptions& BuildOptions)
 {
+	EParallelForFlags ParallelFlags = EParallelForFlags::Unbalanced;
+	//EParallelForFlags ParallelFlags = EParallelForFlags::ForceSingleThread;
+
 	// initial list of spatial wrappers that need to be built
 	TArray<FSpatialWrapperInfo*> ToBuild;
 	for (TPair<void*, TSharedPtr<FSpatialWrapperInfo>> Pair : SpatialAdapters)
@@ -779,7 +782,7 @@ void FMeshSceneAdapter::Build_FullDecompose(const FMeshSceneAdapterBuildOptions&
 	{
 		FSpatialWrapperInfo* WrapperInfo = ToBuild[i];
 		WrapperInfo->SpatialWrapper->Build(TempBuildOptions);
-	});
+	}, ParallelFlags);
 
 	// sort build list by increasing triangle count
 	ToBuild.Sort([&](const FSpatialWrapperInfo& A, const FSpatialWrapperInfo& B)
@@ -789,14 +792,35 @@ void FMeshSceneAdapter::Build_FullDecompose(const FMeshSceneAdapterBuildOptions&
 
 	// stats we will collect during execution
 	int32 NumInitialSources = ToBuild.Num();
+	int32 NumSourceUniqueTris = 0;
 	std::atomic<int32> DecomposedSourceMeshCount = 0;
 	std::atomic<int32> DecomposedMeshesCount = 0;
 	std::atomic<int32> SourceInstancesCount = 0;
 	std::atomic<int32> NewInstancesCount = 0;
 	std::atomic<int32> SkippedDecompositionCount = 0;
+	std::atomic<int32> SingleTriangleMeshes = 0;
 	int32 AddedUniqueTrisCount = 0;
 	int32 InstancedTrisCount = 0;
 
+	struct FProcessedSourceMeshStats
+	{
+		FString AssetName;
+		int32 SourceTriangles = 0;
+		int32 SourceInstances = 0;
+		int32 SourceComponents = 0;
+		int32 UniqueInstanceScales = 0;
+		int32 NumOpen = 0;
+		int32 NumThin = 0;
+		int32 InstancedSubmeshComponents = 0;
+		int32 InstancedSubmeshTris = 0;		// number of source tris used directly, ie copied to a single mesh for original instances
+		TArray<FVector2i> SubmeshSizeCount;
+		int32 NewUniqueTris = 0;
+		int32 NewInstances = 0;
+		int32 TotalNumUniqueTris = 0;
+	};
+	TArray<FProcessedSourceMeshStats> Stats;
+	Stats.SetNum(NumInitialSources);
+
 	// these locks are used below to control access
 	FCriticalSection ToBuildQueueLock;
 	FCriticalSection InternalListsLock;
@@ -819,6 +843,9 @@ void FMeshSceneAdapter::Build_FullDecompose(const FMeshSceneAdapterBuildOptions&
 		PendingBuildsLock.Unlock();
 	};
 
+	std::atomic<int32> NumTinyComponents = 0;
+	std::atomic<int32> NumTinyInstances = 0;
+	std::atomic<int32> TinyInstanceTotalTriangles = 0;
 
 	// Parallel-process all the ToBuild spatial wrappers. If the mesh is closed and all the pieces are good,
 	// this will just emit a Build job. Otherwise it will pull the mesh apart into pieces, move all the closed non-thin
@@ -832,11 +859,13 @@ void FMeshSceneAdapter::Build_FullDecompose(const FMeshSceneAdapterBuildOptions&
 		// ParallelFor will not respect the sorting by triangle-count we did above, so we have to treat the list as a queue and pop from the back
 		ToBuildQueueLock.Lock();
 		check(ToBuild.Num() > 0);
+		FProcessedSourceMeshStats& ItemStats = Stats[ToBuild.Num() - 1];
 		FSpatialWrapperInfo* WrapperInfo = ToBuild.Pop(false);
+		NumSourceUniqueTris += WrapperInfo->SpatialWrapper->GetTriangleCount();
 		ToBuildQueueLock.Unlock();
 
 		// get name for debugging purposes
-		FString AssetName = WrapperInfo->SourceContainer.GetStaticMesh() ? WrapperInfo->SourceContainer.GetStaticMesh()->GetName() : TEXT("Unknown");
+		ItemStats.AssetName = WrapperInfo->SourceContainer.GetStaticMesh() ? WrapperInfo->SourceContainer.GetStaticMesh()->GetName() : TEXT("Unknown");
 
 		// convert this mesh to a dynamicmesh for processing
 		FDynamicMesh3 LocalMesh;
@@ -844,8 +873,33 @@ void FMeshSceneAdapter::Build_FullDecompose(const FMeshSceneAdapterBuildOptions&
 			TRACE_CPUPROFILER_EVENT_SCOPE(MeshScene_Build_ProcessMesh_1Copy);
 			WrapperInfo->SpatialWrapper->AppendMesh(LocalMesh, FTransformSequence3d());
 		}
+		ItemStats.SourceTriangles = LocalMesh.TriangleCount();
+
+		// construct list of per-instance transforms that reference this mesh
+		TArray<FActorChildMesh*> MeshesToDecompose = WrapperInfo->ParentMeshes;
+		TArray<FTransformSequence3d> ParentTransforms;
+		for (FActorChildMesh* MeshInstance : MeshesToDecompose)
+		{
+			ParentTransforms.Add(MeshInstance->WorldTransform);
+			SourceInstancesCount++;
+		}
+		ItemStats.SourceInstances = SourceInstancesCount;
+
+		// Decompose the per-instance transforms into subsets that share the same total scaling ("unique scale").
+		// If we apply these different scales to copies of the mesh, we can generate new instances for the copies,
+		// which can avoid uniquing a lot of geometry
+		TArray<TArray<int32>> UniqueScaleTransformSets;
+		TArray<FTransformSequence3d> UniqueScaleTransforms;
+		ConstructUniqueScalesMapping(ParentTransforms, UniqueScaleTransformSets, UniqueScaleTransforms);
+		int32 NumUniqueScales = UniqueScaleTransformSets.Num();
+		ItemStats.UniqueInstanceScales = NumUniqueScales;
 
-		// should we try to weld here??
+		// if mesh is not too huge, try to weld edges
+		if (LocalMesh.TriangleCount() < 100000)
+		{
+			FMergeCoincidentMeshEdges WeldEdges(&LocalMesh);
+			WeldEdges.Apply();
+		}
 
 		// find separate submeshes of the mesh
 		FMeshConnectedComponents Components(&LocalMesh);
@@ -854,12 +908,15 @@ void FMeshSceneAdapter::Build_FullDecompose(const FMeshSceneAdapterBuildOptions&
 			Components.FindConnectedTriangles();
 		}
 		int32 NumComponents = Components.Num();
+		ItemStats.SourceComponents = NumComponents;
 
 		// for each submesh/component, determine if it is closed, and if it is 'thin'
 		TArray<bool> IsClosed, IsThin;
 		IsClosed.Init(false, NumComponents);
-		std::atomic<int32> NumNonClosed = 0;
 		IsThin.Init(false, NumComponents);
+		std::atomic<int32> NumNonClosed = 0;
+		TArray<TArray<FFrame3d>> BestFitPlanes;
+		BestFitPlanes.SetNum(NumComponents);
 		std::atomic<int32> NumThin = 0;
 		{
 			TRACE_CPUPROFILER_EVENT_SCOPE(MeshScene_Build_ProcessMesh_3Closed);
@@ -873,14 +930,20 @@ void FMeshSceneAdapter::Build_FullDecompose(const FMeshSceneAdapterBuildOptions&
 					NumNonClosed++;
 				}
 
-				FFrame3d TempPlane;
-				IsThin[ci] = IsThinPlanarSubMesh<FDynamicMesh3>(LocalMesh, Triangles, BuildOptions.DesiredMinThickness, TempPlane);
+				BestFitPlanes[ci].SetNum(NumUniqueScales);
+				for (int32 k = 0; k < NumUniqueScales; ++k)
+				{
+					bool bIsThinUnderTransform = IsThinPlanarSubMesh<FDynamicMesh3>(LocalMesh, Triangles, UniqueScaleTransforms[k], BuildOptions.DesiredMinThickness, BestFitPlanes[ci][k]);
+					IsThin[ci] = IsThin[ci] || bIsThinUnderTransform;
+				}
 				if (IsThin[ci])
 				{
 					NumThin++;
 				}
-			});
+			}, ParallelFlags);
 		}
+		ItemStats.NumOpen = NumNonClosed.load();
+		ItemStats.NumThin = NumThin.load();
 
 		// if we have no open meshes and no thin meshes, we can just use the SpatialWrapper we already have, 
 		// but we have to rebuild it because we did not do a full build above
@@ -891,22 +954,6 @@ void FMeshSceneAdapter::Build_FullDecompose(const FMeshSceneAdapterBuildOptions&
 			return;
 		}
 
-		// construct list of per-instance transforms that reference this mesh
-		TArray<FActorChildMesh*> MeshesToDecompose = WrapperInfo->ParentMeshes;
-		TArray<FTransformSequence3d> ParentTransforms;
-		for (FActorChildMesh* MeshInstance : MeshesToDecompose)
-		{
-			ParentTransforms.Add(MeshInstance->WorldTransform);
-			SourceInstancesCount++;
-		}
-
-		// Decompose the per-instance transforms into subsets that share the same total scaling ("unique scale").
-		// If we apply these different scales to copies of the mesh, we can generate new instances for the copies,
-		// which can avoid uniquing a lot of geometry
-		TArray<TArray<int32>> UniqueScaleTransformSets;
-		ConstructUniqueScalesMapping(ParentTransforms, UniqueScaleTransformSets);
-		int32 NumUniqueScales = UniqueScaleTransformSets.Num();
-
 		// Accumulate submesh/components that do *not* need further processing here, that accumulated mesh
 		// (if non-empty) can be shared among all the original FActorChildMesh instances
 		FDynamicMesh3 LocalSpaceParts;
@@ -933,13 +980,16 @@ void FMeshSceneAdapter::Build_FullDecompose(const FMeshSceneAdapterBuildOptions&
 			{
 				const TArray<int32>& Triangles = Components[ci].Indices;
 
+				int32 NewUniqueTriangles = Triangles.Num() * NumUniqueScales;
+
 				// We will make unscaled copies of a mesh if (1) it is "thin" and (2) it has a moderate number of triangles *or* a single usage
 				// TODO: should we always unique a mesh with a single usage? We can just make it unsigned...
 				bool bIsClosed = IsClosed[ci];
-				if (IsThin[ci] == false || (Triangles.Num() > 10000 && WrapperInfo->ParentMeshes.Num() > 1))
+				if (IsThin[ci] == false || (NewUniqueTriangles > 1000000) )
 				{
 					Mappings.Reset(); EditResult.Reset();
 					LocalSpaceAccumulator.AppendTriangles(&LocalMesh, Triangles, Mappings, EditResult, false);
+					ItemStats.InstancedSubmeshComponents++;
 					continue;
 				}
 
@@ -956,27 +1006,35 @@ void FMeshSceneAdapter::Build_FullDecompose(const FMeshSceneAdapterBuildOptions&
 					Editor.AppendTriangles(&LocalMesh, Triangles, Mappings, EditResult, false);
 					// bake in the scaling
 					FVector3d Scale = ParentTransforms[InstanceIndices[0]].GetAccumulatedScale();
+					FAxisAlignedBox3d ScaledBounds = FAxisAlignedBox3d::Empty();
 					for (int32 vid : EditResult.NewVertices)
 					{
 						FVector3d LocalPos = NewSubmesh.Mesh->GetVertex(vid);
-						NewSubmesh.Mesh->SetVertex(vid, LocalPos * Scale);
+						LocalPos *= Scale;
+						NewSubmesh.Mesh->SetVertex(vid, LocalPos);
+						ScaledBounds.Contain(LocalPos);
+					}
+
+					// if this is a tiny submesh we will skip it
+					if (BuildOptions.bFilterTinyObjects && ScaledBounds.MaxDim() < BuildOptions.TinyObjectBoxMaxDimension)
+					{
+						NumTinyComponents++;
+						NumTinyInstances += InstanceIndices.Num();
+						TinyInstanceTotalTriangles += InstanceIndices.Num() * Triangles.Num();
+						continue;
 					}
 
-					// Recompute thickness of scaled mesh and store it. Note that we might fail to be
-					// considered "thin" anymore, in that case we will fall back to using winding number
+					// Recompute thickness of scaled mesh and store it. Note that after scaling we might fail to 
+					// be considered "thin" anymore, in that case we will fall back to using winding number
 					// for this mesh (So, it was a waste to do this separation, but messy to turn back now)
-					FFrame3d TempPlane; double Thickness;
 					if (bIsClosed == false)
 					{
 						NewSubmesh.ComputedThickness = 0;
 					}
-					else if (bIsClosed && IsThinPlanarSubMesh<FDynamicMesh3>(*NewSubmesh.Mesh, BuildOptions.DesiredMinThickness, TempPlane, Thickness))
-					{
-						NewSubmesh.ComputedThickness = Thickness;
-					}
 					else
 					{
-						NewSubmesh.ComputedThickness = BuildOptions.DesiredMinThickness;
+						double NewThickness = MeasureThickness<FDynamicMesh3>(*NewSubmesh.Mesh, BestFitPlanes[ci][k]);
+						NewSubmesh.ComputedThickness = FMathd::Min(NewThickness, BuildOptions.DesiredMinThickness);
 					}
 
 					// make new set of instances
@@ -993,6 +1051,9 @@ void FMeshSceneAdapter::Build_FullDecompose(const FMeshSceneAdapterBuildOptions&
 		// At this point we have processed all the Submeshes/Components. Now we generate new MeshSpatialWrapper's
 		// and any necessary new FActorAdapter's/FActorChildMesh's
 
+		ItemStats.InstancedSubmeshTris = LocalSpaceParts.TriangleCount();
+		ItemStats.TotalNumUniqueTris += ItemStats.InstancedSubmeshTris;
+
 		// First handle the LocalSpaceParts mesh, which can still be shared between the original FActorChildMesh instances
 		if (LocalSpaceParts.TriangleCount() > 0)
 		{
@@ -1077,12 +1138,17 @@ void FMeshSceneAdapter::Build_FullDecompose(const FMeshSceneAdapterBuildOptions&
 				NewInstancesCount++;
 			}
 
+			ItemStats.SubmeshSizeCount.Add(FVector2i(TriangleCount, Submesh.NewTransforms.Num()));
+			ItemStats.NewInstances += Submesh.NewTransforms.Num();
+			ItemStats.NewUniqueTris += TriangleCount;
+			ItemStats.TotalNumUniqueTris += TriangleCount;
+
 			// add actor our actor set
 			InternalListsLock.Lock();
 			SceneActors.Add(MoveTemp(NewActor));
 			InternalListsLock.Unlock();
 		}
-	});		// end outer ParallelFor over ToBuild meshes
+	}, ParallelFlags);		// end outer ParallelFor over ToBuild meshes
 
 	check(ToBuild.Num() == 0);
 
@@ -1102,7 +1168,11 @@ void FMeshSceneAdapter::Build_FullDecompose(const FMeshSceneAdapterBuildOptions&
 		FBuildJob BuildJob = PendingBuildJobs.Pop(false);
 		ToBuildQueueLock.Unlock();
 		BuildJob.BuildWrapper->Build(BuildOptions);
-	});
+		if (BuildJob.BuildWrapper->GetTriangleCount() == 1)
+		{
+			SingleTriangleMeshes++;
+		}
+	}, ParallelFlags);
 	check(PendingBuildJobs.Num() == 0);
 
 	// currently true with the methods used above?
@@ -1110,13 +1180,36 @@ void FMeshSceneAdapter::Build_FullDecompose(const FMeshSceneAdapterBuildOptions&
 
 	if (BuildOptions.bPrintDebugMessages)
 	{
-		UE_LOG(LogTemp, Warning, TEXT("[FMeshSceneAdapter] decomposed %d source meshes used in %d instances (of %d total source meshes), into %d new instances containing %ld unique triangles (%ld total instanced). Skipped %d decompositions."), DecomposedSourceMeshCount.load(), SourceInstancesCount.load(), NumInitialSources, NewInstancesCount.load(), AddedUniqueTrisCount, InstancedTrisCount, SkippedDecompositionCount.load())
+		UE_LOG(LogTemp, Warning, TEXT("[FMeshSceneAdapter] decomposed %d source meshes used in %d instances (of %d total source meshes with %ld unique triangles), into %d new instances containing %ld unique triangles (%ld total instanced). Skipped %d decompositions. Skipped %d tiny components (%d instances, %d total triangles). %d 1-triangle meshes."), 
+			   DecomposedSourceMeshCount.load(), SourceInstancesCount.load(), NumInitialSources, NumSourceUniqueTris, 
+			   NewInstancesCount.load(), AddedUniqueTrisCount, InstancedTrisCount, 
+			   SkippedDecompositionCount.load(), 
+			   NumTinyComponents.load(), NumTinyInstances.load(), TinyInstanceTotalTriangles.load(),
+			   SingleTriangleMeshes.load())
+
+
+		Stats.Sort([](const FProcessedSourceMeshStats& A, const FProcessedSourceMeshStats& B)
+		{
+			return A.NewUniqueTris > B.NewUniqueTris;
+		});
+		for (int32 k = 0; k < 20; ++k)
+		{
+			const FProcessedSourceMeshStats& Stat = Stats[k];
+			FString NewSubmeshesStats;
+			for (FVector2i V : Stat.SubmeshSizeCount)
+			{
+				NewSubmeshesStats += FString::Printf(TEXT(" (%d,%d)"), V.X, V.Y);
+			}
+			UE_LOG(LogTemp, Warning, TEXT("  %s : SourceTris %d Inst %d ResultTris %d Inst %d | UniqueScales %d Components %d | KeptTris %d NewUniqueTris %d | (Tri,Inst) %s"),
+				   *Stat.AssetName, Stat.SourceTriangles, Stat.SourceInstances, Stat.TotalNumUniqueTris, (Stat.NewInstances + Stat.SourceInstances),
+				   Stat.UniqueInstanceScales, Stat.SourceComponents,
+				   Stat.InstancedSubmeshTris, Stat.NewUniqueTris, *NewSubmeshesStats);
+		}
 	}
 }
 
 
 
-
 void FMeshSceneAdapter::GetGeometryStatistics(FStatistics& StatsOut)
 {
 	StatsOut.UniqueMeshCount = 0;
diff --git a/Engine/Plugins/Experimental/MeshModelingToolset/Source/ModelingComponents/Public/Scene/MeshSceneAdapter.h b/Engine/Plugins/Experimental/MeshModelingToolset/Source/ModelingComponents/Public/Scene/MeshSceneAdapter.h
--- a/Engine/Plugins/Experimental/MeshModelingToolset/Source/ModelingComponents/Public/Scene/MeshSceneAdapter.h
+++ b/Engine/Plugins/Experimental/MeshModelingToolset/Source/ModelingComponents/Public/Scene/MeshSceneAdapter.h
@@ -71,6 +71,11 @@ struct FMeshSceneAdapterBuildOptions
 	/** Thickness used for bThickenThinMeshes processing */
 	double DesiredMinThickness = 0.1;
 
+	/** If true, tiny objects will be discarded from the mesh scene. This can improve performance when (eg) voxelizing huge scenes */
+	bool bFilterTinyObjects = false;
+	/** If bFilterTinyObjects is enabled, then a tiny object is identified by having a maximum (transformed) bounding-box dimension below this size */
+	double TinyObjectBoxMaxDimension = 0.001;
+
 
 	bool bBuildSpatialDataStructures = true;
 };
@@ -245,7 +250,6 @@ public:
 	virtual void GetAccumulatedMesh(FDynamicMesh3& AccumMesh);
 
 protected:
-
 	// top-level list of ActorAdapters, which represent each Actor and set of Components
 	TArray<TUniquePtr<FActorAdapter>> SceneActors;
 
