ModelingTools: clean up old PDI-based mesh edge rendering in Weld, Generate Polygroups, Remesh, Simplify. Replace with usage of UMeshElementsVisualizer and/or UPreviewGeometry. Also clean up includes and port to UE::ToolTarget:: APIs.

UMeshElementsVisualizer::SetMeshAccessFunction() now takes a TFunction with a TFunctionRef argument, instead of with a FDynamicMesh3* argument. This allows the UMeshElementsVisualizer to access the target mesh via the ProcessMesh()-style call on a UPreviewMesh, UDynamicMesh, etc, rather than receiving direct pointer access. See UWeldMeshEdgesTool::Setup() for example usage.

Add UMeshOpPreviewWithBackgroundCompute::ProcessCurrentMesh(), which forwards to the embedded UPreviewMesh::ProcessMesh()  (which can be used with the above)

Add UE::ToolTarget::GetMeshDescriptionCopy() to ModelingToolTargetUtil, similar to existing DynamicMesh variant that will auto-compute tangents if necessary.

#rnx

diff --git a/Engine/Plugins/Experimental/MeshModelingToolset/Source/MeshModelingTools/Private/ConvertToPolygonsTool.cpp b/Engine/Plugins/Experimental/MeshModelingToolset/Source/MeshModelingTools/Private/ConvertToPolygonsTool.cpp
--- a/Engine/Plugins/Experimental/MeshModelingToolset/Source/MeshModelingTools/Private/ConvertToPolygonsTool.cpp
+++ b/Engine/Plugins/Experimental/MeshModelingToolset/Source/MeshModelingTools/Private/ConvertToPolygonsTool.cpp
@@ -3,27 +3,23 @@
 #include "ConvertToPolygonsTool.h"
 #include "InteractiveToolManager.h"
 #include "ToolBuilderUtil.h"
+#include "ToolSetupUtil.h"
+#include "ModelingToolTargetUtil.h"
+#include "Drawing/PreviewGeometryActor.h"
+#include "PreviewMesh.h"
+
+#include "MeshOpPreviewHelpers.h"
+#include "ModelingOperators.h"
 
 #include "DynamicMesh/DynamicMesh3.h"
 #include "DynamicMesh/MeshNormals.h"
 #include "DynamicMeshEditor.h"
 #include "MeshRegionBoundaryLoops.h"
-#include "MeshDescriptionToDynamicMesh.h"
-#include "DynamicMeshToMeshDescription.h"
 #include "Util/ColorConstants.h"
-#include "ToolSetupUtil.h"
 #include "Polygroups/PolygroupUtil.h"
-#include "PreviewMesh.h"
-
-#include "SceneManagement.h" // for FPrimitiveDrawInterface
-#include "ModelingOperators.h"
-#include "MeshOpPreviewHelpers.h"
-#include "ModelingOperators.h"
+#include "Util/ColorConstants.h"
 
-#include "TargetInterfaces/MaterialProvider.h"
-#include "TargetInterfaces/MeshDescriptionCommitter.h"
-#include "TargetInterfaces/MeshDescriptionProvider.h"
-#include "TargetInterfaces/PrimitiveComponentBackedTarget.h"
+#include "FindPolygonsAlgorithm.h"
 
 #include "ExplicitUseGeometryMathTypes.h"		// using UE::Geometry::(math types)
 using namespace UE::Geometry;
@@ -149,73 +145,61 @@ UConvertToPolygonsTool::UConvertToPolygonsTool()
 
 bool UConvertToPolygonsTool::CanAccept() const
 {
-	return Super::CanAccept() && (PreviewWithBackgroundCompute == nullptr || PreviewWithBackgroundCompute->HaveValidResult());
+	return Super::CanAccept() && (PreviewCompute == nullptr || PreviewCompute->HaveValidResult());
 }
 
 void UConvertToPolygonsTool::Setup()
 {
 	UInteractiveTool::Setup();
 
-	IPrimitiveComponentBackedTarget* TargetComponent = Cast<IPrimitiveComponentBackedTarget>(Target);
+	FComponentMaterialSet MaterialSet = UE::ToolTarget::GetMaterialSet(Target);
 
-	FComponentMaterialSet MaterialSet;
-	Cast<IMaterialProvider>(Target)->GetMaterialSet(MaterialSet);
-
-	// populate the OriginalDynamicMesh with a conversion of the input mesh.
-	{
-		OriginalDynamicMesh = MakeShared<FDynamicMesh3, ESPMode::ThreadSafe>();
-		FMeshDescriptionToDynamicMesh Converter;
-		Converter.Convert(Cast<IMeshDescriptionProvider>(Target)->GetMeshDescription(), *OriginalDynamicMesh);
-	}
+	OriginalDynamicMesh = MakeShared<FDynamicMesh3, ESPMode::ThreadSafe>(UE::ToolTarget::GetDynamicMeshCopy(Target));
 
 	Settings = NewObject<UConvertToPolygonsToolProperties>(this);
 	Settings->RestoreProperties(this);
 	AddToolPropertySource(Settings);
-	FTransform MeshTransform = TargetComponent->GetWorldTransform();
-	// hide existing mesh
-	TargetComponent->SetOwnerVisibility(false);
-	// Set up the preview object
+	FTransform MeshTransform = (FTransform)UE::ToolTarget::GetLocalToWorldTransform(Target);
+	UE::ToolTarget::HideSourceObject(Target);
+
 	{
 		// create the operator factory
 		UConvertToPolygonsOperatorFactory* ConvertToPolygonsOperatorFactory = NewObject<UConvertToPolygonsOperatorFactory>(this);
 		ConvertToPolygonsOperatorFactory->ConvertToPolygonsTool = this; // set the back pointer
 
-
-		PreviewWithBackgroundCompute = NewObject<UMeshOpPreviewWithBackgroundCompute>(ConvertToPolygonsOperatorFactory, "Preview");
-		PreviewWithBackgroundCompute->Setup(this->TargetWorld, ConvertToPolygonsOperatorFactory);
-		PreviewWithBackgroundCompute->SetIsMeshTopologyConstant(true, EMeshRenderAttributeFlags::Positions | EMeshRenderAttributeFlags::VertexNormals);
+		PreviewCompute = NewObject<UMeshOpPreviewWithBackgroundCompute>(ConvertToPolygonsOperatorFactory);
+		PreviewCompute->Setup(this->TargetWorld, ConvertToPolygonsOperatorFactory);
+		PreviewCompute->SetIsMeshTopologyConstant(true, EMeshRenderAttributeFlags::Positions | EMeshRenderAttributeFlags::VertexNormals);
 
 		// Give the preview something to display
-		PreviewWithBackgroundCompute->PreviewMesh->SetTransform(TargetComponent->GetWorldTransform());
-		PreviewWithBackgroundCompute->PreviewMesh->SetTangentsMode(EDynamicMeshComponentTangentsMode::AutoCalculated);
-		PreviewWithBackgroundCompute->PreviewMesh->UpdatePreview(OriginalDynamicMesh.Get());
+		PreviewCompute->PreviewMesh->SetTransform(MeshTransform);
+		PreviewCompute->PreviewMesh->SetTangentsMode(EDynamicMeshComponentTangentsMode::AutoCalculated);
+		PreviewCompute->PreviewMesh->UpdatePreview(OriginalDynamicMesh.Get());
 		
-		PreviewWithBackgroundCompute->ConfigureMaterials(MaterialSet.Materials,
+		PreviewCompute->ConfigureMaterials(MaterialSet.Materials,
 			ToolSetupUtil::GetDefaultWorkingMaterial(GetToolManager())
 		);
 
 		// show the preview mesh
-		PreviewWithBackgroundCompute->SetVisibility(true);
+		PreviewCompute->SetVisibility(true);
 
 		// something to capture the polygons from the async task when it is done
-		PreviewWithBackgroundCompute->OnOpCompleted.AddLambda([this](const FDynamicMeshOperator* MeshOp) 
+		PreviewCompute->OnOpCompleted.AddLambda([this](const FDynamicMeshOperator* MeshOp) 
 		{ 
 			const FConvertToPolygonsOp*  ConvertToPolygonsOp = static_cast<const FConvertToPolygonsOp*>(MeshOp);
-				
-			// edges used for tool ::Render() method 
 			this->PolygonEdges = ConvertToPolygonsOp->Polygons.PolygonEdges;
-				
-			// we have new triangle groups to color
 			UpdateVisualization();
 		});
 
-		// updates the triangle color visualization
-		UpdateVisualization();
 		// start the compute
-		PreviewWithBackgroundCompute->InvalidateResult();
-		
+		PreviewCompute->InvalidateResult();
 	}
 	
+	PreviewGeometry = NewObject<UPreviewGeometry>(this);
+	PreviewGeometry->CreateInWorld(TargetWorld, MeshTransform);
+
+	// updates the triangle color visualization
+	UpdateVisualization();
 
 	Settings->WatchProperty(Settings->ConversionMode,
 							[this](EConvertToPolygonsMode NewMode)
@@ -240,74 +224,54 @@ void UConvertToPolygonsTool::UpdateOpParameters(FConvertToPolygonsOp& ConvertToP
 	ConvertToPolygonsOp.AngleTolerance    = Settings->AngleTolerance;
 	ConvertToPolygonsOp.OriginalMesh = OriginalDynamicMesh;
 	
-	FTransform LocalToWorld = Cast<IPrimitiveComponentBackedTarget>(Target)->GetWorldTransform();
+	FTransform LocalToWorld = (FTransform)UE::ToolTarget::GetLocalToWorldTransform(Target);
 	ConvertToPolygonsOp.SetTransform(LocalToWorld);
 }
 
-void UConvertToPolygonsTool::GenerateAsset(const FDynamicMeshOpResult& Result)
-{
-	GetToolManager()->BeginUndoTransaction(LOCTEXT("ConvertToPolygonsToolTransactionName", "Find Polygroups"));
 
-	FDynamicMesh3* DynamicMeshResult = Result.Mesh.Get();
-	check(DynamicMeshResult != nullptr);
-
-	Cast<IMeshDescriptionCommitter>(Target)->CommitMeshDescription([DynamicMeshResult](const IMeshDescriptionCommitter::FCommitterParams& CommitParams)
-	{
-		FDynamicMeshToMeshDescription Converter;
-		Converter.Convert(DynamicMeshResult, *CommitParams.MeshDescriptionOut);
-	});
-
-	GetToolManager()->EndUndoTransaction();
-}
 
 void UConvertToPolygonsTool::Shutdown(EToolShutdownType ShutdownType)
 {
 	Settings->SaveProperties(this);
-	Cast<IPrimitiveComponentBackedTarget>(Target)->SetOwnerVisibility(true);
+	UE::ToolTarget::ShowSourceObject(Target);
+
+	PreviewGeometry->Disconnect();
+	PreviewGeometry = nullptr;
 
-	if (PreviewWithBackgroundCompute)
+	if (PreviewCompute)
 	{
-		FDynamicMeshOpResult Result = PreviewWithBackgroundCompute->Shutdown();
+		FDynamicMeshOpResult Result = PreviewCompute->Shutdown();
 		if (ShutdownType == EToolShutdownType::Accept)
 		{
-			GenerateAsset(Result);
+			GetToolManager()->BeginUndoTransaction(LOCTEXT("ConvertToPolygonsToolTransactionName", "Find Polygroups"));
+			FDynamicMesh3* DynamicMeshResult = Result.Mesh.Get();
+			if (ensure(DynamicMeshResult != nullptr))
+			{
+				// todo: have not actually modified topology here, but groups-only update is not supported yet
+				UE::ToolTarget::CommitDynamicMeshUpdate(Target, *DynamicMeshResult, true);
+			}
+			GetToolManager()->EndUndoTransaction();
 		}
 	}
 }
 
 void UConvertToPolygonsTool::OnSettingsModified()
 {
-	PreviewWithBackgroundCompute->InvalidateResult();
+	PreviewCompute->InvalidateResult();
 }
 
 
 void UConvertToPolygonsTool::OnTick(float DeltaTime)
 {
-	PreviewWithBackgroundCompute->Tick(DeltaTime);
+	PreviewCompute->Tick(DeltaTime);
 }
 
-void UConvertToPolygonsTool::Render(IToolsContextRenderAPI* RenderAPI)
-{
-	FColor LineColor(255, 0, 0);
-
-	FPrimitiveDrawInterface* PDI = RenderAPI->GetPrimitiveDrawInterface();
-	float PDIScale = RenderAPI->GetCameraState().GetPDIScalingFactor();
-	FTransform Transform = Cast<IPrimitiveComponentBackedTarget>(Target)->GetWorldTransform(); //Actor->GetTransform();
 
-	for (int eid : PolygonEdges)
-	{
-		FVector3d A, B;
-		OriginalDynamicMesh->GetEdgeV(eid, A, B);
-		PDI->DrawLine(Transform.TransformPosition((FVector)A), Transform.TransformPosition((FVector)B),
-			LineColor, 0, 2.0*PDIScale, 1.0f, true);
-	}
-}
 
 
 void UConvertToPolygonsTool::UpdateVisualization()
 {
-	
-	if (!PreviewWithBackgroundCompute)
+	if (!PreviewCompute)
 	{
 		return;
 	}
@@ -321,7 +285,7 @@ void UConvertToPolygonsTool::UpdateVisualization()
 		{ 
 			MaterialSet.Materials.Add(ToolSetupUtil::GetSelectionMaterial(GetToolManager()));
 		}
-		PreviewWithBackgroundCompute->PreviewMesh->SetTriangleColorFunction([this](const FDynamicMesh3* Mesh, int TriangleID)
+		PreviewCompute->PreviewMesh->SetTriangleColorFunction([this](const FDynamicMesh3* Mesh, int TriangleID)
 		{
 			return LinearColors::SelectFColor(Mesh->GetTriangleGroup(TriangleID));
 		}, 
@@ -330,10 +294,21 @@ void UConvertToPolygonsTool::UpdateVisualization()
 	else
 	{
 		MaterialTarget->GetMaterialSet(MaterialSet);
-		PreviewWithBackgroundCompute->PreviewMesh->ClearTriangleColorFunction(UPreviewMesh::ERenderUpdateMode::FastUpdate);
+		PreviewCompute->PreviewMesh->ClearTriangleColorFunction(UPreviewMesh::ERenderUpdateMode::FastUpdate);
 	}
-	PreviewWithBackgroundCompute->ConfigureMaterials(MaterialSet.Materials,
+	PreviewCompute->ConfigureMaterials(MaterialSet.Materials,
 		ToolSetupUtil::GetDefaultWorkingMaterial(GetToolManager()));
+
+	FColor GroupLineColor = FColor::Red;
+	float GroupLineThickness = 2.0f;
+
+	PreviewGeometry->CreateOrUpdateLineSet(TEXT("GroupBorders"), PolygonEdges.Num(), 
+		[&](int32 k, TArray<FRenderableLine>& LinesOut) {
+			FVector3d A, B;
+			OriginalDynamicMesh->GetEdgeV(PolygonEdges[k], A, B);
+			LinesOut.Add(FRenderableLine(A, B, GroupLineColor, GroupLineThickness));
+		}, 1);
+
 }
 
 
diff --git a/Engine/Plugins/Experimental/MeshModelingToolset/Source/MeshModelingTools/Private/MeshGroupPaintTool.cpp b/Engine/Plugins/Experimental/MeshModelingToolset/Source/MeshModelingTools/Private/MeshGroupPaintTool.cpp
--- a/Engine/Plugins/Experimental/MeshModelingToolset/Source/MeshModelingTools/Private/MeshGroupPaintTool.cpp
+++ b/Engine/Plugins/Experimental/MeshModelingToolset/Source/MeshModelingTools/Private/MeshGroupPaintTool.cpp
@@ -214,7 +214,9 @@ void UMeshGroupPaintTool::Setup()
 		MeshElementsDisplay->Settings->RestoreProperties(this);
 		AddToolPropertySource(MeshElementsDisplay->Settings);
 	}
-	MeshElementsDisplay->SetMeshAccessFunction([&](void) { return GetSculptMesh(); });
+	MeshElementsDisplay->SetMeshAccessFunction([this](UMeshElementsVisualizer::ProcessDynamicMeshFunc ProcessFunc) {
+		ProcessFunc(*GetSculptMesh());
+	});
 
 	// force colors update... ?
 	DynamicMeshComponent->SetTriangleColorFunction([this](const FDynamicMesh3* Mesh, int TriangleID)
diff --git a/Engine/Plugins/Experimental/MeshModelingToolset/Source/MeshModelingTools/Private/RemeshMeshTool.cpp b/Engine/Plugins/Experimental/MeshModelingToolset/Source/MeshModelingTools/Private/RemeshMeshTool.cpp
--- a/Engine/Plugins/Experimental/MeshModelingToolset/Source/MeshModelingTools/Private/RemeshMeshTool.cpp
+++ b/Engine/Plugins/Experimental/MeshModelingToolset/Source/MeshModelingTools/Private/RemeshMeshTool.cpp
@@ -5,16 +5,20 @@
 #include "InteractiveToolManager.h"
 #include "ToolBuilderUtil.h"
 
+#include "Properties/RemeshProperties.h"
+#include "Properties/MeshStatisticsProperties.h"
+#include "Drawing/MeshElementsVisualizer.h"
+
 #include "Util/ColorConstants.h"
 #include "ToolSetupUtil.h"
-
+#include "ModelingToolTargetUtil.h"
+#include "ToolSetupUtil.h"
 #include "DynamicMesh/DynamicMesh3.h"
+#include "DynamicMesh/DynamicMeshAABBTree3.h"
 
 #include "MeshDescriptionToDynamicMesh.h"
 #include "DynamicMeshToMeshDescription.h"
 
-#include "SceneManagement.h" // for FPrimitiveDrawInterface
-
 #include "TargetInterfaces/MaterialProvider.h"
 #include "TargetInterfaces/MeshDescriptionCommitter.h"
 #include "TargetInterfaces/MeshDescriptionProvider.h"
@@ -70,7 +74,6 @@ URemeshMeshToolProperties::URemeshMeshToolProperties()
 	RemeshType = ERemeshType::Standard;
 	SmoothingType = ERemeshSmoothingType::MeanValue;
 	bPreserveSharpEdges = true;
-	bShowWireframe = true;
 	bShowGroupColors = false;
 
 	TargetEdgeLength = 5.0;
@@ -106,27 +109,19 @@ void URemeshMeshTool::Setup()
 	check(Targets[0]);
 	IPrimitiveComponentBackedTarget* TargetComponent = TargetComponentInterface(0);
 	IMeshDescriptionProvider* TargetMeshProvider = TargetMeshProviderInterface(0);
-	IMaterialProvider* TargetMaterial = TargetMaterialInterface(0);
+
 
 	// hide component and create + show preview
 	TargetComponent->SetOwnerVisibility(false);
-	Preview = NewObject<UMeshOpPreviewWithBackgroundCompute>(this, "Preview");
+	Preview = NewObject<UMeshOpPreviewWithBackgroundCompute>(this);
 	Preview->Setup(this->TargetWorld, this);
-	Preview->OnMeshUpdated.AddLambda([this](UMeshOpPreviewWithBackgroundCompute* Compute)
-	{
-		MeshStatisticsProperties->Update(*Compute->PreviewMesh->GetPreviewDynamicMesh());
-	});
-	FComponentMaterialSet MaterialSet;
-	TargetMaterial->GetMaterialSet(MaterialSet);
+
+	FComponentMaterialSet MaterialSet = UE::ToolTarget::GetMaterialSet(Targets[0]);
 	Preview->ConfigureMaterials( MaterialSet.Materials,
 								 ToolSetupUtil::GetDefaultWorkingMaterial(GetToolManager())
 	);
-	Preview->PreviewMesh->EnableWireframe(BasicProperties->bShowWireframe);
-
 	BasicProperties->WatchProperty(BasicProperties->bShowGroupColors,
 								   [this](bool bNewValue) { UpdateVisualization();});
-	BasicProperties->WatchProperty(BasicProperties->bShowWireframe,
-								   [this](bool bNewValue) { UpdateVisualization();});
 
 	OriginalMesh = MakeShared<FDynamicMesh3, ESPMode::ThreadSafe>();
 	FMeshDescriptionToDynamicMesh Converter;
@@ -156,6 +151,29 @@ void URemeshMeshTool::Setup()
 	AddToolPropertySource(BasicProperties);
 	AddToolPropertySource(MeshStatisticsProperties);
 
+
+	MeshElementsDisplay = NewObject<UMeshElementsVisualizer>(this);
+	MeshElementsDisplay->CreateInWorld(Preview->PreviewMesh->GetWorld(), Preview->PreviewMesh->GetTransform());
+	if (ensure(MeshElementsDisplay->Settings))
+	{
+		MeshElementsDisplay->Settings->bShowWireframe = true;
+		MeshElementsDisplay->Settings->RestoreProperties(this, TEXT("Remesh"));
+		AddToolPropertySource(MeshElementsDisplay->Settings);
+	}
+	MeshElementsDisplay->SetMeshAccessFunction([this](UMeshElementsVisualizer::ProcessDynamicMeshFunc ProcessFunc) {
+		Preview->ProcessCurrentMesh(ProcessFunc);
+	});
+
+	Preview->OnMeshUpdated.AddLambda([this](UMeshOpPreviewWithBackgroundCompute* Compute)
+	{
+		Compute->ProcessCurrentMesh([&](const FDynamicMesh3& ReadMesh)
+		{
+			MeshStatisticsProperties->Update(ReadMesh);
+			MeshElementsDisplay->NotifyMeshChanged();
+		});
+	});
+
+
 	Preview->InvalidateResult();
 
 	SetToolDisplayName(LOCTEXT("ToolName", "Remesh"));
@@ -167,20 +185,35 @@ void URemeshMeshTool::Setup()
 void URemeshMeshTool::Shutdown(EToolShutdownType ShutdownType)
 {
 	BasicProperties->SaveProperties(this);
+
+	if (ensure(MeshElementsDisplay->Settings))
+	{
+		MeshElementsDisplay->Settings->SaveProperties(this, TEXT("Remesh"));
+	}
+	MeshElementsDisplay->Disconnect();
+
+
 	for (int ComponentIdx = 0; ComponentIdx < Targets.Num(); ComponentIdx++)
 	{
-		TargetComponentInterface(ComponentIdx)->SetOwnerVisibility(true);
+		UE::ToolTarget::ShowSourceObject(Targets[ComponentIdx]);
 	}
+
 	FDynamicMeshOpResult Result = Preview->Shutdown();
 	if (ShutdownType == EToolShutdownType::Accept)
 	{
-		GenerateAsset(Result);
+		if (ShutdownType == EToolShutdownType::Accept)
+		{
+			GetToolManager()->BeginUndoTransaction(LOCTEXT("RemeshMeshToolTransactionName", "Remesh Mesh"));
+			UE::ToolTarget::CommitDynamicMeshUpdate(Targets[0], *Result.Mesh, true);
+			GetToolManager()->EndUndoTransaction();
+		}
 	}
 }
 
 void URemeshMeshTool::OnTick(float DeltaTime)
 {
 	Preview->Tick(DeltaTime);
+	MeshElementsDisplay->OnTick(DeltaTime);
 }
 
 TUniquePtr<FDynamicMeshOperator> URemeshMeshTool::MakeNewOperator()
@@ -215,8 +248,7 @@ TUniquePtr<FDynamicMeshOperator> URemeshMeshTool::MakeNewOperator()
 	Op->SmoothingStrength = BasicProperties->SmoothingStrength;
 	Op->SmoothingType = BasicProperties->SmoothingType;
 
-	check(Targets.Num() > 0);
-	FTransform LocalToWorld = TargetComponentInterface(0)->GetWorldTransform();
+	FTransform LocalToWorld = (FTransform)UE::ToolTarget::GetLocalToWorldTransform(Targets[0]);
 	Op->SetTransform(LocalToWorld);
 
 	Op->OriginalMesh = OriginalMesh;
@@ -228,43 +260,12 @@ TUniquePtr<FDynamicMeshOperator> URemeshMeshTool::MakeNewOperator()
 	return Op;
 }
 
-void URemeshMeshTool::Render(IToolsContextRenderAPI* RenderAPI)
-{
-	FPrimitiveDrawInterface* PDI = RenderAPI->GetPrimitiveDrawInterface();
-	FTransform Transform = TargetComponentInterface(0)->GetWorldTransform();
-
-	if (BasicProperties->bShowConstraintEdges)
-	{
-		FColor LineColor(255, 0, 0);
-		const FDynamicMesh3* TargetMesh = Preview->PreviewMesh->GetPreviewDynamicMesh();
-		if (TargetMesh && TargetMesh->HasAttributes())
-		{
-			float PDIScale = RenderAPI->GetCameraState().GetPDIScalingFactor();
-			const FDynamicMeshUVOverlay* UVOverlay = TargetMesh->Attributes()->PrimaryUV();
-			for (int eid : TargetMesh->EdgeIndicesItr())
-			{
-				if (UVOverlay->IsSeamEdge(eid))
-				{
-					FVector3d A, B;
-					TargetMesh->GetEdgeV(eid, A, B);
-					PDI->DrawLine(Transform.TransformPosition((FVector)A), Transform.TransformPosition((FVector)B),
-						LineColor, 0, 2.0 * PDIScale, 1.0f, true);
-				}
-			}
-		}
-	}
-}
 
 void URemeshMeshTool::OnPropertyModified(UObject* PropertySet, FProperty* Property)
 {
 	if ( Property )
 	{
-		if ( ( Property->GetFName() == GET_MEMBER_NAME_CHECKED(URemeshMeshToolProperties, bShowWireframe) ) ||
-			 ( Property->GetFName() == GET_MEMBER_NAME_CHECKED(URemeshMeshToolProperties, bShowGroupColors) ) )
-		{
-			//UpdateVisualization();
-		}
-		else
+		if ( Property->GetFName() != GET_MEMBER_NAME_CHECKED(URemeshMeshToolProperties, bShowGroupColors) )
 		{
 			Preview->InvalidateResult();
 		}
@@ -273,7 +274,6 @@ void URemeshMeshTool::OnPropertyModified(UObject* PropertySet, FProperty* Proper
 
 void URemeshMeshTool::UpdateVisualization()
 {
-	Preview->PreviewMesh->EnableWireframe(BasicProperties->bShowWireframe);
 	FComponentMaterialSet MaterialSet;
 	if (BasicProperties->bShowGroupColors)
 	{
@@ -286,7 +286,7 @@ void URemeshMeshTool::UpdateVisualization()
 	}
 	else
 	{
-		TargetMaterialInterface(0)->GetMaterialSet(MaterialSet);
+		MaterialSet = UE::ToolTarget::GetMaterialSet(Targets[0]);
 		Preview->PreviewMesh->ClearTriangleColorFunction(UPreviewMesh::ERenderUpdateMode::FastUpdate);
 	}
 	Preview->ConfigureMaterials(MaterialSet.Materials,
@@ -305,20 +305,6 @@ bool URemeshMeshTool::CanAccept() const
 	return Super::CanAccept() && Preview->HaveValidResult();
 }
 
-void URemeshMeshTool::GenerateAsset(const FDynamicMeshOpResult& Result)
-{
-	GetToolManager()->BeginUndoTransaction(LOCTEXT("RemeshMeshToolTransactionName", "Remesh Mesh"));
-
-	check(Result.Mesh.Get() != nullptr);
-	TargetMeshCommitterInterface(0)->CommitMeshDescription([&Result](const IMeshDescriptionCommitter::FCommitterParams& CommitParams)
-	{
-		FDynamicMeshToMeshDescription Converter;
 
-		// full conversion if normal topology changed or faces were inverted
-		Converter.Convert(Result.Mesh.Get(), *CommitParams.MeshDescriptionOut);
-	});
-
-	GetToolManager()->EndUndoTransaction();
-}
 
 #undef LOCTEXT_NAMESPACE
diff --git a/Engine/Plugins/Experimental/MeshModelingToolset/Source/MeshModelingTools/Private/WeldMeshEdgesTool.cpp b/Engine/Plugins/Experimental/MeshModelingToolset/Source/MeshModelingTools/Private/WeldMeshEdgesTool.cpp
--- a/Engine/Plugins/Experimental/MeshModelingToolset/Source/MeshModelingTools/Private/WeldMeshEdgesTool.cpp
+++ b/Engine/Plugins/Experimental/MeshModelingToolset/Source/MeshModelingTools/Private/WeldMeshEdgesTool.cpp
@@ -3,16 +3,14 @@
 #include "WeldMeshEdgesTool.h"
 #include "InteractiveToolManager.h"
 #include "ToolBuilderUtil.h"
+#include "ToolSetupUtil.h"
+#include "ModelingToolTargetUtil.h"
+#include "PreviewMesh.h"
+#include "Drawing/MeshElementsVisualizer.h"
 
 #include "DynamicMesh/DynamicMesh3.h"
-#include "DynamicMesh/DynamicMeshAttributeSet.h"
 #include "DynamicMesh/Operations/MergeCoincidentMeshEdges.h"
 
-#include "Components/DynamicMeshComponent.h"
-#include "ModelingToolTargetUtil.h"
-
-#include "SceneManagement.h" // for FPrimitiveDrawInterface
-#include "Async/ParallelFor.h"
 
 #include "ExplicitUseGeometryMathTypes.h"		// using UE::Geometry::(math types)
 using namespace UE::Geometry;
@@ -30,155 +28,200 @@ USingleSelectionMeshEditingTool* UWeldMeshEdgesToolBuilder::CreateNewTool(const
 
 
 
+
+class FWeldMeshEdgesOp : public  FDynamicMeshOperator
+{
+public:
+
+	// parameters set by the tool
+	TSharedPtr<FDynamicMesh3, ESPMode::ThreadSafe> SourceMesh;
+	double Tolerance;
+	bool bOnlyUnique;
+
+public:
+	FWeldMeshEdgesOp() : FDynamicMeshOperator() {}
+
+	virtual void CalculateResult(FProgressCancel* Progress) override
+	{
+		if ((Progress && Progress->Cancelled()) || !SourceMesh)
+		{
+			return;
+		}
+
+		ResultMesh->Copy(*SourceMesh, true, true, true, true);
+	
+		FMergeCoincidentMeshEdges Merger(ResultMesh.Get());
+		Merger.MergeVertexTolerance = Tolerance;
+		Merger.MergeSearchTolerance = 2 * Merger.MergeVertexTolerance;
+		Merger.OnlyUniquePairs = bOnlyUnique;
+
+		if (Merger.Apply() == false)
+		{
+			ResultMesh->Copy(*SourceMesh, true, true, true, true);
+		}
+		else if (ResultMesh->CheckValidity(true, EValidityCheckFailMode::ReturnOnly) == false)
+		{
+			ResultMesh->Copy(*SourceMesh, true, true, true, true);
+		}
+	}
+
+	void SetTransform(const FTransform& Transform)
+	{
+		ResultTransform = (UE::Geometry::FTransform3d)Transform;
+	}
+};
+
+
+
+TUniquePtr<FDynamicMeshOperator> UWeldMeshEdgesOperatorFactory::MakeNewOperator()
+{
+	check(WeldMeshEdgesTool);
+	TUniquePtr<FWeldMeshEdgesOp> MeshOp = MakeUnique<FWeldMeshEdgesOp>();
+	WeldMeshEdgesTool->UpdateOpParameters(*MeshOp);
+	return MeshOp;
+}
+
+
+
+
 /*
  * Tool
  */
 
+
 UWeldMeshEdgesTool::UWeldMeshEdgesTool()
 {
-	Tolerance = FMergeCoincidentMeshEdges::DEFAULT_TOLERANCE;
-	bOnlyUnique = false;
+	SetToolDisplayName(LOCTEXT("WeldMeshEdgesToolName", "Weld Edges"));
 }
 
+
+
 void UWeldMeshEdgesTool::Setup()
 {
 	UInteractiveTool::Setup();
 
-	// create dynamic mesh component to use for live preview
-	DynamicMeshComponent = NewObject<UDynamicMeshComponent>(UE::ToolTarget::GetTargetActor(Target));
-	DynamicMeshComponent->SetupAttachment(UE::ToolTarget::GetTargetComponent(Target));
-	DynamicMeshComponent->RegisterComponent();
-	DynamicMeshComponent->SetWorldTransform((FTransform)UE::ToolTarget::GetLocalToWorldTransform(Target));
 
-	// transfer materials
+	SourceMesh = MakeShared<FDynamicMesh3, ESPMode::ThreadSafe>(UE::ToolTarget::GetDynamicMeshCopy(Target, true));
+
+	FTransform MeshTransform = (FTransform)UE::ToolTarget::GetLocalToWorldTransform(Target);
+
+	OperatorFactory = NewObject<UWeldMeshEdgesOperatorFactory>(this);
+	OperatorFactory->WeldMeshEdgesTool = this;
+
+	PreviewCompute = NewObject<UMeshOpPreviewWithBackgroundCompute>(OperatorFactory);
+	PreviewCompute->Setup(this->TargetWorld, OperatorFactory);
+
+	// Give the preview something to display
+	PreviewCompute->PreviewMesh->SetTransform(MeshTransform);
+	PreviewCompute->PreviewMesh->SetTangentsMode(EDynamicMeshComponentTangentsMode::ExternallyProvided);
+	PreviewCompute->PreviewMesh->UpdatePreview(SourceMesh.Get());
+
 	FComponentMaterialSet MaterialSet = UE::ToolTarget::GetMaterialSet(Target);
-	for (int k = 0; k < MaterialSet.Materials.Num(); ++k)
+	PreviewCompute->ConfigureMaterials(MaterialSet.Materials,
+									   ToolSetupUtil::GetDefaultWorkingMaterial(GetToolManager()) );
+
+	PreviewCompute->SetVisibility(true);
+	UE::ToolTarget::HideSourceObject(Target);
+
+
+	Settings = NewObject<UWeldMeshEdgesToolProperties>(this);
+	Settings->RestoreProperties(this);
+	AddToolPropertySource(Settings);
+	Settings->WatchProperty(Settings->Tolerance, [this](float) { PreviewCompute->InvalidateResult(); });
+	Settings->WatchProperty(Settings->bOnlyUnique, [this](bool) { PreviewCompute->InvalidateResult(); });
+
+	// create mesh display
+	MeshElementsDisplay = NewObject<UMeshElementsVisualizer>(this);
+	MeshElementsDisplay->CreateInWorld(PreviewCompute->PreviewMesh->GetWorld(), PreviewCompute->PreviewMesh->GetTransform());
+	if (ensure(MeshElementsDisplay->Settings))
 	{
-		DynamicMeshComponent->SetMaterial(k, MaterialSet.Materials[k]);
+		MeshElementsDisplay->Settings->bShowUVSeams = false;
+		MeshElementsDisplay->Settings->bShowNormalSeams = false;
+		MeshElementsDisplay->Settings->bShowColorSeams = false;
+		MeshElementsDisplay->Settings->RestoreProperties(this, TEXT("WeldEdges"));
+		AddToolPropertySource(MeshElementsDisplay->Settings);
 	}
+	MeshElementsDisplay->SetMeshAccessFunction([this](UMeshElementsVisualizer::ProcessDynamicMeshFunc ProcessFunc) {
+		PreviewCompute->ProcessCurrentMesh(ProcessFunc);
+	});
 
-	DynamicMeshComponent->SetTangentsType(EDynamicMeshComponentTangentsMode::AutoCalculated);
-	DynamicMeshComponent->SetMesh(UE::ToolTarget::GetDynamicMeshCopy(Target));
-	DynamicMeshComponent->ProcessMesh([&](const FDynamicMesh3& ReadMesh) { OriginalMesh = ReadMesh; });
 
-	// hide input StaticMeshComponent
-	UE::ToolTarget::HideSourceObject(Target);
+	PreviewCompute->OnMeshUpdated.AddLambda([this](UMeshOpPreviewWithBackgroundCompute* Compute)
+	{
+		Compute->ProcessCurrentMesh([&](const FDynamicMesh3& ReadMesh)
+		{
+			MeshElementsDisplay->NotifyMeshChanged();
+		});
+	});
+
 
-	// initialize our properties
-	ToolPropertyObjects.Add(this);
+	PreviewCompute->InvalidateResult();
 
-	bResultValid = false;
 
 	SetToolDisplayName(LOCTEXT("ToolName", "Weld Edges"));
 	GetToolManager()->DisplayMessage(
 		LOCTEXT("WeldMeshEdgesToolDescription", "Weld overlapping/identical border edges of the selected Mesh, by merging the vertices."),
 		EToolMessageLevel::UserNotification);
+
 }
 
 
 void UWeldMeshEdgesTool::Shutdown(EToolShutdownType ShutdownType)
 {
-	if (DynamicMeshComponent != nullptr)
+	Settings->SaveProperties(this);
+	UE::ToolTarget::ShowSourceObject(Target);
+
+	if (ensure(MeshElementsDisplay->Settings))
 	{
-		UE::ToolTarget::ShowSourceObject(Target);
+		MeshElementsDisplay->Settings->SaveProperties(this, TEXT("WeldEdges"));
+	}
+	MeshElementsDisplay->Disconnect();
 
+	if (PreviewCompute)
+	{
+		FDynamicMeshOpResult Result = PreviewCompute->Shutdown();
 		if (ShutdownType == EToolShutdownType::Accept)
 		{
-			// this block bakes the modified DynamicMeshComponent back into the StaticMeshComponent inside an undo transaction
-			GetToolManager()->BeginUndoTransaction(LOCTEXT("WeldMeshEdgesToolTransactionName", "Remesh Mesh"));
-			DynamicMeshComponent->ProcessMesh([&](const FDynamicMesh3& CurMesh)
+			GetToolManager()->BeginUndoTransaction(LOCTEXT("WeldMeshEdgesToolTransactionName", "Weld Edges"));
+			FDynamicMesh3* DynamicMeshResult = Result.Mesh.Get();
+			if (ensure(DynamicMeshResult != nullptr))
 			{
-				UE::ToolTarget::CommitDynamicMeshUpdate(Target, CurMesh, true);
-			});
+				// todo: have not actually modified topology here, but groups-only update is not supported yet
+				UE::ToolTarget::CommitDynamicMeshUpdate(Target, *DynamicMeshResult, true);
+			}
 			GetToolManager()->EndUndoTransaction();
 		}
-
-		DynamicMeshComponent->UnregisterComponent();
-		DynamicMeshComponent->DestroyComponent();
-		DynamicMeshComponent = nullptr;
 	}
 }
 
 
-void UWeldMeshEdgesTool::Render(IToolsContextRenderAPI* RenderAPI)
-{
-	UpdateResult();
-
-	FPrimitiveDrawInterface* PDI = RenderAPI->GetPrimitiveDrawInterface();
-	FTransform Transform = (FTransform)UE::ToolTarget::GetLocalToWorldTransform(Target);
-
-	FColor LineColor(200, 200, 200);
-	float PDIScale = RenderAPI->GetCameraState().GetPDIScalingFactor();
-	FDynamicMesh3* TargetMesh = DynamicMeshComponent->GetMesh();
-	FDynamicMeshUVOverlay* UVOverlay = TargetMesh->Attributes()->PrimaryUV();
-	for (int eid : TargetMesh->EdgeIndicesItr()) 
-	{
-		if (UVOverlay->IsSeamEdge(eid)) 
-		{
-			FVector3d A, B;
-			TargetMesh->GetEdgeV(eid, A, B);
-			PDI->DrawLine(Transform.TransformPosition((FVector)A), Transform.TransformPosition((FVector)B),
-				LineColor, 0, 1.0f*PDIScale, 1.0f, true);
-		}
-	}
 
+bool UWeldMeshEdgesTool::CanAccept() const
+{
+	return Super::CanAccept() && (PreviewCompute == nullptr || PreviewCompute->HaveValidResult());
+}
 
-	FColor LineColor2(255, 0, 0);
-	for (int eid : TargetMesh->BoundaryEdgeIndicesItr()) 
-	{
-		FVector3d A, B;
-		TargetMesh->GetEdgeV(eid, A, B);
-		PDI->DrawLine(Transform.TransformPosition((FVector)A), Transform.TransformPosition((FVector)B),
-			LineColor2, 0, 2.0f*PDIScale, 1.0f, true);
-	}
 
-}
 
-#if WITH_EDITOR
-void UWeldMeshEdgesTool::PostEditChangeProperty(FPropertyChangedEvent& PropertyChangedEvent)
+void UWeldMeshEdgesTool::OnTick(float DeltaTime)
 {
-	FProperty* PropertyThatChanged = PropertyChangedEvent.Property;
-	bResultValid = false;
+	PreviewCompute->Tick(DeltaTime);
+	MeshElementsDisplay->OnTick(DeltaTime);
 }
-#endif
 
-void UWeldMeshEdgesTool::UpdateResult()
-{
-	if (bResultValid) 
-	{
-		return;
-	}
 
-	ComputeMesh = OriginalMesh;
 
-	FMergeCoincidentMeshEdges Merger(&ComputeMesh);
-	Merger.MergeVertexTolerance = this->Tolerance;
-	Merger.MergeSearchTolerance = 2*Merger.MergeVertexTolerance;
-	Merger.OnlyUniquePairs = this->bOnlyUnique;
-	FDynamicMesh3* SetMesh = &ComputeMesh;
-		
-	if (Merger.Apply() == false)
-	{
-		UE_LOG(LogTemp, Warning, TEXT("WeldMeshEdgesTool : Merger.Apply() returned false"));
-		SetMesh = &OriginalMesh;
-	}
-	else if (ComputeMesh.CheckValidity(true, EValidityCheckFailMode::ReturnOnly) == false)
-	{
-		UE_LOG(LogTemp, Warning, TEXT("WeldMeshEdgesTool : returned mesh is invalid"));
-		SetMesh = &OriginalMesh;
-	}
-
-	DynamicMeshComponent->EditMesh([SetMesh](FDynamicMesh3& EditMesh)
-	{
-		EditMesh = *SetMesh;
-	});
+void UWeldMeshEdgesTool::UpdateOpParameters(FWeldMeshEdgesOp& Op) const
+{
+	Op.bOnlyUnique = Settings->bOnlyUnique;
+	Op.Tolerance = Settings->Tolerance;
+	Op.SourceMesh = SourceMesh;
 
-	GetToolManager()->PostInvalidation();
-	bResultValid = true;
+	FTransform LocalToWorld = (FTransform)UE::ToolTarget::GetLocalToWorldTransform(Target);
+	Op.SetTransform(LocalToWorld);
 }
 
 
 
 
-
 #undef LOCTEXT_NAMESPACE
diff --git a/Engine/Plugins/Experimental/MeshModelingToolset/Source/MeshModelingTools/Public/ConvertToPolygonsTool.h b/Engine/Plugins/Experimental/MeshModelingToolset/Source/MeshModelingTools/Public/ConvertToPolygonsTool.h
--- a/Engine/Plugins/Experimental/MeshModelingToolset/Source/MeshModelingTools/Public/ConvertToPolygonsTool.h
+++ b/Engine/Plugins/Experimental/MeshModelingToolset/Source/MeshModelingTools/Public/ConvertToPolygonsTool.h
@@ -3,20 +3,17 @@
 #pragma once
 
 #include "CoreMinimal.h"
-#include "UObject/NoExportTypes.h"
-#include "SingleSelectionTool.h"
-#include "InteractiveToolBuilder.h"
-#include "DynamicMesh/DynamicMesh3.h"
-#include "FindPolygonsAlgorithm.h"
+#include "BaseTools/SingleSelectionMeshEditingTool.h"
 #include "PreviewMesh.h"
 #include "ModelingOperators.h"
 #include "MeshOpPreviewHelpers.h"
-#include "BaseTools/SingleSelectionMeshEditingTool.h"
 #include "ConvertToPolygonsTool.generated.h"
 
 // predeclaration
 class UConvertToPolygonsTool;
 class FConvertToPolygonsOp;
+class UPreviewGeometry;
+PREDECLARE_GEOMETRY(class FDynamicMesh3);
 
 /**
  *
@@ -25,7 +22,6 @@ UCLASS()
 class MESHMODELINGTOOLS_API UConvertToPolygonsToolBuilder : public USingleSelectionMeshEditingToolBuilder
 {
 	GENERATED_BODY()
-
 public:
 	virtual USingleSelectionMeshEditingTool* CreateNewTool(const FToolBuilderState & SceneState) const override;
 };
@@ -96,7 +92,6 @@ public:
 	virtual void Setup() override;
 	virtual void Shutdown(EToolShutdownType ShutdownType) override;
 
-	virtual void Render(IToolsContextRenderAPI* RenderAPI) override;
 	virtual void OnTick(float DeltaTime) override;
 
 	virtual bool HasCancel() const override { return true; }
@@ -115,7 +110,10 @@ protected:
 	UConvertToPolygonsToolProperties* Settings;
 
 	UPROPERTY()
-	UMeshOpPreviewWithBackgroundCompute* PreviewWithBackgroundCompute = nullptr;
+	UMeshOpPreviewWithBackgroundCompute* PreviewCompute = nullptr;
+
+	UPROPERTY()
+	UPreviewGeometry* PreviewGeometry = nullptr;
 
 protected:
 	TSharedPtr<UE::Geometry::FDynamicMesh3, ESPMode::ThreadSafe> OriginalDynamicMesh;
@@ -124,6 +122,4 @@ protected:
 	TArray<int> PolygonEdges;
 	
 	void UpdateVisualization();
-	void GenerateAsset(const FDynamicMeshOpResult& Result);
-
 };
diff --git a/Engine/Plugins/Experimental/MeshModelingToolset/Source/MeshModelingTools/Public/RemeshMeshTool.h b/Engine/Plugins/Experimental/MeshModelingToolset/Source/MeshModelingTools/Public/RemeshMeshTool.h
--- a/Engine/Plugins/Experimental/MeshModelingToolset/Source/MeshModelingTools/Public/RemeshMeshTool.h
+++ b/Engine/Plugins/Experimental/MeshModelingToolset/Source/MeshModelingTools/Public/RemeshMeshTool.h
@@ -3,18 +3,18 @@
 #pragma once
 
 #include "CoreMinimal.h"
-#include "UObject/NoExportTypes.h"
 #include "MultiSelectionTool.h"
 #include "InteractiveToolBuilder.h"
-#include "DynamicMesh/DynamicMesh3.h"
-#include "DynamicMesh/DynamicMeshAABBTree3.h"
 #include "MeshOpPreviewHelpers.h"
 #include "CleaningOps/RemeshMeshOp.h"
-#include "Properties/MeshStatisticsProperties.h"
 #include "Properties/RemeshProperties.h"
-#include "UObject/UObjectGlobals.h"
 #include "RemeshMeshTool.generated.h"
 
+class UMeshStatisticsProperties;
+class UMeshElementsVisualizer;
+PREDECLARE_GEOMETRY(class FDynamicMesh3);
+PREDECLARE_GEOMETRY(typedef TMeshAABBTree3<FDynamicMesh3> FDynamicMeshAABBTree3);
+
 /**
  *
  */
@@ -59,19 +59,10 @@ public:
 	UPROPERTY(EditAnywhere, Category = Remeshing)
 	bool bDiscardAttributes;
 
-	/** If true, display wireframe */
-	UPROPERTY(EditAnywhere, Category = Display)
-	bool bShowWireframe = true;
-
 	/** Display colors corresponding to the mesh's polygon groups */
 	UPROPERTY(EditAnywhere, Category = Display)
 	bool bShowGroupColors = false;
 
-	/** Display constrained edges */
-	UPROPERTY(EditAnywhere, Category = Display)
-	bool bShowConstraintEdges = false;
-
-
 	/** Remeshing type */
 	UPROPERTY(EditAnywhere, Category = Remeshing, AdvancedDisplay)
 	ERemeshType RemeshType;
@@ -125,7 +116,6 @@ public:
 	virtual void Shutdown(EToolShutdownType ShutdownType) override;
 
 	virtual void OnTick(float DeltaTime) override;
-	virtual void Render(IToolsContextRenderAPI* RenderAPI) override;
 
 	virtual bool HasCancel() const override { return true; }
 	virtual bool HasAccept() const override { return true; }
@@ -145,6 +135,9 @@ public:
 	UPROPERTY()
 	UMeshOpPreviewWithBackgroundCompute* Preview;
 
+	UPROPERTY()
+	UMeshElementsVisualizer* MeshElementsDisplay;
+
 protected:
 
 	UWorld* TargetWorld;
@@ -154,6 +147,5 @@ protected:
 	double InitialMeshArea;
 
 	double CalculateTargetEdgeLength(int TargetTriCount);
-	void GenerateAsset(const FDynamicMeshOpResult& Result);
 	void UpdateVisualization();
 };
diff --git a/Engine/Plugins/Experimental/MeshModelingToolset/Source/MeshModelingTools/Public/WeldMeshEdgesTool.h b/Engine/Plugins/Experimental/MeshModelingToolset/Source/MeshModelingTools/Public/WeldMeshEdgesTool.h
--- a/Engine/Plugins/Experimental/MeshModelingToolset/Source/MeshModelingTools/Public/WeldMeshEdgesTool.h
+++ b/Engine/Plugins/Experimental/MeshModelingToolset/Source/MeshModelingTools/Public/WeldMeshEdgesTool.h
@@ -3,18 +3,19 @@
 #pragma once
 
 #include "CoreMinimal.h"
-#include "UObject/NoExportTypes.h"
-#include "SingleSelectionTool.h"
-#include "InteractiveToolBuilder.h"
+#include "BaseTools/SingleSelectionMeshEditingTool.h"
+#include "PreviewMesh.h"
+#include "ModelingOperators.h"
+#include "MeshOpPreviewHelpers.h"
 #include "DynamicMesh/DynamicMesh3.h"
 #include "DynamicMesh/DynamicMeshAABBTree3.h"
-#include "BaseTools/SingleSelectionMeshEditingTool.h"
 #include "WeldMeshEdgesTool.generated.h"
 
-
 // predeclarations
 struct FMeshDescription;
-class UDynamicMeshComponent;
+class UMeshElementsVisualizer;
+class UWeldMeshEdgesTool;
+class FWeldMeshEdgesOp;
 
 /**
  *
@@ -23,11 +24,44 @@ UCLASS()
 class MESHMODELINGTOOLS_API UWeldMeshEdgesToolBuilder : public USingleSelectionMeshEditingToolBuilder
 {
 	GENERATED_BODY()
-
 public:
 	virtual USingleSelectionMeshEditingTool* CreateNewTool(const FToolBuilderState& SceneState) const override;
 };
 
+
+
+UCLASS()
+class MESHMODELINGTOOLS_API UWeldMeshEdgesToolProperties : public UInteractiveToolPropertySet
+{
+	GENERATED_BODY()
+
+public:
+	/** Edges are considered matching if both pairs of endpoint vertices are closer than this distance */
+	UPROPERTY(EditAnywhere, Category = Options, meta = (UIMin = "0.000001", UIMax = "0.01", ClampMin = "0.00000001", ClampMax = "1000.0"))
+	float Tolerance = FMathf::ZeroTolerance;
+
+	/** Only merge unambiguous pairs that have unique duplicate-edge matches */
+	UPROPERTY(EditAnywhere, Category = Options)
+	bool bOnlyUnique = false;
+};
+
+
+
+
+UCLASS()
+class MESHMODELINGTOOLS_API UWeldMeshEdgesOperatorFactory : public UObject, public UE::Geometry::IDynamicMeshOperatorFactory
+{
+	GENERATED_BODY()
+
+public:
+	// IDynamicMeshOperatorFactory API
+	virtual TUniquePtr<UE::Geometry::FDynamicMeshOperator> MakeNewOperator() override;
+
+	UPROPERTY()
+	UWeldMeshEdgesTool* WeldMeshEdgesTool;
+};
+
+
 /**
  * Mesh Weld Edges Tool
  */
@@ -35,37 +69,39 @@ UCLASS()
 class MESHMODELINGTOOLS_API UWeldMeshEdgesTool : public USingleSelectionMeshEditingTool
 {
 	GENERATED_BODY()
-
 public:
 	UWeldMeshEdgesTool();
 
 	virtual void Setup() override;
 	virtual void Shutdown(EToolShutdownType ShutdownType) override;
 
-	virtual void Render(IToolsContextRenderAPI* RenderAPI) override;
+	virtual void OnTick(float DeltaTime) override;
 
 	virtual bool HasCancel() const override { return true; }
 	virtual bool HasAccept() const override { return true; }
+	virtual bool CanAccept() const override;
+
+	// update parameters in FWeldMeshEdgesOp based on current Settings
+	void UpdateOpParameters(FWeldMeshEdgesOp& Op) const;
 
-#if WITH_EDITOR
-	virtual void PostEditChangeProperty(FPropertyChangedEvent& PropertyChangedEvent) override;
-#endif
 protected:
-	// need to update bResultValid if these are modified, so we don't publicly expose them. 
-	// @todo setters/getters for these
+	UPROPERTY()
+	TObjectPtr<UWeldMeshEdgesToolProperties> Settings;
 
-	/** Edges are considered matching if both pairs of endpoint vertices are closer than this distance */
-	UPROPERTY(EditAnywhere, Category = Options, meta = (UIMin = "0.000001", UIMax = "0.01", ClampMin = "0.00000001", ClampMax = "1000.0"))
-	float Tolerance;
+	UPROPERTY()
+	TObjectPtr<UMeshOpPreviewWithBackgroundCompute> PreviewCompute = nullptr;
 
-	/** Only merge unambiguous pairs that have unique duplicate-edge matches */
-	UPROPERTY(EditAnywhere, Category = Options, AdvancedDisplay)
-	bool bOnlyUnique;
+	UPROPERTY()
+	TObjectPtr<UMeshElementsVisualizer> MeshElementsDisplay = nullptr;
+
+	UPROPERTY()
+	TObjectPtr<UWeldMeshEdgesOperatorFactory> OperatorFactory;
 
 protected:
-	UDynamicMeshComponent* DynamicMeshComponent;
-	UE::Geometry::FDynamicMesh3 OriginalMesh;
-	UE::Geometry::FDynamicMesh3 ComputeMesh;
-	bool bResultValid;
-	void UpdateResult();
+
+	TSharedPtr<UE::Geometry::FDynamicMesh3, ESPMode::ThreadSafe> SourceMesh;
+
+
+
+
 };
diff --git a/Engine/Plugins/Experimental/MeshModelingToolset/Source/MeshModelingToolsEditorOnly/Private/SimplifyMeshTool.cpp b/Engine/Plugins/Experimental/MeshModelingToolset/Source/MeshModelingToolsEditorOnly/Private/SimplifyMeshTool.cpp
--- a/Engine/Plugins/Experimental/MeshModelingToolset/Source/MeshModelingToolsEditorOnly/Private/SimplifyMeshTool.cpp
+++ b/Engine/Plugins/Experimental/MeshModelingToolset/Source/MeshModelingToolsEditorOnly/Private/SimplifyMeshTool.cpp
@@ -3,28 +3,19 @@
 #include "SimplifyMeshTool.h"
 #include "InteractiveToolManager.h"
 #include "Properties/RemeshProperties.h"
+#include "Properties/MeshStatisticsProperties.h"
+#include "Drawing/MeshElementsVisualizer.h"
 #include "ToolBuilderUtil.h"
-
+#include "MeshDescriptionToDynamicMesh.h"
+#include "ModelingToolTargetUtil.h"
 #include "ToolSetupUtil.h"
-#include "Util/ColorConstants.h"
-
 #include "DynamicMesh/DynamicMesh3.h"
+#include "DynamicMesh/DynamicMeshAABBTree3.h"
+#include "Util/ColorConstants.h"
 
-#include "MeshDescriptionToDynamicMesh.h"
-#include "DynamicMeshToMeshDescription.h"
-#include "AssetUtils/MeshDescriptionUtil.h"
-
-#include "SceneManagement.h" // for FPrimitiveDrawInterface
-
-#include "TargetInterfaces/MaterialProvider.h"
-#include "TargetInterfaces/MeshDescriptionCommitter.h"
-#include "TargetInterfaces/MeshDescriptionProvider.h"
-#include "TargetInterfaces/PrimitiveComponentBackedTarget.h"
 
 //#include "ProfilingDebugging/ScopedTimers.h" // enable this to use the timer.
 #include "Modules/ModuleManager.h"
-
-
 #include "IMeshReductionManagerModule.h"
 #include "IMeshReductionInterfaces.h"
 
@@ -62,7 +53,6 @@ USimplifyMeshToolProperties::USimplifyMeshToolProperties()
 	bPreventNormalFlips = true;
 	bDiscardAttributes = false;
 	bGeometricConstraint = false;
-	bShowWireframe = true;
 	bShowGroupColors = false;
 	GroupBoundaryConstraint = EGroupBoundaryConstraint::Ignore;
 	MaterialBoundaryConstraint = EMaterialBoundaryConstraint::Ignore;
@@ -72,18 +62,17 @@ void USimplifyMeshTool::Setup()
 {
 	UInteractiveTool::Setup();
 
-
 	// hide component and create + show preview
-	IPrimitiveComponentBackedTarget* TargetComponent = Cast<IPrimitiveComponentBackedTarget>(Target);
-	TargetComponent->SetOwnerVisibility(false);
-	Preview = NewObject<UMeshOpPreviewWithBackgroundCompute>(this, "Preview");
+	UE::ToolTarget::HideSourceObject(Target);
+
+	Preview = NewObject<UMeshOpPreviewWithBackgroundCompute>(this);
 	Preview->Setup(this->TargetWorld, this);
-	FComponentMaterialSet MaterialSet;
-	Cast<IMaterialProvider>(Target)->GetMaterialSet(MaterialSet);
+	FComponentMaterialSet MaterialSet = UE::ToolTarget::GetMaterialSet(Target);
 	Preview->ConfigureMaterials( MaterialSet.Materials,
 		ToolSetupUtil::GetDefaultWorkingMaterial(GetToolManager())
 	);
 
+	// some of this could be done async...
 	{
 		// if in editor, create progress indicator dialog because building mesh copies can be slow (for very large meshes)
 		// this is especially needed because of the copy we make of the meshdescription; for Reasons, copying meshdescription is pretty slow
@@ -101,29 +90,19 @@ void USimplifyMeshTool::Setup()
 #else
 		auto EnterProgressFrame = [](int Progress) {};
 #endif
-		IMeshDescriptionProvider* TargetMeshProvider = Cast<IMeshDescriptionProvider>(Target);
-		OriginalMeshDescription = MakeShared<FMeshDescription, ESPMode::ThreadSafe>(*TargetMeshProvider->GetMeshDescription());
-		// make sure the mesh description has tangents. 
-		const bool bUseTangents = true;
-		if (bUseTangents)
-		{
-			// if the source was imported with auto-generate tangents, the mesh description tangents won't have been populated
-			FMeshBuildSettings BuildSettings; 
-			BuildSettings.bRecomputeNormals = true;
-			BuildSettings.bRecomputeTangents = true;
-			// this will only generate tangents if the mesh description was missing them
-			UE::MeshDescription::InitializeAutoGeneratedAttributes(*OriginalMeshDescription, &BuildSettings);
-		}
+		OriginalMeshDescription = MakeShared<FMeshDescription, ESPMode::ThreadSafe>(UE::ToolTarget::GetMeshDescriptionCopy(Target, true));
+
 		EnterProgressFrame(1);
+		// UE::ToolTarget::GetDynamicMeshCopy() would recompute the tangents a second time here
 		OriginalMesh = MakeShared<FDynamicMesh3, ESPMode::ThreadSafe>();
 		FMeshDescriptionToDynamicMesh Converter; 
-		// convert with tangent overlay
-		Converter.Convert(TargetMeshProvider->GetMeshDescription(), *OriginalMesh, bUseTangents);
+		Converter.Convert(OriginalMeshDescription.Get(), *OriginalMesh, true);   // convert with tangent overlay
+
 		EnterProgressFrame(2);
 		OriginalMeshSpatial = MakeShared<FDynamicMeshAABBTree3, ESPMode::ThreadSafe>(OriginalMesh.Get(), true);
 	}
 
-	Preview->PreviewMesh->SetTransform(TargetComponent->GetWorldTransform());
+	Preview->PreviewMesh->SetTransform((FTransform)UE::ToolTarget::GetLocalToWorldTransform(Target));
 	Preview->PreviewMesh->SetTangentsMode(EDynamicMeshComponentTangentsMode::AutoCalculated);
 	Preview->PreviewMesh->UpdatePreview(OriginalMesh.Get());
 
@@ -134,15 +113,30 @@ void USimplifyMeshTool::Setup()
 
 	SimplifyProperties->WatchProperty(SimplifyProperties->bShowGroupColors,
 									  [this](bool bNewValue) { UpdateVisualization(); });
-	SimplifyProperties->WatchProperty(SimplifyProperties->bShowWireframe,
-									  [this](bool bNewValue) { UpdateVisualization(); });
 
 	MeshStatisticsProperties = NewObject<UMeshStatisticsProperties>(this);
 	AddToolPropertySource(MeshStatisticsProperties);
 
+	MeshElementsDisplay = NewObject<UMeshElementsVisualizer>(this);
+	MeshElementsDisplay->CreateInWorld(Preview->PreviewMesh->GetWorld(), Preview->PreviewMesh->GetTransform());
+	if (ensure(MeshElementsDisplay->Settings))
+	{
+		MeshElementsDisplay->Settings->bShowWireframe = true;
+		MeshElementsDisplay->Settings->RestoreProperties(this, TEXT("Simplify"));
+		AddToolPropertySource(MeshElementsDisplay->Settings);
+	}
+	MeshElementsDisplay->SetMeshAccessFunction([this](UMeshElementsVisualizer::ProcessDynamicMeshFunc ProcessFunc) {
+		Preview->ProcessCurrentMesh(ProcessFunc);
+	});
+
+
 	Preview->OnMeshUpdated.AddLambda([this](UMeshOpPreviewWithBackgroundCompute* Compute)
 	{
-		MeshStatisticsProperties->Update(*Compute->PreviewMesh->GetPreviewDynamicMesh());
+		Compute->ProcessCurrentMesh([&](const FDynamicMesh3& ReadMesh)
+		{
+			MeshStatisticsProperties->Update(ReadMesh);
+			MeshElementsDisplay->NotifyMeshChanged();
+		});
 	});
 
 	UpdateVisualization();
@@ -164,11 +158,20 @@ bool USimplifyMeshTool::CanAccept() const
 void USimplifyMeshTool::Shutdown(EToolShutdownType ShutdownType)
 {
 	SimplifyProperties->SaveProperties(this);
-	Cast<IPrimitiveComponentBackedTarget>(Target)->SetOwnerVisibility(true);
+	
+	if (ensure(MeshElementsDisplay->Settings))
+	{
+		MeshElementsDisplay->Settings->SaveProperties(this, TEXT("Simplify"));
+	}
+	MeshElementsDisplay->Disconnect();
+
+	UE::ToolTarget::ShowSourceObject(Target);
 	FDynamicMeshOpResult Result = Preview->Shutdown();
 	if (ShutdownType == EToolShutdownType::Accept)
 	{
-		GenerateAsset(Result);
+		GetToolManager()->BeginUndoTransaction(LOCTEXT("SimplifyMeshToolTransactionName", "Simplify Mesh"));
+		UE::ToolTarget::CommitDynamicMeshUpdate(Target, *Result.Mesh, true);
+		GetToolManager()->EndUndoTransaction();
 	}
 }
 
@@ -176,6 +179,7 @@ void USimplifyMeshTool::Shutdown(EToolShutdownType ShutdownType)
 void USimplifyMeshTool::OnTick(float DeltaTime)
 {
 	Preview->Tick(DeltaTime);
+	MeshElementsDisplay->OnTick(DeltaTime);
 }
 
 TUniquePtr<FDynamicMeshOperator> USimplifyMeshTool::MakeNewOperator()
@@ -197,7 +201,7 @@ TUniquePtr<FDynamicMeshOperator> USimplifyMeshTool::MakeNewOperator()
 	Op->MaterialBoundaryConstraint = (EEdgeRefineFlags)SimplifyProperties->MaterialBoundaryConstraint;
 	Op->bGeometricDeviationConstraint = SimplifyProperties->bGeometricConstraint;
 	Op->GeometricTolerance = SimplifyProperties->GeometricTolerance;
-	FTransform LocalToWorld = Cast<IPrimitiveComponentBackedTarget>(Target)->GetWorldTransform();
+	FTransform LocalToWorld = (FTransform)UE::ToolTarget::GetLocalToWorldTransform(Target);
 	Op->SetTransform(LocalToWorld);
 
 	Op->OriginalMeshDescription = OriginalMeshDescription;
@@ -210,45 +214,12 @@ TUniquePtr<FDynamicMeshOperator> USimplifyMeshTool::MakeNewOperator()
 	return Op;
 }
 
-void USimplifyMeshTool::Render(IToolsContextRenderAPI* RenderAPI)
-{
-	FPrimitiveDrawInterface* PDI = RenderAPI->GetPrimitiveDrawInterface();
-	FTransform Transform = Cast<IPrimitiveComponentBackedTarget>(Target)->GetWorldTransform(); //Actor->GetTransform();
-
-	if (SimplifyProperties->bShowUVSeams)
-	{
-		FColor UVSeamColor(15, 240, 15); // same color used in mesh inspector
-		const FDynamicMesh3* TargetMesh = Preview->PreviewMesh->GetPreviewDynamicMesh();
-		if (TargetMesh->HasAttributes())
-		{
-			float PDIScale = RenderAPI->GetCameraState().GetPDIScalingFactor();
-
-			const FDynamicMeshUVOverlay* UVOverlay = TargetMesh->Attributes()->PrimaryUV();
-			for (int eid : TargetMesh->EdgeIndicesItr())
-			{
-				if (UVOverlay->IsSeamEdge(eid))
-				{
-					FVector3d A, B;
-					TargetMesh->GetEdgeV(eid, A, B);
-					PDI->DrawLine(Transform.TransformPosition((FVector)A), Transform.TransformPosition((FVector)B),
-						UVSeamColor, 0, 2.0 * PDIScale, 1.0f, true);
-				}
-			}
-		}
-	}
-}
 
 void USimplifyMeshTool::OnPropertyModified(UObject* PropertySet, FProperty* Property)
 {
 	if ( Property )
 	{
-		if ( Property->GetFName() == GET_MEMBER_NAME_CHECKED(USimplifyMeshToolProperties, bShowUVSeams) )
-		{
-			// nothing required.  Render() handles this case
-			return;
-		}
-		if ( ( Property->GetFName() == GET_MEMBER_NAME_CHECKED(USimplifyMeshToolProperties, bShowWireframe) ) ||
-			 ( Property->GetFName() == GET_MEMBER_NAME_CHECKED(USimplifyMeshToolProperties, bShowGroupColors) )  )
+		if ( Property->GetFName() == GET_MEMBER_NAME_CHECKED(USimplifyMeshToolProperties, bShowGroupColors) )
 		{
 			UpdateVisualization();
 		}
@@ -261,7 +232,6 @@ void USimplifyMeshTool::OnPropertyModified(UObject* PropertySet, FProperty* Prop
 
 void USimplifyMeshTool::UpdateVisualization()
 {
-	Preview->PreviewMesh->EnableWireframe(SimplifyProperties->bShowWireframe);
 	FComponentMaterialSet MaterialSet;
 	if (SimplifyProperties->bShowGroupColors)
 	{
@@ -274,27 +244,13 @@ void USimplifyMeshTool::UpdateVisualization()
 	}
 	else
 	{
-		Cast<IMaterialProvider>(Target)->GetMaterialSet(MaterialSet);
+		MaterialSet = UE::ToolTarget::GetMaterialSet(Target);
 		Preview->PreviewMesh->ClearTriangleColorFunction(UPreviewMesh::ERenderUpdateMode::FastUpdate);
 	}
 	Preview->ConfigureMaterials(MaterialSet.Materials,
 								ToolSetupUtil::GetDefaultWorkingMaterial(GetToolManager()));
 }
 
-void USimplifyMeshTool::GenerateAsset(const FDynamicMeshOpResult& Result)
-{
-	GetToolManager()->BeginUndoTransaction(LOCTEXT("SimplifyMeshToolTransactionName", "Simplify Mesh"));
-
-	check(Result.Mesh.Get() != nullptr);
-	Cast<IMeshDescriptionCommitter>(Target)->CommitMeshDescription([&Result](const IMeshDescriptionCommitter::FCommitterParams& CommitParams)
-	{
-		FDynamicMeshToMeshDescription Converter;
-		Converter.Convert(Result.Mesh.Get(), *CommitParams.MeshDescriptionOut);
-	});
-
-	GetToolManager()->EndUndoTransaction();
-}
-
 
 
 
diff --git a/Engine/Plugins/Experimental/MeshModelingToolset/Source/MeshModelingToolsEditorOnly/Public/SimplifyMeshTool.h b/Engine/Plugins/Experimental/MeshModelingToolset/Source/MeshModelingToolsEditorOnly/Public/SimplifyMeshTool.h
--- a/Engine/Plugins/Experimental/MeshModelingToolset/Source/MeshModelingToolsEditorOnly/Public/SimplifyMeshTool.h
+++ b/Engine/Plugins/Experimental/MeshModelingToolset/Source/MeshModelingToolsEditorOnly/Public/SimplifyMeshTool.h
@@ -4,20 +4,16 @@
 
 #include "CoreMinimal.h"
 #include "UObject/NoExportTypes.h"
-#include "SingleSelectionTool.h"
-#include "InteractiveToolBuilder.h"
-#include "DynamicMesh/DynamicMesh3.h"
-#include "DynamicMesh/DynamicMeshAABBTree3.h"
+#include "BaseTools/SingleSelectionMeshEditingTool.h"
 #include "MeshOpPreviewHelpers.h"
-#include "CleaningOps/SimplifyMeshOp.h"
-#include "Properties/MeshStatisticsProperties.h"
 #include "Properties/RemeshProperties.h"
-#include "BaseTools/SingleSelectionMeshEditingTool.h"
+#include "CleaningOps/SimplifyMeshOp.h"		// required in header for enum types
 #include "SimplifyMeshTool.generated.h"
 
-
-
-
+class UMeshStatisticsProperties;
+class UMeshElementsVisualizer;
+PREDECLARE_GEOMETRY(class FDynamicMesh3);
+PREDECLARE_GEOMETRY(typedef TMeshAABBTree3<FDynamicMesh3> FDynamicMeshAABBTree3);
 
 /**
  *
@@ -75,18 +71,10 @@ public:
 	UPROPERTY(EditAnywhere, Category = Options, meta = (UIMin = "0.0", UIMax = "10.0", ClampMin = "0.0", ClampMax = "10000000.0", EditCondition = "bGeometricConstraint && SimplifierType != ESimplifyType::UEStandard"))
 	float GeometricTolerance;
 
-	/** If true, display wireframe */
-	UPROPERTY(EditAnywhere, Category = Display)
-	bool bShowWireframe;
-
 	/** Display colors corresponding to the mesh's polygon groups */
 	UPROPERTY(EditAnywhere, Category = Display)
 	bool bShowGroupColors = false;
 
-	/** Display seams in first UV channel */
-	UPROPERTY(EditAnywhere, Category = Display, DisplayName ="Show UV Seams")
-	bool bShowUVSeams = false;
-	
 	/** Enable projection back to input mesh */
 	UPROPERTY(EditAnywhere, Category = Options, AdvancedDisplay)
 	bool bReproject;
@@ -108,7 +96,6 @@ public:
 	virtual void Shutdown(EToolShutdownType ShutdownType) override;
 
 	virtual void OnTick(float DeltaTime) override;
-	virtual void Render(IToolsContextRenderAPI* RenderAPI) override;
 
 	virtual bool HasCancel() const override { return true; }
 	virtual bool HasAccept() const override { return true; }
@@ -129,11 +116,13 @@ private:
 	UPROPERTY()
 	UMeshOpPreviewWithBackgroundCompute* Preview;
 
+	UPROPERTY()
+	UMeshElementsVisualizer* MeshElementsDisplay;
+
 	TSharedPtr<FMeshDescription, ESPMode::ThreadSafe> OriginalMeshDescription;
 	// Dynamic Mesh versions precomputed in Setup (rather than recomputed for every simplify op)
 	TSharedPtr<UE::Geometry::FDynamicMesh3, ESPMode::ThreadSafe> OriginalMesh;
 	TSharedPtr<UE::Geometry::FDynamicMeshAABBTree3, ESPMode::ThreadSafe> OriginalMeshSpatial;
 
-	void GenerateAsset(const FDynamicMeshOpResult& Result);
 	void UpdateVisualization();
 };
diff --git a/Engine/Plugins/Experimental/MeshModelingToolset/Source/ModelingComponents/Private/Drawing/MeshElementsVisualizer.cpp b/Engine/Plugins/Experimental/MeshModelingToolset/Source/ModelingComponents/Private/Drawing/MeshElementsVisualizer.cpp
--- a/Engine/Plugins/Experimental/MeshModelingToolset/Source/ModelingComponents/Private/Drawing/MeshElementsVisualizer.cpp
+++ b/Engine/Plugins/Experimental/MeshModelingToolset/Source/ModelingComponents/Private/Drawing/MeshElementsVisualizer.cpp
@@ -14,53 +14,55 @@ using namespace UE::Geometry;
 /**
  * IMeshWireframeSource implementation for a FDynamicMesh3
  */
-class FDynamicMeshWireframeSource : public IMeshWireframeSource
+class FTemporaryDynamicMeshWireframeSource : public IMeshWireframeSource
 {
 public:
-	const FDynamicMesh3* Mesh = nullptr;
+	const FDynamicMesh3& Mesh;
 
-	virtual bool IsValid() const { return Mesh != nullptr; }
+	FTemporaryDynamicMeshWireframeSource(const FDynamicMesh3& MeshIn) : Mesh(MeshIn) {}
+
+	virtual bool IsValid() const { return true; }
 
 	virtual FBoxSphereBounds GetBounds() const 
 	{ 
-		FAxisAlignedBox3d Bounds = Mesh->GetBounds();
+		FAxisAlignedBox3d Bounds = Mesh.GetBounds();
 		return FBoxSphereBounds((FBox)Bounds);
 	}
 
 	virtual FVector GetVertex(int32 Index) const
 	{
-		return (FVector)Mesh->GetVertex(Index);
+		return (FVector)Mesh.GetVertex(Index);
 	}
 
 	virtual int32 GetEdgeCount() const
 	{
-		return Mesh->EdgeCount();
+		return Mesh.EdgeCount();
 	}
 
 	virtual int32 GetMaxEdgeIndex() const
 	{
-		return Mesh->MaxEdgeID();
+		return Mesh.MaxEdgeID();
 	}
 
 	virtual bool IsEdge(int32 Index) const
 	{
-		return Mesh->IsEdge(Index);
+		return Mesh.IsEdge(Index);
 	}
 
 	virtual void GetEdge(int32 EdgeIndex, int32& VertIndexAOut, int32& VertIndexBOut, EMeshEdgeType& TypeOut) const
 	{
-		FIndex2i EdgeV = Mesh->GetEdgeV(EdgeIndex);
+		FIndex2i EdgeV = Mesh.GetEdgeV(EdgeIndex);
 		VertIndexAOut = EdgeV.A;
 		VertIndexBOut = EdgeV.B;
 		int32 EdgeType = (int32)EMeshEdgeType::Regular;
-		if (Mesh->IsBoundaryEdge(EdgeIndex))
+		if (Mesh.IsBoundaryEdge(EdgeIndex))
 		{
 			EdgeType |= (int32)EMeshEdgeType::MeshBoundary;
 		}
-		if (Mesh->HasAttributes())
+		if (Mesh.HasAttributes())
 		{
 			bool bIsUVSeam = false, bIsNormalSeam = false, bIsColorSeam = false;
-			if (Mesh->Attributes()->IsSeamEdge(EdgeIndex, bIsUVSeam, bIsNormalSeam, bIsColorSeam))
+			if (Mesh.Attributes()->IsSeamEdge(EdgeIndex, bIsUVSeam, bIsNormalSeam, bIsColorSeam))
 			{
 				if (bIsUVSeam)
 				{
@@ -84,20 +86,20 @@ public:
 class FDynamicMeshWireframeSourceProvider : public IMeshWireframeSourceProvider
 {
 public:
-	TUniqueFunction<const FDynamicMesh3* (void)> MeshAccessFunction;
+	TUniqueFunction<void(UMeshElementsVisualizer::ProcessDynamicMeshFunc)> MeshAccessFunction;
 
-	FDynamicMeshWireframeSourceProvider(TUniqueFunction<const FDynamicMesh3* (void)>&& MeshAccessFuncIn)
+	FDynamicMeshWireframeSourceProvider(TUniqueFunction<void(UMeshElementsVisualizer::ProcessDynamicMeshFunc)>&& MeshAccessFuncIn)
 	{
 		MeshAccessFunction = MoveTemp(MeshAccessFuncIn);
 	}
 
 	virtual void AccessMesh(TFunctionRef<void(const IMeshWireframeSource&)> ProcessingFunc) override
 	{
-		const FDynamicMesh3* Mesh = MeshAccessFunction();
-		check(Mesh);
-		FDynamicMeshWireframeSource WireSource;
-		WireSource.Mesh = Mesh;
-		ProcessingFunc(WireSource);
+		MeshAccessFunction([&](const FDynamicMesh3& ReadMesh)
+		{
+			FTemporaryDynamicMeshWireframeSource WireSource(ReadMesh);
+			ProcessingFunc(WireSource);
+		});
 	}
 };
 
@@ -132,7 +134,7 @@ void UMeshElementsVisualizer::OnCreated()
 }
 
 
-void UMeshElementsVisualizer::SetMeshAccessFunction(TUniqueFunction<const FDynamicMesh3* (void)>&& MeshAccessFunctionIn)
+void UMeshElementsVisualizer::SetMeshAccessFunction(TUniqueFunction<void(ProcessDynamicMeshFunc)>&& MeshAccessFunctionIn)
 {
 	WireframeSourceProvider = MakeShared<FDynamicMeshWireframeSourceProvider>(MoveTemp(MeshAccessFunctionIn));
 
@@ -157,6 +159,7 @@ void UMeshElementsVisualizer::UpdateLineDepthBiasScale()
 	}
 }
 
+
 void UMeshElementsVisualizer::OnTick(float DeltaTime)
 {
 	if (bSettingsModified)
diff --git a/Engine/Plugins/Experimental/MeshModelingToolset/Source/ModelingComponents/Private/MeshOpPreviewHelpers.cpp b/Engine/Plugins/Experimental/MeshModelingToolset/Source/ModelingComponents/Private/MeshOpPreviewHelpers.cpp
--- a/Engine/Plugins/Experimental/MeshModelingToolset/Source/ModelingComponents/Private/MeshOpPreviewHelpers.cpp
+++ b/Engine/Plugins/Experimental/MeshModelingToolset/Source/ModelingComponents/Private/MeshOpPreviewHelpers.cpp
@@ -148,7 +148,20 @@ bool UMeshOpPreviewWithBackgroundCompute::GetCurrentResultCopy(FDynamicMesh3& Me
 {
 	if ( HaveValidResult() || bOnlyIfValid == false)
 	{
-		MeshOut.Copy( *PreviewMesh->GetMesh() );
+		PreviewMesh->ProcessMesh([&](const FDynamicMesh3& ReadMesh)
+		{
+			MeshOut = ReadMesh;
+		});
+		return true;
+	}
+	return false;
+}
+
+bool UMeshOpPreviewWithBackgroundCompute::ProcessCurrentMesh(TFunctionRef<void(const UE::Geometry::FDynamicMesh3&)> ProcessFunc, bool bOnlyIfValid)
+{
+	if (HaveValidResult() || bOnlyIfValid == false)
+	{
+		PreviewMesh->ProcessMesh(ProcessFunc);
 		return true;
 	}
 	return false;
diff --git a/Engine/Plugins/Experimental/MeshModelingToolset/Source/ModelingComponents/Private/ModelingToolTargetUtil.cpp b/Engine/Plugins/Experimental/MeshModelingToolset/Source/ModelingComponents/Private/ModelingToolTargetUtil.cpp
--- a/Engine/Plugins/Experimental/MeshModelingToolset/Source/ModelingComponents/Private/ModelingToolTargetUtil.cpp
+++ b/Engine/Plugins/Experimental/MeshModelingToolset/Source/ModelingComponents/Private/ModelingToolTargetUtil.cpp
@@ -120,6 +120,23 @@ const FMeshDescription* UE::ToolTarget::GetMeshDescription(UToolTarget* Target)
 }
 
 
+FMeshDescription UE::ToolTarget::GetMeshDescriptionCopy(UToolTarget* Target, bool bWantMeshTangents)
+{
+	IMeshDescriptionProvider* MeshDescriptionProvider = Cast<IMeshDescriptionProvider>(Target);
+	if (MeshDescriptionProvider)
+	{
+		FMeshDescription NewMeshDescription = *MeshDescriptionProvider->GetMeshDescription();
+		if (bWantMeshTangents)
+		{
+			MeshDescriptionProvider->CalculateAutoGeneratedAttributes(NewMeshDescription);
+		}
+		return NewMeshDescription;
+	}
+	ensure(false);
+	return FMeshDescription();
+}
+
+
 FDynamicMesh3 UE::ToolTarget::GetDynamicMeshCopy(UToolTarget* Target, bool bWantMeshTangents)
 {
 	IMeshDescriptionProvider* MeshDescriptionProvider = Cast<IMeshDescriptionProvider>(Target);
diff --git a/Engine/Plugins/Experimental/MeshModelingToolset/Source/ModelingComponents/Public/Drawing/MeshElementsVisualizer.h b/Engine/Plugins/Experimental/MeshModelingToolset/Source/ModelingComponents/Public/Drawing/MeshElementsVisualizer.h
--- a/Engine/Plugins/Experimental/MeshModelingToolset/Source/ModelingComponents/Public/Drawing/MeshElementsVisualizer.h
+++ b/Engine/Plugins/Experimental/MeshModelingToolset/Source/ModelingComponents/Public/Drawing/MeshElementsVisualizer.h
@@ -92,14 +92,33 @@ class MODELINGCOMPONENTS_API UMeshElementsVisualizer : public UPreviewGeometry
 	GENERATED_BODY()
 
 public:
+
+	/**
+	 * A reference to a function that can be called with a const FDynamicMesh3 to do some computation on it.
+	 */
+	typedef TFunctionRef<void(const UE::Geometry::FDynamicMesh3& Mesh)> ProcessDynamicMeshFunc;
+
 	/**
-	 * UMeshElementsVisualizer must be provided with a callback function that can be used to
-	 * access the mesh. UMeshElementsVisualizer will hold onto this callback, and wrap it in
-	 * an API that is passed to child Components used to render mesh geometry. Those Components
-	 * will access the mesh when creating their Scene Proxies/etc. The callback can return null
-	 * if the mesh has become invalid/etc
+	 * UMeshElementsVisualizer needs access to a mesh to generate it's data structures. However
+	 * we may not be able to (or shouldn't) directly pass a pointer to UMeshElementsVisualizer.
+	 * For example a UDynamicMeshComponent can expose it's internal mesh via a function ProcessMesh( TFunctionRef<void(const FDynamicMesh3&)> Func ).
+	 * UMeshElementsVisualizer cannot directly call that function without knowing about the UDynamicMeshComponent,
+	 * which we would like to avoid. So instead the client of UMeshElementsVisualizer provides a 
+	 * MeshAccessFunction that we can call with an internal function/lambda that UMeshElementsVisualizer
+	 * internals will produce, that can be passed a const FDynamicMesh3& for processing (this is a ProcessDynamicMeshFunc).
+	 * The job of the MeshAccessFunction is to call this ProcessDynamicMeshFunc on a FDynamicMesh3,
+	 * which it may have direct access to, or it may get by calling (eg) UDynamicMeshComponent::ProcessMesh() internally. 
+	 * 
+	 * So for example usage with a Tool that has a DynamicMeshComponent could be as follows:
+	 *   MeshElementsVisualizer->SetMeshAccessFunction([this](UMeshElementsVisualizer::ProcessDynamicMeshFunc ProcessFunc) {
+	 *       DynamicMeshComponent->ProcessMesh(ProcessFunc);
+	 *   });
+	 * And for usage with a mesh that a Tool class owns:
+	 *   MeshElementsVisualizer->SetMeshAccessFunction([this](UMeshElementsVisualizer::ProcessDynamicMeshFunc ProcessFunc) {
+	 *       ProcessFunc(MyInternalMesh);
+	 *   });
 	 */
-	void SetMeshAccessFunction(TUniqueFunction<const UE::Geometry::FDynamicMesh3* (void)>&& MeshAccessFunction);
+	void SetMeshAccessFunction(TUniqueFunction<void(ProcessDynamicMeshFunc)>&& MeshAccessFunction);
 
 	/**
 	 * Call if mesh provided by MeshAccessFunction has been modified, will cause a full recomputation
diff --git a/Engine/Plugins/Experimental/MeshModelingToolset/Source/ModelingComponents/Public/MeshOpPreviewHelpers.h b/Engine/Plugins/Experimental/MeshModelingToolset/Source/ModelingComponents/Public/MeshOpPreviewHelpers.h
--- a/Engine/Plugins/Experimental/MeshModelingToolset/Source/ModelingComponents/Public/MeshOpPreviewHelpers.h
+++ b/Engine/Plugins/Experimental/MeshModelingToolset/Source/ModelingComponents/Public/MeshOpPreviewHelpers.h
@@ -129,6 +129,15 @@ public:
 	 */
 	bool GetCurrentResultCopy(FDynamicMesh3& MeshOut, bool bOnlyIfValid = true);
 
+
+	/**
+	 * Allow an external function to safely access the PreviewMesh's mesh 
+	 * @param bOnlyIfValid if true, then only call ProcessFunc if current result is valid, ie HaveValidResult() == true. Default false.
+	 * @return true if ProcessFunc was called
+	 */
+	bool ProcessCurrentMesh(TFunctionRef<void(const UE::Geometry::FDynamicMesh3&)> ProcessFunc, bool bOnlyIfValid = false);
+
+	/** @return UWorld that the created PreviewMesh exist in */
 	virtual UWorld* GetWorld() const override { return PreviewWorld; }
 
 	//
diff --git a/Engine/Plugins/Experimental/MeshModelingToolset/Source/ModelingComponents/Public/ModelingToolTargetUtil.h b/Engine/Plugins/Experimental/MeshModelingToolset/Source/ModelingComponents/Public/ModelingToolTargetUtil.h
--- a/Engine/Plugins/Experimental/MeshModelingToolset/Source/ModelingComponents/Public/ModelingToolTargetUtil.h
+++ b/Engine/Plugins/Experimental/MeshModelingToolset/Source/ModelingComponents/Public/ModelingToolTargetUtil.h
@@ -79,7 +79,14 @@ MODELINGCOMPONENTS_API bool CommitMaterialSetUpdate(
 MODELINGCOMPONENTS_API const FMeshDescription* GetMeshDescription(UToolTarget* Target);
 
 /**
- * Fetch a DynamicMesh3 representing the given ToolTarget. This may be a conersion of the output of GetMeshDescription().
+ * @return a copy of the MeshDescription underlying a ToolTarget
+ * @param bWantMeshTangents if true, tangents will be returned if the target has them available. This may require that they be auto-calculated in some cases (which may be expensive)
+ * @return a new MeshDescription, which may be empty if the Target doesn't have a mesh  
+ */
+MODELINGCOMPONENTS_API FMeshDescription GetMeshDescriptionCopy(UToolTarget* Target, bool bWantMeshTangents = false);
+
+/**
+ * Fetch a DynamicMesh3 representing the given ToolTarget. This may be a conversion of the output of GetMeshDescription().
  * This function returns a copy, so the caller can take ownership of this Mesh.
  * @param bWantMeshTangents if true, tangents will be returned if the target has them available. This may require that they be auto-calculated in some cases (which may be expensive)
  * @return a created DynamicMesh3, which may be empty if the Target doesn't have a mesh 
diff --git a/Engine/Plugins/Experimental/UVEditor/Source/UVEditor/Private/UVEditorMode.cpp b/Engine/Plugins/Experimental/UVEditor/Source/UVEditor/Private/UVEditorMode.cpp
--- a/Engine/Plugins/Experimental/UVEditor/Source/UVEditor/Private/UVEditorMode.cpp
+++ b/Engine/Plugins/Experimental/UVEditor/Source/UVEditor/Private/UVEditorMode.cpp
@@ -23,6 +23,7 @@
 #include "UVEditorSubsystem.h"
 #include "UVEditorToolUtil.h"
 #include "UVToolContextObjects.h"
+#include "Editor.h"
 
 #define LOCTEXT_NAMESPACE "UUVEditorMode"
 
@@ -191,9 +192,9 @@ void UUVEditorMode::InitializeTargets(TArray<TObjectPtr<UObject>>& AssetsIn, TAr
 		WireframeDisplay->WireframeComponent->BoundaryEdgeThickness = 2;
 
 		// The wireframe will track the unwrap preview mesh
-		WireframeDisplay->SetMeshAccessFunction([ToolInputObject](void) { 
-			return ToolInputObject->UnwrapPreview->PreviewMesh->GetMesh(); 
-			});
+		WireframeDisplay->SetMeshAccessFunction([ToolInputObject](UMeshElementsVisualizer::ProcessDynamicMeshFunc ProcessFunc) {
+			ToolInputObject->UnwrapPreview->ProcessCurrentMesh(ProcessFunc);
+		});
 
 		// The settings object and wireframe are not part of a tool, so they won't get ticked like they
 		// are supposed to (to enable property watching), unless we add this here.
