ModelingTools: remove duplicate branch identified by PVS
#rnx

diff --git a/Engine/Plugins/Experimental/MeshModelingToolset/Source/MeshModelingTools/Private/RemeshMeshTool.cpp b/Engine/Plugins/Experimental/MeshModelingToolset/Source/MeshModelingTools/Private/RemeshMeshTool.cpp
--- a/Engine/Plugins/Experimental/MeshModelingToolset/Source/MeshModelingTools/Private/RemeshMeshTool.cpp
+++ b/Engine/Plugins/Experimental/MeshModelingToolset/Source/MeshModelingTools/Private/RemeshMeshTool.cpp
@@ -201,12 +201,9 @@ void URemeshMeshTool::Shutdown(EToolShutdownType ShutdownType)
 	FDynamicMeshOpResult Result = Preview->Shutdown();
 	if (ShutdownType == EToolShutdownType::Accept)
 	{
-		if (ShutdownType == EToolShutdownType::Accept)
-		{
-			GetToolManager()->BeginUndoTransaction(LOCTEXT("RemeshMeshToolTransactionName", "Remesh Mesh"));
-			UE::ToolTarget::CommitDynamicMeshUpdate(Targets[0], *Result.Mesh, true);
-			GetToolManager()->EndUndoTransaction();
-		}
+		GetToolManager()->BeginUndoTransaction(LOCTEXT("RemeshMeshToolTransactionName", "Remesh Mesh"));
+		UE::ToolTarget::CommitDynamicMeshUpdate(Targets[0], *Result.Mesh, true);
+		GetToolManager()->EndUndoTransaction();
 	}
 }
 
