ModelingComponents: handle meshes without materialID attribute in MeshRenderDecomposition
#rnx

diff --git a/Engine/Plugins/Experimental/MeshModelingToolset/Source/ModelingComponents/Private/Drawing/MeshRenderDecomposition.cpp b/Engine/Plugins/Experimental/MeshModelingToolset/Source/ModelingComponents/Private/Drawing/MeshRenderDecomposition.cpp
--- a/Engine/Plugins/Experimental/MeshModelingToolset/Source/ModelingComponents/Private/Drawing/MeshRenderDecomposition.cpp
+++ b/Engine/Plugins/Experimental/MeshModelingToolset/Source/ModelingComponents/Private/Drawing/MeshRenderDecomposition.cpp
@@ -36,8 +36,8 @@ void FMeshRenderDecomposition::BuildMaterialDecomposition(const FDynamicMesh3* M
 	LocalMaterials.Materials.Add(nullptr);
 	const FComponentMaterialSet* UseMaterials = (ensure(MaterialSet != nullptr)) ? MaterialSet : &LocalMaterials;
 
-	const FDynamicMeshMaterialAttribute* MaterialID = Mesh->Attributes()->GetMaterialID();
-	ensure(MaterialID != nullptr);
+	// may not have MaterialID, in that case all triangles will be material 0
+	const FDynamicMeshMaterialAttribute* MaterialID = Mesh->HasAttributes() ? Mesh->Attributes()->GetMaterialID() : nullptr;
 
 	int32 NumMaterials = UseMaterials->Materials.Num();
 	Decomp.Initialize(NumMaterials);
