ModelingComponents: Add UBaseDynamicmeshComponent::SetShadowsEnabled() and UPreviewMesh::SetShadowsEnabled() wrapper, and update Tools that are directly modifying the bCastDynamicShadow variable on Components. That doesn't work as it doesn't update the SceneProxy, and is currently causing crashes. ::SetShadowsEnabled() forces a full immediate rebuild of the SceneProxy, which is very heavy but seems to be required.
#rnx

diff --git a/Engine/Plugins/Experimental/MeshModelingToolset/Source/MeshModelingTools/Private/DynamicMeshSculptTool.cpp b/Engine/Plugins/Experimental/MeshModelingToolset/Source/MeshModelingTools/Private/DynamicMeshSculptTool.cpp
--- a/Engine/Plugins/Experimental/MeshModelingToolset/Source/MeshModelingTools/Private/DynamicMeshSculptTool.cpp
+++ b/Engine/Plugins/Experimental/MeshModelingToolset/Source/MeshModelingTools/Private/DynamicMeshSculptTool.cpp
@@ -101,6 +101,7 @@ void UDynamicMeshSculptTool::Setup()
 
 	// create dynamic mesh component to use for live preview
 	DynamicMeshComponent = NewObject<UOctreeDynamicMeshComponent>(UE::ToolTarget::GetTargetActor(Target));
+	DynamicMeshComponent->SetShadowsEnabled(false);
 	DynamicMeshComponent->SetupAttachment(UE::ToolTarget::GetTargetActor(Target)->GetRootComponent());
 	DynamicMeshComponent->RegisterComponent();
 
@@ -2448,8 +2449,8 @@ void UDynamicMeshSculptTool::UpdateMaterialMode(EMeshEditingMaterialModes Materi
 	if (MaterialMode == EMeshEditingMaterialModes::ExistingMaterial)
 	{
 		DynamicMeshComponent->ClearOverrideRenderMaterial();
-		DynamicMeshComponent->bCastDynamicShadow = UE::ToolTarget::GetTargetComponent(Target)->bCastDynamicShadow;
-		ActiveOverrideMaterial = nullptr;
+		DynamicMeshComponent->SetShadowsEnabled(UE::ToolTarget::GetTargetComponent(Target)->bCastDynamicShadow);
+		ActiveOverrideMaterial = nullptr; 
 	}
 	else 
 	{
@@ -2511,7 +2512,7 @@ void UDynamicMeshSculptTool::UpdateMaterialMode(EMeshEditingMaterialModes Materi
 			ActiveOverrideMaterial->SetScalarParameterValue(TEXT("FlatShading"), (ViewProperties->bFlatShading) ? 1.0f : 0.0f);
 		}
 
-		DynamicMeshComponent->bCastDynamicShadow = false;
+		DynamicMeshComponent->SetShadowsEnabled(false);
 	}
 }
 
diff --git a/Engine/Plugins/Experimental/MeshModelingToolset/Source/MeshModelingTools/Private/MeshAttributePaintTool.cpp b/Engine/Plugins/Experimental/MeshModelingToolset/Source/MeshModelingTools/Private/MeshAttributePaintTool.cpp
--- a/Engine/Plugins/Experimental/MeshModelingToolset/Source/MeshModelingTools/Private/MeshAttributePaintTool.cpp
+++ b/Engine/Plugins/Experimental/MeshModelingToolset/Source/MeshModelingTools/Private/MeshAttributePaintTool.cpp
@@ -165,7 +165,7 @@ void UMeshAttributePaintTool::Setup()
 	// configure preview mesh
 	PreviewMesh->SetTangentsMode(EDynamicMeshComponentTangentsMode::AutoCalculated);
 	//PreviewMesh->EnableWireframe(SelectionProps->bShowWireframe);
-	PreviewMesh->GetRootComponent()->bCastDynamicShadow = false;
+	PreviewMesh->SetShadowsEnabled(false);
 
 	// enable vtx colors on preview mesh
 	PreviewMesh->EditMesh([](FDynamicMesh3& Mesh)
diff --git a/Engine/Plugins/Experimental/MeshModelingToolset/Source/MeshModelingTools/Private/MeshSelectionTool.cpp b/Engine/Plugins/Experimental/MeshModelingToolset/Source/MeshModelingTools/Private/MeshSelectionTool.cpp
--- a/Engine/Plugins/Experimental/MeshModelingToolset/Source/MeshModelingTools/Private/MeshSelectionTool.cpp
+++ b/Engine/Plugins/Experimental/MeshModelingToolset/Source/MeshModelingTools/Private/MeshSelectionTool.cpp
@@ -86,7 +86,7 @@ void UMeshSelectionTool::Setup()
 	PreviewMesh->EnableWireframe(SelectionProps->bShowWireframe);
 
 	// disable shadows
-	PreviewMesh->GetRootComponent()->bCastDynamicShadow = false;
+	PreviewMesh->SetShadowsEnabled(false);
 
 	// configure secondary render material
 	UMaterialInterface* SelectionMaterial = ToolSetupUtil::GetSelectionMaterial(FLinearColor(0.9f, 0.1f, 0.1f), GetToolManager());
diff --git a/Engine/Plugins/Experimental/MeshModelingToolset/Source/MeshModelingTools/Private/Sculpting/MeshSculptToolBase.cpp b/Engine/Plugins/Experimental/MeshModelingToolset/Source/MeshModelingTools/Private/Sculpting/MeshSculptToolBase.cpp
--- a/Engine/Plugins/Experimental/MeshModelingToolset/Source/MeshModelingTools/Private/Sculpting/MeshSculptToolBase.cpp
+++ b/Engine/Plugins/Experimental/MeshModelingToolset/Source/MeshModelingTools/Private/Sculpting/MeshSculptToolBase.cpp
@@ -235,6 +235,9 @@ void UMeshSculptToolBase::Render(IToolsContextRenderAPI* RenderAPI)
 
 void UMeshSculptToolBase::InitializeSculptMeshComponent(UBaseDynamicMeshComponent* Component)
 {
+	// disable shadows initially, as changing shadow settings invalidates the SceneProxy
+	Component->SetShadowsEnabled(false);
+
 	Component->SetupAttachment(UE::ToolTarget::GetTargetActor(Target)->GetRootComponent());
 	Component->RegisterComponent();
 
@@ -957,7 +960,7 @@ void UMeshSculptToolBase::UpdateMaterialMode(EMeshEditingMaterialModes MaterialM
 	if (MaterialMode == EMeshEditingMaterialModes::ExistingMaterial)
 	{
 		GetSculptMeshComponent()->ClearOverrideRenderMaterial();
-		GetSculptMeshComponent()->bCastDynamicShadow = UE::ToolTarget::GetTargetComponent(Target)->bCastDynamicShadow;
+		GetSculptMeshComponent()->SetShadowsEnabled(UE::ToolTarget::GetTargetComponent(Target)->bCastDynamicShadow);
 		ActiveOverrideMaterial = nullptr;
 	}
 	else
@@ -1021,7 +1024,7 @@ void UMeshSculptToolBase::UpdateMaterialMode(EMeshEditingMaterialModes MaterialM
 			ActiveOverrideMaterial->SetScalarParameterValue(TEXT("FlatShading"), (ViewProperties->bFlatShading) ? 1.0f : 0.0f);
 		}
 
-		GetSculptMeshComponent()->bCastDynamicShadow = false;
+		GetSculptMeshComponent()->SetShadowsEnabled(false);
 	}
 }
 
diff --git a/Engine/Plugins/Experimental/MeshModelingToolset/Source/MeshModelingTools/Private/SkinWeightsBindingTool.cpp b/Engine/Plugins/Experimental/MeshModelingToolset/Source/MeshModelingTools/Private/SkinWeightsBindingTool.cpp
--- a/Engine/Plugins/Experimental/MeshModelingToolset/Source/MeshModelingTools/Private/SkinWeightsBindingTool.cpp
+++ b/Engine/Plugins/Experimental/MeshModelingToolset/Source/MeshModelingTools/Private/SkinWeightsBindingTool.cpp
@@ -858,7 +858,7 @@ void USkinWeightsBindingTool::Setup()
 
 	Preview->PreviewMesh->SetTransform(TargetComponent->GetWorldTransform());
 	Preview->PreviewMesh->SetTangentsMode(EDynamicMeshComponentTangentsMode::AutoCalculated);
-	Preview->PreviewMesh->GetRootComponent()->bCastDynamicShadow = false;
+	Preview->PreviewMesh->SetShadowsEnabled(false);
 	Preview->PreviewMesh->UpdatePreview(OriginalMesh.Get());
 
 	Occupancy = MakeShared<FOccupancyGrid>(*OriginalMesh, Properties->VoxelResolution);
diff --git a/Engine/Plugins/Experimental/MeshModelingToolset/Source/MeshModelingTools/Private/SkinWeightsPaintTool.cpp b/Engine/Plugins/Experimental/MeshModelingToolset/Source/MeshModelingTools/Private/SkinWeightsPaintTool.cpp
--- a/Engine/Plugins/Experimental/MeshModelingToolset/Source/MeshModelingTools/Private/SkinWeightsPaintTool.cpp
+++ b/Engine/Plugins/Experimental/MeshModelingToolset/Source/MeshModelingTools/Private/SkinWeightsPaintTool.cpp
@@ -100,7 +100,7 @@ void USkinWeightsPaintTool::Setup()
 	// configure preview mesh
 	PreviewMesh->SetTangentsMode(EDynamicMeshComponentTangentsMode::AutoCalculated);
 	//PreviewMesh->EnableWireframe(SelectionProps->bShowWireframe);
-	PreviewMesh->GetRootComponent()->bCastDynamicShadow = false;
+	PreviewMesh->SetShadowsEnabled(false);
 
 	// enable vtx colors on preview mesh
 	PreviewMesh->EditMesh([](FDynamicMesh3& Mesh)
diff --git a/Engine/Plugins/Experimental/MeshModelingToolset/Source/ModelingComponents/Private/Components/BaseDynamicMeshComponent.cpp b/Engine/Plugins/Experimental/MeshModelingToolset/Source/ModelingComponents/Private/Components/BaseDynamicMeshComponent.cpp
--- a/Engine/Plugins/Experimental/MeshModelingToolset/Source/ModelingComponents/Private/Components/BaseDynamicMeshComponent.cpp
+++ b/Engine/Plugins/Experimental/MeshModelingToolset/Source/ModelingComponents/Private/Components/BaseDynamicMeshComponent.cpp
@@ -1,6 +1,7 @@
 // Copyright Epic Games, Inc. All Rights Reserved. 
 
 #include "Components/BaseDynamicMeshComponent.h"
+#include "Components/BaseDynamicMeshSceneProxy.h"
 
 using namespace UE::Geometry;
 
@@ -11,6 +12,28 @@ UBaseDynamicMeshComponent::UBaseDynamicMeshComponent(const FObjectInitializer& O
 
 
 
+void UBaseDynamicMeshComponent::SetShadowsEnabled(bool bEnabled)
+{
+	// finish any drawing so that we can be certain our SceneProxy is no longer in use before we rebuild it below
+	FlushRenderingCommands();
+
+	SetCastShadow(bEnabled);
+	//bCastDynamicShadow = bEnabled;
+
+	// apparently SceneProxy has to be fully rebuilt to change shadow state
+
+	// this marks the SceneProxy for rebuild, but not immediately, and possibly allows bad things to happen
+	// before the end of the frame
+	//MarkRenderStateDirty();
+
+	// force immediate rebuild of the SceneProxy
+	if (IsRegistered())
+	{
+		ReregisterComponent();
+	}
+}
+
+
 void UBaseDynamicMeshComponent::SetOverrideRenderMaterial(UMaterialInterface* Material)
 {
 	if (OverrideRenderMaterial != Material)
diff --git a/Engine/Plugins/Experimental/MeshModelingToolset/Source/ModelingComponents/Private/Components/BaseDynamicMeshSceneProxy.h b/Engine/Plugins/Experimental/MeshModelingToolset/Source/ModelingComponents/Private/Components/BaseDynamicMeshSceneProxy.h
--- a/Engine/Plugins/Experimental/MeshModelingToolset/Source/ModelingComponents/Private/Components/BaseDynamicMeshSceneProxy.h
+++ b/Engine/Plugins/Experimental/MeshModelingToolset/Source/ModelingComponents/Private/Components/BaseDynamicMeshSceneProxy.h
@@ -7,6 +7,10 @@
 #include "VertexFactory.h"
 #include "LocalVertexFactory.h"
 #include "PrimitiveViewRelevance.h"
+#include "Rendering/StaticMeshVertexBuffer.h"
+#include "Rendering/PositionVertexBuffer.h"
+#include "Rendering/ColorVertexBuffer.h"
+#include "DynamicMeshBuilder.h"
 #include "Components/BaseDynamicMeshComponent.h"
 
 using UE::Geometry::FDynamicMesh3;
diff --git a/Engine/Plugins/Experimental/MeshModelingToolset/Source/ModelingComponents/Private/PreviewMesh.cpp b/Engine/Plugins/Experimental/MeshModelingToolset/Source/ModelingComponents/Private/PreviewMesh.cpp
--- a/Engine/Plugins/Experimental/MeshModelingToolset/Source/ModelingComponents/Private/PreviewMesh.cpp
+++ b/Engine/Plugins/Experimental/MeshModelingToolset/Source/ModelingComponents/Private/PreviewMesh.cpp
@@ -187,6 +187,12 @@ void UPreviewMesh::EnableWireframe(bool bEnable)
 	DynamicMeshComponent->bExplicitShowWireframe = bEnable;
 }
 
+void UPreviewMesh::SetShadowsEnabled(bool bEnable)
+{
+	check(DynamicMeshComponent);
+	DynamicMeshComponent->SetShadowsEnabled(bEnable);
+}
+
 
 FTransform UPreviewMesh::GetTransform() const
 {
diff --git a/Engine/Plugins/Experimental/MeshModelingToolset/Source/ModelingComponents/Public/Components/BaseDynamicMeshComponent.h b/Engine/Plugins/Experimental/MeshModelingToolset/Source/ModelingComponents/Public/Components/BaseDynamicMeshComponent.h
--- a/Engine/Plugins/Experimental/MeshModelingToolset/Source/ModelingComponents/Public/Components/BaseDynamicMeshComponent.h
+++ b/Engine/Plugins/Experimental/MeshModelingToolset/Source/ModelingComponents/Public/Components/BaseDynamicMeshComponent.h
@@ -18,6 +18,7 @@
 struct FMeshDescription;
 class FMeshVertexChange;
 class FMeshChange;
+class FBaseDynamicMeshSceneProxy;
 using UE::Geometry::FDynamicMesh3;
 
 /**
@@ -57,7 +58,7 @@ enum class EDynamicMeshComponentTangentsMode : uint8
 /**
  * UBaseDynamicMeshComponent is a base interface for a UMeshComponent based on a UDynamicMesh.
  */
-UCLASS(Abstract, hidecategories = (LOD, Physics, Collision), ClassGroup = Rendering)
+UCLASS(Abstract, hidecategories = (LOD), ClassGroup = Rendering)
 class MODELINGCOMPONENTS_API UBaseDynamicMeshComponent : 
 	public UMeshComponent, 
 	public IToolFrameworkComponent, 
@@ -160,6 +161,16 @@ public:
 
 
 protected:
+
+	/**
+	 * Subclass must implement this to return scene proxy if available, or nullptr
+	 */
+	virtual FBaseDynamicMeshSceneProxy* GetBaseSceneProxy()
+	{
+		unimplemented();
+		return nullptr;
+	}
+
 	/**
 	 * Subclass must implement this to notify allocated proxies of updated materials
 	 */
@@ -170,7 +181,6 @@ protected:
 
 
 
-
 	//===============================================================================================================
 	// Built-in Wireframe-on-Shaded Rendering support. The wireframe looks terrible but this is a convenient 
 	// way to enable/disable it.
@@ -180,7 +190,7 @@ public:
 	 * if true, we always show the wireframe on top of the shaded mesh, even when not in wireframe mode
 	 * @todo: this should not be public, access via Set/Get once all usage is cleaned up
 	 */
-	UPROPERTY(EditAnywhere, BlueprintReadWrite, Category = "Dynamic Mesh Component")
+	UPROPERTY(EditAnywhere, BlueprintReadWrite, Category = "Dynamic Mesh Component", meta = (DisplayName = "Wireframe Overlay") )
 	bool bExplicitShowWireframe = false;
 
 	/**
@@ -197,6 +207,19 @@ public:
 
 
 
+	//===============================================================================================================
+	// API for changing Rendering settings. Although some of these settings are available publicly
+	// on the Component (in some cases as public members), generally changing them requires more complex 
+	// Rendering invalidation.
+	//
+public:
+
+	UFUNCTION(BlueprintCallable, Category = "Dynamic Mesh Component")
+	virtual void SetShadowsEnabled(bool bEnabled);
+
+	UFUNCTION(BlueprintCallable, Category = "Dynamic Mesh Component")
+	virtual bool GetShadowsEnabled() const { return CastShadow; }
+
 
 
 	//===============================================================================================================
diff --git a/Engine/Plugins/Experimental/MeshModelingToolset/Source/ModelingComponents/Public/PreviewMesh.h b/Engine/Plugins/Experimental/MeshModelingToolset/Source/ModelingComponents/Public/PreviewMesh.h
--- a/Engine/Plugins/Experimental/MeshModelingToolset/Source/ModelingComponents/Public/PreviewMesh.h
+++ b/Engine/Plugins/Experimental/MeshModelingToolset/Source/ModelingComponents/Public/PreviewMesh.h
@@ -96,6 +96,11 @@ public:
 	 */
 	void EnableWireframe(bool bEnable);
 
+	/**
+	 * Enable/disable shadow rendering
+	 */
+	void SetShadowsEnabled(bool bEnable);
+
 	/**
 	 * Set material on the preview mesh
 	 */
