ModelingTools: Clean up AutoUV Tools and add support for Polygroup Layers.

- Refactor recompute-UVs-for-groups/islands out of ParameterizeMeshOp/Tool into RecomputeUVsOp/Tool. RecomputeUVsTool supports selection of active PolygroupLayer when applicable. Add support for auto-repacking.
- Clean up ParameterizeMeshOp/Tool, remove functionality related to recomputation that was moved to Recompute tool, add separate property sets for UVAtlas and XAtlas
- Add FDynamicMeshUVEditor::QuickPack() function to compute UV layout
- Add support for Polygroup Layer selection in SetCollisionGeometryTool, if only one source object is in use

#rnx

diff --git a/Engine/Plugins/Experimental/GeometryFlow/Source/GeometryFlowMeshProcessingEditor/Private/MeshProcessingNodes/MeshAutoGenerateUVsNode.cpp b/Engine/Plugins/Experimental/GeometryFlow/Source/GeometryFlowMeshProcessingEditor/Private/MeshProcessingNodes/MeshAutoGenerateUVsNode.cpp
--- a/Engine/Plugins/Experimental/GeometryFlow/Source/GeometryFlowMeshProcessingEditor/Private/MeshProcessingNodes/MeshAutoGenerateUVsNode.cpp
+++ b/Engine/Plugins/Experimental/GeometryFlow/Source/GeometryFlowMeshProcessingEditor/Private/MeshProcessingNodes/MeshAutoGenerateUVsNode.cpp
@@ -17,9 +17,7 @@ void FMeshAutoGenerateUVsNode::GenerateUVs(const FDynamicMesh3& MeshIn, const FM
 	ParameterizeMeshOp.Stretch = Settings.Stretch;
 	ParameterizeMeshOp.NumCharts = Settings.NumCharts;
 	ParameterizeMeshOp.InputMesh = InputMesh;
-
-	ParameterizeMeshOp.IslandMode = EParamOpIslandMode::Auto;
-	ParameterizeMeshOp.UnwrapType = EParamOpUnwrapType::MinStretch;
+	ParameterizeMeshOp.Method = UE::Geometry::EParamOpBackend::UVAtlas;
 
 	FProgressCancel Progress;
 	ParameterizeMeshOp.CalculateResult(&Progress);
diff --git a/Engine/Plugins/Experimental/GeometryProcessing/Source/DynamicMesh/Private/Parameterization/DynamicMeshUVEditor.cpp b/Engine/Plugins/Experimental/GeometryProcessing/Source/DynamicMesh/Private/Parameterization/DynamicMeshUVEditor.cpp
--- a/Engine/Plugins/Experimental/GeometryProcessing/Source/DynamicMesh/Private/Parameterization/DynamicMeshUVEditor.cpp
+++ b/Engine/Plugins/Experimental/GeometryProcessing/Source/DynamicMesh/Private/Parameterization/DynamicMeshUVEditor.cpp
@@ -14,6 +14,7 @@
 
 #include "Parameterization/MeshDijkstra.h"
 #include "Parameterization/MeshLocalParam.h"
+#include "Parameterization/MeshUVPacking.h"
 #include "Solvers/MeshParameterizationSolvers.h"
 
 #include "Async/ParallelFor.h"
@@ -989,4 +990,22 @@ void FDynamicMeshUVEditor::SetTriangleUVsFromCylinderProjection(
 	{
 		Result->NewUVElements = MoveTemp(NewUVIndices);
 	}
+}
+
+
+
+
+bool FDynamicMeshUVEditor::QuickPack(int32 TargetTextureResolution, float GutterSize)
+{
+	// always split bowties before packing
+	UVOverlay->SplitBowties();
+
+	FDynamicMeshUVPacker Packer(UVOverlay);
+	Packer.TextureResolution = TargetTextureResolution;
+	Packer.GutterSize = GutterSize;
+	Packer.bAllowFlips = false;
+
+	bool bOK = Packer.StandardPack();
+
+	return bOK;
 }
\ No newline at end of file
diff --git a/Engine/Plugins/Experimental/GeometryProcessing/Source/DynamicMesh/Public/Parameterization/DynamicMeshUVEditor.h b/Engine/Plugins/Experimental/GeometryProcessing/Source/DynamicMesh/Public/Parameterization/DynamicMeshUVEditor.h
--- a/Engine/Plugins/Experimental/GeometryProcessing/Source/DynamicMesh/Public/Parameterization/DynamicMeshUVEditor.h
+++ b/Engine/Plugins/Experimental/GeometryProcessing/Source/DynamicMesh/Public/Parameterization/DynamicMeshUVEditor.h
@@ -146,6 +146,11 @@ public:
 
 
 
+	/**
+	 * Pack UVs into unit rectangle
+	 */
+	bool QuickPack(int32 TargetTextureResolution = 512, float GutterSize = 1.0f);
+
 
 public:
 
diff --git a/Engine/Plugins/Experimental/GeometryProcessing/Source/DynamicMesh/Public/Polygroups/PolygroupSet.h b/Engine/Plugins/Experimental/GeometryProcessing/Source/DynamicMesh/Public/Polygroups/PolygroupSet.h
--- a/Engine/Plugins/Experimental/GeometryProcessing/Source/DynamicMesh/Public/Polygroups/PolygroupSet.h
+++ b/Engine/Plugins/Experimental/GeometryProcessing/Source/DynamicMesh/Public/Polygroups/PolygroupSet.h
@@ -65,6 +65,14 @@ struct DYNAMICMESH_API FPolygroupSet
 		return (PolygroupAttrib) ? PolygroupAttrib->GetValue(TriangleID) : Mesh->GetTriangleGroup(TriangleID);
 	}
 
+	/**
+	 * @return PolygroupID for a TriangleID
+	 */
+	int32 GetTriangleGroup(int32 TriangleID)
+	{
+		return (PolygroupAttrib) ? PolygroupAttrib->GetValue(TriangleID) : Mesh->GetTriangleGroup(TriangleID);
+	}
+
 	/**
 	 * Set the PolygroupID for a TriangleID
 	 */
diff --git a/Engine/Plugins/Experimental/MeshModelingToolset/Source/GeometryProcessingAdapters/Private/ApproximateActorsImpl.cpp b/Engine/Plugins/Experimental/MeshModelingToolset/Source/GeometryProcessingAdapters/Private/ApproximateActorsImpl.cpp
--- a/Engine/Plugins/Experimental/MeshModelingToolset/Source/GeometryProcessingAdapters/Private/ApproximateActorsImpl.cpp
+++ b/Engine/Plugins/Experimental/MeshModelingToolset/Source/GeometryProcessingAdapters/Private/ApproximateActorsImpl.cpp
@@ -603,8 +603,7 @@ static TSharedPtr<FApproximationMeshData> GenerateApproximationMesh(
 	ParameterizeMeshOp.Stretch = 0.1;
 	ParameterizeMeshOp.NumCharts = 0;
 	ParameterizeMeshOp.InputMesh = UVInputMesh;
-	ParameterizeMeshOp.IslandMode = EParamOpIslandMode::Auto;
-	ParameterizeMeshOp.UnwrapType = EParamOpUnwrapType::MinStretch;
+	ParameterizeMeshOp.Method = UE::Geometry::EParamOpBackend::UVAtlas;
 	FProgressCancel UVProgressCancel;
 	{
 		TRACE_CPUPROFILER_EVENT_SCOPE(ApproximateActorsImpl_Generate_GenerateUVs);
diff --git a/Engine/Plugins/Experimental/MeshModelingToolset/Source/MeshModelingTools/Private/MeshGroupPaintTool.cpp b/Engine/Plugins/Experimental/MeshModelingToolset/Source/MeshModelingTools/Private/MeshGroupPaintTool.cpp
--- a/Engine/Plugins/Experimental/MeshModelingToolset/Source/MeshModelingTools/Private/MeshGroupPaintTool.cpp
+++ b/Engine/Plugins/Experimental/MeshModelingToolset/Source/MeshModelingTools/Private/MeshGroupPaintTool.cpp
@@ -146,7 +146,7 @@ void UMeshGroupPaintTool::Setup()
 	PolyLassoMechanic->OnDrawPolyLassoFinished.AddUObject(this, &UMeshGroupPaintTool::OnPolyLassoFinished);
 
 	PolygroupLayerProperties = NewObject<UPolygroupLayersProperties>(this);
-	PolygroupLayerProperties->RestoreProperties(this);
+	PolygroupLayerProperties->RestoreProperties(this, TEXT("MeshGroupPaintTool"));
 	PolygroupLayerProperties->InitializeGroupLayers(GetSculptMesh());
 	PolygroupLayerProperties->WatchProperty(PolygroupLayerProperties->ActiveGroupLayer, [&](FName) { OnSelectedGroupLayerChanged(); });
 	UpdateActiveGroupLayer();
@@ -211,7 +211,7 @@ void UMeshGroupPaintTool::Setup()
 	MeshElementsDisplay->CreateInWorld(DynamicMeshComponent->GetWorld(), DynamicMeshComponent->GetComponentTransform());
 	if (ensure(MeshElementsDisplay->Settings))
 	{
-		MeshElementsDisplay->Settings->RestoreProperties(this);
+		MeshElementsDisplay->Settings->RestoreProperties(this, TEXT("MeshGroupPaintTool"));
 		AddToolPropertySource(MeshElementsDisplay->Settings);
 	}
 	MeshElementsDisplay->SetMeshAccessFunction([this](UMeshElementsVisualizer::ProcessDynamicMeshFunc ProcessFunc) {
@@ -246,12 +246,12 @@ void UMeshGroupPaintTool::Shutdown(EToolShutdownType ShutdownType)
 
 	if (ensure(MeshElementsDisplay->Settings))
 	{
-		MeshElementsDisplay->Settings->SaveProperties(this);
+		MeshElementsDisplay->Settings->SaveProperties(this, TEXT("MeshGroupPaintTool"));
 	}
 	MeshElementsDisplay->Disconnect();
 
 	FilterProperties->SaveProperties(this);
-	PolygroupLayerProperties->SaveProperties(this);
+	PolygroupLayerProperties->SaveProperties(this, TEXT("MeshGroupPaintTool"));
 
 
 	// TODO
diff --git a/Engine/Plugins/Experimental/MeshModelingToolset/Source/MeshModelingTools/Private/Physics/SetCollisionGeometryTool.cpp b/Engine/Plugins/Experimental/MeshModelingToolset/Source/MeshModelingTools/Private/Physics/SetCollisionGeometryTool.cpp
--- a/Engine/Plugins/Experimental/MeshModelingToolset/Source/MeshModelingTools/Private/Physics/SetCollisionGeometryTool.cpp
+++ b/Engine/Plugins/Experimental/MeshModelingToolset/Source/MeshModelingTools/Private/Physics/SetCollisionGeometryTool.cpp
@@ -4,12 +4,11 @@
 #include "InteractiveToolManager.h"
 #include "ToolBuilderUtil.h"
 #include "ToolSetupUtil.h"
-#include "Drawing/PreviewGeometryActor.h"
-#include "MeshDescriptionToDynamicMesh.h"
 #include "DynamicMesh/MeshTransforms.h"
 #include "DynamicMeshEditor.h"
 #include "Selections/MeshConnectedComponents.h"
 #include "DynamicSubmesh3.h"
+#include "Polygroups/PolygroupUtil.h"
 
 #include "ShapeApproximation/ShapeDetection3.h"
 #include "ShapeApproximation/MeshSimpleShapeApproximation.h"
@@ -26,8 +25,8 @@
 
 #include "TargetInterfaces/MeshDescriptionProvider.h"
 #include "TargetInterfaces/PrimitiveComponentBackedTarget.h"
-#include "TargetInterfaces/StaticMeshBackedTarget.h"
 #include "ToolTargetManager.h"
+#include "ModelingToolTargetUtil.h"
 
 #include "ExplicitUseGeometryMathTypes.h"		// using UE::Geometry::(math types)
 using namespace UE::Geometry;
@@ -77,23 +76,31 @@ void USetCollisionGeometryTool::Setup()
 		for (int32 k = 0; k < Targets.Num() -1; ++k)
 		{
 			SourceObjectIndices.Add(k);
-			TargetComponentInterface(k)->SetOwnerVisibility(false);
+			UE::ToolTarget::HideSourceObject(Targets[k]);
 		}
 	}
 
+	// collect input meshes
+	InitialSourceMeshes.SetNum(SourceObjectIndices.Num());
+	ParallelFor(SourceObjectIndices.Num(), [&](int32 k)
+	{
+		InitialSourceMeshes[k] = UE::ToolTarget::GetDynamicMeshCopy(Targets[k]);
+	});
+
+	// why is this here?
 	bInputMeshesValid = true;
 
-	IPrimitiveComponentBackedTarget* CollisionTarget = TargetComponentInterface(Targets.Num() - 1);
+	UToolTarget* CollisionTarget = Targets[Targets.Num() - 1];
 	PreviewGeom = NewObject<UPreviewGeometry>(this);
-	FTransform PreviewTransform = CollisionTarget->GetWorldTransform();
+	FTransform PreviewTransform = (FTransform)UE::ToolTarget::GetLocalToWorldTransform(CollisionTarget);
 	OrigTargetTransform = PreviewTransform;
 	TargetScale3D = PreviewTransform.GetScale3D();
 	PreviewTransform.SetScale3D(FVector::OneVector);
-	PreviewGeom->CreateInWorld(CollisionTarget->GetOwnerActor()->GetWorld(), PreviewTransform);
+	PreviewGeom->CreateInWorld(UE::ToolTarget::GetTargetActor(CollisionTarget)->GetWorld(), PreviewTransform);
 
 	// initialize initial collision object
 	InitialCollision = MakeShared<FPhysicsDataCollection, ESPMode::ThreadSafe>();
-	InitialCollision->InitializeFromComponent(CollisionTarget->GetOwnerComponent(), true);
+	InitialCollision->InitializeFromComponent(UE::ToolTarget::GetTargetComponent(CollisionTarget), true);
 	InitialCollision->ExternalScale3D = TargetScale3D;
 
 	// create tool options
@@ -101,7 +108,7 @@ void USetCollisionGeometryTool::Setup()
 	Settings->RestoreProperties(this);
 	AddToolPropertySource(Settings);
 	Settings->bUseWorldSpace = (SourceObjectIndices.Num() > 1);
-	Settings->WatchProperty(Settings->InputMode, [this](ESetCollisionGeometryInputMode) { bResultValid = false; });
+	Settings->WatchProperty(Settings->InputMode, [this](ESetCollisionGeometryInputMode) { OnInputModeChanged(); });
 	Settings->WatchProperty(Settings->GeometryType, [this](ECollisionGeometryType) { bResultValid = false; });
 	Settings->WatchProperty(Settings->bUseWorldSpace, [this](bool) { bInputMeshesValid = false; });
 	Settings->WatchProperty(Settings->bAppendToExisting, [this](bool) { bResultValid = false; });
@@ -117,7 +124,16 @@ void USetCollisionGeometryTool::Setup()
 	Settings->WatchProperty(Settings->bSimplifyPolygons, [this](bool) { bResultValid = false; });
 	Settings->WatchProperty(Settings->HullTolerance, [this](float) { bResultValid = false; });
 	Settings->WatchProperty(Settings->SweepAxis, [this](EProjectedHullAxis) { bResultValid = false; });
-		
+
+	if (InitialSourceMeshes.Num() == 1)
+	{
+		PolygroupLayerProperties = NewObject<UPolygroupLayersProperties>(this);
+		PolygroupLayerProperties->RestoreProperties(this, TEXT("SetCollisionGeometryTool"));
+		PolygroupLayerProperties->InitializeGroupLayers(&InitialSourceMeshes[0]);
+		PolygroupLayerProperties->WatchProperty(PolygroupLayerProperties->ActiveGroupLayer, [&](FName) { OnSelectedGroupLayerChanged(); });
+		AddToolPropertySource(PolygroupLayerProperties);
+	}
+
 	bResultValid = false;
 
 	VizSettings = NewObject<UCollisionGeometryVisualizationProperties>(this);
@@ -142,6 +158,11 @@ void USetCollisionGeometryTool::Shutdown(EToolShutdownType ShutdownType)
 {
 	VizSettings->SaveProperties(this);
 	Settings->SaveProperties(this);
+	if (PolygroupLayerProperties)
+	{
+		PolygroupLayerProperties->SaveProperties(this, TEXT("SetCollisionGeometryTool"));
+	}
+
 	PreviewGeom->Disconnect();
 
 	// show hidden sources
@@ -149,7 +170,7 @@ void USetCollisionGeometryTool::Shutdown(EToolShutdownType ShutdownType)
 	{
 		for (int32 k : SourceObjectIndices)
 		{
-			TargetComponentInterface(k)->SetOwnerVisibility(true);
+			UE::ToolTarget::ShowSourceObject(Targets[k]);
 		}
 	}
 
@@ -162,51 +183,53 @@ void USetCollisionGeometryTool::Shutdown(EToolShutdownType ShutdownType)
 
 		// code below derived from FStaticMeshEditor::DuplicateSelectedPrims(), FStaticMeshEditor::OnCollisionSphere(), and GeomFitUtils.cpp::GenerateSphylAsSimpleCollision()
 
-		IPrimitiveComponentBackedTarget* CollisionTarget = TargetComponentInterface(Targets.Num() - 1);
-		UStaticMeshComponent* StaticMeshComponent = CastChecked<UStaticMeshComponent>(CollisionTarget->GetOwnerComponent());
-		UStaticMesh* StaticMesh = StaticMeshComponent->GetStaticMesh();
-		UBodySetup* BodySetup = StaticMesh->GetBodySetup();
-
-		// mark the BodySetup for modification. Do we need to modify the UStaticMesh??
-		BodySetup->Modify();
+		UPrimitiveComponent* Component = UE::ToolTarget::GetTargetComponent(Targets[Targets.Num() - 1]);
+		UStaticMeshComponent* StaticMeshComponent = Cast<UStaticMeshComponent>(Component);
+		UStaticMesh* StaticMesh = (StaticMeshComponent) ? StaticMeshComponent->GetStaticMesh() : nullptr;
+		UBodySetup* BodySetup = (StaticMesh) ? StaticMesh->GetBodySetup() : nullptr;
+		if (BodySetup != nullptr)
+		{
+			// mark the BodySetup for modification. Do we need to modify the UStaticMesh??
+			BodySetup->Modify();
 
-		// clear existing simple collision. This will call BodySetup->InvalidatePhysicsData()
-		BodySetup->RemoveSimpleCollision();
+			// clear existing simple collision. This will call BodySetup->InvalidatePhysicsData()
+			BodySetup->RemoveSimpleCollision();
 
-		// set new collision geometry
-		BodySetup->AggGeom = GeneratedCollision->AggGeom;
+			// set new collision geometry
+			BodySetup->AggGeom = GeneratedCollision->AggGeom;
 
-		// update collision type
-		BodySetup->CollisionTraceFlag = (ECollisionTraceFlag)(int32)Settings->SetCollisionType;
+			// update collision type
+			BodySetup->CollisionTraceFlag = (ECollisionTraceFlag)(int32)Settings->SetCollisionType;
 
-		// rebuild physics meshes
-		BodySetup->CreatePhysicsMeshes();
+			// rebuild physics meshes
+			BodySetup->CreatePhysicsMeshes();
 
-		// rebuild nav collision (? StaticMeshEditor does this)
-		StaticMesh->CreateNavCollision(/*bIsUpdate=*/true);
+			// rebuild nav collision (? StaticMeshEditor does this)
+			StaticMesh->CreateNavCollision(/*bIsUpdate=*/true);
 
-		// update physics state on all components using this StaticMesh
-		for (FThreadSafeObjectIterator Iter(UStaticMeshComponent::StaticClass()); Iter; ++Iter)
-		{
-			UStaticMeshComponent* SMComponent = Cast<UStaticMeshComponent>(*Iter);
-			if (SMComponent->GetStaticMesh() == StaticMesh)
+			// update physics state on all components using this StaticMesh
+			for (FThreadSafeObjectIterator Iter(UStaticMeshComponent::StaticClass()); Iter; ++Iter)
 			{
-				if (SMComponent->IsPhysicsStateCreated())
+				UStaticMeshComponent* SMComponent = Cast<UStaticMeshComponent>(*Iter);
+				if (SMComponent->GetStaticMesh() == StaticMesh)
 				{
-					SMComponent->RecreatePhysicsState();
+					if (SMComponent->IsPhysicsStateCreated())
+					{
+						SMComponent->RecreatePhysicsState();
+					}
 				}
 			}
-		}
 
-		// do we need to do a post edit change here??
+			// do we need to do a post edit change here??
 
-		// mark static mesh as dirty so it gets resaved?
-		StaticMesh->MarkPackageDirty();
+			// mark static mesh as dirty so it gets resaved?
+			StaticMesh->MarkPackageDirty();
 
 #if WITH_EDITORONLY_DATA
-		// mark the static mesh as having customized collision so it is not regenerated on reimport
-		StaticMesh->bCustomizedCollision = true;
+			// mark the static mesh as having customized collision so it is not regenerated on reimport
+			StaticMesh->bCustomizedCollision = true;
 #endif // WITH_EDITORONLY_DATA
+		}
 
 		// post the undo transaction
 		GetToolManager()->EndUndoTransaction();
@@ -240,6 +263,48 @@ void USetCollisionGeometryTool::OnTick(float DeltaTime)
 }
 
 
+
+void USetCollisionGeometryTool::OnInputModeChanged()
+{
+	if (PolygroupLayerProperties != nullptr)
+	{
+		SetToolPropertySourceEnabled(PolygroupLayerProperties, Settings->InputMode == ESetCollisionGeometryInputMode::PerMeshGroup);
+	}
+	bResultValid = false;
+}
+
+void USetCollisionGeometryTool::OnSelectedGroupLayerChanged()
+{
+	bInputMeshesValid = false;
+	bResultValid = false;
+}
+
+
+void USetCollisionGeometryTool::UpdateActiveGroupLayer()
+{
+	if (InitialSourceMeshes.Num() != 1)
+	{
+		ensure(false);		// should not get here
+		return;
+	}
+	FDynamicMesh3* GroupLayersMesh = &InitialSourceMeshes[0];
+
+	if (PolygroupLayerProperties->HasSelectedPolygroup() == false)
+	{
+		ActiveGroupSet = MakeUnique<UE::Geometry::FPolygroupSet>(GroupLayersMesh);
+	}
+	else
+	{
+		FName SelectedName = PolygroupLayerProperties->ActiveGroupLayer;
+		FDynamicMeshPolygroupAttribute* FoundAttrib = UE::Geometry::FindPolygroupLayerByName(*GroupLayersMesh, SelectedName);
+		ensureMsgf(FoundAttrib, TEXT("Selected Attribute Not Found! Falling back to Default group layer."));
+		ActiveGroupSet = MakeUnique<UE::Geometry::FPolygroupSet>(GroupLayersMesh, FoundAttrib);
+	}
+}
+
+
+
+
 void USetCollisionGeometryTool::UpdateVisualization()
 {
 	float UseThickness = VizSettings->LineThickness;
@@ -415,22 +480,23 @@ TArray<const T*> MakeRawPointerList(const TArray<TSharedPtr<T, ESPMode::ThreadSa
 
 void USetCollisionGeometryTool::PrecomputeInputMeshes()
 {
-	IPrimitiveComponentBackedTarget* CollisionTarget = TargetComponentInterface(Targets.Num()-1);
-	UE::Geometry::FTransform3d TargetTransform(CollisionTarget->GetWorldTransform());
+	if (InitialSourceMeshes.Num() == 1)
+	{
+		UpdateActiveGroupLayer();
+	}
+
+	UToolTarget* CollisionTarget = Targets[Targets.Num() - 1];
+	UE::Geometry::FTransform3d TargetTransform(UE::ToolTarget::GetLocalToWorldTransform(CollisionTarget));
 	UE::Geometry::FTransform3d TargetTransformInv = TargetTransform.Inverse();
 
 	InputMeshes.Reset();
 	InputMeshes.SetNum(SourceObjectIndices.Num());
 	ParallelFor(SourceObjectIndices.Num(), [&](int32 k)
 	{
-		FMeshDescriptionToDynamicMesh Converter;
-		Converter.bCalculateMaps = false;
-		Converter.bDisableAttributes = true;
-		FDynamicMesh3 SourceMesh;
-		Converter.Convert(TargetMeshProviderInterface(k)->GetMeshDescription(), SourceMesh);
+		FDynamicMesh3 SourceMesh = InitialSourceMeshes[k];
 		if (Settings->bUseWorldSpace)
 		{
-			UE::Geometry::FTransform3d ToWorld(TargetComponentInterface(k)->GetWorldTransform());
+			UE::Geometry::FTransform3d ToWorld(UE::ToolTarget::GetLocalToWorldTransform(Targets[k]));
 			MeshTransforms::ApplyTransform(SourceMesh, ToWorld);
 			MeshTransforms::ApplyTransform(SourceMesh, TargetTransformInv);
 		}
@@ -465,8 +531,17 @@ void USetCollisionGeometryTool::PrecomputeInputMeshes()
 
 	// build per-group input meshes
 	PerGroupInputMeshes.Reset();
-	InitializeDerivedMeshSet(InputMeshes, PerGroupInputMeshes,
-		[&](const FDynamicMesh3* Mesh, int32 Tri0, int32 Tri1) { return Mesh->GetTriangleGroup(Tri0) == Mesh->GetTriangleGroup(Tri1); });
+	if (ActiveGroupSet.IsValid())
+	{
+		check(InputMeshes.Num() == 1);
+		InitializeDerivedMeshSet(InputMeshes, PerGroupInputMeshes,
+			[&](const FDynamicMesh3* Mesh, int32 Tri0, int32 Tri1) { return ActiveGroupSet->GetTriangleGroup(Tri0) == ActiveGroupSet->GetTriangleGroup(Tri1); });
+	}
+	else
+	{
+		InitializeDerivedMeshSet(InputMeshes, PerGroupInputMeshes,
+			[&](const FDynamicMesh3* Mesh, int32 Tri0, int32 Tri1) { return Mesh->GetTriangleGroup(Tri0) == Mesh->GetTriangleGroup(Tri1); });
+	}
 	PerGroupMeshesApproximator = MakeShared<FMeshSimpleShapeApproximation, ESPMode::ThreadSafe>();
 	PerGroupMeshesApproximator->InitializeSourceMeshes(MakeRawPointerList<FDynamicMesh3>(PerGroupInputMeshes));
 
diff --git a/Engine/Plugins/Experimental/MeshModelingToolset/Source/MeshModelingTools/Private/RecomputeUVsTool.cpp b/Engine/Plugins/Experimental/MeshModelingToolset/Source/MeshModelingTools/Private/RecomputeUVsTool.cpp
new file mode 100644
--- /dev/null
+++ b/Engine/Plugins/Experimental/MeshModelingToolset/Source/MeshModelingTools/Private/RecomputeUVsTool.cpp
@@ -0,0 +1,218 @@
+// Copyright Epic Games, Inc. All Rights Reserved.
+
+#include "RecomputeUVsTool.h"
+#include "InteractiveToolManager.h"
+#include "ToolBuilderUtil.h"
+#include "DynamicMesh/DynamicMesh3.h"
+#include "Polygroups/PolygroupUtil.h"
+#include "FaceGroupUtil.h"
+#include "ToolSetupUtil.h"
+#include "ModelingToolTargetUtil.h"
+#include "ParameterizationOps/RecomputeUVsOp.h"
+
+
+#include "ExplicitUseGeometryMathTypes.h"		// using UE::Geometry::(math types)
+using namespace UE::Geometry;
+
+#define LOCTEXT_NAMESPACE "URecomputeUVsTool"
+
+
+/*
+ * ToolBuilder
+ */
+
+USingleSelectionMeshEditingTool* URecomputeUVsToolBuilder::CreateNewTool(const FToolBuilderState& SceneState) const
+{
+	URecomputeUVsTool* NewTool = NewObject<URecomputeUVsTool>(SceneState.ToolManager);
+	return NewTool;
+}
+
+/*
+ * Tool
+ */
+
+
+void URecomputeUVsTool::Setup()
+{
+	UInteractiveTool::Setup();
+
+	InputMesh = MakeShared<FDynamicMesh3, ESPMode::ThreadSafe>();
+	*InputMesh = UE::ToolTarget::GetDynamicMeshCopy(Target);
+
+	Preview = NewObject<UMeshOpPreviewWithBackgroundCompute>(this);
+	Preview->Setup(this->TargetWorld, this);
+	Preview->PreviewMesh->SetTangentsMode(EDynamicMeshComponentTangentsMode::AutoCalculated);
+	Preview->PreviewMesh->ReplaceMesh(*InputMesh);
+	Preview->ConfigureMaterials(UE::ToolTarget::GetMaterialSet(Target).Materials,
+								ToolSetupUtil::GetDefaultWorkingMaterial(GetToolManager()) );
+	Preview->PreviewMesh->SetTransform((FTransform)UE::ToolTarget::GetLocalToWorldTransform(Target));
+
+	Preview->OnMeshUpdated.AddLambda([this](UMeshOpPreviewWithBackgroundCompute* Op)
+	{
+		MaterialSettings->UpdateMaterials();
+	});
+
+	UE::ToolTarget::HideSourceObject(Target);
+
+	// initialize our properties
+
+	UVChannelProperties = NewObject<UMeshUVChannelProperties>(this);
+	UVChannelProperties->RestoreProperties(this);
+	UVChannelProperties->Initialize(InputMesh.Get(), false);
+	UVChannelProperties->ValidateSelection(true);
+	UVChannelProperties->WatchProperty(UVChannelProperties->UVChannel, [this](const FString& NewValue) 
+	{
+		MaterialSettings->UVChannel = UVChannelProperties->GetSelectedChannelIndex(true);
+	});
+	AddToolPropertySource(UVChannelProperties);
+
+	Settings = NewObject<URecomputeUVsToolProperties>(this);
+	Settings->RestoreProperties(this);
+	AddToolPropertySource(Settings);
+
+	PolygroupLayerProperties = NewObject<UPolygroupLayersProperties>(this);
+	PolygroupLayerProperties->RestoreProperties(this, TEXT("RecomputeUVsTool"));
+	PolygroupLayerProperties->InitializeGroupLayers(InputMesh.Get());
+	PolygroupLayerProperties->WatchProperty(PolygroupLayerProperties->ActiveGroupLayer, [&](FName) { OnSelectedGroupLayerChanged(); });
+	AddToolPropertySource(PolygroupLayerProperties);
+	UpdateActiveGroupLayer();
+
+	MaterialSettings = NewObject<UExistingMeshMaterialProperties>(this);
+	MaterialSettings->MaterialMode = ESetMeshMaterialMode::Checkerboard;
+	MaterialSettings->RestoreProperties(this, TEXT("ModelingUVTools"));
+	AddToolPropertySource(MaterialSettings);
+	// force update
+	MaterialSettings->UpdateMaterials();
+	Preview->OverrideMaterial = MaterialSettings->GetActiveOverrideMaterial();
+
+	Preview->InvalidateResult();    // start compute
+
+	SetToolDisplayName(LOCTEXT("ToolNameLocal", "UV Unwrap"));
+	GetToolManager()->DisplayMessage(
+		LOCTEXT("OnStartTool_Regions", "Generate UVs for Polygroups or existing UV charts of the Mesh using various strategies."),
+		EToolMessageLevel::UserNotification);
+}
+
+
+void URecomputeUVsTool::OnPropertyModified(UObject* PropertySet, FProperty* Property)
+{
+	bool bForceMaterialUpdate = false;
+	if (PropertySet == Settings || PropertySet == UVChannelProperties)
+	{
+		// One of the UV generation properties must have changed.  Dirty the result to force a recompute
+		Preview->InvalidateResult();
+		bForceMaterialUpdate = true;
+	}
+
+	if (PropertySet == MaterialSettings || bForceMaterialUpdate)
+	{
+		MaterialSettings->UpdateMaterials();
+		Preview->OverrideMaterial = MaterialSettings->GetActiveOverrideMaterial();
+	}
+}
+
+
+void URecomputeUVsTool::Shutdown(EToolShutdownType ShutdownType)
+{
+	UVChannelProperties->SaveProperties(this);
+	Settings->SaveProperties(this);
+	PolygroupLayerProperties->RestoreProperties(this, TEXT("RecomputeUVsTool"));
+	MaterialSettings->SaveProperties(this, TEXT("ModelingUVTools"));
+
+	UE::ToolTarget::ShowSourceObject(Target);
+
+	FDynamicMeshOpResult Result = Preview->Shutdown();
+	if (ShutdownType == EToolShutdownType::Accept)
+	{
+		GetToolManager()->BeginUndoTransaction(LOCTEXT("RecomputeUVs", "Recompute UVs"));
+		FDynamicMesh3* NewDynamicMesh = Result.Mesh.Get();
+		if (ensure(NewDynamicMesh))
+		{
+			UE::ToolTarget::CommitDynamicMeshUVUpdate(Target, NewDynamicMesh);
+		}
+		GetToolManager()->EndUndoTransaction();
+	}
+
+}
+
+void URecomputeUVsTool::OnTick(float DeltaTime)
+{
+	Preview->Tick(DeltaTime);
+}
+
+bool URecomputeUVsTool::CanAccept() const
+{
+	return Super::CanAccept() && Preview->HaveValidResult();
+}
+
+TUniquePtr<FDynamicMeshOperator> URecomputeUVsTool::MakeNewOperator()
+{
+	FAxisAlignedBox3d MeshBounds = Preview->PreviewMesh->GetMesh()->GetBounds();
+	TUniquePtr<FRecomputeUVsOp> RecomputeUVsOp = MakeUnique<FRecomputeUVsOp>();
+	RecomputeUVsOp->InputMesh = InputMesh;
+	RecomputeUVsOp->InputGroups = ActiveGroupSet;
+	RecomputeUVsOp->UVLayer = UVChannelProperties->GetSelectedChannelIndex(true);
+	
+	RecomputeUVsOp->IslandMode = (UE::Geometry::ERecomputeUVsIslandMode)(int)Settings->IslandMode;
+	RecomputeUVsOp->UnwrapType = (UE::Geometry::ERecomputeUVsUnwrapType)(int)Settings->UnwrapType;
+
+	RecomputeUVsOp->bPackUVs = Settings->bAutoPack;
+	if (Settings->bAutoPack)
+	{
+		RecomputeUVsOp->PackingTextureResolution = Settings->TextureResolution;
+		RecomputeUVsOp->bNormalizeAreas = false;
+		RecomputeUVsOp->AreaScaling = 1.0;
+	}
+	else
+	{
+		switch (Settings->UVScaleMode)
+		{
+		case ERecomputeUVsToolUVScaleMode::NoScaling:
+			RecomputeUVsOp->bNormalizeAreas = false;
+			RecomputeUVsOp->AreaScaling = 1.0;
+			break;
+		case ERecomputeUVsToolUVScaleMode::NormalizeToBounds:
+			RecomputeUVsOp->bNormalizeAreas = true;
+			RecomputeUVsOp->AreaScaling = Settings->UVScale / MeshBounds.MaxDim();
+			break;
+		case ERecomputeUVsToolUVScaleMode::NormalizeToWorld:
+			RecomputeUVsOp->bNormalizeAreas = true;
+			RecomputeUVsOp->AreaScaling = Settings->UVScale;
+			break;
+		}
+	}
+
+
+	RecomputeUVsOp->SetTransform(UE::ToolTarget::GetLocalToWorldTransform(Target));
+
+	return RecomputeUVsOp;
+}
+
+
+
+void URecomputeUVsTool::OnSelectedGroupLayerChanged()
+{
+	UpdateActiveGroupLayer();
+	Preview->InvalidateResult();
+}
+
+
+void URecomputeUVsTool::UpdateActiveGroupLayer()
+{
+	if (PolygroupLayerProperties->HasSelectedPolygroup() == false)
+	{
+		ActiveGroupSet = MakeShared<UE::Geometry::FPolygroupSet, ESPMode::ThreadSafe>(InputMesh.Get());
+	}
+	else
+	{
+		FName SelectedName = PolygroupLayerProperties->ActiveGroupLayer;
+		FDynamicMeshPolygroupAttribute* FoundAttrib = UE::Geometry::FindPolygroupLayerByName(*InputMesh, SelectedName);
+		ensureMsgf(FoundAttrib, TEXT("Selected Attribute Not Found! Falling back to Default group layer."));
+		ActiveGroupSet = MakeShared<UE::Geometry::FPolygroupSet, ESPMode::ThreadSafe>(InputMesh.Get(), FoundAttrib);
+	}
+
+}
+
+
+
+#undef LOCTEXT_NAMESPACE
diff --git a/Engine/Plugins/Experimental/MeshModelingToolset/Source/MeshModelingTools/Public/Physics/SetCollisionGeometryTool.h b/Engine/Plugins/Experimental/MeshModelingToolset/Source/MeshModelingTools/Public/Physics/SetCollisionGeometryTool.h
--- a/Engine/Plugins/Experimental/MeshModelingToolset/Source/MeshModelingTools/Public/Physics/SetCollisionGeometryTool.h
+++ b/Engine/Plugins/Experimental/MeshModelingToolset/Source/MeshModelingTools/Public/Physics/SetCollisionGeometryTool.h
@@ -10,6 +10,8 @@
 #include "OrientedBoxTypes.h"
 #include "CapsuleTypes.h"
 #include "Physics/CollisionPropertySets.h"
+#include "PropertySets/PolygroupLayersProperties.h"
+#include "Polygroups/PolygroupSet.h"
 #include "SetCollisionGeometryTool.generated.h"
 
 class UPreviewGeometry;
@@ -156,6 +158,8 @@ protected:
 	UPROPERTY()
 	USetCollisionGeometryToolProperties* Settings = nullptr;
 
+	UPROPERTY()
+	UPolygroupLayersProperties* PolygroupLayerProperties = nullptr;
 
 	UPROPERTY()
 	UCollisionGeometryVisualizationProperties* VizSettings = nullptr;
@@ -173,6 +177,10 @@ protected:
 	TArray<int32> SourceObjectIndices;
 	bool bSourcesHidden = false;
 
+	TArray<FDynamicMesh3> InitialSourceMeshes;
+
+	void OnInputModeChanged();
+
 	enum class EDetectedCollisionGeometry
 	{
 		None,
@@ -211,6 +219,10 @@ protected:
 		TFunctionRef<bool(const FDynamicMesh3*, int32, int32)> TrisConnectedPredicate);
 	TSharedPtr<UE::Geometry::FMeshSimpleShapeApproximation, ESPMode::ThreadSafe>& GetApproximator(ESetCollisionGeometryInputMode MeshSetMode);
 
+	TUniquePtr<UE::Geometry::FPolygroupSet> ActiveGroupSet;
+	void OnSelectedGroupLayerChanged();
+	void UpdateActiveGroupLayer();
+
 	FTransform OrigTargetTransform;
 	FVector TargetScale3D;
 
diff --git a/Engine/Plugins/Experimental/MeshModelingToolset/Source/MeshModelingTools/Public/RecomputeUVsTool.h b/Engine/Plugins/Experimental/MeshModelingToolset/Source/MeshModelingTools/Public/RecomputeUVsTool.h
new file mode 100644
--- /dev/null
+++ b/Engine/Plugins/Experimental/MeshModelingToolset/Source/MeshModelingTools/Public/RecomputeUVsTool.h
@@ -0,0 +1,151 @@
+// Copyright Epic Games, Inc. All Rights Reserved.
+
+#pragma once
+
+#include "CoreMinimal.h"
+#include "BaseTools/SingleSelectionMeshEditingTool.h"
+#include "DynamicMesh/DynamicMesh3.h"
+#include "MeshOpPreviewHelpers.h"
+#include "Properties/MeshMaterialProperties.h"
+#include "Properties/MeshUVChannelProperties.h"
+#include "PropertySets/PolygroupLayersProperties.h"
+#include "Polygroups/PolygroupSet.h"
+
+#include "RecomputeUVsTool.generated.h"
+
+
+// predeclarations
+class UDynamicMeshComponent;
+class UMaterialInterface;
+class UMaterialInstanceDynamic;
+
+
+/**
+ *
+ */
+UCLASS()
+class MESHMODELINGTOOLS_API URecomputeUVsToolBuilder : public USingleSelectionMeshEditingToolBuilder
+{
+	GENERATED_BODY()
+
+public:
+	virtual USingleSelectionMeshEditingTool* CreateNewTool(const FToolBuilderState& SceneState) const override;
+};
+
+
+
+UENUM()
+enum class ERecomputeUVsUnwrapType
+{
+	// Values must match UE::Geometry::ERecomputeUVsUnwrapType
+
+	/** */
+	ExpMap = 0,
+	/** */
+	Conformal = 1
+};
+
+
+
+UENUM()
+enum class ERecomputeUVsIslandMode
+{
+	// Values must match UE::Geometry::ERecomputeUVsIslandMode
+
+	/** */
+	PolyGroups = 0,
+	/** */
+	ExistingUVs = 1
+};
+
+UENUM()
+enum class ERecomputeUVsToolUVScaleMode
+{
+	/** No scaling is applied to UV islands */
+	NoScaling,
+	/** Scale UV islands such that they have constant relative area, relative to object bounds */
+	NormalizeToBounds,
+	/** Scale UV islands such that they have constant relative area, relative to world space */
+	NormalizeToWorld
+};
+
+
+UCLASS()
+class MESHMODELINGTOOLS_API URecomputeUVsToolProperties : public UInteractiveToolPropertySet
+{
+	GENERATED_BODY()
+public:
+
+	UPROPERTY(EditAnywhere, Category = "Compute UVs")
+	ERecomputeUVsIslandMode IslandMode = ERecomputeUVsIslandMode::PolyGroups;
+
+	UPROPERTY(EditAnywhere, Category = "Compute UVs")
+	ERecomputeUVsUnwrapType UnwrapType = ERecomputeUVsUnwrapType::ExpMap;
+
+	UPROPERTY(EditAnywhere, Category = "UV Layout")
+	bool bAutoPack = true;
+
+	UPROPERTY(EditAnywhere, Category = "UV Layout", meta = (UIMin = "64", UIMax = "2048", ClampMin = "2", ClampMax = "4096", EditCondition = bAutoPack))
+	int TextureResolution = 1024;
+
+	/** Scaling applied to UV islands after solve */
+	UPROPERTY(EditAnywhere, Category = "Transform UVs", meta = (EditCondition = "bAutoPack == false"))
+	ERecomputeUVsToolUVScaleMode UVScaleMode = ERecomputeUVsToolUVScaleMode::NormalizeToBounds;
+
+	/** Scaling factor used for UV island normalization/scaling */
+	UPROPERTY(EditAnywhere, Category = "Transform UVs", meta = (EditCondition = "UVScaleMode!=ERecomputeUVsToolUVScaleMode::NoScaling && bAutoPack == false", UIMin = "0.001", UIMax = "10", ClampMin = "0.00001", ClampMax = "1000000.0") )
+	float UVScale = 1.0;
+};
+
+
+
+
+
+/**
+ * URecomputeUVsTool Recomputes UVs based on existing segmentations of the mesh
+ */
+UCLASS()
+class MESHMODELINGTOOLS_API URecomputeUVsTool : public USingleSelectionMeshEditingTool, public UE::Geometry::IDynamicMeshOperatorFactory
+{
+	GENERATED_BODY()
+
+public:
+	virtual void Setup() override;
+	virtual void Shutdown(EToolShutdownType ShutdownType) override;
+
+	virtual void OnTick(float DeltaTime) override;
+
+	virtual bool HasCancel() const override { return true; }
+	virtual bool HasAccept() const override { return true; }
+	virtual bool CanAccept() const override;
+
+	virtual void OnPropertyModified(UObject* PropertySet, FProperty* Property) override;
+
+	// IDynamicMeshOperatorFactory API
+	virtual TUniquePtr<UE::Geometry::FDynamicMeshOperator> MakeNewOperator() override;
+
+protected:
+	UPROPERTY()
+	UMeshUVChannelProperties* UVChannelProperties = nullptr;
+
+	UPROPERTY()
+	URecomputeUVsToolProperties* Settings = nullptr;
+
+	UPROPERTY()
+	UPolygroupLayersProperties* PolygroupLayerProperties = nullptr;
+
+	UPROPERTY()
+	UExistingMeshMaterialProperties* MaterialSettings = nullptr;
+
+protected:
+	UPROPERTY()
+	UMeshOpPreviewWithBackgroundCompute* Preview = nullptr;
+
+protected:
+	TSharedPtr<UE::Geometry::FDynamicMesh3, ESPMode::ThreadSafe> InputMesh;
+
+
+	TSharedPtr<UE::Geometry::FPolygroupSet, ESPMode::ThreadSafe> ActiveGroupSet;
+	void OnSelectedGroupLayerChanged();
+	void UpdateActiveGroupLayer();
+};
diff --git a/Engine/Plugins/Experimental/MeshModelingToolset/Source/MeshModelingToolsEditorOnly/Private/ParameterizeMeshTool.cpp b/Engine/Plugins/Experimental/MeshModelingToolset/Source/MeshModelingToolsEditorOnly/Private/ParameterizeMeshTool.cpp
--- a/Engine/Plugins/Experimental/MeshModelingToolset/Source/MeshModelingToolsEditorOnly/Private/ParameterizeMeshTool.cpp
+++ b/Engine/Plugins/Experimental/MeshModelingToolset/Source/MeshModelingToolsEditorOnly/Private/ParameterizeMeshTool.cpp
@@ -3,31 +3,17 @@
 #include "ParameterizeMeshTool.h"
 #include "InteractiveToolManager.h"
 #include "ToolBuilderUtil.h"
-
+#include "ToolSetupUtil.h"
+#include "ModelingToolTargetUtil.h"
 #include "DynamicMesh/DynamicMesh3.h"
-#include "DynamicMeshToMeshDescription.h"
 #include "FaceGroupUtil.h"
-
-#include "Components/DynamicMeshComponent.h"
-#include "Materials/MaterialInstanceDynamic.h"
-#include "MeshDescription.h"
-#include "MeshDescriptionToDynamicMesh.h"
-
 #include "ParameterizationOps/ParameterizeMeshOp.h"
-#include "ToolSetupUtil.h"
-#include "ModelingToolTargetUtil.h"
-#include "TargetInterfaces/MaterialProvider.h"
-#include "TargetInterfaces/MeshDescriptionCommitter.h"
-#include "TargetInterfaces/MeshDescriptionProvider.h"
-#include "TargetInterfaces/PrimitiveComponentBackedTarget.h"
 
 #include "ExplicitUseGeometryMathTypes.h"		// using UE::Geometry::(math types)
 using namespace UE::Geometry;
 
 #define LOCTEXT_NAMESPACE "UParameterizeMeshTool"
 
-DEFINE_LOG_CATEGORY_STATIC(LogParameterizeMeshTool, Log, All);
-
 /*
  * ToolBuilder
  */
@@ -35,136 +21,98 @@ DEFINE_LOG_CATEGORY_STATIC(LogParameterizeMeshTool, Log, All);
 USingleSelectionMeshEditingTool* UParameterizeMeshToolBuilder::CreateNewTool(const FToolBuilderState& SceneState) const
 {
 	UParameterizeMeshTool* NewTool = NewObject<UParameterizeMeshTool>(SceneState.ToolManager);
-	NewTool->SetUseAutoGlobalParameterizationMode(bDoAutomaticGlobalUnwrap);
 	return NewTool;
 }
 
 /*
  * Tool
  */
-UParameterizeMeshTool::UParameterizeMeshTool()
-{
-}
-
-void UParameterizeMeshTool::SetUseAutoGlobalParameterizationMode(bool bEnable)
-{
-	bDoAutomaticGlobalUnwrap = bEnable;
-}
 
 
 void UParameterizeMeshTool::Setup()
 {
 	UInteractiveTool::Setup();
 
-	// Copy existing material if there is one	
-	IMaterialProvider* TargetMaterial = Cast<IMaterialProvider>(Target);
-	DefaultMaterial = TargetMaterial->GetMaterial(0);
-	if (DefaultMaterial == nullptr)
-	{
-		DefaultMaterial = LoadObject<UMaterial>(nullptr, TEXT("/Engine/EngineMaterials/DefaultMaterial"));  
-	}
-
-	// hide input StaticMeshComponent
-	IPrimitiveComponentBackedTarget* TargetComponent = Cast<IPrimitiveComponentBackedTarget>(Target);
-	TargetComponent->SetOwnerVisibility(false);
+	InputMesh = MakeShared<FDynamicMesh3, ESPMode::ThreadSafe>();
+	*InputMesh = UE::ToolTarget::GetDynamicMeshCopy(Target);
 
-	// Construct the preview object and set the material on it
-	IMeshDescriptionProvider* TargetMeshProvider = Cast<IMeshDescriptionProvider>(Target);
-	Preview = NewObject<UMeshOpPreviewWithBackgroundCompute>(this, "Preview");
+	Preview = NewObject<UMeshOpPreviewWithBackgroundCompute>(this);
 	Preview->Setup(this->TargetWorld, this);
 	Preview->PreviewMesh->SetTangentsMode(EDynamicMeshComponentTangentsMode::AutoCalculated);
-	Preview->PreviewMesh->ReplaceMesh(UE::ToolTarget::GetDynamicMeshCopy(Target));
+	Preview->PreviewMesh->ReplaceMesh(*InputMesh);
+	Preview->ConfigureMaterials(UE::ToolTarget::GetMaterialSet(Target).Materials,
+								ToolSetupUtil::GetDefaultWorkingMaterial(GetToolManager()));
+	Preview->PreviewMesh->SetTransform((FTransform)UE::ToolTarget::GetLocalToWorldTransform(Target));
 
 	Preview->OnMeshUpdated.AddLambda([this](UMeshOpPreviewWithBackgroundCompute* Op)
 	{
 		MaterialSettings->UpdateMaterials();
 	});
 
-	// Initialize the preview mesh with a copy of the source mesh.
-	bool bHasGroups = false;
-	{
-		InputMesh = MakeShared<FDynamicMesh3, ESPMode::ThreadSafe>( * Preview->PreviewMesh->GetMesh() );
-		FMeshDescriptionToDynamicMesh Converter;
-		bHasGroups = FaceGroupUtil::HasMultipleGroups(*InputMesh);
-
-		FComponentMaterialSet MaterialSet;
-		TargetMaterial->GetMaterialSet(MaterialSet);
-		Preview->ConfigureMaterials(MaterialSet.Materials,
-			ToolSetupUtil::GetDefaultWorkingMaterial(GetToolManager())
-		);
-
-		Preview->PreviewMesh->SetTransform(TargetComponent->GetWorldTransform());
-	}
-
-	if (bDoAutomaticGlobalUnwrap == false && bHasGroups == false)
-	{
-		GetToolManager()->DisplayMessage(
-			LOCTEXT("NoGroupsWarning", "This mesh has no PolyGroups!"),
-			EToolMessageLevel::UserWarning);
-		//bDoAutomaticGlobalUnwrap = true;
-	}
+	UE::ToolTarget::HideSourceObject(Target);
 
 	// initialize our properties
 
 	UVChannelProperties = NewObject<UMeshUVChannelProperties>(this);
 	UVChannelProperties->RestoreProperties(this);
-	UVChannelProperties->Initialize(TargetMeshProvider->GetMeshDescription(), false);
+	UVChannelProperties->Initialize(InputMesh.Get(), false);
 	UVChannelProperties->ValidateSelection(true);
 	UVChannelProperties->WatchProperty(UVChannelProperties->UVChannel, [this](const FString& NewValue) 
 	{
 		MaterialSettings->UVChannel = UVChannelProperties->GetSelectedChannelIndex(true);
 	});
-
 	AddToolPropertySource(UVChannelProperties);
 
 	Settings = NewObject<UParameterizeMeshToolProperties>(this);
 	Settings->RestoreProperties(this);
-	Settings->bIsGlobalMode = bDoAutomaticGlobalUnwrap;
 	AddToolPropertySource(Settings);
+	Settings->WatchProperty(Settings->Method, [&](EParameterizeMeshUVMethod) { OnMethodTypeChanged(); });
+
+
+	UVAtlasProperties = NewObject<UParameterizeMeshToolUVAtlasProperties>(this);
+	UVAtlasProperties->RestoreProperties(this);
+	AddToolPropertySource(UVAtlasProperties);
+	SetToolPropertySourceEnabled(UVAtlasProperties, false);
+
+	XAtlasProperties = NewObject<UParameterizeMeshToolXAtlasProperties>(this);
+	XAtlasProperties->RestoreProperties(this);
+	AddToolPropertySource(XAtlasProperties);
+	SetToolPropertySourceEnabled(XAtlasProperties, false);
 
 	MaterialSettings = NewObject<UExistingMeshMaterialProperties>(this);
-	MaterialSettings->RestoreProperties(this);
+	MaterialSettings->MaterialMode = ESetMeshMaterialMode::Checkerboard;
+	MaterialSettings->RestoreProperties(this, TEXT("ModelingUVTools"));
 	AddToolPropertySource(MaterialSettings);
 	// force update
 	MaterialSettings->UpdateMaterials();
 	Preview->OverrideMaterial = MaterialSettings->GetActiveOverrideMaterial();
 
-
-	Preview->SetVisibility(true);
 	Preview->InvalidateResult();    // start compute
 
-
-	if (bDoAutomaticGlobalUnwrap)
-	{
-		SetToolDisplayName(LOCTEXT("ToolNameGlobal", "AutoUV"));
-		GetToolManager()->DisplayMessage(
-			LOCTEXT("OnStartTool_Global", "Automatically partition the selected Mesh into UV islands, flatten, and pack into a single UV chart"),
-			EToolMessageLevel::UserNotification);
-	}
-	else
-	{
-		SetToolDisplayName(LOCTEXT("ToolNameLocal", "UV Unwrap"));
-		GetToolManager()->DisplayMessage(
-			LOCTEXT("OnStartTool_Regions", "Generate UVs for polygroups or existing UV charts of the Mesh using various strategies. Does not calculate layout/packing."),
-			EToolMessageLevel::UserNotification);
-	}
+	SetToolDisplayName(LOCTEXT("ToolNameGlobal", "AutoUV"));
+	GetToolManager()->DisplayMessage(
+		LOCTEXT("OnStartTool_Global", "Automatically partition the selected Mesh into UV islands, flatten, and pack into a single UV chart"),
+		EToolMessageLevel::UserNotification);
 }
 
 void UParameterizeMeshTool::OnPropertyModified(UObject* PropertySet, FProperty* Property)
 {
-	bool bForceMaterialUpdate = false;
-	if (PropertySet == Settings || PropertySet == UVChannelProperties)
+	if (PropertySet != MaterialSettings)
 	{
-		// One of the UV generation properties must have changed.  Dirty the result to force a recompute
 		Preview->InvalidateResult();
-		bForceMaterialUpdate = true;
 	}
 
-	if (PropertySet == MaterialSettings || bForceMaterialUpdate)
-	{
-		MaterialSettings->UpdateMaterials();
-		Preview->OverrideMaterial = MaterialSettings->GetActiveOverrideMaterial();
-	}
+	MaterialSettings->UpdateMaterials();
+	Preview->OverrideMaterial = MaterialSettings->GetActiveOverrideMaterial();
+}
+
+
+void UParameterizeMeshTool::OnMethodTypeChanged()
+{
+	SetToolPropertySourceEnabled(UVAtlasProperties, Settings->Method == EParameterizeMeshUVMethod::UVAtlas);
+	SetToolPropertySourceEnabled(XAtlasProperties, Settings->Method == EParameterizeMeshUVMethod::XAtlas);
+
+	Preview->InvalidateResult();
 }
 
 
@@ -172,37 +120,21 @@ void UParameterizeMeshTool::Shutdown(EToolShutdownType ShutdownType)
 {
 	UVChannelProperties->SaveProperties(this);
 	Settings->SaveProperties(this);
-	MaterialSettings->SaveProperties(this);
+	MaterialSettings->SaveProperties(this, TEXT("ModelingUVTools"));
 
 	FDynamicMeshOpResult Result = Preview->Shutdown();
 	
 	// Restore (unhide) the source meshes
-	Cast<IPrimitiveComponentBackedTarget>(Target)->SetOwnerVisibility(true);
+	UE::ToolTarget::ShowSourceObject(Target);
 
 	if (ShutdownType == EToolShutdownType::Accept)
 	{
-		FDynamicMesh3* DynamicMeshResult = Result.Mesh.Get();
-		check(DynamicMeshResult != nullptr);
-		GetToolManager()->BeginUndoTransaction(LOCTEXT("ParameterizeMesh", "Parameterize Mesh"));
-
-		Cast<IMeshDescriptionCommitter>(Target)->CommitMeshDescription([DynamicMeshResult](const IMeshDescriptionCommitter::FCommitterParams& CommitParams)
+		GetToolManager()->BeginUndoTransaction(LOCTEXT("ParameterizeMesh", "Auto UVs"));
+		FDynamicMesh3* NewDynamicMesh = Result.Mesh.Get();
+		if (ensure(NewDynamicMesh))
 		{
-			FMeshDescription* MeshDescription = CommitParams.MeshDescriptionOut;
-
-			bool bVerticesOnly = false;
-			bool bAttributesOnly = true;
-			if (FDynamicMeshToMeshDescription::HaveMatchingElementCounts(DynamicMeshResult, MeshDescription, bVerticesOnly, bAttributesOnly))
-			{
-				FDynamicMeshToMeshDescription Converter;
-				Converter.UpdateAttributes(DynamicMeshResult, *MeshDescription, false, false, true/*update uvs*/);
-			}
-			else
-			{
-				// must have been duplicate tris in the mesh description; we can't count on 1-to-1 mapping of TriangleIDs.  Just convert 
-				FDynamicMeshToMeshDescription Converter;
-				Converter.Convert(DynamicMeshResult, *MeshDescription);
-			}
-		});
+			UE::ToolTarget::CommitDynamicMeshUVUpdate(Target, NewDynamicMesh);
+		}
 		GetToolManager()->EndUndoTransaction();
 	}
 
@@ -222,42 +154,19 @@ TUniquePtr<FDynamicMeshOperator> UParameterizeMeshTool::MakeNewOperator()
 {
 	FAxisAlignedBox3d MeshBounds = Preview->PreviewMesh->GetMesh()->GetBounds();
 	TUniquePtr<FParameterizeMeshOp> ParameterizeMeshOp = MakeUnique<FParameterizeMeshOp>();
-	ParameterizeMeshOp->Stretch   = Settings->ChartStretch;
-	ParameterizeMeshOp->NumCharts = 0;
+
 	ParameterizeMeshOp->InputMesh = InputMesh;
 	ParameterizeMeshOp->UVLayer = UVChannelProperties->GetSelectedChannelIndex(true);
-	
-	if (bDoAutomaticGlobalUnwrap)
-	{
-		ParameterizeMeshOp->IslandMode = EParamOpIslandMode::Auto;
-		ParameterizeMeshOp->UnwrapType = EParamOpUnwrapType::MinStretch;
-	}
-	else
-	{
-		ParameterizeMeshOp->IslandMode = (EParamOpIslandMode)(int)Settings->IslandMode;
-		ParameterizeMeshOp->UnwrapType = (EParamOpUnwrapType)(int)Settings->UnwrapType;
-	}
-
-	switch (Settings->UVScaleMode)
-	{
-		case EParameterizeMeshToolUVScaleMode::NoScaling:
-			ParameterizeMeshOp->bNormalizeAreas = false;
-			ParameterizeMeshOp->AreaScaling = 1.0;
-			break;
-		case EParameterizeMeshToolUVScaleMode::NormalizeToBounds:
-			ParameterizeMeshOp->bNormalizeAreas = true;
-			ParameterizeMeshOp->AreaScaling = Settings->UVScale / MeshBounds.MaxDim();
-			break;
-		case EParameterizeMeshToolUVScaleMode::NormalizeToWorld:
-			ParameterizeMeshOp->bNormalizeAreas = true;
-			ParameterizeMeshOp->AreaScaling = Settings->UVScale;
-			break;
-	}
-
 	ParameterizeMeshOp->Method = (EParamOpBackend)(int)Settings->Method;
 
-	UE::Geometry::FTransform3d XForm3d(Cast<IPrimitiveComponentBackedTarget>(Target)->GetWorldTransform());
-	ParameterizeMeshOp->SetTransform(XForm3d);
+	// uvatlas options
+	ParameterizeMeshOp->Stretch   = UVAtlasProperties->ChartStretch;
+	ParameterizeMeshOp->NumCharts = UVAtlasProperties->NumCharts;
+	
+	// xatlas options
+	ParameterizeMeshOp->XAtlasMaxIterations = XAtlasProperties->MaxIterations;
+
+	ParameterizeMeshOp->SetTransform(UE::ToolTarget::GetLocalToWorldTransform(Target));
 
 	return ParameterizeMeshOp;
 }
diff --git a/Engine/Plugins/Experimental/MeshModelingToolset/Source/MeshModelingToolsEditorOnly/Public/ParameterizeMeshTool.h b/Engine/Plugins/Experimental/MeshModelingToolset/Source/MeshModelingToolsEditorOnly/Public/ParameterizeMeshTool.h
--- a/Engine/Plugins/Experimental/MeshModelingToolset/Source/MeshModelingToolsEditorOnly/Public/ParameterizeMeshTool.h
+++ b/Engine/Plugins/Experimental/MeshModelingToolset/Source/MeshModelingToolsEditorOnly/Public/ParameterizeMeshTool.h
@@ -3,141 +3,86 @@
 #pragma once
 
 #include "CoreMinimal.h"
-#include "UObject/NoExportTypes.h"
-#include "SingleSelectionTool.h"
-#include "InteractiveToolBuilder.h"
+#include "BaseTools/SingleSelectionMeshEditingTool.h"
 #include "DynamicMesh/DynamicMesh3.h"
 #include "MeshOpPreviewHelpers.h"
 #include "Properties/MeshMaterialProperties.h"
 #include "Properties/MeshUVChannelProperties.h"
-#include "BaseTools/SingleSelectionMeshEditingTool.h"
 
 #include "ParameterizeMeshTool.generated.h"
 
-
 // predeclarations
 class UDynamicMeshComponent;
 class UMaterialInterface;
 class UMaterialInstanceDynamic;
 
 
-
-/** Material modes for MeshInspectorTool */
-UENUM()
-enum class EParameterizeMeshMaterialMode : uint8
-{
-	/** Input material */
-	Default,
-
-	/** Checkerboard material */
-	Checkerboard,
-
-	/** Override material */
-	Override
-};
-
-/**
- *
- */
 UCLASS()
 class MESHMODELINGTOOLSEDITORONLY_API UParameterizeMeshToolBuilder : public USingleSelectionMeshEditingToolBuilder
 {
 	GENERATED_BODY()
-
 public:
-	bool bDoAutomaticGlobalUnwrap = false;
 	virtual USingleSelectionMeshEditingTool* CreateNewTool(const FToolBuilderState& SceneState) const override;
 };
 
 
-
-
 UENUM()
-enum class EUVUnwrapType
-{
-	/**  */
-	MinStretch = 0,
-	/** */
-	ExpMap = 1,
-	/** */
-	Conformal = 2
-};
-
-
-
-UENUM()
-enum class EUVIslandMode
-{
-	/**  */
-	Auto = 0,
-	/** */
-	PolyGroups = 1,
-	/** */
-	ExistingUVs = 2
-};
-
-UENUM()
-enum class EUVMethod
+enum class EParameterizeMeshUVMethod
 {
+	// keep values the same as UE::Geometry::EParamOpBackend!
+	/** Compute Automatic UVs using UVAtlas */
 	UVAtlas = 0,
+	/** Compute Automatic UVs using XAtlas */
 	XAtlas = 1
 };
 
 
-UENUM()
-enum class EParameterizeMeshToolUVScaleMode
-{
-	/** No scaling is applied to UV islands */
-	NoScaling,
-	/** Scale UV islands such that they have constant relative area, relative to object bounds */
-	NormalizeToBounds,
-	/** Scale UV islands such that they have constant relative area, relative to world space */
-	NormalizeToWorld
-};
-
-
-
 
 UCLASS()
 class MESHMODELINGTOOLSEDITORONLY_API UParameterizeMeshToolProperties : public UInteractiveToolPropertySet
 {
 	GENERATED_BODY()
 public:
-
-	UPROPERTY(EditAnywhere, Category = Options, meta = (EditConditionHides, HideEditConditionToggle, EditCondition = "bIsGlobalMode == false"))
-	EUVIslandMode IslandMode = EUVIslandMode::PolyGroups;
-
-	UPROPERTY(EditAnywhere, Category = Options, meta = (EditConditionHides, HideEditConditionToggle, EditCondition = "bIsGlobalMode == false"))
-	EUVUnwrapType UnwrapType = EUVUnwrapType::ExpMap;
-
+	/** Automatic UV Generation technique to use */
 	UPROPERTY(EditAnywhere, Category = Options)
-	EUVMethod Method = EUVMethod::UVAtlas;
-
-	/** Maximum amount of stretch, from none to any.  If zero stretch is specified each triangle will likey be its own chart */
-	UPROPERTY(EditAnywhere, Category = Options, 
-			  meta = (Default = "0.166", UIMin = "0", UIMax = "1", ClampMin = "0", ClampMax = "1", 
-					  EditCondition = "bIsGlobalMode && Method == EUVMethod::UVAtlas || UnwrapType == EUVUnwrapType::MinStretch && Method == EUVMethod::UVAtlas"))
-	float ChartStretch = 0.11f;
+	EParameterizeMeshUVMethod Method = EParameterizeMeshUVMethod::UVAtlas;
+};
 
 
-	/** Scaling applied to UV islands */
-	UPROPERTY(EditAnywhere, Category = Options)
-	EParameterizeMeshToolUVScaleMode UVScaleMode = EParameterizeMeshToolUVScaleMode::NormalizeToBounds;
 
-	/** Scaling factor used for UV island normalization/scaling */
-	UPROPERTY(EditAnywhere, Category = Options, meta = (EditCondition = "UVScaleMode!=EParameterizeMeshToolUVScaleMode::NoScaling", UIMin = "0.001", UIMax = "10", ClampMin = "0.00001", ClampMax = "1000000.0") )
-	float UVScale = 1.0;
+/**
+ * Settings for the UVAtlas Automatic UV Generation Method
+ */
+UCLASS()
+class MESHMODELINGTOOLSEDITORONLY_API UParameterizeMeshToolUVAtlasProperties : public UInteractiveToolPropertySet
+{
+	GENERATED_BODY()
+public:
+	/** Maximum amount of stretch, from none to any.  If zero stretch is specified each triangle will likely be its own chart */
+	UPROPERTY(EditAnywhere, Category = UVAtlas,  meta = (UIMin = "0", UIMax = "1", ClampMin = "0", ClampMax = "1"))
+	float ChartStretch = 0.11f;
 
-	UPROPERTY(meta = (TransientToolProperty))
-	bool bIsGlobalMode = false;
+	/** Hint at number of Charts. 0 (Default) means UVAtlas will decide */
+	UPROPERTY(EditAnywhere, Category = UVAtlas, meta = (UIMin = "0", UIMax = "100", ClampMin = "0", ClampMax = "10000"))
+	int NumCharts = 0;
 };
 
 
+UCLASS()
+class MESHMODELINGTOOLSEDITORONLY_API UParameterizeMeshToolXAtlasProperties : public UInteractiveToolPropertySet
+{
+	GENERATED_BODY()
+public:
+	/** Number of solve iterations. Higher values generally result in better charts. */
+	UPROPERTY(EditAnywhere, Category = XAtlas, meta = (UIMin = "1", UIMax = "10", ClampMin = "1", ClampMax = "1000"))
+	int MaxIterations = 1;
+};
 
 
 
 /**
- * Simple Mesh Simplifying Tool
+ * UParameterizeMeshTool automatically decomposes the input mesh into charts, solves for UVs,
+ * and then packs the resulting charts
  */
 UCLASS()
 class MESHMODELINGTOOLSEDITORONLY_API UParameterizeMeshTool : public USingleSelectionMeshEditingTool, public UE::Geometry::IDynamicMeshOperatorFactory
@@ -145,10 +90,6 @@ class MESHMODELINGTOOLSEDITORONLY_API UParameterizeMeshTool : public USingleSele
 	GENERATED_BODY()
 
 public:
-	UParameterizeMeshTool();
-
-	virtual void SetUseAutoGlobalParameterizationMode(bool bEnable);
-
 	virtual void Setup() override;
 	virtual void Shutdown(EToolShutdownType ShutdownType) override;
 
@@ -158,7 +99,6 @@ public:
 	virtual bool HasAccept() const override { return true; }
 	virtual bool CanAccept() const override;
 
-
 	virtual void OnPropertyModified(UObject* PropertySet, FProperty* Property) override;
 
 	// IDynamicMeshOperatorFactory API
@@ -172,25 +112,22 @@ protected:
 	UParameterizeMeshToolProperties* Settings = nullptr;
 
 	UPROPERTY()
-	UExistingMeshMaterialProperties* MaterialSettings = nullptr;
+	UParameterizeMeshToolUVAtlasProperties* UVAtlasProperties = nullptr;
 
-protected:
 	UPROPERTY()
-	UMaterialInterface* DefaultMaterial = nullptr;
+	UParameterizeMeshToolXAtlasProperties* XAtlasProperties = nullptr;
 
-	UPROPERTY()
-	UMaterialInstanceDynamic* DisplayMaterial = nullptr;
 
 	UPROPERTY()
-	UMaterialInstanceDynamic* CheckerMaterial = nullptr;
-	
-	UPROPERTY()
-	UMeshOpPreviewWithBackgroundCompute* Preview = nullptr;
+	UExistingMeshMaterialProperties* MaterialSettings = nullptr;
 
+protected:
 	UPROPERTY()
-	bool bDoAutomaticGlobalUnwrap = false;
+	UMeshOpPreviewWithBackgroundCompute* Preview = nullptr;
 
 protected:
 	TSharedPtr<UE::Geometry::FDynamicMesh3, ESPMode::ThreadSafe> InputMesh;
 
+	void OnMethodTypeChanged();
+
 };
diff --git a/Engine/Plugins/Experimental/MeshModelingToolset/Source/ModelingOperators/Private/ParameterizationOps/RecomputeUVsOp.cpp b/Engine/Plugins/Experimental/MeshModelingToolset/Source/ModelingOperators/Private/ParameterizationOps/RecomputeUVsOp.cpp
new file mode 100644
--- /dev/null
+++ b/Engine/Plugins/Experimental/MeshModelingToolset/Source/ModelingOperators/Private/ParameterizationOps/RecomputeUVsOp.cpp
@@ -0,0 +1,158 @@
+// Copyright Epic Games, Inc. All Rights Reserved.
+
+#include "ParameterizationOps/RecomputeUVsOp.h"
+
+#include "DynamicMesh/DynamicMeshAttributeSet.h"
+#include "Selections/MeshConnectedComponents.h"
+#include "Parameterization/MeshLocalParam.h"
+#include "Parameterization/DynamicMeshUVEditor.h"
+
+#include "ExplicitUseGeometryMathTypes.h"		// using UE::Geometry::(math types)
+using namespace UE::Geometry;
+
+
+void FRecomputeUVsOp::NormalizeUVAreas(const FDynamicMesh3& Mesh, FDynamicMeshUVOverlay* Overlay, float GlobalScale)
+{
+	FMeshConnectedComponents UVComponents(&Mesh);
+	UVComponents.FindConnectedTriangles([&](int32 Triangle0, int32 Triangle1) {
+		return Overlay->AreTrianglesConnected(Triangle0, Triangle1);
+	});
+
+	// TODO ParallelFor
+	for (FMeshConnectedComponents::FComponent& Component : UVComponents)
+	{
+		const TArray<int>& Triangles = Component.Indices;
+		TSet<int> UVElements;
+		UVElements.Reserve(Triangles.Num() * 3);
+		double AreaUV = 0;
+		double Area3D = 0;
+		FVector3d Triangle3D[3];
+		FVector2d TriangleUV[3];
+		FAxisAlignedBox2d BoundsUV = FAxisAlignedBox2d::Empty();
+
+		for (int tid : Triangles)
+		{
+			FIndex3i TriElements = Overlay->GetTriangle(tid);
+			if (!TriElements.Contains(FDynamicMesh3::InvalidID))
+			{
+				for (int j = 0; j < 3; ++j)
+				{
+					TriangleUV[j] = FVector2d(Overlay->GetElement(TriElements[j]));
+					BoundsUV.Contain(TriangleUV[j]);
+					Triangle3D[j] = Mesh.GetVertex(Overlay->GetParentVertex(TriElements[j]));
+					UVElements.Add(TriElements[j]);
+				}
+				AreaUV += VectorUtil::Area(TriangleUV[0], TriangleUV[1], TriangleUV[2]);
+				Area3D += VectorUtil::Area(Triangle3D[0], Triangle3D[1], Triangle3D[2]);
+			}
+		}
+		
+		double LinearScale = (AreaUV > 0.00001) ? ( FMathd::Sqrt(Area3D) / FMathd::Sqrt(AreaUV)) : 1.0;
+		LinearScale = LinearScale * GlobalScale;
+		FVector2d ComponentOrigin = BoundsUV.Center();
+
+		for (int elemid : UVElements)
+		{
+			FVector2d UV = FVector2d(Overlay->GetElement(elemid));
+			UV = (UV - ComponentOrigin) * LinearScale + ComponentOrigin;
+			Overlay->SetElement(elemid, FVector2f(UV));
+		}
+	}
+}
+
+
+void FRecomputeUVsOp::CalculateResult(FProgressCancel* Progress)
+{
+	if (!InputMesh.IsValid())
+	{
+		return;
+	}
+
+	ResultMesh = MakeUnique<FDynamicMesh3>(*InputMesh);
+
+	FDynamicMesh3* BaseMesh = ResultMesh.Get();
+	FDynamicMeshUVEditor UVEditor(BaseMesh, UVLayer, true);
+	FDynamicMeshUVOverlay* UseOverlay = UVEditor.GetOverlay();
+	if (ensure(UseOverlay != nullptr) == false)
+	{
+		return;
+	}
+
+	if (Progress && Progress->Cancelled())
+	{
+		return;
+	}
+
+	// find group-connected-components
+	FMeshConnectedComponents ConnectedComponents(ResultMesh.Get());
+	if (IslandMode == ERecomputeUVsIslandMode::PolyGroups)
+	{
+		if (InputGroups != nullptr)
+		{
+			ConnectedComponents.FindConnectedTriangles([this](int32 CurTri, int32 NbrTri) {
+				return InputGroups->GetTriangleGroup(CurTri) == InputGroups->GetTriangleGroup(NbrTri);
+			});
+		}
+		else
+		{
+			ConnectedComponents.FindConnectedTriangles([this](int32 CurTri, int32 NbrTri) {
+				return ResultMesh->GetTriangleGroup(CurTri) == ResultMesh->GetTriangleGroup(NbrTri);
+			});
+		}
+	}
+	else
+	{
+		ConnectedComponents.FindConnectedTriangles([&](int32 Triangle0, int32 Triangle1) {
+			return UseOverlay->AreTrianglesConnected(Triangle0, Triangle1);
+		});
+	}
+
+	if (Progress && Progress->Cancelled())
+	{
+		return;
+	}
+
+	// TODO: the solves here could be done in parallel if we pre-allocated the island element IDs
+
+	int32 NumComponents = ConnectedComponents.Num();
+	int32 SuccessCount = 0;
+	for (int32 k = 0; k < NumComponents; ++k)
+	{
+		const TArray<int32>& ComponentTris = ConnectedComponents[k].Indices;
+
+		bool bComputedUVs = false;
+		switch (UnwrapType)
+		{
+		case ERecomputeUVsUnwrapType::ExpMap:
+			bComputedUVs = UVEditor.SetTriangleUVsFromExpMap(ComponentTris);
+			break;
+
+		case ERecomputeUVsUnwrapType::ConformalFreeBoundary:
+			bComputedUVs = UVEditor.SetTriangleUVsFromFreeBoundaryConformal(ComponentTris);
+			break;
+		}
+
+		if (bComputedUVs)
+		{
+			SuccessCount++;
+		}
+
+		if (Progress && Progress->Cancelled())
+		{
+			return;
+		}
+	}
+
+
+	if (bNormalizeAreas)
+	{
+		// todo should be a DynamicUVEditor function?
+		NormalizeUVAreas(*ResultMesh, UseOverlay, AreaScaling);
+	}
+
+	if (bPackUVs)
+	{
+		UVEditor.QuickPack(PackingTextureResolution, PackingGutterWidth);
+	}
+
+}
\ No newline at end of file
diff --git a/Engine/Plugins/Experimental/MeshModelingToolset/Source/ModelingOperators/Public/ParameterizationOps/RecomputeUVsOp.h b/Engine/Plugins/Experimental/MeshModelingToolset/Source/ModelingOperators/Public/ParameterizationOps/RecomputeUVsOp.h
new file mode 100644
--- /dev/null
+++ b/Engine/Plugins/Experimental/MeshModelingToolset/Source/ModelingOperators/Public/ParameterizationOps/RecomputeUVsOp.h
@@ -0,0 +1,85 @@
+// Copyright Epic Games, Inc. All Rights Reserved.
+
+#pragma once
+
+#include "CoreMinimal.h"
+#include "Util/ProgressCancel.h"
+#include "ModelingOperators.h"
+#include "DynamicMesh/DynamicMeshAttributeSet.h"
+#include "Polygroups/PolygroupSet.h"
+
+
+namespace UE
+{
+namespace Geometry
+{
+
+
+enum class ERecomputeUVsUnwrapType
+{
+	ExpMap = 0,
+	ConformalFreeBoundary = 1
+};
+
+
+
+enum class ERecomputeUVsIslandMode
+{
+	PolyGroups = 0,
+	UVIslands = 1
+};
+
+
+class MODELINGOPERATORS_API FRecomputeUVsOp : public FDynamicMeshOperator
+{
+public:
+	virtual ~FRecomputeUVsOp() {}
+
+	//
+	// Inputs
+	// 
+
+	// source mesh
+	TSharedPtr<FDynamicMesh3, ESPMode::ThreadSafe> InputMesh;
+	
+	// source groups (optional)
+	TSharedPtr<UE::Geometry::FPolygroupSet, ESPMode::ThreadSafe> InputGroups;
+
+	// area scaling
+	bool bNormalizeAreas = true;
+	float AreaScaling = 1.0;
+
+	// UV layer
+	int32 UVLayer = 0;
+
+	// Atlas Packing parameters
+	bool bPackUVs = true;
+	int32 PackingTextureResolution = 512;
+	float PackingGutterWidth = 1.0f;
+
+	ERecomputeUVsUnwrapType UnwrapType = ERecomputeUVsUnwrapType::ExpMap;
+	ERecomputeUVsIslandMode IslandMode = ERecomputeUVsIslandMode::PolyGroups;
+
+	// set ability on protected transform.
+	void SetTransform(const FTransform3d& XForm)
+	{
+		ResultTransform = XForm;
+	}
+
+	//
+	// FDynamicMeshOperator implementation
+	// 
+
+	virtual void CalculateResult(FProgressCancel* Progress) override;
+
+
+protected:
+
+	bool ComputeUVs_ExpMap(FDynamicMesh3& InOutMesh, TFunction<bool(float)>& Interrupter, float GlobalScale = 1.0f);
+	bool ComputeUVs_ConformalFreeBoundary(FDynamicMesh3& InOutMesh, TFunction<bool(float)>& Interrupter, float GlobalScale = 1.0f);
+
+	void NormalizeUVAreas(const FDynamicMesh3& Mesh, FDynamicMeshUVOverlay* Overlay, float GlobalScale = 1.0f);
+};
+
+} // end namespace UE::Geometry
+} // end namespace UE
diff --git a/Engine/Plugins/Experimental/MeshModelingToolset/Source/ModelingOperatorsEditorOnly/Private/ParameterizationOps/ParameterizeMeshOp.cpp b/Engine/Plugins/Experimental/MeshModelingToolset/Source/ModelingOperatorsEditorOnly/Private/ParameterizationOps/ParameterizeMeshOp.cpp
--- a/Engine/Plugins/Experimental/MeshModelingToolset/Source/ModelingOperatorsEditorOnly/Private/ParameterizationOps/ParameterizeMeshOp.cpp
+++ b/Engine/Plugins/Experimental/MeshModelingToolset/Source/ModelingOperatorsEditorOnly/Private/ParameterizationOps/ParameterizeMeshOp.cpp
@@ -142,9 +142,9 @@ FParameterizeMeshOp::FLinearMesh::FLinearMesh(const FDynamicMesh3& Mesh, const b
 	}
 }
 
-bool FParameterizeMeshOp::ComputeUVs(FDynamicMesh3& Mesh,  TFunction<bool(float)>& Interrupter, const bool bUsePolygroups, float GlobalScale)
+bool FParameterizeMeshOp::ComputeUVs(FDynamicMesh3& Mesh,  TFunction<bool(float)>& Interrupter)
 {
-	// the UVAltas code is unhappy if you feed it a single degenerate triangle
+	// the UVAtlas code is unhappy if you feed it a single degenerate triangle
 	bool bNonDegenerate = true;
 	if (Mesh.TriangleCount() == 1)
 	{
@@ -172,7 +172,7 @@ bool FParameterizeMeshOp::ComputeUVs(FDynamicMesh3& Mesh,  TFunction<bool(float)
 	}
 
 	// Convert to a dense form.
-	FLinearMesh LinearMesh(FlippedMesh, bUsePolygroups);
+	FLinearMesh LinearMesh(FlippedMesh, false);
 
 	// Data to be populated by the UV generation tool
 	TArray<FVector2D> UVVertexBuffer;
@@ -189,6 +189,7 @@ bool FParameterizeMeshOp::ComputeUVs(FDynamicMesh3& Mesh,  TFunction<bool(float)
 	if (bUseXAtlas)
 	{
 		XAtlasWrapper::XAtlasChartOptions ChartOptions;
+		ChartOptions.MaxIterations = XAtlasMaxIterations;
 		XAtlasWrapper::XAtlasPackOptions PackOptions;
 		bSuccess = XAtlasWrapper::ComputeUVs(LinearMesh.IndexBuffer, 
 											 LinearMesh.VertexBuffer, 
@@ -285,10 +286,6 @@ bool FParameterizeMeshOp::ComputeUVs(FDynamicMesh3& Mesh,  TFunction<bool(float)
 			UVOverlay->SetTriangle(TriID, ElTri);
 		}
 
-		//if (bNormalizeAreas)
-		//{
-		//	NormalizeUVAreas(&Mesh, UVOverlay, GlobalScale);
-		//}
 	}
 
 
@@ -302,99 +299,6 @@ bool FParameterizeMeshOp::ComputeUVs(FDynamicMesh3& Mesh,  TFunction<bool(float)
 
 
 
-bool FParameterizeMeshOp::ComputeUVs_ExpMap(FDynamicMesh3& Mesh, TFunction<bool(float)>& Interrupter, float GlobalScale)
-{
-	FDynamicMeshUVEditor UVEditor(&Mesh, UVLayer, true);
-
-	// Delete existing UVs, if there are any.
-	UVEditor.ResetUVs();
-	
-	TArray<int32> AllTriangles;
-	for (int32 tid : Mesh.TriangleIndicesItr())
-	{
-		AllTriangles.Add(tid);
-	}
-
-	UVEditor.SetTriangleUVsFromExpMap(AllTriangles);
-
-	return true;
-}
-
-
-
-
-bool FParameterizeMeshOp::ComputeUVs_ConformalFreeBoundary(FDynamicMesh3& InOutMesh, TFunction<bool(float)>& Interrupter, float GlobalScale)
-{
-	FDynamicMeshUVEditor UVEditor(&InOutMesh, UVLayer, true);
-
-	// Delete existing UVs, if there are any.
-	UVEditor.ResetUVs();
-	
-	TArray<int32> AllTriangles;
-	for (int32 tid : InOutMesh.TriangleIndicesItr())
-	{
-		AllTriangles.Add(tid);
-	}
-
-	bool bUseExistingUVTopology = (IslandMode == EParamOpIslandMode::UVIslands);
-	UVEditor.SetTriangleUVsFromFreeBoundaryConformal(AllTriangles, bUseExistingUVTopology);
-
-	return true;
-}
-
-
-
-
-void FParameterizeMeshOp::NormalizeUVAreas(const FDynamicMesh3& Mesh, FDynamicMeshUVOverlay* Overlay, float GlobalScale)
-{
-	FMeshConnectedComponents UVComponents(&Mesh);
-	UVComponents.FindConnectedTriangles([&](int32 Triangle0, int32 Triangle1) {
-		return Overlay->AreTrianglesConnected(Triangle0, Triangle1);
-	});
-
-	// TODO ParallelFor
-	for (FMeshConnectedComponents::FComponent& Component : UVComponents)
-	{
-		const TArray<int>& Triangles = Component.Indices;
-		TSet<int> UVElements;
-		UVElements.Reserve(Triangles.Num() * 3);
-		double AreaUV = 0;
-		double Area3D = 0;
-		FVector3d Triangle3D[3];
-		FVector2d TriangleUV[3];
-		FAxisAlignedBox2d BoundsUV = FAxisAlignedBox2d::Empty();
-
-		for (int tid : Triangles)
-		{
-			FIndex3i TriElements = Overlay->GetTriangle(tid);
-			if (!TriElements.Contains(FDynamicMesh3::InvalidID))
-			{
-				for (int j = 0; j < 3; ++j)
-				{
-					TriangleUV[j] = FVector2d(Overlay->GetElement(TriElements[j]));
-					BoundsUV.Contain(TriangleUV[j]);
-					Triangle3D[j] = Mesh.GetVertex(Overlay->GetParentVertex(TriElements[j]));
-					UVElements.Add(TriElements[j]);
-				}
-				AreaUV += VectorUtil::Area(TriangleUV[0], TriangleUV[1], TriangleUV[2]);
-				Area3D += VectorUtil::Area(Triangle3D[0], Triangle3D[1], Triangle3D[2]);
-			}
-		}
-		
-		double LinearScale = (AreaUV > 0.00001) ? ( FMathd::Sqrt(Area3D) / FMathd::Sqrt(AreaUV)) : 1.0;
-		LinearScale = LinearScale * GlobalScale;
-		FVector2d ComponentOrigin = BoundsUV.Center();
-
-		for (int elemid : UVElements)
-		{
-			FVector2d UV = FVector2d(Overlay->GetElement(elemid));
-			UV = (UV - ComponentOrigin) * LinearScale + ComponentOrigin;
-			Overlay->SetElement(elemid, FVector2f(UV));
-		}
-	}
-}
-
-
 void FParameterizeMeshOp::CalculateResult(FProgressCancel* Progress)
 {
 	if (!InputMesh.IsValid())
@@ -409,139 +313,10 @@ void FParameterizeMeshOp::CalculateResult(FProgressCancel* Progress)
 		return;
 	}
 
-	FDynamicMesh3* BaseMesh = ResultMesh.Get();
-	FDynamicMeshUVOverlay* UVOverlay = (BaseMesh->HasAttributes()) ?
-		BaseMesh->Attributes()->GetUVLayer(UVLayer) : nullptr;
-
-	// We can either split the mesh into multiple meshes - one for each poly group and create UVs for each with calls to UVAtlas (potentially in parallel)
-	// or we can update the adjacency and use a single call to UVAtlas.
-	//
-	// We might disable the Split because running our LayoutUV tool on the results it could mangle some triangles, while the single call
-	// to UVAtlas automatically does some layout.  This mangling happens some UVtriangles have opposite winding order from the rest of the group and they
-	// get put in their own chart..  I noticed that only when I remesh with "prevent normal flipping" I get these inverted uv triangles..
-	const bool bSplitIntoMultiple = true;
-
 	// The UV atlas callback uses a float progress-based interrupter. 
 	TFunction<bool(float)> Iterrupter = [Progress](float)->bool {return !(Progress && Progress->Cancelled()); };
 
-	if (UnwrapType == EParamOpUnwrapType::MinStretch && IslandMode == EParamOpIslandMode::Auto)
-	{
-		// Single call to UVAtlas - perhaps respecting the poly groups.
-		bool bSuccess = ComputeUVs(*ResultMesh, Iterrupter, false, AreaScaling);
-		return;
-	}
-
-
-	// Count the number of components for which we successufully generate UVs
-	int32 SuccessCount = 0;
-
-	// find group-connected-components
-	FMeshConnectedComponents ConnectedComponents(ResultMesh.Get());
-	if (IslandMode == EParamOpIslandMode::PolyGroups)
-	{
-		ConnectedComponents.FindConnectedTriangles([this](int32 CurTri, int32 NbrTri) {
-			return ResultMesh->GetTriangleGroup(CurTri) == ResultMesh->GetTriangleGroup(NbrTri);
-		});
-	}
-	else
-	{
-		check(UVOverlay);
-		ConnectedComponents.FindConnectedTriangles([&](int32 Triangle0, int32 Triangle1) {
-			return UVOverlay->AreTrianglesConnected(Triangle0, Triangle1);
-		});
-	}
-
-	// Create an array of meshes - one for each poly group
-	// NB TIndirectArray automatically deletes the meshes when it goes out of scope.
-	TIndirectArray<FDynamicSubmesh3> SubmeshArray;
-	for ( FMeshConnectedComponents::FComponent& Component : ConnectedComponents)
-	{
-		FDynamicSubmesh3* ComponentSubmesh = new FDynamicSubmesh3(BaseMesh, Component.Indices, (int)EMeshComponents::None, true);
-		SubmeshArray.Add(ComponentSubmesh);
-	}
-
-	// Compute UVs for each component mesh.  NB: this could be parallelized over each
-	for (FDynamicSubmesh3& ComponentSubmesh : SubmeshArray)
-	{
-		FDynamicMesh3& ComponentMesh = ComponentSubmesh.GetSubmesh();
-		if (ComponentMesh.TriangleCount() > 0)
-		{
-			bool bComputedUVs = false;
-			switch (UnwrapType)
-			{
-			case EParamOpUnwrapType::ExpMap:
-				bComputedUVs = ComputeUVs_ExpMap(ComponentMesh, Iterrupter, AreaScaling);
-				break;
-
-			case EParamOpUnwrapType::ConformalFreeBoundary:
-				bComputedUVs = ComputeUVs_ConformalFreeBoundary(ComponentMesh, Iterrupter, AreaScaling);
-				break;
-
-			case EParamOpUnwrapType::MinStretch:
-				bComputedUVs = ComputeUVs(ComponentMesh, Iterrupter, false, AreaScaling);
-				break;
-			}
-
-			if (bComputedUVs) SuccessCount++;
-		}
-
-		if (Progress && Progress->Cancelled())
-		{
-			return;
-		}
-	}
-
-	// transfer the UVs back to the ResultMesh
-	if (SuccessCount > 0)
-	{
-		if (ResultMesh->HasAttributes() == false)
-		{
-			ResultMesh->EnableAttributes();
-		}
-
-		UVOverlay->ClearElements(); // delete existing UVs
-
-		for (FDynamicSubmesh3& ComponentSubmesh : SubmeshArray)
-		{
-			FDynamicMesh3& ComponentMesh = ComponentSubmesh.GetSubmesh();
-			if (ComponentMesh.HasAttributes() == false)
-			{
-				continue;
-			}
-
-			// copy the elements
-			const FDynamicMeshUVOverlay* ComponentUVOverlay = ComponentMesh.Attributes()->GetUVLayer(UVLayer);
-			TArray<int32> ElemIDMap;
-			ElemIDMap.SetNum(ComponentUVOverlay->MaxElementID());
-			for (int32 ElemID : ComponentUVOverlay->ElementIndicesItr())
-			{
-				int ParentVertexID = ComponentUVOverlay->GetParentVertex(ElemID);
-				int32 BaseElemID = UVOverlay->AppendElement(ComponentUVOverlay->GetElement(ElemID));
-				ElemIDMap[ElemID] = BaseElemID;
-			}
-
-			// set the triangle values to the new elements
-			for (int32 TriangleID : ComponentMesh.TriangleIndicesItr())
-			{
-				int32 BaseTriangleID = ComponentSubmesh.MapTriangleToBaseMesh(TriangleID);
-				FIndex3i SubUVTriangle = ComponentUVOverlay->GetTriangle(TriangleID);
-				if (SubUVTriangle.A >= 0 && SubUVTriangle.B >= 0 && SubUVTriangle.C >= 0)
-				{
-					FIndex3i BaseUVTriangle(ElemIDMap[SubUVTriangle.A], ElemIDMap[SubUVTriangle.B], ElemIDMap[SubUVTriangle.C]);
-					UVOverlay->SetTriangle(BaseTriangleID, BaseUVTriangle);
-				}
-			}
-
-			if (Progress && Progress->Cancelled())
-			{
-				return;
-			}
-		}
-
-		if (bNormalizeAreas)
-		{
-			NormalizeUVAreas(*ResultMesh, UVOverlay, AreaScaling);
-		}
-	}
+	// Single call to UVAtlas - perhaps respecting the poly groups.
+	ComputeUVs(*ResultMesh, Iterrupter);
 
 }
\ No newline at end of file
diff --git a/Engine/Plugins/Experimental/MeshModelingToolset/Source/ModelingOperatorsEditorOnly/Public/ParameterizationOps/ParameterizeMeshOp.h b/Engine/Plugins/Experimental/MeshModelingToolset/Source/ModelingOperatorsEditorOnly/Public/ParameterizationOps/ParameterizeMeshOp.h
--- a/Engine/Plugins/Experimental/MeshModelingToolset/Source/ModelingOperatorsEditorOnly/Public/ParameterizationOps/ParameterizeMeshOp.h
+++ b/Engine/Plugins/Experimental/MeshModelingToolset/Source/ModelingOperatorsEditorOnly/Public/ParameterizationOps/ParameterizeMeshOp.h
@@ -13,26 +13,9 @@ namespace UE
 namespace Geometry
 {
 
-
-enum class EParamOpUnwrapType
-{
-	MinStretch = 0,
-	ExpMap = 1,
-	ConformalFreeBoundary = 2
-};
-
-
-
-enum class EParamOpIslandMode
-{
-	Auto = 0,
-	PolyGroups = 1,
-	UVIslands = 2
-};
-
 enum class EParamOpBackend
 {
-	UVAltas = 0,
+	UVAtlas = 0,
 	XAtlas = 1
 };
 
@@ -48,13 +31,12 @@ public:
 	// source mesh
 	TSharedPtr<FDynamicMesh3, ESPMode::ThreadSafe> InputMesh;
 		
-	// UV generation parameters
-	float Stretch;
-	int32 NumCharts;
+	// UVAtlas generation parameters
+	float Stretch = 0.11f;
+	int32 NumCharts = 0;
 
-	// area scaling
-	bool bNormalizeAreas = true;
-	float AreaScaling = 1.0;
+	// XAtlas generation parameters
+	int32 XAtlasMaxIterations = 1;
 
 	// UV layer
 	int32 UVLayer = 0;
@@ -64,12 +46,10 @@ public:
 	int32 Width = 512;
 	float Gutter = 2.5;
 
-	EParamOpUnwrapType UnwrapType = EParamOpUnwrapType::ExpMap;
-	EParamOpIslandMode IslandMode = EParamOpIslandMode::PolyGroups;
-	EParamOpBackend Method = EParamOpBackend::UVAltas;
+	EParamOpBackend Method = EParamOpBackend::UVAtlas;
 
 	// set ability on protected transform.
-	void SetTransform(FTransform3d& XForm)
+	void SetTransform(const FTransform3d& XForm)
 	{
 		ResultTransform = XForm;
 	}
@@ -101,11 +81,7 @@ protected:
 		
 	};
 
-	bool ComputeUVs(FDynamicMesh3& InOutMesh,  TFunction<bool(float)>& Interrupter, const bool bUsePolygroups = false, float GlobalScale = 1.0f);
-	bool ComputeUVs_ExpMap(FDynamicMesh3& InOutMesh, TFunction<bool(float)>& Interrupter, float GlobalScale = 1.0f);
-	bool ComputeUVs_ConformalFreeBoundary(FDynamicMesh3& InOutMesh, TFunction<bool(float)>& Interrupter, float GlobalScale = 1.0f);
-
-	void NormalizeUVAreas(const FDynamicMesh3& Mesh, FDynamicMeshUVOverlay* Overlay, float GlobalScale = 1.0f);
+	bool ComputeUVs(FDynamicMesh3& InOutMesh,  TFunction<bool(float)>& Interrupter);
 };
 
 } // end namespace UE::Geometry
diff --git a/Engine/Plugins/Experimental/ModelingToolsEditorMode/Source/ModelingToolsEditorMode/Private/ModelingToolsEditorMode.cpp b/Engine/Plugins/Experimental/ModelingToolsEditorMode/Source/ModelingToolsEditorMode/Private/ModelingToolsEditorMode.cpp
--- a/Engine/Plugins/Experimental/ModelingToolsEditorMode/Source/ModelingToolsEditorMode/Private/ModelingToolsEditorMode.cpp
+++ b/Engine/Plugins/Experimental/ModelingToolsEditorMode/Source/ModelingToolsEditorMode/Private/ModelingToolsEditorMode.cpp
@@ -78,6 +78,7 @@
 #include "BakeMeshAttributeMapsTool.h"
 #include "MeshAttributePaintTool.h"
 #include "ParameterizeMeshTool.h"
+#include "RecomputeUVsTool.h"
 #include "MeshTangentsTool.h"
 #include "ProjectToTargetTool.h"
 #include "LatticeDeformerTool.h"
@@ -529,12 +530,10 @@ void UModelingToolsEditorMode::Enter()
 	RegisterTool(ToolManagerCommands.BeginPolygonCutTool, TEXT("BeginPolygonCutTool"), PolygonCutToolBuilder);
 
 	auto GlobalUVGenerateToolBuilder = NewObject<UParameterizeMeshToolBuilder>();
-	GlobalUVGenerateToolBuilder->bDoAutomaticGlobalUnwrap = true;
 	RegisterTool(ToolManagerCommands.BeginGlobalUVGenerateTool, TEXT("BeginGlobalUVGenerateTool"), GlobalUVGenerateToolBuilder);
 
-	auto GroupUVGenerateToolBuilder = NewObject<UParameterizeMeshToolBuilder>();
-	GroupUVGenerateToolBuilder->bDoAutomaticGlobalUnwrap = false;
-	RegisterTool(ToolManagerCommands.BeginGroupUVGenerateTool, TEXT("BeginGroupUVGenerateTool"), GroupUVGenerateToolBuilder);
+	auto RecomputeUVsToolBuilder = NewObject<URecomputeUVsToolBuilder>();
+	RegisterTool(ToolManagerCommands.BeginGroupUVGenerateTool, TEXT("BeginGroupUVGenerateTool"), RecomputeUVsToolBuilder);
 
 	RegisterTool(ToolManagerCommands.BeginUVSeamEditTool, TEXT("BeginUVSeamEditTool"), NewObject< USeamSculptToolBuilder>());
 
