ModelingComponents: add support for a UBodySetup and Complex Collision to DynamicMeshComponent
#rnx

diff --git a/Engine/Plugins/Experimental/MeshModelingToolset/Source/ModelingComponents/Private/Components/DynamicMeshComponent.cpp b/Engine/Plugins/Experimental/MeshModelingToolset/Source/ModelingComponents/Private/Components/DynamicMeshComponent.cpp
--- a/Engine/Plugins/Experimental/MeshModelingToolset/Source/ModelingComponents/Private/Components/DynamicMeshComponent.cpp
+++ b/Engine/Plugins/Experimental/MeshModelingToolset/Source/ModelingComponents/Private/Components/DynamicMeshComponent.cpp
@@ -1,22 +1,12 @@
 // Copyright Epic Games, Inc. All Rights Reserved. 
 
 #include "Components/DynamicMeshComponent.h"
-#include "RenderingThread.h"
-#include "RenderResource.h"
-#include "PrimitiveViewRelevance.h"
 #include "PrimitiveSceneProxy.h"
-#include "VertexFactory.h"
 #include "MaterialShared.h"
 #include "Engine/CollisionProfile.h"
 #include "Materials/Material.h"
-#include "LocalVertexFactory.h"
-#include "SceneManagement.h"
-#include "DynamicMeshBuilder.h"
-#include "EngineGlobals.h"
-#include "Engine/Engine.h"
-#include "StaticMeshResources.h"
-#include "StaticMeshAttributes.h"
 #include "Async/Async.h"
+#include "Engine/CollisionProfile.h"
 
 #include "DynamicMesh/DynamicMeshAttributeSet.h"
 #include "DynamicMesh/MeshNormals.h"
@@ -24,7 +14,6 @@
 
 #include "Changes/MeshVertexChange.h"
 #include "Changes/MeshChange.h"
-#include "DynamicMesh/DynamicMeshChangeTracker.h"
 #include "DynamicMesh/MeshTransforms.h"
 
 // default proxy for this component
@@ -70,9 +59,34 @@ void UDynamicMeshComponent::PostLoad()
 	MeshObjectChangedHandle = MeshObject->OnMeshChanged().AddUObject(this, &UDynamicMeshComponent::OnMeshObjectChanged);
 
 	ResetProxy();
+
+	// make sure BodySetup is created
+	GetBodySetup();
 }
 
 
+#if WITH_EDITOR
+void UDynamicMeshComponent::PostEditChangeProperty(FPropertyChangedEvent& PropertyChangedEvent)
+{
+	Super::PostEditChangeProperty(PropertyChangedEvent);
+
+	const FName PropName = PropertyChangedEvent.GetPropertyName();
+	if ( (PropName == GET_MEMBER_NAME_CHECKED(UDynamicMeshComponent, bEnableComplexCollision)) ||
+		(PropName == GET_MEMBER_NAME_CHECKED(UDynamicMeshComponent, CollisionType)) ||
+		(PropName == GET_MEMBER_NAME_CHECKED(UDynamicMeshComponent, bDeferCollisionUpdates))  )
+	{
+		if (bDeferCollisionUpdates)
+		{
+			InvalidatePhysicsData();
+		}
+		else
+		{
+			RebuildPhysicsData();
+		}
+	}
+}
+#endif
+
 
 void UDynamicMeshComponent::SetMesh(UE::Geometry::FDynamicMesh3&& MoveMesh)
 {
@@ -838,6 +852,7 @@ void UDynamicMeshComponent::DisableSecondaryTriangleBuffers()
 
 void UDynamicMeshComponent::SetExternalDecomposition(TUniquePtr<FMeshRenderDecomposition> DecompositionIn)
 {
+	check(DecompositionIn->Num() > 0);
 	Decomposition = MoveTemp(DecompositionIn);
 	NotifyMeshUpdated();
 }
@@ -921,6 +936,17 @@ void UDynamicMeshComponent::OnMeshObjectChanged(UDynamicMesh* ChangedMeshObject,
 		NotifyMeshUpdated();
 		OnMeshChanged.Broadcast();
 	}
+
+	// Rebuild body setup. Should this be deferred until proxy creation? Sometimes multiple changes are emitted...
+	// todo: can possibly skip this in some change situations, eg if only changing attributes
+	if (bDeferCollisionUpdates)
+	{
+		InvalidatePhysicsData();
+	}
+	else
+	{
+		RebuildPhysicsData();
+	}
 }
 
 
@@ -941,6 +967,16 @@ void UDynamicMeshComponent::SetDynamicMesh(UDynamicMesh* NewMesh)
 
 	NotifyMeshUpdated();
 	OnMeshChanged.Broadcast();
+
+	// Rebuild physics data
+	if (bDeferCollisionUpdates)
+	{
+		InvalidatePhysicsData();
+	}
+	else
+	{
+		RebuildPhysicsData();
+	}
 }
 
 
@@ -955,3 +991,140 @@ void UDynamicMeshComponent::OnChildDetached(USceneComponent* ChildComponent)
 	Super::OnChildDetached(ChildComponent);
 	OnChildAttachmentModified.Broadcast(ChildComponent, false);
 }
+
+
+
+
+
+bool UDynamicMeshComponent::GetPhysicsTriMeshData(struct FTriMeshCollisionData* CollisionData, bool InUseAllTriData)
+{
+	// todo: support physical materials
+	// todo: support UPhysicsSettings::Get()->bSupportUVFromHitResults
+
+	// this is something we currently assume, if you hit this ensure, we made a mistake
+	ensure(bEnableComplexCollision);
+
+	ProcessMesh([&](const FDynamicMesh3& Mesh)
+	{
+		TArray<int32> VertexMap;
+		bool bIsSparseV = !Mesh.IsCompactV();
+		if (bIsSparseV)
+		{
+			VertexMap.SetNum(Mesh.MaxVertexID());
+		}
+
+		// copy vertices
+		CollisionData->Vertices.Reserve(Mesh.VertexCount());
+		for (int32 vid : Mesh.VertexIndicesItr())
+		{
+			int32 Index = CollisionData->Vertices.Add((FVector)Mesh.GetVertex(vid));
+			if (bIsSparseV)
+			{
+				VertexMap[vid] = Index;
+			}
+			else
+			{
+				check(vid == Index);
+			}
+		}
+
+		// copy triangles
+		CollisionData->Indices.Reserve(Mesh.TriangleCount());
+		CollisionData->MaterialIndices.Reserve(Mesh.TriangleCount());
+		for (int32 tid : Mesh.TriangleIndicesItr())
+		{
+			FIndex3i Tri = Mesh.GetTriangle(tid);
+			FTriIndices Triangle;
+			Triangle.v0 = (bIsSparseV) ? VertexMap[Tri.A] : Tri.A;
+			Triangle.v1 = (bIsSparseV) ? VertexMap[Tri.B] : Tri.B;
+			Triangle.v2 = (bIsSparseV) ? VertexMap[Tri.C] : Tri.C;
+			CollisionData->Indices.Add(Triangle);
+
+			CollisionData->MaterialIndices.Add(0);		// not supporting physical materials yet
+		}
+
+		CollisionData->bFlipNormals = true;
+		CollisionData->bDeformableMesh = true;
+		CollisionData->bFastCook = true;
+	});
+
+	return true;
+}
+
+bool UDynamicMeshComponent::ContainsPhysicsTriMeshData(bool InUseAllTriData) const
+{
+	return bEnableComplexCollision && ((MeshObject != nullptr) ? (MeshObject->GetTriangleCount() > 0) : false);
+}
+
+bool UDynamicMeshComponent::WantsNegXTriMesh()
+{
+	return true;
+}
+
+UBodySetup* UDynamicMeshComponent::GetBodySetup()
+{
+	if (MeshBodySetup == nullptr)
+	{
+		UBodySetup* NewBodySetup = nullptr;
+		{
+			FGCScopeGuard Scope;
+			NewBodySetup = NewObject<UBodySetup>(this);
+		}
+		NewBodySetup->BodySetupGuid = FGuid::NewGuid();
+		
+		NewBodySetup->bGenerateMirroredCollision = false;
+		NewBodySetup->CollisionTraceFlag = this->CollisionType;
+
+		NewBodySetup->DefaultInstance.SetCollisionProfileName(UCollisionProfile::BlockAll_ProfileName);
+		NewBodySetup->bSupportUVsAndFaceRemap = false; /* bSupportPhysicalMaterialMasks; */
+		
+		SetBodySetup(NewBodySetup);
+	}
+
+	return MeshBodySetup;
+}
+
+void UDynamicMeshComponent::SetBodySetup(UBodySetup* NewSetup)
+{
+	if (ensure(NewSetup))
+	{
+		MeshBodySetup = NewSetup;
+	}
+}
+
+void UDynamicMeshComponent::InvalidatePhysicsData()
+{
+	if (GetBodySetup())
+	{
+		GetBodySetup()->InvalidatePhysicsData();
+		bCollisionUpdatePending = true;
+	}
+}
+
+void UDynamicMeshComponent::RebuildPhysicsData()
+{
+	UBodySetup* BodySetup = GetBodySetup();
+	if (BodySetup)
+	{
+		// New GUID as collision has changed
+		BodySetup->BodySetupGuid = FGuid::NewGuid();
+		// Also we want cooked data for this
+		BodySetup->bHasCookedCollisionData = true;
+
+		BodySetup->CollisionTraceFlag = this->CollisionType;
+
+		BodySetup->InvalidatePhysicsData();
+		BodySetup->CreatePhysicsMeshes();
+		RecreatePhysicsState();
+
+		bCollisionUpdatePending = false;
+	}
+}
+
+void UDynamicMeshComponent::UpdateCollision(bool bOnlyIfPending)
+{
+	if (bOnlyIfPending == false || bCollisionUpdatePending)
+	{
+		RebuildPhysicsData();
+	}
+}
\ No newline at end of file
diff --git a/Engine/Plugins/Experimental/MeshModelingToolset/Source/ModelingComponents/Private/UDynamicMesh.cpp b/Engine/Plugins/Experimental/MeshModelingToolset/Source/ModelingComponents/Private/UDynamicMesh.cpp
--- a/Engine/Plugins/Experimental/MeshModelingToolset/Source/ModelingComponents/Private/UDynamicMesh.cpp
+++ b/Engine/Plugins/Experimental/MeshModelingToolset/Source/ModelingComponents/Private/UDynamicMesh.cpp
@@ -83,6 +83,17 @@ void UDynamicMesh::InitializeMesh()
 }
 
 
+bool UDynamicMesh::IsEmpty() const
+{
+	return Mesh->TriangleCount() == 0;
+}
+
+int32 UDynamicMesh::GetTriangleCount() const
+{
+	return Mesh->TriangleCount();
+}
+
+
 void UDynamicMesh::SetMesh(const UE::Geometry::FDynamicMesh3& MoveMesh)
 {
 	FDynamicMeshChangeInfo ChangeInfo;
@@ -124,6 +135,21 @@ void UDynamicMesh::EditMeshInternal(TFunctionRef<void(FDynamicMesh3&)> EditFunc,
 		PreMeshChangedEvent.Broadcast(this, ChangeInfo);
 	}
 	EditFunc(GetMeshRef());
+
+	// Enforce our mesh attribute invariants. This should probably be optional to support compute-only UDynamicMeshes....
+	if (Mesh->HasTriangleGroups() == false)
+	{
+		Mesh->EnableTriangleGroups();
+	}
+	if (Mesh->HasAttributes() == false)
+	{
+		Mesh->EnableAttributes();
+	}
+	if (Mesh->Attributes()->HasMaterialID() == false)
+	{
+		Mesh->Attributes()->EnableMaterialID();
+	}
+
 	if (!bDeferChangeEvents)
 	{
 		MeshChangedEvent.Broadcast(this, ChangeInfo);
diff --git a/Engine/Plugins/Experimental/MeshModelingToolset/Source/ModelingComponents/Public/Components/DynamicMeshComponent.h b/Engine/Plugins/Experimental/MeshModelingToolset/Source/ModelingComponents/Public/Components/DynamicMeshComponent.h
--- a/Engine/Plugins/Experimental/MeshModelingToolset/Source/ModelingComponents/Public/Components/DynamicMeshComponent.h
+++ b/Engine/Plugins/Experimental/MeshModelingToolset/Source/ModelingComponents/Public/Components/DynamicMeshComponent.h
@@ -10,6 +10,7 @@
 #include "TransformTypes.h"
 #include "Async/Future.h"
 #include "UDynamicMesh.h"
+#include "PhysicsEngine/BodySetup.h"
 
 #include "DynamicMeshComponent.generated.h"
 
@@ -18,7 +19,7 @@
 
 /** internal FPrimitiveSceneProxy defined in DynamicMeshSceneProxy.h */
 class FDynamicMeshSceneProxy;
+class FBaseDynamicMeshSceneProxy;
 
-
 /**
  * Interface for a render mesh processor. Use this to process the Mesh stored in UDynamicMeshComponent before
@@ -58,8 +59,8 @@
  * See comment sections below for details.
  * 
  */
-UCLASS(hidecategories = (LOD, Physics, Collision), meta = (BlueprintSpawnableComponent), ClassGroup = Rendering)
-class MODELINGCOMPONENTS_API UDynamicMeshComponent : public UBaseDynamicMeshComponent
+UCLASS(hidecategories = (LOD), meta = (BlueprintSpawnableComponent), ClassGroup = Rendering)
+class MODELINGCOMPONENTS_API UDynamicMeshComponent : public UBaseDynamicMeshComponent, public IInterface_CollisionDataProvider
 {
 	GENERATED_UCLASS_BODY()
 
@@ -390,6 +391,79 @@
 
 
 
+	//===============================================================================================================
+	// Triangle-Vertex Tangents support. The default behavior is to not use Tangents, this will lead to incorrect
+	// rendering for any material with Normal Maps and some other shaders.
+	// If TangentsType == EDynamicMeshComponentTangentsMode::ExternallyProvided, the Tangent and Bitangent attributes of
+	//   the FDynamicMesh3 AttributeSet are used at the SceneProxy level, the Component is not involved
+	// If TangentsType == EDynamicMeshComponentTangentsMode::AutoCalculated, the Tangents are computed internally using
+	//   a fast MikkT approximation via FMeshTangentsf. They will be recomputed when the mesh is modified, however
+	//   they are *not* recomputed when using the Fast Update functions above (in that case InvalidateAutoCalculatedTangents() 
+	//   can be used to force recomputation)
+	//
+public:
+	UFUNCTION(BlueprintCallable, Category = "Dynamic Mesh Component")
+	void SetTangentsType(EDynamicMeshComponentTangentsMode NewTangentsType);
+
+	UFUNCTION(BlueprintCallable, Category = "Dynamic Mesh Component")
+	EDynamicMeshComponentTangentsMode GetTangentsType() const { return TangentsType; }
+
+	/** This function marks the auto tangents as dirty, they will be recomputed before they are used again */
+	virtual void InvalidateAutoCalculatedTangents();
+
+	/** @return AutoCalculated Tangent Set, which may require that they be recomputed, or nullptr if not enabled/available */
+	const UE::Geometry::FMeshTangentsf* GetAutoCalculatedTangents();
+
+protected:
+	/** How should Tangents be calculated/handled */
+	UPROPERTY()
+	EDynamicMeshComponentTangentsMode TangentsType = EDynamicMeshComponentTangentsMode::NoTangents;
+
+	/** true if AutoCalculatedTangents has been computed for current mesh */
+	bool bAutoCalculatedTangentsValid = false;
+
+	/** Set of per-triangle-vertex tangents computed for the current mesh. Only valid if bAutoCalculatedTangentsValid == true */
+	UE::Geometry::FMeshTangentsf AutoCalculatedTangents;
+
+	void UpdateAutoCalculatedTangents();
+
+
+
+	//===============================================================================================================
+	//
+	// Physics APIs
+	//
+
+public:
+	virtual bool GetPhysicsTriMeshData(struct FTriMeshCollisionData* CollisionData, bool InUseAllTriData) override;
+	virtual bool ContainsPhysicsTriMeshData(bool InUseAllTriData) const override;
+	virtual bool WantsNegXTriMesh() override;
+
+	virtual UBodySetup* GetBodySetup() override;
+	virtual void SetBodySetup(UBodySetup* NewSetup);
+
+	/** Type of Collision Geometry to use for this Mesh */
+	UPROPERTY(EditAnywhere, BlueprintReadWrite, Category = "Dynamic Mesh Component|Collision")
+	TEnumAsByte<enum ECollisionTraceFlag> CollisionType = ECollisionTraceFlag::CTF_UseSimpleAsComplex;
+
+	/** If true, current mesh will be used as Complex Collision source mesh. Independent of the CollisionType setting. */
+	UPROPERTY(EditAnywhere, BlueprintReadWrite, Category = "Dynamic Mesh Component|Collision");
+	bool bEnableComplexCollision = false;
+
+	/** If true, updates to the mesh will not result in immediate collision regeneration. Useful when the mesh will be modified multiple times before collision is needed. */
+	UPROPERTY(EditAnywhere, BlueprintReadWrite, Category = "Dynamic Mesh Component|Collision");
+	bool bDeferCollisionUpdates = false;
+
+protected:
+	UPROPERTY(Instanced)
+	TObjectPtr<UBodySetup> MeshBodySetup;
+
+	virtual void InvalidatePhysicsData();
+	virtual void RebuildPhysicsData();
+
+	bool bCollisionUpdatePending = false;
+	virtual void UpdateCollision(bool bOnlyIfPending);
+
 protected:
 
 	//
@@ -416,11 +490,14 @@
 	 */
 	bool bProxyValid = false;
 
+	virtual FBaseDynamicMeshSceneProxy* GetBaseSceneProxy() override { return (FBaseDynamicMeshSceneProxy*)GetCurrentSceneProxy(); }
+
 	/** 
 	 * @return current render proxy, if valid, otherwise nullptr 
 	 */
 	FDynamicMeshSceneProxy* GetCurrentSceneProxy();
 
+
 	/**
 	 * Fully invalidate all rendering data for this Component. Current Proxy will be discarded, Bounds and possibly Tangents recomputed, etc
 	 */
@@ -436,44 +513,10 @@
 
 	//~ UObject Interface.
 	virtual void PostLoad() override;
+#if WITH_EDITOR
+	void PostEditChangeProperty(FPropertyChangedEvent& PropertyChangedEvent) override;
+#endif
 
 
 
-
-	//===============================================================================================================
-	// Triangle-Vertex Tangents support. The default behavior is to not use Tangents, this will lead to incorrect
-	// rendering for any material with Normal Maps and some other shaders.
-	// If TangentsType == EDynamicMeshComponentTangentsMode::ExternallyProvided, the Tangent and Bitangent attributes of
-	//   the FDynamicMesh3 AttributeSet are used at the SceneProxy level, the Component is not involved
-	// If TangentsType == EDynamicMeshComponentTangentsMode::AutoCalculated, the Tangents are computed internally using
-	//   a fast MikkT approximation via FMeshTangentsf. They will be recomputed when the mesh is modified, however
-	//   they are *not* recomputed when using the Fast Update functions above (in that case InvalidateAutoCalculatedTangents() 
-	//   can be used to force recomputation)
-	//
-public:
-	UFUNCTION(BlueprintCallable, Category = "Dynamic Mesh Component")
-	void SetTangentsType(EDynamicMeshComponentTangentsMode NewTangentsType);
-
-	UFUNCTION(BlueprintCallable, Category = "Dynamic Mesh Component")
-	EDynamicMeshComponentTangentsMode GetTangentsType() const { return TangentsType; }
-
-	/** This function marks the auto tangents as dirty, they will be recomputed before they are used again */
-	virtual void InvalidateAutoCalculatedTangents();
-
-	/** @return AutoCalculated Tangent Set, which may require that they be recomputed, or nullptr if not enabled/available */
-	const UE::Geometry::FMeshTangentsf* GetAutoCalculatedTangents();
-
-protected:
-	/** How should Tangents be calculated/handled */
-	UPROPERTY()
-	EDynamicMeshComponentTangentsMode TangentsType = EDynamicMeshComponentTangentsMode::NoTangents;
-
-	/** true if AutoCalculatedTangents has been computed for current mesh */
-	bool bAutoCalculatedTangentsValid = false;
-
-	/** Set of per-triangle-vertex tangents computed for the current mesh. Only valid if bAutoCalculatedTangentsValid == true */
-	UE::Geometry::FMeshTangentsf AutoCalculatedTangents;
-
-	void UpdateAutoCalculatedTangents();
-
 };
diff --git a/Engine/Plugins/Experimental/MeshModelingToolset/Source/ModelingComponents/Public/Components/OctreeDynamicMeshComponent.h b/Engine/Plugins/Experimental/MeshModelingToolset/Source/ModelingComponents/Public/Components/OctreeDynamicMeshComponent.h
--- a/Engine/Plugins/Experimental/MeshModelingToolset/Source/ModelingComponents/Public/Components/OctreeDynamicMeshComponent.h
+++ b/Engine/Plugins/Experimental/MeshModelingToolset/Source/ModelingComponents/Public/Components/OctreeDynamicMeshComponent.h
@@ -16,7 +16,7 @@ struct FMeshDescription;
 
 /** internal FPrimitiveSceneProxy defined in OctreeDynamicMeshSceneProxy.h */
 class FOctreeDynamicMeshSceneProxy;
-
+class FBaseDynamicMeshSceneProxy;
 
 /** 
  * UOctreeDynamicMeshComponent is a mesh component similar to UProceduralMeshComponent,
@@ -116,6 +116,7 @@ protected:
 
 private:
 
+	virtual FBaseDynamicMeshSceneProxy* GetBaseSceneProxy() override { return (FBaseDynamicMeshSceneProxy*)GetCurrentSceneProxy(); }
 	FOctreeDynamicMeshSceneProxy* GetCurrentSceneProxy() { return (FOctreeDynamicMeshSceneProxy*)SceneProxy; }
 
 	//~ Begin UPrimitiveComponent Interface.
diff --git a/Engine/Plugins/Experimental/MeshModelingToolset/Source/ModelingComponents/Public/UDynamicMesh.h b/Engine/Plugins/Experimental/MeshModelingToolset/Source/ModelingComponents/Public/UDynamicMesh.h
--- a/Engine/Plugins/Experimental/MeshModelingToolset/Source/ModelingComponents/Public/UDynamicMesh.h
+++ b/Engine/Plugins/Experimental/MeshModelingToolset/Source/ModelingComponents/Public/UDynamicMesh.h
@@ -128,6 +128,19 @@ public:
 	 */
 	virtual void InitializeMesh();
 
+	/**
+	 * @return true if the mesh has no triangles
+	 */
+	UFUNCTION(BlueprintCallable, Category = "Dynamic Mesh")
+	bool IsEmpty() const;
+
+	/**
+	 * @return number of triangles in the mesh
+	 */
+	UFUNCTION(BlueprintCallable, Category = "Dynamic Mesh")
+	int32 GetTriangleCount() const;
+
+
 	/** 
 	 * @return Const reference to the internal FDynamicMesh3.
 	 * @warning Calling ProcessMesh() is preferred! This interface may be  removed in the future
