Approximate Actors
Exclude from the mesh generation all mesh sections that are using materials for which the material domain isn't "Surface" (ex: decals)
Those will still be rendered properly in the material baking phase.
Prevent black patches in places where decals aren't resting on surfaces

diff --git a/Engine/Plugins/Experimental/MeshModelingToolset/Source/ModelingComponents/Private/Scene/MeshSceneAdapter.cpp b/Engine/Plugins/Experimental/MeshModelingToolset/Source/ModelingComponents/Private/Scene/MeshSceneAdapter.cpp
--- a/Engine/Plugins/Experimental/MeshModelingToolset/Source/ModelingComponents/Private/Scene/MeshSceneAdapter.cpp
+++ b/Engine/Plugins/Experimental/MeshModelingToolset/Source/ModelingComponents/Private/Scene/MeshSceneAdapter.cpp
@@ -23,6 +23,7 @@
 #include "Components/StaticMeshComponent.h"
 #include "Components/InstancedStaticMeshComponent.h"
 #include "Engine/StaticMesh.h"
+#include "Materials/Material.h"
 
 #include "Async/ParallelFor.h"
 #include "Async/Async.h"
@@ -364,9 +365,84 @@ public:
 
 };
 
+/**
+ * Mesh adapter to filter out all mesh sections that are using materials for which the material domain isn't "Surface"
+ * This excludes decals for example.
+ */
+struct FMeshDescriptionTriangleMeshSurfaceAdapter : public FMeshDescriptionTriangleMeshAdapter
+{
+public:
+	FMeshDescriptionTriangleMeshSurfaceAdapter(const FMeshDescription* MeshIn, UStaticMesh* StaticMeshIn)
+		: FMeshDescriptionTriangleMeshAdapter(MeshIn)
+		, NumValidTriangles(0)
+		, bHasInvalidPolyGroup(false)
+	{
+		FStaticMeshConstAttributes MeshDescriptionAttributes(*MeshIn);
+		TPolygonGroupAttributesConstRef<FName> MaterialSlotNames = MeshDescriptionAttributes.GetPolygonGroupMaterialSlotNames();
+
+		for (FPolygonGroupID PolygonGroupID : MeshIn->PolygonGroups().GetElementIDs())
+		{
+			bool bValidPolyGroup = false;
+
+			int32 MaterialIndex = StaticMeshIn->GetMaterialIndex(MaterialSlotNames[PolygonGroupID]);
+			if (MaterialIndex > INDEX_NONE)
+			{
+				UMaterialInterface* MaterialInterface = StaticMeshIn->GetMaterial(MaterialIndex);
+				if (MaterialInterface && MaterialInterface->GetMaterial())
+				{
+					bValidPolyGroup = MaterialInterface->GetMaterial()->MaterialDomain == EMaterialDomain::MD_Surface;
+				}
+			}
+
+			if (bValidPolyGroup)
+			{
+				ValidPolyGroups.Add(PolygonGroupID);
+				NumValidTriangles += MeshIn->GetNumPolygonGroupTriangles(PolygonGroupID);
+			}
+			else
+			{
+				bHasInvalidPolyGroup = true;
+			}
+		}
+	}
+
+	bool IsTriangle(int32 TID) const
+	{
+		return bHasInvalidPolyGroup ? ValidPolyGroups.Contains(Mesh->GetTrianglePolygonGroup(TID)) : FMeshDescriptionTriangleMeshAdapter::IsTriangle(TID);
+	}
+
+	int32 TriangleCount() const
+	{
+		return bHasInvalidPolyGroup ? NumValidTriangles : FMeshDescriptionTriangleMeshAdapter::TriangleCount();
+	}
 
+private:
+	TSet<FPolygonGroupID> ValidPolyGroups;
+	int32 NumValidTriangles;
+	bool bHasInvalidPolyGroup;
+};
 
-class FStaticMeshSpatialWrapper : public IMeshSpatialWrapper
+template <class TriangleMeshType>
+struct TMeshDescriptionMeshAdapterd : public UE::Geometry::TTriangleMeshAdapter<double>
+{
+	TriangleMeshType ParentAdapter;
+
+	TMeshDescriptionMeshAdapterd(TriangleMeshType ParentAdapterIn) : ParentAdapter(ParentAdapterIn)
+	{
+		IsTriangle = [&](int index) { return ParentAdapter.IsTriangle(index); };
+		IsVertex = [&](int index) { return ParentAdapter.IsVertex(index); };
+		MaxTriangleID = [&]() { return ParentAdapter.MaxTriangleID(); };
+		MaxVertexID = [&]() { return ParentAdapter.MaxVertexID(); };
+		TriangleCount = [&]() { return ParentAdapter.TriangleCount(); };
+		VertexCount = [&]() { return ParentAdapter.VertexCount(); };
+		GetShapeTimestamp = [&]() { return ParentAdapter.GetShapeTimestamp(); };
+		GetTriangle = [&](int32 TriangleID) { return ParentAdapter.GetTriangle(TriangleID); };
+		GetVertex = [&](int32 VertexID) { return ParentAdapter.GetVertex(VertexID); };
+	}
+};
+
+template <class TriangleMeshType>
+class TStaticMeshSpatialWrapper : public IMeshSpatialWrapper
 {
 public:
 	UStaticMesh* StaticMesh = nullptr;
@@ -375,22 +451,21 @@ public:
 	
 	FMeshDescription* SourceMesh = nullptr;
 
-	TUniquePtr<FMeshDescriptionTriangleMeshAdapter> Adapter;
-	TUniquePtr<TMeshAABBTree3<FMeshDescriptionTriangleMeshAdapter>> AABBTree;
-	TUniquePtr<TFastWindingTree<FMeshDescriptionTriangleMeshAdapter>> FWNTree;
+	TUniquePtr<TriangleMeshType> Adapter;
+
+	TUniquePtr<TMeshAABBTree3<TriangleMeshType>> AABBTree;
+	TUniquePtr<TFastWindingTree<TriangleMeshType>> FWNTree;
 
 	virtual bool Build(const FMeshSceneAdapterBuildOptions& BuildOptions) override
 	{
 		check(StaticMesh);
-#if WITH_EDITOR
-		SourceMesh = StaticMesh->GetMeshDescription(LODIndex);
-#else
-		checkf(false, TEXT("Not currently supported - to build at Runtime it is necessary to read from the StaticMesh RenderBuffers"));
+
+#if !WITH_EDITOR
 		SourceMesh = nullptr;
 #endif
 		if (SourceMesh)
 		{
-			Adapter = MakeUnique<FMeshDescriptionTriangleMeshAdapter>(SourceMesh);
+			check(Adapter);
 
 			BuildScale = FVector3d::One();
 #if WITH_EDITOR
@@ -403,11 +478,11 @@ public:
 			{
 				{
 					TRACE_CPUPROFILER_EVENT_SCOPE(MeshScene_WrapperBuild_StaticMesh_AABBTree);
-					AABBTree = MakeUnique<TMeshAABBTree3<FMeshDescriptionTriangleMeshAdapter>>(Adapter.Get(), true);
+					AABBTree = MakeUnique<TMeshAABBTree3<TriangleMeshType>>(Adapter.Get(), true);
 				}
 				{
 					TRACE_CPUPROFILER_EVENT_SCOPE(MeshScene_WrapperBuild_StaticMesh_FWNTree);
-					FWNTree = MakeUnique<TFastWindingTree<FMeshDescriptionTriangleMeshAdapter>>(AABBTree.Get(), true);
+					FWNTree = MakeUnique<TFastWindingTree<TriangleMeshType>>(AABBTree.Get(), true);
 				}
 			}
 			return true;
@@ -427,13 +502,13 @@ public:
 	virtual FAxisAlignedBox3d GetWorldBounds(TFunctionRef<FVector3d(const FVector3d&)> LocalToWorldFunc) override
 	{
 		if (!SourceMesh) return FAxisAlignedBox3d::Empty();
-		return GetTransformedVertexBounds<FMeshDescriptionTriangleMeshAdapter>(*Adapter, LocalToWorldFunc);
+		return GetTransformedVertexBounds<TriangleMeshType>(*Adapter, LocalToWorldFunc);
 	}
 
 	virtual void CollectSeedPoints(TArray<FVector3d>& WorldPoints, TFunctionRef<FVector3d(const FVector3d&)> LocalToWorldFunc) override
 	{
 		if (!SourceMesh) return;
-		CollectSeedPointsFromMeshVertices<FMeshDescriptionTriangleMeshAdapter>(*Adapter, LocalToWorldFunc, WorldPoints);
+		CollectSeedPointsFromMeshVertices<TriangleMeshType>(*Adapter, LocalToWorldFunc, WorldPoints);
 	}
 
 	virtual double FastWindingNumber(const FVector3d& P, const FTransformSequence3d& LocalToWorldTransform) override
@@ -458,7 +533,10 @@ public:
 		if (!SourceMesh) return;
 
 #if WITH_EDITOR
-		if (AppendTo.TriangleCount() == 0 && TransformSeq.Num() == 0)
+		// Fast path only works on non-filtered meshes
+		// as it relies on the mesh description directly rather than the adapter
+		bool bFilteredMesh = SourceMesh->Triangles().Num() != Adapter->TriangleCount();
+		if (!bFilteredMesh && AppendTo.TriangleCount() == 0 && TransformSeq.Num() == 0)
 		{
 			// this is somewhat faster in profiling
 			FMeshDescription* UseMeshDescription = StaticMesh->GetMeshDescription(LODIndex);
@@ -473,7 +551,7 @@ public:
 
 		FDynamicMeshEditor Editor(&AppendTo);
 		FMeshIndexMappings Mappings;
-		FMeshDescriptionMeshAdapterd AdapterWrapper(*Adapter);
+		TMeshDescriptionMeshAdapterd<TriangleMeshType> AdapterWrapper(*Adapter);
 		Editor.AppendMesh(&AdapterWrapper, Mappings,
 			[&](int, const FVector3d& Pos) { return TransformSeq.TransformPosition(Pos); });
 	}
@@ -485,16 +563,25 @@ public:
 
 static TUniquePtr<IMeshSpatialWrapper> SpatialWrapperFactory( const FMeshTypeContainer& MeshContainer )
 {
+#if WITH_EDITOR
 	if (MeshContainer.MeshType == ESceneMeshType::StaticMeshAsset)
 	{
-		TUniquePtr<FStaticMeshSpatialWrapper> SMWrapper = MakeUnique<FStaticMeshSpatialWrapper>();
+		TUniquePtr<TStaticMeshSpatialWrapper<FMeshDescriptionTriangleMeshSurfaceAdapter>> SMWrapper = MakeUnique<TStaticMeshSpatialWrapper<FMeshDescriptionTriangleMeshSurfaceAdapter>>();
+
+		const int32 LODIndex = 0;
 		SMWrapper->SourceContainer = MeshContainer;
 		SMWrapper->StaticMesh = MeshContainer.GetStaticMesh();
+		SMWrapper->SourceMesh = SMWrapper->StaticMesh->GetMeshDescription(LODIndex);
+		SMWrapper->Adapter = MakeUnique<FMeshDescriptionTriangleMeshSurfaceAdapter>(SMWrapper->SourceMesh, SMWrapper->StaticMesh);
+
 		if (ensure(SMWrapper->StaticMesh != nullptr))
 		{
 			return SMWrapper;
 		}
 	}
+#else
+	checkf(false, TEXT("Not currently supported - to build at Runtime it is necessary to read from the StaticMesh RenderBuffers"));
+#endif
 
 	return TUniquePtr<IMeshSpatialWrapper>();
 }
