ModelingComponents: Add DynamicMeshActor
- ADynamicMeshActor type has a UDynamicMeshComponent subobject (similar to StaticMeshActor)
- Add IPersistentDynamicMeshSource, a ToolTarget Interface for accessing a UDynamicMesh
- Add UDynamicMeshComponentToolTarget, a UToolTarget implementing the above as well as various standard interfaces, to allow Tools to operate on an external DynamicMeshComponent
- Update UE::ToolTarget:: helper functions that get/commit via DynamicMesh to handle IPersistentDynamicMeshSource directly, instead of going via existing MeshDescription/DynamicMesh interfaces. Add UE::ToolTarget::SetSourceObjectVisible() helper function.
- Add support for creating DynamicMeshActors in UCreateMeshObjectTypeProperties and UEditorModelingObjectsCreationAPI
- register UDynamicMeshComponentToolTargetFactory in ModelingToolsEditorMode
- Rename existing internal ADynamicMeshActor to AOldDynamicMeshActor. Does not appear to be in use and will be removed in a future CL.

#rnx

diff --git a/Engine/Plugins/Experimental/MeshModelingToolset/Source/MeshModelingTools/Private/Properties/CreateMeshObjectTypeProperties.cpp b/Engine/Plugins/Experimental/MeshModelingToolset/Source/MeshModelingTools/Private/Properties/CreateMeshObjectTypeProperties.cpp
--- a/Engine/Plugins/Experimental/MeshModelingToolset/Source/MeshModelingTools/Private/Properties/CreateMeshObjectTypeProperties.cpp
+++ b/Engine/Plugins/Experimental/MeshModelingToolset/Source/MeshModelingTools/Private/Properties/CreateMeshObjectTypeProperties.cpp
@@ -5,20 +5,22 @@
 
 const FString UCreateMeshObjectTypeProperties::StaticMeshIdentifier = TEXT("Static Mesh");
 const FString UCreateMeshObjectTypeProperties::VolumeIdentifier = TEXT("Volume");
+const FString UCreateMeshObjectTypeProperties::DynamicMeshActorIdentifier = TEXT("Dynamic Mesh");
 
 void UCreateMeshObjectTypeProperties::InitializeDefault()
 {
 	bool bStaticMeshes = true;
 	bool bVolumes = false;
+	bool bDynamicMeshes = true;
 
 #if WITH_EDITOR
 	bVolumes = true;
 #endif
 
-	Initialize(bStaticMeshes, bVolumes);
+	Initialize(bStaticMeshes, bVolumes, bDynamicMeshes);
 }
 
-void UCreateMeshObjectTypeProperties::Initialize(bool bEnableStaticMeshes, bool bEnableVolumes)
+void UCreateMeshObjectTypeProperties::Initialize(bool bEnableStaticMeshes, bool bEnableVolumes, bool bEnableDynamicMeshActor)
 {
 	if (bEnableStaticMeshes)
 	{
@@ -28,6 +30,10 @@ void UCreateMeshObjectTypeProperties::Initialize(bool bEnableStaticMeshes, bool
 	{
 		OutputTypeNamesList.Add(VolumeIdentifier);
 	}
+	if (bEnableDynamicMeshActor)
+	{
+		OutputTypeNamesList.Add(DynamicMeshActorIdentifier);
+	}
 
 	if ((OutputType.Len() == 0) || (OutputType.Len() > 0 && OutputTypeNamesList.Contains(OutputType) == false))
 	{
@@ -58,6 +64,10 @@ ECreateObjectTypeHint UCreateMeshObjectTypeProperties::GetCurrentCreateMeshType(
 	{
 		return ECreateObjectTypeHint::Volume;
 	}
+	else if (OutputType == DynamicMeshActorIdentifier)
+	{
+		return ECreateObjectTypeHint::DynamicMeshActor;
+	}
 	return ECreateObjectTypeHint::Undefined;
 }
 
@@ -80,5 +90,10 @@ bool UCreateMeshObjectTypeProperties::ConfigureCreateMeshObjectParams(FCreateMes
 		ParamsOut.TypeHintClass = VolumeType.Get();
 		return true;
 	}
+	else if (OutputType == DynamicMeshActorIdentifier)
+	{
+		ParamsOut.TypeHint = ECreateObjectTypeHint::DynamicMeshActor;
+		return true;
+	}
 	return false;
 }
diff --git a/Engine/Plugins/Experimental/MeshModelingToolset/Source/MeshModelingTools/Public/Properties/CreateMeshObjectTypeProperties.h b/Engine/Plugins/Experimental/MeshModelingToolset/Source/MeshModelingTools/Public/Properties/CreateMeshObjectTypeProperties.h
--- a/Engine/Plugins/Experimental/MeshModelingToolset/Source/MeshModelingTools/Public/Properties/CreateMeshObjectTypeProperties.h
+++ b/Engine/Plugins/Experimental/MeshModelingToolset/Source/MeshModelingTools/Public/Properties/CreateMeshObjectTypeProperties.h
@@ -21,7 +21,7 @@ class MESHMODELINGTOOLS_API UCreateMeshObjectTypeProperties : public UInteractiv
 public:
 
 	void InitializeDefault();
-	void Initialize(bool bEnableStaticMeshes = true, bool bEnableVolumes = false);
+	void Initialize(bool bEnableStaticMeshes = true, bool bEnableVolumes = false, bool bEnableDynamicMeshActor = false);
 
 	/** Type of Object to create */
 	UPROPERTY(EditAnywhere, Category = OutputType, meta = (DisplayName = "Output Type", GetOptions = GetOutputTypeNamesFunc))
@@ -65,5 +65,6 @@ public:
 	// constants used for different known types
 	static const FString StaticMeshIdentifier;
 	static const FString VolumeIdentifier;
+	static const FString DynamicMeshActorIdentifier;
 
 };
\ No newline at end of file
diff --git a/Engine/Plugins/Experimental/MeshModelingToolset/Source/ModelingComponents/Private/DynamicMeshActor.cpp b/Engine/Plugins/Experimental/MeshModelingToolset/Source/ModelingComponents/Private/DynamicMeshActor.cpp
new file mode 100644
--- /dev/null
+++ b/Engine/Plugins/Experimental/MeshModelingToolset/Source/ModelingComponents/Private/DynamicMeshActor.cpp
@@ -0,0 +1,30 @@
+// Copyright Epic Games, Inc. All Rights Reserved.
+
+#include "DynamicMeshActor.h"
+#include "ToolSetupUtil.h"
+
+
+#define LOCTEXT_NAMESPACE "ADynamicMeshActor"
+
+ADynamicMeshActor::ADynamicMeshActor(const FObjectInitializer& ObjectInitializer)
+	: Super(ObjectInitializer)
+{
+	DynamicMeshComponent = CreateDefaultSubobject<UDynamicMeshComponent>(TEXT("DynamicMeshComponent"));
+	DynamicMeshComponent->SetMobility(EComponentMobility::Movable);
+	DynamicMeshComponent->SetGenerateOverlapEvents(false);
+	DynamicMeshComponent->SetCollisionProfileName(UCollisionProfile::BlockAll_ProfileName);
+
+	DynamicMeshComponent->CollisionType = ECollisionTraceFlag::CTF_UseDefault;
+
+	DynamicMeshComponent->SetMaterial(0, ToolSetupUtil::GetDefaultMaterial());		// is this necessary?
+
+	SetRootComponent(DynamicMeshComponent);
+}
+
+
+
+
+
+
+
+#undef LOCTEXT_NAMESPACE
\ No newline at end of file
diff --git a/Engine/Plugins/Experimental/MeshModelingToolset/Source/ModelingComponents/Private/ModelingToolTargetUtil.cpp b/Engine/Plugins/Experimental/MeshModelingToolset/Source/ModelingComponents/Private/ModelingToolTargetUtil.cpp
--- a/Engine/Plugins/Experimental/MeshModelingToolset/Source/ModelingComponents/Private/ModelingToolTargetUtil.cpp
+++ b/Engine/Plugins/Experimental/MeshModelingToolset/Source/ModelingComponents/Private/ModelingToolTargetUtil.cpp
@@ -7,6 +7,7 @@
 #include "TargetInterfaces/MeshDescriptionCommitter.h"
 #include "TargetInterfaces/MeshDescriptionProvider.h"
 #include "TargetInterfaces/PrimitiveComponentBackedTarget.h"
+#include "TargetInterfaces/DynamicMeshSource.h"
 
 #include "ModelingObjectsCreationAPI.h"
 
@@ -19,6 +20,8 @@
 #include "MeshDescriptionToDynamicMesh.h"
 #include "DynamicMeshToMeshDescription.h"
 
+#define LOCTEXT_NAMESPACE "ModelingToolTargetUtil"
+
 using namespace UE::Geometry;
 
 AActor* UE::ToolTarget::GetTargetActor(UToolTarget* Target)
@@ -68,6 +71,19 @@ bool UE::ToolTarget::ShowSourceObject(UToolTarget* Target)
 }
 
 
+bool UE::ToolTarget::SetSourceObjectVisible(UToolTarget* Target, bool bVisible)
+{
+	if (bVisible)
+	{
+		return ShowSourceObject(Target);
+	}
+	else
+	{
+		return HideSourceObject(Target);
+	}
+}
+
+
 FTransform3d UE::ToolTarget::GetLocalToWorldTransform(UToolTarget* Target)
 {
 	IPrimitiveComponentBackedTarget* TargetComponent = Cast<IPrimitiveComponentBackedTarget>(Target);
@@ -139,6 +155,15 @@ FMeshDescription UE::ToolTarget::GetMeshDescriptionCopy(UToolTarget* Target, boo
 
 FDynamicMesh3 UE::ToolTarget::GetDynamicMeshCopy(UToolTarget* Target, bool bWantMeshTangents)
 {
+	IPersistentDynamicMeshSource* DynamicMeshSource = Cast<IPersistentDynamicMeshSource>(Target);
+	if (DynamicMeshSource)
+	{
+		UDynamicMesh* DynamicMesh = DynamicMeshSource->GetDynamicMeshContainer();
+		FDynamicMesh3 Mesh;
+		DynamicMesh->ProcessMesh([&](const FDynamicMesh3& ReadMesh) { Mesh = ReadMesh; });
+		return Mesh;
+	}
+
 	IMeshDescriptionProvider* MeshDescriptionProvider = Cast<IMeshDescriptionProvider>(Target);
 	FDynamicMesh3 Mesh(EMeshComponents::FaceGroups);
 	Mesh.EnableAttributes();
@@ -197,6 +222,27 @@ UE::ToolTarget::EDynamicMeshUpdateResult UE::ToolTarget::CommitDynamicMeshUpdate
 	const FConversionToMeshDescriptionOptions& ConversionOptions,
 	const FComponentMaterialSet* UpdatedMaterials)
 {
+	IPersistentDynamicMeshSource* DynamicMeshSource = Cast<IPersistentDynamicMeshSource>(Target);
+	if (DynamicMeshSource)
+	{
+		UDynamicMesh* DynamicMesh = DynamicMeshSource->GetDynamicMeshContainer();
+		TUniquePtr<FDynamicMesh3> CurrentMesh = DynamicMesh->ExtractMesh();
+		TSharedPtr<FDynamicMesh3> CurrentMeshShared(CurrentMesh.Release());
+
+		DynamicMesh->EditMesh([&](FDynamicMesh3& EditMesh) { EditMesh = UpdatedMesh; });
+
+		TSharedPtr<FDynamicMesh3> NewMeshShared = MakeShared<FDynamicMesh3>();
+		DynamicMesh->ProcessMesh([&](const FDynamicMesh3& ReadMesh) { *NewMeshShared = ReadMesh; });
+
+		TUniquePtr<FMeshReplacementChange> ReplaceChange = MakeUnique<FMeshReplacementChange>(CurrentMeshShared, NewMeshShared);
+		
+		DynamicMeshSource->CommitDynamicMeshChange(MoveTemp(ReplaceChange), LOCTEXT("CommitDynamicMeshUpdate_MeshSource", "Edit Mesh"));
+
+
+		// todo support bModifiedTopology flag?
+		return EDynamicMeshUpdateResult::Ok;
+	}
+
 	IMeshDescriptionCommitter* MeshDescriptionCommitter = Cast<IMeshDescriptionCommitter>(Target);
 	if (MeshDescriptionCommitter)
 	{
@@ -235,6 +281,16 @@ UE::ToolTarget::EDynamicMeshUpdateResult UE::ToolTarget::CommitDynamicMeshUpdate
 
 UE::ToolTarget::EDynamicMeshUpdateResult UE::ToolTarget::CommitDynamicMeshUVUpdate(UToolTarget* Target, const UE::Geometry::FDynamicMesh3* UpdatedMesh)
 {
+	IPersistentDynamicMeshSource* DynamicMeshSource = Cast<IPersistentDynamicMeshSource>(Target);
+	if (DynamicMeshSource)
+	{
+		// todo actually only update UVs? 
+		UDynamicMesh* DynamicMesh = DynamicMeshSource->GetDynamicMeshContainer();
+		DynamicMesh->EditMesh([&](FDynamicMesh3& EditMesh) { EditMesh = *UpdatedMesh; });
+		return EDynamicMeshUpdateResult::Ok;
+	}
+
+
 	IMeshDescriptionCommitter* MeshDescriptionCommitter = Cast<IMeshDescriptionCommitter>(Target);
 	if (!ensure(MeshDescriptionCommitter))
 	{
@@ -293,4 +349,8 @@ bool UE::ToolTarget::ConfigureCreateMeshObjectParams(UToolTarget* SourceTarget,
 		}
 	}
 	return false;
-}
\ No newline at end of file
+}
+
+
+
+#undef LOCTEXT_NAMESPACE
\ No newline at end of file
diff --git a/Engine/Plugins/Experimental/MeshModelingToolset/Source/ModelingComponents/Public/DynamicMeshActor.h b/Engine/Plugins/Experimental/MeshModelingToolset/Source/ModelingComponents/Public/DynamicMeshActor.h
new file mode 100644
--- /dev/null
+++ b/Engine/Plugins/Experimental/MeshModelingToolset/Source/ModelingComponents/Public/DynamicMeshActor.h
@@ -0,0 +1,29 @@
+// Copyright Epic Games, Inc. All Rights Reserved.
+
+#pragma once
+
+#include "CoreMinimal.h"
+#include "UDynamicMesh.h"
+#include "Components/DynamicMeshComponent.h"
+
+#include "DynamicMeshActor.generated.h"
+
+
+/**
+ * ADynamicMeshActor is an Actor that has a USimpleDynamicMeshComponent as it's RootObject.
+ * Used for experimental procedural geometry generation Blueprints. Currently no physics support.
+ */
+UCLASS(ConversionRoot, ComponentWrapperClass, ClassGroup=DynamicMesh, meta = (ChildCanTick))
+class MODELINGCOMPONENTS_API ADynamicMeshActor : public AActor
+{
+	GENERATED_UCLASS_BODY()
+private:
+	UPROPERTY(Category = DynamicMeshActor, VisibleAnywhere, BlueprintReadOnly, meta = (ExposeFunctionCategories = "Mesh,Rendering,Physics,Components|StaticMesh", AllowPrivateAccess = "true"))
+	TObjectPtr<class UDynamicMeshComponent> DynamicMeshComponent;
+
+public:
+	UFUNCTION(BlueprintCallable, Category = DynamicMeshActor)
+		UDynamicMeshComponent* GetDynamicMeshComponent() const { return DynamicMeshComponent; }
+};
+
+
diff --git a/Engine/Plugins/Experimental/MeshModelingToolset/Source/ModelingComponents/Public/ModelingToolTargetUtil.h b/Engine/Plugins/Experimental/MeshModelingToolset/Source/ModelingComponents/Public/ModelingToolTargetUtil.h
--- a/Engine/Plugins/Experimental/MeshModelingToolset/Source/ModelingComponents/Public/ModelingToolTargetUtil.h
+++ b/Engine/Plugins/Experimental/MeshModelingToolset/Source/ModelingComponents/Public/ModelingToolTargetUtil.h
@@ -49,6 +49,12 @@ MODELINGCOMPONENTS_API bool HideSourceObject(UToolTarget* Target);
  */
 MODELINGCOMPONENTS_API bool ShowSourceObject(UToolTarget* Target);
 
+/**
+ * Show or Hide the "Source Object" (eg PrimitiveComponent, Actor, etc) backing a ToolTarget
+ * @return true on success
+ */
+MODELINGCOMPONENTS_API bool SetSourceObjectVisible(UToolTarget* Target, bool bVisible);
+
 /**
  * @return the local-to-world Transform underlying a ToolTarget, eg the Component or Actor transform
  */
diff --git a/Engine/Plugins/Experimental/MeshModelingToolset/Source/ModelingComponents/Public/TargetInterfaces/DynamicMeshSource.h b/Engine/Plugins/Experimental/MeshModelingToolset/Source/ModelingComponents/Public/TargetInterfaces/DynamicMeshSource.h
new file mode 100644
--- /dev/null
+++ b/Engine/Plugins/Experimental/MeshModelingToolset/Source/ModelingComponents/Public/TargetInterfaces/DynamicMeshSource.h
@@ -0,0 +1,53 @@
+// Copyright Epic Games, Inc. All Rights Reserved.
+
+#pragma once
+
+#include "CoreMinimal.h"
+#include "UObject/Interface.h"
+#include "InteractiveToolChange.h"
+
+#include "DynamicMeshSource.generated.h"
+
+class UDynamicMesh;
+class UDynamicMeshComponent;
+
+UINTERFACE()
+class MODELINGCOMPONENTS_API UPersistentDynamicMeshSource : public UInterface
+{
+	GENERATED_BODY()
+};
+
+class MODELINGCOMPONENTS_API IPersistentDynamicMeshSource
+{
+	GENERATED_BODY()
+
+public:
+
+	/**
+	 * @return a UDynamicMesh for this source. 
+	 */
+	virtual UDynamicMesh* GetDynamicMeshContainer() = 0;
+
+	/**
+	 * Commit a change to the UDynamicMesh. The assumption here is that the Change has already been applied to
+	 * the UDynamicMesh returned by GetDynamicMeshContainer(). This function is used by Tools to provide that
+	 * change to higher levels for storage in undo systems/etc (eg the Editor Transaction system)
+	 */
+	virtual void CommitDynamicMeshChange(TUniquePtr<FToolCommandChange> Change, const FText& ChangeMessage) = 0;
+
+	//
+	// optional
+	//
+
+	/**
+	 * @return true if the UDynamicMesh comes from a UDynamicMeshComponent (or subclass)
+	 */
+	virtual bool HasDynamicMeshComponent() const { return false; }
+
+	/**
+	 * @return the UDynamicMeshComponent that owns the source UDynamicMesh, if available
+	 */
+	virtual UDynamicMeshComponent* GetDynamicMeshComponent() { return nullptr; }
+
+
+};
diff --git a/Engine/Plugins/Experimental/MeshModelingToolset/Source/ModelingComponentsEditorOnly/Private/EditorModelingObjectsCreationAPI.cpp b/Engine/Plugins/Experimental/MeshModelingToolset/Source/ModelingComponentsEditorOnly/Private/EditorModelingObjectsCreationAPI.cpp
--- a/Engine/Plugins/Experimental/MeshModelingToolset/Source/ModelingComponentsEditorOnly/Private/EditorModelingObjectsCreationAPI.cpp
+++ b/Engine/Plugins/Experimental/MeshModelingToolset/Source/ModelingComponentsEditorOnly/Private/EditorModelingObjectsCreationAPI.cpp
@@ -17,13 +17,16 @@
 #include "Materials/Material.h"
 
 #include "ToolTargets/VolumeDynamicMeshToolTarget.h"  // for CVarModelingMaxVolumeTriangleCount
-
 #include "Engine/BlockingVolume.h"
 #include "Components/BrushComponent.h"
 #include "Engine/Polys.h"
 #include "Model.h"
 #include "BSPOps.h"		// in UnrealEd
 
+#include "DynamicMeshActor.h"
+#include "Components/DynamicMeshComponent.h"
+
+
 using namespace UE::Geometry;
 
 
@@ -97,6 +100,10 @@ FCreateMeshObjectResult UEditorModelingObjectsCreationAPI::CreateMeshObject(FCre
 	{
 		ResultOut = CreateVolume(MoveTemp(CreateMeshParams));
 	}
+	else if (CreateMeshParams.TypeHint == ECreateObjectTypeHint::DynamicMeshActor)
+	{
+		ResultOut = CreateDynamicMeshActor(MoveTemp(CreateMeshParams));
+	}
 	else
 	{
 		ResultOut = CreateStaticMeshAsset(MoveTemp(CreateMeshParams));
@@ -180,6 +187,56 @@ FCreateMeshObjectResult UEditorModelingObjectsCreationAPI::CreateVolume(FCreateM
 
 
 
+FCreateMeshObjectResult UEditorModelingObjectsCreationAPI::CreateDynamicMeshActor(FCreateMeshObjectParams&& CreateMeshParams)
+{
+	// spawn new actor
+	FRotator Rotation(0.0f, 0.0f, 0.0f);
+	FActorSpawnParameters SpawnInfo;
+	ADynamicMeshActor* NewActor = CreateMeshParams.TargetWorld->SpawnActor<ADynamicMeshActor>(SpawnInfo);
+
+	UDynamicMeshComponent* NewComponent = NewActor->GetDynamicMeshComponent();
+
+	if (CreateMeshParams.MeshType == ECreateMeshObjectSourceMeshType::DynamicMesh)
+	{
+		NewComponent->SetMesh(MoveTemp(CreateMeshParams.DynamicMesh.GetValue()));
+		NewComponent->NotifyMeshUpdated();
+	}
+	else if (CreateMeshParams.MeshType == ECreateMeshObjectSourceMeshType::MeshDescription)
+	{
+		const FMeshDescription* MeshDescription = &CreateMeshParams.MeshDescription.GetValue();
+		FDynamicMesh3 Mesh(EMeshComponents::FaceGroups);
+		Mesh.EnableAttributes();
+		FMeshDescriptionToDynamicMesh Converter;
+		Converter.Convert(MeshDescription, Mesh, true);
+		NewComponent->SetMesh(MoveTemp(Mesh));
+	}
+	else
+	{
+		return FCreateMeshObjectResult{ ECreateModelingObjectResult::Failed_InvalidMesh };
+	}
+
+	NewActor->SetActorTransform(CreateMeshParams.Transform);
+
+	// set materials
+	TArray<UMaterialInterface*> ComponentMaterials = FilterMaterials(CreateMeshParams.Materials);
+	for (int32 k = 0; k < ComponentMaterials.Num(); ++k)
+	{
+		NewComponent->SetMaterial(k, ComponentMaterials[k]);
+	}
+
+	NewActor->PostEditChange();
+
+	// emit result
+	FCreateMeshObjectResult ResultOut;
+	ResultOut.ResultCode = ECreateModelingObjectResult::Ok;
+	ResultOut.NewActor = NewActor;
+	ResultOut.NewComponent = NewComponent;
+	ResultOut.NewAsset = nullptr;
+	return ResultOut;
+}
+
+
+
 
 FCreateMeshObjectResult UEditorModelingObjectsCreationAPI::CreateStaticMeshAsset(FCreateMeshObjectParams&& CreateMeshParams)
 {
diff --git a/Engine/Plugins/Experimental/MeshModelingToolset/Source/ModelingComponentsEditorOnly/Private/ToolTargets/DynamicMeshComponentToolTarget.cpp b/Engine/Plugins/Experimental/MeshModelingToolset/Source/ModelingComponentsEditorOnly/Private/ToolTargets/DynamicMeshComponentToolTarget.cpp
new file mode 100644
--- /dev/null
+++ b/Engine/Plugins/Experimental/MeshModelingToolset/Source/ModelingComponentsEditorOnly/Private/ToolTargets/DynamicMeshComponentToolTarget.cpp
@@ -0,0 +1,196 @@
+// Copyright Epic Games, Inc. All Rights Reserved.
+
+#include "ToolTargets/DynamicMeshComponentToolTarget.h"
+
+#include "Components/DynamicMeshComponent.h"
+#include "ComponentReregisterContext.h"
+#include "DynamicMeshToMeshDescription.h"
+#include "MeshDescriptionToDynamicMesh.h"
+#include "StaticMeshAttributes.h"
+#include "RenderingThread.h"
+#include "Materials/Material.h"
+
+#include "Misc/ITransaction.h"
+#include "ScopedTransaction.h"
+
+
+#define LOCTEXT_NAMESPACE "UDynamicMeshComponentToolTarget"
+
+
+bool UDynamicMeshComponentToolTarget::IsValid() const
+{
+	if (!UPrimitiveComponentToolTarget::IsValid())
+	{
+		return false;
+	}
+	UDynamicMeshComponent* DynamicMeshComponent = Cast<UDynamicMeshComponent>(Component);
+	if (DynamicMeshComponent == nullptr)
+	{
+		return false;
+	}
+	UDynamicMesh* DynamicMesh = DynamicMeshComponent->GetDynamicMesh();
+	if (DynamicMesh == nullptr)
+	{
+		return false;
+	}
+
+	return true;
+}
+
+
+int32 UDynamicMeshComponentToolTarget::GetNumMaterials() const
+{
+	return ensure(IsValid()) ? Component->GetNumMaterials() : 0;
+}
+
+UMaterialInterface* UDynamicMeshComponentToolTarget::GetMaterial(int32 MaterialIndex) const
+{
+	return ensure(IsValid()) ? Component->GetMaterial(MaterialIndex) : nullptr;
+}
+
+void UDynamicMeshComponentToolTarget::GetMaterialSet(FComponentMaterialSet& MaterialSetOut, bool bPreferAssetMaterials) const
+{
+	if (!ensure(IsValid())) return;
+
+	int32 NumMaterials = Component->GetNumMaterials();
+	MaterialSetOut.Materials.SetNum(NumMaterials);
+	for (int32 k = 0; k < NumMaterials; ++k)
+	{
+		MaterialSetOut.Materials[k] = Component->GetMaterial(k);
+	}
+}
+
+bool UDynamicMeshComponentToolTarget::CommitMaterialSetUpdate(const FComponentMaterialSet& MaterialSet, bool bApplyToAsset)
+{
+	if (!ensure(IsValid())) return false;
+
+	// filter out any Engine materials that we don't want to be permanently assigning
+	TArray<UMaterialInterface*> FilteredMaterials = MaterialSet.Materials;
+	for (int32 k = 0; k < FilteredMaterials.Num(); ++k)
+	{
+		FString AssetPath = FilteredMaterials[k]->GetPathName();
+		if (AssetPath.StartsWith(TEXT("/MeshModelingToolset/")))
+		{
+			FilteredMaterials[k] = UMaterial::GetDefaultMaterial(MD_Surface);
+		}
+	}
+
+	UDynamicMeshComponent* DynamicMeshComponent = Cast<UDynamicMeshComponent>(Component);
+
+	int32 NumMaterialsNeeded = Component->GetNumMaterials();
+	int32 NumMaterialsGiven = FilteredMaterials.Num();
+
+	DynamicMeshComponent->Modify();
+	for (int32 k = 0; k < NumMaterialsGiven; ++k)
+	{
+		DynamicMeshComponent->SetMaterial(k, FilteredMaterials[k]);
+	}
+
+	return true;
+}
+
+
+const FMeshDescription* UDynamicMeshComponentToolTarget::GetMeshDescription()
+{
+	if (ensure(IsValid()))
+	{
+		if (bHaveMeshDescription)
+		{
+			return ConvertedMeshDescription.Get();
+		}
+
+		UDynamicMeshComponent* DynamicMeshComponent = Cast<UDynamicMeshComponent>(Component);
+
+		ConvertedMeshDescription = MakeUnique<FMeshDescription>();
+		FStaticMeshAttributes Attributes(*ConvertedMeshDescription);
+		Attributes.Register();
+
+		FDynamicMeshToMeshDescription Converter;
+		Converter.Convert(DynamicMeshComponent->GetMesh(), *ConvertedMeshDescription);
+
+		bHaveMeshDescription = true;
+		return ConvertedMeshDescription.Get();
+	}
+	return nullptr;
+}
+
+
+void UDynamicMeshComponentToolTarget::CommitMeshDescription(const FCommitter& Committer)
+{
+	if (ensure(IsValid()) == false) return;
+
+	UDynamicMesh* DynamicMesh = GetDynamicMeshContainer();
+	TUniquePtr<FDynamicMesh3> CurrentMesh = DynamicMesh->ExtractMesh();
+	TSharedPtr<FDynamicMesh3> CurrentMeshShared(CurrentMesh.Release());
+
+	FMeshDescription EditingMeshDescription(*GetMeshDescription());
+	FCommitterParams CommitterParams;
+	CommitterParams.MeshDescriptionOut = &EditingMeshDescription;
+	Committer(CommitterParams);
+	FMeshDescriptionToDynamicMesh Converter;
+	TSharedPtr<FDynamicMesh3> NewMeshShared = MakeShared<FDynamicMesh3>();
+	NewMeshShared->EnableAttributes();
+	Converter.Convert(CommitterParams.MeshDescriptionOut, *NewMeshShared);
+
+	DynamicMesh->EditMesh([&](FDynamicMesh3& EditMesh) { EditMesh = *NewMeshShared; });
+
+	TUniquePtr<FMeshReplacementChange> ReplaceChange = MakeUnique<FMeshReplacementChange>(CurrentMeshShared, NewMeshShared);
+
+	CommitDynamicMeshChange(MoveTemp(ReplaceChange), LOCTEXT("DynamicMeshComponentToolTargetCommit", "Update Mesh"));
+}
+
+
+UDynamicMesh* UDynamicMeshComponentToolTarget::GetDynamicMeshContainer()
+{
+	UDynamicMeshComponent* DynamicMeshComponent = Cast<UDynamicMeshComponent>(Component);
+	return DynamicMeshComponent->GetDynamicMesh();
+}
+
+bool UDynamicMeshComponentToolTarget::HasDynamicMeshComponent() const
+{
+	return true;
+}
+
+UDynamicMeshComponent* UDynamicMeshComponentToolTarget::GetDynamicMeshComponent()
+{
+	return Cast<UDynamicMeshComponent>(Component);
+}
+
+
+
+void UDynamicMeshComponentToolTarget::CommitDynamicMeshChange(TUniquePtr<FToolCommandChange> Change, const FText& ChangeMessage)
+{
+	FScopedTransaction Transaction(ChangeMessage);
+	check(GUndo != nullptr);
+	UDynamicMesh* DynamicMesh = GetDynamicMeshContainer();
+	DynamicMesh->Modify();
+	GUndo->StoreUndo(DynamicMesh, MoveTemp(Change));
+}
+
+
+
+
+
+// Factory
+
+bool UDynamicMeshComponentToolTargetFactory::CanBuildTarget(UObject* SourceObject, const FToolTargetTypeRequirements& Requirements) const
+{
+	UDynamicMeshComponent* Component = Cast<UDynamicMeshComponent>(SourceObject);
+	return Component 
+		&& !Component->IsPendingKillOrUnreachable() 
+		&& Component->IsValidLowLevel() 
+		&& Component->GetDynamicMesh()
+		&& Requirements.AreSatisfiedBy(UDynamicMeshComponentToolTarget::StaticClass());
+}
+
+UToolTarget* UDynamicMeshComponentToolTargetFactory::BuildTarget(UObject* SourceObject, const FToolTargetTypeRequirements& Requirements)
+{
+	UDynamicMeshComponentToolTarget* Target = NewObject<UDynamicMeshComponentToolTarget>();// TODO: Should we set an outer here?
+	Target->Component = Cast<UDynamicMeshComponent>(SourceObject);
+	check(Target->Component && Requirements.AreSatisfiedBy(Target));
+
+	return Target;
+}
+
+
+#undef LOCTEXT_NAMESPACE
\ No newline at end of file
diff --git a/Engine/Plugins/Experimental/MeshModelingToolset/Source/ModelingComponentsEditorOnly/Public/EditorModelingObjectsCreationAPI.h b/Engine/Plugins/Experimental/MeshModelingToolset/Source/ModelingComponentsEditorOnly/Public/EditorModelingObjectsCreationAPI.h
--- a/Engine/Plugins/Experimental/MeshModelingToolset/Source/ModelingComponentsEditorOnly/Public/EditorModelingObjectsCreationAPI.h
+++ b/Engine/Plugins/Experimental/MeshModelingToolset/Source/ModelingComponentsEditorOnly/Public/EditorModelingObjectsCreationAPI.h
@@ -77,6 +77,7 @@ public:
 	//
 	FCreateMeshObjectResult CreateStaticMeshAsset(FCreateMeshObjectParams&& CreateMeshParams);
 	FCreateMeshObjectResult CreateVolume(FCreateMeshObjectParams&& CreateMeshParams);
+	FCreateMeshObjectResult CreateDynamicMeshActor(FCreateMeshObjectParams&& CreateMeshParams);
 
 	TArray<UMaterialInterface*> FilterMaterials(const TArray<UMaterialInterface*>& MaterialsIn);
 };
diff --git a/Engine/Plugins/Experimental/MeshModelingToolset/Source/ModelingComponentsEditorOnly/Public/ToolTargets/DynamicMeshComponentToolTarget.h b/Engine/Plugins/Experimental/MeshModelingToolset/Source/ModelingComponentsEditorOnly/Public/ToolTargets/DynamicMeshComponentToolTarget.h
new file mode 100644
--- /dev/null
+++ b/Engine/Plugins/Experimental/MeshModelingToolset/Source/ModelingComponentsEditorOnly/Public/ToolTargets/DynamicMeshComponentToolTarget.h
@@ -0,0 +1,76 @@
+// Copyright Epic Games, Inc. All Rights Reserved.
+
+#pragma once
+
+#include "CoreMinimal.h"
+
+#include "TargetInterfaces/MaterialProvider.h"
+#include "TargetInterfaces/MeshDescriptionCommitter.h"
+#include "TargetInterfaces/MeshDescriptionProvider.h"
+#include "TargetInterfaces/DynamicMeshSource.h"
+#include "ToolTargets/PrimitiveComponentToolTarget.h"
+
+#include "DynamicMeshComponentToolTarget.generated.h"
+
+class UDynamicMesh;
+
+/**
+ * A ToolTarget backed by a DynamicMeshComponent
+ */
+UCLASS(Transient)
+class MODELINGCOMPONENTSEDITORONLY_API UDynamicMeshComponentToolTarget : 
+	public UPrimitiveComponentToolTarget,
+	public IMeshDescriptionCommitter, 
+	public IMeshDescriptionProvider, 
+	public IMaterialProvider,
+	public IPersistentDynamicMeshSource
+{
+	GENERATED_BODY()
+
+public:
+	virtual bool IsValid() const override;
+
+
+public:
+	// IMeshDescriptionProvider implementation
+	const FMeshDescription* GetMeshDescription() override;
+
+	// IMeshDescritpionCommitter implementation
+	virtual void CommitMeshDescription(const FCommitter& Committer) override;
+	using IMeshDescriptionCommitter::CommitMeshDescription; // unhide the other overload
+
+	// IMaterialProvider implementation
+	int32 GetNumMaterials() const override;
+	UMaterialInterface* GetMaterial(int32 MaterialIndex) const override;
+	void GetMaterialSet(FComponentMaterialSet& MaterialSetOut, bool bPreferAssetMaterials) const override;
+	virtual bool CommitMaterialSetUpdate(const FComponentMaterialSet& MaterialSet, bool bApplyToAsset) override;	
+
+	// IPersistentDynamicMeshSource implementation
+	virtual UDynamicMesh* GetDynamicMeshContainer() override;
+	virtual void CommitDynamicMeshChange(TUniquePtr<FToolCommandChange> Change, const FText& ChangeMessage) override;
+	virtual bool HasDynamicMeshComponent() const override;
+	virtual UDynamicMeshComponent* GetDynamicMeshComponent() override;
+
+	// Rest provided by parent class
+
+protected:
+	TUniquePtr<FMeshDescription> ConvertedMeshDescription;
+	bool bHaveMeshDescription = false;
+
+protected:
+	friend class UDynamicMeshComponentToolTargetFactory;
+};
+
+
+/** Factory for UDynamicMeshComponentToolTarget to be used by the target manager. */
+UCLASS(Transient)
+class MODELINGCOMPONENTSEDITORONLY_API UDynamicMeshComponentToolTargetFactory : public UToolTargetFactory
+{
+	GENERATED_BODY()
+
+public:
+
+	virtual bool CanBuildTarget(UObject* SourceObject, const FToolTargetTypeRequirements& TargetTypeInfo) const override;
+
+	virtual UToolTarget* BuildTarget(UObject* SourceObject, const FToolTargetTypeRequirements& TargetTypeInfo) override;
+};
\ No newline at end of file
diff --git a/Engine/Plugins/Experimental/ModelingToolsEditorMode/Source/ModelingToolsEditorMode/Private/ModelingToolsEditorMode.cpp b/Engine/Plugins/Experimental/ModelingToolsEditorMode/Source/ModelingToolsEditorMode/Private/ModelingToolsEditorMode.cpp
--- a/Engine/Plugins/Experimental/ModelingToolsEditorMode/Source/ModelingToolsEditorMode/Private/ModelingToolsEditorMode.cpp
+++ b/Engine/Plugins/Experimental/ModelingToolsEditorMode/Source/ModelingToolsEditorMode/Private/ModelingToolsEditorMode.cpp
@@ -10,7 +10,7 @@
 #include "ToolTargets/StaticMeshComponentToolTarget.h"
 #include "ToolTargets/VolumeDynamicMeshToolTarget.h"
 #include "ToolTargets/VolumeMeshDescriptionToolTarget.h"
-//#include "ToolTargets/DynamicMeshComponentToolTarget.h"
+#include "ToolTargets/DynamicMeshComponentToolTarget.h"
 #include "Framework/Commands/UICommandList.h"
 #include "Framework/Application/SlateApplication.h"
 #include "Framework/MultiBox/MultiBoxBuilder.h"
@@ -306,7 +306,7 @@ void UModelingToolsEditorMode::Enter()
 	ToolsContext->TargetManager->AddTargetFactory(NewObject<UStaticMeshComponentToolTargetFactory>(ToolsContext->TargetManager));
 	ToolsContext->TargetManager->AddTargetFactory(NewObject<UVolumeDynamicMeshToolTargetFactory>(ToolsContext->TargetManager));
 	ToolsContext->TargetManager->AddTargetFactory(NewObject<UVolumeMeshDescriptionToolTargetFactory>(ToolsContext->TargetManager));
-	//ToolsContext->TargetManager->AddTargetFactory(NewObject<UDynamicMeshComponentToolTargetFactory>(ToolsContext->TargetManager));
+	ToolsContext->TargetManager->AddTargetFactory(NewObject<UDynamicMeshComponentToolTargetFactory>(ToolsContext->TargetManager));
 
 	// register stylus event handler
 	StylusStateTracker = MakeUnique<FStylusStateTracker>();
