ModelingTools: add cvar to enable/disable creation of DynamicMeshActors, disable by default
#rnx

diff --git a/Engine/Plugins/Experimental/MeshModelingToolset/Source/MeshModelingTools/Private/Properties/CreateMeshObjectTypeProperties.cpp b/Engine/Plugins/Experimental/MeshModelingToolset/Source/MeshModelingTools/Private/Properties/CreateMeshObjectTypeProperties.cpp
--- a/Engine/Plugins/Experimental/MeshModelingToolset/Source/MeshModelingTools/Private/Properties/CreateMeshObjectTypeProperties.cpp
+++ b/Engine/Plugins/Experimental/MeshModelingToolset/Source/MeshModelingTools/Private/Properties/CreateMeshObjectTypeProperties.cpp
@@ -3,6 +3,11 @@
 #include "Properties/CreateMeshObjectTypeProperties.h"
 #include "ModelingObjectsCreationAPI.h"
 
+static TAutoConsoleVariable<int32> CVarEnableDynamicMeshActors(
+	TEXT("geometry.DynamicMeshActor.Enable"),
+	0,
+	TEXT("Enable DynamicMeshActor Output Type in various Modeling Mode Tools"));
+
 const FString UCreateMeshObjectTypeProperties::StaticMeshIdentifier = TEXT("Static Mesh");
 const FString UCreateMeshObjectTypeProperties::VolumeIdentifier = TEXT("Volume");
 const FString UCreateMeshObjectTypeProperties::DynamicMeshActorIdentifier = TEXT("Dynamic Mesh");
@@ -30,7 +35,7 @@ void UCreateMeshObjectTypeProperties::Initialize(bool bEnableStaticMeshes, bool
 	{
 		OutputTypeNamesList.Add(VolumeIdentifier);
 	}
-	if (bEnableDynamicMeshActor)
+	if (bEnableDynamicMeshActor && CVarEnableDynamicMeshActors.GetValueOnGameThread() != 0)
 	{
 		OutputTypeNamesList.Add(DynamicMeshActorIdentifier);
 	}
