ModelingTools: add Volume output support to ExtractCollisionGeometryTool. Add support for emitting each collision shape as a separate object, and emitting complex collision. Fix TObjectPtr issue in SetCollisionGeometryTool.
#rnx

diff --git a/Engine/Plugins/Experimental/MeshModelingToolset/Source/MeshModelingTools/Private/Physics/ExtractCollisionGeometryTool.cpp b/Engine/Plugins/Experimental/MeshModelingToolset/Source/MeshModelingTools/Private/Physics/ExtractCollisionGeometryTool.cpp
--- a/Engine/Plugins/Experimental/MeshModelingToolset/Source/MeshModelingTools/Private/Physics/ExtractCollisionGeometryTool.cpp
+++ b/Engine/Plugins/Experimental/MeshModelingToolset/Source/MeshModelingTools/Private/Physics/ExtractCollisionGeometryTool.cpp
@@ -17,6 +17,7 @@
 #include "Generators/CapsuleGenerator.h"
 #include "DynamicMesh/MeshTransforms.h"
 #include "Parameterization/DynamicMeshUVEditor.h"
+#include "DynamicMesh/Operations/MergeCoincidentMeshEdges.h"
 
 #include "Physics/PhysicsDataCollection.h"
 #include "Physics/CollisionGeometryVisualization.h"
@@ -27,11 +28,10 @@
 #include "Engine/Classes/PhysicsEngine/BodySetup.h"
 #include "Engine/Classes/PhysicsEngine/AggregateGeom.h"
 
-#include "TargetInterfaces/MaterialProvider.h"
-#include "TargetInterfaces/MeshDescriptionCommitter.h"
 #include "TargetInterfaces/MeshDescriptionProvider.h"
 #include "TargetInterfaces/PrimitiveComponentBackedTarget.h"
 #include "TargetInterfaces/StaticMeshBackedTarget.h"
+#include "ModelingToolTargetUtil.h"
 
 #include "ExplicitUseGeometryMathTypes.h"		// using UE::Geometry::(math types)
 using namespace UE::Geometry;
@@ -41,8 +41,6 @@ using namespace UE::Geometry;
 const FToolTargetTypeRequirements& UExtractCollisionGeometryToolBuilder::GetTargetRequirements() const
 {
 	static FToolTargetTypeRequirements TypeRequirements({
-		UMaterialProvider::StaticClass(),
-		UMeshDescriptionCommitter::StaticClass(),
 		UMeshDescriptionProvider::StaticClass(),
 		UPrimitiveComponentBackedTarget::StaticClass(),
 		UStaticMeshBackedTarget::StaticClass()
@@ -62,11 +60,10 @@ void UExtractCollisionGeometryTool::Setup()
 	UInteractiveTool::Setup();
 
 	// create preview mesh
-	IPrimitiveComponentBackedTarget* TargetComponent = Cast<IPrimitiveComponentBackedTarget>(Target);
 	PreviewMesh = NewObject<UPreviewMesh>(this);
 	PreviewMesh->bBuildSpatialDataStructure = false;
 	PreviewMesh->CreateInWorld(TargetWorld, FTransform::Identity);
-	PreviewMesh->SetTransform(TargetComponent->GetWorldTransform());
+	PreviewMesh->SetTransform((FTransform)UE::ToolTarget::GetLocalToWorldTransform(Target));
 	PreviewMesh->SetMaterial(ToolSetupUtil::GetDefaultSculptMaterial(GetToolManager()));
 	PreviewMesh->SetOverrideRenderMaterial(ToolSetupUtil::GetSelectionMaterial(GetToolManager()));
 	PreviewMesh->SetTriangleColorFunction([this](const FDynamicMesh3* Mesh, int TriangleID)
@@ -74,6 +71,23 @@ void UExtractCollisionGeometryTool::Setup()
 		return LinearColors::SelectFColor(Mesh->GetTriangleGroup(TriangleID));
 	});
 
+	OutputTypeProperties = NewObject<UCreateMeshObjectTypeProperties>(this);
+	OutputTypeProperties->OutputType = UCreateMeshObjectTypeProperties::VolumeIdentifier;		// prefer volumes for extracting simple collision
+	OutputTypeProperties->RestoreProperties(this, TEXT("ExtractCollisionTool"));
+	OutputTypeProperties->InitializeDefault();
+	OutputTypeProperties->WatchProperty(OutputTypeProperties->OutputType, [this](FString) { OutputTypeProperties->UpdatePropertyVisibility(); });
+	AddToolPropertySource(OutputTypeProperties);
+
+	Settings = NewObject<UExtractCollisionToolProperties>(this);
+	Settings->RestoreProperties(this);
+	AddToolPropertySource(Settings);
+	Settings->WatchProperty(Settings->CollisionType, [this](EExtractCollisionOutputType NewValue) { bResultValid = false; });
+	Settings->WatchProperty(Settings->bWeldEdges, [this](bool bNewValue) { bResultValid = false; });
+	Settings->WatchProperty(Settings->bShowPreview, [this](bool bNewValue) { PreviewMesh->SetVisible(bNewValue); });
+	PreviewMesh->SetVisible(Settings->bShowPreview);
+	Settings->WatchProperty(Settings->bShowInputMesh, [this](bool bNewValue) { UE::ToolTarget::SetSourceObjectVisible(Target, bNewValue); });
+	UE::ToolTarget::SetSourceObjectVisible(Target, Settings->bShowInputMesh);
+
 	VizSettings = NewObject<UCollisionGeometryVisualizationProperties>(this);
 	VizSettings->RestoreProperties(this);
 	AddToolPropertySource(VizSettings);
@@ -81,18 +95,20 @@ void UExtractCollisionGeometryTool::Setup()
 	VizSettings->WatchProperty(VizSettings->Color, [this](FColor NewValue) { bVisualizationDirty = true; });
 
 
-	const UStaticMeshComponent* Component = CastChecked<UStaticMeshComponent>(TargetComponent->GetOwnerComponent());
-	const UStaticMesh* StaticMesh = Component->GetStaticMesh();
-	if (ensure(StaticMesh && StaticMesh->GetBodySetup()))
+	UPrimitiveComponent* TargetComponent = UE::ToolTarget::GetTargetComponent(Target);
+	UStaticMeshComponent* StaticMeshComponent = (TargetComponent) ? Cast<UStaticMeshComponent>(TargetComponent) : nullptr;
+	UStaticMesh* StaticMesh = (StaticMeshComponent) ? StaticMeshComponent->GetStaticMesh() : nullptr;
+	UBodySetup* BodySetup = (StaticMesh) ? StaticMesh->GetBodySetup() : nullptr;
+	if (ensure(BodySetup))
 	{
 		PhysicsInfo = MakeShared<FPhysicsDataCollection>();
-		PhysicsInfo->InitializeFromComponent(Component, true);
+		PhysicsInfo->InitializeFromComponent(StaticMeshComponent, true);
 
 		PreviewElements = NewObject<UPreviewGeometry>(this);
-		FTransform TargetTransform = TargetComponent->GetWorldTransform();
+		FTransform TargetTransform = (FTransform)UE::ToolTarget::GetLocalToWorldTransform(Target);
 		//PhysicsInfo->ExternalScale3D = TargetTransform.GetScale3D();
 		//TargetTransform.SetScale3D(FVector::OneVector);
-		PreviewElements->CreateInWorld(TargetComponent->GetOwnerActor()->GetWorld(), TargetTransform);
+		PreviewElements->CreateInWorld(UE::ToolTarget::GetTargetActor(Target)->GetWorld(), TargetTransform);
 		
 		UE::PhysicsTools::InitializePreviewGeometryLines(*PhysicsInfo, PreviewElements,
 			VizSettings->Color, VizSettings->LineThickness, 0.0f, 16);
@@ -102,46 +118,76 @@ void UExtractCollisionGeometryTool::Setup()
 		AddToolPropertySource(ObjectProps);
 	}
 
-
 	SetToolDisplayName(LOCTEXT("ToolName", "Collision To Mesh"));
 	GetToolManager()->DisplayMessage(
-		LOCTEXT("OnStartTool", "Convert Collision geometry to Static Mesh"),
+		LOCTEXT("OnStartTool", "Convert Collision Geometry to Mesh Objects"),
 		EToolMessageLevel::UserNotification);
 }
 
 
 void UExtractCollisionGeometryTool::Shutdown(EToolShutdownType ShutdownType)
 {
+	OutputTypeProperties->SaveProperties(this, TEXT("ExtractCollisionTool"));
+	Settings->SaveProperties(this);
 	VizSettings->SaveProperties(this);
 
-	UE::Geometry::FTransform3d Transform(PreviewMesh->GetTransform());
+	UE::Geometry::FTransform3d ActorTransform(PreviewMesh->GetTransform());
 
 	PreviewElements->Disconnect();
 	PreviewMesh->SetVisible(false);
 	PreviewMesh->Disconnect();
 	PreviewMesh = nullptr;
 
+	UE::ToolTarget::ShowSourceObject(Target);
+
 	if (ShutdownType == EToolShutdownType::Accept)
 	{
 		UMaterialInterface* UseMaterial = UMaterial::GetDefaultMaterial(MD_Surface);
 
-		FString NewName = Target->IsValid() ?
-			FString::Printf(TEXT("%s_Collision"), *Cast<IPrimitiveComponentBackedTarget>(Target)->GetOwnerComponent()->GetName()) : TEXT("CollisionGeo");
+		FString TargetName = UE::ToolTarget::GetTargetComponent(Target)->GetName();
 
 		GetToolManager()->BeginUndoTransaction(LOCTEXT("CreateCollisionMesh", "Collision To Mesh"));
 
-		FCreateMeshObjectParams NewMeshObjectParams;
-		NewMeshObjectParams.TargetWorld = TargetWorld;
-		NewMeshObjectParams.Transform = (FTransform)Transform;
-		NewMeshObjectParams.BaseName = NewName;
-		NewMeshObjectParams.Materials.Add(UseMaterial);
-		NewMeshObjectParams.SetMesh(&CurrentMesh);
-		FCreateMeshObjectResult Result = UE::Modeling::CreateMeshObject(GetToolManager(), MoveTemp(NewMeshObjectParams));
-		if (Result.IsOK() && Result.NewActor != nullptr)
+		TArray<AActor*> OutputSelection;
+
+		auto EmitNewMesh = [&](FDynamicMesh3&& Mesh, FTransform3d UseTransform, FString UseName)
+		{
+			FCreateMeshObjectParams NewMeshObjectParams;
+			NewMeshObjectParams.TargetWorld = TargetWorld;
+			NewMeshObjectParams.Transform = (FTransform)UseTransform;
+			NewMeshObjectParams.BaseName = UseName;
+			NewMeshObjectParams.Materials.Add(UseMaterial);
+			NewMeshObjectParams.SetMesh(MoveTemp(Mesh));
+			OutputTypeProperties->ConfigureCreateMeshObjectParams(NewMeshObjectParams);
+			FCreateMeshObjectResult Result = UE::Modeling::CreateMeshObject(GetToolManager(), MoveTemp(NewMeshObjectParams));
+			if (Result.IsOK() && Result.NewActor != nullptr)
+			{
+				OutputSelection.Add(Result.NewActor);
+			}
+		};
+
+		int32 NumParts = CurrentMeshParts.Num();
+		if (Settings->bOutputSeparateMeshes && NumParts > 1)
 		{
-			ToolSelectionUtil::SetNewActorSelection(GetToolManager(), Result.NewActor);
+			for ( int32 k = 0; k < NumParts; ++k)
+			{
+				FDynamicMesh3& MeshPart = CurrentMeshParts[k];
+				FAxisAlignedBox3d Bounds = MeshPart.GetBounds();
+				MeshTransforms::Translate(MeshPart, -Bounds.Center());
+				FTransform3d CenterTransform = ActorTransform;
+				CenterTransform.SetTranslation(CenterTransform.GetTranslation() + Bounds.Center());
+				FString NewName = FString::Printf(TEXT("%s_Collision%d"), *TargetName, k);
+				EmitNewMesh(MoveTemp(MeshPart), CenterTransform, NewName);
+			}
+		}
+		else
+		{
+			FString NewName = FString::Printf(TEXT("%s_Collision"), *TargetName);
+			EmitNewMesh(MoveTemp(CurrentMesh), ActorTransform, NewName);
 		}
 
+		ToolSelectionUtil::SetNewActorSelection(GetToolManager(), OutputSelection);
+
 		GetToolManager()->EndUndoTransaction();
 	}
 }
@@ -159,7 +205,16 @@ void UExtractCollisionGeometryTool::OnTick(float DeltaTime)
 {
 	if (bResultValid == false)
 	{
-		RecalculateMesh();
+		GetToolManager()->DisplayMessage(FText(), EToolMessageLevel::UserWarning);
+		if (Settings->CollisionType == EExtractCollisionOutputType::Simple)
+		{
+			RecalculateMesh_Simple();
+		}
+		else
+		{
+			RecalculateMesh_Complex();
+		}
+		bResultValid = true;
 	}
 
 	if (bVisualizationDirty)
@@ -173,24 +228,29 @@ void UExtractCollisionGeometryTool::OnTick(float DeltaTime)
 
 void UExtractCollisionGeometryTool::UpdateVisualization()
 {
-	float UseThickness = VizSettings->LineThickness / 10.0f;
+	float UseThickness = VizSettings->LineThickness;
 	FColor UseColor = VizSettings->Color;
 	PreviewElements->UpdateAllLineSets([&](ULineSetComponent* LineSet)
 		{
 			LineSet->SetAllLinesThickness(UseThickness);
 			LineSet->SetAllLinesColor(UseColor);
 		});
+	
+	UMaterialInterface* LineMaterial = ToolSetupUtil::GetDefaultLineComponentMaterial(GetToolManager(), !VizSettings->bShowHidden);
+	PreviewElements->SetAllLineSetsMaterial(LineMaterial);
 }
 
 
 
-void UExtractCollisionGeometryTool::RecalculateMesh()
+void UExtractCollisionGeometryTool::RecalculateMesh_Simple()
 {
 	int32 SphereResolution = 16;
 
 	CurrentMesh = FDynamicMesh3(EMeshComponents::FaceGroups);
 	CurrentMesh.EnableAttributes();
 
+	CurrentMeshParts.Reset();
+
 	FDynamicMeshEditor Editor(&CurrentMesh);
 
 	const FKAggregateGeom& AggGeom = PhysicsInfo->AggGeom;
@@ -208,6 +268,7 @@ void UExtractCollisionGeometryTool::RecalculateMesh()
 
 		FMeshIndexMappings Mappings;
 		Editor.AppendMesh(&SphereMesh, Mappings);
+		CurrentMeshParts.Add(MoveTemp(SphereMesh));
 	}
 
 
@@ -224,6 +285,7 @@ void UExtractCollisionGeometryTool::RecalculateMesh()
 
 		FMeshIndexMappings Mappings;
 		Editor.AppendMesh(&BoxMesh, Mappings);
+		CurrentMeshParts.Add(MoveTemp(BoxMesh));
 	}
 
 
@@ -245,6 +307,7 @@ void UExtractCollisionGeometryTool::RecalculateMesh()
 
 		FMeshIndexMappings Mappings;
 		Editor.AppendMesh(&CapsuleMesh, Mappings);
+		CurrentMeshParts.Add(MoveTemp(CapsuleMesh));
 	}
 
 
@@ -270,57 +333,81 @@ void UExtractCollisionGeometryTool::RecalculateMesh()
 		UVEditor.SetPerTriangleUVs();
 		FMeshIndexMappings Mappings;
 		Editor.AppendMesh(&ConvexMesh, Mappings);
+		CurrentMeshParts.Add(MoveTemp(ConvexMesh));
 	}
 
+	for (FDynamicMesh3& MeshPart : CurrentMeshParts)
+	{
+		FMeshNormals::InitializeMeshToPerTriangleNormals(&MeshPart);
+	}
 	FMeshNormals::InitializeMeshToPerTriangleNormals(&CurrentMesh);
-	PreviewMesh->UpdatePreview(&CurrentMesh);
 
+	PreviewMesh->UpdatePreview(&CurrentMesh);
 
-	//PreviewGeom->CreateOrUpdateLineSet(TEXT("Capsules"), AggGeom.SphylElems.Num(), [&](int32 Index, TArray<FRenderableLine>& LinesOut)
-	//{
-	//	const FKSphylElem& Capsule = AggGeom.SphylElems[Index];
-	//	FTransform ElemTransform = Capsule.GetTransform();
-	//	ElemTransform.ScaleTranslation(PhysicsData.ExternalScale3D);
-	//	FTransform3f ElemTransformf(ElemTransform);
-	//	const float HalfLength = Capsule.GetScaledCylinderLength(PhysicsData.ExternalScale3D) * .5f;
-	//	const float Radius = Capsule.GetScaledRadius(PhysicsData.ExternalScale3D);
-	//	FVector3f Top(0, 0, HalfLength), Bottom(0, 0, -HalfLength);
-
-	//	// top and bottom circles
-	//	UE::Geometry::GenerateCircleSegments<float>(CircleSteps, Radius, Top, FVector3f::UnitX(), FVector3f::UnitY(), ElemTransformf,
-	//		[&](const FVector3f& A, const FVector3f& B) { LinesOut.Add(FRenderableLine((FVector)A, (FVector)B, CapsuleColor, LineThickness, DepthBias)); });
-	//	UE::Geometry::GenerateCircleSegments<float>(CircleSteps, Radius, Bottom, FVector3f::UnitX(), FVector3f::UnitY(), ElemTransformf,
-	//		[&](const FVector3f& A, const FVector3f& B) { LinesOut.Add(FRenderableLine((FVector)A, (FVector)B, CapsuleColor, LineThickness, DepthBias)); });
-
-	//	// top dome
-	//	UE::Geometry::GenerateArcSegments<float>(CircleSteps, Radius, 0.0, PI, Top, FVector3f::UnitY(), FVector3f::UnitZ(), ElemTransformf,
-	//		[&](const FVector3f& A, const FVector3f& B) { LinesOut.Add(FRenderableLine((FVector)A, (FVector)B, CapsuleColor, LineThickness, DepthBias)); });
-	//	UE::Geometry::GenerateArcSegments<float>(CircleSteps, Radius, 0.0, PI, Top, FVector3f::UnitX(), FVector3f::UnitZ(), ElemTransformf,
-	//		[&](const FVector3f& A, const FVector3f& B) { LinesOut.Add(FRenderableLine((FVector)A, (FVector)B, CapsuleColor, LineThickness, DepthBias)); });
-
-	//	// bottom dome
-	//	UE::Geometry::GenerateArcSegments<float>(CircleSteps, Radius, 0.0, -PI, Bottom, FVector3f::UnitY(), FVector3f::UnitZ(), ElemTransformf,
-	//		[&](const FVector3f& A, const FVector3f& B) { LinesOut.Add(FRenderableLine((FVector)A, (FVector)B, CapsuleColor, LineThickness, DepthBias)); });
-	//	UE::Geometry::GenerateArcSegments<float>(CircleSteps, Radius, 0.0, -PI, Bottom, FVector3f::UnitX(), FVector3f::UnitZ(), ElemTransformf,
-	//		[&](const FVector3f& A, const FVector3f& B) { LinesOut.Add(FRenderableLine((FVector)A, (FVector)B, CapsuleColor, LineThickness, DepthBias)); });
-
-	//	// connecting lines
-	//	for (int k = 0; k < 2; ++k)
-	//	{
-	//		FVector DX = (k < 1) ? FVector(-Radius, 0, 0) : FVector(Radius, 0, 0);
-	//		LinesOut.Add(FRenderableLine(
-	//			ElemTransform.TransformPosition((FVector)Top + DX),
-	//			ElemTransform.TransformPosition((FVector)Bottom + DX), CapsuleColor, LineThickness, DepthBias));
-	//		FVector DY = (k < 1) ? FVector(0, -Radius, 0) : FVector(0, Radius, 0);
-	//		LinesOut.Add(FRenderableLine(
-	//			ElemTransform.TransformPosition((FVector)Top + DY),
-	//			ElemTransform.TransformPosition((FVector)Bottom + DY), CapsuleColor, LineThickness, DepthBias));
-	//	}
-	//});
+	if (CurrentMeshParts.Num() == 0)
+	{
+		GetToolManager()->DisplayMessage(LOCTEXT("NoSimpleCollisionShapes", "This Mesh has no Simple Collision Shapes"), EToolMessageLevel::UserWarning);
+	}
 }
 
 
 
+void UExtractCollisionGeometryTool::RecalculateMesh_Complex()
+{
+	CurrentMesh = FDynamicMesh3(EMeshComponents::FaceGroups);
+	CurrentMesh.EnableAttributes();
+	CurrentMeshParts.Reset();
+
+	bool bMeshErrors = false;
+
+	UStaticMeshComponent* StaticMeshComponent = Cast<UStaticMeshComponent>(UE::ToolTarget::GetTargetComponent(Target));
+	UStaticMesh* StaticMesh = (StaticMeshComponent) ? StaticMeshComponent->GetStaticMesh() : nullptr;
+	IInterface_CollisionDataProvider* CollisionProvider = (StaticMesh) ? Cast<IInterface_CollisionDataProvider>(StaticMesh) : nullptr;
+	if (CollisionProvider && CollisionProvider->ContainsPhysicsTriMeshData(true))
+	{
+		FTriMeshCollisionData CollisionData;
+		if (CollisionProvider->GetPhysicsTriMeshData(&CollisionData, true))
+		{
+			TArray<int32> VertexIDMap;
+			for (int32 k = 0; k < CollisionData.Vertices.Num(); ++k)
+			{
+				int32 vid = CurrentMesh.AppendVertex(CollisionData.Vertices[k]);
+				VertexIDMap.Add(vid);
+			}
+			for (const FTriIndices& TriIndices : CollisionData.Indices)
+			{
+				FIndex3i Triangle(TriIndices.v0, TriIndices.v1, TriIndices.v2);
+				int32 tid = CurrentMesh.AppendTriangle(Triangle);
+				if (tid < 0)
+				{
+					bMeshErrors = true;
+				}
+			}
+
+			if (Settings->bWeldEdges)
+			{
+				FMergeCoincidentMeshEdges Weld(&CurrentMesh);
+				Weld.OnlyUniquePairs = true;
+				Weld.Apply();
+				Weld.OnlyUniquePairs = false;
+				Weld.Apply();
+			}
+
+			if (!CollisionData.bFlipNormals)		// collision mesh has reversed orientation
+			{
+				CurrentMesh.ReverseOrientation(false);
+			}
+
+			FMeshNormals::InitializeMeshToPerTriangleNormals(&CurrentMesh);
+		}
+	}
+
+	PreviewMesh->UpdatePreview(&CurrentMesh);
 
+	if (CurrentMesh.TriangleCount() == 0)
+	{
+		GetToolManager()->DisplayMessage(LOCTEXT("EmptyComplexCollision", "This Mesh has no Complex Collision geometry"), EToolMessageLevel::UserWarning);
+	}
+}
 
 #undef LOCTEXT_NAMESPACE
\ No newline at end of file
diff --git a/Engine/Plugins/Experimental/MeshModelingToolset/Source/MeshModelingTools/Private/Physics/SetCollisionGeometryTool.cpp b/Engine/Plugins/Experimental/MeshModelingToolset/Source/MeshModelingTools/Private/Physics/SetCollisionGeometryTool.cpp
--- a/Engine/Plugins/Experimental/MeshModelingToolset/Source/MeshModelingTools/Private/Physics/SetCollisionGeometryTool.cpp
+++ b/Engine/Plugins/Experimental/MeshModelingToolset/Source/MeshModelingTools/Private/Physics/SetCollisionGeometryTool.cpp
@@ -185,7 +185,7 @@ void USetCollisionGeometryTool::Shutdown(EToolShutdownType ShutdownType)
 
 		UPrimitiveComponent* Component = UE::ToolTarget::GetTargetComponent(Targets[Targets.Num() - 1]);
 		UStaticMeshComponent* StaticMeshComponent = Cast<UStaticMeshComponent>(Component);
-		UStaticMesh* StaticMesh = (StaticMeshComponent) ? StaticMeshComponent->GetStaticMesh() : nullptr;
+		TObjectPtr<UStaticMesh> StaticMesh = (StaticMeshComponent) ? StaticMeshComponent->GetStaticMesh() : nullptr;
 		UBodySetup* BodySetup = (StaticMesh) ? StaticMesh->GetBodySetup() : nullptr;
 		if (BodySetup != nullptr)
 		{
diff --git a/Engine/Plugins/Experimental/MeshModelingToolset/Source/MeshModelingTools/Public/Physics/ExtractCollisionGeometryTool.h b/Engine/Plugins/Experimental/MeshModelingToolset/Source/MeshModelingTools/Public/Physics/ExtractCollisionGeometryTool.h
--- a/Engine/Plugins/Experimental/MeshModelingToolset/Source/MeshModelingTools/Public/Physics/ExtractCollisionGeometryTool.h
+++ b/Engine/Plugins/Experimental/MeshModelingToolset/Source/MeshModelingTools/Public/Physics/ExtractCollisionGeometryTool.h
@@ -3,11 +3,10 @@
 #pragma once
 
 #include "CoreMinimal.h"
-#include "SingleSelectionTool.h"
-#include "InteractiveToolBuilder.h"
+#include "BaseTools/SingleSelectionMeshEditingTool.h"
 #include "DynamicMesh/DynamicMesh3.h"
 #include "Physics/CollisionPropertySets.h"
-#include "BaseTools/SingleSelectionMeshEditingTool.h"
+#include "Properties/CreateMeshObjectTypeProperties.h"
 #include "ExtractCollisionGeometryTool.generated.h"
 
 class UPreviewGeometry;
@@ -26,6 +25,42 @@ protected:
 };
 
 
+UENUM()
+enum class EExtractCollisionOutputType : uint8
+{
+	/** Simple Collision shapes (Box, Sphere, Capsule, Convex) */
+	Simple = 0,
+	/** Complex Collision Mesh */
+	Complex = 1
+};
+
+
+UCLASS()
+class MESHMODELINGTOOLS_API UExtractCollisionToolProperties : public UInteractiveToolPropertySet
+{
+	GENERATED_BODY()
+public:
+	/** Type of collision geometry to convert to Mesh */
+	UPROPERTY(EditAnywhere, Category = Options)
+	EExtractCollisionOutputType CollisionType = EExtractCollisionOutputType::Simple;
+
+	/** Whether or not to weld coincident border edges of the Complex Collision Mesh (if possible) */
+	UPROPERTY(EditAnywhere, Category = Options, Meta = (EditCondition = "CollisionType == EExtractCollisionOutputType::Complex"))
+	bool bWeldEdges = true;
+
+	/** Whether or not to generate a seperate Mesh Object for each Simple Collision Shape  */
+	UPROPERTY(EditAnywhere, Category = Options, Meta = (EditCondition = "CollisionType == EExtractCollisionOutputType::Simple"))
+	bool bOutputSeparateMeshes = true;
+
+	/** Show/Hide a preview of the generated mesh (overlaps source mesh) */
+	UPROPERTY(EditAnywhere, Category = Options)
+	bool bShowPreview = false;
+
+	/** Show/Hide input mesh */
+	UPROPERTY(EditAnywhere, Category = Options)
+	bool bShowInputMesh = true;
+};
+
 
 
 /**
@@ -47,6 +82,13 @@ public:
 
 protected:
 
+	/** Property set for type of output object (StaticMesh, Volume, etc) */
+	UPROPERTY()
+	UCreateMeshObjectTypeProperties* OutputTypeProperties;
+
+	UPROPERTY()
+	UExtractCollisionToolProperties* Settings;
+
 	UPROPERTY()
 	UCollisionGeometryVisualizationProperties* VizSettings = nullptr;
 
@@ -63,9 +105,11 @@ protected:
 	// these are TSharedPtr because TPimplPtr cannot currently be added to a TArray?
 	TSharedPtr<FPhysicsDataCollection> PhysicsInfo;
 
+	TArray<UE::Geometry::FDynamicMesh3> CurrentMeshParts;
 	UE::Geometry::FDynamicMesh3 CurrentMesh;
 	bool bResultValid = false;
-	void RecalculateMesh();
+	void RecalculateMesh_Simple();
+	void RecalculateMesh_Complex();
 
 	bool bVisualizationDirty = false;
 	void UpdateVisualization();
