ModelingTools: remove UAssetBackedTarget requirement from various Tools so that they work on DynamicMeshActor/Components. The GetMapToSharedSourceData() function in UMultiSelectionTool now handles this internally.
#rnx

diff --git a/Engine/Plugins/Experimental/MeshModelingToolset/Source/MeshModelingTools/Private/BakeTransformTool.cpp b/Engine/Plugins/Experimental/MeshModelingToolset/Source/MeshModelingTools/Private/BakeTransformTool.cpp
--- a/Engine/Plugins/Experimental/MeshModelingToolset/Source/MeshModelingTools/Private/BakeTransformTool.cpp
+++ b/Engine/Plugins/Experimental/MeshModelingToolset/Source/MeshModelingTools/Private/BakeTransformTool.cpp
@@ -21,8 +21,8 @@
 #include "TargetInterfaces/MeshDescriptionCommitter.h"
 #include "TargetInterfaces/MeshDescriptionProvider.h"
 #include "TargetInterfaces/PrimitiveComponentBackedTarget.h"
-#include "TargetInterfaces/AssetBackedTarget.h"
 #include "ToolTargetManager.h"
+#include "ModelingToolTargetUtil.h"
 
 #include "ExplicitUseGeometryMathTypes.h"		// using UE::Geometry::(math types)
 using namespace UE::Geometry;
@@ -39,8 +39,7 @@ const FToolTargetTypeRequirements& UBakeTransformToolBuilder::GetTargetRequireme
 	static FToolTargetTypeRequirements TypeRequirements({
 		UMeshDescriptionCommitter::StaticClass(),
 		UMeshDescriptionProvider::StaticClass(),
-		UPrimitiveComponentBackedTarget::StaticClass(),
-		UAssetBackedTarget::StaticClass()
+		UPrimitiveComponentBackedTarget::StaticClass()
 		});
 	return TypeRequirements;
 }
@@ -102,8 +101,8 @@ void UBakeTransformTool::Setup()
 	bool bHasZeroScales = false;
 	for (int32 ComponentIdx = 0; ComponentIdx < Targets.Num(); ComponentIdx++)
 	{
-		IPrimitiveComponentBackedTarget* TargetComponent = TargetComponentInterface(ComponentIdx);
-		if (TargetComponent->GetWorldTransform().GetScale3D().GetAbsMin() < KINDA_SMALL_NUMBER)
+		FTransform Transform = (FTransform)UE::ToolTarget::GetLocalToWorldTransform(Targets[ComponentIdx]);
+		if (Transform.GetScale3D().GetAbsMin() < KINDA_SMALL_NUMBER)
 		{
 			bHasZeroScales = true;
 		}
@@ -138,9 +137,11 @@ void UBakeTransformTool::UpdateAssets()
 	// Make sure mesh descriptions are deserialized before we open transaction.
 	// This is to avoid potential stability issues related to creation/load of
 	// mesh descriptions inside a transaction.
+	// TODO: this may not be necessary anymore. Also may not be the most efficient
+	TArray<FMeshDescription> SourceMeshes;
 	for (int32 ComponentIdx = 0; ComponentIdx < Targets.Num(); ComponentIdx++)
 	{
-		TargetMeshProviderInterface(ComponentIdx)->GetMeshDescription();
+		UE::ToolTarget::GetMeshDescription(Targets[ComponentIdx]);
 	}
 
 	GetToolManager()->BeginUndoTransaction(LOCTEXT("BakeTransformToolTransactionName", "Bake Transforms"));
@@ -148,12 +149,11 @@ void UBakeTransformTool::UpdateAssets()
 	TArray<UE::Geometry::FTransform3d> BakedTransforms;
 	for (int32 ComponentIdx = 0; ComponentIdx < Targets.Num(); ComponentIdx++)
 	{
-		IPrimitiveComponentBackedTarget* TargetComponent = TargetComponentInterface(ComponentIdx);
-		IMeshDescriptionCommitter* TargetMeshCommitter = TargetMeshCommitterInterface(ComponentIdx);
-		UPrimitiveComponent* Component = TargetComponent->GetOwnerComponent();
+		UToolTarget* Target = Targets[ComponentIdx];
+		UPrimitiveComponent* Component = UE::ToolTarget::GetTargetComponent(Target);
 		Component->Modify();
 
-		UE::Geometry::FTransform3d ComponentToWorld(TargetComponent->GetWorldTransform());
+		UE::Geometry::FTransform3d ComponentToWorld = UE::ToolTarget::GetLocalToWorldTransform(Target);
 		UE::Geometry::FTransform3d ToBakePart = UE::Geometry::FTransform3d::Identity();
 		UE::Geometry::FTransform3d NewWorldPart = ComponentToWorld;
 
@@ -220,37 +220,47 @@ void UBakeTransformTool::UpdateAssets()
 				check(false); // must explicitly handle all cases
 			}
 
-			TargetMeshCommitter->CommitMeshDescription([this, &ToBakePart, &NewWorldPart](const IMeshDescriptionCommitter::FCommitterParams& CommitParams)
+			// apply edit
+			FMeshDescription SourceMesh(UE::ToolTarget::GetMeshDescriptionCopy(Targets[ComponentIdx], false));
+			FMeshDescriptionEditableTriangleMeshAdapter EditableMeshDescAdapter(&SourceMesh);
+
+			// do this part within the commit because we have the MeshDescription already computed
+			if (BasicProperties->bRecenterPivot)
 			{
-				FMeshDescriptionEditableTriangleMeshAdapter EditableMeshDescAdapter(CommitParams.MeshDescriptionOut);
+				FBox BBox = SourceMesh.ComputeBoundingBox();
+				FVector3d Center(BBox.GetCenter());
+				FFrame3d LocalFrame(Center);
+				ToBakePart.SetTranslation(ToBakePart.GetTranslation() - Center);
+				NewWorldPart.SetTranslation(NewWorldPart.GetTranslation() + NewWorldPart.TransformVector(Center));
+			}
 
-				// do this part within the commit because we have the MeshDescription already computed
-				if (BasicProperties->bRecenterPivot)
-				{
-					FBox BBox = CommitParams.MeshDescriptionOut->ComputeBoundingBox();
-					FVector3d Center(BBox.GetCenter());
-					FFrame3d LocalFrame(Center);
-					ToBakePart.SetTranslation(ToBakePart.GetTranslation() - Center);
-					NewWorldPart.SetTranslation(NewWorldPart.GetTranslation() + NewWorldPart.TransformVector(Center));
-				}
+			MeshAdapterTransforms::ApplyTransform(EditableMeshDescAdapter, ToBakePart);
 
-				MeshAdapterTransforms::ApplyTransform(EditableMeshDescAdapter, ToBakePart);
+			FVector3d BakeScaleVec = ToBakePart.GetScale();
+			if (BakeScaleVec.X * BakeScaleVec.Y * BakeScaleVec.Z < 0)
+			{
+				SourceMesh.ReverseAllPolygonFacing();
+			}
 
-				FVector3d ScaleVec = ToBakePart.GetScale();
-				if (ScaleVec.X * ScaleVec.Y * ScaleVec.Z < 0)
-				{
-					CommitParams.MeshDescriptionOut->ReverseAllPolygonFacing();
-				}
-			});
+			// todo: support vertex-only update
+			UE::ToolTarget::CommitMeshDescriptionUpdate(Target, &SourceMesh);
 
 			// try to transform simple collision
-			UE::Geometry::TransformSimpleCollision(Component, ToBakePart);
+			if (UE::Geometry::ComponentTypeSupportsCollision(Component))
+			{
+				UE::Geometry::TransformSimpleCollision(Component, ToBakePart);
+			}
 
 			BakedTransforms.Add(ToBakePart);
 		}
 
 		Component->SetWorldTransform((FTransform)NewWorldPart);
-		TargetComponent->GetOwnerActor()->MarkComponentsRenderStateDirty();
+
+		AActor* TargetActor = UE::ToolTarget::GetTargetActor(Target);
+		if (TargetActor)
+		{
+			TargetActor->MarkComponentsRenderStateDirty();
+		}
 	}
 
 	GetToolManager()->EndUndoTransaction();
diff --git a/Engine/Plugins/Experimental/MeshModelingToolset/Source/MeshModelingTools/Private/EditPivotTool.cpp b/Engine/Plugins/Experimental/MeshModelingToolset/Source/MeshModelingTools/Private/EditPivotTool.cpp
--- a/Engine/Plugins/Experimental/MeshModelingToolset/Source/MeshModelingTools/Private/EditPivotTool.cpp
+++ b/Engine/Plugins/Experimental/MeshModelingToolset/Source/MeshModelingTools/Private/EditPivotTool.cpp
@@ -24,8 +24,8 @@
 #include "TargetInterfaces/MeshDescriptionCommitter.h"
 #include "TargetInterfaces/MeshDescriptionProvider.h"
 #include "TargetInterfaces/PrimitiveComponentBackedTarget.h"
-#include "TargetInterfaces/AssetBackedTarget.h"
 #include "ToolTargetManager.h"
+#include "ModelingToolTargetUtil.h"
 
 #include "ExplicitUseGeometryMathTypes.h"		// using UE::Geometry::(math types)
 using namespace UE::Geometry;
@@ -42,8 +42,7 @@ const FToolTargetTypeRequirements& UEditPivotToolBuilder::GetTargetRequirements(
 	static FToolTargetTypeRequirements TypeRequirements({
 		UMeshDescriptionCommitter::StaticClass(),
 		UMeshDescriptionProvider::StaticClass(),
-		UPrimitiveComponentBackedTarget::StaticClass(),
-		UAssetBackedTarget::StaticClass()
+		UPrimitiveComponentBackedTarget::StaticClass()
 	});
 	return TypeRequirements;
 }
@@ -130,7 +129,7 @@ void UEditPivotTool::Setup()
 	bool bHasISMCs = false;
 	for (int32 k = 0; k < Targets.Num(); ++k)
 	{
-		if (Cast<UInstancedStaticMeshComponent>(TargetComponentInterface(k)->GetOwnerComponent()) != nullptr)
+		if (Cast<UInstancedStaticMeshComponent>(UE::ToolTarget::GetTargetComponent(Targets[k])) != nullptr)
 		{
 			bHasISMCs = true;
 		}
@@ -185,9 +184,9 @@ void UEditPivotTool::Precompute()
 	int NumComponents = Targets.Num();
 	if (NumComponents == 1)
 	{
-		Transform = UE::Geometry::FTransform3d(TargetComponentInterface(0)->GetWorldTransform());
+		Transform = UE::ToolTarget::GetLocalToWorldTransform(Targets[0]);
 
-		const FMeshDescription* Mesh = TargetMeshProviderInterface(0)->GetMeshDescription();
+		const FMeshDescription* Mesh = UE::ToolTarget::GetMeshDescription(Targets[0]);
 		VertexIteration(Mesh, [&](int32 VertexID, const FVector& Position) {
 			ObjectBounds.Contain((FVector3d)Position);
 			WorldBounds.Contain(Transform.TransformPosition((FVector3d)Position));
@@ -198,10 +197,8 @@ void UEditPivotTool::Precompute()
 		Transform = UE::Geometry::FTransform3d::Identity();
 		for (int k = 0; k < NumComponents; ++k)
 		{
-			IPrimitiveComponentBackedTarget* TargetComponent = TargetComponentInterface(k);
-			IMeshDescriptionProvider* TargetMeshProvider = TargetMeshProviderInterface(k);
-			UE::Geometry::FTransform3d CurTransform(TargetComponent->GetWorldTransform());
-			const FMeshDescription* Mesh = TargetMeshProvider->GetMeshDescription();
+			UE::Geometry::FTransform3d CurTransform = UE::ToolTarget::GetLocalToWorldTransform(Targets[k]);
+			const FMeshDescription* Mesh = UE::ToolTarget::GetMeshDescription(Targets[k]);
 			VertexIteration(Mesh, [&](int32 VertexID, const FVector& Position) {
 				ObjectBounds.Contain(CurTransform.TransformPosition((FVector3d)Position));
 				WorldBounds.Contain(CurTransform.TransformPosition((FVector3d)Position));
@@ -326,7 +323,7 @@ void UEditPivotTool::SetActiveGizmos_Single(bool bLocalRotations)
 
 	for (int32 ComponentIdx = 0; ComponentIdx < Targets.Num(); ComponentIdx++)
 	{
-		Transformable.TransformProxy->AddComponent(TargetComponentInterface(ComponentIdx)->GetOwnerComponent());
+		Transformable.TransformProxy->AddComponent(UE::ToolTarget::GetTargetComponent(Targets[ComponentIdx]));
 	}
 	Transformable.TransformGizmo = UE::TransformGizmoUtil::CreateCustomTransformGizmo(GizmoManager,
 		ETransformGizmoSubElements::StandardTranslateRotate, this
@@ -437,15 +434,14 @@ void UEditPivotTool::UpdateAssets(const FFrame3d& NewPivotWorldFrame)
 	TArray<FTransform> OriginalTransforms;
 	for (int32 ComponentIdx = 0; ComponentIdx < Targets.Num(); ComponentIdx++)
 	{
-		OriginalTransforms.Add(TargetComponentInterface(ComponentIdx)->GetWorldTransform());
+		OriginalTransforms.Add((FTransform)UE::ToolTarget::GetLocalToWorldTransform(Targets[ComponentIdx]));
 	}
 	for (int32 ComponentIdx = 0; ComponentIdx < Targets.Num(); ComponentIdx++)
 	{
-		IPrimitiveComponentBackedTarget* TargetComponent = TargetComponentInterface(ComponentIdx);
-		UPrimitiveComponent* Component = TargetComponent->GetOwnerComponent();
+		UPrimitiveComponent* Component = UE::ToolTarget::GetTargetComponent(Targets[ComponentIdx]);
 		Component->Modify();
 
-		UInstancedStaticMeshComponent* InstancedComponent = Cast<UInstancedStaticMeshComponent>(TargetComponent->GetOwnerComponent());
+		UInstancedStaticMeshComponent* InstancedComponent = Cast<UInstancedStaticMeshComponent>(Component);
 		if (InstancedComponent != nullptr)
 		{
 			// For ISMC, we will not bake in a mesh transform, instead we will update the instance transforms relative to the new pivot
@@ -475,13 +471,16 @@ void UEditPivotTool::UpdateAssets(const FFrame3d& NewPivotWorldFrame)
 			UE::Geometry::FTransform3d ToBake(OriginalTransforms[ComponentIdx] * NewWorldInverse);
 
 			// transform simple collision geometry
-			UE::Geometry::TransformSimpleCollision(Component, ToBake);
-
-			TargetMeshCommitterInterface(ComponentIdx)->CommitMeshDescription([&ToBake](const IMeshDescriptionCommitter::FCommitterParams& CommitParams)
+			if (UE::Geometry::ComponentTypeSupportsCollision(Component))
 			{
-				FMeshDescriptionEditableTriangleMeshAdapter EditableMeshDescAdapter(CommitParams.MeshDescriptionOut);
-				MeshAdapterTransforms::ApplyTransform(EditableMeshDescAdapter, ToBake);
-			});
+				UE::Geometry::TransformSimpleCollision(Component, ToBake);
+			}
+
+			FMeshDescription SourceMesh(UE::ToolTarget::GetMeshDescriptionCopy(Targets[ComponentIdx], false));
+			FMeshDescriptionEditableTriangleMeshAdapter EditableMeshDescAdapter(&SourceMesh);
+			MeshAdapterTransforms::ApplyTransform(EditableMeshDescAdapter, ToBake);
+			// todo: support vertex-only update
+			UE::ToolTarget::CommitMeshDescriptionUpdate(Targets[ComponentIdx], &SourceMesh);
 
 			Component->SetWorldTransform(NewWorldTransform);
 		}
@@ -491,8 +490,13 @@ void UEditPivotTool::UpdateAssets(const FFrame3d& NewPivotWorldFrame)
 			FTransform Baked = OriginalTransforms[MapToFirstOccurrences[ComponentIdx]] * NewWorldInverse;
 			Component->SetWorldTransform(Baked.Inverse() * OriginalTransforms[ComponentIdx]);
 		}
-		TargetComponent->GetOwnerActor()->MarkComponentsRenderStateDirty();
-		TargetComponent->GetOwnerActor()->UpdateComponentTransforms();
+
+		AActor* OwnerActor = UE::ToolTarget::GetTargetActor(Targets[ComponentIdx]);
+		if (OwnerActor)
+		{
+			OwnerActor->MarkComponentsRenderStateDirty();
+			OwnerActor->UpdateComponentTransforms();
+		}
 	}
 
 	// hack to ensure user sees the updated pivot immediately: request re-select of the original selection
@@ -500,8 +504,11 @@ void UEditPivotTool::UpdateAssets(const FFrame3d& NewPivotWorldFrame)
 	NewSelection.ModificationType = ESelectedObjectsModificationType::Replace;
 	for (int OrigMeshIdx = 0; OrigMeshIdx < Targets.Num(); OrigMeshIdx++)
 	{
-		IPrimitiveComponentBackedTarget* TargetComponent = TargetComponentInterface(OrigMeshIdx);
-		NewSelection.Actors.Add(TargetComponent->GetOwnerActor());
+		AActor* OwnerActor = UE::ToolTarget::GetTargetActor(Targets[OrigMeshIdx]);
+		if (OwnerActor)
+		{
+			NewSelection.Actors.Add(OwnerActor);
+		}
 	}
 	GetToolManager()->RequestSelectionChange(NewSelection);
 
diff --git a/Engine/Plugins/Experimental/MeshModelingToolset/Source/MeshModelingTools/Private/MirrorTool.cpp b/Engine/Plugins/Experimental/MeshModelingToolset/Source/MeshModelingTools/Private/MirrorTool.cpp
--- a/Engine/Plugins/Experimental/MeshModelingToolset/Source/MeshModelingTools/Private/MirrorTool.cpp
+++ b/Engine/Plugins/Experimental/MeshModelingToolset/Source/MeshModelingTools/Private/MirrorTool.cpp
@@ -18,7 +18,6 @@
 #include "TargetInterfaces/MeshDescriptionCommitter.h"
 #include "TargetInterfaces/MeshDescriptionProvider.h"
 #include "TargetInterfaces/PrimitiveComponentBackedTarget.h"
-#include "TargetInterfaces/AssetBackedTarget.h"
 #include "ToolTargetManager.h"
 
 #include "ExplicitUseGeometryMathTypes.h"		// using UE::Geometry::(math types)
@@ -35,8 +34,7 @@ const FToolTargetTypeRequirements& UMirrorToolBuilder::GetTargetRequirements() c
 		UMaterialProvider::StaticClass(),
 		UMeshDescriptionCommitter::StaticClass(),
 		UMeshDescriptionProvider::StaticClass(),
-		UPrimitiveComponentBackedTarget::StaticClass(),
-		UAssetBackedTarget::StaticClass()
+		UPrimitiveComponentBackedTarget::StaticClass()
 		});
 	return TypeRequirements;
 }
diff --git a/Engine/Plugins/Experimental/MeshModelingToolset/Source/MeshModelingTools/Private/RemoveOccludedTrianglesTool.cpp b/Engine/Plugins/Experimental/MeshModelingToolset/Source/MeshModelingTools/Private/RemoveOccludedTrianglesTool.cpp
--- a/Engine/Plugins/Experimental/MeshModelingToolset/Source/MeshModelingTools/Private/RemoveOccludedTrianglesTool.cpp
+++ b/Engine/Plugins/Experimental/MeshModelingToolset/Source/MeshModelingTools/Private/RemoveOccludedTrianglesTool.cpp
@@ -27,7 +27,6 @@
 #include "TargetInterfaces/MeshDescriptionCommitter.h"
 #include "TargetInterfaces/MeshDescriptionProvider.h"
 #include "TargetInterfaces/PrimitiveComponentBackedTarget.h"
-#include "TargetInterfaces/AssetBackedTarget.h"
 #include "ToolTargetManager.h"
 
 #include "ExplicitUseGeometryMathTypes.h"		// using UE::Geometry::(math types)
@@ -46,8 +45,7 @@ const FToolTargetTypeRequirements& URemoveOccludedTrianglesToolBuilder::GetTarge
 		UMaterialProvider::StaticClass(),
 		UMeshDescriptionCommitter::StaticClass(),
 		UMeshDescriptionProvider::StaticClass(),
-		UPrimitiveComponentBackedTarget::StaticClass(),
-		UAssetBackedTarget::StaticClass()
+		UPrimitiveComponentBackedTarget::StaticClass()
 		});
 	return TypeRequirements;
 }
