ModelingTools: fix TObjectPtr static analysis issue
#rnx

diff --git a/Engine/Plugins/Experimental/MeshModelingToolset/Source/MeshModelingTools/Private/Physics/ExtractCollisionGeometryTool.cpp b/Engine/Plugins/Experimental/MeshModelingToolset/Source/MeshModelingTools/Private/Physics/ExtractCollisionGeometryTool.cpp
--- a/Engine/Plugins/Experimental/MeshModelingToolset/Source/MeshModelingTools/Private/Physics/ExtractCollisionGeometryTool.cpp
+++ b/Engine/Plugins/Experimental/MeshModelingToolset/Source/MeshModelingTools/Private/Physics/ExtractCollisionGeometryTool.cpp
@@ -97,7 +97,7 @@ void UExtractCollisionGeometryTool::Setup()
 
 	UPrimitiveComponent* TargetComponent = UE::ToolTarget::GetTargetComponent(Target);
 	UStaticMeshComponent* StaticMeshComponent = (TargetComponent) ? Cast<UStaticMeshComponent>(TargetComponent) : nullptr;
-	UStaticMesh* StaticMesh = (StaticMeshComponent) ? StaticMeshComponent->GetStaticMesh() : nullptr;
+	TObjectPtr<UStaticMesh> StaticMesh = (StaticMeshComponent) ? StaticMeshComponent->GetStaticMesh() : nullptr;
 	UBodySetup* BodySetup = (StaticMesh) ? StaticMesh->GetBodySetup() : nullptr;
 	if (ensure(BodySetup))
 	{
@@ -361,7 +361,7 @@ void UExtractCollisionGeometryTool::RecalculateMesh_Complex()
 	bool bMeshErrors = false;
 
 	UStaticMeshComponent* StaticMeshComponent = Cast<UStaticMeshComponent>(UE::ToolTarget::GetTargetComponent(Target));
-	UStaticMesh* StaticMesh = (StaticMeshComponent) ? StaticMeshComponent->GetStaticMesh() : nullptr;
+	TObjectPtr<UStaticMesh> StaticMesh = (StaticMeshComponent) ? StaticMeshComponent->GetStaticMesh() : nullptr;
 	IInterface_CollisionDataProvider* CollisionProvider = (StaticMesh) ? Cast<IInterface_CollisionDataProvider>(StaticMesh) : nullptr;
 	if (CollisionProvider && CollisionProvider->ContainsPhysicsTriMeshData(true))
 	{
