ModelingTools: Initialized some uninitialized struct members to fix a startup warning (soon to be an error)

diff --git a/Engine/Plugins/Experimental/MeshModelingToolset/Source/ModelingComponents/Public/ModelingObjectsCreationAPI.h b/Engine/Plugins/Experimental/MeshModelingToolset/Source/ModelingComponents/Public/ModelingObjectsCreationAPI.h
--- a/Engine/Plugins/Experimental/MeshModelingToolset/Source/ModelingComponents/Public/ModelingObjectsCreationAPI.h
+++ b/Engine/Plugins/Experimental/MeshModelingToolset/Source/ModelingComponents/Public/ModelingObjectsCreationAPI.h
@@ -183,7 +183,7 @@ struct MODELINGCOMPONENTS_API FCreateMeshObjectResult
 
 	/** Success/Failure status for the requested operation */
 	UPROPERTY(Category = "CreateMeshObjectResult", VisibleAnywhere)
-	ECreateModelingObjectResult ResultCode;
+	ECreateModelingObjectResult ResultCode = ECreateModelingObjectResult::Ok;
 
 	/** A pointer to a newly-created Actor for the mesh object, if applicable (eg StaticMeshActor) */
 	UPROPERTY(Category="CreateMeshObjectResult", VisibleAnywhere)
@@ -270,7 +270,7 @@ struct MODELINGCOMPONENTS_API FCreateTextureObjectResult
 
 	/** Success/Failure status for the requested operation */
 	UPROPERTY(Category = "CreateTextureObjectResult", VisibleAnywhere)
-	ECreateModelingObjectResult ResultCode;
+	ECreateModelingObjectResult ResultCode = ECreateModelingObjectResult::Ok;
 
 	/** A pointer to a newly-created Asset for the texture object, if applicable (eg StaticMeshAsset) */
 	UPROPERTY(Category = "CreateTextureObjectResult", VisibleAnywhere)
