ModelingTools: DynamicMeshActor support Tools and Settings.

- Move UCreateMeshObjectTypeProperties to ModelingComponents and update all include sites
- Added Editor Modeling Mode settings to Enable/Disable showing DynamicMeshActor creation options, and select default mesh object type. Removed CVarEnableDynamicMeshActors.
- Added optional 'Auto' output mesh type to UCreateMeshObjectTypeProperties. This can be used in Tools that have input objects and want to allow optional conversion, but default to 'just use input mesh object type'.
- Added ConvertMeshesTool, this does in-place conversion between Mesh Object types for a set of selected objects (Duplicate tool can also do this, but only for a single object, and functionality is expected to further diverge)
- Added SplitMeshesTool, decomposes a mesh into parts and creates a new output object for each part
- CombineMeshesTool now supports variable output object type. Cleaned up internals.

#rnx

diff --git a/Engine/Plugins/Experimental/MeshModelingToolset/Source/MeshModelingTools/Private/CombineMeshesTool.cpp b/Engine/Plugins/Experimental/MeshModelingToolset/Source/MeshModelingTools/Private/CombineMeshesTool.cpp
--- a/Engine/Plugins/Experimental/MeshModelingToolset/Source/MeshModelingTools/Private/CombineMeshesTool.cpp
+++ b/Engine/Plugins/Experimental/MeshModelingToolset/Source/MeshModelingTools/Private/CombineMeshesTool.cpp
@@ -3,18 +3,12 @@
 #include "CombineMeshesTool.h"
 #include "InteractiveToolManager.h"
 #include "ToolBuilderUtil.h"
-
 #include "ToolSetupUtil.h"
 
 #include "DynamicMesh/DynamicMesh3.h"
-#include "BaseBehaviors/MultiClickSequenceInputBehavior.h"
-#include "Selection/SelectClickedAction.h"
 #include "DynamicMeshEditor.h"
 #include "DynamicMesh/MeshTransforms.h"
 
-#include "MeshDescriptionToDynamicMesh.h"
-#include "DynamicMeshToMeshDescription.h"
-
 #include "ModelingObjectsCreationAPI.h"
 #include "Selection/ToolSelectionUtil.h"
 #include "Physics/ComponentCollisionUtil.h"
@@ -100,23 +94,35 @@ void UCombineMeshesTool::Setup()
 	AddToolPropertySource(BasicProperties);
 	BasicProperties->RestoreProperties(this);
 	BasicProperties->bIsDuplicateMode = this->bDuplicateMode;
+
+	OutputTypeProperties = NewObject<UCreateMeshObjectTypeProperties>(this);
+	OutputTypeProperties->InitializeDefaultWithAuto();
+	OutputTypeProperties->OutputType = UCreateMeshObjectTypeProperties::AutoIdentifier;
+	OutputTypeProperties->RestoreProperties(this, TEXT("OutputTypeFromInputTool"));
+	OutputTypeProperties->WatchProperty(OutputTypeProperties->OutputType, [this](FString) { OutputTypeProperties->UpdatePropertyVisibility(); });
+	AddToolPropertySource(OutputTypeProperties);
+
 	BasicProperties->WatchProperty(BasicProperties->WriteOutputTo, [&](ECombineTargetType NewType)
 	{
 		if (NewType == ECombineTargetType::NewAsset)
 		{
 			BasicProperties->OutputAsset = TEXT("");
+			SetToolPropertySourceEnabled(OutputTypeProperties, true);
 		}
 		else
 		{
 			int32 Index = (BasicProperties->WriteOutputTo == ECombineTargetType::FirstInputAsset) ? 0 : Targets.Num() - 1;
-			BasicProperties->OutputAsset = UE::Modeling::GetComponentAssetBaseName(TargetComponentInterface(Index)->GetOwnerComponent(), false);
+			BasicProperties->OutputAsset = UE::Modeling::GetComponentAssetBaseName(UE::ToolTarget::GetTargetComponent(Targets[Index]), false);
+			SetToolPropertySourceEnabled(OutputTypeProperties, false);
 		}
 	});
 
+	SetToolPropertySourceEnabled(OutputTypeProperties, BasicProperties->WriteOutputTo == ECombineTargetType::NewAsset);
+
 	if (bDuplicateMode)
 	{
 		SetToolDisplayName(LOCTEXT("DuplicateMeshesToolName", "Duplicate"));
-		BasicProperties->OutputName = UE::Modeling::GetComponentAssetBaseName(TargetComponentInterface(0)->GetOwnerComponent());
+		BasicProperties->OutputName = UE::Modeling::GetComponentAssetBaseName(UE::ToolTarget::GetTargetComponent(Targets[0]));
 	}
 	else
 	{
@@ -149,6 +155,7 @@ void UCombineMeshesTool::Setup()
 void UCombineMeshesTool::Shutdown(EToolShutdownType ShutdownType)
 {
 	BasicProperties->SaveProperties(this);
+	OutputTypeProperties->SaveProperties(this, TEXT("OutputTypeFromInputTool"));
 	HandleSourceProperties->SaveProperties(this);
 
 	if (ShutdownType == EToolShutdownType::Accept)
@@ -166,59 +173,33 @@ void UCombineMeshesTool::Shutdown(EToolShutdownType ShutdownType)
 
 void UCombineMeshesTool::CreateNewAsset()
 {
-	// Make sure mesh descriptions are deserialized before we open transaction.
-	// This is to avoid potential stability issues related to creation/load of
-	// mesh descriptions inside a transaction.
-	TArray<const FMeshDescription*> MeshDescriptions;
+	// Make sure meshes are available before we open transaction. This is to avoid potential stability issues related 
+	// to creation/load of meshes inside a transaction, for assets that possibly do not have bulk data currently loaded.
+	TArray<FDynamicMesh3> InputMeshes;
+	InputMeshes.Reserve(Targets.Num());
 	for (int32 ComponentIdx = 0; ComponentIdx < Targets.Num(); ComponentIdx++)
 	{
-		MeshDescriptions.Add(TargetMeshProviderInterface(ComponentIdx)->GetMeshDescription());
+		InputMeshes.Add(UE::ToolTarget::GetDynamicMeshCopy(Targets[ComponentIdx], true));
 	}
 
-	GetToolManager()->BeginUndoTransaction(
-		bDuplicateMode ? 
+	GetToolManager()->BeginUndoTransaction( bDuplicateMode ? 
 		LOCTEXT("DuplicateMeshToolTransactionName", "Duplicate Mesh") :
 		LOCTEXT("CombineMeshesToolTransactionName", "Combine Meshes"));
 
-	// note there is a MergeMeshUtilities.h w/ a very feature-filled mesh merging class, but for simplicity (and to fit modeling tool needs)
-	// this tool currently converts things through dynamic mesh instead
 	FBox Box(ForceInit);
 	for (int32 ComponentIdx = 0; ComponentIdx < Targets.Num(); ComponentIdx++)
 	{
-		Box += TargetComponentInterface(ComponentIdx)->GetOwnerComponent()->Bounds.GetBox();
+		Box += UE::ToolTarget::GetTargetComponent(Targets[ComponentIdx])->Bounds.GetBox();
 	}
 
-	bool bMergeSameMaterials = true;
 	TArray<UMaterialInterface*> AllMaterials;
-	TMap<UMaterialInterface*, int> KnownMaterials;
-	TArray<int> CombinedMatToOutMatIdx;
-	for (int32 ComponentIdx = 0; ComponentIdx < Targets.Num(); ComponentIdx++)
-	{
-		IMaterialProvider* ComponentMaterial = TargetMaterialInterface(ComponentIdx);
-		for (int MaterialIdx = 0, NumMaterials = ComponentMaterial->GetNumMaterials(); MaterialIdx < NumMaterials; MaterialIdx++)
-		{
-			UMaterialInterface* Mat = ComponentMaterial->GetMaterial(MaterialIdx);
-			int32 OutMatIdx = -1;
-			if (!bMergeSameMaterials || !KnownMaterials.Contains(Mat))
-			{
-				OutMatIdx = AllMaterials.Num();
-				if (bMergeSameMaterials)
-				{
-					KnownMaterials.Add(Mat, AllMaterials.Num());
-				}
-				AllMaterials.Add(Mat);
-			}
-			else
-			{
-				OutMatIdx = KnownMaterials[Mat];
-			}
-			CombinedMatToOutMatIdx.Add(OutMatIdx);
-		}
-	}
+	TArray<TArray<int32>> MaterialIDRemaps;
+	BuildCombinedMaterialSet(AllMaterials, MaterialIDRemaps);
 
 	FDynamicMesh3 AccumulateDMesh;
 	AccumulateDMesh.EnableTriangleGroups();
 	AccumulateDMesh.EnableAttributes();
+	AccumulateDMesh.Attributes()->EnableTangents();
 	AccumulateDMesh.Attributes()->EnableMaterialID();
 	AccumulateDMesh.Attributes()->EnablePrimaryColors();
 	FTransform AccumToWorld(Box.GetCenter());
@@ -242,13 +223,9 @@ void UCombineMeshesTool::CreateNewAsset()
 #if WITH_EDITOR
 			SlowTask.EnterProgressFrame(1);
 #endif
-			IPrimitiveComponentBackedTarget* TargetComponent = TargetComponentInterface(ComponentIdx);
-			IMeshDescriptionProvider* TargetMeshProvider = TargetMeshProviderInterface(ComponentIdx);
+			UPrimitiveComponent* PrimitiveComponent = UE::ToolTarget::GetTargetComponent(Targets[ComponentIdx]);
 
-			FMeshDescriptionToDynamicMesh Converter;
-			FDynamicMesh3 ComponentDMesh;
-			const FMeshDescription* MeshDescription = MeshDescriptions[ComponentIdx];
-			Converter.Convert(MeshDescription, ComponentDMesh);
+			FDynamicMesh3& ComponentDMesh = InputMeshes[ComponentIdx];
 			bNeedColorAttr = bNeedColorAttr || (ComponentDMesh.HasAttributes() && ComponentDMesh.Attributes()->HasPrimaryColors());
 
 			if (ComponentDMesh.HasAttributes() && ComponentDMesh.Attributes()->NumUVLayers() > AccumulateDMesh.Attributes()->NumUVLayers())
@@ -256,7 +233,7 @@ void UCombineMeshesTool::CreateNewAsset()
 				AccumulateDMesh.Attributes()->SetNumUVLayers(ComponentDMesh.Attributes()->NumUVLayers());
 			}
 
-			UE::Geometry::FTransform3d XF = (UE::Geometry::FTransform3d)(TargetComponentInterface(ComponentIdx)->GetWorldTransform() * ToAccum);
+			UE::Geometry::FTransform3d XF = (UE::Geometry::FTransform3d)((FTransform)UE::ToolTarget::GetLocalToWorldTransform(Targets[ComponentIdx]) * ToAccum);
 			if (XF.GetDeterminant() < 0)
 			{
 				ComponentDMesh.ReverseOrientation(false);
@@ -266,7 +243,8 @@ void UCombineMeshesTool::CreateNewAsset()
 			FDynamicMeshMaterialAttribute* MatAttrib = ComponentDMesh.Attributes()->GetMaterialID();
 			for (int TID : ComponentDMesh.TriangleIndicesItr())
 			{
-				MatAttrib->SetValue(TID, CombinedMatToOutMatIdx[MatIndexBase + MatAttrib->GetValue(TID)]);
+				int MatID = MatAttrib->GetValue(TID);
+				MatAttrib->SetValue(TID, MaterialIDRemaps[ComponentIdx][MatID]);
 			}
 
 			FDynamicMeshEditor Editor(&AccumulateDMesh);
@@ -274,10 +252,11 @@ void UCombineMeshesTool::CreateNewAsset()
 			if (bDuplicateMode) // no transform if duplicating
 			{
 				Editor.AppendMesh(&ComponentDMesh, IndexMapping);
-				if (UE::Geometry::ComponentTypeSupportsCollision(TargetComponent->GetOwnerComponent()))
+
+				if (UE::Geometry::ComponentTypeSupportsCollision(PrimitiveComponent))
 				{
-					CollisionSettings = UE::Geometry::GetCollisionSettings(TargetComponent->GetOwnerComponent());
-					UE::Geometry::AppendSimpleCollision(TargetComponent->GetOwnerComponent(), &SimpleCollision, UE::Geometry::FTransform3d::Identity());
+					CollisionSettings = UE::Geometry::GetCollisionSettings(PrimitiveComponent);
+					UE::Geometry::AppendSimpleCollision(PrimitiveComponent, &SimpleCollision, UE::Geometry::FTransform3d::Identity());
 				}
 			}
 			else
@@ -285,13 +264,14 @@ void UCombineMeshesTool::CreateNewAsset()
 				Editor.AppendMesh(&ComponentDMesh, IndexMapping,
 					[&XF](int Unused, const FVector3d P) { return XF.TransformPosition(P); },
 					[&XF](int Unused, const FVector3d N) { return XF.TransformNormal(N); });
-				if (UE::Geometry::ComponentTypeSupportsCollision(TargetComponent->GetOwnerComponent()))
+				if (UE::Geometry::ComponentTypeSupportsCollision(PrimitiveComponent))
 				{
-					UE::Geometry::AppendSimpleCollision(TargetComponent->GetOwnerComponent(), &SimpleCollision, XF);
+					UE::Geometry::AppendSimpleCollision(PrimitiveComponent, &SimpleCollision, XF);
 				}
 			}
 
-			MatIndexBase += TargetMaterialInterface(ComponentIdx)->GetNumMaterials();
+			FComponentMaterialSet MaterialSet = UE::ToolTarget::GetMaterialSet(Targets[ComponentIdx]);
+			MatIndexBase += MaterialSet.Materials.Num();
 		}
 
 		if (!bNeedColorAttr)
@@ -307,7 +287,7 @@ void UCombineMeshesTool::CreateNewAsset()
 		{
 			// TODO: will need to refactor this when we support duplicating multiple
 			check(Targets.Num() == 1);
-			AccumToWorld = TargetComponentInterface(0)->GetWorldTransform();
+			AccumToWorld = (FTransform)UE::ToolTarget::GetLocalToWorldTransform(Targets[0]);
 		}
 
 		// max len explicitly enforced here, would ideally notify user
@@ -323,7 +303,14 @@ void UCombineMeshesTool::CreateNewAsset()
 		NewMeshObjectParams.BaseName = UseBaseName;
 		NewMeshObjectParams.Materials = AllMaterials;
 		NewMeshObjectParams.SetMesh(&AccumulateDMesh);
-		UE::ToolTarget::ConfigureCreateMeshObjectParams(Targets[0], NewMeshObjectParams);
+		if (OutputTypeProperties->OutputType == UCreateMeshObjectTypeProperties::AutoIdentifier)
+		{
+			UE::ToolTarget::ConfigureCreateMeshObjectParams(Targets[0], NewMeshObjectParams);
+		}
+		else
+		{
+			OutputTypeProperties->ConfigureCreateMeshObjectParams(NewMeshObjectParams);
+		}
 		FCreateMeshObjectResult Result = UE::Modeling::CreateMeshObject(GetToolManager(), MoveTemp(NewMeshObjectParams));
 		if (Result.IsOK() && Result.NewActor != nullptr)
 		{
@@ -341,7 +328,7 @@ void UCombineMeshesTool::CreateNewAsset()
 	TArray<AActor*> Actors;
 	for (int32 Idx = 0; Idx < Targets.Num(); Idx++)
 	{
-		Actors.Add(TargetComponentInterface(Idx)->GetOwnerActor());
+		Actors.Add(UE::ToolTarget::GetTargetActor(Targets[Idx]));
 	}
 	HandleSourceProperties->ApplyMethod(Actors, GetToolManager());
 
@@ -355,74 +342,46 @@ void UCombineMeshesTool::CreateNewAsset()
 
 
 
-
-
-
 void UCombineMeshesTool::UpdateExistingAsset()
 {
-	TArray<const FMeshDescription*> MeshDescriptions;
+	// Make sure meshes are available before we open transaction. This is to avoid potential stability issues related 
+	// to creation/load of meshes inside a transaction, for assets that possibly do not have bulk data currently loaded.
+	TArray<FDynamicMesh3> InputMeshes;
+	InputMeshes.Reserve(Targets.Num());
 	for (int32 ComponentIdx = 0; ComponentIdx < Targets.Num(); ComponentIdx++)
 	{
-		MeshDescriptions.Add(TargetMeshProviderInterface(ComponentIdx)->GetMeshDescription());
+		InputMeshes.Add(UE::ToolTarget::GetDynamicMeshCopy(Targets[ComponentIdx], true));
 	}
 
 	check(!bDuplicateMode);
 	GetToolManager()->BeginUndoTransaction(LOCTEXT("CombineMeshesToolTransactionName", "Combine Meshes"));
 
-	// note there is a MergeMeshUtilities.h w/ a very feature-filled mesh merging class, but for simplicity (and to fit modeling tool needs)
-	// this tool currently converts things through dynamic mesh instead
-
 	AActor* SkipActor = nullptr;
 
-	bool bMergeSameMaterials = true;
 	TArray<UMaterialInterface*> AllMaterials;
-	TMap<UMaterialInterface*, int> KnownMaterials;
-	TArray<int> CombinedMatToOutMatIdx;
-	for (int32 ComponentIdx = 0; ComponentIdx < Targets.Num(); ComponentIdx++)
-	{
-		IMaterialProvider* TargetMaterial = TargetMaterialInterface(ComponentIdx);
-		for (int MaterialIdx = 0, NumMaterials = TargetMaterial->GetNumMaterials(); MaterialIdx < NumMaterials; MaterialIdx++)
-		{
-			UMaterialInterface* Mat = TargetMaterial->GetMaterial(MaterialIdx);
-			int32 OutMatIdx = -1;
-			if (!bMergeSameMaterials || !KnownMaterials.Contains(Mat))
-			{
-				OutMatIdx = AllMaterials.Num();
-				if (bMergeSameMaterials)
-				{
-					KnownMaterials.Add(Mat, AllMaterials.Num());
-				}
-				AllMaterials.Add(Mat);
-			}
-			else
-			{
-				OutMatIdx = KnownMaterials[Mat];
-			}
-			CombinedMatToOutMatIdx.Add(OutMatIdx);
-		}
-	}
+	TArray<TArray<int32>> MaterialIDRemaps;
+	BuildCombinedMaterialSet(AllMaterials, MaterialIDRemaps);
 
 	FDynamicMesh3 AccumulateDMesh;
 	AccumulateDMesh.EnableTriangleGroups();
 	AccumulateDMesh.EnableAttributes();
+	AccumulateDMesh.Attributes()->EnableTangents();
 	AccumulateDMesh.Attributes()->EnableMaterialID();
 	AccumulateDMesh.Attributes()->EnablePrimaryColors();
 
 	int32 SkipIndex = (BasicProperties->WriteOutputTo == ECombineTargetType::FirstInputAsset) ? 0 : (Targets.Num() - 1);
-	IPrimitiveComponentBackedTarget* UpdateTarget = TargetComponentInterface(SkipIndex);
-	IMeshDescriptionCommitter* UpdateTargetCommitter = TargetMeshCommitterInterface(SkipIndex);
-	IMaterialProvider* UpdateTargetMaterial = TargetMaterialInterface(SkipIndex);
-	SkipActor = UpdateTarget->GetOwnerActor();
+	UPrimitiveComponent* UpdateComponent = UE::ToolTarget::GetTargetComponent(Targets[SkipIndex]);
+	SkipActor = UE::ToolTarget::GetTargetActor(Targets[SkipIndex]);
 
-	UE::Geometry::FTransform3d TargetToWorld = (UE::Geometry::FTransform3d)UpdateTarget->GetWorldTransform();
+	UE::Geometry::FTransform3d TargetToWorld = (UE::Geometry::FTransform3d)UE::ToolTarget::GetLocalToWorldTransform(Targets[SkipIndex]);
 	UE::Geometry::FTransform3d WorldToTarget = TargetToWorld.Inverse();
 
 	FSimpleShapeSet3d SimpleCollision;
 	UE::Geometry::FComponentCollisionSettings CollisionSettings;
-	bool bSupportsCollision = UE::Geometry::ComponentTypeSupportsCollision(UpdateTarget->GetOwnerComponent());
-	if (bSupportsCollision)
+	bool bOutputComponentSupportsCollision = UE::Geometry::ComponentTypeSupportsCollision(UpdateComponent);
+	if (bOutputComponentSupportsCollision)
 	{
-		CollisionSettings = UE::Geometry::GetCollisionSettings(UpdateTarget->GetOwnerComponent());
+		CollisionSettings = UE::Geometry::GetCollisionSettings(UpdateComponent);
 	}
 	TArray<UE::Geometry::FTransform3d> Transforms;
 	Transforms.SetNum(2);
@@ -436,32 +395,27 @@ void UCombineMeshesTool::UpdateExistingAsset()
 		SlowTask.MakeDialog();
 #endif
 		bool bNeedColorAttr = false;
-		int MatIndexBase = 0;
 		for (int32 ComponentIdx = 0; ComponentIdx < Targets.Num(); ComponentIdx++)
 		{
 #if WITH_EDITOR
 			SlowTask.EnterProgressFrame(1);
 #endif
+			UPrimitiveComponent* PrimitiveComponent = UE::ToolTarget::GetTargetComponent(Targets[ComponentIdx]);
 
-			IPrimitiveComponentBackedTarget* TargetComponent = TargetComponentInterface(ComponentIdx);
-
-			FMeshDescriptionToDynamicMesh Converter;
-			FDynamicMesh3 ComponentDMesh;
-			Converter.Convert(MeshDescriptions[ComponentIdx], ComponentDMesh);
+			FDynamicMesh3& ComponentDMesh = InputMeshes[ComponentIdx];
 			bNeedColorAttr = bNeedColorAttr || (ComponentDMesh.HasAttributes() && ComponentDMesh.Attributes()->HasPrimaryColors());
 
 			// update material IDs to account for combined material set
 			FDynamicMeshMaterialAttribute* MatAttrib = ComponentDMesh.Attributes()->GetMaterialID();
 			for (int TID : ComponentDMesh.TriangleIndicesItr())
 			{
-				MatAttrib->SetValue(TID, CombinedMatToOutMatIdx[MatIndexBase + MatAttrib->GetValue(TID)]);
+				int MatID = MatAttrib->GetValue(TID);
+				MatAttrib->SetValue(TID, MaterialIDRemaps[ComponentIdx][MatID]);
 			}
-			MatIndexBase += TargetMaterialInterface(ComponentIdx)->GetNumMaterials();
-
 
 			if (ComponentIdx != SkipIndex)
 			{
-				UE::Geometry::FTransform3d ComponentToWorld = (UE::Geometry::FTransform3d)TargetComponent->GetWorldTransform();
+				UE::Geometry::FTransform3d ComponentToWorld = (UE::Geometry::FTransform3d)UE::ToolTarget::GetLocalToWorldTransform(Targets[ComponentIdx]);
 				MeshTransforms::ApplyTransform(ComponentDMesh, ComponentToWorld);
 				if (ComponentToWorld.GetDeterminant() < 0)
 				{
@@ -474,16 +428,16 @@ void UCombineMeshesTool::UpdateExistingAsset()
 				}
 				Transforms[0] = ComponentToWorld;
 				Transforms[1] = WorldToTarget;
-				if (bSupportsCollision)
+				if (bOutputComponentSupportsCollision && UE::Geometry::ComponentTypeSupportsCollision(PrimitiveComponent))
 				{
-					UE::Geometry::AppendSimpleCollision(TargetComponent->GetOwnerComponent(), &SimpleCollision, Transforms);
+					UE::Geometry::AppendSimpleCollision(PrimitiveComponent, &SimpleCollision, Transforms);
 				}
 			}
 			else
 			{
-				if (bSupportsCollision)
+				if (bOutputComponentSupportsCollision && UE::Geometry::ComponentTypeSupportsCollision(PrimitiveComponent))
 				{
-					UE::Geometry::AppendSimpleCollision(TargetComponent->GetOwnerComponent(), &SimpleCollision, UE::Geometry::FTransform3d::Identity());
+					UE::Geometry::AppendSimpleCollision(PrimitiveComponent, &SimpleCollision, UE::Geometry::FTransform3d::Identity());
 				}
 			}
 
@@ -501,31 +455,24 @@ void UCombineMeshesTool::UpdateExistingAsset()
 		SlowTask.EnterProgressFrame(1);
 #endif
 
-		UpdateTargetCommitter->CommitMeshDescription([&](const IMeshDescriptionCommitter::FCommitterParams& CommitParams)
-		{
-			FDynamicMeshToMeshDescription Converter;
-			Converter.Convert(&AccumulateDMesh, *CommitParams.MeshDescriptionOut);
-		});
+		FComponentMaterialSet NewMaterialSet;
+		NewMaterialSet.Materials = AllMaterials;
+		UE::ToolTarget::CommitDynamicMeshUpdate(Targets[SkipIndex], AccumulateDMesh, true, FConversionToMeshDescriptionOptions(), &NewMaterialSet);
 
-		if (bSupportsCollision)
+		if (bOutputComponentSupportsCollision)
 		{
-			UE::Geometry::SetSimpleCollision(UpdateTarget->GetOwnerComponent(), &SimpleCollision, CollisionSettings);
+			UE::Geometry::SetSimpleCollision(UpdateComponent, &SimpleCollision, CollisionSettings);
 		}
 
-		FComponentMaterialSet MaterialSet;
-		MaterialSet.Materials = AllMaterials;
-		UpdateTargetMaterial->CommitMaterialSetUpdate(MaterialSet, true);
-
 		// select the new actor
 		ToolSelectionUtil::SetNewActorSelection(GetToolManager(), SkipActor);
 	}
 
-
 	
 	TArray<AActor*> Actors;
 	for (int Idx = 0; Idx < Targets.Num(); Idx++)
 	{
-		AActor* Actor = TargetComponentInterface(Idx)->GetOwnerActor();
+		AActor* Actor = UE::ToolTarget::GetTargetActor(Targets[Idx]);
 		if (Actor != SkipActor)
 		{
 			Actors.Add(Actor);
@@ -533,7 +480,6 @@ void UCombineMeshesTool::UpdateExistingAsset()
 	}
 	HandleSourceProperties->ApplyMethod(Actors, GetToolManager());
 
-
 	GetToolManager()->EndUndoTransaction();
 }
 
@@ -543,4 +489,41 @@ void UCombineMeshesTool::UpdateExistingAsset()
 
 
 
+void UCombineMeshesTool::BuildCombinedMaterialSet(TArray<UMaterialInterface*>& NewMaterialsOut, TArray<TArray<int32>>& MaterialIDRemapsOut)
+{
+	NewMaterialsOut.Reset();
+
+	TMap<UMaterialInterface*, int> KnownMaterials;
+
+	MaterialIDRemapsOut.SetNum(Targets.Num());
+	for (int32 ComponentIdx = 0; ComponentIdx < Targets.Num(); ComponentIdx++)
+	{
+		FComponentMaterialSet MaterialSet = UE::ToolTarget::GetMaterialSet(Targets[ComponentIdx]);
+		int32 NumMaterials = MaterialSet.Materials.Num();
+		for (int MaterialIdx = 0; MaterialIdx < NumMaterials; MaterialIdx++)
+		{
+			UMaterialInterface* Mat = MaterialSet.Materials[MaterialIdx];
+			int32 NewMaterialIdx = 0;
+			if (KnownMaterials.Contains(Mat) == false)
+			{
+				NewMaterialIdx = NewMaterialsOut.Num();
+				KnownMaterials.Add(Mat, NewMaterialIdx);
+				NewMaterialsOut.Add(Mat);
+			}
+			else
+			{
+				NewMaterialIdx = KnownMaterials[Mat];
+			}
+			MaterialIDRemapsOut[ComponentIdx].Add(NewMaterialIdx);
+		}
+	}
+}
+
+
+
+
+
+
+
+
 #undef LOCTEXT_NAMESPACE
diff --git a/Engine/Plugins/Experimental/MeshModelingToolset/Source/MeshModelingTools/Private/ConvertMeshesTool.cpp b/Engine/Plugins/Experimental/MeshModelingToolset/Source/MeshModelingTools/Private/ConvertMeshesTool.cpp
new file mode 100644
--- /dev/null
+++ b/Engine/Plugins/Experimental/MeshModelingToolset/Source/MeshModelingTools/Private/ConvertMeshesTool.cpp
@@ -0,0 +1,139 @@
+// Copyright Epic Games, Inc. All Rights Reserved.
+
+#include "ConvertMeshesTool.h"
+#include "ComponentSourceInterfaces.h"
+#include "InteractiveToolManager.h"
+#include "ToolTargetManager.h"
+#include "ToolBuilderUtil.h"
+#include "ToolSetupUtil.h"
+#include "ModelingToolTargetUtil.h"
+#include "ModelingObjectsCreationAPI.h"
+#include "Selection/ToolSelectionUtil.h"
+
+#include "ExplicitUseGeometryMathTypes.h"		// using UE::Geometry::(math types)
+using namespace UE::Geometry;
+
+#define LOCTEXT_NAMESPACE "UConvertMeshesTool"
+
+/*
+ * ToolBuilder
+ */
+const FToolTargetTypeRequirements& UConvertMeshesToolBuilder::GetTargetRequirements() const
+{
+	static FToolTargetTypeRequirements TypeRequirements({
+		UMaterialProvider::StaticClass(),
+		UMeshDescriptionProvider::StaticClass(),
+		UPrimitiveComponentBackedTarget::StaticClass()
+		});
+	return TypeRequirements;
+}
+
+bool UConvertMeshesToolBuilder::CanBuildTool(const FToolBuilderState& SceneState) const
+{
+	return SceneState.TargetManager->CountSelectedAndTargetable(SceneState, GetTargetRequirements()) > 0;
+}
+
+UInteractiveTool* UConvertMeshesToolBuilder::BuildTool(const FToolBuilderState& SceneState) const
+{
+	UConvertMeshesTool* NewTool = NewObject<UConvertMeshesTool>(SceneState.ToolManager);
+	TArray<TObjectPtr<UToolTarget>> Targets = SceneState.TargetManager->BuildAllSelectedTargetable(SceneState, GetTargetRequirements());
+	NewTool->SetTargets(MoveTemp(Targets));
+	NewTool->SetWorld(SceneState.World);
+	return NewTool;
+}
+
+
+
+/*
+ * Tool
+ */
+
+
+void UConvertMeshesTool::SetWorld(UWorld* World)
+{
+	this->TargetWorld = World;
+}
+
+void UConvertMeshesTool::Setup()
+{
+	UInteractiveTool::Setup();
+
+	OutputTypeProperties = NewObject<UCreateMeshObjectTypeProperties>(this);
+	OutputTypeProperties->RestoreProperties(this, TEXT("ConvertMeshesTool"));
+	OutputTypeProperties->InitializeDefault();
+	OutputTypeProperties->WatchProperty(OutputTypeProperties->OutputType, [this](FString) { OutputTypeProperties->UpdatePropertyVisibility(); });
+	AddToolPropertySource(OutputTypeProperties);
+
+	BasicProperties = NewObject<UConvertMeshesToolProperties>(this);
+	BasicProperties->RestoreProperties(this);
+	AddToolPropertySource(BasicProperties);
+
+	SetToolDisplayName(LOCTEXT("ToolName", "Convert"));
+	GetToolManager()->DisplayMessage(
+		LOCTEXT("OnStartTool", "Convert Meshes to a different Mesh Object type"),
+		EToolMessageLevel::UserNotification);
+}
+
+bool UConvertMeshesTool::CanAccept() const
+{
+	return Super::CanAccept();
+}
+
+void UConvertMeshesTool::Shutdown(EToolShutdownType ShutdownType)
+{
+	OutputTypeProperties->SaveProperties(this, TEXT("ConvertMeshesTool"));
+	BasicProperties->SaveProperties(this);
+
+	if (ShutdownType == EToolShutdownType::Accept)
+	{
+		GetToolManager()->BeginUndoTransaction(LOCTEXT("ConvertMeshesToolTransactionName", "Convert Meshes"));
+
+		TArray<AActor*> NewSelectedActors;
+		TSet<AActor*> DeleteActors;
+
+		for (int32 k = 0; k < Targets.Num(); ++k)
+		{
+			AActor* TargetActor = UE::ToolTarget::GetTargetActor(Targets[k]);
+			check(TargetActor != nullptr);
+			DeleteActors.Add(TargetActor);
+
+			FTransform SourceTransform = (FTransform)UE::ToolTarget::GetLocalToWorldTransform(Targets[k]);
+			FDynamicMesh3 SourceMesh = UE::ToolTarget::GetDynamicMeshCopy(Targets[k], true);
+			FString AssetName = TargetActor->GetName();
+			FComponentMaterialSet Materials = UE::ToolTarget::GetMaterialSet(Targets[0]);
+			const FComponentMaterialSet* TransferMaterials = (BasicProperties->bTransferMaterials) ? &Materials : nullptr;
+
+			FCreateMeshObjectParams NewMeshObjectParams;
+			NewMeshObjectParams.TargetWorld = TargetWorld;
+			NewMeshObjectParams.Transform = (FTransform)SourceTransform;
+			NewMeshObjectParams.BaseName = AssetName;
+			if (TransferMaterials != nullptr)
+			{
+				NewMeshObjectParams.Materials = TransferMaterials->Materials;
+			}
+			NewMeshObjectParams.SetMesh(MoveTemp(SourceMesh));
+
+			OutputTypeProperties->ConfigureCreateMeshObjectParams(NewMeshObjectParams);
+			FCreateMeshObjectResult Result = UE::Modeling::CreateMeshObject(GetToolManager(), MoveTemp(NewMeshObjectParams));
+			if (Result.IsOK())
+			{
+				NewSelectedActors.Add(Result.NewActor);
+			}
+		}
+
+		for (AActor* DeleteActor : DeleteActors)
+		{
+			DeleteActor->Destroy();
+		}
+
+		ToolSelectionUtil::SetNewActorSelection(GetToolManager(), NewSelectedActors);
+
+		GetToolManager()->EndUndoTransaction();
+	}
+}
+
+
+
+
+
+#undef LOCTEXT_NAMESPACE
diff --git a/Engine/Plugins/Experimental/MeshModelingToolset/Source/MeshModelingTools/Private/Properties/CreateMeshObjectTypeProperties.cpp b/Engine/Plugins/Experimental/MeshModelingToolset/Source/MeshModelingTools/Private/Properties/CreateMeshObjectTypeProperties.cpp
deleted file mode 100644
diff --git a/Engine/Plugins/Experimental/MeshModelingToolset/Source/MeshModelingTools/Private/SplitMeshesTool.cpp b/Engine/Plugins/Experimental/MeshModelingToolset/Source/MeshModelingTools/Private/SplitMeshesTool.cpp
new file mode 100644
--- /dev/null
+++ b/Engine/Plugins/Experimental/MeshModelingToolset/Source/MeshModelingTools/Private/SplitMeshesTool.cpp
@@ -0,0 +1,250 @@
+// Copyright Epic Games, Inc. All Rights Reserved.
+
+#include "SplitMeshesTool.h"
+#include "ComponentSourceInterfaces.h"
+#include "InteractiveToolManager.h"
+#include "ToolTargetManager.h"
+#include "ToolBuilderUtil.h"
+#include "ToolSetupUtil.h"
+#include "ModelingToolTargetUtil.h"
+#include "ModelingObjectsCreationAPI.h"
+#include "Selection/ToolSelectionUtil.h"
+#include "Selections/MeshConnectedComponents.h"
+#include "DynamicMesh/MeshTransforms.h"
+#include "DynamicSubmesh3.h"
+
+#include "ExplicitUseGeometryMathTypes.h"		// using UE::Geometry::(math types)
+using namespace UE::Geometry;
+
+#define LOCTEXT_NAMESPACE "USplitMeshesTool"
+
+/*
+ * ToolBuilder
+ */
+const FToolTargetTypeRequirements& USplitMeshesToolBuilder::GetTargetRequirements() const
+{
+	static FToolTargetTypeRequirements TypeRequirements({
+		UMaterialProvider::StaticClass(),
+		UMeshDescriptionProvider::StaticClass(),
+		UPrimitiveComponentBackedTarget::StaticClass()
+		});
+	return TypeRequirements;
+}
+
+bool USplitMeshesToolBuilder::CanBuildTool(const FToolBuilderState& SceneState) const
+{
+	return SceneState.TargetManager->CountSelectedAndTargetable(SceneState, GetTargetRequirements()) > 0;
+}
+
+UInteractiveTool* USplitMeshesToolBuilder::BuildTool(const FToolBuilderState& SceneState) const
+{
+	USplitMeshesTool* NewTool = NewObject<USplitMeshesTool>(SceneState.ToolManager);
+	TArray<TObjectPtr<UToolTarget>> Targets = SceneState.TargetManager->BuildAllSelectedTargetable(SceneState, GetTargetRequirements());
+	NewTool->SetTargets(MoveTemp(Targets));
+	NewTool->SetWorld(SceneState.World);
+	return NewTool;
+}
+
+
+
+/*
+ * Tool
+ */
+
+
+void USplitMeshesTool::SetWorld(UWorld* World)
+{
+	this->TargetWorld = World;
+}
+
+void USplitMeshesTool::Setup()
+{
+	UInteractiveTool::Setup();
+
+	OutputTypeProperties = NewObject<UCreateMeshObjectTypeProperties>(this);
+	OutputTypeProperties->InitializeDefaultWithAuto();
+	OutputTypeProperties->OutputType = UCreateMeshObjectTypeProperties::AutoIdentifier;
+	OutputTypeProperties->RestoreProperties(this, TEXT("OutputTypeFromInputTool"));
+	OutputTypeProperties->WatchProperty(OutputTypeProperties->OutputType, [this](FString) { OutputTypeProperties->UpdatePropertyVisibility(); });
+	AddToolPropertySource(OutputTypeProperties);
+
+	BasicProperties = NewObject<USplitMeshesToolProperties>(this);
+	BasicProperties->RestoreProperties(this);
+	AddToolPropertySource(BasicProperties);
+
+	SourceMeshes.SetNum(Targets.Num());
+	for (int32 k = 0; k < Targets.Num(); ++k)
+	{
+		SourceMeshes[k].Mesh = UE::ToolTarget::GetDynamicMeshCopy(Targets[k], true);
+		SourceMeshes[k].Materials = UE::ToolTarget::GetMaterialSet(Targets[k]).Materials;
+	}
+
+	UpdateSplitMeshes();
+
+	SetToolDisplayName(LOCTEXT("ToolName", "Split"));
+	GetToolManager()->DisplayMessage(
+		LOCTEXT("OnStartTool", "Split Meshes into parts"),
+		EToolMessageLevel::UserNotification);
+}
+
+bool USplitMeshesTool::CanAccept() const
+{
+	return Super::CanAccept();
+}
+
+void USplitMeshesTool::Shutdown(EToolShutdownType ShutdownType)
+{
+	OutputTypeProperties->SaveProperties(this, TEXT("OutputTypeFromInputTool"));
+	BasicProperties->SaveProperties(this);
+
+	if (ShutdownType == EToolShutdownType::Accept)
+	{
+		GetToolManager()->BeginUndoTransaction(LOCTEXT("SplitMeshesToolTransactionName", "Split Meshes"));
+
+		TArray<AActor*> NewSelectedActors;
+		TSet<AActor*> DeleteActors;
+
+		for (int32 ti = 0; ti < Targets.Num(); ++ti)
+		{
+			FComponentsInfo& SplitInfo = SplitMeshes[ti];
+			if (SplitInfo.bNoComponents)
+			{
+				continue;
+			}
+			AActor* TargetActor = UE::ToolTarget::GetTargetActor(Targets[ti]);
+			check(TargetActor != nullptr);
+			DeleteActors.Add(TargetActor);
+
+			FTransform3d SourceTransform = UE::ToolTarget::GetLocalToWorldTransform(Targets[ti]);
+			FDynamicMesh3 SourceMesh = UE::ToolTarget::GetDynamicMeshCopy(Targets[ti], true);
+			FString AssetName = TargetActor->GetName();
+
+			FCreateMeshObjectParams BaseMeshObjectParams;
+			BaseMeshObjectParams.TargetWorld = TargetWorld;
+
+			if (OutputTypeProperties->OutputType == UCreateMeshObjectTypeProperties::AutoIdentifier)
+			{
+				UE::ToolTarget::ConfigureCreateMeshObjectParams(Targets[0], BaseMeshObjectParams);
+			}
+			else
+			{
+				OutputTypeProperties->ConfigureCreateMeshObjectParams(BaseMeshObjectParams);
+			}
+
+			int32 NumComponents = SplitInfo.Meshes.Num();
+			for (int32 k = 0; k < NumComponents; ++k)
+			{
+				FCreateMeshObjectParams NewMeshObjectParams = BaseMeshObjectParams;
+				NewMeshObjectParams.BaseName = FString::Printf(TEXT("%s_%d"), *AssetName, k);
+				FTransform3d PartTransform = SourceTransform;
+				PartTransform.SetTranslation(SourceTransform.GetTranslation() + SplitInfo.Origins[k]);
+				NewMeshObjectParams.Transform = (FTransform)PartTransform;
+				if (BasicProperties->bTransferMaterials)
+				{
+					NewMeshObjectParams.Materials = SplitInfo.Materials[k];
+				}
+				NewMeshObjectParams.SetMesh(MoveTemp(SplitInfo.Meshes[k]));
+
+				FCreateMeshObjectResult Result = UE::Modeling::CreateMeshObject(GetToolManager(), MoveTemp(NewMeshObjectParams));
+				if (Result.IsOK())
+				{
+					NewSelectedActors.Add(Result.NewActor);
+				}
+			}
+		}
+
+		for (AActor* DeleteActor : DeleteActors)
+		{
+			DeleteActor->Destroy();
+		}
+
+		ToolSelectionUtil::SetNewActorSelection(GetToolManager(), NewSelectedActors);
+
+		GetToolManager()->EndUndoTransaction();
+	}
+}
+
+
+
+
+void USplitMeshesTool::UpdateSplitMeshes()
+{
+	SplitMeshes.Reset();
+	SplitMeshes.SetNum(SourceMeshes.Num());
+	NoSplitCount = 0;
+
+	for (int32 si = 0; si < SourceMeshes.Num(); ++si)
+	{
+		FComponentsInfo& SplitInfo = SplitMeshes[si];
+		const FDynamicMesh3* SourceMesh = &SourceMeshes[si].Mesh;
+		const TArray<UMaterialInterface*> SourceMaterials = SourceMeshes[si].Materials;
+
+		FMeshConnectedComponents MeshComponents(SourceMesh);
+		MeshComponents.FindConnectedTriangles();
+		int32 NumComponents = MeshComponents.Num();
+
+		if (NumComponents < 2)
+		{
+			SplitInfo.bNoComponents = true;
+			NoSplitCount++;
+			continue;
+		}
+		SplitInfo.bNoComponents = false;
+
+		SplitInfo.Meshes.SetNum(NumComponents);
+		SplitInfo.Materials.SetNum(NumComponents);
+		SplitInfo.Origins.SetNum(NumComponents);
+		for (int32 k = 0; k < NumComponents; ++k)
+		{
+			FDynamicSubmesh3 SubmeshCalc(SourceMesh, MeshComponents[k].Indices);
+			FDynamicMesh3& Submesh = SubmeshCalc.GetSubmesh();
+			TArray<UMaterialInterface*> NewMaterials;
+
+			// remap materials
+			FDynamicMeshMaterialAttribute* MaterialIDs = Submesh.HasAttributes() ? Submesh.Attributes()->GetMaterialID() : nullptr;
+			if (MaterialIDs)
+			{
+				TArray<int32> UniqueIDs;
+				for (int32 tid : Submesh.TriangleIndicesItr())
+				{
+					int32 MaterialID = MaterialIDs->GetValue(tid);
+					int32 Index = UniqueIDs.IndexOfByKey(MaterialID);
+					if (Index == INDEX_NONE)
+					{
+						int32 NewMaterialID = UniqueIDs.Num();
+						UniqueIDs.Add(MaterialID);
+						NewMaterials.Add(SourceMaterials[MaterialID]);
+						MaterialIDs->SetValue(tid, NewMaterialID);
+					}
+					else
+					{
+						MaterialIDs->SetValue(tid, Index);
+					}
+				}
+			}
+			
+			// reposition mesh
+			FAxisAlignedBox3d Bounds = Submesh.GetBounds();
+			MeshTransforms::Translate(Submesh, -Bounds.Center());
+
+			SplitInfo.Meshes[k] = MoveTemp(Submesh);
+			SplitInfo.Materials[k] = MoveTemp(NewMaterials);
+			SplitInfo.Origins[k] = Bounds.Center();
+		}
+	}
+
+	if (NoSplitCount > 0)
+	{
+		GetToolManager()->DisplayMessage(
+			FText::Format(LOCTEXT("NoComponentsMessage", "{0} of {1} Input Meshes cannot be Split."), NoSplitCount, SourceMeshes.Num()), EToolMessageLevel::UserWarning);
+	}
+	else
+	{
+		GetToolManager()->DisplayMessage(FText(), EToolMessageLevel::UserWarning);
+	}
+
+}
+
+
+
+#undef LOCTEXT_NAMESPACE
diff --git a/Engine/Plugins/Experimental/MeshModelingToolset/Source/MeshModelingTools/Private/TransferMeshTool.cpp b/Engine/Plugins/Experimental/MeshModelingToolset/Source/MeshModelingTools/Private/TransferMeshTool.cpp
--- a/Engine/Plugins/Experimental/MeshModelingToolset/Source/MeshModelingTools/Private/TransferMeshTool.cpp
+++ b/Engine/Plugins/Experimental/MeshModelingToolset/Source/MeshModelingTools/Private/TransferMeshTool.cpp
@@ -45,10 +45,6 @@ UInteractiveTool* UTransferMeshToolBuilder::BuildTool(const FToolBuilderState& S
  * Tool
  */
 
-UTransferMeshTool::UTransferMeshTool(const FObjectInitializer&)
-{
-}
-
 void UTransferMeshTool::SetWorld(UWorld* World)
 {
 	this->TargetWorld = World;
@@ -64,7 +60,7 @@ void UTransferMeshTool::Setup()
 
 	SetToolDisplayName(LOCTEXT("ToolName", "Transfer"));
 	GetToolManager()->DisplayMessage(
-		LOCTEXT("OnStartTool", "Copy mesh from Source object to Target object"),
+		LOCTEXT("OnStartTool", "Copy Mesh from Source object to Target object"),
 		EToolMessageLevel::UserNotification);
 }
 
diff --git a/Engine/Plugins/Experimental/MeshModelingToolset/Source/MeshModelingTools/Public/AddPrimitiveTool.h b/Engine/Plugins/Experimental/MeshModelingToolset/Source/MeshModelingTools/Public/AddPrimitiveTool.h
--- a/Engine/Plugins/Experimental/MeshModelingToolset/Source/MeshModelingTools/Public/AddPrimitiveTool.h
+++ b/Engine/Plugins/Experimental/MeshModelingToolset/Source/MeshModelingTools/Public/AddPrimitiveTool.h
@@ -8,7 +8,7 @@
 #include "MaterialShared.h"
 #include "PreviewMesh.h"
 #include "Properties/MeshMaterialProperties.h"
-#include "Properties/CreateMeshObjectTypeProperties.h"
+#include "PropertySets/CreateMeshObjectTypeProperties.h"
 #include "UObject/NoExportTypes.h"
 #include "UObject/UObjectGlobals.h"
 
diff --git a/Engine/Plugins/Experimental/MeshModelingToolset/Source/MeshModelingTools/Public/CombineMeshesTool.h b/Engine/Plugins/Experimental/MeshModelingToolset/Source/MeshModelingTools/Public/CombineMeshesTool.h
--- a/Engine/Plugins/Experimental/MeshModelingToolset/Source/MeshModelingTools/Public/CombineMeshesTool.h
+++ b/Engine/Plugins/Experimental/MeshModelingToolset/Source/MeshModelingTools/Public/CombineMeshesTool.h
@@ -3,25 +3,16 @@
 #pragma once
 
 #include "CoreMinimal.h"
-#include "UObject/NoExportTypes.h"
 #include "MultiSelectionTool.h"
 #include "InteractiveToolBuilder.h"
 #include "MeshOpPreviewHelpers.h"
 #include "DynamicMesh/DynamicMesh3.h"
-#include "BaseTools/SingleClickTool.h"
 #include "PropertySets/OnAcceptProperties.h"
+#include "PropertySets/CreateMeshObjectTypeProperties.h"
 #include "CombineMeshesTool.generated.h"
 
-
 // predeclarations
 struct FMeshDescription;
-class UDynamicMeshComponent;
-class IPrimitiveComponentBackedTarget;
-class IMeshDescriptionCommitter;
-class IMeshDescriptionProvider;
-class IMaterialProvider;
-
-
 
 /**
  *
@@ -106,6 +97,9 @@ protected:
 	UPROPERTY()
 	UCombineMeshesToolProperties* BasicProperties;
 
+	UPROPERTY()
+	UCreateMeshObjectTypeProperties* OutputTypeProperties;
+
 	UPROPERTY()
 	UOnAcceptHandleSourcesProperties* HandleSourceProperties;
 
@@ -116,4 +110,7 @@ protected:
 
 	void CreateNewAsset();
 	void UpdateExistingAsset();
+
+
+	void BuildCombinedMaterialSet(TArray<UMaterialInterface*>& NewMaterialsOut, TArray<TArray<int32>>& MaterialIDRemapsOut);
 };
diff --git a/Engine/Plugins/Experimental/MeshModelingToolset/Source/MeshModelingTools/Public/ConvertMeshesTool.h b/Engine/Plugins/Experimental/MeshModelingToolset/Source/MeshModelingTools/Public/ConvertMeshesTool.h
new file mode 100644
--- /dev/null
+++ b/Engine/Plugins/Experimental/MeshModelingToolset/Source/MeshModelingTools/Public/ConvertMeshesTool.h
@@ -0,0 +1,63 @@
+// Copyright Epic Games, Inc. All Rights Reserved.
+
+#pragma once
+
+#include "CoreMinimal.h"
+#include "MultiSelectionTool.h"
+#include "InteractiveToolBuilder.h"
+#include "PropertySets/CreateMeshObjectTypeProperties.h"
+#include "ConvertMeshesTool.generated.h"
+
+/**
+ *
+ */
+UCLASS()
+class MESHMODELINGTOOLS_API UConvertMeshesToolBuilder : public UInteractiveToolBuilder
+{
+	GENERATED_BODY()
+public:
+	virtual bool CanBuildTool(const FToolBuilderState& SceneState) const override;
+	virtual UInteractiveTool* BuildTool(const FToolBuilderState& SceneState) const override;
+
+protected:
+	virtual const FToolTargetTypeRequirements& GetTargetRequirements() const override;
+};
+
+/**
+ * Standard properties of the Transfer operation
+ */
+UCLASS()
+class MESHMODELINGTOOLS_API UConvertMeshesToolProperties : public UInteractiveToolPropertySet
+{
+	GENERATED_BODY()
+public:
+	UPROPERTY(EditAnywhere, Category = Options)
+	bool bTransferMaterials = true;
+};
+
+
+
+UCLASS()
+class MESHMODELINGTOOLS_API UConvertMeshesTool : public UMultiSelectionTool
+{
+	GENERATED_BODY()
+
+public:
+	virtual void SetWorld(UWorld* World);
+
+	virtual void Setup() override;
+	virtual void Shutdown(EToolShutdownType ShutdownType) override;
+
+	virtual bool HasCancel() const override { return true; }
+	virtual bool HasAccept() const override { return true; }
+	virtual bool CanAccept() const override;
+
+	UPROPERTY()
+	UConvertMeshesToolProperties* BasicProperties;
+
+	UPROPERTY()
+	UCreateMeshObjectTypeProperties* OutputTypeProperties;
+
+protected:
+	UWorld* TargetWorld;
+};
diff --git a/Engine/Plugins/Experimental/MeshModelingToolset/Source/MeshModelingTools/Public/DrawAndRevolveTool.h b/Engine/Plugins/Experimental/MeshModelingToolset/Source/MeshModelingTools/Public/DrawAndRevolveTool.h
--- a/Engine/Plugins/Experimental/MeshModelingToolset/Source/MeshModelingTools/Public/DrawAndRevolveTool.h
+++ b/Engine/Plugins/Experimental/MeshModelingToolset/Source/MeshModelingTools/Public/DrawAndRevolveTool.h
@@ -11,7 +11,7 @@
 #include "InteractiveToolChange.h" //FToolCommandChange
 #include "MeshOpPreviewHelpers.h" //FDynamicMeshOpResult
 #include "Properties/MeshMaterialProperties.h"
-#include "Properties/CreateMeshObjectTypeProperties.h"
+#include "PropertySets/CreateMeshObjectTypeProperties.h"
 #include "Properties/RevolveProperties.h"
 
 #include "DrawAndRevolveTool.generated.h"
diff --git a/Engine/Plugins/Experimental/MeshModelingToolset/Source/MeshModelingTools/Public/DrawPolyPathTool.h b/Engine/Plugins/Experimental/MeshModelingToolset/Source/MeshModelingTools/Public/DrawPolyPathTool.h
--- a/Engine/Plugins/Experimental/MeshModelingToolset/Source/MeshModelingTools/Public/DrawPolyPathTool.h
+++ b/Engine/Plugins/Experimental/MeshModelingToolset/Source/MeshModelingTools/Public/DrawPolyPathTool.h
@@ -10,7 +10,7 @@
 #include "Mechanics/CollectSurfacePathMechanic.h"
 #include "Mechanics/ConstructionPlaneMechanic.h"
 #include "Drawing/PolyEditPreviewMesh.h"
-#include "Properties/CreateMeshObjectTypeProperties.h"
+#include "PropertySets/CreateMeshObjectTypeProperties.h"
 #include "Properties/MeshMaterialProperties.h"
 #include "DrawPolyPathTool.generated.h"
 
diff --git a/Engine/Plugins/Experimental/MeshModelingToolset/Source/MeshModelingTools/Public/DrawPolygonTool.h b/Engine/Plugins/Experimental/MeshModelingToolset/Source/MeshModelingTools/Public/DrawPolygonTool.h
--- a/Engine/Plugins/Experimental/MeshModelingToolset/Source/MeshModelingTools/Public/DrawPolygonTool.h
+++ b/Engine/Plugins/Experimental/MeshModelingToolset/Source/MeshModelingTools/Public/DrawPolygonTool.h
@@ -12,7 +12,7 @@
 #include "Snapping/PointPlanarSnapSolver.h"
 #include "ToolSceneQueriesUtil.h"
 #include "Properties/MeshMaterialProperties.h"
-#include "Properties/CreateMeshObjectTypeProperties.h"
+#include "PropertySets/CreateMeshObjectTypeProperties.h"
 #include "Mechanics/PlaneDistanceFromHitMechanic.h"
 #include "DrawPolygonTool.generated.h"
 
diff --git a/Engine/Plugins/Experimental/MeshModelingToolset/Source/MeshModelingTools/Public/Physics/ExtractCollisionGeometryTool.h b/Engine/Plugins/Experimental/MeshModelingToolset/Source/MeshModelingTools/Public/Physics/ExtractCollisionGeometryTool.h
--- a/Engine/Plugins/Experimental/MeshModelingToolset/Source/MeshModelingTools/Public/Physics/ExtractCollisionGeometryTool.h
+++ b/Engine/Plugins/Experimental/MeshModelingToolset/Source/MeshModelingTools/Public/Physics/ExtractCollisionGeometryTool.h
@@ -6,7 +6,7 @@
 #include "BaseTools/SingleSelectionMeshEditingTool.h"
 #include "DynamicMesh/DynamicMesh3.h"
 #include "Physics/CollisionPropertySets.h"
-#include "Properties/CreateMeshObjectTypeProperties.h"
+#include "PropertySets/CreateMeshObjectTypeProperties.h"
 #include "ExtractCollisionGeometryTool.generated.h"
 
 class UPreviewGeometry;
diff --git a/Engine/Plugins/Experimental/MeshModelingToolset/Source/MeshModelingTools/Public/Properties/CreateMeshObjectTypeProperties.h b/Engine/Plugins/Experimental/MeshModelingToolset/Source/MeshModelingTools/Public/Properties/CreateMeshObjectTypeProperties.h
deleted file mode 100644
diff --git a/Engine/Plugins/Experimental/MeshModelingToolset/Source/MeshModelingTools/Public/SplitMeshesTool.h b/Engine/Plugins/Experimental/MeshModelingToolset/Source/MeshModelingTools/Public/SplitMeshesTool.h
new file mode 100644
--- /dev/null
+++ b/Engine/Plugins/Experimental/MeshModelingToolset/Source/MeshModelingTools/Public/SplitMeshesTool.h
@@ -0,0 +1,83 @@
+// Copyright Epic Games, Inc. All Rights Reserved.
+
+#pragma once
+
+#include "CoreMinimal.h"
+#include "MultiSelectionTool.h"
+#include "InteractiveToolBuilder.h"
+#include "PropertySets/CreateMeshObjectTypeProperties.h"
+#include "SplitMeshesTool.generated.h"
+
+class UMaterialInterface;
+
+
+UCLASS()
+class MESHMODELINGTOOLS_API USplitMeshesToolBuilder : public UInteractiveToolBuilder
+{
+	GENERATED_BODY()
+public:
+	virtual bool CanBuildTool(const FToolBuilderState& SceneState) const override;
+	virtual UInteractiveTool* BuildTool(const FToolBuilderState& SceneState) const override;
+
+protected:
+	virtual const FToolTargetTypeRequirements& GetTargetRequirements() const override;
+};
+
+
+
+UCLASS()
+class MESHMODELINGTOOLS_API USplitMeshesToolProperties : public UInteractiveToolPropertySet
+{
+	GENERATED_BODY()
+public:
+	UPROPERTY(EditAnywhere, Category = Options)
+	bool bTransferMaterials = true;
+};
+
+
+
+UCLASS()
+class MESHMODELINGTOOLS_API USplitMeshesTool : public UMultiSelectionTool
+{
+	GENERATED_BODY()
+
+public:
+	virtual void SetWorld(UWorld* World);
+
+	virtual void Setup() override;
+	virtual void Shutdown(EToolShutdownType ShutdownType) override;
+
+	virtual bool HasCancel() const override { return true; }
+	virtual bool HasAccept() const override { return true; }
+	virtual bool CanAccept() const override;
+
+	UPROPERTY()
+	USplitMeshesToolProperties* BasicProperties;
+
+	UPROPERTY()
+	UCreateMeshObjectTypeProperties* OutputTypeProperties;
+
+protected:
+	UWorld* TargetWorld;
+
+	struct FSourceMeshInfo
+	{
+		UE::Geometry::FDynamicMesh3 Mesh;
+		TArray<UMaterialInterface*> Materials;
+	};
+	TArray<FSourceMeshInfo> SourceMeshes;
+
+
+	struct FComponentsInfo
+	{
+		bool bNoComponents;
+		TArray<UE::Geometry::FDynamicMesh3> Meshes;
+		TArray<TArray<UMaterialInterface*>> Materials;
+		TArray<FVector3d> Origins;
+	};
+	TArray<FComponentsInfo> SplitMeshes;
+
+	int32 NoSplitCount = 0;
+
+	void UpdateSplitMeshes();
+};
diff --git a/Engine/Plugins/Experimental/MeshModelingToolset/Source/MeshModelingTools/Public/TransferMeshTool.h b/Engine/Plugins/Experimental/MeshModelingToolset/Source/MeshModelingTools/Public/TransferMeshTool.h
--- a/Engine/Plugins/Experimental/MeshModelingToolset/Source/MeshModelingTools/Public/TransferMeshTool.h
+++ b/Engine/Plugins/Experimental/MeshModelingToolset/Source/MeshModelingTools/Public/TransferMeshTool.h
@@ -43,8 +43,6 @@ class MESHMODELINGTOOLS_API UTransferMeshTool : public UMultiSelectionTool
 	GENERATED_BODY()
 
 public:
-	UTransferMeshTool(const FObjectInitializer&);
-
 	virtual void SetWorld(UWorld* World);
 
 	virtual void Setup() override;
diff --git a/Engine/Plugins/Experimental/MeshModelingToolset/Source/ModelingComponents/Private/BaseTools/BaseCreateFromSelectedTool.cpp b/Engine/Plugins/Experimental/MeshModelingToolset/Source/ModelingComponents/Private/BaseTools/BaseCreateFromSelectedTool.cpp
--- a/Engine/Plugins/Experimental/MeshModelingToolset/Source/ModelingComponents/Private/BaseTools/BaseCreateFromSelectedTool.cpp
+++ b/Engine/Plugins/Experimental/MeshModelingToolset/Source/ModelingComponents/Private/BaseTools/BaseCreateFromSelectedTool.cpp
@@ -85,11 +85,18 @@ void UBaseCreateFromSelectedTool::Setup()
 	TransformProperties->RestoreProperties(this);
 	AddToolPropertySource(TransformProperties);
 	
+	OutputTypeProperties = NewObject<UCreateMeshObjectTypeProperties>(this);
+	OutputTypeProperties->InitializeDefaultWithAuto();
+	OutputTypeProperties->OutputType = UCreateMeshObjectTypeProperties::AutoIdentifier;
+	OutputTypeProperties->RestoreProperties(this, TEXT("OutputTypeFromInputTool"));
+	OutputTypeProperties->WatchProperty(OutputTypeProperties->OutputType, [this](FString) { OutputTypeProperties->UpdatePropertyVisibility(); });
+	AddToolPropertySource(OutputTypeProperties);
+
 	HandleSourcesProperties = NewObject<UBaseCreateFromSelectedHandleSourceProperties>(this);
 	HandleSourcesProperties->RestoreProperties(this);
 	AddToolPropertySource(HandleSourcesProperties);
 
-	Preview = NewObject<UMeshOpPreviewWithBackgroundCompute>(this, "Preview");
+	Preview = NewObject<UMeshOpPreviewWithBackgroundCompute>(this);
 	Preview->Setup(this->TargetWorld, this);
 
 	SetPreviewCallbacks();
@@ -108,6 +115,7 @@ void UBaseCreateFromSelectedTool::Setup()
 	HandleSourcesProperties->OutputName = PrefixWithSourceNameIfSingleSelection(GetCreatedAssetName());
 	HandleSourcesProperties->WatchProperty(HandleSourcesProperties->WriteOutputTo, [&](EBaseCreateFromSelectedTargetType NewType)
 	{
+		SetToolPropertySourceEnabled(OutputTypeProperties, (NewType == EBaseCreateFromSelectedTargetType::NewAsset));
 		if (NewType == EBaseCreateFromSelectedTargetType::NewAsset)
 		{
 			HandleSourcesProperties->OutputAsset = TEXT("");
@@ -241,7 +249,14 @@ void UBaseCreateFromSelectedTool::GenerateAsset(const FDynamicMeshOpResult& OpRe
 	NewMeshObjectParams.BaseName = UseBaseName;
 	NewMeshObjectParams.Materials = GetOutputMaterials();
 	NewMeshObjectParams.SetMesh(OpResult.Mesh.Get());
-	UE::ToolTarget::ConfigureCreateMeshObjectParams(Targets[0], NewMeshObjectParams);
+	if (OutputTypeProperties->OutputType == UCreateMeshObjectTypeProperties::AutoIdentifier)
+	{
+		UE::ToolTarget::ConfigureCreateMeshObjectParams(Targets[0], NewMeshObjectParams);
+	}
+	else
+	{
+		OutputTypeProperties->ConfigureCreateMeshObjectParams(NewMeshObjectParams);
+	}
 	FCreateMeshObjectResult Result = UE::Modeling::CreateMeshObject(GetToolManager(), MoveTemp(NewMeshObjectParams));
 	if (Result.IsOK() && Result.NewActor != nullptr)
 	{
@@ -315,6 +330,7 @@ void UBaseCreateFromSelectedTool::Shutdown(EToolShutdownType ShutdownType)
 {
 	SaveProperties();
 	HandleSourcesProperties->SaveProperties(this);
+	OutputTypeProperties->SaveProperties(this, TEXT("OutputTypeFromInputTool"));
 	TransformProperties->SaveProperties(this);
 
 	FDynamicMeshOpResult Result = Preview->Shutdown();
diff --git a/Engine/Plugins/Experimental/MeshModelingToolset/Source/ModelingComponents/Private/PropertySets/CreateMeshObjectTypeProperties.cpp b/Engine/Plugins/Experimental/MeshModelingToolset/Source/ModelingComponents/Private/PropertySets/CreateMeshObjectTypeProperties.cpp
new file mode 100644
--- /dev/null
+++ b/Engine/Plugins/Experimental/MeshModelingToolset/Source/ModelingComponents/Private/PropertySets/CreateMeshObjectTypeProperties.cpp
@@ -0,0 +1,126 @@
+// Copyright Epic Games, Inc. All Rights Reserved.
+
+#include "PropertySets/CreateMeshObjectTypeProperties.h"
+#include "ModelingObjectsCreationAPI.h"
+
+
+const FString UCreateMeshObjectTypeProperties::StaticMeshIdentifier = TEXT("Static Mesh");
+const FString UCreateMeshObjectTypeProperties::VolumeIdentifier = TEXT("Volume");
+const FString UCreateMeshObjectTypeProperties::DynamicMeshActorIdentifier = TEXT("Dynamic Mesh");
+const FString UCreateMeshObjectTypeProperties::AutoIdentifier = TEXT("From Input");
+
+bool UCreateMeshObjectTypeProperties::bEnableDynamicMeshActorSupport = false;
+FString UCreateMeshObjectTypeProperties::DefaultObjectTypeIdentifier = UCreateMeshObjectTypeProperties::StaticMeshIdentifier;
+
+
+void UCreateMeshObjectTypeProperties::InitializeDefault()
+{
+	bool bStaticMeshes = true;
+	bool bVolumes = false;
+	bool bDynamicMeshes = true;
+
+#if WITH_EDITOR
+	bVolumes = true;
+#endif
+
+	Initialize(bStaticMeshes, bVolumes, bDynamicMeshes);
+}
+
+void UCreateMeshObjectTypeProperties::InitializeDefaultWithAuto()
+{
+	InitializeDefault();
+	OutputTypeNamesList.Add(AutoIdentifier);
+}
+
+void UCreateMeshObjectTypeProperties::Initialize(bool bEnableStaticMeshes, bool bEnableVolumes, bool bEnableDynamicMeshActor)
+{
+	if (bEnableStaticMeshes)
+	{
+		OutputTypeNamesList.Add(StaticMeshIdentifier);
+	}
+	if (bEnableVolumes)
+	{
+		OutputTypeNamesList.Add(VolumeIdentifier);
+	}
+	if (bEnableDynamicMeshActor && bEnableDynamicMeshActorSupport)
+	{
+		OutputTypeNamesList.Add(DynamicMeshActorIdentifier);
+	}
+
+	if (OutputTypeNamesList.Num() == 0)
+	{
+		return;
+	}
+
+	if ((OutputType.Len() == 0) || (OutputType.Len() > 0 && OutputTypeNamesList.Contains(OutputType) == false))
+	{
+		if (OutputTypeNamesList.Contains(DefaultObjectTypeIdentifier))
+		{
+			OutputType = DefaultObjectTypeIdentifier;
+		}
+		else
+		{
+			OutputType = OutputTypeNamesList[0];
+		}
+	}
+}
+
+const TArray<FString>& UCreateMeshObjectTypeProperties::GetOutputTypeNamesFunc()
+{
+	return OutputTypeNamesList;
+}
+
+
+
+bool UCreateMeshObjectTypeProperties::ShouldShowPropertySet() const
+{
+	return (OutputTypeNamesList.Num() > 1)
+		|| OutputTypeNamesList.Contains(VolumeIdentifier);
+}
+
+ECreateObjectTypeHint UCreateMeshObjectTypeProperties::GetCurrentCreateMeshType() const
+{
+	if (OutputType == StaticMeshIdentifier)
+	{
+		return ECreateObjectTypeHint::StaticMesh;
+	}
+	else if (OutputType == VolumeIdentifier)
+	{
+		return ECreateObjectTypeHint::Volume;
+	}
+	else if (OutputType == DynamicMeshActorIdentifier)
+	{
+		return ECreateObjectTypeHint::DynamicMeshActor;
+	}
+	return ECreateObjectTypeHint::Undefined;
+}
+
+void UCreateMeshObjectTypeProperties::UpdatePropertyVisibility()
+{
+	bShowVolumeList = (OutputType == VolumeIdentifier);
+}
+
+
+bool UCreateMeshObjectTypeProperties::ConfigureCreateMeshObjectParams(FCreateMeshObjectParams& ParamsOut) const
+{
+	// client has to handle this case
+	ensure(OutputType != AutoIdentifier);
+
+	if (OutputType == StaticMeshIdentifier)
+	{
+		ParamsOut.TypeHint = ECreateObjectTypeHint::StaticMesh;
+		return true;
+	}
+	else if (OutputType == VolumeIdentifier)
+	{
+		ParamsOut.TypeHint = ECreateObjectTypeHint::Volume;
+		ParamsOut.TypeHintClass = VolumeType.Get();
+		return true;
+	}
+	else if (OutputType == DynamicMeshActorIdentifier)
+	{
+		ParamsOut.TypeHint = ECreateObjectTypeHint::DynamicMeshActor;
+		return true;
+	}
+	return false;
+}
diff --git a/Engine/Plugins/Experimental/MeshModelingToolset/Source/ModelingComponents/Public/BaseTools/BaseCreateFromSelectedTool.h b/Engine/Plugins/Experimental/MeshModelingToolset/Source/ModelingComponents/Public/BaseTools/BaseCreateFromSelectedTool.h
--- a/Engine/Plugins/Experimental/MeshModelingToolset/Source/ModelingComponents/Public/BaseTools/BaseCreateFromSelectedTool.h
+++ b/Engine/Plugins/Experimental/MeshModelingToolset/Source/ModelingComponents/Public/BaseTools/BaseCreateFromSelectedTool.h
@@ -7,6 +7,7 @@
 #include "InteractiveToolBuilder.h"
 #include "MeshOpPreviewHelpers.h"
 #include "PropertySets/OnAcceptProperties.h"
+#include "PropertySets/CreateMeshObjectTypeProperties.h"
 #include "BaseCreateFromSelectedTool.generated.h"
 
 class UTransformGizmo;
@@ -191,6 +192,9 @@ protected:
 	UPROPERTY()
 	UTransformInputsToolProperties* TransformProperties;
 
+	UPROPERTY()
+	UCreateMeshObjectTypeProperties* OutputTypeProperties;
+
 	UPROPERTY()
 	UBaseCreateFromSelectedHandleSourceProperties* HandleSourcesProperties;
 
diff --git a/Engine/Plugins/Experimental/MeshModelingToolset/Source/ModelingComponents/Public/PropertySets/CreateMeshObjectTypeProperties.h b/Engine/Plugins/Experimental/MeshModelingToolset/Source/ModelingComponents/Public/PropertySets/CreateMeshObjectTypeProperties.h
new file mode 100644
--- /dev/null
+++ b/Engine/Plugins/Experimental/MeshModelingToolset/Source/ModelingComponents/Public/PropertySets/CreateMeshObjectTypeProperties.h
@@ -0,0 +1,78 @@
+// Copyright Epic Games, Inc. All Rights Reserved.
+
+#pragma once
+
+#include "CoreMinimal.h"
+#include "InteractiveToolBuilder.h"
+#include "Engine/BlockingVolume.h"
+#include "ModelingObjectsCreationAPI.h"
+#include "CreateMeshObjectTypeProperties.generated.h"
+
+struct FCreateMeshObjectParams;
+
+/**
+ * UCreateMeshObjectTypeProperties is a InteractiveTool PropertySet used to select
+ * what type of object to create, in creation tools (ie StaticMesh, Volume, etc).
+ */
+UCLASS()
+class MODELINGCOMPONENTS_API UCreateMeshObjectTypeProperties : public UInteractiveToolPropertySet
+{
+	GENERATED_BODY()
+public:
+
+	void InitializeDefault();
+	void InitializeDefaultWithAuto();
+	void Initialize(bool bEnableStaticMeshes = true, bool bEnableVolumes = false, bool bEnableDynamicMeshActor = false);
+
+	/** Type of Object to create */
+	UPROPERTY(EditAnywhere, Category = OutputType, meta = (DisplayName = "Output Type", GetOptions = GetOutputTypeNamesFunc, NoResetToDefault))
+	FString OutputType;
+
+	/** Type of new Volume to create */
+	UPROPERTY(EditAnywhere, Category = OutputType, meta = (EditCondition = "bShowVolumeList == true", EditConditionHides, HideEditConditionToggle) )
+	TSubclassOf<class AVolume> VolumeType = ABlockingVolume::StaticClass();
+
+	// This function returns a list that is shown to select OutputType
+	UFUNCTION()
+	const TArray<FString>& GetOutputTypeNamesFunc();
+
+	// The list returned by GetOutputTypeNamesFunc()
+	UPROPERTY(meta = (TransientToolProperty))
+	TArray<FString> OutputTypeNamesList;
+
+	// control whether the VolumeType field is enabled
+	UPROPERTY(meta = (TransientToolProperty))
+	bool bShowVolumeList = false;
+
+	/** @return true if we should show this property set on a Tool, ie if there are options the user can actually change (false if only one possible output type that has no options) */
+	UFUNCTION()
+	bool ShouldShowPropertySet() const;
+
+	/** Convert current OutputType selection into a type that can be passed to UModelingObjectsCreationAPI */
+	UFUNCTION()
+	ECreateObjectTypeHint GetCurrentCreateMeshType() const;
+
+	// update visibility of properties based on current OutputType selection
+	virtual void UpdatePropertyVisibility();
+
+	/**
+	 * Utility function to configure a FCreateMeshObjectParams based on current settings
+	 *    - for StaticMesh, set ParamsOut.TypeHint
+	 *    - for Volume, if in UE Editor, set ParamsOut.TypeHint and ParamsOut.TypeHintClass
+	 * @return true if a type this function can handle was configured
+	 */
+	virtual bool ConfigureCreateMeshObjectParams(FCreateMeshObjectParams& ParamsOut) const;
+
+	// constants used for different known types
+	static const FString AutoIdentifier;
+	static const FString StaticMeshIdentifier;
+	static const FString VolumeIdentifier;
+	static const FString DynamicMeshActorIdentifier;
+
+
+	//
+	// Public static values used to configure behavior of this class, set (eg) from Editor settings
+	//
+	static bool bEnableDynamicMeshActorSupport;
+	static FString DefaultObjectTypeIdentifier;
+};
\ No newline at end of file
diff --git a/Engine/Plugins/Experimental/ModelingToolsEditorMode/Source/ModelingToolsEditorMode/Private/ModelingToolsEditorMode.cpp b/Engine/Plugins/Experimental/ModelingToolsEditorMode/Source/ModelingToolsEditorMode/Private/ModelingToolsEditorMode.cpp
--- a/Engine/Plugins/Experimental/ModelingToolsEditorMode/Source/ModelingToolsEditorMode/Private/ModelingToolsEditorMode.cpp
+++ b/Engine/Plugins/Experimental/ModelingToolsEditorMode/Source/ModelingToolsEditorMode/Private/ModelingToolsEditorMode.cpp
@@ -16,7 +16,7 @@
 #include "Framework/MultiBox/MultiBoxBuilder.h"
 #include "EditorViewportClient.h"
 #include "EngineAnalytics.h"
-#include "BaseGIzmos/TransformGizmoUtil.h"
+#include "BaseGizmos/TransformGizmoUtil.h"
 
 #include "Features/IModularFeatures.h"
 #include "ModelingModeToolExtensions.h"
@@ -85,6 +85,8 @@
 #include "SeamSculptTool.h"
 #include "MeshGroupPaintTool.h"
 #include "TransferMeshTool.h"
+#include "ConvertMeshesTool.h"
+#include "SplitMeshesTool.h"
 
 #include "Physics/PhysicsInspectorTool.h"
 #include "Physics/SetCollisionGeometryTool.h"
@@ -432,6 +434,8 @@ void UModelingToolsEditorMode::Enter()
 	RegisterTool(ToolManagerCommands.BeginEditPivotTool, TEXT("BeginEditPivotTool"), NewObject<UEditPivotToolBuilder>());
 	RegisterTool(ToolManagerCommands.BeginAlignObjectsTool, TEXT("BeginAlignObjectsTool"), NewObject<UAlignObjectsToolBuilder>());
 	RegisterTool(ToolManagerCommands.BeginTransferMeshTool, TEXT("BeginTransferMeshTool"), NewObject<UTransferMeshToolBuilder>());
+	RegisterTool(ToolManagerCommands.BeginConvertMeshesTool, TEXT("BeginConvertMeshesTool"), NewObject<UConvertMeshesToolBuilder>());
+	RegisterTool(ToolManagerCommands.BeginSplitMeshesTool, TEXT("BeginSplitMeshesTool"), NewObject<USplitMeshesToolBuilder>());
 	RegisterTool(ToolManagerCommands.BeginBakeTransformTool, TEXT("BeginBakeTransformTool"), NewObject<UBakeTransformToolBuilder>());
 	RegisterTool(ToolManagerCommands.BeginTransformUVIslandsTool, TEXT("BeginTransformUVIslandsTool"), NewObject<UEditUVIslandsToolBuilder>());
 	RegisterTool(ToolManagerCommands.BeginLatticeDeformerTool, TEXT("BeginLatticeDeformerTool"), NewObject<ULatticeDeformerToolBuilder>());
diff --git a/Engine/Plugins/Experimental/ModelingToolsEditorMode/Source/ModelingToolsEditorMode/Private/ModelingToolsEditorModeToolkit.cpp b/Engine/Plugins/Experimental/ModelingToolsEditorMode/Source/ModelingToolsEditorMode/Private/ModelingToolsEditorModeToolkit.cpp
--- a/Engine/Plugins/Experimental/ModelingToolsEditorMode/Source/ModelingToolsEditorMode/Private/ModelingToolsEditorModeToolkit.cpp
+++ b/Engine/Plugins/Experimental/ModelingToolsEditorMode/Source/ModelingToolsEditorMode/Private/ModelingToolsEditorModeToolkit.cpp
@@ -25,12 +25,16 @@
 #include "Features/IModularFeatures.h"
 #include "ModelingModeToolExtensions.h"
 
+// for Object Type properties
+#include "PropertySets/CreateMeshObjectTypeProperties.h"
+
 // for LOD setting
 #include "EditorInteractiveToolsFrameworkModule.h"
 #include "Tools/EditorComponentSourceFactory.h"
 #include "ToolTargetManager.h"
 #include "ToolTargets/StaticMeshComponentToolTarget.h"
 
+
 #define LOCTEXT_NAMESPACE "FModelingToolsEditorModeToolkit"
 
 
@@ -134,6 +138,8 @@ void FModelingToolsEditorModeToolkit::Init(const TSharedPtr<IToolkitHost>& InitT
 	GetScriptableEditorMode()->GetInteractiveToolsContext()->OnToolNotificationMessage.AddSP(this, &FModelingToolsEditorModeToolkit::PostNotification);
 	GetScriptableEditorMode()->GetInteractiveToolsContext()->OnToolWarningMessage.AddSP(this, &FModelingToolsEditorModeToolkit::PostWarning);
 
+	UpdateObjectCreationOptionsFromSettings();
+
 	SAssignNew(ViewportOverlayWidget, SHorizontalBox)
 
 	+SHorizontalBox::Slot()
@@ -552,6 +558,8 @@ void FModelingToolsEditorModeToolkit::BuildToolPalette_Experimental(FName Palett
 		ToolbarBuilder.AddToolBarButton(Commands.BeginEditPivotTool);
 		ToolbarBuilder.AddToolBarButton(Commands.BeginBakeTransformTool);
 		ToolbarBuilder.AddToolBarButton(Commands.BeginTransferMeshTool);
+		ToolbarBuilder.AddToolBarButton(Commands.BeginConvertMeshesTool);
+		ToolbarBuilder.AddToolBarButton(Commands.BeginSplitMeshesTool);
 	}
 	else if (PaletteIndex == DeformTabName)
 	{
@@ -872,8 +880,32 @@ void FModelingToolsEditorModeToolkit::UpdateAssetPanelFromSettings()
 }
 
 
+void FModelingToolsEditorModeToolkit::UpdateObjectCreationOptionsFromSettings()
+{
+	// update DynamicMeshActor Settings
+	const UModelingToolsEditorModeSettings* Settings = GetDefault<UModelingToolsEditorModeSettings>();
+
+	// enable/disable dynamic mesh actors
+	UCreateMeshObjectTypeProperties::bEnableDynamicMeshActorSupport = Settings->bEnableDynamicMeshActors;
+
+	// set configured default type
+	if (Settings->DefaultMeshObjectType == EModelingModeDefaultMeshObjectType::DynamicMeshActor && Settings->bEnableDynamicMeshActors)
+	{
+		UCreateMeshObjectTypeProperties::DefaultObjectTypeIdentifier = UCreateMeshObjectTypeProperties::DynamicMeshActorIdentifier;
+	}
+	else if (Settings->DefaultMeshObjectType == EModelingModeDefaultMeshObjectType::VolumeActor)
+	{
+		UCreateMeshObjectTypeProperties::DefaultObjectTypeIdentifier = UCreateMeshObjectTypeProperties::VolumeIdentifier;
+	}
+	else
+	{
+		UCreateMeshObjectTypeProperties::DefaultObjectTypeIdentifier = UCreateMeshObjectTypeProperties::StaticMeshIdentifier;
+	}
+}
+
 void FModelingToolsEditorModeToolkit::OnAssetSettingsModified()
 {
+	UpdateObjectCreationOptionsFromSettings();
 	UpdateAssetPanelFromSettings();
 }
 
diff --git a/Engine/Plugins/Experimental/ModelingToolsEditorMode/Source/ModelingToolsEditorMode/Private/ModelingToolsManagerActions.cpp b/Engine/Plugins/Experimental/ModelingToolsEditorMode/Source/ModelingToolsEditorMode/Private/ModelingToolsManagerActions.cpp
--- a/Engine/Plugins/Experimental/ModelingToolsEditorMode/Source/ModelingToolsEditorMode/Private/ModelingToolsManagerActions.cpp
+++ b/Engine/Plugins/Experimental/ModelingToolsEditorMode/Source/ModelingToolsEditorMode/Private/ModelingToolsManagerActions.cpp
@@ -56,6 +56,8 @@ void FModelingToolsManagerCommands::RegisterCommands()
 	UI_COMMAND(BeginDuplicateMeshesTool, "Dupe", "Duplicate Selection into new StaticMesh Asset", EUserInterfaceActionType::ToggleButton, FInputChord());
 	UI_COMMAND(BeginAlignObjectsTool, "Align", "Align Objects", EUserInterfaceActionType::ToggleButton, FInputChord());
 	UI_COMMAND(BeginTransferMeshTool, "Transfer", "Transfer Mesh", EUserInterfaceActionType::ToggleButton, FInputChord());
+	UI_COMMAND(BeginConvertMeshesTool, "Convert", "Convert Meshes", EUserInterfaceActionType::ToggleButton, FInputChord());
+	UI_COMMAND(BeginSplitMeshesTool, "Split", "Split Meshes", EUserInterfaceActionType::ToggleButton, FInputChord());
 
 	UI_COMMAND(BeginRemeshSculptMeshTool, "DSclpt", "Dynamic Mesh Sculpting", EUserInterfaceActionType::ToggleButton, FInputChord());
 	UI_COMMAND(BeginRemeshMeshTool, "Remesh", "Retriangulate Mesh", EUserInterfaceActionType::ToggleButton, FInputChord());
diff --git a/Engine/Plugins/Experimental/ModelingToolsEditorMode/Source/ModelingToolsEditorMode/Public/ModelingToolsEditorModeSettings.h b/Engine/Plugins/Experimental/ModelingToolsEditorMode/Source/ModelingToolsEditorMode/Public/ModelingToolsEditorModeSettings.h
--- a/Engine/Plugins/Experimental/ModelingToolsEditorMode/Source/ModelingToolsEditorMode/Public/ModelingToolsEditorModeSettings.h
+++ b/Engine/Plugins/Experimental/ModelingToolsEditorMode/Source/ModelingToolsEditorMode/Public/ModelingToolsEditorModeSettings.h
@@ -7,6 +7,20 @@
 #include "ModelingToolsEditorModeSettings.generated.h"
 
 
+
+UENUM()
+enum class EModelingModeDefaultMeshObjectType
+{
+	/** Generate a new Static Mesh Asset (using Generated Asset settings below) and AStaticMeshActor */
+	StaticMeshAsset,
+	/** Generate a new AVolume Actor */
+	VolumeActor,
+	/** Generate a new ADynamicMeshActor (stored locally in the Level) */
+	DynamicMeshActor
+};
+
+
+
 UENUM()
 enum class EModelingModeAssetGenerationBehavior
 {
@@ -46,28 +60,36 @@ class MODELINGTOOLSEDITORMODE_API UModelingToolsEditorModeSettings
 	GENERATED_BODY()
 public:
 
+	/** Enable/Disable the options to emit Dynamic Mesh Actors in Modeling Mode Tools */
+	UPROPERTY(config, EditAnywhere, Category = "Modeling Mode")
+	bool bEnableDynamicMeshActors = false;
+
+	/** What type of Mesh Object should Output Type Setting default to in Modeling Mode Tools (takes effect after Editor restart) */
+	UPROPERTY(config, EditAnywhere, Category = "Modeling Mode")
+	EModelingModeDefaultMeshObjectType DefaultMeshObjectType = EModelingModeDefaultMeshObjectType::StaticMeshAsset;
+
 	/** How should Assets auto-generated by Modeling Tools be handled in terms of saving, naming, etc */
-	UPROPERTY(config, EditAnywhere, Category = GeneratedModelingAssets)
+	UPROPERTY(config, EditAnywhere, Category = "Modeling Mode|Generated Assets")
 	EModelingModeAssetGenerationBehavior AssetGenerationMode = EModelingModeAssetGenerationBehavior::AutoGenerateButDoNotAutosave;
 
 	/** Where should Assets auto-generated by Modeling Tools be stored by default */
-	UPROPERTY(config, EditAnywhere, Category = GeneratedModelingAssets)
+	UPROPERTY(config, EditAnywhere, Category = "Modeling Mode|Generated Assets")
 	EModelingModeAssetGenerationLocation AssetGenerationLocation = EModelingModeAssetGenerationLocation::AutoGeneratedWorldRelativeAssetPath;
 
 	/** Assets auto-generated by Modeling Tools are stored at this path, relative to the parent path defined by the Location. Set to an empty string to disable. */
-	UPROPERTY(config, EditAnywhere, Category = GeneratedModelingAssets)
+	UPROPERTY(config, EditAnywhere, Category = "Modeling Mode|Generated Assets")
 	FString AutoGeneratedAssetPath = TEXT("_GENERATED");
 
 	/** If true, Auto-Generated Assets created in an unsaved Level will be stored relative to top-level /Game folder, otherwise they will be stored in /Temp and cannot be saved until they are explicitly moved to a permanent location */
-	UPROPERTY(config, EditAnywhere, Category = GeneratedModelingAssets)
+	UPROPERTY(config, EditAnywhere, Category = "Modeling Mode|Generated Assets")
 	bool bStoreUnsavedLevelAssetsInTopLevelGameFolder = true;
 
 	/** If true, Autogenerated Assets are stored in per-user folders below the Autogen path */
-	UPROPERTY(config, EditAnywhere, Category = GeneratedModelingAssets)
+	UPROPERTY(config, EditAnywhere, Category = "Modeling Mode|Generated Assets")
 	bool bUsePerUserAutogenSubfolder = true;
 
 	/** If true, Autogenerated Assets have a short random string generated and appended to their name */
-	UPROPERTY(config, EditAnywhere, Category = GeneratedModelingAssets)
+	UPROPERTY(config, EditAnywhere, Category = "Modeling Mode|Generated Assets")
 	bool bAppendRandomStringToName = true;
 
 public:
diff --git a/Engine/Plugins/Experimental/ModelingToolsEditorMode/Source/ModelingToolsEditorMode/Public/ModelingToolsEditorModeToolkit.h b/Engine/Plugins/Experimental/ModelingToolsEditorMode/Source/ModelingToolsEditorMode/Public/ModelingToolsEditorModeToolkit.h
--- a/Engine/Plugins/Experimental/ModelingToolsEditorMode/Source/ModelingToolsEditorMode/Public/ModelingToolsEditorModeToolkit.h
+++ b/Engine/Plugins/Experimental/ModelingToolsEditorMode/Source/ModelingToolsEditorMode/Public/ModelingToolsEditorModeToolkit.h
@@ -97,6 +97,7 @@ private:
 	void UpdateAssetLocationMode(TSharedPtr<FString> NewString);
 	void UpdateAssetSaveMode(TSharedPtr<FString> NewString);
 	void UpdateAssetPanelFromSettings();
+	void UpdateObjectCreationOptionsFromSettings();
 	void OnAssetSettingsModified();
 	FDelegateHandle AssetSettingsModifiedHandle;
 	void OnShowAssetSettings();
diff --git a/Engine/Plugins/Experimental/ModelingToolsEditorMode/Source/ModelingToolsEditorMode/Public/ModelingToolsManagerActions.h b/Engine/Plugins/Experimental/ModelingToolsEditorMode/Source/ModelingToolsEditorMode/Public/ModelingToolsManagerActions.h
--- a/Engine/Plugins/Experimental/ModelingToolsEditorMode/Source/ModelingToolsEditorMode/Public/ModelingToolsManagerActions.h
+++ b/Engine/Plugins/Experimental/ModelingToolsEditorMode/Source/ModelingToolsEditorMode/Public/ModelingToolsManagerActions.h
@@ -50,6 +50,8 @@ public:
 	TSharedPtr<FUICommandInfo> BeginDuplicateMeshesTool;
 	TSharedPtr<FUICommandInfo> BeginAlignObjectsTool;
 	TSharedPtr<FUICommandInfo> BeginTransferMeshTool;
+	TSharedPtr<FUICommandInfo> BeginConvertMeshesTool;
+	TSharedPtr<FUICommandInfo> BeginSplitMeshesTool;
 
 	TSharedPtr<FUICommandInfo> BeginRemeshSculptMeshTool;
 	TSharedPtr<FUICommandInfo> BeginRemeshMeshTool;
