ModelingTools: Add support for selection of UV layer and Polygroup layer in MeshSelectionTool. Switch to using MeshElementVisualizer for wireframe/borders, and add a UMeshStatisticsProperties. Improve button text/tooltips. Switch EditMeshMaterialsTool active-material to use a dynamic dropdown list with material names, instead of material index. Change default UV calculation to Conformal in ComputeUVs tool.
GeometryProcessing: make PolygroupSet take a const mesh, and require caller to provide writable mesh when calling SetGroup()

#rnx

diff --git a/Engine/Plugins/Experimental/GeometryProcessing/Source/DynamicMesh/Private/Polygroups/PolygroupSet.cpp b/Engine/Plugins/Experimental/GeometryProcessing/Source/DynamicMesh/Private/Polygroups/PolygroupSet.cpp
--- a/Engine/Plugins/Experimental/GeometryProcessing/Source/DynamicMesh/Private/Polygroups/PolygroupSet.cpp
+++ b/Engine/Plugins/Experimental/GeometryProcessing/Source/DynamicMesh/Private/Polygroups/PolygroupSet.cpp
@@ -15,14 +15,14 @@ FPolygroupSet::FPolygroupSet(const FPolygroupSet* CopyIn)
 	MaxGroupID = CopyIn->MaxGroupID;
 }
 
-FPolygroupSet::FPolygroupSet(FDynamicMesh3* MeshIn)
+FPolygroupSet::FPolygroupSet(const FDynamicMesh3* MeshIn)
 {
 	Mesh = MeshIn;
 	GroupLayerIndex = -1;
 	RecalculateMaxGroupID();
 }
 
-FPolygroupSet::FPolygroupSet(FDynamicMesh3* MeshIn, FDynamicMeshPolygroupAttribute* PolygroupAttribIn)
+FPolygroupSet::FPolygroupSet(const FDynamicMesh3* MeshIn, const FDynamicMeshPolygroupAttribute* PolygroupAttribIn)
 {
 	Mesh = MeshIn;
 	PolygroupAttrib = PolygroupAttribIn;
@@ -30,7 +30,7 @@ FPolygroupSet::FPolygroupSet(FDynamicMesh3* MeshIn, FDynamicMeshPolygroupAttribu
 	RecalculateMaxGroupID();
 }
 
-FPolygroupSet::FPolygroupSet(FDynamicMesh3* MeshIn, int32 PolygroupLayerIndex)
+FPolygroupSet::FPolygroupSet(const FDynamicMesh3* MeshIn, int32 PolygroupLayerIndex)
 {
 	Mesh = MeshIn;
 	if (ensure(Mesh->Attributes()))
@@ -47,7 +47,7 @@ FPolygroupSet::FPolygroupSet(FDynamicMesh3* MeshIn, int32 PolygroupLayerIndex)
 }
 
 
-FPolygroupSet::FPolygroupSet(FDynamicMesh3* MeshIn, FName AttribName)
+FPolygroupSet::FPolygroupSet(const FDynamicMesh3* MeshIn, FName AttribName)
 {
 	Mesh = MeshIn;
 	PolygroupAttrib = UE::Geometry::FindPolygroupLayerByName(*MeshIn, AttribName);
diff --git a/Engine/Plugins/Experimental/GeometryProcessing/Source/DynamicMesh/Private/Polygroups/PolygroupUtil.cpp b/Engine/Plugins/Experimental/GeometryProcessing/Source/DynamicMesh/Private/Polygroups/PolygroupUtil.cpp
--- a/Engine/Plugins/Experimental/GeometryProcessing/Source/DynamicMesh/Private/Polygroups/PolygroupUtil.cpp
+++ b/Engine/Plugins/Experimental/GeometryProcessing/Source/DynamicMesh/Private/Polygroups/PolygroupUtil.cpp
@@ -20,6 +20,20 @@ FDynamicMeshPolygroupAttribute* UE::Geometry::FindPolygroupLayerByName(FDynamicM
 	return nullptr;
 }
 
+const FDynamicMeshPolygroupAttribute* UE::Geometry::FindPolygroupLayerByName(const FDynamicMesh3& Mesh, FName Name)
+{
+	const FDynamicMeshAttributeSet* AttributeSet = Mesh.Attributes();
+	if (AttributeSet == nullptr) return nullptr;
+	int32 NumPolygroupLayers = AttributeSet->NumPolygroupLayers();
+	for (int32 k = 0; k < NumPolygroupLayers; ++k)
+	{
+		if (AttributeSet->GetPolygroupLayer(k)->GetName() == Name)
+		{
+			return AttributeSet->GetPolygroupLayer(k);
+		}
+	}
+	return nullptr;
+}
 
 int32 UE::Geometry::FindPolygroupLayerIndexByName(const FDynamicMesh3& Mesh, FName Name)
 {
diff --git a/Engine/Plugins/Experimental/GeometryProcessing/Source/DynamicMesh/Public/Polygroups/PolygroupSet.h b/Engine/Plugins/Experimental/GeometryProcessing/Source/DynamicMesh/Public/Polygroups/PolygroupSet.h
--- a/Engine/Plugins/Experimental/GeometryProcessing/Source/DynamicMesh/Public/Polygroups/PolygroupSet.h
+++ b/Engine/Plugins/Experimental/GeometryProcessing/Source/DynamicMesh/Public/Polygroups/PolygroupSet.h
@@ -26,32 +26,32 @@ namespace Geometry
  */
 struct DYNAMICMESH_API FPolygroupSet
 {
-	FDynamicMesh3* Mesh = nullptr;
-	FDynamicMeshPolygroupAttribute* PolygroupAttrib = nullptr;
+	const FDynamicMesh3* Mesh = nullptr;
+	const FDynamicMeshPolygroupAttribute* PolygroupAttrib = nullptr;
 	int32 GroupLayerIndex = -1;
 	int32 MaxGroupID = 0; // Note: all group IDs are less than MaxGroupID
 
 	/** Initialize a PolygroupSet for the given Mesh, and standard triangle group layer */
-	explicit FPolygroupSet(FDynamicMesh3* MeshIn);
+	explicit FPolygroupSet(const FDynamicMesh3* MeshIn);
 
 	/** Initialize a PolygroupSet for given Mesh and specific Polygroup attribute layer */
-	explicit FPolygroupSet(FDynamicMesh3* MeshIn, FDynamicMeshPolygroupAttribute* PolygroupAttribIn);
+	explicit FPolygroupSet(const FDynamicMesh3* MeshIn, const FDynamicMeshPolygroupAttribute* PolygroupAttribIn);
 
 	/** Initialize a PolygroupSet for given Mesh and specific Polygroup attribute layer, found by index. If not valid, fall back to standard triangle group layer. */
-	explicit FPolygroupSet(FDynamicMesh3* MeshIn, int32 PolygroupLayerIndex);
+	explicit FPolygroupSet(const FDynamicMesh3* MeshIn, int32 PolygroupLayerIndex);
 
 	/** Initialize a PolygroupSet for given Mesh and specific Polygroup attribute layer, found by name. If not valid, fall back to standard triangle group layer. */
-	explicit FPolygroupSet(FDynamicMesh3* MeshIn, FName AttribName);
+	explicit FPolygroupSet(const FDynamicMesh3* MeshIn, FName AttribName);
 
 	/** Initialize a PolygroupSet by copying an existing PolygroupSet */
 	explicit FPolygroupSet(const FPolygroupSet* CopyIn);
 
 
 	/** @return Mesh this PolygroupSet references  */
-	FDynamicMesh3* GetMesh() { return Mesh; }
+	const FDynamicMesh3* GetMesh() { return Mesh; }
 
 	/** @return PolygroupAttribute this PolygroupSet references, or null if no PolygroupAttribute is in use */
-	FDynamicMeshPolygroupAttribute* GetPolygroup() { return PolygroupAttrib; }
+	const FDynamicMeshPolygroupAttribute* GetPolygroup() { return PolygroupAttrib; }
 
 	/** @return index of current PolygroupAttribute into Mesh AttributeSet, or -1 if this information does not exist */
 	int32 GetPolygroupIndex() const { return GroupLayerIndex; }
@@ -60,7 +60,7 @@ struct DYNAMICMESH_API FPolygroupSet
 	/**
 	 * @return PolygroupID for a TriangleID
 	 */
-	int32 GetGroup(int32 TriangleID)
+	int32 GetGroup(int32 TriangleID) const
 	{
 		return (PolygroupAttrib) ? PolygroupAttrib->GetValue(TriangleID) : Mesh->GetTriangleGroup(TriangleID);
 	}
@@ -68,7 +68,7 @@ struct DYNAMICMESH_API FPolygroupSet
 	/**
 	 * @return PolygroupID for a TriangleID
 	 */
-	int32 GetTriangleGroup(int32 TriangleID)
+	int32 GetTriangleGroup(int32 TriangleID) const
 	{
 		return (PolygroupAttrib) ? PolygroupAttrib->GetValue(TriangleID) : Mesh->GetTriangleGroup(TriangleID);
 	}
@@ -76,17 +76,20 @@ struct DYNAMICMESH_API FPolygroupSet
 	/**
 	 * Set the PolygroupID for a TriangleID
 	 */
-	void SetGroup(int32 TriangleID, int32 NewGroupID)
+	void SetGroup(int32 TriangleID, int32 NewGroupID, FDynamicMesh3& WritableMesh)
 	{
-		if (Mesh->IsTriangle(TriangleID))
+		checkSlow(&WritableMesh == this->Mesh);		// require the same mesh
+		if (WritableMesh.IsTriangle(TriangleID))
 		{
 			if (PolygroupAttrib)
 			{
-				PolygroupAttrib->SetValue(TriangleID, NewGroupID);
+				FDynamicMeshPolygroupAttribute* WritableGroupAttrib = WritableMesh.Attributes()->GetPolygroupLayer(GroupLayerIndex);
+				checkSlow(WritableGroupAttrib == PolygroupAttrib)
+				WritableGroupAttrib->SetValue(TriangleID, NewGroupID);
 			}
 			else
 			{
-				Mesh->SetTriangleGroup(TriangleID, NewGroupID);
+				WritableMesh.SetTriangleGroup(TriangleID, NewGroupID);
 			}
 		}
 		MaxGroupID = FMath::Max(MaxGroupID, NewGroupID + 1);
diff --git a/Engine/Plugins/Experimental/GeometryProcessing/Source/DynamicMesh/Public/Polygroups/PolygroupUtil.h b/Engine/Plugins/Experimental/GeometryProcessing/Source/DynamicMesh/Public/Polygroups/PolygroupUtil.h
--- a/Engine/Plugins/Experimental/GeometryProcessing/Source/DynamicMesh/Public/Polygroups/PolygroupUtil.h
+++ b/Engine/Plugins/Experimental/GeometryProcessing/Source/DynamicMesh/Public/Polygroups/PolygroupUtil.h
@@ -17,6 +17,12 @@ namespace Geometry
  */
 DYNAMICMESH_API FDynamicMeshPolygroupAttribute* FindPolygroupLayerByName(FDynamicMesh3& Mesh, FName Name);
 
+/**
+ * Find the first FDynamicMeshPolygroupAttribute with the given FName in a the AttributeSet of a Mesh.
+ * @return nullptr if no Polygroup layer is found
+ */
+DYNAMICMESH_API const FDynamicMeshPolygroupAttribute* FindPolygroupLayerByName(const FDynamicMesh3& Mesh, FName Name);
+
 /**
  * @return index of the first Layer with the given FName in Mesh AttributeSet, or -1 if not found
  */
diff --git a/Engine/Plugins/Experimental/MeshModelingToolset/Source/MeshModelingTools/Private/MeshGroupPaintTool.cpp b/Engine/Plugins/Experimental/MeshModelingToolset/Source/MeshModelingTools/Private/MeshGroupPaintTool.cpp
--- a/Engine/Plugins/Experimental/MeshModelingToolset/Source/MeshModelingTools/Private/MeshGroupPaintTool.cpp
+++ b/Engine/Plugins/Experimental/MeshModelingToolset/Source/MeshModelingTools/Private/MeshGroupPaintTool.cpp
@@ -556,7 +556,7 @@ void UMeshGroupPaintTool::SyncMeshWithGroupBuffer(FDynamicMesh3* Mesh)
 
 		ActiveGroupEditBuilder->SaveTriangle(TriIdx, CurGroupID, ROIGroupBuffer[k]);
 
-		ActiveGroupSet->SetGroup(TriIdx, ROIGroupBuffer[k]);
+		ActiveGroupSet->SetGroup(TriIdx, ROIGroupBuffer[k], *Mesh);
 
 		//ActiveVertexChange->UpdateVertexColor(VertIdx, OrigColor, NewColor);
 	}
@@ -831,7 +831,7 @@ void UMeshGroupPaintTool::SetTrianglesToGroupID(const TSet<int32>& Triangles, in
 	ActiveGroupEditBuilder->SaveTriangles(TempROIBuffer);
 	for (int32 tid : TempROIBuffer)
 	{
-		ActiveGroupSet->SetGroup(tid, ToGroupID);
+		ActiveGroupSet->SetGroup(tid, ToGroupID, *GetSculptMesh());
 	}
 	ActiveGroupEditBuilder->SaveTriangles(TempROIBuffer);
 
@@ -1227,7 +1227,7 @@ void UMeshGroupPaintTool::GrowCurrentGroupAction()
 	ActiveGroupEditBuilder->SaveTriangles(TempROIBuffer);
 	for (int32 tid : TempROIBuffer)
 	{
-		ActiveGroupSet->SetGroup(tid, CurrentGroupID);
+		ActiveGroupSet->SetGroup(tid, CurrentGroupID, *GetSculptMesh());
 	}
 	ActiveGroupEditBuilder->SaveTriangles(TempROIBuffer);
 
@@ -1254,7 +1254,7 @@ void UMeshGroupPaintTool::ShrinkCurrentGroupAction()
 	for (int32 tid : TempROIBuffer)
 	{
 		// todo: could probably guess boundary groups here...
-		ActiveGroupSet->SetGroup(tid, 0);
+		ActiveGroupSet->SetGroup(tid, 0, *GetSculptMesh());
 	}
 	ActiveGroupEditBuilder->SaveTriangles(TempROIBuffer);
 
@@ -1282,7 +1282,7 @@ void UMeshGroupPaintTool::ClearCurrentGroupAction()
 	ActiveGroupEditBuilder->SaveTriangles(TempROIBuffer);
 	for (int32 tid : TempROIBuffer)
 	{
-		ActiveGroupSet->SetGroup(tid, 0);
+		ActiveGroupSet->SetGroup(tid, 0, *GetSculptMesh());
 	}
 	ActiveGroupEditBuilder->SaveTriangles(TempROIBuffer);
 
@@ -1311,7 +1311,7 @@ void UMeshGroupPaintTool::FloodFillCurrentGroupAction()
 	ActiveGroupEditBuilder->SaveTriangles(TempROIBuffer);
 	for (int32 tid : TempROIBuffer)
 	{
-		ActiveGroupSet->SetGroup(tid, SetGroupID);
+		ActiveGroupSet->SetGroup(tid, SetGroupID, *GetSculptMesh());
 	}
 	ActiveGroupEditBuilder->SaveTriangles(TempROIBuffer);
 
@@ -1338,7 +1338,7 @@ void UMeshGroupPaintTool::ClearAllGroupsAction()
 	ActiveGroupEditBuilder->SaveTriangles(TempROIBuffer);
 	for (int32 tid : TempROIBuffer)
 	{
-		ActiveGroupSet->SetGroup(tid, 0);
+		ActiveGroupSet->SetGroup(tid, 0, *GetSculptMesh());
 	}
 	ActiveGroupEditBuilder->SaveTriangles(TempROIBuffer);
 
diff --git a/Engine/Plugins/Experimental/MeshModelingToolset/Source/MeshModelingTools/Private/MeshSelectionTool.cpp b/Engine/Plugins/Experimental/MeshModelingToolset/Source/MeshModelingTools/Private/MeshSelectionTool.cpp
--- a/Engine/Plugins/Experimental/MeshModelingToolset/Source/MeshModelingTools/Private/MeshSelectionTool.cpp
+++ b/Engine/Plugins/Experimental/MeshModelingToolset/Source/MeshModelingTools/Private/MeshSelectionTool.cpp
@@ -17,6 +17,11 @@
 #include "Selections/MeshConnectedComponents.h"
 #include "Selections/MeshFaceSelection.h"
 #include "ModelingToolTargetUtil.h"
+#include "Properties/MeshStatisticsProperties.h"
+#include "Drawing/MeshElementsVisualizer.h"
+#include "Properties/MeshUVChannelProperties.h"
+#include "PropertySets/PolygroupLayersProperties.h"
+#include "Polygroups/PolygroupUtil.h"
 
 #include "Algo/MaxElement.h"
 
@@ -70,6 +75,23 @@ void UMeshSelectionTool::Setup()
 	SelectionProps->RestoreProperties(this);
 	AddToolPropertySource(SelectionProps);
 
+	PolygroupLayerProperties = NewObject<UPolygroupLayersProperties>(this);
+	PolygroupLayerProperties->RestoreProperties(this);
+	PreviewMesh->ProcessMesh([&](const FDynamicMesh3& ReadMesh) { PolygroupLayerProperties->InitializeGroupLayers(&ReadMesh); });
+	PolygroupLayerProperties->WatchProperty(PolygroupLayerProperties->ActiveGroupLayer, [&](FName) { UpdateActiveGroupLayer(); bColorsUpdatePending = true; bFullMeshInvalidationPending = true; });
+	AddToolPropertySource(PolygroupLayerProperties);
+	UpdateActiveGroupLayer();
+
+	UVChannelProperties = NewObject<UMeshUVChannelProperties>(this);
+	UVChannelProperties->RestoreProperties(this);
+	PreviewMesh->ProcessMesh([&](const FDynamicMesh3& ReadMesh) { UVChannelProperties->Initialize(&ReadMesh, false); });
+	UVChannelProperties->ValidateSelection(true);
+	UVChannelProperties->WatchProperty(UVChannelProperties->UVChannel, [this](const FString& NewValue) { CacheUVIslandIDs(); bColorsUpdatePending = true; bFullMeshInvalidationPending = true; });
+	AddToolPropertySource(UVChannelProperties);
+
+	// we could probably calculate this on-demand but we need to do it before making any mesh changes? or update?
+	CacheUVIslandIDs();
+
 	AddSubclassPropertySets();
 
 	SelectionActions = NewObject<UMeshSelectionEditActions>(this);
@@ -82,9 +104,6 @@ void UMeshSelectionTool::Setup()
 	// set autocalculated tangents
 	PreviewMesh->SetTangentsMode(EDynamicMeshComponentTangentsMode::AutoCalculated);
 
-	// enable wireframe on component
-	PreviewMesh->EnableWireframe(SelectionProps->bShowWireframe);
-
 	// disable shadows
 	PreviewMesh->SetShadowsEnabled(false);
 
@@ -109,9 +128,6 @@ void UMeshSelectionTool::Setup()
 	SelectedVertices = TBitArray<>(false, Mesh->MaxVertexID());
 	SelectedTriangles = TBitArray<>(false, Mesh->MaxTriangleID());
 
-	// we could probably calculate this on-demand but we need to do it before making any mesh changes? or update?
-	CacheUVIslandIDs();
-
 	this->Selection = NewObject<UMeshSelectionSet>(this);
 	Selection->GetOnModified().AddLambda([this](USelectionSet* SelectionObj)
 	{
@@ -119,10 +135,14 @@ void UMeshSelectionTool::Setup()
 	});
 
 	// rebuild octree if mesh changes
-	PreviewMesh->GetOnMeshChanged().AddLambda([this]() { bOctreeValid = false; bFullMeshInvalidationPending = true; });
+	PreviewMesh->GetOnMeshChanged().AddLambda([this]() { 
+		bOctreeValid = false; 
+		bFullMeshInvalidationPending = true; 
+		bColorsUpdatePending = true; 
+		CacheUVIslandIDs(); 
+		UpdateActiveGroupLayer();
+	});
 
-	SelectionProps->WatchProperty(SelectionProps->bShowWireframe,
-								  [this](bool bNewValue) { PreviewMesh->EnableWireframe(bNewValue); });
 	SelectionProps->WatchProperty(SelectionProps->FaceColorMode,
 								  [this](EMeshFacesColorMode NewValue)
 								  {
@@ -130,6 +150,24 @@ void UMeshSelectionTool::Setup()
 								  });
 	bColorsUpdatePending = (SelectionProps->FaceColorMode != EMeshFacesColorMode::None);
 
+	MeshElementsDisplay = NewObject<UMeshElementsVisualizer>(this);
+	MeshElementsDisplay->CreateInWorld(PreviewMesh->GetWorld(), PreviewMesh->GetTransform());
+	if (ensure(MeshElementsDisplay->Settings))
+	{
+		MeshElementsDisplay->Settings->bShowWireframe = true;
+		MeshElementsDisplay->Settings->bShowUVSeams = false;
+		MeshElementsDisplay->Settings->bShowNormalSeams = false;
+		MeshElementsDisplay->Settings->bShowColorSeams = false;
+		MeshElementsDisplay->Settings->RestoreProperties(this, TEXT("MeshSelection"));
+		AddToolPropertySource(MeshElementsDisplay->Settings);
+	}
+	MeshElementsDisplay->SetMeshAccessFunction([this](UMeshElementsVisualizer::ProcessDynamicMeshFunc ProcessFunc) {
+		PreviewMesh->ProcessMesh(ProcessFunc);
+	});
+
+	MeshStatisticsProperties = NewObject<UMeshStatisticsProperties>(this);
+	AddToolPropertySource(MeshStatisticsProperties);
+
 	RecalculateBrushRadius();
 	UpdateVisualization(true);
 
@@ -152,14 +190,10 @@ UMeshSelectionToolActionPropertySet* UMeshSelectionTool::CreateEditActions()
 
 
 
-void UMeshSelectionTool::OnShutdown(EToolShutdownType ShutdownType)
+void UMeshSelectionTool::ApplyShutdownAction(EToolShutdownType ShutdownType)
 {
-	SelectionProps->SaveProperties(this);
-	BrushProperties->SaveProperties(this);
-
 	if (bHaveModifiedMesh && ShutdownType == EToolShutdownType::Accept)
 	{
-		// this block bakes the modified DynamicMeshComponent back into the StaticMeshComponent inside an undo transaction
 		GetToolManager()->BeginUndoTransaction(LOCTEXT("MeshSelectionToolTransactionName", "Edit Mesh"));
 		UE::ToolTarget::CommitDynamicMeshUpdate(Target, *PreviewMesh->GetMesh(), true);
 		GetToolManager()->EndUndoTransaction();
@@ -174,6 +208,23 @@ void UMeshSelectionTool::OnShutdown(EToolShutdownType ShutdownType)
 }
 
 
+void UMeshSelectionTool::OnShutdown(EToolShutdownType ShutdownType)
+{
+	SelectionProps->SaveProperties(this);
+	BrushProperties->SaveProperties(this);
+	UVChannelProperties->SaveProperties(this);
+	PolygroupLayerProperties->SaveProperties(this);
+
+	if (ensure(MeshElementsDisplay->Settings))
+	{
+		MeshElementsDisplay->Settings->SaveProperties(this, TEXT("MeshSelection"));
+	}
+	MeshElementsDisplay->Disconnect();
+
+	ApplyShutdownAction(ShutdownType);
+}
+
+
 
 
 void UMeshSelectionTool::RegisterActions(FInteractiveToolActionSet& ActionSet)
@@ -204,13 +255,6 @@ void UMeshSelectionTool::RegisterActions(FInteractiveToolActionSet& ActionSet)
 		[this]() { DeleteSelectedTriangles(); });
 
 
-	ActionSet.RegisterAction(this, (int32)EStandardToolActions::ToggleWireframe,
-		TEXT("ToggleWireframe"),
-		LOCTEXT("ToggleWireframe", "Toggle Wireframe"),
-		LOCTEXT("ToggleWireframeTooltip", "Toggle visibility of wireframe overlay"),
-		EModifierKey::Alt, EKeys::W,
-		[this]() { SelectionProps->bShowWireframe = !SelectionProps->bShowWireframe; });
-
 #if WITH_EDITOR  	// enum HasMetaData()  is not available at runtime
 	ActionSet.RegisterAction(this, (int32)EMeshSelectionToolActions::CycleSelectionMode,
 		TEXT("CycleSelectionMode"),
@@ -500,7 +544,7 @@ void UMeshSelectionTool::UpdateFaceSelection(const FBrushStampData& Stamp, const
 	else if (SelectionProps->SelectionMode == EMeshSelectionToolPrimaryMode::AllInGroup)
 	{
 		FMeshConnectedComponents::GrowToConnectedTriangles(Mesh, TriangleROI, LocalROI, &TemporaryBuffer, &TemporarySet,
-			[Mesh](int t1, int t2) { return Mesh->GetTriangleGroup(t1) == Mesh->GetTriangleGroup(t2); } );
+			[&](int t1, int t2) { return ActiveGroupSet->GetTriangleGroup(t1) == ActiveGroupSet->GetTriangleGroup(t2); } );
 		UseROI = &LocalROI;
 	}
 	else if (SelectionProps->SelectionMode == EMeshSelectionToolPrimaryMode::ByMaterial)
@@ -700,7 +744,19 @@ void UMeshSelectionTool::UpdateVisualization(bool bSelectionModified)
 {
 	check(SelectionType == EMeshSelectionElementType::Face);  // only face selection supported so far
 
-	bFullMeshInvalidationPending = false;
+	SetToolPropertySourceEnabled(UVChannelProperties,
+								 SelectionProps->SelectionMode == EMeshSelectionToolPrimaryMode::ByUVIsland ||
+								 SelectionProps->FaceColorMode == EMeshFacesColorMode::ByUVIsland);
+
+	if (bFullMeshInvalidationPending)
+	{
+		PreviewMesh->ProcessMesh([&](const FDynamicMesh3& ReadMesh)
+		{
+			MeshStatisticsProperties->Update(ReadMesh);
+		});
+		MeshElementsDisplay->NotifyMeshChanged();
+		bFullMeshInvalidationPending = false;
+	}
 
 	// force an update of renderbuffers
 	if (bSelectionModified)
@@ -735,7 +791,7 @@ FColor UMeshSelectionTool::GetCurrentFaceColor(const FDynamicMesh3* Mesh, int Tr
 {
 	if (SelectionProps->FaceColorMode == EMeshFacesColorMode::ByGroup)
 	{
-		return LinearColors::SelectFColor(Mesh->GetTriangleGroup(TriangleID));
+		return LinearColors::SelectFColor(ActiveGroupSet->GetTriangleGroup(TriangleID));
 	}
 	else if (SelectionProps->FaceColorMode == EMeshFacesColorMode::ByMaterialID)
 	{
@@ -757,7 +813,8 @@ void UMeshSelectionTool::CacheUVIslandIDs()
 
 	TriangleToUVIsland.SetNum(Mesh->MaxTriangleID());
 
-	const FDynamicMeshUVOverlay* UV = Mesh->Attributes()->GetUVLayer(0);
+	int32 ActiveUVLayer = UVChannelProperties->GetSelectedChannelIndex(true);
+	const FDynamicMeshUVOverlay* UV = Mesh->Attributes()->GetUVLayer(ActiveUVLayer);
 
 	Components.FindConnectedTriangles([&](int32 TriIdx0, int32 TriIdx1)
 	{
@@ -775,6 +832,26 @@ void UMeshSelectionTool::CacheUVIslandIDs()
 }
 
 
+void UMeshSelectionTool::UpdateActiveGroupLayer()
+{
+	// todo: need this to be const
+	const FDynamicMesh3* SourceMesh = PreviewMesh->GetMesh();
+
+	if (PolygroupLayerProperties->HasSelectedPolygroup() == false)
+	{
+		ActiveGroupSet = MakeShared<UE::Geometry::FPolygroupSet, ESPMode::ThreadSafe>(SourceMesh);
+	}
+	else
+	{
+		FName SelectedName = PolygroupLayerProperties->ActiveGroupLayer;
+		const FDynamicMeshPolygroupAttribute* FoundAttrib = UE::Geometry::FindPolygroupLayerByName(*SourceMesh, SelectedName);
+		ensureMsgf(FoundAttrib, TEXT("Selected Attribute Not Found! Falling back to Default group layer."));
+		ActiveGroupSet = MakeShared<UE::Geometry::FPolygroupSet, ESPMode::ThreadSafe>(SourceMesh, FoundAttrib);
+	}
+
+}
+
+
 void UMeshSelectionTool::Render(IToolsContextRenderAPI* RenderAPI)
 {
 	UDynamicMeshBrushTool::Render(RenderAPI);
@@ -815,6 +892,13 @@ void UMeshSelectionTool::OnTick(float DeltaTime)
 		bHavePendingAction = false;
 		PendingAction = EMeshSelectionToolActions::NoAction;
 	}
+
+	if (bFullMeshInvalidationPending)
+	{
+		UpdateVisualization(false);
+	}
+
+	MeshElementsDisplay->OnTick(DeltaTime);
 }
 
 
@@ -1061,18 +1145,18 @@ void UMeshSelectionTool::GrowShrinkSelection(bool bGrow)
 			}
 		}
 	}
-	if( Mesh->HasTriangleGroups() && (SelectionProps->SelectionMode == EMeshSelectionToolPrimaryMode::AllInGroup) )
+	if( SelectionProps->SelectionMode == EMeshSelectionToolPrimaryMode::AllInGroup )
 	{
 		TSet<int32> AdjacentFaces{ChangeFaces};
 		TSet<int32> AdjacentGroups{};
 		ChangeFaces.Empty();
 		for ( int32 TID : AdjacentFaces )
 		{
-			AdjacentGroups.Add(Mesh->GetTriangleGroup(TID));
+			AdjacentGroups.Add(ActiveGroupSet->GetTriangleGroup(TID));
 		}
 		for ( int32 TID : Mesh->TriangleIndicesItr() )
 		{
-			if ( AdjacentGroups.Contains(Mesh->GetTriangleGroup(TID)) )
+			if ( AdjacentGroups.Contains(ActiveGroupSet->GetTriangleGroup(TID)) )
 			{
 				ChangeFaces.Add(TID);
 			}
@@ -1463,7 +1547,7 @@ void UMeshSelectionTool::AssignNewGroupToSelectedTriangles()
 	// assign new groups to triangles
 	// note: using an FMeshChange is kind of overkill here
 	TUniquePtr<FMeshChange> MeshChange = PreviewMesh->TrackedEditMesh(
-		[&SelectedFaces](FDynamicMesh3& Mesh, FDynamicMeshChangeTracker& ChangeTracker)
+		[&SelectedFaces, this](FDynamicMesh3& Mesh, FDynamicMeshChangeTracker& ChangeTracker)
 	{
 		// each component gets its own group id
 		FMeshConnectedComponents Components(&Mesh);
@@ -1475,7 +1559,7 @@ void UMeshSelectionTool::AssignNewGroupToSelectedTriangles()
 			for (int tid : Component.Indices)
 			{
 				ChangeTracker.SaveTriangle(tid, true);
-				Mesh.SetTriangleGroup(tid, NewGroupID);
+				ActiveGroupSet->SetGroup(tid, NewGroupID, Mesh);
 			}
 		}
 	});
diff --git a/Engine/Plugins/Experimental/MeshModelingToolset/Source/MeshModelingTools/Public/MeshSelectionTool.h b/Engine/Plugins/Experimental/MeshModelingToolset/Source/MeshModelingTools/Public/MeshSelectionTool.h
--- a/Engine/Plugins/Experimental/MeshModelingToolset/Source/MeshModelingTools/Public/MeshSelectionTool.h
+++ b/Engine/Plugins/Experimental/MeshModelingToolset/Source/MeshModelingTools/Public/MeshSelectionTool.h
@@ -7,8 +7,14 @@
 #include "SelectionSet.h"
 #include "Changes/MeshSelectionChange.h"
 #include "DynamicMesh/DynamicMeshOctree3.h"
+#include "Polygroups/PolygroupSet.h"
 #include "MeshSelectionTool.generated.h"
 
+class UMeshStatisticsProperties;
+class UMeshElementsVisualizer;
+class UMeshUVChannelProperties;
+class UPolygroupLayersProperties;
+
 /**
  *
  */
@@ -74,57 +80,65 @@ class MESHMODELINGTOOLS_API UMeshSelectionEditActions : public UMeshSelectionToo
 	GENERATED_BODY()
 
 public:
-	UFUNCTION(CallInEditor, Category = SelectionEdits, meta = (DisplayName = "Select All", DisplayPriority = 1))
-	void SelectAll()
+	/** Clear the active triangle selection */
+	UFUNCTION(CallInEditor, Category = SelectionEdits, meta = (DisplayPriority = 0))
+	void Clear()
 	{
-		PostAction(EMeshSelectionToolActions::SelectAll);
+		PostAction(EMeshSelectionToolActions::ClearSelection);
 	}
 
-	UFUNCTION(CallInEditor, Category = SelectionEdits, meta = (DisplayName = "Clear", DisplayPriority = 1))
-	void Clear()
+	/** Select all triangles in the mesh */
+	UFUNCTION(CallInEditor, Category = SelectionEdits, meta = (DisplayPriority = 1))
+	void SelectAll()
 	{
-		PostAction(EMeshSelectionToolActions::ClearSelection);
+		PostAction(EMeshSelectionToolActions::SelectAll);
 	}
 
-	UFUNCTION(CallInEditor, Category = SelectionEdits, meta = (DisplayName = "Invert", DisplayPriority = 2))
+	/** Invert the active triangle selection */
+	UFUNCTION(CallInEditor, Category = SelectionEdits, meta = (DisplayPriority = 2))
 	void Invert()
 	{
 		PostAction(EMeshSelectionToolActions::InvertSelection);
 	}
 
-
-	UFUNCTION(CallInEditor, Category = SelectionEdits, meta = (DisplayName = "Grow", DisplayPriority = 3))
+	/** Grow the active triangle selection to include any triangles touching a vertex on the selection boundary */
+	UFUNCTION(CallInEditor, Category = SelectionEdits, meta = (DisplayPriority = 3))
 	void Grow()
 	{
 		PostAction(EMeshSelectionToolActions::GrowSelection);
 	}
 
-	UFUNCTION(CallInEditor, Category = SelectionEdits, meta = (DisplayName = "Shrink", DisplayPriority = 4))
+	/** Shrink the active triangle selection by removing any triangles touching a vertex on the selection boundary */
+	UFUNCTION(CallInEditor, Category = SelectionEdits, meta = ( DisplayPriority = 4))
 	void Shrink()
 	{
 		PostAction(EMeshSelectionToolActions::ShrinkSelection);
 	}
 
-	UFUNCTION(CallInEditor, Category = SelectionEdits, meta = (DisplayName = "ExpandToConnected", DisplayPriority = 5))
-	void ExpandToConnected()
+	/** Grow the active selection to include any triangle connected via shared edges (ie flood-fill) */
+	UFUNCTION(CallInEditor, Category = SelectionEdits, meta = (DisplayPriority = 5))
+	void FloodFill()
 	{
 		PostAction(EMeshSelectionToolActions::ExpandToConnected);
 	}
 
-	UFUNCTION(CallInEditor, Category = SelectionEdits, meta = (DisplayName = "SelectLargestComponentByTriCount", DisplayPriority = 5))
-	void SelectLargestComponentByTriCount()
+	/** Select the largest connected mesh component by triangle count */
+	UFUNCTION(CallInEditor, Category = SelectionEdits, meta = (DisplayPriority = 6))
+	void LargestAreaPart()
 	{
 		PostAction(EMeshSelectionToolActions::SelectLargestComponentByTriCount);
 	}
 
-	UFUNCTION(CallInEditor, Category = SelectionEdits, meta = (DisplayName = "SelectLargestComponentByArea", DisplayPriority = 5))
-	void SelectLargestComponentByArea()
+	/** Select the largest connected mesh component by mesh area */
+	UFUNCTION(CallInEditor, Category = SelectionEdits, meta = (DisplayPriority = 7))
+	void LargestTriCountPart()
 	{
 		PostAction(EMeshSelectionToolActions::SelectLargestComponentByArea);
 	}
 
-	UFUNCTION(CallInEditor, Category = SelectionEdits, meta = (DisplayName = "OptimizeSelection", DisplayPriority = 5))
-	void OptimizeSelection()
+	/** Optimize the selection border */
+	UFUNCTION(CallInEditor, Category = SelectionEdits, meta = (DisplayPriority = 8))
+	void OptimizeBorder()
 	{
 		PostAction(EMeshSelectionToolActions::OptimizeSelection);
 	}
@@ -140,35 +154,41 @@ class MESHMODELINGTOOLS_API UMeshSelectionMeshEditActions : public UMeshSelectio
 	GENERATED_BODY()
 
 public:
+	/** Delete the selected triangles */
 	UFUNCTION(CallInEditor, Category = MeshEdits, meta = (DisplayName = "Delete", DisplayPriority = 1))
-	void DeleteTriangles()
+	void Delete()
 	{
 		PostAction(EMeshSelectionToolActions::DeleteSelected);
 	}
 
-	UFUNCTION(CallInEditor, Category = MeshEdits, meta = (DisplayName = "Separate", DisplayPriority = 2))
-	void SeparateTriangles() 
-	{
-		PostAction(EMeshSelectionToolActions::SeparateSelected);
-	}
-
+	/** Disconnected the selected triangles from their neighbours, to create mesh boundaries along the selection borders */
 	UFUNCTION(CallInEditor, Category = MeshEdits, meta = (DisplayName = "Disconnect", DisplayPriority = 3))
-	void DisconnectTriangles() 
+	void Disconnect() 
 	{
 		PostAction(EMeshSelectionToolActions::DisconnectSelected);
 	}
 
+	/** Flip the normals of the selected triangles. This will create hard normals at selection borders. */
 	UFUNCTION(CallInEditor, Category = MeshEdits, meta = (DisplayName = "Flip Normals", DisplayPriority = 4))
 	void FlipNormals() 
 	{
 		PostAction(EMeshSelectionToolActions::FlipSelected);
 	}
 
+	/** Assign a new unique Polygroup index to the selected triangles */
 	UFUNCTION(CallInEditor, Category = MeshEdits, meta = (DisplayName = "Create Polygroup", DisplayPriority = 5))
 	void CreatePolygroup()
 	{
 		PostAction(EMeshSelectionToolActions::CreateGroup);
 	}
+
+	/** Delete the selected triangles from the active Mesh Object and create a new Mesh containing those triangles */
+	UFUNCTION(CallInEditor, Category = MeshEdits, meta = (DisplayName = "Separate", DisplayPriority = 10))
+	void Separate() 
+	{
+		PostAction(EMeshSelectionToolActions::SeparateSelected);
+	}
+
 };
 
 
@@ -241,14 +261,11 @@ public:
 	UPROPERTY(EditAnywhere, Category = Selection)
 	bool bHitBackFaces = true;
 
-	/** Toggle drawing of wireframe overlay on/off [Alt+W] */
-	UPROPERTY(EditAnywhere, Category = Selection)
-	bool bShowWireframe = true;
-
 	/** Toggle drawing of highlight points on/off */
 	UPROPERTY(EditAnywhere, Category = Selection)
 	bool bShowPoints = false;
 
+	/** Color each triangle based on the selected mesh attribute */
 	UPROPERTY(EditAnywhere, Category = Selection)
 	EMeshFacesColorMode FaceColorMode = EMeshFacesColorMode::None;
 };
@@ -302,6 +319,17 @@ public:
 	UPROPERTY()
 	UMeshSelectionToolActionPropertySet* EditActions;
 
+	UPROPERTY()
+	UMeshStatisticsProperties* MeshStatisticsProperties;
+
+	UPROPERTY()
+	UMeshElementsVisualizer* MeshElementsDisplay;
+
+	UPROPERTY()
+	UMeshUVChannelProperties* UVChannelProperties = nullptr;
+
+	UPROPERTY()
+	UPolygroupLayersProperties* PolygroupLayerProperties = nullptr;
 
 protected:
 	virtual UMeshSelectionToolActionPropertySet* CreateEditActions();
@@ -385,10 +413,15 @@ protected:
 	void SeparateSelectedTriangles(); // separates out selected triangles to a new mesh, removing them from the current mesh
 	void FlipSelectedTriangles();
 	void AssignNewGroupToSelectedTriangles();
-	
+
+	TSharedPtr<UE::Geometry::FPolygroupSet, ESPMode::ThreadSafe> ActiveGroupSet;
+	void OnSelectedGroupLayerChanged();
+	void UpdateActiveGroupLayer();
 
 	// if true, mesh has been edited
 	bool bHaveModifiedMesh = false;
+
+	virtual void ApplyShutdownAction(EToolShutdownType ShutdownType);
 };
 
 
diff --git a/Engine/Plugins/Experimental/MeshModelingToolset/Source/MeshModelingTools/Public/RecomputeUVsTool.h b/Engine/Plugins/Experimental/MeshModelingToolset/Source/MeshModelingTools/Public/RecomputeUVsTool.h
--- a/Engine/Plugins/Experimental/MeshModelingToolset/Source/MeshModelingTools/Public/RecomputeUVsTool.h
+++ b/Engine/Plugins/Experimental/MeshModelingToolset/Source/MeshModelingTools/Public/RecomputeUVsTool.h
@@ -80,7 +80,7 @@ public:
 	ERecomputeUVsIslandMode IslandMode = ERecomputeUVsIslandMode::PolyGroups;
 
 	UPROPERTY(EditAnywhere, Category = "Compute UVs")
-	ERecomputeUVsUnwrapType UnwrapType = ERecomputeUVsUnwrapType::ExpMap;
+	ERecomputeUVsUnwrapType UnwrapType = ERecomputeUVsUnwrapType::Conformal;
 
 	UPROPERTY(EditAnywhere, Category = "UV Layout")
 	bool bAutoPack = true;
diff --git a/Engine/Plugins/Experimental/MeshModelingToolset/Source/MeshModelingToolsEditorOnly/Private/EditMeshMaterialsTool.cpp b/Engine/Plugins/Experimental/MeshModelingToolset/Source/MeshModelingToolsEditorOnly/Private/EditMeshMaterialsTool.cpp
--- a/Engine/Plugins/Experimental/MeshModelingToolset/Source/MeshModelingToolsEditorOnly/Private/EditMeshMaterialsTool.cpp
+++ b/Engine/Plugins/Experimental/MeshModelingToolset/Source/MeshModelingToolsEditorOnly/Private/EditMeshMaterialsTool.cpp
@@ -13,10 +13,9 @@
 #include "MeshRegionBoundaryLoops.h"
 #include "DynamicMesh/MeshIndexUtil.h"
 #include "ToolSetupUtil.h"
+#include "ModelingToolTargetUtil.h"
+#include "Materials/MaterialInterface.h"
 
-#include "TargetInterfaces/MaterialProvider.h"
-#include "TargetInterfaces/MeshDescriptionCommitter.h"
-#include "TargetInterfaces/PrimitiveComponentBackedTarget.h"
 
 #include "ExplicitUseGeometryMathTypes.h"		// using UE::Geometry::(math types)
 using namespace UE::Geometry;
@@ -32,6 +31,50 @@ void UEditMeshMaterialsEditActions::PostMaterialAction(EEditMeshMaterialsToolAct
 }
 
 
+
+void UEditMeshMaterialsToolProperties::UpdateFromMaterialsList()
+{
+	MaterialNamesList.Reset();
+
+	for ( int32 k = 0; k < Materials.Num(); ++k)
+	{
+		UMaterialInterface* Mat = Materials[k];
+		FString MatName = (Mat != nullptr) ? Mat->GetName() : "(none)";
+		FString UseName = FString::Printf(TEXT("[%d] %s"), k, *MatName);
+		MaterialNamesList.Add(UseName);
+	}
+
+	if (MaterialNamesList.Num() == 0)
+	{
+		ActiveMaterial = TEXT("(no materials)");
+		return;
+	}
+
+	// udpate active material if it no longer exists
+	for (int32 k = 0; k < MaterialNamesList.Num(); ++k)
+	{
+		if (MaterialNamesList[k] == ActiveMaterial)
+		{
+			return;
+		}
+	}
+	ActiveMaterial = MaterialNamesList[0];
+}
+
+int32 UEditMeshMaterialsToolProperties::GetSelectedMaterialIndex() const
+{
+	for (int32 k = 0; k < MaterialNamesList.Num(); ++k)
+	{
+		if (MaterialNamesList[k] == ActiveMaterial)
+		{
+			return k;
+		}
+	}
+	return 0;
+}
+
+
+
 /*
  * ToolBuilder
  */
@@ -54,9 +97,7 @@ void UEditMeshMaterialsTool::Setup()
 
 	PreviewMesh->ClearOverrideRenderMaterial();
 
-	IMaterialProvider* TargetMaterial = Cast<IMaterialProvider>(Target);
-	FComponentMaterialSet AssetMaterials;
-	TargetMaterial->GetMaterialSet(AssetMaterials, true);
+	FComponentMaterialSet AssetMaterials = UE::ToolTarget::GetMaterialSet(Target, true);
 	MaterialProps->Materials = AssetMaterials.Materials;
 	CurrentMaterials = MaterialProps->Materials;
 	InitialMaterialKey = GetMaterialKey();
@@ -65,8 +106,7 @@ void UEditMeshMaterialsTool::Setup()
 		[this](){ return GetMaterialKey(); },
 		[this](FMaterialSetKey NewKey) { OnMaterialSetChanged(); });
 
-	FComponentMaterialSet ComponentMaterials;
-	TargetMaterial->GetMaterialSet(ComponentMaterials, false);
+	FComponentMaterialSet ComponentMaterials = UE::ToolTarget::GetMaterialSet(Target, false);
 	if (ComponentMaterials != AssetMaterials)
 	{
 		GetToolManager()->DisplayMessage(
@@ -154,7 +194,7 @@ void UEditMeshMaterialsTool::AssignMaterialToSelectedTriangles()
 	TUniquePtr<FToolCommandChange> SelectionChange = EndChange();
 	ChangeSeq->AppendChange(Selection, MoveTemp(SelectionChange));
 
-	int32 SetMaterialID = MaterialProps->SelectedMaterial;
+	int32 SetMaterialID = MaterialProps->GetSelectedMaterialIndex();
 
 	// assign new groups to triangles
 	// note: using an FMeshChange is kind of overkill here
@@ -176,6 +216,7 @@ void UEditMeshMaterialsTool::AssignMaterialToSelectedTriangles()
 	// emit combined change sequence
 	GetToolManager()->EmitObjectChange(this, MoveTemp(ChangeSeq), LOCTEXT("MeshSelectionToolAssignMaterial", "Assign Material"));
 
+	bFullMeshInvalidationPending = true;
 	OnExternalSelectionChange();
 	bHaveModifiedMesh = true;
 }
@@ -195,6 +236,8 @@ void UEditMeshMaterialsTool::OnMaterialSetChanged()
 
 	GetToolManager()->EmitObjectChange(this, MoveTemp(MaterialChange), LOCTEXT("MaterialSetChange", "Material Change"));
 
+	MaterialProps->UpdateFromMaterialsList();
+
 	bHaveModifiedMaterials = true;
 }
 
@@ -212,13 +255,8 @@ void UEditMeshMaterialsTool::ExternalUpdateMaterialSet(const TArray<UMaterialInt
 
 
 
-
-
-void UEditMeshMaterialsTool::OnShutdown(EToolShutdownType ShutdownType)
+void UEditMeshMaterialsTool::ApplyShutdownAction(EToolShutdownType ShutdownType)
 {
-	// this is a bit of a hack, UMeshSelectionTool::OnShutdown will also do this...
-	SelectionProps->SaveProperties(this);
-
 	if (ShutdownType == EToolShutdownType::Accept)
 	{
 		GetToolManager()->BeginUndoTransaction(LOCTEXT("EditMeshMaterialsTransactionName", "Edit Materials"));
@@ -227,13 +265,20 @@ void UEditMeshMaterialsTool::OnShutdown(EToolShutdownType ShutdownType)
 		{
 			FComponentMaterialSet NewMaterialSet;
 			NewMaterialSet.Materials = CurrentMaterials;
-			Cast<IMaterialProvider>(Target)->CommitMaterialSetUpdate(NewMaterialSet, true);
+			UE::ToolTarget::CommitMaterialSetUpdate(Target, NewMaterialSet, true);
 		}
 
-		UMeshSelectionTool::OnShutdown(ShutdownType);
+		if (bHaveModifiedMesh)
+		{
+			UE::ToolTarget::CommitDynamicMeshUpdate(Target, *PreviewMesh->GetMesh(), true);
+		}
 
 		GetToolManager()->EndUndoTransaction();
 	}
+	else
+	{
+		UMeshSelectionTool::ApplyShutdownAction(ShutdownType);
+	}
 }
 
 
diff --git a/Engine/Plugins/Experimental/MeshModelingToolset/Source/MeshModelingToolsEditorOnly/Public/EditMeshMaterialsTool.h b/Engine/Plugins/Experimental/MeshModelingToolset/Source/MeshModelingToolsEditorOnly/Public/EditMeshMaterialsTool.h
--- a/Engine/Plugins/Experimental/MeshModelingToolset/Source/MeshModelingToolsEditorOnly/Public/EditMeshMaterialsTool.h
+++ b/Engine/Plugins/Experimental/MeshModelingToolset/Source/MeshModelingToolsEditorOnly/Public/EditMeshMaterialsTool.h
@@ -28,9 +28,18 @@ class MESHMODELINGTOOLSEDITORONLY_API UEditMeshMaterialsToolProperties : public
 {
 	GENERATED_BODY()
 public:
+	/** */
+	UPROPERTY(EditAnywhere, Category = Materials, meta = (TransientToolProperty, DisplayName = "Active Material", GetOptions = GetMaterialNamesFunc, NoResetToDefault))
+	FString ActiveMaterial;
 
-	UPROPERTY(EditAnywhere, Category = Materials, meta = (ArrayClamp = "Materials"))
-	int SelectedMaterial = 0;
+	UFUNCTION()
+	const TArray<FString>& GetMaterialNamesFunc() { return MaterialNamesList; }
+
+	UPROPERTY(meta = (TransientToolProperty))
+	TArray<FString> MaterialNamesList;
+
+	void UpdateFromMaterialsList();
+	int32 GetSelectedMaterialIndex() const;
 
 	UPROPERTY(EditAnywhere, Category=Materials)
 	TArray<UMaterialInterface*> Materials;
@@ -55,8 +64,8 @@ class MESHMODELINGTOOLSEDITORONLY_API UEditMeshMaterialsEditActions : public UMe
 	GENERATED_BODY()
 
 public:
-	UFUNCTION(CallInEditor, Category = MaterialEdits, meta = (DisplayName = "Assign Selected Material", DisplayPriority = 1))
-	void AssignSelectedMaterial()
+	UFUNCTION(CallInEditor, Category = MaterialEdits, meta = (DisplayName = "Assign Active Material", DisplayPriority = 1))
+	void AssignActiveMaterial()
 	{
 		PostMaterialAction(EEditMeshMaterialsToolActions::AssignMaterial);
 	}
@@ -83,8 +92,6 @@ public:
 
 	virtual bool CanAccept() const override { return UMeshSelectionTool::CanAccept() || bHaveModifiedMaterials; }
 
-	virtual void OnShutdown(EToolShutdownType ShutdownType) override;
-
 	void RequestMaterialAction(EEditMeshMaterialsToolActions ActionType);
 
 protected:
@@ -113,6 +120,8 @@ protected:
 	FMaterialSetKey InitialMaterialKey;
 	bool bHaveModifiedMaterials = false;
 
+	virtual void ApplyShutdownAction(EToolShutdownType ShutdownType) override;
+
 	void ExternalUpdateMaterialSet(const TArray<UMaterialInterface*>& NewMaterialSet);
 	friend class FEditMeshMaterials_MaterialSetChange;
 };
diff --git a/Engine/Plugins/Experimental/MeshModelingToolset/Source/ModelingOperators/Private/CleaningOps/RemoveOccludedTrianglesOp.cpp b/Engine/Plugins/Experimental/MeshModelingToolset/Source/ModelingOperators/Private/CleaningOps/RemoveOccludedTrianglesOp.cpp
--- a/Engine/Plugins/Experimental/MeshModelingToolset/Source/ModelingOperators/Private/CleaningOps/RemoveOccludedTrianglesOp.cpp
+++ b/Engine/Plugins/Experimental/MeshModelingToolset/Source/ModelingOperators/Private/CleaningOps/RemoveOccludedTrianglesOp.cpp
@@ -47,13 +47,13 @@ void FRemoveOccludedTrianglesOp::CalculateResult(FProgressCancel* Progress)
 			return FIndex2i::Invalid();
 		}
 
-		auto SetGroup = [](UE::Geometry::FPolygroupSet& ActiveGroupSet, const TArray<int>& SelectedTris)
+		auto SetGroup = [](FDynamicMesh3& WriteMesh, UE::Geometry::FPolygroupSet& ActiveGroupSet, const TArray<int>& SelectedTris)
 		{
 			int32 JacketGroupID = ActiveGroupSet.AllocateNewGroupID();
 
 			for (int TID : SelectedTris)
 			{
-				ActiveGroupSet.SetGroup(TID, JacketGroupID);
+				ActiveGroupSet.SetGroup(TID, JacketGroupID, WriteMesh);
 			}
 
 			return FIndex2i(JacketGroupID, ActiveGroupSet.GroupLayerIndex);
@@ -62,12 +62,12 @@ void FRemoveOccludedTrianglesOp::CalculateResult(FProgressCancel* Progress)
 		if (!bActiveGroupIsDefault)
 		{
 			UE::Geometry::FPolygroupSet ActiveGroupSet(&Mesh, LayerName);
-			return SetGroup(ActiveGroupSet, SelectedTris);
+			return SetGroup(Mesh, ActiveGroupSet, SelectedTris);
 		}
 		else
 		{
 			UE::Geometry::FPolygroupSet ActiveGroupSet(&Mesh);
-			return SetGroup(ActiveGroupSet, SelectedTris);
+			return SetGroup(Mesh, ActiveGroupSet, SelectedTris);
 		}
 	};
 
