MeshModelingTools: Exposing planar tolerance parameter to user in mirror tool.

#rnx

diff --git a/Engine/Plugins/Experimental/MeshModelingToolset/Source/MeshModelingTools/Private/MirrorTool.cpp b/Engine/Plugins/Experimental/MeshModelingToolset/Source/MeshModelingTools/Private/MirrorTool.cpp
--- a/Engine/Plugins/Experimental/MeshModelingToolset/Source/MeshModelingTools/Private/MirrorTool.cpp
+++ b/Engine/Plugins/Experimental/MeshModelingToolset/Source/MeshModelingTools/Private/MirrorTool.cpp
@@ -67,6 +67,7 @@ TUniquePtr<FDynamicMeshOperator> UMirrorOperatorFactory::MakeNewOperator()
 	MirrorOp->bAppendToOriginal = MirrorTool->Settings->OperationMode == EMirrorOperationMode::MirrorAndAppend;
 	MirrorOp->bCropFirst = MirrorTool->Settings->bCropAlongMirrorPlaneFirst;
 	MirrorOp->bWeldAlongPlane = MirrorTool->Settings->bWeldVerticesOnMirrorPlane;
+	MirrorOp->PlaneTolerance = MirrorTool->Settings->PlaneTolerance;
 	MirrorOp->bAllowBowtieVertexCreation = MirrorTool->Settings->bAllowBowtieVertexCreation;
 
 	FTransform LocalToWorld = MirrorTool->TargetComponentInterface(ComponentIndex)->GetWorldTransform();
diff --git a/Engine/Plugins/Experimental/MeshModelingToolset/Source/MeshModelingTools/Public/MirrorTool.h b/Engine/Plugins/Experimental/MeshModelingToolset/Source/MeshModelingTools/Public/MirrorTool.h
--- a/Engine/Plugins/Experimental/MeshModelingToolset/Source/MeshModelingTools/Public/MirrorTool.h
+++ b/Engine/Plugins/Experimental/MeshModelingToolset/Source/MeshModelingTools/Public/MirrorTool.h
@@ -79,11 +79,19 @@ public:
 	bool bCropAlongMirrorPlaneFirst = true;
 
 	/** Weld vertices that lie on the mirror plane. Vertices will not be welded if doing so would give an edge more than two faces, or if they are part of a face in the plane. */
-	UPROPERTY(EditAnywhere, Category = Options, meta = (EditCondition = "OperationMode == EMirrorOperationMode::MirrorAndAppend"))
+	UPROPERTY(EditAnywhere, Category = Options, meta = (
+		EditCondition = "OperationMode == EMirrorOperationMode::MirrorAndAppend", EditConditionHides))
 	bool bWeldVerticesOnMirrorPlane = true;
 
+	/** Distance (in unscaled mesh space) to allow a point to be from the plane and still consider it "on the mirror plane". */
+	UPROPERTY(EditAnywhere, Category = Options, meta = (
+		EditCondition = "OperationMode == EMirrorOperationMode::MirrorAndAppend && bWeldVerticesOnMirrorPlane", EditConditionHides,
+		UIMin = 0, UIMax = 0.01, ClampMin = 0, ClampMax = 10))
+	double PlaneTolerance = KINDA_SMALL_NUMBER;
+
 	/** When welding, whether to allow bowtie vertices to be created, or to duplicate the vertex. */
-	UPROPERTY(EditAnywhere, Category = Options, AdvancedDisplay, meta = (EditCondition = "bWeldVerticesOnMirrorPlane && OperationMode == EMirrorOperationMode::MirrorAndAppend"))
+	UPROPERTY(EditAnywhere, Category = Options, AdvancedDisplay, meta = (
+		EditCondition = "bWeldVerticesOnMirrorPlane && OperationMode == EMirrorOperationMode::MirrorAndAppend", EditConditionHides))
 	bool bAllowBowtieVertexCreation = false;
 
 	/** Snap the mirror plane to the world grid. */
