ModelingTools: Add mesh vertex baking infrastructure.
- Added MeshMapEvaluator.EvaluateColor interface as a post sampling pass conversion from raw float data to color.
- Added FMeshConstantMapEvaluator
- Added Vertex Color output for Bake Texture.
- Added RGB mask support to MeshVertexColorMaterial.
- Added MeshVertexAlphaMaterial to preview vertex color alpha.

#rnx

diff --git a/Engine/Plugins/Experimental/GeometryProcessing/Source/DynamicMesh/Private/Sampling/MeshBakerCommon.cpp b/Engine/Plugins/Experimental/GeometryProcessing/Source/DynamicMesh/Private/Sampling/MeshBakerCommon.cpp
new file mode 100644
--- /dev/null
+++ b/Engine/Plugins/Experimental/GeometryProcessing/Source/DynamicMesh/Private/Sampling/MeshBakerCommon.cpp
@@ -0,0 +1,127 @@
+// Copyright Epic Games, Inc. All Rights Reserved.
+
+#include "Sampling/MeshBakerCommon.h"
+
+#include "ExplicitUseGeometryMathTypes.h"		// using UE::Geometry::(math types)
+using namespace UE::Geometry;
+
+/**
+ * Find point on Detail mesh that corresponds to point on Base mesh.
+ * Strategy is:
+ *    1) cast a ray inwards along -Normal from BasePoint + Thickness*Normal
+ *    2) cast a ray outwards along Normal from BasePoint
+ *    3) cast a ray inwards along -Normal from BasePoint
+ * We take (1) preferentially, and then (2), and then (3)
+ *
+ * If all of those fail, if bFailToNearestPoint is true we fall back to nearest-point,
+ *
+ * If all the above fail, return false
+ */
+bool UE::Geometry::GetDetailMeshTrianglePoint_Raycast(
+	const FDynamicMesh3& DetailMesh,
+	const FDynamicMeshAABBTree3& DetailSpatial,
+	const FVector3d& BasePoint,
+	const FVector3d& BaseNormal,
+	int32& DetailTriangleOut,
+	FVector3d& DetailTriBaryCoords,
+	double Thickness,
+	bool bFailToNearestPoint
+)
+{
+	// TODO: should we check normals here? inverse normal should probably not be considered valid
+
+	// shoot rays forwards and backwards
+	FRay3d InwardRay = FRay3d(BasePoint + Thickness * BaseNormal, -BaseNormal);
+	FRay3d ForwardRay(BasePoint, BaseNormal);
+	FRay3d BackwardRay(BasePoint, -BaseNormal);
+	int32 ForwardHitTID = IndexConstants::InvalidID, InwardHitTID = IndexConstants::InvalidID, BackwardHitTID = IndexConstants::InvalidID;
+	double ForwardHitDist, InwardHitDist, BackwardHitDist;
+
+	IMeshSpatial::FQueryOptions Options;
+	Options.MaxDistance = Thickness;
+	bool bHitInward = DetailSpatial.FindNearestHitTriangle(InwardRay, InwardHitDist, InwardHitTID, Options);
+	bool bHitForward = DetailSpatial.FindNearestHitTriangle(ForwardRay, ForwardHitDist, ForwardHitTID, Options);
+	bool bHitBackward = DetailSpatial.FindNearestHitTriangle(BackwardRay, BackwardHitDist, BackwardHitTID, Options);
+
+	FRay3d HitRay;
+	int32 HitTID = IndexConstants::InvalidID;
+	double HitDist = TNumericLimits<double>::Max();
+
+	if (bHitInward)
+	{
+		HitRay = InwardRay;
+		HitTID = InwardHitTID;
+		HitDist = InwardHitDist;
+	}
+	else if (bHitForward)
+	{
+		HitRay = ForwardRay;
+		HitTID = ForwardHitTID;
+		HitDist = ForwardHitDist;
+	}
+	else if (bHitBackward)
+	{
+		HitRay = BackwardRay;
+		HitTID = BackwardHitTID;
+		HitDist = BackwardHitDist;
+	}
+
+	// if we got a valid ray hit, use it
+	if (DetailMesh.IsTriangle(HitTID))
+	{
+		DetailTriangleOut = HitTID;
+		FIntrRay3Triangle3d IntrQuery = TMeshQueries<FDynamicMesh3>::TriangleIntersection(DetailMesh, HitTID, HitRay);
+		DetailTriBaryCoords = IntrQuery.TriangleBaryCoords;
+		return true;
+	}
+	else
+	{
+		// if we did not find any hits, try nearest-point
+		IMeshSpatial::FQueryOptions OnSurfQueryOptions;
+		OnSurfQueryOptions.MaxDistance = Thickness;
+		double NearDistSqr = 0;
+		int32 NearestTriID = -1;
+		// if we are using absolute nearest point as a fallback, then ignore max distance
+		if (bFailToNearestPoint)
+		{
+			NearestTriID = DetailSpatial.FindNearestTriangle(BasePoint, NearDistSqr);
+		}
+		else
+		{
+			NearestTriID = DetailSpatial.FindNearestTriangle(BasePoint, NearDistSqr, OnSurfQueryOptions);
+		}
+		if (DetailMesh.IsTriangle(NearestTriID))
+		{
+			DetailTriangleOut = NearestTriID;
+			FDistPoint3Triangle3d DistQuery = TMeshQueries<FDynamicMesh3>::TriangleDistance(DetailMesh, NearestTriID, BasePoint);
+			DetailTriBaryCoords = DistQuery.TriangleBaryCoords;
+			return true;
+		}
+	}
+
+	return false;
+}
+
+/**
+ * Find point on Detail mesh that corresponds to point on Base mesh using minimum distance
+ */
+bool UE::Geometry::GetDetailMeshTrianglePoint_Nearest(
+	const FDynamicMesh3& DetailMesh,
+	const FDynamicMeshAABBTree3& DetailSpatial,
+	const FVector3d& BasePoint,
+	int32& DetailTriangleOut,
+	FVector3d& DetailTriBaryCoords)
+{
+	double NearDistSqr = 0;
+	int32 NearestTriID = DetailSpatial.FindNearestTriangle(BasePoint, NearDistSqr);
+	if (DetailMesh.IsTriangle(NearestTriID))
+	{
+		DetailTriangleOut = NearestTriID;
+		FDistPoint3Triangle3d DistQuery = TMeshQueries<FDynamicMesh3>::TriangleDistance(DetailMesh, NearestTriID, BasePoint);
+		DetailTriBaryCoords = DistQuery.TriangleBaryCoords;
+		return true;
+	}
+
+	return false;
+}
+
diff --git a/Engine/Plugins/Experimental/GeometryProcessing/Source/DynamicMesh/Private/Sampling/MeshConstantMapEvaluator.cpp b/Engine/Plugins/Experimental/GeometryProcessing/Source/DynamicMesh/Private/Sampling/MeshConstantMapEvaluator.cpp
new file mode 100644
--- /dev/null
+++ b/Engine/Plugins/Experimental/GeometryProcessing/Source/DynamicMesh/Private/Sampling/MeshConstantMapEvaluator.cpp
@@ -0,0 +1,43 @@
+// Copyright Epic Games, Inc. All Rights Reserved.
+
+#include "Sampling/MeshConstantMapEvaluator.h"
+#include "Sampling/MeshMapBaker.h"
+
+#include "ExplicitUseGeometryMathTypes.h"		// using UE::Geometry::(math types)
+using namespace UE::Geometry;
+
+
+FMeshConstantMapEvaluator::FMeshConstantMapEvaluator(float ValueIn)
+	: Value(ValueIn)
+{
+}
+
+void FMeshConstantMapEvaluator::Setup(const FMeshBaseBaker& Baker, FEvaluationContext& Context)
+{
+	Context.Evaluate = &EvaluateSample;
+	Context.EvaluateDefault = &EvaluateDefault;
+	Context.EvaluateColor = &EvaluateColor;
+	Context.EvalData = this;
+	Context.AccumulateMode = EAccumulateMode::Overwrite;
+	Context.DataLayout = { EComponents::Float1 };
+}
+
+void FMeshConstantMapEvaluator::EvaluateSample(float*& Out, const FCorrespondenceSample& Sample, void* EvalData)
+{
+	FMeshConstantMapEvaluator* Eval = static_cast<FMeshConstantMapEvaluator*>(EvalData);
+	WriteToBuffer(Out, Eval->Value);
+}
+
+void FMeshConstantMapEvaluator::EvaluateDefault(float*& Out, void* EvalData)
+{
+	FMeshConstantMapEvaluator* Eval = static_cast<FMeshConstantMapEvaluator*>(EvalData);
+	WriteToBuffer(Out, Eval->Value);
+}
+
+void FMeshConstantMapEvaluator::EvaluateColor(const int DataIdx, float*& In, FVector4f& Out, void* EvalData)
+{
+	Out = FVector4f(In[0], In[0], In[0], 1.0f);
+	In += 1;
+}
+
+
diff --git a/Engine/Plugins/Experimental/GeometryProcessing/Source/DynamicMesh/Private/Sampling/MeshCurvatureMapEvaluator.cpp b/Engine/Plugins/Experimental/GeometryProcessing/Source/DynamicMesh/Private/Sampling/MeshCurvatureMapEvaluator.cpp
--- a/Engine/Plugins/Experimental/GeometryProcessing/Source/DynamicMesh/Private/Sampling/MeshCurvatureMapEvaluator.cpp
+++ b/Engine/Plugins/Experimental/GeometryProcessing/Source/DynamicMesh/Private/Sampling/MeshCurvatureMapEvaluator.cpp
@@ -18,13 +18,14 @@ void FMeshCurvatureMapEvaluator::CacheDetailCurvatures(const FDynamicMesh3* Deta
 	check(Curvatures->Num() == DetailMeshIn->MaxVertexID());
 }
 
-void FMeshCurvatureMapEvaluator::Setup(const FMeshMapBaker& Baker, FEvaluationContext& Context)
+void FMeshCurvatureMapEvaluator::Setup(const FMeshBaseBaker& Baker, FEvaluationContext& Context)
 {
 	Context.Evaluate = &EvaluateSample;
 	Context.EvaluateDefault = &EvaluateDefault;
+	Context.EvaluateColor = &EvaluateColor;
 	Context.EvalData = this;
 	Context.AccumulateMode = EAccumulateMode::Add;
-	Context.DataLayout = { EComponents::Float3 };
+	Context.DataLayout = { EComponents::Float1 };
 
 	// Cache data from the baker
 	DetailMesh = Baker.GetDetailMesh();
@@ -65,7 +66,7 @@ void FMeshCurvatureMapEvaluator::Setup(const FMeshMapBaker& Baker, FEvaluationCo
 	{
 		ClampMax = RangeScale * OverrideRangeMax;
 	}
-	double MinClamp = MinRangeScale * ClampMax;
+	const double MinClamp = MinRangeScale * ClampMax;
 	ClampRange = FInterval1d(MinClamp, ClampMax);
 
 	GetColorMapRange(NegativeColor, ZeroColor, PositiveColor);
@@ -74,17 +75,25 @@ void FMeshCurvatureMapEvaluator::Setup(const FMeshMapBaker& Baker, FEvaluationCo
 void FMeshCurvatureMapEvaluator::EvaluateSample(float*& Out, const FCorrespondenceSample& Sample, void* EvalData)
 {
 	FMeshCurvatureMapEvaluator* Eval = static_cast<FMeshCurvatureMapEvaluator*>(EvalData);
-	double Curvature = Eval->SampleFunction(Sample);
-	double Sign = FMathd::Sign(Curvature);
-	float T = (float)Eval->ClampRange.GetT(FMathd::Abs(Curvature));
-	FVector3f CurvatureColor = (Sign < 0) ?
-		Lerp(Eval->ZeroColor, Eval->NegativeColor, T) : Lerp(Eval->ZeroColor, Eval->PositiveColor, T);
-	WriteToBuffer(Out, CurvatureColor);
+	const float Curvature = (float) Eval->SampleFunction(Sample);
+	WriteToBuffer(Out, Curvature);
 }
 
 void FMeshCurvatureMapEvaluator::EvaluateDefault(float*& Out, void* EvalData)
 {
-	WriteToBuffer(Out, FVector3f::Zero());
+	WriteToBuffer(Out, 0.0f);
+}
+
+void FMeshCurvatureMapEvaluator::EvaluateColor(const int DataIdx, float*& In, FVector4f& Out, void* EvalData)
+{
+	FMeshCurvatureMapEvaluator* Eval = static_cast<FMeshCurvatureMapEvaluator*>(EvalData);
+	const float Curvature = In[0];
+	const float Sign = FMathf::Sign(Curvature);
+	const float T = Eval->ClampRange.GetT(FMathf::Abs(Curvature));
+	FVector3f CurvatureColor = (Sign < 0) ?
+		Lerp(Eval->ZeroColor, Eval->NegativeColor, T) : Lerp(Eval->ZeroColor, Eval->PositiveColor, T);
+	Out = FVector4f(CurvatureColor[0], CurvatureColor[1], CurvatureColor[2], 1.0f);
+	In += 1;
 }
 
 double FMeshCurvatureMapEvaluator::SampleFunction(const FCorrespondenceSample& SampleData)
diff --git a/Engine/Plugins/Experimental/GeometryProcessing/Source/DynamicMesh/Private/Sampling/MeshMapBaker.cpp b/Engine/Plugins/Experimental/GeometryProcessing/Source/DynamicMesh/Private/Sampling/MeshMapBaker.cpp
--- a/Engine/Plugins/Experimental/GeometryProcessing/Source/DynamicMesh/Private/Sampling/MeshMapBaker.cpp
+++ b/Engine/Plugins/Experimental/GeometryProcessing/Source/DynamicMesh/Private/Sampling/MeshMapBaker.cpp
@@ -1,132 +1,13 @@
 // Copyright Epic Games, Inc. All Rights Reserved.
 
 #include "Sampling/MeshMapBaker.h"
+#include "Sampling/MeshBakerCommon.h"
 #include "Image/ImageOccupancyMap.h"
 #include "Spatial/DenseGrid2.h"
 
 #include "ExplicitUseGeometryMathTypes.h"		// using UE::Geometry::(math types)
 using namespace UE::Geometry;
 
-/**
- * Find point on Detail mesh that corresponds to point on Base mesh.
- * Strategy is:
- *    1) cast a ray inwards along -Normal from BasePoint + Thickness*Normal
- *    2) cast a ray outwards along Normal from BasePoint
- *    3) cast a ray inwards along -Normal from BasePoint
- * We take (1) preferentially, and then (2), and then (3)
- *
- * If all of those fail, if bFailToNearestPoint is true we fall back to nearest-point,
- *
- * If all the above fail, return false
- */
-static bool GetDetailMeshTrianglePoint_Raycast(
-	const FDynamicMesh3& DetailMesh,
-	const FDynamicMeshAABBTree3& DetailSpatial,
-	const FVector3d& BasePoint,
-	const FVector3d& BaseNormal,
-	int32& DetailTriangleOut,
-	FVector3d& DetailTriBaryCoords,
-	double Thickness,
-	bool bFailToNearestPoint
-)
-{
-	// TODO: should we check normals here? inverse normal should probably not be considered valid
-
-	// shoot rays forwards and backwards
-	FRay3d InwardRay = FRay3d(BasePoint + Thickness * BaseNormal, -BaseNormal);
-	FRay3d ForwardRay(BasePoint, BaseNormal);
-	FRay3d BackwardRay(BasePoint, -BaseNormal);
-	int32 ForwardHitTID = IndexConstants::InvalidID, InwardHitTID = IndexConstants::InvalidID, BackwardHitTID = IndexConstants::InvalidID;
-	double ForwardHitDist, InwardHitDist, BackwardHitDist;
-
-	IMeshSpatial::FQueryOptions Options;
-	Options.MaxDistance = Thickness;
-	bool bHitInward = DetailSpatial.FindNearestHitTriangle(InwardRay, InwardHitDist, InwardHitTID, Options);
-	bool bHitForward = DetailSpatial.FindNearestHitTriangle(ForwardRay, ForwardHitDist, ForwardHitTID, Options);
-	bool bHitBackward = DetailSpatial.FindNearestHitTriangle(BackwardRay, BackwardHitDist, BackwardHitTID, Options);
-
-	FRay3d HitRay;
-	int32 HitTID = IndexConstants::InvalidID;
-	double HitDist = TNumericLimits<double>::Max();
-
-	if (bHitInward)
-	{
-		HitRay = InwardRay;
-		HitTID = InwardHitTID;
-		HitDist = InwardHitDist;
-	}
-	else if (bHitForward)
-	{
-		HitRay = ForwardRay;
-		HitTID = ForwardHitTID;
-		HitDist = ForwardHitDist;
-	}
-	else if (bHitBackward)
-	{
-		HitRay = BackwardRay;
-		HitTID = BackwardHitTID;
-		HitDist = BackwardHitDist;
-	}
-
-	// if we got a valid ray hit, use it
-	if (DetailMesh.IsTriangle(HitTID))
-	{
-		DetailTriangleOut = HitTID;
-		FIntrRay3Triangle3d IntrQuery = TMeshQueries<FDynamicMesh3>::TriangleIntersection(DetailMesh, HitTID, HitRay);
-		DetailTriBaryCoords = IntrQuery.TriangleBaryCoords;
-		return true;
-	}
-	else
-	{
-		// if we did not find any hits, try nearest-point
-		IMeshSpatial::FQueryOptions OnSurfQueryOptions;
-		OnSurfQueryOptions.MaxDistance = Thickness;
-		double NearDistSqr = 0;
-		int32 NearestTriID = -1;
-		// if we are using absolute nearest point as a fallback, then ignore max distance
-		if (bFailToNearestPoint)
-		{
-			NearestTriID = DetailSpatial.FindNearestTriangle(BasePoint, NearDistSqr);
-		}
-		else
-		{
-			NearestTriID = DetailSpatial.FindNearestTriangle(BasePoint, NearDistSqr, OnSurfQueryOptions);
-		}
-		if (DetailMesh.IsTriangle(NearestTriID))
-		{
-			DetailTriangleOut = NearestTriID;
-			FDistPoint3Triangle3d DistQuery = TMeshQueries<FDynamicMesh3>::TriangleDistance(DetailMesh, NearestTriID, BasePoint);
-			DetailTriBaryCoords = DistQuery.TriangleBaryCoords;
-			return true;
-		}
-	}
-
-	return false;
-}
-
-/**
- * Find point on Detail mesh that corresponds to point on Base mesh using minimum distance
- */
-static bool GetDetailMeshTrianglePoint_Nearest(
-	const FDynamicMesh3& DetailMesh,
-	const FDynamicMeshAABBTree3& DetailSpatial,
-	const FVector3d& BasePoint,
-	int32& DetailTriangleOut,
-	FVector3d& DetailTriBaryCoords)
-{
-	double NearDistSqr = 0;
-	int32 NearestTriID = DetailSpatial.FindNearestTriangle(BasePoint, NearDistSqr);
-	if (DetailMesh.IsTriangle(NearestTriID))
-	{
-		DetailTriangleOut = NearestTriID;
-		FDistPoint3Triangle3d DistQuery = TMeshQueries<FDynamicMesh3>::TriangleDistance(DetailMesh, NearestTriID, BasePoint);
-		DetailTriBaryCoords = DistQuery.TriangleBaryCoords;
-		return true;
-	}
-
-	return false;
-}
-
 
 //
 // FMeshMapBaker
@@ -210,11 +91,12 @@ void FMeshMapBaker::InitBakeDefaults()
 	BakeDefaultColors.SetNumUninitialized(NumBakeResults);
 	for (int32 Idx = 0; Idx < NumBakers; ++Idx)
 	{
-		const int32 NumData = BakeContexts[Idx].DataLayout.Num();
+		const FMeshMapEvaluator::FEvaluationContext& Context = BakeContexts[Idx];
+		const int32 NumData = Context.DataLayout.Num();
 		for (int32 DataIdx = 0; DataIdx < NumData; ++DataIdx)
 		{
 			const int32 ResultIdx = BakeOffsets[Idx] + DataIdx;
-			BakeDefaultColors[ResultIdx] = FloatToPixel(BufferPtr, BakeContexts[Idx].DataLayout[DataIdx], 1.0f);
+			Context.EvaluateColor(DataIdx, BufferPtr, BakeDefaultColors[ResultIdx], Context.EvalData);
 		}
 	}
 	checkSlow((BufferPtr - Buffer) == BakeSampleBufferSize);
@@ -347,8 +229,16 @@ void FMeshMapBaker::Bake()
 								const int32 Offset = BakeSampleOffsets[ResultIdx];
 								float* BufferPtr = &PixelBuffer[Offset];
 								const FMeshMapEvaluator::EComponents Stride = Context.DataLayout[DataIdx];
+
+								// Apply weight to raw float data.
+								for (int32 FloatIdx = 0; FloatIdx < (int32)Stride; ++FloatIdx)
+								{
+									BufferPtr[FloatIdx] *= Weight;
+								}
+
+								// Convert float data to color.
 								FVector4f& Pixel = BakeResults[ResultIdx]->GetPixel(ImageLinearIdx);
-								Pixel = FloatToPixel(BufferPtr, Stride, Weight);
+								Context.EvaluateColor(DataIdx, BufferPtr, Pixel, Context.EvalData);
 							}
 						}
 					};
@@ -471,45 +361,6 @@ void FMeshMapBaker::BakeSample(FMeshMapTileBuffer& TileBuffer, const FMeshMapEva
 	AccumulateFn(BakeAccumulateLists[(int32)FMeshMapEvaluator::EAccumulateMode::Overwrite], OverwriteFn);
 }
 
-FVector4f FMeshMapBaker::FloatToPixel(float*& Buffer, FMeshMapEvaluator::EComponents Stride, float Weight)
-{
-	// TODO: Decompose this method into non-branching parts.
-	FVector4f Result = FVector4f(0.0, 0.0, 0.0, 1.0f);
-	switch (Stride)
-	{
-	case FMeshMapEvaluator::EComponents::Float1:
-	{
-		float X = Buffer[0] * Weight;
-		Result = FVector4f(X, X, X, 1.0f);
-		break;
-	}
-	case FMeshMapEvaluator::EComponents::Float2:
-	{
-		FVector2f Vec(Buffer[0], Buffer[1]);
-		Vec *= Weight;
-		Result = FVector4f(Vec.X, Vec.Y, 0.0f, 1.0f);
-		break;
-	}
-	case FMeshMapEvaluator::EComponents::Float3:
-	{
-		FVector3f Vec(Buffer[0], Buffer[1], Buffer[2]);
-		Vec *= Weight;
-		Result = FVector4f(Vec.X, Vec.Y, Vec.Z, 1.0f);
-		break;
-	}
-	case FMeshMapEvaluator::EComponents::Float4:
-	{
-		Result = FVector4f(Buffer[0], Buffer[1], Buffer[2], Buffer[3]);
-		Result *= Weight;
-		break;
-	}
-	default:
-		break;
-	}
-	Buffer += static_cast<int32>(Stride);
-	return Result;
-}
-
 int32 FMeshMapBaker::AddBaker(TSharedPtr<FMeshMapEvaluator> Sampler)
 {
 	return Bakers.Add(Sampler);
@@ -538,81 +389,11 @@ const TArrayView<TUniquePtr<TImageBuilder<FVector4f>>> FMeshMapBaker::GetBakeRes
 	return TArrayView<TUniquePtr<TImageBuilder<FVector4f>>>(&BakeResults[ResultIdx], NumResults);
 }
 
-void FMeshMapBaker::SetDetailMesh(const FDynamicMesh3* Mesh, const FDynamicMeshAABBTree3* Spatial)
-{
-	DetailMesh = Mesh;
-	DetailSpatial = Spatial;
-}
-
-void FMeshMapBaker::SetTargetMesh(const FDynamicMesh3* Mesh)
-{
-	TargetMesh = Mesh;
-}
-
-void FMeshMapBaker::SetTargetMeshTangents(TSharedPtr<FMeshTangentsd, ESPMode::ThreadSafe> Tangents)
-{
-	TargetMeshTangents = Tangents;
-}
-
-const FDynamicMesh3* FMeshMapBaker::GetDetailMesh() const
-{
-	return DetailMesh;
-}
-
-const FDynamicMeshAABBTree3* FMeshMapBaker::GetDetailMeshSpatial() const
-{
-	return DetailSpatial;
-}
-
-const FDynamicMeshNormalOverlay* FMeshMapBaker::GetDetailMeshNormals() const
-{
-	check(DetailMesh && DetailMesh->HasAttributes());
-	return DetailMesh->Attributes()->PrimaryNormals();
-}
-
-const FDynamicMeshUVOverlay* FMeshMapBaker::GetDetailMeshUVs(int32 UVLayerIn /*=0*/) const
-{
-	check(DetailMesh && DetailMesh->HasAttributes());
-	return DetailMesh->Attributes()->GetUVLayer(UVLayerIn);
-}
-
-const FDynamicMesh3* FMeshMapBaker::GetTargetMesh() const
-{
-	return TargetMesh;
-}
-
-const FDynamicMeshUVOverlay* FMeshMapBaker::GetTargetMeshUVs() const
-{
-	check(TargetMesh && TargetMesh->HasAttributes() && UVLayer < TargetMesh->Attributes()->NumUVLayers());
-	return TargetMesh->Attributes()->GetUVLayer(UVLayer);
-}
-
-const FDynamicMeshNormalOverlay* FMeshMapBaker::GetTargetMeshNormals() const
-{
-	check(TargetMesh && TargetMesh->HasAttributes());
-	return TargetMesh->Attributes()->PrimaryNormals();
-}
-
-TSharedPtr<FMeshTangentsd, ESPMode::ThreadSafe> FMeshMapBaker::GetTargetMeshTangents() const
-{
-	return TargetMeshTangents;
-}
-
 void FMeshMapBaker::SetDimensions(FImageDimensions DimensionsIn)
 {
 	Dimensions = DimensionsIn;
 }
 
-void FMeshMapBaker::SetUVLayer(int32 UVLayerIn)
-{
-	UVLayer = UVLayerIn;
-}
-
-void FMeshMapBaker::SetThickness(double ThicknessIn)
-{
-	Thickness = ThicknessIn;
-}
-
 void FMeshMapBaker::SetGutterSize(int32 GutterSizeIn)
 {
 	GutterSize = GutterSizeIn;
@@ -623,9 +404,5 @@ void FMeshMapBaker::SetMultisampling(int32 MultisamplingIn)
 	Multisampling = MultisamplingIn;
 }
 
-void FMeshMapBaker::SetCorrespondenceStrategy(ECorrespondenceStrategy Strategy)
-{
-	CorrespondenceStrategy = Strategy;
-}
 
 
diff --git a/Engine/Plugins/Experimental/GeometryProcessing/Source/DynamicMesh/Private/Sampling/MeshNormalMapEvaluator.cpp b/Engine/Plugins/Experimental/GeometryProcessing/Source/DynamicMesh/Private/Sampling/MeshNormalMapEvaluator.cpp
--- a/Engine/Plugins/Experimental/GeometryProcessing/Source/DynamicMesh/Private/Sampling/MeshNormalMapEvaluator.cpp
+++ b/Engine/Plugins/Experimental/GeometryProcessing/Source/DynamicMesh/Private/Sampling/MeshNormalMapEvaluator.cpp
@@ -7,10 +7,11 @@
 using namespace UE::Geometry;
 
 
-void FMeshNormalMapEvaluator::Setup(const FMeshMapBaker& Baker, FEvaluationContext& Context)
+void FMeshNormalMapEvaluator::Setup(const FMeshBaseBaker& Baker, FEvaluationContext& Context)
 {
 	Context.Evaluate = &EvaluateSample;
 	Context.EvaluateDefault = &EvaluateDefault;
+	Context.EvaluateColor = &EvaluateColor;
 	Context.EvalData = this;
 	Context.AccumulateMode = EAccumulateMode::Add;
 	Context.DataLayout = { EComponents::Float3 };
@@ -25,21 +26,26 @@ void FMeshNormalMapEvaluator::Setup(const FMeshMapBaker& Baker, FEvaluationConte
 void FMeshNormalMapEvaluator::EvaluateSample(float*& Out, const FCorrespondenceSample& Sample, void* EvalData)
 {
 	FMeshNormalMapEvaluator* Eval = static_cast<FMeshNormalMapEvaluator*>(EvalData);
-	FVector3f SampleResult = Eval->SampleFunction(Sample);
-	// Map normal space [-1,1] to floating point color space [0,1]
-	FVector3f Result = (SampleResult + FVector3f::One()) * 0.5f;
-	WriteToBuffer(Out, Result);
+	const FVector3f SampleResult = Eval->SampleFunction(Sample);
+	WriteToBuffer(Out, SampleResult);
 }
 
 void FMeshNormalMapEvaluator::EvaluateDefault(float*& Out, void* EvalData)
 {
 	FMeshNormalMapEvaluator* Eval = static_cast<FMeshNormalMapEvaluator*>(EvalData);
-	// Map normal space [-1,1] to floating point color space [0,1]
-	FVector3f Normal = (Eval->DefaultNormal + FVector3f::One()) * 0.5f;
-	WriteToBuffer(Out, Normal);
+	WriteToBuffer(Out, Eval->DefaultNormal);
 }
 
-FVector3f FMeshNormalMapEvaluator::SampleFunction(const FCorrespondenceSample& SampleData)
+void FMeshNormalMapEvaluator::EvaluateColor(const int DataIdx, float*& In, FVector4f& Out, void* EvalData)
+{
+	// Map normal space [-1,1] to color space [0,1]
+	const FVector3f Normal(In[0], In[1], In[2]);
+	const FVector3f Color = (Normal + FVector3f::One()) * 0.5f;
+	Out = FVector4f(Color.X, Color.Y, Color.Z, 1.0f);
+	In += 3;
+}
+
+FVector3f FMeshNormalMapEvaluator::SampleFunction(const FCorrespondenceSample& SampleData) const
 {
 	int32 DetailTriID = SampleData.DetailTriID;
 	if (DetailMesh->IsTriangle(DetailTriID))
diff --git a/Engine/Plugins/Experimental/GeometryProcessing/Source/DynamicMesh/Private/Sampling/MeshOcclusionMapEvaluator.cpp b/Engine/Plugins/Experimental/GeometryProcessing/Source/DynamicMesh/Private/Sampling/MeshOcclusionMapEvaluator.cpp
--- a/Engine/Plugins/Experimental/GeometryProcessing/Source/DynamicMesh/Private/Sampling/MeshOcclusionMapEvaluator.cpp
+++ b/Engine/Plugins/Experimental/GeometryProcessing/Source/DynamicMesh/Private/Sampling/MeshOcclusionMapEvaluator.cpp
@@ -21,7 +21,7 @@ using namespace UE::Geometry;
 const FVector3d FMeshOcclusionMapEvaluator::DefaultTangentNormal = FVector3d::UnitZ();
 const FVector3d FMeshOcclusionMapEvaluator::DefaultObjectNormal = FVector3d::Zero();
 
-void FMeshOcclusionMapEvaluator::Setup(const FMeshMapBaker& Baker, FEvaluationContext& Context)
+void FMeshOcclusionMapEvaluator::Setup(const FMeshBaseBaker& Baker, FEvaluationContext& Context)
 {
 	if (WANT_ALL(OcclusionType))
 	{
@@ -29,11 +29,13 @@ void FMeshOcclusionMapEvaluator::Setup(const FMeshMapBaker& Baker, FEvaluationCo
 		{
 			Context.Evaluate = &EvaluateSample<EMeshOcclusionMapType::All, ESpace::Tangent>;
 			Context.EvaluateDefault = &EvaluateDefault<EMeshOcclusionMapType::All, ESpace::Tangent>;
+			Context.EvaluateColor = &EvaluateColor<EMeshOcclusionMapType::All, ESpace::Tangent>;
 		}
 		else // NormalSpace == ESpace::Object
 		{
 			Context.Evaluate = &EvaluateSample<EMeshOcclusionMapType::All, ESpace::Object>;
 			Context.EvaluateDefault = &EvaluateDefault<EMeshOcclusionMapType::All, ESpace::Object>;
+			Context.EvaluateColor = &EvaluateColor<EMeshOcclusionMapType::All, ESpace::Object>;
 		}
 		Context.EvalData = this;
 		Context.AccumulateMode = EAccumulateMode::Add;
@@ -43,6 +45,7 @@ void FMeshOcclusionMapEvaluator::Setup(const FMeshMapBaker& Baker, FEvaluationCo
 	{
 		Context.Evaluate = &EvaluateSample<EMeshOcclusionMapType::AmbientOcclusion, ESpace::Tangent>;
 		Context.EvaluateDefault = &EvaluateDefault<EMeshOcclusionMapType::AmbientOcclusion, ESpace::Tangent>;
+		Context.EvaluateColor = &EvaluateColor<EMeshOcclusionMapType::AmbientOcclusion, ESpace::Tangent>;
 		Context.EvalData = this;
 		Context.AccumulateMode = EAccumulateMode::Add;
 		Context.DataLayout = { EComponents::Float1 };
@@ -53,11 +56,13 @@ void FMeshOcclusionMapEvaluator::Setup(const FMeshMapBaker& Baker, FEvaluationCo
 		{
 			Context.Evaluate = &EvaluateSample<EMeshOcclusionMapType::BentNormal, ESpace::Tangent>;
 			Context.EvaluateDefault = &EvaluateDefault<EMeshOcclusionMapType::BentNormal, ESpace::Tangent>;
+			Context.EvaluateColor = &EvaluateColor<EMeshOcclusionMapType::BentNormal, ESpace::Tangent>;
 		}
 		else // NormalSpace == ESpace::Object
 		{
 			Context.Evaluate = &EvaluateSample<EMeshOcclusionMapType::BentNormal, ESpace::Object>;
 			Context.EvaluateDefault = &EvaluateDefault<EMeshOcclusionMapType::BentNormal, ESpace::Object>;
+			Context.EvaluateColor = &EvaluateColor<EMeshOcclusionMapType::BentNormal, ESpace::Object>;
 		}
 		Context.EvalData = this;
 		Context.AccumulateMode = EAccumulateMode::Add;
@@ -105,13 +110,13 @@ void FMeshOcclusionMapEvaluator::Setup(const FMeshMapBaker& Baker, FEvaluationCo
 	}
 
 	// Map occlusion ray hemisphere to conical area (SpreadAngle/2)
-	double ConicalAngle = FMathd::Clamp(SpreadAngle * 0.5, 0.0001, 90.0);
+	const double ConicalAngle = FMathd::Clamp(SpreadAngle * 0.5, 0.0001, 90.0);
 	for (int32 k = 0; k < RayDirections.Num(); ++k)
 	{
 		FVector3d& RayDir = RayDirections[k];
-		double RayAngle = AngleD(RayDir, FVector3d::UnitZ());
+		const double RayAngle = AngleD(RayDir, FVector3d::UnitZ());
 		FVector3d RayCross = RayDir.Cross(FVector3d::UnitZ());
-		double RotationAngle = RayAngle - FMathd::Lerp(0.0, ConicalAngle, RayAngle / 90.0);
+		const double RotationAngle = RayAngle - FMathd::Lerp(0.0, ConicalAngle, RayAngle / 90.0);
 		FQuaterniond Rotation(RayCross, RotationAngle, true);
 		RayDir = Rotation * RayDir;
 	}
@@ -134,9 +139,7 @@ void FMeshOcclusionMapEvaluator::EvaluateSample(float*& Out, const FCorresponden
 
 	if constexpr (WANT_BENT_NORMAL(ComputeType))
 	{
-		// Map normal space[-1, 1] to floating point color space[0, 1]
-		FVector3f NormalColor = (FVector3f(BentNormal.X, BentNormal.Y, BentNormal.Z) + FVector3f::One()) * 0.5;
-		WriteToBuffer(Out, NormalColor);
+		WriteToBuffer(Out, BentNormal);
 	}
 }
 
@@ -151,10 +154,38 @@ void FMeshOcclusionMapEvaluator::EvaluateDefault(float*& Out, void* EvalData)
 	if constexpr (WANT_BENT_NORMAL(ComputeType))
 	{
 		const FVector3d& DefaultNormal = ComputeSpace == ESpace::Tangent ? DefaultTangentNormal : DefaultObjectNormal;
+		WriteToBuffer(Out, DefaultNormal);
+	}
+}
 
-		// Map normal space [-1,1] to floating point color space [0,1]
-		FVector3f Normal = (DefaultObjectNormal + FVector3f::One()) * 0.5f;
-		WriteToBuffer(Out, Normal);
+template <EMeshOcclusionMapType ComputeType, FMeshOcclusionMapEvaluator::ESpace ComputeSpace>
+void FMeshOcclusionMapEvaluator::EvaluateColor(const int DataIdx, float*& In, FVector4f& Out, void* EvalData)
+{
+	auto EvalOcclusion = [&In, &Out]()
+	{
+		Out = FVector4f(In[0], In[0], In[0], 1.0f);
+		In += 1;
+	};
+
+	auto EvalNormal = [&In, &Out]()
+	{
+		const FVector3f Normal(In[0], In[1], In[2]);
+		const FVector3f Color = (Normal + FVector3f::One()) * 0.5f;
+		Out = FVector4f(Color.X, Color.Y, Color.Z, 1.0f);
+		In += 3;
+	};
+
+	if constexpr (WANT_AMBIENT_OCCLUSION(ComputeType) && WANT_BENT_NORMAL(ComputeType))
+	{
+		DataIdx == 0 ? EvalOcclusion() : EvalNormal();
+	}
+	else if constexpr (WANT_AMBIENT_OCCLUSION(ComputeType))
+	{
+		EvalOcclusion();
+	}
+	else if constexpr (WANT_BENT_NORMAL(ComputeType))
+	{
+		EvalNormal();
 	}
 }
 
diff --git a/Engine/Plugins/Experimental/GeometryProcessing/Source/DynamicMesh/Private/Sampling/MeshPropertyMapEvaluator.cpp b/Engine/Plugins/Experimental/GeometryProcessing/Source/DynamicMesh/Private/Sampling/MeshPropertyMapEvaluator.cpp
--- a/Engine/Plugins/Experimental/GeometryProcessing/Source/DynamicMesh/Private/Sampling/MeshPropertyMapEvaluator.cpp
+++ b/Engine/Plugins/Experimental/GeometryProcessing/Source/DynamicMesh/Private/Sampling/MeshPropertyMapEvaluator.cpp
@@ -7,10 +7,11 @@
 #include "ExplicitUseGeometryMathTypes.h"		// using UE::Geometry::(math types)
 using namespace UE::Geometry;
 
-void FMeshPropertyMapEvaluator::Setup(const FMeshMapBaker& Baker, FEvaluationContext& Context)
+void FMeshPropertyMapEvaluator::Setup(const FMeshBaseBaker& Baker, FEvaluationContext& Context)
 {
 	Context.Evaluate = &EvaluateSample;
 	Context.EvaluateDefault = &EvaluateDefault;
+	Context.EvaluateColor = &EvaluateColor;
 	Context.EvalData = this;
 	Context.AccumulateMode = EAccumulateMode::Add;
 	Context.DataLayout = { EComponents::Float3 };
@@ -20,6 +21,7 @@ void FMeshPropertyMapEvaluator::Setup(const FMeshMapBaker& Baker, FEvaluationCon
 	DetailNormalOverlay = Baker.GetDetailMeshNormals();
 	check(DetailNormalOverlay);
 	DetailUVOverlay = Baker.GetDetailMeshUVs();
+	DetailColorOverlay = Baker.GetDetailMeshColors();
 
 	Bounds = DetailMesh->GetBounds();
 	for (int32 j = 0; j < 3; ++j)
@@ -48,6 +50,9 @@ void FMeshPropertyMapEvaluator::Setup(const FMeshMapBaker& Baker, FEvaluationCon
 	case EMeshPropertyMapType::MaterialID:
 		DefaultValue = FVector3f(LinearColors::LightPink3f());
 		break;
+	case EMeshPropertyMapType::VertexColor:
+		DefaultValue = FVector3f::One();
+		break;
 	}
 }
 
@@ -63,6 +68,13 @@ void FMeshPropertyMapEvaluator::EvaluateDefault(float*& Out, void* EvalData)
 	WriteToBuffer(Out, FVector3f::Zero());
 }
 
+void FMeshPropertyMapEvaluator::EvaluateColor(const int DataIdx, float*& In, FVector4f& Out, void* EvalData)
+{
+	// TODO: Move property color space transformation from EvaluateSample/Default to here.
+	Out = FVector4f(In[0], In[1], In[2], 1.0f);
+	In += 3;
+}
+
 FVector3f FMeshPropertyMapEvaluator::SampleFunction(const FCorrespondenceSample& SampleData)
 {
 	FVector3f Color = DefaultValue;
@@ -114,6 +126,16 @@ FVector3f FMeshPropertyMapEvaluator::SampleFunction(const FCorrespondenceSample&
 			}
 		}
 		break;
+		case EMeshPropertyMapType::VertexColor:
+		{
+			if (DetailColorOverlay && DetailColorOverlay->IsSetTriangle(DetailTriID))
+			{
+				FVector4d DetailColor;
+				DetailColorOverlay->GetTriBaryInterpolate<double>(DetailTriID, &SampleData.DetailBaryCoords.X, &DetailColor.X);
+				Color = FVector3f(DetailColor.X, DetailColor.Y, DetailColor.Z);
+			}
+		}
+		break;
 		}
 	}
 	return Color;
diff --git a/Engine/Plugins/Experimental/GeometryProcessing/Source/DynamicMesh/Private/Sampling/MeshResampleImageEvaluator.cpp b/Engine/Plugins/Experimental/GeometryProcessing/Source/DynamicMesh/Private/Sampling/MeshResampleImageEvaluator.cpp
--- a/Engine/Plugins/Experimental/GeometryProcessing/Source/DynamicMesh/Private/Sampling/MeshResampleImageEvaluator.cpp
+++ b/Engine/Plugins/Experimental/GeometryProcessing/Source/DynamicMesh/Private/Sampling/MeshResampleImageEvaluator.cpp
@@ -10,10 +10,11 @@ using namespace UE::Geometry;
 // FMeshResampleImage Evaluator
 //
 
-void FMeshResampleImageEvaluator::Setup(const FMeshMapBaker& Baker, FEvaluationContext& Context)
+void FMeshResampleImageEvaluator::Setup(const FMeshBaseBaker& Baker, FEvaluationContext& Context)
 {
 	Context.Evaluate = &EvaluateSample;
 	Context.EvaluateDefault = &EvaluateDefault;
+	Context.EvaluateColor = &EvaluateColor;
 	Context.EvalData = this;
 	Context.AccumulateMode = EAccumulateMode::Add;
 	Context.DataLayout = { EComponents::Float4 };
@@ -34,6 +35,12 @@ void FMeshResampleImageEvaluator::EvaluateDefault(float*& Out, void* EvalData)
 	WriteToBuffer(Out, FVector4f(0.0f, 0.0f, 0.0f, 1.0f));
 }
 
+void FMeshResampleImageEvaluator::EvaluateColor(const int DataIdx, float*& In, FVector4f& Out, void* EvalData)
+{
+	Out = FVector4f(In[0], In[1], In[2], In[3]);
+	In += 4;
+}
+
 FVector4f FMeshResampleImageEvaluator::ImageSampleFunction(const FCorrespondenceSample& SampleData)
 {
 	FVector4f Color = DefaultColor;
@@ -51,10 +58,11 @@ FVector4f FMeshResampleImageEvaluator::ImageSampleFunction(const FCorrespondence
 // FMeshMultiResampleImage Evaluator
 //
 
-void FMeshMultiResampleImageEvaluator::Setup(const FMeshMapBaker& Baker, FEvaluationContext& Context)
+void FMeshMultiResampleImageEvaluator::Setup(const FMeshBaseBaker& Baker, FEvaluationContext& Context)
 {
 	Context.Evaluate = &EvaluateSampleMulti;
 	Context.EvaluateDefault = &EvaluateDefault;
+	Context.EvaluateColor = &EvaluateColor;
 	Context.EvalData = this;
 	Context.AccumulateMode = EAccumulateMode::Add;
 	Context.DataLayout = { EComponents::Float4 };
diff --git a/Engine/Plugins/Experimental/GeometryProcessing/Source/DynamicMesh/Private/Sampling/MeshVertexBaker.cpp b/Engine/Plugins/Experimental/GeometryProcessing/Source/DynamicMesh/Private/Sampling/MeshVertexBaker.cpp
new file mode 100644
--- /dev/null
+++ b/Engine/Plugins/Experimental/GeometryProcessing/Source/DynamicMesh/Private/Sampling/MeshVertexBaker.cpp
@@ -0,0 +1,200 @@
+// Copyright Epic Games, Inc. All Rights Reserved.
+
+#include "Sampling/MeshVertexBaker.h"
+#include "Sampling/MeshBakerCommon.h"
+
+#include "ExplicitUseGeometryMathTypes.h"		// using UE::Geometry::(math types)
+using namespace UE::Geometry;
+
+FMeshConstantMapEvaluator FMeshVertexBaker::ZeroEvaluator(0.0f);
+FMeshConstantMapEvaluator FMeshVertexBaker::OneEvaluator(1.0f);
+
+void FMeshVertexBaker::Bake()
+{
+	if (!ensure(TargetMesh && TargetMesh->HasAttributes() && TargetMesh->Attributes()->HasPrimaryColors()))
+	{
+		return;
+	}
+
+	// Convert Bake mode into internal list of bakers.
+	Bakers.Reset();
+	if (BakeMode == EBakeMode::Color)
+	{
+		FMeshMapEvaluator* Evaluator = ColorEvaluator.Get();
+		Bakers.Add(Evaluator ? Evaluator : &ZeroEvaluator);
+		BakeInternal = &BakeImpl<EBakeMode::Color>;
+	}
+	else // Mode == EBakeMode::Channel
+	{
+		for (int Idx = 0; Idx < 4; ++Idx)
+		{
+			// For alpha channel, default to 1.0, otherwise 0.0.
+			FMeshMapEvaluator* DefaultEvaluator = Idx == 3 ? &OneEvaluator : &ZeroEvaluator;
+			FMeshMapEvaluator* Evaluator = ChannelEvaluators[Idx].Get();
+			Bakers.Add(Evaluator ? Evaluator : DefaultEvaluator);
+		}
+		BakeInternal = &BakeImpl<EBakeMode::Channel>;
+	}
+
+	const int NumBakers = Bakers.Num();
+	if (NumBakers == 0)
+	{
+		return;
+	}
+
+	// Initialize BakeContext(s) and BakeDefaults
+	BakeDefaults = FVector4f(0.0f, 0.0f, 0.0f, 1.0f);
+	float* DefaultBufferPtr = &BakeDefaults[0];
+	BakeContexts.Reset();
+	BakeContexts.SetNum(NumBakers);
+	BakeSampleBufferSize = 0;
+	for (int Idx = 0; Idx < NumBakers; ++Idx)
+	{
+		Bakers[Idx]->Setup(*this, BakeContexts[Idx]);
+
+		for (FMeshMapEvaluator::EComponents Components : BakeContexts[Idx].DataLayout)
+		{
+			BakeSampleBufferSize += (int) Components;
+		}
+		if (!ensure(BakeSampleBufferSize <= 4))
+		{
+			return;
+		}
+
+		BakeContexts[Idx].EvaluateDefault(DefaultBufferPtr, BakeContexts[Idx].EvalData);
+	}
+
+	// Initialize BakeResult to unique vertex color elements
+	const FDynamicMeshColorOverlay* ColorOverlay = TargetMesh->Attributes()->PrimaryColors();
+	const int NumColors = ColorOverlay->ElementCount();
+	Dimensions = FImageDimensions(NumColors, 1);
+
+	BakeResult = MakeUnique<TImageBuilder<FVector4f>>();
+	BakeResult->SetDimensions(Dimensions);
+	BakeResult->Clear(BakeDefaults);
+
+	BakeInternal(this);
+}
+
+const TImageBuilder<FVector4f>* FMeshVertexBaker::GetBakeResult() const
+{
+	return BakeResult.Get();
+}
+
+template<FMeshVertexBaker::EBakeMode ComputeMode>
+void FMeshVertexBaker::BakeImpl(void* Data)
+{
+	if (!ensure(Data))
+	{
+		return;
+	}
+
+	FMeshVertexBaker* Baker = static_cast<FMeshVertexBaker*>(Data);
+	
+	ECorrespondenceStrategy UseStrategy = Baker->CorrespondenceStrategy;
+	if (UseStrategy == ECorrespondenceStrategy::Identity && ensure(Baker->DetailMesh == Baker->TargetMesh) == false)
+	{
+		// Identity strategy requires mesh to be the same. Could potentially have two copies, in which
+		// case this ensure is too conservative, but for now we will assume this
+		UseStrategy = ECorrespondenceStrategy::NearestPoint;
+	}
+	
+	const FDynamicMesh3* Mesh = Baker->TargetMesh;
+	const FDynamicMeshColorOverlay* ColorOverlay = Baker->TargetMesh->Attributes()->PrimaryColors();
+
+	// TODO: Refactor into TMeshSurfaceSampler class (future vertex bake enhancements will require non-UV based surface sampling)
+	auto SampleSurface = [&Baker, &Mesh, &ColorOverlay, &UseStrategy](int32 ElementIdx,
+	                                                                FMeshMapEvaluator::FCorrespondenceSample& ValueOut)
+	{
+		const int32 VId = ColorOverlay->GetParentVertex(ElementIdx);
+
+		// Compute ray direction
+		// TODO: Precompute a color to normal overlay mapping data structure and replace this with a lookup.
+		FVector3f SurfaceNormal = FVector3f::Zero();
+		TArray<int> ColorElemTris;
+		ColorOverlay->GetElementTriangles(ElementIdx, ColorElemTris);
+		for (int TId : ColorElemTris)
+		{
+			SurfaceNormal += Mesh->GetTriNormal(TId);
+		}
+		Normalize(SurfaceNormal);
+
+		// Compute surface point and barycentric coords
+		const FVector3d SurfacePoint = Mesh->GetVertex(VId);
+		const int32 TriangleIndex = ColorElemTris[0];
+		FVector3d BaryCoords = FVector3d::Zero();
+		FIndex3i TriVerts = Mesh->GetTriangle(TriangleIndex);
+		for (int Idx = 0; Idx < 3; ++Idx)
+		{
+			if (VId == TriVerts[Idx])
+			{
+				BaryCoords[Idx] = 1.0;
+				break;
+			}
+		}
+
+		ValueOut.BaseSample.TriangleIndex = TriangleIndex;
+		ValueOut.BaseSample.SurfacePoint = SurfacePoint;
+		ValueOut.BaseSample.BaryCoords = BaryCoords;
+		ValueOut.DetailTriID = FDynamicMesh3::InvalidID;
+		ValueOut.BaseNormal = SurfaceNormal;
+
+		const FVector3d RayDir = SurfaceNormal;
+		if (UseStrategy == ECorrespondenceStrategy::Identity)
+		{
+			ValueOut.DetailTriID = TriangleIndex;
+			ValueOut.DetailBaryCoords = BaryCoords;
+		}
+		else if (UseStrategy == ECorrespondenceStrategy::NearestPoint)
+		{
+			GetDetailMeshTrianglePoint_Nearest(*Baker->DetailMesh, *Baker->DetailSpatial, SurfacePoint,
+			                                   ValueOut.DetailTriID, ValueOut.DetailBaryCoords);
+		}
+		else // Fall back to raycast strategy
+		{
+			const double SampleThickness = Baker->GetThickness(); // could modulate w/ a map here...
+
+			// Find detail mesh triangle point
+			GetDetailMeshTrianglePoint_Raycast(*Baker->DetailMesh, *Baker->DetailSpatial, SurfacePoint, RayDir,
+			                                   ValueOut.DetailTriID, ValueOut.DetailBaryCoords,
+			                                   SampleThickness,
+			                                   (UseStrategy == ECorrespondenceStrategy::RaycastStandardThenNearest));
+		}
+	};
+
+	// Perform bake
+	const int32 TileWidth = 256;
+	const int32 TileHeight = 1;
+	const FImageTiling Tiles(Baker->Dimensions, TileWidth, TileHeight);
+	const int32 NumTiles = Tiles.Num();
+	const int NumBakers = Baker->Bakers.Num();
+	ParallelFor(NumTiles, [&Baker, &Tiles, &TileWidth, &NumBakers, &SampleSurface](const int32 TileIdx)
+	{
+		const FImageDimensions Tile = Tiles.GetTile(TileIdx);
+		const int Width = Tile.GetWidth();
+		for (int32 Idx = 0; Idx < Width; ++Idx)
+		{
+			const int ElemIdx = TileIdx * TileWidth + Idx;
+			FMeshMapEvaluator::FCorrespondenceSample Sample;
+			SampleSurface(ElemIdx, Sample);
+
+			FVector4f& Pixel = Baker->BakeResult->GetPixel(ElemIdx);
+			float* BufferPtr = &Pixel[0];
+			for (int32 BakerIdx = 0; BakerIdx < NumBakers; ++BakerIdx)
+			{
+				Baker->BakeContexts[BakerIdx].Evaluate(BufferPtr, Sample, Baker->BakeContexts[BakerIdx].EvalData);
+			}
+
+			// For color bakes, ask our evaluators to convert the float data to color.
+			if constexpr(ComputeMode == EBakeMode::Color)
+			{
+				// TODO: Use a separate buffer rather than R/W from the same pixel.
+				BufferPtr = &Pixel[0];
+				for (int32 BakerIdx = 0; BakerIdx < NumBakers; ++BakerIdx)
+				{
+					Baker->BakeContexts[BakerIdx].EvaluateColor(0, BufferPtr, Pixel, Baker->BakeContexts[BakerIdx].EvalData);
+				}
+			}
+		}
+	}, !Baker->bParallel ? EParallelForFlags::ForceSingleThread : EParallelForFlags::None);
+}
diff --git a/Engine/Plugins/Experimental/GeometryProcessing/Source/DynamicMesh/Public/Sampling/MeshBakerCommon.h b/Engine/Plugins/Experimental/GeometryProcessing/Source/DynamicMesh/Public/Sampling/MeshBakerCommon.h
new file mode 100644
--- /dev/null
+++ b/Engine/Plugins/Experimental/GeometryProcessing/Source/DynamicMesh/Public/Sampling/MeshBakerCommon.h
@@ -0,0 +1,70 @@
+// Copyright Epic Games, Inc. All Rights Reserved.
+
+#pragma once
+
+#include "DynamicMesh/DynamicMesh3.h"
+#include "DynamicMesh/DynamicMeshAABBTree3.h"
+#include "Image/ImageDimensions.h"
+
+namespace UE
+{
+namespace Geometry
+{
+
+/** Compute base/detail mesh correspondence by raycast */
+bool GetDetailMeshTrianglePoint_Raycast(
+	const FDynamicMesh3& DetailMesh,
+	const FDynamicMeshAABBTree3& DetailSpatial,
+	const FVector3d& BasePoint,
+	const FVector3d& BaseNormal,
+	int32& DetailTriangleOut,
+	FVector3d& DetailTriBaryCoords,
+	double Thickness,
+	bool bFailToNearestPoint
+);
+
+/** Compute base/detail mesh correspondence by nearest distance */
+bool GetDetailMeshTrianglePoint_Nearest(
+	const FDynamicMesh3& DetailMesh,
+	const FDynamicMeshAABBTree3& DetailSpatial,
+	const FVector3d& BasePoint,
+	int32& DetailTriangleOut,
+	FVector3d& DetailTriBaryCoords);
+
+/** Image tile storage for map bakes. */
+class FMeshMapTileBuffer
+{
+public:
+	FMeshMapTileBuffer(const FImageDimensions& TileIn, const int32 PixelSizeIn)
+		: Tile(TileIn)
+		, PixelSize(PixelSizeIn + 1) // + 1 for accumulated pixel weight.
+	{
+		Buffer = static_cast<float*>(FMemory::MallocZeroed(sizeof(float) * PixelSize * Tile.Num()));
+	}
+
+	~FMeshMapTileBuffer()
+	{
+		FMemory::Free(Buffer);
+	}
+
+	float& GetPixelWeight(const int64 LinearIdx) const
+	{
+		checkSlow(LinearIdx >= 0 && LinearIdx < Tile.Num());
+		return Buffer[LinearIdx * PixelSize];
+	}
+
+	float* GetPixel(const int64 LinearIdx) const
+	{
+		checkSlow(LinearIdx >= 0 && LinearIdx < Tile.Num());
+		return &Buffer[LinearIdx * PixelSize + 1];
+	}
+
+private:
+	const FImageDimensions Tile;
+	const int32 PixelSize;
+	float* Buffer;
+};
+
+} // end namespace UE::Geometry
+} // end namespace UE
+	
\ No newline at end of file
diff --git a/Engine/Plugins/Experimental/GeometryProcessing/Source/DynamicMesh/Public/Sampling/MeshBaseBaker.h b/Engine/Plugins/Experimental/GeometryProcessing/Source/DynamicMesh/Public/Sampling/MeshBaseBaker.h
new file mode 100644
--- /dev/null
+++ b/Engine/Plugins/Experimental/GeometryProcessing/Source/DynamicMesh/Public/Sampling/MeshBaseBaker.h
@@ -0,0 +1,137 @@
+// Copyright Epic Games, Inc. All Rights Reserved.
+
+#pragma once
+
+#include "DynamicMesh/DynamicMesh3.h"
+#include "DynamicMesh/DynamicMeshAttributeSet.h"
+#include "DynamicMesh/DynamicMeshAABBTree3.h"
+#include "DynamicMesh/MeshTangents.h"
+
+namespace UE
+{
+namespace Geometry
+{
+
+class FMeshBaseBaker
+{
+public:
+	virtual ~FMeshBaseBaker() = default;
+	
+	/**
+	* ECorrespondenceStrategy determines the basic approach that will be used to establish a
+	* mapping from points on the BakeTarget Mesh (usually low-poly) to points on the Detail Mesh (eg highpoly).
+	* Geometrically this is not a 1-1 mapping so there are various options
+	*/
+	enum class ECorrespondenceStrategy
+	{
+		/** Raycast inwards from Point+Thickness*Normal, if that misses, try Outwards from Point, then Inwards from Point */
+		RaycastStandard,
+		/** Use geometrically nearest point. Thickness is ignored */
+		NearestPoint,
+		/** Use RaycastStandard but fall back to NearestPoint if none of the rays hit */
+		RaycastStandardThenNearest,
+		/** Assume that BakeTarget == DetailMesh and so no mapping is necessary */
+		Identity
+	};
+
+	// Setters
+	void SetDetailMesh(const FDynamicMesh3* Mesh, const FDynamicMeshAABBTree3* Spatial)
+	{
+		DetailMesh = Mesh;
+		DetailSpatial = Spatial;
+	}
+	void SetTargetMesh(const FDynamicMesh3* Mesh)
+	{
+		TargetMesh = Mesh;
+	}
+	void SetTargetMeshTangents(const TSharedPtr<FMeshTangentsd, ESPMode::ThreadSafe> Tangents)
+	{
+		TargetMeshTangents = Tangents;
+	}
+	void SetUVLayer(const int32 UVLayerIn)
+	{
+		UVLayer = UVLayerIn;
+	}
+	void SetThickness(const double ThicknessIn)
+	{
+		Thickness = ThicknessIn;
+	}
+	void SetCorrespondenceStrategy(const ECorrespondenceStrategy Strategy)
+	{
+		CorrespondenceStrategy = Strategy;
+	}
+
+	// TargetMesh Getters
+	const FDynamicMesh3* GetTargetMesh() const
+	{
+		return TargetMesh;
+	}
+	const FDynamicMeshUVOverlay* GetTargetMeshUVs() const
+	{
+		check(TargetMesh && TargetMesh->HasAttributes() && UVLayer < TargetMesh->Attributes()->NumUVLayers());
+        return TargetMesh->Attributes()->GetUVLayer(UVLayer);
+	}
+	const FDynamicMeshNormalOverlay* GetTargetMeshNormals() const
+	{
+		check(TargetMesh && TargetMesh->HasAttributes());
+		return TargetMesh->Attributes()->PrimaryNormals();
+	}
+	TSharedPtr<FMeshTangentsd, ESPMode::ThreadSafe> GetTargetMeshTangents() const
+	{
+		return TargetMeshTangents;
+	}
+
+	// DetailMesh Getters
+	const FDynamicMesh3* GetDetailMesh() const
+	{
+		return DetailMesh;
+	}
+	const FDynamicMeshAABBTree3* GetDetailMeshSpatial() const
+	{
+		return DetailSpatial;
+	}
+	const FDynamicMeshNormalOverlay* GetDetailMeshNormals() const
+	{
+		check(DetailMesh && DetailMesh->HasAttributes());
+		return DetailMesh->Attributes()->PrimaryNormals();
+	}
+	const FDynamicMeshUVOverlay* GetDetailMeshUVs(int32 UVLayerIn=0) const
+	{
+		check(DetailMesh && DetailMesh->HasAttributes());
+		return DetailMesh->Attributes()->GetUVLayer(UVLayerIn);
+	}
+	const FDynamicMeshColorOverlay* GetDetailMeshColors() const
+	{
+		check(DetailMesh && DetailMesh->HasAttributes());
+		return DetailMesh->Attributes()->PrimaryColors();
+	}
+
+	// Other Getters
+	int32 GetUVLayer() const
+	{
+		return UVLayer;
+	}
+	double GetThickness() const
+	{
+		return Thickness;
+	}
+	ECorrespondenceStrategy GetCorrespondenceStrategy() const
+	{
+		return CorrespondenceStrategy;
+	}
+
+protected:
+	const FDynamicMesh3* DetailMesh = nullptr;
+	const FDynamicMeshAABBTree3* DetailSpatial = nullptr;
+	const FDynamicMesh3* TargetMesh = nullptr;
+	TSharedPtr<FMeshTangentsd, ESPMode::ThreadSafe> TargetMeshTangents;
+
+	int32 UVLayer = 0;
+	double Thickness = 3.0;
+	ECorrespondenceStrategy CorrespondenceStrategy = ECorrespondenceStrategy::RaycastStandard;
+};
+	
+
+} // end namespace UE::Geometry
+} // end namespace UE
+	
\ No newline at end of file
diff --git a/Engine/Plugins/Experimental/GeometryProcessing/Source/DynamicMesh/Public/Sampling/MeshConstantMapEvaluator.h b/Engine/Plugins/Experimental/GeometryProcessing/Source/DynamicMesh/Public/Sampling/MeshConstantMapEvaluator.h
new file mode 100644
--- /dev/null
+++ b/Engine/Plugins/Experimental/GeometryProcessing/Source/DynamicMesh/Public/Sampling/MeshConstantMapEvaluator.h
@@ -0,0 +1,41 @@
+// Copyright Epic Games, Inc. All Rights Reserved.
+
+#pragma once
+
+#include "Sampling/MeshMapEvaluator.h"
+#include "DynamicMesh/MeshTangents.h"
+
+namespace UE
+{
+namespace Geometry
+{
+
+/**
+ * A mesh evaluator for constant data. This evaluator can be useful as a filler
+ * when computing per channel color data.
+ */
+class DYNAMICMESH_API FMeshConstantMapEvaluator : public FMeshMapEvaluator
+{
+public:
+	FMeshConstantMapEvaluator() = default;
+	explicit FMeshConstantMapEvaluator(float ValueIn);
+	
+	// Begin FMeshMapEvaluator interface
+	virtual void Setup(const FMeshBaseBaker& Baker, FEvaluationContext& Context) override;
+
+	virtual EMeshMapEvaluatorType Type() const override { return EMeshMapEvaluatorType::Constant; }
+	// End FMeshMapEvaluator interface
+
+	static void EvaluateSample(float*& Out, const FCorrespondenceSample& Sample, void* EvalData);
+
+	static void EvaluateDefault(float*& Out, void* EvalData);
+
+	static void EvaluateColor(const int DataIdx, float*& In, FVector4f& Out, void* EvalData);
+
+public:
+	float Value = 0.0f;
+};
+
+} // end namespace UE::Geometry
+} // end namespace UE
+
diff --git a/Engine/Plugins/Experimental/GeometryProcessing/Source/DynamicMesh/Public/Sampling/MeshCurvatureMapEvaluator.h b/Engine/Plugins/Experimental/GeometryProcessing/Source/DynamicMesh/Public/Sampling/MeshCurvatureMapEvaluator.h
--- a/Engine/Plugins/Experimental/GeometryProcessing/Source/DynamicMesh/Public/Sampling/MeshCurvatureMapEvaluator.h
+++ b/Engine/Plugins/Experimental/GeometryProcessing/Source/DynamicMesh/Public/Sampling/MeshCurvatureMapEvaluator.h
@@ -11,6 +11,9 @@ namespace Geometry
 
 class FMeshVertexCurvatureCache;
 
+/**
+ * A mesh evaluator for mesh curvatures.
+ */
 class DYNAMICMESH_API FMeshCurvatureMapEvaluator : public FMeshMapEvaluator
 {
 public:
@@ -51,7 +54,7 @@ public:
 
 public:
 	// Begin FMeshMapEvaluator interface
-	virtual void Setup(const FMeshMapBaker& Baker, FEvaluationContext& Context) override;
+	virtual void Setup(const FMeshBaseBaker& Baker, FEvaluationContext& Context) override;
 
 	virtual EMeshMapEvaluatorType Type() const override { return EMeshMapEvaluatorType::Curvature; }
 	// End FMeshMapEvaluator interface
@@ -60,6 +63,8 @@ public:
 
 	static void EvaluateDefault(float*& Out, void* EvalData);
 
+	static void EvaluateColor(const int DataIdx, float*& In, FVector4f& Out, void* EvalData);
+
 	/** Populate Curvatures member if valid data has not been provided */
 	void CacheDetailCurvatures(const FDynamicMesh3* DetailMesh);
 
diff --git a/Engine/Plugins/Experimental/GeometryProcessing/Source/DynamicMesh/Public/Sampling/MeshMapBaker.h b/Engine/Plugins/Experimental/GeometryProcessing/Source/DynamicMesh/Public/Sampling/MeshMapBaker.h
--- a/Engine/Plugins/Experimental/GeometryProcessing/Source/DynamicMesh/Public/Sampling/MeshMapBaker.h
+++ b/Engine/Plugins/Experimental/GeometryProcessing/Source/DynamicMesh/Public/Sampling/MeshMapBaker.h
@@ -2,10 +2,7 @@
 
 #pragma once
 
-#include "DynamicMesh/DynamicMesh3.h"
-#include "DynamicMesh/DynamicMeshAttributeSet.h"
-#include "DynamicMesh/DynamicMeshAABBTree3.h"
-#include "DynamicMesh/MeshTangents.h"
+#include "Sampling/MeshBaseBaker.h"
 #include "Sampling/MeshMapEvaluator.h"
 #include "Sampling/MeshSurfaceSampler.h"
 #include "Spatial/DenseGrid2.h"
@@ -18,41 +15,9 @@ namespace Geometry
 {
 
 class FImageOccupancyMap;
+class FMeshMapTileBuffer;
 
-class FMeshMapTileBuffer
-{
-public:
-	FMeshMapTileBuffer(const FImageDimensions& TileIn, const int32 PixelSizeIn)
-		: Tile(TileIn)
-		, PixelSize(PixelSizeIn + 1) // + 1 for accumulated pixel weight.
-	{
-		Buffer = static_cast<float*>(FMemory::MallocZeroed(sizeof(float) * PixelSize * Tile.Num()));
-	}
-
-	~FMeshMapTileBuffer()
-	{
-		FMemory::Free(Buffer);
-	}
-
-	float& GetPixelWeight(int64 LinearIdx)
-	{
-		checkSlow(LinearIdx >= 0 && LinearIdx < Tile.Num());
-		return Buffer[LinearIdx * PixelSize];
-	}
-
-	float* GetPixel(int64 LinearIdx)
-	{
-		checkSlow(LinearIdx >= 0 && LinearIdx < Tile.Num());
-		return &Buffer[LinearIdx * PixelSize + 1];
-	}
-
-private:
-	const FImageDimensions Tile;
-	const int32 PixelSize;
-	float* Buffer;
-};
-
-class DYNAMICMESH_API FMeshMapBaker
+class DYNAMICMESH_API FMeshMapBaker : public FMeshBaseBaker
 {
 public:
 
@@ -85,52 +50,14 @@ public:
 	//
 	// Parameters
 	//
-
-	/**
-	 * ECorrespondenceStrategy determines the basic approach that will be used to establish a
-	 * mapping from points on the BakeTarget Mesh (usually low-poly) to points on the Detail Mesh (eg highpoly).
-	 * Geometrically this is not a 1-1 mapping so there are various options
-	 */
-	enum class ECorrespondenceStrategy
-	{
-		/** Raycast inwards from Point+Thickness*Normal, if that misses, try Outwards from Point, then Inwards from Point */
-		RaycastStandard,
-		/** Use geometrically nearest point. Thickness is ignored */
-		NearestPoint,
-		/** Use RaycastStandard but fall back to NearestPoint if none of the rays hit */
-		RaycastStandardThenNearest,
-		/** Assume that BakeTarget == DetailMesh and so no mapping is necessary */
-		Identity
-	};
-
-	void SetDetailMesh(const FDynamicMesh3* Mesh, const FDynamicMeshAABBTree3* Spatial);
-	void SetTargetMesh(const FDynamicMesh3* Mesh);
-
-	void SetTargetMeshTangents(TSharedPtr<FMeshTangentsd, ESPMode::ThreadSafe> Tangents);
-
-	const FDynamicMesh3* GetTargetMesh() const;
-	const FDynamicMeshUVOverlay* GetTargetMeshUVs() const;
-	const FDynamicMeshNormalOverlay* GetTargetMeshNormals() const;
-	TSharedPtr<FMeshTangentsd, ESPMode::ThreadSafe> GetTargetMeshTangents() const;
-
-	const FDynamicMesh3* GetDetailMesh() const;
-	const FDynamicMeshAABBTree3* GetDetailMeshSpatial() const;
-	const FDynamicMeshNormalOverlay* GetDetailMeshNormals() const;
-	const FDynamicMeshUVOverlay* GetDetailMeshUVs(int32 UVLayer=0) const;
-
+	
 	void SetDimensions(FImageDimensions DimensionsIn);
-	void SetUVLayer(int32 UVLayerIn);
-	void SetThickness(double ThicknessIn);
 	void SetGutterSize(int32 GutterSizeIn);
 	void SetMultisampling(int32 MultisamplingIn);
-	void SetCorrespondenceStrategy(ECorrespondenceStrategy Strategy);
 
 	FImageDimensions GetDimensions() const { return Dimensions; }
-	int32 GetUVLayer() const { return UVLayer; }
-	double GetThickness() const { return Thickness; }
 	int32 GetGutterSize() const { return GutterSize; }
 	int32 GetMultisampling() const { return Multisampling; }
-	ECorrespondenceStrategy GetCorrespondenceStrategy() const { return CorrespondenceStrategy; }
 
 protected:
 	/** Evaluate samples for this tile pixel. */
@@ -146,31 +73,13 @@ protected:
 	/** Initialize bake sample default floats and colors. */
 	void InitBakeDefaults();
 
-	/**
-	 * Convert float buffer value to float4 color. This function
-	 * will advance the buffer pointer by the stride value.
-	 * 
-	 * @param Buffer [out] input buffer data.
-	 * @param Stride the data layout of the buffer pointer.
-	 * @param Weight the weight to apply to the buffer data.
-	 */
-	FVector4f FloatToPixel(float*& Buffer, FMeshMapEvaluator::EComponents Stride, float Weight);
-
 protected:
 	const bool bParallel = true;
 
-	const FDynamicMesh3* DetailMesh = nullptr;
-	const FDynamicMeshAABBTree3* DetailSpatial = nullptr;
-	const FDynamicMesh3* TargetMesh = nullptr;
-	TSharedPtr<FMeshTangentsd, ESPMode::ThreadSafe> TargetMeshTangents;
-
 	FDynamicMesh3 FlatMesh;
 	TMeshSurfaceUVSampler<FMeshMapEvaluator::FCorrespondenceSample> DetailMeshSampler;
 
 	FImageDimensions Dimensions = FImageDimensions(128, 128);
-	int32 UVLayer = 0;
-	double Thickness = 3.0;
-	ECorrespondenceStrategy CorrespondenceStrategy = ECorrespondenceStrategy::RaycastStandard;
 	int32 GutterSize = 4;
 
 	/** The square dimensions for multisampling each pixel. */
diff --git a/Engine/Plugins/Experimental/GeometryProcessing/Source/DynamicMesh/Public/Sampling/MeshMapEvaluator.h b/Engine/Plugins/Experimental/GeometryProcessing/Source/DynamicMesh/Public/Sampling/MeshMapEvaluator.h
--- a/Engine/Plugins/Experimental/GeometryProcessing/Source/DynamicMesh/Public/Sampling/MeshMapEvaluator.h
+++ b/Engine/Plugins/Experimental/GeometryProcessing/Source/DynamicMesh/Public/Sampling/MeshMapEvaluator.h
@@ -9,10 +9,11 @@ namespace UE
 namespace Geometry
 {
 
-class FMeshMapBaker;
+class FMeshBaseBaker;
 
 enum class EMeshMapEvaluatorType
 {
+	Constant,
 	Normal,
 	Occlusion,
 	Curvature,
@@ -59,7 +60,7 @@ public:
 	 * 1. The evaluation functions.
 	 * 2. The accumulation method for data samples.
 	 * 3. The number of output results/targets for this evaluator.
-	 * 3. The data layout of each output result/target for this evaluator.
+	 * 4. The data layout of each output result/target for this evaluator.
 	 */
 	struct FEvaluationContext
 	{
@@ -68,7 +69,7 @@ public:
 		 * float buffer. This function should:
 		 * 
 		 * - assume that the size of the buffer is correct.
-		 * - advance the the buffer pointer by the DataLayout stride.
+		 * - advance the buffer pointer by the DataLayout stride.
 		 * 
 		 * @param Out the buffer to populate.
 		 * @param Sample the detail/target mesh correspondence sample data.
@@ -87,7 +88,7 @@ public:
 		 * This function should:
 		 * 
 		 * - assume that the size of the buffer is correct.
-		 * - advance the the buffer pointer by the DataLayout stride.
+		 * - advance the buffer pointer by the DataLayout stride.
 		 * 
 		 * @param Out the buffer to populate.
 		 * @param Sample the detail/target mesh correspondence sample data.
@@ -96,6 +97,25 @@ public:
 		using EvaluateDefaultFn = void(*)(float*& Out, void* EvalData);
 		EvaluateDefaultFn EvaluateDefault = nullptr;
 
+		/**
+		 * Function pointer to evaluate the color representation of
+		 * an evaluated sample result.
+		 *
+		 * This function will be invoked once per pixel after the
+		 * evaluated samples have been accumulated. This function
+		 * should:
+		 *
+		 * - assume that the size of the buffer is correct.
+		 * - advance the [In] buffer pointer by the DataLayout stride.
+		 *
+		 * @param DataIdx the index into the DataLayout being processed.
+		 * @param In the buffer containing the accumulated sample result.
+		 * @param Out the output float4 color to populate
+		 * @param EvalData custom data pointer provided by the evaluation context.
+		 */
+		using EvaluateColorFn = void(*)(const int DataIdx, float*& In, FVector4f& Out, void* EvalData);
+		EvaluateColorFn EvaluateColor = nullptr;
+
 		/** Define custom data to be passed to the evaluation function. */
 		void* EvalData = nullptr;
 
@@ -121,7 +141,7 @@ public:
 	 * @param Baker the baker that owns this evaluator.
 	 * @param Context [out] the evaluation context.
 	 */
-	virtual void Setup(const FMeshMapBaker& Baker, FEvaluationContext& Context) = 0;
+	virtual void Setup(const FMeshBaseBaker& Baker, FEvaluationContext& Context) = 0;
 
 	/** @return the type of evaluator. */
 	virtual EMeshMapEvaluatorType Type() const = 0;
diff --git a/Engine/Plugins/Experimental/GeometryProcessing/Source/DynamicMesh/Public/Sampling/MeshNormalMapEvaluator.h b/Engine/Plugins/Experimental/GeometryProcessing/Source/DynamicMesh/Public/Sampling/MeshNormalMapEvaluator.h
--- a/Engine/Plugins/Experimental/GeometryProcessing/Source/DynamicMesh/Public/Sampling/MeshNormalMapEvaluator.h
+++ b/Engine/Plugins/Experimental/GeometryProcessing/Source/DynamicMesh/Public/Sampling/MeshNormalMapEvaluator.h
@@ -10,11 +10,14 @@ namespace UE
 namespace Geometry
 {
 
+/**
+ * A mesh evaluator for tangent space normals.
+ */
 class DYNAMICMESH_API FMeshNormalMapEvaluator : public FMeshMapEvaluator
 {
 public:
 	// Begin FMeshMapEvaluator interface
-	virtual void Setup(const FMeshMapBaker& Baker, FEvaluationContext& Context) override;
+	virtual void Setup(const FMeshBaseBaker& Baker, FEvaluationContext& Context) override;
 
 	virtual EMeshMapEvaluatorType Type() const override { return EMeshMapEvaluatorType::Normal; }
 	// End FMeshMapEvaluator interface
@@ -23,16 +26,18 @@ public:
 
 	static void EvaluateDefault(float*& Out, void* EvalData);
 
+	static void EvaluateColor(const int DataIdx, float*& In, FVector4f& Out, void* EvalData);
+
 protected:
 	// Cached data
 	const FDynamicMesh3* DetailMesh = nullptr;
 	const FDynamicMeshNormalOverlay* DetailNormalOverlay = nullptr;
-	const TMeshTangents<double>* BaseMeshTangents;
+	const TMeshTangents<double>* BaseMeshTangents = nullptr;
 
 	FVector3f DefaultNormal = FVector3f::UnitZ();
 
 private:
-	FVector3f SampleFunction(const FCorrespondenceSample& Sample);
+	FVector3f SampleFunction(const FCorrespondenceSample& SampleData) const;
 };
 
 } // end namespace UE::Geometry
diff --git a/Engine/Plugins/Experimental/GeometryProcessing/Source/DynamicMesh/Public/Sampling/MeshOcclusionMapEvaluator.h b/Engine/Plugins/Experimental/GeometryProcessing/Source/DynamicMesh/Public/Sampling/MeshOcclusionMapEvaluator.h
--- a/Engine/Plugins/Experimental/GeometryProcessing/Source/DynamicMesh/Public/Sampling/MeshOcclusionMapEvaluator.h
+++ b/Engine/Plugins/Experimental/GeometryProcessing/Source/DynamicMesh/Public/Sampling/MeshOcclusionMapEvaluator.h
@@ -22,6 +22,9 @@ enum class EMeshOcclusionMapType : uint8
 };
 ENUM_CLASS_FLAGS(EMeshOcclusionMapType);
 
+/**
+ * A mesh evaluator for occlusion data (Ambient Occlusion & Bent Normals).
+ */
 class DYNAMICMESH_API FMeshOcclusionMapEvaluator : public FMeshMapEvaluator
 {
 public:
@@ -53,7 +56,7 @@ public:
 
 public:
 	// Begin FMeshMapEvaluator interface
-	virtual void Setup(const FMeshMapBaker& Baker, FEvaluationContext& Context) override;
+	virtual void Setup(const FMeshBaseBaker& Baker, FEvaluationContext& Context) override;
 
 	virtual EMeshMapEvaluatorType Type() const override { return EMeshMapEvaluatorType::Occlusion; }
 	// End FMeshMapEvaluator interface
@@ -64,6 +67,9 @@ public:
 	template <EMeshOcclusionMapType OcclusionType, ESpace NormalSpace>
 	static void EvaluateDefault(float*& Out, void* EvalData);
 
+	template <EMeshOcclusionMapType OcclusionType, ESpace NormalSpace>
+	static void EvaluateColor(const int DataIdx, float*& In, FVector4f& Out, void* EvalData);
+
 protected:
 	// Cached data
 	const FDynamicMesh3* DetailMesh = nullptr;
diff --git a/Engine/Plugins/Experimental/GeometryProcessing/Source/DynamicMesh/Public/Sampling/MeshPropertyMapEvaluator.h b/Engine/Plugins/Experimental/GeometryProcessing/Source/DynamicMesh/Public/Sampling/MeshPropertyMapEvaluator.h
--- a/Engine/Plugins/Experimental/GeometryProcessing/Source/DynamicMesh/Public/Sampling/MeshPropertyMapEvaluator.h
+++ b/Engine/Plugins/Experimental/GeometryProcessing/Source/DynamicMesh/Public/Sampling/MeshPropertyMapEvaluator.h
@@ -15,9 +15,13 @@ enum class EMeshPropertyMapType
 	Normal = 2,
 	FacetNormal = 3,
 	UVPosition = 4,
-	MaterialID = 5
+	MaterialID = 5,
+	VertexColor = 6
 };
 
+/**
+ * A mesh evaluator for mesh properties as color data.
+ */
 class DYNAMICMESH_API FMeshPropertyMapEvaluator : public FMeshMapEvaluator
 {
 public:
@@ -25,7 +29,7 @@ public:
 
 public:
 	// Begin FMeshMapEvaluator interface
-	virtual void Setup(const FMeshMapBaker& Baker, FEvaluationContext& Context) override;
+	virtual void Setup(const FMeshBaseBaker& Baker, FEvaluationContext& Context) override;
 
 	virtual EMeshMapEvaluatorType Type() const override { return EMeshMapEvaluatorType::Property; }
 	// End FMeshMapEvaluator interface
@@ -34,11 +38,14 @@ public:
 
 	static void EvaluateDefault(float*& Out, void* EvalData);
 
+	static void EvaluateColor(const int DataIdx, float*& In, FVector4f& Out, void* EvalData);
+
 protected:
 	// Cached data
 	const FDynamicMesh3* DetailMesh = nullptr;
 	const FDynamicMeshNormalOverlay* DetailNormalOverlay = nullptr;
 	const FDynamicMeshUVOverlay* DetailUVOverlay = nullptr;
+	const FDynamicMeshColorOverlay* DetailColorOverlay = nullptr;
 	FAxisAlignedBox3d Bounds;
 	FVector3f DefaultValue = FVector3f::Zero();
 
diff --git a/Engine/Plugins/Experimental/GeometryProcessing/Source/DynamicMesh/Public/Sampling/MeshResampleImageEvaluator.h b/Engine/Plugins/Experimental/GeometryProcessing/Source/DynamicMesh/Public/Sampling/MeshResampleImageEvaluator.h
--- a/Engine/Plugins/Experimental/GeometryProcessing/Source/DynamicMesh/Public/Sampling/MeshResampleImageEvaluator.h
+++ b/Engine/Plugins/Experimental/GeometryProcessing/Source/DynamicMesh/Public/Sampling/MeshResampleImageEvaluator.h
@@ -10,7 +10,9 @@ namespace UE
 namespace Geometry
 {
 
-/** Resample image */
+/**
+ * A mesh evaluator for sampling 2D texture data. 
+ */
 class DYNAMICMESH_API FMeshResampleImageEvaluator : public FMeshMapEvaluator
 {
 public:
@@ -21,7 +23,7 @@ public:
 
 public:
 	// Begin FMeshMapEvaluator interface
-	virtual void Setup(const FMeshMapBaker& Baker, FEvaluationContext& Context) override;
+	virtual void Setup(const FMeshBaseBaker& Baker, FEvaluationContext& Context) override;
 
 	virtual EMeshMapEvaluatorType Type() const override { return EMeshMapEvaluatorType::ResampleImage; }
 	// End FMeshMapEvaluator interface
@@ -30,6 +32,8 @@ public:
 
 	static void EvaluateDefault(float*& Out, void* EvalData);
 
+	static void EvaluateColor(const int DataIdx, float*& In, FVector4f& Out, void* EvalData);
+
 protected:
 	// Cached data
 	const FDynamicMesh3* DetailMesh = nullptr;
@@ -38,7 +42,9 @@ private:
 	FVector4f ImageSampleFunction(const FCorrespondenceSample& Sample);
 };
 
-/** Resample multi texture by material ID */
+/**
+ * A mesh evaluator for sampling multiple 2D textures by material ID
+ */
 class DYNAMICMESH_API FMeshMultiResampleImageEvaluator : public FMeshResampleImageEvaluator
 {
 public:
@@ -46,7 +52,7 @@ public:
 
 public:
 	// Begin FMeshMapEvaluator interface
-	virtual void Setup(const FMeshMapBaker& Baker, FEvaluationContext& Context) override;
+	virtual void Setup(const FMeshBaseBaker& Baker, FEvaluationContext& Context) override;
 
 	virtual EMeshMapEvaluatorType Type() const override { return EMeshMapEvaluatorType::MultiResampleImage; }
 	// End FMeshMapEvaluator interface
diff --git a/Engine/Plugins/Experimental/GeometryProcessing/Source/DynamicMesh/Public/Sampling/MeshVertexBaker.h b/Engine/Plugins/Experimental/GeometryProcessing/Source/DynamicMesh/Public/Sampling/MeshVertexBaker.h
new file mode 100644
--- /dev/null
+++ b/Engine/Plugins/Experimental/GeometryProcessing/Source/DynamicMesh/Public/Sampling/MeshVertexBaker.h
@@ -0,0 +1,83 @@
+// Copyright Epic Games, Inc. All Rights Reserved.
+
+#pragma once
+
+#include "Sampling/MeshBaseBaker.h"
+#include "Sampling/MeshMapEvaluator.h"
+#include "Sampling/MeshConstantMapEvaluator.h"
+#include "Image/ImageBuilder.h"
+#include "Image/ImageDimensions.h"
+
+namespace UE
+{
+namespace Geometry
+{
+
+class DYNAMICMESH_API FMeshVertexBaker : public FMeshBaseBaker
+{
+public:
+	enum class EBakeMode
+	{
+		Channel,	// Bake scalar data into color channels
+		Color		// Bake color data to RGB
+	};
+
+	EBakeMode BakeMode = EBakeMode::Color;
+
+	TSharedPtr<FMeshMapEvaluator, ESPMode::ThreadSafe> ColorEvaluator;
+	TSharedPtr<FMeshMapEvaluator, ESPMode::ThreadSafe> ChannelEvaluators[4];
+
+public:	
+	//
+	// Bake
+	//
+
+	/** Process all bakers to generate the image result. */
+	void Bake();
+
+	/** @return the bake result image. */
+	const TImageBuilder<FVector4f>* GetBakeResult() const;
+
+	/** if this function returns true, we should abort calculation */
+	TFunction<bool(void)> CancelF = []() { return false; };
+
+	//
+	// Parameters
+	//
+
+protected:
+	template<EBakeMode ComputeMode>
+	static void BakeImpl(void* Data);
+
+protected:
+	const bool bParallel = true;
+
+	/** Internally cached image dimensions proportional to the number of unique vertex color elements. */
+	FImageDimensions Dimensions;
+
+	/** Bake output image. */
+	TUniquePtr<TImageBuilder<FVector4f>> BakeResult;
+
+	/** Internal list of bake evaluators. */
+	TArray<FMeshMapEvaluator*> Bakers;
+
+	/** Evaluation contexts for each mesh evaluator. */
+	TArray<FMeshMapEvaluator::FEvaluationContext> BakeContexts;
+
+	/** Internal cached default bake data. */
+	FVector4f BakeDefaults;
+
+	/** The total size of the temporary float buffer for BakeSample. */
+	int32 BakeSampleBufferSize = 0;
+
+	/** */
+	using BakeFn = void(*)(void* Data);
+	BakeFn BakeInternal = nullptr;
+
+	/** Constant evaluator for empty channels. */
+	static FMeshConstantMapEvaluator ZeroEvaluator;
+	static FMeshConstantMapEvaluator OneEvaluator;
+};
+
+} // end namespace UE::Geometry
+} // end namespace UE
\ No newline at end of file
diff --git a/Engine/Plugins/Experimental/MeshModelingToolset/Source/MeshModelingTools/Private/BakeMeshAttributeMapsTool.cpp b/Engine/Plugins/Experimental/MeshModelingToolset/Source/MeshModelingTools/Private/BakeMeshAttributeMapsTool.cpp
--- a/Engine/Plugins/Experimental/MeshModelingToolset/Source/MeshModelingTools/Private/BakeMeshAttributeMapsTool.cpp
+++ b/Engine/Plugins/Experimental/MeshModelingToolset/Source/MeshModelingTools/Private/BakeMeshAttributeMapsTool.cpp
@@ -60,7 +60,8 @@ static constexpr EBakeMapType ALL_BAKE_MAP_TYPES[] =
 	EBakeMapType::FaceNormalImage,
 	EBakeMapType::PositionImage,
 	EBakeMapType::MaterialID,
-	EBakeMapType::MultiTexture
+	EBakeMapType::MultiTexture,
+	EBakeMapType::VertexColorImage
 };
 
 
@@ -121,11 +122,11 @@ public:
 
 	// Map Type settings
 	EBakeMapType Maps;
-	UBakeMeshAttributeMapsTool::FNormalMapSettings NormalSettings;
-	UBakeMeshAttributeMapsTool::FOcclusionMapSettings OcclusionSettings;
-	UBakeMeshAttributeMapsTool::FCurvatureMapSettings CurvatureSettings;
-	UBakeMeshAttributeMapsTool::FMeshPropertyMapSettings PropertySettings;
-	UBakeMeshAttributeMapsTool::FTexture2DImageSettings TextureSettings;
+	FNormalMapSettings NormalSettings;
+	FOcclusionMapSettings OcclusionSettings;
+	FCurvatureMapSettings CurvatureSettings;
+	FMeshPropertyMapSettings PropertySettings;
+	FTexture2DImageSettings TextureSettings;
 
 	// Texture2DImage & MultiTexture settings
 	using ImagePtr = TSharedPtr<UE::Geometry::TImageBuilder<FVector4f>, ESPMode::ThreadSafe>;
@@ -238,6 +239,13 @@ public:
 				Baker->AddBaker(PropertyBaker);
 				break;
 			}
+			case EBakeMapType::VertexColorImage:
+			{
+				TSharedPtr<FMeshPropertyMapEvaluator, ESPMode::ThreadSafe> PropertyBaker = MakeShared<FMeshPropertyMapEvaluator, ESPMode::ThreadSafe>();
+				PropertyBaker->Property = EMeshPropertyMapType::VertexColor;
+				Baker->AddBaker(PropertyBaker);
+				break;
+			}
 			case EBakeMapType::Texture2DImage:
 			{
 				TSharedPtr<FMeshResampleImageEvaluator, ESPMode::ThreadSafe> ResampleBaker = MakeShared<FMeshResampleImageEvaluator, ESPMode::ThreadSafe>();
@@ -469,7 +477,8 @@ TUniquePtr<UE::Geometry::TGenericDataOperator<FMeshMapBaker>> UBakeMeshAttribute
 	if ((bool)(CachedBakeCacheSettings.BakeMapTypes & EBakeMapType::NormalImage) ||
 		(bool)(CachedBakeCacheSettings.BakeMapTypes & EBakeMapType::FaceNormalImage) ||
 		(bool)(CachedBakeCacheSettings.BakeMapTypes & EBakeMapType::PositionImage) ||
-		(bool)(CachedBakeCacheSettings.BakeMapTypes & EBakeMapType::MaterialID))
+		(bool)(CachedBakeCacheSettings.BakeMapTypes & EBakeMapType::MaterialID) ||
+		(bool)(CachedBakeCacheSettings.BakeMapTypes & EBakeMapType::VertexColorImage))
 	{
 		Op->PropertySettings = CachedMeshPropertyMapSettings;
 	}
@@ -553,6 +562,9 @@ void UBakeMeshAttributeMapsTool::Shutdown(EToolShutdownType ShutdownType)
 				case EBakeMapType::MaterialID:
 					TexName = FString::Printf(TEXT("%s_MaterialIDImg"), *BaseName);
 					break;
+				case EBakeMapType::VertexColorImage:
+					TexName = FString::Printf(TEXT("%s_VertexColorIDImg"), *BaseName);
+					break;
 				case EBakeMapType::PositionImage:
 					TexName = FString::Printf(TEXT("%s_PositionImg"), *BaseName);
 					TexType = FTexture2DBuilder::ETextureType::Color;
@@ -801,7 +813,8 @@ void UBakeMeshAttributeMapsTool::UpdateResult()
 	if ((bool)(CachedBakeCacheSettings.BakeMapTypes & EBakeMapType::NormalImage) ||
 		(bool)(CachedBakeCacheSettings.BakeMapTypes & EBakeMapType::FaceNormalImage) ||
 		(bool)(CachedBakeCacheSettings.BakeMapTypes & EBakeMapType::PositionImage) ||
-		(bool)(CachedBakeCacheSettings.BakeMapTypes & EBakeMapType::MaterialID))
+		(bool)(CachedBakeCacheSettings.BakeMapTypes & EBakeMapType::MaterialID) ||
+		(bool)(CachedBakeCacheSettings.BakeMapTypes & EBakeMapType::VertexColorImage))
 	{
 		OpState |= UpdateResult_MeshProperty();
 	}
@@ -969,102 +982,6 @@ EBakeOpState UBakeMeshAttributeMapsTool::UpdateResult_MeshProperty()
 }
 
 
-
-
-
-
-
-
-class FTempTextureAccess
-{
-public:
-	FTempTextureAccess(UTexture2D* DisplacementMap)
-		: DisplacementMap(DisplacementMap)
-	{
-		check(DisplacementMap);
-		OldCompressionSettings = DisplacementMap->CompressionSettings;
-		bOldSRGB = DisplacementMap->SRGB;
-#if WITH_EDITOR
-		OldMipGenSettings = DisplacementMap->MipGenSettings;
-#endif
-		DisplacementMap->CompressionSettings = TextureCompressionSettings::TC_VectorDisplacementmap;
-		DisplacementMap->SRGB = false;
-#if WITH_EDITOR
-		DisplacementMap->MipGenSettings = TextureMipGenSettings::TMGS_NoMipmaps;
-#endif
-		DisplacementMap->UpdateResource();
-
-		FormattedImageData = reinterpret_cast<const FColor*>(DisplacementMap->PlatformData->Mips[0].BulkData.LockReadOnly());
-	}
-	FTempTextureAccess(const FTempTextureAccess&) = delete;
-	FTempTextureAccess(FTempTextureAccess&&) = delete;
-	void operator=(const FTempTextureAccess&) = delete;
-	void operator=(FTempTextureAccess&&) = delete;
-
-	~FTempTextureAccess()
-	{
-		DisplacementMap->PlatformData->Mips[0].BulkData.Unlock();
-
-		DisplacementMap->CompressionSettings = OldCompressionSettings;
-		DisplacementMap->SRGB = bOldSRGB;
-#if WITH_EDITOR
-		DisplacementMap->MipGenSettings = OldMipGenSettings;
-#endif
-
-		DisplacementMap->UpdateResource();
-	}
-
-	bool HasData() const
-	{
-		return FormattedImageData != nullptr;
-	}
-	const FColor* GetData() const
-	{
-		return FormattedImageData;
-	}
-
-	FImageDimensions GetDimensions() const
-	{
-		int32 Width = DisplacementMap->PlatformData->Mips[0].SizeX;
-		int32 Height = DisplacementMap->PlatformData->Mips[0].SizeY;
-		return FImageDimensions(Width, Height);
-	}
-
-
-	bool CopyTo(TImageBuilder<FVector4f>& DestImage) const
-	{
-		if (!HasData()) return false;
-
-		FImageDimensions TextureDimensions = GetDimensions();
-		if (ensure(DestImage.GetDimensions() == TextureDimensions) == false)
-		{
-			return false;
-		}
-
-		int64 Num = TextureDimensions.Num();
-		for (int32 i = 0; i < Num; ++i)
-		{
-			FColor ByteColor = FormattedImageData[i];
-			FLinearColor FloatColor(ByteColor);
-			DestImage.SetPixel(i, FVector4f(FloatColor));
-		}
-		return true;
-	}
-
-private:
-	UTexture2D* DisplacementMap{ nullptr };
-	TextureCompressionSettings OldCompressionSettings{};
-	TextureMipGenSettings OldMipGenSettings{};
-	bool bOldSRGB{ false };
-	const FColor* FormattedImageData{ nullptr };
-};
-
-
-
-
-
-
-
 EBakeOpState UBakeMeshAttributeMapsTool::UpdateResult_Texture2DImage()
 {
 	EBakeOpState ResultState = EBakeOpState::Complete;
@@ -1234,6 +1151,7 @@ void UBakeMeshAttributeMapsTool::UpdateVisualization()
 			case EBakeMapType::MaterialID:
 			case EBakeMapType::Texture2DImage:
 			case EBakeMapType::MultiTexture:
+			case EBakeMapType::VertexColorImage:
 				PreviewMaterial->SetTextureParameterValue(TEXT("NormalMap"), EmptyNormalMap);
 				PreviewMaterial->SetTextureParameterValue(TEXT("OcclusionMap"), EmptyColorMapWhite);
 				PreviewMaterial->SetTextureParameterValue(TEXT("ColorMap"), PreviewMap);
@@ -1272,6 +1190,7 @@ void UBakeMeshAttributeMapsTool::UpdateOnModeChange()
 		case EBakeMapType::FaceNormalImage:
 		case EBakeMapType::PositionImage:
 		case EBakeMapType::MaterialID:
+		case EBakeMapType::VertexColorImage:
 			break;
 		case EBakeMapType::Texture2DImage:
 			SetToolPropertySourceEnabled(Texture2DProps, true);
@@ -1386,6 +1305,9 @@ void UBakeMeshAttributeMapsTool::OnMapsUpdated(const TUniquePtr<FMeshMapBaker>&
 			case EMeshPropertyMapType::MaterialID:
 				MapType = EBakeMapType::MaterialID;
 				break;
+			case EMeshPropertyMapType::VertexColor:
+				MapType = EBakeMapType::VertexColorImage;
+				break;
 			case EMeshPropertyMapType::UVPosition:
 			default:
 				break;
diff --git a/Engine/Plugins/Experimental/MeshModelingToolset/Source/MeshModelingTools/Private/BakeMeshAttributeToolCommon.cpp b/Engine/Plugins/Experimental/MeshModelingToolset/Source/MeshModelingTools/Private/BakeMeshAttributeToolCommon.cpp
new file mode 100644
--- /dev/null
+++ b/Engine/Plugins/Experimental/MeshModelingToolset/Source/MeshModelingTools/Private/BakeMeshAttributeToolCommon.cpp
@@ -0,0 +1,71 @@
+// Copyright Epic Games, Inc. All Rights Reserved.
+
+#include "BakeMeshAttributeToolCommon.h"
+
+FTempTextureAccess::FTempTextureAccess(UTexture2D* DisplacementMap)
+	: DisplacementMap(DisplacementMap)
+{
+	check(DisplacementMap);
+	OldCompressionSettings = DisplacementMap->CompressionSettings;
+	bOldSRGB = DisplacementMap->SRGB;
+#if WITH_EDITOR
+	OldMipGenSettings = DisplacementMap->MipGenSettings;
+#endif
+	DisplacementMap->CompressionSettings = TextureCompressionSettings::TC_VectorDisplacementmap;
+	DisplacementMap->SRGB = false;
+#if WITH_EDITOR
+	DisplacementMap->MipGenSettings = TextureMipGenSettings::TMGS_NoMipmaps;
+#endif
+	DisplacementMap->UpdateResource();
+
+	FormattedImageData = reinterpret_cast<const FColor*>(DisplacementMap->PlatformData->Mips[0].BulkData.LockReadOnly());
+}
+
+FTempTextureAccess::~FTempTextureAccess()
+{
+	DisplacementMap->PlatformData->Mips[0].BulkData.Unlock();
+	DisplacementMap->CompressionSettings = OldCompressionSettings;
+	DisplacementMap->SRGB = bOldSRGB;
+#if WITH_EDITOR
+	DisplacementMap->MipGenSettings = OldMipGenSettings;
+#endif
+	DisplacementMap->UpdateResource();
+}
+
+bool FTempTextureAccess::HasData() const
+{
+	return FormattedImageData != nullptr;
+}
+	
+const FColor* FTempTextureAccess::GetData() const
+{
+	return FormattedImageData;
+}
+
+FImageDimensions FTempTextureAccess::GetDimensions() const
+{
+	const int32 Width = DisplacementMap->PlatformData->Mips[0].SizeX;
+	const int32 Height = DisplacementMap->PlatformData->Mips[0].SizeY;
+	return FImageDimensions(Width, Height);
+}
+
+bool FTempTextureAccess::CopyTo(TImageBuilder<FVector4f>& DestImage) const
+{
+	if (!HasData()) return false;
+
+	const FImageDimensions TextureDimensions = GetDimensions();
+	if (ensure(DestImage.GetDimensions() == TextureDimensions) == false)
+	{
+		return false;
+	}
+
+	const int64 Num = TextureDimensions.Num();
+	for (int32 i = 0; i < Num; ++i)
+	{
+		FColor ByteColor = FormattedImageData[i];
+		FLinearColor FloatColor(ByteColor);
+		DestImage.SetPixel(i, FVector4f(FloatColor));
+	}
+	return true;
+}
+
diff --git a/Engine/Plugins/Experimental/MeshModelingToolset/Source/MeshModelingTools/Private/BakeMeshAttributeVertexTool.cpp b/Engine/Plugins/Experimental/MeshModelingToolset/Source/MeshModelingTools/Private/BakeMeshAttributeVertexTool.cpp
new file mode 100644
--- /dev/null
+++ b/Engine/Plugins/Experimental/MeshModelingToolset/Source/MeshModelingTools/Private/BakeMeshAttributeVertexTool.cpp
@@ -0,0 +1,946 @@
+// Copyright Epic Games, Inc. All Rights Reserved.
+
+#include "BakeMeshAttributeVertexTool.h"
+#include "InteractiveToolManager.h"
+#include "ToolBuilderUtil.h"
+
+#include "DynamicMesh/MeshTransforms.h"
+#include "MeshDescriptionToDynamicMesh.h"
+
+#include "Sampling/MeshNormalMapEvaluator.h"
+#include "Sampling/MeshOcclusionMapEvaluator.h"
+#include "Sampling/MeshCurvatureMapEvaluator.h"
+#include "Sampling/MeshPropertyMapEvaluator.h"
+#include "Sampling/MeshResampleImageEvaluator.h"
+
+#include "TargetInterfaces/MaterialProvider.h"
+#include "TargetInterfaces/MeshDescriptionProvider.h"
+#include "TargetInterfaces/PrimitiveComponentBackedTarget.h"
+#include "TargetInterfaces/StaticMeshBackedTarget.h"
+#include "ToolTargetManager.h"
+#include "ModelingToolTargetUtil.h"
+
+#include "ExplicitUseGeometryMathTypes.h"		// using UE::Geometry::(math types)
+using namespace UE::Geometry;
+
+#define LOCTEXT_NAMESPACE "UBakeMeshAttributeVertexTool"
+
+/*
+ * ToolBuilder
+ */
+
+const FToolTargetTypeRequirements& UBakeMeshAttributeVertexToolBuilder::GetTargetRequirements() const
+{
+	static FToolTargetTypeRequirements TypeRequirements({
+		UMeshDescriptionProvider::StaticClass(),
+		UPrimitiveComponentBackedTarget::StaticClass(),
+		UStaticMeshBackedTarget::StaticClass(), // currently only supports StaticMesh targets
+		UMaterialProvider::StaticClass()
+	});
+	return TypeRequirements;
+}
+
+bool UBakeMeshAttributeVertexToolBuilder::CanBuildTool(const FToolBuilderState& SceneState) const
+{
+	const int32 NumTargets = SceneState.TargetManager->CountSelectedAndTargetable(SceneState, GetTargetRequirements());
+	return (NumTargets == 1 || NumTargets == 2);
+}
+
+UInteractiveTool* UBakeMeshAttributeVertexToolBuilder::BuildTool(const FToolBuilderState& SceneState) const
+{
+	UBakeMeshAttributeVertexTool* NewTool = NewObject<UBakeMeshAttributeVertexTool>(SceneState.ToolManager);
+	TArray<TObjectPtr<UToolTarget>> Targets = SceneState.TargetManager->BuildAllSelectedTargetable(
+		SceneState, GetTargetRequirements());
+	NewTool->SetTargets(MoveTemp(Targets));
+	NewTool->SetWorld(SceneState.World);
+	return NewTool;
+}
+
+
+/*
+ * Operators
+ */
+
+class FMeshVertexBakerOp : public TGenericDataOperator<FMeshVertexBaker>
+{
+public:
+	// General bake settings
+	TSharedPtr<UE::Geometry::FDynamicMesh3, ESPMode::ThreadSafe> DetailMesh;
+	TSharedPtr<UE::Geometry::FDynamicMeshAABBTree3, ESPMode::ThreadSafe> DetailSpatial;
+	UE::Geometry::FDynamicMesh3* BaseMesh;
+	TSharedPtr<UE::Geometry::TMeshTangents<double>, ESPMode::ThreadSafe> BaseMeshTangents;
+	TUniquePtr<UE::Geometry::FMeshVertexBaker> Baker;
+
+	UBakeMeshAttributeVertexTool::FBakeSettings BakeSettings;
+	UBakeMeshAttributeVertexTool::FBakeColorSettings ColorSettings;
+	UBakeMeshAttributeVertexTool::FBakeChannelSettings ChannelSettings;
+	FOcclusionMapSettings OcclusionSettings;
+	FCurvatureMapSettings CurvatureSettings;
+	FTexture2DImageSettings TextureSettings;
+
+	// Texture2DImage & MultiTexture settings
+	using ImagePtr = TSharedPtr<UE::Geometry::TImageBuilder<FVector4f>, ESPMode::ThreadSafe>;
+	const FDynamicMeshUVOverlay* UVOverlay = nullptr;
+	ImagePtr TextureImage;
+	TMap<int32, ImagePtr> MaterialToTextureImageMap;
+
+	// Begin TGenericDataOperator interface
+	virtual void CalculateResult(FProgressCancel* Progress) override
+	{
+		Baker = MakeUnique<FMeshVertexBaker>();
+		Baker->CancelF = [Progress]()
+		{
+			return Progress && Progress->Cancelled();
+		};
+		Baker->SetTargetMesh(BaseMesh);
+		Baker->SetDetailMesh(DetailMesh.Get(), DetailSpatial.Get());
+		Baker->SetTargetMeshTangents(BaseMeshTangents);
+		Baker->SetThickness(BakeSettings.Thickness);
+		Baker->BakeMode = BakeSettings.VertexMode == EBakeVertexMode::Color ? FMeshVertexBaker::EBakeMode::Color : FMeshVertexBaker::EBakeMode::Channel;
+
+		auto InitOcclusionEvaluator = [this] (FMeshOcclusionMapEvaluator* OcclusionEval, const EMeshOcclusionMapType OcclusionType)
+		{
+			OcclusionEval->OcclusionType = OcclusionType;
+			OcclusionEval->NumOcclusionRays = OcclusionSettings.OcclusionRays;
+			OcclusionEval->MaxDistance = OcclusionSettings.MaxDistance;
+			OcclusionEval->SpreadAngle = OcclusionSettings.SpreadAngle;
+			OcclusionEval->BiasAngleDeg = OcclusionSettings.BiasAngle;
+
+			switch (OcclusionSettings.Distribution)
+			{
+			case EOcclusionMapDistribution::Cosine:
+				OcclusionEval->Distribution = FMeshOcclusionMapEvaluator::EDistribution::Cosine;
+				break;
+			case EOcclusionMapDistribution::Uniform:
+				OcclusionEval->Distribution = FMeshOcclusionMapEvaluator::EDistribution::Uniform;
+				break;
+			}
+
+			switch (OcclusionSettings.NormalSpace)
+			{
+			case ENormalMapSpace::Tangent:
+				OcclusionEval->NormalSpace = FMeshOcclusionMapEvaluator::ESpace::Tangent;
+				break;
+			case ENormalMapSpace::Object:
+				OcclusionEval->NormalSpace = FMeshOcclusionMapEvaluator::ESpace::Object;
+				break;
+			}
+		};
+
+		if (BakeSettings.VertexMode == EBakeVertexMode::PerChannel)
+		{
+			for(int ChannelIdx = 0; ChannelIdx < 4; ++ChannelIdx)
+			{
+				switch(ChannelSettings.BakeType[ChannelIdx])
+				{
+				case EBakeVertexTypeChannel::AmbientOcclusion:
+				{
+					TSharedPtr<FMeshOcclusionMapEvaluator, ESPMode::ThreadSafe> OcclusionEval = MakeShared<FMeshOcclusionMapEvaluator, ESPMode::ThreadSafe>();
+					InitOcclusionEvaluator(OcclusionEval.Get(), EMeshOcclusionMapType::AmbientOcclusion);
+					Baker->ChannelEvaluators[ChannelIdx] = OcclusionEval;
+					break;				
+				}
+				case EBakeVertexTypeChannel::Curvature:
+				{
+					TSharedPtr<FMeshCurvatureMapEvaluator, ESPMode::ThreadSafe> CurvatureEval = MakeShared<FMeshCurvatureMapEvaluator, ESPMode::ThreadSafe>();
+					Baker->ChannelEvaluators[ChannelIdx] = CurvatureEval;
+					break;
+				}
+				default:
+				case EBakeVertexTypeChannel::None:
+				{
+					Baker->ChannelEvaluators[ChannelIdx] = nullptr;
+					break;
+				}
+				}
+			}
+		}
+		else // EBakeVertexMode::Color
+		{
+			switch (ColorSettings.BakeType)
+			{
+			case EBakeVertexTypeColor::TangentSpaceNormal:
+			{
+				Baker->ColorEvaluator = MakeShared<FMeshNormalMapEvaluator, ESPMode::ThreadSafe>();
+				break;
+			}
+			case EBakeVertexTypeColor::AmbientOcclusion:
+			{
+				TSharedPtr<FMeshOcclusionMapEvaluator, ESPMode::ThreadSafe> OcclusionEval = MakeShared<FMeshOcclusionMapEvaluator, ESPMode::ThreadSafe>();
+				InitOcclusionEvaluator(OcclusionEval.Get(), EMeshOcclusionMapType::AmbientOcclusion);
+				Baker->ColorEvaluator = OcclusionEval;
+				break;
+			}
+			case EBakeVertexTypeColor::BentNormal:
+			{
+				TSharedPtr<FMeshOcclusionMapEvaluator, ESPMode::ThreadSafe> OcclusionEval = MakeShared<FMeshOcclusionMapEvaluator, ESPMode::ThreadSafe>();
+				InitOcclusionEvaluator(OcclusionEval.Get(), EMeshOcclusionMapType::BentNormal);
+				Baker->ColorEvaluator = OcclusionEval;
+				break;
+			}
+			case EBakeVertexTypeColor::Curvature:
+			{
+				TSharedPtr<FMeshCurvatureMapEvaluator, ESPMode::ThreadSafe> CurvatureBaker = MakeShared<FMeshCurvatureMapEvaluator, ESPMode::ThreadSafe>();
+				CurvatureBaker->RangeScale = FMathd::Clamp(CurvatureSettings.RangeMultiplier, 0.0001, 1000.0);
+				CurvatureBaker->MinRangeScale = FMathd::Clamp(CurvatureSettings.MinRangeMultiplier, 0.0, 1.0);
+				CurvatureBaker->UseCurvatureType = (FMeshCurvatureMapEvaluator::ECurvatureType)CurvatureSettings.CurvatureType;
+				CurvatureBaker->UseColorMode = (FMeshCurvatureMapEvaluator::EColorMode)CurvatureSettings.ColorMode;
+				CurvatureBaker->UseClampMode = (FMeshCurvatureMapEvaluator::EClampMode)CurvatureSettings.ClampMode;
+				Baker->ColorEvaluator = CurvatureBaker;
+				break;
+			}
+			case EBakeVertexTypeColor::PositionImage:
+			{
+				TSharedPtr<FMeshPropertyMapEvaluator, ESPMode::ThreadSafe> PropertyBaker = MakeShared<FMeshPropertyMapEvaluator, ESPMode::ThreadSafe>();
+				PropertyBaker->Property = EMeshPropertyMapType::Position;
+				Baker->ColorEvaluator = PropertyBaker;
+				break;
+			}
+			case EBakeVertexTypeColor::NormalImage:
+			{
+				TSharedPtr<FMeshPropertyMapEvaluator, ESPMode::ThreadSafe> PropertyBaker = MakeShared<FMeshPropertyMapEvaluator, ESPMode::ThreadSafe>();
+				PropertyBaker->Property = EMeshPropertyMapType::Normal;
+				Baker->ColorEvaluator = PropertyBaker;
+				break;
+			}
+			case EBakeVertexTypeColor::FaceNormalImage:
+			{
+				TSharedPtr<FMeshPropertyMapEvaluator, ESPMode::ThreadSafe> PropertyBaker = MakeShared<FMeshPropertyMapEvaluator, ESPMode::ThreadSafe>();
+				PropertyBaker->Property = EMeshPropertyMapType::FacetNormal;
+				Baker->ColorEvaluator = PropertyBaker;
+				break;
+			}
+			case EBakeVertexTypeColor::MaterialID:
+			{
+				TSharedPtr<FMeshPropertyMapEvaluator, ESPMode::ThreadSafe> PropertyBaker = MakeShared<FMeshPropertyMapEvaluator, ESPMode::ThreadSafe>();
+				PropertyBaker->Property = EMeshPropertyMapType::MaterialID;
+				Baker->ColorEvaluator = PropertyBaker;
+				break;
+			}
+			case EBakeVertexTypeColor::Texture2DImage:
+			{
+				TSharedPtr<FMeshResampleImageEvaluator, ESPMode::ThreadSafe> TextureBaker = MakeShared<FMeshResampleImageEvaluator, ESPMode::ThreadSafe>();
+				TextureBaker->DetailUVOverlay = UVOverlay;
+				TextureBaker->SampleFunction = [this](FVector2d UVCoord) {
+					return TextureImage->BilinearSampleUV<float>(UVCoord, FVector4f(0, 0, 0, 1));
+				};
+				Baker->ColorEvaluator = TextureBaker;
+				break;
+			}
+			case EBakeVertexTypeColor::MultiTexture:
+			{
+				TSharedPtr<FMeshMultiResampleImageEvaluator, ESPMode::ThreadSafe> TextureBaker = MakeShared<FMeshMultiResampleImageEvaluator, ESPMode::ThreadSafe>();
+				TextureBaker->DetailUVOverlay = UVOverlay;
+				TextureBaker->MultiTextures = MaterialToTextureImageMap;
+				Baker->ColorEvaluator = TextureBaker;
+				break;
+			}
+			}
+		}
+
+		Baker->Bake();
+		SetResult(MoveTemp(Baker));
+	}
+
+	// End TGenericDataOperator interface
+};
+
+
+/*
+ * Tool
+ */
+
+void UBakeMeshAttributeVertexTool::Setup()
+{
+	UInteractiveTool::Setup();
+
+	UMaterial* Material = LoadObject<UMaterial>(nullptr, TEXT("/MeshModelingToolset/Materials/MeshVertexColorMaterial"));
+	check(Material);
+	if (Material != nullptr)
+	{
+		PreviewMaterial = UMaterialInstanceDynamic::Create(Material, GetToolManager());
+	}
+
+	UMaterial* AlphaMaterial = LoadObject<UMaterial>(nullptr, TEXT("/MeshModelingToolset/Materials/MeshVertexAlphaMaterial"));
+	check(AlphaMaterial);
+	if (AlphaMaterial != nullptr)
+	{
+		PreviewAlphaMaterial = UMaterialInstanceDynamic::Create(AlphaMaterial, GetToolManager());
+	}
+
+	UMaterial* WorkingMaterial = LoadObject<UMaterial>(nullptr, TEXT("/MeshModelingToolset/Materials/InProgressMaterial"));
+	check(WorkingMaterial);
+	if (WorkingMaterial != nullptr)
+	{
+		WorkingPreviewMaterial = UMaterialInstanceDynamic::Create(WorkingMaterial, GetToolManager());
+	}
+
+	bIsBakeToSelf = (Targets.Num() == 1);
+
+	UE::ToolTarget::HideSourceObject(Targets[0]);
+
+	const FDynamicMesh3 InputMeshWithTangents = UE::ToolTarget::GetDynamicMeshCopy(Targets[0], true);
+	PreviewMesh = NewObject<UPreviewMesh>(this);
+	PreviewMesh->CreateInWorld(TargetWorld, FTransform::Identity);
+	PreviewMesh->SetTransform(static_cast<FTransform>(UE::ToolTarget::GetLocalToWorldTransform(Targets[0])));
+	PreviewMesh->SetTangentsMode(EDynamicMeshComponentTangentsMode::ExternallyProvided);
+	PreviewMesh->ReplaceMesh(InputMeshWithTangents);
+	PreviewMesh->SetMaterials(UE::ToolTarget::GetMaterialSet(Targets[0]).Materials);
+	PreviewMesh->SetOverrideRenderMaterial(PreviewMaterial);
+	PreviewMesh->SetVisible(true);
+
+	PreviewMesh->ProcessMesh([this](const FDynamicMesh3& Mesh)
+	{
+		BaseMesh.Copy(Mesh);
+		BaseSpatial.SetMesh(&BaseMesh, true);
+		BaseMeshTangents = MakeShared<FMeshTangentsd, ESPMode::ThreadSafe>(&BaseMesh);
+		BaseMeshTangents->CopyTriVertexTangents(Mesh);
+	});
+
+	// Setup tool property sets
+	Settings = NewObject<UBakeMeshAttributeVertexToolProperties>(this);
+	Settings->RestoreProperties(this);
+	AddToolPropertySource(Settings);
+
+	Settings->WatchProperty(Settings->VertexMode, [this](EBakeVertexMode) { OpState = EBakeOpState::Evaluate; UpdateOnModeChange(); });
+	Settings->WatchProperty(Settings->VertexChannelPreview, [this](EBakeVertexChannel) { UpdateVisualization(); });
+	Settings->WatchProperty(Settings->Thickness, [this](float) { OpState = EBakeOpState::Evaluate; });
+	Settings->WatchProperty(Settings->bUseWorldSpace, [this](bool)	{ bDetailMeshValid = false; OpState = EBakeOpState::Evaluate; });
+	Settings->WatchProperty(Settings->bSplitAtNormalSeams, [this](bool) { bColorTopologyValid = false; OpState = EBakeOpState::Evaluate; });
+	Settings->WatchProperty(Settings->bSplitAtUVSeams, [this](bool) { bColorTopologyValid = false; OpState = EBakeOpState::Evaluate; });
+
+	ColorSettings = NewObject<UBakeMeshAttributeVertexToolColorProperties>(this);
+	ColorSettings->RestoreProperties(this);
+	AddToolPropertySource(ColorSettings);
+	SetToolPropertySourceEnabled(ColorSettings, false);
+	ColorSettings->WatchProperty(ColorSettings->BakeType, [this](EBakeVertexTypeColor) { OpState = EBakeOpState::Evaluate; UpdateOnModeChange(); });
+
+	PerChannelSettings = NewObject<UBakeMeshAttributeVertexToolChannelProperties>(this);
+	PerChannelSettings->RestoreProperties(this);
+	AddToolPropertySource(PerChannelSettings);
+	SetToolPropertySourceEnabled(PerChannelSettings, false);
+	PerChannelSettings->WatchProperty(PerChannelSettings->BakeTypeR, [this](EBakeVertexTypeChannel) { OpState = EBakeOpState::Evaluate; UpdateOnModeChange(); });
+	PerChannelSettings->WatchProperty(PerChannelSettings->BakeTypeG, [this](EBakeVertexTypeChannel) { OpState = EBakeOpState::Evaluate; UpdateOnModeChange(); });
+	PerChannelSettings->WatchProperty(PerChannelSettings->BakeTypeB, [this](EBakeVertexTypeChannel) { OpState = EBakeOpState::Evaluate; UpdateOnModeChange(); });
+	PerChannelSettings->WatchProperty(PerChannelSettings->BakeTypeA, [this](EBakeVertexTypeChannel) { OpState = EBakeOpState::Evaluate; UpdateOnModeChange(); });
+
+	OcclusionSettings = NewObject<UBakedOcclusionMapToolProperties>(this);
+	OcclusionSettings->RestoreProperties(this);
+	AddToolPropertySource(OcclusionSettings);
+	SetToolPropertySourceEnabled(OcclusionSettings, false);
+	OcclusionSettings->WatchProperty(OcclusionSettings->OcclusionRays, [this](int32) { OpState = EBakeOpState::Evaluate; });
+	OcclusionSettings->WatchProperty(OcclusionSettings->MaxDistance, [this](float) { OpState = EBakeOpState::Evaluate; });
+	OcclusionSettings->WatchProperty(OcclusionSettings->SpreadAngle, [this](float) { OpState = EBakeOpState::Evaluate; });
+	OcclusionSettings->WatchProperty(OcclusionSettings->Distribution, [this](EOcclusionMapDistribution) { OpState = EBakeOpState::Evaluate; });
+	OcclusionSettings->WatchProperty(OcclusionSettings->BlurRadius, [this](float) { OpState = EBakeOpState::Evaluate; });
+	OcclusionSettings->WatchProperty(OcclusionSettings->bGaussianBlur, [this](float) { OpState = EBakeOpState::Evaluate; });
+	OcclusionSettings->WatchProperty(OcclusionSettings->BiasAngle, [this](float) { OpState = EBakeOpState::Evaluate; });
+	OcclusionSettings->WatchProperty(OcclusionSettings->NormalSpace, [this](ENormalMapSpace) { OpState = EBakeOpState::Evaluate; });
+
+	CurvatureSettings = NewObject<UBakedCurvatureMapToolProperties>(this);
+	CurvatureSettings->RestoreProperties(this);
+	AddToolPropertySource(CurvatureSettings);
+	SetToolPropertySourceEnabled(CurvatureSettings, false);
+	CurvatureSettings->WatchProperty(CurvatureSettings->RangeMultiplier, [this](float) { OpState = EBakeOpState::Evaluate; });
+	CurvatureSettings->WatchProperty(CurvatureSettings->MinRangeMultiplier, [this](float) { OpState = EBakeOpState::Evaluate; });
+	CurvatureSettings->WatchProperty(CurvatureSettings->CurvatureType, [this](EBakedCurvatureTypeMode) { OpState = EBakeOpState::Evaluate; });
+	CurvatureSettings->WatchProperty(CurvatureSettings->ColorMode, [this](EBakedCurvatureColorMode) { OpState = EBakeOpState::Evaluate; });
+	CurvatureSettings->WatchProperty(CurvatureSettings->Clamping, [this](EBakedCurvatureClampMode) { OpState = EBakeOpState::Evaluate; });
+	CurvatureSettings->WatchProperty(CurvatureSettings->BlurRadius, [this](float) { OpState = EBakeOpState::Evaluate; });
+	CurvatureSettings->WatchProperty(CurvatureSettings->bGaussianBlur, [this](float) { OpState = EBakeOpState::Evaluate; });
+
+	TextureSettings = NewObject<UBakedTexture2DImageProperties>(this);
+	TextureSettings->RestoreProperties(this);
+	AddToolPropertySource(TextureSettings);
+	SetToolPropertySourceEnabled(TextureSettings, false);
+	TextureSettings->WatchProperty(TextureSettings->UVLayer, [this](float) { OpState = EBakeOpState::Evaluate; });
+	TextureSettings->WatchProperty(TextureSettings->SourceTexture, [this](UTexture2D*) { OpState = EBakeOpState::Evaluate; });
+
+	MultiTextureSettings = NewObject<UBakedMultiTexture2DImageProperties>(this);
+	MultiTextureSettings->RestoreProperties(this);
+	AddToolPropertySource(MultiTextureSettings);
+	SetToolPropertySourceEnabled(MultiTextureSettings, false);
+
+	UpdateOnModeChange();
+
+	bDetailMeshValid = false;
+
+	SetToolDisplayName(LOCTEXT("ToolName", "Bake Vertex Colors"));
+	GetToolManager()->DisplayMessage(
+		LOCTEXT("OnStartTool",
+		        "Bake Vertex Colors. Select Bake Mesh (LowPoly) first, then (optionally) Detail Mesh second."),
+		EToolMessageLevel::UserNotification);
+}
+
+void UBakeMeshAttributeVertexTool::Shutdown(EToolShutdownType ShutdownType)
+{
+	Settings->SaveProperties(this);
+	ColorSettings->SaveProperties(this);
+	PerChannelSettings->SaveProperties(this);
+	OcclusionSettings->SaveProperties(this);
+	CurvatureSettings->SaveProperties(this);
+	TextureSettings->SaveProperties(this);
+	MultiTextureSettings->SaveProperties(this);
+
+	UE::ToolTarget::ShowSourceObject(Targets[0]);
+
+	if (Compute)
+	{
+		Compute->Shutdown();
+	}
+
+	if (PreviewMesh)
+	{
+		if (ShutdownType == EToolShutdownType::Accept)
+		{
+			GetToolManager()->BeginUndoTransaction(LOCTEXT("BakeMeshAttributeVertexToolTransactionName",
+			                                               "Bake Mesh Attribute Vertex"));
+			UE::ToolTarget::CommitDynamicMeshColorUpdate(Targets[0], PreviewMesh->GetMesh());
+			GetToolManager()->EndUndoTransaction();
+		}
+
+		PreviewMesh->SetVisible(false);
+		PreviewMesh->Disconnect();
+		PreviewMesh = nullptr;
+	}
+}
+
+void UBakeMeshAttributeVertexTool::OnTick(float DeltaTime)
+{
+	if (Compute)
+	{
+		Compute->Tick(DeltaTime);
+
+		const float ElapsedComputeTime = Compute->GetElapsedComputeTime();
+		if (!CanAccept() && ElapsedComputeTime > SecondsBeforeWorkingMaterial)
+		{
+			PreviewMesh->SetOverrideRenderMaterial(WorkingPreviewMaterial);
+		}
+	}
+}
+
+void UBakeMeshAttributeVertexTool::Render(IToolsContextRenderAPI* RenderAPI)
+{
+	UpdateResult();
+}
+
+bool UBakeMeshAttributeVertexTool::CanAccept() const
+{
+	const bool bValidOp = (OpState & EBakeOpState::Invalid) != EBakeOpState::Invalid;
+	return Compute ? bValidOp && Compute->HaveValidResult() : false;
+}
+
+TUniquePtr<UE::Geometry::TGenericDataOperator<FMeshVertexBaker>> UBakeMeshAttributeVertexTool::MakeNewOperator()
+{
+	TUniquePtr<FMeshVertexBakerOp> Op = MakeUnique<FMeshVertexBakerOp>();
+	Op->DetailMesh = DetailMesh;
+	Op->DetailSpatial = DetailSpatial;
+	Op->BaseMesh = &BaseMesh;
+	Op->BaseMeshTangents = BaseMeshTangents;
+	Op->BakeSettings = CachedBakeSettings;
+	Op->ColorSettings = CachedColorSettings;
+	Op->ChannelSettings = CachedChannelSettings;
+	Op->OcclusionSettings = CachedOcclusionMapSettings;
+	Op->CurvatureSettings = CachedCurvatureMapSettings;
+	Op->TextureSettings = CachedTexture2DImageSettings;
+
+	// Texture2DImage & MultiTexture settings
+	Op->TextureImage = CachedTextureImage;
+	Op->MaterialToTextureImageMap = CachedMultiTextures;
+	Op->UVOverlay = DetailMesh->Attributes()->GetUVLayer(CachedTexture2DImageSettings.UVLayer);
+	return Op;
+}
+
+void UBakeMeshAttributeVertexTool::SetWorld(UWorld* World)
+{
+	TargetWorld = World;
+}
+
+void UBakeMeshAttributeVertexTool::UpdateDetailMesh()
+{
+	IPrimitiveComponentBackedTarget* TargetComponent = TargetComponentInterface(0);
+	IPrimitiveComponentBackedTarget* DetailComponent = TargetComponentInterface(bIsBakeToSelf ? 0 : 1);
+	IMeshDescriptionProvider* DetailMeshProvider = TargetMeshProviderInterface(bIsBakeToSelf ? 0 : 1);
+
+	DetailMesh = MakeShared<FDynamicMesh3, ESPMode::ThreadSafe>();
+	FMeshDescriptionToDynamicMesh Converter;
+	Converter.Convert(DetailMeshProvider->GetMeshDescription(), *DetailMesh);
+	if (Settings->bUseWorldSpace && bIsBakeToSelf == false)
+	{
+		const UE::Geometry::FTransform3d DetailToWorld(DetailComponent->GetWorldTransform());
+		MeshTransforms::ApplyTransform(*DetailMesh, DetailToWorld);
+		const UE::Geometry::FTransform3d WorldToBase(TargetComponent->GetWorldTransform());
+		MeshTransforms::ApplyTransform(*DetailMesh, WorldToBase.Inverse());
+	}
+
+	DetailSpatial = MakeShared<FDynamicMeshAABBTree3, ESPMode::ThreadSafe>();
+	DetailSpatial->SetMesh(DetailMesh.Get(), true);
+
+	OpState = EBakeOpState::Evaluate;
+	DetailMeshTimestamp++;
+}
+
+void UBakeMeshAttributeVertexTool::UpdateOnModeChange()
+{
+	const bool bIsColorMode = Settings->VertexMode == EBakeVertexMode::Color;
+	SetToolPropertySourceEnabled(ColorSettings, bIsColorMode);
+	SetToolPropertySourceEnabled(PerChannelSettings, !bIsColorMode);
+	SetToolPropertySourceEnabled(OcclusionSettings, false);
+	SetToolPropertySourceEnabled(CurvatureSettings, false);
+	SetToolPropertySourceEnabled(TextureSettings, false);
+	SetToolPropertySourceEnabled(MultiTextureSettings, false);
+
+	if (Settings->VertexMode == EBakeVertexMode::Color)
+	{
+		switch (ColorSettings->BakeType)
+		{
+			case EBakeVertexTypeColor::AmbientOcclusion:
+			case EBakeVertexTypeColor::BentNormal:
+				SetToolPropertySourceEnabled(OcclusionSettings, true);
+				break;
+			case EBakeVertexTypeColor::Curvature:
+				SetToolPropertySourceEnabled(CurvatureSettings, true);
+				break;
+			case EBakeVertexTypeColor::Texture2DImage:
+				SetToolPropertySourceEnabled(TextureSettings, true);
+				break;
+			case EBakeVertexTypeColor::MultiTexture:
+				SetToolPropertySourceEnabled(MultiTextureSettings, true);
+				break;
+			default:
+				// No property sets to show.
+				break;
+		}
+	}
+	else // Settings->VertexMode == EBakeVertexMode::PerChannel
+	{
+		EBakeVertexTypeChannel PerChannelTypes[4] = {
+			PerChannelSettings->BakeTypeR,
+			PerChannelSettings->BakeTypeG,
+			PerChannelSettings->BakeTypeB,
+			PerChannelSettings->BakeTypeA
+		};
+		for(int Idx = 0; Idx < 4; ++Idx)
+		{
+			switch(PerChannelTypes[Idx])
+			{
+				case EBakeVertexTypeChannel::AmbientOcclusion:
+					SetToolPropertySourceEnabled(OcclusionSettings, true);
+					break;
+				case EBakeVertexTypeChannel::Curvature:
+					SetToolPropertySourceEnabled(CurvatureSettings, true);
+					break;
+				case EBakeVertexTypeChannel::None:
+				default:
+					break;
+			}
+		}
+	}
+}
+
+void UBakeMeshAttributeVertexTool::UpdateVisualization()
+{
+	if (Settings->VertexChannelPreview == EBakeVertexChannel::A)
+	{
+		PreviewMesh->SetOverrideRenderMaterial(PreviewAlphaMaterial);
+	}
+	else
+	{
+		FLinearColor Mask(FLinearColor::Black);
+		switch(Settings->VertexChannelPreview)
+		{
+		case EBakeVertexChannel::R:
+			Mask.R = 1.0f;
+			break;
+		case EBakeVertexChannel::G:
+			Mask.G = 1.0f;
+			break;
+		case EBakeVertexChannel::B:
+			Mask.B = 1.0f;
+			break;
+		case EBakeVertexChannel::RGBA:
+		default:
+			Mask = FLinearColor::White;
+			break;
+		}
+		PreviewMaterial->SetVectorParameterValue("VertexColorMask", Mask);
+		PreviewMesh->SetOverrideRenderMaterial(PreviewMaterial);
+	}
+}
+
+void UBakeMeshAttributeVertexTool::UpdateColorTopology()
+{
+	// Update PreviewMesh color topology
+	PreviewMesh->EditMesh([this](FDynamicMesh3& Mesh)
+	{
+		Mesh.EnableAttributes();
+		Mesh.Attributes()->DisablePrimaryColors();
+		Mesh.Attributes()->EnablePrimaryColors();
+
+		FDynamicMeshNormalOverlay* NormalOverlay = Mesh.Attributes()->PrimaryNormals();
+		FDynamicMeshUVOverlay* UVOverlay = Mesh.Attributes()->PrimaryUV();
+		Mesh.Attributes()->PrimaryColors()->CreateFromPredicate(
+			[&](int ParentVID, int TriIDA, int TriIDB) -> bool
+			{
+				auto OverlayCanShare = [&] (auto Overlay) -> bool
+				{
+					return Overlay ? Overlay->AreTrianglesConnected(TriIDA, TriIDB) : true;
+				};
+				
+				bool bCanShare = true;
+				if (Settings->bSplitAtNormalSeams)
+				{
+					bCanShare = bCanShare && OverlayCanShare(NormalOverlay);
+				}
+				if (Settings->bSplitAtUVSeams)
+				{
+					bCanShare = bCanShare && OverlayCanShare(UVOverlay);
+				}
+				return bCanShare;
+			}, 0.0f);
+	});
+
+	// Update BaseMesh color topology.
+	BaseMesh.EnableAttributes();
+	BaseMesh.Attributes()->DisablePrimaryColors();
+	BaseMesh.Attributes()->EnablePrimaryColors();
+	PreviewMesh->ProcessMesh([this](const FDynamicMesh3& Mesh)
+	{
+		BaseMesh.Attributes()->PrimaryColors()->Copy(*Mesh.Attributes()->PrimaryColors());
+	});
+}
+
+void UBakeMeshAttributeVertexTool::UpdateResult()
+{
+	if (!bDetailMeshValid)
+	{
+		UpdateDetailMesh();
+		bDetailMeshValid = true;
+	}
+
+	if (!bColorTopologyValid)
+	{
+		UpdateColorTopology();
+		bColorTopologyValid = true;
+	}
+
+	if (OpState == EBakeOpState::Complete)
+	{
+		return;
+	}
+
+	// clear warning (ugh)
+	GetToolManager()->DisplayMessage(FText(), EToolMessageLevel::UserWarning);
+
+	FBakeSettings BakeSettings;
+	BakeSettings.VertexMode = Settings->VertexMode;
+	BakeSettings.bUseWorldSpace = Settings->bUseWorldSpace;
+	BakeSettings.Thickness = Settings->Thickness;
+	if (!(BakeSettings == CachedBakeSettings))
+	{
+		CachedBakeSettings = BakeSettings;
+
+		CachedColorSettings = FBakeColorSettings();
+		CachedChannelSettings = FBakeChannelSettings();
+	}
+
+	FBakeColorSettings BakeColorSettings;
+	BakeColorSettings.BakeType = ColorSettings->BakeType;
+	if (!(BakeColorSettings == CachedColorSettings))
+	{
+		CachedColorSettings = BakeColorSettings;
+	}
+
+	FBakeChannelSettings BakeChannelSettings;
+	BakeChannelSettings.BakeType[0] = PerChannelSettings->BakeTypeR;
+	BakeChannelSettings.BakeType[1] = PerChannelSettings->BakeTypeG;
+	BakeChannelSettings.BakeType[2] = PerChannelSettings->BakeTypeB;
+	BakeChannelSettings.BakeType[3] = PerChannelSettings->BakeTypeA;
+	if (!(BakeChannelSettings == CachedChannelSettings))
+	{
+		CachedChannelSettings = BakeChannelSettings;
+	}
+
+	// Clear our invalid bitflag to check again for valid inputs.
+	OpState = EBakeOpState::Evaluate;
+
+	// Validate bake inputs
+	if (CachedBakeSettings.VertexMode == EBakeVertexMode::Color)
+	{
+		switch(CachedColorSettings.BakeType)
+		{
+		case EBakeVertexTypeColor::TangentSpaceNormal:
+			OpState = UpdateResult_Normal();
+			break;
+		case EBakeVertexTypeColor::AmbientOcclusion:
+		case EBakeVertexTypeColor::BentNormal:
+			OpState = UpdateResult_Occlusion();
+			break;
+		case EBakeVertexTypeColor::Curvature:
+			OpState = UpdateResult_Curvature();
+			break;
+		case EBakeVertexTypeColor::NormalImage:
+		case EBakeVertexTypeColor::FaceNormalImage:
+		case EBakeVertexTypeColor::PositionImage:
+		case EBakeVertexTypeColor::MaterialID:
+			OpState = UpdateResult_MeshProperty();
+			break;
+		case EBakeVertexTypeColor::Texture2DImage:
+			OpState = UpdateResult_Texture2DImage();
+			break;
+		case EBakeVertexTypeColor::MultiTexture:
+			OpState = UpdateResult_MultiTexture();
+			break;
+		}
+	}
+	else // CachedBakeSettings.VertexMode == EBakeVertexMode::PerChannel
+	{
+		// The enabled state of these settings are precomputed in UpdateOnModeChange().
+		if (OcclusionSettings->IsPropertySetEnabled())
+		{
+			OpState |= UpdateResult_Occlusion();
+		}
+		if (CurvatureSettings->IsPropertySetEnabled())
+		{
+			OpState |= UpdateResult_Curvature();
+		}
+	}
+
+	// Early exit if op input parameters are invalid.
+	if ((bool)(OpState & EBakeOpState::Invalid))
+	{
+		return;
+	}
+	else // ((bool)(OpState & EBakeOpState::Evaluate))
+	{
+		if (!Compute)
+		{
+			Compute = MakeUnique<TGenericDataBackgroundCompute<FMeshVertexBaker>>();
+			Compute->Setup(this);
+			Compute->OnResultUpdated.AddLambda([this](const TUniquePtr<FMeshVertexBaker>& NewResult)
+			{
+				OnResultUpdated(NewResult);
+			});
+		}
+		Compute->InvalidateResult();
+	}
+}
+
+void UBakeMeshAttributeVertexTool::OnResultUpdated(const TUniquePtr<FMeshVertexBaker>& NewResult)
+{
+	const TImageBuilder<FVector4f>* ImageResult = NewResult->GetBakeResult();
+	if (!ImageResult)
+	{
+		return;
+	}
+
+	// TODO: Review how to handle the implicit sRGB conversion in the StaticMesh build.
+	PreviewMesh->DeferredEditMesh([this, &ImageResult](FDynamicMesh3& Mesh)
+	{
+		const int NumColors = Mesh.Attributes()->PrimaryColors()->ElementCount();
+		check(NumColors == ImageResult->GetDimensions().GetWidth());
+		for (int Idx = 0; Idx < NumColors; ++Idx)
+		{
+			const FVector4f& Pixel = ImageResult->GetPixel(Idx);
+			Mesh.Attributes()->PrimaryColors()->SetElement(Idx, Pixel);
+		}
+	}, false);
+	PreviewMesh->NotifyDeferredEditCompleted(UPreviewMesh::ERenderUpdateMode::FastUpdate, EMeshRenderAttributeFlags::VertexColors, false);
+	UpdateVisualization();
+
+	OpState = EBakeOpState::Complete;
+}
+
+EBakeOpState UBakeMeshAttributeVertexTool::UpdateResult_Normal()
+{
+	// No settings to configure, always valid to evaluate.
+	return EBakeOpState::Evaluate;
+}
+
+EBakeOpState UBakeMeshAttributeVertexTool::UpdateResult_Occlusion()
+{
+	EBakeOpState ResultState = EBakeOpState::Complete;
+
+	FOcclusionMapSettings OcclusionMapSettings;
+	OcclusionMapSettings.MaxDistance = (OcclusionSettings->MaxDistance == 0) ? TNumericLimits<float>::Max() : OcclusionSettings->MaxDistance;
+	OcclusionMapSettings.OcclusionRays = OcclusionSettings->OcclusionRays;
+	OcclusionMapSettings.SpreadAngle = OcclusionSettings->SpreadAngle;
+	OcclusionMapSettings.Distribution = OcclusionSettings->Distribution;
+	OcclusionMapSettings.BlurRadius = (OcclusionSettings->bGaussianBlur) ? OcclusionSettings->BlurRadius : 0.0;
+	OcclusionMapSettings.BiasAngle = OcclusionSettings->BiasAngle;
+	OcclusionMapSettings.NormalSpace = OcclusionSettings->NormalSpace;
+
+	if (!(CachedOcclusionMapSettings == OcclusionMapSettings))
+	{
+		CachedOcclusionMapSettings = OcclusionMapSettings;
+		ResultState = EBakeOpState::Evaluate;
+	}
+	return ResultState;
+}
+
+EBakeOpState UBakeMeshAttributeVertexTool::UpdateResult_Curvature()
+{
+	EBakeOpState ResultState = EBakeOpState::Complete;
+
+	FCurvatureMapSettings CurvatureMapSettings;
+	CurvatureMapSettings.RangeMultiplier = CurvatureSettings->RangeMultiplier;
+	CurvatureMapSettings.MinRangeMultiplier = CurvatureSettings->MinRangeMultiplier;
+	switch (CurvatureSettings->CurvatureType)
+	{
+	default:
+	case EBakedCurvatureTypeMode::MeanAverage:
+		CurvatureMapSettings.CurvatureType = (int32)FMeshCurvatureMapEvaluator::ECurvatureType::Mean;
+		break;
+	case EBakedCurvatureTypeMode::Gaussian:
+		CurvatureMapSettings.CurvatureType = (int32)FMeshCurvatureMapEvaluator::ECurvatureType::Gaussian;
+		break;
+	case EBakedCurvatureTypeMode::Max:
+		CurvatureMapSettings.CurvatureType = (int32)FMeshCurvatureMapEvaluator::ECurvatureType::MaxPrincipal;
+		break;
+	case EBakedCurvatureTypeMode::Min:
+		CurvatureMapSettings.CurvatureType = (int32)FMeshCurvatureMapEvaluator::ECurvatureType::MinPrincipal;
+		break;
+	}
+	switch (CurvatureSettings->ColorMode)
+	{
+	default:
+	case EBakedCurvatureColorMode::Grayscale:
+		CurvatureMapSettings.ColorMode = (int32)FMeshCurvatureMapEvaluator::EColorMode::BlackGrayWhite;
+		break;
+	case EBakedCurvatureColorMode::RedBlue:
+		CurvatureMapSettings.ColorMode = (int32)FMeshCurvatureMapEvaluator::EColorMode::RedBlue;
+		break;
+	case EBakedCurvatureColorMode::RedGreenBlue:
+		CurvatureMapSettings.ColorMode = (int32)FMeshCurvatureMapEvaluator::EColorMode::RedGreenBlue;
+		break;
+	}
+	switch (CurvatureSettings->Clamping)
+	{
+	default:
+	case EBakedCurvatureClampMode::None:
+		CurvatureMapSettings.ClampMode = (int32)FMeshCurvatureMapEvaluator::EClampMode::FullRange;
+		break;
+	case EBakedCurvatureClampMode::Positive:
+		CurvatureMapSettings.ClampMode = (int32)FMeshCurvatureMapEvaluator::EClampMode::Positive;
+		break;
+	case EBakedCurvatureClampMode::Negative:
+		CurvatureMapSettings.ClampMode = (int32)FMeshCurvatureMapEvaluator::EClampMode::Negative;
+		break;
+	}
+	CurvatureMapSettings.BlurRadius = (CurvatureSettings->bGaussianBlur) ? CurvatureSettings->BlurRadius : 0.0;
+
+	if (!(CachedCurvatureMapSettings == CurvatureMapSettings))
+	{
+		CachedCurvatureMapSettings = CurvatureMapSettings;
+		ResultState = EBakeOpState::Evaluate;
+	}
+	return ResultState;
+}
+
+EBakeOpState UBakeMeshAttributeVertexTool::UpdateResult_MeshProperty()
+{
+	// No settings to configure, always valid to evaluate.
+	return EBakeOpState::Evaluate;
+}
+
+EBakeOpState UBakeMeshAttributeVertexTool::UpdateResult_Texture2DImage()
+{
+	EBakeOpState ResultState = EBakeOpState::Complete;
+
+	FTexture2DImageSettings NewSettings;
+	NewSettings.UVLayer = 0;
+
+	const FDynamicMeshUVOverlay* UVOverlay = DetailMesh->Attributes()->GetUVLayer(NewSettings.UVLayer);
+	if (UVOverlay == nullptr)
+	{
+		GetToolManager()->DisplayMessage(LOCTEXT("InvalidUVWarning", "The Source Mesh does not have the selected UV layer"), EToolMessageLevel::UserWarning);
+		return EBakeOpState::Invalid;
+	}
+	
+	if (TextureSettings->SourceTexture == nullptr)
+	{
+		GetToolManager()->DisplayMessage(LOCTEXT("InvalidTextureWarning", "The Source Texture is not valid"), EToolMessageLevel::UserWarning);
+		return EBakeOpState::Invalid;
+	}
+
+	{
+		FTempTextureAccess TextureAccess(TextureSettings->SourceTexture);
+		CachedTextureImage = MakeShared<UE::Geometry::TImageBuilder<FVector4f>, ESPMode::ThreadSafe>();
+		CachedTextureImage->SetDimensions(TextureAccess.GetDimensions());
+		if (!TextureAccess.CopyTo(*CachedTextureImage))
+		{
+			GetToolManager()->DisplayMessage(LOCTEXT("CannotReadTextureWarning", "Cannot read from the source texture"), EToolMessageLevel::UserWarning);
+			return EBakeOpState::Invalid;
+		}
+	}
+
+	if (!(CachedTexture2DImageSettings == NewSettings))
+	{
+		CachedTexture2DImageSettings = NewSettings;
+		ResultState = EBakeOpState::Evaluate;
+	}
+	return ResultState;
+}
+
+EBakeOpState UBakeMeshAttributeVertexTool::UpdateResult_MultiTexture()
+{
+	EBakeOpState ResultState = EBakeOpState::Complete;
+
+	FTexture2DImageSettings NewSettings;
+	NewSettings.UVLayer = MultiTextureSettings->UVLayer;
+
+	const FDynamicMeshUVOverlay* UVOverlay = DetailMesh->Attributes()->GetUVLayer(NewSettings.UVLayer);
+	if (UVOverlay == nullptr)
+	{
+		GetToolManager()->DisplayMessage(LOCTEXT("InvalidUVWarning", "The Source Mesh does not have the selected UV layer"), EToolMessageLevel::UserWarning);
+		return EBakeOpState::Invalid;
+	}
+
+	for (auto& InputTexture : MultiTextureSettings->MaterialIDSourceTextureMap)
+	{
+		if (InputTexture.Value == nullptr)
+		{
+			GetToolManager()->DisplayMessage(LOCTEXT("InvalidTextureWarning", "The Source Texture is not valid"), EToolMessageLevel::UserWarning);
+			return EBakeOpState::Invalid;
+		}
+	}
+
+	CachedMultiTextures.Reset();
+
+	for (auto& InputTexture : MultiTextureSettings->MaterialIDSourceTextureMap)
+	{
+		UTexture2D* Texture = InputTexture.Value;
+		if (!ensure(Texture != nullptr))
+		{
+			GetToolManager()->DisplayMessage(LOCTEXT("InvalidTextureWarning", "The Source Texture is not valid"), EToolMessageLevel::UserWarning);
+			return EBakeOpState::Invalid;
+		}
+
+		int32 MaterialID = InputTexture.Key;
+		FTempTextureAccess TextureAccess(Texture);
+		CachedMultiTextures.Add(MaterialID, MakeShared<UE::Geometry::TImageBuilder<FVector4f>, ESPMode::ThreadSafe>());
+		CachedMultiTextures[MaterialID]->SetDimensions(TextureAccess.GetDimensions());
+
+		if (!TextureAccess.CopyTo(*CachedMultiTextures[MaterialID]))
+		{
+			GetToolManager()->DisplayMessage(LOCTEXT("CannotReadTextureWarning", "Cannot read from the source texture"), EToolMessageLevel::UserWarning);
+			return EBakeOpState::Invalid;
+		}
+	}
+	if (CachedMultiTextures.Num() == 0)
+	{
+		GetToolManager()->DisplayMessage(LOCTEXT("InvalidTextureWarning", "The Source Texture is not valid"), EToolMessageLevel::UserWarning);
+		return EBakeOpState::Invalid;
+	}
+
+	if (!(CachedTexture2DImageSettings == NewSettings))
+	{
+		CachedTexture2DImageSettings = NewSettings;
+		ResultState = EBakeOpState::Evaluate;
+	}
+	return ResultState;
+}
+
+
+
+#undef LOCTEXT_NAMESPACE
diff --git a/Engine/Plugins/Experimental/MeshModelingToolset/Source/MeshModelingTools/Public/BakeMeshAttributeMapsTool.h b/Engine/Plugins/Experimental/MeshModelingToolset/Source/MeshModelingTools/Public/BakeMeshAttributeMapsTool.h
--- a/Engine/Plugins/Experimental/MeshModelingToolset/Source/MeshModelingTools/Public/BakeMeshAttributeMapsTool.h
+++ b/Engine/Plugins/Experimental/MeshModelingToolset/Source/MeshModelingTools/Public/BakeMeshAttributeMapsTool.h
@@ -3,7 +3,6 @@
 #pragma once
 
 #include "CoreMinimal.h"
-#include "Templates/PimplPtr.h"
 #include "MultiSelectionTool.h"
 #include "InteractiveToolBuilder.h"
 #include "DynamicMesh/DynamicMesh3.h"
@@ -13,6 +12,7 @@
 #include "Sampling/MeshMapBaker.h"
 #include "ModelingOperators.h"
 #include "MeshOpPreviewHelpers.h"
+#include "BakeMeshAttributeToolCommon.h"
 #include "BakeMeshAttributeMapsTool.generated.h"
 
 
@@ -65,8 +65,9 @@ enum class EBakeMapType
 	PositionImage          = 1 << 7,
 	MaterialID             = 1 << 8,
 	MultiTexture           = 1 << 9,
+	VertexColorImage       = 1 << 10,
 	Occlusion              = (AmbientOcclusion | BentNormal) UMETA(Hidden),
-	All                    = 0x3FF UMETA(Hidden)
+	All                    = 0x7FF UMETA(Hidden)
 };
 ENUM_CLASS_FLAGS(EBakeMapType);
 
@@ -142,208 +143,6 @@ public:
 };
 
 
-UENUM()
-enum class ENormalMapSpace
-{
-	/** Tangent space */
-	Tangent UMETA(DisplayName = "Tangent space"),
-	/** Object space */
-	Object UMETA(DisplayName = "Object space")
-};
-
-
-UCLASS()
-class MESHMODELINGTOOLS_API UBakedNormalMapToolProperties : public UInteractiveToolPropertySet
-{
-	GENERATED_BODY()
-public:
-
-};
-
-
-UENUM()
-enum class EOcclusionMapDistribution
-{
-	/** Uniform occlusion rays */
-	Uniform UMETA(DisplayName = "Uniform"),
-	/** Cosine weighted occlusion rays */
-	Cosine UMETA(DisplayName = "Cosine")
-};
-
-
-UCLASS()
-class MESHMODELINGTOOLS_API UBakedOcclusionMapToolProperties : public UInteractiveToolPropertySet
-{
-	GENERATED_BODY()
-public:
-	/** Number of occlusion rays */
-	UPROPERTY(EditAnywhere, Category = OcclusionMap, meta = (UIMin = "1", UIMax = "1024", ClampMin = "0", ClampMax = "50000"))
-	int32 OcclusionRays = 16;
-
-	/** Maximum occlusion distance (0 = infinity) */
-	UPROPERTY(EditAnywhere, Category = OcclusionMap, meta = (UIMin = "0.0", UIMax = "1000.0", ClampMin = "0.0", ClampMax = "99999999.0"))
-	float MaxDistance = 0;
-
-	/** Maximum spread angle of occlusion rays. */
-	UPROPERTY(EditAnywhere, Category = OcclusionMap, meta = (UIMin = "0", UIMax = "180.0", ClampMin = "0", ClampMax = "180.0"))
-	float SpreadAngle = 180.0;
-
-	/** Angular distribution of occlusion rays in the spread angle. */
-	UPROPERTY(EditAnywhere, Category = OcclusionMap)
-	EOcclusionMapDistribution Distribution = EOcclusionMapDistribution::Cosine;
-
-	/** Whether or not to apply Gaussian Blur to computed AO Map (recommended) */
-	UPROPERTY(EditAnywhere, Category = "OcclusionMap|Ambient Occlusion")
-	bool bGaussianBlur = true;
-
-	/** Pixel Radius of Gaussian Blur Kernel */
-	UPROPERTY(EditAnywhere, Category = "OcclusionMap|Ambient Occlusion", meta = (UIMin = "0", UIMax = "10.0", ClampMin = "0", ClampMax = "100.0"))
-	float BlurRadius = 2.25;
-
-	/** Contribution of AO rays that are within this angle (degrees) from horizontal are attenuated. This reduces faceting artifacts. */
-	UPROPERTY(EditAnywhere, Category = "OcclusionMap|Ambient Occlusion", meta = (UIMin = "0", UIMax = "45.0", ClampMin = "0", ClampMax = "89.9"))
-	float BiasAngle = 15.0;
-
-	/** Coordinate space of the bent normal map. */
-	UPROPERTY(EditAnywhere, Category = "OcclusionMap|Bent Normal")
-	ENormalMapSpace NormalSpace = ENormalMapSpace::Tangent;
-};
-
-
-
-UCLASS()
-class MESHMODELINGTOOLS_API UBakedOcclusionMapVisualizationProperties : public UInteractiveToolPropertySet
-{
-	GENERATED_BODY()
-public:
-	UPROPERTY(EditAnywhere, Category = Visualization, meta = (UIMin = "0.0", UIMax = "1.0"))
-	float BaseGrayLevel = 1.0;
-
-	/** AO Multiplier in visualization (does not affect output) */
-	UPROPERTY(EditAnywhere, Category = Visualization, meta = (UIMin = "0.0", UIMax = "1.0"))
-	float OcclusionMultiplier = 1.0;
-};
-
-
-
-UENUM()
-enum class EBakedCurvatureTypeMode
-{
-	/** Mean Curvature is the average of the Max and Min Principal curvatures */
-	MeanAverage,
-	/** Max Principal Curvature */
-	Max,
-	/** Min Principal Curvature */
-	Min,
-	/** Gaussian Curvature is the product of the Max and Min Principal curvatures */
-	Gaussian
-};
-
-UENUM()
-enum class EBakedCurvatureColorMode
-{
-	/** Map curvature values to grayscale such that black is negative, grey is zero, and white is positive */
-	Grayscale,
-	/** Map curvature values to red/blue scale such that red is negative, black is zero, and blue is positive */
-	RedBlue,
-	/** Map curvature values to red/green/blue scale such that red is negative, green is zero, and blue is positive */
-	RedGreenBlue
-};
-
-UENUM()
-enum class EBakedCurvatureClampMode
-{
-	/** Include both negative and positive curvatures */
-	None,
-	/** Clamp negative curvatures to zero */
-	Positive,
-	/** Clamp positive curvatures to zero */
-	Negative
-};
-
-
-
-
-UCLASS()
-class MESHMODELINGTOOLS_API UBakedCurvatureMapToolProperties : public UInteractiveToolPropertySet
-{
-	GENERATED_BODY()
-public:
-	/** Type of curvature to compute */
-	UPROPERTY(EditAnywhere, Category = CurvatureMap)
-	EBakedCurvatureTypeMode CurvatureType = EBakedCurvatureTypeMode::MeanAverage;
-
-	/** Color mapping calculated from curvature values */
-	UPROPERTY(EditAnywhere, Category = CurvatureMap)
-	EBakedCurvatureColorMode ColorMode = EBakedCurvatureColorMode::Grayscale;
-
-	/** Scale the maximum curvature value used to compute the mapping to grayscale/color */
-	UPROPERTY(EditAnywhere, Category = CurvatureMap, meta = (UIMin = "0.1", UIMax = "2.0", ClampMin = "0.001", ClampMax = "100.0"))
-	float RangeMultiplier = 1.0;
-
-	/** Scale the minimum curvature value used to compute the mapping to grayscale/color (fraction of maximum) */
-	UPROPERTY(EditAnywhere, AdvancedDisplay, Category = CurvatureMap, meta = (UIMin = "0.0", UIMax = "1.0"))
-	float MinRangeMultiplier = 0.0;
-
-	/** Clamping to apply to curvature values before scaling to color range */
-	UPROPERTY(EditAnywhere, Category = CurvatureMap)
-	EBakedCurvatureClampMode Clamping = EBakedCurvatureClampMode::None;
-
-	/** Whether or not to apply Gaussian Blur to computed Map */
-	UPROPERTY(EditAnywhere, Category = CurvatureMap)
-	bool bGaussianBlur = false;
-
-	/** Pixel Radius of Gaussian Blur Kernel */
-	UPROPERTY(EditAnywhere, Category = CurvatureMap, meta = (UIMin = "0", UIMax = "10.0", ClampMin = "0", ClampMax = "100.0"))
-	float BlurRadius = 2.25;
-};
-
-
-
-UCLASS()
-class MESHMODELINGTOOLS_API UBakedTexture2DImageProperties : public UInteractiveToolPropertySet
-{
-	GENERATED_BODY()
-public:
-	/** The source texture that is to be resampled into a new texture map */
-	UPROPERTY(EditAnywhere, Category = Texture2D, meta = (TransientToolProperty))
-	UTexture2D* SourceTexture;
-
-	/** The UV layer on the source mesh that corresponds to the SourceTexture */
-	UPROPERTY(EditAnywhere, Category = Texture2D)
-	int32 UVLayer = 0;
-};
-
-
-UCLASS()
-class MESHMODELINGTOOLS_API UBakedMultiTexture2DImageProperties : public UInteractiveToolPropertySet
-{
-	GENERATED_BODY()
-public:
-
-	/** For each material ID, the source texture that will be resampled in that material's region*/
-	UPROPERTY(EditAnywhere, Category = MultiTexture, meta = (DisplayName = "Material IDs / Source Textures"))
-	TMap<int32, UTexture2D*> MaterialIDSourceTextureMap;
-
-	/** UV layer to sample from on the input mesh */
-	UPROPERTY(EditAnywhere, Category = MultiTexture)
-	int32 UVLayer = 0;
-
-	/** The set of all source textures from all input materials */
-	UPROPERTY(VisibleAnywhere, BlueprintReadOnly, AdvancedDisplay, Category = MultiTexture, meta = (DisplayName = "Source Textures"))
-	TArray<UTexture2D*> AllSourceTextures;
-
-};
-
-
-enum class EBakeOpState
-{
-	Complete = 0,		// Inputs valid & Result is valid - no-op.
-	Evaluate = 1 << 0,	// Inputs valid & Result is invalid - re-evaluate.
-	Invalid	 = 1 << 1	// Inputs invalid - pause eval.
-};
-ENUM_CLASS_FLAGS(EBakeOpState);
-
 /**
  * Detail Map Baking Tool
  */
@@ -395,15 +194,8 @@ protected:
 
 
 protected:
-	friend class FBakeMapBaseOp;
-	friend class FBakeNormalMapOp;
-	friend class FBakeOcclusionMapOp;
-	friend class FBakeCurvatureMapOp;
-	friend class FBakeTexture2DImageMapOp;
-	friend class FBakeMeshPropertyMapOp;
-	friend class FBakeMultiTextureOp;
 	friend class FMeshMapBakerOp;
-
+	
 	UDynamicMeshComponent* DynamicMeshComponent;
 
 	UPROPERTY()
@@ -467,96 +259,22 @@ protected:
 	using CachedMapIndex = TMap<EBakeMapType, int32>;
 	CachedMapIndex CachedMapIndices;
 
-
-	struct FNormalMapSettings
-	{
-		FImageDimensions Dimensions;
-
-		bool operator==(const FNormalMapSettings& Other) const
-		{
-			return Dimensions == Other.Dimensions;
-		}
-	};
 	FNormalMapSettings CachedNormalMapSettings;
 	EBakeOpState UpdateResult_Normal();
 
-
-	struct FOcclusionMapSettings
-	{
-		FImageDimensions Dimensions;
-		int32 OcclusionRays;
-		float MaxDistance;
-		float SpreadAngle;
-		EOcclusionMapDistribution Distribution;
-		float BlurRadius;
-		float BiasAngle;
-		ENormalMapSpace NormalSpace;
-
-		bool operator==(const FOcclusionMapSettings& Other) const
-		{
-			return Dimensions == Other.Dimensions &&
-				OcclusionRays == Other.OcclusionRays &&
-				MaxDistance == Other.MaxDistance &&
-				SpreadAngle == Other.SpreadAngle &&
-				Distribution == Other.Distribution &&
-				BlurRadius == Other.BlurRadius &&
-				BiasAngle == Other.BiasAngle &&
-				NormalSpace == Other.NormalSpace;
-		}
-	};
 	FOcclusionMapSettings CachedOcclusionMapSettings;
 	EBakeOpState UpdateResult_Occlusion();
 
-
-	struct FCurvatureMapSettings
-	{
-		FImageDimensions Dimensions;
-		int32 RayCount = 1;
-		int32 CurvatureType = 0;
-		float RangeMultiplier = 1.0;
-		float MinRangeMultiplier = 0.0;
-		int32 ColorMode = 0;
-		int32 ClampMode = 0;
-		float MaxDistance = 1.0;
-		float BlurRadius = 1.0;
-
-		bool operator==(const FCurvatureMapSettings& Other) const
-		{
-			return Dimensions == Other.Dimensions && RayCount == Other.RayCount && CurvatureType == Other.CurvatureType && RangeMultiplier == Other.RangeMultiplier && MinRangeMultiplier == Other.MinRangeMultiplier && ColorMode == Other.ColorMode && ClampMode == Other.ClampMode && MaxDistance == Other.MaxDistance && BlurRadius == Other.BlurRadius;
-		}
-	};
 	FCurvatureMapSettings CachedCurvatureMapSettings;
 	EBakeOpState UpdateResult_Curvature();
 
-
-	struct FMeshPropertyMapSettings
-	{
-		FImageDimensions Dimensions;
-
-		bool operator==(const FMeshPropertyMapSettings& Other) const
-		{
-			return Dimensions == Other.Dimensions;
-		}
-	};
 	FMeshPropertyMapSettings CachedMeshPropertyMapSettings;
 	EBakeOpState UpdateResult_MeshProperty();
 
-
-	struct FTexture2DImageSettings
-	{
-		FImageDimensions Dimensions;
-		int32 UVLayer = 0;
-
-		bool operator==(const FTexture2DImageSettings& Other) const
-		{
-			return Dimensions == Other.Dimensions && UVLayer == Other.UVLayer;
-		}
-	};
 	TSharedPtr<UE::Geometry::TImageBuilder<FVector4f>, ESPMode::ThreadSafe> CachedTextureImage;
 	FTexture2DImageSettings CachedTexture2DImageSettings;
 	EBakeOpState UpdateResult_Texture2DImage();
 
-
 	TMap<int32, TSharedPtr<UE::Geometry::TImageBuilder<FVector4f>, ESPMode::ThreadSafe>> CachedMultiTextures;
 	EBakeOpState UpdateResult_MultiTexture();
 
diff --git a/Engine/Plugins/Experimental/MeshModelingToolset/Source/MeshModelingTools/Public/BakeMeshAttributeToolCommon.h b/Engine/Plugins/Experimental/MeshModelingToolset/Source/MeshModelingTools/Public/BakeMeshAttributeToolCommon.h
new file mode 100644
--- /dev/null
+++ b/Engine/Plugins/Experimental/MeshModelingToolset/Source/MeshModelingTools/Public/BakeMeshAttributeToolCommon.h
@@ -0,0 +1,328 @@
+// Copyright Epic Games, Inc. All Rights Reserved.
+
+#pragma once
+
+#include "CoreMinimal.h"
+#include "MultiSelectionTool.h"
+#include "Image/ImageBuilder.h"
+#include "Image/ImageDimensions.h"
+#include "BakeMeshAttributeToolCommon.generated.h"
+
+// Pre-declarations
+using UE::Geometry::FImageDimensions;
+using UE::Geometry::TImageBuilder;
+
+
+/**
+ * Bake tool property sets
+ */
+
+UENUM()
+enum class ENormalMapSpace
+{
+	/** Tangent space */
+	Tangent UMETA(DisplayName = "Tangent space"),
+	/** Object space */
+	Object UMETA(DisplayName = "Object space")
+};
+
+
+UCLASS()
+class MESHMODELINGTOOLS_API UBakedNormalMapToolProperties : public UInteractiveToolPropertySet
+{
+	GENERATED_BODY()
+public:
+
+};
+
+
+UENUM()
+enum class EOcclusionMapDistribution
+{
+	/** Uniform occlusion rays */
+	Uniform UMETA(DisplayName = "Uniform"),
+	/** Cosine weighted occlusion rays */
+	Cosine UMETA(DisplayName = "Cosine")
+};
+
+
+UCLASS()
+class MESHMODELINGTOOLS_API UBakedOcclusionMapToolProperties : public UInteractiveToolPropertySet
+{
+	GENERATED_BODY()
+public:
+	/** Number of occlusion rays */
+	UPROPERTY(EditAnywhere, Category = OcclusionMap, meta = (UIMin = "1", UIMax = "1024", ClampMin = "0", ClampMax = "50000"))
+	int32 OcclusionRays = 16;
+
+	/** Maximum occlusion distance (0 = infinity) */
+	UPROPERTY(EditAnywhere, Category = OcclusionMap, meta = (UIMin = "0.0", UIMax = "1000.0", ClampMin = "0.0", ClampMax = "99999999.0"))
+	float MaxDistance = 0;
+
+	/** Maximum spread angle of occlusion rays. */
+	UPROPERTY(EditAnywhere, Category = OcclusionMap, meta = (UIMin = "0", UIMax = "180.0", ClampMin = "0", ClampMax = "180.0"))
+	float SpreadAngle = 180.0;
+
+	/** Angular distribution of occlusion rays in the spread angle. */
+	UPROPERTY(EditAnywhere, Category = OcclusionMap)
+	EOcclusionMapDistribution Distribution = EOcclusionMapDistribution::Cosine;
+
+	/** Whether or not to apply Gaussian Blur to computed AO Map (recommended) */
+	UPROPERTY(EditAnywhere, Category = "OcclusionMap|Ambient Occlusion")
+	bool bGaussianBlur = true;
+
+	/** Pixel Radius of Gaussian Blur Kernel */
+	UPROPERTY(EditAnywhere, Category = "OcclusionMap|Ambient Occlusion", meta = (UIMin = "0", UIMax = "10.0", ClampMin = "0", ClampMax = "100.0"))
+	float BlurRadius = 2.25;
+
+	/** Contribution of AO rays that are within this angle (degrees) from horizontal are attenuated. This reduces faceting artifacts. */
+	UPROPERTY(EditAnywhere, Category = "OcclusionMap|Ambient Occlusion", meta = (UIMin = "0", UIMax = "45.0", ClampMin = "0", ClampMax = "89.9"))
+	float BiasAngle = 15.0;
+
+	/** Coordinate space of the bent normal map. */
+	UPROPERTY(EditAnywhere, Category = "OcclusionMap|Bent Normal")
+	ENormalMapSpace NormalSpace = ENormalMapSpace::Tangent;
+};
+
+
+UCLASS()
+class MESHMODELINGTOOLS_API UBakedOcclusionMapVisualizationProperties : public UInteractiveToolPropertySet
+{
+	GENERATED_BODY()
+public:
+	UPROPERTY(EditAnywhere, Category = Visualization, meta = (UIMin = "0.0", UIMax = "1.0"))
+	float BaseGrayLevel = 1.0;
+
+	/** AO Multiplier in visualization (does not affect output) */
+	UPROPERTY(EditAnywhere, Category = Visualization, meta = (UIMin = "0.0", UIMax = "1.0"))
+	float OcclusionMultiplier = 1.0;
+};
+
+
+UENUM()
+enum class EBakedCurvatureTypeMode
+{
+	/** Mean Curvature is the average of the Max and Min Principal curvatures */
+	MeanAverage,
+	/** Max Principal Curvature */
+	Max,
+	/** Min Principal Curvature */
+	Min,
+	/** Gaussian Curvature is the product of the Max and Min Principal curvatures */
+	Gaussian
+};
+
+
+UENUM()
+enum class EBakedCurvatureColorMode
+{
+	/** Map curvature values to grayscale such that black is negative, grey is zero, and white is positive */
+	Grayscale,
+	/** Map curvature values to red/blue scale such that red is negative, black is zero, and blue is positive */
+	RedBlue,
+	/** Map curvature values to red/green/blue scale such that red is negative, green is zero, and blue is positive */
+	RedGreenBlue
+};
+
+
+UENUM()
+enum class EBakedCurvatureClampMode
+{
+	/** Include both negative and positive curvatures */
+	None,
+	/** Clamp negative curvatures to zero */
+	Positive,
+	/** Clamp positive curvatures to zero */
+	Negative
+};
+
+
+UCLASS()
+class MESHMODELINGTOOLS_API UBakedCurvatureMapToolProperties : public UInteractiveToolPropertySet
+{
+	GENERATED_BODY()
+public:
+	/** Type of curvature to compute */
+	UPROPERTY(EditAnywhere, Category = CurvatureMap)
+	EBakedCurvatureTypeMode CurvatureType = EBakedCurvatureTypeMode::MeanAverage;
+
+	/** Color mapping calculated from curvature values */
+	UPROPERTY(EditAnywhere, Category = CurvatureMap)
+	EBakedCurvatureColorMode ColorMode = EBakedCurvatureColorMode::Grayscale;
+
+	/** Scale the maximum curvature value used to compute the mapping to grayscale/color */
+	UPROPERTY(EditAnywhere, Category = CurvatureMap, meta = (UIMin = "0.1", UIMax = "2.0", ClampMin = "0.001", ClampMax = "100.0"))
+	float RangeMultiplier = 1.0;
+
+	/** Scale the minimum curvature value used to compute the mapping to grayscale/color (fraction of maximum) */
+	UPROPERTY(EditAnywhere, AdvancedDisplay, Category = CurvatureMap, meta = (UIMin = "0.0", UIMax = "1.0"))
+	float MinRangeMultiplier = 0.0;
+
+	/** Clamping to apply to curvature values before scaling to color range */
+	UPROPERTY(EditAnywhere, Category = CurvatureMap)
+	EBakedCurvatureClampMode Clamping = EBakedCurvatureClampMode::None;
+
+	/** Whether or not to apply Gaussian Blur to computed Map */
+	UPROPERTY(EditAnywhere, Category = CurvatureMap)
+	bool bGaussianBlur = false;
+
+	/** Pixel Radius of Gaussian Blur Kernel */
+	UPROPERTY(EditAnywhere, Category = CurvatureMap, meta = (UIMin = "0", UIMax = "10.0", ClampMin = "0", ClampMax = "100.0"))
+	float BlurRadius = 2.25;
+};
+
+
+UCLASS()
+class MESHMODELINGTOOLS_API UBakedTexture2DImageProperties : public UInteractiveToolPropertySet
+{
+	GENERATED_BODY()
+public:
+	/** The source texture that is to be resampled into a new texture map */
+	UPROPERTY(EditAnywhere, Category = Texture2D, meta = (TransientToolProperty))
+	UTexture2D* SourceTexture;
+
+	/** The UV layer on the source mesh that corresponds to the SourceTexture */
+	UPROPERTY(EditAnywhere, Category = Texture2D)
+	int32 UVLayer = 0;
+};
+
+
+UCLASS()
+class MESHMODELINGTOOLS_API UBakedMultiTexture2DImageProperties : public UInteractiveToolPropertySet
+{
+	GENERATED_BODY()
+public:
+
+	/** For each material ID, the source texture that will be resampled in that material's region*/
+	UPROPERTY(EditAnywhere, Category = MultiTexture, meta = (DisplayName = "Material IDs / Source Textures"))
+	TMap<int32, UTexture2D*> MaterialIDSourceTextureMap;
+
+	/** UV layer to sample from on the input mesh */
+	UPROPERTY(EditAnywhere, Category = MultiTexture)
+	int32 UVLayer = 0;
+
+	/** The set of all source textures from all input materials */
+	UPROPERTY(VisibleAnywhere, BlueprintReadOnly, AdvancedDisplay, Category = MultiTexture, meta = (DisplayName = "Source Textures"))
+	TArray<UTexture2D*> AllSourceTextures;
+
+};
+
+
+/**
+ * Bake tool property settings structs
+ */
+
+struct FNormalMapSettings
+{
+	FImageDimensions Dimensions;
+
+	bool operator==(const FNormalMapSettings& Other) const
+	{
+		return Dimensions == Other.Dimensions;
+	}
+};
+
+struct FOcclusionMapSettings
+{
+	FImageDimensions Dimensions;
+	int32 OcclusionRays;
+	float MaxDistance;
+	float SpreadAngle;
+	EOcclusionMapDistribution Distribution;
+	float BlurRadius;
+	float BiasAngle;
+	ENormalMapSpace NormalSpace;
+
+	bool operator==(const FOcclusionMapSettings& Other) const
+	{
+		return Dimensions == Other.Dimensions &&
+			OcclusionRays == Other.OcclusionRays &&
+			MaxDistance == Other.MaxDistance &&
+			SpreadAngle == Other.SpreadAngle &&
+			Distribution == Other.Distribution &&
+			BlurRadius == Other.BlurRadius &&
+			BiasAngle == Other.BiasAngle &&
+			NormalSpace == Other.NormalSpace;
+	}
+};
+
+struct FCurvatureMapSettings
+{
+	FImageDimensions Dimensions;
+	int32 RayCount = 1;
+	int32 CurvatureType = 0;
+	float RangeMultiplier = 1.0;
+	float MinRangeMultiplier = 0.0;
+	int32 ColorMode = 0;
+	int32 ClampMode = 0;
+	float MaxDistance = 1.0;
+	float BlurRadius = 1.0;
+
+	bool operator==(const FCurvatureMapSettings& Other) const
+	{
+		return Dimensions == Other.Dimensions && RayCount == Other.RayCount && CurvatureType == Other.CurvatureType && RangeMultiplier == Other.RangeMultiplier && MinRangeMultiplier == Other.MinRangeMultiplier && ColorMode == Other.ColorMode && ClampMode == Other.ClampMode && MaxDistance == Other.MaxDistance && BlurRadius == Other.BlurRadius;
+	}
+};
+
+struct FMeshPropertyMapSettings
+{
+	FImageDimensions Dimensions;
+
+	bool operator==(const FMeshPropertyMapSettings& Other) const
+	{
+		return Dimensions == Other.Dimensions;
+	}
+};
+
+struct FTexture2DImageSettings
+{
+	FImageDimensions Dimensions;
+	int32 UVLayer = 0;
+
+	bool operator==(const FTexture2DImageSettings& Other) const
+	{
+		return Dimensions == Other.Dimensions && UVLayer == Other.UVLayer;
+	}
+};
+
+
+/**
+ * Bake compute state
+ */
+enum class EBakeOpState
+{
+	Complete = 0,		// Inputs valid & Result is valid - no-op.
+	Evaluate = 1 << 0,	// Inputs valid & Result is invalid - re-evaluate.
+	Invalid	 = 1 << 1	// Inputs invalid - pause eval.
+};
+ENUM_CLASS_FLAGS(EBakeOpState);
+
+
+/**
+ * Utility UTexture2D accessor
+ */
+class FTempTextureAccess
+{
+public:
+	FTempTextureAccess(UTexture2D* DisplacementMap);
+	FTempTextureAccess(const FTempTextureAccess&) = delete;
+	FTempTextureAccess(FTempTextureAccess&&) = delete;
+	void operator=(const FTempTextureAccess&) = delete;
+	void operator=(FTempTextureAccess&&) = delete;
+
+	~FTempTextureAccess();
+
+	bool HasData() const;
+	const FColor* GetData() const;	
+	FImageDimensions GetDimensions() const;
+
+	bool CopyTo(TImageBuilder<FVector4f>& DestImage) const;
+
+private:
+	UTexture2D* DisplacementMap{ nullptr };
+	TextureCompressionSettings OldCompressionSettings{};
+	TextureMipGenSettings OldMipGenSettings{};
+	bool bOldSRGB{ false };
+	const FColor* FormattedImageData{ nullptr };
+};
+
diff --git a/Engine/Plugins/Experimental/MeshModelingToolset/Source/MeshModelingTools/Public/BakeMeshAttributeVertexTool.h b/Engine/Plugins/Experimental/MeshModelingToolset/Source/MeshModelingTools/Public/BakeMeshAttributeVertexTool.h
new file mode 100644
--- /dev/null
+++ b/Engine/Plugins/Experimental/MeshModelingToolset/Source/MeshModelingTools/Public/BakeMeshAttributeVertexTool.h
@@ -0,0 +1,297 @@
+// Copyright Epic Games, Inc. All Rights Reserved.
+
+#pragma once
+
+#include "CoreMinimal.h"
+#include "MultiSelectionTool.h"
+#include "InteractiveToolBuilder.h"
+#include "DynamicMesh/DynamicMesh3.h"
+#include "DynamicMesh/DynamicMeshAABBTree3.h"
+#include "ModelingOperators.h"
+#include "MeshOpPreviewHelpers.h"
+#include "PreviewMesh.h"
+#include "Sampling/MeshVertexBaker.h"
+#include "BakeMeshAttributeToolCommon.h"
+#include "BakeMeshAttributeVertexTool.generated.h"
+
+
+/**
+ * Tool Builder
+ */
+UCLASS()
+class MESHMODELINGTOOLS_API UBakeMeshAttributeVertexToolBuilder : public UInteractiveToolBuilder
+{
+	GENERATED_BODY()
+
+public:
+	virtual bool CanBuildTool(const FToolBuilderState& SceneState) const override;
+	virtual UInteractiveTool* BuildTool(const FToolBuilderState& SceneState) const override;
+
+protected:
+	virtual const FToolTargetTypeRequirements& GetTargetRequirements() const override;
+};
+
+
+UENUM()
+enum class EBakeVertexMode
+{
+	/* Bake vertex data to color */
+	Color,
+	/* Bake vertex data to individual color channels */
+	PerChannel
+};
+
+
+UENUM()
+enum class EBakeVertexTypeColor
+{
+	TangentSpaceNormal,
+	AmbientOcclusion,
+	BentNormal,
+	Curvature,
+	Texture2DImage,
+	NormalImage,
+	FaceNormalImage,
+	PositionImage,
+	MaterialID,
+	MultiTexture
+};
+
+
+UENUM()
+enum class EBakeVertexTypeChannel
+{
+	None,
+	AmbientOcclusion,
+	Curvature
+};
+
+
+UENUM()
+enum class EBakeVertexChannel
+{
+	R,
+	G,
+	B,
+	A,
+	RGBA
+};
+
+
+UCLASS()
+class MESHMODELINGTOOLS_API UBakeMeshAttributeVertexToolProperties : public UInteractiveToolPropertySet
+{
+	GENERATED_BODY()
+
+public:
+	/** The bake types to generate */
+	UPROPERTY(EditAnywhere, Category = BakeSettings)
+	EBakeVertexMode VertexMode = EBakeVertexMode::Color;
+
+	/** The vertex channel to preview */
+	UPROPERTY(EditAnywhere, Category = BakeSettings, meta = (TransientToolProperty))
+	EBakeVertexChannel VertexChannelPreview = EBakeVertexChannel::RGBA;
+
+	/** Distance to search for the correspondence between the source and target meshes */
+	UPROPERTY(EditAnywhere, Category = BakeSettings, meta = (ClampMin = "0.001"))
+	float Thickness = 3.0;
+
+	/** Compute target mesh to detail mesh correspondence in world space */
+	UPROPERTY(EditAnywhere, Category = BakeSettings)
+	bool bUseWorldSpace = false;
+
+	/** Split vertex colors at normal seams */
+	UPROPERTY(EditAnywhere, Category = BakeSettings)
+	bool bSplitAtNormalSeams = false;
+
+	/** Split vertex colors at UV seams */
+	UPROPERTY(EditAnywhere, Category = BakeSettings, DisplayName = "Split At UV Seams")
+	bool bSplitAtUVSeams = false;
+};
+
+
+UCLASS()
+class MESHMODELINGTOOLS_API UBakeMeshAttributeVertexToolColorProperties : public UInteractiveToolPropertySet
+{
+	GENERATED_BODY()
+
+public:
+	/** The bake type to generate */
+	UPROPERTY(EditAnywhere, Category = ColorBakeSettings)
+	EBakeVertexTypeColor BakeType = EBakeVertexTypeColor::TangentSpaceNormal;
+};
+
+
+UCLASS()
+class MESHMODELINGTOOLS_API UBakeMeshAttributeVertexToolChannelProperties : public UInteractiveToolPropertySet
+{
+	GENERATED_BODY()
+
+public:
+	/** The bake type to generate in the Red channel */
+	UPROPERTY(EditAnywhere, Category = PerChannelBakeSettings)
+	EBakeVertexTypeChannel BakeTypeR = EBakeVertexTypeChannel::None;
+
+	/** The bake type to generate in the Green channel */
+	UPROPERTY(EditAnywhere, Category = PerChannelBakeSettings)
+	EBakeVertexTypeChannel BakeTypeG = EBakeVertexTypeChannel::None;
+
+	/** The bake type to generate in the Blue channel */
+	UPROPERTY(EditAnywhere, Category = PerChannelBakeSettings)
+	EBakeVertexTypeChannel BakeTypeB = EBakeVertexTypeChannel::None;
+
+	/** The bake type to generate in the Alpha channel */
+	UPROPERTY(EditAnywhere, Category = PerChannelBakeSettings)
+	EBakeVertexTypeChannel BakeTypeA = EBakeVertexTypeChannel::None;
+};
+
+
+/**
+ * Vertex Baking Tool
+ */
+UCLASS()
+class MESHMODELINGTOOLS_API UBakeMeshAttributeVertexTool : public UMultiSelectionTool, public UE::Geometry::IGenericDataOperatorFactory<UE::Geometry::FMeshVertexBaker>
+{
+	GENERATED_BODY()
+
+public:
+	UBakeMeshAttributeVertexTool() = default;
+
+	// Begin UInteractiveTool interface
+	virtual void Setup() override;
+	virtual void Shutdown(EToolShutdownType ShutdownType) override;
+
+	virtual void OnTick(float DeltaTime) override;
+	virtual void Render(IToolsContextRenderAPI* RenderAPI) override;
+
+	virtual bool HasCancel() const override { return true; }
+	virtual bool HasAccept() const override { return true; }
+	virtual bool CanAccept() const override;
+	// End UInteractiveTool interface
+
+	// Begin IGenericDataOperatorFactory interface
+	virtual TUniquePtr<UE::Geometry::TGenericDataOperator<UE::Geometry::FMeshVertexBaker>> MakeNewOperator() override;
+	// End IGenericDataOperatorFactory interface
+
+	void SetWorld(UWorld* World);
+
+protected:
+	UPROPERTY()
+	UBakeMeshAttributeVertexToolProperties* Settings;
+
+	UPROPERTY()
+	UBakeMeshAttributeVertexToolColorProperties* ColorSettings;
+
+	UPROPERTY()
+	UBakeMeshAttributeVertexToolChannelProperties* PerChannelSettings;
+
+	UPROPERTY()
+	UBakedOcclusionMapToolProperties* OcclusionSettings;
+
+	UPROPERTY()
+	UBakedCurvatureMapToolProperties* CurvatureSettings;
+
+	UPROPERTY()
+	UBakedTexture2DImageProperties* TextureSettings;
+
+	UPROPERTY()
+	UBakedMultiTexture2DImageProperties* MultiTextureSettings;
+
+protected:
+	friend class FMeshVertexBakerOp;
+	
+	UPROPERTY()
+	UPreviewMesh* PreviewMesh;
+
+	UPROPERTY()
+	UMaterialInstanceDynamic* PreviewMaterial;
+
+	UPROPERTY()
+	UMaterialInstanceDynamic* PreviewAlphaMaterial;
+
+	UPROPERTY()
+	UMaterialInstanceDynamic* WorkingPreviewMaterial;
+	float SecondsBeforeWorkingMaterial = 0.75;
+
+	UWorld* TargetWorld = nullptr;
+
+	TUniquePtr<TGenericDataBackgroundCompute<UE::Geometry::FMeshVertexBaker>> Compute = nullptr;
+	void OnResultUpdated(const TUniquePtr<UE::Geometry::FMeshVertexBaker>& NewResult);
+
+	TSharedPtr<UE::Geometry::TMeshTangents<double>, ESPMode::ThreadSafe> BaseMeshTangents;
+	UE::Geometry::FDynamicMesh3 BaseMesh;
+	UE::Geometry::FDynamicMeshAABBTree3 BaseSpatial;
+
+	TSharedPtr<UE::Geometry::FDynamicMesh3, ESPMode::ThreadSafe> DetailMesh;
+	TSharedPtr<UE::Geometry::FDynamicMeshAABBTree3, ESPMode::ThreadSafe> DetailSpatial;
+	int32 DetailMeshTimestamp = 0;
+	bool bDetailMeshValid = false;
+	void UpdateDetailMesh();
+
+	bool bColorTopologyValid = false;
+	bool bIsBakeToSelf = false;
+	void UpdateOnModeChange();
+	void UpdateVisualization();
+	void UpdateColorTopology();
+	void UpdateResult();
+
+	EBakeOpState OpState = EBakeOpState::Evaluate;
+
+	struct FBakeSettings
+	{
+		EBakeVertexMode VertexMode = EBakeVertexMode::Color;
+		EBakeVertexChannel VertexChannelPreview = EBakeVertexChannel::RGBA;
+		float Thickness = 3.0;
+		bool bUseWorldSpace = false;
+		bool bSplitAtNormalSeams = false;
+		bool bSplitAtUVSeams = false;
+
+		bool operator==(const FBakeSettings& Other) const
+		{
+			return (VertexMode == Other.VertexMode && bUseWorldSpace == Other.bUseWorldSpace &&
+				Thickness == Other.Thickness && bSplitAtNormalSeams == Other.bSplitAtNormalSeams &&
+				bSplitAtUVSeams == Other.bSplitAtUVSeams);
+		}
+	};
+	FBakeSettings CachedBakeSettings;
+
+	struct FBakeColorSettings
+	{
+		EBakeVertexTypeColor BakeType = EBakeVertexTypeColor::TangentSpaceNormal;
+		bool operator==(const FBakeColorSettings& Other) const
+		{
+			return (BakeType == Other.BakeType);
+		}
+	};
+	FBakeColorSettings CachedColorSettings;
+
+	struct FBakeChannelSettings
+	{
+		EBakeVertexTypeChannel BakeType[4] = { EBakeVertexTypeChannel::None, EBakeVertexTypeChannel::None, EBakeVertexTypeChannel::None, EBakeVertexTypeChannel::None };
+		bool operator==(const FBakeChannelSettings& Other) const
+		{
+			return (BakeType[0] == Other.BakeType[0] && BakeType[1] == Other.BakeType[1] && BakeType[2] == Other.BakeType[2] && BakeType[3] == Other.BakeType[3]);
+		}
+	};
+	FBakeChannelSettings CachedChannelSettings;
+
+	//FNormalMapSettings CachedNormalMapSettings;
+	EBakeOpState UpdateResult_Normal();
+
+	FOcclusionMapSettings CachedOcclusionMapSettings;
+	EBakeOpState UpdateResult_Occlusion();
+
+	FCurvatureMapSettings CachedCurvatureMapSettings;
+	EBakeOpState UpdateResult_Curvature();
+
+	//FMeshPropertyMapSettings CachedMeshPropertyMapSettings;
+	EBakeOpState UpdateResult_MeshProperty();
+
+	TSharedPtr<UE::Geometry::TImageBuilder<FVector4f>, ESPMode::ThreadSafe> CachedTextureImage;
+	FTexture2DImageSettings CachedTexture2DImageSettings;
+	EBakeOpState UpdateResult_Texture2DImage();
+
+	TMap<int32, TSharedPtr<UE::Geometry::TImageBuilder<FVector4f>, ESPMode::ThreadSafe>> CachedMultiTextures;
+	EBakeOpState UpdateResult_MultiTexture();
+};
+
diff --git a/Engine/Plugins/Experimental/MeshModelingToolset/Source/ModelingComponents/Private/ModelingToolTargetUtil.cpp b/Engine/Plugins/Experimental/MeshModelingToolset/Source/ModelingComponents/Private/ModelingToolTargetUtil.cpp
--- a/Engine/Plugins/Experimental/MeshModelingToolset/Source/ModelingComponents/Private/ModelingToolTargetUtil.cpp
+++ b/Engine/Plugins/Experimental/MeshModelingToolset/Source/ModelingComponents/Private/ModelingToolTargetUtil.cpp
@@ -322,6 +322,38 @@ UE::ToolTarget::EDynamicMeshUpdateResult UE::ToolTarget::CommitDynamicMeshUVUpda
 }
 
 
+UE::ToolTarget::EDynamicMeshUpdateResult UE::ToolTarget::CommitDynamicMeshColorUpdate(UToolTarget* Target, const UE::Geometry::FDynamicMesh3* UpdatedMesh)
+{
+	IMeshDescriptionCommitter* MeshDescriptionCommitter = Cast<IMeshDescriptionCommitter>(Target);
+	if (!ensure(MeshDescriptionCommitter))
+	{
+		return EDynamicMeshUpdateResult::Failed;
+	}
+
+	EDynamicMeshUpdateResult Result = EDynamicMeshUpdateResult::Failed;
+	MeshDescriptionCommitter->CommitMeshDescription([UpdatedMesh, &Result](const IMeshDescriptionCommitter::FCommitterParams& CommitParams)
+	{
+		FMeshDescription* MeshDescription = CommitParams.MeshDescriptionOut;
+
+		const bool bVerticesOnly = false;
+		const bool bAttributesOnly = true;
+		if (FDynamicMeshToMeshDescription::HaveMatchingElementCounts(UpdatedMesh, MeshDescription, bVerticesOnly, bAttributesOnly))
+		{
+			FDynamicMeshToMeshDescription Converter;
+			Converter.UpdateVertexColors(UpdatedMesh, *MeshDescription);
+			Result = EDynamicMeshUpdateResult::Ok;
+		}
+		else
+		{
+			// must have been duplicate tris in the mesh description; we can't count on 1-to-1 mapping of TriangleIDs.  Just convert 
+			FDynamicMeshToMeshDescription Converter;
+			Converter.Convert(UpdatedMesh, *MeshDescription);
+			Result = EDynamicMeshUpdateResult::Ok_ForcedFullUpdate;
+		}
+	});
+	return Result;
+}
+
 
 bool UE::ToolTarget::ConfigureCreateMeshObjectParams(UToolTarget* SourceTarget, FCreateMeshObjectParams& DerivedParamsOut)
 {
diff --git a/Engine/Plugins/Experimental/MeshModelingToolset/Source/ModelingComponents/Public/ModelingToolTargetUtil.h b/Engine/Plugins/Experimental/MeshModelingToolset/Source/ModelingComponents/Public/ModelingToolTargetUtil.h
--- a/Engine/Plugins/Experimental/MeshModelingToolset/Source/ModelingComponents/Public/ModelingToolTargetUtil.h
+++ b/Engine/Plugins/Experimental/MeshModelingToolset/Source/ModelingComponents/Public/ModelingToolTargetUtil.h
@@ -145,6 +145,13 @@ MODELINGCOMPONENTS_API EDynamicMeshUpdateResult CommitDynamicMeshUpdate(
 MODELINGCOMPONENTS_API EDynamicMeshUpdateResult CommitDynamicMeshUVUpdate(UToolTarget* Target, const UE::Geometry::FDynamicMesh3* UpdatedMesh);
 
 
+/**
+ * Update the Vertex Colors of the ToolTarget's mesh based on the provided UpdatedMesh.
+ * @return EDynamicMeshUpdateResult::Ok on success, or Ok_ForcedFullUpdate if any dependent mesh topology was modified
+ */
+MODELINGCOMPONENTS_API EDynamicMeshUpdateResult CommitDynamicMeshColorUpdate(UToolTarget* Target, const UE::Geometry::FDynamicMesh3* UpdatedMesh);
+
+
 /**
  * FCreateMeshObjectParams::TypeHint is used by the ModelingObjectsCreationAPI to suggest what type of mesh object to create
  * inside various Tools. This should often be derived from the input mesh object type (eg if you plane-cut a Volume, the output
diff --git a/Engine/Plugins/Experimental/ModelingToolsEditorMode/Source/ModelingToolsEditorMode/Private/ModelingToolsEditorMode.cpp b/Engine/Plugins/Experimental/ModelingToolsEditorMode/Source/ModelingToolsEditorMode/Private/ModelingToolsEditorMode.cpp
--- a/Engine/Plugins/Experimental/ModelingToolsEditorMode/Source/ModelingToolsEditorMode/Private/ModelingToolsEditorMode.cpp
+++ b/Engine/Plugins/Experimental/ModelingToolsEditorMode/Source/ModelingToolsEditorMode/Private/ModelingToolsEditorMode.cpp
@@ -76,6 +76,7 @@
 #include "AlignObjectsTool.h"
 #include "EditUVIslandsTool.h"
 #include "BakeMeshAttributeMapsTool.h"
+#include "BakeMeshAttributeVertexTool.h"
 #include "MeshAttributePaintTool.h"
 #include "ParameterizeMeshTool.h"
 #include "RecomputeUVsTool.h"
@@ -552,6 +553,9 @@ void UModelingToolsEditorMode::Enter()
 	auto BakeMeshAttributeMapsToolBuilder = NewObject<UBakeMeshAttributeMapsToolBuilder>();
 	RegisterTool(ToolManagerCommands.BeginBakeMeshAttributeMapsTool, TEXT("BeginBakeMeshAttributeMapsTool"), BakeMeshAttributeMapsToolBuilder);
 
+	auto BakeMeshAttributeVertexToolBuilder = NewObject<UBakeMeshAttributeVertexToolBuilder>();
+	RegisterTool(ToolManagerCommands.BeginBakeMeshAttributeVertexTool, TEXT("BeginBakeMeshAttributeVertexTool"), BakeMeshAttributeVertexToolBuilder);
+
 	// analysis tools
 
 	RegisterTool(ToolManagerCommands.BeginMeshInspectorTool, TEXT("BeginMeshInspectorTool"), NewObject<UMeshInspectorToolBuilder>());
diff --git a/Engine/Plugins/Experimental/ModelingToolsEditorMode/Source/ModelingToolsEditorMode/Private/ModelingToolsEditorModeStyle.cpp b/Engine/Plugins/Experimental/ModelingToolsEditorMode/Source/ModelingToolsEditorMode/Private/ModelingToolsEditorModeStyle.cpp
--- a/Engine/Plugins/Experimental/ModelingToolsEditorMode/Source/ModelingToolsEditorMode/Private/ModelingToolsEditorModeStyle.cpp
+++ b/Engine/Plugins/Experimental/ModelingToolsEditorMode/Source/ModelingToolsEditorMode/Private/ModelingToolsEditorModeStyle.cpp
@@ -166,6 +166,8 @@ void FModelingToolsEditorModeStyle::Initialize()
 		StyleSet->Set("ModelingToolsManagerCommands.BeginUVSeamEditTool.Small", 		new IMAGE_PLUGIN_BRUSH("Icons/ModelingUVSeamEdit_x40",			Icon20x20));
 		StyleSet->Set("ModelingToolsManagerCommands.BeginBakeMeshAttributeMapsTool", 			new IMAGE_PLUGIN_BRUSH("Icons/ModelingBakeMaps_x40",			Icon20x20));
 		StyleSet->Set("ModelingToolsManagerCommands.BeginBakeMeshAttributeMapsTool.Small", 		new IMAGE_PLUGIN_BRUSH("Icons/ModelingBakeMaps_x40",			Icon20x20));
+		StyleSet->Set("ModelingToolsManagerCommands.BeginBakeMeshAttributeVertexTool", new IMAGE_PLUGIN_BRUSH("Icons/ModelingBakeMaps_x40", Icon20x20));
+		StyleSet->Set("ModelingToolsManagerCommands.BeginBakeMeshAttributeVertexTool.Small", new IMAGE_PLUGIN_BRUSH("Icons/ModelingBakeMaps_x40", Icon20x20));
 		//StyleSet->Set("ModelingToolsManagerCommands.BeginRemoveOccludedTrianglesTool", 				new IMAGE_PLUGIN_BRUSH("Icons/Jacket_40x",			Icon20x20));
 		//StyleSet->Set("ModelingToolsManagerCommands.BeginRemoveOccludedTrianglesTool.Small", 		new IMAGE_PLUGIN_BRUSH("Icons/Jacket_40x",			Icon20x20));
 		StyleSet->Set("ModelingToolsManagerCommands.BeginHoleFillTool", 				new IMAGE_PLUGIN_BRUSH("Icons/ModelingHoleFill_x40",			Icon20x20));
diff --git a/Engine/Plugins/Experimental/ModelingToolsEditorMode/Source/ModelingToolsEditorMode/Private/ModelingToolsEditorModeToolkit.cpp b/Engine/Plugins/Experimental/ModelingToolsEditorMode/Source/ModelingToolsEditorMode/Private/ModelingToolsEditorModeToolkit.cpp
--- a/Engine/Plugins/Experimental/ModelingToolsEditorMode/Source/ModelingToolsEditorMode/Private/ModelingToolsEditorModeToolkit.cpp
+++ b/Engine/Plugins/Experimental/ModelingToolsEditorMode/Source/ModelingToolsEditorMode/Private/ModelingToolsEditorModeToolkit.cpp
@@ -624,6 +624,7 @@ void FModelingToolsEditorModeToolkit::BuildToolPalette_Experimental(FName Palett
 		ToolbarBuilder.AddToolBarButton(Commands.BeginMeshAttributePaintTool);
 		ToolbarBuilder.AddToolBarButton(Commands.BeginEditMeshMaterialsTool);
 		ToolbarBuilder.AddToolBarButton(Commands.BeginBakeMeshAttributeMapsTool);
+		ToolbarBuilder.AddToolBarButton(Commands.BeginBakeMeshAttributeVertexTool);
 	}
 	else if (PaletteIndex == UVTabName)
 	{
diff --git a/Engine/Plugins/Experimental/ModelingToolsEditorMode/Source/ModelingToolsEditorMode/Private/ModelingToolsManagerActions.cpp b/Engine/Plugins/Experimental/ModelingToolsEditorMode/Source/ModelingToolsEditorMode/Private/ModelingToolsManagerActions.cpp
--- a/Engine/Plugins/Experimental/ModelingToolsEditorMode/Source/ModelingToolsEditorMode/Private/ModelingToolsManagerActions.cpp
+++ b/Engine/Plugins/Experimental/ModelingToolsEditorMode/Source/ModelingToolsEditorMode/Private/ModelingToolsManagerActions.cpp
@@ -102,6 +102,7 @@ void FModelingToolsManagerCommands::RegisterCommands()
 	UI_COMMAND(BeginMeshAttributePaintTool, "MapPnt", "Paint Attribute Maps", EUserInterfaceActionType::ToggleButton, FInputChord());
 	UI_COMMAND(BeginAttributeEditorTool, "AttrEd", "Edit/Configure Mesh Attributes", EUserInterfaceActionType::ToggleButton, FInputChord());
 	UI_COMMAND(BeginBakeMeshAttributeMapsTool, "BakeTx", "Bake Image Maps for Target Mesh (optionally from second Source Mesh)", EUserInterfaceActionType::ToggleButton, FInputChord());
+	UI_COMMAND(BeginBakeMeshAttributeVertexTool, "BakeVtx", "Bake Vertex Colors for Target Mesh (optionally from a second Source Mesh)", EUserInterfaceActionType::ToggleButton, FInputChord());
 	UI_COMMAND(BeginUVSeamEditTool, "SeamEd", "Add UV Seams to Mesh", EUserInterfaceActionType::ToggleButton, FInputChord());
 
 	UI_COMMAND(BeginSkinWeightsPaintTool, "SkinWts", "Start the Paint Skin Weights Tool", EUserInterfaceActionType::Button, FInputChord());
diff --git a/Engine/Plugins/Experimental/ModelingToolsEditorMode/Source/ModelingToolsEditorMode/Public/ModelingToolsManagerActions.h b/Engine/Plugins/Experimental/ModelingToolsEditorMode/Source/ModelingToolsEditorMode/Public/ModelingToolsManagerActions.h
--- a/Engine/Plugins/Experimental/ModelingToolsEditorMode/Source/ModelingToolsEditorMode/Public/ModelingToolsManagerActions.h
+++ b/Engine/Plugins/Experimental/ModelingToolsEditorMode/Source/ModelingToolsEditorMode/Public/ModelingToolsManagerActions.h
@@ -97,6 +97,7 @@ public:
 	TSharedPtr<FUICommandInfo> BeginMeshAttributePaintTool;
 	TSharedPtr<FUICommandInfo> BeginAttributeEditorTool;
 	TSharedPtr<FUICommandInfo> BeginBakeMeshAttributeMapsTool;
+	TSharedPtr<FUICommandInfo> BeginBakeMeshAttributeVertexTool;
 	TSharedPtr<FUICommandInfo> BeginUVSeamEditTool;
 
 	TSharedPtr<FUICommandInfo> BeginSkinWeightsPaintTool;
