ApproxActors: restore parallel meshdescription loading
#rnx

diff --git a/Engine/Plugins/Experimental/MeshModelingToolset/Source/ModelingComponents/Private/Scene/MeshSceneAdapter.cpp b/Engine/Plugins/Experimental/MeshModelingToolset/Source/ModelingComponents/Private/Scene/MeshSceneAdapter.cpp
--- a/Engine/Plugins/Experimental/MeshModelingToolset/Source/ModelingComponents/Private/Scene/MeshSceneAdapter.cpp
+++ b/Engine/Plugins/Experimental/MeshModelingToolset/Source/ModelingComponents/Private/Scene/MeshSceneAdapter.cpp
@@ -463,10 +463,11 @@ public:
 #if !WITH_EDITOR
 		SourceMesh = nullptr;
 #endif
-		if (SourceMesh)
-		{
-			check(Adapter);
+		SourceMesh = StaticMesh->GetMeshDescription(LODIndex);
+		Adapter = MakeUnique<FMeshDescriptionTriangleMeshSurfaceAdapter>(SourceMesh, StaticMesh);
 
+		if (SourceMesh && Adapter)
+		{
 			BuildScale = FVector3d::One();
 #if WITH_EDITOR
 			// respect BuildScale build setting
@@ -571,9 +572,6 @@ static TUniquePtr<IMeshSpatialWrapper> SpatialWrapperFactory( const FMeshTypeCon
 		const int32 LODIndex = 0;
 		SMWrapper->SourceContainer = MeshContainer;
 		SMWrapper->StaticMesh = MeshContainer.GetStaticMesh();
-		SMWrapper->SourceMesh = SMWrapper->StaticMesh->GetMeshDescription(LODIndex);
-		SMWrapper->Adapter = MakeUnique<FMeshDescriptionTriangleMeshSurfaceAdapter>(SMWrapper->SourceMesh, SMWrapper->StaticMesh);
-
 		if (ensure(SMWrapper->StaticMesh != nullptr))
 		{
 			return SMWrapper;
@@ -884,8 +882,10 @@ void FMeshSceneAdapter::Build_FullDecompose(const FMeshSceneAdapterBuildOptions&
 	std::atomic<int32> DecomposedMeshesCount = 0;
 	std::atomic<int32> SourceInstancesCount = 0;
 	std::atomic<int32> NewInstancesCount = 0;
+	std::atomic<int32> NewUniqueMeshesCount = 0;
 	std::atomic<int32> SkippedDecompositionCount = 0;
 	std::atomic<int32> SingleTriangleMeshes = 0;
+	std::atomic<int32> NumFinalUniqueTris = 0;
 	int32 AddedUniqueTrisCount = 0;
 	int32 InstancedTrisCount = 0;
 
@@ -1152,6 +1152,9 @@ void FMeshSceneAdapter::Build_FullDecompose(const FMeshSceneAdapterBuildOptions&
 			{
 				MeshInstance->MeshSpatial = WrapperInfo->SpatialWrapper.Get();
 			}
+			InternalListsLock.Lock();
+			InstancedTrisCount += LocalSpaceParts.TriangleCount() * ParentTransforms.Num();
+			InternalListsLock.Unlock();
 		}
 		else
 		{
@@ -1179,6 +1182,7 @@ void FMeshSceneAdapter::Build_FullDecompose(const FMeshSceneAdapterBuildOptions&
 		// pulled these back out of the Actor so it doesn't really matter.
 		for (FInstancedSubmesh& Submesh : NewSubmeshes)
 		{
+			NewUniqueMeshesCount++;
 			TUniquePtr<FActorAdapter> NewActor = MakeUnique<FActorAdapter>();
 			NewActor->SourceActor = nullptr;		// not a "real" actor
 
@@ -1235,6 +1239,9 @@ void FMeshSceneAdapter::Build_FullDecompose(const FMeshSceneAdapterBuildOptions&
 			SceneActors.Add(MoveTemp(NewActor));
 			InternalListsLock.Unlock();
 		}
+
+		NumFinalUniqueTris += ItemStats.TotalNumUniqueTris;
+
 	}, ParallelFlags);		// end outer ParallelFor over ToBuild meshes
 
 	check(ToBuild.Num() == 0);
@@ -1267,9 +1274,10 @@ void FMeshSceneAdapter::Build_FullDecompose(const FMeshSceneAdapterBuildOptions&
 
 	if (BuildOptions.bPrintDebugMessages)
 	{
-		UE_LOG(LogTemp, Warning, TEXT("[FMeshSceneAdapter] decomposed %d source meshes used in %d instances (of %d total source meshes with %ld unique triangles), into %d new instances containing %ld unique triangles (%ld total instanced). Skipped %d decompositions. Skipped %d tiny components (%d instances, %d total triangles). %d 1-triangle meshes."), 
+		UE_LOG(LogTemp, Warning, TEXT("[FMeshSceneAdapter] decomposed %d source meshes used in %d instances (of %d total source meshes with %ld unique triangles), into %d new part meshes used in %d new instances containing %ld unique triangles. Scene has %ld total unique and %ld total instanced. Skipped %d decompositions. Skipped %d tiny components (%d instances, %d total triangles). %d 1-triangle meshes."), 
 			   DecomposedSourceMeshCount.load(), SourceInstancesCount.load(), NumInitialSources, NumSourceUniqueTris, 
-			   NewInstancesCount.load(), AddedUniqueTrisCount, InstancedTrisCount, 
+			   NewUniqueMeshesCount.load(), NewInstancesCount.load(), AddedUniqueTrisCount, 
+			   NumFinalUniqueTris.load(), InstancedTrisCount,
 			   SkippedDecompositionCount.load(), 
 			   NumTinyComponents.load(), NumTinyInstances.load(), TinyInstanceTotalTriangles.load(),
 			   SingleTriangleMeshes.load())
