ModelingComponents: fix editor-only call
#rnx

diff --git a/Engine/Plugins/Experimental/MeshModelingToolset/Source/ModelingComponents/Private/Scene/MeshSceneAdapter.cpp b/Engine/Plugins/Experimental/MeshModelingToolset/Source/ModelingComponents/Private/Scene/MeshSceneAdapter.cpp
--- a/Engine/Plugins/Experimental/MeshModelingToolset/Source/ModelingComponents/Private/Scene/MeshSceneAdapter.cpp
+++ b/Engine/Plugins/Experimental/MeshModelingToolset/Source/ModelingComponents/Private/Scene/MeshSceneAdapter.cpp
@@ -462,8 +462,9 @@ public:
 
 #if !WITH_EDITOR
 		SourceMesh = nullptr;
-#endif
+#else
 		SourceMesh = StaticMesh->GetMeshDescription(LODIndex);
+#endif
 		Adapter = MakeUnique<FMeshDescriptionTriangleMeshSurfaceAdapter>(SourceMesh, StaticMesh);
 
 		if (SourceMesh && Adapter)
