ModelingTools: supporting fallback to no-op in FSubdividePoly if OpenSubdiv is not available, to allow removing Win64-whitelist on ModelingComponentsEditorOnly module
#rnx

diff --git a/Engine/Plugins/Experimental/MeshModelingToolset/Source/ModelingComponentsEditorOnly/ModelingComponentsEditorOnly.Build.cs b/Engine/Plugins/Experimental/MeshModelingToolset/Source/ModelingComponentsEditorOnly/ModelingComponentsEditorOnly.Build.cs
--- a/Engine/Plugins/Experimental/MeshModelingToolset/Source/ModelingComponentsEditorOnly/ModelingComponentsEditorOnly.Build.cs
+++ b/Engine/Plugins/Experimental/MeshModelingToolset/Source/ModelingComponentsEditorOnly/ModelingComponentsEditorOnly.Build.cs
@@ -50,16 +50,21 @@ public class ModelingComponentsEditorOnly : ModuleRules
                 "RenderCore",
 				"PhysicsCore",
                 "RHI",
-
-				"OpenSubdiv",
-
 				"AssetTools",
 				"UnrealEd"			// required for asset factories
-
 			}
 			);
-		
-		
+
+		if (Target.Platform == UnrealTargetPlatform.Win64)
+		{
+			PrivateDependencyModuleNames.AddRange(
+				new string[]
+				{
+					"OpenSubdiv", // currently Win64-only
+				}
+				);
+		}
+
 		DynamicallyLoadedModuleNames.AddRange(
 			new string[]
 			{
diff --git a/Engine/Plugins/Experimental/MeshModelingToolset/Source/ModelingComponentsEditorOnly/Private/Operations/SubdividePoly.cpp b/Engine/Plugins/Experimental/MeshModelingToolset/Source/ModelingComponentsEditorOnly/Private/Operations/SubdividePoly.cpp
--- a/Engine/Plugins/Experimental/MeshModelingToolset/Source/ModelingComponentsEditorOnly/Private/Operations/SubdividePoly.cpp
+++ b/Engine/Plugins/Experimental/MeshModelingToolset/Source/ModelingComponentsEditorOnly/Private/Operations/SubdividePoly.cpp
@@ -4,6 +4,16 @@
 #include "GroupTopology.h"
 #include "DynamicMesh/MeshNormals.h"
 
+// OpenSubdiv currently only available on Windows. On other platforms we will make this a no-op
+#if PLATFORM_WINDOWS
+#define HAVE_OPENSUBDIV 1
+#else
+#define HAVE_OPENSUBDIV 0
+#endif
+
+
+#if HAVE_OPENSUBDIV
+
 // OpenSubdiv needs M_PI defined
 #ifndef M_PI
 #define M_PI PI
@@ -20,6 +30,8 @@
 #undef M_PI
 #endif
 
+#endif
+
 #include "ExplicitUseGeometryMathTypes.h"		// using UE::Geometry::(math types)
 using namespace UE::Geometry;
 
@@ -273,7 +285,9 @@ namespace SubdividePolyLocal
 class FSubdividePoly::RefinerImpl
 {
 public:
+#if HAVE_OPENSUBDIV
 	OpenSubdiv::Far::TopologyRefiner* TopologyRefiner = nullptr;
+#endif
 };
 
 FSubdividePoly::FSubdividePoly(const FGroupTopology& InTopology,
@@ -288,17 +302,20 @@ FSubdividePoly::FSubdividePoly(const FGroupTopology& InTopology,
 
 FSubdividePoly::~FSubdividePoly()
 {
+#if HAVE_OPENSUBDIV
 	if (Refiner && Refiner->TopologyRefiner)
 	{
 		// This was created by TopologyRefinerFactory; looks like we are responsible for cleaning it up.
 		delete Refiner->TopologyRefiner;
 	}
+#endif
 	Refiner = nullptr;
 }
 
 
 bool FSubdividePoly::ComputeTopologySubdivision()
 {
+#if HAVE_OPENSUBDIV
 	if (Level < 1)
 	{
 		return false;
@@ -442,12 +459,15 @@ bool FSubdividePoly::ComputeTopologySubdivision()
 
 	Refiner->TopologyRefiner->RefineUniform(OpenSubdiv::Far::TopologyRefiner::UniformOptions(Level));
 
+#endif		// HAVE_OPENSUBDIV
+
 	return true;
 }
 
 
 bool FSubdividePoly::ComputeSubdividedMesh(FDynamicMesh3& OutMesh)
 {
+#if HAVE_OPENSUBDIV
 	if (Level < 1)
 	{
 		return false;
@@ -656,6 +676,12 @@ bool FSubdividePoly::ComputeSubdividedMesh(FDynamicMesh3& OutMesh)
 			OutMesh.RemoveVertex(Vid, false, false);
 		}
 	}
+#else	// HAVE_OPENSUBDIV
+
+	OutMesh = OriginalMesh;
+
+#endif	// HAVE_OPENSUBDIV
+
 
 	return true;
 }
