Enable soft deformation of FFD lattice using a constrained mesh solver

diff --git a/Engine/Plugins/Experimental/GeometryProcessing/Source/DynamicMesh/Private/Operations/FFDLattice.cpp b/Engine/Plugins/Experimental/GeometryProcessing/Source/DynamicMesh/Private/Operations/FFDLattice.cpp
--- a/Engine/Plugins/Experimental/GeometryProcessing/Source/DynamicMesh/Private/Operations/FFDLattice.cpp
+++ b/Engine/Plugins/Experimental/GeometryProcessing/Source/DynamicMesh/Private/Operations/FFDLattice.cpp
@@ -9,7 +9,7 @@
 using namespace UE::Geometry;
 
 FFFDLattice::FFFDLattice(const FVector3i& InDims, const FDynamicMesh3& Mesh, float Padding) :
-	Dims(InDims)
+	Dimensions(InDims)
 {
 	check(InDims.X > 1 && InDims.Y > 1 && InDims.Z > 1);
 
@@ -31,23 +31,23 @@ FFFDLattice::FFFDLattice(const FVector3i& InDims, const FDynamicMesh3& Mesh, flo
 	InitialBounds.Max = Center + Extents;
 	
 	FVector3d Diag = InitialBounds.Diagonal();
-	CellSize = Diag / FVector3d(Dims - 1);
+	CellSize = Diag / FVector3d(Dimensions - 1);
 
 	ComputeInitialEmbedding(Mesh);
 }
 
 void FFFDLattice::GenerateInitialLatticePositions(TArray<FVector3d>& OutLatticePositions) const
 {
-	int TotalNumLatticePoints = Dims.X * Dims.Y * Dims.Z;
+	int TotalNumLatticePoints = Dimensions.X * Dimensions.Y * Dimensions.Z;
 	OutLatticePositions.SetNum(TotalNumLatticePoints);
 
-	for (int i = 0; i < Dims.X; ++i)
+	for (int i = 0; i < Dimensions.X; ++i)
 	{
 		double X = CellSize.X * i;
-		for (int j = 0; j < Dims.Y; ++j)
+		for (int j = 0; j < Dimensions.Y; ++j)
 		{
 			double Y = CellSize.Y * j;
-			for (int k = 0; k < Dims.Z; ++k)
+			for (int k = 0; k < Dimensions.Z; ++k)
 			{
 				int PointID = ControlPointIndexFromCoordinates(i, j, k);
 				double Z = CellSize.Z * k;
@@ -61,27 +61,27 @@ void FFFDLattice::GenerateInitialLatticePositions(TArray<FVector3d>& OutLatticeP
 
 void FFFDLattice::GenerateLatticeEdges(TArray<FVector2i>& OutLatticeEdges) const
 {
-	OutLatticeEdges.Reset(3 * Dims.X * Dims.Y * Dims.Z);
+	OutLatticeEdges.Reset(3 * Dimensions.X * Dimensions.Y * Dimensions.Z);
 
-	for (int i = 0; i < Dims.X; ++i)
+	for (int i = 0; i < Dimensions.X; ++i)
 	{
-		for (int j = 0; j < Dims.Y; ++j)
+		for (int j = 0; j < Dimensions.Y; ++j)
 		{
-			for (int k = 0; k < Dims.Z; ++k)
+			for (int k = 0; k < Dimensions.Z; ++k)
 			{
 				int PointID = ControlPointIndexFromCoordinates(i, j, k);
 
-				if (i + 1 < Dims.X)
+				if (i + 1 < Dimensions.X)
 				{
 					int IPlusOne = ControlPointIndexFromCoordinates(i + 1, j, k);
 					OutLatticeEdges.Add({ PointID, IPlusOne });
 				}
-				if (j + 1 < Dims.Y)
+				if (j + 1 < Dimensions.Y)
 				{
 					int JPlusOne = ControlPointIndexFromCoordinates(i, j + 1, k);
 					OutLatticeEdges.Add({ PointID, JPlusOne });
 				}
-				if (k + 1 < Dims.Z)
+				if (k + 1 < Dimensions.Z)
 				{
 					int KPlusOne = ControlPointIndexFromCoordinates(i, j, k + 1);
 					OutLatticeEdges.Add({ PointID, KPlusOne });
@@ -351,7 +351,7 @@ FVector3d FFFDLattice::InterpolatedPositionCubic(const FEmbedding& VertexEmbeddi
 				int k = VertexEmbedding.LatticeCell.Z + DK;
 
 				FVector3d LatticePoint;
-				if (i < 0 || i >= Dims.X || j < 0 || j >= Dims.Y || k < 0 || k >= Dims.Z)
+				if (i < 0 || i >= Dimensions.X || j < 0 || j >= Dimensions.Y || k < 0 || k >= Dimensions.Z)
 				{
 					// Get the extrapolated position for a "virtual" control point outside of the deformed lattice
 					LatticePoint = ExtrapolatedLatticePosition({ i,j,k }, LatticeControlPoints);
@@ -502,7 +502,7 @@ FMatrix3d FFFDLattice::CubicInterpolationJacobian(const FEmbedding& VertexEmbedd
 				int k = VertexEmbedding.LatticeCell.Z + DK;
 
 				FVector3d LatticePoint;
-				if (i < 0 || i >= Dims.X || j < 0 || j >= Dims.Y || k < 0 || k >= Dims.Z)
+				if (i < 0 || i >= Dimensions.X || j < 0 || j >= Dimensions.Y || k < 0 || k >= Dimensions.Z)
 				{
 					// Get the extrapolated position for a "virtual" control point outside of the deformed lattice
 					LatticePoint = ExtrapolatedLatticePosition({ i,j,k }, LatticeControlPoints);
@@ -524,7 +524,7 @@ FMatrix3d FFFDLattice::CubicInterpolationJacobian(const FEmbedding& VertexEmbedd
 FVector3d FFFDLattice::ClosestLatticePosition(const FVector3i& VirtualControlPointIndex, const TArray<FVector3d>& LatticeControlPoints) const
 {
 	// Clamp to valid lattice index
-	FVector3i NearestControlPointIndex = Max(Min(VirtualControlPointIndex, Dims - 1), FVector3i(0, 0, 0));
+	FVector3i NearestControlPointIndex = Max(Min(VirtualControlPointIndex, Dimensions - 1), FVector3i(0, 0, 0));
 
 	return LatticeControlPoints[ControlPointIndexFromCoordinates(NearestControlPointIndex)];
 }
@@ -534,14 +534,14 @@ FVector3d FFFDLattice::ExtrapolatedLatticePosition(const FVector3i& VirtualContr
 	// Use the location of the nearest control point and the location of a control point in the opposite direction of 
 	// the extrapolation to get the extrapolated location.
 
-	FVector3i NearestControlPointIndex = Max(Min(VirtualControlPointIndex, Dims - 1), FVector3i(0, 0, 0));
+	FVector3i NearestControlPointIndex = Max(Min(VirtualControlPointIndex, Dimensions - 1), FVector3i(0, 0, 0));
 	FVector3i Delta = VirtualControlPointIndex - NearestControlPointIndex;
 	check(Delta != FVector3i::Zero());
 
 	FVector3i TraceBackControlPointIndex = NearestControlPointIndex - Delta;
-	check(TraceBackControlPointIndex.X >= 0 && TraceBackControlPointIndex.X < Dims.X);
-	check(TraceBackControlPointIndex.Y >= 0 && TraceBackControlPointIndex.Y < Dims.Y);
-	check(TraceBackControlPointIndex.Z >= 0 && TraceBackControlPointIndex.Z < Dims.Z);
+	check(TraceBackControlPointIndex.X >= 0 && TraceBackControlPointIndex.X < Dimensions.X);
+	check(TraceBackControlPointIndex.Y >= 0 && TraceBackControlPointIndex.Y < Dimensions.Y);
+	check(TraceBackControlPointIndex.Z >= 0 && TraceBackControlPointIndex.Z < Dimensions.Z);
 
 	const FVector3d& A = LatticeControlPoints[ControlPointIndexFromCoordinates(TraceBackControlPointIndex)];
 	const FVector3d& B = LatticeControlPoints[ControlPointIndexFromCoordinates(NearestControlPointIndex)];
diff --git a/Engine/Plugins/Experimental/GeometryProcessing/Source/DynamicMesh/Public/Operations/FFDLattice.h b/Engine/Plugins/Experimental/GeometryProcessing/Source/DynamicMesh/Public/Operations/FFDLattice.h
--- a/Engine/Plugins/Experimental/GeometryProcessing/Source/DynamicMesh/Public/Operations/FFDLattice.h
+++ b/Engine/Plugins/Experimental/GeometryProcessing/Source/DynamicMesh/Public/Operations/FFDLattice.h
@@ -80,25 +80,28 @@ public:
 	/// Return the lattice edges by index. The indices refer to the array of lattice positions returned by GenerateInitialLatticePositions
 	void GenerateLatticeEdges(TArray<FVector2i>& OutLatticeEdges) const;
 
-protected:
-
-	// Number of lattice points in each dimension
-	FVector3i Dims;
-
-	// Extents of the lattice before it is deformed
-	FAxisAlignedBox3d InitialBounds;
-
-	// 3D size of a lattice cell
-	FVector3d CellSize;
+	/// Number of lattice points in each dimension
+	const FVector3i& GetDimensions() const
+	{
+		return Dimensions;
+	}
 
-	// Interpolation weights and cell indices per vertex of the input FDynamicMesh3
-	TArray<FEmbedding> VertexEmbeddings;
+	/// 3D size of a lattice cell
+	const FVector3d& GetCellSize() const
+	{
+		return CellSize;
+	}
 
+	/// Extents of the lattice before it is deformed
+	const FAxisAlignedBox3d& GetInitialBounds() const
+	{
+		return InitialBounds;
+	}
 
 	/// Get the index into the flat TArray of positions given the (i,j,k) coordinates in the lattice
 	int ControlPointIndexFromCoordinates(int i, int j, int k) const
 	{
-		int Idx = k + Dims.Z * (j + Dims.Y * i);
+		int Idx = k + Dimensions.Z * (j + Dimensions.Y * i);
 		return Idx;
 	}
 
@@ -108,6 +111,20 @@ protected:
 		return ControlPointIndexFromCoordinates(Index.X, Index.Y, Index.Z);
 	}
 
+protected:
+
+	// Number of lattice points in each dimension
+	FVector3i Dimensions;
+
+	// Extents of the lattice before it is deformed
+	FAxisAlignedBox3d InitialBounds;
+
+	// 3D size of a lattice cell
+	FVector3d CellSize;
+
+	// Interpolation weights and cell indices per vertex of the input FDynamicMesh3
+	TArray<FEmbedding> VertexEmbeddings;
+
 	/// For each vertex in Mesh, compute the lattice cell it resides in and its weighting
 	void ComputeInitialEmbedding(const FDynamicMesh3& Mesh, FLatticeExecutionInfo ExecutionInfo = FLatticeExecutionInfo());
 
diff --git a/Engine/Plugins/Experimental/MeshModelingToolset/Source/MeshModelingTools/Private/LatticeDeformerTool.cpp b/Engine/Plugins/Experimental/MeshModelingToolset/Source/MeshModelingTools/Private/LatticeDeformerTool.cpp
--- a/Engine/Plugins/Experimental/MeshModelingToolset/Source/MeshModelingTools/Private/LatticeDeformerTool.cpp
+++ b/Engine/Plugins/Experimental/MeshModelingToolset/Source/MeshModelingTools/Private/LatticeDeformerTool.cpp
@@ -29,6 +29,78 @@ using namespace UE::Geometry;
 
 #define LOCTEXT_NAMESPACE "ULatticeDeformerTool"
 
+
+namespace
+{
+	void MakeLatticeGraph(const FFFDLattice& Lattice, FDynamicGraph3d& Graph)
+	{
+		const FVector3i& Dims = Lattice.GetDimensions();
+		const FVector3d& CellSize = Lattice.GetCellSize();
+		const FAxisAlignedBox3d& InitialBounds = Lattice.GetInitialBounds();
+
+		// Add cell corners as vertices
+
+		for (int i = 0; i < Dims.X; ++i)
+		{
+			const double X = CellSize.X * i;
+			for (int j = 0; j < Dims.Y; ++j)
+			{
+				const double Y = CellSize.Y * j;
+				for (int k = 0; k < Dims.Z; ++k)
+				{
+					const double Z = CellSize.Z * k;
+
+					const FVector3d Position = InitialBounds.Min + FVector3d{ X,Y,Z };
+					const int P = Lattice.ControlPointIndexFromCoordinates(i, j, k);
+					const int VID = Graph.AppendVertex(Position);
+					ensure(VID == P);
+				}
+			}
+		}
+
+		// Connect cell corners with edges
+
+		for (int i = 0; i < Dims.X; ++i)
+		{
+			for (int j = 0; j < Dims.Y; ++j)
+			{
+				for (int k = 0; k < Dims.Z; ++k)
+				{
+					const int P = Lattice.ControlPointIndexFromCoordinates(i, j, k);
+					if (i + 1 < Dims.X)
+					{
+						const int Pi = Lattice.ControlPointIndexFromCoordinates(i + 1, j, k);
+						Graph.AppendEdge(P, Pi);
+					}
+
+					if (j + 1 < Dims.Y)
+					{
+						const int Pj = Lattice.ControlPointIndexFromCoordinates(i, j + 1, k);
+						Graph.AppendEdge(P, Pj);
+					}
+
+					if (k + 1 < Dims.Z)
+					{
+						const int Pk = Lattice.ControlPointIndexFromCoordinates(i, j, k + 1);
+						Graph.AppendEdge(P, Pk);
+					}
+				}
+			}
+		}
+	}
+}
+
+// Tool properties/actions
+
+void ULatticeDeformerToolProperties::PostAction(ELatticeDeformerToolAction Action)
+{
+	if (ParentTool.IsValid())
+	{
+		ParentTool->RequestAction(Action);
+	}
+}
+
+
 // Tool builder
 
 USingleSelectionMeshEditingTool* ULatticeDeformerToolBuilder::CreateNewTool(const FToolBuilderState& SceneState) const
@@ -102,6 +174,7 @@ void ULatticeDeformerTool::Setup()
 	Converter.Convert(Cast<IMeshDescriptionProvider>(Target)->GetMeshDescription(), *OriginalMesh);
 
 	Settings = NewObject<ULatticeDeformerToolProperties>(this, TEXT("Lattice Deformer Tool Settings"));
+	Settings->Initialize(this);
 	Settings->RestoreProperties(this);
 	AddToolPropertySource(Settings);
 
@@ -126,7 +199,14 @@ void ULatticeDeformerTool::Setup()
 	{
 		ControlPointsMechanic->UpdateSetPivotMode(Settings->bSetPivotMode);
 	});
-	
+	Settings->WatchProperty(Settings->bSoftDeformation, [this](bool)
+	{
+		if (Settings->bSoftDeformation)
+		{
+			RebuildDeformer();
+		}
+	});
+
 
 	TArray<FVector3d> LatticePoints;
 	TArray<FVector2i> LatticeEdges;
@@ -141,17 +221,115 @@ void ULatticeDeformerTool::Setup()
 
 	auto OnPointsChangedLambda = [this]()
 	{
+		if (Settings->bSoftDeformation)
+		{
+			SoftDeformLattice();
+		}
 		Preview->InvalidateResult();
 		Settings->bCanChangeResolution = !ControlPointsMechanic->bHasChanged;
 	};
 	ControlPointsMechanic->OnPointsChanged.AddLambda(OnPointsChangedLambda);
 
+	ControlPointsMechanic->OnSelectionChanged.AddLambda([this]()
+	{
+		if (Settings->bSoftDeformation)
+		{
+			RebuildDeformer();
+		}
+	});
+
 	ControlPointsMechanic->SetCoordinateSystem(Settings->GizmoCoordinateSystem);
 	ControlPointsMechanic->UpdateSetPivotMode(Settings->bSetPivotMode);
 
 	StartPreview();
 }
 
+
+void ULatticeDeformerTool::RebuildDeformer()
+{
+	LatticeGraph = MakePimpl<UE::Geometry::FDynamicGraph3d>();
+	MakeLatticeGraph(*Lattice, *LatticeGraph);
+
+	const TArray<FVector3d>& CurrentLatticePoints = ControlPointsMechanic->GetControlPoints();
+	check(LatticeGraph->VertexCount() == CurrentLatticePoints.Num());
+
+	for (int VID : LatticeGraph->VertexIndicesItr())
+	{
+		LatticeGraph->SetVertex(VID, CurrentLatticePoints[VID]);
+	}
+
+	DeformationSolver = UE::MeshDeformation::ConstructUniformConstrainedMeshDeformer(*LatticeGraph);
+
+	for (int LatticePointIndex = 0; LatticePointIndex < CurrentLatticePoints.Num(); ++LatticePointIndex)
+	{
+		if(ConstrainedLatticePoints.Contains(LatticePointIndex))
+		{
+			// Pin constraint
+			DeformationSolver->AddConstraint(LatticePointIndex, 1.0, ConstrainedLatticePoints[LatticePointIndex], true);
+		}
+		else 
+		{
+			if (ControlPointsMechanic->ControlPointIsSelected(LatticePointIndex))
+			{
+				const FVector3d& MovePosition = CurrentLatticePoints[LatticePointIndex];
+				DeformationSolver->AddConstraint(LatticePointIndex, 1.0, MovePosition, true);
+			}
+		}
+	}
+}
+
+
+void ULatticeDeformerTool::SoftDeformLattice()
+{
+	if (!ensure(Lattice))
+	{
+		return;
+	}
+
+	if (!ensure(ControlPointsMechanic))
+	{
+		return;
+	}
+
+	if (!ensure(DeformationSolver))
+	{
+		return;
+	}
+	
+	const TArray<FVector3d>& CurrentLatticePoints = ControlPointsMechanic->GetControlPoints();
+
+	if (!ensure(LatticeGraph->VertexCount() == CurrentLatticePoints.Num()))
+	{
+		return;
+	}
+
+	for (int LatticePointIndex = 0; LatticePointIndex < CurrentLatticePoints.Num(); ++LatticePointIndex)
+	{
+		if (ControlPointsMechanic->ControlPointIsSelected(LatticePointIndex))
+		{
+			// Don't move pinned points
+			if (ConstrainedLatticePoints.Contains(LatticePointIndex))
+			{
+				continue;
+			}
+
+			if (!ensure(DeformationSolver->IsConstrained(LatticePointIndex)))
+			{
+				continue;
+			}
+
+			const FVector3d& MovePosition = CurrentLatticePoints[LatticePointIndex];
+			DeformationSolver->UpdateConstraintPosition(LatticePointIndex, MovePosition, true);
+		}
+	}
+
+	TArray<FVector3d> DeformedLatticePoints;
+	DeformationSolver->Deform(DeformedLatticePoints);
+
+	ControlPointsMechanic->UpdateControlPointPositions(DeformedLatticePoints);
+}
+
+
 void ULatticeDeformerTool::Shutdown(EToolShutdownType ShutdownType)
 {
 	Settings->SaveProperties(this);
@@ -219,8 +397,31 @@ void ULatticeDeformerTool::StartPreview()
 	Cast<IPrimitiveComponentBackedTarget>(Target)->SetOwnerVisibility(false);
 }
 
+
+void ULatticeDeformerTool::ApplyAction(ELatticeDeformerToolAction Action)
+{
+	switch (Action)
+	{
+	case ELatticeDeformerToolAction::ClearConstraints:
+		ClearConstrainedPoints();
+		break;
+	case ELatticeDeformerToolAction::Constrain:
+		ConstrainSelectedPoints();
+		break;
+	default:
+		break;
+	}
+}
+
+
 void ULatticeDeformerTool::OnTick(float DeltaTime)
 {
+	if (PendingAction != ELatticeDeformerToolAction::NoAction)
+	{
+		ApplyAction(PendingAction);
+		PendingAction = ELatticeDeformerToolAction::NoAction;
+	}
+
 	if (Preview)
 	{
 		if (bShouldRebuild)
@@ -247,4 +448,87 @@ void ULatticeDeformerTool::Render(IToolsContextRenderAPI* RenderAPI)
 	}
 }
 
+void ULatticeDeformerTool::RequestAction(ELatticeDeformerToolAction Action)
+{
+	if (PendingAction == ELatticeDeformerToolAction::NoAction)
+	{
+		PendingAction = Action;
+	}
+}
+
+
+static const FText LatticeConstraintChangeTransactionText = LOCTEXT("LatticeConstraintChange", "Lattice Constraint Change");
+
+void ULatticeDeformerTool::ConstrainSelectedPoints()
+{
+	TMap<int, FVector3d> PrevConstrainedLatticePoints = ConstrainedLatticePoints;
+	const TArray<FVector3d>& CurrentControlPointPositions = ControlPointsMechanic->GetControlPoints();
+	for (int32 VID : ControlPointsMechanic->GetSelectedPointIDs())
+	{
+		ConstrainedLatticePoints.FindOrAdd(VID) = CurrentControlPointPositions[VID];
+	}
+	UpdateMechanicColorOverrides();
+
+	GetToolManager()->EmitObjectChange(this, MakeUnique<FLatticeDeformerToolConstrainedPointsChange>(PrevConstrainedLatticePoints,
+																									 ConstrainedLatticePoints, 
+																									 CurrentChangeStamp), 
+									   LatticeConstraintChangeTransactionText);
+}
+
+void ULatticeDeformerTool::ClearConstrainedPoints()
+{
+	TMap<int, FVector3d> PrevConstrainedLatticePoints = ConstrainedLatticePoints;
+	ConstrainedLatticePoints.Reset();
+	UpdateMechanicColorOverrides();
+
+	GetToolManager()->EmitObjectChange(this, MakeUnique<FLatticeDeformerToolConstrainedPointsChange>(PrevConstrainedLatticePoints,
+																									 ConstrainedLatticePoints,
+																									 CurrentChangeStamp),
+									   LatticeConstraintChangeTransactionText);
+}
+
+
+void ULatticeDeformerTool::UpdateMechanicColorOverrides()
+{
+	ControlPointsMechanic->ClearAllPointColorOverrides();
+	for ( const TPair<int32,FVector3d>& Constraint : ConstrainedLatticePoints)
+	{
+		ControlPointsMechanic->SetPointColorOverride(Constraint.Key, FColor::Cyan);
+	}
+	RebuildDeformer();
+	ControlPointsMechanic->UpdateDrawables();
+}
+
+
+
+void FLatticeDeformerToolConstrainedPointsChange::Apply(UObject* Object)
+{
+	ULatticeDeformerTool* Tool = Cast<ULatticeDeformerTool>(Object);
+	if (!ensure(Tool))
+	{
+		return;
+	}
+
+	Tool->ConstrainedLatticePoints = NewConstrainedLatticePoints;
+	Tool->UpdateMechanicColorOverrides();
+}
+
+void FLatticeDeformerToolConstrainedPointsChange::Revert(UObject* Object)
+{
+	ULatticeDeformerTool* Tool = Cast<ULatticeDeformerTool>(Object);
+	if (!ensure(Tool))
+	{
+		return;
+	}
+
+	Tool->ConstrainedLatticePoints = PrevConstrainedLatticePoints;
+	Tool->UpdateMechanicColorOverrides();
+}
+
+FString FLatticeDeformerToolConstrainedPointsChange::ToString() const
+{
+	return TEXT("FLatticeDeformerToolConstrainedPointsChange");
+}
+
+
 #undef LOCTEXT_NAMESPACE
diff --git a/Engine/Plugins/Experimental/MeshModelingToolset/Source/MeshModelingTools/Public/LatticeDeformerTool.h b/Engine/Plugins/Experimental/MeshModelingToolset/Source/MeshModelingTools/Public/LatticeDeformerTool.h
--- a/Engine/Plugins/Experimental/MeshModelingToolset/Source/MeshModelingTools/Public/LatticeDeformerTool.h
+++ b/Engine/Plugins/Experimental/MeshModelingToolset/Source/MeshModelingTools/Public/LatticeDeformerTool.h
@@ -9,6 +9,8 @@
 #include "SingleSelectionTool.h" //USingleSelectionTool
 #include "Operations/FFDLattice.h"
 #include "BaseTools/SingleSelectionMeshEditingTool.h"
+#include "Solvers/ConstrainedMeshDeformer.h"
+
 #include "LatticeDeformerTool.generated.h"
 
 class ULatticeControlPointsMechanic;
@@ -34,6 +36,13 @@ enum class ELatticeInterpolationType : uint8
 	Cubic UMETA(DisplayName = "Cubic")
 };
 
+UENUM()
+enum class ELatticeDeformerToolAction  : uint8
+{
+	NoAction,
+	Constrain,
+	ClearConstraints
+};
 
 UCLASS()
 class MESHMODELINGTOOLS_API ULatticeDeformerToolProperties : public UInteractiveToolPropertySet
@@ -41,6 +50,9 @@ class MESHMODELINGTOOLS_API ULatticeDeformerToolProperties : public UInteractive
 	GENERATED_BODY()
 
 public:
+	TWeakObjectPtr<ULatticeDeformerTool> ParentTool;
+	void Initialize(ULatticeDeformerTool* ParentToolIn) { ParentTool = ParentToolIn; }
+	void PostAction(ELatticeDeformerToolAction Action);
 
 	/** Number of lattice vertices along the X axis */
 	UPROPERTY(EditAnywhere, Category = Resolution, meta = (UIMin = "2", ClampMin = "2", EditCondition = "bCanChangeResolution", HideEditConditionToggle))
@@ -78,6 +90,24 @@ public:
 	UPROPERTY(EditAnywhere, Category = Gizmo)
 	bool bSetPivotMode = false;
 
+	/** Whether to use soft deformation of the lattice */
+	UPROPERTY(EditAnywhere, Category = Deformation)
+	bool bSoftDeformation = false;
+
+	/** Constrain selected lattice points */
+	UFUNCTION(CallInEditor, Category = Deformation, meta = (DisplayName = "Constrain"))
+	void Constrain() 
+	{
+		PostAction(ELatticeDeformerToolAction::Constrain);
+	}
+
+	/** Clear all constrained lattice points */
+	UFUNCTION(CallInEditor, Category = Deformation, meta = (DisplayName = "Clear Constraints"))
+	void ClearConstraints()
+	{
+		PostAction(ELatticeDeformerToolAction::ClearConstraints);
+	}
+
 };
 
 
@@ -103,6 +133,8 @@ class MESHMODELINGTOOLS_API ULatticeDeformerTool : public USingleSelectionMeshEd
 
 public:
 
+	virtual ~ULatticeDeformerTool() = default;
+
 	virtual void DrawHUD(FCanvas* Canvas, IToolsContextRenderAPI* RenderAPI) override;
 
 	virtual bool HasCancel() const override { return true; }
@@ -143,5 +175,55 @@ protected:
 
 	void StartPreview();
 
+	TUniquePtr<UE::Solvers::IConstrainedMeshSolver> DeformationSolver;
+	TPimplPtr<UE::Geometry::FDynamicGraph3d> LatticeGraph;
+
+	TMap<int, FVector3d> ConstrainedLatticePoints;
+	void ConstrainSelectedPoints();
+	void ClearConstrainedPoints();
+	void UpdateMechanicColorOverrides();
+
+	void RebuildDeformer();
+	void SoftDeformLattice();
+
+	int32 CurrentChangeStamp = 0;
+
+	ELatticeDeformerToolAction PendingAction = ELatticeDeformerToolAction::NoAction;
+	void RequestAction(ELatticeDeformerToolAction Action);
+	void ApplyAction(ELatticeDeformerToolAction Action);
+
 	friend class ULatticeDeformerOperatorFactory;
+	friend class ULatticeDeformerToolProperties;
+	friend class FLatticeDeformerToolConstrainedPointsChange;
+};
+
+
+// Set of constrained points change
+class MESHMODELINGTOOLS_API FLatticeDeformerToolConstrainedPointsChange : public FToolCommandChange
+{
+public:
+
+	FLatticeDeformerToolConstrainedPointsChange(const TMap<int, FVector3d>& PrevConstrainedLatticePointsIn,
+												const TMap<int, FVector3d>& NewConstrainedLatticePointsIn,
+												int32 ChangeStampIn) :
+		PrevConstrainedLatticePoints(PrevConstrainedLatticePointsIn),
+		NewConstrainedLatticePoints(NewConstrainedLatticePointsIn),
+		ChangeStamp(ChangeStampIn)
+	{}
+
+	virtual void Apply(UObject* Object) override;
+	virtual void Revert(UObject* Object) override;
+
+	virtual bool HasExpired(UObject* Object) const override
+	{
+		return Cast<ULatticeDeformerTool>(Object)->CurrentChangeStamp != ChangeStamp;
+	}
+
+	virtual FString ToString() const override;
+
+protected:
+
+	TMap<int, FVector3d> PrevConstrainedLatticePoints;
+	TMap<int, FVector3d> NewConstrainedLatticePoints;
+	int32 ChangeStamp;
 };
diff --git a/Engine/Plugins/Experimental/MeshModelingToolset/Source/ModelingComponents/Private/Mechanics/LatticeControlPointsMechanic.cpp b/Engine/Plugins/Experimental/MeshModelingToolset/Source/ModelingComponents/Private/Mechanics/LatticeControlPointsMechanic.cpp
--- a/Engine/Plugins/Experimental/MeshModelingToolset/Source/ModelingComponents/Private/Mechanics/LatticeControlPointsMechanic.cpp
+++ b/Engine/Plugins/Experimental/MeshModelingToolset/Source/ModelingComponents/Private/Mechanics/LatticeControlPointsMechanic.cpp
@@ -226,6 +226,11 @@ void ULatticeControlPointsMechanic::RebuildDrawables()
 		}
 	}
 
+	for (const TPair<int32, FColor>& Override : ColorOverrides)
+	{
+		DrawnControlPoints->SetPointColor(Override.Key, Override.Value);
+	}
+
 	DrawnLatticeEdges->Clear();
 	for (int32 EdgeIndex = 0; EdgeIndex < LatticeEdges.Num(); ++EdgeIndex)
 	{
@@ -256,6 +261,11 @@ void ULatticeControlPointsMechanic::UpdateDrawables()
 		}
 	}
 
+	for (const TPair<int32, FColor>& Override : ColorOverrides)
+	{
+		DrawnControlPoints->SetPointColor(Override.Key, Override.Value);
+	}
+
 	for (int32 EdgeIndex = 0; EdgeIndex < LatticeEdges.Num(); ++EdgeIndex)
 	{
 		const FVector3d& A = ControlPoints[LatticeEdges[EdgeIndex].X];
@@ -405,6 +415,7 @@ void ULatticeControlPointsMechanic::OnClicked(const FInputDeviceRay& ClickPos)
 void ULatticeControlPointsMechanic::ChangeSelection(int32 NewPointID, bool bAddToSelection)
 {
 	// If not adding to selection, clear it
+	bool bAnyChange = false;
 	if (!bAddToSelection && SelectedPointIDs.Num() > 0)
 	{
 		TSet<int32> PointsToDeselect;
@@ -427,6 +438,12 @@ void ULatticeControlPointsMechanic::ChangeSelection(int32 NewPointID, bool bAddT
 
 		ParentTool->GetToolManager()->EmitObjectChange(this, MakeUnique<FLatticeControlPointsMechanicSelectionChange>(
 			PointsToDeselect, false, PreviousTransform, NewTransform, CurrentChangeStamp), LatticePointDeselectionTransactionText);
+		bAnyChange = true;
+	}
+
+	for (const TPair<int32, FColor>& Override : ColorOverrides)
+	{
+		DrawnControlPoints->SetPointColor(Override.Key, Override.Value);
 	}
 
 	// We check for validity here because giving an invalid id (such as -1) with bAddToSelection == false
@@ -450,6 +467,12 @@ void ULatticeControlPointsMechanic::ChangeSelection(int32 NewPointID, bool bAddT
 			ParentTool->GetToolManager()->EmitObjectChange(this, MakeUnique<FLatticeControlPointsMechanicSelectionChange>(
 				NewPointID, true, PreviousTransform, NewTransform, CurrentChangeStamp), LatticePointSelectionTransactionText);
 		}
+		bAnyChange = true;
+	}
+
+	if (bAnyChange)
+	{
+		OnSelectionChanged.Broadcast();
 	}
 }
 
@@ -492,7 +515,14 @@ bool ULatticeControlPointsMechanic::DeselectPoint(int32 PointID)
 	
 	if (SelectedPointIDs.Remove(PointID) != 0)
 	{
-		DrawnControlPoints->SetPointColor(PointID, NormalPointColor);
+		if (ColorOverrides.Contains(PointID))
+		{
+			DrawnControlPoints->SetPointColor(PointID, ColorOverrides[PointID]);
+		}
+		else
+		{
+			DrawnControlPoints->SetPointColor(PointID, NormalPointColor);
+		}
 		PointFound = true;
 	}
 	
@@ -502,7 +532,15 @@ bool ULatticeControlPointsMechanic::DeselectPoint(int32 PointID)
 void ULatticeControlPointsMechanic::SelectPoint(int32 PointID)
 {
 	SelectedPointIDs.Add(PointID);
-	DrawnControlPoints->SetPointColor(PointID, SelectedColor);
+
+	if (ColorOverrides.Contains(PointID))
+	{
+		DrawnControlPoints->SetPointColor(PointID, ColorOverrides[PointID]);
+	}
+	else
+	{
+		DrawnControlPoints->SetPointColor(PointID, SelectedColor);
+	}
 }
 
 void ULatticeControlPointsMechanic::ClearSelection()
@@ -635,20 +673,29 @@ void ULatticeControlPointsMechanic::OnDragRectangleFinished()
 
 	FTransform NewTransform = PointTransformProxy->GetTransform();
 
+	bool bAnyChange = false;
+
 	if (PreDragSelection.Num() > 0)
 	{
 		ParentTool->GetToolManager()->EmitObjectChange(this, MakeUnique<FLatticeControlPointsMechanicSelectionChange>(
 			PreDragSelection, false, PreviousTransform, NewTransform, CurrentChangeStamp), LatticePointDeselectionTransactionText);
+		bAnyChange = true;
 	}
 
 	if (SelectedPointIDs.Num() > 0)
 	{
 		ParentTool->GetToolManager()->EmitObjectChange(this, MakeUnique<FLatticeControlPointsMechanicSelectionChange>(
 			SelectedPointIDs, true, PreviousTransform, NewTransform, CurrentChangeStamp), LatticePointSelectionTransactionText);
+		bAnyChange = true;
 	}
 
 	ParentTool->GetToolManager()->EndUndoTransaction();
 
+	if (bAnyChange)
+	{
+		OnSelectionChanged.Broadcast();
+	}
+
 	UpdateDrawables();
 }
 
@@ -657,6 +704,18 @@ const TArray<FVector3d>& ULatticeControlPointsMechanic::GetControlPoints() const
 	return ControlPoints;
 }
 
+void ULatticeControlPointsMechanic::UpdateControlPointPositions(const TArray<FVector3d>& NewPoints)
+{
+	if (!ensure(NewPoints.Num() == ControlPoints.Num()))
+	{
+		return;
+	}
+
+	ControlPoints = NewPoints;
+	UpdateDrawables();
+}
+
+
 // ==================== Undo/redo object functions ====================
 
 FLatticeControlPointsMechanicSelectionChange::FLatticeControlPointsMechanicSelectionChange(int32 PointIDIn,
@@ -704,6 +763,7 @@ void FLatticeControlPointsMechanicSelectionChange::Apply(UObject* Object)
 
 	bool bAnyPointSelected = (Mechanic->SelectedPointIDs.Num() > 0);
 	Mechanic->PointTransformGizmo->SetVisibility(bAnyPointSelected);
+	Mechanic->OnSelectionChanged.Broadcast();
 }
 
 void FLatticeControlPointsMechanicSelectionChange::Revert(UObject* Object)
@@ -726,6 +786,7 @@ void FLatticeControlPointsMechanicSelectionChange::Revert(UObject* Object)
 
 	bool bAnyPointSelected = (Mechanic->SelectedPointIDs.Num() > 0);
 	Mechanic->PointTransformGizmo->SetVisibility(bAnyPointSelected);
+	Mechanic->OnSelectionChanged.Broadcast();
 }
 
 FString FLatticeControlPointsMechanicSelectionChange::ToString() const
diff --git a/Engine/Plugins/Experimental/MeshModelingToolset/Source/ModelingComponents/Public/Mechanics/LatticeControlPointsMechanic.h b/Engine/Plugins/Experimental/MeshModelingToolset/Source/ModelingComponents/Public/Mechanics/LatticeControlPointsMechanic.h
--- a/Engine/Plugins/Experimental/MeshModelingToolset/Source/ModelingComponents/Public/Mechanics/LatticeControlPointsMechanic.h
+++ b/Engine/Plugins/Experimental/MeshModelingToolset/Source/ModelingComponents/Public/Mechanics/LatticeControlPointsMechanic.h
@@ -39,16 +39,21 @@ public:
 
 	// TODO: Snapping
 
-	// This delegate is called every time the control points are altered.
+	// This delegate is called every time the control points are moved.
 	DECLARE_MULTICAST_DELEGATE(OnPointsChangedEvent);
 	OnPointsChangedEvent OnPointsChanged;
 
+	// This delegate is called every time the control point selection changes.
+	DECLARE_MULTICAST_DELEGATE(OnSelectionChangedEvent);
+	OnSelectionChangedEvent OnSelectionChanged;
+
 	virtual void Initialize(const TArray<FVector3d>& Points, 
 							const TArray<FVector2i>& Edges,
 							const FTransform3d& LocalToWorldTransform );
 
 	void SetWorld(UWorld* World);
 	const TArray<FVector3d>& GetControlPoints() const;
+	void UpdateControlPointPositions(const TArray<FVector3d>& NewPoints);
 
 	void SetCoordinateSystem(EToolContextCoordinateSystem InCoordinateSystem);
 	EToolContextCoordinateSystem GetCoordinateSystem() const;
@@ -75,6 +80,33 @@ public:
 
 	bool bHasChanged = false;
 
+	bool ControlPointIsSelected(int32 Index)
+	{
+		return SelectedPointIDs.Contains(Index);
+	}
+
+	const TSet<int32>& GetSelectedPointIDs() const
+	{
+		return SelectedPointIDs;
+	}
+
+	void SetPointColorOverride(int32 Index, const FColor& NewColor)
+	{
+		ColorOverrides.FindOrAdd(Index) = NewColor;
+	}
+
+	void ClearPointColorOverride(int32 Index)
+	{
+		ColorOverrides.Remove(Index);
+	}
+
+	void ClearAllPointColorOverrides()
+	{
+		ColorOverrides.Reset();
+	}
+
+	void UpdateDrawables();
+
 protected:
 
 	TArray<FVector3d> ControlPoints;
@@ -167,11 +199,12 @@ protected:
 	void UpdatePointLocations(const TMap<int32, FVector3d>& NewLocations);
 
 	void RebuildDrawables();
-	void UpdateDrawables();
 
 	// Used for expiring undo/redo changes, which compare this to their stored value and expire themselves if they do not match.
 	int32 CurrentChangeStamp = 0;
 
+	TMap<int32, FColor> ColorOverrides;
+
 	friend class FLatticeControlPointsMechanicSelectionChange;
 	friend class FLatticeControlPointsMechanicMovementChange;
 };
