Fix non-unity errors in bake tools.

#rnx

diff --git a/Engine/Plugins/Experimental/MeshModelingToolset/Source/MeshModelingTools/Private/BakeMeshAttributeToolCommon.cpp b/Engine/Plugins/Experimental/MeshModelingToolset/Source/MeshModelingTools/Private/BakeMeshAttributeToolCommon.cpp
--- a/Engine/Plugins/Experimental/MeshModelingToolset/Source/MeshModelingTools/Private/BakeMeshAttributeToolCommon.cpp
+++ b/Engine/Plugins/Experimental/MeshModelingToolset/Source/MeshModelingTools/Private/BakeMeshAttributeToolCommon.cpp
@@ -1,6 +1,7 @@
 // Copyright Epic Games, Inc. All Rights Reserved.
 
 #include "BakeMeshAttributeToolCommon.h"
+#include "Engine/Classes/Engine/Texture2D.h"
 
 FTempTextureAccess::FTempTextureAccess(UTexture2D* DisplacementMap)
 	: DisplacementMap(DisplacementMap)
diff --git a/Engine/Plugins/Experimental/MeshModelingToolset/Source/MeshModelingTools/Private/BakeMeshAttributeVertexTool.cpp b/Engine/Plugins/Experimental/MeshModelingToolset/Source/MeshModelingTools/Private/BakeMeshAttributeVertexTool.cpp
--- a/Engine/Plugins/Experimental/MeshModelingToolset/Source/MeshModelingTools/Private/BakeMeshAttributeVertexTool.cpp
+++ b/Engine/Plugins/Experimental/MeshModelingToolset/Source/MeshModelingTools/Private/BakeMeshAttributeVertexTool.cpp
@@ -4,6 +4,9 @@
 #include "InteractiveToolManager.h"
 #include "ToolBuilderUtil.h"
 
+#include "Materials/Material.h"
+#include "Materials/MaterialInstanceDynamic.h"
+
 #include "DynamicMesh/MeshTransforms.h"
 #include "MeshDescriptionToDynamicMesh.h"
 
diff --git a/Engine/Plugins/Experimental/MeshModelingToolset/Source/MeshModelingTools/Public/BakeMeshAttributeToolCommon.h b/Engine/Plugins/Experimental/MeshModelingToolset/Source/MeshModelingTools/Public/BakeMeshAttributeToolCommon.h
--- a/Engine/Plugins/Experimental/MeshModelingToolset/Source/MeshModelingTools/Public/BakeMeshAttributeToolCommon.h
+++ b/Engine/Plugins/Experimental/MeshModelingToolset/Source/MeshModelingTools/Public/BakeMeshAttributeToolCommon.h
@@ -11,6 +11,7 @@
 // Pre-declarations
 using UE::Geometry::FImageDimensions;
 using UE::Geometry::TImageBuilder;
+class UTexture2D;
 
 
 /**
diff --git a/Engine/Plugins/Experimental/MeshModelingToolset/Source/MeshModelingTools/Public/BakeMeshAttributeVertexTool.h b/Engine/Plugins/Experimental/MeshModelingToolset/Source/MeshModelingTools/Public/BakeMeshAttributeVertexTool.h
--- a/Engine/Plugins/Experimental/MeshModelingToolset/Source/MeshModelingTools/Public/BakeMeshAttributeVertexTool.h
+++ b/Engine/Plugins/Experimental/MeshModelingToolset/Source/MeshModelingTools/Public/BakeMeshAttributeVertexTool.h
@@ -14,6 +14,9 @@
 #include "BakeMeshAttributeToolCommon.h"
 #include "BakeMeshAttributeVertexTool.generated.h"
 
+// predeclarations
+class UMaterialInstanceDynamic;
+
 
 /**
  * Tool Builder
