Non-unity compile fix.

#rnx

diff --git a/Engine/Plugins/Experimental/MeshModelingToolset/Source/MeshModelingTools/Private/BakeMeshAttributeToolCommon.cpp b/Engine/Plugins/Experimental/MeshModelingToolset/Source/MeshModelingTools/Private/BakeMeshAttributeToolCommon.cpp
--- a/Engine/Plugins/Experimental/MeshModelingToolset/Source/MeshModelingTools/Private/BakeMeshAttributeToolCommon.cpp
+++ b/Engine/Plugins/Experimental/MeshModelingToolset/Source/MeshModelingTools/Private/BakeMeshAttributeToolCommon.cpp
@@ -1,7 +1,6 @@
 // Copyright Epic Games, Inc. All Rights Reserved.
 
 #include "BakeMeshAttributeToolCommon.h"
-#include "Engine/Classes/Engine/Texture2D.h"
 
 FTempTextureAccess::FTempTextureAccess(UTexture2D* DisplacementMap)
 	: DisplacementMap(DisplacementMap)
diff --git a/Engine/Plugins/Experimental/MeshModelingToolset/Source/MeshModelingTools/Public/BakeMeshAttributeToolCommon.h b/Engine/Plugins/Experimental/MeshModelingToolset/Source/MeshModelingTools/Public/BakeMeshAttributeToolCommon.h
--- a/Engine/Plugins/Experimental/MeshModelingToolset/Source/MeshModelingTools/Public/BakeMeshAttributeToolCommon.h
+++ b/Engine/Plugins/Experimental/MeshModelingToolset/Source/MeshModelingTools/Public/BakeMeshAttributeToolCommon.h
@@ -3,6 +3,7 @@
 #pragma once
 
 #include "CoreMinimal.h"
+#include "Engine/Classes/Engine/Texture2D.h"
 #include "MultiSelectionTool.h"
 #include "Image/ImageBuilder.h"
 #include "Image/ImageDimensions.h"
@@ -11,7 +12,6 @@
 // Pre-declarations
 using UE::Geometry::FImageDimensions;
 using UE::Geometry::TImageBuilder;
-class UTexture2D;
 
 
 /**
