WorldRenderCapture
* Removed potentially ineffective forced VT prefetching
* Replaced with a warmup by rendering the scene multiple frames before the capture
* Number of warmup frames controlled with the modeling.WorldRenderCapture.VTWarmupFrames cvar (defaults to 5)
* Remove unneeded FlushRenderingCommands() (already performed in ReadLinearColorPixels())

diff --git a/Engine/Plugins/Experimental/MeshModelingToolset/Source/ModelingComponents/Private/Scene/WorldRenderCapture.cpp b/Engine/Plugins/Experimental/MeshModelingToolset/Source/ModelingComponents/Private/Scene/WorldRenderCapture.cpp
--- a/Engine/Plugins/Experimental/MeshModelingToolset/Source/ModelingComponents/Private/Scene/WorldRenderCapture.cpp
+++ b/Engine/Plugins/Experimental/MeshModelingToolset/Source/ModelingComponents/Private/Scene/WorldRenderCapture.cpp
@@ -17,7 +17,10 @@
 
 using namespace UE::Geometry;
 
-
+static TAutoConsoleVariable<int32> CVarModelingWorldRenderCaptureVTWarmupFrames(
+	TEXT("modeling.WorldRenderCapture.VTWarmupFrames"),
+	5,
+	TEXT("Number of frames to render before each capture in order to warmup the VT."));
 
 FRenderCaptureTypeFlags FRenderCaptureTypeFlags::All()
 {
@@ -231,6 +234,22 @@ namespace UE
 namespace Internal
 {
 
+/**
+ * Render the scene to the provided canvas. Will potentially perform the render multiple times, depending on the value of
+ * the CVarModelingWorldRenderCaptureVTWarmupFrames CVar. This is needed to ensure the VT is primed properly before capturing
+ * the scene.
+ */
+static void PerformSceneRender(FCanvas* Canvas, FSceneViewFamily* ViewFamily)
+{
+	for (int32 i = 0; i < CVarModelingWorldRenderCaptureVTWarmupFrames.GetValueOnGameThread(); i++)
+	{
+		GetRendererModule().BeginRenderingViewFamily(Canvas, ViewFamily);
+	}
+
+	GetRendererModule().BeginRenderingViewFamily(Canvas, ViewFamily);
+}
+
+
 /**
  * Reads data from FImagePixelData and stores in an output image, with optional ColorTransformFunc
  */
@@ -324,9 +343,6 @@ bool FWorldRenderCapture::CaptureMRSFromPosition(
 	check(this->World);
 	FSceneInterface* Scene = this->World->Scene;
 
-	// Prefetch all virtual textures so that we have content available
-	UE::AssetUtils::ForceVirtualTexturePrefetch(Dimensions);
-
 	UTextureRenderTarget2D* RenderTargetTexture = GetRenderTexture(true);
 	if (!ensure(RenderTargetTexture))
 	{
@@ -447,7 +463,8 @@ bool FWorldRenderCapture::CaptureMRSFromPosition(
 	// do we actually need for force SM5 here? 
 	FCanvas Canvas = FCanvas(RenderTargetResource, nullptr, this->World, ERHIFeatureLevel::SM5, FCanvas::CDM_DeferDrawing, 1.0f);
 	Canvas.Clear(FLinearColor::Transparent);
-	GetRendererModule().BeginRenderingViewFamily(&Canvas, &ViewFamily);
+
+	UE::Internal::PerformSceneRender(&Canvas, &ViewFamily);
 
 	// wait for render
 	FlushRenderingCommands();
@@ -472,9 +489,6 @@ bool FWorldRenderCapture::CaptureEmissiveFromPosition(
 	double NearPlaneDist,
 	FImageAdapter& ResultImageOut)
 {
-	// Prefetch all virtual textures so that we have content available
-	UE::AssetUtils::ForceVirtualTexturePrefetch(Dimensions);
-
 	UTextureRenderTarget2D* RenderTargetTexture = GetRenderTexture(true);
 	if (ensure(RenderTargetTexture) == false)
 	{
@@ -562,7 +576,7 @@ bool FWorldRenderCapture::CaptureEmissiveFromPosition(
 	FCanvas Canvas(RenderTargetResource, NULL, FApp::GetCurrentTime() - GStartTime, FApp::GetDeltaTime(), FApp::GetCurrentTime() - GStartTime, World->Scene->GetFeatureLevel());
 	Canvas.Clear(FLinearColor::Transparent);
 
-	GetRendererModule().BeginRenderingViewFamily(&Canvas, &ViewFamily);
+	UE::Internal::PerformSceneRender(&Canvas, &ViewFamily);
 
 	// Copy the contents of the remote texture to system memory
 	ReadImageBuffer.Reset();
@@ -571,8 +585,6 @@ bool FWorldRenderCapture::CaptureEmissiveFromPosition(
 	ReadSurfaceDataFlags.SetLinearToGamma(false);
 	RenderTargetResource->ReadLinearColorPixels(ReadImageBuffer, ReadSurfaceDataFlags, FIntRect(0, 0, Width, Height));
 
-	FlushRenderingCommands();
-
 	ResultImageOut.SetDimensions(Dimensions);
 	for (int32 yi = 0; yi < Height; ++yi)
 	{
@@ -609,9 +621,6 @@ namespace Internal
 		TArray<FLinearColor>& OutSamples
 	)
 	{
-		// Prefetch all virtual textures so that we have content available
-		UE::AssetUtils::ForceVirtualTexturePrefetch(Dimensions);
-
 		int32 Width = Dimensions.GetWidth();
 		int32 Height = Dimensions.GetHeight();
 		FTextureRenderTargetResource* RenderTargetResource = RenderTargetTexture->GameThread_GetRenderTargetResource();
@@ -649,7 +658,7 @@ namespace Internal
 		FCanvas Canvas(RenderTargetResource, NULL, FApp::GetCurrentTime() - GStartTime, FApp::GetDeltaTime(), FApp::GetCurrentTime() - GStartTime, Scene->GetFeatureLevel());
 		Canvas.Clear(FLinearColor::Transparent);
 
-		GetRendererModule().BeginRenderingViewFamily(&Canvas, &ViewFamily);
+		UE::Internal::PerformSceneRender(&Canvas, &ViewFamily);
 
 		// Copy the contents of the remote texture to system memory
 		OutSamples.SetNumUninitialized(Width * Height);
@@ -657,8 +666,6 @@ namespace Internal
 		FReadSurfaceDataFlags ReadSurfaceDataFlags;
 		ReadSurfaceDataFlags.SetLinearToGamma(false);
 		RenderTargetResource->ReadLinearColorPixels(OutSamples, ReadSurfaceDataFlags, FIntRect(0, 0, Width, Height));
-
-		FlushRenderingCommands();
 	}
 }
 }
diff --git a/Engine/Plugins/Experimental/MeshModelingToolset/Source/ModelingComponents/Public/Scene/SceneCapturePhotoSet.h b/Engine/Plugins/Experimental/MeshModelingToolset/Source/ModelingComponents/Public/Scene/SceneCapturePhotoSet.h
--- a/Engine/Plugins/Experimental/MeshModelingToolset/Source/ModelingComponents/Public/Scene/SceneCapturePhotoSet.h
+++ b/Engine/Plugins/Experimental/MeshModelingToolset/Source/ModelingComponents/Public/Scene/SceneCapturePhotoSet.h
@@ -149,7 +149,7 @@ protected:
 	FSpatialPhotoSet3f WorldNormalPhotoSet;
 	FSpatialPhotoSet3f EmissivePhotoSet;
 
-	bool bWriteDebugImages;
+	bool bWriteDebugImages = false;
 	FString DebugImagesFolderName = TEXT("SceneCapturePhotoSet");
 };
 
