ApproximateActors - Texture2DBuilder - Use a compressed HDR format for emissive

diff --git a/Engine/Plugins/Experimental/MeshModelingToolset/Source/GeometryProcessingAdapters/Private/ApproximateActorsImpl.cpp b/Engine/Plugins/Experimental/MeshModelingToolset/Source/GeometryProcessingAdapters/Private/ApproximateActorsImpl.cpp
--- a/Engine/Plugins/Experimental/MeshModelingToolset/Source/GeometryProcessingAdapters/Private/ApproximateActorsImpl.cpp
+++ b/Engine/Plugins/Experimental/MeshModelingToolset/Source/GeometryProcessingAdapters/Private/ApproximateActorsImpl.cpp
@@ -325,6 +325,7 @@ static void BakeTexturesFromPhotoCapture(
 		{
 			BuildTexProgress.EnterProgressFrame(1.f);
 			GeneratedTextures.EmissiveMap = FTexture2DBuilder::BuildTextureFromImage(*EmissiveImage, FTexture2DBuilder::ETextureType::EmissiveHDR, false, false);
+			GeneratedTextures.EmissiveMap->CompressionSettings = TC_HDR_Compressed;
 		}
 		if (Options.bBakeNormalMap && NormalImage.IsValid())
 		{
