UVEditor: Fix incremental build error with data type conversion for FVector

#rnx

diff --git a/Engine/Plugins/Experimental/MeshModelingToolset/Source/ModelingComponents/Private/Drawing/MeshDebugDrawing.cpp b/Engine/Plugins/Experimental/MeshModelingToolset/Source/ModelingComponents/Private/Drawing/MeshDebugDrawing.cpp
--- a/Engine/Plugins/Experimental/MeshModelingToolset/Source/ModelingComponents/Private/Drawing/MeshDebugDrawing.cpp
+++ b/Engine/Plugins/Experimental/MeshModelingToolset/Source/ModelingComponents/Private/Drawing/MeshDebugDrawing.cpp
@@ -154,7 +154,7 @@ void MeshDebugDraw::DrawHierarchicalGrid(
 	}
 	
 	//Determine the center of the drawing area for the grid, snapping to grid positions
-	FVector GridOrigin(
+	FVector3f GridOrigin(
 		FMath::GridSnap(LocalFrame.Origin.X, GridScales[0]),
 		FMath::GridSnap(LocalFrame.Origin.Y, GridScales[0]),
 		LocalFrame.Origin.Z);
