ModelingTools: Make all ProxyLOD dependencies optional, when not available on the given platform (ie non-windows) the code paths just early-out or copy the input mesh. Remove all Win64 whitelisting for MeshModelingToolset, MeshLODToolset, ModelingToolsEditorMode.
#rnx

diff --git a/Engine/Plugins/Experimental/GeometryFlow/GeometryFlow.uplugin b/Engine/Plugins/Experimental/GeometryFlow/GeometryFlow.uplugin
--- a/Engine/Plugins/Experimental/GeometryFlow/GeometryFlow.uplugin
+++ b/Engine/Plugins/Experimental/GeometryFlow/GeometryFlow.uplugin
@@ -28,8 +28,7 @@
 		{
 			"Name": "GeometryFlowMeshProcessingEditor",
 			"Type": "Editor",
-			"LoadingPhase": "Default",
-			"WhitelistPlatforms": [ "Win64" ]
+			"LoadingPhase": "Default"
 		},		
 	],
 	"Plugins": [
diff --git a/Engine/Plugins/Experimental/MeshLODToolset/MeshLODToolset.uplugin b/Engine/Plugins/Experimental/MeshLODToolset/MeshLODToolset.uplugin
--- a/Engine/Plugins/Experimental/MeshLODToolset/MeshLODToolset.uplugin
+++ b/Engine/Plugins/Experimental/MeshLODToolset/MeshLODToolset.uplugin
@@ -18,8 +18,7 @@
 		{
 			"Name": "MeshLODToolset",
 			"Type": "Editor",
-			"LoadingPhase": "Default",
-			"WhitelistPlatforms": [ "Win64" ]
+			"LoadingPhase": "Default"
 		},
 	],
 	"Plugins": [
diff --git a/Engine/Plugins/Experimental/MeshLODToolset/Source/MeshLODToolset/Private/Tools/GenerateStaticMeshLODAssetTool.cpp b/Engine/Plugins/Experimental/MeshLODToolset/Source/MeshLODToolset/Private/Tools/GenerateStaticMeshLODAssetTool.cpp
--- a/Engine/Plugins/Experimental/MeshLODToolset/Source/MeshLODToolset/Private/Tools/GenerateStaticMeshLODAssetTool.cpp
+++ b/Engine/Plugins/Experimental/MeshLODToolset/Source/MeshLODToolset/Private/Tools/GenerateStaticMeshLODAssetTool.cpp
@@ -315,13 +315,13 @@ void UGenerateStaticMeshLODAssetTool::Setup()
 	});
 
 
-	PreviewWithBackgroundCompute->OnMeshUpdated.AddLambda([this](const UMeshOpPreviewWithBackgroundCompute* PreviewWithBackgroundCompute)
+	PreviewWithBackgroundCompute->OnMeshUpdated.AddLambda([this](const UMeshOpPreviewWithBackgroundCompute* PreviewCompute)
 	{
 		// GenerateProcess might be in use by an Op somewhere else
 		GenerateProcess->GraphEvalCriticalSection.Lock();
 
 		UE::Geometry::FMeshTangentsd Tangents = GenerateProcess->GetDerivedLOD0MeshTangents();
-		PreviewWithBackgroundCompute->PreviewMesh->EditMesh([&Tangents](FDynamicMesh3& Mesh)
+		PreviewCompute->PreviewMesh->EditMesh([&Tangents](FDynamicMesh3& Mesh)
 		{
 			if (Mesh.HasAttributes() && Tangents.GetTangents().Num() > 0 && Tangents.GetBitangents().Num() > 0)
 			{
diff --git a/Engine/Plugins/Experimental/MeshModelingToolset/MeshModelingToolset.uplugin b/Engine/Plugins/Experimental/MeshModelingToolset/MeshModelingToolset.uplugin
--- a/Engine/Plugins/Experimental/MeshModelingToolset/MeshModelingToolset.uplugin
+++ b/Engine/Plugins/Experimental/MeshModelingToolset/MeshModelingToolset.uplugin
@@ -23,8 +23,7 @@
 		{
 			"Name": "ModelingComponentsEditorOnly",
 			"Type": "Editor",
-			"LoadingPhase": "Default",
-			"WhitelistPlatforms": [ "Win64" ]
+			"LoadingPhase": "Default"
 		},		
 		{
 			"Name": "ModelingOperators",
@@ -34,8 +33,7 @@
 		{
 			"Name": "ModelingOperatorsEditorOnly",
 			"Type": "Editor",
-			"LoadingPhase": "Default",
-			"WhitelistPlatforms": [ "Win64" ]
+			"LoadingPhase": "Default"
 		},
 		{
 			"Name": "MeshModelingTools",
@@ -45,14 +43,12 @@
 		{
 			"Name": "MeshModelingToolsEditorOnly",
 			"Type": "Editor",
-			"LoadingPhase": "Default",
-			"WhitelistPlatforms": [ "Win64" ]
+			"LoadingPhase": "Default"
 		},
 		{
 			"Name": "GeometryProcessingAdapters",
 			"Type": "Editor",
-			"LoadingPhase": "Default",
-			"WhitelistPlatforms": [ "Win64" ]
+			"LoadingPhase": "Default"
 		}
 	],
 	"Plugins": [
@@ -62,7 +58,10 @@
 		},
 		{
 			"Name": "ProxyLODPlugin",
-			"Enabled": true
+			"Enabled": true,
+			"SupportedTargetPlatforms": [
+				"Win64"
+			]
 		},
 		{
 			"Name": "HairStrands",
diff --git a/Engine/Plugins/Experimental/MeshModelingToolset/Source/MeshModelingToolsEditorOnly/MeshModelingToolsEditorOnly.Build.cs b/Engine/Plugins/Experimental/MeshModelingToolset/Source/MeshModelingToolsEditorOnly/MeshModelingToolsEditorOnly.Build.cs
--- a/Engine/Plugins/Experimental/MeshModelingToolset/Source/MeshModelingToolsEditorOnly/MeshModelingToolsEditorOnly.Build.cs
+++ b/Engine/Plugins/Experimental/MeshModelingToolset/Source/MeshModelingToolsEditorOnly/MeshModelingToolsEditorOnly.Build.cs
@@ -39,7 +39,6 @@ public class MeshModelingToolsEditorOnly : ModuleRules
 				"ModelingComponentsEditorOnly",
 				"ModelingOperators",
 				"ModelingOperatorsEditorOnly",
-                "ProxyLODMeshReduction", // currently required to be public due to IVoxelBasedCSG API
 				// ... add other public dependencies that you statically link with here ...
 
 				"HairStrandsCore"		// required for Hair toolset
@@ -62,14 +61,22 @@ public class MeshModelingToolsEditorOnly : ModuleRules
 				"MeshBuilder",
 				"MeshUtilitiesCommon",  
 				"MeshReductionInterface", // for UE standard simplification 
-                "ProxyLODMeshReduction", // for mesh merging voxel-based csg
 				//"Slate",
 				//"SlateCore",
 				// ... add private dependencies that you statically link with here ...	
 			}
 			);
-		
-		
+
+		if (Target.Platform == UnrealTargetPlatform.Win64)
+		{
+			PrivateDependencyModuleNames.AddRange(
+				new string[]
+				{
+					"ProxyLODMeshReduction", // currently Win64-only
+				}
+				);
+		}
+
 		DynamicallyLoadedModuleNames.AddRange(
 			new string[]
 			{
diff --git a/Engine/Plugins/Experimental/MeshModelingToolset/Source/MeshModelingToolsEditorOnly/Private/MergeMeshesTool.cpp b/Engine/Plugins/Experimental/MeshModelingToolset/Source/MeshModelingToolsEditorOnly/Private/MergeMeshesTool.cpp
--- a/Engine/Plugins/Experimental/MeshModelingToolset/Source/MeshModelingToolsEditorOnly/Private/MergeMeshesTool.cpp
+++ b/Engine/Plugins/Experimental/MeshModelingToolset/Source/MeshModelingToolsEditorOnly/Private/MergeMeshesTool.cpp
@@ -183,15 +183,15 @@ TUniquePtr<FDynamicMeshOperator> UMergeMeshesTool::MakeNewOperator()
 
 void UMergeMeshesTool::CacheInputMeshes()
 {
-	InputMeshes = MakeShared<TArray<IVoxelBasedCSG::FPlacedMesh>>();
+	InputMeshes.Reset();
 
 	// Package the selected meshes and transforms for consumption by the CSGTool
 	for (int32 ComponentIdx = 0; ComponentIdx < Targets.Num(); ComponentIdx++)
 	{
-		IVoxelBasedCSG::FPlacedMesh PlacedMesh;
-		PlacedMesh.Mesh = TargetMeshProviderInterface(ComponentIdx)->GetMeshDescription();
-		PlacedMesh.Transform = TargetComponentInterface(ComponentIdx)->GetWorldTransform();
-		InputMeshes->Add(PlacedMesh);
+		UE::Geometry::FVoxelMergeMeshesOp::FInputMesh InputMesh;
+		InputMesh.Mesh = TargetMeshProviderInterface(ComponentIdx)->GetMeshDescription();
+		InputMesh.Transform = TargetComponentInterface(ComponentIdx)->GetWorldTransform();
+		InputMeshes.Add(InputMesh);
 	}
 }
 
diff --git a/Engine/Plugins/Experimental/MeshModelingToolset/Source/MeshModelingToolsEditorOnly/Private/VoxelCSGMeshesTool.cpp b/Engine/Plugins/Experimental/MeshModelingToolset/Source/MeshModelingToolsEditorOnly/Private/VoxelCSGMeshesTool.cpp
--- a/Engine/Plugins/Experimental/MeshModelingToolset/Source/MeshModelingToolsEditorOnly/Private/VoxelCSGMeshesTool.cpp
+++ b/Engine/Plugins/Experimental/MeshModelingToolset/Source/MeshModelingToolsEditorOnly/Private/VoxelCSGMeshesTool.cpp
@@ -178,15 +178,15 @@ TUniquePtr<FDynamicMeshOperator> UVoxelCSGMeshesTool::MakeNewOperator()
 
 void UVoxelCSGMeshesTool::CacheInputMeshes()
 {
-	InputMeshes = MakeShared<TArray<IVoxelBasedCSG::FPlacedMesh>>();
+	InputMeshes.Reset();
 
 	// Package the selected meshes and transforms for consumption by the CSGTool
 	for (int32 ComponentIdx = 0; ComponentIdx < Targets.Num(); ComponentIdx++)
 	{
-		IVoxelBasedCSG::FPlacedMesh PlacedMesh;
-		PlacedMesh.Mesh = TargetMeshProviderInterface(ComponentIdx)->GetMeshDescription();
-		PlacedMesh.Transform = TargetComponentInterface(ComponentIdx)->GetWorldTransform();
-		InputMeshes->Add(PlacedMesh);
+		UE::Geometry::FVoxelBooleanMeshesOp::FInputMesh InputMesh;
+		InputMesh.Mesh = TargetMeshProviderInterface(ComponentIdx)->GetMeshDescription();
+		InputMesh.Transform = TargetComponentInterface(ComponentIdx)->GetWorldTransform();
+		InputMeshes.Add(InputMesh);
 	}
 }
 
diff --git a/Engine/Plugins/Experimental/MeshModelingToolset/Source/MeshModelingToolsEditorOnly/Public/MergeMeshesTool.h b/Engine/Plugins/Experimental/MeshModelingToolset/Source/MeshModelingToolsEditorOnly/Public/MergeMeshesTool.h
--- a/Engine/Plugins/Experimental/MeshModelingToolset/Source/MeshModelingToolsEditorOnly/Public/MergeMeshesTool.h
+++ b/Engine/Plugins/Experimental/MeshModelingToolset/Source/MeshModelingToolsEditorOnly/Public/MergeMeshesTool.h
@@ -6,9 +6,9 @@
 #include "MultiSelectionTool.h"
 #include "InteractiveToolBuilder.h"
 #include "MeshOpPreviewHelpers.h"
-#include "ProxyLODVolume.h"
 #include "Properties/MeshStatisticsProperties.h"
 #include "PropertySets/OnAcceptProperties.h"
+#include "CompositionOps/VoxelMergeMeshesOp.h"
 #include "MergeMeshesTool.generated.h"
 
 /**
@@ -99,7 +99,7 @@ protected:
 protected:
 	UWorld* TargetWorld;
 
-	TSharedPtr<TArray<IVoxelBasedCSG::FPlacedMesh>> InputMeshes;
+	TArray<UE::Geometry::FVoxelMergeMeshesOp::FInputMesh> InputMeshes;
 	/** stash copies of the transforms and pointers to the meshes for consumption by merge Op*/
 	void CacheInputMeshes();
 
diff --git a/Engine/Plugins/Experimental/MeshModelingToolset/Source/MeshModelingToolsEditorOnly/Public/VoxelCSGMeshesTool.h b/Engine/Plugins/Experimental/MeshModelingToolset/Source/MeshModelingToolsEditorOnly/Public/VoxelCSGMeshesTool.h
--- a/Engine/Plugins/Experimental/MeshModelingToolset/Source/MeshModelingToolsEditorOnly/Public/VoxelCSGMeshesTool.h
+++ b/Engine/Plugins/Experimental/MeshModelingToolset/Source/MeshModelingToolsEditorOnly/Public/VoxelCSGMeshesTool.h
@@ -3,13 +3,12 @@
 #pragma once
 
 #include "CoreMinimal.h"
-#include "UObject/NoExportTypes.h"
 #include "InteractiveToolBuilder.h"
 #include "MeshOpPreviewHelpers.h"
 #include "MultiSelectionTool.h"
-#include "ProxyLODVolume.h"
 #include "Properties/MeshStatisticsProperties.h"
 #include "PropertySets/OnAcceptProperties.h"
+#include "CompositionOps/VoxelBooleanMeshesOp.h"
 
 #include "VoxelCSGMeshesTool.generated.h"
 
@@ -126,7 +125,7 @@ protected:
 protected:
 	UWorld* TargetWorld;
 
-	TSharedPtr<TArray<IVoxelBasedCSG::FPlacedMesh>> InputMeshes;
+	TArray<UE::Geometry::FVoxelBooleanMeshesOp::FInputMesh> InputMeshes;
 	/** stash copies of the transforms and pointers to the meshes for consumption by the CSG Op*/
 	void CacheInputMeshes();
 
diff --git a/Engine/Plugins/Experimental/MeshModelingToolset/Source/ModelingOperatorsEditorOnly/ModelingOperatorsEditorOnly.Build.cs b/Engine/Plugins/Experimental/MeshModelingToolset/Source/ModelingOperatorsEditorOnly/ModelingOperatorsEditorOnly.Build.cs
--- a/Engine/Plugins/Experimental/MeshModelingToolset/Source/ModelingOperatorsEditorOnly/ModelingOperatorsEditorOnly.Build.cs
+++ b/Engine/Plugins/Experimental/MeshModelingToolset/Source/ModelingOperatorsEditorOnly/ModelingOperatorsEditorOnly.Build.cs
@@ -34,7 +34,6 @@ public class ModelingOperatorsEditorOnly : ModuleRules
 				"DynamicMesh",
 				"MeshConversion",
 				"GeometryAlgorithms", // required for constrained Delaunay triangulation
-                "ProxyLODMeshReduction", // currently required to be public due to IVoxelBasedCSG API
 				"MeshUtilitiesCommon", // required by uvlayoutop
 				// ... add other public dependencies that you statically link with here ...
 			}
@@ -53,7 +52,16 @@ public class ModelingOperatorsEditorOnly : ModuleRules
 				"MeshUtilities",			// for tangents calculation
 			}
 			);
-		
+
+		if (Target.Platform == UnrealTargetPlatform.Win64)
+		{
+			PrivateDependencyModuleNames.AddRange(
+				new string[]
+				{
+					"ProxyLODMeshReduction", // currently Win64-only
+				}
+				);
+		}
 		
 		DynamicallyLoadedModuleNames.AddRange(
 			new string[]
diff --git a/Engine/Plugins/Experimental/MeshModelingToolset/Source/ModelingOperatorsEditorOnly/Private/CompositionOps/VoxelBooleanMeshesOp.cpp b/Engine/Plugins/Experimental/MeshModelingToolset/Source/ModelingOperatorsEditorOnly/Private/CompositionOps/VoxelBooleanMeshesOp.cpp
--- a/Engine/Plugins/Experimental/MeshModelingToolset/Source/ModelingOperatorsEditorOnly/Private/CompositionOps/VoxelBooleanMeshesOp.cpp
+++ b/Engine/Plugins/Experimental/MeshModelingToolset/Source/ModelingOperatorsEditorOnly/Private/CompositionOps/VoxelBooleanMeshesOp.cpp
@@ -11,11 +11,24 @@
 #include "MeshConstraintsUtil.h"
 #include "CleaningOps/EditNormalsOp.h"
 
+
+// ProxyLOD currently only available on Windows. On other platforms we will make this a no-op
+#if PLATFORM_WINDOWS
+#define HAVE_PROXYLOD 1
+#else
+#define HAVE_PROXYLOD 0
+#endif
+
+#if HAVE_PROXYLOD
+#include "ProxyLODVolume.h"
+#endif
+
 #include "ExplicitUseGeometryMathTypes.h"		// using UE::Geometry::(math types)
 using namespace UE::Geometry;
 
 void FVoxelBooleanMeshesOp::CalculateResult(FProgressCancel* Progress)
 {
+#if HAVE_PROXYLOD
 	FMeshDescription ResultMeshDescription;
 	FStaticMeshAttributes Attributes(ResultMeshDescription);
 	Attributes.Register();
@@ -49,8 +62,17 @@ void FVoxelBooleanMeshesOp::CalculateResult(FProgressCancel* Progress)
 	TUniquePtr<IVoxelBasedCSG> VoxelCSGTool = IVoxelBasedCSG::CreateCSGTool(VoxelSizeD);
 	FVector MergedOrigin;
 
-	IVoxelBasedCSG::FPlacedMesh& A = (*InputMeshArray)[0];
-	IVoxelBasedCSG::FPlacedMesh& B = (*InputMeshArray)[1];
+	TArray<IVoxelBasedCSG::FPlacedMesh> PlacedMeshes;
+	for (const FInputMesh& InputMesh : InputMeshArray)
+	{
+		IVoxelBasedCSG::FPlacedMesh PlacedMesh;
+		PlacedMesh.Mesh = InputMesh.Mesh;
+		PlacedMesh.Transform = InputMesh.Transform;
+		PlacedMeshes.Add(PlacedMesh);
+	}
+
+	IVoxelBasedCSG::FPlacedMesh& A = PlacedMeshes[0];
+	IVoxelBasedCSG::FPlacedMesh& B = PlacedMeshes[1];
 
 	switch (Operation)
 	{
@@ -112,7 +134,6 @@ void FVoxelBooleanMeshesOp::CalculateResult(FProgressCancel* Progress)
 		}
 	}
 
-
 	// Convert to dynamic mesh
 	FMeshDescriptionToDynamicMesh Converter;
 	Converter.Convert(&ResultMeshDescription, *ResultMesh);
@@ -142,7 +163,6 @@ void FVoxelBooleanMeshesOp::CalculateResult(FProgressCancel* Progress)
 
 	if (bFixNormals)
 	{
-
 		TSharedPtr<FDynamicMesh3, ESPMode::ThreadSafe> OpResultMesh(ExtractResult().Release()); // moved the unique pointer
 
 		// Recompute the normals
@@ -163,14 +183,20 @@ void FVoxelBooleanMeshesOp::CalculateResult(FProgressCancel* Progress)
 		ResultMesh = EditNormalsOp.ExtractResult(); // return the edit normals operator copy to this tool.
 	}
 
+#else	// HAVE_PROXYLOD
+
+	FMeshDescriptionToDynamicMesh Converter;
+	Converter.Convert(InputMeshArray[0].Mesh, *ResultMesh);
+	ResultTransform = FTransform3d(InputMeshArray[0].Transform);
+
+#endif	// HAVE_PROXYLOD
+
 }
 
 float FVoxelBooleanMeshesOp::ComputeVoxelSize() const 
 { 
-	const TArray<IVoxelBasedCSG::FPlacedMesh>& PlacedMeshes = *(InputMeshArray);
-
 	float Size = 0.f;
-	for (int32 i = 0; i < PlacedMeshes.Num(); ++i)
+	for (int32 i = 0; i < InputMeshArray.Num(); ++i)
 	{
 		//Bounding box
 		FVector BBoxMin(FLT_MAX, FLT_MAX, FLT_MAX);
@@ -178,8 +204,8 @@ float FVoxelBooleanMeshesOp::ComputeVoxelSize() const
 
 		// Use the scale define by the transform.  We don't care about actual placement
 		// in the world for this.
-		FVector Scale = PlacedMeshes[i].Transform.GetScale3D();
-		const FMeshDescription&  MeshDescription = *PlacedMeshes[i].Mesh;
+		FVector Scale = InputMeshArray[i].Transform.GetScale3D();
+		const FMeshDescription&  MeshDescription = *InputMeshArray[i].Mesh;
 
 		TArrayView<const FVector3f> VertexPositions = MeshDescription.GetVertexPositions().GetRawArray();
 		for (const FVertexID VertexID : MeshDescription.Vertices().GetElementIDs())
diff --git a/Engine/Plugins/Experimental/MeshModelingToolset/Source/ModelingOperatorsEditorOnly/Private/CompositionOps/VoxelMergeMeshesOp.cpp b/Engine/Plugins/Experimental/MeshModelingToolset/Source/ModelingOperatorsEditorOnly/Private/CompositionOps/VoxelMergeMeshesOp.cpp
--- a/Engine/Plugins/Experimental/MeshModelingToolset/Source/ModelingOperatorsEditorOnly/Private/CompositionOps/VoxelMergeMeshesOp.cpp
+++ b/Engine/Plugins/Experimental/MeshModelingToolset/Source/ModelingOperatorsEditorOnly/Private/CompositionOps/VoxelMergeMeshesOp.cpp
@@ -11,11 +11,24 @@
 #include "MeshConstraintsUtil.h"
 #include "CleaningOps/EditNormalsOp.h"
 
+// ProxyLOD currently only available on Windows. On other platforms we will make this a no-op
+#if PLATFORM_WINDOWS
+#define HAVE_PROXYLOD 1
+#else
+#define HAVE_PROXYLOD 0
+#endif
+
+#if HAVE_PROXYLOD
+#include "ProxyLODVolume.h"
+#endif
+
 #include "ExplicitUseGeometryMathTypes.h"		// using UE::Geometry::(math types)
 using namespace UE::Geometry;
 
 void FVoxelMergeMeshesOp::CalculateResult(FProgressCancel* Progress)
 {
+#if HAVE_PROXYLOD
+
 	FMeshDescription MergedMeshesDescription;
 	FStaticMeshAttributes Attributes(MergedMeshesDescription);
 	Attributes.Register();
@@ -29,11 +42,20 @@ void FVoxelMergeMeshesOp::CalculateResult(FProgressCancel* Progress)
 	// Create CSGTool and merge the meshes.
 	TUniquePtr<IVoxelBasedCSG> VoxelCSGTool = IVoxelBasedCSG::CreateCSGTool(VoxelSizeD);
 
+	TArray<IVoxelBasedCSG::FPlacedMesh> PlacedMeshes;
+	for (const FInputMesh& InputMesh : InputMeshArray)
+	{
+		IVoxelBasedCSG::FPlacedMesh PlacedMesh;
+		PlacedMesh.Mesh = InputMesh.Mesh;
+		PlacedMesh.Transform = InputMesh.Transform;
+		PlacedMeshes.Add(PlacedMesh);
+	}
+
 	// world space units.
 	const double MaxNarrowBand = 3 * VoxelSizeD;
 	const double MaxIsoOffset = 2 * VoxelSizeD;
 	const double CSGIsoSurface = FMath::Clamp(IsoSurfaceD, 0., MaxIsoOffset); // the interior distance values maybe messed up when doing a union.
-	FVector MergedOrigin = VoxelCSGTool->ComputeUnion(*InputMeshArray, MergedMeshesDescription, AdaptivityD, CSGIsoSurface);
+	FVector MergedOrigin = VoxelCSGTool->ComputeUnion(PlacedMeshes, MergedMeshesDescription, AdaptivityD, CSGIsoSurface);
 	ResultTransform = FTransform3d(MergedOrigin);
 
 	if (Progress && Progress->Cancelled())
@@ -122,15 +144,22 @@ void FVoxelMergeMeshesOp::CalculateResult(FProgressCancel* Progress)
 		ResultMesh = EditNormalsOp.ExtractResult(); // return the edit normals operator copy to this tool.
 	}
 
+
+#else	// HAVE_PROXYLOD
+
+	FMeshDescriptionToDynamicMesh Converter;
+	Converter.Convert(InputMeshArray[0].Mesh, *ResultMesh);
+	ResultTransform = FTransform3d(InputMeshArray[0].Transform);
+
+#endif	// HAVE_PROXYLOD
+
 }
 
 
 float FVoxelMergeMeshesOp::ComputeVoxelSize() const
 {
-	const TArray<IVoxelBasedCSG::FPlacedMesh>& PlacedMeshes = *(InputMeshArray);
-
 	float Size = 0.f;
-	for (int32 i = 0; i < PlacedMeshes.Num(); ++i)
+	for (int32 i = 0; i < InputMeshArray.Num(); ++i)
 	{
 		//Bounding box
 		FVector BBoxMin(FLT_MAX, FLT_MAX, FLT_MAX);
@@ -138,8 +167,8 @@ float FVoxelMergeMeshesOp::ComputeVoxelSize() const
 
 		// Use the scale define by the transform.  We don't care about actual placement
 		// in the world for this.
-		FVector Scale = PlacedMeshes[i].Transform.GetScale3D();
-		const FMeshDescription&  MeshDescription = *PlacedMeshes[i].Mesh;
+		FVector Scale = InputMeshArray[i].Transform.GetScale3D();
+		const FMeshDescription&  MeshDescription = *InputMeshArray[i].Mesh;
 
 		TArrayView<const FVector3f> VertexPositions = MeshDescription.GetVertexPositions().GetRawArray();
 		for (const FVertexID VertexID : MeshDescription.Vertices().GetElementIDs())
diff --git a/Engine/Plugins/Experimental/MeshModelingToolset/Source/ModelingOperatorsEditorOnly/Private/ParameterizationOps/ParameterizeMeshOp.cpp b/Engine/Plugins/Experimental/MeshModelingToolset/Source/ModelingOperatorsEditorOnly/Private/ParameterizationOps/ParameterizeMeshOp.cpp
--- a/Engine/Plugins/Experimental/MeshModelingToolset/Source/ModelingOperatorsEditorOnly/Private/ParameterizationOps/ParameterizeMeshOp.cpp
+++ b/Engine/Plugins/Experimental/MeshModelingToolset/Source/ModelingOperatorsEditorOnly/Private/ParameterizationOps/ParameterizeMeshOp.cpp
@@ -3,7 +3,6 @@
 #include "ParameterizationOps/ParameterizeMeshOp.h"
 
 #include "DynamicMesh/DynamicMeshAttributeSet.h"
-#include "ProxyLODParameterization.h"
 #include "Selections/MeshConnectedComponents.h"
 
 #include "Parameterization/MeshLocalParam.h"
@@ -13,6 +12,17 @@
 #include "DynamicSubmesh3.h"
 #include "XAtlasWrapper.h"
 
+// ProxyLOD currently only available on Windows. On other platforms we will make this a no-op
+#if PLATFORM_WINDOWS
+#define HAVE_PROXYLOD 1
+#else
+#define HAVE_PROXYLOD 0
+#endif
+
+#if HAVE_PROXYLOD
+#include "ProxyLODParameterization.h"
+#endif
+
 #include "ExplicitUseGeometryMathTypes.h"		// using UE::Geometry::(math types)
 using namespace UE::Geometry;
 
@@ -184,7 +194,11 @@ bool FParameterizeMeshOp::ComputeUVs(FDynamicMesh3& Mesh,  TFunction<bool(float)
 	int32 MaxChartNumber = NumCharts;
 
 	bool bSuccess = true;
+#if HAVE_PROXYLOD
 	bool bUseXAtlas = (Method == EParamOpBackend::XAtlas);
+#else
+	bool bUseXAtlas = true;
+#endif
 
 	if (bUseXAtlas)
 	{
@@ -201,11 +215,15 @@ bool FParameterizeMeshOp::ComputeUVs(FDynamicMesh3& Mesh,  TFunction<bool(float)
 	}
 	else
 	{
+#if HAVE_PROXYLOD
 		TUniquePtr<IProxyLODParameterization> ParameterizationTool = IProxyLODParameterization::CreateTool();
 		bSuccess = ParameterizationTool->GenerateUVs(Width, Height, Gutter, LinearMesh.VertexBuffer, 
 													 LinearMesh.IndexBuffer, LinearMesh.AdjacencyBuffer, Interrupter, 
 													 UVVertexBuffer, UVIndexBuffer, VertexRemapArray, MaxStretch, 
 													 MaxChartNumber);
+#else
+		ensureMsgf(false, TEXT("ProxyLOD not available, this should not be called"));
+#endif
 	}
 	
 
diff --git a/Engine/Plugins/Experimental/MeshModelingToolset/Source/ModelingOperatorsEditorOnly/Public/CompositionOps/VoxelBooleanMeshesOp.h b/Engine/Plugins/Experimental/MeshModelingToolset/Source/ModelingOperatorsEditorOnly/Public/CompositionOps/VoxelBooleanMeshesOp.h
--- a/Engine/Plugins/Experimental/MeshModelingToolset/Source/ModelingOperatorsEditorOnly/Public/CompositionOps/VoxelBooleanMeshesOp.h
+++ b/Engine/Plugins/Experimental/MeshModelingToolset/Source/ModelingOperatorsEditorOnly/Public/CompositionOps/VoxelBooleanMeshesOp.h
@@ -5,8 +5,8 @@
 #include "CoreMinimal.h"
 #include "Util/ProgressCancel.h"
 #include "ModelingOperators.h"
-#include "ProxyLODVolume.h"
 
+struct FMeshDescription;
 
 namespace UE
 {
@@ -27,8 +27,14 @@ public:
 		Union = 3
 	};
 
+	struct FInputMesh
+	{
+		const FMeshDescription* Mesh;
+		FTransform Transform;
+	};
+
 	// inputs
-	TSharedPtr<TArray<IVoxelBasedCSG::FPlacedMesh>> InputMeshArray;
+	TArray<FInputMesh> InputMeshArray;
 	int32 VoxelCount = 128;
 	double VoxelSizeD = 1.0;
 	double AdaptivityD = 0;
diff --git a/Engine/Plugins/Experimental/MeshModelingToolset/Source/ModelingOperatorsEditorOnly/Public/CompositionOps/VoxelMergeMeshesOp.h b/Engine/Plugins/Experimental/MeshModelingToolset/Source/ModelingOperatorsEditorOnly/Public/CompositionOps/VoxelMergeMeshesOp.h
--- a/Engine/Plugins/Experimental/MeshModelingToolset/Source/ModelingOperatorsEditorOnly/Public/CompositionOps/VoxelMergeMeshesOp.h
+++ b/Engine/Plugins/Experimental/MeshModelingToolset/Source/ModelingOperatorsEditorOnly/Public/CompositionOps/VoxelMergeMeshesOp.h
@@ -5,8 +5,8 @@
 #include "CoreMinimal.h"
 #include "Util/ProgressCancel.h"
 #include "ModelingOperators.h"
-#include "ProxyLODVolume.h"
 
+struct FMeshDescription;
 
 namespace UE
 {
@@ -19,8 +19,14 @@ class MODELINGOPERATORSEDITORONLY_API FVoxelMergeMeshesOp : public FDynamicMeshO
 public:
 	virtual ~FVoxelMergeMeshesOp() {}
 
+	struct FInputMesh
+	{
+		const FMeshDescription* Mesh;
+		FTransform Transform;
+	};
+
 	// inputs
-	TSharedPtr<TArray<IVoxelBasedCSG::FPlacedMesh>> InputMeshArray;
+	TArray<FInputMesh> InputMeshArray;
 	int32 VoxelCount = 128;
 	double VoxelSizeD = 1.0;
 	double AdaptivityD = 0;
diff --git a/Engine/Plugins/Experimental/MeshModelingToolset/Source/ModelingOperatorsEditorOnly/Public/CuttingOps/EmbedPolygonsOp.h b/Engine/Plugins/Experimental/MeshModelingToolset/Source/ModelingOperatorsEditorOnly/Public/CuttingOps/EmbedPolygonsOp.h
--- a/Engine/Plugins/Experimental/MeshModelingToolset/Source/ModelingOperatorsEditorOnly/Public/CuttingOps/EmbedPolygonsOp.h
+++ b/Engine/Plugins/Experimental/MeshModelingToolset/Source/ModelingOperatorsEditorOnly/Public/CuttingOps/EmbedPolygonsOp.h
@@ -5,7 +5,6 @@
 #include "CoreMinimal.h"
 #include "Util/ProgressCancel.h"
 #include "ModelingOperators.h"
-#include "ProxyLODVolume.h"
 #include "DynamicMesh/DynamicMesh3.h"
 #include "Polygon2.h"
 
diff --git a/Engine/Plugins/Experimental/ModelingToolsEditorMode/ModelingToolsEditorMode.uplugin b/Engine/Plugins/Experimental/ModelingToolsEditorMode/ModelingToolsEditorMode.uplugin
--- a/Engine/Plugins/Experimental/ModelingToolsEditorMode/ModelingToolsEditorMode.uplugin
+++ b/Engine/Plugins/Experimental/ModelingToolsEditorMode/ModelingToolsEditorMode.uplugin
@@ -17,8 +17,7 @@
 		{
 			"Name": "ModelingToolsEditorMode",
 			"Type": "Editor",
-			"LoadingPhase": "Default",
-			"WhitelistPlatforms": [ "Win64" ]
+			"LoadingPhase": "Default"
 		}
 	],
 	"Plugins": [
diff --git a/Engine/Plugins/Experimental/ProxyLODPlugin/ProxyLODPlugin.uplugin b/Engine/Plugins/Experimental/ProxyLODPlugin/ProxyLODPlugin.uplugin
--- a/Engine/Plugins/Experimental/ProxyLODPlugin/ProxyLODPlugin.uplugin
+++ b/Engine/Plugins/Experimental/ProxyLODPlugin/ProxyLODPlugin.uplugin
@@ -14,6 +14,9 @@
   "CanContainContent" : false,
   "IsBetaVersion" : true,
   "Installed" : false,
+  "SupportedTargetPlatforms": [
+    "Win64"
+  ],
   "Modules" :
 	[
 		{
