Revert whitelisting changes for ModelingMode plugins and ProxyLODPlugin
#rnx

diff --git a/Engine/Plugins/Experimental/GeometryFlow/GeometryFlow.uplugin b/Engine/Plugins/Experimental/GeometryFlow/GeometryFlow.uplugin
--- a/Engine/Plugins/Experimental/GeometryFlow/GeometryFlow.uplugin
+++ b/Engine/Plugins/Experimental/GeometryFlow/GeometryFlow.uplugin
@@ -28,7 +28,8 @@
 		{
 			"Name": "GeometryFlowMeshProcessingEditor",
 			"Type": "Editor",
-			"LoadingPhase": "Default"
+			"LoadingPhase": "Default",
+			"WhitelistPlatforms": [ "Win64" ]
 		},		
 	],
 	"Plugins": [
diff --git a/Engine/Plugins/Experimental/MeshLODToolset/MeshLODToolset.uplugin b/Engine/Plugins/Experimental/MeshLODToolset/MeshLODToolset.uplugin
--- a/Engine/Plugins/Experimental/MeshLODToolset/MeshLODToolset.uplugin
+++ b/Engine/Plugins/Experimental/MeshLODToolset/MeshLODToolset.uplugin
@@ -18,7 +18,8 @@
 		{
 			"Name": "MeshLODToolset",
 			"Type": "Editor",
-			"LoadingPhase": "Default"
+			"LoadingPhase": "Default",
+			"WhitelistPlatforms": [ "Win64" ]
 		},
 	],
 	"Plugins": [
diff --git a/Engine/Plugins/Experimental/MeshModelingToolset/MeshModelingToolset.uplugin b/Engine/Plugins/Experimental/MeshModelingToolset/MeshModelingToolset.uplugin
--- a/Engine/Plugins/Experimental/MeshModelingToolset/MeshModelingToolset.uplugin
+++ b/Engine/Plugins/Experimental/MeshModelingToolset/MeshModelingToolset.uplugin
@@ -23,7 +23,8 @@
 		{
 			"Name": "ModelingComponentsEditorOnly",
 			"Type": "Editor",
-			"LoadingPhase": "Default"
+			"LoadingPhase": "Default",
+			"WhitelistPlatforms": [ "Win64" ]
 		},		
 		{
 			"Name": "ModelingOperators",
@@ -33,7 +34,8 @@
 		{
 			"Name": "ModelingOperatorsEditorOnly",
 			"Type": "Editor",
-			"LoadingPhase": "Default"
+			"LoadingPhase": "Default",
+			"WhitelistPlatforms": [ "Win64" ]
 		},
 		{
 			"Name": "MeshModelingTools",
@@ -43,12 +45,14 @@
 		{
 			"Name": "MeshModelingToolsEditorOnly",
 			"Type": "Editor",
-			"LoadingPhase": "Default"
+			"LoadingPhase": "Default",
+			"WhitelistPlatforms": [ "Win64" ]
 		},
 		{
 			"Name": "GeometryProcessingAdapters",
 			"Type": "Editor",
-			"LoadingPhase": "Default"
+			"LoadingPhase": "Default",
+			"WhitelistPlatforms": [ "Win64" ]
 		}
 	],
 	"Plugins": [
@@ -58,10 +62,7 @@
 		},
 		{
 			"Name": "ProxyLODPlugin",
-			"Enabled": true,
-			"SupportedTargetPlatforms": [
-				"Win64"
-			]
+			"Enabled": true
 		},
 		{
 			"Name": "HairStrands",
diff --git a/Engine/Plugins/Experimental/ModelingToolsEditorMode/ModelingToolsEditorMode.uplugin b/Engine/Plugins/Experimental/ModelingToolsEditorMode/ModelingToolsEditorMode.uplugin
--- a/Engine/Plugins/Experimental/ModelingToolsEditorMode/ModelingToolsEditorMode.uplugin
+++ b/Engine/Plugins/Experimental/ModelingToolsEditorMode/ModelingToolsEditorMode.uplugin
@@ -17,7 +17,8 @@
 		{
 			"Name": "ModelingToolsEditorMode",
 			"Type": "Editor",
-			"LoadingPhase": "Default"
+			"LoadingPhase": "Default",
+			"WhitelistPlatforms": [ "Win64" ]
 		}
 	],
 	"Plugins": [
diff --git a/Engine/Plugins/Experimental/ProxyLODPlugin/ProxyLODPlugin.uplugin b/Engine/Plugins/Experimental/ProxyLODPlugin/ProxyLODPlugin.uplugin
--- a/Engine/Plugins/Experimental/ProxyLODPlugin/ProxyLODPlugin.uplugin
+++ b/Engine/Plugins/Experimental/ProxyLODPlugin/ProxyLODPlugin.uplugin
@@ -14,9 +14,6 @@
   "CanContainContent" : false,
   "IsBetaVersion" : true,
   "Installed" : false,
-  "SupportedTargetPlatforms": [
-    "Win64"
-  ],
   "Modules" :
 	[
 		{
