Wrapped ObjectPtr upgrade refresh for Engine + ShooterGame

diff --git a/Engine/Plugins/Animation/ControlRig/Source/ControlRig/Public/AnimNode_ControlRig.h b/Engine/Plugins/Animation/ControlRig/Source/ControlRig/Public/AnimNode_ControlRig.h
--- a/Engine/Plugins/Animation/ControlRig/Source/ControlRig/Public/AnimNode_ControlRig.h
+++ b/Engine/Plugins/Animation/ControlRig/Source/ControlRig/Public/AnimNode_ControlRig.h
@@ -43,7 +43,7 @@ private:
 	TSubclassOf<UControlRig> ControlRigClass;
 
 	UPROPERTY(transient)
-	UControlRig* ControlRig;
+	TObjectPtr<UControlRig> ControlRig;
 
 	// alpha value handler
 	UPROPERTY(EditAnywhere, Category = Settings, meta = (PinShownByDefault))
diff --git a/Engine/Plugins/Animation/ControlRig/Source/ControlRig/Public/ControlRig.h b/Engine/Plugins/Animation/ControlRig/Source/ControlRig/Public/ControlRig.h
--- a/Engine/Plugins/Animation/ControlRig/Source/ControlRig/Public/ControlRig.h
+++ b/Engine/Plugins/Animation/ControlRig/Source/ControlRig/Public/ControlRig.h
@@ -405,10 +405,10 @@ public:
 private:
 
 	UPROPERTY()
-	URigVM* VM;
+	TObjectPtr<URigVM> VM;
 
 	UPROPERTY()
-	URigHierarchy *DynamicHierarchy;
+	TObjectPtr<URigHierarchy> DynamicHierarchy;
 
 	UPROPERTY()
 	TSoftObjectPtr<UControlRigGizmoLibrary> GizmoLibrary;
@@ -463,7 +463,7 @@ private:
 
 	/** The registry to access data source */
 	UPROPERTY(transient)
-	UAnimationDataSourceRegistry* DataSourceRegistry;
+	TObjectPtr<UAnimationDataSourceRegistry> DataSourceRegistry;
 
 	/** The event name used during an update */
 	UPROPERTY(transient)
@@ -543,7 +543,7 @@ private:
 	const FRigInfluenceMap* FindInfluenceMap(const FName& InEventName);
 
 	UPROPERTY(transient, BlueprintGetter = GetInteractionRig, BlueprintSetter = SetInteractionRig, Category = "Interaction")
-	UControlRig* InteractionRig;
+	TObjectPtr<UControlRig> InteractionRig;
 
 	UPROPERTY(EditAnywhere, transient, BlueprintGetter = GetInteractionRigClass, BlueprintSetter = SetInteractionRigClass, Category = "Interaction", Meta=(DisplayName="Interaction Rig"))
 	TSubclassOf<UControlRig> InteractionRigClass;
@@ -581,7 +581,7 @@ protected
 	:
 	/** Array of user data stored with the asset */
 	UPROPERTY(EditAnywhere, AdvancedDisplay, Instanced, Category = "Default")
-	TArray<UAssetUserData*> AssetUserData;
+	TArray<TObjectPtr<UAssetUserData>> AssetUserData;
 
 private:
 
@@ -663,7 +663,7 @@ private:
 #if WITH_EDITORONLY_DATA
 
 	UPROPERTY(transient)
-	URigVM* VMSnapshotBeforeExecution;
+	TObjectPtr<URigVM> VMSnapshotBeforeExecution;
 
 	/** The current execution mode */
 	UPROPERTY(transient)
diff --git a/Engine/Plugins/Animation/ControlRig/Source/ControlRig/Public/ControlRigComponent.h b/Engine/Plugins/Animation/ControlRig/Source/ControlRig/Public/ControlRigComponent.h
--- a/Engine/Plugins/Animation/ControlRig/Source/ControlRig/Public/ControlRigComponent.h
+++ b/Engine/Plugins/Animation/ControlRig/Source/ControlRig/Public/ControlRigComponent.h
@@ -106,7 +106,7 @@ struct FControlRigComponentMappedElement
 	EControlRigComponentSpace Space;
 
 	UPROPERTY(transient)
-	USceneComponent* SceneComponent;
+	TObjectPtr<USceneComponent> SceneComponent;
 
 	UPROPERTY(transient)
 	int32 ElementIndex;
@@ -130,7 +130,7 @@ struct FControlRigComponentMappedComponent
 	{}
 
 	UPROPERTY(EditAnywhere, BlueprintReadWrite, Category = "Mapping")
-	USceneComponent* Component;
+	TObjectPtr<USceneComponent> Component;
 
 	UPROPERTY(EditAnywhere, BlueprintReadWrite, Category = "Mapping")
 	FName ElementName;
@@ -632,7 +632,7 @@ private:
 #endif
 
 	UPROPERTY(transient)
-	UControlRig* ControlRig;
+	TObjectPtr<UControlRig> ControlRig;
 
 	TMap<USkeletalMeshComponent*, FCachedSkeletalMeshComponentSettings> CachedSkeletalMeshComponentSettings;
 
diff --git a/Engine/Plugins/Animation/ControlRig/Source/ControlRig/Public/ControlRigControlActor.h b/Engine/Plugins/Animation/ControlRig/Source/ControlRig/Public/ControlRigControlActor.h
--- a/Engine/Plugins/Animation/ControlRig/Source/ControlRig/Public/ControlRigControlActor.h
+++ b/Engine/Plugins/Animation/ControlRig/Source/ControlRig/Public/ControlRigControlActor.h
@@ -36,7 +36,7 @@ public:
 #endif
 
 	UPROPERTY(BlueprintReadWrite, EditAnywhere, Category = "Control Actor")
-	class AActor* ActorToTrack;
+	TObjectPtr<class AActor> ActorToTrack;
 
 	UPROPERTY(BlueprintReadWrite, EditAnywhere, Category = "Control Actor")
 	TSubclassOf<UControlRig> ControlRigClass;
@@ -48,7 +48,7 @@ public:
 	bool bIsSelectable;
 
 	UPROPERTY(BlueprintReadWrite, EditAnywhere, Category = "Materials")
-	UMaterialInterface* MaterialOverride;
+	TObjectPtr<UMaterialInterface> MaterialOverride;
 
 	UPROPERTY(BlueprintReadWrite, EditAnywhere, Category = "Materials")
 	FString ColorParameter;
@@ -65,10 +65,10 @@ public:
 private:
 
 	UPROPERTY()
-	class USceneComponent* ActorRootComponent;
+	TObjectPtr<class USceneComponent> ActorRootComponent;
 
 	UPROPERTY(transient)
-	UControlRig * ControlRig;
+	TObjectPtr<UControlRig>  ControlRig;
 
 	UPROPERTY(transient)
 	TArray<FName> ControlNames;
@@ -77,10 +77,10 @@ private:
 	TArray<FTransform> GizmoTransforms;
 
 	UPROPERTY(transient)
-	TArray<UStaticMeshComponent*> Components;
+	TArray<TObjectPtr<UStaticMeshComponent>> Components;
 
 	UPROPERTY(transient)
-	TArray<UMaterialInstanceDynamic*> Materials;
+	TArray<TObjectPtr<UMaterialInstanceDynamic>> Materials;
 
 	UPROPERTY(transient)
 	FName ColorParameterName;
diff --git a/Engine/Plugins/Animation/ControlRig/Source/ControlRig/Public/ControlRigGizmoActor.h b/Engine/Plugins/Animation/ControlRig/Source/ControlRig/Public/ControlRigGizmoActor.h
--- a/Engine/Plugins/Animation/ControlRig/Source/ControlRig/Public/ControlRigGizmoActor.h
+++ b/Engine/Plugins/Animation/ControlRig/Source/ControlRig/Public/ControlRigGizmoActor.h
@@ -55,11 +55,11 @@ public:
 
 	// this is the one holding transform for the controls
 	UPROPERTY()
-	class USceneComponent* ActorRootComponent;
+	TObjectPtr<class USceneComponent> ActorRootComponent;
 
 	// this is visual representation of the transform
 	UPROPERTY(Category = StaticMesh, VisibleAnywhere, BlueprintReadOnly, meta = (ExposeFunctionCategories = "Mesh,Rendering,Physics,Components|StaticMesh", AllowPrivateAccess = "true"))
-	class UStaticMeshComponent* StaticMeshComponent;
+	TObjectPtr<class UStaticMeshComponent> StaticMeshComponent;
 
 	// the name of the control this actor is referencing
 	UPROPERTY()
diff --git a/Engine/Plugins/Animation/ControlRig/Source/ControlRig/Public/ControlRigValidationPass.h b/Engine/Plugins/Animation/ControlRig/Source/ControlRig/Public/ControlRigValidationPass.h
--- a/Engine/Plugins/Animation/ControlRig/Source/ControlRig/Public/ControlRigValidationPass.h
+++ b/Engine/Plugins/Animation/ControlRig/Source/ControlRig/Public/ControlRigValidationPass.h
@@ -63,7 +63,7 @@ class CONTROLRIG_API UControlRigValidator : public UObject
 private:
 
 	UPROPERTY()
-	TArray<UControlRigValidationPass*> Passes;
+	TArray<TObjectPtr<UControlRigValidationPass>> Passes;
 
 	void OnControlRigInitialized(UControlRig* Subject, EControlRigState State, const FName& EventName);
 	void OnControlRigExecuted(UControlRig* Subject, EControlRigState State, const FName& EventName);
diff --git a/Engine/Plugins/Animation/ControlRig/Source/ControlRig/Public/Rigs/AdditiveControlRig.h b/Engine/Plugins/Animation/ControlRig/Source/ControlRig/Public/Rigs/AdditiveControlRig.h
--- a/Engine/Plugins/Animation/ControlRig/Source/ControlRig/Public/Rigs/AdditiveControlRig.h
+++ b/Engine/Plugins/Animation/ControlRig/Source/ControlRig/Public/Rigs/AdditiveControlRig.h
@@ -30,7 +30,7 @@ private:
 	TArray<FRigUnit_AddBoneTransform> AddBoneRigUnits;
 
 	UPROPERTY(transient)
-	URigHierarchyController* Controller;
+	TObjectPtr<URigHierarchyController> Controller;
 
 	/** Create RigElements - bone hierarchy and curves - from incoming skeleton */
 	void CreateRigElements(const USkeletalMesh* InReferenceMesh);
diff --git a/Engine/Plugins/Animation/ControlRig/Source/ControlRig/Public/Rigs/FKControlRig.h b/Engine/Plugins/Animation/ControlRig/Source/ControlRig/Public/Rigs/FKControlRig.h
--- a/Engine/Plugins/Animation/ControlRig/Source/ControlRig/Public/Rigs/FKControlRig.h
+++ b/Engine/Plugins/Animation/ControlRig/Source/ControlRig/Public/Rigs/FKControlRig.h
@@ -71,7 +71,7 @@ private:
 	EControlRigFKRigExecuteMode ApplyMode;
 
 	UPROPERTY(transient)
-	URigHierarchyController* Controller;
+	TObjectPtr<URigHierarchyController> Controller;
 
 	friend class FControlRigInteractionTest;
 };
\ No newline at end of file
diff --git a/Engine/Plugins/Animation/ControlRig/Source/ControlRig/Public/Rigs/RigControlHierarchy.h b/Engine/Plugins/Animation/ControlRig/Source/ControlRig/Public/Rigs/RigControlHierarchy.h
--- a/Engine/Plugins/Animation/ControlRig/Source/ControlRig/Public/Rigs/RigControlHierarchy.h
+++ b/Engine/Plugins/Animation/ControlRig/Source/ControlRig/Public/Rigs/RigControlHierarchy.h
@@ -152,7 +152,7 @@ struct CONTROLRIG_API FRigControl : public FRigElement
 
 	/** If the control is transient and only visible in the control rig editor */
 	UPROPERTY(BlueprintReadOnly, EditAnywhere, Category = Control)
-	UEnum* ControlEnum;
+	TObjectPtr<UEnum> ControlEnum;
 
 	FORCEINLINE_DEBUGGABLE virtual ERigElementType GetElementType() const override
 	{
diff --git a/Engine/Plugins/Animation/ControlRig/Source/ControlRig/Public/Rigs/RigHierarchyController.h b/Engine/Plugins/Animation/ControlRig/Source/ControlRig/Public/Rigs/RigHierarchyController.h
--- a/Engine/Plugins/Animation/ControlRig/Source/ControlRig/Public/Rigs/RigHierarchyController.h
+++ b/Engine/Plugins/Animation/ControlRig/Source/ControlRig/Public/Rigs/RigHierarchyController.h
@@ -501,7 +501,7 @@ public:
 private:
 
 	UPROPERTY(transient)
-	URigHierarchy* Hierarchy;
+	TObjectPtr<URigHierarchy> Hierarchy;
 
 	FRigHierarchyModifiedEvent ModifiedEvent;
 	void Notify(ERigHierarchyNotification InNotifType, const FRigBaseElement* InElement);
diff --git a/Engine/Plugins/Animation/ControlRig/Source/ControlRig/Public/Rigs/RigHierarchyElements.h b/Engine/Plugins/Animation/ControlRig/Source/ControlRig/Public/Rigs/RigHierarchyElements.h
--- a/Engine/Plugins/Animation/ControlRig/Source/ControlRig/Public/Rigs/RigHierarchyElements.h
+++ b/Engine/Plugins/Animation/ControlRig/Source/ControlRig/Public/Rigs/RigHierarchyElements.h
@@ -880,7 +880,7 @@ struct CONTROLRIG_API FRigControlSettings
 
 	/** If the control is 4transient and only visible in the control rig editor */
 	UPROPERTY(BlueprintReadOnly, EditAnywhere, Category = Control)
-	UEnum* ControlEnum;
+	TObjectPtr<UEnum> ControlEnum;
 
 	/** Applies the limits expressed by these settings to a value */
 	FORCEINLINE void ApplyLimits(FRigControlValue& InOutValue) const
diff --git a/Engine/Plugins/Animation/ControlRig/Source/ControlRig/Public/Sequencer/ControlRigObjectSpawner.h b/Engine/Plugins/Animation/ControlRig/Source/ControlRig/Public/Sequencer/ControlRigObjectSpawner.h
--- a/Engine/Plugins/Animation/ControlRig/Source/ControlRig/Public/Sequencer/ControlRigObjectSpawner.h
+++ b/Engine/Plugins/Animation/ControlRig/Source/ControlRig/Public/Sequencer/ControlRigObjectSpawner.h
@@ -13,7 +13,7 @@ class UControlRigObjectHolder : public UObject
 public:
 
 	UPROPERTY()
-	TArray<UObject*> Objects;
+	TArray<TObjectPtr<UObject>> Objects;
 };
 
 class CONTROLRIG_API FControlRigObjectSpawner : public IMovieSceneObjectSpawner
diff --git a/Engine/Plugins/Animation/ControlRig/Source/ControlRig/Public/Sequencer/MovieSceneControlRigParameterSection.h b/Engine/Plugins/Animation/ControlRig/Source/ControlRig/Public/Sequencer/MovieSceneControlRigParameterSection.h
--- a/Engine/Plugins/Animation/ControlRig/Source/ControlRig/Public/Sequencer/MovieSceneControlRigParameterSection.h
+++ b/Engine/Plugins/Animation/ControlRig/Source/ControlRig/Public/Sequencer/MovieSceneControlRigParameterSection.h
@@ -163,7 +163,7 @@ private:
 
 	/** Control Rig that controls us*/
 	UPROPERTY()
-	UControlRig* ControlRig;
+	TObjectPtr<UControlRig> ControlRig;
 
 public:
 
diff --git a/Engine/Plugins/Animation/ControlRig/Source/ControlRig/Public/Sequencer/MovieSceneControlRigParameterTrack.h b/Engine/Plugins/Animation/ControlRig/Source/ControlRig/Public/Sequencer/MovieSceneControlRigParameterTrack.h
--- a/Engine/Plugins/Animation/ControlRig/Source/ControlRig/Public/Sequencer/MovieSceneControlRigParameterTrack.h
+++ b/Engine/Plugins/Animation/ControlRig/Source/ControlRig/Public/Sequencer/MovieSceneControlRigParameterTrack.h
@@ -115,15 +115,15 @@ private:
 
 	/** Control Rig we control*/
 	UPROPERTY()
-	UControlRig* ControlRig;
+	TObjectPtr<UControlRig> ControlRig;
 
 	/** Section we should Key */
 	UPROPERTY()
-	UMovieSceneSection* SectionToKey;
+	TObjectPtr<UMovieSceneSection> SectionToKey;
 
 	/** The sections owned by this track .*/
 	UPROPERTY()
-	TArray<UMovieSceneSection*> Sections;
+	TArray<TObjectPtr<UMovieSceneSection>> Sections;
 
 	/** Unique Name*/
 	UPROPERTY()
diff --git a/Engine/Plugins/Animation/ControlRig/Source/ControlRigDeveloper/Public/ControlRigBlueprint.h b/Engine/Plugins/Animation/ControlRig/Source/ControlRigDeveloper/Public/ControlRigBlueprint.h
--- a/Engine/Plugins/Animation/ControlRig/Source/ControlRigDeveloper/Public/ControlRigBlueprint.h
+++ b/Engine/Plugins/Animation/ControlRig/Source/ControlRigDeveloper/Public/ControlRigBlueprint.h
@@ -251,7 +251,7 @@ public:
 
 #if WITH_EDITORONLY_DATA
 	UPROPERTY()
-	UControlRigGraph* FunctionLibraryEdGraph;
+	TObjectPtr<UControlRigGraph> FunctionLibraryEdGraph;
 #endif
 
 	bool IsFunctionPublic(const FName& InFunctionName) const;
@@ -282,25 +282,25 @@ public:
 protected:
 
 	UPROPERTY(BlueprintReadOnly, Category = "VM")
-	URigVMGraph* Model;
+	TObjectPtr<URigVMGraph> Model;
 
 	UPROPERTY(BlueprintReadOnly, Category = "VM")
-	URigVMFunctionLibrary* FunctionLibrary;
+	TObjectPtr<URigVMFunctionLibrary> FunctionLibrary;
 
 	/** Asset searchable information about exposed public functions on this rig */
 	UPROPERTY(AssetRegistrySearchable)
 	TArray<FControlRigPublicFunctionData> PublicFunctions;
 
 	UPROPERTY(BlueprintReadOnly, transient, Category = "VM")
-	TMap<URigVMGraph*, URigVMController*> Controllers;
+	TMap<TObjectPtr<URigVMGraph>, TObjectPtr<URigVMController>> Controllers;
 
 #if WITH_EDITORONLY_DATA
 
 	UPROPERTY(transient)
-	URigVMGraph* TemplateModel;
+	TObjectPtr<URigVMGraph> TemplateModel;
 
 	UPROPERTY(transient)
-	URigVMController* TemplateController;
+	TObjectPtr<URigVMController> TemplateController;
 
 #endif
 
@@ -378,7 +378,7 @@ public:
 	FRigHierarchyContainer HierarchyContainer_DEPRECATED;
 
 	UPROPERTY(BlueprintReadOnly, Category = "Hierarchy")
-	URigHierarchy* Hierarchy;
+	TObjectPtr<URigHierarchy> Hierarchy;
 
 	UFUNCTION(BlueprintCallable, Category = "Hierarchy")
 	URigHierarchyController* GetHierarchyController() { return HierarchyController; }
@@ -426,7 +426,7 @@ private:
 	int32 VMRecompilationBracket;
 
 	UPROPERTY(transient)
-	URigHierarchyController* HierarchyController;
+	TObjectPtr<URigHierarchyController> HierarchyController;
 
 	FRigVMGraphModifiedEvent ModifiedEvent;
 	void Notify(ERigVMGraphNotifType InNotifType, UObject* InSubject);
@@ -485,7 +485,7 @@ public:
 private:
 
 	UPROPERTY()
-	UControlRigValidator* Validator;
+	TObjectPtr<UControlRigValidator> Validator;
 
 	FRigHierarchyModifiedEvent	HierarchyModifiedEvent;
 
diff --git a/Engine/Plugins/Animation/ControlRig/Source/ControlRigDeveloper/Public/Graph/ControlRigGraph.h b/Engine/Plugins/Animation/ControlRig/Source/ControlRigDeveloper/Public/Graph/ControlRigGraph.h
--- a/Engine/Plugins/Animation/ControlRig/Source/ControlRigDeveloper/Public/Graph/ControlRigGraph.h
+++ b/Engine/Plugins/Animation/ControlRig/Source/ControlRigDeveloper/Public/Graph/ControlRigGraph.h
@@ -139,7 +139,7 @@ private:
 #if WITH_EDITORONLY_DATA
 
 	UPROPERTY(transient)
-	URigVMController* TemplateController;
+	TObjectPtr<URigVMController> TemplateController;
 
 #endif
 #if WITH_EDITOR
diff --git a/Engine/Plugins/Animation/ControlRig/Source/ControlRigDeveloper/Public/Graph/ControlRigGraphNode.h b/Engine/Plugins/Animation/ControlRig/Source/ControlRigDeveloper/Public/Graph/ControlRigGraphNode.h
--- a/Engine/Plugins/Animation/ControlRig/Source/ControlRigDeveloper/Public/Graph/ControlRigGraphNode.h
+++ b/Engine/Plugins/Animation/ControlRig/Source/ControlRigDeveloper/Public/Graph/ControlRigGraphNode.h
@@ -40,10 +40,10 @@ private:
 	FString ModelNodePath;
 
 	UPROPERTY(transient)
-	URigVMNode* CachedModelNode;
+	TObjectPtr<URigVMNode> CachedModelNode;
 
 	UPROPERTY(transient)
-	TMap<FString, URigVMPin*> CachedModelPins;
+	TMap<FString, TObjectPtr<URigVMPin>> CachedModelPins;
 
 	/** The property we represent. For template nodes this represents the struct/property type name. */
 	UPROPERTY()
diff --git a/Engine/Plugins/Animation/ControlRig/Source/ControlRigEditor/Private/EditMode/ControlRigControlsProxy.h b/Engine/Plugins/Animation/ControlRig/Source/ControlRigEditor/Private/EditMode/ControlRigControlsProxy.h
--- a/Engine/Plugins/Animation/ControlRig/Source/ControlRigEditor/Private/EditMode/ControlRigControlsProxy.h
+++ b/Engine/Plugins/Animation/ControlRig/Source/ControlRigEditor/Private/EditMode/ControlRigControlsProxy.h
@@ -173,7 +173,7 @@ struct FControlRigEnumControlProxyValue
 	}
 	
 	UPROPERTY()
-	UEnum* EnumType;
+	TObjectPtr<UEnum> EnumType;
 
 	UPROPERTY(EditAnywhere, Category = Enum)
 	int32 EnumIndex;
@@ -302,10 +302,10 @@ class UControlRigDetailPanelControlProxies :public UObject
 protected:
 
 	UPROPERTY()
-	TMap<FName, UControlRigControlsProxy*> AllProxies;
+	TMap<FName, TObjectPtr<UControlRigControlsProxy>> AllProxies;
 
 	UPROPERTY()
-	TArray< UControlRigControlsProxy*> SelectedProxies;
+	TArray< TObjectPtr<UControlRigControlsProxy>> SelectedProxies;
 
 
 public:
diff --git a/Engine/Plugins/Animation/ControlRig/Source/ControlRigEditor/Private/Editor/SRigHierarchy.h b/Engine/Plugins/Animation/ControlRig/Source/ControlRigEditor/Private/Editor/SRigHierarchy.h
--- a/Engine/Plugins/Animation/ControlRig/Source/ControlRigEditor/Private/Editor/SRigHierarchy.h
+++ b/Engine/Plugins/Animation/ControlRig/Source/ControlRigEditor/Private/Editor/SRigHierarchy.h
@@ -107,7 +107,7 @@ struct FRigHierarchyImportSettings
 	{}
 
 	UPROPERTY(EditAnywhere, Category = "Hierachy Import")
-	USkeletalMesh* Mesh;
+	TObjectPtr<USkeletalMesh> Mesh;
 };
 
 class SRigHierarchyTreeView : public STreeView<TSharedPtr<FRigTreeElement>>
diff --git a/Engine/Plugins/Animation/ControlRig/Source/ControlRigEditor/Private/Graph/NodeSpawners/ControlRigUnitNodeSpawner.h b/Engine/Plugins/Animation/ControlRig/Source/ControlRigEditor/Private/Graph/NodeSpawners/ControlRigUnitNodeSpawner.h
--- a/Engine/Plugins/Animation/ControlRig/Source/ControlRigEditor/Private/Graph/NodeSpawners/ControlRigUnitNodeSpawner.h
+++ b/Engine/Plugins/Animation/ControlRig/Source/ControlRigEditor/Private/Graph/NodeSpawners/ControlRigUnitNodeSpawner.h
@@ -46,7 +46,7 @@ public:
 private:
 	/** The unit type we will spawn */
 	UPROPERTY(Transient)
-	UScriptStruct* StructTemplate;
+	TObjectPtr<UScriptStruct> StructTemplate;
 
 	static UControlRigGraphNode* SpawnNode(UEdGraph* ParentGraph, UBlueprint* Blueprint, UScriptStruct* StructTemplate, FVector2D const Location);
 
diff --git a/Engine/Plugins/Animation/LiveLink/Source/LiveLink/Public/AnimNode_LiveLinkPose.h b/Engine/Plugins/Animation/LiveLink/Source/LiveLink/Public/AnimNode_LiveLinkPose.h
--- a/Engine/Plugins/Animation/LiveLink/Source/LiveLink/Public/AnimNode_LiveLinkPose.h
+++ b/Engine/Plugins/Animation/LiveLink/Source/LiveLink/Public/AnimNode_LiveLinkPose.h
@@ -39,7 +39,7 @@ public:
 	TSubclassOf<ULiveLinkRetargetAsset> RetargetAsset;
 
 	UPROPERTY(transient)
-	ULiveLinkRetargetAsset* CurrentRetargetAsset;
+	TObjectPtr<ULiveLinkRetargetAsset> CurrentRetargetAsset;
 
 public:
 	FAnimNode_LiveLinkPose();
diff --git a/Engine/Plugins/Animation/LiveLink/Source/LiveLink/Public/LiveLinkInstance.h b/Engine/Plugins/Animation/LiveLink/Source/LiveLink/Public/LiveLinkInstance.h
--- a/Engine/Plugins/Animation/LiveLink/Source/LiveLink/Public/LiveLinkInstance.h
+++ b/Engine/Plugins/Animation/LiveLink/Source/LiveLink/Public/LiveLinkInstance.h
@@ -62,7 +62,7 @@ protected:
 
 	// Cache for GC
 	UPROPERTY(transient)
-	ULiveLinkRetargetAsset* CurrentRetargetAsset;
+	TObjectPtr<ULiveLinkRetargetAsset> CurrentRetargetAsset;
 
 	friend FLiveLinkInstanceProxy;
 };
diff --git a/Engine/Plugins/Animation/LiveLink/Source/LiveLinkComponents/Public/LiveLinkComponentController.h b/Engine/Plugins/Animation/LiveLink/Source/LiveLinkComponents/Public/LiveLinkComponentController.h
--- a/Engine/Plugins/Animation/LiveLink/Source/LiveLinkComponents/Public/LiveLinkComponentController.h
+++ b/Engine/Plugins/Animation/LiveLink/Source/LiveLinkComponents/Public/LiveLinkComponentController.h
@@ -25,12 +25,12 @@ public:
 
 #if WITH_EDITORONLY_DATA
 	UPROPERTY(Instanced, NoClear)
-	ULiveLinkControllerBase* Controller_DEPRECATED;
+	TObjectPtr<ULiveLinkControllerBase> Controller_DEPRECATED;
 #endif
 
 	/** Instanced controllers used to control the desired role */
 	UPROPERTY(BlueprintReadOnly, EditAnywhere, Category = "LiveLink", Instanced, NoClear, meta = (ShowOnlyInnerProperties))
-	TMap<TSubclassOf<ULiveLinkRole>, ULiveLinkControllerBase*> ControllerMap;
+	TMap<TSubclassOf<ULiveLinkRole>, TObjectPtr<ULiveLinkControllerBase>> ControllerMap;
 
 	UPROPERTY(EditAnywhere, Category="LiveLink", AdvancedDisplay)
 	bool bUpdateInEditor;
diff --git a/Engine/Plugins/Animation/LiveLink/Source/LiveLinkMovieScene/Public/MovieScene/MovieSceneLiveLinkSection.h b/Engine/Plugins/Animation/LiveLink/Source/LiveLinkMovieScene/Public/MovieScene/MovieSceneLiveLinkSection.h
--- a/Engine/Plugins/Animation/LiveLink/Source/LiveLinkMovieScene/Public/MovieScene/MovieSceneLiveLinkSection.h
+++ b/Engine/Plugins/Animation/LiveLink/Source/LiveLinkMovieScene/Public/MovieScene/MovieSceneLiveLinkSection.h
@@ -84,7 +84,7 @@ public:
 	TArray<bool> ChannelMask;
 
 	UPROPERTY()
-	TArray<UMovieSceneLiveLinkSubSection*> SubSections;
+	TArray<TObjectPtr<UMovieSceneLiveLinkSubSection>> SubSections;
 
 	TSharedPtr<FLiveLinkStaticDataStruct> StaticData;
 
diff --git a/Engine/Plugins/Animation/LiveLink/Source/LiveLinkSequencer/Private/TakeRecorderSource/TakeRecorderLiveLinkSource.h b/Engine/Plugins/Animation/LiveLink/Source/LiveLinkSequencer/Private/TakeRecorderSource/TakeRecorderLiveLinkSource.h
--- a/Engine/Plugins/Animation/LiveLink/Source/LiveLinkSequencer/Private/TakeRecorderSource/TakeRecorderLiveLinkSource.h
+++ b/Engine/Plugins/Animation/LiveLink/Source/LiveLinkSequencer/Private/TakeRecorderSource/TakeRecorderLiveLinkSource.h
@@ -79,7 +79,7 @@ public:
 	* The master track recorder we created.
 	*/
 	UPROPERTY()
-	UMovieSceneLiveLinkTrackRecorder* TrackRecorder;
+	TObjectPtr<UMovieSceneLiveLinkTrackRecorder> TrackRecorder;
 
 private:
 
diff --git a/Engine/Plugins/Enterprise/VariantManagerContent/Source/VariantManagerContent/Public/LevelVariantSetsActor.h b/Engine/Plugins/Enterprise/VariantManagerContent/Source/VariantManagerContent/Public/LevelVariantSetsActor.h
--- a/Engine/Plugins/Enterprise/VariantManagerContent/Source/VariantManagerContent/Public/LevelVariantSetsActor.h
+++ b/Engine/Plugins/Enterprise/VariantManagerContent/Source/VariantManagerContent/Public/LevelVariantSetsActor.h
@@ -48,5 +48,5 @@ private:
 	friend class ULevelVariantSets;
 
 	UPROPERTY(Transient)
-	TMap<UBlueprintGeneratedClass*, ULevelVariantSetsFunctionDirector*> DirectorInstances;
+	TMap<TObjectPtr<UBlueprintGeneratedClass>, TObjectPtr<ULevelVariantSetsFunctionDirector>> DirectorInstances;
 };
diff --git a/Engine/Plugins/Experimental/ChaosEditor/Source/FractureEditor/Private/FractureToolMeshCut.h b/Engine/Plugins/Experimental/ChaosEditor/Source/FractureEditor/Private/FractureToolMeshCut.h
--- a/Engine/Plugins/Experimental/ChaosEditor/Source/FractureEditor/Private/FractureToolMeshCut.h
+++ b/Engine/Plugins/Experimental/ChaosEditor/Source/FractureEditor/Private/FractureToolMeshCut.h
@@ -48,7 +48,7 @@ public:
 private:
 	// Slicing
 	UPROPERTY(EditAnywhere, Category = Slicing)
-	UFractureMeshCutSettings* MeshCutSettings;
+	TObjectPtr<UFractureMeshCutSettings> MeshCutSettings;
 
 	// check if the chosen actor can be used to cut the geometry collection (i.e. if it is a valid static mesh actor with a non-empty static mesh)
 	bool IsCuttingActorValid();
diff --git a/Engine/Plugins/Experimental/ChaosEditor/Source/FractureEditor/Private/FractureToolUV.h b/Engine/Plugins/Experimental/ChaosEditor/Source/FractureEditor/Private/FractureToolUV.h
--- a/Engine/Plugins/Experimental/ChaosEditor/Source/FractureEditor/Private/FractureToolUV.h
+++ b/Engine/Plugins/Experimental/ChaosEditor/Source/FractureEditor/Private/FractureToolUV.h
@@ -82,7 +82,7 @@ public:
 
 	/** The resulting automatically-generated texture map */
 	UPROPERTY(VisibleAnywhere, Category = MapSettings)
-	UTexture2D* Result;
+	TObjectPtr<UTexture2D> Result;
 
 	/** Choose whether to texture only odd material IDs (corresponding to internal faces) or a custom selection */
 	UPROPERTY(EditAnywhere, Category = MapSettings)
diff --git a/Engine/Plugins/Experimental/MeshModelingToolset/Source/MeshModelingTools/Public/AddPatchTool.h b/Engine/Plugins/Experimental/MeshModelingToolset/Source/MeshModelingTools/Public/AddPatchTool.h
--- a/Engine/Plugins/Experimental/MeshModelingToolset/Source/MeshModelingTools/Public/AddPatchTool.h
+++ b/Engine/Plugins/Experimental/MeshModelingToolset/Source/MeshModelingTools/Public/AddPatchTool.h
@@ -93,15 +93,15 @@ public:
 
 protected:
 	UPROPERTY()
-	UAddPatchToolProperties* ShapeSettings;
+	TObjectPtr<UAddPatchToolProperties> ShapeSettings;
 
 	UPROPERTY()
-	UNewMeshMaterialProperties* MaterialProperties;
+	TObjectPtr<UNewMeshMaterialProperties> MaterialProperties;
 
 
 
 	UPROPERTY()
-	UPreviewMesh* PreviewMesh;
+	TObjectPtr<UPreviewMesh> PreviewMesh;
 
 protected:
 	UWorld* TargetWorld;
diff --git a/Engine/Plugins/Experimental/MeshModelingToolset/Source/MeshModelingTools/Public/AddPrimitiveTool.h b/Engine/Plugins/Experimental/MeshModelingToolset/Source/MeshModelingTools/Public/AddPrimitiveTool.h
--- a/Engine/Plugins/Experimental/MeshModelingToolset/Source/MeshModelingTools/Public/AddPrimitiveTool.h
+++ b/Engine/Plugins/Experimental/MeshModelingToolset/Source/MeshModelingTools/Public/AddPrimitiveTool.h
@@ -425,16 +425,16 @@ public:
 	FString Label = "";
 
 	UPROPERTY()
-	AActor* Actor = nullptr;
+	TObjectPtr<AActor> Actor = nullptr;
 
 	UPROPERTY()
-	UStaticMesh* StaticMesh = nullptr;
+	TObjectPtr<UStaticMesh> StaticMesh = nullptr;
 
 	UPROPERTY()
-	UProceduralShapeToolProperties* ShapeSettings;
+	TObjectPtr<UProceduralShapeToolProperties> ShapeSettings;
 
 	UPROPERTY()
-	UNewMeshMaterialProperties* MaterialProperties;
+	TObjectPtr<UNewMeshMaterialProperties> MaterialProperties;
 
 	bool IsInvalid()
 	{
@@ -482,13 +482,13 @@ protected:
 
 	/** Property set for type of output object (StaticMesh, Volume, etc) */
 	UPROPERTY()
-	UCreateMeshObjectTypeProperties* OutputTypeProperties;
+	TObjectPtr<UCreateMeshObjectTypeProperties> OutputTypeProperties;
 
 	UPROPERTY()
-	UProceduralShapeToolProperties* ShapeSettings;
+	TObjectPtr<UProceduralShapeToolProperties> ShapeSettings;
 
 	UPROPERTY()
-	UNewMeshMaterialProperties* MaterialProperties;
+	TObjectPtr<UNewMeshMaterialProperties> MaterialProperties;
 
 
 	/**
@@ -507,10 +507,10 @@ protected:
 
 
 	UPROPERTY()
-	UPreviewMesh* PreviewMesh;
+	TObjectPtr<UPreviewMesh> PreviewMesh;
 
 	UPROPERTY()
-	ULastActorInfo* LastGenerated;
+	TObjectPtr<ULastActorInfo> LastGenerated;
 
 	UPROPERTY()
 	FString AssetName = TEXT("GeneratedAsset");
diff --git a/Engine/Plugins/Experimental/MeshModelingToolset/Source/MeshModelingTools/Public/AlignObjectsTool.h b/Engine/Plugins/Experimental/MeshModelingToolset/Source/MeshModelingTools/Public/AlignObjectsTool.h
--- a/Engine/Plugins/Experimental/MeshModelingToolset/Source/MeshModelingTools/Public/AlignObjectsTool.h
+++ b/Engine/Plugins/Experimental/MeshModelingToolset/Source/MeshModelingTools/Public/AlignObjectsTool.h
@@ -133,7 +133,7 @@ public:
 
 public:
 	UPROPERTY()
-	UAlignObjectsToolProperties* AlignProps;
+	TObjectPtr<UAlignObjectsToolProperties> AlignProps;
 
 
 protected:
diff --git a/Engine/Plugins/Experimental/MeshModelingToolset/Source/MeshModelingTools/Public/BakeMeshAttributeMapsTool.h b/Engine/Plugins/Experimental/MeshModelingToolset/Source/MeshModelingTools/Public/BakeMeshAttributeMapsTool.h
--- a/Engine/Plugins/Experimental/MeshModelingToolset/Source/MeshModelingTools/Public/BakeMeshAttributeMapsTool.h
+++ b/Engine/Plugins/Experimental/MeshModelingToolset/Source/MeshModelingTools/Public/BakeMeshAttributeMapsTool.h
@@ -138,7 +138,7 @@ public:
 	TArray<FString> UVLayerNamesList;
 
 	UPROPERTY(VisibleAnywhere, Category = MapSettings, meta = (TransientToolProperty))
-	TArray<UTexture2D*> Result;
+	TArray<TObjectPtr<UTexture2D>> Result;
 
 };
 
@@ -172,25 +172,25 @@ protected:
 	// @todo setters/getters for these
 
 	UPROPERTY()
-	UBakeMeshAttributeMapsToolProperties* Settings;
+	TObjectPtr<UBakeMeshAttributeMapsToolProperties> Settings;
 
 	UPROPERTY()
-	UBakedNormalMapToolProperties* NormalMapProps;
+	TObjectPtr<UBakedNormalMapToolProperties> NormalMapProps;
 
 	UPROPERTY()
-	UBakedOcclusionMapToolProperties* OcclusionMapProps;
+	TObjectPtr<UBakedOcclusionMapToolProperties> OcclusionMapProps;
 
 	UPROPERTY()
-	UBakedCurvatureMapToolProperties* CurvatureMapProps;
+	TObjectPtr<UBakedCurvatureMapToolProperties> CurvatureMapProps;
 
 	UPROPERTY()
-	UBakedTexture2DImageProperties* Texture2DProps;
+	TObjectPtr<UBakedTexture2DImageProperties> Texture2DProps;
 
 	UPROPERTY()
-	UBakedMultiTexture2DImageProperties* MultiTextureProps;
+	TObjectPtr<UBakedMultiTexture2DImageProperties> MultiTextureProps;
 
 	UPROPERTY()
-	UBakedOcclusionMapVisualizationProperties* VisualizationProps;
+	TObjectPtr<UBakedOcclusionMapVisualizationProperties> VisualizationProps;
 
 
 protected:
@@ -199,13 +199,13 @@ protected:
 	UDynamicMeshComponent* DynamicMeshComponent;
 
 	UPROPERTY()
-	UMaterialInstanceDynamic* PreviewMaterial;
+	TObjectPtr<UMaterialInstanceDynamic> PreviewMaterial;
 
     UPROPERTY()
-	UMaterialInstanceDynamic* BentNormalPreviewMaterial;
+	TObjectPtr<UMaterialInstanceDynamic> BentNormalPreviewMaterial;
 
 	UPROPERTY()
-	UMaterialInstanceDynamic* WorkingPreviewMaterial;
+	TObjectPtr<UMaterialInstanceDynamic> WorkingPreviewMaterial;
 	float SecondsBeforeWorkingMaterial = 0.75;
 
 	TSharedPtr<UE::Geometry::TMeshTangents<double>, ESPMode::ThreadSafe> BaseMeshTangents;
@@ -255,7 +255,7 @@ protected:
 	EBakeOpState OpState = EBakeOpState::Evaluate;
 
 	UPROPERTY()
-	TArray<UTexture2D*> CachedMaps;
+	TArray<TObjectPtr<UTexture2D>> CachedMaps;
 	using CachedMapIndex = TMap<EBakeMapType, int32>;
 	CachedMapIndex CachedMapIndices;
 
@@ -281,13 +281,13 @@ protected:
 
 	// empty maps are shown when nothing is computed
 	UPROPERTY()
-	UTexture2D* EmptyNormalMap;
+	TObjectPtr<UTexture2D> EmptyNormalMap;
 
 	UPROPERTY()
-	UTexture2D* EmptyColorMapBlack;
+	TObjectPtr<UTexture2D> EmptyColorMapBlack;
 
 	UPROPERTY()
-	UTexture2D* EmptyColorMapWhite;
+	TObjectPtr<UTexture2D> EmptyColorMapWhite;
 
 	void InitializeEmptyMaps();
 
diff --git a/Engine/Plugins/Experimental/MeshModelingToolset/Source/MeshModelingTools/Public/BakeMeshAttributeToolCommon.h b/Engine/Plugins/Experimental/MeshModelingToolset/Source/MeshModelingTools/Public/BakeMeshAttributeToolCommon.h
--- a/Engine/Plugins/Experimental/MeshModelingToolset/Source/MeshModelingTools/Public/BakeMeshAttributeToolCommon.h
+++ b/Engine/Plugins/Experimental/MeshModelingToolset/Source/MeshModelingTools/Public/BakeMeshAttributeToolCommon.h
@@ -180,7 +180,7 @@ class MESHMODELINGTOOLS_API UBakedTexture2DImageProperties : public UInteractive
 public:
 	/** The source texture that is to be resampled into a new texture map */
 	UPROPERTY(EditAnywhere, Category = Texture2D, meta = (TransientToolProperty))
-	UTexture2D* SourceTexture;
+	TObjectPtr<UTexture2D> SourceTexture;
 
 	/** The UV layer on the source mesh that corresponds to the SourceTexture */
 	UPROPERTY(EditAnywhere, Category = Texture2D)
@@ -196,7 +196,7 @@ public:
 
 	/** For each material ID, the source texture that will be resampled in that material's region*/
 	UPROPERTY(EditAnywhere, Category = MultiTexture, meta = (DisplayName = "Material IDs / Source Textures"))
-	TMap<int32, UTexture2D*> MaterialIDSourceTextureMap;
+	TMap<int32, TObjectPtr<UTexture2D>> MaterialIDSourceTextureMap;
 
 	/** UV layer to sample from on the input mesh */
 	UPROPERTY(EditAnywhere, Category = MultiTexture)
@@ -204,7 +204,7 @@ public:
 
 	/** The set of all source textures from all input materials */
 	UPROPERTY(VisibleAnywhere, BlueprintReadOnly, AdvancedDisplay, Category = MultiTexture, meta = (DisplayName = "Source Textures"))
-	TArray<UTexture2D*> AllSourceTextures;
+	TArray<TObjectPtr<UTexture2D>> AllSourceTextures;
 
 };
 
diff --git a/Engine/Plugins/Experimental/MeshModelingToolset/Source/MeshModelingTools/Public/BakeMeshAttributeVertexTool.h b/Engine/Plugins/Experimental/MeshModelingToolset/Source/MeshModelingTools/Public/BakeMeshAttributeVertexTool.h
--- a/Engine/Plugins/Experimental/MeshModelingToolset/Source/MeshModelingTools/Public/BakeMeshAttributeVertexTool.h
+++ b/Engine/Plugins/Experimental/MeshModelingToolset/Source/MeshModelingTools/Public/BakeMeshAttributeVertexTool.h
@@ -180,40 +180,40 @@ public:
 
 protected:
 	UPROPERTY()
-	UBakeMeshAttributeVertexToolProperties* Settings;
+	TObjectPtr<UBakeMeshAttributeVertexToolProperties> Settings;
 
 	UPROPERTY()
-	UBakeMeshAttributeVertexToolColorProperties* ColorSettings;
+	TObjectPtr<UBakeMeshAttributeVertexToolColorProperties> ColorSettings;
 
 	UPROPERTY()
-	UBakeMeshAttributeVertexToolChannelProperties* PerChannelSettings;
+	TObjectPtr<UBakeMeshAttributeVertexToolChannelProperties> PerChannelSettings;
 
 	UPROPERTY()
-	UBakedOcclusionMapToolProperties* OcclusionSettings;
+	TObjectPtr<UBakedOcclusionMapToolProperties> OcclusionSettings;
 
 	UPROPERTY()
-	UBakedCurvatureMapToolProperties* CurvatureSettings;
+	TObjectPtr<UBakedCurvatureMapToolProperties> CurvatureSettings;
 
 	UPROPERTY()
-	UBakedTexture2DImageProperties* TextureSettings;
+	TObjectPtr<UBakedTexture2DImageProperties> TextureSettings;
 
 	UPROPERTY()
-	UBakedMultiTexture2DImageProperties* MultiTextureSettings;
+	TObjectPtr<UBakedMultiTexture2DImageProperties> MultiTextureSettings;
 
 protected:
 	friend class FMeshVertexBakerOp;
 	
 	UPROPERTY()
-	UPreviewMesh* PreviewMesh;
+	TObjectPtr<UPreviewMesh> PreviewMesh;
 
 	UPROPERTY()
-	UMaterialInstanceDynamic* PreviewMaterial;
+	TObjectPtr<UMaterialInstanceDynamic> PreviewMaterial;
 
 	UPROPERTY()
-	UMaterialInstanceDynamic* PreviewAlphaMaterial;
+	TObjectPtr<UMaterialInstanceDynamic> PreviewAlphaMaterial;
 
 	UPROPERTY()
-	UMaterialInstanceDynamic* WorkingPreviewMaterial;
+	TObjectPtr<UMaterialInstanceDynamic> WorkingPreviewMaterial;
 	float SecondsBeforeWorkingMaterial = 0.75;
 
 	UWorld* TargetWorld = nullptr;
diff --git a/Engine/Plugins/Experimental/MeshModelingToolset/Source/MeshModelingTools/Public/BakeTransformTool.h b/Engine/Plugins/Experimental/MeshModelingToolset/Source/MeshModelingTools/Public/BakeTransformTool.h
--- a/Engine/Plugins/Experimental/MeshModelingToolset/Source/MeshModelingTools/Public/BakeTransformTool.h
+++ b/Engine/Plugins/Experimental/MeshModelingToolset/Source/MeshModelingTools/Public/BakeTransformTool.h
@@ -91,7 +91,7 @@ public:
 protected:
 
 	UPROPERTY()
-	UBakeTransformToolProperties* BasicProperties;
+	TObjectPtr<UBakeTransformToolProperties> BasicProperties;
 
 protected:
 	UWorld* TargetWorld;
diff --git a/Engine/Plugins/Experimental/MeshModelingToolset/Source/MeshModelingTools/Public/CSGMeshesTool.h b/Engine/Plugins/Experimental/MeshModelingToolset/Source/MeshModelingTools/Public/CSGMeshesTool.h
--- a/Engine/Plugins/Experimental/MeshModelingToolset/Source/MeshModelingTools/Public/CSGMeshesTool.h
+++ b/Engine/Plugins/Experimental/MeshModelingToolset/Source/MeshModelingTools/Public/CSGMeshesTool.h
@@ -129,22 +129,22 @@ protected:
 	void UpdateVisualization();
 
 	UPROPERTY()
-	UCSGMeshesToolProperties* CSGProperties;
+	TObjectPtr<UCSGMeshesToolProperties> CSGProperties;
 
 	UPROPERTY()
-	UTrimMeshesToolProperties* TrimProperties;
+	TObjectPtr<UTrimMeshesToolProperties> TrimProperties;
 
 	TArray<TSharedPtr<FDynamicMesh3, ESPMode::ThreadSafe>> OriginalDynamicMeshes;
 
 	UPROPERTY()
-	TArray<UPreviewMesh*> OriginalMeshPreviews;
+	TArray<TObjectPtr<UPreviewMesh>> OriginalMeshPreviews;
 
 	// Material used to show the otherwise-invisible cutting/trimming mesh
 	UPROPERTY()
-	UMaterialInstanceDynamic* PreviewsGhostMaterial;
+	TObjectPtr<UMaterialInstanceDynamic> PreviewsGhostMaterial;
 
 	UPROPERTY()
-	ULineSetComponent* DrawnLineSet;
+	TObjectPtr<ULineSetComponent> DrawnLineSet;
 
 	// for visualization of any errors in the currently-previewed CSG operation
 	TArray<int> CreatedBoundaryEdges;
diff --git a/Engine/Plugins/Experimental/MeshModelingToolset/Source/MeshModelingTools/Public/CombineMeshesTool.h b/Engine/Plugins/Experimental/MeshModelingToolset/Source/MeshModelingTools/Public/CombineMeshesTool.h
--- a/Engine/Plugins/Experimental/MeshModelingToolset/Source/MeshModelingTools/Public/CombineMeshesTool.h
+++ b/Engine/Plugins/Experimental/MeshModelingToolset/Source/MeshModelingTools/Public/CombineMeshesTool.h
@@ -95,13 +95,13 @@ public:
 protected:
 
 	UPROPERTY()
-	UCombineMeshesToolProperties* BasicProperties;
+	TObjectPtr<UCombineMeshesToolProperties> BasicProperties;
 
 	UPROPERTY()
-	UCreateMeshObjectTypeProperties* OutputTypeProperties;
+	TObjectPtr<UCreateMeshObjectTypeProperties> OutputTypeProperties;
 
 	UPROPERTY()
-	UOnAcceptHandleSourcesProperties* HandleSourceProperties;
+	TObjectPtr<UOnAcceptHandleSourcesProperties> HandleSourceProperties;
 
 protected:
 	UWorld* TargetWorld;
diff --git a/Engine/Plugins/Experimental/MeshModelingToolset/Source/MeshModelingTools/Public/ConvertMeshesTool.h b/Engine/Plugins/Experimental/MeshModelingToolset/Source/MeshModelingTools/Public/ConvertMeshesTool.h
--- a/Engine/Plugins/Experimental/MeshModelingToolset/Source/MeshModelingTools/Public/ConvertMeshesTool.h
+++ b/Engine/Plugins/Experimental/MeshModelingToolset/Source/MeshModelingTools/Public/ConvertMeshesTool.h
@@ -53,10 +53,10 @@ public:
 	virtual bool CanAccept() const override;
 
 	UPROPERTY()
-	UConvertMeshesToolProperties* BasicProperties;
+	TObjectPtr<UConvertMeshesToolProperties> BasicProperties;
 
 	UPROPERTY()
-	UCreateMeshObjectTypeProperties* OutputTypeProperties;
+	TObjectPtr<UCreateMeshObjectTypeProperties> OutputTypeProperties;
 
 protected:
 	UWorld* TargetWorld;
diff --git a/Engine/Plugins/Experimental/MeshModelingToolset/Source/MeshModelingTools/Public/ConvertToPolygonsTool.h b/Engine/Plugins/Experimental/MeshModelingToolset/Source/MeshModelingTools/Public/ConvertToPolygonsTool.h
--- a/Engine/Plugins/Experimental/MeshModelingToolset/Source/MeshModelingTools/Public/ConvertToPolygonsTool.h
+++ b/Engine/Plugins/Experimental/MeshModelingToolset/Source/MeshModelingTools/Public/ConvertToPolygonsTool.h
@@ -75,7 +75,7 @@ public:
 	virtual TUniquePtr<UE::Geometry::FDynamicMeshOperator> MakeNewOperator() override;
 
 	UPROPERTY()
-	UConvertToPolygonsTool* ConvertToPolygonsTool;  // back pointer
+	TObjectPtr<UConvertToPolygonsTool> ConvertToPolygonsTool;  // back pointer
 };
 
 /**
@@ -107,13 +107,13 @@ protected:
 
 protected:
 	UPROPERTY()
-	UConvertToPolygonsToolProperties* Settings;
+	TObjectPtr<UConvertToPolygonsToolProperties> Settings;
 
 	UPROPERTY()
-	UMeshOpPreviewWithBackgroundCompute* PreviewCompute = nullptr;
+	TObjectPtr<UMeshOpPreviewWithBackgroundCompute> PreviewCompute = nullptr;
 
 	UPROPERTY()
-	UPreviewGeometry* PreviewGeometry = nullptr;
+	TObjectPtr<UPreviewGeometry> PreviewGeometry = nullptr;
 
 protected:
 	TSharedPtr<UE::Geometry::FDynamicMesh3, ESPMode::ThreadSafe> OriginalDynamicMesh;
diff --git a/Engine/Plugins/Experimental/MeshModelingToolset/Source/MeshModelingTools/Public/CutMeshWithMeshTool.h b/Engine/Plugins/Experimental/MeshModelingToolset/Source/MeshModelingTools/Public/CutMeshWithMeshTool.h
--- a/Engine/Plugins/Experimental/MeshModelingToolset/Source/MeshModelingTools/Public/CutMeshWithMeshTool.h
+++ b/Engine/Plugins/Experimental/MeshModelingToolset/Source/MeshModelingTools/Public/CutMeshWithMeshTool.h
@@ -78,16 +78,16 @@ protected:
 	void UpdateVisualization();
 
 	UPROPERTY()
-	UCutMeshWithMeshToolProperties* CutProperties;
+	TObjectPtr<UCutMeshWithMeshToolProperties> CutProperties;
 
 	UPROPERTY()
-	UPreviewMesh* IntersectPreviewMesh;
+	TObjectPtr<UPreviewMesh> IntersectPreviewMesh;
 
 	TSharedPtr<FDynamicMesh3, ESPMode::ThreadSafe> OriginalTargetMesh;
 	TSharedPtr<FDynamicMesh3, ESPMode::ThreadSafe> OriginalCuttingMesh;
 
 	UPROPERTY()
-	ULineSetComponent* DrawnLineSet;
+	TObjectPtr<ULineSetComponent> DrawnLineSet;
 
 	// for visualization of any errors in the currently-previewed CSG operation
 	TArray<int> CreatedSubtractBoundaryEdges;
diff --git a/Engine/Plugins/Experimental/MeshModelingToolset/Source/MeshModelingTools/Public/DeformMeshPolygonsTool.h b/Engine/Plugins/Experimental/MeshModelingToolset/Source/MeshModelingTools/Public/DeformMeshPolygonsTool.h
--- a/Engine/Plugins/Experimental/MeshModelingToolset/Source/MeshModelingTools/Public/DeformMeshPolygonsTool.h
+++ b/Engine/Plugins/Experimental/MeshModelingToolset/Source/MeshModelingTools/Public/DeformMeshPolygonsTool.h
@@ -178,10 +178,10 @@ public:
 
 protected:
 	UPROPERTY()
-	UDynamicMeshComponent* DynamicMeshComponent;
+	TObjectPtr<UDynamicMeshComponent> DynamicMeshComponent;
 
 	UPROPERTY()
-	UDeformMeshPolygonsTransformProperties* TransformProps;
+	TObjectPtr<UDeformMeshPolygonsTransformProperties> TransformProps;
 
 	// realtime visualization
 	void OnDynamicMeshComponentChanged();
diff --git a/Engine/Plugins/Experimental/MeshModelingToolset/Source/MeshModelingTools/Public/DisplaceMeshTool.h b/Engine/Plugins/Experimental/MeshModelingToolset/Source/MeshModelingTools/Public/DisplaceMeshTool.h
--- a/Engine/Plugins/Experimental/MeshModelingToolset/Source/MeshModelingTools/Public/DisplaceMeshTool.h
+++ b/Engine/Plugins/Experimental/MeshModelingToolset/Source/MeshModelingTools/Public/DisplaceMeshTool.h
@@ -100,7 +100,7 @@ class MESHMODELINGTOOLS_API UDisplaceMeshTextureMapProperties : public UInteract
 public:
 	/** Displacement map. Only the first channel is used. */
 	UPROPERTY(EditAnywhere, Category = Options)
-	UTexture2D* DisplacementMap = nullptr;
+	TObjectPtr<UTexture2D> DisplacementMap = nullptr;
 
 	//~ TODO: Consider making this an enum
 	/** Channel in the displacement map to use. 0-3 coresponding to the four channels (RGBA) */
@@ -132,7 +132,7 @@ public:
 
 	/** This curve is queried in the range [0,1] to adjust contrast of the displacement map. */
 	UPROPERTY(EditAnywhere, Category = Options, meta = (EditCondition = "bApplyAdjustmentCurve", EditConditionHides))
-	UCurveFloat* AdjustmentCurve = nullptr;
+	TObjectPtr<UCurveFloat> AdjustmentCurve = nullptr;
 
 	/** Recalculate normals from displaced mesh. Disable this if you are applying Displacements that are paired with an existing Normal Map in your Material. */
 	UPROPERTY(EditAnywhere, Category = Options)
@@ -259,27 +259,27 @@ public:
 
 	/** Shared properties for all displacement modes. */
 	UPROPERTY()
-	UDisplaceMeshCommonProperties* CommonProperties;
+	TObjectPtr<UDisplaceMeshCommonProperties> CommonProperties;
 
 	/** Properties defining the directional filter. */
 	UPROPERTY()
-	UDisplaceMeshDirectionalFilterProperties* DirectionalFilterProperties;
+	TObjectPtr<UDisplaceMeshDirectionalFilterProperties> DirectionalFilterProperties;
 
 	/** Properties defining the texture map */
 	UPROPERTY()
-	UDisplaceMeshTextureMapProperties* TextureMapProperties;
+	TObjectPtr<UDisplaceMeshTextureMapProperties> TextureMapProperties;
 
 	/** Multi-layer Perlin noise frequencies and intensities */
 	UPROPERTY()
-	UDisplaceMeshPerlinNoiseProperties* NoiseProperties;
+	TObjectPtr<UDisplaceMeshPerlinNoiseProperties> NoiseProperties;
 
 	/** Sine wave parameters and direction of displacement */
 	UPROPERTY()
-	UDisplaceMeshSineWaveProperties* SineWaveProperties;
+	TObjectPtr<UDisplaceMeshSineWaveProperties> SineWaveProperties;
 
 	/** Contrast Curve we are actively listening to */
 	UPROPERTY()
-	UCurveFloat* ActiveContrastCurveTarget = nullptr;
+	TObjectPtr<UCurveFloat> ActiveContrastCurveTarget = nullptr;
 
 private:
 	void StartComputation();
diff --git a/Engine/Plugins/Experimental/MeshModelingToolset/Source/MeshModelingTools/Public/DrawAndRevolveTool.h b/Engine/Plugins/Experimental/MeshModelingToolset/Source/MeshModelingTools/Public/DrawAndRevolveTool.h
--- a/Engine/Plugins/Experimental/MeshModelingToolset/Source/MeshModelingTools/Public/DrawAndRevolveTool.h
+++ b/Engine/Plugins/Experimental/MeshModelingToolset/Source/MeshModelingTools/Public/DrawAndRevolveTool.h
@@ -76,7 +76,7 @@ public:
 	virtual TUniquePtr<UE::Geometry::FDynamicMeshOperator> MakeNewOperator() override;
 
 	UPROPERTY()
-	UDrawAndRevolveTool* RevolveTool;
+	TObjectPtr<UDrawAndRevolveTool> RevolveTool;
 };
 
 
@@ -120,23 +120,23 @@ protected:
 	void UpdateRevolutionAxis();
 
 	UPROPERTY()
-	UCurveControlPointsMechanic* ControlPointsMechanic = nullptr;
+	TObjectPtr<UCurveControlPointsMechanic> ControlPointsMechanic = nullptr;
 
 	UPROPERTY()
-	UConstructionPlaneMechanic* PlaneMechanic = nullptr;
+	TObjectPtr<UConstructionPlaneMechanic> PlaneMechanic = nullptr;
 
 	/** Property set for type of output object (StaticMesh, Volume, etc) */
 	UPROPERTY()
-	UCreateMeshObjectTypeProperties* OutputTypeProperties;
+	TObjectPtr<UCreateMeshObjectTypeProperties> OutputTypeProperties;
 
 	UPROPERTY()
-	URevolveToolProperties* Settings = nullptr;
+	TObjectPtr<URevolveToolProperties> Settings = nullptr;
 
 	UPROPERTY()
-	UNewMeshMaterialProperties* MaterialProperties;
+	TObjectPtr<UNewMeshMaterialProperties> MaterialProperties;
 
 	UPROPERTY()
-	UMeshOpPreviewWithBackgroundCompute* Preview = nullptr;
+	TObjectPtr<UMeshOpPreviewWithBackgroundCompute> Preview = nullptr;
 
 	void StartPreview();
 
diff --git a/Engine/Plugins/Experimental/MeshModelingToolset/Source/MeshModelingTools/Public/DrawPolyPathTool.h b/Engine/Plugins/Experimental/MeshModelingToolset/Source/MeshModelingTools/Public/DrawPolyPathTool.h
--- a/Engine/Plugins/Experimental/MeshModelingToolset/Source/MeshModelingTools/Public/DrawPolyPathTool.h
+++ b/Engine/Plugins/Experimental/MeshModelingToolset/Source/MeshModelingTools/Public/DrawPolyPathTool.h
@@ -160,16 +160,16 @@ protected:
 
 	/** Property set for type of output object (StaticMesh, Volume, etc) */
 	UPROPERTY()
-	UCreateMeshObjectTypeProperties* OutputTypeProperties;
+	TObjectPtr<UCreateMeshObjectTypeProperties> OutputTypeProperties;
 
 	UPROPERTY()
-	UDrawPolyPathProperties* TransformProps;
+	TObjectPtr<UDrawPolyPathProperties> TransformProps;
 
 	UPROPERTY()
-	UDrawPolyPathExtrudeProperties* ExtrudeProperties;
+	TObjectPtr<UDrawPolyPathExtrudeProperties> ExtrudeProperties;
 
 	UPROPERTY()
-	UNewMeshMaterialProperties* MaterialProperties;
+	TObjectPtr<UNewMeshMaterialProperties> MaterialProperties;
 
 protected:
 	// camera state at last render
@@ -179,7 +179,7 @@ protected:
 	// drawing plane and gizmo
 
 	UPROPERTY()
-	UConstructionPlaneMechanic* PlaneMechanic = nullptr;
+	TObjectPtr<UConstructionPlaneMechanic> PlaneMechanic = nullptr;
 
 	UE::Geometry::FFrame3d DrawPlaneWorld;
 
@@ -198,14 +198,14 @@ protected:
 	bool bPathIsClosed = false;		// If true, CurPathPoints are assumed to define a closed path
 
 	UPROPERTY()
-	UPolyEditPreviewMesh* EditPreview;
+	TObjectPtr<UPolyEditPreviewMesh> EditPreview;
 
 	UPROPERTY()
-	UPlaneDistanceFromHitMechanic* ExtrudeHeightMechanic = nullptr;
+	TObjectPtr<UPlaneDistanceFromHitMechanic> ExtrudeHeightMechanic = nullptr;
 	UPROPERTY()
-	USpatialCurveDistanceMechanic* CurveDistMechanic = nullptr;
+	TObjectPtr<USpatialCurveDistanceMechanic> CurveDistMechanic = nullptr;
 	UPROPERTY()
-	UCollectSurfacePathMechanic* SurfacePathMechanic = nullptr;
+	TObjectPtr<UCollectSurfacePathMechanic> SurfacePathMechanic = nullptr;
 
 	void InitializeNewSurfacePath();
 	void UpdateSurfacePathPlane();
diff --git a/Engine/Plugins/Experimental/MeshModelingToolset/Source/MeshModelingTools/Public/DrawPolygonTool.h b/Engine/Plugins/Experimental/MeshModelingToolset/Source/MeshModelingTools/Public/DrawPolygonTool.h
--- a/Engine/Plugins/Experimental/MeshModelingToolset/Source/MeshModelingTools/Public/DrawPolygonTool.h
+++ b/Engine/Plugins/Experimental/MeshModelingToolset/Source/MeshModelingTools/Public/DrawPolygonTool.h
@@ -213,17 +213,17 @@ protected:
 
 	/** Property set for type of output object (StaticMesh, Volume, etc) */
 	UPROPERTY()
-	UCreateMeshObjectTypeProperties* OutputTypeProperties;
+	TObjectPtr<UCreateMeshObjectTypeProperties> OutputTypeProperties;
 
 	/** Properties that control polygon generation exposed to user via detailsview */
 	UPROPERTY()
-	UDrawPolygonToolStandardProperties* PolygonProperties;
+	TObjectPtr<UDrawPolygonToolStandardProperties> PolygonProperties;
 
 	UPROPERTY()
-	UDrawPolygonToolSnapProperties* SnapProperties;
+	TObjectPtr<UDrawPolygonToolSnapProperties> SnapProperties;
 
 	UPROPERTY()
-	UNewMeshMaterialProperties* MaterialProperties;
+	TObjectPtr<UNewMeshMaterialProperties> MaterialProperties;
 	
 
 	/** Origin of plane we will draw polygon on */
@@ -247,16 +247,16 @@ protected:
 	FViewCameraState CameraState;
 
 	UPROPERTY()
-	UPreviewMesh* PreviewMesh;
+	TObjectPtr<UPreviewMesh> PreviewMesh;
 
 	
 	// drawing plane gizmo
 
 	UPROPERTY()
-	UTransformGizmo* PlaneTransformGizmo;
+	TObjectPtr<UTransformGizmo> PlaneTransformGizmo;
 
 	UPROPERTY()
-	UTransformProxy* PlaneTransformProxy;
+	TObjectPtr<UTransformProxy> PlaneTransformProxy;
 
 	// called on PlaneTransformProxy.OnTransformChanged
 	void PlaneTransformChanged(UTransformProxy* Proxy, FTransform Transform);
@@ -304,7 +304,7 @@ protected:
 	bool bPreviewUpdatePending;
 
 	UPROPERTY()
-	UPlaneDistanceFromHitMechanic* HeightMechanic;
+	TObjectPtr<UPlaneDistanceFromHitMechanic> HeightMechanic;
 
 
 	UE::Geometry::FFrame3f HitPosFrameWorld;
diff --git a/Engine/Plugins/Experimental/MeshModelingToolset/Source/MeshModelingTools/Public/DynamicMeshBrushTool.h b/Engine/Plugins/Experimental/MeshModelingToolset/Source/MeshModelingTools/Public/DynamicMeshBrushTool.h
--- a/Engine/Plugins/Experimental/MeshModelingToolset/Source/MeshModelingTools/Public/DynamicMeshBrushTool.h
+++ b/Engine/Plugins/Experimental/MeshModelingToolset/Source/MeshModelingTools/Public/DynamicMeshBrushTool.h
@@ -42,7 +42,7 @@ protected:
 
 protected:
 	UPROPERTY()
-	UPreviewMesh* PreviewMesh;
+	TObjectPtr<UPreviewMesh> PreviewMesh;
 
 	// this function is called when the component inside the PreviewMesh is modified (eg via an undo/redo event)
 	virtual void OnBaseMeshComponentChanged() {}	
diff --git a/Engine/Plugins/Experimental/MeshModelingToolset/Source/MeshModelingTools/Public/DynamicMeshSculptTool.h b/Engine/Plugins/Experimental/MeshModelingToolset/Source/MeshModelingTools/Public/DynamicMeshSculptTool.h
--- a/Engine/Plugins/Experimental/MeshModelingToolset/Source/MeshModelingTools/Public/DynamicMeshSculptTool.h
+++ b/Engine/Plugins/Experimental/MeshModelingToolset/Source/MeshModelingTools/Public/DynamicMeshSculptTool.h
@@ -262,30 +262,30 @@ public:
 public:
 	/** Properties that control brush size/etc*/
 	UPROPERTY()
-	USculptBrushProperties* BrushProperties;
+	TObjectPtr<USculptBrushProperties> BrushProperties;
 
 	/** Properties that control sculpting*/
 	UPROPERTY()
-	UBrushSculptProperties* SculptProperties;
+	TObjectPtr<UBrushSculptProperties> SculptProperties;
 
 	UPROPERTY()
-	USculptMaxBrushProperties* SculptMaxBrushProperties;
+	TObjectPtr<USculptMaxBrushProperties> SculptMaxBrushProperties;
 	
 	UPROPERTY()
-	UKelvinBrushProperties* KelvinBrushProperties;
+	TObjectPtr<UKelvinBrushProperties> KelvinBrushProperties;
 
 	/** Properties that control dynamic remeshing */
 	UPROPERTY()
-	UBrushRemeshProperties* RemeshProperties;
+	TObjectPtr<UBrushRemeshProperties> RemeshProperties;
 
 	UPROPERTY()
-	UFixedPlaneBrushProperties* GizmoProperties;
+	TObjectPtr<UFixedPlaneBrushProperties> GizmoProperties;
 
 	UPROPERTY()
-	UMeshEditingViewProperties* ViewProperties;
+	TObjectPtr<UMeshEditingViewProperties> ViewProperties;
 
 	UPROPERTY()
-	UDynamicSculptToolActions* SculptToolActions;
+	TObjectPtr<UDynamicSculptToolActions> SculptToolActions;
 
 public:
 	virtual void IncreaseBrushRadiusAction();
@@ -306,19 +306,19 @@ private:
 	FViewCameraState CameraState;
 
 	UPROPERTY()
-	UBrushStampIndicator* BrushIndicator;
+	TObjectPtr<UBrushStampIndicator> BrushIndicator;
 
 	UPROPERTY()
-	UMaterialInstanceDynamic* BrushIndicatorMaterial;
+	TObjectPtr<UMaterialInstanceDynamic> BrushIndicatorMaterial;
 
 	UPROPERTY()
-	UPreviewMesh* BrushIndicatorMesh;
+	TObjectPtr<UPreviewMesh> BrushIndicatorMesh;
 
 	UPROPERTY()
-	UOctreeDynamicMeshComponent* DynamicMeshComponent;
+	TObjectPtr<UOctreeDynamicMeshComponent> DynamicMeshComponent;
 
 	UPROPERTY()
-	UMaterialInstanceDynamic* ActiveOverrideMaterial;
+	TObjectPtr<UMaterialInstanceDynamic> ActiveOverrideMaterial;
 
 	UE::Geometry::FTransform3d InitialTargetTransform;
 	UE::Geometry::FTransform3d CurTargetTransform;
@@ -477,10 +477,10 @@ private:
 
 	// plane gizmo
 	UPROPERTY()
-	UTransformGizmo* PlaneTransformGizmo;
+	TObjectPtr<UTransformGizmo> PlaneTransformGizmo;
 
 	UPROPERTY()
-	UTransformProxy* PlaneTransformProxy;
+	TObjectPtr<UTransformProxy> PlaneTransformProxy;
 
 	void PlaneTransformChanged(UTransformProxy* Proxy, FTransform Transform);
 
diff --git a/Engine/Plugins/Experimental/MeshModelingToolset/Source/MeshModelingTools/Public/EdgeLoopInsertionTool.h b/Engine/Plugins/Experimental/MeshModelingToolset/Source/MeshModelingTools/Public/EdgeLoopInsertionTool.h
--- a/Engine/Plugins/Experimental/MeshModelingToolset/Source/MeshModelingTools/Public/EdgeLoopInsertionTool.h
+++ b/Engine/Plugins/Experimental/MeshModelingToolset/Source/MeshModelingTools/Public/EdgeLoopInsertionTool.h
@@ -112,7 +112,7 @@ public:
 	virtual TUniquePtr<UE::Geometry::FDynamicMeshOperator> MakeNewOperator() override;
 
 	UPROPERTY()
-	UEdgeLoopInsertionTool* Tool;
+	TObjectPtr<UEdgeLoopInsertionTool> Tool;
 };
 
 /** Tool for inserting (group) edge loops into a mesh. */
@@ -152,7 +152,7 @@ public:
 protected:
 
 	UPROPERTY()
-	UEdgeLoopInsertionProperties* Settings = nullptr;
+	TObjectPtr<UEdgeLoopInsertionProperties> Settings = nullptr;
 
 	TSharedPtr<FDynamicMesh3, ESPMode::ThreadSafe> CurrentMesh;
 	TSharedPtr<FGroupTopology, ESPMode::ThreadSafe> CurrentTopology;
@@ -168,7 +168,7 @@ protected:
 	FViewCameraState CameraState;
 
 	UPROPERTY()
-	UMeshOpPreviewWithBackgroundCompute* Preview;
+	TObjectPtr<UMeshOpPreviewWithBackgroundCompute> Preview;
 
 	FToolDataVisualizer ExistingEdgesRenderer;
 	FToolDataVisualizer PreviewEdgeRenderer;
diff --git a/Engine/Plugins/Experimental/MeshModelingToolset/Source/MeshModelingTools/Public/EditMeshPolygonsTool.h b/Engine/Plugins/Experimental/MeshModelingToolset/Source/MeshModelingTools/Public/EditMeshPolygonsTool.h
--- a/Engine/Plugins/Experimental/MeshModelingToolset/Source/MeshModelingTools/Public/EditMeshPolygonsTool.h
+++ b/Engine/Plugins/Experimental/MeshModelingToolset/Source/MeshModelingTools/Public/EditMeshPolygonsTool.h
@@ -553,58 +553,58 @@ protected:
 	bool bTriangleMode;		
 
 	UPROPERTY()
-	UDynamicMeshComponent* DynamicMeshComponent = nullptr;
+	TObjectPtr<UDynamicMeshComponent> DynamicMeshComponent = nullptr;
 
 	UPROPERTY()
-	UPolyEditCommonProperties* CommonProps;
+	TObjectPtr<UPolyEditCommonProperties> CommonProps;
 
 	UPROPERTY()
-	UEditMeshPolygonsToolActions* EditActions;
+	TObjectPtr<UEditMeshPolygonsToolActions> EditActions;
 	UPROPERTY()
-	UEditMeshPolygonsToolActions_Triangles* EditActions_Triangles;
+	TObjectPtr<UEditMeshPolygonsToolActions_Triangles> EditActions_Triangles;
 
 	UPROPERTY()
-	UEditMeshPolygonsToolEdgeActions* EditEdgeActions;
+	TObjectPtr<UEditMeshPolygonsToolEdgeActions> EditEdgeActions;
 	UPROPERTY()
-	UEditMeshPolygonsToolEdgeActions_Triangles* EditEdgeActions_Triangles;
+	TObjectPtr<UEditMeshPolygonsToolEdgeActions_Triangles> EditEdgeActions_Triangles;
 
 	UPROPERTY()
-	UEditMeshPolygonsToolUVActions* EditUVActions;
+	TObjectPtr<UEditMeshPolygonsToolUVActions> EditUVActions;
 
 	UPROPERTY()
-	UPolyEditExtrudeProperties* ExtrudeProperties;
+	TObjectPtr<UPolyEditExtrudeProperties> ExtrudeProperties;
 
 	UPROPERTY()
-	UPolyEditOffsetProperties* OffsetProperties;
+	TObjectPtr<UPolyEditOffsetProperties> OffsetProperties;
 
 	UPROPERTY()
-	UPolyEditInsetProperties* InsetProperties;
+	TObjectPtr<UPolyEditInsetProperties> InsetProperties;
 
 	UPROPERTY()
-	UPolyEditOutsetProperties* OutsetProperties;
+	TObjectPtr<UPolyEditOutsetProperties> OutsetProperties;
 
 	UPROPERTY()
-	UPolyEditCutProperties* CutProperties;
+	TObjectPtr<UPolyEditCutProperties> CutProperties;
 
 	UPROPERTY()
-	UPolyEditSetUVProperties* SetUVProperties;
+	TObjectPtr<UPolyEditSetUVProperties> SetUVProperties;
 
 
 	UPROPERTY()
-	UPolygonSelectionMechanic* SelectionMechanic;
+	TObjectPtr<UPolygonSelectionMechanic> SelectionMechanic;
 
 	UPROPERTY()
-	const UGroupTopologyStorableSelection* StoredToolSelection = nullptr;
+	TObjectPtr<const UGroupTopologyStorableSelection> StoredToolSelection = nullptr;
 	bool IsStoredToolSelectionUsable(const UGroupTopologyStorableSelection* StoredSelection);
 
 	bool bSelectionStateDirty = false;
 	void OnSelectionModifiedEvent();
 
 	UPROPERTY()
-	UMultiTransformer* MultiTransformer = nullptr;
+	TObjectPtr<UMultiTransformer> MultiTransformer = nullptr;
 
 	UPROPERTY()
-	UDragAlignmentMechanic* DragAlignmentMechanic = nullptr;
+	TObjectPtr<UDragAlignmentMechanic> DragAlignmentMechanic = nullptr;
 
 	void OnMultiTransformerTransformBegin();
 	void OnMultiTransformerTransformUpdate();
@@ -711,7 +711,7 @@ protected:
 	bool bPreviewUpdatePending = false;
 
 	UPROPERTY()
-	UPolyEditPreviewMesh* EditPreview;
+	TObjectPtr<UPolyEditPreviewMesh> EditPreview;
 
 	enum class EPreviewMaterialType
 	{
@@ -721,11 +721,11 @@ protected:
 	EPreviewMaterialType CurrentPreviewMaterial;
 
 	UPROPERTY()
-	UPlaneDistanceFromHitMechanic* ExtrudeHeightMechanic = nullptr;
+	TObjectPtr<UPlaneDistanceFromHitMechanic> ExtrudeHeightMechanic = nullptr;
 	UPROPERTY()
-	USpatialCurveDistanceMechanic* CurveDistMechanic = nullptr;
+	TObjectPtr<USpatialCurveDistanceMechanic> CurveDistMechanic = nullptr;
 	UPROPERTY()
-	UCollectSurfacePathMechanic* SurfacePathMechanic = nullptr;
+	TObjectPtr<UCollectSurfacePathMechanic> SurfacePathMechanic = nullptr;
 
 	//
 	// data for current drag
diff --git a/Engine/Plugins/Experimental/MeshModelingToolset/Source/MeshModelingTools/Public/EditNormalsTool.h b/Engine/Plugins/Experimental/MeshModelingToolset/Source/MeshModelingTools/Public/EditNormalsTool.h
--- a/Engine/Plugins/Experimental/MeshModelingToolset/Source/MeshModelingTools/Public/EditNormalsTool.h
+++ b/Engine/Plugins/Experimental/MeshModelingToolset/Source/MeshModelingTools/Public/EditNormalsTool.h
@@ -112,7 +112,7 @@ public:
 	virtual TUniquePtr<UE::Geometry::FDynamicMeshOperator> MakeNewOperator() override;
 
 	UPROPERTY()
-	UEditNormalsTool *Tool;
+	TObjectPtr<UEditNormalsTool> Tool;
 
 	int ComponentIndex;
 
@@ -153,14 +153,14 @@ public:
 protected:
 
 	UPROPERTY()
-	UEditNormalsToolProperties* BasicProperties;
+	TObjectPtr<UEditNormalsToolProperties> BasicProperties;
 
 	UPROPERTY()
-	UEditNormalsAdvancedProperties* AdvancedProperties;
+	TObjectPtr<UEditNormalsAdvancedProperties> AdvancedProperties;
 
 
 	UPROPERTY()
-	TArray<UMeshOpPreviewWithBackgroundCompute*> Previews;
+	TArray<TObjectPtr<UMeshOpPreviewWithBackgroundCompute>> Previews;
 
 
 protected:
diff --git a/Engine/Plugins/Experimental/MeshModelingToolset/Source/MeshModelingTools/Public/EditPivotTool.h b/Engine/Plugins/Experimental/MeshModelingToolset/Source/MeshModelingTools/Public/EditPivotTool.h
--- a/Engine/Plugins/Experimental/MeshModelingToolset/Source/MeshModelingTools/Public/EditPivotTool.h
+++ b/Engine/Plugins/Experimental/MeshModelingToolset/Source/MeshModelingTools/Public/EditPivotTool.h
@@ -79,10 +79,10 @@ struct FEditPivotTarget
 	GENERATED_BODY()
 
 	UPROPERTY()
-	UTransformProxy* TransformProxy = nullptr;
+	TObjectPtr<UTransformProxy> TransformProxy = nullptr;
 
 	UPROPERTY()
-	UTransformGizmo* TransformGizmo = nullptr;
+	TObjectPtr<UTransformGizmo> TransformGizmo = nullptr;
 };
 
 
@@ -183,10 +183,10 @@ public:
 
 public:
 	UPROPERTY()
-	UEditPivotToolProperties* TransformProps;
+	TObjectPtr<UEditPivotToolProperties> TransformProps;
 
 	UPROPERTY()
-	UEditPivotToolActionPropertySet* EditPivotActions;
+	TObjectPtr<UEditPivotToolActionPropertySet> EditPivotActions;
 
 	virtual void RequestAction(EEditPivotToolActions ActionType);
 
@@ -205,7 +205,7 @@ protected:
 	TArray<FEditPivotTarget> ActiveGizmos;
 
 	UPROPERTY()
-	UDragAlignmentMechanic* DragAlignmentMechanic = nullptr;
+	TObjectPtr<UDragAlignmentMechanic> DragAlignmentMechanic = nullptr;
 
 	void UpdateSetPivotModes(bool bEnableSetPivot);
 	void SetActiveGizmos_Single(bool bLocalRotations);
diff --git a/Engine/Plugins/Experimental/MeshModelingToolset/Source/MeshModelingTools/Public/EditUVIslandsTool.h b/Engine/Plugins/Experimental/MeshModelingToolset/Source/MeshModelingTools/Public/EditUVIslandsTool.h
--- a/Engine/Plugins/Experimental/MeshModelingToolset/Source/MeshModelingTools/Public/EditUVIslandsTool.h
+++ b/Engine/Plugins/Experimental/MeshModelingToolset/Source/MeshModelingTools/Public/EditUVIslandsTool.h
@@ -89,23 +89,23 @@ public:
 
 public:
 	UPROPERTY()
-	UExistingMeshMaterialProperties* MaterialSettings = nullptr;
+	TObjectPtr<UExistingMeshMaterialProperties> MaterialSettings = nullptr;
 
 	UPROPERTY()
-	UMaterialInstanceDynamic* CheckerMaterial = nullptr;
+	TObjectPtr<UMaterialInstanceDynamic> CheckerMaterial = nullptr;
 
 protected:
 	UPROPERTY()
-	UDynamicMeshComponent* DynamicMeshComponent;
+	TObjectPtr<UDynamicMeshComponent> DynamicMeshComponent;
 
 	UPROPERTY()
-	UPolygonSelectionMechanic* SelectionMechanic;
+	TObjectPtr<UPolygonSelectionMechanic> SelectionMechanic;
 
 	bool bSelectionStateDirty = false;
 	void OnSelectionModifiedEvent();
 
 	UPROPERTY()
-	UMultiTransformer* MultiTransformer = nullptr;
+	TObjectPtr<UMultiTransformer> MultiTransformer = nullptr;
 
 	void OnMultiTransformerTransformBegin();
 	void OnMultiTransformerTransformUpdate();
diff --git a/Engine/Plugins/Experimental/MeshModelingToolset/Source/MeshModelingTools/Public/GroupEdgeInsertionTool.h b/Engine/Plugins/Experimental/MeshModelingToolset/Source/MeshModelingTools/Public/GroupEdgeInsertionTool.h
--- a/Engine/Plugins/Experimental/MeshModelingToolset/Source/MeshModelingTools/Public/GroupEdgeInsertionTool.h
+++ b/Engine/Plugins/Experimental/MeshModelingToolset/Source/MeshModelingTools/Public/GroupEdgeInsertionTool.h
@@ -71,7 +71,7 @@ public:
 	virtual TUniquePtr<UE::Geometry::FDynamicMeshOperator> MakeNewOperator() override;
 
 	UPROPERTY()
-	UGroupEdgeInsertionTool* Tool;
+	TObjectPtr<UGroupEdgeInsertionTool> Tool;
 };
 
 /** Tool for inserting group edges into polygons of the mesh. */
@@ -120,10 +120,10 @@ public:
 protected:
 
 	UPROPERTY()
-	UGroupEdgeInsertionProperties* Settings = nullptr;
+	TObjectPtr<UGroupEdgeInsertionProperties> Settings = nullptr;
 
 	UPROPERTY()
-	UMeshOpPreviewWithBackgroundCompute* Preview;
+	TObjectPtr<UMeshOpPreviewWithBackgroundCompute> Preview;
 
 	TSharedPtr<FDynamicMesh3, ESPMode::ThreadSafe> CurrentMesh;
 	TSharedPtr<FGroupTopology, ESPMode::ThreadSafe> CurrentTopology;
diff --git a/Engine/Plugins/Experimental/MeshModelingToolset/Source/MeshModelingTools/Public/HoleFillTool.h b/Engine/Plugins/Experimental/MeshModelingToolset/Source/MeshModelingTools/Public/HoleFillTool.h
--- a/Engine/Plugins/Experimental/MeshModelingToolset/Source/MeshModelingTools/Public/HoleFillTool.h
+++ b/Engine/Plugins/Experimental/MeshModelingToolset/Source/MeshModelingTools/Public/HoleFillTool.h
@@ -206,7 +206,7 @@ public:
 	TUniquePtr<UE::Geometry::FDynamicMeshOperator> MakeNewOperator() override;
 
 	UPROPERTY()
-	UHoleFillTool* FillTool;
+	TObjectPtr<UHoleFillTool> FillTool;
 };
 
 /*
@@ -241,22 +241,22 @@ protected:
 	friend UHoleFillStatisticsProperties;
 
 	UPROPERTY()
-	USmoothHoleFillProperties* SmoothHoleFillProperties = nullptr;
+	TObjectPtr<USmoothHoleFillProperties> SmoothHoleFillProperties = nullptr;
 
 	UPROPERTY()
-	UHoleFillToolProperties* Properties = nullptr;
+	TObjectPtr<UHoleFillToolProperties> Properties = nullptr;
 
 	UPROPERTY()
-	UHoleFillToolActions* Actions = nullptr;
+	TObjectPtr<UHoleFillToolActions> Actions = nullptr;
 
 	UPROPERTY()
-	UHoleFillStatisticsProperties* Statistics = nullptr;
+	TObjectPtr<UHoleFillStatisticsProperties> Statistics = nullptr;
 
 	UPROPERTY()
-	UMeshOpPreviewWithBackgroundCompute* Preview = nullptr;
+	TObjectPtr<UMeshOpPreviewWithBackgroundCompute> Preview = nullptr;
 
 	UPROPERTY()
-	UPolygonSelectionMechanic* SelectionMechanic = nullptr;
+	TObjectPtr<UPolygonSelectionMechanic> SelectionMechanic = nullptr;
 
 	// Input mesh. Ownership shared with Op.
 	TSharedPtr<UE::Geometry::FDynamicMesh3, ESPMode::ThreadSafe> OriginalMesh;
diff --git a/Engine/Plugins/Experimental/MeshModelingToolset/Source/MeshModelingTools/Public/LatticeDeformerTool.h b/Engine/Plugins/Experimental/MeshModelingToolset/Source/MeshModelingTools/Public/LatticeDeformerTool.h
--- a/Engine/Plugins/Experimental/MeshModelingToolset/Source/MeshModelingTools/Public/LatticeDeformerTool.h
+++ b/Engine/Plugins/Experimental/MeshModelingToolset/Source/MeshModelingTools/Public/LatticeDeformerTool.h
@@ -121,7 +121,7 @@ public:
 	virtual TUniquePtr<UE::Geometry::FDynamicMeshOperator> MakeNewOperator() override;
 
 	UPROPERTY()
-	ULatticeDeformerTool* LatticeDeformerTool;
+	TObjectPtr<ULatticeDeformerTool> LatticeDeformerTool;
 };
 
 
@@ -157,13 +157,13 @@ protected:
 	TSharedPtr<UE::Geometry::FFFDLattice, ESPMode::ThreadSafe> Lattice;
 
 	UPROPERTY()
-	ULatticeControlPointsMechanic* ControlPointsMechanic = nullptr;
+	TObjectPtr<ULatticeControlPointsMechanic> ControlPointsMechanic = nullptr;
 
 	UPROPERTY()
-	ULatticeDeformerToolProperties* Settings = nullptr;
+	TObjectPtr<ULatticeDeformerToolProperties> Settings = nullptr;
 
 	UPROPERTY()
-	UMeshOpPreviewWithBackgroundCompute* Preview = nullptr;
+	TObjectPtr<UMeshOpPreviewWithBackgroundCompute> Preview = nullptr;
 
 	UPROPERTY()
 	bool bLatticeDeformed = false;
diff --git a/Engine/Plugins/Experimental/MeshModelingToolset/Source/MeshModelingTools/Public/MeshAttributePaintTool.h b/Engine/Plugins/Experimental/MeshModelingToolset/Source/MeshModelingTools/Public/MeshAttributePaintTool.h
--- a/Engine/Plugins/Experimental/MeshModelingToolset/Source/MeshModelingTools/Public/MeshAttributePaintTool.h
+++ b/Engine/Plugins/Experimental/MeshModelingToolset/Source/MeshModelingTools/Public/MeshAttributePaintTool.h
@@ -229,10 +229,10 @@ protected:
 
 protected:
 	UPROPERTY()
-	UMeshAttributePaintBrushOperationProperties* BrushActionProps;
+	TObjectPtr<UMeshAttributePaintBrushOperationProperties> BrushActionProps;
 
 	UPROPERTY()
-	UMeshAttributePaintToolProperties* AttribProps;
+	TObjectPtr<UMeshAttributePaintToolProperties> AttribProps;
 
 	TValueWatcher<int32> SelectedAttributeWatcher;
 
diff --git a/Engine/Plugins/Experimental/MeshModelingToolset/Source/MeshModelingTools/Public/MeshBoundaryToolBase.h b/Engine/Plugins/Experimental/MeshModelingToolset/Source/MeshModelingTools/Public/MeshBoundaryToolBase.h
--- a/Engine/Plugins/Experimental/MeshModelingToolset/Source/MeshModelingTools/Public/MeshBoundaryToolBase.h
+++ b/Engine/Plugins/Experimental/MeshModelingToolset/Source/MeshModelingTools/Public/MeshBoundaryToolBase.h
@@ -45,7 +45,7 @@ protected:
 	UE::Geometry::FDynamicMeshAABBTree3 MeshSpatial;
 
 	UPROPERTY()
-	UPolygonSelectionMechanic* SelectionMechanic = nullptr;
+	TObjectPtr<UPolygonSelectionMechanic> SelectionMechanic = nullptr;
 
 
 	// A variant of group topology that considers all triangles one group, so that group edges are boundary
diff --git a/Engine/Plugins/Experimental/MeshModelingToolset/Source/MeshModelingTools/Public/MeshGroupPaintTool.h b/Engine/Plugins/Experimental/MeshModelingToolset/Source/MeshModelingTools/Public/MeshGroupPaintTool.h
--- a/Engine/Plugins/Experimental/MeshModelingToolset/Source/MeshModelingTools/Public/MeshGroupPaintTool.h
+++ b/Engine/Plugins/Experimental/MeshModelingToolset/Source/MeshModelingTools/Public/MeshGroupPaintTool.h
@@ -305,20 +305,20 @@ public:
 public:
 
 	UPROPERTY()
-	UPolygroupLayersProperties* PolygroupLayerProperties;
+	TObjectPtr<UPolygroupLayersProperties> PolygroupLayerProperties;
 
 	/** Filters on paint brush */
 	UPROPERTY()
-	UGroupPaintBrushFilterProperties* FilterProperties;
+	TObjectPtr<UGroupPaintBrushFilterProperties> FilterProperties;
 
 
 private:
 	// This will be of type UGroupPaintBrushOpProps, we keep a ref so we can change active group ID on pick
 	UPROPERTY()
-	UGroupPaintBrushOpProps* PaintBrushOpProperties;
+	TObjectPtr<UGroupPaintBrushOpProps> PaintBrushOpProperties;
 
 	UPROPERTY()
-	UGroupEraseBrushOpProps* EraseBrushOpProperties;
+	TObjectPtr<UGroupEraseBrushOpProps> EraseBrushOpProperties;
 
 public:
 	void AllocateNewGroupAndSetAsCurrentAction();
@@ -357,7 +357,7 @@ public:
 	virtual void RequestAction(EMeshGroupPaintToolActions ActionType);
 
 	UPROPERTY()
-	UMeshGroupPaintToolFreezeActions* FreezeActions;
+	TObjectPtr<UMeshGroupPaintToolFreezeActions> FreezeActions;
 
 protected:
 	bool bHavePendingAction = false;
@@ -371,7 +371,7 @@ protected:
 	//
 public:
 	UPROPERTY()
-	UPolyLassoMarqueeMechanic* PolyLassoMechanic;
+	TObjectPtr<UPolyLassoMarqueeMechanic> PolyLassoMechanic;
 
 protected:
 	void OnPolyLassoFinished(const FCameraPolyLasso& Lasso, bool bCanceled);
@@ -383,10 +383,10 @@ protected:
 
 protected:
 	UPROPERTY()
-	UDynamicMeshComponent* DynamicMeshComponent;
+	TObjectPtr<UDynamicMeshComponent> DynamicMeshComponent;
 
 	UPROPERTY()
-	UMeshElementsVisualizer* MeshElementsDisplay;
+	TObjectPtr<UMeshElementsVisualizer> MeshElementsDisplay;
 
 	// realtime visualization
 	void OnDynamicMeshComponentChanged(UDynamicMeshComponent* Component, const FMeshVertexChange* Change, bool bRevert);
diff --git a/Engine/Plugins/Experimental/MeshModelingToolset/Source/MeshModelingTools/Public/MeshInspectorTool.h b/Engine/Plugins/Experimental/MeshModelingToolset/Source/MeshModelingTools/Public/MeshInspectorTool.h
--- a/Engine/Plugins/Experimental/MeshModelingToolset/Source/MeshModelingTools/Public/MeshInspectorTool.h
+++ b/Engine/Plugins/Experimental/MeshModelingToolset/Source/MeshModelingTools/Public/MeshInspectorTool.h
@@ -118,23 +118,23 @@ public:
 protected:
 
 	UPROPERTY()
-	UMeshInspectorProperties* Settings = nullptr;
+	TObjectPtr<UMeshInspectorProperties> Settings = nullptr;
 
 	UPROPERTY()
-	UExistingMeshMaterialProperties* MaterialSettings = nullptr;
+	TObjectPtr<UExistingMeshMaterialProperties> MaterialSettings = nullptr;
 
 
 	float LineWidthMultiplier = 1.0f;
 
 protected:
 	UPROPERTY()
-	UPreviewMesh* PreviewMesh;
+	TObjectPtr<UPreviewMesh> PreviewMesh;
 
 	UPROPERTY()
-	ULineSetComponent* DrawnLineSet;
+	TObjectPtr<ULineSetComponent> DrawnLineSet;
 
 	UPROPERTY()
-	UMaterialInterface* DefaultMaterial = nullptr;
+	TObjectPtr<UMaterialInterface> DefaultMaterial = nullptr;
 
 	TArray<int> BoundaryEdges;
 	TArray<int> BoundaryBowties;
diff --git a/Engine/Plugins/Experimental/MeshModelingToolset/Source/MeshModelingTools/Public/MeshSelectionTool.h b/Engine/Plugins/Experimental/MeshModelingToolset/Source/MeshModelingTools/Public/MeshSelectionTool.h
--- a/Engine/Plugins/Experimental/MeshModelingToolset/Source/MeshModelingTools/Public/MeshSelectionTool.h
+++ b/Engine/Plugins/Experimental/MeshModelingToolset/Source/MeshModelingTools/Public/MeshSelectionTool.h
@@ -311,25 +311,25 @@ public:
 	virtual void RequestAction(EMeshSelectionToolActions ActionType);
 
 	UPROPERTY()
-	UMeshSelectionToolProperties* SelectionProps;
+	TObjectPtr<UMeshSelectionToolProperties> SelectionProps;
 
 	UPROPERTY()
-	UMeshSelectionEditActions* SelectionActions;
+	TObjectPtr<UMeshSelectionEditActions> SelectionActions;
 
 	UPROPERTY()
-	UMeshSelectionToolActionPropertySet* EditActions;
+	TObjectPtr<UMeshSelectionToolActionPropertySet> EditActions;
 
 	UPROPERTY()
-	UMeshStatisticsProperties* MeshStatisticsProperties;
+	TObjectPtr<UMeshStatisticsProperties> MeshStatisticsProperties;
 
 	UPROPERTY()
-	UMeshElementsVisualizer* MeshElementsDisplay;
+	TObjectPtr<UMeshElementsVisualizer> MeshElementsDisplay;
 
 	UPROPERTY()
-	UMeshUVChannelProperties* UVChannelProperties = nullptr;
+	TObjectPtr<UMeshUVChannelProperties> UVChannelProperties = nullptr;
 
 	UPROPERTY()
-	UPolygroupLayersProperties* PolygroupLayerProperties = nullptr;
+	TObjectPtr<UPolygroupLayersProperties> PolygroupLayerProperties = nullptr;
 
 protected:
 	virtual UMeshSelectionToolActionPropertySet* CreateEditActions();
@@ -343,10 +343,10 @@ protected:
 
 protected:
 	UPROPERTY()
-	UMeshSelectionSet* Selection;
+	TObjectPtr<UMeshSelectionSet> Selection;
 
 	UPROPERTY()
-	TArray<AActor*> SpawnedActors;
+	TArray<TObjectPtr<AActor>> SpawnedActors;
 
 	UWorld* TargetWorld;
 
diff --git a/Engine/Plugins/Experimental/MeshModelingToolset/Source/MeshModelingTools/Public/MeshSpaceDeformerTool.h b/Engine/Plugins/Experimental/MeshModelingToolset/Source/MeshModelingTools/Public/MeshSpaceDeformerTool.h
--- a/Engine/Plugins/Experimental/MeshModelingToolset/Source/MeshModelingTools/Public/MeshSpaceDeformerTool.h
+++ b/Engine/Plugins/Experimental/MeshModelingToolset/Source/MeshModelingTools/Public/MeshSpaceDeformerTool.h
@@ -181,7 +181,7 @@ public:
 	virtual TUniquePtr<UE::Geometry::FDynamicMeshOperator> MakeNewOperator() override;
 
 	UPROPERTY()
-	UMeshSpaceDeformerTool* SpaceDeformerTool;  // back pointer
+	TObjectPtr<UMeshSpaceDeformerTool> SpaceDeformerTool;  // back pointer
 };
 
 
@@ -215,48 +215,48 @@ public:
 protected:
 
 	UPROPERTY()
-	UMeshSpaceDeformerToolProperties* Settings;
+	TObjectPtr<UMeshSpaceDeformerToolProperties> Settings;
 
 	UPROPERTY()
-	UMeshSpaceDeformerToolActionPropertySet* ToolActions;
+	TObjectPtr<UMeshSpaceDeformerToolActionPropertySet> ToolActions;
 
 	UPROPERTY()
-	UGizmoTransformChangeStateTarget* StateTarget = nullptr;
+	TObjectPtr<UGizmoTransformChangeStateTarget> StateTarget = nullptr;
 	
 	UPROPERTY()
-	UDragAlignmentMechanic* DragAlignmentMechanic;
+	TObjectPtr<UDragAlignmentMechanic> DragAlignmentMechanic;
 
 protected:
 
 	UPROPERTY()
-	UMeshOpPreviewWithBackgroundCompute* Preview = nullptr;   
+	TObjectPtr<UMeshOpPreviewWithBackgroundCompute> Preview = nullptr;   
 
 protected:	
 
 	TSharedPtr<UE::Geometry::FDynamicMesh3, ESPMode::ThreadSafe> OriginalDynamicMesh;
 
 	UPROPERTY()
-	UPreviewMesh* OriginalMeshPreview;
+	TObjectPtr<UPreviewMesh> OriginalMeshPreview;
 
 	UPROPERTY()
-	UIntervalGizmo* IntervalGizmo;
+	TObjectPtr<UIntervalGizmo> IntervalGizmo;
 
 	UPROPERTY()
-	UTransformGizmo* TransformGizmo;
+	TObjectPtr<UTransformGizmo> TransformGizmo;
 
 	UPROPERTY()
-	UTransformProxy* TransformProxy;
+	TObjectPtr<UTransformProxy> TransformProxy;
 
 	/** Interval Parameter sources that reflect UI settings. */
 
 	UPROPERTY()
-	UGizmoLocalFloatParameterSource* UpIntervalSource;
+	TObjectPtr<UGizmoLocalFloatParameterSource> UpIntervalSource;
 
 	UPROPERTY()
-	UGizmoLocalFloatParameterSource* DownIntervalSource;
+	TObjectPtr<UGizmoLocalFloatParameterSource> DownIntervalSource;
 
 	UPROPERTY()
-	UGizmoLocalFloatParameterSource* ForwardIntervalSource;
+	TObjectPtr<UGizmoLocalFloatParameterSource> ForwardIntervalSource;
 
 	// Button click support
 	EMeshSpaceDeformerToolAction PendingAction;
diff --git a/Engine/Plugins/Experimental/MeshModelingToolset/Source/MeshModelingTools/Public/MeshVertexSculptTool.h b/Engine/Plugins/Experimental/MeshModelingToolset/Source/MeshModelingTools/Public/MeshVertexSculptTool.h
--- a/Engine/Plugins/Experimental/MeshModelingToolset/Source/MeshModelingTools/Public/MeshVertexSculptTool.h
+++ b/Engine/Plugins/Experimental/MeshModelingToolset/Source/MeshModelingTools/Public/MeshVertexSculptTool.h
@@ -156,7 +156,7 @@ class MESHMODELINGTOOLS_API UVertexBrushAlphaProperties : public UInteractiveToo
 public:
 	/** Alpha mask applied to brush stamp. Red channel is used. */
 	UPROPERTY(EditAnywhere, Category = Alpha, meta = (DisplayName = "Alpha Mask"))
-	UTexture2D* Alpha = nullptr;
+	TObjectPtr<UTexture2D> Alpha = nullptr;
 
 	/** Alpha is rotated by this angle, inside the brush stamp frame (vertically aligned) */
 	UPROPERTY(EditAnywhere, Category = Alpha, meta = (DisplayName = "Rotation", UIMin = "-180.0", UIMax = "180.0", ClampMin = "-360.0", ClampMax = "360.0"))
@@ -199,13 +199,13 @@ public:
 
 	/** Properties that control sculpting*/
 	UPROPERTY()
-	UVertexBrushSculptProperties* SculptProperties;
+	TObjectPtr<UVertexBrushSculptProperties> SculptProperties;
 
 	UPROPERTY()
-	UVertexBrushAlphaProperties* AlphaProperties;
+	TObjectPtr<UVertexBrushAlphaProperties> AlphaProperties;
 
 	UPROPERTY()
-	UTexture2D* BrushAlpha;
+	TObjectPtr<UTexture2D> BrushAlpha;
 
 public:
 	virtual void IncreaseBrushSpeedAction() override;
@@ -229,7 +229,7 @@ protected:
 
 protected:
 	UPROPERTY()
-	UDynamicMeshComponent* DynamicMeshComponent;
+	TObjectPtr<UDynamicMeshComponent> DynamicMeshComponent;
 
 	// realtime visualization
 	void OnDynamicMeshComponentChanged(UDynamicMeshComponent* Component, const FMeshVertexChange* Change, bool bRevert);
diff --git a/Engine/Plugins/Experimental/MeshModelingToolset/Source/MeshModelingTools/Public/MirrorTool.h b/Engine/Plugins/Experimental/MeshModelingToolset/Source/MeshModelingTools/Public/MirrorTool.h
--- a/Engine/Plugins/Experimental/MeshModelingToolset/Source/MeshModelingTools/Public/MirrorTool.h
+++ b/Engine/Plugins/Experimental/MeshModelingToolset/Source/MeshModelingTools/Public/MirrorTool.h
@@ -126,7 +126,7 @@ public:
 	virtual TUniquePtr<UE::Geometry::FDynamicMeshOperator> MakeNewOperator() override;
 
 	UPROPERTY()
-	UMirrorTool* MirrorTool;
+	TObjectPtr<UMirrorTool> MirrorTool;
 
 	/** Index of the component within MirrorTool->ComponentTargets that this factory creates an operator for. */
 	int ComponentIndex;
@@ -219,16 +219,16 @@ public:
 protected:
 
 	UPROPERTY()
-	UMirrorToolProperties* Settings = nullptr;
+	TObjectPtr<UMirrorToolProperties> Settings = nullptr;
 
 	UPROPERTY()
-	UMirrorToolActionPropertySet* ToolActions = nullptr;
+	TObjectPtr<UMirrorToolActionPropertySet> ToolActions = nullptr;
 
 	UPROPERTY()
-	TArray<UDynamicMeshReplacementChangeTarget*> MeshesToMirror;
+	TArray<TObjectPtr<UDynamicMeshReplacementChangeTarget>> MeshesToMirror;
 
 	UPROPERTY()
-	TArray<UMeshOpPreviewWithBackgroundCompute*> Previews;
+	TArray<TObjectPtr<UMeshOpPreviewWithBackgroundCompute>> Previews;
 
 	FVector3d MirrorPlaneOrigin = FVector3d::Zero();
 	FVector3d MirrorPlaneNormal = FVector3d::UnitZ();
@@ -240,7 +240,7 @@ protected:
 	bool bSnappingToggle = false;		// toggled by hotkey (shift)
 
 	UPROPERTY()
-	UConstructionPlaneMechanic* PlaneMechanic;
+	TObjectPtr<UConstructionPlaneMechanic> PlaneMechanic;
 
 	EMirrorToolAction PendingAction;
 	FBox CombinedBounds;
diff --git a/Engine/Plugins/Experimental/MeshModelingToolset/Source/MeshModelingTools/Public/OffsetMeshTool.h b/Engine/Plugins/Experimental/MeshModelingToolset/Source/MeshModelingTools/Public/OffsetMeshTool.h
--- a/Engine/Plugins/Experimental/MeshModelingToolset/Source/MeshModelingTools/Public/OffsetMeshTool.h
+++ b/Engine/Plugins/Experimental/MeshModelingToolset/Source/MeshModelingTools/Public/OffsetMeshTool.h
@@ -125,16 +125,16 @@ public:
 
 protected:
 	UPROPERTY()
-	UOffsetMeshToolProperties* OffsetProperties = nullptr;
+	TObjectPtr<UOffsetMeshToolProperties> OffsetProperties = nullptr;
 
 	UPROPERTY()
-	UIterativeOffsetProperties* IterativeProperties = nullptr;
+	TObjectPtr<UIterativeOffsetProperties> IterativeProperties = nullptr;
 
 	UPROPERTY()
-	UImplicitOffsetProperties* ImplicitProperties = nullptr;
+	TObjectPtr<UImplicitOffsetProperties> ImplicitProperties = nullptr;
 
 	UPROPERTY()
-	UOffsetWeightMapSetProperties* WeightMapProperties = nullptr;
+	TObjectPtr<UOffsetWeightMapSetProperties> WeightMapProperties = nullptr;
 };
 
 
diff --git a/Engine/Plugins/Experimental/MeshModelingToolset/Source/MeshModelingTools/Public/Physics/ExtractCollisionGeometryTool.h b/Engine/Plugins/Experimental/MeshModelingToolset/Source/MeshModelingTools/Public/Physics/ExtractCollisionGeometryTool.h
--- a/Engine/Plugins/Experimental/MeshModelingToolset/Source/MeshModelingTools/Public/Physics/ExtractCollisionGeometryTool.h
+++ b/Engine/Plugins/Experimental/MeshModelingToolset/Source/MeshModelingTools/Public/Physics/ExtractCollisionGeometryTool.h
@@ -84,23 +84,23 @@ protected:
 
 	/** Property set for type of output object (StaticMesh, Volume, etc) */
 	UPROPERTY()
-	UCreateMeshObjectTypeProperties* OutputTypeProperties;
+	TObjectPtr<UCreateMeshObjectTypeProperties> OutputTypeProperties;
 
 	UPROPERTY()
-	UExtractCollisionToolProperties* Settings;
+	TObjectPtr<UExtractCollisionToolProperties> Settings;
 
 	UPROPERTY()
-	UCollisionGeometryVisualizationProperties* VizSettings = nullptr;
+	TObjectPtr<UCollisionGeometryVisualizationProperties> VizSettings = nullptr;
 
 	UPROPERTY()
-	UPhysicsObjectToolPropertySet* ObjectProps;
+	TObjectPtr<UPhysicsObjectToolPropertySet> ObjectProps;
 
 protected:
 	UPROPERTY()
-	UPreviewGeometry* PreviewElements;
+	TObjectPtr<UPreviewGeometry> PreviewElements;
 
 	UPROPERTY()
-	UPreviewMesh* PreviewMesh;
+	TObjectPtr<UPreviewMesh> PreviewMesh;
 
 	// these are TSharedPtr because TPimplPtr cannot currently be added to a TArray?
 	TSharedPtr<FPhysicsDataCollection> PhysicsInfo;
diff --git a/Engine/Plugins/Experimental/MeshModelingToolset/Source/MeshModelingTools/Public/Physics/PhysicsInspectorTool.h b/Engine/Plugins/Experimental/MeshModelingToolset/Source/MeshModelingTools/Public/Physics/PhysicsInspectorTool.h
--- a/Engine/Plugins/Experimental/MeshModelingToolset/Source/MeshModelingTools/Public/Physics/PhysicsInspectorTool.h
+++ b/Engine/Plugins/Experimental/MeshModelingToolset/Source/MeshModelingTools/Public/Physics/PhysicsInspectorTool.h
@@ -46,17 +46,17 @@ public:
 protected:
 
 	UPROPERTY()
-	UCollisionGeometryVisualizationProperties* VizSettings = nullptr;
+	TObjectPtr<UCollisionGeometryVisualizationProperties> VizSettings = nullptr;
 
 	UPROPERTY()
-	TArray<UPhysicsObjectToolPropertySet*> ObjectData;
+	TArray<TObjectPtr<UPhysicsObjectToolPropertySet>> ObjectData;
 
 	UPROPERTY()
-	UMaterialInterface* LineMaterial = nullptr;
+	TObjectPtr<UMaterialInterface> LineMaterial = nullptr;
 
 protected:
 	UPROPERTY()
-	TArray<UPreviewGeometry*> PreviewElements;
+	TArray<TObjectPtr<UPreviewGeometry>> PreviewElements;
 
 	// these are TSharedPtr because TPimplPtr cannot currently be added to a TArray?
 	TArray<TSharedPtr<FPhysicsDataCollection>> PhysicsInfos;
diff --git a/Engine/Plugins/Experimental/MeshModelingToolset/Source/MeshModelingTools/Public/Physics/SetCollisionGeometryTool.h b/Engine/Plugins/Experimental/MeshModelingToolset/Source/MeshModelingTools/Public/Physics/SetCollisionGeometryTool.h
--- a/Engine/Plugins/Experimental/MeshModelingToolset/Source/MeshModelingTools/Public/Physics/SetCollisionGeometryTool.h
+++ b/Engine/Plugins/Experimental/MeshModelingToolset/Source/MeshModelingTools/Public/Physics/SetCollisionGeometryTool.h
@@ -156,23 +156,23 @@ public:
 protected:
 
 	UPROPERTY()
-	USetCollisionGeometryToolProperties* Settings = nullptr;
+	TObjectPtr<USetCollisionGeometryToolProperties> Settings = nullptr;
 
 	UPROPERTY()
-	UPolygroupLayersProperties* PolygroupLayerProperties = nullptr;
+	TObjectPtr<UPolygroupLayersProperties> PolygroupLayerProperties = nullptr;
 
 	UPROPERTY()
-	UCollisionGeometryVisualizationProperties* VizSettings = nullptr;
+	TObjectPtr<UCollisionGeometryVisualizationProperties> VizSettings = nullptr;
 
 	UPROPERTY()
-	UPhysicsObjectToolPropertySet* CollisionProps;
+	TObjectPtr<UPhysicsObjectToolPropertySet> CollisionProps;
 
 	UPROPERTY()
-	UMaterialInterface* LineMaterial = nullptr;
+	TObjectPtr<UMaterialInterface> LineMaterial = nullptr;
 
 protected:
 	UPROPERTY()
-	UPreviewGeometry* PreviewGeom;
+	TObjectPtr<UPreviewGeometry> PreviewGeom;
 
 	TArray<int32> SourceObjectIndices;
 	bool bSourcesHidden = false;
diff --git a/Engine/Plugins/Experimental/MeshModelingToolset/Source/MeshModelingTools/Public/PlaneCutTool.h b/Engine/Plugins/Experimental/MeshModelingToolset/Source/MeshModelingTools/Public/PlaneCutTool.h
--- a/Engine/Plugins/Experimental/MeshModelingToolset/Source/MeshModelingTools/Public/PlaneCutTool.h
+++ b/Engine/Plugins/Experimental/MeshModelingToolset/Source/MeshModelingTools/Public/PlaneCutTool.h
@@ -120,7 +120,7 @@ public:
 	virtual TUniquePtr<UE::Geometry::FDynamicMeshOperator> MakeNewOperator() override;
 
 	UPROPERTY()
-	UPlaneCutTool *CutTool;
+	TObjectPtr<UPlaneCutTool> CutTool;
 
 	int ComponentIndex;
 };
@@ -166,13 +166,13 @@ public:
 protected:
 
 	UPROPERTY()
-	UPlaneCutToolProperties* BasicProperties;
+	TObjectPtr<UPlaneCutToolProperties> BasicProperties;
 
 	UPROPERTY()
-	UAcceptOutputProperties* AcceptProperties;
+	TObjectPtr<UAcceptOutputProperties> AcceptProperties;
 
 	UPROPERTY()
-	TArray<UMeshOpPreviewWithBackgroundCompute*> Previews;
+	TArray<TObjectPtr<UMeshOpPreviewWithBackgroundCompute>> Previews;
 
 
 	/// Action buttons.
@@ -196,10 +196,10 @@ protected:
 protected:
 
 	UPROPERTY()
-	TArray<UDynamicMeshReplacementChangeTarget*> MeshesToCut;
+	TArray<TObjectPtr<UDynamicMeshReplacementChangeTarget>> MeshesToCut;
 
 	UPROPERTY()
-	UConstructionPlaneMechanic* PlaneMechanic = nullptr;
+	TObjectPtr<UConstructionPlaneMechanic> PlaneMechanic = nullptr;
 
 	// Cutting plane
 	UE::Geometry::FFrame3d CutPlaneWorld;
diff --git a/Engine/Plugins/Experimental/MeshModelingToolset/Source/MeshModelingTools/Public/Properties/MeshMaterialProperties.h b/Engine/Plugins/Experimental/MeshModelingToolset/Source/MeshModelingTools/Public/Properties/MeshMaterialProperties.h
--- a/Engine/Plugins/Experimental/MeshModelingToolset/Source/MeshModelingTools/Public/Properties/MeshMaterialProperties.h
+++ b/Engine/Plugins/Experimental/MeshModelingToolset/Source/MeshModelingTools/Public/Properties/MeshMaterialProperties.h
@@ -76,10 +76,10 @@ public:
 	int UVChannel = 0;
 
 	UPROPERTY(EditAnywhere, Category = MaterialPreview, meta = (EditCondition = "MaterialMode == ESetMeshMaterialMode::Override"))
-	UMaterialInterface* OverrideMaterial = nullptr;
+	TObjectPtr<UMaterialInterface> OverrideMaterial = nullptr;
 
 	UPROPERTY(meta = (TransientToolProperty))
-	UMaterialInstanceDynamic* CheckerMaterial = nullptr;
+	TObjectPtr<UMaterialInstanceDynamic> CheckerMaterial = nullptr;
 
 	// Needs custom restore in order to call setup
 	virtual void RestoreProperties(UInteractiveTool* RestoreToTool, const FString& CacheIdentifier = TEXT("")) override;
@@ -130,7 +130,7 @@ public:
 
 	/** Image used in Image-Based Material */
 	UPROPERTY(EditAnywhere, Category = Rendering, meta = (EditConditionHides, EditCondition = "MaterialMode == EMeshEditingMaterialModes::CustomImage", TransientToolProperty) )
-	UTexture2D* Image;
+	TObjectPtr<UTexture2D> Image;
 
 	//~ Could have used the same property as Color, above, but the user may want different saved values for the two
 	UPROPERTY(EditAnywhere, Category = Rendering, AdvancedDisplay, meta = (EditConditionHides, EditCondition = "MaterialMode == EMeshEditingMaterialModes::Transparent", DisplayName = "Color"))
diff --git a/Engine/Plugins/Experimental/MeshModelingToolset/Source/MeshModelingTools/Public/RecomputeUVsTool.h b/Engine/Plugins/Experimental/MeshModelingToolset/Source/MeshModelingTools/Public/RecomputeUVsTool.h
--- a/Engine/Plugins/Experimental/MeshModelingToolset/Source/MeshModelingTools/Public/RecomputeUVsTool.h
+++ b/Engine/Plugins/Experimental/MeshModelingToolset/Source/MeshModelingTools/Public/RecomputeUVsTool.h
@@ -126,20 +126,20 @@ public:
 
 protected:
 	UPROPERTY()
-	UMeshUVChannelProperties* UVChannelProperties = nullptr;
+	TObjectPtr<UMeshUVChannelProperties> UVChannelProperties = nullptr;
 
 	UPROPERTY()
-	URecomputeUVsToolProperties* Settings = nullptr;
+	TObjectPtr<URecomputeUVsToolProperties> Settings = nullptr;
 
 	UPROPERTY()
-	UPolygroupLayersProperties* PolygroupLayerProperties = nullptr;
+	TObjectPtr<UPolygroupLayersProperties> PolygroupLayerProperties = nullptr;
 
 	UPROPERTY()
-	UExistingMeshMaterialProperties* MaterialSettings = nullptr;
+	TObjectPtr<UExistingMeshMaterialProperties> MaterialSettings = nullptr;
 
 protected:
 	UPROPERTY()
-	UMeshOpPreviewWithBackgroundCompute* Preview = nullptr;
+	TObjectPtr<UMeshOpPreviewWithBackgroundCompute> Preview = nullptr;
 
 protected:
 	TSharedPtr<UE::Geometry::FDynamicMesh3, ESPMode::ThreadSafe> InputMesh;
diff --git a/Engine/Plugins/Experimental/MeshModelingToolset/Source/MeshModelingTools/Public/RemeshMeshTool.h b/Engine/Plugins/Experimental/MeshModelingToolset/Source/MeshModelingTools/Public/RemeshMeshTool.h
--- a/Engine/Plugins/Experimental/MeshModelingToolset/Source/MeshModelingTools/Public/RemeshMeshTool.h
+++ b/Engine/Plugins/Experimental/MeshModelingToolset/Source/MeshModelingTools/Public/RemeshMeshTool.h
@@ -127,16 +127,16 @@ public:
 	virtual TUniquePtr<UE::Geometry::FDynamicMeshOperator> MakeNewOperator() override;
 
 	UPROPERTY()
-	URemeshMeshToolProperties* BasicProperties;
+	TObjectPtr<URemeshMeshToolProperties> BasicProperties;
 
 	UPROPERTY()
-	UMeshStatisticsProperties* MeshStatisticsProperties;
+	TObjectPtr<UMeshStatisticsProperties> MeshStatisticsProperties;
 
 	UPROPERTY()
-	UMeshOpPreviewWithBackgroundCompute* Preview;
+	TObjectPtr<UMeshOpPreviewWithBackgroundCompute> Preview;
 
 	UPROPERTY()
-	UMeshElementsVisualizer* MeshElementsDisplay;
+	TObjectPtr<UMeshElementsVisualizer> MeshElementsDisplay;
 
 protected:
 
diff --git a/Engine/Plugins/Experimental/MeshModelingToolset/Source/MeshModelingTools/Public/RemoveOccludedTrianglesTool.h b/Engine/Plugins/Experimental/MeshModelingToolset/Source/MeshModelingTools/Public/RemoveOccludedTrianglesTool.h
--- a/Engine/Plugins/Experimental/MeshModelingToolset/Source/MeshModelingTools/Public/RemoveOccludedTrianglesTool.h
+++ b/Engine/Plugins/Experimental/MeshModelingToolset/Source/MeshModelingTools/Public/RemoveOccludedTrianglesTool.h
@@ -143,7 +143,7 @@ public:
 	virtual TUniquePtr<UE::Geometry::FDynamicMeshOperator> MakeNewOperator() override;
 
 	UPROPERTY()
-	URemoveOccludedTrianglesTool *Tool;
+	TObjectPtr<URemoveOccludedTrianglesTool> Tool;
 
 	int PreviewIdx;
 };
@@ -183,21 +183,21 @@ public:
 
 protected:
 	UPROPERTY()
-	URemoveOccludedTrianglesToolProperties* BasicProperties;
+	TObjectPtr<URemoveOccludedTrianglesToolProperties> BasicProperties;
 
 	UPROPERTY()
-	UPolygroupLayersProperties* PolygroupLayersProperties;
+	TObjectPtr<UPolygroupLayersProperties> PolygroupLayersProperties;
 
 	UPROPERTY()
-	URemoveOccludedTrianglesAdvancedProperties* AdvancedProperties;
+	TObjectPtr<URemoveOccludedTrianglesAdvancedProperties> AdvancedProperties;
 
 	UPROPERTY()
-	TArray<UMeshOpPreviewWithBackgroundCompute*> Previews;
+	TArray<TObjectPtr<UMeshOpPreviewWithBackgroundCompute>> Previews;
 
 	// When multiple meshes in the selection correspond to the same asset, only one needs a PreviewWithBackgroundCompute
 	//  all others just get a plain PreviewMesh copy that is updated via OnMeshUpdated broadcast from the source Preview
 	UPROPERTY()
-	TArray<UPreviewMesh*> PreviewCopies;
+	TArray<TObjectPtr<UPreviewMesh>> PreviewCopies;
 
 
 protected:
diff --git a/Engine/Plugins/Experimental/MeshModelingToolset/Source/MeshModelingTools/Public/RevolveBoundaryTool.h b/Engine/Plugins/Experimental/MeshModelingToolset/Source/MeshModelingTools/Public/RevolveBoundaryTool.h
--- a/Engine/Plugins/Experimental/MeshModelingToolset/Source/MeshModelingTools/Public/RevolveBoundaryTool.h
+++ b/Engine/Plugins/Experimental/MeshModelingToolset/Source/MeshModelingTools/Public/RevolveBoundaryTool.h
@@ -40,7 +40,7 @@ public:
 	TUniquePtr<UE::Geometry::FDynamicMeshOperator> MakeNewOperator() override;
 
 	UPROPERTY()
-	URevolveBoundaryTool* RevolveBoundaryTool;
+	TObjectPtr<URevolveBoundaryTool> RevolveBoundaryTool;
 };
 
 UCLASS()
@@ -99,16 +99,16 @@ protected:
 	int32 AlignAxisModifier = 2;
 
 	UPROPERTY()
-	URevolveBoundaryToolProperties* Settings = nullptr;
+	TObjectPtr<URevolveBoundaryToolProperties> Settings = nullptr;
 
 	UPROPERTY()
-	UNewMeshMaterialProperties* MaterialProperties;
+	TObjectPtr<UNewMeshMaterialProperties> MaterialProperties;
 
 	UPROPERTY()
-	UConstructionPlaneMechanic* PlaneMechanic = nullptr;
+	TObjectPtr<UConstructionPlaneMechanic> PlaneMechanic = nullptr;
 
 	UPROPERTY()
-	UMeshOpPreviewWithBackgroundCompute* Preview = nullptr;
+	TObjectPtr<UMeshOpPreviewWithBackgroundCompute> Preview = nullptr;
 
 	FVector3d RevolutionAxisOrigin;
 	FVector3d RevolutionAxisDirection;
diff --git a/Engine/Plugins/Experimental/MeshModelingToolset/Source/MeshModelingTools/Public/Sculpting/MeshSculptToolBase.h b/Engine/Plugins/Experimental/MeshModelingToolset/Source/MeshModelingTools/Public/Sculpting/MeshSculptToolBase.h
--- a/Engine/Plugins/Experimental/MeshModelingToolset/Source/MeshModelingTools/Public/Sculpting/MeshSculptToolBase.h
+++ b/Engine/Plugins/Experimental/MeshModelingToolset/Source/MeshModelingTools/Public/Sculpting/MeshSculptToolBase.h
@@ -198,11 +198,11 @@ protected:
 public:
 	/** Properties that control brush size/etc */
 	UPROPERTY()
-	USculptBrushProperties* BrushProperties;
+	TObjectPtr<USculptBrushProperties> BrushProperties;
 
 	/** Properties for 3D workplane / gizmo */
 	UPROPERTY()
-	UWorkPlaneProperties* GizmoProperties;
+	TObjectPtr<UWorkPlaneProperties> GizmoProperties;
 
 
 protected:
@@ -247,12 +247,12 @@ protected:
 	//
 protected:
 	UPROPERTY()
-	TMap<int32, UMeshSculptBrushOpProps*> BrushOpPropSets;
+	TMap<int32, TObjectPtr<UMeshSculptBrushOpProps>> BrushOpPropSets;
 
 	TMap<int32, TUniquePtr<FMeshSculptBrushOpFactory>> BrushOpFactories;
 
 	UPROPERTY()
-	TMap<int32, UMeshSculptBrushOpProps*> SecondaryBrushOpPropSets;
+	TMap<int32, TObjectPtr<UMeshSculptBrushOpProps>> SecondaryBrushOpPropSets;
 
 	TMap<int32, TUniquePtr<FMeshSculptBrushOpFactory>> SecondaryBrushOpFactories;
 
@@ -417,10 +417,10 @@ protected:
 	//
 public:
 	UPROPERTY()
-	UMeshEditingViewProperties* ViewProperties;
+	TObjectPtr<UMeshEditingViewProperties> ViewProperties;
 
 	UPROPERTY()
-	UMaterialInstanceDynamic* ActiveOverrideMaterial;
+	TObjectPtr<UMaterialInstanceDynamic> ActiveOverrideMaterial;
 
 protected:
 	virtual void SetViewPropertiesEnabled(bool bNewValue);
@@ -451,17 +451,17 @@ protected:
 
 protected:
 	UPROPERTY()
-	UBrushStampIndicator* BrushIndicator;
+	TObjectPtr<UBrushStampIndicator> BrushIndicator;
 
 	UPROPERTY()
 	bool bIsVolumetricIndicator;
 
 
 	UPROPERTY()
-	UMaterialInstanceDynamic* BrushIndicatorMaterial;
+	TObjectPtr<UMaterialInstanceDynamic> BrushIndicatorMaterial;
 
 	UPROPERTY()
-	UPreviewMesh* BrushIndicatorMesh;
+	TObjectPtr<UPreviewMesh> BrushIndicatorMesh;
 
 	// creates default sphere indicator
 	UPreviewMesh* MakeDefaultIndicatorSphereMesh(UObject* Parent, UWorld* World, int Resolution = 32);
@@ -475,10 +475,10 @@ protected:
 public:
 	// plane gizmo
 	UPROPERTY()
-	UTransformGizmo* PlaneTransformGizmo;
+	TObjectPtr<UTransformGizmo> PlaneTransformGizmo;
 
 	UPROPERTY()
-	UTransformProxy* PlaneTransformProxy;
+	TObjectPtr<UTransformProxy> PlaneTransformProxy;
 
 protected:
 	virtual void UpdateWorkPlane();
diff --git a/Engine/Plugins/Experimental/MeshModelingToolset/Source/MeshModelingTools/Public/SeamSculptTool.h b/Engine/Plugins/Experimental/MeshModelingToolset/Source/MeshModelingTools/Public/SeamSculptTool.h
--- a/Engine/Plugins/Experimental/MeshModelingToolset/Source/MeshModelingTools/Public/SeamSculptTool.h
+++ b/Engine/Plugins/Experimental/MeshModelingToolset/Source/MeshModelingTools/Public/SeamSculptTool.h
@@ -77,12 +77,12 @@ protected:
 
 public:
 	UPROPERTY()
-	USeamSculptToolProperties* Settings;
+	TObjectPtr<USeamSculptToolProperties> Settings;
 
 
 protected:
 	UPROPERTY()
-	UPreviewGeometry* PreviewGeom;
+	TObjectPtr<UPreviewGeometry> PreviewGeom;
 
 	TSharedPtr<FDynamicMesh3, ESPMode::ThreadSafe> InputMesh;
 	UE::Geometry::FTransform3d MeshTransform;
diff --git a/Engine/Plugins/Experimental/MeshModelingToolset/Source/MeshModelingTools/Public/SelfUnionMeshesTool.h b/Engine/Plugins/Experimental/MeshModelingToolset/Source/MeshModelingTools/Public/SelfUnionMeshesTool.h
--- a/Engine/Plugins/Experimental/MeshModelingToolset/Source/MeshModelingTools/Public/SelfUnionMeshesTool.h
+++ b/Engine/Plugins/Experimental/MeshModelingToolset/Source/MeshModelingTools/Public/SelfUnionMeshesTool.h
@@ -87,10 +87,10 @@ protected:
 protected:
 
 	UPROPERTY()
-	USelfUnionMeshesToolProperties* Properties;
+	TObjectPtr<USelfUnionMeshesToolProperties> Properties;
 
 	UPROPERTY()
-	ULineSetComponent* DrawnLineSet;
+	TObjectPtr<ULineSetComponent> DrawnLineSet;
 
 	TSharedPtr<FDynamicMesh3, ESPMode::ThreadSafe> CombinedSourceMeshes;
 
diff --git a/Engine/Plugins/Experimental/MeshModelingToolset/Source/MeshModelingTools/Public/SkinWeightsBindingTool.h b/Engine/Plugins/Experimental/MeshModelingToolset/Source/MeshModelingTools/Public/SkinWeightsBindingTool.h
--- a/Engine/Plugins/Experimental/MeshModelingToolset/Source/MeshModelingTools/Public/SkinWeightsBindingTool.h
+++ b/Engine/Plugins/Experimental/MeshModelingToolset/Source/MeshModelingTools/Public/SkinWeightsBindingTool.h
@@ -106,10 +106,10 @@ public:
 	TUniquePtr<UE::Geometry::FDynamicMeshOperator> MakeNewOperator() override;
 
 	UPROPERTY()
-	USkinWeightsBindingToolProperties* Properties = nullptr;
+	TObjectPtr<USkinWeightsBindingToolProperties> Properties = nullptr;
 	
 	UPROPERTY()
-	UMeshOpPreviewWithBackgroundCompute* Preview = nullptr;
+	TObjectPtr<UMeshOpPreviewWithBackgroundCompute> Preview = nullptr;
 
 protected:
 	UWorld* TargetWorld = nullptr;
diff --git a/Engine/Plugins/Experimental/MeshModelingToolset/Source/MeshModelingTools/Public/SkinWeightsPaintTool.h b/Engine/Plugins/Experimental/MeshModelingToolset/Source/MeshModelingTools/Public/SkinWeightsPaintTool.h
--- a/Engine/Plugins/Experimental/MeshModelingToolset/Source/MeshModelingTools/Public/SkinWeightsPaintTool.h
+++ b/Engine/Plugins/Experimental/MeshModelingToolset/Source/MeshModelingTools/Public/SkinWeightsPaintTool.h
@@ -113,7 +113,7 @@ protected:
 	void CalculateVertexROI(const FBrushStampData& Stamp, TArray<int>& VertexROI);
 
 	UPROPERTY()
-	USkinWeightsPaintToolProperties* ToolProps;
+	TObjectPtr<USkinWeightsPaintToolProperties> ToolProps;
 
 	TValueWatcher<FBoneReference> CurrentBoneWatcher;
 
diff --git a/Engine/Plugins/Experimental/MeshModelingToolset/Source/MeshModelingTools/Public/SmoothMeshTool.h b/Engine/Plugins/Experimental/MeshModelingToolset/Source/MeshModelingTools/Public/SmoothMeshTool.h
--- a/Engine/Plugins/Experimental/MeshModelingToolset/Source/MeshModelingTools/Public/SmoothMeshTool.h
+++ b/Engine/Plugins/Experimental/MeshModelingToolset/Source/MeshModelingTools/Public/SmoothMeshTool.h
@@ -145,19 +145,19 @@ public:
 
 protected:
 	UPROPERTY()
-	USmoothMeshToolProperties* SmoothProperties = nullptr;
+	TObjectPtr<USmoothMeshToolProperties> SmoothProperties = nullptr;
 
 	UPROPERTY()
-	UIterativeSmoothProperties* IterativeProperties = nullptr;
+	TObjectPtr<UIterativeSmoothProperties> IterativeProperties = nullptr;
 
 	UPROPERTY()
-	UDiffusionSmoothProperties* DiffusionProperties = nullptr;
+	TObjectPtr<UDiffusionSmoothProperties> DiffusionProperties = nullptr;
 
 	UPROPERTY()
-	UImplicitSmoothProperties* ImplicitProperties = nullptr;
+	TObjectPtr<UImplicitSmoothProperties> ImplicitProperties = nullptr;
 
 	UPROPERTY()
-	USmoothWeightMapSetProperties* WeightMapProperties = nullptr;
+	TObjectPtr<USmoothWeightMapSetProperties> WeightMapProperties = nullptr;
 };
 
 
