MeshModelingTools: Fix UV overlay duplication crash bug

#rnx

diff --git a/Engine/Plugins/Experimental/GeometryProcessing/Source/DynamicMesh/Private/Parameterization/DynamicMeshUVEditor.cpp b/Engine/Plugins/Experimental/GeometryProcessing/Source/DynamicMesh/Private/Parameterization/DynamicMeshUVEditor.cpp
--- a/Engine/Plugins/Experimental/GeometryProcessing/Source/DynamicMesh/Private/Parameterization/DynamicMeshUVEditor.cpp
+++ b/Engine/Plugins/Experimental/GeometryProcessing/Source/DynamicMesh/Private/Parameterization/DynamicMeshUVEditor.cpp
@@ -63,7 +63,13 @@ void FDynamicMeshUVEditor::ResetUVs()
 	}
 }
 
-
+void FDynamicMeshUVEditor::ResetUVs(const TArray<int32>& Triangles)
+{
+	if (ensure(UVOverlay))
+	{
+		UVOverlay->ClearElements(Triangles);
+	}
+}
 
 void FDynamicMeshUVEditor::TransformUVElements(const TArray<int32>& ElementIDs, TFunctionRef<FVector2f(const FVector2f&)> TransformFunc)
 {
@@ -85,6 +91,8 @@ void InternalSetPerTriangleUVs(EnumerableType TriangleIDs, const FDynamicMesh3*
 	TMap<int32, int32> BaseToOverlayVIDMap;
 	TArray<int32> NewUVIndices;
 
+	UVOverlay->ClearElements(TriangleIDs);
+	
 	for (int32 TriangleID : TriangleIDs)
 	{
 		FIndex3i MeshTri = Mesh->GetTriangle(TriangleID);
@@ -141,6 +149,8 @@ void FDynamicMeshUVEditor::SetTriangleUVsFromPlanarProjection(
 	if (ensure(UVOverlay) == false) return;
 	if (!Triangles.Num()) return;
 
+	ResetUVs(Triangles);
+
 	double ScaleX = (FMathd::Abs(Dimensions.X) > FMathf::ZeroTolerance) ? (1.0 / Dimensions.X) : 1.0;
 	double ScaleY = (FMathd::Abs(Dimensions.Y) > FMathf::ZeroTolerance) ? (1.0 / Dimensions.Y) : 1.0;
 
@@ -241,6 +251,8 @@ bool FDynamicMeshUVEditor::SetTriangleUVsFromExpMap(const TArray<int32>& Triangl
 	if (ensure(UVOverlay) == false) return false;
 	if (Triangles.Num() == 0) return false;
 
+	ResetUVs(Triangles);
+
 	FDynamicSubmesh3 SubmeshCalc(Mesh, Triangles, (int)EMeshComponents::None, false);
 	FDynamicMesh3& Submesh = SubmeshCalc.GetSubmesh();
 	FMeshNormals::QuickComputeVertexNormals(Submesh);
@@ -308,6 +320,8 @@ bool FDynamicMeshUVEditor::SetTriangleUVsFromExpMap(
 	if (ensure(UVOverlay) == false) return false;
 	if (Triangles.Num() == 0) return false;
 
+	ResetUVs(Triangles);
+
 	double ScaleX = (FMathd::Abs(Dimensions.X) > FMathf::ZeroTolerance) ? (1.0 / Dimensions.X) : 1.0;
 	double ScaleY = (FMathd::Abs(Dimensions.Y) > FMathf::ZeroTolerance) ? (1.0 / Dimensions.Y) : 1.0;
 
@@ -424,6 +438,8 @@ bool FDynamicMeshUVEditor::SetTriangleUVsFromFreeBoundaryConformal(const TArray<
 	if (ensure(UVOverlay) == false) return false;
 	if (Triangles.Num() == 0) return false;
 
+	ResetUVs(Triangles);
+
 	FDynamicMesh3 Submesh(EMeshComponents::None);
 	TMap<int32, int32> BaseToSubmeshV;
 	TArray<int32> SubmeshToBaseV;
@@ -770,6 +786,8 @@ void FDynamicMeshUVEditor::SetTriangleUVsFromBoxProjection(
 	int32 NumTriangles = Triangles.Num();
 	if (!NumTriangles) return;
 
+	ResetUVs(Triangles);
+
 	const int Minor1s[3] = { 1, 0, 0 };
 	const int Minor2s[3] = { 2, 2, 1 };
 	const int Minor1Flip[3] = { -1, 1, 1 };
@@ -874,6 +892,8 @@ void FDynamicMeshUVEditor::SetTriangleUVsFromCylinderProjection(
 	int32 NumTriangles = Triangles.Num();
 	if (!NumTriangles) return;
 
+	ResetUVs(Triangles);
+
 	const int Minor1s[3] = { 1, 0, 0 };
 	const int Minor2s[3] = { 2, 2, 1 };
 	const int Minor1Flip[3] = { -1, 1, 1 };
diff --git a/Engine/Plugins/Experimental/GeometryProcessing/Source/DynamicMesh/Public/Parameterization/DynamicMeshUVEditor.h b/Engine/Plugins/Experimental/GeometryProcessing/Source/DynamicMesh/Public/Parameterization/DynamicMeshUVEditor.h
--- a/Engine/Plugins/Experimental/GeometryProcessing/Source/DynamicMesh/Public/Parameterization/DynamicMeshUVEditor.h
+++ b/Engine/Plugins/Experimental/GeometryProcessing/Source/DynamicMesh/Public/Parameterization/DynamicMeshUVEditor.h
@@ -23,12 +23,13 @@ struct DYNAMICMESH_API FUVEditResult
  */
 class DYNAMICMESH_API FDynamicMeshUVEditor
 {
-public:
+private:
 	/** The mesh we will be editing */
 	FDynamicMesh3* Mesh = nullptr;
 	/** The UV Overlay we will be editing */
 	FDynamicMeshUVOverlay* UVOverlay = nullptr;
 
+public:
 	/**
 	 * Construct UV Editor for a UV Overlay of the given Mesh.
 	 * @param UVLayerIndex index of target UV layer
@@ -47,12 +48,16 @@ public:
 	 */
 	void CreateUVLayer(int32 UVLayerIndex);
 
-
 	/**
-	 * Clear UVs on active layer
+	 * Clear UVs for all triangles on active layer
 	 */
 	void ResetUVs();
 
+	/**
+	* Clear UVs for given triangles on active layer
+	*/
+	void ResetUVs(const TArray<int32>& Triangles);
+
 	/**
 	 * Create new UV island for each Triangle, by planar projection onto plane of Triangle. No transforms/etc are applied.
 	 */
@@ -152,8 +157,6 @@ public:
 	bool QuickPack(int32 TargetTextureResolution = 512, float GutterSize = 1.0f);
 
 
-public:
-
 	//
 	// Utility functions that could maybe go elsewhere but are mainly useful in parameterization context
 	//
diff --git a/Engine/Plugins/Experimental/MeshModelingToolset/Source/ModelingOperators/Private/ParameterizationOps/UVProjectionOp.cpp b/Engine/Plugins/Experimental/MeshModelingToolset/Source/ModelingOperators/Private/ParameterizationOps/UVProjectionOp.cpp
--- a/Engine/Plugins/Experimental/MeshModelingToolset/Source/ModelingOperators/Private/ParameterizationOps/UVProjectionOp.cpp
+++ b/Engine/Plugins/Experimental/MeshModelingToolset/Source/ModelingOperators/Private/ParameterizationOps/UVProjectionOp.cpp
@@ -74,7 +74,6 @@ void FUVProjectionOp::CalculateResult_Box(FProgressCancel* Progress)
 		{
 			Triangles.Add(tid);
 		}
-		UVEditor.ResetUVs();
 		UVEditor.SetTriangleUVsFromBoxProjection(Triangles, [&](const FVector3d P) { return MeshToProjectionSpace.TransformPosition(P); },
 			ProjectionBox.Frame, 2*ProjectionBox.Extents, &EditResult);
 	}
@@ -82,7 +81,7 @@ void FUVProjectionOp::CalculateResult_Box(FProgressCancel* Progress)
 	{
 		UVEditor.SetTriangleUVsFromBoxProjection(*TriangleROI, [&](const FVector3d P) { return MeshToProjectionSpace.TransformPosition(P); },
 			ProjectionBox.Frame, 2*ProjectionBox.Extents, &EditResult);
-		FCompactMaps CompactMaps;
+		//FCompactMaps CompactMaps;
 		//UVEditor.GetOverlay()->CompactInPlace(CompactMaps);
 	}
 
@@ -106,7 +105,6 @@ void FUVProjectionOp::CalculateResult_Plane(FProgressCancel* Progress)
 		{
 			Triangles.Add(tid);
 		}
-		UVEditor.ResetUVs();
 		UVEditor.SetTriangleUVsFromPlanarProjection(Triangles, [&](const FVector3d P) { return MeshToProjectionSpace.TransformPosition(P); },
 			ProjectionFrame, Dimensions, &EditResult);
 	}
@@ -114,7 +112,7 @@ void FUVProjectionOp::CalculateResult_Plane(FProgressCancel* Progress)
 	{
 		UVEditor.SetTriangleUVsFromPlanarProjection(*TriangleROI, [&](const FVector3d P) { return MeshToProjectionSpace.TransformPosition(P); },
 			ProjectionFrame, Dimensions, &EditResult);
-		FCompactMaps CompactMaps;
+		//FCompactMaps CompactMaps;
 		//UVEditor.GetOverlay()->CompactInPlace(CompactMaps);
 	}
 
@@ -138,7 +136,6 @@ void FUVProjectionOp::CalculateResult_ExpMap(FProgressCancel* Progress)
 		{
 			Triangles.Add(tid);
 		}
-		UVEditor.ResetUVs();
 		UVEditor.SetTriangleUVsFromExpMap(Triangles, [&](const FVector3d P) { return MeshToProjectionSpace.TransformPosition(P); },
 			ProjectionFrame, Dimensions, SmoothingRounds, SmoothingAlpha, BlendWeight, &EditResult);
 	}
@@ -146,7 +143,7 @@ void FUVProjectionOp::CalculateResult_ExpMap(FProgressCancel* Progress)
 	{
 		UVEditor.SetTriangleUVsFromExpMap(*TriangleROI, [&](const FVector3d P) { return MeshToProjectionSpace.TransformPosition(P); },
 			ProjectionFrame, Dimensions, SmoothingRounds, SmoothingAlpha, BlendWeight, &EditResult);
-		FCompactMaps CompactMaps;
+		//FCompactMaps CompactMaps;
 		//UVEditor.GetOverlay()->CompactInPlace(CompactMaps);
 	}
 
@@ -167,7 +164,6 @@ void FUVProjectionOp::CalculateResult_Cylinder(FProgressCancel* Progress)
 		{
 			Triangles.Add(tid);
 		}
-		UVEditor.ResetUVs();
 		UVEditor.SetTriangleUVsFromCylinderProjection(Triangles, [&](const FVector3d P) { return MeshToProjectionSpace.TransformPosition(P); },
 			ProjectionBox.Frame, 2 * ProjectionBox.Extents, CylinderSplitAngle, &EditResult);
 	}
@@ -175,7 +171,7 @@ void FUVProjectionOp::CalculateResult_Cylinder(FProgressCancel* Progress)
 	{
 		UVEditor.SetTriangleUVsFromCylinderProjection(*TriangleROI, [&](const FVector3d P) { return MeshToProjectionSpace.TransformPosition(P); },
 			ProjectionBox.Frame, 2 * ProjectionBox.Extents, CylinderSplitAngle, &EditResult);
-		FCompactMaps CompactMaps;
+		//FCompactMaps CompactMaps;
 		//UVEditor.GetOverlay()->CompactInPlace(CompactMaps);
 	}
 
diff --git a/Engine/Plugins/Experimental/MeshModelingToolset/Source/ModelingOperators/Public/ParameterizationOps/RecomputeUVsOp.h b/Engine/Plugins/Experimental/MeshModelingToolset/Source/ModelingOperators/Public/ParameterizationOps/RecomputeUVsOp.h
--- a/Engine/Plugins/Experimental/MeshModelingToolset/Source/ModelingOperators/Public/ParameterizationOps/RecomputeUVsOp.h
+++ b/Engine/Plugins/Experimental/MeshModelingToolset/Source/ModelingOperators/Public/ParameterizationOps/RecomputeUVsOp.h
@@ -75,9 +75,6 @@ public:
 
 protected:
 
-	bool ComputeUVs_ExpMap(FDynamicMesh3& InOutMesh, TFunction<bool(float)>& Interrupter, float GlobalScale = 1.0f);
-	bool ComputeUVs_ConformalFreeBoundary(FDynamicMesh3& InOutMesh, TFunction<bool(float)>& Interrupter, float GlobalScale = 1.0f);
-
 	void NormalizeUVAreas(const FDynamicMesh3& Mesh, FDynamicMeshUVOverlay* Overlay, float GlobalScale = 1.0f);
 };
 
diff --git a/Engine/Source/Runtime/GeometryCore/Public/DynamicMesh/DynamicMeshOverlay.h b/Engine/Source/Runtime/GeometryCore/Public/DynamicMesh/DynamicMeshOverlay.h
--- a/Engine/Source/Runtime/GeometryCore/Public/DynamicMesh/DynamicMeshOverlay.h
+++ b/Engine/Source/Runtime/GeometryCore/Public/DynamicMesh/DynamicMeshOverlay.h
@@ -204,9 +204,19 @@ public:
 		checkSlow(IsCompact());
 	}
 
-	/** Discard current set of elements, but keep triangles */
+	/** Discard all elements. */
 	void ClearElements();
 
+	/** Discard elements for given triangles. */
+	template <typename EnumerableIntType>
+	void ClearElements(const EnumerableIntType& Triangles)
+	{
+		for (int32 TriID : Triangles)
+		{
+			UnsetTriangle(TriID);
+		}
+	}
+
 	/** @return the number of in-use Elements in the overlay */
 	int ElementCount() const { return (int)ElementsRefCounts.GetCount(); }
 	/** @return the maximum element index in the overlay. This may be larger than the count if Elements have been deleted. */
