Add properties save/restore to polycut tool

#rnx

diff --git a/Engine/Plugins/Experimental/MeshModelingToolset/Source/MeshModelingToolsEditorOnly/Private/PolygonOnMeshTool.cpp b/Engine/Plugins/Experimental/MeshModelingToolset/Source/MeshModelingToolsEditorOnly/Private/PolygonOnMeshTool.cpp
--- a/Engine/Plugins/Experimental/MeshModelingToolset/Source/MeshModelingToolsEditorOnly/Private/PolygonOnMeshTool.cpp
+++ b/Engine/Plugins/Experimental/MeshModelingToolset/Source/MeshModelingToolsEditorOnly/Private/PolygonOnMeshTool.cpp
@@ -83,6 +83,7 @@ void UPolygonOnMeshTool::Setup()
 	TargetComponent->SetOwnerVisibility(false);
 
 	BasicProperties = NewObject<UPolygonOnMeshToolProperties>(this);
+	BasicProperties->RestoreProperties(this);
 	AddToolPropertySource(BasicProperties);
 
 	ActionProperties = NewObject<UPolygonOnMeshToolActionPropertySet>(this);
@@ -235,6 +236,7 @@ void UPolygonOnMeshTool::Shutdown(EToolShutdownType ShutdownType)
 	{
 		DrawPolygonMechanic->Shutdown();
 	}
+	BasicProperties->SaveProperties(this);
 
 	// Restore (unhide) the source meshes
 	Cast<IPrimitiveComponentBackedTarget>(Target)->SetOwnerVisibility(true);
