MeshModelingTools: Fixed Bake Texture color transformation for Texture 2D image resample.

#rnx

diff --git a/Engine/Plugins/Experimental/MeshLODToolset/Source/MeshLODToolset/Private/Graphs/GenerateStaticMeshLODProcess.cpp b/Engine/Plugins/Experimental/MeshLODToolset/Source/MeshLODToolset/Private/Graphs/GenerateStaticMeshLODProcess.cpp
--- a/Engine/Plugins/Experimental/MeshLODToolset/Source/MeshLODToolset/Private/Graphs/GenerateStaticMeshLODProcess.cpp
+++ b/Engine/Plugins/Experimental/MeshLODToolset/Source/MeshLODToolset/Private/Graphs/GenerateStaticMeshLODProcess.cpp
@@ -374,7 +374,8 @@ bool UGenerateStaticMeshLODProcess::Initialize(UStaticMesh* StaticMeshIn, FProgr
 	//	{
 	//		FReadTextureJob job = ReadJobs[ji];
 	//		FTextureInfo& TexInfo = SourceMaterials[job.MatIndex].SourceTextures[job.TexIndex];
-	//		UE::AssetUtils::ReadTexture(TexInfo.Texture, TexInfo.Dimensions, TexInfo.Image);
+	//		UE::AssetUtils::ReadTexture(TexInfo.Texture, TexInfo.Image);
+	//		TexInfo.Dimensions = TexInfo.Image.GetDimensions();
 	//	});
 	//});
 
@@ -386,7 +387,8 @@ bool UGenerateStaticMeshLODProcess::Initialize(UStaticMesh* StaticMeshIn, FProgr
 		if (SourceMaterial.bIsPreviouslyGeneratedMaterial == false && SourceMaterial.bIsReusable == false)
 		{
 			FTextureInfo& TexInfo = SourceMaterials[job.MatIndex].SourceTextures[job.TexIndex];
-			UE::AssetUtils::ReadTexture(TexInfo.Texture, TexInfo.Dimensions, TexInfo.Image);
+			UE::AssetUtils::ReadTexture(TexInfo.Texture, TexInfo.Image);
+			TexInfo.Dimensions = TexInfo.Image.GetDimensions();
 		}
 	}
 
diff --git a/Engine/Plugins/Experimental/MeshModelingToolset/Source/MeshModelingTools/Private/BakeMeshAttributeMapsTool.cpp b/Engine/Plugins/Experimental/MeshModelingToolset/Source/MeshModelingTools/Private/BakeMeshAttributeMapsTool.cpp
--- a/Engine/Plugins/Experimental/MeshModelingToolset/Source/MeshModelingTools/Private/BakeMeshAttributeMapsTool.cpp
+++ b/Engine/Plugins/Experimental/MeshModelingToolset/Source/MeshModelingTools/Private/BakeMeshAttributeMapsTool.cpp
@@ -22,6 +22,7 @@
 #include "ImageUtils.h"
 
 #include "AssetUtils/Texture2DBuilder.h"
+#include "AssetUtils/Texture2DUtil.h"
 #include "ModelingObjectsCreationAPI.h"
 
 #include "TargetInterfaces/MaterialProvider.h"
@@ -530,7 +531,7 @@ void UBakeMeshAttributeMapsTool::Shutdown(EToolShutdownType ShutdownType)
 			int NumResults = Settings->Result.Num();
 			for (int ResultIdx = 0; ResultIdx < NumResults; ++ResultIdx)
 			{
-				FTexture2DBuilder::ETextureType TexType = FTexture2DBuilder::ETextureType::Color;
+				FTexture2DBuilder::ETextureType TexType = FTexture2DBuilder::ETextureType::ColorLinear;
 				FString TexName;
 				switch (ResultTypes[ResultIdx])
 				{
@@ -567,15 +568,14 @@ void UBakeMeshAttributeMapsTool::Shutdown(EToolShutdownType ShutdownType)
 					break;
 				case EBakeMapType::PositionImage:
 					TexName = FString::Printf(TEXT("%s_PositionImg"), *BaseName);
-					TexType = FTexture2DBuilder::ETextureType::Color;
 					break;
 				case EBakeMapType::Texture2DImage:
 					TexName = FString::Printf(TEXT("%s_TextureImg"), *BaseName);
-					TexType = FTexture2DBuilder::ETextureType::Color;
+					TexType = CachedTexture2DImageSettings.bSRGB ? FTexture2DBuilder::ETextureType::Color : FTexture2DBuilder::ETextureType::ColorLinear;
 					break;
 				case EBakeMapType::MultiTexture:
 					TexName = FString::Printf(TEXT("%s_MultiTextureImg"), *BaseName);
-					TexType = FTexture2DBuilder::ETextureType::Color;
+					TexType = CachedTexture2DImageSettings.bSRGB ? FTexture2DBuilder::ETextureType::Color : FTexture2DBuilder::ETextureType::ColorLinear;
 					break;
 				}
 				FTexture2DBuilder::CopyPlatformDataToSourceData(Settings->Result[ResultIdx], TexType);
@@ -1006,16 +1006,16 @@ EBakeOpState UBakeMeshAttributeMapsTool::UpdateResult_Texture2DImage()
 		return EBakeOpState::Invalid;
 	}
 
-
 	{
-		FTempTextureAccess TextureAccess(Texture2DProps->SourceTexture);
 		CachedTextureImage = MakeShared<UE::Geometry::TImageBuilder<FVector4f>, ESPMode::ThreadSafe>();
-		CachedTextureImage->SetDimensions(TextureAccess.GetDimensions());
-		if (!TextureAccess.CopyTo(*CachedTextureImage))
+		if (!UE::AssetUtils::ReadTexture(Texture2DProps->SourceTexture, *CachedTextureImage, bPreferPlatformData))
 		{
 			GetToolManager()->DisplayMessage(LOCTEXT("CannotReadTextureWarning", "Cannot read from the source texture"), EToolMessageLevel::UserWarning);
 			return EBakeOpState::Invalid;
 		}
+
+		// The read texture data is always in linear space.
+		NewSettings.bSRGB = false; 
 	}
 
 	if (!(CachedTexture2DImageSettings == NewSettings))
@@ -1054,9 +1054,11 @@ EBakeOpState UBakeMeshAttributeMapsTool::UpdateResult_MultiTexture()
 		}
 	}
 
-
 	CachedMultiTextures.Reset();
 
+	// The read texture data is always in linear space.
+	NewSettings.bSRGB = false;
+	
 	for ( auto& InputTexture : MultiTextureProps->MaterialIDSourceTextureMap)
 	{
 		UTexture2D* Texture = InputTexture.Value;
@@ -1067,11 +1069,8 @@ EBakeOpState UBakeMeshAttributeMapsTool::UpdateResult_MultiTexture()
 		}
 
 		int32 MaterialID = InputTexture.Key;
-		FTempTextureAccess TextureAccess(Texture);
 		CachedMultiTextures.Add(MaterialID, MakeShared<UE::Geometry::TImageBuilder<FVector4f>, ESPMode::ThreadSafe>());
-		CachedMultiTextures[MaterialID]->SetDimensions(TextureAccess.GetDimensions());
-
-		if (!TextureAccess.CopyTo(*CachedMultiTextures[MaterialID]))
+		if (!UE::AssetUtils::ReadTexture(Texture, *CachedMultiTextures[MaterialID], bPreferPlatformData))
 		{
 			GetToolManager()->DisplayMessage(LOCTEXT("CannotReadTextureWarning", "Cannot read from the source texture"), EToolMessageLevel::UserWarning);
 			return EBakeOpState::Invalid;
@@ -1240,7 +1239,7 @@ void UBakeMeshAttributeMapsTool::UpdateOnModeChange()
 void UBakeMeshAttributeMapsTool::OnMapsUpdated(const TUniquePtr<FMeshMapBaker>& NewResult)
 {
 	FImageDimensions BakeDimensions = NewResult->GetDimensions();
-	int32 NumBakers = NewResult->NumBakers();
+	const int32 NumBakers = NewResult->NumBakers();
 	for (int32 BakerIdx = 0; BakerIdx < NumBakers; ++BakerIdx)
 	{
 		FMeshMapEvaluator* Baker = NewResult->GetBaker(BakerIdx);
@@ -1284,7 +1283,7 @@ void UBakeMeshAttributeMapsTool::OnMapsUpdated(const TUniquePtr<FMeshMapBaker>&
 		}
 		case EMeshMapEvaluatorType::Curvature:
 		{
-			UpdateCachedMap(EBakeMapType::Curvature, FTexture2DBuilder::ETextureType::Color, 0);
+			UpdateCachedMap(EBakeMapType::Curvature, FTexture2DBuilder::ETextureType::ColorLinear, 0);
 			break;
 		}
 		case EMeshMapEvaluatorType::Property:
@@ -1313,17 +1312,19 @@ void UBakeMeshAttributeMapsTool::OnMapsUpdated(const TUniquePtr<FMeshMapBaker>&
 				break;
 			}
 
-			UpdateCachedMap(MapType, FTexture2DBuilder::ETextureType::Color, 0);
+			UpdateCachedMap(MapType, FTexture2DBuilder::ETextureType::ColorLinear, 0);
 			break;
 		}
 		case EMeshMapEvaluatorType::ResampleImage:
 		{
-			UpdateCachedMap(EBakeMapType::Texture2DImage, FTexture2DBuilder::ETextureType::Color, 0);
+			FTexture2DBuilder::ETextureType TexType = CachedTexture2DImageSettings.bSRGB ? FTexture2DBuilder::ETextureType::Color : FTexture2DBuilder::ETextureType::ColorLinear;
+			UpdateCachedMap(EBakeMapType::Texture2DImage, TexType, 0);
 			break;
 		}
 		case EMeshMapEvaluatorType::MultiResampleImage:
 		{
-			UpdateCachedMap(EBakeMapType::MultiTexture, FTexture2DBuilder::ETextureType::Color, 0);
+			FTexture2DBuilder::ETextureType TexType = CachedTexture2DImageSettings.bSRGB ? FTexture2DBuilder::ETextureType::Color : FTexture2DBuilder::ETextureType::ColorLinear;
+			UpdateCachedMap(EBakeMapType::MultiTexture, TexType, 0);
 			break;
 		}
 		}
diff --git a/Engine/Plugins/Experimental/MeshModelingToolset/Source/MeshModelingTools/Private/BakeMeshAttributeToolCommon.cpp b/Engine/Plugins/Experimental/MeshModelingToolset/Source/MeshModelingTools/Private/BakeMeshAttributeToolCommon.cpp
deleted file mode 100644
diff --git a/Engine/Plugins/Experimental/MeshModelingToolset/Source/MeshModelingTools/Private/BakeMeshAttributeVertexTool.cpp b/Engine/Plugins/Experimental/MeshModelingToolset/Source/MeshModelingTools/Private/BakeMeshAttributeVertexTool.cpp
--- a/Engine/Plugins/Experimental/MeshModelingToolset/Source/MeshModelingTools/Private/BakeMeshAttributeVertexTool.cpp
+++ b/Engine/Plugins/Experimental/MeshModelingToolset/Source/MeshModelingTools/Private/BakeMeshAttributeVertexTool.cpp
@@ -22,6 +22,7 @@
 #include "TargetInterfaces/StaticMeshBackedTarget.h"
 #include "ToolTargetManager.h"
 #include "ModelingToolTargetUtil.h"
+#include "AssetUtils/Texture2DUtil.h"
 
 #include "ExplicitUseGeometryMathTypes.h"		// using UE::Geometry::(math types)
 using namespace UE::Geometry;
@@ -866,11 +867,12 @@ EBakeOpState UBakeMeshAttributeVertexTool::UpdateResult_Texture2DImage()
 		return EBakeOpState::Invalid;
 	}
 
+	// The read texture data is always in linear space.
+	NewSettings.bSRGB = false;
+
 	{
-		FTempTextureAccess TextureAccess(TextureSettings->SourceTexture);
 		CachedTextureImage = MakeShared<UE::Geometry::TImageBuilder<FVector4f>, ESPMode::ThreadSafe>();
-		CachedTextureImage->SetDimensions(TextureAccess.GetDimensions());
-		if (!TextureAccess.CopyTo(*CachedTextureImage))
+		if (!UE::AssetUtils::ReadTexture(TextureSettings->SourceTexture, *CachedTextureImage, /*bPreferPlatformData*/ false))
 		{
 			GetToolManager()->DisplayMessage(LOCTEXT("CannotReadTextureWarning", "Cannot read from the source texture"), EToolMessageLevel::UserWarning);
 			return EBakeOpState::Invalid;
@@ -910,6 +912,9 @@ EBakeOpState UBakeMeshAttributeVertexTool::UpdateResult_MultiTexture()
 
 	CachedMultiTextures.Reset();
 
+	// The read texture data is always in linear space.
+	NewSettings.bSRGB = false;
+	
 	for (auto& InputTexture : MultiTextureSettings->MaterialIDSourceTextureMap)
 	{
 		UTexture2D* Texture = InputTexture.Value;
@@ -920,11 +925,8 @@ EBakeOpState UBakeMeshAttributeVertexTool::UpdateResult_MultiTexture()
 		}
 
 		int32 MaterialID = InputTexture.Key;
-		FTempTextureAccess TextureAccess(Texture);
 		CachedMultiTextures.Add(MaterialID, MakeShared<UE::Geometry::TImageBuilder<FVector4f>, ESPMode::ThreadSafe>());
-		CachedMultiTextures[MaterialID]->SetDimensions(TextureAccess.GetDimensions());
-
-		if (!TextureAccess.CopyTo(*CachedMultiTextures[MaterialID]))
+		if (!UE::AssetUtils::ReadTexture(Texture, *CachedMultiTextures[MaterialID], /*bPreferPlatformData*/ false))
 		{
 			GetToolManager()->DisplayMessage(LOCTEXT("CannotReadTextureWarning", "Cannot read from the source texture"), EToolMessageLevel::UserWarning);
 			return EBakeOpState::Invalid;
diff --git a/Engine/Plugins/Experimental/MeshModelingToolset/Source/MeshModelingTools/Private/DisplaceMeshTool.cpp b/Engine/Plugins/Experimental/MeshModelingToolset/Source/MeshModelingTools/Private/DisplaceMeshTool.cpp
--- a/Engine/Plugins/Experimental/MeshModelingToolset/Source/MeshModelingTools/Private/DisplaceMeshTool.cpp
+++ b/Engine/Plugins/Experimental/MeshModelingToolset/Source/MeshModelingTools/Private/DisplaceMeshTool.cpp
@@ -665,8 +665,7 @@ namespace DisplaceMeshToolLocals{
 		}
 
 		TImageBuilder<FVector4f> DisplacementMapValues;
-		FImageDimensions DisplacementMapDimensions;
-		if (!UE::AssetUtils::ReadTexture(Parameters.DisplacementMap, DisplacementMapDimensions, DisplacementMapValues,
+		if (!UE::AssetUtils::ReadTexture(Parameters.DisplacementMap, DisplacementMapValues,
 			// need bPreferPlatformData to be true to respond to non-destructive changes to the texture in the editor
 			true)) 
 		{
@@ -675,6 +674,7 @@ namespace DisplaceMeshToolLocals{
 		}
 		else
 		{
+			const FImageDimensions DisplacementMapDimensions = DisplacementMapValues.GetDimensions();
 			int64 TextureWidth = DisplacementMapDimensions.GetWidth();
 			int64 TextureHeight = DisplacementMapDimensions.GetHeight();
 			Parameters.DisplaceField.Resize(TextureWidth, TextureHeight, 0.0f);
diff --git a/Engine/Plugins/Experimental/MeshModelingToolset/Source/MeshModelingTools/Private/MeshVertexSculptTool.cpp b/Engine/Plugins/Experimental/MeshModelingToolset/Source/MeshModelingTools/Private/MeshVertexSculptTool.cpp
--- a/Engine/Plugins/Experimental/MeshModelingToolset/Source/MeshModelingTools/Private/MeshVertexSculptTool.cpp
+++ b/Engine/Plugins/Experimental/MeshModelingToolset/Source/MeshModelingTools/Private/MeshVertexSculptTool.cpp
@@ -952,14 +952,13 @@ void UMeshVertexSculptTool::UpdateBrushAlpha(UTexture2D* NewAlpha)
 		if (this->BrushAlpha != nullptr)
 		{
 			TImageBuilder<FVector4f> AlphaValues;
-			FImageDimensions AlphaDimensions;
 
 			constexpr bool bPreferPlatformData = false;
-			bool bReadOK = UE::AssetUtils::ReadTexture(this->BrushAlpha, AlphaDimensions, AlphaValues, bPreferPlatformData);
+			const bool bReadOK = UE::AssetUtils::ReadTexture(this->BrushAlpha, AlphaValues, bPreferPlatformData);
 			if (bReadOK)
 			{
 				BrushAlphaValues = MoveTemp(AlphaValues);
-				BrushAlphaDimensions = AlphaDimensions;
+				BrushAlphaDimensions = AlphaValues.GetDimensions();
 				bHaveBrushAlpha = true;
 				return;
 			}
diff --git a/Engine/Plugins/Experimental/MeshModelingToolset/Source/MeshModelingTools/Public/BakeMeshAttributeMapsTool.h b/Engine/Plugins/Experimental/MeshModelingToolset/Source/MeshModelingTools/Public/BakeMeshAttributeMapsTool.h
--- a/Engine/Plugins/Experimental/MeshModelingToolset/Source/MeshModelingTools/Public/BakeMeshAttributeMapsTool.h
+++ b/Engine/Plugins/Experimental/MeshModelingToolset/Source/MeshModelingTools/Public/BakeMeshAttributeMapsTool.h
@@ -273,6 +273,7 @@ protected:
 
 	TSharedPtr<UE::Geometry::TImageBuilder<FVector4f>, ESPMode::ThreadSafe> CachedTextureImage;
 	FTexture2DImageSettings CachedTexture2DImageSettings;
+	bool bPreferPlatformData = true;
 	EBakeOpState UpdateResult_Texture2DImage();
 
 	TMap<int32, TSharedPtr<UE::Geometry::TImageBuilder<FVector4f>, ESPMode::ThreadSafe>> CachedMultiTextures;
diff --git a/Engine/Plugins/Experimental/MeshModelingToolset/Source/MeshModelingTools/Public/BakeMeshAttributeToolCommon.h b/Engine/Plugins/Experimental/MeshModelingToolset/Source/MeshModelingTools/Public/BakeMeshAttributeToolCommon.h
--- a/Engine/Plugins/Experimental/MeshModelingToolset/Source/MeshModelingTools/Public/BakeMeshAttributeToolCommon.h
+++ b/Engine/Plugins/Experimental/MeshModelingToolset/Source/MeshModelingTools/Public/BakeMeshAttributeToolCommon.h
@@ -279,10 +279,11 @@ struct FTexture2DImageSettings
 {
 	FImageDimensions Dimensions;
 	int32 UVLayer = 0;
+	bool bSRGB = true;
 
 	bool operator==(const FTexture2DImageSettings& Other) const
 	{
-		return Dimensions == Other.Dimensions && UVLayer == Other.UVLayer;
+		return Dimensions == Other.Dimensions && UVLayer == Other.UVLayer && bSRGB == Other.bSRGB;
 	}
 };
 
@@ -299,31 +300,3 @@ enum class EBakeOpState
 ENUM_CLASS_FLAGS(EBakeOpState);
 
 
-/**
- * Utility UTexture2D accessor
- */
-class FTempTextureAccess
-{
-public:
-	FTempTextureAccess(UTexture2D* DisplacementMap);
-	FTempTextureAccess(const FTempTextureAccess&) = delete;
-	FTempTextureAccess(FTempTextureAccess&&) = delete;
-	void operator=(const FTempTextureAccess&) = delete;
-	void operator=(FTempTextureAccess&&) = delete;
-
-	~FTempTextureAccess();
-
-	bool HasData() const;
-	const FColor* GetData() const;	
-	FImageDimensions GetDimensions() const;
-
-	bool CopyTo(TImageBuilder<FVector4f>& DestImage) const;
-
-private:
-	UTexture2D* DisplacementMap{ nullptr };
-	TextureCompressionSettings OldCompressionSettings{};
-	TextureMipGenSettings OldMipGenSettings{};
-	bool bOldSRGB{ false };
-	const FColor* FormattedImageData{ nullptr };
-};
-
diff --git a/Engine/Plugins/Experimental/MeshModelingToolset/Source/ModelingComponents/Private/AssetUtils/Texture2DBuilder.cpp b/Engine/Plugins/Experimental/MeshModelingToolset/Source/ModelingComponents/Private/AssetUtils/Texture2DBuilder.cpp
--- a/Engine/Plugins/Experimental/MeshModelingToolset/Source/ModelingComponents/Private/AssetUtils/Texture2DBuilder.cpp
+++ b/Engine/Plugins/Experimental/MeshModelingToolset/Source/ModelingComponents/Private/AssetUtils/Texture2DBuilder.cpp
@@ -25,7 +25,8 @@ bool FTexture2DBuilder::Initialize(ETextureType BuildTypeIn, FImageDimensions Di
 	if (BuildType == ETextureType::ColorLinear 
 		|| BuildType == ETextureType::Roughness 
 		|| BuildType == ETextureType::Metallic 
-		|| BuildType == ETextureType::Specular)
+		|| BuildType == ETextureType::Specular
+		|| BuildType == ETextureType::AmbientOcclusion)
 	{
 		RawTexture2D->SRGB = false;
 		RawTexture2D->UpdateResource();
diff --git a/Engine/Plugins/Experimental/MeshModelingToolset/Source/ModelingComponents/Private/AssetUtils/Texture2DUtil.cpp b/Engine/Plugins/Experimental/MeshModelingToolset/Source/ModelingComponents/Private/AssetUtils/Texture2DUtil.cpp
--- a/Engine/Plugins/Experimental/MeshModelingToolset/Source/ModelingComponents/Private/AssetUtils/Texture2DUtil.cpp
+++ b/Engine/Plugins/Experimental/MeshModelingToolset/Source/ModelingComponents/Private/AssetUtils/Texture2DUtil.cpp
@@ -13,7 +13,6 @@ using namespace UE::Geometry;
 
 static bool ReadTexture_PlatformData(
 	UTexture2D* TextureMap,
-	FImageDimensions& Dimensions,
 	TImageBuilder<FVector4f>& DestImage)
 {
 	// Read from PlatformData
@@ -25,39 +24,39 @@ static bool ReadTexture_PlatformData(
 	// which will check() if it's on the Game Thread
 
 	check(TextureMap->PlatformData);
-	int32 Width = TextureMap->PlatformData->Mips[0].SizeX;
-	int32 Height = TextureMap->PlatformData->Mips[0].SizeY;
-	Dimensions = FImageDimensions(Width, Height);
+	const int32 Width = TextureMap->PlatformData->Mips[0].SizeX;
+	const int32 Height = TextureMap->PlatformData->Mips[0].SizeY;
+	const FImageDimensions Dimensions = FImageDimensions(Width, Height);
 	DestImage.SetDimensions(Dimensions);
-	int64 Num = Dimensions.Num();
+	const int64 Num = Dimensions.Num();
 
 	// convert built platform texture data to uncompressed RGBA8 format
-	TextureCompressionSettings InitialCompressionSettings = TextureMap->CompressionSettings;
-	bool bWasSRGB = TextureMap->SRGB;
+	const TextureCompressionSettings InitialCompressionSettings = TextureMap->CompressionSettings;
 #if WITH_EDITOR
-	TextureMipGenSettings InitialMipGenSettings = TextureMap->MipGenSettings;
+	const TextureMipGenSettings InitialMipGenSettings = TextureMap->MipGenSettings;
 #endif
 	TextureMap->CompressionSettings = TextureCompressionSettings::TC_VectorDisplacementmap;
-	TextureMap->SRGB = false;
 #if WITH_EDITOR
 	TextureMap->MipGenSettings = TextureMipGenSettings::TMGS_NoMipmaps;
 #endif
 	TextureMap->UpdateResource();
 
-	// lock texture and represet as FColor
+	// lock texture and read as FColor
 	const FColor* FormattedImageData = reinterpret_cast<const FColor*>(TextureMap->PlatformData->Mips[0].BulkData.LockReadOnly());
 
 	// maybe could be done more quickly by row?
-	for (int32 i = 0; i < Num; ++i)
+	for (int64 i = 0; i < Num; ++i)
 	{
 		FColor ByteColor = FormattedImageData[i];
-		DestImage.SetPixel(i, FVector4f(ByteColor.ReinterpretAsLinear()));
+		FLinearColor FloatColor = (TextureMap->SRGB) ?
+				FLinearColor::FromSRGBColor(ByteColor) :
+				ByteColor.ReinterpretAsLinear();
+		DestImage.SetPixel(i, FVector4f(FloatColor));
 	}
 
 	// restore built platform texture data to initial state
 	TextureMap->PlatformData->Mips[0].BulkData.Unlock();
 	TextureMap->CompressionSettings = InitialCompressionSettings;
-	TextureMap->SRGB = bWasSRGB;
 #if WITH_EDITOR
 	TextureMap->MipGenSettings = InitialMipGenSettings;
 #endif
@@ -70,28 +69,27 @@ static bool ReadTexture_PlatformData(
 #if WITH_EDITOR
 static bool ReadTexture_SourceData(
 	UTexture2D* TextureMap,
-	FImageDimensions& Dimensions,
 	TImageBuilder<FVector4f>& DestImage)
 {
 	FTextureSource& TextureSource = TextureMap->Source;
 
-	int32 Width = TextureSource.GetSizeX();
-	int32 Height = TextureSource.GetSizeY();
-	Dimensions = FImageDimensions(Width, Height);
+	const int32 Width = TextureSource.GetSizeX();
+	const int32 Height = TextureSource.GetSizeY();
+	const FImageDimensions Dimensions = FImageDimensions(Width, Height);
 	DestImage.SetDimensions(Dimensions);
-	int64 Num = Dimensions.Num();
+	const int64 Num = Dimensions.Num();
 
 	TArray64<uint8> SourceData;
 	TextureMap->Source.GetMipData(SourceData, 0, 0, 0);
-	ETextureSourceFormat SourceFormat = TextureSource.GetFormat();
-	int32 BytesPerPixel = TextureSource.GetBytesPerPixel();
+	const ETextureSourceFormat SourceFormat = TextureSource.GetFormat();
+	const int32 BytesPerPixel = TextureSource.GetBytesPerPixel();
 	const uint8* SourceDataPtr = SourceData.GetData();
 
 	// code below is derived from UBlueprintMaterialTextureNodesBPLibrary::Texture2D_SampleUV_EditorOnly()
 	if ((SourceFormat == TSF_BGRA8 || SourceFormat == TSF_BGRE8))
 	{
 		check(BytesPerPixel == sizeof(FColor));
-		for (int32 i = 0; i < Num; ++i)
+		for (int64 i = 0; i < Num; ++i)
 		{
 			const uint8* PixelPtr = SourceDataPtr + (i * BytesPerPixel);
 			FColor PixelColor = *((FColor*)PixelPtr);
@@ -103,23 +101,50 @@ static bool ReadTexture_SourceData(
 			DestImage.SetPixel(i, FVector4f(FloatColor));
 		}
 	}
-	else if ((SourceFormat == TSF_RGBA16 || SourceFormat == TSF_RGBA16F))
+	// code below is also derived from FImage::CopyTo (CopyImage) - invoked during BuildTexture on import.
+	else if (SourceFormat == TSF_RGBA16)
+	{
+		check(BytesPerPixel == sizeof(uint16) * 4);
+		for (int64 i = 0; i < Num; ++i)
+		{
+			const uint8* SourcePtr = SourceDataPtr + (i * BytesPerPixel);
+			const uint16* PixelPtr = (const uint16*)SourcePtr;
+			DestImage.SetPixel(i, FVector4f(
+				PixelPtr[0] / 65535.0f,
+				PixelPtr[1] / 65535.0f,
+				PixelPtr[2] / 65535.0f,
+				PixelPtr[3] / 65535.0f
+				));
+		}
+	}
+	else if (SourceFormat == TSF_RGBA16F)
 	{
 		check(BytesPerPixel == sizeof(FFloat16Color));
-		for (int32 i = 0; i < Num; ++i)
+		for (int64 i = 0; i < Num; ++i)
 		{
 			const uint8* PixelPtr = SourceDataPtr + (i * BytesPerPixel);
 			DestImage.SetPixel(i, FVector4f( ((const FFloat16Color*)PixelPtr)->GetFloats() ));
 		}
 	}
+	else if ((SourceFormat == TSF_G16))
+	{
+		check(BytesPerPixel == 2);
+		for (int64 i = 0; i < Num; ++i)
+		{
+			const uint8* PixelPtr = SourceDataPtr + (i * BytesPerPixel);
+			const uint16 ByteColor = *((const uint16*)PixelPtr);
+			const float FloatColor = float(ByteColor) / 65535.0f;
+			DestImage.SetPixel(i, FVector4f(FloatColor, FloatColor, FloatColor, 1.0));
+		}
+	}
 	else if (SourceFormat == TSF_G8)
 	{
 		check(BytesPerPixel == 1);
-		for (int32 i = 0; i < Num; ++i)
+		for (int64 i = 0; i < Num; ++i)
 		{
 			const uint8* PixelPtr = SourceDataPtr + (i * BytesPerPixel);
-			uint8 PixelColor = *PixelPtr;
-			float PixelColorf = float(PixelColor) / 255.0f;
+			const uint8 PixelColor = *PixelPtr;
+			const float PixelColorf = float(PixelColor) / 255.0f;
 			FLinearColor FloatColor = (TextureMap->SRGB) ?
 				FLinearColor::FromSRGBColor(FColor(PixelColor, PixelColor, PixelColor, 255)) :
 				FLinearColor(PixelColorf, PixelColorf, PixelColorf, 1.0);
@@ -134,24 +159,23 @@ static bool ReadTexture_SourceData(
 
 bool UE::AssetUtils::ReadTexture(
 	UTexture2D* TextureMap,
-	FImageDimensions& Dimensions,
-	TImageBuilder<FVector4f>& DestImage,
-	bool bPreferPlatformData)
+	TImageBuilder<FVector4f>& DestImageOut,
+	const bool bPreferPlatformData)
 {
 	if (ensure(TextureMap) == false) return false;
 
 #if WITH_EDITOR
-	bool bHasVTData = TextureMap->GetPlatformData()->VTData != nullptr;
-	bool bHasMips = TextureMap->GetPlatformData()->Mips.Num() != 0;
+	const bool bHasVTData = TextureMap->GetPlatformData()->VTData != nullptr;
+	const bool bHasMips = TextureMap->GetPlatformData()->Mips.Num() != 0;
 	ensure(bHasVTData != bHasMips);		// should be one or the other
 
 	if (TextureMap->Source.IsValid() && (bPreferPlatformData == false || bHasMips == false))
 	{
-		return ReadTexture_SourceData(TextureMap, Dimensions, DestImage);
+		return ReadTexture_SourceData(TextureMap, DestImageOut);
 	}
 #endif
 
-	return ReadTexture_PlatformData(TextureMap, Dimensions, DestImage);
+	return ReadTexture_PlatformData(TextureMap, DestImageOut);
 }
 
 
diff --git a/Engine/Plugins/Experimental/MeshModelingToolset/Source/ModelingComponents/Public/AssetUtils/Texture2DUtil.h b/Engine/Plugins/Experimental/MeshModelingToolset/Source/ModelingComponents/Public/AssetUtils/Texture2DUtil.h
--- a/Engine/Plugins/Experimental/MeshModelingToolset/Source/ModelingComponents/Public/AssetUtils/Texture2DUtil.h
+++ b/Engine/Plugins/Experimental/MeshModelingToolset/Source/ModelingComponents/Public/AssetUtils/Texture2DUtil.h
@@ -7,6 +7,7 @@
 #include "Image/ImageDimensions.h"
 #include "Image/ImageBuilder.h"
 #include "Engine/Texture2D.h"
+#include "Math/Color.h"
 
 
 
@@ -19,14 +20,15 @@ namespace AssetUtils
 	/**
 	 * Extract the Dimensions and image pixels from an input TextureMap
 	 * By default, the "Source" texture data is read. This is only available in-Editor. At runtime, only "Platform" data is available.
+	 * @param TextureMap the texture map to read
+	 * @param DestImageOut the result of the texture read
 	 * @param bPreferPlatformData if true, platform data will be returned even in-Editor.
 	 * @return true on success
 	 */
 	MODELINGCOMPONENTS_API bool ReadTexture(
 		UTexture2D* TextureMap,
-		FImageDimensions& DimensionsOut,
 		TImageBuilder<FVector4f>& DestImageOut,
-		bool bPreferPlatformData = false);
+		const bool bPreferPlatformData = false);
 
 	/**
 	 * Convert input UTexture2D to single-channel. Assumption is it has more than one channel. Red channel is used.
