Fixing map iteration issues identified by MSVC static analysis following the refresh of wrapped object pointers.

diff --git a/Engine/Plugins/Animation/ControlRig/Source/ControlRigEditor/Private/EditMode/ControlRigControlsProxy.cpp b/Engine/Plugins/Animation/ControlRig/Source/ControlRigEditor/Private/EditMode/ControlRigControlsProxy.cpp
--- a/Engine/Plugins/Animation/ControlRig/Source/ControlRigEditor/Private/EditMode/ControlRigControlsProxy.cpp
+++ b/Engine/Plugins/Animation/ControlRig/Source/ControlRigEditor/Private/EditMode/ControlRigControlsProxy.cpp
@@ -756,7 +756,7 @@ void UControlRigDetailPanelControlProxies::RemoveProxy(const FName& Name)
 
 void UControlRigDetailPanelControlProxies::RemoveAllProxies()
 {
-	for (TPair<FName, UControlRigControlsProxy * >Pair : AllProxies)
+	for (TPair<FName, TObjectPtr<UControlRigControlsProxy> >Pair : AllProxies)
 	{
 		UControlRigControlsProxy* ExistingProxy = Pair.Value;
 		ExistingProxy->Rename(nullptr, GetTransientPackage(), REN_ForceNoResetLoaders);
diff --git a/Engine/Plugins/Experimental/MeshModelingToolset/Source/ModelingComponents/Private/Drawing/PreviewGeometryActor.cpp b/Engine/Plugins/Experimental/MeshModelingToolset/Source/ModelingComponents/Private/Drawing/PreviewGeometryActor.cpp
--- a/Engine/Plugins/Experimental/MeshModelingToolset/Source/ModelingComponents/Private/Drawing/PreviewGeometryActor.cpp
+++ b/Engine/Plugins/Experimental/MeshModelingToolset/Source/ModelingComponents/Private/Drawing/PreviewGeometryActor.cpp
@@ -67,7 +67,7 @@ void UPreviewGeometry::SetTransform(const FTransform& UseTransform)
 
 void UPreviewGeometry::SetAllVisible(bool bVisible)
 {
-	for (TPair<FString, ULineSetComponent*> Entry : LineSets)
+	for (TPair<FString, TObjectPtr<ULineSetComponent>> Entry : LineSets)
 	{
 		Entry.Value->SetVisibility(bVisible);
 	}
@@ -131,7 +131,7 @@ bool UPreviewGeometry::RemoveLineSet(const FString& LineSetIdentifier, bool bDes
 
 void UPreviewGeometry::RemoveAllLineSets(bool bDestroy)
 {
-	for (TPair<FString, ULineSetComponent*> Entry : LineSets)
+	for (TPair<FString, TObjectPtr<ULineSetComponent>> Entry : LineSets)
 	{
 		if (bDestroy)
 		{
@@ -170,7 +170,7 @@ bool UPreviewGeometry::SetLineSetMaterial(const FString& LineSetIdentifier, UMat
 
 void UPreviewGeometry::SetAllLineSetsMaterial(UMaterialInterface* Material)
 {
-	for (TPair<FString, ULineSetComponent*> Entry : LineSets)
+	for (TPair<FString, TObjectPtr<ULineSetComponent>> Entry : LineSets)
 	{
 		Entry.Value->SetLineMaterial(Material);
 	}
diff --git a/Engine/Plugins/VirtualProduction/CameraCalibration/Source/CameraCalibrationEditor/Private/AssetEditor/CameraCalibrationStepsController.cpp b/Engine/Plugins/VirtualProduction/CameraCalibration/Source/CameraCalibrationEditor/Private/AssetEditor/CameraCalibrationStepsController.cpp
--- a/Engine/Plugins/VirtualProduction/CameraCalibration/Source/CameraCalibrationEditor/Private/AssetEditor/CameraCalibrationStepsController.cpp
+++ b/Engine/Plugins/VirtualProduction/CameraCalibration/Source/CameraCalibrationEditor/Private/AssetEditor/CameraCalibrationStepsController.cpp
@@ -86,7 +86,7 @@ namespace CameraCalibrationStepsController
 
 		for (const ULiveLinkComponentController* LLComponentController : LLComponentControllers)
 		{
-			for (const TPair<TSubclassOf<ULiveLinkRole>, ULiveLinkControllerBase*> Pair : LLComponentController->ControllerMap)
+			for (const TPair<TSubclassOf<ULiveLinkRole>, TObjectPtr<ULiveLinkControllerBase>> Pair : LLComponentController->ControllerMap)
 			{
 				if (ULiveLinkCameraController* CameraController = Cast<ULiveLinkCameraController>(Pair.Value))
 				{
