Followup for missed item in CL16902764: Fixing map iteration issues identified by MSVC static analysis following the refresh of wrapped object pointers.

diff --git a/Engine/Plugins/Experimental/MeshModelingToolset/Source/ModelingComponents/Public/Drawing/PreviewGeometryActor.h b/Engine/Plugins/Experimental/MeshModelingToolset/Source/ModelingComponents/Public/Drawing/PreviewGeometryActor.h
--- a/Engine/Plugins/Experimental/MeshModelingToolset/Source/ModelingComponents/Public/Drawing/PreviewGeometryActor.h
+++ b/Engine/Plugins/Experimental/MeshModelingToolset/Source/ModelingComponents/Public/Drawing/PreviewGeometryActor.h
@@ -148,7 +148,7 @@ public:
 	template<typename UpdateFuncType>
 	void UpdateAllLineSets(UpdateFuncType UpdateFunc)
 	{
-		for (TPair<FString, ULineSetComponent*> Entry : LineSets)
+		for (TPair<FString, TObjectPtr<ULineSetComponent>> Entry : LineSets)
 		{
 			UpdateFunc(Entry.Value);
 		}
