ModelingTools: Update BakeVertex tool to use a more generic target interface requirements list.

#rnx

diff --git a/Engine/Plugins/Experimental/MeshModelingToolset/Source/MeshModelingTools/Private/BakeMeshAttributeVertexTool.cpp b/Engine/Plugins/Experimental/MeshModelingToolset/Source/MeshModelingTools/Private/BakeMeshAttributeVertexTool.cpp
--- a/Engine/Plugins/Experimental/MeshModelingToolset/Source/MeshModelingTools/Private/BakeMeshAttributeVertexTool.cpp
+++ b/Engine/Plugins/Experimental/MeshModelingToolset/Source/MeshModelingTools/Private/BakeMeshAttributeVertexTool.cpp
@@ -18,8 +18,8 @@
 
 #include "TargetInterfaces/MaterialProvider.h"
 #include "TargetInterfaces/MeshDescriptionProvider.h"
+#include "TargetInterfaces/MeshDescriptionCommitter.h"
 #include "TargetInterfaces/PrimitiveComponentBackedTarget.h"
-#include "TargetInterfaces/StaticMeshBackedTarget.h"
 #include "ToolTargetManager.h"
 #include "ModelingToolTargetUtil.h"
 #include "AssetUtils/Texture2DUtil.h"
@@ -37,8 +37,8 @@ const FToolTargetTypeRequirements& UBakeMeshAttributeVertexToolBuilder::GetTarge
 {
 	static FToolTargetTypeRequirements TypeRequirements({
 		UMeshDescriptionProvider::StaticClass(),
+		UMeshDescriptionCommitter::StaticClass(),
 		UPrimitiveComponentBackedTarget::StaticClass(),
-		UStaticMeshBackedTarget::StaticClass(), // currently only supports StaticMesh targets
 		UMaterialProvider::StaticClass()
 	});
 	return TypeRequirements;
@@ -399,7 +399,13 @@ void UBakeMeshAttributeVertexTool::Shutdown(EToolShutdownType ShutdownType)
 		{
 			GetToolManager()->BeginUndoTransaction(LOCTEXT("BakeMeshAttributeVertexToolTransactionName",
 			                                               "Bake Mesh Attribute Vertex"));
-			UE::ToolTarget::CommitDynamicMeshColorUpdate(Targets[0], PreviewMesh->GetMesh());
+			FConversionToMeshDescriptionOptions ConvertOptions;
+			ConvertOptions.SetToVertexColorsOnly();
+			UE::ToolTarget::CommitDynamicMeshUpdate(
+				Targets[0],
+				*PreviewMesh->GetMesh(),
+				false, // bHaveModifiedTopology
+				ConvertOptions);
 			GetToolManager()->EndUndoTransaction();
 		}
 
diff --git a/Engine/Plugins/Experimental/MeshModelingToolset/Source/ModelingComponents/Private/ModelingToolTargetUtil.cpp b/Engine/Plugins/Experimental/MeshModelingToolset/Source/ModelingComponents/Private/ModelingToolTargetUtil.cpp
--- a/Engine/Plugins/Experimental/MeshModelingToolset/Source/ModelingComponents/Private/ModelingToolTargetUtil.cpp
+++ b/Engine/Plugins/Experimental/MeshModelingToolset/Source/ModelingComponents/Private/ModelingToolTargetUtil.cpp
@@ -322,37 +322,6 @@ UE::ToolTarget::EDynamicMeshUpdateResult UE::ToolTarget::CommitDynamicMeshUVUpda
 }
 
 
-UE::ToolTarget::EDynamicMeshUpdateResult UE::ToolTarget::CommitDynamicMeshColorUpdate(UToolTarget* Target, const UE::Geometry::FDynamicMesh3* UpdatedMesh)
-{
-	IMeshDescriptionCommitter* MeshDescriptionCommitter = Cast<IMeshDescriptionCommitter>(Target);
-	if (!ensure(MeshDescriptionCommitter))
-	{
-		return EDynamicMeshUpdateResult::Failed;
-	}
-
-	EDynamicMeshUpdateResult Result = EDynamicMeshUpdateResult::Failed;
-	MeshDescriptionCommitter->CommitMeshDescription([UpdatedMesh, &Result](const IMeshDescriptionCommitter::FCommitterParams& CommitParams)
-	{
-		FMeshDescription* MeshDescription = CommitParams.MeshDescriptionOut;
-
-		const bool bVerticesOnly = false;
-		const bool bAttributesOnly = true;
-		if (FDynamicMeshToMeshDescription::HaveMatchingElementCounts(UpdatedMesh, MeshDescription, bVerticesOnly, bAttributesOnly))
-		{
-			FDynamicMeshToMeshDescription Converter;
-			Converter.UpdateVertexColors(UpdatedMesh, *MeshDescription);
-			Result = EDynamicMeshUpdateResult::Ok;
-		}
-		else
-		{
-			// must have been duplicate tris in the mesh description; we can't count on 1-to-1 mapping of TriangleIDs.  Just convert 
-			FDynamicMeshToMeshDescription Converter;
-			Converter.Convert(UpdatedMesh, *MeshDescription);
-			Result = EDynamicMeshUpdateResult::Ok_ForcedFullUpdate;
-		}
-	});
-	return Result;
-}
 
 
 bool UE::ToolTarget::ConfigureCreateMeshObjectParams(UToolTarget* SourceTarget, FCreateMeshObjectParams& DerivedParamsOut)
diff --git a/Engine/Plugins/Experimental/MeshModelingToolset/Source/ModelingComponents/Public/ModelingToolTargetUtil.h b/Engine/Plugins/Experimental/MeshModelingToolset/Source/ModelingComponents/Public/ModelingToolTargetUtil.h
--- a/Engine/Plugins/Experimental/MeshModelingToolset/Source/ModelingComponents/Public/ModelingToolTargetUtil.h
+++ b/Engine/Plugins/Experimental/MeshModelingToolset/Source/ModelingComponents/Public/ModelingToolTargetUtil.h
@@ -144,14 +144,6 @@ MODELINGCOMPONENTS_API EDynamicMeshUpdateResult CommitDynamicMeshUpdate(
  */
 MODELINGCOMPONENTS_API EDynamicMeshUpdateResult CommitDynamicMeshUVUpdate(UToolTarget* Target, const UE::Geometry::FDynamicMesh3* UpdatedMesh);
 
-
-/**
- * Update the Vertex Colors of the ToolTarget's mesh based on the provided UpdatedMesh.
- * @return EDynamicMeshUpdateResult::Ok on success, or Ok_ForcedFullUpdate if any dependent mesh topology was modified
- */
-MODELINGCOMPONENTS_API EDynamicMeshUpdateResult CommitDynamicMeshColorUpdate(UToolTarget* Target, const UE::Geometry::FDynamicMesh3* UpdatedMesh);
-
-
 /**
  * FCreateMeshObjectParams::TypeHint is used by the ModelingObjectsCreationAPI to suggest what type of mesh object to create
  * inside various Tools. This should often be derived from the input mesh object type (eg if you plane-cut a Volume, the output
