Modeling Tools:  Jacketing.  Do not use volume threshold when discarding islands, instead only the user specified area and triangle count.

diff --git a/Engine/Plugins/Experimental/MeshModelingToolset/Source/ModelingOperators/Private/CleaningOps/RemoveOccludedTrianglesOp.cpp b/Engine/Plugins/Experimental/MeshModelingToolset/Source/ModelingOperators/Private/CleaningOps/RemoveOccludedTrianglesOp.cpp
--- a/Engine/Plugins/Experimental/MeshModelingToolset/Source/ModelingOperators/Private/CleaningOps/RemoveOccludedTrianglesOp.cpp
+++ b/Engine/Plugins/Experimental/MeshModelingToolset/Source/ModelingOperators/Private/CleaningOps/RemoveOccludedTrianglesOp.cpp
@@ -9,7 +9,7 @@
 #include "Selections/MeshConnectedComponents.h"
 #include "Selections/MeshFaceSelection.h"
 #include "Polygroups/PolygroupSet.h"
-
+#include "MathUtil.h"
 #include "Async/ParallelFor.h"
 
 #include "ExplicitUseGeometryMathTypes.h"		// using UE::Geometry::(math types)
@@ -112,6 +112,7 @@ void FRemoveOccludedTrianglesOp::CalculateResult(FProgressCancel* Progress)
 	if (MinTriCountConnectedComponent > 0 || MinAreaConnectedComponent > 0)
 	{
 		FDynamicMeshEditor Editor(ResultMesh.Get());
-		Editor.RemoveSmallComponents(0, MinAreaConnectedComponent, MinTriCountConnectedComponent);
+		const double MinVolume = -TMathUtilConstants<double>::MaxReal;
+		Editor.RemoveSmallComponents(MinVolume, MinAreaConnectedComponent, MinTriCountConnectedComponent);
 	}
 }
\ No newline at end of file
