ModelingTools: Report UV generation failures in ParameterizeMeshOp. Handle invalid triangles in UVPacking.
#rnx

diff --git a/Engine/Plugins/Experimental/MeshModelingToolset/Source/ModelingOperatorsEditorOnly/Private/ParameterizationOps/ParameterizeMeshOp.cpp b/Engine/Plugins/Experimental/MeshModelingToolset/Source/ModelingOperatorsEditorOnly/Private/ParameterizationOps/ParameterizeMeshOp.cpp
--- a/Engine/Plugins/Experimental/MeshModelingToolset/Source/ModelingOperatorsEditorOnly/Private/ParameterizationOps/ParameterizeMeshOp.cpp
+++ b/Engine/Plugins/Experimental/MeshModelingToolset/Source/ModelingOperatorsEditorOnly/Private/ParameterizationOps/ParameterizeMeshOp.cpp
@@ -26,9 +26,10 @@
 #include "ExplicitUseGeometryMathTypes.h"		// using UE::Geometry::(math types)
 using namespace UE::Geometry;
 
+#define LOCTEXT_NAMESPACE "ParameterizeMeshOp"
+
 FParameterizeMeshOp::FLinearMesh::FLinearMesh(const FDynamicMesh3& Mesh, const bool bRespectPolygroups)
 {
-
 	TArray<FVector3f>& Positions = this->VertexBuffer;
 
 	// Temporary maps used during construction.
@@ -152,7 +153,95 @@ FParameterizeMeshOp::FLinearMesh::FLinearMesh(const FDynamicMesh3& Mesh, const b
 	}
 }
 
-bool FParameterizeMeshOp::ComputeUVs(FDynamicMesh3& Mesh,  TFunction<bool(float)>& Interrupter)
+
+
+void FParameterizeMeshOp::CopyNewUVsToMesh(
+	FDynamicMesh3& Mesh, 
+	const FLinearMesh& LinearMesh, 
+	const FDynamicMesh3& FlippedMesh,
+	const TArray<FVector2D>& UVVertexBuffer,
+	const TArray<int32>& UVIndexBuffer,
+	const TArray<int32>& VertexRemapArray,
+	bool bReverseOrientation)
+{
+	// Add the UVs to the FDynamicMesh
+	const bool bHasAttributes = Mesh.HasAttributes();
+
+	if (bHasAttributes)
+	{
+		FDynamicMeshAttributeSet* Attributes = Mesh.Attributes();
+		Attributes->GetUVLayer(UVLayer)->ClearElements(); // delete existing UVs
+	}
+	else
+	{
+		// Add attrs for UVS
+		Mesh.EnableAttributes();
+	}
+
+	FDynamicMeshUVOverlay* UVOverlay = Mesh.Attributes()->GetUVLayer(UVLayer);
+
+	// This mesh shouldn't already have UVs.
+
+	checkSlow(UVOverlay->ElementCount() == 0);
+
+	// Add the UVs to the overlay
+	int32 NumUVs = UVVertexBuffer.Num();
+	TArray<int32> UVOffsetToElID;  UVOffsetToElID.Reserve(NumUVs);
+
+	for (int32 i = 0; i < NumUVs; ++i)
+	{
+		FVector2D UV = UVVertexBuffer[i];
+
+		// The associated VertID in the dynamic mesh
+		const int32 VertOffset = VertexRemapArray[i];
+
+		// add the UV to the mesh overlay
+		const int32 NewID = UVOverlay->AppendElement(FVector2f(UV));
+		UVOffsetToElID.Add(NewID);
+	}
+
+	int32 NumUVTris = UVIndexBuffer.Num() / 3;
+	for (int32 i = 0; i < NumUVTris; ++i)
+	{
+		int32 t = i * 3;
+		// The triangle in UV space
+		FIndex3i UVTri(UVIndexBuffer[t], UVIndexBuffer[t + 1], UVIndexBuffer[t + 2]);
+
+		// the triangle in terms of the VertIDs in the DynamicMesh
+		FIndex3i TriVertIDs;
+		for (int c = 0; c < 3; ++c)
+		{
+			// the offset for this vertex in the LinearMesh
+			int32 Offset = VertexRemapArray[UVTri[c]];
+
+			int32 VertID = LinearMesh.VertToID[Offset];
+
+			TriVertIDs[c] = VertID;
+		}
+
+		// NB: this could be slow.. 
+		int32 TriID = FlippedMesh.FindTriangle(TriVertIDs[0], TriVertIDs[1], TriVertIDs[2]);
+
+		checkSlow(TriID != FDynamicMesh3::InvalidID);
+
+		FIndex3i ElTri;
+		if (bReverseOrientation)
+		{
+			ElTri = FIndex3i(UVOffsetToElID[UVTri[1]], UVOffsetToElID[UVTri[0]], UVOffsetToElID[UVTri[2]]);
+		}
+		else
+		{
+			ElTri = FIndex3i(UVOffsetToElID[UVTri[0]], UVOffsetToElID[UVTri[1]], UVOffsetToElID[UVTri[2]]);
+		}
+
+		// add the triangle to the overlay
+		UVOverlay->SetTriangle(TriID, ElTri);
+	}
+}
+
+
+
+bool FParameterizeMeshOp::ComputeUVs_UVAtlas(FDynamicMesh3& Mesh,  TFunction<bool(float)>& Interrupter)
 {
 	// the UVAtlas code is unhappy if you feed it a single degenerate triangle
 	bool bNonDegenerate = true;
@@ -167,6 +256,7 @@ bool FParameterizeMeshOp::ComputeUVs(FDynamicMesh3& Mesh,  TFunction<bool(float)
 
 	if (!bNonDegenerate)
 	{
+		NewResultInfo.AddError(FGeometryError(0, LOCTEXT("DegenerateMeshError", "Mesh Contains Degenerate Triangles, Cannot be processed by UVAtlas")));
 		return false;
 	}
 
@@ -193,119 +283,78 @@ bool FParameterizeMeshOp::ComputeUVs(FDynamicMesh3& Mesh,  TFunction<bool(float)
 	float MaxStretch     = Stretch;
 	int32 MaxChartNumber = NumCharts;
 
-	bool bSuccess = true;
+	bool bSuccess = false;
 #if HAVE_PROXYLOD
-	bool bUseXAtlas = (Method == EParamOpBackend::XAtlas);
+	TUniquePtr<IProxyLODParameterization> ParameterizationTool = IProxyLODParameterization::CreateTool();
+	bSuccess = ParameterizationTool->GenerateUVs(Width, Height, Gutter, LinearMesh.VertexBuffer, 
+													LinearMesh.IndexBuffer, LinearMesh.AdjacencyBuffer, Interrupter, 
+													UVVertexBuffer, UVIndexBuffer, VertexRemapArray, MaxStretch, 
+													MaxChartNumber);
 #else
-	bool bUseXAtlas = true;
+	ensureMsgf(false, TEXT("ProxyLOD not available, this should not be called"));
 #endif
 
-	if (bUseXAtlas)
+
+	// Add the UVs to the FDynamicMesh
+	if (bSuccess)
 	{
-		XAtlasWrapper::XAtlasChartOptions ChartOptions;
-		ChartOptions.MaxIterations = XAtlasMaxIterations;
-		XAtlasWrapper::XAtlasPackOptions PackOptions;
-		bSuccess = XAtlasWrapper::ComputeUVs(LinearMesh.IndexBuffer, 
-											 LinearMesh.VertexBuffer, 
-											 ChartOptions,
-											 PackOptions,
-											 UVVertexBuffer, 
-											 UVIndexBuffer, 
-											 VertexRemapArray);
+		CopyNewUVsToMesh(Mesh, LinearMesh, FlippedMesh, UVVertexBuffer, UVIndexBuffer, VertexRemapArray, bFixOrientation);
 	}
 	else
 	{
-#if HAVE_PROXYLOD
-		TUniquePtr<IProxyLODParameterization> ParameterizationTool = IProxyLODParameterization::CreateTool();
-		bSuccess = ParameterizationTool->GenerateUVs(Width, Height, Gutter, LinearMesh.VertexBuffer, 
-													 LinearMesh.IndexBuffer, LinearMesh.AdjacencyBuffer, Interrupter, 
-													 UVVertexBuffer, UVIndexBuffer, VertexRemapArray, MaxStretch, 
-													 MaxChartNumber);
-#else
-		ensureMsgf(false, TEXT("ProxyLOD not available, this should not be called"));
-#endif
+		NewResultInfo.AddError(FGeometryError(0, LOCTEXT("UVAtlasFailed", "UVAtlas failed")));
 	}
-	
-
-	// Add the UVs to the FDynamicMesh
-	if (bSuccess)
-	{
-
-		const bool bHasAttributes = Mesh.HasAttributes();
-
-		if (bHasAttributes)
-		{
-			FDynamicMeshAttributeSet* Attributes = Mesh.Attributes();
-			Attributes->GetUVLayer(UVLayer)->ClearElements(); // delete existing UVs
-		}
-		else
-		{
-			// Add attrs for UVS
-			Mesh.EnableAttributes();
-		}
 
-		FDynamicMeshUVOverlay* UVOverlay = Mesh.Attributes()->GetUVLayer(UVLayer);
+	return bSuccess;
+}
 
-		// This mesh shouldn't already have UVs.
 
-		checkSlow(UVOverlay->ElementCount() == 0);
 
-		// Add the UVs to the overlay
-		int32 NumUVs = UVVertexBuffer.Num();
-		TArray<int32> UVOffsetToElID;  UVOffsetToElID.Reserve(NumUVs);
 
-		for (int32 i = 0; i < NumUVs; ++i)
-		{
-			FVector2D UV = UVVertexBuffer[i];
 
-			// The associated VertID in the dynamic mesh
-			const int32 VertOffset = VertexRemapArray[i];
 
-			// add the UV to the mesh overlay
-			const int32 NewID = UVOverlay->AppendElement(FVector2f(UV));
-			UVOffsetToElID.Add(NewID);
-		}
 
-		int32 NumUVTris = UVIndexBuffer.Num() / 3;
-		for (int32 i = 0; i < NumUVTris; ++i)
-		{
-			int32 t = i * 3;
-			// The triangle in UV space
-			FIndex3i UVTri(UVIndexBuffer[t], UVIndexBuffer[t + 1], UVIndexBuffer[t + 2]);
-
-			// the triangle in terms of the VertIDs in the DynamicMesh
-			FIndex3i TriVertIDs;
-			for (int c = 0; c < 3; ++c)
-			{
-				// the offset for this vertex in the LinearMesh
-				int32 Offset = VertexRemapArray[UVTri[c]];
 
-				int32 VertID = LinearMesh.VertToID[Offset];
 
-				TriVertIDs[c] = VertID;
-			}
-
-			// NB: this could be slow.. 
-			int32 TriID = FlippedMesh.FindTriangle(TriVertIDs[0], TriVertIDs[1], TriVertIDs[2]);
+bool FParameterizeMeshOp::ComputeUVs_XAtlas(FDynamicMesh3& Mesh,  TFunction<bool(float)>& Interrupter)
+{
+	// reverse mesh orientation
+	const bool bFixOrientation = true;
+	FDynamicMesh3 FlippedMesh(EMeshComponents::FaceGroups);
+	FlippedMesh.Copy(Mesh, false, false, false, false);
+	if (bFixOrientation)
+	{
+		FlippedMesh.ReverseOrientation(false);
+	}
 
-			checkSlow(TriID != FDynamicMesh3::InvalidID);
+	// Convert to a dense form.
+	FLinearMesh LinearMesh(FlippedMesh, false);
 
-			FIndex3i ElTri;
-			if (bFixOrientation)
-			{
-				ElTri = FIndex3i(UVOffsetToElID[UVTri[1]], UVOffsetToElID[UVTri[0]], UVOffsetToElID[UVTri[2]]);
-			}
-			else
-			{
-				ElTri = FIndex3i(UVOffsetToElID[ UVTri[0]], UVOffsetToElID[UVTri[1]], UVOffsetToElID[UVTri[2]]);
-			}
+	// Data to be populated by the UV generation tool
+	TArray<FVector2D> UVVertexBuffer;
+	TArray<int32>     UVIndexBuffer;
+	TArray<int32>     VertexRemapArray; // This maps the UV vertices to the original position vertices.  Note multiple UV vertices might share the same positional vertex (due to UV boundaries)
 
-			// add the triangle to the overlay
-			UVOverlay->SetTriangle(TriID, ElTri);
-		}
+	XAtlasWrapper::XAtlasChartOptions ChartOptions;
+	ChartOptions.MaxIterations = XAtlasMaxIterations;
+	XAtlasWrapper::XAtlasPackOptions PackOptions;
+	bool bSuccess = XAtlasWrapper::ComputeUVs(LinearMesh.IndexBuffer,
+											LinearMesh.VertexBuffer, 
+											ChartOptions,
+											PackOptions,
+											UVVertexBuffer, 
+											UVIndexBuffer, 
+											VertexRemapArray);
 
+	// Add the UVs to the FDynamicMesh
+	if (bSuccess)
+	{
+		CopyNewUVsToMesh(Mesh, LinearMesh, FlippedMesh, UVVertexBuffer, UVIndexBuffer, VertexRemapArray, bFixOrientation);
+	}
+	else
+	{
+		NewResultInfo.AddError(FGeometryError(0, LOCTEXT("XAtlasFailed", "XAtlas failed")));
 	}
-
 
 	return bSuccess;
 }
@@ -314,13 +363,11 @@ bool FParameterizeMeshOp::ComputeUVs(FDynamicMesh3& Mesh,  TFunction<bool(float)
 
 
 
-
-
-
 void FParameterizeMeshOp::CalculateResult(FProgressCancel* Progress)
 {
 	if (!InputMesh.IsValid())
 	{
+		SetResultInfo(FGeometryResult::Failed());
 		return;
 	}
 
@@ -328,13 +375,35 @@ void FParameterizeMeshOp::CalculateResult(FProgressCancel* Progress)
 
 	if (Progress && Progress->Cancelled())
 	{
+		SetResultInfo(FGeometryResult::Cancelled());
 		return;
 	}
 
 	// The UV atlas callback uses a float progress-based interrupter. 
-	TFunction<bool(float)> Iterrupter = [Progress](float)->bool {return !(Progress && Progress->Cancelled()); };
+	TFunction<bool(float)> Interrupter = [Progress](float)->bool {return !(Progress && Progress->Cancelled()); };
+
+#if HAVE_PROXYLOD
+	bool bUseXAtlas = (Method == EParamOpBackend::XAtlas);
+#else
+	bool bUseXAtlas = true;
+#endif
+
+	NewResultInfo = FGeometryResult(EGeometryResultType::InProgress);
 
 	// Single call to UVAtlas - perhaps respecting the poly groups.
-	ComputeUVs(*ResultMesh, Iterrupter);
+	bool bOK = false;
+	if (bUseXAtlas)
+	{
+		bOK = ComputeUVs_XAtlas(*ResultMesh, Interrupter);
+	}
+	else
+	{
+		bOK = ComputeUVs_UVAtlas(*ResultMesh, Interrupter);
+	}
+
+	NewResultInfo.SetSuccess(bOK, Progress);
+	SetResultInfo(NewResultInfo);
+}
+
 
-}
\ No newline at end of file
+#undef LOCTEXT_NAMESPACE
\ No newline at end of file
diff --git a/Engine/Plugins/Experimental/MeshModelingToolset/Source/ModelingOperatorsEditorOnly/Public/ParameterizationOps/ParameterizeMeshOp.h b/Engine/Plugins/Experimental/MeshModelingToolset/Source/ModelingOperatorsEditorOnly/Public/ParameterizationOps/ParameterizeMeshOp.h
--- a/Engine/Plugins/Experimental/MeshModelingToolset/Source/ModelingOperatorsEditorOnly/Public/ParameterizationOps/ParameterizeMeshOp.h
+++ b/Engine/Plugins/Experimental/MeshModelingToolset/Source/ModelingOperatorsEditorOnly/Public/ParameterizationOps/ParameterizeMeshOp.h
@@ -81,7 +81,19 @@ protected:
 		
 	};
 
-	bool ComputeUVs(FDynamicMesh3& InOutMesh,  TFunction<bool(float)>& Interrupter);
+	FGeometryResult NewResultInfo;
+
+	bool ComputeUVs_UVAtlas(FDynamicMesh3& InOutMesh, TFunction<bool(float)>& Interrupter);
+	bool ComputeUVs_XAtlas(FDynamicMesh3& InOutMesh, TFunction<bool(float)>& Interrupter);
+
+	void CopyNewUVsToMesh(
+		FDynamicMesh3& Mesh,
+		const FLinearMesh& LinearMesh,
+		const FDynamicMesh3& FlippedMesh,
+		const TArray<FVector2D>& UVVertexBuffer,
+		const TArray<int32>& UVIndexBuffer,
+		const TArray<int32>& VertexRemapArray,
+		bool bReverseOrientation);
 };
 
 } // end namespace UE::Geometry
diff --git a/Engine/Source/Runtime/GeometryCore/Private/Parameterization/UVPacking.cpp b/Engine/Source/Runtime/GeometryCore/Private/Parameterization/UVPacking.cpp
--- a/Engine/Source/Runtime/GeometryCore/Private/Parameterization/UVPacking.cpp
+++ b/Engine/Source/Runtime/GeometryCore/Private/Parameterization/UVPacking.cpp
@@ -947,6 +947,12 @@ void FUVPacker::GetIslandStats(IUVMeshView* Mesh, const TArray<int32>& Island, F
 		FIndex3i Triangle3D = Mesh->GetTriangle(tid);
 		FIndex3i TriangleUV = Mesh->GetUVTriangle(tid);
 
+		// skip invalid UV-triangles
+		if (TriangleUV.A < 0 || TriangleUV.B < 0 || TriangleUV.C < 0)
+		{
+			continue;
+		}
+
 		FVector3d Positions[3];
 		FVector2d UVs[3];
 
