UVEditor: Adding vertex and mesh selection options

#rnx

diff --git a/Engine/Plugins/Experimental/MeshModelingToolset/Source/ModelingComponents/Private/Selection/MeshSelectionMechanic.cpp b/Engine/Plugins/Experimental/MeshModelingToolset/Source/ModelingComponents/Private/Selection/MeshSelectionMechanic.cpp
--- a/Engine/Plugins/Experimental/MeshModelingToolset/Source/ModelingComponents/Private/Selection/MeshSelectionMechanic.cpp
+++ b/Engine/Plugins/Experimental/MeshModelingToolset/Source/ModelingComponents/Private/Selection/MeshSelectionMechanic.cpp
@@ -5,6 +5,7 @@
 #include "BaseBehaviors/BehaviorTargetInterfaces.h"
 #include "BaseBehaviors/SingleClickBehavior.h"
 #include "Drawing/LineSetComponent.h"
+#include "Drawing/PointSetComponent.h"
 #include "Drawing/PreviewGeometryActor.h"
 #include "InteractiveToolManager.h"
 #include "Polyline3.h"
@@ -31,6 +32,11 @@ void UMeshSelectionMechanic::Setup(UInteractiveTool* ParentToolIn)
 	LineSet = NewObject<ULineSetComponent>();
 	LineSet->SetLineMaterial(ToolSetupUtil::GetDefaultLineComponentMaterial(
 		GetParentTool()->GetToolManager(), /*bDepthTested*/ true));
+
+	PointSet = NewObject<UPointSetComponent>();
+	PointSet->SetPointMaterial(ToolSetupUtil::GetDefaultPointComponentMaterial(
+		GetParentTool()->GetToolManager(), /*bDepthTested*/ true));
+
 }
 
 void UMeshSelectionMechanic::SetWorld(UWorld* World)
@@ -57,6 +63,17 @@ void UMeshSelectionMechanic::SetWorld(UWorld* World)
 	{
 		LineSet->RegisterComponent();
 	}
+
+	PointSet->Rename(nullptr, PreviewGeometryActor); // Changes the "outer"
+	PointSet->AttachToComponent(LineSet, FAttachmentTransformRules::KeepWorldTransform);
+	if (PointSet->IsRegistered())
+	{
+		PointSet->ReregisterComponent();
+	}
+	else
+	{
+		PointSet->RegisterComponent();
+	}
 }
 
 void UMeshSelectionMechanic::Shutdown()
@@ -108,6 +125,7 @@ void UMeshSelectionMechanic::SetSelection(const FDynamicMeshSelection& Selection
 void UMeshSelectionMechanic::RebuildDrawnElements(const FTransform& StartTransform)
 {
 	LineSet->Clear();
+	PointSet->Clear();
 	PreviewGeometryActor->SetActorTransform(StartTransform);
 
 	// For us to end up with the StartTransform, we have to invert it, but only after applying
@@ -147,6 +165,17 @@ void UMeshSelectionMechanic::RebuildDrawnElements(const FTransform& StartTransfo
 					LineColor, LineThickness, DepthBias);
 		}
 	}
+	else if (CurrentSelection.Type == FDynamicMeshSelection::EType::Vertex)
+	{
+		PointSet->ReservePoints(CurrentSelection.SelectedIDs.Num());
+		for (int32 Vid : CurrentSelection.SelectedIDs)
+		{
+			FRenderablePoint PointToRender(TransformToApply(CurrentSelection.Mesh->GetVertex(Vid)),
+				                           PointColor,
+				                           PointThickness);
+			PointSet->AddPoint(PointToRender);
+		}
+	}
 }
 
 void UMeshSelectionMechanic::UpdateCentroid()
@@ -173,6 +202,14 @@ void UMeshSelectionMechanic::UpdateCentroid()
 		}
 		CurrentSelectionCentroid /= CurrentSelection.SelectedIDs.Num();
 	}
+	if (CurrentSelection.Type == FDynamicMeshSelection::EType::Vertex)
+	{
+		for (int32 Vid : CurrentSelection.SelectedIDs)
+		{
+			CurrentSelectionCentroid += CurrentSelection.Mesh->GetVertex(Vid);
+		}
+		CurrentSelectionCentroid /= CurrentSelection.SelectedIDs.Num();
+	}
 }
 
 FVector3d UMeshSelectionMechanic::GetCurrentSelectionCentroid()
@@ -213,7 +250,10 @@ void UMeshSelectionMechanic::OnClicked(const FInputDeviceRay& ClickPos)
 
 	if (!InMultiSelectMode() ||
 		(SelectionMode == EMeshSelectionMechanicMode::Component && CurrentSelection.Type != FDynamicMeshSelection::EType::Triangle) ||
-		(SelectionMode == EMeshSelectionMechanicMode::Edge && CurrentSelection.Type != FDynamicMeshSelection::EType::Edge) ) 
+		(SelectionMode == EMeshSelectionMechanicMode::Edge && CurrentSelection.Type != FDynamicMeshSelection::EType::Edge) ||
+		(SelectionMode == EMeshSelectionMechanicMode::Vertex && CurrentSelection.Type != FDynamicMeshSelection::EType::Vertex) || 
+		(SelectionMode == EMeshSelectionMechanicMode::Triangle && CurrentSelection.Type != FDynamicMeshSelection::EType::Triangle) || 
+		(SelectionMode == EMeshSelectionMechanicMode::Mesh && CurrentSelection.Type != FDynamicMeshSelection::EType::Triangle) )
 	{ // If we're not enabling multiselect, or our desired mode doesn't match our current selection type, clear the existing selection
 		CurrentSelection.SelectedIDs.Empty();
 		CurrentSelection.Mesh = nullptr;
@@ -286,6 +326,47 @@ void UMeshSelectionMechanic::OnClicked(const FInputDeviceRay& ClickPos)
 				CurrentSelection.Type = FDynamicMeshSelection::EType::Edge;
 			}
 		}
+		else if (SelectionMode == EMeshSelectionMechanicMode::Vertex)
+		{
+			// TODO: Improve this to handle super narrow, sliver triangles better, where testing near vertices can be difficult.
+
+			// Try to snap to one of the vertices
+			FIndex3i Vids = CurrentSelection.Mesh->GetTriangle(HitTid);
+
+			FGeometrySet3 GeometrySet;
+			for (int i = 0; i < 3; ++i)
+			{
+				GeometrySet.AddPoint(Vids[i], CurrentSelection.Mesh->GetTriVertex(HitTid, i));
+			}
+			FGeometrySet3::FNearest Result;
+			if (GeometrySet.FindNearestPointToRay(ClickPos.WorldRay, Result,
+				[this](const FVector3d& Position1, const FVector3d& Position2) {
+					return ToolSceneQueriesUtil::PointSnapQuery(CameraState,
+						Position1, Position2,
+						ToolSceneQueriesUtil::GetDefaultVisualAngleSnapThreshD()); }))
+			{
+				// Since we've already handled multiselect and not holding vertices above, this should be safe.	
+				// But just to be sure, we'll toss in a check to validate everything.
+				ensure((InMultiSelectMode() && CurrentSelection.Type == FDynamicMeshSelection::EType::Vertex) || CurrentSelection.SelectedIDs.IsEmpty());
+				CurrentSelection.SelectedIDs.Add(Result.ID);
+				CurrentSelection.Type = FDynamicMeshSelection::EType::Vertex;
+			}
+		}
+		else if (SelectionMode == EMeshSelectionMechanicMode::Triangle)
+		{
+			ensure((InMultiSelectMode() && CurrentSelection.Type == FDynamicMeshSelection::EType::Triangle) || CurrentSelection.SelectedIDs.IsEmpty());
+			CurrentSelection.SelectedIDs.Add(HitTid);
+			CurrentSelection.Type = FDynamicMeshSelection::EType::Triangle;
+		}
+		else if (SelectionMode == EMeshSelectionMechanicMode::Mesh)
+		{
+			ensure((InMultiSelectMode() && CurrentSelection.Type == FDynamicMeshSelection::EType::Triangle) || CurrentSelection.SelectedIDs.IsEmpty());
+			for (int32 Tid : CurrentSelection.Mesh->TriangleIndicesItr())
+			{
+				CurrentSelection.SelectedIDs.Add( Tid );
+			}
+			CurrentSelection.Type = FDynamicMeshSelection::EType::Triangle;
+		}
 	}
 
 	if (!(OriginalSelection == CurrentSelection))
diff --git a/Engine/Plugins/Experimental/MeshModelingToolset/Source/ModelingComponents/Public/Selection/MeshSelectionMechanic.h b/Engine/Plugins/Experimental/MeshModelingToolset/Source/ModelingComponents/Public/Selection/MeshSelectionMechanic.h
--- a/Engine/Plugins/Experimental/MeshModelingToolset/Source/ModelingComponents/Public/Selection/MeshSelectionMechanic.h
+++ b/Engine/Plugins/Experimental/MeshModelingToolset/Source/ModelingComponents/Public/Selection/MeshSelectionMechanic.h
@@ -15,6 +15,7 @@
 #include "MeshSelectionMechanic.generated.h"
 
 class ULineSetComponent;
+class UPointSetComponent;
 class APreviewGeometryActor;
 
 enum class EMeshSelectionMechanicMode
@@ -23,7 +24,12 @@ enum class EMeshSelectionMechanicMode
 	
 	// Not yet fully implemented for UV mesh purposes, since 
 	// we need to be able to select occluded edges
-	Edge
+	Edge,
+	Vertex,
+	Triangle,
+
+	// TODO: This might be good to rename later. And determine how it might interact with multi-mesh selection?
+	Mesh
 };
 
 UCLASS()
@@ -95,6 +101,9 @@ protected:
 	UPROPERTY()
 	TObjectPtr<ULineSetComponent> LineSet = nullptr;
 
+	UPROPERTY()
+	TObjectPtr<UPointSetComponent> PointSet = nullptr;
+
 
 	TArray<TSharedPtr<FDynamicMeshAABBTree3>> MeshSpatials;
 	TArray<FTransform> MeshTransforms;
@@ -103,11 +112,13 @@ protected:
 	FViewCameraState CameraState;
 
 	FColor LineColor = FColor::Yellow;
+	FColor PointColor = FColor::Yellow;
 
 	bool bShiftToggle = false;
 	static const int32 ShiftModifierID = 1;
 
 	float LineThickness = 3;
+	float PointThickness = 6;
 	float DepthBias = 0.3;
 
 	FVector3d CurrentSelectionCentroid;
diff --git a/Engine/Plugins/Experimental/UVEditor/Source/UVEditorTools/Private/UVSelectTool.cpp b/Engine/Plugins/Experimental/UVEditor/Source/UVEditorTools/Private/UVSelectTool.cpp
--- a/Engine/Plugins/Experimental/UVEditor/Source/UVEditorTools/Private/UVSelectTool.cpp
+++ b/Engine/Plugins/Experimental/UVEditor/Source/UVEditorTools/Private/UVSelectTool.cpp
@@ -65,9 +65,7 @@ void UUVSelectTool::Setup()
 	SelectionMechanic->SetWorld(Targets[0]->UnwrapPreview->GetWorld());
 	SelectionMechanic->OnSelectionChanged.AddUObject(this, &UUVSelectTool::OnSelectionChanged);
 
-	SelectionMechanic->SelectionMode =
-		Settings->SelectionMode == EUVSelectToolSelectionMode::Island ? EMeshSelectionMechanicMode::Component
-		: EMeshSelectionMechanicMode::Edge;
+	ConfigureSelectionModeFromControls();
 
 	// Retrieve cached AABB tree storage, or else set it up
 	UContextObjectStore* ContextStore = GetToolManager()->GetContextObjectStore();
@@ -170,9 +168,7 @@ void UUVSelectTool::Shutdown(EToolShutdownType ShutdownType)
 
 void UUVSelectTool::OnPropertyModified(UObject* PropertySet, FProperty* Property)
 {
-	SelectionMechanic->SelectionMode =
-		Settings->SelectionMode == EUVSelectToolSelectionMode::Island ? EMeshSelectionMechanicMode::Component
-		: EMeshSelectionMechanicMode::Edge;
+	ConfigureSelectionModeFromControls();
 }
 
 void UUVSelectTool::UpdateGizmo()
@@ -194,6 +190,31 @@ void UUVSelectTool::UpdateGizmo()
 	TransformGizmo->SetVisibility(bGizmoEnabled && !SelectionMechanic->GetCurrentSelection().IsEmpty());
 }
 
+void UUVSelectTool::ConfigureSelectionModeFromControls()
+{
+	switch (Settings->SelectionMode)
+	{
+	case EUVSelectToolSelectionMode::Island:
+		SelectionMechanic->SelectionMode = EMeshSelectionMechanicMode::Component;
+		break;
+	case EUVSelectToolSelectionMode::Edge:
+		SelectionMechanic->SelectionMode = EMeshSelectionMechanicMode::Edge;
+		break;
+	case EUVSelectToolSelectionMode::Vertex:
+		SelectionMechanic->SelectionMode = EMeshSelectionMechanicMode::Vertex;
+		break;
+	case EUVSelectToolSelectionMode::Triangle:
+		SelectionMechanic->SelectionMode = EMeshSelectionMechanicMode::Triangle;
+		break;
+	case EUVSelectToolSelectionMode::Mesh:
+		SelectionMechanic->SelectionMode = EMeshSelectionMechanicMode::Mesh;
+		break;
+	default:
+		ensure(false);
+		break;
+	}
+}
+
 void UUVSelectTool::OnSelectionChanged()
 {
 	const FDynamicMeshSelection& Selection = SelectionMechanic->GetCurrentSelection();
@@ -280,6 +301,29 @@ void UUVSelectTool::OnSelectionChanged()
 				}
 			}
 		}
+		else if (Selection.Type == FDynamicMeshSelection::EType::Vertex)
+		{
+			for (int32 Vid : Selection.SelectedIDs)
+			{
+				if (!VidSet.Contains(Vid))
+				{
+					VidSet.Add(Vid);
+					MovingVids.Add(Vid);
+				}
+
+				TArray<int> TidOneRing;
+				Selection.Mesh->GetVertexOneRingTriangles(Vid, TidOneRing);
+				for (int32 Tid : TidOneRing)
+				{
+					if (!TidSet.Contains(Tid))
+					{
+						TidSet.Add(Tid);
+						SelectedTids.Add(Tid);
+					}
+				}
+
+			}
+		}
 		else
 		{
 			check(false);
@@ -300,7 +344,7 @@ void UUVSelectTool::UpdateLivePreviewLines()
 	{
 		FTransform MeshTransform = Targets[SelectionTargetIndex]->AppliedPreview->PreviewMesh->GetTransform();
 		const FDynamicMesh3* LivePreviewMesh = Targets[SelectionTargetIndex]->AppliedCanonical.Get();
-
+		
 		for (int32 Eid : BoundaryEids)
 		{
 			FVector3d Vert1, Vert2;
@@ -311,6 +355,7 @@ void UUVSelectTool::UpdateLivePreviewLines()
 				MeshTransform.TransformPosition(Vert2), 
 				FColor::Yellow, 2, 1.5);
 		}
+		
 	}
 }
 
diff --git a/Engine/Plugins/Experimental/UVEditor/Source/UVEditorTools/Public/UVSelectTool.h b/Engine/Plugins/Experimental/UVEditor/Source/UVEditorTools/Public/UVSelectTool.h
--- a/Engine/Plugins/Experimental/UVEditor/Source/UVEditorTools/Public/UVSelectTool.h
+++ b/Engine/Plugins/Experimental/UVEditor/Source/UVEditorTools/Public/UVSelectTool.h
@@ -47,7 +47,10 @@ UENUM()
 enum class EUVSelectToolSelectionMode : uint8
 {
 	Island,
-	Edge
+	Edge,
+	Vertex,
+	Triangle,
+	Mesh
 };
 
 UCLASS()
@@ -112,6 +115,8 @@ protected:
 	virtual void UpdateGizmo();
 	virtual void UpdateLivePreviewLines();
 
+	void ConfigureSelectionModeFromControls();
+
 	UWorld* TargetWorld;
 
 	UPROPERTY()
