ApproximateActors Improvements
MeshSceneAdapter: handle mis-configured MeshDescription material slot names that reference ImportedMaterialSlotName instead of MaterialSlotName, and fall back to material 0 / default surfarce material in cases where material is not defined (which is normal engine behavior elsewhere). Add CVar to control multi-threading, to support debugging. Fix handling of nearly-thick-enough objects that were incorrectly being converted to very thin shells.
ApproximateActors: Expose options for ground plane clipping, UV generation method, tiny parts filtering, implement support for these options. Add optional base-plane occluder to allow occlusion filtering to remove "bottom" geometry. Implement fallback to Boxmap UVs if UV generation fails. Set default output material on generated Collision meshes.
Expose above options in FMeshApproximationSettings for MergeActors/etc UI

#rnx

diff --git a/Engine/Plugins/Experimental/MeshModelingToolset/Source/GeometryProcessingAdapters/Private/ApproximateActorsImpl.cpp b/Engine/Plugins/Experimental/MeshModelingToolset/Source/GeometryProcessingAdapters/Private/ApproximateActorsImpl.cpp
--- a/Engine/Plugins/Experimental/MeshModelingToolset/Source/GeometryProcessingAdapters/Private/ApproximateActorsImpl.cpp
+++ b/Engine/Plugins/Experimental/MeshModelingToolset/Source/GeometryProcessingAdapters/Private/ApproximateActorsImpl.cpp
@@ -11,11 +11,15 @@
 #include "Implicit/Solidify.h"
 #include "Implicit/Morphology.h"
 #include "Operations/RemoveOccludedTriangles.h"
+#include "Operations/MeshPlaneCut.h"
+#include "ConstrainedDelaunay2.h"
 #include "ParameterizationOps/ParameterizeMeshOp.h"
 #include "Parameterization/MeshUVPacking.h"
 #include "MeshQueries.h"
 #include "ProjectionTargets.h"
 #include "Selections/MeshFaceSelection.h"
+#include "Generators/RectangleMeshGenerator.h"
+#include "Parameterization/DynamicMeshUVEditor.h"
 
 #include "AssetUtils/CreateStaticMeshUtil.h"
 #include "AssetUtils/CreateTexture2DUtil.h"
@@ -401,10 +405,21 @@ static TSharedPtr<FApproximationMeshData> GenerateApproximationMesh(
 	// avoid insane memory usage
 	if (VoxelDimTarget > Options.ClampVoxelDimension)
 	{
-		UE_LOG(LogApproximateActors, Warning, TEXT("FApproximateActorsImpl - very large voxel size %d clamped to %d"), VoxelDimTarget, Options.ClampVoxelDimension);
+		UE_LOG(LogApproximateActors, Warning, TEXT("very large voxel size %d clamped to %d"), VoxelDimTarget, Options.ClampVoxelDimension);
 		VoxelDimTarget = Options.ClampVoxelDimension;
 	}
 
+	// make ground plane
+	FVector3d GroundPlaneOrigin;
+	FPlane3d GroundClipPlane;
+	bool bHaveGroundClipPlane = false;
+	if (Options.GroundPlanePolicy == IGeometryProcessing_ApproximateActors::EGroundPlanePolicy::FixedZHeightGroundPlane)
+	{
+		GroundPlaneOrigin = FVector3d(SceneBounds.Center().X, SceneBounds.Center().Y, Options.GroundPlaneZHeight);
+		GroundClipPlane = FPlane3d(FVector3d::UnitZ(), GroundPlaneOrigin);
+		bHaveGroundClipPlane = true;
+	}
+
 	Progress.EnterProgressFrame(1.f, LOCTEXT("SolidifyMesh", "Approximating Mesh..."));
 
 	FWindingNumberBasedSolidify Solidify(
@@ -451,6 +466,8 @@ static TSharedPtr<FApproximationMeshData> GenerateApproximationMesh(
 		}
 	}
 
+	// TODO: try doing base clipping here to speed up simplification? slight risk of introducing border issues...
+
 	// if mesh has no triangles, something has gone wrong
 	if (CurResultMesh == nullptr || CurResultMesh->TriangleCount() == 0)
 	{
@@ -539,10 +556,29 @@ static TSharedPtr<FApproximationMeshData> GenerateApproximationMesh(
 		NoTransforms.Add(FTransform3d::Identity());
 		TArray<FDynamicMeshAABBTree3*> Spatials;
 		Spatials.Add(&CurResultMeshSpatial);
+
+		FAxisAlignedBox3d Bounds = CurResultMesh->GetBounds();
+
+		FDynamicMesh3 BasePlaneOccluderMesh;
+		FDynamicMeshAABBTree3 BasePlaneOccluderSpatial;
+		if (Options.bAddDownwardFacesOccluder)
+		{
+			FRectangleMeshGenerator RectGen;
+			RectGen.Origin = Bounds.Center();
+			RectGen.Origin.Z = Bounds.Min.Z - 1.0;
+			RectGen.Normal = FVector3f::UnitZ();
+			RectGen.Width = RectGen.Height = 10.0 * Bounds.MaxDim();
+			BasePlaneOccluderMesh.Copy(&RectGen.Generate());
+			BasePlaneOccluderSpatial.SetMesh(&BasePlaneOccluderMesh, true);
+			NoTransforms.Add(FTransform3d::Identity());
+			Spatials.Add(&BasePlaneOccluderSpatial);
+		}
+
 		{
 			TRACE_CPUPROFILER_EVENT_SCOPE(ApproximateActorsImpl_Generate_Occlusion_Compute);
 			Remover.Select(NoTransforms, Spatials, {}, NoTransforms);
 		}
+		int32 NumRemoved = 0;
 		if (Remover.RemovedT.Num() > 0)
 		{
 			FMeshFaceSelection Selection(CurResultMesh);
@@ -551,21 +587,62 @@ static TSharedPtr<FApproximationMeshData> GenerateApproximationMesh(
 				Selection.Select(Remover.RemovedT);
 				Selection.ExpandToOneRingNeighbours(1);
 				Selection.ContractBorderByOneRingNeighbours(2);
+
+				// select any tris w/ all verts below clip plane
+				if (Options.GroundPlaneClippingPolicy == IGeometryProcessing_ApproximateActors::EGroundPlaneClippingPolicy::DiscardFullyHiddenFaces)
+				{
+					if (bHaveGroundClipPlane)
+					{
+						for (int32 tid : CurResultMesh->TriangleIndicesItr())
+						{
+							FVector3d A, B, C;
+							CurResultMesh->GetTriVertices(tid, A, B, C);
+							if (GroundClipPlane.WhichSide(A) <= 0 && GroundClipPlane.WhichSide(B) <= 0 && GroundClipPlane.WhichSide(C) <= 0)
+							{
+								Selection.Select(tid);
+							}
+						}
+					}
+					else
+					{
+						UE_LOG(LogApproximateActors, Warning, TEXT("DiscardFullyHiddenFaces Ground Plane Clipping Policy ignored because no Ground Clip Plane is set"));
+					}
+				}
 			}
 			FDynamicMeshEditor Editor(CurResultMesh);
 			{
 				TRACE_CPUPROFILER_EVENT_SCOPE(ApproximateActorsImpl_Generate_Occlusion_Delete);
-				Editor.RemoveTriangles(Selection.AsArray(), true);
+				TArray SelectionArray(Selection.AsArray());
+				NumRemoved = SelectionArray.Num();
+				Editor.RemoveTriangles(SelectionArray, true);
 			}
 		}
 
 		if (Options.bVerbose)
 		{
-			UE_LOG(LogApproximateActors, Warning, TEXT("Occlusion-Filtered mesh has %d triangles"), CurResultMesh->TriangleCount());
+			UE_LOG(LogApproximateActors, Warning, TEXT("Occlusion-Filtered mesh has %d triangles (removed %d)"), CurResultMesh->TriangleCount(), NumRemoved);
 		}
 	}
 
 
+	if (Options.GroundPlaneClippingPolicy == IGeometryProcessing_ApproximateActors::EGroundPlaneClippingPolicy::CutFaces ||
+		Options.GroundPlaneClippingPolicy == IGeometryProcessing_ApproximateActors::EGroundPlaneClippingPolicy::CutFacesAndFill)
+	{
+		if (bHaveGroundClipPlane)
+		{
+			FMeshPlaneCut PlaneCut(CurResultMesh, GroundPlaneOrigin, -GroundClipPlane.Normal);
+			PlaneCut.Cut();
+			if (Options.GroundPlaneClippingPolicy == IGeometryProcessing_ApproximateActors::EGroundPlaneClippingPolicy::CutFacesAndFill)
+			{
+				PlaneCut.HoleFill(ConstrainedDelaunayTriangulate<double>, true);
+			}
+		}
+		else
+		{
+			UE_LOG(LogApproximateActors, Warning, TEXT("Ground Plane Cut/Fill Policy ignored because no Ground Clip Plane is set"));
+		}
+	}
+
 
 	// re-enable attributes
 	CurResultMesh->EnableAttributes();
@@ -598,34 +675,72 @@ static TSharedPtr<FApproximationMeshData> GenerateApproximationMesh(
 	Progress.EnterProgressFrame(1.f, LOCTEXT("ComputingUVs", "Computing UVs..."));
 
 	// compute UVs
+	bool bHaveValidUVs = true;
 	TSharedPtr<FDynamicMesh3, ESPMode::ThreadSafe> UVInputMesh = MakeShared<FDynamicMesh3, ESPMode::ThreadSafe>();
 	*UVInputMesh = MoveTemp(*CurResultMesh);
 	FParameterizeMeshOp ParameterizeMeshOp;
-	ParameterizeMeshOp.Stretch = 0.1;
+	ParameterizeMeshOp.Stretch = Options.UVAtlasStretchTarget;
 	ParameterizeMeshOp.NumCharts = 0;
 	ParameterizeMeshOp.InputMesh = UVInputMesh;
-	ParameterizeMeshOp.Method = UE::Geometry::EParamOpBackend::UVAtlas;
+	ParameterizeMeshOp.Method = UE::Geometry::EParamOpBackend::XAtlas;
+	if (Options.UVPolicy == IGeometryProcessing_ApproximateActors::EUVGenerationPolicy::PreferUVAtlas)
+	{
+		ParameterizeMeshOp.Method = UE::Geometry::EParamOpBackend::UVAtlas;
+	}
 	FProgressCancel UVProgressCancel;
 	{
 		TRACE_CPUPROFILER_EVENT_SCOPE(ApproximateActorsImpl_Generate_GenerateUVs);
 		ParameterizeMeshOp.CalculateResult(&UVProgressCancel);
 	}
 
-	TUniquePtr<FDynamicMesh3> FinalMesh = ParameterizeMeshOp.ExtractResult();
+	TUniquePtr<FDynamicMesh3> FinalMesh;
+
+	FGeometryResult UVResultInfo = ParameterizeMeshOp.GetResultInfo();
+	if (! UVResultInfo.HasResult())
+	{
+		UE_LOG(LogApproximateActors, Warning, TEXT("UV Auto-Generation Failed for target path %s"), *Options.BasePackagePath);
+		bHaveValidUVs = false;
+		FinalMesh = MakeUnique<FDynamicMesh3>(MoveTemp(*UVInputMesh));
+	}
+	else
+	{
+		FinalMesh = ParameterizeMeshOp.ExtractResult();
+	}
+
+	// if UVs failed, fall back to box projection
+	if (!bHaveValidUVs)
+	{
+		FDynamicMeshUVEditor UVEditor(FinalMesh.Get(), 0, true);
+		TArray<int32> AllTriangles;
+		for (int32 tid : FinalMesh->TriangleIndicesItr())
+		{
+			AllTriangles.Add(tid);
+		}
+		UVEditor.SetTriangleUVsFromBoxProjection(AllTriangles, [&](const FVector3d& P) { return P; },
+												 FFrame3d(FinalMesh->GetBounds().Center()), FVector3d::One());
+		bHaveValidUVs = true;
+	}
+
 
 	Progress.EnterProgressFrame(1.f, LOCTEXT("PackingUVs", "Packing UVs..."));
 
 	// repack UVs
-	FDynamicMeshUVOverlay* RepackUVLayer = FinalMesh->Attributes()->PrimaryUV();
-	RepackUVLayer->SplitBowties();
-	FDynamicMeshUVPacker Packer(RepackUVLayer);
-	Packer.TextureResolution = Options.TextureImageSize / 4;		// maybe too conservative? We don't have gutter control currently.
-	Packer.GutterSize = 1.0;		// not clear this works
-	Packer.bAllowFlips = false;
+	if (bHaveValidUVs)
 	{
-		TRACE_CPUPROFILER_EVENT_SCOPE(ApproximateActorsImpl_Generate_PackUVs);
-		bool bOK = Packer.StandardPack();
-		ensure(bOK);
+		FDynamicMeshUVOverlay* RepackUVLayer = FinalMesh->Attributes()->PrimaryUV();
+		RepackUVLayer->SplitBowties();
+		FDynamicMeshUVPacker Packer(RepackUVLayer);
+		Packer.TextureResolution = Options.TextureImageSize / 4;		// maybe too conservative? We don't have gutter control currently.
+		Packer.GutterSize = 1.0;		// not clear this works
+		Packer.bAllowFlips = false;
+		{
+			TRACE_CPUPROFILER_EVENT_SCOPE(ApproximateActorsImpl_Generate_PackUVs);
+			bool bPackingOK = Packer.StandardPack();
+			if (!bPackingOK)
+			{
+				UE_LOG(LogApproximateActors, Warning, TEXT("UV Packing Failed for target path %s"), *Options.BasePackagePath);
+			}
+		}
 	}
 
 	Progress.EnterProgressFrame(1.f, LOCTEXT("ComputingTangents", "Computing Tangents..."));
@@ -705,6 +820,8 @@ IGeometryProcessing_ApproximateActors::FOptions FApproximateActorsImpl::Construc
 
 	Options.bAutoThickenThinParts = UseSettings.bAttemptAutoThickening;
 	Options.AutoThickenThicknessMeters = UseSettings.TargetMinThicknessMultiplier * UseSettings.ApproximationAccuracy;
+	Options.bIgnoreTinyParts = UseSettings.bIgnoreTinyParts;
+	Options.TinyPartMaxDimensionMeters = UseSettings.TinyPartSizeMultiplier * UseSettings.ApproximationAccuracy;
 
 	Options.BaseCappingPolicy = IGeometryProcessing_ApproximateActors::EBaseCappingPolicy::NoBaseCapping;
 	if (UseSettings.BaseCapping == EMeshApproximationBaseCappingType::ConvexPolygon)
@@ -721,8 +838,36 @@ IGeometryProcessing_ApproximateActors::FOptions FApproximateActorsImpl::Construc
 	Options.bApplyMorphology = UseSettings.bFillGaps;
 	Options.MorphologyDistanceMeters = UseSettings.GapDistance;
 
+	if (UseSettings.GroundClipping == EMeshApproximationGroundPlaneClippingPolicy::NoGroundClipping)
+	{
+		Options.GroundPlanePolicy = EGroundPlanePolicy::NoGroundPlane;
+		Options.GroundPlaneClippingPolicy = IGeometryProcessing_ApproximateActors::EGroundPlaneClippingPolicy::NoClipping;
+	}
+	else if (UseSettings.GroundClipping == EMeshApproximationGroundPlaneClippingPolicy::DiscardWithZPlane)
+	{
+		Options.GroundPlanePolicy = EGroundPlanePolicy::FixedZHeightGroundPlane;
+		Options.GroundPlaneZHeight = UseSettings.GroundClippingZHeight;
+		Options.GroundPlaneClippingPolicy = IGeometryProcessing_ApproximateActors::EGroundPlaneClippingPolicy::DiscardFullyHiddenFaces;
+	}
+	else if (UseSettings.GroundClipping == EMeshApproximationGroundPlaneClippingPolicy::CutWithZPlane)
+	{
+		Options.GroundPlanePolicy = EGroundPlanePolicy::FixedZHeightGroundPlane;
+		Options.GroundPlaneZHeight = UseSettings.GroundClippingZHeight;
+		Options.GroundPlaneClippingPolicy = IGeometryProcessing_ApproximateActors::EGroundPlaneClippingPolicy::CutFaces;
+	}
+	else if (UseSettings.GroundClipping == EMeshApproximationGroundPlaneClippingPolicy::CutAndFillWithZPlane)
+	{
+		Options.GroundPlanePolicy = EGroundPlanePolicy::FixedZHeightGroundPlane;
+		Options.GroundPlaneZHeight = UseSettings.GroundClippingZHeight;
+		Options.GroundPlaneClippingPolicy = IGeometryProcessing_ApproximateActors::EGroundPlaneClippingPolicy::CutFacesAndFill;
+	}
+
+
+
 	Options.OcclusionPolicy = (UseSettings.OcclusionMethod == EOccludedGeometryFilteringPolicy::VisibilityBasedFiltering) ?
 		IGeometryProcessing_ApproximateActors::EOcclusionPolicy::VisibilityBased : IGeometryProcessing_ApproximateActors::EOcclusionPolicy::None;
+	Options.bAddDownwardFacesOccluder = UseSettings.bOccludeFromBottom;
+
 	Options.FixedTriangleCount = UseSettings.TargetTriCount;
 	if (UseSettings.SimplifyMethod == EMeshApproximationSimplificationPolicy::TrianglesPerArea)
 	{
@@ -739,6 +884,9 @@ IGeometryProcessing_ApproximateActors::FOptions FApproximateActorsImpl::Construc
 		Options.MeshSimplificationPolicy = IGeometryProcessing_ApproximateActors::ESimplificationPolicy::FixedTriangleCount;
 	}
 
+	Options.UVPolicy = (UseSettings.UVGenerationMethod == EMeshApproximationUVGenerationPolicy::PreferUVAtlas) ?
+		IGeometryProcessing_ApproximateActors::EUVGenerationPolicy::PreferUVAtlas : IGeometryProcessing_ApproximateActors::EUVGenerationPolicy::PreferXAtlas;
+
 	Options.bCalculateHardNormals = UseSettings.bEstimateHardNormals;
 	Options.HardNormalsAngleDeg = FMath::Clamp(UseSettings.HardNormalAngle, 0.001f, 89.99f);
 
@@ -813,8 +961,8 @@ void FApproximateActorsImpl::GenerateApproximationForActorSet(const TArray<AActo
 	SceneBuildOptions.bThickenThinMeshes = Options.bAutoThickenThinParts;
 	SceneBuildOptions.DesiredMinThickness = Options.AutoThickenThicknessMeters * 100.0;		// convert to cm (UE Units)
 	// filter out objects smaller than 10% of voxel size
-	SceneBuildOptions.bFilterTinyObjects = true;
-	SceneBuildOptions.TinyObjectBoxMaxDimension = ApproxAccuracy * 0.1;
+	SceneBuildOptions.bFilterTinyObjects = Options.bIgnoreTinyParts;
+	SceneBuildOptions.TinyObjectBoxMaxDimension = Options.TinyPartMaxDimensionMeters;
 	SceneBuildOptions.bPrintDebugMessages = Options.bVerbose;
 	{
 		TRACE_CPUPROFILER_EVENT_SCOPE(ApproximateActorsImpl_Generate_BuildScene);
@@ -1044,6 +1192,10 @@ UStaticMesh* FApproximateActorsImpl::EmitGeneratedMeshAsset(
 	{
 		MeshAssetOptions.AssetMaterials.Add(Material);
 	}
+	else
+	{
+		MeshAssetOptions.AssetMaterials.Add(UMaterial::GetDefaultMaterial(MD_Surface));
+	}
 	FStaticMeshResults MeshAssetOutputs;
 	ECreateStaticMeshResult ResultCode = UE::AssetUtils::CreateStaticMeshAsset(MeshAssetOptions, MeshAssetOutputs);
 	ensure(ResultCode == ECreateStaticMeshResult::Ok);
diff --git a/Engine/Plugins/Experimental/MeshModelingToolset/Source/ModelingComponents/Private/Scene/MeshSceneAdapter.cpp b/Engine/Plugins/Experimental/MeshModelingToolset/Source/ModelingComponents/Private/Scene/MeshSceneAdapter.cpp
--- a/Engine/Plugins/Experimental/MeshModelingToolset/Source/ModelingComponents/Private/Scene/MeshSceneAdapter.cpp
+++ b/Engine/Plugins/Experimental/MeshModelingToolset/Source/ModelingComponents/Private/Scene/MeshSceneAdapter.cpp
@@ -31,6 +31,12 @@
 using namespace UE::Geometry;
 
 
+static TAutoConsoleVariable<int32> CVarMeshSceneAdapterDisableMultiThreading(
+	TEXT("geometry.MeshSceneAdapter.SingleThreaded"),
+	0,
+	TEXT("Determines whether or not to use multi-threading in MeshSceneAdapter.\n"));
+
+
 /** 
  * Compute the bounds of the vertices of Mesh, under 3D transformation TransformFunc
  * @return computed bounding box
@@ -244,6 +250,10 @@ public:
 	bool bUseDistanceShellForWinding = false;
 	// unsigned distance isovalue that defines inside
 	double WindingShellThickness = 0.0;
+	// if true, do winding query before testing shell thickness. This is necessary in cases where the
+	// shell thickness is very small (eg slightly thickening a relatively thick object). Then the thickness
+	// query would only create a shell instead of a full solid
+	bool bRequiresWindingQueryFallback = false;
 
 	TUniquePtr<TMeshAABBTree3<FDynamicMesh3>> AABBTree;
 	TUniquePtr<TFastWindingTree<FDynamicMesh3>> FWNTree;
@@ -257,7 +267,7 @@ public:
 				TRACE_CPUPROFILER_EVENT_SCOPE(MeshScene_WrapperBuild_DMesh_AABBTree);
 				AABBTree = MakeUnique<TMeshAABBTree3<FDynamicMesh3>>(&Mesh, true);
 			}
-			if (bUseDistanceShellForWinding == false)
+			if (bUseDistanceShellForWinding == false || bRequiresWindingQueryFallback)
 			{
 				{
 					TRACE_CPUPROFILER_EVENT_SCOPE(MeshScene_WrapperBuild_DMesh_FWNTree);
@@ -296,6 +306,17 @@ public:
 	{
 		if (bUseDistanceShellForWinding)
 		{
+			if (bRequiresWindingQueryFallback)
+			{
+				double WindingNumber = bHasBakedTransform ?
+					FWNTree->FastWindingNumber(P) :
+					FWNTree->FastWindingNumber(LocalToWorldTransform.InverseTransformPosition(P));
+				if (WindingNumber > 0.99)
+				{
+					return 1.0;
+				}
+			}
+
 			if (bHasBakedTransform || bHasBakedScale)
 			{
 				FVector3d UseP = (bHasBakedTransform) ? P : LocalToWorldTransform.InverseTransformPosition(P);
@@ -365,6 +386,58 @@ public:
 
 };
 
+
+
+/**
+ * Utility function for handling static mesh material queries
+ */
+static int32 GetStaticMeshMaterialIndexFromSlotName(UStaticMesh* StaticMesh, FName MaterialSlotName, bool bFallbackToSlot0OnFailure)
+{
+	const TArray<FStaticMaterial>& StaticMaterials = StaticMesh->GetStaticMaterials();
+	int32 NumMaterials = StaticMaterials.Num();
+
+	// search primary slot names first
+	for (int32 MaterialIndex = 0; MaterialIndex < NumMaterials; ++MaterialIndex)
+	{
+		const FStaticMaterial& StaticMaterial = StaticMaterials[MaterialIndex];
+		if (StaticMaterial.MaterialSlotName == MaterialSlotName)
+		{
+			return MaterialIndex;
+		}
+	}
+
+	// sometimes we are given the ImportedSlotName
+	for (int32 MaterialIndex = 0; MaterialIndex < NumMaterials; ++MaterialIndex)
+	{
+		const FStaticMaterial& StaticMaterial = StaticMaterials[MaterialIndex];
+		if (StaticMaterial.ImportedMaterialSlotName == MaterialSlotName)
+		{
+			return MaterialIndex;
+		}
+	}
+
+	return (bFallbackToSlot0OnFailure && NumMaterials > 0) ? 0 : INDEX_NONE;
+}
+
+
+
+static UMaterialInterface* GetStaticMeshMaterialFromSlotName(UStaticMesh* StaticMesh, FName MaterialSlotName, bool bFallbackToSlot0OnFailure)
+{
+	int32 MaterialIndex = GetStaticMeshMaterialIndexFromSlotName(StaticMesh, MaterialSlotName, bFallbackToSlot0OnFailure);
+	if (MaterialIndex > INDEX_NONE)
+	{
+		UMaterialInterface* MaterialInterface = StaticMesh->GetMaterial(MaterialIndex);
+		return (MaterialInterface != nullptr) ? MaterialInterface : UMaterial::GetDefaultMaterial(MD_Surface);
+	}
+	else
+	{
+		return UMaterial::GetDefaultMaterial(MD_Surface);
+	}
+}
+
+
+
+
 /**
  * Mesh adapter to filter out all mesh sections that are using materials for which the material domain isn't "Surface"
  * This excludes decals for example.
@@ -382,17 +455,15 @@ public:
 
 		for (FPolygonGroupID PolygonGroupID : MeshIn->PolygonGroups().GetElementIDs())
 		{
-			bool bValidPolyGroup = false;
+			// Try to find correct material index for the slot name in the MeshDescription. 
+			// Sometimes this data is wrong, in that case fall back to slot 0
+			FName MaterialSlotName = MaterialSlotNames[PolygonGroupID];
+			int32 MaterialIndex = GetStaticMeshMaterialIndexFromSlotName(StaticMeshIn, MaterialSlotName, true);
 
-			int32 MaterialIndex = StaticMeshIn->GetMaterialIndex(MaterialSlotNames[PolygonGroupID]);
-			if (MaterialIndex > INDEX_NONE)
-			{
-				UMaterialInterface* MaterialInterface = StaticMeshIn->GetMaterial(MaterialIndex);
-				if (MaterialInterface && MaterialInterface->GetMaterial())
-				{
-					bValidPolyGroup = MaterialInterface->GetMaterial()->MaterialDomain == EMaterialDomain::MD_Surface;
-				}
-			}
+			UMaterialInterface* MaterialInterface = GetStaticMeshMaterialFromSlotName(StaticMeshIn, MaterialSlotName, true);
+			bool bValidPolyGroup = (MaterialInterface != nullptr)
+				&& (MaterialInterface->GetMaterial() != nullptr)
+				&& (MaterialInterface->GetMaterial()->MaterialDomain == EMaterialDomain::MD_Surface);
 
 			if (bValidPolyGroup)
 			{
@@ -851,8 +922,8 @@ void ConstructUniqueScalesMapping(
 
 void FMeshSceneAdapter::Build_FullDecompose(const FMeshSceneAdapterBuildOptions& BuildOptions)
 {
-	EParallelForFlags ParallelFlags = EParallelForFlags::Unbalanced;
-	//EParallelForFlags ParallelFlags = EParallelForFlags::ForceSingleThread;
+	EParallelForFlags ParallelFlags = (CVarMeshSceneAdapterDisableMultiThreading.GetValueOnAnyThread() != 0) ?
+		EParallelForFlags::ForceSingleThread : EParallelForFlags::Unbalanced;
 
 	// initial list of spatial wrappers that need to be built
 	TArray<FSpatialWrapperInfo*> ToBuild;
@@ -1200,6 +1271,7 @@ void FMeshSceneAdapter::Build_FullDecompose(const FMeshSceneAdapterBuildOptions&
 			{
 				NewInstancedMesh->bUseDistanceShellForWinding = true;
 				NewInstancedMesh->WindingShellThickness = 0.5 * (BuildOptions.DesiredMinThickness - Submesh.ComputedThickness);
+				NewInstancedMesh->bRequiresWindingQueryFallback = (NewInstancedMesh->WindingShellThickness < 0.6*BuildOptions.DesiredMinThickness);
 			}
 
 			TSharedPtr<FSpatialWrapperInfo> NewWrapperInfo = MakeShared<FSpatialWrapperInfo>();
@@ -1288,7 +1360,7 @@ void FMeshSceneAdapter::Build_FullDecompose(const FMeshSceneAdapterBuildOptions&
 		{
 			return A.NewUniqueTris > B.NewUniqueTris;
 		});
-		for (int32 k = 0; k < 20; ++k)
+		for (int32 k = 0; k < FMath::Min(Stats.Num(), 20); ++k)
 		{
 			const FProcessedSourceMeshStats& Stat = Stats[k];
 			FString NewSubmeshesStats;
diff --git a/Engine/Source/Developer/GeometryProcessingInterfaces/Public/GeometryProcessingInterfaces/ApproximateActors.h b/Engine/Source/Developer/GeometryProcessingInterfaces/Public/GeometryProcessingInterfaces/ApproximateActors.h
--- a/Engine/Source/Developer/GeometryProcessingInterfaces/Public/GeometryProcessingInterfaces/ApproximateActors.h
+++ b/Engine/Source/Developer/GeometryProcessingInterfaces/Public/GeometryProcessingInterfaces/ApproximateActors.h
@@ -38,6 +38,20 @@ public:
 		CollisionMesh
 	};
 
+	enum class EGroundPlanePolicy
+	{
+		NoGroundPlane = 0,
+		FixedZHeightGroundPlane = 1
+	};
+
+	enum class EGroundPlaneClippingPolicy
+	{
+		NoClipping = 0,
+		DiscardFullyHiddenFaces = 1,
+		CutFaces = 2,
+		CutFacesAndFill = 3
+	};
+
 	enum class EBaseCappingPolicy
 	{
 		NoBaseCapping = 0,
@@ -58,6 +72,12 @@ public:
 		GeometricTolerance = 2
 	};
 
+	enum class EUVGenerationPolicy : uint8
+	{
+		PreferUVAtlas = 0,
+		PreferXAtlas = 1
+	};
+
 	enum class ETextureSizePolicy : uint8
 	{
 		TextureSize = 0,
@@ -82,6 +102,9 @@ public:
 		bool bAutoThickenThinParts = false;
 		double AutoThickenThicknessMeters = 0.1;
 
+		bool bIgnoreTinyParts = false;
+		double TinyPartMaxDimensionMeters = 0.1;
+
 		EBaseCappingPolicy BaseCappingPolicy = EBaseCappingPolicy::NoBaseCapping;
 		double BaseThicknessOverrideMeters = 0.0;		// use this thickness for Solid. If zero, use AutoThickenThickness if bAutoThickenThinParts=true, otherwise use WorldSpaceApproximationAccuracyMeters
 		double BaseHeightOverrideMeters = 0.0;		// consider this height from MinZ as "base" region. If zero, use 2.0*WorldSpaceApproximationAccuracyMeters
@@ -100,13 +123,23 @@ public:
 		double MorphologyDistanceMeters = 0.1;
 
 		EOcclusionPolicy OcclusionPolicy = EOcclusionPolicy::VisibilityBased;
+		bool bAddDownwardFacesOccluder = true;
+
 		ESimplificationPolicy MeshSimplificationPolicy = ESimplificationPolicy::FixedTriangleCount;
 		int32 FixedTriangleCount = 5000;
 		double SimplificationTargetMetric = 0.1;		// interpretation varies depending on MeshSimplificationPolicy
 
+		EGroundPlanePolicy GroundPlanePolicy = EGroundPlanePolicy::NoGroundPlane;
+		double GroundPlaneZHeight = 0.0;
+		
+		EGroundPlaneClippingPolicy GroundPlaneClippingPolicy = EGroundPlaneClippingPolicy::NoClipping;
+
 		bool bCalculateHardNormals = true;
 		double HardNormalsAngleDeg = 60.0;
 
+		EUVGenerationPolicy UVPolicy = EUVGenerationPolicy::PreferXAtlas;
+		double UVAtlasStretchTarget = 0.1;
+
 		//
 		// Material approximation settings
 		//
diff --git a/Engine/Source/Runtime/Engine/Classes/Engine/MeshMerging.h b/Engine/Source/Runtime/Engine/Classes/Engine/MeshMerging.h
--- a/Engine/Source/Runtime/Engine/Classes/Engine/MeshMerging.h
+++ b/Engine/Source/Runtime/Engine/Classes/Engine/MeshMerging.h
@@ -706,6 +706,23 @@ enum class EMeshApproximationSimplificationPolicy : uint8
 	GeometricTolerance = 2
 };
 
+UENUM()
+enum class EMeshApproximationGroundPlaneClippingPolicy : uint8
+{
+	NoGroundClipping = 0,
+	DiscardWithZPlane = 1,
+	CutWithZPlane = 2,
+	CutAndFillWithZPlane = 3
+};
+
+
+UENUM()
+enum class EMeshApproximationUVGenerationPolicy : uint8
+{
+	PreferUVAtlas = 0,
+	PreferXAtlas = 1
+};
+
 
 USTRUCT(Blueprintable)
 struct FMeshApproximationSettings
@@ -716,6 +733,11 @@ struct FMeshApproximationSettings
 	UPROPERTY(BlueprintReadWrite, EditAnywhere, Category = ShapeSettings)
 	EMeshApproximationType OutputType = EMeshApproximationType::MeshAndMaterials;
 
+
+	//
+	// Mesh Generation Settings
+	//
+
 	/** Approximation Accuracy in Meters, will determine (eg) voxel resolution */
 	UPROPERTY(BlueprintReadWrite, EditAnywhere, Category = ShapeSettings, meta = (DisplayName = "Approximation Accuracy (meters)", ClampMin = "0.001"))
 	float ApproximationAccuracy = 1.0;
@@ -732,16 +754,24 @@ struct FMeshApproximationSettings
 	UPROPERTY(BlueprintReadWrite, EditAnywhere, Category = ShapeSettings, meta = (ClampMin = "0.001", EditCondition = "bAttemptAutoThickening"))
 	float TargetMinThicknessMultiplier = 1.5;
 
+	/** If enabled, tiny parts will be excluded from the mesh merging, which can improve performance */
+	UPROPERTY(BlueprintReadWrite, EditAnywhere, Category = ShapeSettings)
+	bool bIgnoreTinyParts = true;
+
+	/** Multiplier on Approximation Accuracy used to define tiny-part threshold, using maximum bounding-box dimension */
+	UPROPERTY(BlueprintReadWrite, EditAnywhere, Category = ShapeSettings, meta = (ClampMin = "0.001", EditCondition = "bIgnoreTinyParts"))
+	float TinyPartSizeMultiplier = 0.05;
+
+
 	/** Optional methods to attempt to close off the bottom of open meshes */
 	UPROPERTY(BlueprintReadWrite, EditAnywhere, Category = ShapeSettings)
 	EMeshApproximationBaseCappingType BaseCapping = EMeshApproximationBaseCappingType::NoBaseCapping;
 
 
 	/** Winding Threshold controls hole filling at open mesh borders. Smaller value means "more/rounder" filling */
-	UPROPERTY(BlueprintReadWrite, EditAnywhere, Category = ShapeSettings, meta = (ClampMin = "0.01", ClampMax = "0.99"))
+	UPROPERTY(BlueprintReadWrite, EditAnywhere, AdvancedDisplay, Category = ShapeSettings, meta = (ClampMin = "0.01", ClampMax = "0.99"))
 	float WindingThreshold = 0.5;
 
-
 	/** If true, topological expand/contract is used to try to fill small gaps between objects. */
 	UPROPERTY(BlueprintReadWrite, EditAnywhere, Category = ShapeSettings)
 	bool bFillGaps = true;
@@ -751,18 +781,27 @@ struct FMeshApproximationSettings
 	float GapDistance = 0.1;
 
 
+	//
+	// Output Mesh Filtering and Simplification Settings
+	//
+
 	/** Type of hidden geometry removal to apply */
 	UPROPERTY(BlueprintReadWrite, EditAnywhere, Category = SimplifySettings)
 	EOccludedGeometryFilteringPolicy OcclusionMethod = EOccludedGeometryFilteringPolicy::VisibilityBasedFiltering;
 
-	/** Mesh simplification criteria */
+	/** If true, then the OcclusionMethod computation is configured to try to consider downward-facing "bottom" geometry as occluded */
+	UPROPERTY(BlueprintReadWrite, EditAnywhere, Category = SimplifySettings)
+	bool bOccludeFromBottom = true;
+
+	/** Mesh Simplification criteria */
 	UPROPERTY(BlueprintReadWrite, EditAnywhere, Category = SimplifySettings)
 	EMeshApproximationSimplificationPolicy SimplifyMethod = EMeshApproximationSimplificationPolicy::GeometricTolerance;
 
+	/** Target triangle count for Mesh Simplification, for SimplifyMethods that use a Count*/
 	UPROPERTY(BlueprintReadWrite, EditAnywhere, Category = SimplifySettings, meta = (ClampMin = "16", EditCondition = "SimplifyMethod == EMeshApproximationSimplificationPolicy::FixedTriangleCount" ))
 	int32 TargetTriCount = 2000;
 
-	/** Approximate Number of triangles per Square Meter */
+	/** Approximate Number of triangles per Square Meter, for SimplifyMethods that use such a constraint */
 	UPROPERTY(BlueprintReadWrite, EditAnywhere, Category = SimplifySettings, meta = (ClampMin = "0.01", EditCondition = "SimplifyMethod == EMeshApproximationSimplificationPolicy::TrianglesPerArea" ))
 	float TrianglesPerM = 2.0f;
 
@@ -770,6 +809,18 @@ struct FMeshApproximationSettings
 	UPROPERTY(BlueprintReadWrite, EditAnywhere, Category = SimplifySettings, meta = (DisplayName = "Geometric Deviation (meters)", ClampMin = "0.0001", EditCondition = "SimplifyMethod == EMeshApproximationSimplificationPolicy::GeometricTolerance"))
 	float GeometricDeviation = 0.1f;
 
+	/** Configure how the final mesh should be clipped with a ground plane, if desired */
+	UPROPERTY(BlueprintReadWrite, EditAnywhere, Category = SimplifySettings)
+	EMeshApproximationGroundPlaneClippingPolicy GroundClipping = EMeshApproximationGroundPlaneClippingPolicy::NoGroundClipping;
+
+	/** Z-Height for the ground clipping plane, if enabled */
+	UPROPERTY(BlueprintReadWrite, EditAnywhere, Category = SimplifySettings, meta = (EditCondition = "GroundClipping != EGroundPlaneClippingPolicy::NoGroundClipping"))
+	float GroundClippingZHeight = 0.0f;
+
+
+	//
+	// Mesh Normals and Tangents Settings
+	//
 
 	/** If true, normal angle will be used to estimate hard normals */
 	UPROPERTY(BlueprintReadWrite, EditAnywhere, Category = NormalsSettings)
@@ -779,6 +830,17 @@ struct FMeshApproximationSettings
 	float HardNormalAngle = 60.0f;
 
 
+	//
+	// Mesh UV Generation Settings
+	//
+
+	UPROPERTY(BlueprintReadWrite, EditAnywhere, Category = UVSettings)
+	EMeshApproximationUVGenerationPolicy UVGenerationMethod = EMeshApproximationUVGenerationPolicy::PreferXAtlas;
+
+
+	//
+	// Output Static Mesh Settings
+	//
 
 	/** Whether to generate a nanite-enabled mesh */
 	UPROPERTY(BlueprintReadWrite, EditAnywhere, Category = MeshSettings)
@@ -797,6 +859,10 @@ struct FMeshApproximationSettings
 	bool bAllowDistanceField = true;
 
 
+	//
+	// Material Baking Settings
+	//
+
 	/** If Value is > 1, Multisample output baked textures by this amount in each direction (eg 4 == 16x supersampling) */
 	UPROPERTY(BlueprintReadWrite, EditAnywhere, Category = MaterialSettings, meta = (ClampMin = "0", ClampMax = "8", UIMin = "0", UIMax = "4"))
 	int32 MultiSamplingAA = 0;
@@ -816,6 +882,10 @@ struct FMeshApproximationSettings
 	float NearPlaneDist = 1.0f;
 
 
+	//
+	// Debug Output Settings
+	//
+
 
 	/** If true, print out debugging messages */
 	UPROPERTY(BlueprintReadWrite, EditAnywhere, AdvancedDisplay, Category = DebugSettings)
