ModelingComponents: Compact meshes before passing to DynamicMeshComponent ToolTarget, as some Tools implicitly assume the input mesh is compact (and it generally reduces memory/disk size)
#rnx

diff --git a/Engine/Plugins/Experimental/MeshModelingToolset/Source/ModelingComponents/Private/ModelingToolTargetUtil.cpp b/Engine/Plugins/Experimental/MeshModelingToolset/Source/ModelingComponents/Private/ModelingToolTargetUtil.cpp
--- a/Engine/Plugins/Experimental/MeshModelingToolset/Source/ModelingComponents/Private/ModelingToolTargetUtil.cpp
+++ b/Engine/Plugins/Experimental/MeshModelingToolset/Source/ModelingComponents/Private/ModelingToolTargetUtil.cpp
@@ -234,7 +234,10 @@ UE::ToolTarget::EDynamicMeshUpdateResult UE::ToolTarget::CommitDynamicMeshUpdate
 		TUniquePtr<FDynamicMesh3> CurrentMesh = DynamicMesh->ExtractMesh();
 		TSharedPtr<FDynamicMesh3> CurrentMeshShared(CurrentMesh.Release());
 
-		DynamicMesh->EditMesh([&](FDynamicMesh3& EditMesh) { EditMesh = UpdatedMesh; });
+		DynamicMesh->EditMesh([&](FDynamicMesh3& EditMesh) 
+		{ 
+			EditMesh.CompactCopy(UpdatedMesh);
+		});
 
 		TSharedPtr<FDynamicMesh3> NewMeshShared = MakeShared<FDynamicMesh3>();
 		DynamicMesh->ProcessMesh([&](const FDynamicMesh3& ReadMesh) { *NewMeshShared = ReadMesh; });
diff --git a/Engine/Plugins/Experimental/MeshModelingToolset/Source/ModelingComponentsEditorOnly/Private/EditorModelingObjectsCreationAPI.cpp b/Engine/Plugins/Experimental/MeshModelingToolset/Source/ModelingComponentsEditorOnly/Private/EditorModelingObjectsCreationAPI.cpp
--- a/Engine/Plugins/Experimental/MeshModelingToolset/Source/ModelingComponentsEditorOnly/Private/EditorModelingObjectsCreationAPI.cpp
+++ b/Engine/Plugins/Experimental/MeshModelingToolset/Source/ModelingComponentsEditorOnly/Private/EditorModelingObjectsCreationAPI.cpp
@@ -201,7 +201,12 @@ FCreateMeshObjectResult UEditorModelingObjectsCreationAPI::CreateDynamicMeshActo
 
 	if (CreateMeshParams.MeshType == ECreateMeshObjectSourceMeshType::DynamicMesh)
 	{
-		NewComponent->SetMesh(MoveTemp(CreateMeshParams.DynamicMesh.GetValue()));
+		FDynamicMesh3 SetMesh = MoveTemp(CreateMeshParams.DynamicMesh.GetValue());
+		if (SetMesh.IsCompact() == false)
+		{
+			SetMesh.CompactInPlace();
+		}
+		NewComponent->SetMesh(MoveTemp(SetMesh));
 		NewComponent->NotifyMeshUpdated();
 	}
 	else if (CreateMeshParams.MeshType == ECreateMeshObjectSourceMeshType::MeshDescription)
