MeshModelingTools: Allow polygon selection mechanic to be enabled and disabled at will, and expose the topology selector inside it.

#rnx

diff --git a/Engine/Plugins/Experimental/MeshModelingToolset/Source/ModelingComponents/Private/Selection/PolygonSelectionMechanic.cpp b/Engine/Plugins/Experimental/MeshModelingToolset/Source/ModelingComponents/Private/Selection/PolygonSelectionMechanic.cpp
--- a/Engine/Plugins/Experimental/MeshModelingToolset/Source/ModelingComponents/Private/Selection/PolygonSelectionMechanic.cpp
+++ b/Engine/Plugins/Experimental/MeshModelingToolset/Source/ModelingComponents/Private/Selection/PolygonSelectionMechanic.cpp
@@ -25,6 +25,8 @@ void UPolygonSelectionMechanic::Setup(UInteractiveTool* ParentToolIn)
 {
 	UInteractionMechanic::Setup(ParentToolIn);
 
+	TopoSelector = MakeShared<FGroupTopologySelector, ESPMode::ThreadSafe>();
+
 	ClickBehavior = NewObject<USingleClickInputBehavior>();
 	ClickBehavior->Initialize(this);
 	ClickBehavior->SetDefaultPriority(BasePriority);
@@ -49,12 +51,15 @@ void UPolygonSelectionMechanic::Setup(UInteractiveTool* ParentToolIn)
 	{
 		AddToolPropertySource(Properties);
 	}
-	Properties->WatchProperty(Properties->bSelectVertices, [this](bool bSelectVertices) { 
-		MarqueeMechanic->SetIsEnabled(Properties->bEnableMarquee && (Properties->bSelectVertices || Properties->bSelectEdges)); });
+	Properties->WatchProperty(Properties->bSelectVertices, [this](bool bSelectVertices) {
+		UpdateMarqueeEnabled(); 
+	});
 	Properties->WatchProperty(Properties->bSelectEdges, [this](bool bSelectVertices) { 
-		MarqueeMechanic->SetIsEnabled(Properties->bEnableMarquee && (Properties->bSelectVertices || Properties->bSelectEdges)); });
+		UpdateMarqueeEnabled(); 
+	});
 	Properties->WatchProperty(Properties->bEnableMarquee, [this](bool bEnableMarquee) { 
-		MarqueeMechanic->SetIsEnabled(Properties->bEnableMarquee && (Properties->bSelectVertices || Properties->bSelectEdges)); });
+		UpdateMarqueeEnabled(); 
+	});
 
 	// set up visualizers
 	PolyEdgesRenderer.LineColor = FLinearColor::Red;
@@ -106,10 +111,10 @@ void UPolygonSelectionMechanic::Initialize(
 	this->Topology = TopologyIn;
 	this->TargetTransform = FTransform3d(TargetTransformIn);
 
-	TopoSelector.Initialize(Mesh, Topology);
+	TopoSelector->Initialize(Mesh, Topology);
 	this->GetSpatialFunc = GetSpatialSourceFuncIn;
-	TopoSelector.SetSpatialSource(GetSpatialFunc);
-	TopoSelector.PointsWithinToleranceTest = [this](const FVector3d& Position1, const FVector3d& Position2, double TolScale) {
+	TopoSelector->SetSpatialSource(GetSpatialFunc);
+	TopoSelector->PointsWithinToleranceTest = [this](const FVector3d& Position1, const FVector3d& Position2, double TolScale) {
 		if (CameraState.bIsOrthographic)
 		{
 			// We could just always use ToolSceneQueriesUtil::PointSnapQuery. But in ortho viewports, we happen to know
@@ -171,6 +176,12 @@ void UPolygonSelectionMechanic::DisableBehaviors(UInteractiveTool* ParentToolIn)
 	// TODO: Is it worth adding a way to remove the property watchers for marquee?
 }
 
+void UPolygonSelectionMechanic::SetIsEnabled(bool bOn)
+{
+	bIsEnabled = bOn;
+	UpdateMarqueeEnabled();
+}
+
 void UPolygonSelectionMechanic::SetBasePriority(const FInputCapturePriority &Priority)
 {
 	BasePriority = Priority;
@@ -226,13 +237,13 @@ void UPolygonSelectionMechanic::Render(IToolsContextRenderAPI* RenderAPI)
 	{
 		SelectionRenderer.BeginFrame(RenderAPI, RenderCameraState);
 		SelectionRenderer.SetTransform(Transform);
-		TopoSelector.DrawSelection(PersistentSelection, &SelectionRenderer, &RenderCameraState);
+		TopoSelector->DrawSelection(PersistentSelection, &SelectionRenderer, &RenderCameraState);
 		SelectionRenderer.EndFrame();
 	}
 
 	HilightRenderer.BeginFrame(RenderAPI, RenderCameraState);
 	HilightRenderer.SetTransform(Transform);
-	TopoSelector.DrawSelection(HilightSelection, &HilightRenderer, &RenderCameraState);
+	TopoSelector->DrawSelection(HilightSelection, &HilightRenderer, &RenderCameraState);
 	HilightRenderer.EndFrame();
 }
 
@@ -257,7 +268,7 @@ void UPolygonSelectionMechanic::ClearHighlight()
 void UPolygonSelectionMechanic::NotifyMeshChanged(bool bTopologyModified)
 {
 	ClearHighlight();
-	TopoSelector.Invalidate(true, bTopologyModified);
+	TopoSelector->Invalidate(true, bTopologyModified);
 	if (bTopologyModified)
 	{
 		PersistentSelection.Clear();
@@ -282,7 +293,7 @@ bool UPolygonSelectionMechanic::TopologyHitTest(const FRay& WorldRay, FHitResult
 	FVector3d LocalPosition, LocalNormal;
 	int32 EdgeSegmentId; // Only used if hit is an edge
 	FGroupTopologySelector::FSelectionSettings TopoSelectorSettings = GetTopoSelectorSettings(bUseOrthoSettings);
-	if (TopoSelector.FindSelectedElement(TopoSelectorSettings, LocalRay, OutSelection, LocalPosition, LocalNormal, &EdgeSegmentId) == false)
+	if (TopoSelector->FindSelectedElement(TopoSelectorSettings, LocalRay, OutSelection, LocalPosition, LocalNormal, &EdgeSegmentId) == false)
 	{
 		return false;
 	}
@@ -366,15 +377,15 @@ bool UPolygonSelectionMechanic::UpdateHighlight(const FRay& WorldRay)
 	HilightSelection.Clear();
 	FVector3d LocalPosition, LocalNormal;
 	FGroupTopologySelector::FSelectionSettings TopoSelectorSettings = GetTopoSelectorSettings(CameraState.bIsOrthographic);
-	bool bHit = TopoSelector.FindSelectedElement(TopoSelectorSettings, LocalRay, HilightSelection, LocalPosition, LocalNormal);
+	bool bHit = TopoSelector->FindSelectedElement(TopoSelectorSettings, LocalRay, HilightSelection, LocalPosition, LocalNormal);
 
 	if (HilightSelection.SelectedEdgeIDs.Num() > 0 && Properties->bSelectEdgeRings && ShouldSelectEdgeRingsFunc())
 	{
-		TopoSelector.ExpandSelectionByEdgeRings(HilightSelection);
+		TopoSelector->ExpandSelectionByEdgeRings(HilightSelection);
 	}
 	if (HilightSelection.SelectedEdgeIDs.Num() > 0 && Properties->bSelectEdgeLoops && ShouldSelectEdgeLoopsFunc())
 	{
-		TopoSelector.ExpandSelectionByEdgeLoops(HilightSelection);
+		TopoSelector->ExpandSelectionByEdgeLoops(HilightSelection);
 	}
 
 	// Don't hover highlight a selection that we already selected, because people didn't like that
@@ -446,18 +457,18 @@ bool UPolygonSelectionMechanic::UpdateSelection(const FRay& WorldRay, FVector3d&
 	FVector3d LocalPosition, LocalNormal;
 	FGroupTopologySelection Selection;
 	FGroupTopologySelector::FSelectionSettings TopoSelectorSettings = GetTopoSelectorSettings(CameraState.bIsOrthographic);
-	if (TopoSelector.FindSelectedElement(TopoSelectorSettings, LocalRay, Selection, LocalPosition, LocalNormal))
+	if (TopoSelector->FindSelectedElement(TopoSelectorSettings, LocalRay, Selection, LocalPosition, LocalNormal))
 	{
 		LocalHitPositionOut = LocalPosition;
 		LocalHitNormalOut = LocalNormal;
 
 		if (Selection.SelectedEdgeIDs.Num() > 0 && Properties->bSelectEdgeRings && ShouldSelectEdgeRingsFunc())
 		{
-			TopoSelector.ExpandSelectionByEdgeRings(Selection);
+			TopoSelector->ExpandSelectionByEdgeRings(Selection);
 		}
 		if (Selection.SelectedEdgeIDs.Num() > 0 && Properties->bSelectEdgeLoops && ShouldSelectEdgeLoopsFunc())
 		{
-			TopoSelector.ExpandSelectionByEdgeLoops(Selection);
+			TopoSelector->ExpandSelectionByEdgeLoops(Selection);
 		}
 
 		if (ShouldAddToSelectionFunc())
@@ -518,6 +529,11 @@ void UPolygonSelectionMechanic::ClearSelection()
 
 FInputRayHit UPolygonSelectionMechanic::IsHitByClick(const FInputDeviceRay& ClickPos)
 {
+	if (!bIsEnabled)
+	{
+		return FInputRayHit(); // bHit is false
+	}
+
 	FHitResult OutHit;
 	FGroupTopologySelection Selection;
 	if (TopologyHitTest(ClickPos.WorldRay, OutHit, Selection, CameraState.bIsOrthographic))
@@ -548,7 +564,7 @@ void UPolygonSelectionMechanic::OnClicked(const FInputDeviceRay& ClickPos)
 FInputRayHit UPolygonSelectionMechanic::BeginHoverSequenceHitTest(const FInputDeviceRay& PressPos)
 {
 	FHitResult OutHit;
-	if (TopologyHitTest(PressPos.WorldRay, OutHit))
+	if (bIsEnabled && TopologyHitTest(PressPos.WorldRay, OutHit))
 	{
 		return FInputRayHit(OutHit.Distance);
 	}
@@ -599,7 +615,7 @@ void UPolygonSelectionMechanic::OnDragRectangleChanged(const FCameraRectangle& C
 	FGroupTopologySelector::FSelectionSettings TopoSelectorSettings = GetTopoSelectorSettings(false);
 	TopoSelectorSettings.bIgnoreOcclusion = Properties->bMarqueeIgnoreOcclusion; // uses a separate setting for marquee
 	
-	TopoSelector.FindSelectedElement(TopoSelectorSettings, CurrentRectangle, TargetTransform,
+	TopoSelector->FindSelectedElement(TopoSelectorSettings, CurrentRectangle, TargetTransform,
 		RectangleSelection);
 
 	if (ShouldAddToSelectionFunc())
@@ -638,6 +654,14 @@ void UPolygonSelectionMechanic::OnDragRectangleFinished()
 	ParentTool->GetToolManager()->EndUndoTransaction();
 }
 
+void UPolygonSelectionMechanic::UpdateMarqueeEnabled()
+{
+	MarqueeMechanic->SetIsEnabled(
+		bIsEnabled 
+		&& Properties->bEnableMarquee 
+		&& (Properties->bSelectVertices || Properties->bSelectEdges));
+}
+
 
 void UPolygonSelectionMechanic::BeginChange()
 {
diff --git a/Engine/Plugins/Experimental/MeshModelingToolset/Source/ModelingComponents/Public/Selection/PolygonSelectionMechanic.h b/Engine/Plugins/Experimental/MeshModelingToolset/Source/ModelingComponents/Public/Selection/PolygonSelectionMechanic.h
--- a/Engine/Plugins/Experimental/MeshModelingToolset/Source/ModelingComponents/Public/Selection/PolygonSelectionMechanic.h
+++ b/Engine/Plugins/Experimental/MeshModelingToolset/Source/ModelingComponents/Public/Selection/PolygonSelectionMechanic.h
@@ -128,6 +128,11 @@ public:
 	 */
 	void DisableBehaviors(UInteractiveTool* ParentToolIn);
 
+	/**
+	 * Enable/disable the mechanic without permanently removing behaviors or shutting it down.
+	 */
+	void SetIsEnabled(bool bOn);
+
 	/**
 	 * Sets the base priority so that tools can make sure that their own behaviors are higher
 	 * priority. The mechanic will not use any priority value higher than this, but it may use
@@ -190,6 +195,8 @@ public:
 	bool TopologyHitTest(const FRay& WorldRay, FHitResult& OutHit, FGroupTopologySelection& OutSelection, bool bUseOrthoSettings = false);
 	bool TopologyHitTest(const FRay& WorldRay, FHitResult& OutHit, bool bUseOrthoSettings = false);
 
+	TSharedPtr<FGroupTopologySelector, ESPMode::ThreadSafe> GetTopologySelector() { return TopoSelector; }
+
 	//
 	// Hover API
 	//
@@ -306,6 +313,8 @@ protected:
 	virtual void OnDragRectangleChanged(const FCameraRectangle& CurrentRectangle);
 	virtual void OnDragRectangleFinished();
 
+	virtual void UpdateMarqueeEnabled();
+
 public:
 	/** 
 	 * OnSelectionChanged is broadcast whenever the selection is modified (including by FChanges, which
@@ -320,6 +329,8 @@ public:
 	TObjectPtr<UPolygonSelectionMechanicProperties> Properties;
 
 protected:
+	bool bIsEnabled = true;
+
 	const FDynamicMesh3* Mesh;
 	const FGroupTopology* Topology;
 	TFunction<FDynamicMeshAABBTree3*()> GetSpatialFunc;
@@ -350,7 +361,7 @@ protected:
 
 	FTransform3d TargetTransform;
 
-	FGroupTopologySelector TopoSelector;
+	TSharedPtr<FGroupTopologySelector, ESPMode::ThreadSafe> TopoSelector;
 
 	/**
 	 * Get the topology selector settings to use given the current selection settings.
