MeshModelingTools: Fix busy material sometimes showing up when compute is cancelled. Also, forward secondary buffer controls through PreviewMesh.

#rnx

diff --git a/Engine/Plugins/Experimental/MeshModelingToolset/Source/ModelingComponents/Private/MeshOpPreviewHelpers.cpp b/Engine/Plugins/Experimental/MeshModelingToolset/Source/ModelingComponents/Private/MeshOpPreviewHelpers.cpp
--- a/Engine/Plugins/Experimental/MeshModelingToolset/Source/ModelingComponents/Private/MeshOpPreviewHelpers.cpp
+++ b/Engine/Plugins/Experimental/MeshModelingToolset/Source/ModelingComponents/Private/MeshOpPreviewHelpers.cpp
@@ -108,10 +108,10 @@ void UMeshOpPreviewWithBackgroundCompute::Tick(float DeltaTime)
 
 void UMeshOpPreviewWithBackgroundCompute::UpdateResults()
 {
-	EBackgroundComputeTaskStatus Status = BackgroundCompute ? BackgroundCompute->CheckStatus()
+	LastComputeStatus = BackgroundCompute ? BackgroundCompute->CheckStatus()
 		: EBackgroundComputeTaskStatus::NotComputing;
 
-	if (Status == EBackgroundComputeTaskStatus::NewResultAvailable)
+	if (LastComputeStatus == EBackgroundComputeTaskStatus::NewResultAvailable)
 	{
 		TUniquePtr<FDynamicMeshOperator> MeshOp = BackgroundCompute->ExtractResult();
 		OnOpCompleted.Broadcast(MeshOp.Get());
@@ -203,6 +203,7 @@ void UMeshOpPreviewWithBackgroundCompute::SetIsMeshTopologyConstant(bool bOn, EM
 
 bool UMeshOpPreviewWithBackgroundCompute::IsUsingWorkingMaterial()
 {
-	bool bIsLongDelay = BackgroundCompute && BackgroundCompute->GetElapsedComputeTime() > SecondsBeforeWorkingMaterial;
-	return !(bResultValid || WorkingMaterial == nullptr || bIsLongDelay == false);
+	return WorkingMaterial && BackgroundCompute 
+		&& LastComputeStatus == EBackgroundComputeTaskStatus::InProgress
+		&& BackgroundCompute->GetElapsedComputeTime() > SecondsBeforeWorkingMaterial;
 }
diff --git a/Engine/Plugins/Experimental/MeshModelingToolset/Source/ModelingComponents/Private/PreviewMesh.cpp b/Engine/Plugins/Experimental/MeshModelingToolset/Source/ModelingComponents/Private/PreviewMesh.cpp
--- a/Engine/Plugins/Experimental/MeshModelingToolset/Source/ModelingComponents/Private/PreviewMesh.cpp
+++ b/Engine/Plugins/Experimental/MeshModelingToolset/Source/ModelingComponents/Private/PreviewMesh.cpp
@@ -172,6 +172,17 @@ void UPreviewMesh::DisableSecondaryTriangleBuffers()
 	DynamicMeshComponent->DisableSecondaryTriangleBuffers();
 }
 
+void UPreviewMesh::SetSecondaryBuffersVisibility(bool bSecondaryVisibility)
+{
+	check(DynamicMeshComponent);
+	DynamicMeshComponent->SetSecondaryBuffersVisibility(bSecondaryVisibility);
+}
+
+void UPreviewMesh::FastNotifySecondaryTrianglesChanged()
+{
+	check(DynamicMeshComponent);
+	DynamicMeshComponent->FastNotifySecondaryTrianglesChanged();
+}
 
 void UPreviewMesh::SetTangentsMode(EDynamicMeshComponentTangentsMode TangentsType)
 {
diff --git a/Engine/Plugins/Experimental/MeshModelingToolset/Source/ModelingComponents/Public/MeshOpPreviewHelpers.h b/Engine/Plugins/Experimental/MeshModelingToolset/Source/ModelingComponents/Public/MeshOpPreviewHelpers.h
--- a/Engine/Plugins/Experimental/MeshModelingToolset/Source/ModelingComponents/Public/MeshOpPreviewHelpers.h
+++ b/Engine/Plugins/Experimental/MeshModelingToolset/Source/ModelingComponents/Public/MeshOpPreviewHelpers.h
@@ -220,6 +220,11 @@ protected:
 	// state flag, if true then we have valid result
 	bool bResultValid = false;
 
+	// Stored status of last compute, mainly so that we know when we should
+	// show the "busy" material.
+	UE::Geometry::EBackgroundComputeTaskStatus LastComputeStatus = 
+		UE::Geometry::EBackgroundComputeTaskStatus::NotComputing;
+
 	bool bVisible = true;
 
 	// Used for partial/fast updates of the preview mesh render proxy.
diff --git a/Engine/Plugins/Experimental/MeshModelingToolset/Source/ModelingComponents/Public/PreviewMesh.h b/Engine/Plugins/Experimental/MeshModelingToolset/Source/ModelingComponents/Public/PreviewMesh.h
--- a/Engine/Plugins/Experimental/MeshModelingToolset/Source/ModelingComponents/Public/PreviewMesh.h
+++ b/Engine/Plugins/Experimental/MeshModelingToolset/Source/ModelingComponents/Public/PreviewMesh.h
@@ -170,7 +170,16 @@ public:
 	 */
 	void DisableSecondaryTriangleBuffers();
 
+	/**
+	 * Show/Hide the secondary triangle buffers.
+	 */
+	void SetSecondaryBuffersVisibility(bool bSecondaryVisibility);
 
+	/** 
+	 * Call this after updating the secondary triangle sorting.
+	 * This function will update the existing buffers if possible, without rebuilding entire RenderProxy.
+	 */
+	void FastNotifySecondaryTrianglesChanged();
 
 	/**
 	 * Set the tangents mode for the underlying component, if available. 
