ModelingTools: Change default BakeTexture setting to None.

#rnx

diff --git a/Engine/Plugins/Experimental/MeshModelingToolset/Source/MeshModelingTools/Private/BakeMeshAttributeMapsTool.cpp b/Engine/Plugins/Experimental/MeshModelingToolset/Source/MeshModelingTools/Private/BakeMeshAttributeMapsTool.cpp
--- a/Engine/Plugins/Experimental/MeshModelingToolset/Source/MeshModelingTools/Private/BakeMeshAttributeMapsTool.cpp
+++ b/Engine/Plugins/Experimental/MeshModelingToolset/Source/MeshModelingTools/Private/BakeMeshAttributeMapsTool.cpp
@@ -283,6 +283,8 @@ void UBakeMeshAttributeMapsTool::Setup()
 {
 	UInteractiveTool::Setup();
 
+	InitializeEmptyMaps();
+
 	// create dynamic mesh component to use for live preview
 	// TODO: convert to UPreviewMesh
 	AActor* ParentActor = UE::ToolTarget::GetTargetActor(Targets[0]);
@@ -315,6 +317,9 @@ void UBakeMeshAttributeMapsTool::Setup()
 	if (Material != nullptr)
 	{
 		PreviewMaterial = UMaterialInstanceDynamic::Create(Material, GetToolManager());
+		PreviewMaterial->SetTextureParameterValue(TEXT("NormalMap"), EmptyNormalMap);
+		PreviewMaterial->SetTextureParameterValue(TEXT("OcclusionMap"), EmptyColorMapWhite);
+		PreviewMaterial->SetTextureParameterValue(TEXT("ColorMap"), EmptyColorMapWhite);
 		DynamicMeshComponent->SetOverrideRenderMaterial(PreviewMaterial);
 	}
 	UMaterial* BentNormalMaterial = LoadObject<UMaterial>(nullptr, TEXT("/MeshModelingToolset/Materials/BakeBentNormalPreviewMaterial"));
@@ -414,7 +419,6 @@ void UBakeMeshAttributeMapsTool::Setup()
 	VisualizationProps->RestoreProperties(this);
 	AddToolPropertySource(VisualizationProps);
 
-	InitializeEmptyMaps();
 	UpdateOnModeChange();
 
 	bInputsDirty = true;
@@ -431,7 +435,7 @@ void UBakeMeshAttributeMapsTool::Setup()
 
 bool UBakeMeshAttributeMapsTool::CanAccept() const
 {
-	bool bCanAccept = Compute ? Compute->HaveValidResult() : false;
+	bool bCanAccept = Compute ? Compute->HaveValidResult() && Settings->MapTypes != (int) EBakeMapType::None : false;
 	if (bCanAccept)
 	{
 		// Allow Accept if all non-None types have valid results.
diff --git a/Engine/Plugins/Experimental/MeshModelingToolset/Source/MeshModelingTools/Public/BakeMeshAttributeMapsTool.h b/Engine/Plugins/Experimental/MeshModelingToolset/Source/MeshModelingTools/Public/BakeMeshAttributeMapsTool.h
--- a/Engine/Plugins/Experimental/MeshModelingToolset/Source/MeshModelingTools/Public/BakeMeshAttributeMapsTool.h
+++ b/Engine/Plugins/Experimental/MeshModelingToolset/Source/MeshModelingTools/Public/BakeMeshAttributeMapsTool.h
@@ -107,7 +107,7 @@ class MESHMODELINGTOOLS_API UBakeMeshAttributeMapsToolProperties : public UInter
 public:
 	/** The map types to generate */
 	UPROPERTY(EditAnywhere, Category = MapSettings, meta=(Bitmask, BitmaskEnum=EBakeMapType))
-	int32 MapTypes = (int32) EBakeMapType::TangentSpaceNormalMap;
+	int32 MapTypes = (int32) EBakeMapType::None;
 
 	/** The map type index to preview */
 	UPROPERTY(EditAnywhere, Category = MapSettings, meta=(ArrayClamp="Result"))
