MeshModelingTools: Refactor PolyEdit to use tool activity objects for multi-interaction activities rather than being monolithic. Merge group edge insertion tool and edge loop insertion tool into PolyEdit. Merge inset and outset, and offset and extrude. Remove use of MultiTransformer, and do some other cleanup in PolyEdit.

#rnx

diff --git a/Engine/Plugins/Experimental/MeshModelingToolset/Source/MeshModelingTools/Private/EdgeLoopInsertionTool.cpp b/Engine/Plugins/Experimental/MeshModelingToolset/Source/MeshModelingTools/Private/EdgeLoopInsertionTool.cpp
deleted file mode 100644
diff --git a/Engine/Plugins/Experimental/MeshModelingToolset/Source/MeshModelingTools/Private/EditMeshPolygonsTool.cpp b/Engine/Plugins/Experimental/MeshModelingToolset/Source/MeshModelingTools/Private/EditMeshPolygonsTool.cpp
--- a/Engine/Plugins/Experimental/MeshModelingToolset/Source/MeshModelingTools/Private/EditMeshPolygonsTool.cpp
+++ b/Engine/Plugins/Experimental/MeshModelingToolset/Source/MeshModelingTools/Private/EditMeshPolygonsTool.cpp
@@ -1,61 +1,95 @@
 // Copyright Epic Games, Inc. All Rights Reserved.
 
 #include "EditMeshPolygonsTool.h"
-#include "InteractiveToolManager.h"
-#include "ToolBuilderUtil.h"
 
+#include "Algo/ForEach.h"
+#include "BaseBehaviors/SingleClickBehavior.h"
+#include "BaseGizmos/TransformGizmo.h"
+#include "BaseGizmos/TransformGizmoUtil.h"
+#include "BaseGizmos/TransformProxy.h"
 #include "CompGeom/PolygonTriangulation.h"
-#include "SegmentTypes.h"
+#include "Components/BrushComponent.h"
+#include "ContextObjectStore.h"
 #include "DynamicMesh/DynamicMeshAttributeSet.h"
-#include "Mechanics/DragAlignmentMechanic.h"
+#include "DynamicMesh/MeshIndexUtil.h" // TriangleToVertexIDs
 #include "DynamicMesh/MeshNormals.h"
-#include "ToolSceneQueriesUtil.h"
-#include "Intersection/IntersectionUtil.h"
-#include "Transforms/MultiTransformer.h"
-#include "TransformTypes.h"
-#include "BaseBehaviors/SingleClickBehavior.h"
-#include "Util/ColorConstants.h"
-#include "Util/CompactMaps.h"
-#include "ToolSetupUtil.h"
-#include "Operations/MeshPlaneCut.h"
-#include "Selection/GroupTopologyStorableSelection.h"
-#include "Selections/MeshEdgeSelection.h"
-#include "Selections/MeshFaceSelection.h"
-#include "Selections/MeshConnectedComponents.h"
-#include "FaceGroupUtil.h"
+#include "DynamicMesh/MeshTransforms.h"
 #include "DynamicMeshEditor.h"
-#include "DynamicMesh/DynamicMeshChangeTracker.h"
-#include "Changes/MeshChange.h"
-#include "DynamicMesh/MeshIndexUtil.h"
+#include "FaceGroupUtil.h"
+#include "GroupTopology.h"
+#include "InteractiveGizmoManager.h"
+#include "InteractiveToolManager.h"
+#include "Mechanics/DragAlignmentMechanic.h"
+#include "MeshBoundaryLoops.h"
+#include "MeshOpPreviewHelpers.h" // UMeshOpPreviewWithBackgroundCompute
 #include "MeshRegionBoundaryLoops.h"
-#include "ModelingToolTargetUtil.h"
-
-#include "Operations/OffsetMeshRegion.h"
-#include "Operations/InsetMeshRegion.h"
+#include "ModelingToolTargetUtil.h" // UE::ToolTarget:: functions
 #include "Operations/SimpleHoleFiller.h"
-#include "DynamicMesh/MeshTransforms.h"
-
-#include "Algo/ForEach.h"
-#include "Async/ParallelFor.h"
-#include "Containers/BitArray.h"
-#include "Materials/MaterialInstanceDynamic.h"
-#include "Components/BrushComponent.h"
+#include "Selection/GroupTopologyStorableSelection.h"
+#include "Selection/PolygonSelectionMechanic.h"
+#include "Selections/MeshConnectedComponents.h"
+#include "TargetInterfaces/MaterialProvider.h"
+#include "TargetInterfaces/DynamicMeshCommitter.h"
+#include "TargetInterfaces/DynamicMeshProvider.h"
+#include "TargetInterfaces/PrimitiveComponentBackedTarget.h"
+#include "ToolActivities/PolyEditActivityContext.h"
+#include "ToolActivities/PolyEditExtrudeActivity.h"
+#include "ToolActivities/PolyEditInsertEdgeActivity.h"
+#include "ToolActivities/PolyEditInsertEdgeLoopActivity.h"
+#include "ToolActivities/PolyEditInsetOutsetActivity.h"
+#include "ToolActivities/PolyEditCutFacesActivity.h"
+#include "ToolActivities/PolyEditPlanarProjectionUVActivity.h"
+#include "ToolContextInterfaces.h" // FToolBuilderState
+#include "ToolSetupUtil.h"
+#include "ToolTargetManager.h"
+#include "TransformTypes.h"
+#include "Util/CompactMaps.h"
 
 #include "ExplicitUseGeometryMathTypes.h"		// using UE::Geometry::(math types)
 using namespace UE::Geometry;
 
 #define LOCTEXT_NAMESPACE "UEditMeshPolygonsTool"
 
+namespace EditMeshPolygonsToolLocals
+{
+	FText PolyEditDefaultMessage = LOCTEXT("OnStartEditMeshPolygonsTool_TriangleMode", "Select triangles to edit mesh. Use middle mouse on gizmo to "
+		"reposition it. Hold Ctrl while translating or (in local mode) rotating to align to scene. Shift and Ctrl "
+		"change marquee select behavior. Q toggles Gizmo Orientation Lock.");
+
+	FText TriEditDefaultMessage = LOCTEXT("OnStartEditMeshPolygonsTool", "Select PolyGroups to edit mesh. Use middle mouse on gizmo to reposition it. "
+		"Hold Ctrl while translating or (in local mode) rotating to align to scene. Shift and Ctrl change marquee select "
+		"behavior. Q toggles Gizmo Orientation Lock.");
+}
+
 /*
  * ToolBuilder
  */
-UMeshSurfacePointTool* UEditMeshPolygonsToolBuilder::CreateNewTool(const FToolBuilderState& SceneState) const
+
+const FToolTargetTypeRequirements& UEditMeshPolygonsToolBuilder::GetTargetRequirements() const
+{
+	static FToolTargetTypeRequirements TypeRequirements({
+		UMaterialProvider::StaticClass(),
+		UDynamicMeshCommitter::StaticClass(),
+		UDynamicMeshProvider::StaticClass(),
+		UPrimitiveComponentBackedTarget::StaticClass()
+		});
+	return TypeRequirements;
+}
+
+bool UEditMeshPolygonsToolBuilder::CanBuildTool(const FToolBuilderState& SceneState) const
+{
+	return SceneState.TargetManager->CountSelectedAndTargetable(SceneState, GetTargetRequirements()) == 1;
+}
+
+UInteractiveTool* UEditMeshPolygonsToolBuilder::BuildTool(const FToolBuilderState& SceneState) const
 {
 	UEditMeshPolygonsTool* EditPolygonsTool = NewObject<UEditMeshPolygonsTool>(SceneState.ToolManager);
 	if (bTriangleMode)
 	{
 		EditPolygonsTool->EnableTriangleMode();
 	}
+	EditPolygonsTool->SetWorld(SceneState.World);
+	EditPolygonsTool->SetTarget(SceneState.TargetManager->BuildFirstSelectedTargetable(SceneState, GetTargetRequirements()));
 
 	// This passes in nullptr if the stored selection is not the right type. The tool
 	// will figure out whether the selection is still relevant.
@@ -65,28 +99,23 @@ UMeshSurfacePointTool* UEditMeshPolygonsToolBuilder::CreateNewTool(const FToolBu
 }
 
 
-
-
-
-void UEditMeshPolygonsActionModeToolBuilder::InitializeNewTool(UMeshSurfacePointTool* ToolIn, const FToolBuilderState& SceneState) const
+UInteractiveTool* UEditMeshPolygonsActionModeToolBuilder::BuildTool(const FToolBuilderState& SceneState) const
 {
-	UEditMeshPolygonsToolBuilder::InitializeNewTool(ToolIn, SceneState);
-
-	UEditMeshPolygonsTool* Tool = Cast<UEditMeshPolygonsTool>(ToolIn);
+	UEditMeshPolygonsTool* Tool = Cast<UEditMeshPolygonsTool>(UEditMeshPolygonsToolBuilder::BuildTool(SceneState));
 	EEditMeshPolygonsToolActions UseAction = StartupAction;
 	Tool->PostSetupFunction = [UseAction](UEditMeshPolygonsTool* PolyTool)
 	{
 		PolyTool->SetToSelectionModeInterface();
 		PolyTool->RequestAction(UseAction);
 	};
+
+	return Tool;
 }
 
 
-void UEditMeshPolygonsSelectionModeToolBuilder::InitializeNewTool(UMeshSurfacePointTool* ToolIn, const FToolBuilderState& SceneState) const
+UInteractiveTool* UEditMeshPolygonsSelectionModeToolBuilder::BuildTool(const FToolBuilderState& SceneState) const
 {
-	UEditMeshPolygonsToolBuilder::InitializeNewTool(ToolIn, SceneState);
-
-	UEditMeshPolygonsTool* Tool = Cast<UEditMeshPolygonsTool>(ToolIn);
+	UEditMeshPolygonsTool* Tool = Cast<UEditMeshPolygonsTool>(UEditMeshPolygonsToolBuilder::BuildTool(SceneState));
 	EEditMeshPolygonsToolSelectionMode UseMode = SelectionMode;
 	Tool->PostSetupFunction = [UseMode](UEditMeshPolygonsTool* PolyTool)
 	{
@@ -122,6 +151,8 @@ void UEditMeshPolygonsSelectionModeToolBuilder::InitializeNewTool(UMeshSurfacePo
 			break;
 		}
 	};
+
+	return Tool;
 }
 
 void UEditMeshPolygonsTool::SetToSelectionModeInterface()
@@ -152,74 +183,119 @@ UEditMeshPolygonsTool::UEditMeshPolygonsTool()
 
 void UEditMeshPolygonsTool::EnableTriangleMode()
 {
-	check(DynamicMeshComponent == nullptr);		// must not have been initialized!
+	check(Preview == nullptr);		// must not have been initialized!
 	bTriangleMode = true;
 }
 
 void UEditMeshPolygonsTool::Setup()
 {
-	UMeshSurfacePointTool::Setup();
+	using namespace EditMeshPolygonsToolLocals;
+
+	// Start by adding the actions, because we want them at the top.
+	if (bTriangleMode)
+	{
+		EditActions_Triangles = NewObject<UEditMeshPolygonsToolActions_Triangles>();
+		EditActions_Triangles->Initialize(this);
+		AddToolPropertySource(EditActions_Triangles);
+
+		EditEdgeActions_Triangles = NewObject<UEditMeshPolygonsToolEdgeActions_Triangles>();
+		EditEdgeActions_Triangles->Initialize(this);
+		AddToolPropertySource(EditEdgeActions_Triangles);
+
+		SetToolDisplayName(LOCTEXT("EditMeshTrianglesToolName", "Triangle Edit"));
+		DefaultMessage = PolyEditDefaultMessage;
+	}
+	else
+	{
+		EditActions = NewObject<UEditMeshPolygonsToolActions>();
+		EditActions->Initialize(this);
+		AddToolPropertySource(EditActions);
+
+		EditEdgeActions = NewObject<UEditMeshPolygonsToolEdgeActions>();
+		EditEdgeActions->Initialize(this);
+		AddToolPropertySource(EditEdgeActions);
+
+		EditUVActions = NewObject<UEditMeshPolygonsToolUVActions>();
+		EditUVActions->Initialize(this);
+		AddToolPropertySource(EditUVActions);
+
+		DefaultMessage = TriEditDefaultMessage;
+		
+	}
+
+	GetToolManager()->DisplayMessage(DefaultMessage,
+		EToolMessageLevel::UserNotification);
+
+	// We add an empty line for the error message so that things don't jump when we use it.
+	GetToolManager()->DisplayMessage(FText(), EToolMessageLevel::UserWarning);
+
+	CancelAction = NewObject<UEditMeshPolygonsToolCancelAction>();
+	CancelAction->Initialize(this);
+	AddToolPropertySource(CancelAction);
+	SetToolPropertySourceEnabled(CancelAction, false);
+
+	// Initialize the common properties but don't add them yet, because we want them to be under the activity-specific ones.
+	CommonProps = NewObject<UPolyEditCommonProperties>(this);
+	CommonProps->RestoreProperties(this);
+	CommonProps->WatchProperty(CommonProps->LocalFrameMode,
+		[this](ELocalFrameMode) { UpdateGizmoFrame(); });
+	CommonProps->WatchProperty(CommonProps->bLockRotation,
+		[this](bool)
+		{
+			LockedTransfomerFrame = LastTransformerFrame;
+		});
+	// We are going to SilentUpdate here because otherwise the Watches above will immediately fire
+	// and cause UpdateGizmoFrame() to be called emitting a spurious Transform change. 
+	CommonProps->SilentUpdateWatched();
+
+	CurrentMesh = Cast<IDynamicMeshProvider>(Target)->GetDynamicMesh();
 
-	// register click behavior
-	USingleClickInputBehavior* ClickBehavior = NewObject<USingleClickInputBehavior>();
-	ClickBehavior->Initialize(this);
-	AddInputBehavior(ClickBehavior);
+	// TODO: Do we need this?
+	FMeshNormals::QuickComputeVertexNormals(*CurrentMesh);
 
-	// create dynamic mesh component to use for live preview
-	DynamicMeshComponent = NewObject<UDynamicMeshComponent>(UE::ToolTarget::GetTargetActor(Target), "DynamicMesh");
-	DynamicMeshComponent->SetupAttachment(UE::ToolTarget::GetTargetActor(Target)->GetRootComponent());
-	DynamicMeshComponent->RegisterComponent();
+	// Create the preview object
+	Preview = NewObject<UMeshOpPreviewWithBackgroundCompute>();
+	Preview->Setup(TargetWorld);
 	WorldTransform = UE::ToolTarget::GetLocalToWorldTransform(Target);
-	DynamicMeshComponent->SetWorldTransform((FTransform)WorldTransform);
+	Preview->PreviewMesh->SetTransform((FTransform)WorldTransform);
+
+	// We'll use the spatial inside preview mesh mainly for the convenience of having it update automatically.
+	Preview->PreviewMesh->bBuildSpatialDataStructure = true;
 
 	// set materials
 	FComponentMaterialSet MaterialSet = UE::ToolTarget::GetMaterialSet(Target);
-	for (int k = 0; k < MaterialSet.Materials.Num(); ++k)
-	{
-		DynamicMeshComponent->SetMaterial(k, MaterialSet.Materials[k]);
-	}
+	Preview->ConfigureMaterials(MaterialSet.Materials, ToolSetupUtil::GetDefaultWorkingMaterial(GetToolManager()));
 
 	// configure secondary render material
 	UMaterialInterface* SelectionMaterial = ToolSetupUtil::GetSelectionMaterial(FLinearColor::Yellow, GetToolManager());
 	if (SelectionMaterial != nullptr)
 	{
-		DynamicMeshComponent->SetSecondaryRenderMaterial(SelectionMaterial);
+		// Note that you have to do it this way rather than reaching into the PreviewMesh because the background compute
+		// mesh has to be able to swap in/out a working material and restore the primary/secondary ones.
+		Preview->SecondaryMaterial = SelectionMaterial;
 	}
 
-	// enable secondary triangle buffers
-	DynamicMeshComponent->EnableSecondaryTriangleBuffers(
+	Preview->PreviewMesh->EnableSecondaryTriangleBuffers(
 		[this](const FDynamicMesh3* Mesh, int32 TriangleID)
 	{
 		return SelectionMechanic->GetActiveSelection().IsSelectedTriangle(Mesh, Topology.Get(), TriangleID);
 	});
 
+	Preview->PreviewMesh->SetTangentsMode(EDynamicMeshComponentTangentsMode::AutoCalculated);
+	Preview->PreviewMesh->UpdatePreview(CurrentMesh.Get());
+	Preview->PreviewMesh->EnableWireframe(CommonProps->bShowWireframe);
+	Preview->SetVisibility(true);
 
-	// dynamic mesh configuration settings
-	DynamicMeshComponent->SetTangentsType(EDynamicMeshComponentTangentsMode::AutoCalculated);
-	DynamicMeshComponent->SetMesh(UE::ToolTarget::GetDynamicMeshCopy(Target));
-	FMeshNormals::QuickComputeVertexNormals(*DynamicMeshComponent->GetMesh());
-	OnDynamicMeshComponentChangedHandle = DynamicMeshComponent->OnMeshChanged.Add(
-		FSimpleMulticastDelegate::FDelegate::CreateUObject(this, &UEditMeshPolygonsTool::OnDynamicMeshComponentChanged));
-
-
+	// initialize AABBTree
+	MeshSpatial = MakeShared<FDynamicMeshAABBTree3>();
+	MeshSpatial->SetMesh(CurrentMesh.Get());
 
-	// add properties
-	CommonProps = NewObject<UPolyEditCommonProperties>(this);
-	CommonProps->RestoreProperties(this);
-	AddToolPropertySource(CommonProps);
-	CommonProps->WatchProperty(CommonProps->LocalFrameMode,
-								  [this](ELocalFrameMode) { UpdateMultiTransformerFrame(); });
-	CommonProps->WatchProperty(CommonProps->bLockRotation,
-								  [this](bool)
-								  {
-									  LockedTransfomerFrame = LastTransformerFrame;
-								  });
-	// We are going to SilentUpdate here because otherwise the Watches above will immediately fire (why??)
-	// and cause UpdateMultiTransformerFrame() to be called for each, emitting two spurious Transform changes. 
-	CommonProps->SilentUpdateWatched();
+	Topology = (bTriangleMode) ? MakeShared<FTriangleGroupTopology, ESPMode::ThreadSafe>(CurrentMesh.Get(), true)
+		: MakeShared<FGroupTopology, ESPMode::ThreadSafe>(CurrentMesh.Get(), true);
 
 	// set up SelectionMechanic
 	SelectionMechanic = NewObject<UPolygonSelectionMechanic>(this);
+	SelectionMechanic->bAddSelectionFilterPropertiesToParentTool = false; // We'll do this ourselves later
 	SelectionMechanic->Setup(this);
 	SelectionMechanic->Properties->RestoreProperties(this);
 	SelectionMechanic->OnSelectionChanged.AddUObject(this, &UEditMeshPolygonsTool::OnSelectionModifiedEvent);
@@ -227,10 +303,14 @@ void UEditMeshPolygonsTool::Setup()
 	{
 		SelectionMechanic->PolyEdgesRenderer.LineThickness = 1.0;
 	}
+	SelectionMechanic->Initialize(CurrentMesh.Get(),
+		Preview->PreviewMesh->GetTransform(),
+		TargetWorld,
+		Topology.Get(),
+		[this]() { return &GetSpatial(); }
+	);
 
-	// initialize AABBTree
-	MeshSpatial.SetMesh(DynamicMeshComponent->GetMesh());
-	PrecomputeTopology();
+	LinearDeformer.Initialize(CurrentMesh.Get(), Topology.Get());
 
 	// Have to load selection after initializing the selection mechanic since we need to have
 	// the topology built.
@@ -242,121 +322,93 @@ void UEditMeshPolygonsTool::Setup()
 	bSelectionStateDirty = SelectionMechanic->HasSelection();
 
 	// Set UV Scale factor based on initial mesh bounds
-	float BoundsMaxDim = DynamicMeshComponent->GetMesh()->GetBounds().MaxDim();
+	float BoundsMaxDim = CurrentMesh->GetBounds().MaxDim();
 	if (BoundsMaxDim > 0)
 	{
 		UVScaleFactor = 1.0 / BoundsMaxDim;
 	}
 
+	// Wrap the data structures into a context that we can give to the activities
+	UPolyEditActivityContext* ActivityContext = NewObject<UPolyEditActivityContext>();
+	ActivityContext->bTriangleMode = bTriangleMode;
+	ActivityContext->CommonProperties = CommonProps;
+	ActivityContext->CurrentMesh = CurrentMesh;
+	ActivityContext->Preview = Preview;
+	ActivityContext->CurrentTopology = Topology;
+	ActivityContext->MeshSpatial = MeshSpatial;
+	ActivityContext->SelectionMechanic = SelectionMechanic;
+	ActivityContext->EmitActivityStart = [this](const FText& TransactionLabel)
+	{
+		EmitActivityStart(TransactionLabel);
+	};
+	ActivityContext->EmitCurrentMeshChangeAndUpdate = [this](const FText& TransactionLabel,
+		TUniquePtr<UE::Geometry::FDynamicMeshChange> MeshChangeIn,
+		const UE::Geometry::FGroupTopologySelection& OutputSelection, 
+		bool bTopologyModified) 
+	{
+		EmitCurrentMeshChangeAndUpdate(TransactionLabel, MoveTemp(MeshChangeIn), OutputSelection, bTopologyModified);
+	};
+	GetToolManager()->GetContextObjectStore()->RemoveContextObjectsOfType(UPolyEditActivityContext::StaticClass());
+	GetToolManager()->GetContextObjectStore()->AddContextObject(ActivityContext);
+
+	ExtrudeActivity = NewObject<UPolyEditExtrudeActivity>();
+	ExtrudeActivity->Setup(this);
+	
+	InsetOutsetActivity = NewObject<UPolyEditInsetOutsetActivity>();
+	InsetOutsetActivity->Setup(this);
+
+	CutFacesActivity = NewObject<UPolyEditCutFacesActivity>();
+	CutFacesActivity->Setup(this);
+
+	PlanarProjectionUVActivity = NewObject<UPolyEditPlanarProjectionUVActivity>();
+	PlanarProjectionUVActivity->Setup(this);
+
+	InsertEdgeLoopActivity = NewObject<UPolyEditInsertEdgeLoopActivity>();
+	InsertEdgeLoopActivity->Setup(this);
+
+	InsertEdgeActivity = NewObject<UPolyEditInsertEdgeActivity>();
+	InsertEdgeActivity->Setup(this);
+
+	// Now that we've initialized the activities, add in the selection settings and 
+	// CommonProps so that they are at the bottom.
+	AddToolPropertySource(SelectionMechanic->Properties);
+	AddToolPropertySource(CommonProps);
+
 	// hide input StaticMeshComponent
 	UE::ToolTarget::HideSourceObject(Target);
 
-	// init state flags flags
-	bInDrag = false;
-
-	MultiTransformer = NewObject<UMultiTransformer>(this);
-	MultiTransformer->Setup(GetToolManager()->GetPairedGizmoManager(), GetToolManager());
-	MultiTransformer->SetGizmoVisibility(false);
-	MultiTransformer->SetGizmoRepositionable(true);
-	MultiTransformer->SetDisallowNegativeScaling(true);
-	MultiTransformer->OnTransformStarted.AddUObject(this, &UEditMeshPolygonsTool::OnMultiTransformerTransformBegin);
-	MultiTransformer->OnTransformUpdated.AddUObject(this, &UEditMeshPolygonsTool::OnMultiTransformerTransformUpdate);
-	MultiTransformer->OnTransformCompleted.AddUObject(this, &UEditMeshPolygonsTool::OnMultiTransformerTransformEnd);
-	MultiTransformer->OnEndPivotEdit.AddWeakLambda(this, [this]() {
-		LastTransformerFrame = MultiTransformer->GetCurrentGizmoFrame();
+	UInteractiveGizmoManager* GizmoManager = GetToolManager()->GetPairedGizmoManager();
+	
+	TransformGizmo = UE::TransformGizmoUtil::CreateCustomRepositionableTransformGizmo(GizmoManager,
+		ETransformGizmoSubElements::FullTranslateRotateScale, this);
+	// Stop scaling at 0 rather than going negative
+	TransformGizmo->SetDisallowNegativeScaling(true);
+	// We allow non uniform scale even when the gizmo mode is set to "world" because we're not scaling components- we're
+	// moving vertices, so we don't care which axes we "scale" along.
+	TransformGizmo->SetIsNonUniformScaleAllowedFunction([]() {
+		return true;
+	});
+	// Hook up callbacks
+	TransformProxy = NewObject<UTransformProxy>(this);
+	TransformProxy->OnTransformChanged.AddUObject(this, &UEditMeshPolygonsTool::OnGizmoTransformChanged);
+	TransformProxy->OnBeginTransformEdit.AddUObject(this, &UEditMeshPolygonsTool::OnBeginGizmoTransform);
+	TransformProxy->OnEndTransformEdit.AddUObject(this, &UEditMeshPolygonsTool::OnEndGizmoTransform);
+	TransformProxy->OnEndPivotEdit.AddWeakLambda(this, [this](UTransformProxy* Proxy) {
+		LastTransformerFrame = FFrame3d(Proxy->GetTransform());
 		if (CommonProps->bLockRotation)
 		{
 			LockedTransfomerFrame = LastTransformerFrame;
 		}
 	});
-	MultiTransformer->SetSnapToWorldGridSourceFunc([this]() {
-		return CommonProps->bSnapToWorldGrid
-			&& GetToolManager()->GetContextQueriesAPI()->GetCurrentCoordinateSystem() == EToolContextCoordinateSystem::World;
-	});
-	// We allow non uniform scale even when the gizmo mode is set to "world" because we're not scaling components- we're
-	// moving vertices, so we don't care which axes we "scale" along.
-	MultiTransformer->SetIsNonUniformScaleAllowedFunction([]() {
-		return true;
-		});
+	TransformGizmo->SetActiveTarget(TransformProxy, GetToolManager());
+	TransformGizmo->SetVisibility(false);
 
 	DragAlignmentMechanic = NewObject<UDragAlignmentMechanic>(this);
 	DragAlignmentMechanic->Setup(this);
 	DragAlignmentMechanic->InitializeDeformedMeshRayCast(
 		[this]() { return &GetSpatial(); },
-		WorldTransform, &LinearDeformer); // Should happen after PrecomputeTopology so that LinearDeformer is valid
-	MultiTransformer->AddAlignmentMechanic(DragAlignmentMechanic);
-
-	if (bTriangleMode == false)
-	{
-		EditActions = NewObject<UEditMeshPolygonsToolActions>();
-		EditActions->Initialize(this);
-		AddToolPropertySource(EditActions);
-
-		EditEdgeActions = NewObject<UEditMeshPolygonsToolEdgeActions>();
-		EditEdgeActions->Initialize(this);
-		AddToolPropertySource(EditEdgeActions);
-
-		EditUVActions = NewObject<UEditMeshPolygonsToolUVActions>();
-		EditUVActions->Initialize(this);
-		AddToolPropertySource(EditUVActions);
-	}
-	else
-	{
-		EditActions_Triangles = NewObject<UEditMeshPolygonsToolActions_Triangles>();
-		EditActions_Triangles->Initialize(this);
-		AddToolPropertySource(EditActions_Triangles);
-
-		EditEdgeActions_Triangles = NewObject<UEditMeshPolygonsToolEdgeActions_Triangles>();
-		EditEdgeActions_Triangles->Initialize(this);
-		AddToolPropertySource(EditEdgeActions_Triangles);
-	}
-
-	ExtrudeProperties = NewObject<UPolyEditExtrudeProperties>();
-	ExtrudeProperties->RestoreProperties(this);
-	AddToolPropertySource(ExtrudeProperties);
-	SetToolPropertySourceEnabled(ExtrudeProperties, false);
-	ExtrudeProperties->WatchProperty(ExtrudeProperties->Direction,
-									 [this](EPolyEditExtrudeDirection){ RestartExtrude(); });
-
-	OffsetProperties = NewObject<UPolyEditOffsetProperties>();
-	OffsetProperties->RestoreProperties(this);
-	AddToolPropertySource(OffsetProperties);
-	SetToolPropertySourceEnabled(OffsetProperties, false);
-
-	InsetProperties = NewObject<UPolyEditInsetProperties>();
-	InsetProperties->RestoreProperties(this);
-	AddToolPropertySource(InsetProperties);
-	SetToolPropertySourceEnabled(InsetProperties, false);
-
-	OutsetProperties = NewObject<UPolyEditOutsetProperties>();
-	OutsetProperties->RestoreProperties(this);
-	AddToolPropertySource(OutsetProperties);
-	SetToolPropertySourceEnabled(OutsetProperties, false);
-
-	CutProperties = NewObject<UPolyEditCutProperties>();
-	CutProperties->RestoreProperties(this);
-	AddToolPropertySource(CutProperties);
-	SetToolPropertySourceEnabled(CutProperties, false);
-
-	SetUVProperties = NewObject<UPolyEditSetUVProperties>();
-	SetUVProperties->RestoreProperties(this);
-	AddToolPropertySource(SetUVProperties);
-	SetToolPropertySourceEnabled(SetUVProperties, false);
-
-
-	if (bTriangleMode)
-	{
-		SetToolDisplayName(LOCTEXT("EditMeshTrianglesToolName", "Triangle Edit"));
-		GetToolManager()->DisplayMessage(
-			LOCTEXT("OnStartEditMeshPolygonsTool_TriangleMode", "Select triangles to edit mesh. Use middle mouse on gizmo to reposition it. Hold Ctrl while translating or (in local mode) rotating to align to scene. Shift and Ctrl change marquee select behavior. Q toggles Gizmo Orientation Lock."),
-			EToolMessageLevel::UserNotification);
-	}
-	else
-	{
-		GetToolManager()->DisplayMessage(
-			LOCTEXT("OnStartEditMeshPolygonsTool", "Select PolyGroups to edit mesh. Use middle mouse on gizmo to reposition it. Hold Ctrl while translating or (in local mode) rotating to align to scene. Shift and Ctrl change marquee select behavior. Q toggles Gizmo Orientation Lock."),
-			EToolMessageLevel::UserNotification);
-	}
+		WorldTransform, &LinearDeformer); // Should happen after LinearDeformer is initialized
+	DragAlignmentMechanic->AddToGizmo(TransformGizmo);
 
 	if (Topology->Groups.Num() < 2)
 	{
@@ -388,28 +440,29 @@ bool UEditMeshPolygonsTool::IsStoredToolSelectionUsable(const UGroupTopologyStor
 
 void UEditMeshPolygonsTool::Shutdown(EToolShutdownType ShutdownType)
 {
+	if (CurrentActivity)
+	{
+		CurrentActivity->End(ShutdownType);
+		CurrentActivity = nullptr;
+	}
 	CommonProps->SaveProperties(this);
-	ExtrudeProperties->SaveProperties(this);
-	OffsetProperties->SaveProperties(this);
-	InsetProperties->SaveProperties(this);
-	CutProperties->SaveProperties(this);
-	SetUVProperties->SaveProperties(this);
-	SelectionMechanic->Properties->SaveProperties(this);
-
-	MultiTransformer->Shutdown();
+
+	GetToolManager()->GetContextObjectStore()->RemoveContextObjectsOfType(UPolyEditActivityContext::StaticClass());
+
+	ExtrudeActivity->Shutdown(ShutdownType);
+	InsetOutsetActivity->Shutdown(ShutdownType);
+	CutFacesActivity->Shutdown(ShutdownType);
+	PlanarProjectionUVActivity->Shutdown(ShutdownType);
+	InsertEdgeActivity->Shutdown(ShutdownType);
+	InsertEdgeLoopActivity->Shutdown(ShutdownType);
+
+	GetToolManager()->GetPairedGizmoManager()->DestroyAllGizmosByOwner(this);
+
 	DragAlignmentMechanic->Shutdown();
 	// We wait to shut down the selection mechanic in case we need to do work to store the selection.
 
-	if (EditPreview != nullptr)
-	{
-		EditPreview->Disconnect();
-		EditPreview = nullptr;
-	}
-
-	if (DynamicMeshComponent != nullptr)
+	if (Preview != nullptr)
 	{
-		DynamicMeshComponent->OnMeshChanged.Remove(OnDynamicMeshComponentChangedHandle);
-
 		UE::ToolTarget::ShowSourceObject(Target);
 
 		if (ShutdownType == EToolShutdownType::Accept)
@@ -429,25 +482,26 @@ void UEditMeshPolygonsTool::Shutdown(EToolShutdownType ShutdownType)
 					: UGroupTopologyStorableSelection::ETopologyType::FGroupTopology);
 			}
 
+			bool bModifiedTopology = (ModifiedTopologyCounter > 0);
+
 			// may need to compact the mesh if we did undo on a mesh edit, then vertices will be dense but compact checks will fail...
-			if (bWasTopologyEdited)
+			if (bModifiedTopology)
 			{
 				// Store the compact maps if we have a selection that we need to update
-				DynamicMeshComponent->GetMesh()->CompactInPlace(NewStoredToolSelection ? &CompactMaps : nullptr);
+				CurrentMesh->CompactInPlace(NewStoredToolSelection ? &CompactMaps : nullptr);
 			}
 
 			// Finish prepping the stored selection
 			if (NewStoredToolSelection)
 			{
-				SelectionMechanic->GetStorableSelection(*NewStoredToolSelection, bWasTopologyEdited ? &CompactMaps : nullptr);
+				SelectionMechanic->GetStorableSelection(*NewStoredToolSelection, bModifiedTopology ? &CompactMaps : nullptr);
 			}
 
 			// this block bakes the modified DynamicMeshComponent back into the StaticMeshComponent inside an undo transaction
 			GetToolManager()->BeginUndoTransaction(LOCTEXT("EditMeshPolygonsToolTransactionName", "Deform Mesh"));
-			DynamicMeshComponent->ProcessMesh([&](const FDynamicMesh3& ReadMesh)
+			Preview->PreviewMesh->ProcessMesh([&](const FDynamicMesh3& ReadMesh)
 			{
 				FConversionToMeshDescriptionOptions ConversionOptions;
-				bool bModifiedTopology = (ModifiedTopologyCounter > 0);
 				UE::ToolTarget::CommitDynamicMeshUpdate(Target, ReadMesh, bModifiedTopology, ConversionOptions);
 			});
 
@@ -468,16 +522,43 @@ void UEditMeshPolygonsTool::Shutdown(EToolShutdownType ShutdownType)
 			GetToolManager()->EndUndoTransaction();
 		}
 
-		DynamicMeshComponent->UnregisterComponent();
-		DynamicMeshComponent->DestroyComponent();
-		DynamicMeshComponent = nullptr;
+		Preview->Shutdown();
+		Preview = nullptr;
 	}
 
 	// The seleciton mechanic shutdown has to happen after (potentially) saving selection above
 	SelectionMechanic->Shutdown();
-}
 
+	// We null out as many pointers as we can because the tool pointer usually ends up sticking
+	// around in the undo stack.
+	TargetWorld = nullptr;
+	CommonProps = nullptr;
+	EditActions = nullptr;
+	EditActions_Triangles = nullptr;
+	EditEdgeActions = nullptr;
+	EditEdgeActions_Triangles = nullptr;
+	EditUVActions = nullptr;
+	CancelAction = nullptr;
+
+	ExtrudeActivity = nullptr;
+	InsetOutsetActivity = nullptr;
+	CutFacesActivity = nullptr;
+	PlanarProjectionUVActivity = nullptr;
+	InsertEdgeActivity = nullptr;
+	InsertEdgeLoopActivity = nullptr;
+
+	SelectionMechanic = nullptr;
+	DragAlignmentMechanic = nullptr;
+	StoredToolSelection = nullptr;
+
+	TransformGizmo = nullptr;
+	TransformProxy = nullptr;
 
+	CurrentMesh.Reset();
+	Topology.Reset();
+	MeshSpatial.Reset();
+
+}
 
 
 void UEditMeshPolygonsTool::RegisterActions(FInteractiveToolActionSet& ActionSet)
@@ -509,6 +590,10 @@ void UEditMeshPolygonsTool::RegisterActions(FInteractiveToolActionSet& ActionSet
 		LOCTEXT("DeleteSelectionUIName", "Delete Selection"),
 		LOCTEXT("DeleteSelectionTooltip", "Delete Selection"),
 		EModifierKey::None, EKeys::Delete, OnDeletionKeyPress);
+
+	// TODO: Esc should be made to exit out of current activity if one is active. However this
+	// requires a bit of work because we don't seem to be able to register conditional actions,
+	// and we don't want to always capture Esc.
 }
 
 
@@ -522,100 +607,17 @@ void UEditMeshPolygonsTool::RequestAction(EEditMeshPolygonsToolActions ActionTyp
 	PendingAction = ActionType;
 }
 
-
-
-
-
 FDynamicMeshAABBTree3& UEditMeshPolygonsTool::GetSpatial()
 {
 	if (bSpatialDirty)
 	{
-		MeshSpatial.Build();
+		MeshSpatial->Build();
 		bSpatialDirty = false;
 	}
-	return MeshSpatial;
-}
-
-
-
-
-
-
-
-bool UEditMeshPolygonsTool::HitTest(const FRay& WorldRay, FHitResult& OutHit)
-{
-	// If we're in the middle of an action, take the click (to finish an inset, etc).
-	if (CurrentToolMode != ECurrentToolMode::TransformSelection)
-	{
-		OutHit.Distance = 100.0;
-		OutHit.ImpactPoint = WorldRay.PointAt(100.0);
-		return true;
-	}
-
-	// The selection mechanic and gizmo will take care of the TransformSelection state.
-
-	return false;
-}
-
-
-
-FInputRayHit UEditMeshPolygonsTool::IsHitByClick(const FInputDeviceRay& ClickPos)
-{
-	FHitResult OutHit;
-	if (HitTest(ClickPos.WorldRay, OutHit))
-	{
-		return FInputRayHit(OutHit.Distance);
-	}
-	return FInputRayHit(); // bHit is set to false.
-}
-
-void UEditMeshPolygonsTool::OnClicked(const FInputDeviceRay& ClickPos)
-{
-	if ( CurrentToolMode == ECurrentToolMode::ExtrudeSelection )
-	{
-		ApplyExtrude(false);
-		return;
-	}
-	if (CurrentToolMode == ECurrentToolMode::OffsetSelection)
-	{
-		ApplyExtrude(true);
-		return;
-	}
-	else if (CurrentToolMode == ECurrentToolMode::InsetSelection || CurrentToolMode == ECurrentToolMode::OutsetSelection)
-	{
-		ApplyInset(CurrentToolMode == ECurrentToolMode::OutsetSelection);
-		return;
-	}
-	else if (CurrentToolMode == ECurrentToolMode::CutSelection)
-	{
-		if (SurfacePathMechanic->TryAddPointFromRay(ClickPos.WorldRay))
-		{
-			if (SurfacePathMechanic->IsDone())
-			{
-				ApplyCutFaces();
-			}
-		}
-		return;
-	}
-	else if (CurrentToolMode == ECurrentToolMode::SetUVs)
-	{
-		if (SurfacePathMechanic->TryAddPointFromRay(ClickPos.WorldRay))
-		{
-			if (SurfacePathMechanic->IsDone())
-			{
-				ApplySetUVs();
-			}
-		}
-		return;
-	}
-
-	ensureMsgf(CurrentToolMode != ECurrentToolMode::TransformSelection, TEXT("EditMeshPolygonsTool: "
-		"Should not receive click requests in transform mode- they should have been handled by "
-		"selection mechanic or gizmo."));
+	return *MeshSpatial;
 }
 
-
-void UEditMeshPolygonsTool::UpdateMultiTransformerFrame(const FFrame3d* UseFrame)
+void UEditMeshPolygonsTool::UpdateGizmoFrame(const FFrame3d* UseFrame)
 {
 	FFrame3d SetFrame = LastTransformerFrame;
 	if (UseFrame == nullptr)
@@ -640,8 +642,9 @@ void UEditMeshPolygonsTool::UpdateMultiTransformerFrame(const FFrame3d* UseFrame
 	}
 
 	LastTransformerFrame = SetFrame;
-	//MultiTransformer->UpdateGizmoPositionFromWorldFrame(SetFrame, true);
-	MultiTransformer->InitializeGizmoPositionFromWorldFrame(SetFrame, true);
+
+	// This resets the scale as well
+	TransformGizmo->ReinitializeGizmoTransform(SetFrame.ToFTransform());
 }
 
 
@@ -652,7 +655,7 @@ FBox UEditMeshPolygonsTool::GetWorldSpaceFocusBox()
 		FAxisAlignedBox3d Bounds = SelectionMechanic->GetSelectionBounds(true);
 		return (FBox)Bounds;
 	}
-	return UMeshSurfacePointTool::GetWorldSpaceFocusBox();
+	return FBox();
 }
 
 bool UEditMeshPolygonsTool::GetWorldSpaceFocusPoint(const FRay& WorldRay, FVector& PointOut)
@@ -675,133 +678,63 @@ bool UEditMeshPolygonsTool::GetWorldSpaceFocusPoint(const FRay& WorldRay, FVecto
 
 void UEditMeshPolygonsTool::OnSelectionModifiedEvent()
 {
-	FVector3d LocalLastHitPosition, LocalLastHitNormal;
-	SelectionMechanic->GetClickedHitPosition(LocalLastHitPosition, LocalLastHitNormal);
-	FFrame3d LocalFrame(LocalLastHitPosition, LocalLastHitNormal);
-	LastGeometryFrame = SelectionMechanic->GetSelectionFrame(true, &LocalFrame);
-	UpdateMultiTransformerFrame();
 	bSelectionStateDirty = true;
 }
 
 
-
-
-FInputRayHit UEditMeshPolygonsTool::CanBeginClickDragSequence(const FInputDeviceRay& PressPos)
+void UEditMeshPolygonsTool::OnBeginGizmoTransform(UTransformProxy* Proxy)
 {
-	// disable this for now
-	return FInputRayHit();
-	//return UMeshSurfacePointTool::CanBeginClickDragSequence(PressPos);
-}
-
+	SelectionMechanic->ClearHighlight();
+	UpdateDeformerFromSelection( SelectionMechanic->GetActiveSelection() );
+	
+	FTransform Transform = Proxy->GetTransform();
+	InitialGizmoFrame = FFrame3d(Transform);
+	InitialGizmoScale = FVector3d(Transform.GetScale3D());
 
+	BeginDeformerChange();
 
-void UEditMeshPolygonsTool::OnBeginDrag(const FRay& WorldRay)
-{
+	bInGizmoDrag = true;
 }
 
-
-
-void UEditMeshPolygonsTool::OnUpdateDrag(const FRay& Ray)
+void UEditMeshPolygonsTool::OnGizmoTransformChanged(UTransformProxy* Proxy, FTransform Transform)
 {
-	check(false);
+	if (bInGizmoDrag)
+	{
+		LastUpdateGizmoFrame = FFrame3d(Transform);
+		LastUpdateGizmoScale = FVector3d(Transform.GetScale3D());
+		GetToolManager()->PostInvalidation();
+		bGizmoUpdatePending = true;
+		bLastUpdateUsedWorldFrame = (TransformGizmo->CurrentCoordinateSystem == EToolContextCoordinateSystem::World);
+	}
 }
 
-void UEditMeshPolygonsTool::OnEndDrag(const FRay& Ray)
+void UEditMeshPolygonsTool::OnEndGizmoTransform(UTransformProxy* Proxy)
 {
-	check(false);
-}
-
+	bInGizmoDrag = false;
+	bGizmoUpdatePending = false;
+	bSpatialDirty = true;
+	SelectionMechanic->NotifyMeshChanged(false);
 
-
-void UEditMeshPolygonsTool::OnMultiTransformerTransformBegin()
-{
-	SelectionMechanic->ClearHighlight();
-	UpdateDeformerFromSelection( SelectionMechanic->GetActiveSelection() );
-	InitialGizmoFrame = MultiTransformer->GetCurrentGizmoFrame();
-	InitialGizmoScale = MultiTransformer->GetCurrentGizmoScale();
-	BeginChange();
-}
-
-void UEditMeshPolygonsTool::OnMultiTransformerTransformUpdate()
-{
-	if (MultiTransformer->InGizmoEdit())
-	{
-		CacheUpdate_Gizmo();
-	}
-}
-
-void UEditMeshPolygonsTool::OnMultiTransformerTransformEnd()
-{
-	bGizmoUpdatePending = false;
-	bSpatialDirty = true;
-	SelectionMechanic->NotifyMeshChanged(false);
+	FFrame3d TransformFrame(Proxy->GetTransform());
 
 	if (CommonProps->bLockRotation)
 	{
-		FFrame3d SetFrame = MultiTransformer->GetCurrentGizmoFrame();
+		FFrame3d SetFrame = TransformFrame;
 		SetFrame.Rotation = LockedTransfomerFrame.Rotation;
-		MultiTransformer->InitializeGizmoPositionFromWorldFrame(SetFrame, true);
+		TransformGizmo->ReinitializeGizmoTransform(SetFrame.ToFTransform());
 	}
 	else
 	{
-		MultiTransformer->ResetScale();
+		TransformGizmo->SetNewChildScale(FVector::OneVector);
 	}
 
-	LastTransformerFrame = MultiTransformer->GetCurrentGizmoFrame();
+	LastTransformerFrame = TransformFrame;
 
 	// close change record
-	EndChange();
-}
-
-
-
-bool UEditMeshPolygonsTool::OnUpdateHover(const FInputDeviceRay& DevicePos)
-{
-	if (CurrentToolMode == ECurrentToolMode::ExtrudeSelection)
-	{
-		ExtrudeHeightMechanic->UpdateCurrentDistance(DevicePos.WorldRay);
-		bPreviewUpdatePending = true;
-		return true;
-	}
-	else if (CurrentToolMode == ECurrentToolMode::OffsetSelection)
-	{
-		ExtrudeHeightMechanic->UpdateCurrentDistance(DevicePos.WorldRay);
-		bPreviewUpdatePending = true;
-		return true;
-	}
-	else if (CurrentToolMode == ECurrentToolMode::InsetSelection || CurrentToolMode == ECurrentToolMode::OutsetSelection)
-	{
-		CurveDistMechanic->UpdateCurrentDistance(DevicePos.WorldRay);
-		bPreviewUpdatePending = true;
-		return true;
-	}
-	else if (CurrentToolMode == ECurrentToolMode::CutSelection)
-	{
-		SurfacePathMechanic->UpdatePreviewPoint(DevicePos.WorldRay);
-		return true;
-	}
-	else if (CurrentToolMode == ECurrentToolMode::SetUVs)
-	{
-		SurfacePathMechanic->UpdatePreviewPoint(DevicePos.WorldRay);
-		bPreviewUpdatePending = true;
-		return true;
-	}
-
-	ensureMsgf(CurrentToolMode != ECurrentToolMode::TransformSelection, TEXT("EditMeshPolygonsTool: "
-		"Should not receive hover requests in transform mode- they should have been handled by "
-		"selection mechanic or gizmo."));
-	return true;
-}
-
-
-void UEditMeshPolygonsTool::OnEndHover()
-{
-	SelectionMechanic->ClearHighlight();
+	EndDeformerChange();
 }
 
 
-
-
 void UEditMeshPolygonsTool::UpdateDeformerFromSelection(const FGroupTopologySelection& Selection)
 {
 	//Determine which of the following (corners, edges or faces) has been selected by counting the associated feature's IDs
@@ -821,17 +754,6 @@ void UEditMeshPolygonsTool::UpdateDeformerFromSelection(const FGroupTopologySele
 	}
 }
 
-
-
-void UEditMeshPolygonsTool::CacheUpdate_Gizmo()
-{
-	LastUpdateGizmoFrame = MultiTransformer->GetCurrentGizmoFrame();
-	LastUpdateGizmoScale = MultiTransformer->GetCurrentGizmoScale();
-	GetToolManager()->PostInvalidation();
-	bGizmoUpdatePending = true;
-	bLastUpdateUsedWorldFrame = (MultiTransformer->GetGizmoCoordinateSystem() == EToolContextCoordinateSystem::World);
-}
-
 void UEditMeshPolygonsTool::ComputeUpdate_Gizmo()
 {
 	if (SelectionMechanic->HasSelection() == false || bGizmoUpdatePending == false)
@@ -847,7 +769,7 @@ void UEditMeshPolygonsTool::ComputeUpdate_Gizmo()
 	FVector3d CurScaleDelta = CurScale - InitialGizmoScale;
 	FVector3d LocalTranslation = WorldTransform.InverseTransformVector(TranslationDelta);
 
-	FDynamicMesh3* Mesh = DynamicMeshComponent->GetMesh();
+	FDynamicMesh3* Mesh = CurrentMesh.Get();
 	if (TranslationDelta.SquaredLength() > 0.0001 || RotateDelta.SquaredLength() > 0.0001 || CurScaleDelta.SquaredLength() > 0.0001)
 	{
 		if (bLastUpdateUsedWorldFrame)
@@ -879,765 +801,301 @@ void UEditMeshPolygonsTool::ComputeUpdate_Gizmo()
 				FVector3d NewPosWorld = CurFrame.FromFramePoint(PosGizmo);
 				FVector3d NewPosLocal = WorldTransform.InverseTransformPosition(NewPosWorld);
 				return NewPosLocal;
-			});
-		}
-	}
-	else
-	{
-		// Reset mesh to initial positions.
-		LinearDeformer.ClearSolution(Mesh);
-	}
-	DynamicMeshComponent->FastNotifyPositionsUpdated(true);
-	GetToolManager()->PostInvalidation();
-}
-
-
-
-void UEditMeshPolygonsTool::OnTick(float DeltaTime)
-{
-	MultiTransformer->Tick(DeltaTime);
-
-	bool bLocalCoordSystem = GetToolManager()->GetPairedGizmoManager()
-		->GetContextQueriesAPI()->GetCurrentCoordinateSystem() == EToolContextCoordinateSystem::Local;
-	if (CommonProps->bLocalCoordSystem != bLocalCoordSystem)
-	{
-		CommonProps->bLocalCoordSystem = bLocalCoordSystem;
-		NotifyOfPropertyChangeByTool(CommonProps);
-	}
-
-	if (bGizmoUpdatePending)
-	{
-		ComputeUpdate_Gizmo();
-	}
-
-	if (bSelectionStateDirty)
-	{
-		// update color highlights
-		DynamicMeshComponent->FastNotifySecondaryTrianglesChanged();
-
-		if (SelectionMechanic->HasSelection())
-		{
-			MultiTransformer->SetGizmoVisibility(true);
-
-			// update frame because we might be here due to an undo event/etc, rather than an explicit selection change
-			LastGeometryFrame = SelectionMechanic->GetSelectionFrame(true, &LastGeometryFrame);
-			UpdateMultiTransformerFrame();
-		}
-		else
-		{
-			MultiTransformer->SetGizmoVisibility(false);
-		}
-
-		bSelectionStateDirty = false;
-	}
-
-	if (PendingAction != EEditMeshPolygonsToolActions::NoAction)
-	{
-		CancelMeshEditChange();
-
-		if (PendingAction == EEditMeshPolygonsToolActions::Extrude || PendingAction == EEditMeshPolygonsToolActions::Offset)
-		{
-			GetToolManager()->EmitObjectChange(this, MakeUnique<FBeginInteractivePolyEditChange>(CurrentOperationTimestamp), LOCTEXT("PolyMeshEditBeginExtrude", "Begin Extrude"));
-			BeginExtrude( (PendingAction == EEditMeshPolygonsToolActions::Offset) );
-		}
-		else if (PendingAction == EEditMeshPolygonsToolActions::Inset)
-		{
-			GetToolManager()->EmitObjectChange(this, MakeUnique<FBeginInteractivePolyEditChange>(CurrentOperationTimestamp), LOCTEXT("PolyMeshEditBeginInset", "Begin Inset"));
-			BeginInset(false);
-		}
-		else if (PendingAction == EEditMeshPolygonsToolActions::Outset)
-		{
-			GetToolManager()->EmitObjectChange(this, MakeUnique<FBeginInteractivePolyEditChange>(CurrentOperationTimestamp), LOCTEXT("PolyMeshEditBeginOutset", "Begin Outset"));
-			BeginInset(true);
-		}
-		else if (PendingAction == EEditMeshPolygonsToolActions::CutFaces)
-		{
-			GetToolManager()->EmitObjectChange(this, MakeUnique<FBeginInteractivePolyEditChange>(CurrentOperationTimestamp), LOCTEXT("PolyMeshEditBeginCutFaces", "Begin Cut Faces"));
-			BeginCutFaces();
-		}
-		else if (PendingAction == EEditMeshPolygonsToolActions::PlanarProjectionUV)
-		{
-			GetToolManager()->EmitObjectChange(this, MakeUnique<FBeginInteractivePolyEditChange>(CurrentOperationTimestamp), LOCTEXT("PolyMeshEditBeginUVPlanarProjection", "Begin Set UVs"));
-			BeginSetUVs();
-		}
-		else if (PendingAction == EEditMeshPolygonsToolActions::Merge)
-		{
-			ApplyMerge();
-		}
-		else if (PendingAction == EEditMeshPolygonsToolActions::Delete)
-		{
-			ApplyDelete();
-		}
-		else if (PendingAction == EEditMeshPolygonsToolActions::RecalculateNormals)
-		{
-			ApplyRecalcNormals();
-		}
-		else if (PendingAction == EEditMeshPolygonsToolActions::FlipNormals)
-		{
-			ApplyFlipNormals();
-		}
-		else if (PendingAction == EEditMeshPolygonsToolActions::CollapseEdge)
-		{
-			ApplyCollapseEdge();
-		}
-		else if (PendingAction == EEditMeshPolygonsToolActions::WeldEdges)
-		{
-			ApplyWeldEdges();
-		}
-		else if (PendingAction == EEditMeshPolygonsToolActions::StraightenEdge)
-		{
-			ApplyStraightenEdges();
-		}
-		else if (PendingAction == EEditMeshPolygonsToolActions::FillHole)
-		{
-			ApplyFillHole();
-		}
-		else if (PendingAction == EEditMeshPolygonsToolActions::Retriangulate)
-		{
-			ApplyRetriangulate();
-		}
-		else if (PendingAction == EEditMeshPolygonsToolActions::Decompose)
-		{
-			ApplyDecompose();
-		}
-		else if (PendingAction == EEditMeshPolygonsToolActions::Disconnect)
-		{
-			ApplyDisconnect();
-		}
-		else if (PendingAction == EEditMeshPolygonsToolActions::Duplicate)
-		{
-			ApplyDuplicate();
-		}
-		else if (PendingAction == EEditMeshPolygonsToolActions::PokeSingleFace)
-		{
-			ApplyPokeSingleFace();
-		}
-		else if (PendingAction == EEditMeshPolygonsToolActions::SplitSingleEdge)
-		{
-			ApplySplitSingleEdge();
-		}
-		else if (PendingAction == EEditMeshPolygonsToolActions::CollapseSingleEdge)
-		{
-			ApplyCollapseSingleEdge();
-		}
-		else if (PendingAction == EEditMeshPolygonsToolActions::FlipSingleEdge)
-		{
-			ApplyFlipSingleEdge();
-		}
-
-		PendingAction = EEditMeshPolygonsToolActions::NoAction;
-	}
-
-
-	// todo: convert to ValueWatcher
-	if (CurrentToolMode == ECurrentToolMode::SetUVs)
-	{
-		EPreviewMaterialType WantMaterial = (SetUVProperties->bShowMaterial) ? EPreviewMaterialType::SourceMaterials : EPreviewMaterialType::UVMaterial;
-		if (CurrentPreviewMaterial != WantMaterial)
-		{
-			UpdateEditPreviewMaterials(WantMaterial);
-		}
-	}
-
-
-	if (bPreviewUpdatePending)
-	{
-		if (CurrentToolMode == ECurrentToolMode::ExtrudeSelection)
-		{
-			EditPreview->UpdateExtrudeType(ExtrudeHeightMechanic->CurrentHeight);
-		}
-		else if (CurrentToolMode == ECurrentToolMode::OffsetSelection)
-		{
-			if (OffsetProperties->bUseFaceNormals)
-			{
-				EditPreview->UpdateExtrudeType_FaceNormalAvg(ExtrudeHeightMechanic->CurrentHeight);
-			}
-			else
-			{
-				EditPreview->UpdateExtrudeType(ExtrudeHeightMechanic->CurrentHeight, true);
-			}
-		}
-		else if (CurrentToolMode == ECurrentToolMode::InsetSelection || CurrentToolMode == ECurrentToolMode::OutsetSelection)
-		{
-			bool bOutset = (CurrentToolMode == ECurrentToolMode::OutsetSelection);
-			double Sign = bOutset ? -1.0 : 1.0;
-			bool bReproject = (bOutset) ? false : InsetProperties->bReproject;
-			double Softness = (bOutset) ? OutsetProperties->Softness : InsetProperties->Softness;
-			bool bBoundaryOnly = (bOutset) ? OutsetProperties->bBoundaryOnly : InsetProperties->bBoundaryOnly;
-			double AreaCorrection = (bOutset) ? OutsetProperties->AreaScale : InsetProperties->AreaScale;
-			EditPreview->UpdateInsetType(Sign* CurveDistMechanic->CurrentDistance, bReproject, Softness, AreaCorrection, bBoundaryOnly);
-		}
-		else if (CurrentToolMode == ECurrentToolMode::SetUVs)
-		{
-			UpdateSetUVS();
-		}
-		bPreviewUpdatePending = false;
-	}
-}
-
-
-void UEditMeshPolygonsTool::PrecomputeTopology()
-{
-	FDynamicMesh3* Mesh = DynamicMeshComponent->GetMesh();
-	Topology = (bTriangleMode) ? MakeUnique<FTriangleGroupTopology>(Mesh, false) : MakeUnique<FGroupTopology>(Mesh, false);
-	Topology->RebuildTopology();
-
-	// update selection mechanic
-	SelectionMechanic->Initialize(DynamicMeshComponent, Topology.Get(),
-		[this]() { return &GetSpatial(); }
-		);
-
-	LinearDeformer.Initialize(Mesh, Topology.Get());
-}
-
-
-
-
-void UEditMeshPolygonsTool::Render(IToolsContextRenderAPI* RenderAPI)
-{
-	GetToolManager()->GetContextQueriesAPI()->GetCurrentViewState(CameraState);
-	DynamicMeshComponent->bExplicitShowWireframe = CommonProps->bShowWireframe;
-
-	SelectionMechanic->Render(RenderAPI);
-
-	DragAlignmentMechanic->Render(RenderAPI);
-
-	if (ExtrudeHeightMechanic != nullptr)
-	{
-		ExtrudeHeightMechanic->Render(RenderAPI);
-	}
-	if (CurveDistMechanic != nullptr)
-	{
-		CurveDistMechanic->Render(RenderAPI);
-	}
-	if (SurfacePathMechanic != nullptr)
-	{
-		SurfacePathMechanic->Render(RenderAPI);
-	}
-}
-
-void UEditMeshPolygonsTool::DrawHUD(FCanvas* Canvas, IToolsContextRenderAPI* RenderAPI)
-{
-	SelectionMechanic->DrawHUD(Canvas, RenderAPI);
-}
-
-
-
-
-//
-// Change Tracking
-//
-
-
-void UEditMeshPolygonsTool::UpdateChangeFromROI(bool bFinal)
-{
-	if (ActiveVertexChange == nullptr)
-	{
-		return;
-	}
-
-	FDynamicMesh3* Mesh = DynamicMeshComponent->GetMesh();
-	ActiveVertexChange->SaveVertices(Mesh, LinearDeformer.GetModifiedVertices(), !bFinal);
-	ActiveVertexChange->SaveOverlayNormals(Mesh, LinearDeformer.GetModifiedOverlayNormals(), !bFinal);
-}
-
-
-void UEditMeshPolygonsTool::BeginChange()
-{
-	if (ActiveVertexChange == nullptr)
-	{
-		ActiveVertexChange = new FMeshVertexChangeBuilder(EMeshVertexChangeComponents::VertexPositions | EMeshVertexChangeComponents::OverlayNormals);
-		UpdateChangeFromROI(false);
-	}
-}
-
-
-void UEditMeshPolygonsTool::EndChange()
-{
-	if (ActiveVertexChange != nullptr)
-	{
-		UpdateChangeFromROI(true);
-		GetToolManager()->EmitObjectChange(DynamicMeshComponent, MoveTemp(ActiveVertexChange->Change), LOCTEXT("PolyMeshDeformationChange", "PolyMesh Edit"));
-	}
-
-	delete ActiveVertexChange;
-	ActiveVertexChange = nullptr;
-
-	CurrentOperationTimestamp++;
-}
-
-
-void UEditMeshPolygonsTool::OnDynamicMeshComponentChanged()
-{
-	bSpatialDirty = true;
-	SelectionMechanic->NotifyMeshChanged(false);
-}
-
-void UEditMeshPolygonsTool::AfterTopologyEdit()
-{
-	bSpatialDirty = true;
-	bWasTopologyEdited = true;
-	SelectionMechanic->NotifyMeshChanged(true);
-
-	DynamicMeshComponent->NotifyMeshUpdated();
-	MeshSpatial.SetMesh(DynamicMeshComponent->GetMesh(), true);
-	PrecomputeTopology();
-}
-
-
-
-
-void UEditMeshPolygonsTool::ApplyPlaneCut()
-{
-	FFrame3d PlaneFrame;
-
-	FDynamicMesh3* Mesh = DynamicMeshComponent->GetMesh();
-	FMeshPlaneCut Cut(Mesh, PlaneFrame.Origin, PlaneFrame.Z());
-	Cut.UVScaleFactor = UVScaleFactor;
-
-	FMeshEdgeSelection Edges(Mesh);
-	const FGroupTopologySelection& ActiveSelection = SelectionMechanic->GetActiveSelection();
-	if (ActiveSelection.SelectedGroupIDs.Num() > 0)
-	{
-		for (int32 GroupID : ActiveSelection.SelectedGroupIDs)
-		{
-			Edges.SelectTriangleEdges( Topology->GetGroupTriangles(GroupID) );
-		}
-		Cut.EdgeFilterFunc = [&](int EdgeID) { return Edges.IsSelected(EdgeID); };
-	}
-
-	Cut.SplitEdgesOnly(true);
-
-	DynamicMeshComponent->NotifyMeshUpdated();
-	MeshSpatial.SetMesh(DynamicMeshComponent->GetMesh(), true);
-	PrecomputeTopology();
-}
-
-
-
-
-
-void UEditMeshPolygonsTool::BeginExtrude(bool bIsNormalOffset)
-{
-	FDynamicMesh3* Mesh = DynamicMeshComponent->GetMesh();
-	if (bIsNormalOffset)
-	{
-		// yikes...
-	}
-	if (BeginMeshFaceEditChangeWithPreview() == false)
-	{
-		return;
-	}
-
-	ActiveSelectionFrameWorld.AlignAxis(2, GetExtrudeDirection());
-	EditPreview->InitializeExtrudeType(Mesh, ActiveTriangleSelection, ActiveSelectionFrameWorld.Z(), &WorldTransform, true);
-	// move world extrude frame to point on surface
-	ActiveSelectionFrameWorld.Origin = EditPreview->GetInitialPatchMeshSpatial().FindNearestPoint(ActiveSelectionFrameWorld.Origin);
-
-	// make inifinite-extent hit-test mesh
-	FDynamicMesh3 ExtrudeHitTargetMesh;
-	EditPreview->MakeExtrudeTypeHitTargetMesh(ExtrudeHitTargetMesh);
-
-	ExtrudeHeightMechanic = NewObject<UPlaneDistanceFromHitMechanic>(this);
-	ExtrudeHeightMechanic->Setup(this);
-
-	ExtrudeHeightMechanic->WorldHitQueryFunc = [this](const FRay& WorldRay, FHitResult& HitResult)
-	{
-		return ToolSceneQueriesUtil::FindNearestVisibleObjectHit(DynamicMeshComponent->GetWorld(), HitResult, WorldRay);
-	};
-	ExtrudeHeightMechanic->WorldPointSnapFunc = [this](const FVector3d& WorldPos, FVector3d& SnapPos)
-	{
-		return CommonProps->bSnapToWorldGrid && ToolSceneQueriesUtil::FindWorldGridSnapPoint(this, WorldPos, SnapPos);
-	};
-	ExtrudeHeightMechanic->CurrentHeight = 1.0f;  // initialize to something non-zero...prob should be based on polygon bounds maybe?
-
-	ExtrudeHeightMechanic->Initialize(MoveTemp(ExtrudeHitTargetMesh), ActiveSelectionFrameWorld, true);
-	CurrentToolMode = (bIsNormalOffset) ? ECurrentToolMode::OffsetSelection : ECurrentToolMode::ExtrudeSelection;
-
-	if (bIsNormalOffset == false)
-	{
-		SetToolPropertySourceEnabled(ExtrudeProperties, true);
-	}
-	else
-	{
-		SetToolPropertySourceEnabled(OffsetProperties, true);
-	}
-	SetActionButtonPanelsVisible(false);
-}
-
-
-
-void UEditMeshPolygonsTool::ApplyExtrude(bool bIsOffset)
-{
-	check(ExtrudeHeightMechanic != nullptr && EditPreview != nullptr);
-
-	FVector3d ExtrudeDir = WorldTransform.InverseTransformVector(ActiveSelectionFrameWorld.Z());
-	double ExtrudeDist = ExtrudeHeightMechanic->CurrentHeight;
-
-	FDynamicMesh3* Mesh = DynamicMeshComponent->GetMesh();
-	FOffsetMeshRegion Extruder(Mesh);
-	Extruder.UVScaleFactor = UVScaleFactor;
-	Extruder.Triangles = ActiveTriangleSelection;
-	TSet<int32> TriangleSet(ActiveTriangleSelection);
-	Extruder.OffsetPositionFunc = [&](const FVector3d& Pos, const FVector3f& Normal, int32 VertexID) {
-		return Pos + ExtrudeDist * (bIsOffset ? (FVector3d)Normal : ExtrudeDir);
-	};
-	Extruder.bIsPositiveOffset = (ExtrudeDist > 0);
-	Extruder.bUseFaceNormals = (bIsOffset && OffsetProperties->bUseFaceNormals);
-	Extruder.bOffsetFullComponentsAsSolids = bIsOffset || ExtrudeProperties->bShellsToSolids;
-	Extruder.ChangeTracker = MakeUnique<FDynamicMeshChangeTracker>(Mesh);
-	Extruder.ChangeTracker->BeginChange();
-	Extruder.Apply();
-
-	FMeshNormals::QuickComputeVertexNormalsForTriangles(*Mesh, Extruder.AllModifiedTriangles);
-
-	// construct new selection
-	FGroupTopologySelection NewSelection;
-	if (!bTriangleMode)
-	{
-		for (const FOffsetMeshRegion::FOffsetInfo& Info : Extruder.OffsetRegions)
-		{
-			NewSelection.SelectedGroupIDs.Append(Info.OffsetGroups);
+			});
 		}
 	}
 	else
 	{
-		for (const FOffsetMeshRegion::FOffsetInfo& Info : Extruder.OffsetRegions)
-		{
-			NewSelection.SelectedGroupIDs.Append(Info.InitialTriangles);
-		}
+		// Reset mesh to initial positions.
+		LinearDeformer.ClearSolution(Mesh);
 	}
 
-	// emit undo
-	TUniquePtr<FMeshChange> MeshChange = MakeUnique<FMeshChange>(Extruder.ChangeTracker->EndChange());
-	CompleteMeshEditChange( (bIsOffset) ? LOCTEXT("PolyMeshOffsetChange", "Offset") : LOCTEXT("PolyMeshExtrudeChange", "Extrude"),
-		MoveTemp(MeshChange), NewSelection);
-
-	ExtrudeHeightMechanic = nullptr;
-	CurrentToolMode = ECurrentToolMode::TransformSelection;
+	Preview->PreviewMesh->UpdatePreview(CurrentMesh.Get());
 
-	SetToolPropertySourceEnabled(ExtrudeProperties, false);
-	SetToolPropertySourceEnabled(OffsetProperties, false);
-	SetActionButtonPanelsVisible(true);
+	GetToolManager()->PostInvalidation();
 }
 
 
-void UEditMeshPolygonsTool::RestartExtrude()
+
+void UEditMeshPolygonsTool::OnTick(float DeltaTime)
 {
-	if (CurrentToolMode == ECurrentToolMode::ExtrudeSelection)
+	Preview->Tick(DeltaTime);
+
+	if (CurrentActivity)
 	{
-		CancelMeshEditChange();
-		BeginExtrude(false);
+		CurrentActivity->Tick(DeltaTime);
 	}
-}
 
-
-FVector3d UEditMeshPolygonsTool::GetExtrudeDirection() const
-{
-	switch (ExtrudeProperties->Direction)
+	bool bLocalCoordSystem = GetToolManager()->GetPairedGizmoManager()
+		->GetContextQueriesAPI()->GetCurrentCoordinateSystem() == EToolContextCoordinateSystem::Local;
+	if (CommonProps->bLocalCoordSystem != bLocalCoordSystem)
 	{
-	default:
-	case EPolyEditExtrudeDirection::SelectionNormal:
-		return ActiveSelectionFrameWorld.Z();
-	case EPolyEditExtrudeDirection::WorldX:
-		return FVector3d::UnitX();
-	case EPolyEditExtrudeDirection::WorldY:
-		return FVector3d::UnitY();
-	case EPolyEditExtrudeDirection::WorldZ:
-		return FVector3d::UnitZ();
-	case EPolyEditExtrudeDirection::LocalX:
-		return WorldTransform.GetRotation().AxisX();
-	case EPolyEditExtrudeDirection::LocalY:
-		return WorldTransform.GetRotation().AxisY();
-	case EPolyEditExtrudeDirection::LocalZ:
-		return WorldTransform.GetRotation().AxisZ();
+		CommonProps->bLocalCoordSystem = bLocalCoordSystem;
+		NotifyOfPropertyChangeByTool(CommonProps);
 	}
-	return ActiveSelectionFrameWorld.Z();
-}
 
+	// TODO: We should probably allow gizmo snapping in local mode as well. It requires some gizmo changes.
+	TransformGizmo->bSnapToWorldGrid = CommonProps->bSnapToWorldGrid && !bLocalCoordSystem;
 
-
-
-void UEditMeshPolygonsTool::BeginInset(bool bOutset)
-{
-	FDynamicMesh3* Mesh = DynamicMeshComponent->GetMesh();
-	if (BeginMeshFaceEditChangeWithPreview() == false)
+	if (bGizmoUpdatePending)
 	{
-		return;
+		ComputeUpdate_Gizmo();
 	}
 
-	EditPreview->InitializeInsetType(Mesh, ActiveTriangleSelection, &WorldTransform);
-
-	// make infinite-extent hit-test mesh
-	FDynamicMesh3 InsetHitTargetMesh;
-	EditPreview->MakeInsetTypeTargetMesh(InsetHitTargetMesh);
-
-	CurveDistMechanic = NewObject<USpatialCurveDistanceMechanic>(this);
-	CurveDistMechanic->Setup(this);
-	CurveDistMechanic->WorldPointSnapFunc = [this](const FVector3d& WorldPos, FVector3d& SnapPos)
+	if (bSelectionStateDirty)
 	{
-		return CommonProps->bSnapToWorldGrid && ToolSceneQueriesUtil::FindWorldGridSnapPoint(this, WorldPos, SnapPos);
-	};
-	CurveDistMechanic->CurrentDistance = 1.0f;  // initialize to something non-zero...prob should be based on polygon bounds maybe?
-
-	FMeshBoundaryLoops Loops(&InsetHitTargetMesh);
-	TArray<FVector3d> LoopVertices;
-	Loops.Loops[0].GetVertices(LoopVertices);
-	CurveDistMechanic->InitializePolyLoop(LoopVertices, FTransform3d::Identity());
-	CurrentToolMode = (bOutset) ? ECurrentToolMode::OutsetSelection : ECurrentToolMode::InsetSelection;
-
-	SetToolPropertySourceEnabled((bOutset) ? 
-		(UInteractiveToolPropertySet*)OutsetProperties : (UInteractiveToolPropertySet*)InsetProperties, true);
-	SetActionButtonPanelsVisible(false);
-}
-
-
+		// update color highlights
+		Preview->PreviewMesh->FastNotifySecondaryTrianglesChanged();
 
+		UpdateGizmoVisibility();
 
-void UEditMeshPolygonsTool::ApplyInset(bool bOutset)
-{
-	check(CurveDistMechanic != nullptr && EditPreview != nullptr);
+		bSelectionStateDirty = false;
+	}
 
-	FDynamicMesh3* Mesh = DynamicMeshComponent->GetMesh();
-	FInsetMeshRegion Inset(Mesh);
-	Inset.UVScaleFactor = UVScaleFactor;
-	Inset.Triangles = ActiveTriangleSelection;
-	Inset.InsetDistance = (bOutset) ? -CurveDistMechanic->CurrentDistance : CurveDistMechanic->CurrentDistance;
-	Inset.bReproject = (bOutset) ? false : InsetProperties->bReproject;
-	Inset.Softness = (bOutset) ? OutsetProperties->Softness : InsetProperties->Softness;
-	Inset.bSolveRegionInteriors = (bOutset) ? (!OutsetProperties->bBoundaryOnly) : (!InsetProperties->bBoundaryOnly);
-	Inset.AreaCorrection = (bOutset) ? OutsetProperties->AreaScale : InsetProperties->AreaScale;
+	if (PendingAction != EEditMeshPolygonsToolActions::NoAction)
+	{
+		// Clear any existing error messages.
+		GetToolManager()->DisplayMessage(FText(), EToolMessageLevel::UserWarning);
 
-	Inset.ChangeTracker = MakeUnique<FDynamicMeshChangeTracker>(Mesh);
-	Inset.ChangeTracker->BeginChange();
-	Inset.Apply();
+		switch (PendingAction)
+		{
 
-	FMeshNormals::QuickComputeVertexNormalsForTriangles(*Mesh, Inset.AllModifiedTriangles);
+		//Interactive operations:
+		case EEditMeshPolygonsToolActions::Extrude:
+		{
+			StartActivity(ExtrudeActivity);
+			break;
+		}
+		case EEditMeshPolygonsToolActions::InsetOutset:
+		{
+			StartActivity(InsetOutsetActivity);
+			break;
+		}
+		case EEditMeshPolygonsToolActions::CutFaces:
+		{
+			StartActivity(CutFacesActivity);
+			break;
+		}
+		case EEditMeshPolygonsToolActions::PlanarProjectionUV:
+		{
+			StartActivity(PlanarProjectionUVActivity);
+			break;
+		}
+		case EEditMeshPolygonsToolActions::InsertEdge:
+		{
+			StartActivity(InsertEdgeActivity);
+			break;
+		}
+		case EEditMeshPolygonsToolActions::InsertEdgeLoop:
+		{
+			StartActivity(InsertEdgeLoopActivity);
+			break;
+		}
 
-	// emit undo
-	FGroupTopologySelection CurSelection = SelectionMechanic->GetActiveSelection();
-	TUniquePtr<FMeshChange> MeshChange = MakeUnique<FMeshChange>(Inset.ChangeTracker->EndChange());
-	CompleteMeshEditChange( bOutset ? LOCTEXT("PolyMeshOutsetChange", "Outset") : LOCTEXT("PolyMeshInsetChange", "Inset"), MoveTemp(MeshChange), CurSelection);
+		case EEditMeshPolygonsToolActions::CancelCurrent:
+		{
+			EndCurrentActivity();
+			break;
+		}
 
-	CurveDistMechanic = nullptr;
-	CurrentToolMode = ECurrentToolMode::TransformSelection;
+		// Single action operations:
+		case EEditMeshPolygonsToolActions::Merge:
+			ApplyMerge();
+			break;
+		case  EEditMeshPolygonsToolActions::Delete:
+			ApplyDelete();
+			break;
+		case EEditMeshPolygonsToolActions::RecalculateNormals:
+			ApplyRecalcNormals();
+			break;
+		case EEditMeshPolygonsToolActions::FlipNormals:
+			ApplyFlipNormals();
+			break;
+		case EEditMeshPolygonsToolActions::CollapseEdge:
+			ApplyCollapseEdge();
+			break;
+		case EEditMeshPolygonsToolActions::WeldEdges:
+			ApplyWeldEdges();
+			break;
+		case EEditMeshPolygonsToolActions::StraightenEdge:
+			ApplyStraightenEdges();
+			break;
+		case EEditMeshPolygonsToolActions::FillHole:
+			ApplyFillHole();
+			break;
+		case EEditMeshPolygonsToolActions::Retriangulate:
+			ApplyRetriangulate();
+			break;
+		case EEditMeshPolygonsToolActions::Decompose:
+			ApplyDecompose();
+			break;
+		case EEditMeshPolygonsToolActions::Disconnect:
+			ApplyDisconnect();
+			break;
+		case EEditMeshPolygonsToolActions::Duplicate:
+			ApplyDuplicate();
+			break;
+		case EEditMeshPolygonsToolActions::PokeSingleFace:
+			ApplyPokeSingleFace();
+			break;
+		case EEditMeshPolygonsToolActions::SplitSingleEdge:
+			ApplySplitSingleEdge();
+			break;
+		case EEditMeshPolygonsToolActions::CollapseSingleEdge:
+			ApplyCollapseSingleEdge();
+			break;
+		case EEditMeshPolygonsToolActions::FlipSingleEdge:
+			ApplyFlipSingleEdge();
+			break;
+		}
 
-	SetToolPropertySourceEnabled((bOutset) ?
-		(UInteractiveToolPropertySet*)OutsetProperties : (UInteractiveToolPropertySet*)InsetProperties, false);
-	SetActionButtonPanelsVisible(true);
+		PendingAction = EEditMeshPolygonsToolActions::NoAction;
+	}
 }
 
-
-
-
-void UEditMeshPolygonsTool::BeginCutFaces()
+void UEditMeshPolygonsTool::StartActivity(TObjectPtr<UInteractiveToolActivity> Activity)
 {
-	FDynamicMesh3* Mesh = DynamicMeshComponent->GetMesh();
-	if (BeginMeshFaceEditChangeWithPreview() == false)
+	EndCurrentActivity();
+
+	// Right now we rely on the activity to fail to start or to issue an error message if the
+	// conditions are not right. Someday, we are going to disable the buttons based on a CanStart
+	// call.
+	if (Activity->Start() == EToolActivityStartResult::Running)
 	{
-		GetToolManager()->DisplayMessage(
-			LOCTEXT("OnCutFacesFailedMessage", "Cannot Cut Current Selection"),
-			EToolMessageLevel::UserWarning);
-		return;
+		TransformGizmo->SetVisibility(false);
+		SelectionMechanic->SetIsEnabled(false);
+		CurrentActivity = Activity;
+		SetToolPropertySourceEnabled(CancelAction, true);
+		SetActionButtonPanelsVisible(false);
 	}
-	GetToolManager()->DisplayMessage(
-		LOCTEXT("OnBeginCutFacesMessage", "Click twice on selected face to define cut line"),
-		EToolMessageLevel::UserMessage);
-
+}
 
-	EditPreview->InitializeStaticType(Mesh, ActiveTriangleSelection, &WorldTransform);
+void UEditMeshPolygonsTool::EndCurrentActivity()
+{
+	if (CurrentActivity)
+	{
+		if (CurrentActivity->IsRunning())
+		{
+			CurrentActivity->End(EToolShutdownType::Cancel);
 
-	FDynamicMesh3 StaticHitTargetMesh;
-	EditPreview->MakeInsetTypeTargetMesh(StaticHitTargetMesh);
+			// Reset info message.
+			GetToolManager()->DisplayMessage(DefaultMessage,
+				EToolMessageLevel::UserNotification);
+		}
 
-	SurfacePathMechanic = NewObject<UCollectSurfacePathMechanic>(this);
-	SurfacePathMechanic->Setup(this);
-	SurfacePathMechanic->InitializeMeshSurface(MoveTemp(StaticHitTargetMesh));
-	SurfacePathMechanic->SetFixedNumPointsMode(2);
-	SurfacePathMechanic->bSnapToTargetMeshVertices = true;
-	double SnapTol = ToolSceneQueriesUtil::GetDefaultVisualAngleSnapThreshD();
-	SurfacePathMechanic->SpatialSnapPointsFunc = [this, SnapTol](FVector3d Position1, FVector3d Position2)
-	{
-		return CutProperties->bSnapToVertices && ToolSceneQueriesUtil::PointSnapQuery(this->CameraState, Position1, Position2, SnapTol);
-	};
+		CurrentActivity = nullptr;
+		++ActivityTimestamp;
 
-	CurrentToolMode = ECurrentToolMode::CutSelection;
-	SetToolPropertySourceEnabled(CutProperties, true);
-	SetActionButtonPanelsVisible(false);
+		SetToolPropertySourceEnabled(CancelAction, false);
+		SetActionButtonPanelsVisible(true);
+		SelectionMechanic->SetIsEnabled(true);
+		UpdateGizmoVisibility();
+	}
 }
 
-void UEditMeshPolygonsTool::ApplyCutFaces()
+void UEditMeshPolygonsTool::NotifyActivitySelfEnded(UInteractiveToolActivity* Activity)
 {
-	check(SurfacePathMechanic != nullptr && EditPreview != nullptr);
-
-	FDynamicMesh3* Mesh = DynamicMeshComponent->GetMesh();
+	EndCurrentActivity();
+}
 
-	// construct cut plane normal from line points
-	FFrame3d Point0(SurfacePathMechanic->HitPath[0]), Point1(SurfacePathMechanic->HitPath[1]);
-	FVector3d PlaneNormal;
-	if (CutProperties->Orientation == EPolyEditCutPlaneOrientation::ViewDirection)
+void UEditMeshPolygonsTool::UpdateGizmoVisibility()
+{
+	if (SelectionMechanic->HasSelection())
 	{
-		FVector3d Direction0 = UE::Geometry::Normalized(Point0.Origin - (FVector3d)CameraState.Position);
-		FVector3d Direction1 = UE::Geometry::Normalized(Point1.Origin - (FVector3d)CameraState.Position);
-		PlaneNormal = Direction1.Cross(Direction0);
+		TransformGizmo->SetVisibility(true);
+
+		// update frame because we might be here due to an undo event/etc, rather than an explicit 
+		// selection change
+		LastGeometryFrame = SelectionMechanic->GetSelectionFrame(true, &LastGeometryFrame);
+		UpdateGizmoFrame();
 	}
 	else
 	{
-		FVector3d LineDirection = UE::Geometry::Normalized(Point1.Origin - Point0.Origin);
-		FVector3d UpVector = UE::Geometry::Normalized(Point0.Z() + Point1.Z());
-		PlaneNormal = LineDirection.Cross(UpVector);
-	}
-	FVector3d PlaneOrigin = 0.5 * (Point0.Origin + Point1.Origin);
-	// map into local space of target mesh
-	PlaneOrigin = WorldTransform.InverseTransformPosition(PlaneOrigin);
-	PlaneNormal = WorldTransform.InverseTransformNormal(PlaneNormal);
-	UE::Geometry::Normalize(PlaneNormal);
-
-	// track changes
-	FDynamicMeshChangeTracker ChangeTracker(Mesh);
-	ChangeTracker.BeginChange();
-	TArray<int32> VertexSelection;
-	UE::Geometry::TriangleToVertexIDs(Mesh, ActiveTriangleSelection, VertexSelection);
-	ChangeTracker.SaveVertexOneRingTriangles(VertexSelection, true);
-
-	// apply the cut to edges of selected triangles
-	FGroupTopologySelection OutputSelection;
-	FMeshPlaneCut Cut(Mesh, PlaneOrigin, PlaneNormal);
-	FMeshEdgeSelection Edges(Mesh);
-	Edges.SelectTriangleEdges(ActiveTriangleSelection);
-	Cut.EdgeFilterFunc = [&](int EdgeID) { return Edges.IsSelected(EdgeID); };
-	if (Cut.SplitEdgesOnly(true))
-	{
-		if (!bTriangleMode)
-		{
-			for (FMeshPlaneCut::FCutResultRegion& Region : Cut.ResultRegions)
-			{
-				OutputSelection.SelectedGroupIDs.Add(Region.GroupID);
-			}
-		}
-		else
-		{
-			// Retain the selection along the cut. ResultSeedTriangles does not
-			// contain selected tris that are not cut, so re-add the original selected
-			// tris.
-			OutputSelection.SelectedGroupIDs.Append(Cut.ResultSeedTriangles);
-			OutputSelection.SelectedGroupIDs.Append(ActiveTriangleSelection);
-		}
+		TransformGizmo->SetVisibility(false);
 	}
+}
 
-	// emit undo
-	TUniquePtr<FMeshChange> MeshChange = MakeUnique<FMeshChange>(ChangeTracker.EndChange());
-	CompleteMeshEditChange(LOCTEXT("PolyMeshCutFacesChange", "Cut Faces"), MoveTemp(MeshChange), OutputSelection);
+void UEditMeshPolygonsTool::Render(IToolsContextRenderAPI* RenderAPI)
+{
+	Preview->PreviewMesh->EnableWireframe(CommonProps->bShowWireframe);
+	SelectionMechanic->Render(RenderAPI);
+	DragAlignmentMechanic->Render(RenderAPI);
 
-	SurfacePathMechanic = nullptr;
-	CurrentToolMode = ECurrentToolMode::TransformSelection;
-	SetToolPropertySourceEnabled(CutProperties, false);
-	SetActionButtonPanelsVisible(true);
+	if (CurrentActivity)
+	{
+		CurrentActivity->Render(RenderAPI);
+	}
 }
 
+void UEditMeshPolygonsTool::DrawHUD(FCanvas* Canvas, IToolsContextRenderAPI* RenderAPI)
+{
+	SelectionMechanic->DrawHUD(Canvas, RenderAPI);
+}
 
 
 
 
-void UEditMeshPolygonsTool::BeginSetUVs()
+//
+// Gizmo change tracking
+//
+void UEditMeshPolygonsTool::UpdateDeformerChangeFromROI(bool bFinal)
 {
-	FDynamicMesh3* Mesh = DynamicMeshComponent->GetMesh();
-	if (BeginMeshFaceEditChangeWithPreview() == false)
+	if (ActiveVertexChange == nullptr)
 	{
-		GetToolManager()->DisplayMessage( LOCTEXT("OnSetUVsFailedMesssage", "Cannot Set UVs for Current Selection"), EToolMessageLevel::UserWarning);
 		return;
 	}
-	GetToolManager()->DisplayMessage( LOCTEXT("OnBeginSetUVsMessage", "Click on the face to Set UVs"), EToolMessageLevel::UserMessage);
-
-	EditPreview->InitializeStaticType(Mesh, ActiveTriangleSelection, &WorldTransform);
-	UpdateEditPreviewMaterials((SetUVProperties->bShowMaterial) ? EPreviewMaterialType::SourceMaterials : EPreviewMaterialType::UVMaterial);
-
-	FDynamicMesh3 StaticHitTargetMesh;
-	EditPreview->MakeInsetTypeTargetMesh(StaticHitTargetMesh);
-
-	SurfacePathMechanic = NewObject<UCollectSurfacePathMechanic>(this);
-	SurfacePathMechanic->Setup(this);
-	SurfacePathMechanic->InitializeMeshSurface(MoveTemp(StaticHitTargetMesh));
-	SurfacePathMechanic->SetFixedNumPointsMode(2);
-	SurfacePathMechanic->bSnapToTargetMeshVertices = true;
-	double SnapTol = ToolSceneQueriesUtil::GetDefaultVisualAngleSnapThreshD();
-	SurfacePathMechanic->SpatialSnapPointsFunc = [this, SnapTol](FVector3d Position1, FVector3d Position2)
-	{
-		return ToolSceneQueriesUtil::PointSnapQuery(this->CameraState, Position1, Position2, SnapTol);
-	};
 
-	CurrentToolMode = ECurrentToolMode::SetUVs;
-	SetToolPropertySourceEnabled(SetUVProperties, true);
-	SetActionButtonPanelsVisible(false);
+	FDynamicMesh3* Mesh = CurrentMesh.Get();
+	ActiveVertexChange->SaveVertices(Mesh, LinearDeformer.GetModifiedVertices(), !bFinal);
+	ActiveVertexChange->SaveOverlayNormals(Mesh, LinearDeformer.GetModifiedOverlayNormals(), !bFinal);
 }
 
-void UEditMeshPolygonsTool::UpdateSetUVS()
+void UEditMeshPolygonsTool::BeginDeformerChange()
 {
-	// align projection frame to line user is drawing out from plane origin
-	FFrame3d PlanarFrame = SurfacePathMechanic->PreviewPathPoint;
-	double UVScale = 1.0 / ActiveSelectionBounds.MaxDim();
-	if (SurfacePathMechanic->HitPath.Num() == 1)
+	if (ActiveVertexChange == nullptr)
 	{
-		SurfacePathMechanic->InitializePlaneSurface(PlanarFrame);
-
-		FVector3d Delta = PlanarFrame.Origin - SurfacePathMechanic->HitPath[0].Origin;
-		double Dist = UE::Geometry::Normalize(Delta);
-		UVScale *= FMathd::Lerp(1.0, 25.0, Dist / ActiveSelectionBounds.MaxDim());
-		PlanarFrame = SurfacePathMechanic->HitPath[0];
-		PlanarFrame.ConstrainedAlignAxis(0, Delta, PlanarFrame.Z());
+		ActiveVertexChange = new FMeshVertexChangeBuilder(EMeshVertexChangeComponents::VertexPositions | EMeshVertexChangeComponents::OverlayNormals);
+		UpdateDeformerChangeFromROI(false);
 	}
+}
 
-	EditPreview->UpdateStaticType([&](FDynamicMesh3& Mesh)
+void UEditMeshPolygonsTool::EndDeformerChange()
+{
+	if (ActiveVertexChange != nullptr)
 	{
-		FDynamicMeshEditor Editor(&Mesh);
-		TArray<int32> AllTriangles;
-		for (int32 tid : Mesh.TriangleIndicesItr())
-		{
-			AllTriangles.Add(tid);
-		}
-		Editor.SetTriangleUVsFromProjection(AllTriangles, PlanarFrame, UVScale, FVector2f::Zero(), false);
-	}, false);
+		UpdateDeformerChangeFromROI(true);
+		GetToolManager()->EmitObjectChange(this, MoveTemp(ActiveVertexChange->Change), 
+			LOCTEXT("PolyMeshDeformationChange", "PolyMesh Edit"));
+	}
 
+	delete ActiveVertexChange;
+	ActiveVertexChange = nullptr;
 }
 
-
-void UEditMeshPolygonsTool::ApplySetUVs()
+// This gets called by vertex change events emitted via gizmo (deformer) interaction
+void UEditMeshPolygonsTool::ApplyChange(const FMeshVertexChange* Change, bool bRevert)
 {
-	FDynamicMesh3* Mesh = DynamicMeshComponent->GetMesh();
-	FGroupTopologySelection ActiveSelection = SelectionMechanic->GetActiveSelection();
+	Preview->PreviewMesh->ApplyChange(Change, bRevert);
+	CurrentMesh->Copy(*Preview->PreviewMesh->GetMesh());
+	bSpatialDirty = true;
+	SelectionMechanic->NotifyMeshChanged(false);
 
-	// align projection frame to line user drew
-	FFrame3d PlanarFrame = SurfacePathMechanic->HitPath[0];
-	double UVScale = 1.0 / ActiveSelectionBounds.MaxDim();
-	FVector3d Delta = SurfacePathMechanic->HitPath[1].Origin - PlanarFrame.Origin;
-	double Dist = UE::Geometry::Normalize(Delta);
-	UVScale *= FMathd::Lerp(1.0, 25.0, Dist / ActiveSelectionBounds.MaxDim());
-	PlanarFrame.ConstrainedAlignAxis(0, Delta, PlanarFrame.Z());
-
-	// transform to local, use 3D point to transfer UV scale value
-	FVector3d ScalePt = PlanarFrame.Origin + UVScale * PlanarFrame.Z();
-	FTransform3d ToLocalXForm(WorldTransform.Inverse());
-	PlanarFrame.Transform(ToLocalXForm);
-	ScalePt = ToLocalXForm.TransformPosition(ScalePt);
-	UVScale = Distance(ScalePt, PlanarFrame.Origin);
-
-	// track changes
-	FDynamicMeshChangeTracker ChangeTracker(Mesh);
-	ChangeTracker.BeginChange();
-	ChangeTracker.SaveTriangles(ActiveTriangleSelection, true);
-	FDynamicMeshEditor Editor(Mesh);
-	Editor.SetTriangleUVsFromProjection(ActiveTriangleSelection, PlanarFrame, UVScale, FVector2f::Zero(), false, 0);
+	// Topology does not need updating
+}
 
-	// emit undo
-	TUniquePtr<FMeshChange> MeshChange = MakeUnique<FMeshChange>(ChangeTracker.EndChange());
-	CompleteMeshEditChange(LOCTEXT("PolyMeshSetUVsChange", "Set UVs"), MoveTemp(MeshChange), ActiveSelection);
 
-	SurfacePathMechanic = nullptr;
-	CurrentToolMode = ECurrentToolMode::TransformSelection;
-	SetToolPropertySourceEnabled(SetUVProperties, false);
-	SetActionButtonPanelsVisible(true);
-}
+void UEditMeshPolygonsTool::UpdateFromCurrentMesh(bool bGroupTopologyModified)
+{
+	Preview->PreviewMesh->UpdatePreview(CurrentMesh.Get(),
+		bGroupTopologyModified ? UPreviewMesh::ERenderUpdateMode::FullUpdate : UPreviewMesh::ERenderUpdateMode::FastUpdate);
+	bSpatialDirty = true;
+	SelectionMechanic->NotifyMeshChanged(bGroupTopologyModified);
 
+	if (bGroupTopologyModified)
+	{
+		Topology->RebuildTopology();
+	}
+}
 
 
 
@@ -1651,7 +1109,7 @@ void UEditMeshPolygonsTool::ApplyMerge()
 		return;
 	}
 
-	FDynamicMesh3* Mesh = DynamicMeshComponent->GetMesh();
+	FDynamicMesh3* Mesh = CurrentMesh.Get();
 	FDynamicMeshChangeTracker ChangeTracker(Mesh);
 	ChangeTracker.BeginChange();
 	ChangeTracker.SaveTriangles(ActiveTriangleSelection, true);
@@ -1665,11 +1123,7 @@ void UEditMeshPolygonsTool::ApplyMerge()
 		NewSelection.SelectedGroupIDs.Add(NewGroupID);
 	}
 	
-	// emit undo
-	TUniquePtr<FMeshChange> MeshChange = MakeUnique<FMeshChange>(ChangeTracker.EndChange());
-	CompleteMeshEditChange(LOCTEXT("PolyMeshMergeChange", "Merge"), MoveTemp(MeshChange), NewSelection);
-
-	CurrentToolMode = ECurrentToolMode::TransformSelection;
+	EmitCurrentMeshChangeAndUpdate(LOCTEXT("PolyMeshMergeChange", "Merge"), ChangeTracker.EndChange(), NewSelection, true);
 }
 
 
@@ -1687,7 +1141,7 @@ void UEditMeshPolygonsTool::ApplyDelete()
 		return;
 	}
 
-	FDynamicMesh3* Mesh = DynamicMeshComponent->GetMesh();
+	FDynamicMesh3* Mesh = CurrentMesh.Get();
 
 	// prevent deleting all triangles
 	if (ActiveTriangleSelection.Num() >= Mesh->TriangleCount())
@@ -1704,12 +1158,8 @@ void UEditMeshPolygonsTool::ApplyDelete()
 	FDynamicMeshEditor Editor(Mesh);
 	Editor.RemoveTriangles(ActiveTriangleSelection, true);
 
-	// emit undo
-	TUniquePtr<FMeshChange> MeshChange = MakeUnique<FMeshChange>(ChangeTracker.EndChange());
 	FGroupTopologySelection NewSelection;
-	CompleteMeshEditChange(LOCTEXT("PolyMeshDeleteChange", "Delete"), MoveTemp(MeshChange), NewSelection);
-
-	CurrentToolMode = ECurrentToolMode::TransformSelection;
+	EmitCurrentMeshChangeAndUpdate(LOCTEXT("PolyMeshDeleteChange", "Delete"), ChangeTracker.EndChange(), NewSelection, true);
 }
 
 
@@ -1724,7 +1174,7 @@ void UEditMeshPolygonsTool::ApplyRecalcNormals()
 		return;
 	}
 
-	FDynamicMesh3* Mesh = DynamicMeshComponent->GetMesh();
+	FDynamicMesh3* Mesh = CurrentMesh.Get();
 	FDynamicMeshChangeTracker ChangeTracker(Mesh);
 	ChangeTracker.BeginChange();
 	FDynamicMeshEditor Editor(Mesh);
@@ -1735,11 +1185,11 @@ void UEditMeshPolygonsTool::ApplyRecalcNormals()
 		Editor.SetTriangleNormals(Topology->GetGroupTriangles(GroupID));
 	}
 
-	// emit undo
-	TUniquePtr<FMeshChange> MeshChange = MakeUnique<FMeshChange>(ChangeTracker.EndChange());
-	CompleteMeshEditChange(LOCTEXT("PolyMeshRecalcNormalsChange", "Recalc Normals"), MoveTemp(MeshChange), ActiveSelection);
-
-	CurrentToolMode = ECurrentToolMode::TransformSelection;
+	// We actually don't even need any of the wrapper around this change since we're not altering
+	// positions or topology (so no other structures need updating), but we go ahead and go the
+	// same route as everything else.
+	EmitCurrentMeshChangeAndUpdate(LOCTEXT("PolyMeshRecalcNormalsChange", "Recalc Normals"), 
+		ChangeTracker.EndChange(), ActiveSelection, false);
 }
 
 
@@ -1753,7 +1203,7 @@ void UEditMeshPolygonsTool::ApplyFlipNormals()
 		return;
 	}
 
-	FDynamicMesh3* Mesh = DynamicMeshComponent->GetMesh();
+	FDynamicMesh3* Mesh = CurrentMesh.Get();
 	FDynamicMeshChangeTracker ChangeTracker(Mesh);
 	ChangeTracker.BeginChange();
 	FDynamicMeshEditor Editor(Mesh);
@@ -1767,11 +1217,11 @@ void UEditMeshPolygonsTool::ApplyFlipNormals()
 		}
 	}
 
-	// emit undo
-	TUniquePtr<FMeshChange> MeshChange = MakeUnique<FMeshChange>(ChangeTracker.EndChange());
-	CompleteMeshEditChange(LOCTEXT("PolyMeshFlipNormalsChange", "Flip Normals"), MoveTemp(MeshChange), ActiveSelection);
-
-	CurrentToolMode = ECurrentToolMode::TransformSelection;
+	// We actually don't even need any of the wrapper around this change since we're not altering
+	// positions or topology (so no other structures need updating), but we go ahead and go the
+	// same route as everything else.
+	EmitCurrentMeshChangeAndUpdate(LOCTEXT("PolyMeshFlipNormalsChange", "Flip Normals"), 
+		ChangeTracker.EndChange(), ActiveSelection, false);
 }
 
 
@@ -1784,7 +1234,7 @@ void UEditMeshPolygonsTool::ApplyRetriangulate()
 	}
 
 	int32 nCompleted = 0;
-	FDynamicMesh3* Mesh = DynamicMeshComponent->GetMesh();
+	FDynamicMesh3* Mesh = CurrentMesh.Get();
 	FDynamicMeshChangeTracker ChangeTracker(Mesh);
 	ChangeTracker.BeginChange();
 	FDynamicMeshEditor Editor(Mesh);
@@ -1842,9 +1292,8 @@ void UEditMeshPolygonsTool::ApplyRetriangulate()
 		GetToolManager()->DisplayMessage(LOCTEXT("OnRetriangulateFailures", "Some faces could not be retriangulated"), EToolMessageLevel::UserWarning);
 	}
 
-	TUniquePtr<FMeshChange> MeshChange = MakeUnique<FMeshChange>(ChangeTracker.EndChange());
-	CompleteMeshEditChange(LOCTEXT("PolyMeshRetriangulateChange", "Retriangulate"), MoveTemp(MeshChange), ActiveSelection);
-	CurrentToolMode = ECurrentToolMode::TransformSelection;
+	EmitCurrentMeshChangeAndUpdate(LOCTEXT("PolyMeshRetriangulateChange", "Retriangulate"), 
+		ChangeTracker.EndChange(), ActiveSelection, true);
 }
 
 
@@ -1858,7 +1307,7 @@ void UEditMeshPolygonsTool::ApplyDecompose()
 		return;
 	}
 
-	FDynamicMesh3* Mesh = DynamicMeshComponent->GetMesh();
+	FDynamicMesh3* Mesh = CurrentMesh.Get();
 	FDynamicMeshChangeTracker ChangeTracker(Mesh);
 	ChangeTracker.BeginChange();
 	FGroupTopologySelection NewSelection;
@@ -1875,9 +1324,7 @@ void UEditMeshPolygonsTool::ApplyDecompose()
 	}
 
 
-	TUniquePtr<FMeshChange> MeshChange = MakeUnique<FMeshChange>(ChangeTracker.EndChange());
-	CompleteMeshEditChange(LOCTEXT("PolyMeshDecomposeChange", "Decompose"), MoveTemp(MeshChange), NewSelection);
-	CurrentToolMode = ECurrentToolMode::TransformSelection;
+	EmitCurrentMeshChangeAndUpdate(LOCTEXT("PolyMeshDecomposeChange", "Decompose"), ChangeTracker.EndChange(), NewSelection, true);
 }
 
 
@@ -1889,7 +1336,7 @@ void UEditMeshPolygonsTool::ApplyDisconnect()
 		return;
 	}
 
-	FDynamicMesh3* Mesh = DynamicMeshComponent->GetMesh();
+	FDynamicMesh3* Mesh = CurrentMesh.Get();
 	FDynamicMeshChangeTracker ChangeTracker(Mesh);
 	ChangeTracker.BeginChange();
 	FGroupTopologySelection ActiveSelection = SelectionMechanic->GetActiveSelection();
@@ -1902,9 +1349,7 @@ void UEditMeshPolygonsTool::ApplyDisconnect()
 	FDynamicMeshEditor Editor(Mesh);
 	Editor.DisconnectTriangles(AllTriangles, false);
 
-	TUniquePtr<FMeshChange> MeshChange = MakeUnique<FMeshChange>(ChangeTracker.EndChange());
-	CompleteMeshEditChange(LOCTEXT("PolyMeshDisconnectChange", "Disconnect"), MoveTemp(MeshChange), ActiveSelection);
-	CurrentToolMode = ECurrentToolMode::TransformSelection;
+	EmitCurrentMeshChangeAndUpdate(LOCTEXT("PolyMeshDisconnectChange", "Disconnect"), ChangeTracker.EndChange(), ActiveSelection, true);
 }
 
 
@@ -1918,7 +1363,7 @@ void UEditMeshPolygonsTool::ApplyDuplicate()
 		return;
 	}
 
-	FDynamicMesh3* Mesh = DynamicMeshComponent->GetMesh();
+	FDynamicMesh3* Mesh = CurrentMesh.Get();
 	FDynamicMeshChangeTracker ChangeTracker(Mesh);
 	ChangeTracker.BeginChange();
 	FGroupTopologySelection ActiveSelection = SelectionMechanic->GetActiveSelection();
@@ -1935,9 +1380,7 @@ void UEditMeshPolygonsTool::ApplyDuplicate()
 	FGroupTopologySelection NewSelection;
 	NewSelection.SelectedGroupIDs.Append(bTriangleMode ? EditResult.NewTriangles : EditResult.NewGroups);
 
-	TUniquePtr<FMeshChange> MeshChange = MakeUnique<FMeshChange>(ChangeTracker.EndChange());
-	CompleteMeshEditChange(LOCTEXT("PolyMeshDisconnectChange", "Disconnect"), MoveTemp(MeshChange), NewSelection);
-	CurrentToolMode = ECurrentToolMode::TransformSelection;
+	EmitCurrentMeshChangeAndUpdate(LOCTEXT("PolyMeshDisconnectChange", "Disconnect"), ChangeTracker.EndChange(), NewSelection, true);
 }
 
 
@@ -1956,7 +1399,7 @@ void UEditMeshPolygonsTool::ApplyCollapseEdge()
 		return;
 	}
 
-	FDynamicMesh3* Mesh = DynamicMeshComponent->GetMesh();
+	FDynamicMesh3* Mesh = CurrentMesh.Get();
 
 	FDynamicMeshChangeTracker ChangeTracker(Mesh);
 	ChangeTracker.BeginChange();
@@ -1974,11 +1417,8 @@ void UEditMeshPolygonsTool::ApplyCollapseEdge()
 	//}
 
 	// emit undo
-	TUniquePtr<FMeshChange> MeshChange = MakeUnique<FMeshChange>(ChangeTracker.EndChange());
 	FGroupTopologySelection NewSelection;
-	CompleteMeshEditChange(LOCTEXT("PolyMeshEdgeCollapseChange", "Collapse"), MoveTemp(MeshChange), NewSelection);
-
-	CurrentToolMode = ECurrentToolMode::TransformSelection;
+	EmitCurrentMeshChangeAndUpdate(LOCTEXT("PolyMeshEdgeCollapseChange", "Collapse"), ChangeTracker.EndChange(), NewSelection, true);
 }
 
 
@@ -1991,11 +1431,10 @@ void UEditMeshPolygonsTool::ApplyWeldEdges()
 	if ( bValidInput == false )
 	{
 		GetToolManager()->DisplayMessage( LOCTEXT("OnWeldEdgesFailed", "Cannot Weld current selection"), EToolMessageLevel::UserWarning);
-		CancelMeshEditChange();
 		return;
 	}
 
-	FDynamicMesh3* Mesh = DynamicMeshComponent->GetMesh();
+	FDynamicMesh3* Mesh = CurrentMesh.Get();
 
 	FDynamicMeshChangeTracker ChangeTracker(Mesh);
 	ChangeTracker.BeginChange();
@@ -2014,14 +1453,11 @@ void UEditMeshPolygonsTool::ApplyWeldEdges()
 	if (Result != EMeshResult::Ok)
 	{
 		GetToolManager()->DisplayMessage( LOCTEXT("OnWeldEdgesFailed", "Cannot Weld current selection"), EToolMessageLevel::UserWarning);
-		CancelMeshEditChange();
 		return;
 	}
 
-	TUniquePtr<FMeshChange> MeshChange = MakeUnique<FMeshChange>(ChangeTracker.EndChange());
 	FGroupTopologySelection NewSelection;
-	CompleteMeshEditChange(LOCTEXT("PolyMeshWeldEdgeChange", "Weld Edges"), MoveTemp(MeshChange), NewSelection);
-	CurrentToolMode = ECurrentToolMode::TransformSelection;
+	EmitCurrentMeshChangeAndUpdate(LOCTEXT("PolyMeshWeldEdgeChange", "Weld Edges"), ChangeTracker.EndChange(), NewSelection, true);
 }
 
 
@@ -2030,11 +1466,10 @@ void UEditMeshPolygonsTool::ApplyStraightenEdges()
 	if (BeginMeshEdgeEditChange() == false)
 	{
 		GetToolManager()->DisplayMessage(LOCTEXT("OnStraightenEdgesFailed", "Cannot Straighten current selection"), EToolMessageLevel::UserWarning);
-		CancelMeshEditChange();
 		return;
 	}
 
-	FDynamicMesh3* Mesh = DynamicMeshComponent->GetMesh();
+	FDynamicMesh3* Mesh = CurrentMesh.Get();
 
 	FDynamicMeshChangeTracker ChangeTracker(Mesh);
 	ChangeTracker.BeginChange();
@@ -2057,10 +1492,9 @@ void UEditMeshPolygonsTool::ApplyStraightenEdges()
 		}
 	}
 
-	TUniquePtr<FMeshChange> MeshChange = MakeUnique<FMeshChange>(ChangeTracker.EndChange());
 	FGroupTopologySelection NewSelection;
-	CompleteMeshEditChange(LOCTEXT("PolyMeshStraightenEdgeChange", "Straighten Edges"), MoveTemp(MeshChange), NewSelection);
-	CurrentToolMode = ECurrentToolMode::TransformSelection;
+	EmitCurrentMeshChangeAndUpdate(LOCTEXT("PolyMeshStraightenEdgeChange", "Straighten Edges"), 
+		ChangeTracker.EndChange(), NewSelection, false);
 }
 
 
@@ -2070,11 +1504,10 @@ void UEditMeshPolygonsTool::ApplyFillHole()
 	if (BeginMeshBoundaryEdgeEditChange(false) == false)
 	{
 		GetToolManager()->DisplayMessage( LOCTEXT("OnEdgeFillFailed", "Cannot Fill current selection"), EToolMessageLevel::UserWarning);
-		CancelMeshEditChange();
 		return;
 	}
 
-	FDynamicMesh3* Mesh = DynamicMeshComponent->GetMesh();
+	FDynamicMesh3* Mesh = CurrentMesh.Get();
 	FDynamicMeshChangeTracker ChangeTracker(Mesh);
 	ChangeTracker.BeginChange();
 	FGroupTopologySelection NewSelection;
@@ -2118,10 +1551,7 @@ void UEditMeshPolygonsTool::ApplyFillHole()
 		}
 	}
 
-	// emit undo
-	TUniquePtr<FMeshChange> MeshChange = MakeUnique<FMeshChange>(ChangeTracker.EndChange());
-	CompleteMeshEditChange(LOCTEXT("PolyMeshFillHoleChange", "Fill Hole"), MoveTemp(MeshChange), NewSelection);
-	CurrentToolMode = ECurrentToolMode::TransformSelection;
+	EmitCurrentMeshChangeAndUpdate(LOCTEXT("PolyMeshFillHoleChange", "Fill Hole"), ChangeTracker.EndChange(), NewSelection, true);
 }
 
 
@@ -2135,7 +1565,7 @@ void UEditMeshPolygonsTool::ApplyPokeSingleFace()
 		return;
 	}
 
-	FDynamicMesh3* Mesh = DynamicMeshComponent->GetMesh();
+	FDynamicMesh3* Mesh = CurrentMesh.Get();
 	FDynamicMeshChangeTracker ChangeTracker(Mesh);
 	ChangeTracker.BeginChange();
 	ChangeTracker.SaveTriangles(ActiveTriangleSelection, true);
@@ -2151,9 +1581,7 @@ void UEditMeshPolygonsTool::ApplyPokeSingleFace()
 		}
 	}
 
-	TUniquePtr<FMeshChange> MeshChange = MakeUnique<FMeshChange>(ChangeTracker.EndChange());
-	CompleteMeshEditChange(LOCTEXT("PolyMeshPokeChange", "Poke Faces"), MoveTemp(MeshChange), NewSelection);
-	CurrentToolMode = ECurrentToolMode::TransformSelection;
+	EmitCurrentMeshChangeAndUpdate(LOCTEXT("PolyMeshPokeChange", "Poke Faces"), ChangeTracker.EndChange(), NewSelection, true);
 }
 
 
@@ -2166,7 +1594,7 @@ void UEditMeshPolygonsTool::ApplyFlipSingleEdge()
 		return;
 	}
 
-	FDynamicMesh3* Mesh = DynamicMeshComponent->GetMesh();
+	FDynamicMesh3* Mesh = CurrentMesh.Get();
 	FGroupTopologySelection ActiveSelection = SelectionMechanic->GetActiveSelection();
 	FDynamicMeshChangeTracker ChangeTracker(Mesh);
 	ChangeTracker.BeginChange();
@@ -2183,9 +1611,8 @@ void UEditMeshPolygonsTool::ApplyFlipSingleEdge()
 		}
 	}
 
-	TUniquePtr<FMeshChange> MeshChange = MakeUnique<FMeshChange>(ChangeTracker.EndChange());
-	CompleteMeshEditChange(LOCTEXT("PolyMeshFlipChange", "Flip Edges"), MoveTemp(MeshChange), ActiveSelection);
-	CurrentToolMode = ECurrentToolMode::TransformSelection;
+	// Group topology may or may not change, but just assume that it does.
+	EmitCurrentMeshChangeAndUpdate(LOCTEXT("PolyMeshFlipChange", "Flip Edges"), ChangeTracker.EndChange(), ActiveSelection, true);
 }
 
 void UEditMeshPolygonsTool::ApplyCollapseSingleEdge()
@@ -2196,7 +1623,7 @@ void UEditMeshPolygonsTool::ApplyCollapseSingleEdge()
 		return;
 	}
 
-	FDynamicMesh3* Mesh = DynamicMeshComponent->GetMesh();
+	FDynamicMesh3* Mesh = CurrentMesh.Get();
 	FGroupTopologySelection ActiveSelection = SelectionMechanic->GetActiveSelection();
 	FDynamicMeshChangeTracker ChangeTracker(Mesh);
 	ChangeTracker.BeginChange();
@@ -2227,9 +1654,8 @@ void UEditMeshPolygonsTool::ApplyCollapseSingleEdge()
 		}
 	}
 
-	TUniquePtr<FMeshChange> MeshChange = MakeUnique<FMeshChange>(ChangeTracker.EndChange());
-	CompleteMeshEditChange(LOCTEXT("PolyMeshCollapseChange", "Collapse Edges"), MoveTemp(MeshChange), FGroupTopologySelection());
-	CurrentToolMode = ECurrentToolMode::TransformSelection;
+	EmitCurrentMeshChangeAndUpdate(LOCTEXT("PolyMeshCollapseChange", "Collapse Edges"), 
+		ChangeTracker.EndChange(), FGroupTopologySelection(), true);
 }
 
 void UEditMeshPolygonsTool::ApplySplitSingleEdge()
@@ -2240,7 +1666,7 @@ void UEditMeshPolygonsTool::ApplySplitSingleEdge()
 		return;
 	}
 
-	FDynamicMesh3* Mesh = DynamicMeshComponent->GetMesh();
+	FDynamicMesh3* Mesh = CurrentMesh.Get();
 	FGroupTopologySelection NewSelection;
 	FDynamicMeshChangeTracker ChangeTracker(Mesh);
 	ChangeTracker.BeginChange();
@@ -2269,9 +1695,7 @@ void UEditMeshPolygonsTool::ApplySplitSingleEdge()
 		}
 	}
 
-	TUniquePtr<FMeshChange> MeshChange = MakeUnique<FMeshChange>(ChangeTracker.EndChange());
-	CompleteMeshEditChange(LOCTEXT("PolyMeshSplitChange", "Split Edges"), MoveTemp(MeshChange), NewSelection);
-	CurrentToolMode = ECurrentToolMode::TransformSelection;
+	EmitCurrentMeshChangeAndUpdate(LOCTEXT("PolyMeshSplitChange", "Split Edges"), ChangeTracker.EndChange(), NewSelection, true);
 }
 
 
@@ -2279,8 +1703,6 @@ void UEditMeshPolygonsTool::ApplySplitSingleEdge()
 
 bool UEditMeshPolygonsTool::BeginMeshFaceEditChange()
 {
-	check(EditPreview == nullptr);
-
 	ActiveTriangleSelection.Reset();
 
 	// need some selected faces
@@ -2291,7 +1713,7 @@ bool UEditMeshPolygonsTool::BeginMeshFaceEditChange()
 		return false;
 	}
 
-	const FDynamicMesh3* Mesh = DynamicMeshComponent->GetMesh();
+	const FDynamicMesh3* Mesh = CurrentMesh.Get();
 	ActiveSelectionBounds = FAxisAlignedBox3d::Empty();
 	for (int tid : ActiveTriangleSelection)
 	{
@@ -2306,71 +1728,69 @@ bool UEditMeshPolygonsTool::BeginMeshFaceEditChange()
 	return true;
 }
 
-
-bool UEditMeshPolygonsTool::BeginMeshFaceEditChangeWithPreview()
+void UEditMeshPolygonsTool::EmitCurrentMeshChangeAndUpdate(const FText& TransactionLabel,
+	TUniquePtr<FDynamicMeshChange> MeshChangeIn,
+	const FGroupTopologySelection& OutputSelection, 
+	bool bGroupTopologyModified)
 {
-	bool bOK = BeginMeshFaceEditChange();
-	if (bOK)
-	{
-		EditPreview = NewObject<UPolyEditPreviewMesh>(this);
-		EditPreview->CreateInWorld(UE::ToolTarget::GetTargetActor(Target)->GetWorld(), FTransform::Identity);
-		UpdateEditPreviewMaterials(EPreviewMaterialType::PreviewMaterial);
-		EditPreview->EnableWireframe(true);
+	// open top-level transaction
+	GetToolManager()->BeginUndoTransaction(TransactionLabel);
 
-		// hide gizmo and selected triangles
-		MultiTransformer->SetGizmoVisibility(false);
-		DynamicMeshComponent->SetSecondaryBuffersVisibility(false);
-	}
-	return bOK;
-}
+	// Since we clear the selection in the selection mechanic when topology changes, we need to know
+	// when OutputSelection is pointing to the selection in the selection mechanic and is not empty,
+	// so that we can copy it ahead of time and reinstate it.
+	bool bReferencingSameSelection = (&SelectionMechanic->GetActiveSelection() == &OutputSelection);
 
+	bool bSelectionModified = !bReferencingSameSelection && SelectionMechanic->GetActiveSelection() != OutputSelection;
 
-void UEditMeshPolygonsTool::CompleteMeshEditChange(
-	const FText& TransactionLabel, 
-	TUniquePtr<FToolCommandChange> EditChange,
-	const FGroupTopologySelection& OutputSelection)
-{
-	// open top-level transaction
-	GetToolManager()->BeginUndoTransaction(TransactionLabel);
+	// In case we need to make a selection copy
+	FGroupTopologySelection TempSelection;
+	const FGroupTopologySelection* OutputSelectionToUse = &OutputSelection;
 
-	// clear current selection
-	SelectionMechanic->BeginChange();
-	SelectionMechanic->ClearSelection();
-	GetToolManager()->EmitObjectChange(SelectionMechanic, SelectionMechanic->EndChange(), LOCTEXT("PolyMeshExtrudeChangeClearSelection", "ClearSelection"));
+	// If the selection is going to be cleared, we need to do it explicitly ourselves so that we can emit a change.
+	if (!SelectionMechanic->GetActiveSelection().IsEmpty() && (bSelectionModified || bGroupTopologyModified))
+	{
+		if (bReferencingSameSelection)
+		{
+			// Need to make a copy because OutputSelection will get cleared due to bGroupTopologyModified
+			TempSelection = OutputSelection;
+			OutputSelectionToUse = &TempSelection;
+		}
 
-	// emit the pre-edit change
-	GetToolManager()->EmitObjectChange(this, MakeUnique<FEditPolygonsTopologyPreEditChange>(), LOCTEXT("PolyMeshExtrudeChangePreEdit", "PreEdit"));
+		SelectionMechanic->BeginChange();
+		SelectionMechanic->ClearSelection();
+		GetToolManager()->EmitObjectChange(SelectionMechanic, SelectionMechanic->EndChange(), LOCTEXT("PolyMeshExtrudeChangeClearSelection", "ClearSelection"));
+	}
 
-	// emit the mesh change
-	GetToolManager()->EmitObjectChange(DynamicMeshComponent, MoveTemp(EditChange), TransactionLabel);
+	GetToolManager()->EmitObjectChange(this, 
+		MakeUnique<FEditMeshPolygonsToolMeshChange>(MoveTemp(MeshChangeIn), bGroupTopologyModified),
+		TransactionLabel);
 
-	// emit the post-edit change
-	GetToolManager()->EmitObjectChange(this, MakeUnique<FEditPolygonsTopologyPostEditChange>(), TransactionLabel);
-	// call this (PostEditChange will do this)
-	AfterTopologyEdit();
-	// increment topology-change counter
-	ModifiedTopologyCounter++;
+	// Update related structures
+	UpdateFromCurrentMesh(bGroupTopologyModified);
+	ModifiedTopologyCounter += bGroupTopologyModified;
 
-	// set output selection
-	if (OutputSelection.IsEmpty() == false)
+	// Set output selection either if we changed selections (to something non-empty), or if
+	// our selection got cleared due to bGroupTopologyModified.
+	if (!OutputSelectionToUse->IsEmpty() && (bSelectionModified || bGroupTopologyModified))
 	{
 		SelectionMechanic->BeginChange();
-		SelectionMechanic->SetSelection(OutputSelection);
+		SelectionMechanic->SetSelection(*OutputSelectionToUse);
 		GetToolManager()->EmitObjectChange(SelectionMechanic, SelectionMechanic->EndChange(), LOCTEXT("PolyMeshExtrudeChangeSetSelection", "SetSelection"));
 	}
 
-	// complete the transaction
 	GetToolManager()->EndUndoTransaction();
+}
 
-	// clean up preview mesh, hiding of things, etc
-	if (EditPreview != nullptr)
-	{
-		EditPreview->Disconnect();
-		EditPreview = nullptr;
-	}
-	DynamicMeshComponent->SetSecondaryBuffersVisibility(true);
+void UEditMeshPolygonsTool::EmitActivityStart(const FText& TransactionLabel)
+{
+	++ActivityTimestamp;
 
-	CurrentOperationTimestamp++;
+	GetToolManager()->BeginUndoTransaction(TransactionLabel);
+	GetToolManager()->EmitObjectChange(this,
+		MakeUnique<FPolyEditActivityStartChange>(ActivityTimestamp),
+		TransactionLabel);
+	GetToolManager()->EndUndoTransaction();
 }
 
 
@@ -2395,8 +1815,6 @@ bool UEditMeshPolygonsTool::BeginMeshBoundaryEdgeEditChange(bool bOnlySimple)
 
 bool UEditMeshPolygonsTool::BeginMeshEdgeEditChange(TFunctionRef<bool(int32)> GroupEdgeIDFilterFunc)
 {
-	check(EditPreview == nullptr);
-
 	ActiveEdgeSelection.Reset();
 
 	const FGroupTopologySelection& ActiveSelection = SelectionMechanic->GetActiveSelection();
@@ -2419,71 +1837,6 @@ bool UEditMeshPolygonsTool::BeginMeshEdgeEditChange(TFunctionRef<bool(int32)> Gr
 	return ActiveEdgeSelection.Num() > 0;
 }
 
-
-
-void UEditMeshPolygonsTool::CancelMeshEditChange()
-{
-	if (EditPreview != nullptr)
-	{
-		EditPreview->Disconnect();
-		EditPreview = nullptr;
-	}
-	DynamicMeshComponent->SetSecondaryBuffersVisibility(true);
-
-	// disable any mechanics
-	ExtrudeHeightMechanic = nullptr;
-	CurveDistMechanic = nullptr;
-	SurfacePathMechanic = nullptr;
-
-	// hide properties that might be visible
-	SetToolPropertySourceEnabled(ExtrudeProperties, false);
-	SetToolPropertySourceEnabled(OffsetProperties, false);
-	SetToolPropertySourceEnabled(InsetProperties, false);
-	SetToolPropertySourceEnabled(OutsetProperties, false);
-	SetToolPropertySourceEnabled(CutProperties, false);
-	SetToolPropertySourceEnabled(SetUVProperties, false);
-	SetActionButtonPanelsVisible(true);
-
-	CurrentToolMode = ECurrentToolMode::TransformSelection;
-}
-
-
-
-
-void UEditMeshPolygonsTool::UpdateEditPreviewMaterials(EPreviewMaterialType MaterialType)
-{
-	if (EditPreview != nullptr)
-	{
-		if (MaterialType == EPreviewMaterialType::SourceMaterials)
-		{
-			EditPreview->ClearOverrideRenderMaterial();
-			EditPreview->SetMaterials(DynamicMeshComponent->GetMaterials());
-		}
-		else if (MaterialType == EPreviewMaterialType::PreviewMaterial)
-		{
-			EditPreview->ClearOverrideRenderMaterial();
-			EditPreview->SetMaterial(
-				ToolSetupUtil::GetSelectionMaterial(FLinearColor(0.8f, 0.75f, 0.0f), GetToolManager()));
-		}
-		else if (MaterialType == EPreviewMaterialType::UVMaterial)
-		{
-			UMaterial* CheckerMaterialBase = LoadObject<UMaterial>(nullptr, TEXT("/MeshModelingToolset/Materials/CheckerMaterial"));
-			if (CheckerMaterialBase != nullptr)
-			{
-				UMaterialInstanceDynamic* CheckerMaterial = UMaterialInstanceDynamic::Create(CheckerMaterialBase, NULL);
-				CheckerMaterial->SetScalarParameterValue("Density", 1);
-				EditPreview->SetOverrideRenderMaterial(CheckerMaterial);
-			}
-		}
-
-		CurrentPreviewMaterial = MaterialType;
-	}
-}
-
-
-
-
-
 void UEditMeshPolygonsTool::SetActionButtonPanelsVisible(bool bVisible)
 {
 	if (bTriangleMode == false)
@@ -2514,51 +1867,42 @@ void UEditMeshPolygonsTool::SetActionButtonPanelsVisible(bool bVisible)
 	}
 }
 
-
-
-
-void FEditPolygonsTopologyPreEditChange::Apply(UObject* Object)
-{
-}
-void FEditPolygonsTopologyPreEditChange::Revert(UObject* Object)
+void FEditMeshPolygonsToolMeshChange::Apply(UObject* Object)
 {
-	Cast<UEditMeshPolygonsTool>(Object)->AfterTopologyEdit();
-	Cast<UEditMeshPolygonsTool>(Object)->ModifiedTopologyCounter--;
-}
-FString FEditPolygonsTopologyPreEditChange::ToString() const 
-{ 
-	return TEXT("FEditPolygonsTopologyPreEditChange"); 
+	UEditMeshPolygonsTool* Tool = Cast<UEditMeshPolygonsTool>(Object);
+	
+	MeshChange->Apply(Tool->CurrentMesh.Get(), false);
+	Tool->UpdateFromCurrentMesh(bGroupTopologyModified);
+	Tool->ModifiedTopologyCounter += bGroupTopologyModified;
 }
 
-
-void FEditPolygonsTopologyPostEditChange::Apply(UObject* Object)
+void FEditMeshPolygonsToolMeshChange::Revert(UObject* Object)
 {
-	Cast<UEditMeshPolygonsTool>(Object)->AfterTopologyEdit();
-	Cast<UEditMeshPolygonsTool>(Object)->ModifiedTopologyCounter++;
+	UEditMeshPolygonsTool* Tool = Cast<UEditMeshPolygonsTool>(Object);
+	
+	MeshChange->Apply(Tool->CurrentMesh.Get(), true);
+	Tool->UpdateFromCurrentMesh(bGroupTopologyModified);
+	Tool->ModifiedTopologyCounter -= bGroupTopologyModified;
 }
-void FEditPolygonsTopologyPostEditChange::Revert(UObject* Object)
+
+FString FEditMeshPolygonsToolMeshChange::ToString() const
 {
+	return TEXT("FEditMeshPolygonsToolMeshChange");
 }
-FString FEditPolygonsTopologyPostEditChange::ToString() const 
-{ 
-	return TEXT("FEditPolygonsTopologyPostEditChange"); 
-}
-
 
-
-void FBeginInteractivePolyEditChange::Revert(UObject* Object)
+void FPolyEditActivityStartChange::Revert(UObject* Object)
 {
-	Cast<UEditMeshPolygonsTool>(Object)->CancelMeshEditChange();
+	Cast<UEditMeshPolygonsTool>(Object)->EndCurrentActivity();
 	bHaveDoneUndo = true;
 }
-bool FBeginInteractivePolyEditChange::HasExpired(UObject* Object) const
+bool FPolyEditActivityStartChange::HasExpired(UObject* Object) const
 {
-	return bHaveDoneUndo || (Cast<UEditMeshPolygonsTool>(Object)->CheckInOperation(OperationTimestamp) == false);
+	return bHaveDoneUndo 
+		|| Cast<UEditMeshPolygonsTool>(Object)->ActivityTimestamp != ActivityTimestamp;
 }
-FString FBeginInteractivePolyEditChange::ToString() const
+FString FPolyEditActivityStartChange::ToString() const
 {
-	return TEXT("FBeginInteractivePolyEditChange");
+	return TEXT("FPolyEditActivityStartChange");
 }
 
-
 #undef LOCTEXT_NAMESPACE
diff --git a/Engine/Plugins/Experimental/MeshModelingToolset/Source/MeshModelingTools/Private/GroupEdgeInsertionTool.cpp b/Engine/Plugins/Experimental/MeshModelingToolset/Source/MeshModelingTools/Private/GroupEdgeInsertionTool.cpp
deleted file mode 100644
diff --git a/Engine/Plugins/Experimental/MeshModelingToolset/Source/MeshModelingTools/Private/ToolActivities/PolyEditActivityUtil.cpp b/Engine/Plugins/Experimental/MeshModelingToolset/Source/MeshModelingTools/Private/ToolActivities/PolyEditActivityUtil.cpp
new file mode 100644
--- /dev/null
+++ b/Engine/Plugins/Experimental/MeshModelingToolset/Source/MeshModelingTools/Private/ToolActivities/PolyEditActivityUtil.cpp
@@ -0,0 +1,52 @@
+// Copyright Epic Games, Inc. All Rights Reserved.
+
+#include "ToolActivities/PolyEditActivityUtil.h"
+
+#include "Drawing/PolyEditPreviewMesh.h"
+#include "InteractiveTool.h"
+#include "Materials/MaterialInstanceDynamic.h"
+#include "Materials/MaterialInterface.h"
+#include "MeshOpPreviewHelpers.h"
+#include "ToolActivities/PolyEditActivityContext.h"
+#include "ToolSetupUtil.h"
+
+using namespace UE::Geometry;
+
+UPolyEditPreviewMesh* PolyEditActivityUtil::CreatePolyEditPreviewMesh(UInteractiveTool& Tool, const UPolyEditActivityContext& ActivityContext)
+{
+	UPolyEditPreviewMesh* EditPreview = NewObject<UPolyEditPreviewMesh>(&Tool);
+	EditPreview->CreateInWorld(Tool.GetWorld(), FTransform::Identity);
+	UpdatePolyEditPreviewMaterials(Tool, ActivityContext, *EditPreview, EPreviewMaterialType::PreviewMaterial);
+	EditPreview->EnableWireframe(true);
+
+	return EditPreview;
+}
+
+void PolyEditActivityUtil::UpdatePolyEditPreviewMaterials(UInteractiveTool& Tool, const UPolyEditActivityContext& ActivityContext,
+	UPolyEditPreviewMesh& EditPreview, EPreviewMaterialType MaterialType)
+{
+	if (MaterialType == EPreviewMaterialType::SourceMaterials)
+	{
+		TArray<UMaterialInterface*> Materials;
+		ActivityContext.Preview->PreviewMesh->GetMaterials(Materials);
+
+		EditPreview.ClearOverrideRenderMaterial();
+		EditPreview.SetMaterials(Materials);
+	}
+	else if (MaterialType == EPreviewMaterialType::PreviewMaterial)
+	{
+		EditPreview.ClearOverrideRenderMaterial();
+		EditPreview.SetMaterial(
+			ToolSetupUtil::GetSelectionMaterial(FLinearColor(0.8f, 0.75f, 0.0f), Tool.GetToolManager()));
+	}
+	else if (MaterialType == EPreviewMaterialType::UVMaterial)
+	{
+		UMaterial* CheckerMaterialBase = LoadObject<UMaterial>(nullptr, TEXT("/MeshModelingToolset/Materials/CheckerMaterial"));
+		if (CheckerMaterialBase != nullptr)
+		{
+			UMaterialInstanceDynamic* CheckerMaterial = UMaterialInstanceDynamic::Create(CheckerMaterialBase, NULL);
+			CheckerMaterial->SetScalarParameterValue("Density", 1);
+			EditPreview.SetOverrideRenderMaterial(CheckerMaterial);
+		}
+	}
+}
diff --git a/Engine/Plugins/Experimental/MeshModelingToolset/Source/MeshModelingTools/Private/ToolActivities/PolyEditCutFacesActivity.cpp b/Engine/Plugins/Experimental/MeshModelingToolset/Source/MeshModelingTools/Private/ToolActivities/PolyEditCutFacesActivity.cpp
new file mode 100644
--- /dev/null
+++ b/Engine/Plugins/Experimental/MeshModelingToolset/Source/MeshModelingTools/Private/ToolActivities/PolyEditCutFacesActivity.cpp
@@ -0,0 +1,254 @@
+// Copyright Epic Games, Inc. All Rights Reserved.
+
+#include "ToolActivities/PolyEditCutFacesActivity.h"
+
+#include "BaseBehaviors/SingleClickBehavior.h"
+#include "BaseBehaviors/MouseHoverBehavior.h"
+#include "ContextObjectStore.h"
+#include "Drawing/PolyEditPreviewMesh.h"
+#include "InteractiveToolManager.h"
+#include "Mechanics/CollectSurfacePathMechanic.h"
+#include "DynamicMesh/MeshIndexUtil.h" // TriangleToVertexIDs
+#include "MeshOpPreviewHelpers.h"
+#include "Operations/MeshPlaneCut.h"
+#include "Selections/MeshEdgeSelection.h"
+#include "Selection/PolygonSelectionMechanic.h"
+#include "ToolActivities/PolyEditActivityContext.h"
+#include "ToolActivities/PolyEditActivityUtil.h"
+#include "ToolSceneQueriesUtil.h"
+
+#define LOCTEXT_NAMESPACE "UPolyEditInsetOutsetActivity"
+
+using namespace UE::Geometry;
+
+void UPolyEditCutFacesActivity::Setup(UInteractiveTool* ParentToolIn)
+{
+	Super::Setup(ParentToolIn);
+
+	CutProperties = NewObject<UPolyEditCutProperties>();
+	CutProperties->RestoreProperties(ParentToolIn);
+	AddToolPropertySource(CutProperties);
+	SetToolPropertySourceEnabled(CutProperties, false);
+
+	// Register ourselves to receive clicks and hover
+	USingleClickInputBehavior* ClickBehavior = NewObject<USingleClickInputBehavior>();
+	ClickBehavior->Initialize(this);
+	ParentTool->AddInputBehavior(ClickBehavior);
+	UMouseHoverBehavior* HoverBehavior = NewObject<UMouseHoverBehavior>();
+	HoverBehavior->Initialize(this);
+	ParentTool->AddInputBehavior(HoverBehavior);
+
+	ActivityContext = ParentTool->GetToolManager()->GetContextObjectStore()->FindContext<UPolyEditActivityContext>();
+}
+
+void UPolyEditCutFacesActivity::Shutdown(EToolShutdownType ShutdownType)
+{
+	Clear();
+	CutProperties->SaveProperties(ParentTool.Get());
+
+	CutProperties = nullptr;
+	ParentTool = nullptr;
+	ActivityContext = nullptr;
+}
+
+bool UPolyEditCutFacesActivity::CanStart() const
+{
+	if (!ActivityContext)
+	{
+		return false;
+	}
+	const FGroupTopologySelection& Selection = ActivityContext->SelectionMechanic->GetActiveSelection();
+	return !Selection.SelectedGroupIDs.IsEmpty();
+}
+
+EToolActivityStartResult UPolyEditCutFacesActivity::Start()
+{
+	if (!CanStart())
+	{
+		ParentTool->GetToolManager()->DisplayMessage(
+			LOCTEXT("OnCutFacesFailedMessage", "Cannot cut without face selection."),
+			EToolMessageLevel::UserWarning);
+		return EToolActivityStartResult::FailedStart;
+	}
+
+	Clear();
+	BeginCutFaces();
+	bIsRunning = true;
+
+	ActivityContext->EmitActivityStart(LOCTEXT("BeginCutFacesActivity", "Begin Cut Faces"));
+
+	return EToolActivityStartResult::Running;
+}
+
+bool UPolyEditCutFacesActivity::CanAccept() const
+{
+	return false;
+}
+
+EToolActivityEndResult UPolyEditCutFacesActivity::End(EToolShutdownType)
+{
+	Clear();
+	EToolActivityEndResult ToReturn = bIsRunning ? EToolActivityEndResult::Cancelled 
+		: EToolActivityEndResult::ErrorDuringEnd;
+	bIsRunning = false;
+	return ToReturn;
+}
+
+void UPolyEditCutFacesActivity::BeginCutFaces()
+{
+	const FGroupTopologySelection& ActiveSelection = ActivityContext->SelectionMechanic->GetActiveSelection();
+	TArray<int32> ActiveTriangleSelection;
+	ActivityContext->CurrentTopology->GetSelectedTriangles(ActiveSelection, ActiveTriangleSelection);
+
+	FTransform3d WorldTransform(ActivityContext->Preview->PreviewMesh->GetTransform());
+
+	EditPreview = PolyEditActivityUtil::CreatePolyEditPreviewMesh(*ParentTool, *ActivityContext);
+	EditPreview->InitializeStaticType(ActivityContext->CurrentMesh.Get(), ActiveTriangleSelection, &WorldTransform);
+
+	FDynamicMesh3 StaticHitTargetMesh;
+	EditPreview->MakeInsetTypeTargetMesh(StaticHitTargetMesh);
+
+	SurfacePathMechanic = NewObject<UCollectSurfacePathMechanic>(this);
+	SurfacePathMechanic->Setup(ParentTool.Get());
+	SurfacePathMechanic->InitializeMeshSurface(MoveTemp(StaticHitTargetMesh));
+	SurfacePathMechanic->SetFixedNumPointsMode(2);
+	SurfacePathMechanic->bSnapToTargetMeshVertices = true;
+	double SnapTol = ToolSceneQueriesUtil::GetDefaultVisualAngleSnapThreshD();
+	SurfacePathMechanic->SpatialSnapPointsFunc = [this, SnapTol](FVector3d Position1, FVector3d Position2)
+	{
+		return CutProperties->bSnapToVertices && ToolSceneQueriesUtil::PointSnapQuery(this->CameraState, Position1, Position2, SnapTol);
+	};
+
+	SetToolPropertySourceEnabled(CutProperties, true);
+}
+
+void UPolyEditCutFacesActivity::ApplyCutFaces()
+{
+	check(SurfacePathMechanic != nullptr && EditPreview != nullptr);
+
+	const FGroupTopologySelection& ActiveSelection = ActivityContext->SelectionMechanic->GetActiveSelection();
+	TArray<int32> ActiveTriangleSelection;
+	ActivityContext->CurrentTopology->GetSelectedTriangles(ActiveSelection, ActiveTriangleSelection);
+
+	// construct cut plane normal from line points
+	FFrame3d Point0(SurfacePathMechanic->HitPath[0]), Point1(SurfacePathMechanic->HitPath[1]);
+	FVector3d PlaneNormal;
+	if (CutProperties->Orientation == EPolyEditCutPlaneOrientation::ViewDirection)
+	{
+		FVector3d Direction0 = UE::Geometry::Normalized(Point0.Origin - (FVector3d)CameraState.Position);
+		FVector3d Direction1 = UE::Geometry::Normalized(Point1.Origin - (FVector3d)CameraState.Position);
+		PlaneNormal = Direction1.Cross(Direction0);
+	}
+	else
+	{
+		FVector3d LineDirection = UE::Geometry::Normalized(Point1.Origin - Point0.Origin);
+		FVector3d UpVector = UE::Geometry::Normalized(Point0.Z() + Point1.Z());
+		PlaneNormal = LineDirection.Cross(UpVector);
+	}
+	FVector3d PlaneOrigin = 0.5 * (Point0.Origin + Point1.Origin);
+	// map into local space of target mesh
+	FTransform3d WorldTransform(ActivityContext->Preview->PreviewMesh->GetTransform());
+	PlaneOrigin = WorldTransform.InverseTransformPosition(PlaneOrigin);
+	PlaneNormal = WorldTransform.InverseTransformNormal(PlaneNormal);
+	UE::Geometry::Normalize(PlaneNormal);
+
+	// track changes
+	FDynamicMeshChangeTracker ChangeTracker(ActivityContext->CurrentMesh.Get());
+	ChangeTracker.BeginChange();
+	TArray<int32> VertexSelection;
+	UE::Geometry::TriangleToVertexIDs(ActivityContext->CurrentMesh.Get(), ActiveTriangleSelection, VertexSelection);
+	ChangeTracker.SaveVertexOneRingTriangles(VertexSelection, true);
+
+	// apply the cut to edges of selected triangles
+	FGroupTopologySelection OutputSelection;
+	FMeshPlaneCut Cut(ActivityContext->CurrentMesh.Get(), PlaneOrigin, PlaneNormal);
+	FMeshEdgeSelection Edges(ActivityContext->CurrentMesh.Get());
+	Edges.SelectTriangleEdges(ActiveTriangleSelection);
+	Cut.EdgeFilterFunc = [&](int EdgeID) { return Edges.IsSelected(EdgeID); };
+	if (Cut.SplitEdgesOnly(true))
+	{
+		if (!ActivityContext->bTriangleMode)
+		{
+			for (FMeshPlaneCut::FCutResultRegion& Region : Cut.ResultRegions)
+			{
+				OutputSelection.SelectedGroupIDs.Add(Region.GroupID);
+			}
+		}
+		else
+		{
+			// Retain the selection along the cut. ResultSeedTriangles does not
+			// contain selected tris that are not cut, so re-add the original selected
+			// tris.
+			OutputSelection.SelectedGroupIDs.Append(Cut.ResultSeedTriangles);
+			OutputSelection.SelectedGroupIDs.Append(ActiveTriangleSelection);
+		}
+	}
+
+	// Emit undo (also updates relevant structures)
+	ActivityContext->EmitCurrentMeshChangeAndUpdate(LOCTEXT("PolyMeshCutFacesChange", "Cut Faces"),
+		ChangeTracker.EndChange(), OutputSelection, true);
+
+	// End activity
+	Clear();
+	bIsRunning = false;
+	Cast<IToolActivityHost>(ParentTool)->NotifyActivitySelfEnded(this);
+}
+
+void UPolyEditCutFacesActivity::Clear()
+{
+	if (EditPreview != nullptr)
+	{
+		EditPreview->Disconnect();
+		EditPreview = nullptr;
+	}
+
+	SurfacePathMechanic = nullptr;
+	SetToolPropertySourceEnabled(CutProperties, false);
+}
+
+void UPolyEditCutFacesActivity::Render(IToolsContextRenderAPI* RenderAPI)
+{
+	ParentTool->GetToolManager()->GetContextQueriesAPI()->GetCurrentViewState(CameraState);
+
+	if (SurfacePathMechanic != nullptr)
+	{
+		SurfacePathMechanic->Render(RenderAPI);
+	}
+}
+
+void UPolyEditCutFacesActivity::Tick(float DeltaTime)
+{
+}
+
+FInputRayHit UPolyEditCutFacesActivity::IsHitByClick(const FInputDeviceRay& ClickPos)
+{
+	FInputRayHit OutHit;
+	OutHit.bHit = bIsRunning;
+	return OutHit;
+}
+
+void UPolyEditCutFacesActivity::OnClicked(const FInputDeviceRay& ClickPos)
+{
+	if (bIsRunning && SurfacePathMechanic->TryAddPointFromRay(ClickPos.WorldRay))
+	{
+		if (SurfacePathMechanic->IsDone())
+		{
+			ApplyCutFaces();
+		}
+	}
+}
+
+FInputRayHit UPolyEditCutFacesActivity::BeginHoverSequenceHitTest(const FInputDeviceRay& PressPos)
+{
+	FInputRayHit OutHit;
+	OutHit.bHit = bIsRunning;
+	return OutHit;
+}
+
+bool UPolyEditCutFacesActivity::OnUpdateHover(const FInputDeviceRay& DevicePos)
+{
+	SurfacePathMechanic->UpdatePreviewPoint(DevicePos.WorldRay);
+	return bIsRunning;
+}
+
+#undef LOCTEXT_NAMESPACE
diff --git a/Engine/Plugins/Experimental/MeshModelingToolset/Source/MeshModelingTools/Private/ToolActivities/PolyEditExtrudeActivity.cpp b/Engine/Plugins/Experimental/MeshModelingToolset/Source/MeshModelingTools/Private/ToolActivities/PolyEditExtrudeActivity.cpp
new file mode 100644
--- /dev/null
+++ b/Engine/Plugins/Experimental/MeshModelingToolset/Source/MeshModelingTools/Private/ToolActivities/PolyEditExtrudeActivity.cpp
@@ -0,0 +1,311 @@
+// Copyright Epic Games, Inc. All Rights Reserved.
+
+#include "ToolActivities/PolyEditExtrudeActivity.h"
+
+#include "BaseBehaviors/SingleClickBehavior.h"
+#include "BaseBehaviors/MouseHoverBehavior.h"
+#include "ContextObjectStore.h"
+#include "Drawing/PolyEditPreviewMesh.h"
+#include "DynamicMesh/MeshNormals.h"
+#include "EditMeshPolygonsTool.h"
+#include "InteractiveToolManager.h"
+#include "Mechanics/PlaneDistanceFromHitMechanic.h"
+#include "MeshOpPreviewHelpers.h"
+#include "Operations/OffsetMeshRegion.h"
+#include "Selection/PolygonSelectionMechanic.h"
+#include "ToolActivities/PolyEditActivityContext.h"
+#include "ToolActivities/PolyEditActivityUtil.h"
+#include "ToolSceneQueriesUtil.h"
+
+#define LOCTEXT_NAMESPACE "UPolyEditExtrudeActivity"
+
+using namespace UE::Geometry;
+
+void UPolyEditExtrudeActivity::Setup(UInteractiveTool* ParentToolIn)
+{
+	Super::Setup(ParentToolIn);
+
+	ExtrudeProperties = NewObject<UPolyEditExtrudeProperties>();
+	ExtrudeProperties->RestoreProperties(ParentTool.Get());
+	AddToolPropertySource(ExtrudeProperties);
+	SetToolPropertySourceEnabled(ExtrudeProperties, false);
+	ExtrudeProperties->WatchProperty(ExtrudeProperties->Direction,
+		[this](EPolyEditExtrudeDirection) { 
+			Clear();
+			BeginExtrude();
+		});
+	ExtrudeProperties->WatchProperty(ExtrudeProperties->ExtrudeMode,
+		[this](EPolyEditExtrudeMode) {
+			Clear();
+			BeginExtrude();
+		});
+
+	// Register ourselves to receive clicks and hover
+	USingleClickInputBehavior* ClickBehavior = NewObject<USingleClickInputBehavior>();
+	ClickBehavior->Initialize(this);
+	ParentTool->AddInputBehavior(ClickBehavior);
+	UMouseHoverBehavior* HoverBehavior = NewObject<UMouseHoverBehavior>();
+	HoverBehavior->Initialize(this);
+	ParentTool->AddInputBehavior(HoverBehavior);
+
+	ActivityContext = ParentTool->GetToolManager()->GetContextObjectStore()->FindContext<UPolyEditActivityContext>();
+}
+
+void UPolyEditExtrudeActivity::Shutdown(EToolShutdownType ShutdownType)
+{
+	Clear();
+	ExtrudeProperties->SaveProperties(ParentTool.Get());
+
+	ExtrudeProperties = nullptr;
+	ParentTool = nullptr;
+	ActivityContext = nullptr;
+}
+
+bool UPolyEditExtrudeActivity::CanStart() const
+{
+	if (!ActivityContext)
+	{
+		return false;
+	}
+	const FGroupTopologySelection& Selection = ActivityContext->SelectionMechanic->GetActiveSelection();
+	return !Selection.SelectedGroupIDs.IsEmpty();
+}
+
+EToolActivityStartResult UPolyEditExtrudeActivity::Start()
+{
+	if (!CanStart())
+	{
+		ParentTool->GetToolManager()->DisplayMessage(
+			LOCTEXT("OnExtrudeFailedMesssage", "Cannot extrude without face selection."),
+			EToolMessageLevel::UserWarning);
+		return EToolActivityStartResult::FailedStart;
+	}
+
+	Clear();
+	BeginExtrude();
+	bIsRunning = true;
+
+	ActivityContext->EmitActivityStart(LOCTEXT("BeginExtrudeActivity", "Begin Extrude"));
+
+	return EToolActivityStartResult();
+}
+
+bool UPolyEditExtrudeActivity::CanAccept() const
+{
+	return false;
+}
+
+EToolActivityEndResult UPolyEditExtrudeActivity::End(EToolShutdownType)
+{
+	Clear();
+	EToolActivityEndResult ToReturn = bIsRunning ? EToolActivityEndResult::Cancelled : EToolActivityEndResult::ErrorDuringEnd;
+	bIsRunning = false;
+	return ToReturn;
+}
+
+
+void UPolyEditExtrudeActivity::BeginExtrude()
+{
+	const FGroupTopologySelection& ActiveSelection = ActivityContext->SelectionMechanic->GetActiveSelection();
+	TArray<int32> ActiveTriangleSelection;
+	ActivityContext->CurrentTopology->GetSelectedTriangles(ActiveSelection, ActiveTriangleSelection);
+
+	FTransform3d WorldTransform(ActivityContext->Preview->PreviewMesh->GetTransform());
+
+	// Get the world frame
+	FFrame3d ActiveSelectionFrameLocal = ActivityContext->CurrentTopology->GetSelectionFrame(ActiveSelection);
+	ActiveSelectionFrameWorld = ActiveSelectionFrameLocal;
+	ActiveSelectionFrameWorld.Transform(WorldTransform);
+	ActiveSelectionFrameWorld.AlignAxis(2, GetExtrudeDirection());
+
+	// Set up a preview of the extruded portion of the mesh
+	EditPreview = PolyEditActivityUtil::CreatePolyEditPreviewMesh(*ParentTool, *ActivityContext);
+	EditPreview->InitializeExtrudeType(ActivityContext->CurrentMesh.Get(), ActiveTriangleSelection, ActiveSelectionFrameWorld.Z(), &WorldTransform, true);
+	// move world extrude frame to point on surface
+	ActiveSelectionFrameWorld.Origin = EditPreview->GetInitialPatchMeshSpatial().FindNearestPoint(ActiveSelectionFrameWorld.Origin);
+
+	// Hide the selected triangles (that are being replaced by the extruded portion)
+	ActivityContext->Preview->PreviewMesh->SetSecondaryBuffersVisibility(false);
+
+	// Set up the mechanic we use to determine how far to extrude
+	ExtrudeHeightMechanic = NewObject<UPlaneDistanceFromHitMechanic>(this);
+	ExtrudeHeightMechanic->Setup(ParentTool.Get());
+	ExtrudeHeightMechanic->WorldHitQueryFunc = [this](const FRay& WorldRay, FHitResult& HitResult)
+	{
+		return ToolSceneQueriesUtil::FindNearestVisibleObjectHit(ActivityContext->Preview->GetWorld(), HitResult, WorldRay);
+	};
+	ExtrudeHeightMechanic->WorldPointSnapFunc = [this](const FVector3d& WorldPos, FVector3d& SnapPos)
+	{
+		return ActivityContext->CommonProperties->bSnapToWorldGrid && ToolSceneQueriesUtil::FindWorldGridSnapPoint(ParentTool.Get(), WorldPos, SnapPos);
+	};
+	ExtrudeHeightMechanic->CurrentHeight = 1.0f;  // initialize to something non-zero...prob should be based on polygon bounds maybe?
+
+	// make inifinite-extent hit-test mesh to use in the mechanic
+	FDynamicMesh3 ExtrudeHitTargetMesh;
+	EditPreview->MakeExtrudeTypeHitTargetMesh(ExtrudeHitTargetMesh);
+	ExtrudeHeightMechanic->Initialize(MoveTemp(ExtrudeHitTargetMesh), ActiveSelectionFrameWorld, true);
+
+	SetToolPropertySourceEnabled(ExtrudeProperties, true);
+
+	float BoundsMaxDim = ActivityContext->CurrentMesh->GetBounds().MaxDim();
+	if (BoundsMaxDim > 0)
+	{
+		UVScaleFactor = 1.0 / BoundsMaxDim;
+	}
+
+	bPreviewUpdatePending = true;
+}
+
+void UPolyEditExtrudeActivity::ApplyExtrude()
+{
+	check(ExtrudeHeightMechanic != nullptr && EditPreview != nullptr);
+
+	const FGroupTopologySelection& ActiveSelection = ActivityContext->SelectionMechanic->GetActiveSelection();
+	TArray<int32> ActiveTriangleSelection;
+	ActivityContext->CurrentTopology->GetSelectedTriangles(ActiveSelection, ActiveTriangleSelection);
+
+	FTransform3d WorldTransform(ActivityContext->Preview->PreviewMesh->GetTransform());
+	FVector3d MeshSpaceExtrudeDirection = WorldTransform.InverseTransformVector(ActiveSelectionFrameWorld.Z());
+	double ExtrudeDistance = ExtrudeHeightMechanic->CurrentHeight;
+
+	FOffsetMeshRegion Extruder(ActivityContext->CurrentMesh.Get());
+	Extruder.UVScaleFactor = UVScaleFactor;
+	Extruder.Triangles = ActiveTriangleSelection;
+	TSet<int32> TriangleSet(ActiveTriangleSelection);
+	bool bUseNormals = (ExtrudeProperties->ExtrudeMode != EPolyEditExtrudeMode::SingleDirection);
+	Extruder.OffsetPositionFunc = [ExtrudeDistance, bUseNormals, &MeshSpaceExtrudeDirection](const FVector3d& Pos, const FVector3f& Normal, int32 VertexID) {
+		return Pos + ExtrudeDistance * (bUseNormals ? (FVector3d)Normal : MeshSpaceExtrudeDirection);
+	};
+	Extruder.bIsPositiveOffset = (ExtrudeDistance > 0);
+	Extruder.bUseFaceNormals = (ExtrudeProperties->ExtrudeMode == EPolyEditExtrudeMode::SelectedTriangleNormals);
+	Extruder.bOffsetFullComponentsAsSolids = ExtrudeProperties->bShellsToSolids;
+	Extruder.ChangeTracker = MakeUnique<FDynamicMeshChangeTracker>(ActivityContext->CurrentMesh.Get());
+	Extruder.ChangeTracker->BeginChange();
+	Extruder.Apply();
+
+	FMeshNormals::QuickComputeVertexNormalsForTriangles(*ActivityContext->CurrentMesh, Extruder.AllModifiedTriangles);
+
+	// construct new selection
+	FGroupTopologySelection NewSelection;
+	if (!ActivityContext->bTriangleMode)
+	{
+		for (const FOffsetMeshRegion::FOffsetInfo& Info : Extruder.OffsetRegions)
+		{
+			NewSelection.SelectedGroupIDs.Append(Info.OffsetGroups);
+		}
+	}
+	else
+	{
+		for (const FOffsetMeshRegion::FOffsetInfo& Info : Extruder.OffsetRegions)
+		{
+			NewSelection.SelectedGroupIDs.Append(Info.InitialTriangles);
+		}
+	}
+
+	// Emit undo  (also updates relevant structures)
+	ActivityContext->EmitCurrentMeshChangeAndUpdate(LOCTEXT("PolyMeshExtrudeChange", "Extrude"),
+		Extruder.ChangeTracker->EndChange(), NewSelection, true);
+
+	// End activity
+	Clear();
+	bIsRunning = false;
+	Cast<IToolActivityHost>(ParentTool)->NotifyActivitySelfEnded(this);
+}
+
+void UPolyEditExtrudeActivity::Clear()
+{
+	if (EditPreview != nullptr)
+	{
+		EditPreview->Disconnect();
+		EditPreview = nullptr;
+	}
+
+	ActivityContext->Preview->PreviewMesh->SetSecondaryBuffersVisibility(true);
+
+	ExtrudeHeightMechanic = nullptr;
+	SetToolPropertySourceEnabled(ExtrudeProperties, false);
+}
+
+
+FVector3d UPolyEditExtrudeActivity::GetExtrudeDirection() const
+{
+	switch (ExtrudeProperties->Direction)
+	{
+	default:
+	case EPolyEditExtrudeDirection::SelectionNormal:
+		return ActiveSelectionFrameWorld.Z();
+	case EPolyEditExtrudeDirection::WorldX:
+		return FVector3d::UnitX();
+	case EPolyEditExtrudeDirection::WorldY:
+		return FVector3d::UnitY();
+	case EPolyEditExtrudeDirection::WorldZ:
+		return FVector3d::UnitZ();
+	case EPolyEditExtrudeDirection::LocalX:
+		return FTransform3d(ActivityContext->Preview->PreviewMesh->GetTransform()).GetRotation().AxisX();
+	case EPolyEditExtrudeDirection::LocalY:
+		return FTransform3d(ActivityContext->Preview->PreviewMesh->GetTransform()).GetRotation().AxisY();
+	case EPolyEditExtrudeDirection::LocalZ:
+		return FTransform3d(ActivityContext->Preview->PreviewMesh->GetTransform()).GetRotation().AxisZ();
+	}
+	return ActiveSelectionFrameWorld.Z();
+}
+
+void UPolyEditExtrudeActivity::Render(IToolsContextRenderAPI* RenderAPI)
+{
+	if (ExtrudeHeightMechanic != nullptr)
+	{
+		ExtrudeHeightMechanic->Render(RenderAPI);
+	}
+}
+
+void UPolyEditExtrudeActivity::Tick(float DeltaTime)
+{
+	if (EditPreview && bPreviewUpdatePending)
+	{
+		switch (ExtrudeProperties->ExtrudeMode)
+		{
+		case EPolyEditExtrudeMode::SingleDirection:
+			EditPreview->UpdateExtrudeType(ExtrudeHeightMechanic->CurrentHeight);
+			break;
+		case EPolyEditExtrudeMode::SelectedTriangleNormals:
+			EditPreview->UpdateExtrudeType_FaceNormalAvg(ExtrudeHeightMechanic->CurrentHeight);
+			break;
+		case EPolyEditExtrudeMode::VertexNormals:
+			EditPreview->UpdateExtrudeType(ExtrudeHeightMechanic->CurrentHeight, true);
+			break;
+		}
+
+		bPreviewUpdatePending = false;
+	}
+}
+
+FInputRayHit UPolyEditExtrudeActivity::IsHitByClick(const FInputDeviceRay& ClickPos)
+{
+	FInputRayHit OutHit;
+	OutHit.bHit = bIsRunning;
+	return OutHit;
+}
+
+void UPolyEditExtrudeActivity::OnClicked(const FInputDeviceRay& ClickPos)
+{
+	if (bIsRunning)
+	{
+		ApplyExtrude();
+	}
+}
+
+FInputRayHit UPolyEditExtrudeActivity::BeginHoverSequenceHitTest(const FInputDeviceRay& PressPos)
+{
+	FInputRayHit OutHit;
+	OutHit.bHit = bIsRunning;
+	return OutHit;
+}
+
+bool UPolyEditExtrudeActivity::OnUpdateHover(const FInputDeviceRay& DevicePos)
+{
+	ExtrudeHeightMechanic->UpdateCurrentDistance(DevicePos.WorldRay);
+	bPreviewUpdatePending = true;
+	return bIsRunning;
+}
+
+#undef LOCTEXT_NAMESPACE
diff --git a/Engine/Plugins/Experimental/MeshModelingToolset/Source/MeshModelingTools/Private/ToolActivities/PolyEditInsertEdgeActivity.cpp b/Engine/Plugins/Experimental/MeshModelingToolset/Source/MeshModelingTools/Private/ToolActivities/PolyEditInsertEdgeActivity.cpp
new file mode 100644
--- /dev/null
+++ b/Engine/Plugins/Experimental/MeshModelingToolset/Source/MeshModelingTools/Private/ToolActivities/PolyEditInsertEdgeActivity.cpp
@@ -0,0 +1,860 @@
+// Copyright Epic Games, Inc. All Rights Reserved.
+
+#include "ToolActivities/PolyEditInsertEdgeActivity.h"
+
+#include "BaseBehaviors/SingleClickBehavior.h"
+#include "BaseBehaviors/MouseHoverBehavior.h"
+#include "ContextObjectStore.h"
+#include "CuttingOps/GroupEdgeInsertionOp.h"
+#include "DynamicMesh/DynamicMeshChangeTracker.h"
+#include "InteractiveToolManager.h"
+#include "MeshOpPreviewHelpers.h"
+#include "Selection/PolygonSelectionMechanic.h"
+#include "ToolActivities/PolyEditActivityContext.h"
+
+#include "ExplicitUseGeometryMathTypes.h"		// using UE::Geometry::(math types)
+using namespace UE::Geometry;
+
+#define LOCTEXT_NAMESPACE "UPolyEditInsertEdgeActivity"
+
+namespace PolyEditInsertEdgeActivityLocals
+{
+	bool GetSharedBoundary(const FGroupTopology& Topology,
+		const FGroupEdgeInserter::FGroupEdgeSplitPoint& StartPoint, int32 StartTopologyID, bool bStartIsCorner,
+		const FGroupEdgeInserter::FGroupEdgeSplitPoint& EndPoint, int32 EndTopologyID, bool bEndIsCorner,
+		int32& GroupIDOut, int32& BoundaryIndexOut);
+	bool DoesBoundaryContainPoint(const FGroupTopology& Topology,
+		const FGroupTopology::FGroupBoundary& Boundary, int32 PointTopologyID, bool bPointIsCorner);
+}
+
+TUniquePtr<FDynamicMeshOperator> UPolyEditInsertEdgeActivity::MakeNewOperator()
+{
+	TUniquePtr<FGroupEdgeInsertionOp> Op = MakeUnique<FGroupEdgeInsertionOp>();
+
+	Op->OriginalMesh = ActivityContext->CurrentMesh;
+	Op->OriginalTopology = ActivityContext->CurrentTopology;
+	Op->SetTransform(TargetTransform);
+
+	if (Settings->InsertionMode == EGroupEdgeInsertionMode::PlaneCut)
+	{
+		Op->Mode = FGroupEdgeInserter::EInsertionMode::PlaneCut;
+	}
+	else
+	{
+		Op->Mode = FGroupEdgeInserter::EInsertionMode::Retriangulate;
+	}
+
+	Op->VertexTolerance = Settings->VertexTolerance;
+
+	Op->StartPoint = StartPoint;
+	Op->EndPoint = EndPoint;
+	Op->CommonGroupID = CommonGroupID;
+	Op->CommonBoundaryIndex = CommonBoundaryIndex;
+
+	return Op;
+}
+
+void UPolyEditInsertEdgeActivity::Setup(UInteractiveTool* ParentToolIn)
+{
+	Super::Setup(ParentToolIn);
+
+	// Set up properties
+	Settings = NewObject<UGroupEdgeInsertionProperties>(this);
+	Settings->RestoreProperties(ParentTool.Get());
+	AddToolPropertySource(Settings);
+	SetToolPropertySourceEnabled(Settings, false);
+	Settings->GetOnModified().AddUObject(this, &UPolyEditInsertEdgeActivity::OnPropertyModified);
+
+	// These draw the group edges and the loops to be inserted
+	ExistingEdgesRenderer.LineColor = FLinearColor::Red;
+	ExistingEdgesRenderer.LineThickness = 2.0;
+	PreviewEdgeRenderer.LineColor = FLinearColor::Green;
+	PreviewEdgeRenderer.LineThickness = 4.0;
+	PreviewEdgeRenderer.PointColor = FLinearColor::Green;
+	PreviewEdgeRenderer.PointSize = 8.0;
+	PreviewEdgeRenderer.bDepthTested = false;
+
+	// Set up the topology selector settings
+	TopologySelectorSettings.bEnableEdgeHits = true;
+	TopologySelectorSettings.bEnableCornerHits = true;
+	TopologySelectorSettings.bEnableFaceHits = false;
+
+	// Set up our input routing
+	USingleClickInputBehavior* ClickBehavior = NewObject<USingleClickInputBehavior>();
+	ClickBehavior->Initialize(this);
+	ParentTool->AddInputBehavior(ClickBehavior);
+	UMouseHoverBehavior* HoverBehavior = NewObject<UMouseHoverBehavior>();
+	HoverBehavior->Initialize(this);
+	ParentTool->AddInputBehavior(HoverBehavior);
+
+	ActivityContext = ParentTool->GetToolManager()->GetContextObjectStore()->FindContext<UPolyEditActivityContext>();
+
+	TopologySelector = ActivityContext->SelectionMechanic->GetTopologySelector();
+}
+
+void UPolyEditInsertEdgeActivity::Shutdown(EToolShutdownType ShutdownType)
+{
+	if (bIsRunning)
+	{
+		End(ShutdownType);
+	}
+
+	Settings->SaveProperties(ParentTool.Get());
+	Settings->GetOnModified().Clear();
+
+	TopologySelector.Reset();
+	Settings = nullptr;
+	ActivityContext = nullptr;
+
+	Super::Shutdown(ShutdownType);
+}
+
+EToolActivityStartResult UPolyEditInsertEdgeActivity::Start()
+{
+	ParentTool->GetToolManager()->DisplayMessage(
+		LOCTEXT("InsertEdgeActivityDescription", "Click two points on the boundary of a face to "
+			"insert a new edge between the points and split the face."),
+		EToolMessageLevel::UserNotification);
+
+	SetToolPropertySourceEnabled(Settings, true);
+
+	// We don't use selection, so clear it if necessary (have to issue an undo/redo event)
+	if (!ActivityContext->SelectionMechanic->GetActiveSelection().IsEmpty())
+	{
+		ActivityContext->SelectionMechanic->BeginChange();
+		ActivityContext->SelectionMechanic->ClearSelection();
+		ParentTool->GetToolManager()->EmitObjectChange(ActivityContext->SelectionMechanic,
+			ActivityContext->SelectionMechanic->EndChange(), LOCTEXT("ClearSelection", "Clear Selection"));
+	}
+
+	TargetTransform = ActivityContext->Preview->PreviewMesh->GetTransform();
+
+	SetupPreview();
+
+	ToolState = EState::GettingStart;
+	bLastComputeSucceeded = false;
+
+	bIsRunning = true;
+
+	// Emit activity start transaction
+	ActivityContext->EmitActivityStart(LOCTEXT("BeginInsertEdgeActivity", "Begin Insert Edge"));
+
+	return EToolActivityStartResult::Running;
+}
+
+EToolActivityEndResult UPolyEditInsertEdgeActivity::End(EToolShutdownType ShutdownType)
+{
+	if (!bIsRunning)
+	{
+		return EToolActivityEndResult::ErrorDuringEnd;
+	}
+
+	SetToolPropertySourceEnabled(Settings, false);
+	Settings->GetOnModified().Clear();
+
+	ActivityContext->Preview->OnOpCompleted.RemoveAll(this);
+	ActivityContext->Preview->OnMeshUpdated.RemoveAll(this);
+	ActivityContext->Preview->ClearOpFactory();
+
+	LatestOpTopologyResult.Reset();
+	LatestOpChangedTids.Reset();
+
+	bIsRunning = false;
+
+	return CanAccept() ? EToolActivityEndResult::Completed : EToolActivityEndResult::Cancelled;
+}
+
+void UPolyEditInsertEdgeActivity::SetupPreview()
+{
+	ActivityContext->Preview->ChangeOpFactory(this);
+
+	// Whenever we get a new result from the op, we need to extract the preview edges so that
+	// we can draw them if we want to.
+	ActivityContext->Preview->OnOpCompleted.AddWeakLambda(this, [this](const FDynamicMeshOperator* UncastOp) {
+		const FGroupEdgeInsertionOp* Op = static_cast<const FGroupEdgeInsertionOp*>(UncastOp);
+
+		bLastComputeSucceeded = Op->bSucceeded;
+		LatestOpTopologyResult.Reset();
+		LatestOpChangedTids.Reset();
+
+		PreviewEdges.Reset();
+		if (bLastComputeSucceeded)
+		{
+			Op->GetEdgeLocations(PreviewEdges);
+			LatestOpTopologyResult = Op->ResultTopology;
+			LatestOpChangedTids = Op->ChangedTids;
+		}
+	});
+
+	// In case of failure, we want to hide the broken preview, since we wouldn't accept it on
+	// a click. Note that this can't be fired OnOpCompleted because the preview is updated
+	// with the op result after that callback, which would undo the reset. The preview edge
+	// extraction can't be lumped in here because it needs the op rather than the preview object.
+	ActivityContext->Preview->OnMeshUpdated.AddWeakLambda(this, [this](UMeshOpPreviewWithBackgroundCompute*) {
+		if (!bLastComputeSucceeded)
+		{
+			ActivityContext->Preview->PreviewMesh->UpdatePreview(ActivityContext->CurrentMesh.Get());
+		}
+	});
+}
+
+bool UPolyEditInsertEdgeActivity::CanStart() const
+{
+	return true;
+}
+
+void UPolyEditInsertEdgeActivity::Tick(float DeltaTime)
+{
+	if (ToolState == EState::WaitingForInsertComplete && ActivityContext->Preview->HaveValidResult())
+	{
+		if (bLastComputeSucceeded)
+		{
+			FDynamicMeshChangeTracker ChangeTracker(ActivityContext->CurrentMesh.Get());
+			ChangeTracker.BeginChange();
+			ChangeTracker.SaveTriangles(*LatestOpChangedTids, true /*bSaveVertices*/);
+
+			// Update current mesh
+			ActivityContext->CurrentMesh->Copy(*ActivityContext->Preview->PreviewMesh->GetMesh(), 
+				true, true, true, true);
+			
+			*ActivityContext->CurrentTopology = *LatestOpTopologyResult;
+			ActivityContext->CurrentTopology->RetargetOnClonedMesh(ActivityContext->CurrentMesh.Get());
+			
+			// Emit transaction
+			FGroupTopologySelection EmptySelection;
+			ActivityContext->EmitCurrentMeshChangeAndUpdate(LOCTEXT("EdgeInsertionTransactionName", "Edge Insertion"),
+				ChangeTracker.EndChange(), EmptySelection, true);
+
+			ToolState = EState::GettingStart;
+		}
+		else
+		{
+			ToolState = EState::GettingEnd;
+		}
+
+		PreviewEdges.Reset();
+	}
+
+}
+
+void UPolyEditInsertEdgeActivity::Render(IToolsContextRenderAPI* RenderAPI)
+{
+	ParentTool->GetToolManager()->GetContextQueriesAPI()->GetCurrentViewState(CameraState);
+
+	// Draw the existing group edges
+	FViewCameraState RenderCameraState = RenderAPI->GetCameraState();
+	ExistingEdgesRenderer.BeginFrame(RenderAPI, RenderCameraState);
+	ExistingEdgesRenderer.SetTransform(TargetTransform);
+
+	for (const FGroupTopology::FGroupEdge& Edge : ActivityContext->CurrentTopology->Edges)
+	{
+		FVector3d A, B;
+		for (int32 eid : Edge.Span.Edges)
+		{
+			ActivityContext->CurrentMesh->GetEdgeV(eid, A, B);
+			ExistingEdgesRenderer.DrawLine(A, B);
+		}
+	}
+	ExistingEdgesRenderer.EndFrame();
+
+	// Draw the preview edges and points
+	PreviewEdgeRenderer.BeginFrame(RenderAPI, RenderCameraState);
+	PreviewEdgeRenderer.SetTransform(TargetTransform);
+	for (const TPair<FVector3d, FVector3d>& EdgeVerts : PreviewEdges)
+	{
+		PreviewEdgeRenderer.DrawLine(EdgeVerts.Key, EdgeVerts.Value);
+	}
+	for (const FVector3d& Point : PreviewPoints)
+	{
+		PreviewEdgeRenderer.DrawPoint(Point);
+	}
+	PreviewEdgeRenderer.EndFrame();
+}
+
+bool UPolyEditInsertEdgeActivity::CanAccept() const
+{
+	return ToolState != EState::WaitingForInsertComplete;
+}
+
+void UPolyEditInsertEdgeActivity::OnPropertyModified(UObject* PropertySet, FProperty* Property)
+{
+	PreviewEdges.Reset();
+	ActivityContext->Preview->InvalidateResult();
+}
+
+void UPolyEditInsertEdgeActivity::ClearPreview(bool bClearDrawnElements)
+{
+	ActivityContext->Preview->CancelCompute();
+	ActivityContext->Preview->PreviewMesh->UpdatePreview(ActivityContext->CurrentMesh.Get());
+	bShowingBaseMesh = true;
+
+	if (bClearDrawnElements)
+	{
+		PreviewEdges.Reset();
+		PreviewPoints.Reset();
+	}
+}
+
+/** 
+ * Update the preview unless we've already computed one with the same parameters (such as when snapping to
+ * the same vertex despite moving the mouse).
+ */
+void UPolyEditInsertEdgeActivity::ConditionallyUpdatePreview(
+	const FGroupEdgeInserter::FGroupEdgeSplitPoint& NewEndPoint, int32 NewEndTopologyID, bool bNewEndIsCorner,
+	int32 NewCommonGroupID, int32 NewBoundaryIndex)
+{
+	//if (InputGroupEdgeID != NewGroupID
+	//	|| (NewInteractiveInputLength && Settings->PositionMode != EEdgeLoopPositioningMode::Even
+	//		&& *NewInteractiveInputLength != InteractiveInputLength))
+	//{
+	//	InputGroupEdgeID = NewGroupID;
+	//	if (NewInteractiveInputLength)
+	//	{
+	//		InteractiveInputLength = *NewInteractiveInputLength;
+	//	}
+	//	PreviewEdges.Reset();
+	//	ActivityContext->Preview->InvalidateResult();
+	//}
+
+	if (bShowingBaseMesh 
+		|| bEndIsCorner != bNewEndIsCorner || EndTopologyID != NewEndTopologyID
+		|| EndPoint.bIsVertex != NewEndPoint.bIsVertex || EndPoint.ElementID != NewEndPoint.ElementID
+		|| (!NewEndPoint.bIsVertex && NewEndPoint.EdgeTValue != EndPoint.EdgeTValue)
+		|| CommonGroupID != NewCommonGroupID || CommonBoundaryIndex != NewBoundaryIndex)
+	{
+		// Update the end variables, since they are apparently different
+		EndPoint = NewEndPoint;
+		EndTopologyID = NewEndTopologyID;
+		bEndIsCorner = bNewEndIsCorner;
+		CommonGroupID = NewCommonGroupID;
+		CommonBoundaryIndex = NewBoundaryIndex;
+
+		// If either endpoint is a corner, we need to calculate its tangent. This will differ based on which
+		// boundary it is a part of.
+		if (bStartIsCorner)
+		{
+			GetCornerTangent(StartTopologyID, CommonGroupID, CommonBoundaryIndex, StartPoint.Tangent);
+		}
+		if (bEndIsCorner)
+		{
+			GetCornerTangent(EndTopologyID, CommonGroupID, CommonBoundaryIndex, EndPoint.Tangent);
+		}
+
+		bShowingBaseMesh = false;
+		PreviewEdges.Reset();
+		ActivityContext->Preview->InvalidateResult();
+	}
+}
+
+FInputRayHit UPolyEditInsertEdgeActivity::BeginHoverSequenceHitTest(const FInputDeviceRay& PressPos)
+{
+	FInputRayHit Hit;
+
+	// Early out if the activity is not running. This is actually important because the behavior is
+	// always in the behavior list while the tool is running (we don't have a way to add/remove at
+	// will).
+	if (!bIsRunning)
+	{
+		return Hit; // Hit.bHit is false
+	}
+
+	switch (ToolState)
+	{
+	case EState::WaitingForInsertComplete:
+		break; // Keep hit invalid
+
+	case EState::GettingStart:
+	{
+		PreviewPoints.Reset();
+		FVector3d RayPoint;
+		if (TopologyHitTest(PressPos.WorldRay, RayPoint))
+		{
+			Hit = FInputRayHit(PressPos.WorldRay.GetParameter((FVector)RayPoint));
+		}
+		break;
+	}
+	case EState::GettingEnd:
+	{
+		FVector3d RayPoint;
+		FRay3d LocalRay;
+		if (TopologyHitTest(PressPos.WorldRay, RayPoint, &LocalRay))
+		{
+			Hit = FInputRayHit(PressPos.WorldRay.GetParameter((FVector)RayPoint));
+		}
+		else
+		{
+			// If we don't hit a valid element, we still do a hover if we hit the mesh.
+			// We still do the topology check in the first place because it accepts missing
+			// rays that are close enough to snap.
+			double RayT = 0;
+			int32 Tid = FDynamicMesh3::InvalidID;
+			if (ActivityContext->MeshSpatial->FindNearestHitTriangle(LocalRay, RayT, Tid))
+			{
+				Hit = FInputRayHit(RayT);
+			}
+		}
+		break;
+	}
+	}
+
+	return Hit;
+}
+
+bool UPolyEditInsertEdgeActivity::OnUpdateHover(const FInputDeviceRay& DevicePos)
+{
+	using namespace PolyEditInsertEdgeActivityLocals;
+
+	switch (ToolState)
+	{
+	case EState::WaitingForInsertComplete:
+		return false; // Do nothing.
+
+	case EState::GettingStart:
+	{
+		// Update start variables and show a preview of a point if it's on an edge or corner
+		PreviewPoints.Reset();
+		FVector3d PreviewPoint;
+		if (GetHoveredItem(DevicePos.WorldRay, StartPoint, StartTopologyID, bStartIsCorner, PreviewPoint))
+		{
+			PreviewPoints.Add(PreviewPoint);
+			return true;
+		}
+		return false;
+	}
+	case EState::GettingEnd:
+	{
+		check(PreviewPoints.Num() > 0);
+		PreviewPoints.SetNum(1); // Keep the first element, which is the start point
+
+		// Don't update the end variables right away so that we can check if they actually changed (they
+		// won't when we snap to the same corner as before).
+		FGroupEdgeInserter::FGroupEdgeSplitPoint SnappedPoint;
+		int32 PointTopologyID, GroupID, BoundaryIndex;
+		bool bPointIsCorner;
+		FVector3d PreviewPoint;
+		FRay3d LocalRay;
+		if (GetHoveredItem(DevicePos.WorldRay, SnappedPoint, PointTopologyID, bPointIsCorner, PreviewPoint, &LocalRay))
+		{
+			// See if the point is not on the same vertex/edge but is on the same boundary
+			if (!(SnappedPoint.bIsVertex == StartPoint.bIsVertex && SnappedPoint.ElementID == StartPoint.ElementID)
+				&& GetSharedBoundary(*ActivityContext->CurrentTopology, StartPoint, StartTopologyID, bStartIsCorner,
+					SnappedPoint, PointTopologyID, bPointIsCorner, GroupID, BoundaryIndex))
+			{
+				ConditionallyUpdatePreview(SnappedPoint, PointTopologyID, bPointIsCorner, GroupID, BoundaryIndex);
+			}
+			else
+			{
+				PreviewEdges.Reset(); // TODO: Maybe we should show a different color edge on a fail, rather than hiding it?
+			}
+			PreviewPoints.Add(PreviewPoint);
+			
+			return true;
+		}
+
+		// If we don't have a valid endpoint, draw a line to the current hit location.
+		if (!bShowingBaseMesh)
+		{
+			ClearPreview(false);
+		}
+		PreviewEdges.Reset();
+		double RayT = 0;
+		int32 Tid = FDynamicMesh3::InvalidID;
+		if (ActivityContext->MeshSpatial->FindNearestHitTriangle(LocalRay, RayT, Tid))
+		{
+			PreviewEdges.Emplace(PreviewPoints[0], LocalRay.PointAt(RayT));
+			return true;
+		}
+		return false;
+	}
+	}
+
+	check(false); // Each case has its own return, so shouldn't get here
+	return false;
+}
+
+void UPolyEditInsertEdgeActivity::OnEndHover()
+{
+	switch (ToolState)
+	{
+	case EState::WaitingForInsertComplete:
+	case EState::GettingStart:
+		ClearPreview(true);
+		break;
+	case EState::GettingEnd:
+		// Keep the first preview point.
+		ClearPreview(false);
+		PreviewPoints.SetNum(1);
+		PreviewEdges.Reset();
+	}
+}
+
+FInputRayHit UPolyEditInsertEdgeActivity::IsHitByClick(const FInputDeviceRay& ClickPos)
+{
+	FInputRayHit Hit;
+
+	// Early out if the activity is not running. 
+	if (!bIsRunning)
+	{
+		return Hit; // Hit.bHit is false
+	}
+
+	switch (ToolState)
+	{
+	case EState::WaitingForInsertComplete:
+		break; // Keep hit invalid
+
+	// Same requirement for the other two cases: the click should go on an edge
+	case EState::GettingStart:
+	case EState::GettingEnd:
+	{
+		FVector3d RayPoint;
+		if (TopologyHitTest(ClickPos.WorldRay, RayPoint))
+		{
+			Hit = FInputRayHit(ClickPos.WorldRay.GetParameter((FVector)RayPoint));
+		}
+		break;
+	}
+	}
+	return Hit;
+}
+
+void UPolyEditInsertEdgeActivity::OnClicked(const FInputDeviceRay& ClickPos)
+{
+	using namespace PolyEditInsertEdgeActivityLocals;
+
+	switch (ToolState)
+	{
+	case EState::WaitingForInsertComplete:
+		break; // Do nothing
+
+	case EState::GettingStart:
+	{
+		// Update start variables and switch state if successful
+		FVector3d PreviewPoint;
+		if (GetHoveredItem(ClickPos.WorldRay, StartPoint, StartTopologyID, bStartIsCorner, PreviewPoint))
+		{
+			PreviewPoints.Reset();
+			PreviewPoints.Add(PreviewPoint);
+			ToolState = EState::GettingEnd;
+
+			ParentTool->GetToolManager()->BeginUndoTransaction(LOCTEXT("GroupEdgeStartTransactionName", "Group Edge Start"));
+			ParentTool->GetToolManager()->EmitObjectChange(this, MakeUnique<FGroupEdgeInsertionFirstPointChange>(CurrentChangeStamp),
+				LOCTEXT("GroupEdgeStart", "Group Edge Start"));
+			ParentTool->GetToolManager()->EndUndoTransaction();
+		}
+		break;
+	}
+	case EState::GettingEnd:
+	{
+		// Don't update the end variables right away so that we can check if they actually changed (they
+		// won't when we snap to the same corner as before).
+		FVector3d PreviewPoint;
+		FGroupEdgeInserter::FGroupEdgeSplitPoint SnappedPoint;
+		int32 PointTopologyID, GroupID, BoundaryIndex;
+		bool bPointIsCorner;
+		if (GetHoveredItem(ClickPos.WorldRay, SnappedPoint, PointTopologyID, bPointIsCorner, PreviewPoint))
+		{
+			// See if the point is not on the same vertex/edge but is on the same boundary
+			if (!(SnappedPoint.bIsVertex == StartPoint.bIsVertex && SnappedPoint.ElementID == StartPoint.ElementID) 
+				&& GetSharedBoundary(*ActivityContext->CurrentTopology, StartPoint, StartTopologyID, bStartIsCorner,
+					SnappedPoint, PointTopologyID, bPointIsCorner, GroupID, BoundaryIndex))
+			{
+				ConditionallyUpdatePreview(SnappedPoint, PointTopologyID, bPointIsCorner, GroupID, BoundaryIndex);
+				ToolState = EState::WaitingForInsertComplete;
+			}
+			else
+			{
+				ClearPreview(false);
+			}
+		}
+		break;
+	}
+	}
+}
+
+bool UPolyEditInsertEdgeActivity::TopologyHitTest(const FRay& WorldRay,
+	FVector3d& RayPositionOut, FRay3d* LocalRayOut)
+{
+	FRay3d LocalRay((FVector3d)TargetTransform.InverseTransformPosition(WorldRay.Origin),
+		(FVector3d)TargetTransform.InverseTransformVector(WorldRay.Direction), false);
+
+	if (LocalRayOut)
+	{
+		*LocalRayOut = LocalRay;
+	}
+
+	FGroupTopologySelection Selection;
+	FVector3d Position, Normal;
+	if (TopologySelector->FindSelectedElement(TopologySelectorSettings,
+		LocalRay, Selection, Position, Normal))
+	{
+		RayPositionOut = TargetTransform.TransformPosition(Position);
+		return true;
+	}
+	return false;
+}
+
+bool UPolyEditInsertEdgeActivity::GetHoveredItem(const FRay& WorldRay,
+	FGroupEdgeInserter::FGroupEdgeSplitPoint& PointOut, 
+	int32& TopologyElementIDOut, bool& bIsCornerOut, FVector3d& PositionOut,
+	FRay3d* LocalRayOut)
+{
+	TopologyElementIDOut = FDynamicMesh3::InvalidID;
+	PointOut.ElementID = FDynamicMesh3::InvalidID;
+
+	// Cast the ray to see what we hit.
+	FRay3d LocalRay((FVector3d)TargetTransform.InverseTransformPosition(WorldRay.Origin),
+		(FVector3d)TargetTransform.InverseTransformVector(WorldRay.Direction), false);
+	if (LocalRayOut)
+	{
+		*LocalRayOut = LocalRay;
+	}
+	FGroupTopologySelection Selection;
+	FVector3d Position, Normal;
+	int32 EdgeSegmentID;
+	if (!TopologySelector->FindSelectedElement(
+		TopologySelectorSettings, LocalRay, Selection, Position, Normal, &EdgeSegmentID))
+	{
+		return false; // Didn't hit anything
+	}
+	else if (Selection.SelectedCornerIDs.Num() > 0)
+	{
+		// Point is a corner
+		TopologyElementIDOut = Selection.GetASelectedCornerID();
+		bIsCornerOut = true;
+		PointOut.bIsVertex = true;
+		PointOut.ElementID = ActivityContext->CurrentTopology->GetCornerVertexID(TopologyElementIDOut);
+		// We can't initialize the tangent yet because the tangent of a corner will
+		// depend on which boundary it is a part of.
+
+		PositionOut = ActivityContext->CurrentMesh->GetVertex(PointOut.ElementID);
+	}
+	else 
+	{
+		// Point is an edge. We'll need to calculate the t value and some other things.
+		check(Selection.SelectedEdgeIDs.Num() > 0);
+
+		TopologyElementIDOut = Selection.GetASelectedEdgeID();
+		bIsCornerOut = false;
+
+		const FGroupTopology::FGroupEdge& GroupEdge = ActivityContext->CurrentTopology->Edges[TopologyElementIDOut];
+
+		int32 Eid = GroupEdge.Span.Edges[EdgeSegmentID];
+		int32 StartVid = GroupEdge.Span.Vertices[EdgeSegmentID];
+		int32 EndVid = GroupEdge.Span.Vertices[EdgeSegmentID + 1];
+		FVector3d StartVert = ActivityContext->CurrentMesh->GetVertex(StartVid);
+		FVector3d EndVert = ActivityContext->CurrentMesh->GetVertex(EndVid);
+		FVector3d EdgeVector = EndVert - StartVert;
+		double EdgeLength = EdgeVector.Length();
+		check(EdgeLength > 0);
+
+		PointOut.Tangent = EdgeVector / EdgeLength;
+
+		FRay EdgeRay((FVector)StartVert, (FVector)PointOut.Tangent, true);
+		float DistDownEdge = EdgeRay.GetParameter((FVector)Position);
+
+		PositionOut = (FVector3d)EdgeRay.PointAt(DistDownEdge);
+
+		// See if the point is at a vertex in the group edge span.
+		if (DistDownEdge <= Settings->VertexTolerance)
+		{
+			PointOut.bIsVertex = true;
+			PointOut.ElementID = StartVid;
+			if (EdgeSegmentID > 0)
+			{
+				// Average with previous normalized edge vector
+				PointOut.Tangent += 
+					UE::Geometry::Normalized(StartVert - ActivityContext->CurrentMesh->GetVertex(GroupEdge.Span.Vertices[EdgeSegmentID - 1]));
+				UE::Geometry::Normalize(PointOut.Tangent);
+			}
+		}
+		else if (abs(DistDownEdge - EdgeLength) <= Settings->VertexTolerance)
+		{
+			PointOut.bIsVertex = true;
+			PointOut.ElementID = EndVid;
+			if (EdgeSegmentID + 2 < GroupEdge.Span.Vertices.Num())
+			{
+				PointOut.Tangent += UE::Geometry::Normalized(
+					ActivityContext->CurrentMesh->GetVertex(GroupEdge.Span.Vertices[EdgeSegmentID + 2]) - EndVert);
+				UE::Geometry::Normalize(PointOut.Tangent);
+			}
+		}
+		else
+		{
+			PointOut.bIsVertex = false;
+			PointOut.ElementID = Eid;
+			PointOut.EdgeTValue = DistDownEdge / EdgeLength;
+			if (ActivityContext->CurrentMesh->GetEdgeV(Eid).A != StartVid)
+			{
+				PointOut.EdgeTValue = 1 - PointOut.EdgeTValue;
+			}
+		}
+	}
+	return true;
+}
+
+void UPolyEditInsertEdgeActivity::GetCornerTangent(int32 CornerID, int32 GroupID, int32 BoundaryIndex, FVector3d& TangentOut)
+{
+	TangentOut = FVector3d::Zero();
+
+	int32 CornerVid = ActivityContext->CurrentTopology->GetCornerVertexID(CornerID);
+	check(CornerVid != FDynamicMesh3::InvalidID);
+
+	const FGroupTopology::FGroup* Group = ActivityContext->CurrentTopology->FindGroupByID(GroupID);
+	check(Group && BoundaryIndex >= 0 && BoundaryIndex < Group->Boundaries.Num());
+
+	const FGroupTopology::FGroupBoundary& Boundary = Group->Boundaries[BoundaryIndex];
+	TArray<FVector3d> AdjacentPoints;
+	for (int32 GroupEdgeID : Boundary.GroupEdges)
+	{
+		TArray<int32> Vertices = ActivityContext->CurrentTopology->Edges[GroupEdgeID].Span.Vertices;
+		if (Vertices[0] == CornerVid)
+		{
+			AdjacentPoints.Add(ActivityContext->CurrentMesh->GetVertex(Vertices[1]));
+		}
+		else if (Vertices.Last() == CornerVid)
+		{
+			AdjacentPoints.Add(ActivityContext->CurrentMesh->GetVertex(Vertices[Vertices.Num()-2]));
+		}
+	}
+	check(AdjacentPoints.Num() == 2);
+
+	FVector3d CornerPosition = ActivityContext->CurrentMesh->GetVertex(CornerVid);
+	TangentOut = UE::Geometry::Normalized(CornerPosition - AdjacentPoints[0]);
+	TangentOut += UE::Geometry::Normalized(AdjacentPoints[1] - CornerPosition);
+	UE::Geometry::Normalize(TangentOut);
+}
+
+bool PolyEditInsertEdgeActivityLocals::GetSharedBoundary(const FGroupTopology& Topology,
+	const FGroupEdgeInserter::FGroupEdgeSplitPoint& StartPoint, int32 StartTopologyID, bool bStartIsCorner,
+	const FGroupEdgeInserter::FGroupEdgeSplitPoint& EndPoint, int32 EndTopologyID, bool bEndIsCorner,
+	int32& GroupIDOut, int32& BoundaryIndexOut)
+{
+	// The start and endpoints could be on the same boundary of multiple groups at
+	// the same time, and sometimes we won't be able to resolve the ambiguity
+	// (one example is a sphere split into two equal groups, but could even happen
+	// with more than two groups when endpoints are corners).
+	// Sometimes there are things we can do to eliminate some contenders- the best
+	// approach is probably trying to do a plane cut for all of the options and 
+	// removing those that fail. However, it's worth noting that such issues won't
+	// arise in the standard application of this tool for low-poly modeling, where
+	// groups are planar, so it's not worth the bother.
+	// Instead, we'll just take one of the results arbitrarily, though we will try to
+	// take one that has a single boundary (this will prefer a cylinder cap over
+	// a cylinder side).
+	// TODO: The code would be simpler if we didn't even want to do that filtering- we'd
+	// just return the first result we found. Should we consider doing that?
+
+	GroupIDOut = FDynamicMesh3::InvalidID;
+	BoundaryIndexOut = FDynamicMesh3::InvalidID;
+
+	TArray<TPair<int32, int32>> CandidateGroupIDsAndBoundaryIndices;
+	if (bStartIsCorner)
+	{
+		// Go through all neighboring groups and their boundaries to find a shared one.
+		const FGroupTopology::FCorner& StartCorner = Topology.Corners[StartTopologyID];
+		for (int32 GroupID : StartCorner.NeighbourGroupIDs)
+		{
+			const FGroupTopology::FGroup* Group = Topology.FindGroupByID(GroupID);
+			for (int32 i = 0; i < Group->Boundaries.Num(); ++i)
+			{
+				const FGroupTopology::FGroupBoundary& Boundary = Group->Boundaries[i];
+				if (DoesBoundaryContainPoint(Topology, Boundary, EndTopologyID, bEndIsCorner)
+					&& DoesBoundaryContainPoint(Topology, Boundary, StartTopologyID, bStartIsCorner))
+				{
+					CandidateGroupIDsAndBoundaryIndices.Emplace(GroupID, i);
+					break; // Can't share more than one boundary in the same group
+				}
+			}
+		}
+	}
+	else
+	{
+		// Start is on an edge, so there are fewer boundaries to look through.
+		const FGroupTopology::FGroupEdge& GroupEdge = Topology.Edges[StartTopologyID];
+		const FGroupTopology::FGroup* Group = Topology.FindGroupByID(GroupEdge.Groups.A);
+		for (int32 i = 0; i < Group->Boundaries.Num(); ++i)
+		{
+			const FGroupTopology::FGroupBoundary& Boundary = Group->Boundaries[i];
+
+			if (DoesBoundaryContainPoint(Topology, Boundary, EndTopologyID, bEndIsCorner)
+				&& DoesBoundaryContainPoint(Topology, Boundary, StartTopologyID, bStartIsCorner))
+			{
+				CandidateGroupIDsAndBoundaryIndices.Emplace(GroupEdge.Groups.A, i);
+				break;
+			}
+		}
+		if (GroupEdge.Groups.B != FDynamicMesh3::InvalidID)
+		{
+			Group = Topology.FindGroupByID(GroupEdge.Groups.B);
+			for (int32 i = 0; i < Group->Boundaries.Num(); ++i)
+			{
+				const FGroupTopology::FGroupBoundary& Boundary = Group->Boundaries[i];
+
+				if (DoesBoundaryContainPoint(Topology, Boundary, EndTopologyID, bEndIsCorner)
+					&& DoesBoundaryContainPoint(Topology, Boundary, StartTopologyID, bStartIsCorner))
+				{
+					CandidateGroupIDsAndBoundaryIndices.Emplace(GroupEdge.Groups.B, i);
+					break;
+				}
+			}
+		}
+	}
+
+	if (CandidateGroupIDsAndBoundaryIndices.Num() == 0)
+	{
+		return false;
+	}
+
+	// Prefer a result that has a single boundary if there are multiple.
+	if (CandidateGroupIDsAndBoundaryIndices.Num() > 1)
+	{
+		for (const TPair<int32, int32>& GroupIDBoundaryIdxPair : CandidateGroupIDsAndBoundaryIndices)
+		{
+			if (Topology.FindGroupByID(GroupIDBoundaryIdxPair.Key)->Boundaries.Num() == 1)
+			{
+				GroupIDOut = GroupIDBoundaryIdxPair.Key;
+				BoundaryIndexOut = 0;
+				return true;
+			}
+		}
+	}
+
+	GroupIDOut = CandidateGroupIDsAndBoundaryIndices[0].Key;
+	BoundaryIndexOut = CandidateGroupIDsAndBoundaryIndices[0].Value;
+	return true;
+}
+
+bool PolyEditInsertEdgeActivityLocals::DoesBoundaryContainPoint(const FGroupTopology& Topology,
+	const FGroupTopology::FGroupBoundary& Boundary, int32 PointTopologyID, bool bPointIsCorner)
+{
+	for (int32 GroupEdgeID : Boundary.GroupEdges)
+	{
+		if (!bPointIsCorner && GroupEdgeID == PointTopologyID)
+		{
+			return true;
+		}
+
+		const FGroupTopology::FGroupEdge& GroupEdge = Topology.Edges[GroupEdgeID];
+		if (bPointIsCorner && (GroupEdge.EndpointCorners.A == PointTopologyID
+			|| GroupEdge.EndpointCorners.B == PointTopologyID))
+		{
+			return true;
+		}
+	}
+	return false;
+}
+
+void FGroupEdgeInsertionFirstPointChange::Revert(UObject* Object)
+{
+	UPolyEditInsertEdgeActivity* Activity = Cast<UPolyEditInsertEdgeActivity>(Object);
+
+	check(Activity->ToolState == UPolyEditInsertEdgeActivity::EState::GettingEnd);
+	Activity->ToolState = UPolyEditInsertEdgeActivity::EState::GettingStart;
+
+	Activity->ClearPreview();
+
+	bHaveDoneUndo = true;
+}
+
+#undef LOCTEXT_NAMESPACE
\ No newline at end of file
diff --git a/Engine/Plugins/Experimental/MeshModelingToolset/Source/MeshModelingTools/Private/ToolActivities/PolyEditInsertEdgeLoopActivity.cpp b/Engine/Plugins/Experimental/MeshModelingToolset/Source/MeshModelingTools/Private/ToolActivities/PolyEditInsertEdgeLoopActivity.cpp
new file mode 100644
--- /dev/null
+++ b/Engine/Plugins/Experimental/MeshModelingToolset/Source/MeshModelingTools/Private/ToolActivities/PolyEditInsertEdgeLoopActivity.cpp
@@ -0,0 +1,540 @@
+// Copyright Epic Games, Inc. All Rights Reserved.
+
+#include "ToolActivities/PolyEditInsertEdgeLoopActivity.h"
+
+#include "BaseBehaviors/SingleClickBehavior.h"
+#include "BaseBehaviors/MouseHoverBehavior.h"
+#include "ContextObjectStore.h"
+#include "CuttingOps/EdgeLoopInsertionOp.h"
+#include "DynamicMeshToMeshDescription.h"
+#include "DynamicMesh/DynamicMeshChangeTracker.h"
+#include "InteractiveToolManager.h"
+#include "MeshOpPreviewHelpers.h"
+#include "Operations/GroupEdgeInserter.h"
+#include "Selection/PolygonSelectionMechanic.h"
+#include "ToolActivities/PolyEditActivityContext.h"
+
+#include "ExplicitUseGeometryMathTypes.h"		// using UE::Geometry::(math types)
+using namespace UE::Geometry;
+
+#define LOCTEXT_NAMESPACE "UPolyEditInsertEdgeLoopActivity"
+
+TUniquePtr<FDynamicMeshOperator> UPolyEditInsertEdgeLoopActivity::MakeNewOperator()
+{
+	TUniquePtr<FEdgeLoopInsertionOp> Op = MakeUnique<FEdgeLoopInsertionOp>();
+
+	Op->OriginalMesh = ActivityContext->CurrentMesh;
+	Op->OriginalTopology = ActivityContext->CurrentTopology;
+	Op->SetTransform(TargetTransform);
+
+	if (Settings->InsertionMode == EEdgeLoopInsertionMode::PlaneCut)
+	{
+		Op->Mode = FGroupEdgeInserter::EInsertionMode::PlaneCut;
+	}
+	else
+	{
+		Op->Mode = FGroupEdgeInserter::EInsertionMode::Retriangulate;
+	}
+
+	Op->VertexTolerance = Settings->VertexTolerance;
+
+	Op->GroupEdgeID = InputGroupEdgeID;
+
+	Op->StartCornerID = Settings->bFlipOffsetDirection ?
+		Op->OriginalTopology->Edges[Op->GroupEdgeID].EndpointCorners.B
+		: Op->OriginalTopology->Edges[Op->GroupEdgeID].EndpointCorners.A;
+
+	// Set up the inputs
+	if (Settings->PositionMode == EEdgeLoopPositioningMode::Even)
+	{
+		int32 NumLoops = Settings->NumLoops;
+		for (int32 i = 0; i < NumLoops; ++i)
+		{
+			Op->InputLengths.Add((i + 1.0) / (NumLoops + 1.0));
+		}
+	}
+	else if (Settings->bInteractive)
+	{
+		Op->InputLengths.Add(InteractiveInputLength);
+	}
+	else if (Settings->PositionMode == EEdgeLoopPositioningMode::ProportionOffset)
+	{
+		Op->InputLengths.Add(Settings->ProportionOffset);
+	}
+	else
+	{
+		Op->InputLengths.Add(Settings->DistanceOffset);
+	}
+
+	Op->bInputsAreProportions = (Settings->PositionMode == EEdgeLoopPositioningMode::Even
+		|| Settings->PositionMode == EEdgeLoopPositioningMode::ProportionOffset);
+
+	return Op;
+}
+
+void UPolyEditInsertEdgeLoopActivity::Setup(UInteractiveTool* ParentToolIn)
+{
+	Super::Setup(ParentToolIn);
+
+	// Set up properties
+	Settings = NewObject<UEdgeLoopInsertionProperties>(this);
+	Settings->RestoreProperties(ParentTool.Get());
+	AddToolPropertySource(Settings);
+	SetToolPropertySourceEnabled(Settings, false);
+	Settings->GetOnModified().AddUObject(this, &UPolyEditInsertEdgeLoopActivity::OnPropertyModified);
+
+	// Draws the new group edges that are added
+	PreviewEdgeRenderer.LineColor = FLinearColor::Green;
+	PreviewEdgeRenderer.LineThickness = 2.0;
+
+	// Highlights non-quad groups that stop the loop;
+	ProblemTopologyRenderer.LineColor = FLinearColor::Red;
+	ProblemTopologyRenderer.LineThickness = 3.0;
+	ProblemTopologyRenderer.DepthBias = 1.0;
+
+	TopologySelectorSettings.bEnableEdgeHits = true;
+	TopologySelectorSettings.bEnableFaceHits = false;
+	TopologySelectorSettings.bEnableCornerHits = false;
+
+	// Set up our input routing
+	USingleClickInputBehavior* ClickBehavior = NewObject<USingleClickInputBehavior>();
+	ClickBehavior->Initialize(this);
+	ParentTool->AddInputBehavior(ClickBehavior, this);
+	UMouseHoverBehavior* HoverBehavior = NewObject<UMouseHoverBehavior>();
+	HoverBehavior->Initialize(this);
+	ParentTool->AddInputBehavior(HoverBehavior, this);
+
+	ActivityContext = ParentTool->GetToolManager()->GetContextObjectStore()->FindContext<UPolyEditActivityContext>();
+
+	TopologySelector = ActivityContext->SelectionMechanic->GetTopologySelector();
+}
+
+void UPolyEditInsertEdgeLoopActivity::Shutdown(EToolShutdownType ShutdownType)
+{
+	if (bIsRunning)
+	{
+		End(ShutdownType);
+	}
+
+	Settings->SaveProperties(ParentTool.Get());
+	Settings->GetOnModified().Clear();
+
+	TopologySelector.Reset();
+	Settings = nullptr;
+	ActivityContext = nullptr;
+
+	Super::Shutdown(ShutdownType);
+}
+
+EToolActivityStartResult UPolyEditInsertEdgeLoopActivity::Start()
+{
+	ParentTool->GetToolManager()->DisplayMessage(
+		LOCTEXT("InsertEdgeLoopActivityDescription", "Click an edge to insert an edge loop passing across "
+			"that edge. Edge loops follow a sequence of quad-like polygroups."),
+		EToolMessageLevel::UserNotification);
+
+	SetToolPropertySourceEnabled(Settings, true);
+
+	// We don't use selection, so clear it if necessary (have to issue an undo/redo event)
+	if (!ActivityContext->SelectionMechanic->GetActiveSelection().IsEmpty())
+	{
+		ActivityContext->SelectionMechanic->BeginChange();
+		ActivityContext->SelectionMechanic->ClearSelection();
+		ParentTool->GetToolManager()->EmitObjectChange(ActivityContext->SelectionMechanic,
+			ActivityContext->SelectionMechanic->EndChange(), LOCTEXT("ClearSelection", "Clear Selection"));
+	}
+
+	TargetTransform = ActivityContext->Preview->PreviewMesh->GetTransform();
+
+	SetupPreview();
+
+	bWaitingForInsertionCompletion = false;
+	bLastComputeSucceeded = false;
+	InputGroupEdgeID = FDynamicMesh3::InvalidID;
+
+	bIsRunning = true;
+
+	// Emit activity start transaction
+	ActivityContext->EmitActivityStart(LOCTEXT("BeginInsertEdgeLoopActivity", "Begin Insert Edge Loop"));
+
+	return EToolActivityStartResult::Running;
+}
+
+EToolActivityEndResult UPolyEditInsertEdgeLoopActivity::End(EToolShutdownType ShutdownType)
+{
+	if (!bIsRunning)
+	{
+		return EToolActivityEndResult::ErrorDuringEnd;
+	}
+
+	SetToolPropertySourceEnabled(Settings, false);
+	Settings->GetOnModified().Clear();
+
+	ActivityContext->Preview->OnOpCompleted.RemoveAll(this);
+	ActivityContext->Preview->OnMeshUpdated.RemoveAll(this);
+	ActivityContext->Preview->ClearOpFactory();
+
+	bIsRunning = false;
+
+	return CanAccept() ? EToolActivityEndResult::Completed : EToolActivityEndResult::Cancelled;
+}
+
+void UPolyEditInsertEdgeLoopActivity::SetupPreview()
+{
+	ActivityContext->Preview->ChangeOpFactory(this);
+
+	// Whenever we get a new result from the op, we need to extract the preview edges so that
+	// we can draw them if we want to, and the additional outputs we need (changed triangles and
+	// topology).
+	ActivityContext->Preview->OnOpCompleted.AddWeakLambda(this, [this](const FDynamicMeshOperator* UncastOp) {
+		const FEdgeLoopInsertionOp* Op = static_cast<const FEdgeLoopInsertionOp*>(UncastOp);
+
+		bLastComputeSucceeded = Op->bSucceeded;
+		LatestOpTopologyResult.Reset();
+		PreviewEdges.Reset();
+		LatestOpChangedTids.Reset();
+
+		if (bLastComputeSucceeded)
+		{
+			Op->GetLoopEdgeLocations(PreviewEdges);
+			LatestOpTopologyResult = Op->ResultTopology;
+			LatestOpChangedTids = Op->ChangedTids;
+		}
+
+		// Regardless of success, extract things for highlighting any non-quads that stopped our loop.
+		ProblemTopologyEdges.Reset();
+		ProblemTopologyVerts.Reset();
+		for (int32 GroupEdgeID : Op->ProblemGroupEdgeIDs)
+		{
+			for (int32 Eid : ActivityContext->CurrentTopology->GetGroupEdgeEdges(GroupEdgeID))
+			{
+				TPair<FVector3d, FVector3d> Endpoints;
+				ActivityContext->CurrentMesh->GetEdgeV(Eid, Endpoints.Key, Endpoints.Value);
+				ProblemTopologyEdges.Add(MoveTemp(Endpoints));
+			}
+			FGroupTopology::FGroupEdge& GroupEdge = ActivityContext->CurrentTopology->Edges[GroupEdgeID];
+			if (GroupEdge.EndpointCorners.A != FDynamicMesh3::InvalidID)
+			{
+				ProblemTopologyVerts.AddUnique(ActivityContext->CurrentMesh->GetVertex(
+					ActivityContext->CurrentTopology->Corners[GroupEdge.EndpointCorners.A].VertexID));
+				ProblemTopologyVerts.AddUnique(ActivityContext->CurrentMesh->GetVertex(
+					ActivityContext->CurrentTopology->Corners[GroupEdge.EndpointCorners.B].VertexID));
+			}
+		}
+		});
+
+	// In case of failure, we want to hide the broken preview, since we wouldn't accept it on
+	// a click. Note that this can't be fired OnOpCompleted because the preview is updated
+	// with the op result after that callback, which would undo the reset. The preview edge
+	// extraction can't be lumped in here because it needs the op rather than the preview object.
+	ActivityContext->Preview->OnMeshUpdated.AddWeakLambda(this, [this](UMeshOpPreviewWithBackgroundCompute*) {
+		if (!bLastComputeSucceeded)
+		{
+			ActivityContext->Preview->PreviewMesh->UpdatePreview(ActivityContext->CurrentMesh.Get());
+		}
+		});
+}
+
+bool UPolyEditInsertEdgeLoopActivity::CanStart() const
+{
+	return true;
+}
+
+void UPolyEditInsertEdgeLoopActivity::Tick(float DeltaTime)
+{
+	// See if we clicked and are waiting for the result to be computed and put in.
+	if (bWaitingForInsertionCompletion && ActivityContext->Preview->HaveValidResult())
+	{
+		if (bLastComputeSucceeded)
+		{
+			FDynamicMeshChangeTracker ChangeTracker(ActivityContext->CurrentMesh.Get());
+			ChangeTracker.BeginChange();
+			ChangeTracker.SaveTriangles(*LatestOpChangedTids, true /*bSaveVertices*/);
+
+			// Update current mesh
+			ActivityContext->CurrentMesh->Copy(*ActivityContext->Preview->PreviewMesh->GetMesh(), 
+				true, true, true, true);
+			
+
+			*ActivityContext->CurrentTopology = *LatestOpTopologyResult;
+			ActivityContext->CurrentTopology->RetargetOnClonedMesh(ActivityContext->CurrentMesh.Get());
+
+			// Emit transaction
+			FGroupTopologySelection EmptySelection;
+			ActivityContext->EmitCurrentMeshChangeAndUpdate(LOCTEXT("EdgeLoopInsertionTransactionName", "Edge Loop Insertion"),
+				ChangeTracker.EndChange(), EmptySelection, true);
+		}
+
+		PreviewEdges.Reset();
+		ProblemTopologyEdges.Reset();
+		ProblemTopologyVerts.Reset();
+
+		bWaitingForInsertionCompletion = false;
+	}
+}
+
+void UPolyEditInsertEdgeLoopActivity::Render(IToolsContextRenderAPI* RenderAPI)
+{
+	ParentTool->GetToolManager()->GetContextQueriesAPI()->GetCurrentViewState(CameraState);
+
+	// Draw the preview edges
+	PreviewEdgeRenderer.BeginFrame(RenderAPI, CameraState);
+	PreviewEdgeRenderer.SetTransform(TargetTransform);
+	for (TPair<FVector3d, FVector3d>& EdgeVerts : PreviewEdges)
+	{
+		PreviewEdgeRenderer.DrawLine(EdgeVerts.Key, EdgeVerts.Value);
+	}
+	PreviewEdgeRenderer.EndFrame();
+
+	if (Settings->bHighlightProblemGroups)
+	{
+		// Highlight any non-quad groups that stopped the loop.
+		ProblemTopologyRenderer.BeginFrame(RenderAPI, CameraState);
+		ProblemTopologyRenderer.SetTransform(TargetTransform);
+		for (TPair<FVector3d, FVector3d>& EdgeVerts : ProblemTopologyEdges)
+		{
+			ProblemTopologyRenderer.DrawLine(EdgeVerts.Key, EdgeVerts.Value);
+		}
+		for (FVector3d& Vert : ProblemTopologyVerts)
+		{
+			ProblemTopologyRenderer.DrawViewFacingX(Vert, ProblemVertTickWidth);
+		}
+		ProblemTopologyRenderer.EndFrame();
+	}
+}
+
+bool UPolyEditInsertEdgeLoopActivity::CanAccept() const
+{
+	return !bWaitingForInsertionCompletion;
+}
+
+void UPolyEditInsertEdgeLoopActivity::OnPropertyModified(UObject* PropertySet, FProperty* Property)
+{
+	PreviewEdges.Reset();
+	ActivityContext->Preview->InvalidateResult();
+}
+
+FInputRayHit UPolyEditInsertEdgeLoopActivity::HitTest(const FRay& WorldRay)
+{
+	FInputRayHit Hit;
+
+	// Early out if the activity is not running. This is actually important because the behavior is
+	// always in the behavior list while the tool is running (we don't have a way to add/remove at
+	// will).
+	if (!bIsRunning)
+	{
+		return Hit; // Hit.bHit is false
+	}
+
+	// See if we hit an edge
+	FRay3d LocalRay((FVector3d)TargetTransform.InverseTransformPosition(WorldRay.Origin),
+		(FVector3d)TargetTransform.InverseTransformVector(WorldRay.Direction), false);
+	FGroupTopologySelection Selection;
+	FVector3d Position, Normal;
+	if (TopologySelector->FindSelectedElement(
+		TopologySelectorSettings, LocalRay, Selection, Position, Normal))
+	{
+		// TODO: We could check here that the edge has some quad-like neighbor. For now we
+		// just check that the edge isn't a loop unto itself (in which case the neighbor groups
+		// are definitely not quad-like).
+		int32 GroupEdgeID = Selection.GetASelectedEdgeID();
+		const FGroupTopology::FGroupEdge& GroupEdge = ActivityContext->CurrentTopology->Edges[GroupEdgeID];
+		if (GroupEdge.EndpointCorners.A != FDynamicMesh3::InvalidID)
+		{
+			Hit = FInputRayHit(LocalRay.Project(Position));
+		}
+	}
+
+	return Hit;
+}
+
+bool UPolyEditInsertEdgeLoopActivity::UpdateHoveredItem(const FRay& WorldRay)
+{
+	// Check that we hit an edge
+	FRay3d LocalRay((FVector3d)TargetTransform.InverseTransformPosition(WorldRay.Origin),
+		(FVector3d)TargetTransform.InverseTransformVector(WorldRay.Direction), false);
+
+	FGroupTopologySelection Selection;
+	FVector3d Position, Normal;
+	int32 EdgeSegmentID;
+	if (!TopologySelector->FindSelectedElement(
+		TopologySelectorSettings, LocalRay, Selection, Position, Normal, &EdgeSegmentID))
+	{
+		ClearPreview();
+		return false; // Didn't hit anything
+	}
+
+	// Check that the edge has endpoints
+	int32 GroupEdgeID = Selection.GetASelectedEdgeID();
+	FGroupTopology::FGroupEdge GroupEdge = ActivityContext->CurrentTopology->Edges[GroupEdgeID];
+	if (GroupEdge.EndpointCorners.A == FDynamicMesh3::InvalidID)
+	{
+		ClearPreview();
+		return false; // Edge definitely doesn't have quad-like neighbors
+	}
+
+	if (Settings->PositionMode == EEdgeLoopPositioningMode::Even)
+	{
+		// In even mode and non-interactive mode, all that matters is the group edge 
+		// that we're hovering, not where our pointer is exactly.
+		ConditionallyUpdatePreview(GroupEdgeID);
+		return true;
+	}
+	if (!Settings->bInteractive)
+	{
+		// Don't try to insert a loop when our inputs don't make sense.
+		double TotalLength = ActivityContext->CurrentTopology->GetEdgeArcLength(GroupEdgeID);
+		if (Settings->PositionMode == EEdgeLoopPositioningMode::DistanceOffset)
+		{
+			if (Settings->DistanceOffset > TotalLength || Settings->DistanceOffset <= Settings->VertexTolerance)
+			{
+				ClearPreview();
+				return false;
+			}
+		}
+		else if (Settings->PositionMode == EEdgeLoopPositioningMode::ProportionOffset)
+		{
+			if (abs(Settings->ProportionOffset * TotalLength - TotalLength) <= Settings->VertexTolerance)
+			{
+				ClearPreview();
+				return false;
+			}
+		}
+
+		ConditionallyUpdatePreview(GroupEdgeID);
+		return true;
+	}
+
+	// Otherwise, we need to figure out where along the edge we are hovering.
+	double NewInputLength = 0;
+	int32 StartVid = GroupEdge.Span.Vertices[EdgeSegmentID];
+	int32 EndVid = GroupEdge.Span.Vertices[EdgeSegmentID + 1];
+	FVector3d StartVert = ActivityContext->CurrentMesh->GetVertex(StartVid);
+	FVector3d EndVert = ActivityContext->CurrentMesh->GetVertex(EndVid);
+
+	FRay EdgeRay((FVector)StartVert, (FVector)(EndVert - StartVert), false);
+
+	double DistDownEdge = EdgeRay.GetParameter((FVector)Position);
+
+	TArray<double> PerVertexLengths;
+	double TotalLength = ActivityContext->CurrentTopology->GetEdgeArcLength(GroupEdgeID, &PerVertexLengths);
+
+	NewInputLength = PerVertexLengths[EdgeSegmentID] + DistDownEdge;
+	if (Settings->bFlipOffsetDirection)
+	{
+		// If we flipped start corner, we should be measuring from the opposite direction
+		NewInputLength = TotalLength - NewInputLength;
+	}
+	// We avoid trying to insert loops that are guaranteed to follow existing group edges.
+	// Distance offset with total length may work if the group widens on the other side.
+	// Though it's worth noting that this filter as a whole is assuming straight group edges...
+	if (NewInputLength <= Settings->VertexTolerance ||
+		(Settings->PositionMode == EEdgeLoopPositioningMode::ProportionOffset
+			&& abs(NewInputLength - TotalLength) <= Settings->VertexTolerance))
+	{
+		ClearPreview();
+		return false;
+	}
+	if (Settings->PositionMode == EEdgeLoopPositioningMode::ProportionOffset)
+	{
+		NewInputLength /= TotalLength;
+	}
+	ConditionallyUpdatePreview(GroupEdgeID, &NewInputLength);
+	return true;
+}
+
+void UPolyEditInsertEdgeLoopActivity::ClearPreview()
+{
+	ActivityContext->Preview->CancelCompute();
+	ActivityContext->Preview->PreviewMesh->UpdatePreview(ActivityContext->CurrentMesh.Get());
+	InputGroupEdgeID = FDynamicMesh3::InvalidID;
+	PreviewEdges.Reset();
+	ProblemTopologyEdges.Reset();
+	ProblemTopologyVerts.Reset();
+}
+
+/** 
+ * Update the preview unless we've already computed one with the same parameters (such as when using "even" or non-interactive
+ * parameter setting) 
+ */
+void UPolyEditInsertEdgeLoopActivity::ConditionallyUpdatePreview(int32 NewGroupID, double* NewInteractiveInputLength)
+{
+	if (InputGroupEdgeID != NewGroupID
+		|| (NewInteractiveInputLength && Settings->PositionMode != EEdgeLoopPositioningMode::Even
+			&& *NewInteractiveInputLength != InteractiveInputLength))
+	{
+		InputGroupEdgeID = NewGroupID;
+		if (NewInteractiveInputLength)
+		{
+			InteractiveInputLength = *NewInteractiveInputLength;
+		}
+		PreviewEdges.Reset();
+		ActivityContext->Preview->InvalidateResult();
+	}
+}
+
+FInputRayHit UPolyEditInsertEdgeLoopActivity::BeginHoverSequenceHitTest(const FInputDeviceRay& PressPos)
+{
+	if (bWaitingForInsertionCompletion)
+	{
+		return FInputRayHit();
+	}
+
+	return HitTest(PressPos.WorldRay);
+}
+
+bool UPolyEditInsertEdgeLoopActivity::OnUpdateHover(const FInputDeviceRay& DevicePos)
+{
+	if (bWaitingForInsertionCompletion)
+	{
+		return false;
+	}
+
+	return UpdateHoveredItem(DevicePos.WorldRay);
+}
+
+void UPolyEditInsertEdgeLoopActivity::OnEndHover()
+{
+	// Note that OnEndHover actually gets called before MouseDown, and therefore before OnClicked
+	// (which gets called on MouseUp). Currently, this means that we don't actually end up using
+	// an existing hover computation for our click even if MouseUp happened in the same place- we just
+	// always redo the hover in OnClicked and immediately say that we want to accept it.
+	// 
+	// If we ever want to change things to do all this a bit more properly, we would need to make a
+	// couple changes: 1. We need to use a ClickDrag behavior instead of SingleClick. 2. OnEndHover
+	// should set a boolean to request a clear to happen in Tick(). 3. In the click behavior, we
+	// cancel any requested clear on mouse down, and keep updating hover during drag. On mouse up
+	// we accept whatever hover we have.
+
+	if (!bWaitingForInsertionCompletion)
+	{
+		ClearPreview();
+	}
+}
+
+FInputRayHit UPolyEditInsertEdgeLoopActivity::IsHitByClick(const FInputDeviceRay& ClickPos)
+{
+	FInputRayHit Hit;
+	if (bWaitingForInsertionCompletion)
+	{
+		return Hit;
+	}
+
+	return HitTest(ClickPos.WorldRay);
+}
+
+void UPolyEditInsertEdgeLoopActivity::OnClicked(const FInputDeviceRay& ClickPos)
+{
+	if (bWaitingForInsertionCompletion)
+	{
+		return;
+	}
+
+	if (UpdateHoveredItem(ClickPos.WorldRay))
+	{
+		bWaitingForInsertionCompletion = true;
+	}
+
+}
+
+#undef LOCTEXT_NAMESPACE
\ No newline at end of file
diff --git a/Engine/Plugins/Experimental/MeshModelingToolset/Source/MeshModelingTools/Private/ToolActivities/PolyEditInsetOutsetActivity.cpp b/Engine/Plugins/Experimental/MeshModelingToolset/Source/MeshModelingTools/Private/ToolActivities/PolyEditInsetOutsetActivity.cpp
new file mode 100644
--- /dev/null
+++ b/Engine/Plugins/Experimental/MeshModelingToolset/Source/MeshModelingTools/Private/ToolActivities/PolyEditInsetOutsetActivity.cpp
@@ -0,0 +1,239 @@
+// Copyright Epic Games, Inc. All Rights Reserved.
+
+#include "ToolActivities/PolyEditInsetOutsetActivity.h"
+
+#include "BaseBehaviors/SingleClickBehavior.h"
+#include "BaseBehaviors/MouseHoverBehavior.h"
+#include "ContextObjectStore.h"
+#include "Drawing/PolyEditPreviewMesh.h"
+#include "DynamicMesh/MeshNormals.h"
+#include "EditMeshPolygonsTool.h"
+#include "InteractiveToolManager.h"
+#include "Mechanics/SpatialCurveDistanceMechanic.h"
+#include "MeshBoundaryLoops.h"
+#include "MeshOpPreviewHelpers.h"
+#include "Operations/InsetMeshRegion.h"
+#include "Selection/PolygonSelectionMechanic.h"
+#include "ToolActivities/PolyEditActivityContext.h"
+#include "ToolActivities/PolyEditActivityUtil.h"
+#include "ToolSceneQueriesUtil.h"
+
+using namespace UE::Geometry;
+
+#define LOCTEXT_NAMESPACE "UPolyEditInsetOutsetActivity"
+
+void UPolyEditInsetOutsetActivity::Setup(UInteractiveTool* ParentToolIn)
+{
+	Super::Setup(ParentToolIn);
+
+	InsetProperties = NewObject<UPolyEditInsetOutsetProperties>();
+	InsetProperties->RestoreProperties(ParentTool.Get());
+	AddToolPropertySource(InsetProperties);
+	SetToolPropertySourceEnabled(InsetProperties, false);
+
+	// Register ourselves to receive clicks and hover
+	USingleClickInputBehavior* ClickBehavior = NewObject<USingleClickInputBehavior>();
+	ClickBehavior->Initialize(this);
+	ParentTool->AddInputBehavior(ClickBehavior);
+	UMouseHoverBehavior* HoverBehavior = NewObject<UMouseHoverBehavior>();
+	HoverBehavior->Initialize(this);
+	ParentTool->AddInputBehavior(HoverBehavior);
+
+	ActivityContext = ParentTool->GetToolManager()->GetContextObjectStore()->FindContext<UPolyEditActivityContext>();
+}
+
+void UPolyEditInsetOutsetActivity::Shutdown(EToolShutdownType ShutdownType)
+{
+	Clear();
+	InsetProperties->SaveProperties(ParentTool.Get());
+
+	InsetProperties = nullptr;
+	ParentTool = nullptr;
+	ActivityContext = nullptr;
+}
+
+bool UPolyEditInsetOutsetActivity::CanStart() const
+{
+	if (!ActivityContext)
+	{
+		return false;
+	}
+	const FGroupTopologySelection& Selection = ActivityContext->SelectionMechanic->GetActiveSelection();
+	return !Selection.SelectedGroupIDs.IsEmpty();
+}
+
+EToolActivityStartResult UPolyEditInsetOutsetActivity::Start()
+{
+	if (!CanStart())
+	{
+		return EToolActivityStartResult::FailedStart;
+	}
+
+	Clear();
+	BeginInset();
+	bIsRunning = true;
+
+	ActivityContext->EmitActivityStart(LOCTEXT("BeginInsetOutsetActivity", "Begin Inset/Outset"));
+
+	return EToolActivityStartResult::Running;
+}
+
+bool UPolyEditInsetOutsetActivity::CanAccept() const
+{
+	return false;
+}
+
+EToolActivityEndResult UPolyEditInsetOutsetActivity::End(EToolShutdownType)
+{
+	Clear();
+	EToolActivityEndResult ToReturn = bIsRunning ? EToolActivityEndResult::Cancelled 
+		: EToolActivityEndResult::ErrorDuringEnd;
+	bIsRunning = false;
+	return ToReturn;
+}
+
+void UPolyEditInsetOutsetActivity::BeginInset()
+{
+	const FGroupTopologySelection& ActiveSelection = ActivityContext->SelectionMechanic->GetActiveSelection();
+	TArray<int32> ActiveTriangleSelection;
+	ActivityContext->CurrentTopology->GetSelectedTriangles(ActiveSelection, ActiveTriangleSelection);
+
+	FTransform3d WorldTransform(ActivityContext->Preview->PreviewMesh->GetTransform());
+
+	EditPreview = PolyEditActivityUtil::CreatePolyEditPreviewMesh(*ParentTool, *ActivityContext);
+	EditPreview->InitializeInsetType(ActivityContext->CurrentMesh.Get(), ActiveTriangleSelection, &WorldTransform);
+
+	// Hide the selected triangles (that are being replaced by the inset/outset portion)
+	ActivityContext->Preview->PreviewMesh->SetSecondaryBuffersVisibility(false);
+
+	// make infinite-extent hit-test mesh
+	FDynamicMesh3 InsetHitTargetMesh;
+	EditPreview->MakeInsetTypeTargetMesh(InsetHitTargetMesh);
+
+	CurveDistMechanic = NewObject<USpatialCurveDistanceMechanic>(this);
+	CurveDistMechanic->Setup(ParentTool.Get());
+	CurveDistMechanic->WorldPointSnapFunc = [this](const FVector3d& WorldPos, FVector3d& SnapPos)
+	{
+		return ActivityContext->CommonProperties->bSnapToWorldGrid
+			&& ToolSceneQueriesUtil::FindWorldGridSnapPoint(ParentTool.Get(), WorldPos, SnapPos);
+	};
+	CurveDistMechanic->CurrentDistance = 1.0f;  // initialize to something non-zero...prob should be based on polygon bounds maybe?
+
+	FMeshBoundaryLoops Loops(&InsetHitTargetMesh);
+	TArray<FVector3d> LoopVertices;
+	Loops.Loops[0].GetVertices(LoopVertices);
+	CurveDistMechanic->InitializePolyLoop(LoopVertices, FTransform3d::Identity());
+
+	SetToolPropertySourceEnabled(InsetProperties, true);
+
+	float BoundsMaxDim = ActivityContext->CurrentMesh->GetBounds().MaxDim();
+	if (BoundsMaxDim > 0)
+	{
+		UVScaleFactor = 1.0 / BoundsMaxDim;
+	}
+
+	bPreviewUpdatePending = true;
+}
+
+void UPolyEditInsetOutsetActivity::ApplyInset()
+{
+	check(CurveDistMechanic != nullptr && EditPreview != nullptr);
+
+	const FGroupTopologySelection& ActiveSelection = ActivityContext->SelectionMechanic->GetActiveSelection();
+	TArray<int32> ActiveTriangleSelection;
+	ActivityContext->CurrentTopology->GetSelectedTriangles(ActiveSelection, ActiveTriangleSelection);
+
+	FInsetMeshRegion Inset(ActivityContext->CurrentMesh.Get());
+	Inset.UVScaleFactor = UVScaleFactor;
+	Inset.Triangles = ActiveTriangleSelection;
+	Inset.InsetDistance = (InsetProperties->bOutset) ? -CurveDistMechanic->CurrentDistance
+		: CurveDistMechanic->CurrentDistance;
+	Inset.bReproject = (InsetProperties->bOutset) ? false : InsetProperties->bReproject;
+	Inset.Softness = InsetProperties->Softness;
+	Inset.bSolveRegionInteriors = !InsetProperties->bBoundaryOnly;
+	Inset.AreaCorrection = InsetProperties->AreaScale;
+
+	Inset.ChangeTracker = MakeUnique<FDynamicMeshChangeTracker>(ActivityContext->CurrentMesh.Get());
+	Inset.ChangeTracker->BeginChange();
+	Inset.Apply();
+
+	FMeshNormals::QuickComputeVertexNormalsForTriangles(*ActivityContext->CurrentMesh, Inset.AllModifiedTriangles);
+
+	// Emit undo (also updates relevant structures)
+	ActivityContext->EmitCurrentMeshChangeAndUpdate(LOCTEXT("PolyMeshInsetOutsetChange", "Inset/Outset"),
+		Inset.ChangeTracker->EndChange(), ActiveSelection, true);
+
+	// End activity
+	Clear();
+	bIsRunning = false;
+	Cast<IToolActivityHost>(ParentTool)->NotifyActivitySelfEnded(this);
+}
+
+void UPolyEditInsetOutsetActivity::Clear()
+{
+	if (EditPreview != nullptr)
+	{
+		EditPreview->Disconnect();
+		EditPreview = nullptr;
+	}
+
+	ActivityContext->Preview->PreviewMesh->SetSecondaryBuffersVisibility(true);
+
+	CurveDistMechanic = nullptr;
+	SetToolPropertySourceEnabled(InsetProperties, false);
+}
+
+void UPolyEditInsetOutsetActivity::Render(IToolsContextRenderAPI* RenderAPI)
+{
+	if (CurveDistMechanic != nullptr)
+	{
+		CurveDistMechanic->Render(RenderAPI);
+	}
+}
+
+void UPolyEditInsetOutsetActivity::Tick(float DeltaTime)
+{
+	if (EditPreview && bPreviewUpdatePending)
+	{
+		double Sign = InsetProperties->bOutset ? -1.0 : 1.0;
+		bool bReproject = (InsetProperties->bOutset) ? false : InsetProperties->bReproject;
+		double Softness = InsetProperties->Softness;
+		bool bBoundaryOnly = InsetProperties->bBoundaryOnly;
+		double AreaCorrection = InsetProperties->AreaScale;
+		EditPreview->UpdateInsetType(Sign * CurveDistMechanic->CurrentDistance,
+			bReproject, Softness, AreaCorrection, bBoundaryOnly);
+
+		bPreviewUpdatePending = false;
+	}
+}
+
+FInputRayHit UPolyEditInsetOutsetActivity::IsHitByClick(const FInputDeviceRay& ClickPos)
+{
+	FInputRayHit OutHit;
+	OutHit.bHit = bIsRunning;
+	return OutHit;
+}
+
+void UPolyEditInsetOutsetActivity::OnClicked(const FInputDeviceRay& ClickPos)
+{
+	if (bIsRunning)
+	{
+		ApplyInset();
+	}
+}
+
+FInputRayHit UPolyEditInsetOutsetActivity::BeginHoverSequenceHitTest(const FInputDeviceRay& PressPos)
+{
+	FInputRayHit OutHit;
+	OutHit.bHit = bIsRunning;
+	return OutHit;
+}
+
+bool UPolyEditInsetOutsetActivity::OnUpdateHover(const FInputDeviceRay& DevicePos)
+{
+	CurveDistMechanic->UpdateCurrentDistance(DevicePos.WorldRay);
+	bPreviewUpdatePending = true;
+	return bIsRunning;
+}
+
+#undef LOCTEXT_NAMESPACE
diff --git a/Engine/Plugins/Experimental/MeshModelingToolset/Source/MeshModelingTools/Private/ToolActivities/PolyEditPlanarProjectionUVActivity.cpp b/Engine/Plugins/Experimental/MeshModelingToolset/Source/MeshModelingTools/Private/ToolActivities/PolyEditPlanarProjectionUVActivity.cpp
new file mode 100644
--- /dev/null
+++ b/Engine/Plugins/Experimental/MeshModelingToolset/Source/MeshModelingTools/Private/ToolActivities/PolyEditPlanarProjectionUVActivity.cpp
@@ -0,0 +1,282 @@
+// Copyright Epic Games, Inc. All Rights Reserved.
+
+#include "ToolActivities/PolyEditPlanarProjectionUVActivity.h"
+
+#include "BaseBehaviors/SingleClickBehavior.h"
+#include "BaseBehaviors/MouseHoverBehavior.h"
+#include "ContextObjectStore.h"
+#include "Drawing/PolyEditPreviewMesh.h"
+#include "InteractiveToolManager.h"
+#include "Mechanics/CollectSurfacePathMechanic.h"
+#include "MeshOpPreviewHelpers.h"
+#include "Selection/PolygonSelectionMechanic.h"
+#include "ToolActivities/PolyEditActivityContext.h"
+#include "ToolSceneQueriesUtil.h"
+
+#define LOCTEXT_NAMESPACE "UPolyEditPlanarProjectionUVActivity"
+
+using namespace UE::Geometry;
+
+void UPolyEditPlanarProjectionUVActivity::Setup(UInteractiveTool* ParentToolIn)
+{
+	Super::Setup(ParentToolIn);
+
+	SetUVProperties = NewObject<UPolyEditSetUVProperties>();
+	SetUVProperties->RestoreProperties(ParentTool.Get());
+	AddToolPropertySource(SetUVProperties);
+	SetToolPropertySourceEnabled(SetUVProperties, false);
+
+	// Register ourselves to receive clicks and hover
+	USingleClickInputBehavior* ClickBehavior = NewObject<USingleClickInputBehavior>();
+	ClickBehavior->Initialize(this);
+	ParentTool->AddInputBehavior(ClickBehavior);
+	UMouseHoverBehavior* HoverBehavior = NewObject<UMouseHoverBehavior>();
+	HoverBehavior->Initialize(this);
+	ParentTool->AddInputBehavior(HoverBehavior);
+
+	ActivityContext = ParentTool->GetToolManager()->GetContextObjectStore()->FindContext<UPolyEditActivityContext>();
+}
+
+void UPolyEditPlanarProjectionUVActivity::Shutdown(EToolShutdownType ShutdownType)
+{
+	Clear();
+	SetUVProperties->SaveProperties(ParentTool.Get());
+
+	SetUVProperties = nullptr;
+	ParentTool = nullptr;
+	ActivityContext = nullptr;
+}
+
+bool UPolyEditPlanarProjectionUVActivity::CanStart() const
+{
+	if (!ActivityContext)
+	{
+		return false;
+	}
+	const FGroupTopologySelection& Selection = ActivityContext->SelectionMechanic->GetActiveSelection();
+	return !Selection.SelectedGroupIDs.IsEmpty();
+}
+
+EToolActivityStartResult UPolyEditPlanarProjectionUVActivity::Start()
+{
+	if (!CanStart())
+	{
+		ParentTool->GetToolManager()->DisplayMessage(
+			LOCTEXT("OnSetUVsFailedMesssage", "Cannot set UVs without face selection."), 
+			EToolMessageLevel::UserWarning);
+		return EToolActivityStartResult::FailedStart;
+	}
+
+	Clear();
+	BeginSetUVs();
+	bIsRunning = true;
+
+	ParentTool->GetToolManager()->DisplayMessage(
+		LOCTEXT("OnBeginSetUVsMessage", "Click on the face to Set UVs"), EToolMessageLevel::UserMessage);
+
+	ActivityContext->EmitActivityStart(LOCTEXT("BeginSetUVsActivity", "Begin Set UVs"));
+
+	return EToolActivityStartResult::Running;
+}
+
+bool UPolyEditPlanarProjectionUVActivity::CanAccept() const
+{
+	return false;
+}
+
+EToolActivityEndResult UPolyEditPlanarProjectionUVActivity::End(EToolShutdownType)
+{
+	Clear();
+	EToolActivityEndResult ToReturn = bIsRunning ? EToolActivityEndResult::Cancelled 
+		: EToolActivityEndResult::ErrorDuringEnd;
+	bIsRunning = false;
+	return ToReturn;
+}
+
+void UPolyEditPlanarProjectionUVActivity::BeginSetUVs()
+{
+	using namespace UE::Geometry::PolyEditActivityUtil;
+
+	TArray<int32> ActiveTriangleSelection;
+	const FGroupTopologySelection& ActiveSelection = ActivityContext->SelectionMechanic->GetActiveSelection();
+	ActivityContext->CurrentTopology->GetSelectedTriangles(ActiveSelection, ActiveTriangleSelection);
+
+	FTransform3d WorldTransform(ActivityContext->Preview->PreviewMesh->GetTransform());
+
+	EditPreview = PolyEditActivityUtil::CreatePolyEditPreviewMesh(*ParentTool, *ActivityContext);
+	EditPreview->InitializeStaticType(ActivityContext->CurrentMesh.Get(), ActiveTriangleSelection, &WorldTransform);
+
+	UpdatePolyEditPreviewMaterials(*ParentTool, *ActivityContext, *EditPreview, (SetUVProperties->bShowMaterial) ?
+		EPreviewMaterialType::SourceMaterials : EPreviewMaterialType::UVMaterial);
+	CurrentPreviewMaterial = SetUVProperties->bShowMaterial ?
+		EPreviewMaterialType::SourceMaterials : EPreviewMaterialType::UVMaterial;
+
+	FDynamicMesh3 StaticHitTargetMesh;
+	EditPreview->MakeInsetTypeTargetMesh(StaticHitTargetMesh);
+
+	// Hide the selected triangles (that are being replaced by the portion we're editing)
+	ActivityContext->Preview->PreviewMesh->SetSecondaryBuffersVisibility(false);
+
+	SurfacePathMechanic = NewObject<UCollectSurfacePathMechanic>(this);
+	SurfacePathMechanic->Setup(ParentTool.Get());
+	SurfacePathMechanic->InitializeMeshSurface(MoveTemp(StaticHitTargetMesh));
+	SurfacePathMechanic->SetFixedNumPointsMode(2);
+	SurfacePathMechanic->bSnapToTargetMeshVertices = true;
+	double SnapTol = ToolSceneQueriesUtil::GetDefaultVisualAngleSnapThreshD();
+	SurfacePathMechanic->SpatialSnapPointsFunc = [this, SnapTol](FVector3d Position1, FVector3d Position2)
+	{
+		return ToolSceneQueriesUtil::PointSnapQuery(this->CameraState, Position1, Position2, SnapTol);
+	};
+
+	ActiveSelectionBounds = FAxisAlignedBox3d::Empty();
+	for (int tid : ActiveTriangleSelection)
+	{
+		ActiveSelectionBounds.Contain(ActivityContext->CurrentMesh->GetTriBounds(tid));
+	}
+
+	SetToolPropertySourceEnabled(SetUVProperties, true);
+}
+
+void UPolyEditPlanarProjectionUVActivity::UpdateSetUVS()
+{
+	// align projection frame to line user is drawing out from plane origin
+	FFrame3d PlanarFrame = SurfacePathMechanic->PreviewPathPoint;
+	double UVScale = 1.0 / ActiveSelectionBounds.MaxDim();
+	if (SurfacePathMechanic->HitPath.Num() == 1)
+	{
+		SurfacePathMechanic->InitializePlaneSurface(PlanarFrame);
+
+		FVector3d Delta = PlanarFrame.Origin - SurfacePathMechanic->HitPath[0].Origin;
+		double Dist = UE::Geometry::Normalize(Delta);
+		UVScale *= FMathd::Lerp(1.0, 25.0, Dist / ActiveSelectionBounds.MaxDim());
+		PlanarFrame = SurfacePathMechanic->HitPath[0];
+		PlanarFrame.ConstrainedAlignAxis(0, Delta, PlanarFrame.Z());
+	}
+
+	EditPreview->UpdateStaticType([&](FDynamicMesh3& Mesh)
+		{
+			FDynamicMeshEditor Editor(&Mesh);
+			TArray<int32> AllTriangles;
+			for (int32 tid : Mesh.TriangleIndicesItr())
+			{
+				AllTriangles.Add(tid);
+			}
+			Editor.SetTriangleUVsFromProjection(AllTriangles, PlanarFrame, UVScale, FVector2f::Zero(), false);
+		}, false);
+
+}
+
+void UPolyEditPlanarProjectionUVActivity::ApplySetUVs()
+{
+	const FGroupTopologySelection& ActiveSelection = ActivityContext->SelectionMechanic->GetActiveSelection();
+	TArray<int32> ActiveTriangleSelection;
+	ActivityContext->CurrentTopology->GetSelectedTriangles(ActiveSelection, ActiveTriangleSelection);
+
+	// align projection frame to line user drew
+	FFrame3d PlanarFrame = SurfacePathMechanic->HitPath[0];
+	double UVScale = 1.0 / ActiveSelectionBounds.MaxDim();
+	FVector3d Delta = SurfacePathMechanic->HitPath[1].Origin - PlanarFrame.Origin;
+	double Dist = UE::Geometry::Normalize(Delta);
+	UVScale *= FMathd::Lerp(1.0, 25.0, Dist / ActiveSelectionBounds.MaxDim());
+	PlanarFrame.ConstrainedAlignAxis(0, Delta, PlanarFrame.Z());
+
+	// transform to local, use 3D point to transfer UV scale value
+	FTransform3d WorldTransform(ActivityContext->Preview->PreviewMesh->GetTransform());
+	FVector3d ScalePt = PlanarFrame.Origin + UVScale * PlanarFrame.Z();
+	FTransform3d ToLocalXForm(WorldTransform.Inverse());
+	PlanarFrame.Transform(ToLocalXForm);
+	ScalePt = ToLocalXForm.TransformPosition(ScalePt);
+	UVScale = Distance(ScalePt, PlanarFrame.Origin);
+
+	// track changes
+	FDynamicMeshChangeTracker ChangeTracker(ActivityContext->CurrentMesh.Get());
+	ChangeTracker.BeginChange();
+	ChangeTracker.SaveTriangles(ActiveTriangleSelection, true);
+	FDynamicMeshEditor Editor(ActivityContext->CurrentMesh.Get());
+	Editor.SetTriangleUVsFromProjection(ActiveTriangleSelection, PlanarFrame, UVScale, FVector2f::Zero(), false, 0);
+
+	// Emit undo (also updates relevant structures)
+	ActivityContext->EmitCurrentMeshChangeAndUpdate(LOCTEXT("PolyMeshSetUVsChange", "Set UVs"),
+		ChangeTracker.EndChange(), ActiveSelection, false);
+
+	// End activity
+	Clear();
+	bIsRunning = false;
+	Cast<IToolActivityHost>(ParentTool)->NotifyActivitySelfEnded(this);
+}
+
+void UPolyEditPlanarProjectionUVActivity::Clear()
+{
+	if (EditPreview != nullptr)
+	{
+		EditPreview->Disconnect();
+		EditPreview = nullptr;
+	}
+
+	ActivityContext->Preview->PreviewMesh->SetSecondaryBuffersVisibility(true);
+
+	SurfacePathMechanic = nullptr;
+	SetToolPropertySourceEnabled(SetUVProperties, false);
+}
+
+void UPolyEditPlanarProjectionUVActivity::Render(IToolsContextRenderAPI* RenderAPI)
+{
+	ParentTool->GetToolManager()->GetContextQueriesAPI()->GetCurrentViewState(CameraState);
+
+	if (SurfacePathMechanic != nullptr)
+	{
+		SurfacePathMechanic->Render(RenderAPI);
+	}
+}
+
+void UPolyEditPlanarProjectionUVActivity::Tick(float DeltaTime)
+{
+	using namespace UE::Geometry::PolyEditActivityUtil;
+
+	EPreviewMaterialType WantMaterial = (SetUVProperties->bShowMaterial) ? EPreviewMaterialType::SourceMaterials : EPreviewMaterialType::UVMaterial;
+	if (CurrentPreviewMaterial != WantMaterial)
+	{
+		UpdatePolyEditPreviewMaterials(*ParentTool, *ActivityContext, *EditPreview, WantMaterial);
+		CurrentPreviewMaterial = WantMaterial;
+	}
+
+	if (bPreviewUpdatePending)
+	{
+		UpdateSetUVS();
+		bPreviewUpdatePending = false;
+	}
+}
+
+FInputRayHit UPolyEditPlanarProjectionUVActivity::IsHitByClick(const FInputDeviceRay& ClickPos)
+{
+	FInputRayHit OutHit;
+	OutHit.bHit = bIsRunning;
+	return OutHit;
+}
+
+void UPolyEditPlanarProjectionUVActivity::OnClicked(const FInputDeviceRay& ClickPos)
+{
+	if (bIsRunning && SurfacePathMechanic->TryAddPointFromRay(ClickPos.WorldRay))
+	{
+		if (SurfacePathMechanic->IsDone())
+		{
+			ApplySetUVs();
+		}
+	}
+}
+
+FInputRayHit UPolyEditPlanarProjectionUVActivity::BeginHoverSequenceHitTest(const FInputDeviceRay& PressPos)
+{
+	FInputRayHit OutHit;
+	OutHit.bHit = bIsRunning;
+	return OutHit;
+}
+
+bool UPolyEditPlanarProjectionUVActivity::OnUpdateHover(const FInputDeviceRay& DevicePos)
+{
+	SurfacePathMechanic->UpdatePreviewPoint(DevicePos.WorldRay);
+	bPreviewUpdatePending = true;
+	return bIsRunning;
+}
+
+#undef LOCTEXT_NAMESPACE
diff --git a/Engine/Plugins/Experimental/MeshModelingToolset/Source/MeshModelingTools/Public/EdgeLoopInsertionTool.h b/Engine/Plugins/Experimental/MeshModelingToolset/Source/MeshModelingTools/Public/EdgeLoopInsertionTool.h
deleted file mode 100644
diff --git a/Engine/Plugins/Experimental/MeshModelingToolset/Source/MeshModelingTools/Public/EditMeshPolygonsTool.h b/Engine/Plugins/Experimental/MeshModelingToolset/Source/MeshModelingTools/Public/EditMeshPolygonsTool.h
--- a/Engine/Plugins/Experimental/MeshModelingToolset/Source/MeshModelingTools/Public/EditMeshPolygonsTool.h
+++ b/Engine/Plugins/Experimental/MeshModelingToolset/Source/MeshModelingTools/Public/EditMeshPolygonsTool.h
@@ -3,44 +3,53 @@
 #pragma once
 
 #include "CoreMinimal.h"
-#include "UObject/NoExportTypes.h"
-#include "BaseTools/MeshSurfacePointTool.h"
-#include "Components/DynamicMeshComponent.h"
+
+#include "Changes/MeshVertexChange.h" // IMeshVertexCommandChangeTarget
 #include "DynamicMesh/DynamicMeshAABBTree3.h"
-#include "ToolDataVisualizer.h"
-#include "Transforms/QuickAxisTranslater.h"
-#include "Transforms/QuickAxisRotator.h"
-#include "Changes/MeshVertexChange.h"
-#include "GroupTopology.h"
-#include "Spatial/GeometrySet3.h"
-#include "Selection/GroupTopologySelector.h"
+#include "DynamicMesh/DynamicMeshChangeTracker.h" // FDynamicMeshChange for TUniquePtr
+#include "InteractiveToolActivity.h" // IToolActivityHost
+#include "InteractiveToolBuilder.h"
 #include "Operations/GroupTopologyDeformer.h"
-#include "ModelingOperators/Public/ModelingTaskTypes.h"
-#include "Transforms/MultiTransformer.h"
-#include "Selection/PolygonSelectionMechanic.h"
-#include "Mechanics/PlaneDistanceFromHitMechanic.h"
-#include "Mechanics/SpatialCurveDistanceMechanic.h"
-#include "Mechanics/CollectSurfacePathMechanic.h"
-#include "Drawing/PolyEditPreviewMesh.h"
+#include "SingleSelectionTool.h"
+
+#include "GeometryBase.h"
+
 #include "EditMeshPolygonsTool.generated.h"
 
-class FMeshVertexChangeBuilder;
+PREDECLARE_GEOMETRY(class FGroupTopology);
+PREDECLARE_GEOMETRY(struct FGroupTopologySelection);
+
 class UDragAlignmentMechanic;
 class UGroupTopologyStorableSelection;
+class UMeshOpPreviewWithBackgroundCompute; 
+class FMeshVertexChangeBuilder;
+class UPolyEditInsertEdgeActivity;
+class UPolyEditInsertEdgeLoopActivity;
+class UPolyEditExtrudeActivity;
+class UPolyEditInsetOutsetActivity;
+class UPolyEditCutFacesActivity;
+class UPolyEditPlanarProjectionUVActivity;
+class UPolygonSelectionMechanic;
+class UTransformGizmo;	
+class UTransformProxy;
+
 
 /**
  * ToolBuilder
  */
 UCLASS()
-class MESHMODELINGTOOLS_API UEditMeshPolygonsToolBuilder : public UMeshSurfacePointToolBuilder
+class MESHMODELINGTOOLS_API UEditMeshPolygonsToolBuilder : public UInteractiveToolBuilder
 {
 	GENERATED_BODY()
 public:
 	bool bTriangleMode = false;
 
-	virtual UMeshSurfacePointTool* CreateNewTool(const FToolBuilderState& SceneState) const override;
-};
+	virtual bool CanBuildTool(const FToolBuilderState& SceneState) const override;
+	virtual UInteractiveTool* BuildTool(const FToolBuilderState& SceneState) const override;
 
+protected:
+	virtual const FToolTargetTypeRequirements& GetTargetRequirements() const override;
+};
 
 
 UENUM()
@@ -83,19 +92,18 @@ public:
 };
 
 
-
-
-
-
 UENUM()
 enum class EEditMeshPolygonsToolActions
 {
 	NoAction,
-	PlaneCut,
+	CancelCurrent,
 	Extrude,
-	Offset,
-	Inset,
-	Outset,
+	InsetOutset,
+	InsertEdge,
+	InsertEdgeLoop,
+	Complete,
+
+	PlaneCut,
 	Merge,
 	Delete,
 	CutFaces,
@@ -118,11 +126,8 @@ enum class EEditMeshPolygonsToolActions
 	SplitSingleEdge,
 	FlipSingleEdge,
 	CollapseSingleEdge
-
 };
 
-
-
 UCLASS()
 class MESHMODELINGTOOLS_API UEditMeshPolygonsActionModeToolBuilder : public UEditMeshPolygonsToolBuilder
 {
@@ -130,12 +135,9 @@ class MESHMODELINGTOOLS_API UEditMeshPolygonsActionModeToolBuilder : public UEdi
 public:
 	EEditMeshPolygonsToolActions StartupAction = EEditMeshPolygonsToolActions::Extrude;
 
-	virtual void InitializeNewTool(UMeshSurfacePointTool* Tool, const FToolBuilderState& SceneState) const override;
+	virtual UInteractiveTool* BuildTool(const FToolBuilderState& SceneState) const override;
 };
 
-
-
-
 UENUM()
 enum class EEditMeshPolygonsToolSelectionMode
 {
@@ -154,16 +156,11 @@ class MESHMODELINGTOOLS_API UEditMeshPolygonsSelectionModeToolBuilder : public U
 public:
 	EEditMeshPolygonsToolSelectionMode SelectionMode = EEditMeshPolygonsToolSelectionMode::Faces;
 
-	virtual void InitializeNewTool(UMeshSurfacePointTool* Tool, const FToolBuilderState& SceneState) const override;
+	virtual UInteractiveTool* BuildTool(const FToolBuilderState& SceneState) const override;
 };
 
 
 
-
-
-
-
-
 UCLASS()
 class MESHMODELINGTOOLS_API UEditMeshPolygonsToolActionPropertySet : public UInteractiveToolPropertySet
 {
@@ -178,7 +175,7 @@ public:
 };
 
 
-
+/** PolyEdit Actions */
 UCLASS()
 class MESHMODELINGTOOLS_API UEditMeshPolygonsToolActions : public UEditMeshPolygonsToolActionPropertySet
 {
@@ -188,17 +185,9 @@ public:
 	UFUNCTION(CallInEditor, Category = FaceEdits, meta = (DisplayName = "Extrude", DisplayPriority = 1))
 	void Extrude() { PostAction(EEditMeshPolygonsToolActions::Extrude); }
 
-	/** Offset the current set of selected faces. Click in viewport to confirm offset distance. */
-	UFUNCTION(CallInEditor, Category = FaceEdits, meta = (DisplayName = "Offset", DisplayPriority = 2))
-	void Offset() { PostAction(EEditMeshPolygonsToolActions::Offset); }
-
-	/** Inset the current set of selected faces. Click in viewport to confirm inset distance. */
-	UFUNCTION(CallInEditor, Category = FaceEdits, meta = (DisplayName = "Inset", DisplayPriority = 3))
-	void Inset() { PostAction(EEditMeshPolygonsToolActions::Inset);	}
-
-	/** Outset the current set of selected faces. Click in viewport to confirm outset distance. */
-	UFUNCTION(CallInEditor, Category = FaceEdits, meta = (DisplayName = "Outset", DisplayPriority = 3))
-	void Outset() { PostAction(EEditMeshPolygonsToolActions::Outset); }
+	/** Inset/Outset the current set of selected faces. Click in viewport to confirm inset/outset distance. */
+	UFUNCTION(CallInEditor, Category = FaceEdits, meta = (DisplayName = "Inset/Outset", DisplayPriority = 3))
+	void InsetOutset() { PostAction(EEditMeshPolygonsToolActions::InsetOutset);	}
 
 	/** Merge the current set of selected faces into a single face. */
 	UFUNCTION(CallInEditor, Category = FaceEdits, meta = (DisplayName = "Merge", DisplayPriority = 4))
@@ -249,17 +238,9 @@ public:
 	UFUNCTION(CallInEditor, Category = TriangleEdits, meta = (DisplayName = "Extrude", DisplayPriority = 1))
 	void Extrude() { PostAction(EEditMeshPolygonsToolActions::Extrude); }
 
-	/** Offset the current set of selected faces. Click in viewport to confirm offset distance. */
-	UFUNCTION(CallInEditor, Category = TriangleEdits, meta = (DisplayName = "Offset", DisplayPriority = 2))
-	void Offset() { PostAction(EEditMeshPolygonsToolActions::Offset); }
-
 	/** Inset the current set of selected faces. Click in viewport to confirm inset distance. */
-	UFUNCTION(CallInEditor, Category = TriangleEdits, meta = (DisplayName = "Inset", DisplayPriority = 3))
-	void Inset() { PostAction(EEditMeshPolygonsToolActions::Inset);	}
-
-	/** Outset the current set of selected faces. Click in viewport to confirm outset distance. */
-	UFUNCTION(CallInEditor, Category = TriangleEdits, meta = (DisplayName = "Outset", DisplayPriority = 3))
-	void Outset() { PostAction(EEditMeshPolygonsToolActions::Outset); }
+	UFUNCTION(CallInEditor, Category = TriangleEdits, meta = (DisplayName = "Inset/Outset", DisplayPriority = 3))
+	void InsetOutset() { PostAction(EEditMeshPolygonsToolActions::InsetOutset);	}
 
 	/** Delete the current set of selected faces */
 	UFUNCTION(CallInEditor, Category = TriangleEdits, meta = (DisplayName = "Delete", DisplayPriority = 4))
@@ -318,15 +299,23 @@ class MESHMODELINGTOOLS_API UEditMeshPolygonsToolEdgeActions : public UEditMeshP
 {
 	GENERATED_BODY()
 public:
-	UFUNCTION(CallInEditor, Category = EdgeEdits, meta = (DisplayName = "Weld", DisplayPriority = 1))
+	UFUNCTION(CallInEditor, Category = EdgeEdits, meta = (DisplayName = "InsertEdgeLoop", DisplayPriority = 1))
+	void InsertEdgeLoop() { PostAction(EEditMeshPolygonsToolActions::InsertEdgeLoop); }
+
+	UFUNCTION(CallInEditor, Category = EdgeEdits, meta = (DisplayName = "Insert Edge", DisplayPriority = 2))
+	void InsertEdge() { PostAction(EEditMeshPolygonsToolActions::InsertEdge); }
+
+	UFUNCTION(CallInEditor, Category = EdgeEdits, meta = (DisplayName = "Weld", DisplayPriority = 3))
 	void Weld() { PostAction(EEditMeshPolygonsToolActions::WeldEdges); }
 
-	UFUNCTION(CallInEditor, Category = EdgeEdits, meta = (DisplayName = "Straighten", DisplayPriority = 1))
+	UFUNCTION(CallInEditor, Category = EdgeEdits, meta = (DisplayName = "Straighten", DisplayPriority = 4))
 	void Straighten() { PostAction(EEditMeshPolygonsToolActions::StraightenEdge); }
 
 	/** Fill the adjacent hole for any selected boundary edges */
-	UFUNCTION(CallInEditor, Category = EdgeEdits, meta = (DisplayName = "Fill Hole", DisplayPriority = 1))
+	UFUNCTION(CallInEditor, Category = EdgeEdits, meta = (DisplayName = "Fill Hole", DisplayPriority = 5))
 	void FillHole()	{ PostAction(EEditMeshPolygonsToolActions::FillHole); }
+
+	
 };
 
 
@@ -354,150 +343,30 @@ public:
 };
 
 
-
-
-UENUM()
-enum class EPolyEditExtrudeDirection
-{
-	SelectionNormal,
-	WorldX,
-	WorldY,
-	WorldZ,
-	LocalX,
-	LocalY,
-	LocalZ
-};
-
-
-UCLASS()
-class MESHMODELINGTOOLS_API UPolyEditExtrudeProperties : public UInteractiveToolPropertySet
-{
-	GENERATED_BODY()
-
-public:
-	UPROPERTY(EditAnywhere, Category = Extrude)
-	EPolyEditExtrudeDirection Direction = EPolyEditExtrudeDirection::SelectionNormal;
-
-	/** Controls whether extruding an entire patch should create a solid or an open shell */
-	UPROPERTY(EditAnywhere, Category = Extrude)
-	bool bShellsToSolids = true;
-};
-
-
-
-UCLASS()
-class MESHMODELINGTOOLS_API UPolyEditOffsetProperties : public UInteractiveToolPropertySet
-{
-	GENERATED_BODY()
-
-public:
-	/** Offset by averaged face normals instead of per-vertex normals */
-	UPROPERTY(EditAnywhere, Category = Offset)
-	bool bUseFaceNormals = false;
-};
-
-
-
 /**
- * Settings for Inset operation
+ * TODO: This is currently a separate action set so that we can show/hide it depending on whether
+ * we have an activity running. We should have a cleaner alternative.
  */
 UCLASS()
-class MESHMODELINGTOOLS_API UPolyEditInsetProperties : public UInteractiveToolPropertySet
-{
-	GENERATED_BODY()
-
-public:
-	/** Determines whether vertices in inset region should be projected back onto input surface */
-	UPROPERTY(EditAnywhere, Category = Inset)
-	bool bReproject = true;
-
-	/** Amount of smoothing applied to inset boundary */
-	UPROPERTY(EditAnywhere, Category = Inset, meta = (UIMin = "0.0", UIMax = "1.0", EditCondition = "bBoundaryOnly == false"))
-	float Softness = 0.5;
-
-	/** Controls whether inset operation will move interior vertices as well as border vertices */
-	UPROPERTY(EditAnywhere, Category = Inset, AdvancedDisplay)
-	bool bBoundaryOnly = false;
-
-	/** Tweak area scaling when solving for interior vertices */
-	UPROPERTY(EditAnywhere, Category = Inset, AdvancedDisplay, meta = (UIMin = "0.0", UIMax = "1.0", EditCondition = "bBoundaryOnly == false"))
-	float AreaScale = true;
-};
-
-
-
-UCLASS()
-class MESHMODELINGTOOLS_API UPolyEditOutsetProperties : public UInteractiveToolPropertySet
-{
-	GENERATED_BODY()
-
-public:
-	/** Amount of smoothing applied to outset boundary */
-	UPROPERTY(EditAnywhere, Category = Inset, meta = (UIMin = "0.0", UIMax = "1.0", EditCondition = "bBoundaryOnly == false"))
-	float Softness = 0.5;
-
-	/** Controls whether outset operation will move interior vertices as well as border vertices */
-	UPROPERTY(EditAnywhere, Category = Inset, AdvancedDisplay)
-	bool bBoundaryOnly = false;
-
-	/** Tweak area scaling when solving for interior vertices */
-	UPROPERTY(EditAnywhere, Category = Inset, AdvancedDisplay, meta = (UIMin = "0.0", UIMax = "1.0", EditCondition = "bBoundaryOnly == false"))
-	float AreaScale = true;
-};
-
-
-
-
-
-
-UENUM()
-enum class EPolyEditCutPlaneOrientation
-{
-	FaceNormals,
-	ViewDirection
-};
-
-
-
-UCLASS()
-class MESHMODELINGTOOLS_API UPolyEditCutProperties : public UInteractiveToolPropertySet
+class MESHMODELINGTOOLS_API UEditMeshPolygonsToolCancelAction : public UEditMeshPolygonsToolActionPropertySet
 {
 	GENERATED_BODY()
-
 public:
-	UPROPERTY(EditAnywhere, Category = Cut)
-	EPolyEditCutPlaneOrientation Orientation = EPolyEditCutPlaneOrientation::FaceNormals;
-
-	UPROPERTY(EditAnywhere, Category = Cut)
-	bool bSnapToVertices = true;
+	UFUNCTION(CallInEditor, Category = CurrentOperation, meta = (DisplayName = "Cancel", DisplayPriority = 1))
+	void Done() { PostAction(EEditMeshPolygonsToolActions::CancelCurrent); }
 };
 
 
 
-
-UCLASS()
-class MESHMODELINGTOOLS_API UPolyEditSetUVProperties : public UInteractiveToolPropertySet
-{
-	GENERATED_BODY()
-
-public:
-	UPROPERTY(EditAnywhere, Category = PlanarProjectUV)
-	bool bShowMaterial = false;
-};
-
-
-
-
 /**
  *
  */
 UCLASS()
-class MESHMODELINGTOOLS_API UEditMeshPolygonsTool : public UMeshSurfacePointTool, public IClickBehaviorTarget
+class MESHMODELINGTOOLS_API UEditMeshPolygonsTool : public USingleSelectionTool, public IToolActivityHost, public IMeshVertexCommandChangeTarget
 {
 	GENERATED_BODY()
-	using FTransform3d = UE::Geometry::FTransform3d;
 	using FFrame3d = UE::Geometry::FFrame3d;
-	using FAxisAlignedBox3d = UE::Geometry::FAxisAlignedBox3d;
+
 public:
 	UEditMeshPolygonsTool();
 
@@ -513,6 +382,8 @@ public:
 		StoredToolSelection = StoredToolSelectionIn; 
 	}
 
+	virtual void SetWorld(UWorld* World) { this->TargetWorld = World; }
+
 	virtual void Setup() override;
 	virtual void Shutdown(EToolShutdownType ShutdownType) override;
 
@@ -523,122 +394,117 @@ public:
 	virtual bool HasCancel() const override { return true; }
 	virtual bool HasAccept() const override { return true; }
 
-	// UMeshSurfacePointTool API
-	virtual bool HitTest(const FRay& Ray, FHitResult& OutHit) override;
-	virtual void OnBeginDrag(const FRay& Ray) override;
-	virtual void OnUpdateDrag(const FRay& Ray) override;
-	virtual void OnEndDrag(const FRay& Ray) override;
-	virtual bool OnUpdateHover(const FInputDeviceRay& DevicePos) override;
-	virtual void OnEndHover() override;
-
-	// IClickDragBehaviorTarget API
-	virtual FInputRayHit CanBeginClickDragSequence(const FInputDeviceRay& PressPos) override;
-
-	// IClickBehaviorTarget API
-	virtual FInputRayHit IsHitByClick(const FInputDeviceRay& ClickPos) override;
-	virtual void OnClicked(const FInputDeviceRay& ClickPos) override;
-
 	// IInteractiveToolCameraFocusAPI implementation
 	virtual FBox GetWorldSpaceFocusBox() override;
 	virtual bool GetWorldSpaceFocusPoint(const FRay& WorldRay, FVector& PointOut) override;
 
+	// IToolActivityHost
+	virtual void NotifyActivitySelfEnded(UInteractiveToolActivity* Activity) override;
+
+	// IMeshVertexCommandChangeTarget
+	virtual void ApplyChange(const FMeshVertexChange* Change, bool bRevert) override;
+
 public:
 
 	virtual void RequestAction(EEditMeshPolygonsToolActions ActionType);
 
+	void SetActionButtonsVisibility(bool bVisible);
+
 protected:
 	// If bTriangleMode = true, then we use a per-triangle FTriangleGroupTopology instead of polygroup topology.
 	// This allows low-level mesh editing with mainly the same code, at a significant cost in overhead.
 	// This is a fundamental mode switch, must be set before ::Setup() is called!
 	bool bTriangleMode;		
 
+	TObjectPtr<UWorld> TargetWorld = nullptr;
+
 	UPROPERTY()
-	TObjectPtr<UDynamicMeshComponent> DynamicMeshComponent = nullptr;
+	TObjectPtr<UMeshOpPreviewWithBackgroundCompute> Preview = nullptr;
 
 	UPROPERTY()
-	TObjectPtr<UPolyEditCommonProperties> CommonProps;
+	TObjectPtr<UPolyEditCommonProperties> CommonProps = nullptr;
 
 	UPROPERTY()
-	TObjectPtr<UEditMeshPolygonsToolActions> EditActions;
+	TObjectPtr<UEditMeshPolygonsToolActions> EditActions = nullptr;
 	UPROPERTY()
-	TObjectPtr<UEditMeshPolygonsToolActions_Triangles> EditActions_Triangles;
+	TObjectPtr<UEditMeshPolygonsToolActions_Triangles> EditActions_Triangles = nullptr;
 
 	UPROPERTY()
-	TObjectPtr<UEditMeshPolygonsToolEdgeActions> EditEdgeActions;
+	TObjectPtr<UEditMeshPolygonsToolEdgeActions> EditEdgeActions = nullptr;
 	UPROPERTY()
-	TObjectPtr<UEditMeshPolygonsToolEdgeActions_Triangles> EditEdgeActions_Triangles;
+	TObjectPtr<UEditMeshPolygonsToolEdgeActions_Triangles> EditEdgeActions_Triangles = nullptr;
 
 	UPROPERTY()
-	TObjectPtr<UEditMeshPolygonsToolUVActions> EditUVActions;
+	TObjectPtr<UEditMeshPolygonsToolUVActions> EditUVActions = nullptr;
 
 	UPROPERTY()
-	TObjectPtr<UPolyEditExtrudeProperties> ExtrudeProperties;
+	TObjectPtr<UEditMeshPolygonsToolCancelAction> CancelAction = nullptr;
 
+	/**
+	 * Activity objects that handle multi-interaction operations
+	 */
 	UPROPERTY()
-	TObjectPtr<UPolyEditOffsetProperties> OffsetProperties;
-
+	TObjectPtr<UPolyEditExtrudeActivity> ExtrudeActivity = nullptr;
 	UPROPERTY()
-	TObjectPtr<UPolyEditInsetProperties> InsetProperties;
-
+	TObjectPtr<UPolyEditInsetOutsetActivity> InsetOutsetActivity = nullptr;
 	UPROPERTY()
-	TObjectPtr<UPolyEditOutsetProperties> OutsetProperties;
-
+	TObjectPtr<UPolyEditCutFacesActivity> CutFacesActivity = nullptr;
 	UPROPERTY()
-	TObjectPtr<UPolyEditCutProperties> CutProperties;
-
+	TObjectPtr<UPolyEditPlanarProjectionUVActivity> PlanarProjectionUVActivity = nullptr;
+	UPROPERTY()
+	TObjectPtr<UPolyEditInsertEdgeActivity> InsertEdgeActivity = nullptr;
 	UPROPERTY()
-	TObjectPtr<UPolyEditSetUVProperties> SetUVProperties;
+	TObjectPtr<UPolyEditInsertEdgeLoopActivity> InsertEdgeLoopActivity = nullptr;
+
+	/**
+	 * Points to one of the activities when it is active
+	 */
+	TObjectPtr<UInteractiveToolActivity> CurrentActivity = nullptr;
 
+	TSharedPtr<UE::Geometry::FDynamicMesh3> CurrentMesh;
+	TSharedPtr<UE::Geometry::FGroupTopology> Topology;
+	TSharedPtr<UE::Geometry::FDynamicMeshAABBTree3> MeshSpatial;
 
 	UPROPERTY()
 	TObjectPtr<UPolygonSelectionMechanic> SelectionMechanic;
 
 	UPROPERTY()
-	TObjectPtr<const UGroupTopologyStorableSelection> StoredToolSelection = nullptr;
-	bool IsStoredToolSelectionUsable(const UGroupTopologyStorableSelection* StoredSelection);
+	TObjectPtr<UDragAlignmentMechanic> DragAlignmentMechanic = nullptr;
 
-	bool bSelectionStateDirty = false;
-	void OnSelectionModifiedEvent();
+	UPROPERTY()
+	TObjectPtr<const UGroupTopologyStorableSelection> StoredToolSelection = nullptr;
 
 	UPROPERTY()
-	TObjectPtr<UMultiTransformer> MultiTransformer = nullptr;
+	TObjectPtr<UTransformGizmo> TransformGizmo = nullptr;
 
 	UPROPERTY()
-	TObjectPtr<UDragAlignmentMechanic> DragAlignmentMechanic = nullptr;
+	TObjectPtr<UTransformProxy> TransformProxy = nullptr;
+
+	FText DefaultMessage;
+
+	bool IsStoredToolSelectionUsable(const UGroupTopologyStorableSelection* StoredSelection);
+	bool bSelectionStateDirty = false;
+	void OnSelectionModifiedEvent();
 
-	void OnMultiTransformerTransformBegin();
-	void OnMultiTransformerTransformUpdate();
-	void OnMultiTransformerTransformEnd();
-	void UpdateMultiTransformerFrame(const FFrame3d* UseFrame = nullptr);
+	void OnBeginGizmoTransform(UTransformProxy* Proxy);
+	void OnEndGizmoTransform(UTransformProxy* Proxy);
+	void OnGizmoTransformChanged(UTransformProxy* Proxy, FTransform Transform);
+	void UpdateGizmoFrame(const FFrame3d* UseFrame = nullptr);
 	FFrame3d LastGeometryFrame;
 	FFrame3d LastTransformerFrame;
 	FFrame3d LockedTransfomerFrame;
+	bool bInGizmoDrag = false;
 
-	// realtime visualization
-	void OnDynamicMeshComponentChanged();
-	FDelegateHandle OnDynamicMeshComponentChangedHandle;
-
-
-	// camera state at last render
 	UE::Geometry::FTransform3d WorldTransform;
-	FViewCameraState CameraState;
-
-	// True for the duration of UI click+drag
-	bool bInDrag;
 
 	FFrame3d InitialGizmoFrame;
 	FVector3d InitialGizmoScale;
-	void CacheUpdate_Gizmo();
 	bool bGizmoUpdatePending = false;
 	FFrame3d LastUpdateGizmoFrame;
 	FVector3d LastUpdateGizmoScale;
 	bool bLastUpdateUsedWorldFrame = false;
 	void ComputeUpdate_Gizmo();
 
-	TUniquePtr<FGroupTopology> Topology;
-	void PrecomputeTopology();
-
-	UE::Geometry::FDynamicMeshAABBTree3 MeshSpatial;
 	UE::Geometry::FDynamicMeshAABBTree3& GetSpatial();
 	bool bSpatialDirty;
 
@@ -647,36 +513,23 @@ protected:
 
 	EEditMeshPolygonsToolActions PendingAction = EEditMeshPolygonsToolActions::NoAction;
 
-	enum class ECurrentToolMode
-	{
-		TransformSelection,
-		ExtrudeSelection,
-		OffsetSelection,
-		InsetSelection,
-		OutsetSelection,
-		CutSelection,
-		SetUVs
-	};
-	ECurrentToolMode CurrentToolMode = ECurrentToolMode::TransformSelection;
-	int32 CurrentOperationTimestamp = 1;
-	bool CheckInOperation(int32 Timestamp) const { return CurrentOperationTimestamp == Timestamp; }
+	int32 ActivityTimestamp = 1;
 
-	void BeginExtrude(bool bIsNormalOffset);
-	void ApplyExtrude(bool bIsOffset);
-	void RestartExtrude();
-	FVector3d GetExtrudeDirection() const;
+	void StartActivity(TObjectPtr<UInteractiveToolActivity> Activity);
+	void EndCurrentActivity();
+	void SetActionButtonPanelsVisible(bool bVisible);
 
-	void BeginInset(bool bOutset);
-	void ApplyInset(bool bOutset);
+	// Emit an undoable change to CurrentMesh and update related structures (preview, spatial, etc)
+	void EmitCurrentMeshChangeAndUpdate(const FText& TransactionLabel,
+		TUniquePtr<UE::Geometry::FDynamicMeshChange> MeshChangeIn,
+		const UE::Geometry::FGroupTopologySelection& OutputSelection,
+		bool bTopologyChanged);
 
-	void BeginCutFaces();
-	void ApplyCutFaces();
+	// Emit an undoable start of an activity
+	void EmitActivityStart(const FText& TransactionLabel);
 
-	void BeginSetUVs();
-	void UpdateSetUVS();
-	void ApplySetUVs();
+	void UpdateGizmoVisibility();
 
-	void ApplyPlaneCut();
 	void ApplyMerge();
 	void ApplyDelete();
 	void ApplyRecalcNormals();
@@ -699,7 +552,7 @@ protected:
 	FFrame3d ActiveSelectionFrameLocal;
 	FFrame3d ActiveSelectionFrameWorld;
 	TArray<int32> ActiveTriangleSelection;
-	FAxisAlignedBox3d ActiveSelectionBounds;
+	UE::Geometry::FAxisAlignedBox3d ActiveSelectionBounds;
 
 	struct FSelectedEdge
 	{
@@ -708,11 +561,6 @@ protected:
 	};
 	TArray<FSelectedEdge> ActiveEdgeSelection;
 
-	bool bPreviewUpdatePending = false;
-
-	UPROPERTY()
-	TObjectPtr<UPolyEditPreviewMesh> EditPreview;
-
 	enum class EPreviewMaterialType
 	{
 		SourceMaterials, PreviewMaterial, UVMaterial
@@ -720,47 +568,30 @@ protected:
 	void UpdateEditPreviewMaterials(EPreviewMaterialType MaterialType);
 	EPreviewMaterialType CurrentPreviewMaterial;
 
-	UPROPERTY()
-	TObjectPtr<UPlaneDistanceFromHitMechanic> ExtrudeHeightMechanic = nullptr;
-	UPROPERTY()
-	TObjectPtr<USpatialCurveDistanceMechanic> CurveDistMechanic = nullptr;
-	UPROPERTY()
-	TObjectPtr<UCollectSurfacePathMechanic> SurfacePathMechanic = nullptr;
 
 	//
 	// data for current drag
 	//
-
 	UE::Geometry::FGroupTopologyDeformer LinearDeformer;
-	void UpdateDeformerFromSelection(const FGroupTopologySelection& Selection);
-
-
+	void UpdateDeformerFromSelection(const UE::Geometry::FGroupTopologySelection& Selection);
 
 	FMeshVertexChangeBuilder* ActiveVertexChange;
-	void BeginChange();
-	void EndChange();
-	void UpdateChangeFromROI(bool bFinal);
+	void UpdateDeformerChangeFromROI(bool bFinal);
+	void BeginDeformerChange();
+	void EndDeformerChange();
 
 	bool BeginMeshFaceEditChange();
-	bool BeginMeshFaceEditChangeWithPreview();
-	void CompleteMeshEditChange(const FText& TransactionLabel, TUniquePtr<FToolCommandChange> EditChange, const FGroupTopologySelection& OutputSelection);
-	void CancelMeshEditChange();
 
 	bool BeginMeshEdgeEditChange();
 	bool BeginMeshBoundaryEdgeEditChange(bool bOnlySimple);
 	bool BeginMeshEdgeEditChange(TFunctionRef<bool(int32)> GroupEdgeIDFilterFunc);
 
-	void AfterTopologyEdit();
+	void UpdateFromCurrentMesh(bool bGroupTopologyModified);
 	int32 ModifiedTopologyCounter = 0;
 	bool bWasTopologyEdited = false;
 
-	void SetActionButtonPanelsVisible(bool bVisible);
-
-	friend class FEditPolygonsTopologyPreEditChange;
-	friend class FEditPolygonsTopologyPostEditChange;
-	friend class FBeginInteractivePolyEditChange;
-
-
+	friend class FEditMeshPolygonsToolMeshChange;
+	friend class FPolyEditActivityStartChange;
 
 	// custom setup support
 	friend class UEditMeshPolygonsSelectionModeToolBuilder;
@@ -770,41 +601,46 @@ protected:
 };
 
 
-
-
-
-class MESHMODELINGTOOLS_API FEditPolygonsTopologyPreEditChange : public FToolCommandChange
+/**
+ * Wraps a FDynamicMeshChange so that it can be expired and so that other data
+ * structures in the tool can be updated.
+ */
+class MESHMODELINGTOOLS_API FEditMeshPolygonsToolMeshChange : public FToolCommandChange
 {
 public:
-	virtual void Apply(UObject* Object) override;
-	virtual void Revert(UObject* Object) override;
-	virtual FString ToString() const override;
-};
+	FEditMeshPolygonsToolMeshChange(TUniquePtr<UE::Geometry::FDynamicMeshChange> MeshChangeIn, bool bGroupTopologyModified)
+		: MeshChange(MoveTemp(MeshChangeIn))
+		, bGroupTopologyModified(bGroupTopologyModified)
+	{};
 
-class MESHMODELINGTOOLS_API FEditPolygonsTopologyPostEditChange : public FToolCommandChange
-{
-public:
 	virtual void Apply(UObject* Object) override;
 	virtual void Revert(UObject* Object) override;
 	virtual FString ToString() const override;
+
+protected:
+	TUniquePtr<UE::Geometry::FDynamicMeshChange> MeshChange;
+	bool bGroupTopologyModified;
 };
 
 
+
 /**
- * FBeginInteractivePolyEditChange is used to cancel out of an active action on Undo. No action is taken on Redo
+ * FPolyEditActivityStartChange is used to cancel out of an active action on Undo. 
  * No action is taken on Redo, ie we do not re-start the Tool on Redo.
  */
-class MESHMODELINGTOOLS_API FBeginInteractivePolyEditChange : public FToolCommandChange
+class MESHMODELINGTOOLS_API FPolyEditActivityStartChange : public FToolCommandChange
 {
 public:
-	bool bHaveDoneUndo = false;
-	int32 OperationTimestamp = 0;
-	FBeginInteractivePolyEditChange(int32 CurrentTimestamp)
+	FPolyEditActivityStartChange(int32 ActivityTimestampIn)
 	{
-		OperationTimestamp = CurrentTimestamp;
+		ActivityTimestamp = ActivityTimestampIn;
 	}
 	virtual void Apply(UObject* Object) override {}
 	virtual void Revert(UObject* Object) override;
 	virtual bool HasExpired(UObject* Object) const override;
 	virtual FString ToString() const override;
+
+protected:
+	bool bHaveDoneUndo = false;
+	int32 ActivityTimestamp = 0;
 };
diff --git a/Engine/Plugins/Experimental/MeshModelingToolset/Source/MeshModelingTools/Public/GroupEdgeInsertionTool.h b/Engine/Plugins/Experimental/MeshModelingToolset/Source/MeshModelingTools/Public/GroupEdgeInsertionTool.h
deleted file mode 100644
diff --git a/Engine/Plugins/Experimental/MeshModelingToolset/Source/MeshModelingTools/Public/ToolActivities/PolyEditActivityContext.h b/Engine/Plugins/Experimental/MeshModelingToolset/Source/MeshModelingTools/Public/ToolActivities/PolyEditActivityContext.h
new file mode 100644
--- /dev/null
+++ b/Engine/Plugins/Experimental/MeshModelingToolset/Source/MeshModelingTools/Public/ToolActivities/PolyEditActivityContext.h
@@ -0,0 +1,78 @@
+// Copyright Epic Games, Inc. All Rights Reserved.
+
+#pragma once
+
+#include "CoreMinimal.h"
+
+#include "DynamicMesh/DynamicMeshAABBTree3.h"
+#include "GeometryBase.h"
+
+#include "PolyEditActivityContext.generated.h"
+
+PREDECLARE_GEOMETRY(class FDynamicMesh3);
+PREDECLARE_GEOMETRY(class FDynamicMeshChange);
+PREDECLARE_GEOMETRY(struct FGroupTopologySelection);
+PREDECLARE_GEOMETRY(class FGroupTopology);
+class FToolCommandChange;
+class UPolyEditCommonProperties;
+class UPolygonSelectionMechanic;
+class UMeshOpPreviewWithBackgroundCompute;
+
+UCLASS()
+class MESHMODELINGTOOLS_API UPolyEditActivityContext : public UObject
+{
+	GENERATED_BODY()
+
+public:
+	
+	UPROPERTY()
+	TObjectPtr<UPolyEditCommonProperties> CommonProperties;
+
+	/** 
+	 * The CurrentMesh is the authoritative current version of the mesh, which would be baked back
+	 * to the asset on accept.
+	 */
+	TSharedPtr<UE::Geometry::FDynamicMesh3, ESPMode::ThreadSafe> CurrentMesh;
+
+	/**
+	 * The mesh stored in the preview is not authoritative. It may be altered in various ways as the 
+	 * user previews potential changes, and may be reset back to CurrentMesh if an activity
+	 * is cancelled.
+	 */
+	UPROPERTY()
+	TObjectPtr<UMeshOpPreviewWithBackgroundCompute> Preview;
+
+	/** 
+	 * The activity should update CurrentTopology and MeshSpatial if it alters CurrentMesh. 
+	 */
+	TSharedPtr<UE::Geometry::FGroupTopology, ESPMode::ThreadSafe> CurrentTopology;
+	TSharedPtr<UE::Geometry::FDynamicMeshAABBTree3, ESPMode::ThreadSafe> MeshSpatial;
+
+	bool bTriangleMode = false;
+
+	/** 
+	 * The activity may use the selection mechanic to get (or alter) the current selection, though
+	 * if selection is just being changed at the end of the activity, it should probably be done
+	 * through EmitCurrentMeshChangeAndUpdate so that it is lumped with the same undo transaction.
+	 */
+	UPROPERTY()
+	UPolygonSelectionMechanic* SelectionMechanic;
+
+	/**
+	 * Given a change to CurrentMesh (already applied and stored in MeshChangeIn), use this
+	 * function to properly emit it from the host. Be aware that the host will update any 
+	 * related data structures at this time (spatial, preview, etc) so that state is consistent
+	 * with undo/redo.
+	 */
+	TUniqueFunction<void (const FText& TransactionLabel,
+			TUniquePtr<UE::Geometry::FDynamicMeshChange> MeshChangeIn,
+			const UE::Geometry::FGroupTopologySelection& OutputSelection,
+			bool bGroupTopologyChanged)> 
+		EmitCurrentMeshChangeAndUpdate;
+
+	/**
+	 * If an activity starts running, it should use this function to have the host emit an
+	 * appropriate undoable transaction with the given transaction label.
+	 */
+	TUniqueFunction<void(const FText& TransactionLabel)> EmitActivityStart;
+};
diff --git a/Engine/Plugins/Experimental/MeshModelingToolset/Source/MeshModelingTools/Public/ToolActivities/PolyEditActivityUtil.h b/Engine/Plugins/Experimental/MeshModelingToolset/Source/MeshModelingTools/Public/ToolActivities/PolyEditActivityUtil.h
new file mode 100644
--- /dev/null
+++ b/Engine/Plugins/Experimental/MeshModelingToolset/Source/MeshModelingTools/Public/ToolActivities/PolyEditActivityUtil.h
@@ -0,0 +1,25 @@
+// Copyright Epic Games, Inc. All Rights Reserved.
+
+#pragma once
+
+#include "CoreMinimal.h"
+
+class UInteractiveTool;
+class UPolyEditActivityContext;
+class UPolyEditPreviewMesh;
+
+namespace UE {
+namespace Geometry {
+namespace PolyEditActivityUtil {
+
+	enum class EPreviewMaterialType
+	{
+		SourceMaterials, PreviewMaterial, UVMaterial
+	};
+
+	UPolyEditPreviewMesh* CreatePolyEditPreviewMesh(UInteractiveTool& Tool, const UPolyEditActivityContext& ActivityContext);
+
+	void UpdatePolyEditPreviewMaterials(UInteractiveTool& Tool, const UPolyEditActivityContext& ActivityContext,
+		UPolyEditPreviewMesh& EditPreviewMesh, EPreviewMaterialType MaterialType);
+
+}}}
\ No newline at end of file
diff --git a/Engine/Plugins/Experimental/MeshModelingToolset/Source/MeshModelingTools/Public/ToolActivities/PolyEditCutFacesActivity.h b/Engine/Plugins/Experimental/MeshModelingToolset/Source/MeshModelingTools/Public/ToolActivities/PolyEditCutFacesActivity.h
new file mode 100644
--- /dev/null
+++ b/Engine/Plugins/Experimental/MeshModelingToolset/Source/MeshModelingTools/Public/ToolActivities/PolyEditCutFacesActivity.h
@@ -0,0 +1,89 @@
+// Copyright Epic Games, Inc. All Rights Reserved.
+
+#pragma once
+
+#include "CoreMinimal.h"
+
+#include "BaseBehaviors/BehaviorTargetInterfaces.h"
+#include "InteractiveToolActivity.h"
+#include "ToolContextInterfaces.h" // FViewCameraState
+
+#include "PolyEditCutFacesActivity.generated.h"
+
+class UPolyEditActivityContext;
+class UPolyEditPreviewMesh;
+class UCollectSurfacePathMechanic;
+
+UENUM()
+enum class EPolyEditCutPlaneOrientation
+{
+	FaceNormals,
+	ViewDirection
+};
+
+UCLASS()
+class MESHMODELINGTOOLS_API UPolyEditCutProperties : public UInteractiveToolPropertySet
+{
+	GENERATED_BODY()
+
+public:
+	UPROPERTY(EditAnywhere, Category = Cut)
+	EPolyEditCutPlaneOrientation Orientation = EPolyEditCutPlaneOrientation::FaceNormals;
+
+	UPROPERTY(EditAnywhere, Category = Cut)
+	bool bSnapToVertices = true;
+};
+
+
+/**
+ *
+ */
+UCLASS()
+class MESHMODELINGTOOLS_API UPolyEditCutFacesActivity : public UInteractiveToolActivity,
+	public IClickBehaviorTarget, public IHoverBehaviorTarget
+{
+	GENERATED_BODY()
+
+public:
+	// IInteractiveToolActivity
+	virtual void Setup(UInteractiveTool* ParentTool) override;
+	virtual void Shutdown(EToolShutdownType ShutdownType) override;
+	virtual bool CanStart() const override;
+	virtual EToolActivityStartResult Start() override;
+	virtual bool IsRunning() const override { return bIsRunning; }
+	virtual bool CanAccept() const override;
+	virtual EToolActivityEndResult End(EToolShutdownType) override;
+	virtual void Render(IToolsContextRenderAPI* RenderAPI) override;
+	virtual void Tick(float DeltaTime) override;
+
+	// IClickBehaviorTarget API
+	virtual FInputRayHit IsHitByClick(const FInputDeviceRay& ClickPos) override;
+	virtual void OnClicked(const FInputDeviceRay& ClickPos) override;
+
+	// IHoverBehaviorTarget implementation
+	virtual FInputRayHit BeginHoverSequenceHitTest(const FInputDeviceRay& PressPos) override;
+	virtual void OnBeginHover(const FInputDeviceRay& DevicePos) override {}
+	virtual bool OnUpdateHover(const FInputDeviceRay& DevicePos) override;
+	virtual void OnEndHover() override {}
+
+protected:
+	void Clear();
+	void BeginCutFaces();
+	void ApplyCutFaces();
+
+	UPROPERTY()
+	UPolyEditCutProperties* CutProperties;
+
+	UPROPERTY()
+	UPolyEditPreviewMesh* EditPreview;
+
+	UPROPERTY()
+	UCollectSurfacePathMechanic* SurfacePathMechanic = nullptr;
+
+	UPROPERTY()
+	TObjectPtr<UPolyEditActivityContext> ActivityContext;
+
+	bool bIsRunning = false;
+
+	FViewCameraState CameraState;
+};
diff --git a/Engine/Plugins/Experimental/MeshModelingToolset/Source/MeshModelingTools/Public/ToolActivities/PolyEditExtrudeActivity.h b/Engine/Plugins/Experimental/MeshModelingToolset/Source/MeshModelingTools/Public/ToolActivities/PolyEditExtrudeActivity.h
new file mode 100644
--- /dev/null
+++ b/Engine/Plugins/Experimental/MeshModelingToolset/Source/MeshModelingTools/Public/ToolActivities/PolyEditExtrudeActivity.h
@@ -0,0 +1,115 @@
+// Copyright Epic Games, Inc. All Rights Reserved.
+
+#pragma once
+
+#include "CoreMinimal.h"
+
+#include "BaseBehaviors/BehaviorTargetInterfaces.h"
+#include "InteractiveToolActivity.h"
+#include "FrameTypes.h"
+
+#include "PolyEditExtrudeActivity.generated.h"
+
+class UPolyEditActivityContext;
+class UPolyEditPreviewMesh;
+class UPlaneDistanceFromHitMechanic;
+
+UENUM()
+enum class EPolyEditExtrudeMode
+{
+	//~ TODO: this is what we actually want, but not yet implemented.
+	//~ SelectedFaceNormals,
+
+	SingleDirection,
+
+	//~ These are likely not very useful, but they can sometimes take
+	//~ the place of SelectedFaceNormals, until we implement that.
+	SelectedTriangleNormals,
+	VertexNormals,
+};
+
+UENUM()
+enum class EPolyEditExtrudeDirection
+{
+	SelectionNormal,
+	WorldX,
+	WorldY,
+	WorldZ,
+	LocalX,
+	LocalY,
+	LocalZ
+};
+
+
+UCLASS()
+class MESHMODELINGTOOLS_API UPolyEditExtrudeProperties : public UInteractiveToolPropertySet
+{
+	GENERATED_BODY()
+
+public:
+	UPROPERTY(EditAnywhere, Category = Extrude)
+	EPolyEditExtrudeMode ExtrudeMode = EPolyEditExtrudeMode::SingleDirection;
+
+	UPROPERTY(EditAnywhere, Category = Extrude)
+	EPolyEditExtrudeDirection Direction = EPolyEditExtrudeDirection::SelectionNormal;
+
+	/** Controls whether extruding an entire open-border patch should create a solid or an open shell */
+	UPROPERTY(EditAnywhere, Category = Extrude)
+	bool bShellsToSolids = true;
+};
+
+/**
+ * 
+ */
+UCLASS()
+class MESHMODELINGTOOLS_API UPolyEditExtrudeActivity : public UInteractiveToolActivity, 
+	public IClickBehaviorTarget, public IHoverBehaviorTarget
+{
+	GENERATED_BODY()
+
+public:
+	// IInteractiveToolActivity
+	virtual void Setup(UInteractiveTool* ParentTool) override;
+	virtual void Shutdown(EToolShutdownType ShutdownType) override;
+	virtual bool CanStart() const override;
+	virtual EToolActivityStartResult Start() override;
+	virtual bool IsRunning() const override { return bIsRunning; }
+	virtual bool CanAccept() const override;
+	virtual EToolActivityEndResult End(EToolShutdownType) override;
+	virtual void Render(IToolsContextRenderAPI* RenderAPI) override;
+	virtual void Tick(float DeltaTime) override;
+
+	// IClickBehaviorTarget API
+	virtual FInputRayHit IsHitByClick(const FInputDeviceRay& ClickPos) override;
+	virtual void OnClicked(const FInputDeviceRay& ClickPos) override;
+
+	// IHoverBehaviorTarget implementation
+	virtual FInputRayHit BeginHoverSequenceHitTest(const FInputDeviceRay& PressPos) override;
+	virtual void OnBeginHover(const FInputDeviceRay& DevicePos) override {}
+	virtual bool OnUpdateHover(const FInputDeviceRay& DevicePos) override;
+	virtual void OnEndHover() override {}
+
+	UPROPERTY()
+	UPolyEditExtrudeProperties* ExtrudeProperties;
+
+protected:
+	FVector3d GetExtrudeDirection() const;
+	void BeginExtrude();
+	void ApplyExtrude();
+	void Clear();
+
+	UPROPERTY()
+	UPolyEditPreviewMesh* EditPreview;
+
+	UPROPERTY()
+	UPlaneDistanceFromHitMechanic* ExtrudeHeightMechanic = nullptr;
+
+	UPROPERTY()
+	TObjectPtr<UPolyEditActivityContext> ActivityContext;
+
+	bool bIsRunning = false;
+
+	UE::Geometry::FFrame3d ActiveSelectionFrameWorld;
+	float UVScaleFactor = 1.0f;
+	bool bPreviewUpdatePending = false;
+};
diff --git a/Engine/Plugins/Experimental/MeshModelingToolset/Source/MeshModelingTools/Public/ToolActivities/PolyEditInsertEdgeActivity.h b/Engine/Plugins/Experimental/MeshModelingToolset/Source/MeshModelingTools/Public/ToolActivities/PolyEditInsertEdgeActivity.h
new file mode 100644
--- /dev/null
+++ b/Engine/Plugins/Experimental/MeshModelingToolset/Source/MeshModelingTools/Public/ToolActivities/PolyEditInsertEdgeActivity.h
@@ -0,0 +1,201 @@
+// Copyright Epic Games, Inc. All Rights Reserved.
+
+#pragma once
+
+#include "CoreMinimal.h"
+
+#include "BaseBehaviors/BehaviorTargetInterfaces.h"
+#include "ModelingOperators.h" //IDynamicMeshOperatorFactory
+#include "InteractiveTool.h" //UInteractiveToolPropertySet
+#include "InteractiveToolActivity.h"
+#include "InteractiveToolBuilder.h" //UInteractiveToolBuilder
+#include "InteractiveToolChange.h" //FToolCommandChange
+#include "MeshOpPreviewHelpers.h" //FDynamicMeshOpResult
+#include "Operations/GroupEdgeInserter.h"
+#include "Selection/GroupTopologySelector.h"
+#include "SingleSelectionTool.h"
+#include "ToolDataVisualizer.h"
+
+#include "PolyEditInsertEdgeActivity.generated.h"
+
+PREDECLARE_USE_GEOMETRY_CLASS(FDynamicMeshChange);
+class UPolyEditActivityContext;
+
+UENUM()
+enum class EGroupEdgeInsertionMode
+{
+	/** Existing groups will be deleted and new triangles will be created for the new groups.
+	 Keeps topology simple but breaks non-planar groups. */
+	Retriangulate,
+
+	/** Keeps existing triangles and cuts them to create a new path. May result in fragmented triangles over time.*/
+	PlaneCut
+};
+
+UCLASS()
+class MESHMODELINGTOOLS_API UGroupEdgeInsertionProperties : public UInteractiveToolPropertySet
+{
+	GENERATED_BODY()
+
+public:
+	/** Determines how group edges are added to the geometry */
+	UPROPERTY(EditAnywhere, Category = InsertEdge)
+	EGroupEdgeInsertionMode InsertionMode = EGroupEdgeInsertionMode::Retriangulate;
+
+	UPROPERTY(EditAnywhere, Category = InsertEdge)
+	bool bWireframe = true;
+
+	/** How close a new loop edge needs to pass next to an existing vertex to use that vertex rather than creating a new one (used for plane cut). */
+	UPROPERTY(EditAnywhere, Category = InsertEdge, AdvancedDisplay)
+	double VertexTolerance = 0.001;
+};
+
+/** Interactive activity for inserting a group edge into a mesh. */
+UCLASS()
+class MESHMODELINGTOOLS_API UPolyEditInsertEdgeActivity : public UInteractiveToolActivity, 
+	public UE::Geometry::IDynamicMeshOperatorFactory, 
+	public IHoverBehaviorTarget, public IClickBehaviorTarget
+{
+	GENERATED_BODY()
+
+	using FGroupEdgeInserter = UE::Geometry::FGroupEdgeInserter;
+	using FDynamicMesh3 = UE::Geometry::FDynamicMesh3;
+	using FRay3d = UE::Geometry::FRay3d;
+
+	enum class EState
+	{
+		GettingStart,
+		GettingEnd,
+		WaitingForInsertComplete
+	};
+
+public:
+
+	friend class FGroupEdgeInsertionFirstPointChange;
+
+	UPolyEditInsertEdgeActivity() {};
+
+	virtual void OnPropertyModified(UObject* PropertySet, FProperty* Property);
+
+	// IInteractiveToolActivity
+	virtual void Setup(UInteractiveTool* ParentTool) override;
+	virtual void Shutdown(EToolShutdownType ShutdownType) override;
+	virtual bool CanStart() const override;
+	virtual EToolActivityStartResult Start() override;
+	virtual bool IsRunning() const override { return bIsRunning; }
+	virtual bool CanAccept() const override;
+	virtual EToolActivityEndResult End(EToolShutdownType) override;
+	virtual void Render(IToolsContextRenderAPI* RenderAPI) override;
+	virtual void Tick(float DeltaTime) override;
+
+	// IDynamicMeshOperatorFactory
+	virtual TUniquePtr<UE::Geometry::FDynamicMeshOperator> MakeNewOperator() override;
+
+	// IClickBehaviorTarget
+	virtual FInputRayHit IsHitByClick(const FInputDeviceRay& ClickPos) override;
+	virtual void OnClicked(const FInputDeviceRay& ClickPos) override;
+
+	// IHoverBehaviorTarget
+	virtual FInputRayHit BeginHoverSequenceHitTest(const FInputDeviceRay& PressPos) override;
+	virtual void OnBeginHover(const FInputDeviceRay& DevicePos) override {}
+	virtual bool OnUpdateHover(const FInputDeviceRay& DevicePos) override;
+	virtual void OnEndHover() override;
+
+protected:
+	UPROPERTY()
+	TObjectPtr<UGroupEdgeInsertionProperties> Settings = nullptr;
+
+	UPROPERTY()
+	UPolyEditActivityContext* ActivityContext;
+
+	bool bIsRunning = false;
+
+	FTransform TargetTransform;
+	TSharedPtr<FGroupTopologySelector, ESPMode::ThreadSafe> TopologySelector;
+
+	TArray<TPair<FVector3d, FVector3d>> PreviewEdges;
+	TArray<FVector3d> PreviewPoints;
+
+	FViewCameraState CameraState;
+
+	FToolDataVisualizer ExistingEdgesRenderer;
+	FToolDataVisualizer PreviewEdgeRenderer;
+	FGroupTopologySelector::FSelectionSettings TopologySelectorSettings;
+
+
+	// Inputs from user interaction:
+	FGroupEdgeInserter::FGroupEdgeSplitPoint StartPoint;
+	int32 StartTopologyID = FDynamicMesh3::InvalidID;
+	bool bStartIsCorner = false;
+
+	FGroupEdgeInserter::FGroupEdgeSplitPoint EndPoint;
+	int32 EndTopologyID = FDynamicMesh3::InvalidID;
+	bool bEndIsCorner = false;
+
+	int32 CommonGroupID = FDynamicMesh3::InvalidID;
+	int32 CommonBoundaryIndex = FDynamicMesh3::InvalidID;
+
+
+	// State control:
+	EState ToolState = EState::GettingStart;
+
+	bool bShowingBaseMesh = false;
+	bool bLastComputeSucceeded = false;
+	TSharedPtr<FGroupTopology, ESPMode::ThreadSafe> LatestOpTopologyResult;
+	TSharedPtr<TSet<int32>, ESPMode::ThreadSafe> LatestOpChangedTids;
+
+	int32 CurrentChangeStamp = 0;
+
+	void SetupPreview();
+
+	bool TopologyHitTest(const FRay& WorldRay, FVector3d& RayPositionOut, FRay3d* LocalRayOut = nullptr);
+	bool GetHoveredItem(const FRay& WorldRay, FGroupEdgeInserter::FGroupEdgeSplitPoint& PointOut,
+		int32& TopologyElementIDOut, bool& bIsCornerOut, FVector3d& PositionOut, FRay3d* LocalRayOut = nullptr);
+
+	void ConditionallyUpdatePreview(const FGroupEdgeInserter::FGroupEdgeSplitPoint& NewEndPoint, 
+		int32 NewEndTopologyID, bool bNewEndIsCorner, int32 NewCommonGroupID, int32 NewBoundaryIndex);
+
+	void ClearPreview(bool bClearDrawnElements = true);
+
+	void GetCornerTangent(int32 CornerID, int32 GroupID, int32 BoundaryIndex, FVector3d& TangentOut);
+
+	/**
+	 * Expires the tool-associated changes in the undo/redo stack. The ComponentTarget
+	 * changes will stay (we want this).
+	 */
+	inline void ExpireChanges()
+	{
+		++CurrentChangeStamp;
+	}
+};
+
+/**
+ * This should get emitted when selecting the first point in an edge insertion so that we can undo it.
+ */
+class MESHMODELINGTOOLS_API FGroupEdgeInsertionFirstPointChange : public FToolCommandChange
+{
+public:
+	FGroupEdgeInsertionFirstPointChange(int32 CurrentChangeStamp)
+		: ChangeStamp(CurrentChangeStamp)
+	{};
+
+	virtual void Apply(UObject* Object) override {};
+	virtual void Revert(UObject* Object) override;
+	virtual bool HasExpired(UObject* Object) const override
+	{
+		UPolyEditInsertEdgeActivity* Activity = Cast<UPolyEditInsertEdgeActivity>(Object);
+		return bHaveDoneUndo || Activity->CurrentChangeStamp != ChangeStamp
+			|| Activity->ToolState != UPolyEditInsertEdgeActivity::EState::GettingEnd;
+		// TODO: this is a bit of a hack in that we should probably have a separate stamp
+		// for expiring these instead of letting the tool state help (they expire after
+		// each new insertion unlike the other changes, which expire on tool close).
+	}
+	virtual FString ToString() const override
+	{
+		return TEXT("FGroupEdgeInsertionFirstPointChange");
+	}
+
+protected:
+	int32 ChangeStamp;
+	bool bHaveDoneUndo = false;
+};
diff --git a/Engine/Plugins/Experimental/MeshModelingToolset/Source/MeshModelingTools/Public/ToolActivities/PolyEditInsertEdgeLoopActivity.h b/Engine/Plugins/Experimental/MeshModelingToolset/Source/MeshModelingTools/Public/ToolActivities/PolyEditInsertEdgeLoopActivity.h
new file mode 100644
--- /dev/null
+++ b/Engine/Plugins/Experimental/MeshModelingToolset/Source/MeshModelingTools/Public/ToolActivities/PolyEditInsertEdgeLoopActivity.h
@@ -0,0 +1,179 @@
+// Copyright Epic Games, Inc. All Rights Reserved.
+
+#pragma once
+
+#include "CoreMinimal.h"
+
+#include "BaseBehaviors/BehaviorTargetInterfaces.h"
+#include "ModelingOperators.h" //IDynamicMeshOperatorFactory
+#include "InteractiveTool.h" //UInteractiveToolPropertySet
+#include "InteractiveToolActivity.h"
+#include "InteractiveToolBuilder.h" //UInteractiveToolBuilder
+#include "InteractiveToolChange.h" //FToolCommandChange
+#include "MeshOpPreviewHelpers.h" //FDynamicMeshOpResult
+#include "Selection/GroupTopologySelector.h"
+#include "SingleSelectionTool.h"
+#include "ToolDataVisualizer.h"
+
+#include "PolyEditInsertEdgeLoopActivity.generated.h"
+
+PREDECLARE_USE_GEOMETRY_CLASS(FDynamicMeshChange);
+class UPolyEditActivityContext;
+
+UENUM()
+enum class EEdgeLoopPositioningMode
+{
+	/** Edge loops will be evenly centered within a group. Allows for multiple insertions at a time. */
+	Even,
+
+	/** Edge loops will fall at the same length proportion at each edge they intersect (e.g., a quarter way down). */
+	ProportionOffset,
+
+	/** Edge loops will fall a constant distance away from the start of each edge they intersect
+	 (e.g., 20 units down). Clamps to end if edge is too short. */
+	 DistanceOffset
+};
+
+UENUM()
+enum class EEdgeLoopInsertionMode
+{
+	/** Existing groups will be deleted and new triangles will be created for the new groups.
+	 Keeps topology simple but breaks non-planar groups. */
+	Retriangulate,
+
+	/** Keeps existing triangles and cuts them to create a new path. May result in fragmented triangles over time.*/
+	PlaneCut
+};
+
+UCLASS()
+class MESHMODELINGTOOLS_API UEdgeLoopInsertionProperties : public UInteractiveToolPropertySet
+{
+	GENERATED_BODY()
+
+public:
+	/** Determines how edge loops position themselves vertically relative to loop direction. */
+	UPROPERTY(EditAnywhere, Category = InsertEdgeLoop)
+	EEdgeLoopPositioningMode PositionMode = EEdgeLoopPositioningMode::ProportionOffset;
+
+	/** Determines how edge loops are added to the geometry */
+	UPROPERTY(EditAnywhere, Category = InsertEdgeLoop)
+	EEdgeLoopInsertionMode InsertionMode = EEdgeLoopInsertionMode::Retriangulate;
+
+	/** How many loops to insert at a time. Only used with "even" positioning mode. */
+	UPROPERTY(EditAnywhere, Category = InsertEdgeLoop, meta = (UIMin = "0", UIMax = "20", ClampMin = "0", ClampMax = "500",
+		EditCondition = "PositionMode == EEdgeLoopPositioningMode::Even", EditConditionHides))
+	int32 NumLoops = 1;
+
+	UPROPERTY(EditAnywhere, Category = InsertEdgeLoop, AdvancedDisplay, meta = (UIMin = "0", UIMax = "1.0", ClampMin = "0", ClampMax = "1.0",
+		EditCondition = "PositionMode == EEdgeLoopPositioningMode::ProportionOffset && !bInteractive", EditConditionHides))
+	double ProportionOffset = 0.5;
+
+	UPROPERTY(EditAnywhere, Category = InsertEdgeLoop, AdvancedDisplay, meta = (UIMin = "0", ClampMin = "0",
+		EditCondition = "PositionMode == EEdgeLoopPositioningMode::DistanceOffset && !bInteractive", EditConditionHides))
+	double DistanceOffset = 10.0;
+
+	/** When false, the distance/proportion offset is numerically specified, and mouse clicks just choose the edge. */
+	UPROPERTY(EditAnywhere, Category = InsertEdgeLoop, AdvancedDisplay, meta = (
+		EditCondition = "PositionMode != EEdgeLoopPositioningMode::Even", EditConditionHides))
+	bool bInteractive = true;
+
+	/** Measure the distance offset from the opposite side of the edges. */
+	UPROPERTY(EditAnywhere, Category = InsertEdgeLoop, meta = (
+		EditCondition = "PositionMode == EEdgeLoopPositioningMode::DistanceOffset", EditConditionHides))
+	bool bFlipOffsetDirection = false;
+
+	/** When true, non-quad-like groups that stop the loop will be highlighted, with X's marking the corners. */
+	UPROPERTY(EditAnywhere, Category = InsertEdgeLoop)
+	bool bHighlightProblemGroups = true;
+
+	/** How close a new loop edge needs to pass next to an existing vertex to use that vertex rather than creating a new one. */
+	UPROPERTY(EditAnywhere, Category = InsertEdgeLoop, AdvancedDisplay)
+	double VertexTolerance = 0.001;
+};
+
+/** Interactive activity for inserting (group) edge loops into a mesh. */
+UCLASS()
+class MESHMODELINGTOOLS_API UPolyEditInsertEdgeLoopActivity : public UInteractiveToolActivity, 
+	public UE::Geometry::IDynamicMeshOperatorFactory,
+	public IHoverBehaviorTarget, public IClickBehaviorTarget
+{
+	GENERATED_BODY()
+public:
+
+	friend class FEdgeLoopInsertionChange;
+
+	UPolyEditInsertEdgeLoopActivity() {};
+
+	virtual void OnPropertyModified(UObject* PropertySet, FProperty* Property);
+
+	// IInteractiveToolActivity
+	virtual void Setup(UInteractiveTool* ParentTool) override;
+	virtual void Shutdown(EToolShutdownType ShutdownType) override;
+	virtual bool CanStart() const override;
+	virtual EToolActivityStartResult Start() override;
+	virtual bool IsRunning() const override { return bIsRunning; }
+	virtual bool CanAccept() const override;
+	virtual EToolActivityEndResult End(EToolShutdownType) override;
+	virtual void Render(IToolsContextRenderAPI* RenderAPI) override;
+	virtual void Tick(float DeltaTime) override;
+
+	// IDynamicMeshOperatorFactory
+	virtual TUniquePtr<UE::Geometry::FDynamicMeshOperator> MakeNewOperator() override;
+
+	// IClickBehaviorTarget
+	virtual FInputRayHit IsHitByClick(const FInputDeviceRay& ClickPos) override;
+	virtual void OnClicked(const FInputDeviceRay& ClickPos) override;
+
+	// IHoverBehaviorTarget
+	virtual FInputRayHit BeginHoverSequenceHitTest(const FInputDeviceRay& PressPos) override;
+	virtual void OnBeginHover(const FInputDeviceRay& DevicePos) override {}
+	virtual bool OnUpdateHover(const FInputDeviceRay& DevicePos) override;
+	virtual void OnEndHover() override;
+
+protected:
+	UPROPERTY()
+	UEdgeLoopInsertionProperties* Settings = nullptr;
+
+	UPROPERTY()
+	UPolyEditActivityContext* ActivityContext;
+
+	bool bIsRunning = false;
+
+	FTransform TargetTransform;
+	TSharedPtr<FGroupTopologySelector, ESPMode::ThreadSafe> TopologySelector;
+
+	TArray<TPair<FVector3d, FVector3d>> PreviewEdges;
+
+	// Used to highlight problematic topology (non-quad groups) when it stops a loop.
+	TArray<TPair<FVector3d, FVector3d>> ProblemTopologyEdges;
+	TArray<FVector3d> ProblemTopologyVerts;
+
+	FViewCameraState CameraState;
+
+	//FToolDataVisualizer ExistingEdgesRenderer;
+	FToolDataVisualizer PreviewEdgeRenderer;
+	FToolDataVisualizer ProblemTopologyRenderer;
+	FGroupTopologySelector::FSelectionSettings TopologySelectorSettings;
+	float ProblemVertTickWidth = 8;
+
+	void SetupPreview();
+
+	FInputRayHit HitTest(const FRay& WorldRay);
+	bool UpdateHoveredItem(const FRay& WorldRay);
+
+	void ConditionallyUpdatePreview(int32 NewGroupID, double* NewInputLength = nullptr);
+	void ClearPreview();
+
+	// Taken from user interaction, read as inputs by the op factory
+	int32 InputGroupEdgeID = FDynamicMesh3::InvalidID;
+	double InteractiveInputLength = 0;
+
+	// On valid clicks, we wait to finish the background op and apply it before taking more input.
+	// Gets reset OnTick when the result is ready.
+	bool bWaitingForInsertionCompletion = false;
+
+	// Copied over on op completion
+	bool bLastComputeSucceeded = false;
+	TSharedPtr<FGroupTopology, ESPMode::ThreadSafe> LatestOpTopologyResult;
+	TSharedPtr<TSet<int32>, ESPMode::ThreadSafe> LatestOpChangedTids;
+};
diff --git a/Engine/Plugins/Experimental/MeshModelingToolset/Source/MeshModelingTools/Public/ToolActivities/PolyEditInsetOutsetActivity.h b/Engine/Plugins/Experimental/MeshModelingToolset/Source/MeshModelingTools/Public/ToolActivities/PolyEditInsetOutsetActivity.h
new file mode 100644
--- /dev/null
+++ b/Engine/Plugins/Experimental/MeshModelingToolset/Source/MeshModelingTools/Public/ToolActivities/PolyEditInsetOutsetActivity.h
@@ -0,0 +1,101 @@
+// Copyright Epic Games, Inc. All Rights Reserved.
+
+#pragma once
+
+#include "CoreMinimal.h"
+
+#include "BaseBehaviors/BehaviorTargetInterfaces.h"
+#include "InteractiveToolActivity.h"
+
+#include "PolyEditInsetOutsetActivity.generated.h"
+
+class UPolyEditActivityContext;
+class UPolyEditPreviewMesh;
+class USpatialCurveDistanceMechanic;
+
+UCLASS()
+class MESHMODELINGTOOLS_API UPolyEditInsetOutsetProperties : public UInteractiveToolPropertySet
+{
+	GENERATED_BODY()
+
+public:
+	//~ TODO: We should not have to have the user tell us when it's an inset/offset, we should just do the proper
+	//~ thing based on where the user clicks.
+	UPROPERTY(EditAnywhere, Category = "Inset/Outset")
+	bool bOutset = false;
+
+	/** Amount of smoothing applied to outset boundary */
+	UPROPERTY(EditAnywhere, Category = "Inset/Outset", 
+		meta = (UIMin = "0.0", UIMax = "1.0", EditCondition = "bBoundaryOnly == false"))
+	float Softness = 0.5;
+
+	/** Controls whether outset operation will move interior vertices as well as border vertices */
+	UPROPERTY(EditAnywhere, Category = "Inset/Outset", AdvancedDisplay)
+	bool bBoundaryOnly = false;
+
+	/** Tweak area scaling when solving for interior vertices */
+	UPROPERTY(EditAnywhere, Category = "Inset/Outset", AdvancedDisplay, 
+		meta = (UIMin = "0.0", UIMax = "1.0", EditCondition = "bBoundaryOnly == false"))
+	float AreaScale = true;
+
+	/** When insetting, determines whether vertices in inset region should be projected back onto input surface */
+	UPROPERTY(EditAnywhere, Category = "Inset/Outset")
+	bool bReproject = true;
+
+
+};
+
+
+/**
+ *
+ */
+UCLASS()
+class MESHMODELINGTOOLS_API UPolyEditInsetOutsetActivity : public UInteractiveToolActivity,
+	public IClickBehaviorTarget, public IHoverBehaviorTarget
+{
+	GENERATED_BODY()
+
+public:
+	// IInteractiveToolActivity
+	virtual void Setup(UInteractiveTool* ParentTool) override;
+	virtual void Shutdown(EToolShutdownType ShutdownType) override;
+	virtual bool CanStart() const override;
+	virtual EToolActivityStartResult Start() override;
+	virtual bool IsRunning() const override { return bIsRunning; }
+	virtual bool CanAccept() const override;
+	virtual EToolActivityEndResult End(EToolShutdownType) override;
+	virtual void Render(IToolsContextRenderAPI* RenderAPI) override;
+	virtual void Tick(float DeltaTime) override;
+
+	// IClickBehaviorTarget API
+	virtual FInputRayHit IsHitByClick(const FInputDeviceRay& ClickPos) override;
+	virtual void OnClicked(const FInputDeviceRay& ClickPos) override;
+
+	// IHoverBehaviorTarget implementation
+	virtual FInputRayHit BeginHoverSequenceHitTest(const FInputDeviceRay& PressPos) override;
+	virtual void OnBeginHover(const FInputDeviceRay& DevicePos) override {}
+	virtual bool OnUpdateHover(const FInputDeviceRay& DevicePos) override;
+	virtual void OnEndHover() override {}
+
+	UPROPERTY()
+	UPolyEditInsetOutsetProperties* InsetProperties;
+
+protected:
+	void Clear();
+	void BeginInset();
+	void ApplyInset();
+
+	bool bIsRunning = false;
+	bool bPreviewUpdatePending = false;
+
+	UPROPERTY()
+	UPolyEditPreviewMesh* EditPreview;
+
+	UPROPERTY()
+	USpatialCurveDistanceMechanic* CurveDistMechanic = nullptr;
+
+	UPROPERTY()
+	TObjectPtr<UPolyEditActivityContext> ActivityContext;
+
+	float UVScaleFactor = 1.0f;
+};
diff --git a/Engine/Plugins/Experimental/MeshModelingToolset/Source/MeshModelingTools/Public/ToolActivities/PolyEditPlanarProjectionUVActivity.h b/Engine/Plugins/Experimental/MeshModelingToolset/Source/MeshModelingTools/Public/ToolActivities/PolyEditPlanarProjectionUVActivity.h
new file mode 100644
--- /dev/null
+++ b/Engine/Plugins/Experimental/MeshModelingToolset/Source/MeshModelingTools/Public/ToolActivities/PolyEditPlanarProjectionUVActivity.h
@@ -0,0 +1,85 @@
+// Copyright Epic Games, Inc. All Rights Reserved.
+
+#pragma once
+
+#include "CoreMinimal.h"
+
+#include "BaseBehaviors/BehaviorTargetInterfaces.h"
+#include "BoxTypes.h"
+#include "InteractiveToolActivity.h"
+#include "ToolActivities/PolyEditActivityUtil.h"
+#include "ToolContextInterfaces.h" // FViewCameraState
+
+#include "PolyEditPlanarProjectionUVActivity.generated.h"
+
+class UPolyEditActivityContext;
+class UPolyEditPreviewMesh;
+class UCollectSurfacePathMechanic;
+
+UCLASS()
+class MESHMODELINGTOOLS_API UPolyEditSetUVProperties : public UInteractiveToolPropertySet
+{
+	GENERATED_BODY()
+
+public:
+	UPROPERTY(EditAnywhere, Category = PlanarProjectUV)
+	bool bShowMaterial = false;
+};
+
+
+/**
+ *
+ */
+UCLASS()
+class MESHMODELINGTOOLS_API UPolyEditPlanarProjectionUVActivity : public UInteractiveToolActivity,
+	public IClickBehaviorTarget, public IHoverBehaviorTarget
+{
+	GENERATED_BODY()
+
+public:
+	// IInteractiveToolActivity
+	virtual void Setup(UInteractiveTool* ParentTool) override;
+	virtual void Shutdown(EToolShutdownType ShutdownType) override;
+	virtual bool CanStart() const override;
+	virtual EToolActivityStartResult Start() override;
+	virtual bool IsRunning() const override { return bIsRunning; }
+	virtual bool CanAccept() const override;
+	virtual EToolActivityEndResult End(EToolShutdownType) override;
+	virtual void Render(IToolsContextRenderAPI* RenderAPI) override;
+	virtual void Tick(float DeltaTime) override;
+
+	// IClickBehaviorTarget API
+	virtual FInputRayHit IsHitByClick(const FInputDeviceRay& ClickPos) override;
+	virtual void OnClicked(const FInputDeviceRay& ClickPos) override;
+
+	// IHoverBehaviorTarget implementation
+	virtual FInputRayHit BeginHoverSequenceHitTest(const FInputDeviceRay& PressPos) override;
+	virtual void OnBeginHover(const FInputDeviceRay& DevicePos) override {}
+	virtual bool OnUpdateHover(const FInputDeviceRay& DevicePos) override;
+	virtual void OnEndHover() override {}
+
+protected:
+	void Clear();
+	void BeginSetUVs();
+	void UpdateSetUVS();
+	void ApplySetUVs();
+
+	UPROPERTY()
+	UPolyEditSetUVProperties* SetUVProperties;
+
+	UPROPERTY()
+	UPolyEditPreviewMesh* EditPreview;
+
+	UPROPERTY()
+	UCollectSurfacePathMechanic* SurfacePathMechanic = nullptr;
+
+	UPROPERTY()
+	TObjectPtr<UPolyEditActivityContext> ActivityContext;
+
+	bool bIsRunning = false;
+
+	bool bPreviewUpdatePending = false;
+	UE::Geometry::PolyEditActivityUtil::EPreviewMaterialType CurrentPreviewMaterial;
+	UE::Geometry::FAxisAlignedBox3d ActiveSelectionBounds;
+	FViewCameraState CameraState;
+};
diff --git a/Engine/Plugins/Experimental/MeshModelingToolset/Source/ModelingComponents/Public/InteractiveToolActivity.h b/Engine/Plugins/Experimental/MeshModelingToolset/Source/ModelingComponents/Public/InteractiveToolActivity.h
new file mode 100644
--- /dev/null
+++ b/Engine/Plugins/Experimental/MeshModelingToolset/Source/ModelingComponents/Public/InteractiveToolActivity.h
@@ -0,0 +1,162 @@
+// Copyright Epic Games, Inc. All Rights Reserved.
+
+#pragma once
+
+#include "CoreMinimal.h"
+
+#include "InteractionMechanic.h"
+#include "InteractiveTool.h" //EToolShutdownType
+
+#include "InteractiveToolActivity.generated.h"
+
+class IToolsContextRenderAPI;
+
+enum class EToolActivityStartResult
+{
+	Running,
+	FailedStart,
+
+	// It's not clear whether we want this option. It's certainly the case that depending
+	// on settings, an activity may want to end immediately, or perhaps we may want to
+	// put single-action subtools on the same system. But the activity can just end on the
+	// next tick, and if we allow a third start option, the tool may need to be aware of
+	// it as a possibility, plus we have to decide whether or not NotifyActivitySelfEnded
+	// applies... Seems messy to have it, so let's keep it out until we really want it.
+	//ImmediatelyCompleted,
+};
+
+enum class EToolActivityEndResult
+{
+	Completed,
+	Cancelled,
+	ErrorDuringEnd,
+};
+
+
+/**
+ * A tool activity is a sort of "sub-tool" used to break apart functionality in tools that
+ * provide support for different multi-interaction subtasks. It is meant to limit the sprawl
+ * of switch statements in the base tool, to allow subtasks to be designed similar to how a
+ * tool might be designed, and to ease extendability.
+ *
+ * An activity has the following expectations:
+ * - Setup() is called in host tool setup and Shutdown() is called in host tool Shutdown()
+ * - Start() is called to start the activity (such as when user clicks a button).
+ * - If the activity returns a result of EStartResult::ActivityRunning on Start(), it will expect
+ *   Render() and Tick() calls from the host until either (a)- the host calls End() on the activity,
+ *   or (b)- the activity reaches a stopping point itself and calls NotifyActivitySelfEnded() on the
+ *   host. The activity should not require Render() and Tick() if it is not running.
+ * 
+ *  Compared to a UInteractionMechanic, a tool activity:
+ * - Expects that it is the main consumer of input, i.e. takes over the tool. Mechanics, by contrast,
+ *   are currently often used together with other mechanics, or to support main tool functionality.
+ * - Should not require the hosting tool to have specific knowledge about it or be heavily involved.
+ *   Mechanics, by contrast, currently often require tools to use various mechanic-specific getters/setters
+ *   during the tool.
+ * 
+ * Passing data back and forth can be done either by letting a tool activity use a specific
+ * context object that the tool can prep in the context store, or by requiring the host to
+ * implement specific interfaces (that the activity can check for in Setup()).
+ */
+UCLASS()
+class MODELINGCOMPONENTS_API UInteractiveToolActivity : public UInteractionMechanic
+{
+	GENERATED_BODY()
+
+public:
+	// Note about the below: we didn't really need to repeat Setup(), Tick(), and Render() here
+	// since these are staying the same as UInteractionMechanic and get inherited. But it is 
+	// helpful to have the entire API that we need to implement described in one place
+
+	virtual ~UInteractiveToolActivity() {}
+
+	/**
+	 * Should be called during a tool's Setup(). Does not start the activity.
+	 */
+	virtual void Setup(UInteractiveTool* ParentToolIn) override
+	{
+		Super::Setup(ParentToolIn);
+	}
+
+	/**
+	 * Should be called during a tool's Shutdown()
+	 */
+	virtual void Shutdown(EToolShutdownType ShutdownType)
+	{
+		Super::Shutdown();
+	}
+
+	/**
+	 * Check whether a Start() call will result in a success.
+	 */
+	virtual bool CanStart() const { return false; }
+
+	/**
+	 * Attempt to start the activity.
+	 */
+	virtual EToolActivityStartResult Start() { return EToolActivityStartResult::FailedStart; }
+
+	/**
+	 * Check whether the activity is running (though the tool can just check the result of
+	 * Start() itself).
+	 */
+	virtual bool IsRunning() const 
+	{
+		// We could keep a boolean here in the base class and update it in the base class Start() and End()
+		// functions so that derived classes don't have to remember to implement this. But it is probably better
+		// to have the inconvenience of having to implement IsRunning in the derived class than to force
+		// derived classes to remember to call base Start() and End() or risk giving an incorrect result here...
+		return ensure(false); 
+	}
+
+	/**
+	 * If true, calling End with EToolShutdownType::Accept will result in a valid completion of
+	 * the activity. It's up to the tool whether to respect this or not, since End() should still
+	 * end the activity.
+	 */
+	virtual bool CanAccept() const { return false; };
+
+	/**
+	 * Force an end to the activity.
+	 */
+	virtual EToolActivityEndResult End(EToolShutdownType) { return EToolActivityEndResult::ErrorDuringEnd; };
+
+	/**
+	 * If the activity is running, should be called from the tool's OnTick()
+	 */
+	virtual void Tick(float DeltaTime) override {};
+
+	/**
+	 * If the activity is running, should be called from the tool's Render()
+	 */
+	virtual void Render(IToolsContextRenderAPI* RenderAPI) override {};
+
+protected:
+
+
+private:
+	// We don't want to use this overload of Shutdown because an activity may need to know
+	// the ShutdownType if it is still running when Shutdown is called (to End properly).
+	virtual void Shutdown() override final
+	{
+		ensure(false);
+		Shutdown(EToolShutdownType::Cancel);
+	}
+};
+
+UINTERFACE()
+class MODELINGCOMPONENTS_API UToolActivityHost : public UInterface
+{
+	GENERATED_BODY()
+};
+
+class MODELINGCOMPONENTS_API IToolActivityHost
+{
+	GENERATED_BODY()
+public:
+
+	/**
+	 * Notify the tool that the activity has ended.
+	 */
+	virtual void NotifyActivitySelfEnded(UInteractiveToolActivity* Activity) = 0;
+};
\ No newline at end of file
diff --git a/Engine/Plugins/Experimental/ModelingToolsEditorMode/Source/ModelingToolsEditorMode/Private/ModelingToolsEditorMode.cpp b/Engine/Plugins/Experimental/ModelingToolsEditorMode/Source/ModelingToolsEditorMode/Private/ModelingToolsEditorMode.cpp
--- a/Engine/Plugins/Experimental/ModelingToolsEditorMode/Source/ModelingToolsEditorMode/Private/ModelingToolsEditorMode.cpp
+++ b/Engine/Plugins/Experimental/ModelingToolsEditorMode/Source/ModelingToolsEditorMode/Private/ModelingToolsEditorMode.cpp
@@ -29,8 +29,6 @@
 #include "EditMeshPolygonsTool.h"
 #include "DeformMeshPolygonsTool.h"
 #include "SubdividePolyTool.h"
-#include "GroupEdgeInsertionTool.h"
-#include "EdgeLoopInsertionTool.h"
 #include "ConvertToPolygonsTool.h"
 #include "AddPrimitiveTool.h"
 #include "AddPatchTool.h"
@@ -469,10 +467,6 @@ void UModelingToolsEditorMode::Enter()
 	RegisterTool(ToolManagerCommands.BeginProjectToTargetTool, TEXT("BeginProjectToTargetTool"), NewObject<UProjectToTargetToolBuilder>());
 	RegisterTool(ToolManagerCommands.BeginSimplifyMeshTool, TEXT("BeginSimplifyMeshTool"), NewObject<USimplifyMeshToolBuilder>());
 
-	RegisterTool(ToolManagerCommands.BeginGroupEdgeInsertionTool, TEXT("BeginGroupEdgeInsertionTool"), NewObject<UGroupEdgeInsertionToolBuilder>());
-
-	RegisterTool(ToolManagerCommands.BeginEdgeLoopInsertionTool, TEXT("BeginEdgeLoopInsertionTool"), NewObject<UEdgeLoopInsertionToolBuilder>());
-
 	auto EditNormalsToolBuilder = NewObject<UEditNormalsToolBuilder>();
 	RegisterTool(ToolManagerCommands.BeginEditNormalsTool, TEXT("BeginEditNormalsTool"), EditNormalsToolBuilder);
 
@@ -606,9 +600,7 @@ void UModelingToolsEditorMode::Enter()
 		RegisterTool(UICommand, StringName, ActionModeBuilder);
 	};
 	RegisterPolyModelActionTool(EEditMeshPolygonsToolActions::Extrude, ToolManagerCommands.BeginPolyModelTool_Extrude, TEXT("PolyEdit_Extrude"));
-	RegisterPolyModelActionTool(EEditMeshPolygonsToolActions::Offset, ToolManagerCommands.BeginPolyModelTool_Offset, TEXT("PolyEdit_Offset"));
-	RegisterPolyModelActionTool(EEditMeshPolygonsToolActions::Inset, ToolManagerCommands.BeginPolyModelTool_Inset, TEXT("PolyEdit_Inset"));
-	RegisterPolyModelActionTool(EEditMeshPolygonsToolActions::Outset, ToolManagerCommands.BeginPolyModelTool_Outset, TEXT("PolyEdit_Outset"));
+	RegisterPolyModelActionTool(EEditMeshPolygonsToolActions::InsetOutset, ToolManagerCommands.BeginPolyModelTool_InsetOutset, TEXT("PolyEdit_InsetOutset"));
 	RegisterPolyModelActionTool(EEditMeshPolygonsToolActions::CutFaces, ToolManagerCommands.BeginPolyModelTool_CutFaces, TEXT("PolyEdit_CutFaces"));
 
 
diff --git a/Engine/Plugins/Experimental/ModelingToolsEditorMode/Source/ModelingToolsEditorMode/Private/ModelingToolsEditorModeToolkit.cpp b/Engine/Plugins/Experimental/ModelingToolsEditorMode/Source/ModelingToolsEditorMode/Private/ModelingToolsEditorModeToolkit.cpp
--- a/Engine/Plugins/Experimental/ModelingToolsEditorMode/Source/ModelingToolsEditorMode/Private/ModelingToolsEditorModeToolkit.cpp
+++ b/Engine/Plugins/Experimental/ModelingToolsEditorMode/Source/ModelingToolsEditorMode/Private/ModelingToolsEditorModeToolkit.cpp
@@ -614,8 +614,6 @@ void FModelingToolsEditorModeToolkit::BuildToolPalette_Experimental(FName Palett
 	{
 		ToolbarBuilder.AddToolBarButton(Commands.BeginPolyEditTool);
 		ToolbarBuilder.AddToolBarButton(Commands.BeginPolyDeformTool);
-		ToolbarBuilder.AddToolBarButton(Commands.BeginGroupEdgeInsertionTool);
-		ToolbarBuilder.AddToolBarButton(Commands.BeginEdgeLoopInsertionTool);
 		ToolbarBuilder.AddToolBarButton(Commands.BeginMeshBooleanTool);
 		ToolbarBuilder.AddToolBarButton(Commands.BeginCutMeshWithMeshTool);
 		ToolbarBuilder.AddToolBarButton(Commands.BeginSubdividePolyTool);
@@ -675,14 +673,9 @@ void FModelingToolsEditorModeToolkit::BuildToolPalette_Experimental(FName Palett
 		ToolbarBuilder.AddToolBarButton(Commands.BeginPolyModelTool_RingSelect);
 		ToolbarBuilder.AddSeparator();
 		ToolbarBuilder.AddToolBarButton(Commands.BeginPolyModelTool_Extrude);
-		ToolbarBuilder.AddToolBarButton(Commands.BeginPolyModelTool_Offset);
-		ToolbarBuilder.AddToolBarButton(Commands.BeginPolyModelTool_Inset);
-		ToolbarBuilder.AddToolBarButton(Commands.BeginPolyModelTool_Outset);
+		ToolbarBuilder.AddToolBarButton(Commands.BeginPolyModelTool_InsetOutset);
 		ToolbarBuilder.AddToolBarButton(Commands.BeginPolyModelTool_CutFaces);
 		ToolbarBuilder.AddSeparator();
-		ToolbarBuilder.AddToolBarButton(Commands.BeginGroupEdgeInsertionTool);
-		ToolbarBuilder.AddToolBarButton(Commands.BeginEdgeLoopInsertionTool);
-		ToolbarBuilder.AddSeparator();
 		ToolbarBuilder.AddToolBarButton(Commands.BeginSubdividePolyTool);
 		ToolbarBuilder.AddToolBarButton(Commands.BeginPolyEditTool);
 	}
diff --git a/Engine/Plugins/Experimental/ModelingToolsEditorMode/Source/ModelingToolsEditorMode/Private/ModelingToolsManagerActions.cpp b/Engine/Plugins/Experimental/ModelingToolsEditorMode/Source/ModelingToolsEditorMode/Private/ModelingToolsManagerActions.cpp
--- a/Engine/Plugins/Experimental/ModelingToolsEditorMode/Source/ModelingToolsEditorMode/Private/ModelingToolsManagerActions.cpp
+++ b/Engine/Plugins/Experimental/ModelingToolsEditorMode/Source/ModelingToolsEditorMode/Private/ModelingToolsManagerActions.cpp
@@ -41,8 +41,6 @@ void FModelingToolsManagerCommands::RegisterCommands()
 
 	UI_COMMAND(BeginSculptMeshTool, "VSclpt", "Vertex Sculpting", EUserInterfaceActionType::ToggleButton, FInputChord());
 	UI_COMMAND(BeginPolyEditTool, "PolyEd", "Edit Mesh via PolyGroups", EUserInterfaceActionType::ToggleButton, FInputChord());
-	UI_COMMAND(BeginGroupEdgeInsertionTool, "EdgeIns", "Insert PolyGroup Edge", EUserInterfaceActionType::ToggleButton, FInputChord());
-	UI_COMMAND(BeginEdgeLoopInsertionTool, "LoopIns", "Insert PolyGroup Edge Loop", EUserInterfaceActionType::ToggleButton, FInputChord());
 	UI_COMMAND(BeginTriEditTool, "TriEd", "Edit Mesh via Triangles", EUserInterfaceActionType::ToggleButton, FInputChord());
 	UI_COMMAND(BeginPolyDeformTool, "PolyDef", "Deform Mesh via PolyGroups", EUserInterfaceActionType::ToggleButton, FInputChord());
 	UI_COMMAND(BeginSmoothMeshTool, "Smooth", "Smooth Mesh surface", EUserInterfaceActionType::ToggleButton, FInputChord());
@@ -129,9 +127,7 @@ void FModelingToolsManagerCommands::RegisterCommands()
 	UI_COMMAND(BeginPolyModelTool_RingSelect, "Rings", "PolyGroup Ring Selection Tool", EUserInterfaceActionType::ToggleButton, FInputChord());
 
 	UI_COMMAND(BeginPolyModelTool_Extrude, "Extrude", "PolyGroup Extrude Tool", EUserInterfaceActionType::ToggleButton, FInputChord());
-	UI_COMMAND(BeginPolyModelTool_Offset, "Offset", "PolyGroup Offset Tool", EUserInterfaceActionType::ToggleButton, FInputChord());
-	UI_COMMAND(BeginPolyModelTool_Inset, "Inset", "PolyGroup Inset Tool", EUserInterfaceActionType::ToggleButton, FInputChord());
-	UI_COMMAND(BeginPolyModelTool_Outset, "Outset", "PolyGroup Outset Tool", EUserInterfaceActionType::ToggleButton, FInputChord());
+	UI_COMMAND(BeginPolyModelTool_InsetOutset, "Inset/Outset", "PolyGroup Inset/Outset Tool", EUserInterfaceActionType::ToggleButton, FInputChord());
 	UI_COMMAND(BeginPolyModelTool_CutFaces, "Cut", "PolyGroup Cut Faces Tool", EUserInterfaceActionType::ToggleButton, FInputChord());
 
 	UI_COMMAND(AcceptActiveTool, "Accept", "Accept the active tool", EUserInterfaceActionType::Button, FInputChord());
diff --git a/Engine/Plugins/Experimental/ModelingToolsEditorMode/Source/ModelingToolsEditorMode/Public/ModelingToolsManagerActions.h b/Engine/Plugins/Experimental/ModelingToolsEditorMode/Source/ModelingToolsEditorMode/Public/ModelingToolsManagerActions.h
--- a/Engine/Plugins/Experimental/ModelingToolsEditorMode/Source/ModelingToolsEditorMode/Public/ModelingToolsManagerActions.h
+++ b/Engine/Plugins/Experimental/ModelingToolsEditorMode/Source/ModelingToolsEditorMode/Public/ModelingToolsManagerActions.h
@@ -35,8 +35,6 @@ public:
 	TSharedPtr<FUICommandInfo> BeginSculptMeshTool;
 	TSharedPtr<FUICommandInfo> BeginPolyEditTool;
 	TSharedPtr<FUICommandInfo> BeginSubdividePolyTool;
-	TSharedPtr<FUICommandInfo> BeginGroupEdgeInsertionTool;
-	TSharedPtr<FUICommandInfo> BeginEdgeLoopInsertionTool;
 	TSharedPtr<FUICommandInfo> BeginTriEditTool;
 	TSharedPtr<FUICommandInfo> BeginPolyDeformTool;
 	TSharedPtr<FUICommandInfo> BeginSmoothMeshTool;
@@ -120,9 +118,7 @@ public:
 	TSharedPtr<FUICommandInfo> BeginPolyModelTool_RingSelect;
 
 	TSharedPtr<FUICommandInfo> BeginPolyModelTool_Extrude;
-	TSharedPtr<FUICommandInfo> BeginPolyModelTool_Offset;
-	TSharedPtr<FUICommandInfo> BeginPolyModelTool_Inset;
-	TSharedPtr<FUICommandInfo> BeginPolyModelTool_Outset;
+	TSharedPtr<FUICommandInfo> BeginPolyModelTool_InsetOutset;
 	TSharedPtr<FUICommandInfo> BeginPolyModelTool_CutFaces;
 
 
diff --git a/Engine/Source/Runtime/Experimental/InteractiveToolsFramework/Private/InteractionMechanic.cpp b/Engine/Source/Runtime/Experimental/InteractiveToolsFramework/Private/InteractionMechanic.cpp
--- a/Engine/Source/Runtime/Experimental/InteractiveToolsFramework/Private/InteractionMechanic.cpp
+++ b/Engine/Source/Runtime/Experimental/InteractiveToolsFramework/Private/InteractionMechanic.cpp
@@ -45,4 +45,13 @@ void UInteractionMechanic::AddToolPropertySource(UInteractiveToolPropertySet* Pr
 	{
 		ParentTool->AddToolPropertySource(PropertySet);
 	}
+}
+
+bool UInteractionMechanic::SetToolPropertySourceEnabled(UInteractiveToolPropertySet* PropertySet, bool bEnabled)
+{
+	if (ParentTool.IsValid())
+	{
+		return ParentTool->SetToolPropertySourceEnabled(PropertySet, bEnabled);
+	}
+	return false;
 }
\ No newline at end of file
diff --git a/Engine/Source/Runtime/Experimental/InteractiveToolsFramework/Public/InteractionMechanic.h b/Engine/Source/Runtime/Experimental/InteractiveToolsFramework/Public/InteractionMechanic.h
--- a/Engine/Source/Runtime/Experimental/InteractiveToolsFramework/Public/InteractionMechanic.h
+++ b/Engine/Source/Runtime/Experimental/InteractiveToolsFramework/Public/InteractionMechanic.h
@@ -55,5 +55,8 @@ protected:
 	 * @param PropertySet Property Set object to add
 	 */
 	virtual void AddToolPropertySource(UInteractiveToolPropertySet* PropertySet);
+
+	/** Enable/Disable a PropertySet object for the parent tool. */
+	virtual bool SetToolPropertySourceEnabled(UInteractiveToolPropertySet* PropertySet, bool bEnabled);
 };
 
