ModelingTools: add GroupFill mode to MeshGroupPaintTool. In this mode, any group touched by the brush will be fully painted w/ the current paint group, allowing existing groups to be easily merged. Updated internals so that Volumetic mode now supports seam filtering
#rnx

diff --git a/Engine/Plugins/Experimental/MeshModelingToolset/Source/MeshModelingTools/Private/MeshGroupPaintTool.cpp b/Engine/Plugins/Experimental/MeshModelingToolset/Source/MeshModelingTools/Private/MeshGroupPaintTool.cpp
--- a/Engine/Plugins/Experimental/MeshModelingToolset/Source/MeshModelingTools/Private/MeshGroupPaintTool.cpp
+++ b/Engine/Plugins/Experimental/MeshModelingToolset/Source/MeshModelingTools/Private/MeshGroupPaintTool.cpp
@@ -313,7 +313,8 @@ void UMeshGroupPaintTool::OnPropertyModified(UObject* PropertySet, FProperty* Pr
 bool UMeshGroupPaintTool::IsInBrushSubMode() const
 {
 	return FilterProperties->SubToolType == EMeshGroupPaintInteractionType::Brush
-		|| FilterProperties->SubToolType == EMeshGroupPaintInteractionType::Fill;
+		|| FilterProperties->SubToolType == EMeshGroupPaintInteractionType::Fill
+		|| FilterProperties->SubToolType == EMeshGroupPaintInteractionType::GroupFill;
 }
 
 
@@ -386,7 +387,32 @@ void UMeshGroupPaintTool::UpdateROI(const FSculptBrushStamp& BrushStamp)
 		TriangleROI.Add(CenterTID);
 	}
 
+	FVector3d CenterNormal = Mesh->IsTriangle(CenterTID) ? TriNormals[CenterTID] : FVector3d::One();		// One so that normal check always passes
+
+	bool bUseAngleThreshold = (FilterProperties->AngleThreshold < 180.0f);
+	double DotAngleThreshold = FMathd::Cos(FilterProperties->AngleThreshold * FMathd::DegToRad);
+	bool bStopAtUVSeams = FilterProperties->bUVSeams;
+	bool bStopAtNormalSeams = FilterProperties->bNormalSeams;
+
+
+	auto CheckEdgeCriteria = [&](int32 t1, int32 t2) -> bool
+	{
+		if (bUseAngleThreshold == false || CenterNormal.Dot(TriNormals[t2]) > DotAngleThreshold)
+		{
+			int32 eid = Mesh->FindEdgeFromTriPair(t1, t2);
+			if (bStopAtUVSeams == false || UVSeamEdges[eid] == false)
+			{
+				if (bStopAtNormalSeams == false || NormalSeamEdges[eid] == false)
+				{
+					return true;
+				}
+			}
+		}
+		return false;
+	};
+
 	bool bFill = (FilterProperties->SubToolType == EMeshGroupPaintInteractionType::Fill);
+	bool bGroupFill = (FilterProperties->SubToolType == EMeshGroupPaintInteractionType::GroupFill);
 
 	if (FilterProperties->BrushAreaMode == EMeshGroupPaintBrushAreaType::Volumetric)
 	{
@@ -398,56 +424,45 @@ void UMeshGroupPaintTool::UpdateROI(const FSculptBrushStamp& BrushStamp)
 				TriangleROI.Add(TriIdx);
 			}
 		});
-
-		if (bFill)
-		{
-			TArray<int32> StartROI;
-			for (int32 tid : TriangleROI)
-			{
-				StartROI.Add(tid);
-			}
-			FMeshConnectedComponents::GrowToConnectedTriangles(Mesh, StartROI, TriangleROI, &TempROIBuffer,
-				[&](int t1, int t2)
-			{
-				return true;
-			});
-		}
 	}
 	else
 	{
 		if (Mesh->IsTriangle(CenterTID))
 		{
-			FVector3d CenterNormal = TriNormals[CenterTID];
-			bool bUseAngleThreshold = (FilterProperties->AngleThreshold < 180.0f);
-			double DotAngleThreshold = FMathd::Cos(FilterProperties->AngleThreshold * FMathd::DegToRad);
-
-			bool bStopAtUVSeams = FilterProperties->bUVSeams;
-			bool bStopAtNormalSeams = FilterProperties->bNormalSeams;
-
 			TArray<int32> StartROI;
 			StartROI.Add(CenterTID);
 			FMeshConnectedComponents::GrowToConnectedTriangles(Mesh, StartROI, TriangleROI, &TempROIBuffer,
 				[&](int t1, int t2) 
 			{ 
-				if (bFill || (Mesh->GetTriCentroid(t2) - BrushPos).SquaredLength() < RadiusSqr)
+				if ((Mesh->GetTriCentroid(t2) - BrushPos).SquaredLength() < RadiusSqr)
 				{
-					if (bUseAngleThreshold == false || CenterNormal.Dot(TriNormals[t2]) > DotAngleThreshold)
-					{
-						int32 eid = Mesh->FindEdgeFromTriPair(t1, t2);
-						if (bStopAtUVSeams == false || UVSeamEdges[eid] == false)
-						{
-							if (bStopAtNormalSeams == false || NormalSeamEdges[eid] == false)
-							{
-								return true;
-							}
-						}
-					}
+					return CheckEdgeCriteria(t1, t2);
 				}
 				return false;
 			});
 		}
 	}
 
+	if (bFill || bGroupFill)
+	{
+		TArray<int32> StartROI;
+		TSet<int32> FillGroups;
+		for (int32 tid : TriangleROI)
+		{
+			StartROI.Add(tid);
+			FillGroups.Add(ActiveGroupSet->GetGroup(tid));
+		}
+		FMeshConnectedComponents::GrowToConnectedTriangles(Mesh, StartROI, TriangleROI, &TempROIBuffer,
+			[&](int t1, int t2)
+		{
+			if (bFill || (bGroupFill && FillGroups.Contains(ActiveGroupSet->GetGroup(t2))))
+			{
+				return CheckEdgeCriteria(t1, t2);
+			}
+			return false;
+		});
+	}
+
 
 	// apply visibility filter
 	if (FilterProperties->VisibilityFilter != EMeshGroupPaintVisibilityType::None)
diff --git a/Engine/Plugins/Experimental/MeshModelingToolset/Source/MeshModelingTools/Public/MeshGroupPaintTool.h b/Engine/Plugins/Experimental/MeshModelingToolset/Source/MeshModelingTools/Public/MeshGroupPaintTool.h
--- a/Engine/Plugins/Experimental/MeshModelingToolset/Source/MeshModelingTools/Public/MeshGroupPaintTool.h
+++ b/Engine/Plugins/Experimental/MeshModelingToolset/Source/MeshModelingTools/Public/MeshGroupPaintTool.h
@@ -64,6 +64,7 @@ enum class EMeshGroupPaintInteractionType : uint8
 {
 	Brush,
 	Fill,
+	GroupFill,
 	PolyLasso,
 
 	LastValue UMETA(Hidden)
@@ -125,17 +126,17 @@ public:
 
 	/** Relative size of brush */
 	UPROPERTY(EditAnywhere, Category = ActionType, meta = (DisplayName = "Brush Size", UIMin = "0.0", UIMax = "1.0", ClampMin = "0.0", ClampMax = "10.0", 
-		HideEditConditionToggle, EditConditionHides, EditCondition = "SubToolType == EMeshGroupPaintInteractionType::Brush"))
+		HideEditConditionToggle, EditConditionHides, EditCondition = "SubToolType != EMeshGroupPaintInteractionType::PolyLasso"))
 	float BrushSize;
 
 	/** When Volumetric, all faces inside the brush sphere are selected, otherwise only connected faces are selected */
 	UPROPERTY(EditAnywhere, Category = ActionType, meta = (DisplayName = "Brush Area Mode",
-		HideEditConditionToggle, EditConditionHides, EditCondition = "SubToolType == EMeshGroupPaintInteractionType::Brush"))
+		HideEditConditionToggle, EditConditionHides, EditCondition = "SubToolType != EMeshGroupPaintInteractionType::PolyLasso"))
 	EMeshGroupPaintBrushAreaType BrushAreaMode = EMeshGroupPaintBrushAreaType::Connected;
 
 	/** Allow the Brush to hit the back-side of the mesh */
 	UPROPERTY(EditAnywhere, Category = ActionType, meta = (DisplayName = "Hit Back Faces",
-		HideEditConditionToggle, EditConditionHides, EditCondition = "SubToolType == EMeshGroupPaintInteractionType::Brush"))
+		HideEditConditionToggle, EditConditionHides, EditCondition = "SubToolType != EMeshGroupPaintInteractionType::PolyLasso"))
 	bool bHitBackFaces = true;
 
 	/** The group that will be assigned to triangles */
@@ -154,24 +155,22 @@ public:
 	UPROPERTY(EditAnywhere, Category = ActionType)
 	bool bOnlyEraseCurrent = false;
 
-	
-	/** Control which triangles can be affected by the current operation based on visibility */
-	UPROPERTY(EditAnywhere, Category = Filters)
-	EMeshGroupPaintVisibilityType VisibilityFilter = EMeshGroupPaintVisibilityType::None;
-
 	/** The Region affected by the current operation will be bounded by edge angles larger than this threshold */
-	UPROPERTY(EditAnywhere, Category = Filters, meta = (UIMin = "0.0", UIMax = "180.0",
-		EditCondition = "bVolumetric == false"))
+	UPROPERTY(EditAnywhere, Category = Filters, meta = (UIMin = "0.0", UIMax = "180.0", EditCondition = "SubToolType != EMeshGroupPaintInteractionType::PolyLasso && BrushAreaMode == EMeshGroupPaintBrushAreaType::Connected"))
 	float AngleThreshold = 180.0f;
 
 	/** The Region affected by the current operation will be bounded by UV borders/seams */
-	UPROPERTY(EditAnywhere, Category = Filters, meta = (EditCondition = "bVolumetric == false"))
+	UPROPERTY(EditAnywhere, Category = Filters, meta = (EditCondition = "SubToolType != EMeshGroupPaintInteractionType::PolyLasso && BrushAreaMode == EMeshGroupPaintBrushAreaType::Connected"))
 	bool bUVSeams = false;
 
 	/** The Region affected by the current operation will be bounded by Hard Normal edges/seams */
-	UPROPERTY(EditAnywhere, Category = Filters, meta = (EditCondition = "bVolumetric == false"))
+	UPROPERTY(EditAnywhere, Category = Filters, meta = (EditCondition = "SubToolType != EMeshGroupPaintInteractionType::PolyLasso && BrushAreaMode == EMeshGroupPaintBrushAreaType::Connected"))
 	bool bNormalSeams = false;
 
+	/** Control which triangles can be affected by the current operation based on visibility. Applied after all other filters. */
+	UPROPERTY(EditAnywhere, Category = Filters)
+	EMeshGroupPaintVisibilityType VisibilityFilter = EMeshGroupPaintVisibilityType::None;
+
 
 	/** Number of vertices in a triangle the Lasso must hit to be counted as "inside" */
 	UPROPERTY(EditAnywhere, Category = Filters, AdvancedDisplay, meta = (UIMin = 1, UIMax = 3, EditCondition = "SubToolType == EMeshGroupPaintInteractionType::PolyLasso"))
