ModelingTools: performance improvements in MeshGroupPaintTool
#rnx

diff --git a/Engine/Plugins/Experimental/MeshModelingToolset/Source/MeshModelingTools/Private/MeshGroupPaintTool.cpp b/Engine/Plugins/Experimental/MeshModelingToolset/Source/MeshModelingTools/Private/MeshGroupPaintTool.cpp
--- a/Engine/Plugins/Experimental/MeshModelingToolset/Source/MeshModelingTools/Private/MeshGroupPaintTool.cpp
+++ b/Engine/Plugins/Experimental/MeshModelingToolset/Source/MeshModelingTools/Private/MeshGroupPaintTool.cpp
@@ -372,6 +372,8 @@ void UMeshGroupPaintTool::UpdateROI(const FSculptBrushStamp& BrushStamp)
 {
 	SCOPE_CYCLE_COUNTER(GroupPaintTool_UpdateROI);
 
+	int32 SetGroupID = GetInInvertStroke() ? FilterProperties->EraseGroup : FilterProperties->SetGroup;
+
 	const FVector3d& BrushPos = BrushStamp.LocalFrame.Origin;
 	const FDynamicMesh3* Mesh = GetSculptMesh();
 	float RadiusSqr = GetCurrentBrushRadius() * GetCurrentBrushRadius();
@@ -389,12 +391,12 @@ void UMeshGroupPaintTool::UpdateROI(const FSculptBrushStamp& BrushStamp)
 
 	FVector3d CenterNormal = Mesh->IsTriangle(CenterTID) ? TriNormals[CenterTID] : FVector3d::One();		// One so that normal check always passes
 
-	bool bUseAngleThreshold = (FilterProperties->AngleThreshold < 180.0f);
+	bool bVolumetric = (FilterProperties->BrushAreaMode == EMeshGroupPaintBrushAreaType::Volumetric);
+	bool bUseAngleThreshold = (bVolumetric == false) && (FilterProperties->AngleThreshold < 180.0f);
 	double DotAngleThreshold = FMathd::Cos(FilterProperties->AngleThreshold * FMathd::DegToRad);
 	bool bStopAtUVSeams = FilterProperties->bUVSeams;
 	bool bStopAtNormalSeams = FilterProperties->bNormalSeams;
 
-
 	auto CheckEdgeCriteria = [&](int32 t1, int32 t2) -> bool
 	{
 		if (bUseAngleThreshold == false || CenterNormal.Dot(TriNormals[t2]) > DotAngleThreshold)
@@ -414,7 +416,7 @@ void UMeshGroupPaintTool::UpdateROI(const FSculptBrushStamp& BrushStamp)
 	bool bFill = (FilterProperties->SubToolType == EMeshGroupPaintInteractionType::Fill);
 	bool bGroupFill = (FilterProperties->SubToolType == EMeshGroupPaintInteractionType::GroupFill);
 
-	if (FilterProperties->BrushAreaMode == EMeshGroupPaintBrushAreaType::Volumetric)
+	if (bVolumetric)
 	{
 		Octree.RangeQuery(BrushBox,
 			[&](int TriIdx) {
@@ -443,23 +445,35 @@ void UMeshGroupPaintTool::UpdateROI(const FSculptBrushStamp& BrushStamp)
 		}
 	}
 
-	if (bFill || bGroupFill)
+	if (bFill)
 	{
 		TArray<int32> StartROI;
-		TSet<int32> FillGroups;
 		for (int32 tid : TriangleROI)
 		{
 			StartROI.Add(tid);
-			FillGroups.Add(ActiveGroupSet->GetGroup(tid));
 		}
 		FMeshConnectedComponents::GrowToConnectedTriangles(Mesh, StartROI, TriangleROI, &TempROIBuffer,
-			[&](int t1, int t2)
+														   [&](int t1, int t2)
+		{
+			return CheckEdgeCriteria(t1, t2);
+		});
+	}
+	else if (bGroupFill)
+	{
+		TArray<int32> StartROI;
+		TSet<int32> FillGroups;
+		for (int32 tid : TriangleROI)
 		{
-			if (bFill || (bGroupFill && FillGroups.Contains(ActiveGroupSet->GetGroup(t2))))
+			if (ActiveGroupSet->GetGroup(tid) != SetGroupID)
 			{
-				return CheckEdgeCriteria(t1, t2);
+				StartROI.Add(tid);
+				FillGroups.Add(ActiveGroupSet->GetGroup(tid));
 			}
-			return false;
+		}
+		FMeshConnectedComponents::GrowToConnectedTriangles(Mesh, StartROI, TriangleROI, &TempROIBuffer,
+			[&](int t1, int t2)
+		{
+			return (FillGroups.Contains(ActiveGroupSet->GetGroup(t2)));
 		});
 	}
 
@@ -534,7 +548,7 @@ bool UMeshGroupPaintTool::UpdateStampPosition(const FRay& WorldRay)
 }
 
 
-void UMeshGroupPaintTool::ApplyStamp()
+bool UMeshGroupPaintTool::ApplyStamp()
 {
 	SCOPE_CYCLE_COUNTER(GroupPaintToolApplyStamp);
 
@@ -546,35 +560,40 @@ void UMeshGroupPaintTool::ApplyStamp()
 	FDynamicMesh3* Mesh = GetSculptMesh();
 	GroupBrushOp->ApplyStampByTriangles(Mesh, CurrentStamp, ROITriangleBuffer, ROIGroupBuffer);
 
-	SyncMeshWithGroupBuffer(Mesh);
+	bool bUpdated = SyncMeshWithGroupBuffer(Mesh);
 
 	LastStamp = CurrentStamp;
 	LastStamp.TimeStamp = FDateTime::Now();
+
+	return bUpdated;
 }
 
 
 
 
-void UMeshGroupPaintTool::SyncMeshWithGroupBuffer(FDynamicMesh3* Mesh)
+bool UMeshGroupPaintTool::SyncMeshWithGroupBuffer(FDynamicMesh3* Mesh)
 {
+	int NumModified = 0;
 	const int32 NumT = ROITriangleBuffer.Num();
 	// change update could be async here if we collected array of <idx,orig,new> and dispatched independenlty
 	for ( int32 k = 0; k < NumT; ++k)
 	{
 		int TriIdx = ROITriangleBuffer[k];
 		int32 CurGroupID = ActiveGroupSet->GetGroup(TriIdx);
-
 		if (FrozenGroups.Contains(CurGroupID))		// skip frozen groups
 		{
 			continue;
 		}
 
-		ActiveGroupEditBuilder->SaveTriangle(TriIdx, CurGroupID, ROIGroupBuffer[k]);
-
-		ActiveGroupSet->SetGroup(TriIdx, ROIGroupBuffer[k], *Mesh);
-
-		//ActiveVertexChange->UpdateVertexColor(VertIdx, OrigColor, NewColor);
+		if (ROIGroupBuffer[k] != CurGroupID)
+		{
+			ActiveGroupEditBuilder->SaveTriangle(TriIdx, CurGroupID, ROIGroupBuffer[k]);
+			ActiveGroupSet->SetGroup(TriIdx, ROIGroupBuffer[k], *Mesh);
+			//ActiveVertexChange->UpdateVertexColor(VertIdx, OrigColor, NewColor);
+			NumModified++;
+		}
 	}
+	return (NumModified > 0);
 }
 
 
@@ -1101,8 +1120,9 @@ void UMeshGroupPaintTool::OnTick(float DeltaTime)
 			});
 
 			// apply the stamp
-			ApplyStamp();
+			bool bGroupsModified = ApplyStamp();
 
+			if (bGroupsModified)
 			{
 				SCOPE_CYCLE_COUNTER(GroupPaintTool_Tick_UpdateMeshBlock);
 				DynamicMeshComponent->FastNotifyTriangleVerticesUpdated(TriangleROI, EMeshRenderAttributeFlags::VertexColors);
diff --git a/Engine/Plugins/Experimental/MeshModelingToolset/Source/MeshModelingTools/Public/MeshGroupPaintTool.h b/Engine/Plugins/Experimental/MeshModelingToolset/Source/MeshModelingTools/Public/MeshGroupPaintTool.h
--- a/Engine/Plugins/Experimental/MeshModelingToolset/Source/MeshModelingTools/Public/MeshGroupPaintTool.h
+++ b/Engine/Plugins/Experimental/MeshModelingToolset/Source/MeshModelingTools/Public/MeshGroupPaintTool.h
@@ -414,7 +414,7 @@ protected:
 	EMeshGroupPaintBrushType PendingStampType = EMeshGroupPaintBrushType::Paint;
 
 	bool UpdateStampPosition(const FRay& WorldRay);
-	void ApplyStamp();
+	bool ApplyStamp();
 
 	UE::Geometry::FDynamicMeshOctree3 Octree;
 
@@ -427,7 +427,7 @@ protected:
 
 	TArray<int32> ROITriangleBuffer;
 	TArray<int32> ROIGroupBuffer;
-	void SyncMeshWithGroupBuffer(FDynamicMesh3* Mesh);
+	bool SyncMeshWithGroupBuffer(FDynamicMesh3* Mesh);
 
 	TUniquePtr<FDynamicMeshGroupEditBuilder> ActiveGroupEditBuilder;
 	void BeginChange();
