ModelingTools: minor fixes to group and UV generation workflow
- restore default group generation mode to normal-angles in ConvertToPolygonsTool
- fix bug that was resulting in vertex IDs instead of triangle IDs being used to initialize furthest-point sampling in FindPolygonsAlgorithm
#rnx

diff --git a/Engine/Plugins/Experimental/MeshModelingToolset/Source/MeshModelingTools/Private/FindPolygonsAlgorithm.cpp b/Engine/Plugins/Experimental/MeshModelingToolset/Source/MeshModelingTools/Private/FindPolygonsAlgorithm.cpp
--- a/Engine/Plugins/Experimental/MeshModelingToolset/Source/MeshModelingTools/Private/FindPolygonsAlgorithm.cpp
+++ b/Engine/Plugins/Experimental/MeshModelingToolset/Source/MeshModelingTools/Private/FindPolygonsAlgorithm.cpp
@@ -204,8 +204,7 @@ bool FFindPolygonsAlgorithm::FindPolygonsFromFurthestPointSampling(int32 NumPoin
 	int32 NumConnected = Components.Num();
 	for (int32 k = 0; k < NumConnected; ++k)
 	{
-		FIndex3i TriVerts = Mesh->GetTriangle(Components[k].Indices[0]);
-		SeedIndices.Add(TriVerts.A);
+		SeedIndices.Add(Components[k].Indices[0]);
 	}
 
 	// initial incremental update from per-component points
diff --git a/Engine/Plugins/Experimental/MeshModelingToolset/Source/MeshModelingTools/Public/ConvertToPolygonsTool.h b/Engine/Plugins/Experimental/MeshModelingToolset/Source/MeshModelingTools/Public/ConvertToPolygonsTool.h
--- a/Engine/Plugins/Experimental/MeshModelingToolset/Source/MeshModelingTools/Public/ConvertToPolygonsTool.h
+++ b/Engine/Plugins/Experimental/MeshModelingToolset/Source/MeshModelingTools/Public/ConvertToPolygonsTool.h
@@ -52,8 +52,7 @@ class MESHMODELINGTOOLS_API UConvertToPolygonsToolProperties : public UInteracti
 public:
 	/** Strategy to use to group triangles */
 	UPROPERTY(EditAnywhere, Category = PolyGroups)
-		EConvertToPolygonsMode ConversionMode = EConvertToPolygonsMode::FromFurthestPointSampling;
-	//EConvertToPolygonsMode ConversionMode = EConvertToPolygonsMode::FaceNormalDeviation;
+	EConvertToPolygonsMode ConversionMode = EConvertToPolygonsMode::FaceNormalDeviation;
 
 	/** Tolerance for planarity */
 	UPROPERTY(EditAnywhere, Category = PolyGroups, meta = (UIMin = "0.001", UIMax = "20.0", ClampMin = "0.0", ClampMax = "90.0", EditCondition = "ConversionMode == EConvertToPolygonsMode::FaceNormalDeviation", EditConditionHides))
