MeshModelingTools: Made dynamic mesh getter/setter functions in ModelingToolTargetUtil support DynamicMeshProvider/Committer in cases where we still have those.

#rnx

diff --git a/Engine/Plugins/Experimental/MeshModelingToolset/Source/ModelingComponents/Private/ModelingToolTargetUtil.cpp b/Engine/Plugins/Experimental/MeshModelingToolset/Source/ModelingComponents/Private/ModelingToolTargetUtil.cpp
--- a/Engine/Plugins/Experimental/MeshModelingToolset/Source/ModelingComponents/Private/ModelingToolTargetUtil.cpp
+++ b/Engine/Plugins/Experimental/MeshModelingToolset/Source/ModelingComponents/Private/ModelingToolTargetUtil.cpp
@@ -3,6 +3,8 @@
 #include "ModelingToolTargetUtil.h"
 
 #include "ToolTargets/ToolTarget.h"
+#include "TargetInterfaces/DynamicMeshCommitter.h"
+#include "TargetInterfaces/DynamicMeshProvider.h"
 #include "TargetInterfaces/MaterialProvider.h"
 #include "TargetInterfaces/MeshDescriptionCommitter.h"
 #include "TargetInterfaces/MeshDescriptionProvider.h"
@@ -164,6 +166,12 @@ FDynamicMesh3 UE::ToolTarget::GetDynamicMeshCopy(UToolTarget* Target, bool bWant
 		return Mesh;
 	}
 
+	IDynamicMeshProvider* DynamicMeshProvider = Cast<IDynamicMeshProvider>(Target);
+	if (DynamicMeshProvider)
+	{
+		return *DynamicMeshProvider->GetDynamicMesh();
+	}
+
 	IMeshDescriptionProvider* MeshDescriptionProvider = Cast<IMeshDescriptionProvider>(Target);
 	FDynamicMesh3 Mesh(EMeshComponents::FaceGroups);
 	Mesh.EnableAttributes();
@@ -251,6 +259,23 @@ UE::ToolTarget::EDynamicMeshUpdateResult UE::ToolTarget::CommitDynamicMeshUpdate
 		return EDynamicMeshUpdateResult::Ok;
 	}
 
+	IDynamicMeshCommitter* DynamicMeshCommitter = Cast<IDynamicMeshCommitter>(Target);
+	if (DynamicMeshCommitter)
+	{
+		IDynamicMeshCommitter::FDynamicMeshCommitInfo CommitInfo;
+		CommitInfo.bTopologyChanged = bHaveModifiedTopology;
+		CommitInfo.bPolygroupsChanged = ConversionOptions.bSetPolyGroups;
+		CommitInfo.bPositionsChanged = ConversionOptions.bUpdatePositions;
+		CommitInfo.bNormalsChanged = ConversionOptions.bUpdateNormals;
+		CommitInfo.bTangentsChanged = ConversionOptions.bUpdateTangents;
+		CommitInfo.bUVsChanged = ConversionOptions.bUpdateUVs;
+		CommitInfo.bVertexColorsChanged = ConversionOptions.bUpdateVtxColors;
+
+		DynamicMeshCommitter->CommitDynamicMesh(UpdatedMesh, CommitInfo);
+
+		return EDynamicMeshUpdateResult::Ok;
+	}
+
 	IMeshDescriptionCommitter* MeshDescriptionCommitter = Cast<IMeshDescriptionCommitter>(Target);
 	if (MeshDescriptionCommitter)
 	{
