MeshModelingTools: Additional cleanup and improvements following PolyEdit refactor. Specifically, split inset/outset again, fixed a check being hit on property change in edge insertions, made default be plane cut for edge insertions, made PolyEdit use mesh descriptions as requirement again, and added ability to catch Enter/Esc key presses.

#rnx

diff --git a/Engine/Plugins/Editor/ModelingToolsEditorMode/Source/ModelingToolsEditorMode/Private/ModelingToolsEditorMode.cpp b/Engine/Plugins/Editor/ModelingToolsEditorMode/Source/ModelingToolsEditorMode/Private/ModelingToolsEditorMode.cpp
--- a/Engine/Plugins/Editor/ModelingToolsEditorMode/Source/ModelingToolsEditorMode/Private/ModelingToolsEditorMode.cpp
+++ b/Engine/Plugins/Editor/ModelingToolsEditorMode/Source/ModelingToolsEditorMode/Private/ModelingToolsEditorMode.cpp
@@ -600,7 +600,8 @@ void UModelingToolsEditorMode::Enter()
 		RegisterTool(UICommand, StringName, ActionModeBuilder);
 	};
 	RegisterPolyModelActionTool(EEditMeshPolygonsToolActions::Extrude, ToolManagerCommands.BeginPolyModelTool_Extrude, TEXT("PolyEdit_Extrude"));
-	RegisterPolyModelActionTool(EEditMeshPolygonsToolActions::InsetOutset, ToolManagerCommands.BeginPolyModelTool_InsetOutset, TEXT("PolyEdit_InsetOutset"));
+	RegisterPolyModelActionTool(EEditMeshPolygonsToolActions::Inset, ToolManagerCommands.BeginPolyModelTool_Inset, TEXT("PolyEdit_Inset"));
+	RegisterPolyModelActionTool(EEditMeshPolygonsToolActions::Outset, ToolManagerCommands.BeginPolyModelTool_Outset, TEXT("PolyEdit_Outset"));
 	RegisterPolyModelActionTool(EEditMeshPolygonsToolActions::CutFaces, ToolManagerCommands.BeginPolyModelTool_CutFaces, TEXT("PolyEdit_CutFaces"));
 
 
diff --git a/Engine/Plugins/Editor/ModelingToolsEditorMode/Source/ModelingToolsEditorMode/Private/ModelingToolsEditorModeToolkit.cpp b/Engine/Plugins/Editor/ModelingToolsEditorMode/Source/ModelingToolsEditorMode/Private/ModelingToolsEditorModeToolkit.cpp
--- a/Engine/Plugins/Editor/ModelingToolsEditorMode/Source/ModelingToolsEditorMode/Private/ModelingToolsEditorModeToolkit.cpp
+++ b/Engine/Plugins/Editor/ModelingToolsEditorMode/Source/ModelingToolsEditorMode/Private/ModelingToolsEditorModeToolkit.cpp
@@ -673,7 +673,8 @@ void FModelingToolsEditorModeToolkit::BuildToolPalette_Experimental(FName Palett
 		ToolbarBuilder.AddToolBarButton(Commands.BeginPolyModelTool_RingSelect);
 		ToolbarBuilder.AddSeparator();
 		ToolbarBuilder.AddToolBarButton(Commands.BeginPolyModelTool_Extrude);
-		ToolbarBuilder.AddToolBarButton(Commands.BeginPolyModelTool_InsetOutset);
+		ToolbarBuilder.AddToolBarButton(Commands.BeginPolyModelTool_Inset);
+		ToolbarBuilder.AddToolBarButton(Commands.BeginPolyModelTool_Outset);
 		ToolbarBuilder.AddToolBarButton(Commands.BeginPolyModelTool_CutFaces);
 		ToolbarBuilder.AddSeparator();
 		ToolbarBuilder.AddToolBarButton(Commands.BeginSubdividePolyTool);
diff --git a/Engine/Plugins/Editor/ModelingToolsEditorMode/Source/ModelingToolsEditorMode/Private/ModelingToolsManagerActions.cpp b/Engine/Plugins/Editor/ModelingToolsEditorMode/Source/ModelingToolsEditorMode/Private/ModelingToolsManagerActions.cpp
--- a/Engine/Plugins/Editor/ModelingToolsEditorMode/Source/ModelingToolsEditorMode/Private/ModelingToolsManagerActions.cpp
+++ b/Engine/Plugins/Editor/ModelingToolsEditorMode/Source/ModelingToolsEditorMode/Private/ModelingToolsManagerActions.cpp
@@ -127,7 +127,8 @@ void FModelingToolsManagerCommands::RegisterCommands()
 	UI_COMMAND(BeginPolyModelTool_RingSelect, "Rings", "PolyGroup Ring Selection Tool", EUserInterfaceActionType::ToggleButton, FInputChord());
 
 	UI_COMMAND(BeginPolyModelTool_Extrude, "Extrude", "PolyGroup Extrude Tool", EUserInterfaceActionType::ToggleButton, FInputChord());
-	UI_COMMAND(BeginPolyModelTool_InsetOutset, "Inset/Outset", "PolyGroup Inset/Outset Tool", EUserInterfaceActionType::ToggleButton, FInputChord());
+	UI_COMMAND(BeginPolyModelTool_Inset, "Inset", "PolyGroup Inset Tool", EUserInterfaceActionType::ToggleButton, FInputChord());
+	UI_COMMAND(BeginPolyModelTool_Outset, "Outset", "PolyGroup Outset Tool", EUserInterfaceActionType::ToggleButton, FInputChord());
 	UI_COMMAND(BeginPolyModelTool_CutFaces, "Cut", "PolyGroup Cut Faces Tool", EUserInterfaceActionType::ToggleButton, FInputChord());
 
 	UI_COMMAND(AcceptActiveTool, "Accept", "Accept the active tool", EUserInterfaceActionType::Button, FInputChord());
diff --git a/Engine/Plugins/Editor/ModelingToolsEditorMode/Source/ModelingToolsEditorMode/Public/ModelingToolsManagerActions.h b/Engine/Plugins/Editor/ModelingToolsEditorMode/Source/ModelingToolsEditorMode/Public/ModelingToolsManagerActions.h
--- a/Engine/Plugins/Editor/ModelingToolsEditorMode/Source/ModelingToolsEditorMode/Public/ModelingToolsManagerActions.h
+++ b/Engine/Plugins/Editor/ModelingToolsEditorMode/Source/ModelingToolsEditorMode/Public/ModelingToolsManagerActions.h
@@ -118,7 +118,8 @@ public:
 	TSharedPtr<FUICommandInfo> BeginPolyModelTool_RingSelect;
 
 	TSharedPtr<FUICommandInfo> BeginPolyModelTool_Extrude;
-	TSharedPtr<FUICommandInfo> BeginPolyModelTool_InsetOutset;
+	TSharedPtr<FUICommandInfo> BeginPolyModelTool_Inset;
+	TSharedPtr<FUICommandInfo> BeginPolyModelTool_Outset;
 	TSharedPtr<FUICommandInfo> BeginPolyModelTool_CutFaces;
 
 
diff --git a/Engine/Plugins/Experimental/MeshModelingToolset/Source/MeshModelingTools/Private/EditMeshPolygonsTool.cpp b/Engine/Plugins/Experimental/MeshModelingToolset/Source/MeshModelingTools/Private/EditMeshPolygonsTool.cpp
--- a/Engine/Plugins/Experimental/MeshModelingToolset/Source/MeshModelingTools/Private/EditMeshPolygonsTool.cpp
+++ b/Engine/Plugins/Experimental/MeshModelingToolset/Source/MeshModelingTools/Private/EditMeshPolygonsTool.cpp
@@ -29,8 +29,8 @@
 #include "Selection/PolygonSelectionMechanic.h"
 #include "Selections/MeshConnectedComponents.h"
 #include "TargetInterfaces/MaterialProvider.h"
-#include "TargetInterfaces/DynamicMeshCommitter.h"
-#include "TargetInterfaces/DynamicMeshProvider.h"
+#include "TargetInterfaces/MeshDescriptionCommitter.h"
+#include "TargetInterfaces/MeshDescriptionProvider.h"
 #include "TargetInterfaces/PrimitiveComponentBackedTarget.h"
 #include "ToolActivities/PolyEditActivityContext.h"
 #include "ToolActivities/PolyEditExtrudeActivity.h"
@@ -69,8 +69,16 @@ const FToolTargetTypeRequirements& UEditMeshPolygonsToolBuilder::GetTargetRequir
 {
 	static FToolTargetTypeRequirements TypeRequirements({
 		UMaterialProvider::StaticClass(),
-		UDynamicMeshCommitter::StaticClass(),
-		UDynamicMeshProvider::StaticClass(),
+
+		// TODO: In reality we don't actually need a mesh description- we need anything that we
+		// can get a dynamic mesh out of. Unfortunately we have not yet settled on an interface that
+		// provides this- DynamicMeshProvider (which currently gives a dynamic mesh shared ptr) is 
+		// likely to be changed or replaced with something that gives a UDynamicMesh instead. For now
+		// we ask for a mesh description as a hack, since all the targets we currently use are able to
+		// provide both it and a dynamic mesh.
+		UMeshDescriptionCommitter::StaticClass(),
+		UMeshDescriptionProvider::StaticClass(),
+		
 		UPrimitiveComponentBackedTarget::StaticClass()
 		});
 	return TypeRequirements;
@@ -248,7 +256,7 @@ void UEditMeshPolygonsTool::Setup()
 	// and cause UpdateGizmoFrame() to be called emitting a spurious Transform change. 
 	CommonProps->SilentUpdateWatched();
 
-	CurrentMesh = Cast<IDynamicMeshProvider>(Target)->GetDynamicMesh();
+	CurrentMesh = MakeShared<FDynamicMesh3>(UE::ToolTarget::GetDynamicMeshCopy(Target));
 
 	// TODO: Do we need this?
 	FMeshNormals::QuickComputeVertexNormals(*CurrentMesh);
@@ -497,13 +505,9 @@ void UEditMeshPolygonsTool::Shutdown(EToolShutdownType ShutdownType)
 				SelectionMechanic->GetStorableSelection(*NewStoredToolSelection, bModifiedTopology ? &CompactMaps : nullptr);
 			}
 
-			// this block bakes the modified DynamicMeshComponent back into the StaticMeshComponent inside an undo transaction
+			// Bake CurrentMesh back to target inside an undo transaction
 			GetToolManager()->BeginUndoTransaction(LOCTEXT("EditMeshPolygonsToolTransactionName", "Deform Mesh"));
-			Preview->PreviewMesh->ProcessMesh([&](const FDynamicMesh3& ReadMesh)
-			{
-				FConversionToMeshDescriptionOptions ConversionOptions;
-				UE::ToolTarget::CommitDynamicMeshUpdate(Target, ReadMesh, bModifiedTopology, ConversionOptions);
-			});
+			UE::ToolTarget::CommitDynamicMeshUpdate(Target, *CurrentMesh, bModifiedTopology);
 
 			// The stored selection change should go into this transaction as well.
 			// If we're keeping the same selection, we still need to store it back, though we could do it outside
@@ -866,8 +870,15 @@ void UEditMeshPolygonsTool::OnTick(float DeltaTime)
 			StartActivity(ExtrudeActivity);
 			break;
 		}
-		case EEditMeshPolygonsToolActions::InsetOutset:
+		case EEditMeshPolygonsToolActions::Inset:
 		{
+			InsetOutsetActivity->Settings->bOutset = false;
+			StartActivity(InsetOutsetActivity);
+			break;
+		}
+		case EEditMeshPolygonsToolActions::Outset:
+		{
+			InsetOutsetActivity->Settings->bOutset = true;
 			StartActivity(InsetOutsetActivity);
 			break;
 		}
@@ -970,13 +981,13 @@ void UEditMeshPolygonsTool::StartActivity(TObjectPtr<UInteractiveToolActivity> A
 	}
 }
 
-void UEditMeshPolygonsTool::EndCurrentActivity()
+void UEditMeshPolygonsTool::EndCurrentActivity(EToolShutdownType ShutdownType)
 {
 	if (CurrentActivity)
 	{
 		if (CurrentActivity->IsRunning())
 		{
-			CurrentActivity->End(EToolShutdownType::Cancel);
+			CurrentActivity->End(ShutdownType);
 
 			// Reset info message.
 			GetToolManager()->DisplayMessage(DefaultMessage,
@@ -1759,7 +1770,7 @@ void UEditMeshPolygonsTool::EmitCurrentMeshChangeAndUpdate(const FText& Transact
 
 		SelectionMechanic->BeginChange();
 		SelectionMechanic->ClearSelection();
-		GetToolManager()->EmitObjectChange(SelectionMechanic, SelectionMechanic->EndChange(), LOCTEXT("PolyMeshExtrudeChangeClearSelection", "ClearSelection"));
+		GetToolManager()->EmitObjectChange(SelectionMechanic, SelectionMechanic->EndChange(), LOCTEXT("ClearSelection", "Clear Selection"));
 	}
 
 	GetToolManager()->EmitObjectChange(this, 
@@ -1776,7 +1787,7 @@ void UEditMeshPolygonsTool::EmitCurrentMeshChangeAndUpdate(const FText& Transact
 	{
 		SelectionMechanic->BeginChange();
 		SelectionMechanic->SetSelection(*OutputSelectionToUse);
-		GetToolManager()->EmitObjectChange(SelectionMechanic, SelectionMechanic->EndChange(), LOCTEXT("PolyMeshExtrudeChangeSetSelection", "SetSelection"));
+		GetToolManager()->EmitObjectChange(SelectionMechanic, SelectionMechanic->EndChange(), LOCTEXT("SetSelection", "Set Selection"));
 	}
 
 	GetToolManager()->EndUndoTransaction();
@@ -1867,6 +1878,45 @@ void UEditMeshPolygonsTool::SetActionButtonPanelsVisible(bool bVisible)
 	}
 }
 
+
+bool UEditMeshPolygonsTool::CanCurrentlyNestedCancel()
+{
+	return CurrentActivity != nullptr 
+		|| (SelectionMechanic && !SelectionMechanic->GetActiveSelection().IsEmpty());
+}
+
+bool UEditMeshPolygonsTool::ExecuteNestedCancelCommand()
+{
+	if (CurrentActivity)
+	{
+		EndCurrentActivity(EToolShutdownType::Cancel);
+		return true;
+	}
+	else if (SelectionMechanic && !SelectionMechanic->GetActiveSelection().IsEmpty())
+	{
+		SelectionMechanic->BeginChange();
+		SelectionMechanic->ClearSelection();
+		GetToolManager()->EmitObjectChange(SelectionMechanic, SelectionMechanic->EndChange(), LOCTEXT("ClearSelection", "Clear Selection"));
+		return true;
+	}
+	return false;
+}
+
+bool UEditMeshPolygonsTool::CanCurrentlyNestedAccept()
+{
+	return CurrentActivity != nullptr;
+}
+
+bool UEditMeshPolygonsTool::ExecuteNestedAcceptCommand()
+{
+	if (CurrentActivity)
+	{
+		EndCurrentActivity(EToolShutdownType::Accept);
+		return true;
+	}
+	return false;
+}
+
 void FEditMeshPolygonsToolMeshChange::Apply(UObject* Object)
 {
 	UEditMeshPolygonsTool* Tool = Cast<UEditMeshPolygonsTool>(Object);
diff --git a/Engine/Plugins/Experimental/MeshModelingToolset/Source/MeshModelingTools/Private/ToolActivities/PolyEditExtrudeActivity.cpp b/Engine/Plugins/Experimental/MeshModelingToolset/Source/MeshModelingTools/Private/ToolActivities/PolyEditExtrudeActivity.cpp
--- a/Engine/Plugins/Experimental/MeshModelingToolset/Source/MeshModelingTools/Private/ToolActivities/PolyEditExtrudeActivity.cpp
+++ b/Engine/Plugins/Experimental/MeshModelingToolset/Source/MeshModelingTools/Private/ToolActivities/PolyEditExtrudeActivity.cpp
@@ -95,12 +95,27 @@ bool UPolyEditExtrudeActivity::CanAccept() const
 	return false;
 }
 
-EToolActivityEndResult UPolyEditExtrudeActivity::End(EToolShutdownType)
+EToolActivityEndResult UPolyEditExtrudeActivity::End(EToolShutdownType ShutdownType)
 {
-	Clear();
-	EToolActivityEndResult ToReturn = bIsRunning ? EToolActivityEndResult::Cancelled : EToolActivityEndResult::ErrorDuringEnd;
-	bIsRunning = false;
-	return ToReturn;
+	if (!bIsRunning)
+	{
+		Clear();
+		return EToolActivityEndResult::ErrorDuringEnd;
+	}
+
+	if (ShutdownType == EToolShutdownType::Cancel)
+	{
+		Clear();
+		bIsRunning = false;
+		return EToolActivityEndResult::Cancelled;
+	}
+	else
+	{
+		ApplyExtrude();
+		Clear();
+		bIsRunning = false;
+		return EToolActivityEndResult::Completed;
+	}
 }
 
 
@@ -205,11 +220,6 @@ void UPolyEditExtrudeActivity::ApplyExtrude()
 	// Emit undo  (also updates relevant structures)
 	ActivityContext->EmitCurrentMeshChangeAndUpdate(LOCTEXT("PolyMeshExtrudeChange", "Extrude"),
 		Extruder.ChangeTracker->EndChange(), NewSelection, true);
-
-	// End activity
-	Clear();
-	bIsRunning = false;
-	Cast<IToolActivityHost>(ParentTool)->NotifyActivitySelfEnded(this);
 }
 
 void UPolyEditExtrudeActivity::Clear()
@@ -291,6 +301,11 @@ void UPolyEditExtrudeActivity::OnClicked(const FInputDeviceRay& ClickPos)
 	if (bIsRunning)
 	{
 		ApplyExtrude();
+
+		// End activity
+		Clear();
+		bIsRunning = false;
+		Cast<IToolActivityHost>(ParentTool)->NotifyActivitySelfEnded(this);
 	}
 }
 
diff --git a/Engine/Plugins/Experimental/MeshModelingToolset/Source/MeshModelingTools/Private/ToolActivities/PolyEditInsertEdgeActivity.cpp b/Engine/Plugins/Experimental/MeshModelingToolset/Source/MeshModelingTools/Private/ToolActivities/PolyEditInsertEdgeActivity.cpp
--- a/Engine/Plugins/Experimental/MeshModelingToolset/Source/MeshModelingTools/Private/ToolActivities/PolyEditInsertEdgeActivity.cpp
+++ b/Engine/Plugins/Experimental/MeshModelingToolset/Source/MeshModelingTools/Private/ToolActivities/PolyEditInsertEdgeActivity.cpp
@@ -154,6 +154,7 @@ EToolActivityEndResult UPolyEditInsertEdgeActivity::End(EToolShutdownType Shutdo
 
 	ActivityContext->Preview->OnOpCompleted.RemoveAll(this);
 	ActivityContext->Preview->OnMeshUpdated.RemoveAll(this);
+	ClearPreview(true);
 	ActivityContext->Preview->ClearOpFactory();
 
 	LatestOpTopologyResult.Reset();
@@ -279,7 +280,10 @@ bool UPolyEditInsertEdgeActivity::CanAccept() const
 void UPolyEditInsertEdgeActivity::OnPropertyModified(UObject* PropertySet, FProperty* Property)
 {
 	PreviewEdges.Reset();
-	ActivityContext->Preview->InvalidateResult();
+
+	// Don't clear drawn elements because we may be getting the second endpoint, and still
+	// need to keep the first.
+	ClearPreview(false); 
 }
 
 void UPolyEditInsertEdgeActivity::ClearPreview(bool bClearDrawnElements)
@@ -404,6 +408,11 @@ bool UPolyEditInsertEdgeActivity::OnUpdateHover(const FInputDeviceRay& DevicePos
 {
 	using namespace PolyEditInsertEdgeActivityLocals;
 
+	if (!bIsRunning)
+	{
+		return false;
+	}
+
 	switch (ToolState)
 	{
 	case EState::WaitingForInsertComplete:
@@ -474,6 +483,10 @@ bool UPolyEditInsertEdgeActivity::OnUpdateHover(const FInputDeviceRay& DevicePos
 
 void UPolyEditInsertEdgeActivity::OnEndHover()
 {
+	if (!bIsRunning)
+	{
+		return;
+	}
 	switch (ToolState)
 	{
 	case EState::WaitingForInsertComplete:
diff --git a/Engine/Plugins/Experimental/MeshModelingToolset/Source/MeshModelingTools/Private/ToolActivities/PolyEditInsertEdgeLoopActivity.cpp b/Engine/Plugins/Experimental/MeshModelingToolset/Source/MeshModelingTools/Private/ToolActivities/PolyEditInsertEdgeLoopActivity.cpp
--- a/Engine/Plugins/Experimental/MeshModelingToolset/Source/MeshModelingTools/Private/ToolActivities/PolyEditInsertEdgeLoopActivity.cpp
+++ b/Engine/Plugins/Experimental/MeshModelingToolset/Source/MeshModelingTools/Private/ToolActivities/PolyEditInsertEdgeLoopActivity.cpp
@@ -38,6 +38,10 @@ TUniquePtr<FDynamicMeshOperator> UPolyEditInsertEdgeLoopActivity::MakeNewOperato
 
 	Op->VertexTolerance = Settings->VertexTolerance;
 
+	if (!ensure(InputGroupEdgeID >= 0))
+	{
+		return Op;
+	}
 	Op->GroupEdgeID = InputGroupEdgeID;
 
 	Op->StartCornerID = Settings->bFlipOffsetDirection ?
@@ -172,6 +176,7 @@ EToolActivityEndResult UPolyEditInsertEdgeLoopActivity::End(EToolShutdownType Sh
 
 	ActivityContext->Preview->OnOpCompleted.RemoveAll(this);
 	ActivityContext->Preview->OnMeshUpdated.RemoveAll(this);
+	ClearPreview();
 	ActivityContext->Preview->ClearOpFactory();
 
 	bIsRunning = false;
@@ -311,21 +316,13 @@ bool UPolyEditInsertEdgeLoopActivity::CanAccept() const
 void UPolyEditInsertEdgeLoopActivity::OnPropertyModified(UObject* PropertySet, FProperty* Property)
 {
 	PreviewEdges.Reset();
-	ActivityContext->Preview->InvalidateResult();
+	ClearPreview();
 }
 
 FInputRayHit UPolyEditInsertEdgeLoopActivity::HitTest(const FRay& WorldRay)
 {
 	FInputRayHit Hit;
 
-	// Early out if the activity is not running. This is actually important because the behavior is
-	// always in the behavior list while the tool is running (we don't have a way to add/remove at
-	// will).
-	if (!bIsRunning)
-	{
-		return Hit; // Hit.bHit is false
-	}
-
 	// See if we hit an edge
 	FRay3d LocalRay((FVector3d)TargetTransform.InverseTransformPosition(WorldRay.Origin),
 		(FVector3d)TargetTransform.InverseTransformVector(WorldRay.Direction), false);
@@ -475,7 +472,10 @@ void UPolyEditInsertEdgeLoopActivity::ConditionallyUpdatePreview(int32 NewGroupI
 
 FInputRayHit UPolyEditInsertEdgeLoopActivity::BeginHoverSequenceHitTest(const FInputDeviceRay& PressPos)
 {
-	if (bWaitingForInsertionCompletion)
+	// Note that the check for bIsRunning is important here and other input handlers since the
+	// hover/click behaviors are always in the behavior list while the tool is running (we don't
+	// currently have a way to add/remove at will).
+	if (bWaitingForInsertionCompletion || !bIsRunning)
 	{
 		return FInputRayHit();
 	}
@@ -485,7 +485,7 @@ FInputRayHit UPolyEditInsertEdgeLoopActivity::BeginHoverSequenceHitTest(const FI
 
 bool UPolyEditInsertEdgeLoopActivity::OnUpdateHover(const FInputDeviceRay& DevicePos)
 {
-	if (bWaitingForInsertionCompletion)
+	if (bWaitingForInsertionCompletion || !bIsRunning)
 	{
 		return false;
 	}
@@ -506,7 +506,7 @@ void UPolyEditInsertEdgeLoopActivity::OnEndHover()
 	// cancel any requested clear on mouse down, and keep updating hover during drag. On mouse up
 	// we accept whatever hover we have.
 
-	if (!bWaitingForInsertionCompletion)
+	if (bIsRunning && !bWaitingForInsertionCompletion)
 	{
 		ClearPreview();
 	}
@@ -515,7 +515,7 @@ void UPolyEditInsertEdgeLoopActivity::OnEndHover()
 FInputRayHit UPolyEditInsertEdgeLoopActivity::IsHitByClick(const FInputDeviceRay& ClickPos)
 {
 	FInputRayHit Hit;
-	if (bWaitingForInsertionCompletion)
+	if (bWaitingForInsertionCompletion || !bIsRunning)
 	{
 		return Hit;
 	}
diff --git a/Engine/Plugins/Experimental/MeshModelingToolset/Source/MeshModelingTools/Private/ToolActivities/PolyEditInsetOutsetActivity.cpp b/Engine/Plugins/Experimental/MeshModelingToolset/Source/MeshModelingTools/Private/ToolActivities/PolyEditInsetOutsetActivity.cpp
--- a/Engine/Plugins/Experimental/MeshModelingToolset/Source/MeshModelingTools/Private/ToolActivities/PolyEditInsetOutsetActivity.cpp
+++ b/Engine/Plugins/Experimental/MeshModelingToolset/Source/MeshModelingTools/Private/ToolActivities/PolyEditInsetOutsetActivity.cpp
@@ -26,10 +26,10 @@ void UPolyEditInsetOutsetActivity::Setup(UInteractiveTool* ParentToolIn)
 {
 	Super::Setup(ParentToolIn);
 
-	InsetProperties = NewObject<UPolyEditInsetOutsetProperties>();
-	InsetProperties->RestoreProperties(ParentTool.Get());
-	AddToolPropertySource(InsetProperties);
-	SetToolPropertySourceEnabled(InsetProperties, false);
+	Settings = NewObject<UPolyEditInsetOutsetProperties>();
+	Settings->RestoreProperties(ParentTool.Get());
+	AddToolPropertySource(Settings);
+	SetToolPropertySourceEnabled(Settings, false);
 
 	// Register ourselves to receive clicks and hover
 	USingleClickInputBehavior* ClickBehavior = NewObject<USingleClickInputBehavior>();
@@ -45,9 +45,9 @@ void UPolyEditInsetOutsetActivity::Setup(UInteractiveTool* ParentToolIn)
 void UPolyEditInsetOutsetActivity::Shutdown(EToolShutdownType ShutdownType)
 {
 	Clear();
-	InsetProperties->SaveProperties(ParentTool.Get());
+	Settings->SaveProperties(ParentTool.Get());
 
-	InsetProperties = nullptr;
+	Settings = nullptr;
 	ParentTool = nullptr;
 	ActivityContext = nullptr;
 }
@@ -66,6 +66,9 @@ EToolActivityStartResult UPolyEditInsetOutsetActivity::Start()
 {
 	if (!CanStart())
 	{
+		ParentTool->GetToolManager()->DisplayMessage(
+			LOCTEXT("OnInsetOutsetFailedMesssage", "Cannot inset or outset without face selection."),
+			EToolMessageLevel::UserWarning);
 		return EToolActivityStartResult::FailedStart;
 	}
 
@@ -83,13 +86,27 @@ bool UPolyEditInsetOutsetActivity::CanAccept() const
 	return false;
 }
 
-EToolActivityEndResult UPolyEditInsetOutsetActivity::End(EToolShutdownType)
+EToolActivityEndResult UPolyEditInsetOutsetActivity::End(EToolShutdownType ShutdownType)
 {
-	Clear();
-	EToolActivityEndResult ToReturn = bIsRunning ? EToolActivityEndResult::Cancelled 
-		: EToolActivityEndResult::ErrorDuringEnd;
-	bIsRunning = false;
-	return ToReturn;
+	if (!bIsRunning)
+	{
+		Clear();
+		return EToolActivityEndResult::ErrorDuringEnd;
+	}
+
+	if (ShutdownType == EToolShutdownType::Cancel)
+	{
+		Clear();
+		bIsRunning = false;
+		return EToolActivityEndResult::Cancelled;
+	}
+	else
+	{
+		ApplyInset();
+		Clear();
+		bIsRunning = false;
+		return EToolActivityEndResult::Completed;
+	}
 }
 
 void UPolyEditInsetOutsetActivity::BeginInset()
@@ -124,7 +141,7 @@ void UPolyEditInsetOutsetActivity::BeginInset()
 	Loops.Loops[0].GetVertices(LoopVertices);
 	CurveDistMechanic->InitializePolyLoop(LoopVertices, FTransform3d::Identity());
 
-	SetToolPropertySourceEnabled(InsetProperties, true);
+	SetToolPropertySourceEnabled(Settings, true);
 
 	float BoundsMaxDim = ActivityContext->CurrentMesh->GetBounds().MaxDim();
 	if (BoundsMaxDim > 0)
@@ -146,12 +163,12 @@ void UPolyEditInsetOutsetActivity::ApplyInset()
 	FInsetMeshRegion Inset(ActivityContext->CurrentMesh.Get());
 	Inset.UVScaleFactor = UVScaleFactor;
 	Inset.Triangles = ActiveTriangleSelection;
-	Inset.InsetDistance = (InsetProperties->bOutset) ? -CurveDistMechanic->CurrentDistance
+	Inset.InsetDistance = (Settings->bOutset) ? -CurveDistMechanic->CurrentDistance
 		: CurveDistMechanic->CurrentDistance;
-	Inset.bReproject = (InsetProperties->bOutset) ? false : InsetProperties->bReproject;
-	Inset.Softness = InsetProperties->Softness;
-	Inset.bSolveRegionInteriors = !InsetProperties->bBoundaryOnly;
-	Inset.AreaCorrection = InsetProperties->AreaScale;
+	Inset.bReproject = (Settings->bOutset) ? false : Settings->bReproject;
+	Inset.Softness = Settings->Softness;
+	Inset.bSolveRegionInteriors = !Settings->bBoundaryOnly;
+	Inset.AreaCorrection = Settings->AreaScale;
 
 	Inset.ChangeTracker = MakeUnique<FDynamicMeshChangeTracker>(ActivityContext->CurrentMesh.Get());
 	Inset.ChangeTracker->BeginChange();
@@ -162,11 +179,6 @@ void UPolyEditInsetOutsetActivity::ApplyInset()
 	// Emit undo (also updates relevant structures)
 	ActivityContext->EmitCurrentMeshChangeAndUpdate(LOCTEXT("PolyMeshInsetOutsetChange", "Inset/Outset"),
 		Inset.ChangeTracker->EndChange(), ActiveSelection, true);
-
-	// End activity
-	Clear();
-	bIsRunning = false;
-	Cast<IToolActivityHost>(ParentTool)->NotifyActivitySelfEnded(this);
 }
 
 void UPolyEditInsetOutsetActivity::Clear()
@@ -180,7 +192,7 @@ void UPolyEditInsetOutsetActivity::Clear()
 	ActivityContext->Preview->PreviewMesh->SetSecondaryBuffersVisibility(true);
 
 	CurveDistMechanic = nullptr;
-	SetToolPropertySourceEnabled(InsetProperties, false);
+	SetToolPropertySourceEnabled(Settings, false);
 }
 
 void UPolyEditInsetOutsetActivity::Render(IToolsContextRenderAPI* RenderAPI)
@@ -195,11 +207,11 @@ void UPolyEditInsetOutsetActivity::Tick(float DeltaTime)
 {
 	if (EditPreview && bPreviewUpdatePending)
 	{
-		double Sign = InsetProperties->bOutset ? -1.0 : 1.0;
-		bool bReproject = (InsetProperties->bOutset) ? false : InsetProperties->bReproject;
-		double Softness = InsetProperties->Softness;
-		bool bBoundaryOnly = InsetProperties->bBoundaryOnly;
-		double AreaCorrection = InsetProperties->AreaScale;
+		double Sign = Settings->bOutset ? -1.0 : 1.0;
+		bool bReproject = (Settings->bOutset) ? false : Settings->bReproject;
+		double Softness = Settings->Softness;
+		bool bBoundaryOnly = Settings->bBoundaryOnly;
+		double AreaCorrection = Settings->AreaScale;
 		EditPreview->UpdateInsetType(Sign * CurveDistMechanic->CurrentDistance,
 			bReproject, Softness, AreaCorrection, bBoundaryOnly);
 
@@ -219,6 +231,11 @@ void UPolyEditInsetOutsetActivity::OnClicked(const FInputDeviceRay& ClickPos)
 	if (bIsRunning)
 	{
 		ApplyInset();
+
+		// End activity
+		Clear();
+		bIsRunning = false;
+		Cast<IToolActivityHost>(ParentTool)->NotifyActivitySelfEnded(this);
 	}
 }
 
diff --git a/Engine/Plugins/Experimental/MeshModelingToolset/Source/MeshModelingTools/Public/EditMeshPolygonsTool.h b/Engine/Plugins/Experimental/MeshModelingToolset/Source/MeshModelingTools/Public/EditMeshPolygonsTool.h
--- a/Engine/Plugins/Experimental/MeshModelingToolset/Source/MeshModelingTools/Public/EditMeshPolygonsTool.h
+++ b/Engine/Plugins/Experimental/MeshModelingToolset/Source/MeshModelingTools/Public/EditMeshPolygonsTool.h
@@ -9,6 +9,7 @@
 #include "DynamicMesh/DynamicMeshChangeTracker.h" // FDynamicMeshChange for TUniquePtr
 #include "InteractiveToolActivity.h" // IToolActivityHost
 #include "InteractiveToolBuilder.h"
+#include "InteractiveToolQueryInterfaces.h" // IInteractiveToolNestedAcceptCancelAPI
 #include "Operations/GroupTopologyDeformer.h"
 #include "SingleSelectionTool.h"
 
@@ -98,7 +99,8 @@ enum class EEditMeshPolygonsToolActions
 	NoAction,
 	CancelCurrent,
 	Extrude,
-	InsetOutset,
+	Inset,
+	Outset,
 	InsertEdge,
 	InsertEdgeLoop,
 	Complete,
@@ -185,9 +187,13 @@ public:
 	UFUNCTION(CallInEditor, Category = FaceEdits, meta = (DisplayName = "Extrude", DisplayPriority = 1))
 	void Extrude() { PostAction(EEditMeshPolygonsToolActions::Extrude); }
 
-	/** Inset/Outset the current set of selected faces. Click in viewport to confirm inset/outset distance. */
-	UFUNCTION(CallInEditor, Category = FaceEdits, meta = (DisplayName = "Inset/Outset", DisplayPriority = 3))
-	void InsetOutset() { PostAction(EEditMeshPolygonsToolActions::InsetOutset);	}
+	/** Inset the current set of selected faces. Click in viewport to confirm inset distance. */
+	UFUNCTION(CallInEditor, Category = FaceEdits, meta = (DisplayName = "Inset", DisplayPriority = 2))
+	void Inset() { PostAction(EEditMeshPolygonsToolActions::Inset);	}
+
+	/** Outset the current set of selected faces. Click in viewport to confirm outset distance. */
+	UFUNCTION(CallInEditor, Category = FaceEdits, meta = (DisplayName = "Outset", DisplayPriority = 3))
+	void Outset() { PostAction(EEditMeshPolygonsToolActions::Outset);	}
 
 	/** Merge the current set of selected faces into a single face. */
 	UFUNCTION(CallInEditor, Category = FaceEdits, meta = (DisplayName = "Merge", DisplayPriority = 4))
@@ -239,8 +245,12 @@ public:
 	void Extrude() { PostAction(EEditMeshPolygonsToolActions::Extrude); }
 
 	/** Inset the current set of selected faces. Click in viewport to confirm inset distance. */
-	UFUNCTION(CallInEditor, Category = TriangleEdits, meta = (DisplayName = "Inset/Outset", DisplayPriority = 3))
-	void InsetOutset() { PostAction(EEditMeshPolygonsToolActions::InsetOutset);	}
+	UFUNCTION(CallInEditor, Category = FaceEdits, meta = (DisplayName = "Inset", DisplayPriority = 2))
+	void Inset() { PostAction(EEditMeshPolygonsToolActions::Inset);	}
+
+	/** Outset the current set of selected faces. Click in viewport to confirm outset distance. */
+	UFUNCTION(CallInEditor, Category = FaceEdits, meta = (DisplayName = "Outset", DisplayPriority = 3))
+	void Outset() { PostAction(EEditMeshPolygonsToolActions::Outset);	}
 
 	/** Delete the current set of selected faces */
 	UFUNCTION(CallInEditor, Category = TriangleEdits, meta = (DisplayName = "Delete", DisplayPriority = 4))
@@ -299,10 +309,10 @@ class MESHMODELINGTOOLS_API UEditMeshPolygonsToolEdgeActions : public UEditMeshP
 {
 	GENERATED_BODY()
 public:
-	UFUNCTION(CallInEditor, Category = EdgeEdits, meta = (DisplayName = "InsertEdgeLoop", DisplayPriority = 1))
+	UFUNCTION(CallInEditor, Category = ShapeEdits, meta = (DisplayName = "InsertEdgeLoop", DisplayPriority = 1))
 	void InsertEdgeLoop() { PostAction(EEditMeshPolygonsToolActions::InsertEdgeLoop); }
 
-	UFUNCTION(CallInEditor, Category = EdgeEdits, meta = (DisplayName = "Insert Edge", DisplayPriority = 2))
+	UFUNCTION(CallInEditor, Category = ShapeEdits, meta = (DisplayName = "Insert Edge", DisplayPriority = 2))
 	void InsertEdge() { PostAction(EEditMeshPolygonsToolActions::InsertEdge); }
 
 	UFUNCTION(CallInEditor, Category = EdgeEdits, meta = (DisplayName = "Weld", DisplayPriority = 3))
@@ -362,7 +372,10 @@ public:
  *
  */
 UCLASS()
-class MESHMODELINGTOOLS_API UEditMeshPolygonsTool : public USingleSelectionTool, public IToolActivityHost, public IMeshVertexCommandChangeTarget
+class MESHMODELINGTOOLS_API UEditMeshPolygonsTool : public USingleSelectionTool, 
+	public IToolActivityHost, 
+	public IMeshVertexCommandChangeTarget,
+	public IInteractiveToolNestedAcceptCancelAPI
 {
 	GENERATED_BODY()
 	using FFrame3d = UE::Geometry::FFrame3d;
@@ -404,6 +417,14 @@ public:
 	// IMeshVertexCommandChangeTarget
 	virtual void ApplyChange(const FMeshVertexChange* Change, bool bRevert) override;
 
+	// IInteractiveToolNestedAcceptCancelAPI
+	virtual bool SupportsNestedCancelCommand() override { return true; }
+	virtual bool CanCurrentlyNestedCancel() override;
+	virtual bool ExecuteNestedCancelCommand() override;
+	virtual bool SupportsNestedAcceptCommand() override { return true; }
+	virtual bool CanCurrentlyNestedAccept() override;
+	virtual bool ExecuteNestedAcceptCommand() override;
+
 public:
 
 	virtual void RequestAction(EEditMeshPolygonsToolActions ActionType);
@@ -516,7 +537,7 @@ protected:
 	int32 ActivityTimestamp = 1;
 
 	void StartActivity(TObjectPtr<UInteractiveToolActivity> Activity);
-	void EndCurrentActivity();
+	void EndCurrentActivity(EToolShutdownType ShutdownType = EToolShutdownType::Cancel);
 	void SetActionButtonPanelsVisible(bool bVisible);
 
 	// Emit an undoable change to CurrentMesh and update related structures (preview, spatial, etc)
diff --git a/Engine/Plugins/Experimental/MeshModelingToolset/Source/MeshModelingTools/Public/ToolActivities/PolyEditInsertEdgeActivity.h b/Engine/Plugins/Experimental/MeshModelingToolset/Source/MeshModelingTools/Public/ToolActivities/PolyEditInsertEdgeActivity.h
--- a/Engine/Plugins/Experimental/MeshModelingToolset/Source/MeshModelingTools/Public/ToolActivities/PolyEditInsertEdgeActivity.h
+++ b/Engine/Plugins/Experimental/MeshModelingToolset/Source/MeshModelingTools/Public/ToolActivities/PolyEditInsertEdgeActivity.h
@@ -40,13 +40,10 @@ class MESHMODELINGTOOLS_API UGroupEdgeInsertionProperties : public UInteractiveT
 public:
 	/** Determines how group edges are added to the geometry */
 	UPROPERTY(EditAnywhere, Category = InsertEdge)
-	EGroupEdgeInsertionMode InsertionMode = EGroupEdgeInsertionMode::Retriangulate;
-
-	UPROPERTY(EditAnywhere, Category = InsertEdge)
-	bool bWireframe = true;
+	EGroupEdgeInsertionMode InsertionMode = EGroupEdgeInsertionMode::PlaneCut;
 
 	/** How close a new loop edge needs to pass next to an existing vertex to use that vertex rather than creating a new one (used for plane cut). */
-	UPROPERTY(EditAnywhere, Category = InsertEdge, AdvancedDisplay)
+	UPROPERTY(EditAnywhere, Category = InsertEdge, AdvancedDisplay, meta = (UIMin = "0", UIMax = "0.01", ClampMin = "0", ClampMax = "10"))
 	double VertexTolerance = 0.001;
 };
 
diff --git a/Engine/Plugins/Experimental/MeshModelingToolset/Source/MeshModelingTools/Public/ToolActivities/PolyEditInsertEdgeLoopActivity.h b/Engine/Plugins/Experimental/MeshModelingToolset/Source/MeshModelingTools/Public/ToolActivities/PolyEditInsertEdgeLoopActivity.h
--- a/Engine/Plugins/Experimental/MeshModelingToolset/Source/MeshModelingTools/Public/ToolActivities/PolyEditInsertEdgeLoopActivity.h
+++ b/Engine/Plugins/Experimental/MeshModelingToolset/Source/MeshModelingTools/Public/ToolActivities/PolyEditInsertEdgeLoopActivity.h
@@ -57,7 +57,7 @@ public:
 
 	/** Determines how edge loops are added to the geometry */
 	UPROPERTY(EditAnywhere, Category = InsertEdgeLoop)
-	EEdgeLoopInsertionMode InsertionMode = EEdgeLoopInsertionMode::Retriangulate;
+	EEdgeLoopInsertionMode InsertionMode = EEdgeLoopInsertionMode::PlaneCut;
 
 	/** How many loops to insert at a time. Only used with "even" positioning mode. */
 	UPROPERTY(EditAnywhere, Category = InsertEdgeLoop, meta = (UIMin = "0", UIMax = "20", ClampMin = "0", ClampMax = "500",
@@ -87,7 +87,7 @@ public:
 	bool bHighlightProblemGroups = true;
 
 	/** How close a new loop edge needs to pass next to an existing vertex to use that vertex rather than creating a new one. */
-	UPROPERTY(EditAnywhere, Category = InsertEdgeLoop, AdvancedDisplay)
+	UPROPERTY(EditAnywhere, Category = InsertEdgeLoop, AdvancedDisplay, meta = (UIMin = "0", UIMax = "0.01", ClampMin = "0", ClampMax = "10"))
 	double VertexTolerance = 0.001;
 };
 
diff --git a/Engine/Plugins/Experimental/MeshModelingToolset/Source/MeshModelingTools/Public/ToolActivities/PolyEditInsetOutsetActivity.h b/Engine/Plugins/Experimental/MeshModelingToolset/Source/MeshModelingTools/Public/ToolActivities/PolyEditInsetOutsetActivity.h
--- a/Engine/Plugins/Experimental/MeshModelingToolset/Source/MeshModelingTools/Public/ToolActivities/PolyEditInsetOutsetActivity.h
+++ b/Engine/Plugins/Experimental/MeshModelingToolset/Source/MeshModelingTools/Public/ToolActivities/PolyEditInsetOutsetActivity.h
@@ -19,11 +19,6 @@ class MESHMODELINGTOOLS_API UPolyEditInsetOutsetProperties : public UInteractive
 	GENERATED_BODY()
 
 public:
-	//~ TODO: We should not have to have the user tell us when it's an inset/offset, we should just do the proper
-	//~ thing based on where the user clicks.
-	UPROPERTY(EditAnywhere, Category = "Inset/Outset")
-	bool bOutset = false;
-
 	/** Amount of smoothing applied to outset boundary */
 	UPROPERTY(EditAnywhere, Category = "Inset/Outset", 
 		meta = (UIMin = "0.0", UIMax = "1.0", EditCondition = "bBoundaryOnly == false"))
@@ -39,10 +34,16 @@ public:
 	float AreaScale = true;
 
 	/** When insetting, determines whether vertices in inset region should be projected back onto input surface */
-	UPROPERTY(EditAnywhere, Category = "Inset/Outset")
+	UPROPERTY(EditAnywhere, Category = "Inset/Outset", Meta = (EditCondition = "bOutset", HideEditConditionToggle))
 	bool bReproject = true;
 
-
+	//~ This is not user editable- it gets set by PolyEdit depending on whether the user clicks
+	//~ inset or outset. Currently, both operations share the same code, and one may argue that
+	//~ we should just determine which to do based on where the user clicks. However, our long
+	//~ term plan is that they will be more differentiated in operation, to the point that we
+	//~ may split them into separate activities.
+	UPROPERTY()
+	bool bOutset = false;
 };
 
 
@@ -78,7 +79,7 @@ public:
 	virtual void OnEndHover() override {}
 
 	UPROPERTY()
-	UPolyEditInsetOutsetProperties* InsetProperties;
+	UPolyEditInsetOutsetProperties* Settings;
 
 protected:
 	void Clear();
