MeshModelingTools: Making PolyEdit and PolyDeform respect the editor snapping options.

#rnx

diff --git a/Engine/Plugins/Experimental/MeshModelingToolset/Source/MeshModelingTools/Private/DeformMeshPolygonsTool.cpp b/Engine/Plugins/Experimental/MeshModelingToolset/Source/MeshModelingTools/Private/DeformMeshPolygonsTool.cpp
--- a/Engine/Plugins/Experimental/MeshModelingToolset/Source/MeshModelingTools/Private/DeformMeshPolygonsTool.cpp
+++ b/Engine/Plugins/Experimental/MeshModelingToolset/Source/MeshModelingTools/Private/DeformMeshPolygonsTool.cpp
@@ -300,7 +300,6 @@ UDeformMeshPolygonsTransformProperties::UDeformMeshPolygonsTransformProperties()
 	bSelectFaces        = true;
 	bSelectEdges        = true;
 	bShowWireframe      = false;
-	bSnapToWorldGrid    = false;
 }
 
 /*
@@ -1472,9 +1471,9 @@ void UDeformMeshPolygonsTool::ComputeUpdate_Translate()
 	FGroupTopologyDeformer& SelectedDeformer = (bIsLaplacian) ? *LaplacianDeformer : LinearDeformer;
 
 	TFunction<FVector3d(const FVector3d&)> PointConstraintFunc = nullptr;
-	if (TransformProps->bSnapToWorldGrid &&
-	    GetToolManager()->GetContextQueriesAPI()->GetCurrentCoordinateSystem() == EToolContextCoordinateSystem::World)
+	if (GetToolManager()->GetContextQueriesAPI()->GetCurrentCoordinateSystem() == EToolContextCoordinateSystem::World)
 	{
+		// We currently don't support grid snapping in local mode
 		PointConstraintFunc = [&](const FVector3d& Pos) {
 			FVector3d GridSnapPos;
 			return ToolSceneQueriesUtil::FindWorldGridSnapPoint(this, Pos, GridSnapPos) ? GridSnapPos : Pos;
diff --git a/Engine/Plugins/Experimental/MeshModelingToolset/Source/MeshModelingTools/Private/EditMeshPolygonsTool.cpp b/Engine/Plugins/Experimental/MeshModelingToolset/Source/MeshModelingTools/Private/EditMeshPolygonsTool.cpp
--- a/Engine/Plugins/Experimental/MeshModelingToolset/Source/MeshModelingTools/Private/EditMeshPolygonsTool.cpp
+++ b/Engine/Plugins/Experimental/MeshModelingToolset/Source/MeshModelingTools/Private/EditMeshPolygonsTool.cpp
@@ -838,9 +838,6 @@ void UEditMeshPolygonsTool::OnTick(float DeltaTime)
 		NotifyOfPropertyChangeByTool(CommonProps);
 	}
 
-	// TODO: We should probably allow gizmo snapping in local mode as well. It requires some gizmo changes.
-	TransformGizmo->bSnapToWorldGrid = CommonProps->bSnapToWorldGrid && !bLocalCoordSystem;
-
 	if (bGizmoUpdatePending)
 	{
 		ComputeUpdate_Gizmo();
diff --git a/Engine/Plugins/Experimental/MeshModelingToolset/Source/MeshModelingTools/Private/ToolActivities/PolyEditExtrudeActivity.cpp b/Engine/Plugins/Experimental/MeshModelingToolset/Source/MeshModelingTools/Private/ToolActivities/PolyEditExtrudeActivity.cpp
--- a/Engine/Plugins/Experimental/MeshModelingToolset/Source/MeshModelingTools/Private/ToolActivities/PolyEditExtrudeActivity.cpp
+++ b/Engine/Plugins/Experimental/MeshModelingToolset/Source/MeshModelingTools/Private/ToolActivities/PolyEditExtrudeActivity.cpp
@@ -151,7 +151,7 @@ void UPolyEditExtrudeActivity::BeginExtrude()
 	};
 	ExtrudeHeightMechanic->WorldPointSnapFunc = [this](const FVector3d& WorldPos, FVector3d& SnapPos)
 	{
-		return ActivityContext->CommonProperties->bSnapToWorldGrid && ToolSceneQueriesUtil::FindWorldGridSnapPoint(ParentTool.Get(), WorldPos, SnapPos);
+		return ToolSceneQueriesUtil::FindWorldGridSnapPoint(ParentTool.Get(), WorldPos, SnapPos);
 	};
 	ExtrudeHeightMechanic->CurrentHeight = 1.0f;  // initialize to something non-zero...prob should be based on polygon bounds maybe?
 
diff --git a/Engine/Plugins/Experimental/MeshModelingToolset/Source/MeshModelingTools/Private/ToolActivities/PolyEditInsetOutsetActivity.cpp b/Engine/Plugins/Experimental/MeshModelingToolset/Source/MeshModelingTools/Private/ToolActivities/PolyEditInsetOutsetActivity.cpp
--- a/Engine/Plugins/Experimental/MeshModelingToolset/Source/MeshModelingTools/Private/ToolActivities/PolyEditInsetOutsetActivity.cpp
+++ b/Engine/Plugins/Experimental/MeshModelingToolset/Source/MeshModelingTools/Private/ToolActivities/PolyEditInsetOutsetActivity.cpp
@@ -131,8 +131,7 @@ void UPolyEditInsetOutsetActivity::BeginInset()
 	CurveDistMechanic->Setup(ParentTool.Get());
 	CurveDistMechanic->WorldPointSnapFunc = [this](const FVector3d& WorldPos, FVector3d& SnapPos)
 	{
-		return ActivityContext->CommonProperties->bSnapToWorldGrid
-			&& ToolSceneQueriesUtil::FindWorldGridSnapPoint(ParentTool.Get(), WorldPos, SnapPos);
+		return ToolSceneQueriesUtil::FindWorldGridSnapPoint(ParentTool.Get(), WorldPos, SnapPos);
 	};
 	CurveDistMechanic->CurrentDistance = 1.0f;  // initialize to something non-zero...prob should be based on polygon bounds maybe?
 
diff --git a/Engine/Plugins/Experimental/MeshModelingToolset/Source/MeshModelingTools/Public/DeformMeshPolygonsTool.h b/Engine/Plugins/Experimental/MeshModelingToolset/Source/MeshModelingTools/Public/DeformMeshPolygonsTool.h
--- a/Engine/Plugins/Experimental/MeshModelingToolset/Source/MeshModelingTools/Public/DeformMeshPolygonsTool.h
+++ b/Engine/Plugins/Experimental/MeshModelingToolset/Source/MeshModelingTools/Public/DeformMeshPolygonsTool.h
@@ -109,9 +109,6 @@ public:
 	UPROPERTY(EditAnywhere, Category = SelectionFilter)
 	bool bSelectVertices;
 
-	UPROPERTY(EditAnywhere, Category = Snapping)
-	bool bSnapToWorldGrid;
-
 	UPROPERTY(EditAnywhere, Category = Display)
 	bool bShowWireframe;
 
diff --git a/Engine/Plugins/Experimental/MeshModelingToolset/Source/MeshModelingTools/Public/DrawPolygonTool.h b/Engine/Plugins/Experimental/MeshModelingToolset/Source/MeshModelingTools/Public/DrawPolygonTool.h
--- a/Engine/Plugins/Experimental/MeshModelingToolset/Source/MeshModelingTools/Public/DrawPolygonTool.h
+++ b/Engine/Plugins/Experimental/MeshModelingToolset/Source/MeshModelingTools/Public/DrawPolygonTool.h
@@ -124,6 +124,7 @@ public:
 	UPROPERTY(EditAnywhere, NonTransactional, Category = Snapping)
 	bool bEnableSnapping = true;
 
+	/** When true, allows snapping to world grid according to editor settings. */
 	UPROPERTY(EditAnywhere, NonTransactional, Category = Snapping, Meta = (EditCondition = "bEnableSnapping", EditConditionHides))
 	bool bSnapToWorldGrid = false;
 
diff --git a/Engine/Plugins/Experimental/MeshModelingToolset/Source/MeshModelingTools/Public/EditMeshPolygonsTool.h b/Engine/Plugins/Experimental/MeshModelingToolset/Source/MeshModelingTools/Public/EditMeshPolygonsTool.h
--- a/Engine/Plugins/Experimental/MeshModelingToolset/Source/MeshModelingTools/Public/EditMeshPolygonsTool.h
+++ b/Engine/Plugins/Experimental/MeshModelingToolset/Source/MeshModelingTools/Public/EditMeshPolygonsTool.h
@@ -83,10 +83,6 @@ public:
 	UPROPERTY(EditAnywhere, Category = Gizmo, meta = (HideEditConditionToggle, EditCondition = "bLocalCoordSystem"))
 	bool bLockRotation = false;
 
-	/** Snaps to grid in translation. Only active with a global coordinate system. */
-	UPROPERTY(EditAnywhere, Category = Gizmo, meta = (HideEditConditionToggle, EditCondition = "!bLocalCoordSystem"))
-	bool bSnapToWorldGrid = false;
-
 	/** This gets updated internally so that properties can respond to whether the coordinate system is set to local or global */
 	UPROPERTY()
 	bool bLocalCoordSystem = true;
