MeshModelingToolset: Split plugin into non-experimental and experimental

diff --git a/Engine/Plugins/Runtime/MeshModelingToolset/Source/MeshModelingTools/Public/AddPrimitiveTool.h b/Engine/Plugins/Runtime/MeshModelingToolset/Source/MeshModelingTools/Public/AddPrimitiveTool.h
new file mode 100644
--- /dev/null
+++ b/Engine/Plugins/Runtime/MeshModelingToolset/Source/MeshModelingTools/Public/AddPrimitiveTool.h
@@ -0,0 +1,618 @@
+// Copyright Epic Games, Inc. All Rights Reserved.
+
+#pragma once
+
+#include "BaseTools/SingleClickTool.h"
+#include "CoreMinimal.h"
+#include "InteractiveToolBuilder.h"
+#include "MaterialShared.h"
+#include "PreviewMesh.h"
+#include "Properties/MeshMaterialProperties.h"
+#include "PropertySets/CreateMeshObjectTypeProperties.h"
+#include "UObject/NoExportTypes.h"
+#include "UObject/UObjectGlobals.h"
+
+#include "AddPrimitiveTool.generated.h"
+
+PREDECLARE_USE_GEOMETRY_CLASS(FDynamicMesh3);
+
+/**
+ * Builder
+ */
+UCLASS()
+class MESHMODELINGTOOLS_API UAddPrimitiveToolBuilder : public UInteractiveToolBuilder
+{
+	GENERATED_BODY()
+
+public:
+	enum class EMakeMeshShapeType : uint32
+	{
+		Box,
+		Cylinder,
+		Cone,
+		Arrow,
+		Rectangle,
+		Disc,
+		Torus,
+		Sphere,
+		Stairs
+	};
+
+	virtual bool CanBuildTool(const FToolBuilderState& SceneState) const override;
+	virtual UInteractiveTool* BuildTool(const FToolBuilderState& SceneState) const override;
+
+	EMakeMeshShapeType ShapeType{EMakeMeshShapeType::Box};
+};
+
+/** Placement Target Types */
+UENUM()
+enum class EMakeMeshPlacementType : uint8
+{
+	GroundPlane = 0,
+	OnScene     = 1
+};
+
+/** Placement Pivot Location */
+UENUM()
+enum class EMakeMeshPivotLocation : uint8
+{
+	Base,
+	Centered,
+	Top
+};
+
+/** Polygroup mode for primitive */
+UENUM()
+enum class EMakeMeshPolygroupMode : uint8
+{
+	/** One polygroup for entire output mesh */
+	Single,
+	/** One polygroup per geometric face of primitive */
+	PerFace,
+	/** One polygroup per mesh quad/triangle */
+	PerQuad
+};
+
+UCLASS()
+class MESHMODELINGTOOLS_API UProceduralShapeToolProperties : public UInteractiveToolPropertySet
+{
+	GENERATED_BODY()
+
+public:
+	UProceduralShapeToolProperties() = default;
+
+	/** If the shape settings haven't changed, create instances of the last created asset rather than creating a whole new asset.  If false, all created actors will have separate underlying mesh assets. */
+	UPROPERTY(EditAnywhere, Category = AssetSettings)
+	bool bInstanceIfPossible = false;
+
+	/** How should Polygroups be assigned to triangles of Primitive */
+	UPROPERTY(EditAnywhere, Category = ShapeSettings, meta = (ProceduralShapeSetting))
+	EMakeMeshPolygroupMode PolygroupMode = EMakeMeshPolygroupMode::PerFace;
+
+	/** How to place Primitive in the Scene */
+	UPROPERTY(EditAnywhere, Category = Positioning, meta = (DisplayName = "Target Surface"))
+	EMakeMeshPlacementType PlaceMode = EMakeMeshPlacementType::OnScene;
+
+	/** If true, placement location will be snapped to grid. Only relevant when coordinate system is set to World. */
+	UPROPERTY(EditAnywhere, Category = Positioning)
+	bool bSnapToGrid = true;
+
+	/** Location of Pivot within Primitive */
+	UPROPERTY(EditAnywhere, Category = Positioning, meta = (ProceduralShapeSetting))
+	EMakeMeshPivotLocation PivotLocation = EMakeMeshPivotLocation::Base;
+
+	/** Rotation of Primitive around up axis */
+	UPROPERTY(EditAnywhere, Category = Positioning, meta = (DisplayName = "Rotation", UIMin = "0.0", UIMax = "360.0"))
+	float Rotation = 0.0;
+
+	/** Align shape to Placement Surface */
+	UPROPERTY(EditAnywhere, Category = Positioning, meta = (DisplayName = "Align to Normal", EditCondition = "PlaceMode == EMakeMeshPlacementType::OnScene"))
+	bool bAlignShapeToPlacementSurface = true;
+
+	///** Start Angle of Shape */
+	//UPROPERTY(EditAnywhere, Category = ShapeSettings, meta = (DisplayName = "Start Angle", UIMin = "0.0", UIMax = "360.0", ClampMin = "-10000", ClampMax = "10000.0"))
+	//float StartAngle = 0.f;
+
+	///** End Angle of Shape */
+	//UPROPERTY(EditAnywhere, Category = ShapeSettings, meta = (DisplayName = "End Angle", UIMin = "0.0", UIMax = "360.0", ClampMin = "-10000", ClampMax = "10000.0"))
+	//float EndAngle = 360.f;
+
+	bool IsEquivalent( const UProceduralShapeToolProperties* ) const;
+};
+
+UCLASS()
+class MESHMODELINGTOOLS_API UProceduralBoxToolProperties : public UProceduralShapeToolProperties
+{
+	GENERATED_BODY()
+public:
+	UProceduralBoxToolProperties() = default;
+
+	/** Width of Shape */
+	UPROPERTY(EditAnywhere, Category = "ShapeSettings", meta = (DisplayName = "Width", UIMin = "1.0", UIMax = "1000.0", ClampMin = "0.0001", ClampMax = "1000000.0", ProceduralShapeSetting))
+	float Width = 100.f;
+
+	/** Depth of Shape */
+	UPROPERTY(EditAnywhere, Category = "ShapeSettings", meta = (DisplayName = "Depth", UIMin = "1.0", UIMax = "1000.0", ClampMin = "0.0001", ClampMax = "1000000.0", ProceduralShapeSetting))
+	float Depth = 100.f;
+
+	/** Height of Shape */
+	UPROPERTY(EditAnywhere, Category = "ShapeSettings", meta = (DisplayName = "Height", UIMin = "1.0", UIMax = "1000.0", ClampMin = "0.0001", ClampMax = "1000000.0", ProceduralShapeSetting))
+	float Height = 100.f;
+
+	/** Number of Subdivisions Along the Width */
+	UPROPERTY(EditAnywhere, Category = "ShapeSettings", meta = (DisplayName = "Width Subdivisions", UIMin = "1", UIMax = "100", ClampMin = "1", ClampMax = "500", ProceduralShapeSetting))
+	int WidthSubdivisions = 1;
+
+	/** Number of Subdivisions Along the Depth */
+	UPROPERTY(EditAnywhere, Category = "ShapeSettings", meta = (DisplayName = "Depth Subdivisions", UIMin = "1", UIMax = "100", ClampMin = "1", ClampMax = "500", ProceduralShapeSetting))
+	int DepthSubdivisions = 1;
+
+	/** Number of Subdivisions Along the Height */
+	UPROPERTY(EditAnywhere, Category = "ShapeSettings", meta = (DisplayName = "Height Subdivisions", UIMin = "1", UIMax = "100", ClampMin = "1", ClampMax = "500", ProceduralShapeSetting))
+	int HeightSubdivisions = 1;
+};
+
+UENUM()
+enum class EProceduralRectType
+{
+	/** Create a Rectangle */
+	Rectangle UMETA(DisplayName = "Rectangle"),
+	/** Create a Rounded Rectangle */
+	RoundedRectangle UMETA(DisplayName = "Rounded Rectangle")
+};
+
+UCLASS()
+class MESHMODELINGTOOLS_API UProceduralRectangleToolProperties : public UProceduralShapeToolProperties
+{
+	GENERATED_BODY()
+public:
+	UProceduralRectangleToolProperties() = default;
+
+	/** Type of rectangle to create */
+	UPROPERTY(EditAnywhere, Category = "ShapeSettings")
+	EProceduralRectType RectType = EProceduralRectType::Rectangle;
+
+	/** Width of Shape */
+	UPROPERTY(EditAnywhere, Category = "ShapeSettings", meta = (DisplayName = "Width", UIMin = "1.0", UIMax = "1000.0", ClampMin = "0.0001", ClampMax = "1000000.0", ProceduralShapeSetting))
+	float Width = 100.f;
+
+	/** Depth of Shape */
+	UPROPERTY(EditAnywhere, Category = "ShapeSettings", meta = (DisplayName = "Depth", UIMin = "1.0", UIMax = "1000.0", ClampMin = "0.0001", ClampMax = "1000000.0", ProceduralShapeSetting))
+	float Depth = 100.f;
+
+	/** Number of Subdivisions Along the Width */
+	UPROPERTY(EditAnywhere, Category = "ShapeSettings", meta = (DisplayName = "Width Subdivisions", UIMin = "1", UIMax = "100", ClampMin = "1", ClampMax = "500", ProceduralShapeSetting))
+	int WidthSubdivisions = 1;
+
+	/** Number of Subdivisions Along the Depth */
+	UPROPERTY(EditAnywhere, Category = "ShapeSettings", meta = (DisplayName = "Depth Subdivisions", UIMin = "1", UIMax = "100", ClampMin = "1", ClampMax = "500", ProceduralShapeSetting))
+	int DepthSubdivisions = 1;
+
+	/** Radius of Rounded Corners */
+	UPROPERTY(EditAnywhere, Category = "ShapeSettings", meta = (DisplayName = "Corner Radius", UIMin = "1.0", UIMax = "1000.0", ClampMin = "0.0001", ClampMax = "1000000.0", ProceduralShapeSetting, EditCondition = "RectType == EProceduralRectType::RoundedRectangle", EditConditionHides))
+	float CornerRadius = 25.f;
+
+	/** Number of Angular Slices in Each Rounded Corner */
+	UPROPERTY(EditAnywhere, Category = "ShapeSettings", meta = (DisplayName = "Corner Slices", UIMin = "3", UIMax = "128", ClampMin = "3", ClampMax = "500", ProceduralShapeSetting, EditCondition = "RectType == EProceduralRectType::RoundedRectangle", EditConditionHides))
+	int CornerSlices = 16;
+};
+
+UENUM()
+enum class EProceduralDiscType
+{
+	/** Create a solid Disc */
+	Disc UMETA(DisplayName = "Disc"),
+	/** Create a Disc with a hole */
+	PuncturedDisc UMETA(DisplayName = "Punctured Disc")
+};
+
+UCLASS()
+class MESHMODELINGTOOLS_API UProceduralDiscToolProperties : public UProceduralShapeToolProperties
+{
+	GENERATED_BODY()
+public:
+	/** Type of disc to create */
+	UPROPERTY(EditAnywhere, Category = "ShapeSettings")
+	EProceduralDiscType DiscType = EProceduralDiscType::Disc;
+
+	/** Radius of Disc */
+	UPROPERTY(EditAnywhere, Category = "ShapeSettings", meta = (DisplayName = "Radius", UIMin = "1.0", UIMax = "1000.0", ClampMin = "0.0001", ClampMax = "1000000.0", ProceduralShapeSetting))
+	float Radius = 50.f;
+
+	/** Number of Angular Slices around the Disc */
+	UPROPERTY(EditAnywhere, Category = "ShapeSettings", meta = (DisplayName = "Radial Slices", UIMin = "3", UIMax = "128", ClampMin = "3", ClampMax = "500", ProceduralShapeSetting))
+	int RadialSlices = 16;
+
+	/** Number of Radial Subdivisions in the Disc */
+	UPROPERTY(EditAnywhere, Category = "ShapeSettings", meta = (DisplayName = "Radial Subdivisions", UIMin = "1", UIMax = "100", ClampMin = "1", ClampMax = "500", ProceduralShapeSetting))
+	int RadialSubdivisions = 1;
+
+	/** Radius of the Disc's Hole */
+	UPROPERTY(EditAnywhere, Category = "ShapeSettings", meta = (DisplayName = "Hole Radius", UIMin = "1.0", UIMax = "1000.0", ClampMin = "0.0001", ClampMax = "1000000.0", ProceduralShapeSetting, EditCondition = "DiscType == EProceduralDiscType::PuncturedDisc", EditConditionHides))
+	float HoleRadius = 25.f;
+};
+
+UCLASS()
+class MESHMODELINGTOOLS_API UProceduralTorusToolProperties : public UProceduralShapeToolProperties
+{
+	GENERATED_BODY()
+public:
+	/** Radius from the Torus Center to the Center of the Torus Tube */
+	UPROPERTY(EditAnywhere, Category = "ShapeSettings", meta = (DisplayName = "Major Radius", UIMin = "1.0", UIMax = "1000.0", ClampMin = "0.0001", ClampMax = "1000000.0", ProceduralShapeSetting))
+	float MajorRadius = 50.f;
+
+	/** Radius of the Torus Tube */
+	UPROPERTY(EditAnywhere, Category = "ShapeSettings", meta = (DisplayName = "Minor Radius", UIMin = "1.0", UIMax = "1000.0", ClampMin = "0.0001", ClampMax = "1000000.0", ProceduralShapeSetting))
+	float MinorRadius = 25.f;
+
+	/** Number of Angular Slices Along the Torus Tube */
+	UPROPERTY(EditAnywhere, Category = "ShapeSettings", meta = (DisplayName = "Major Slices", UIMin = "3", UIMax = "128", ClampMin = "3", ClampMax = "500", ProceduralShapeSetting))
+	int TubeSlices = 16;
+
+	/** Number of Angular Slices Around the Tube of the Torus */
+	UPROPERTY(EditAnywhere, Category = "ShapeSettings", meta = (DisplayName = "Minor Slices", UIMin = "3", UIMax = "128", ClampMin = "3", ClampMax = "500", ProceduralShapeSetting))
+	int CrossSectionSlices = 16;
+};
+
+UCLASS()
+class MESHMODELINGTOOLS_API UProceduralCylinderToolProperties : public UProceduralShapeToolProperties
+{
+	GENERATED_BODY()
+public:
+
+	/** Radius of The Cylinder */
+	UPROPERTY(EditAnywhere, Category = "ShapeSettings", meta = (DisplayName = "Radius", UIMin = "1.0", UIMax = "1000.0", ClampMin = "0.0001", ClampMax = "1000000.0", ProceduralShapeSetting))
+	float Radius = 50.f;
+
+	/** Height of Cylinder */
+	UPROPERTY(EditAnywhere, Category = "ShapeSettings", meta = (DisplayName = "Height", UIMin = "1.0", UIMax = "1000.0", ClampMin = "0.0001", ClampMax = "1000000.0", ProceduralShapeSetting))
+	float Height = 200.f;
+
+	/** Number of Slices on the Cylinder Caps */
+	UPROPERTY(EditAnywhere, Category = "ShapeSettings", meta = (DisplayName = "Radial Slices", UIMin = "3", UIMax = "128", ClampMin = "3", ClampMax = "500", ProceduralShapeSetting))
+	int RadialSlices = 16;
+
+	/** Number of Vertical Subdivisions Along the Height of the Cylinder */
+	UPROPERTY(EditAnywhere, Category = "ShapeSettings", meta = (DisplayName = "Height Subdivisions", UIMin = "1", UIMax = "100", ClampMin = "1", ClampMax = "500", ProceduralShapeSetting))
+	int HeightSubdivisions = 1;
+};
+
+UCLASS()
+class MESHMODELINGTOOLS_API UProceduralConeToolProperties : public UProceduralShapeToolProperties
+{
+	GENERATED_BODY()
+public:
+	/** Radius of the Cone */
+	UPROPERTY(EditAnywhere, Category = "ShapeSettings", meta = (DisplayName = "Radius", UIMin = "1.0", UIMax = "1000.0", ClampMin = "0.0001", ClampMax = "1000000.0", ProceduralShapeSetting))
+	float Radius = 50.f;
+
+	/** Height of Cone */
+	UPROPERTY(EditAnywhere, Category = "ShapeSettings", meta = (DisplayName = "Height", UIMin = "1.0", UIMax = "1000.0", ClampMin = "0.0001", ClampMax = "1000000.0", ProceduralShapeSetting))
+	float Height = 200.f;
+
+	/** Number of Slices on the Cone Base */
+	UPROPERTY(EditAnywhere, Category = "ShapeSettings", meta = (DisplayName = "Radial Slices", UIMin = "3", UIMax = "128", ClampMin = "3", ClampMax = "500", ProceduralShapeSetting))
+	int RadialSlices = 16;
+
+	/** Number of Vertical Subdivisions Along the Hight of the Cone */
+	UPROPERTY(EditAnywhere, Category = "ShapeSettings", meta = (DisplayName = "Height Subdivisions", UIMin = "1", UIMax = "100", ClampMin = "1", ClampMax = "500", ProceduralShapeSetting))
+	int HeightSubdivisions = 1;
+};
+
+UCLASS()
+class MESHMODELINGTOOLS_API UProceduralArrowToolProperties : public UProceduralShapeToolProperties
+{
+	GENERATED_BODY()
+public:
+	/** Radius of the Arrow Shaft */
+	UPROPERTY(EditAnywhere, Category = "ShapeSettings", meta = (DisplayName = "Shaft Radius", UIMin = "1.0", UIMax = "1000.0", ClampMin = "0.0001", ClampMax = "1000000.0", ProceduralShapeSetting))
+	float ShaftRadius = 20.f;
+
+	/** Height of Arrow Shaft */
+	UPROPERTY(EditAnywhere, Category = "ShapeSettings", meta = (DisplayName = "Shaft Height", UIMin = "1.0", UIMax = "1000.0", ClampMin = "0.0001", ClampMax = "1000000.0", ProceduralShapeSetting))
+	float ShaftHeight = 200.f;
+
+	/** Radius of the Arrow Head */
+	UPROPERTY(EditAnywhere, Category = "ShapeSettings", meta = (DisplayName = "Head Radius", UIMin = "1.0", UIMax = "1000.0", ClampMin = "0.0001", ClampMax = "1000000.0", ProceduralShapeSetting))
+	float HeadRadius = 60.f;
+
+	/** Height of Arrow's Head */
+	UPROPERTY(EditAnywhere, Category = "ShapeSettings", meta = (DisplayName = "Head Height", UIMin = "1.0", UIMax = "1000.0", ClampMin = "0.0001", ClampMax = "1000000.0", ProceduralShapeSetting))
+	float HeadHeight = 120.f;
+
+	/** Number of Angular Slices Around the Arrow */
+	UPROPERTY(EditAnywhere, Category = "ShapeSettings", meta = (DisplayName = "Radial Slices", UIMin = "3", UIMax = "100", ClampMin = "3", ClampMax = "500", ProceduralShapeSetting))
+	int RadialSlices = 16;
+
+	/** Number of Vertical Subdivisions Along in the Arrow */
+	UPROPERTY(EditAnywhere, Category = "ShapeSettings", meta = (DisplayName = "Total Subdivisions", UIMin = "1", UIMax = "100", ClampMin = "1", ClampMax = "500", ProceduralShapeSetting))
+	int TotalSubdivisions = 1;
+};
+
+UENUM()
+enum class EProceduralSphereType
+{
+	/** Create a Sphere with Lat Long parameterization */
+	LatLong UMETA(DisplayName = "Lat Long"),
+	/** Create a Sphere with Box parameterization */
+	Box UMETA(DisplayName = "Box")
+};
+
+UCLASS()
+class MESHMODELINGTOOLS_API UProceduralSphereToolProperties : public UProceduralShapeToolProperties
+{
+	GENERATED_BODY()
+public:
+	/** Type of Sphere to create */
+	UPROPERTY(EditAnywhere, Category = "ShapeSettings")
+	EProceduralSphereType SphereType = EProceduralSphereType::Box;
+
+	/** Radius of the Sphere */
+	UPROPERTY(EditAnywhere, Category = "ShapeSettings", meta = (DisplayName = "Radius", UIMin = "1.0", UIMax = "1000.0", ClampMin = "0.0001", ClampMax = "1000000.0", ProceduralShapeSetting))
+	float Radius = 50.f;
+
+	/** Number of Latitudinal Slices of the Sphere */
+	UPROPERTY(EditAnywhere, Category = "ShapeSettings", meta = (DisplayName = "Latitude Slices", UIMin = "3", UIMax = "100", ClampMin = "4", ClampMax = "500", ProceduralShapeSetting, EditCondition = "SphereType == EProceduralSphereType::LatLong", EditConditionHides))
+	int LatitudeSlices = 16;
+
+	/** Number of Longitudinal Slices around the Sphere */
+	UPROPERTY(EditAnywhere, Category = "ShapeSettings", meta = (DisplayName = "Longitude Slices", UIMin = "3", UIMax = "100", ClampMin = "4", ClampMax = "500", ProceduralShapeSetting, EditCondition = "SphereType == EProceduralSphereType::LatLong", EditConditionHides))
+	int LongitudeSlices = 16;
+
+	/** Number of Subdivisions of each Side of the Sphere */
+	UPROPERTY(EditAnywhere, Category = "ShapeSettings", meta = (DisplayName = "Side Subdivisions", UIMin = "1", UIMax = "100", ClampMin = "1", ClampMax = "500", ProceduralShapeSetting, EditCondition = "SphereType == EProceduralSphereType::Box", EditConditionHides))
+	int Subdivisions = 16;
+};
+
+UENUM()
+enum class EProceduralStairsType
+{
+	/** Create a linear staircase */
+	Linear UMETA(DisplayName = "Linear"),
+	/** Create a floating staircase */
+	Floating UMETA(DisplayName = "Floating"),
+	/** Create a curved staircase */
+	Curved UMETA(DisplayName = "Curved"),
+	/** Create a spiral staircase */
+	Spiral UMETA(DisplayName = "Spiral")
+};
+
+UCLASS()
+class MESHMODELINGTOOLS_API UProceduralStairsToolProperties : public UProceduralShapeToolProperties
+{
+	GENERATED_BODY()
+public:
+	/** Type of staircase to create */
+	UPROPERTY(EditAnywhere, Category = "ShapeSettings")
+		EProceduralStairsType StairsType = EProceduralStairsType::Linear;
+
+	/** Number of Steps of Shape */
+	UPROPERTY(EditAnywhere, Category = "ShapeSettings", meta = (DisplayName = "Number of Steps", UIMin = "1", UIMax = "100", ClampMin = "1", ClampMax = "1000", ProceduralShapeSetting))
+	int NumSteps = 8;
+
+	/** Width of each step */
+	UPROPERTY(EditAnywhere, Category = "ShapeSettings", meta = (DisplayName = "Step Width", UIMin = "1.0", UIMax = "1000.0", ClampMin = "0.0001", ClampMax = "1000000.0", ProceduralShapeSetting))
+	float StepWidth = 150.0f;
+
+	/** Height of each step */
+	UPROPERTY(EditAnywhere, Category = "ShapeSettings", meta = (DisplayName = "Step Height", UIMin = "1.0", UIMax = "1000.0", ClampMin = "0.0001", ClampMax = "1000000.0", ProceduralShapeSetting))
+	float StepHeight = 20.0f;
+
+	/** Depth of each step for linear stair shapes */
+	UPROPERTY(EditAnywhere, Category = "ShapeSettings", meta = (DisplayName = "Step Depth", UIMin = "1.0", UIMax = "1000.0", ClampMin = "0.0001", ClampMax = "1000000.0", ProceduralShapeSetting, EditCondition = "StairsType == EProceduralStairsType::Linear || StairsType == EProceduralStairsType::Floating", EditConditionHides))
+	float StepDepth = 30.0f;
+
+	/** Angular length of curved stair shapes */
+	UPROPERTY(EditAnywhere, Category = "ShapeSettings", meta = (DisplayName = "Curve Angle", UIMin = "-360.0", UIMax = "360.0", ClampMin = "-360.0", ClampMax = "360.0", ProceduralShapeSetting, EditCondition = "StairsType == EProceduralStairsType::Curved", EditConditionHides))
+	float CurveAngle = 90.0f;
+
+	/** Angular length of spiral stair shapes */
+	UPROPERTY(EditAnywhere, Category = "ShapeSettings", meta = (DisplayName = "Spiral Angle", UIMin = "-3600.0", UIMax = "3600.0", ClampMin = "-360000.0", ClampMax = "360000.0", ProceduralShapeSetting, EditCondition = "StairsType == EProceduralStairsType::Spiral", EditConditionHides))
+	float SpiralAngle = 90.0f;
+
+	/** Inner radius for curved/spiral stair shapes */
+	UPROPERTY(EditAnywhere, Category = "ShapeSettings", meta = (DisplayName = "Inner Radius", UIMin = "1.0", UIMax = "1000.0", ClampMin = "0.0001", ClampMax = "1000000.0", ProceduralShapeSetting, EditCondition = "StairsType == EProceduralStairsType::Curved || StairsType == EProceduralStairsType::Spiral", EditConditionHides))
+	float InnerRadius = 150.0f;
+};
+
+
+UCLASS(Transient)
+class MESHMODELINGTOOLS_API ULastActorInfo : public UObject
+{
+	GENERATED_BODY()
+
+public:
+	FString Label = "";
+
+	UPROPERTY()
+	TObjectPtr<AActor> Actor = nullptr;
+
+	UPROPERTY()
+	TObjectPtr<UStaticMesh> StaticMesh = nullptr;
+
+	UPROPERTY()
+	TObjectPtr<UProceduralShapeToolProperties> ShapeSettings;
+
+	UPROPERTY()
+	TObjectPtr<UNewMeshMaterialProperties> MaterialProperties;
+
+	bool IsInvalid()
+	{
+		return Actor == nullptr || StaticMesh == nullptr || ShapeSettings == nullptr || MaterialProperties == nullptr;
+	}
+};
+
+/**
+ * Base tool to create primitives
+ */
+UCLASS()
+class MESHMODELINGTOOLS_API UAddPrimitiveTool : public USingleClickTool, public IHoverBehaviorTarget
+{
+	GENERATED_BODY()
+
+public:
+	UAddPrimitiveTool(const FObjectInitializer&);
+
+	virtual void SetWorld(UWorld* World);
+
+	virtual void Setup() override;
+	virtual void Shutdown(EToolShutdownType ShutdownType) override;
+
+	virtual void Render(IToolsContextRenderAPI* RenderAPI) override;
+
+	virtual bool HasCancel() const override { return false; }
+	virtual bool HasAccept() const override { return false; }
+	virtual bool CanAccept() const override { return false; }
+
+	virtual void OnPropertyModified(UObject* PropertySet, FProperty* Property) override;
+
+	virtual void OnClicked(const FInputDeviceRay& ClickPos) override;
+
+
+	// IHoverBehaviorTarget interface
+	virtual FInputRayHit BeginHoverSequenceHitTest(const FInputDeviceRay& PressPos) override;
+	virtual void OnBeginHover(const FInputDeviceRay& DevicePos) override;
+	virtual bool OnUpdateHover(const FInputDeviceRay& DevicePos) override;
+	virtual void OnEndHover() override;
+
+
+protected:
+	virtual void GenerateMesh(FDynamicMesh3* OutMesh) const {}
+	virtual UProceduralShapeToolProperties* CreateShapeSettings(){return nullptr;}
+
+	/** Property set for type of output object (StaticMesh, Volume, etc) */
+	UPROPERTY()
+	TObjectPtr<UCreateMeshObjectTypeProperties> OutputTypeProperties;
+
+	UPROPERTY()
+	TObjectPtr<UProceduralShapeToolProperties> ShapeSettings;
+
+	UPROPERTY()
+	TObjectPtr<UNewMeshMaterialProperties> MaterialProperties;
+
+
+	/**
+	 * Checks if the passed-in settings would create the same asset as the current settings
+	 */
+	bool IsEquivalentLastGeneratedAsset()
+	{
+		if (LastGenerated == nullptr || LastGenerated->IsInvalid())
+		{
+			return false;
+		}
+		return (LastGenerated->MaterialProperties->UVScale == MaterialProperties->UVScale) &&
+			(LastGenerated->MaterialProperties->bWorldSpaceUVScale == MaterialProperties->bWorldSpaceUVScale) &&
+			ShapeSettings->IsEquivalent(LastGenerated->ShapeSettings);
+	}
+
+
+	UPROPERTY()
+	TObjectPtr<UPreviewMesh> PreviewMesh;
+
+	UPROPERTY()
+	TObjectPtr<ULastActorInfo> LastGenerated;
+
+	UPROPERTY()
+	FString AssetName = TEXT("GeneratedAsset");
+
+protected:
+	UWorld* TargetWorld;
+
+	void UpdatePreviewPosition(const FInputDeviceRay& ClickPos);
+	UE::Geometry::FFrame3f ShapeFrame;
+
+	void UpdatePreviewMesh();
+};
+
+
+UCLASS()
+class UAddBoxPrimitiveTool : public UAddPrimitiveTool
+{
+	GENERATED_BODY()
+public:
+	UAddBoxPrimitiveTool(const FObjectInitializer& ObjectInitializer);
+protected:
+	void GenerateMesh(FDynamicMesh3* OutMesh) const override;
+};
+
+UCLASS()
+class UAddCylinderPrimitiveTool : public UAddPrimitiveTool
+{
+	GENERATED_BODY()
+public:
+	UAddCylinderPrimitiveTool(const FObjectInitializer& ObjectInitializer);
+protected:
+	void GenerateMesh(FDynamicMesh3* OutMesh) const override;
+};
+
+UCLASS()
+class UAddConePrimitiveTool : public UAddPrimitiveTool
+{
+	GENERATED_BODY()
+public:
+	UAddConePrimitiveTool(const FObjectInitializer& ObjectInitializer);
+protected:
+	void GenerateMesh(FDynamicMesh3* OutMesh) const override;
+};
+
+UCLASS()
+class UAddRectanglePrimitiveTool : public UAddPrimitiveTool
+{
+	GENERATED_BODY()
+public:
+	UAddRectanglePrimitiveTool(const FObjectInitializer& ObjectInitializer);
+protected:
+	void GenerateMesh(FDynamicMesh3* OutMesh) const override;
+};
+
+UCLASS()
+class UAddDiscPrimitiveTool : public UAddPrimitiveTool
+{
+	GENERATED_BODY()
+public:
+	UAddDiscPrimitiveTool(const FObjectInitializer& ObjectInitializer);
+protected:
+	void GenerateMesh(FDynamicMesh3* OutMesh) const override;
+};
+
+UCLASS()
+class UAddTorusPrimitiveTool : public UAddPrimitiveTool
+{
+	GENERATED_BODY()
+public:
+	UAddTorusPrimitiveTool(const FObjectInitializer& ObjectInitializer);
+protected:
+	void GenerateMesh(FDynamicMesh3* OutMesh) const override;
+};
+
+UCLASS()
+class UAddArrowPrimitiveTool : public UAddPrimitiveTool
+{
+	GENERATED_BODY()
+public:
+	UAddArrowPrimitiveTool(const FObjectInitializer& ObjectInitializer);
+protected:
+	void GenerateMesh(FDynamicMesh3* OutMesh) const override;
+};
+
+UCLASS()
+class UAddSpherePrimitiveTool : public UAddPrimitiveTool
+{
+	GENERATED_BODY()
+public:
+	UAddSpherePrimitiveTool(const FObjectInitializer& ObjectInitializer);
+protected:
+	void GenerateMesh(FDynamicMesh3* OutMesh) const override;
+};
+
+UCLASS()
+class UAddStairsPrimitiveTool : public UAddPrimitiveTool
+{
+	GENERATED_BODY()
+public:
+	UAddStairsPrimitiveTool(const FObjectInitializer& ObjectInitializer);
+protected:
+	void GenerateMesh(FDynamicMesh3* OutMesh) const override;
+};
+
+
diff --git a/Engine/Source/Developer/GeometryProcessingInterfaces/Public/GeometryProcessingInterfaces/ApproximateActors.h b/Engine/Source/Developer/GeometryProcessingInterfaces/Public/GeometryProcessingInterfaces/ApproximateActors.h
--- a/Engine/Source/Developer/GeometryProcessingInterfaces/Public/GeometryProcessingInterfaces/ApproximateActors.h
+++ b/Engine/Source/Developer/GeometryProcessingInterfaces/Public/GeometryProcessingInterfaces/ApproximateActors.h
@@ -155,7 +155,7 @@ public:
 		//
 
 		// A new MIC derived from this material will be created and assigned to the generated mesh
-		UMaterialInterface* BakeMaterial = nullptr;		// if null, will use /MeshModelingToolset/Materials/FullMaterialBakePreviewMaterial_PackedMRS instead
+		UMaterialInterface* BakeMaterial = nullptr;		// if null, will use /MeshModelingToolsetExp/Materials/FullMaterialBakePreviewMaterial_PackedMRS instead
 		FName BaseColorTexParamName = FName("BaseColor");
 		bool bBakeBaseColor = true;
 		FName RoughnessTexParamName = FName("Roughness");
