MeshModelingToolset: Move selected "Create" tools out of experimental

#rnx

diff --git a/Engine/Plugins/Experimental/MeshModelingToolsetExp/Source/MeshModelingToolsExp/Private/CombineMeshesTool.cpp b/Engine/Plugins/Experimental/MeshModelingToolsetExp/Source/MeshModelingToolsExp/Private/CombineMeshesTool.cpp
deleted file mode 100644
diff --git a/Engine/Plugins/Experimental/MeshModelingToolsetExp/Source/MeshModelingToolsExp/Private/DrawAndRevolveTool.cpp b/Engine/Plugins/Experimental/MeshModelingToolsetExp/Source/MeshModelingToolsExp/Private/DrawAndRevolveTool.cpp
deleted file mode 100644
diff --git a/Engine/Plugins/Experimental/MeshModelingToolsetExp/Source/MeshModelingToolsExp/Private/DrawPolygonTool.cpp b/Engine/Plugins/Experimental/MeshModelingToolsetExp/Source/MeshModelingToolsExp/Private/DrawPolygonTool.cpp
deleted file mode 100644
diff --git a/Engine/Plugins/Experimental/MeshModelingToolsetExp/Source/MeshModelingToolsExp/Public/CombineMeshesTool.h b/Engine/Plugins/Experimental/MeshModelingToolsetExp/Source/MeshModelingToolsExp/Public/CombineMeshesTool.h
deleted file mode 100644
diff --git a/Engine/Plugins/Experimental/MeshModelingToolsetExp/Source/MeshModelingToolsExp/Public/DrawAndRevolveTool.h b/Engine/Plugins/Experimental/MeshModelingToolsetExp/Source/MeshModelingToolsExp/Public/DrawAndRevolveTool.h
deleted file mode 100644
diff --git a/Engine/Plugins/Experimental/MeshModelingToolsetExp/Source/MeshModelingToolsExp/Public/DrawPolygonTool.h b/Engine/Plugins/Experimental/MeshModelingToolsetExp/Source/MeshModelingToolsExp/Public/DrawPolygonTool.h
deleted file mode 100644
diff --git a/Engine/Plugins/Runtime/MeshModelingToolset/Source/MeshModelingTools/MeshModelingTools.Build.cs b/Engine/Plugins/Runtime/MeshModelingToolset/Source/MeshModelingTools/MeshModelingTools.Build.cs
--- a/Engine/Plugins/Runtime/MeshModelingToolset/Source/MeshModelingTools/MeshModelingTools.Build.cs
+++ b/Engine/Plugins/Runtime/MeshModelingToolset/Source/MeshModelingTools/MeshModelingTools.Build.cs
@@ -25,17 +25,18 @@ public class MeshModelingTools : ModuleRules
 			{
 				"Core",
 				// "Eigen",
+				"InputCore",
                 "InteractiveToolsFramework",
 				"GeometryCore",
 				// "GeometryFramework",
-				// "GeometryAlgorithms",
+				"GeometryAlgorithms",
 				"DynamicMesh",
-				// "MeshConversion",
+				"MeshConversion",
 				// "MeshDescription",
 				// "StaticMeshDescription",
 				// "SkeletalMeshDescription",
 				"ModelingComponents",
-				// "ModelingOperators",
+				"ModelingOperators",
 
 				// ... add other public dependencies that you statically link with here ...
 			}
diff --git a/Engine/Plugins/Runtime/MeshModelingToolset/Source/MeshModelingTools/Private/CombineMeshesTool.cpp b/Engine/Plugins/Runtime/MeshModelingToolset/Source/MeshModelingTools/Private/CombineMeshesTool.cpp
new file mode 100644
--- /dev/null
+++ b/Engine/Plugins/Runtime/MeshModelingToolset/Source/MeshModelingTools/Private/CombineMeshesTool.cpp
@@ -0,0 +1,529 @@
+// Copyright Epic Games, Inc. All Rights Reserved.
+
+#include "CombineMeshesTool.h"
+#include "InteractiveToolManager.h"
+#include "ToolBuilderUtil.h"
+#include "ToolSetupUtil.h"
+
+#include "DynamicMesh/DynamicMesh3.h"
+#include "DynamicMeshEditor.h"
+#include "DynamicMesh/MeshTransforms.h"
+
+#include "ModelingObjectsCreationAPI.h"
+#include "Selection/ToolSelectionUtil.h"
+#include "Physics/ComponentCollisionUtil.h"
+#include "ShapeApproximation/SimpleShapeSet3.h"
+
+#include "Components/StaticMeshComponent.h"
+#include "Engine/StaticMesh.h"
+
+#include "TargetInterfaces/MaterialProvider.h"
+#include "TargetInterfaces/MeshDescriptionCommitter.h"
+#include "TargetInterfaces/MeshDescriptionProvider.h"
+#include "TargetInterfaces/PrimitiveComponentBackedTarget.h"
+#include "ToolTargetManager.h"
+#include "ModelingToolTargetUtil.h"
+
+#if WITH_EDITOR
+#include "Misc/ScopedSlowTask.h"
+#endif
+
+#include "ExplicitUseGeometryMathTypes.h"		// using UE::Geometry::(math types)
+using namespace UE::Geometry;
+
+#define LOCTEXT_NAMESPACE "UCombineMeshesTool"
+
+/*
+ * ToolBuilder
+ */
+
+
+const FToolTargetTypeRequirements& UCombineMeshesToolBuilder::GetTargetRequirements() const
+{
+	static FToolTargetTypeRequirements TypeRequirements({
+		UMeshDescriptionCommitter::StaticClass(),
+		UMeshDescriptionProvider::StaticClass(),
+		UPrimitiveComponentBackedTarget::StaticClass(),
+		UMaterialProvider::StaticClass()
+		});
+	return TypeRequirements;
+}
+
+bool UCombineMeshesToolBuilder::CanBuildTool(const FToolBuilderState& SceneState) const
+{
+	return (bIsDuplicateTool) ?
+		  (SceneState.TargetManager->CountSelectedAndTargetable(SceneState, GetTargetRequirements()) == 1)
+		: (SceneState.TargetManager->CountSelectedAndTargetable(SceneState, GetTargetRequirements()) > 1);
+}
+
+UInteractiveTool* UCombineMeshesToolBuilder::BuildTool(const FToolBuilderState& SceneState) const
+{
+	UCombineMeshesTool* NewTool = NewObject<UCombineMeshesTool>(SceneState.ToolManager);
+
+	TArray<TObjectPtr<UToolTarget>> Targets = SceneState.TargetManager->BuildAllSelectedTargetable(SceneState, GetTargetRequirements());
+	NewTool->SetTargets(MoveTemp(Targets));
+	NewTool->SetWorld(SceneState.World);
+	NewTool->SetDuplicateMode(bIsDuplicateTool);
+
+	return NewTool;
+}
+
+
+
+
+/*
+ * Tool
+ */
+
+
+void UCombineMeshesTool::SetWorld(UWorld* World)
+{
+	this->TargetWorld = World;
+}
+
+void UCombineMeshesTool::SetDuplicateMode(bool bDuplicateModeIn)
+{
+	this->bDuplicateMode = bDuplicateModeIn;
+}
+
+void UCombineMeshesTool::Setup()
+{
+	UInteractiveTool::Setup();
+
+	BasicProperties = NewObject<UCombineMeshesToolProperties>(this);
+	AddToolPropertySource(BasicProperties);
+	BasicProperties->RestoreProperties(this);
+	BasicProperties->bIsDuplicateMode = this->bDuplicateMode;
+
+	OutputTypeProperties = NewObject<UCreateMeshObjectTypeProperties>(this);
+	OutputTypeProperties->InitializeDefaultWithAuto();
+	OutputTypeProperties->OutputType = UCreateMeshObjectTypeProperties::AutoIdentifier;
+	OutputTypeProperties->RestoreProperties(this, TEXT("OutputTypeFromInputTool"));
+	OutputTypeProperties->WatchProperty(OutputTypeProperties->OutputType, [this](FString) { OutputTypeProperties->UpdatePropertyVisibility(); });
+	AddToolPropertySource(OutputTypeProperties);
+
+	BasicProperties->WatchProperty(BasicProperties->WriteOutputTo, [&](ECombineTargetType NewType)
+	{
+		if (NewType == ECombineTargetType::NewAsset)
+		{
+			BasicProperties->OutputAsset = TEXT("");
+			SetToolPropertySourceEnabled(OutputTypeProperties, true);
+		}
+		else
+		{
+			int32 Index = (BasicProperties->WriteOutputTo == ECombineTargetType::FirstInputAsset) ? 0 : Targets.Num() - 1;
+			BasicProperties->OutputAsset = UE::Modeling::GetComponentAssetBaseName(UE::ToolTarget::GetTargetComponent(Targets[Index]), false);
+			SetToolPropertySourceEnabled(OutputTypeProperties, false);
+		}
+	});
+
+	SetToolPropertySourceEnabled(OutputTypeProperties, BasicProperties->WriteOutputTo == ECombineTargetType::NewAsset);
+
+	if (bDuplicateMode)
+	{
+		SetToolDisplayName(LOCTEXT("DuplicateMeshesToolName", "Duplicate"));
+		BasicProperties->OutputName = UE::Modeling::GetComponentAssetBaseName(UE::ToolTarget::GetTargetComponent(Targets[0]));
+	}
+	else
+	{
+		SetToolDisplayName(LOCTEXT("CombineMeshesToolName", "Append"));
+		BasicProperties->OutputName = FString("Combined");
+	}
+
+
+	HandleSourceProperties = NewObject<UOnAcceptHandleSourcesProperties>(this);
+	AddToolPropertySource(HandleSourceProperties);
+	HandleSourceProperties->RestoreProperties(this);
+
+
+	if (bDuplicateMode)
+	{
+		GetToolManager()->DisplayMessage(
+			LOCTEXT("OnStartToolDuplicate", "This Tool duplicates input Asset into a new Asset, and optionally replaces the input Actor with a new Actor containing the new Asset."),
+			EToolMessageLevel::UserNotification);
+	}
+	else
+	{
+		GetToolManager()->DisplayMessage(
+			LOCTEXT("OnStartToolCombine", "This Tool appends the meshes from the input Assets into a new Asset, and optionally replaces the source Actors with a new Actor containing the new Asset."),
+			EToolMessageLevel::UserNotification);
+	}
+}
+
+
+
+void UCombineMeshesTool::Shutdown(EToolShutdownType ShutdownType)
+{
+	BasicProperties->SaveProperties(this);
+	OutputTypeProperties->SaveProperties(this, TEXT("OutputTypeFromInputTool"));
+	HandleSourceProperties->SaveProperties(this);
+
+	if (ShutdownType == EToolShutdownType::Accept)
+	{
+		if (bDuplicateMode || BasicProperties->WriteOutputTo == ECombineTargetType::NewAsset)
+		{
+			CreateNewAsset();
+		}
+		else
+		{
+			UpdateExistingAsset();
+		}
+	}
+}
+
+void UCombineMeshesTool::CreateNewAsset()
+{
+	// Make sure meshes are available before we open transaction. This is to avoid potential stability issues related 
+	// to creation/load of meshes inside a transaction, for assets that possibly do not have bulk data currently loaded.
+	TArray<FDynamicMesh3> InputMeshes;
+	InputMeshes.Reserve(Targets.Num());
+	for (int32 ComponentIdx = 0; ComponentIdx < Targets.Num(); ComponentIdx++)
+	{
+		InputMeshes.Add(UE::ToolTarget::GetDynamicMeshCopy(Targets[ComponentIdx], true));
+	}
+
+	GetToolManager()->BeginUndoTransaction( bDuplicateMode ? 
+		LOCTEXT("DuplicateMeshToolTransactionName", "Duplicate Mesh") :
+		LOCTEXT("CombineMeshesToolTransactionName", "Combine Meshes"));
+
+	FBox Box(ForceInit);
+	for (int32 ComponentIdx = 0; ComponentIdx < Targets.Num(); ComponentIdx++)
+	{
+		Box += UE::ToolTarget::GetTargetComponent(Targets[ComponentIdx])->Bounds.GetBox();
+	}
+
+	TArray<UMaterialInterface*> AllMaterials;
+	TArray<TArray<int32>> MaterialIDRemaps;
+	BuildCombinedMaterialSet(AllMaterials, MaterialIDRemaps);
+
+	FDynamicMesh3 AccumulateDMesh;
+	AccumulateDMesh.EnableTriangleGroups();
+	AccumulateDMesh.EnableAttributes();
+	AccumulateDMesh.Attributes()->EnableTangents();
+	AccumulateDMesh.Attributes()->EnableMaterialID();
+	AccumulateDMesh.Attributes()->EnablePrimaryColors();
+	FTransform AccumToWorld(Box.GetCenter());
+	FTransform ToAccum(-Box.GetCenter());
+
+	FSimpleShapeSet3d SimpleCollision;
+	UE::Geometry::FComponentCollisionSettings CollisionSettings;
+
+	{
+#if WITH_EDITOR
+		FScopedSlowTask SlowTask(Targets.Num()+1, 
+			bDuplicateMode ? 
+			LOCTEXT("DuplicateMeshBuild", "Building duplicate mesh ...") :
+			LOCTEXT("CombineMeshesBuild", "Building combined mesh ..."));
+		SlowTask.MakeDialog();
+#endif
+		bool bNeedColorAttr = false;
+		int MatIndexBase = 0;
+		for (int32 ComponentIdx = 0; ComponentIdx < Targets.Num(); ComponentIdx++)
+		{
+#if WITH_EDITOR
+			SlowTask.EnterProgressFrame(1);
+#endif
+			UPrimitiveComponent* PrimitiveComponent = UE::ToolTarget::GetTargetComponent(Targets[ComponentIdx]);
+
+			FDynamicMesh3& ComponentDMesh = InputMeshes[ComponentIdx];
+			bNeedColorAttr = bNeedColorAttr || (ComponentDMesh.HasAttributes() && ComponentDMesh.Attributes()->HasPrimaryColors());
+
+			if (ComponentDMesh.HasAttributes() && ComponentDMesh.Attributes()->NumUVLayers() > AccumulateDMesh.Attributes()->NumUVLayers())
+			{
+				AccumulateDMesh.Attributes()->SetNumUVLayers(ComponentDMesh.Attributes()->NumUVLayers());
+			}
+
+			UE::Geometry::FTransform3d XF = (UE::Geometry::FTransform3d)((FTransform)UE::ToolTarget::GetLocalToWorldTransform(Targets[ComponentIdx]) * ToAccum);
+			if (XF.GetDeterminant() < 0)
+			{
+				ComponentDMesh.ReverseOrientation(false);
+			}
+
+			// update material IDs to account for combined material set
+			FDynamicMeshMaterialAttribute* MatAttrib = ComponentDMesh.Attributes()->GetMaterialID();
+			for (int TID : ComponentDMesh.TriangleIndicesItr())
+			{
+				int MatID = MatAttrib->GetValue(TID);
+				MatAttrib->SetValue(TID, MaterialIDRemaps[ComponentIdx][MatID]);
+			}
+
+			FDynamicMeshEditor Editor(&AccumulateDMesh);
+			FMeshIndexMappings IndexMapping;
+			if (bDuplicateMode) // no transform if duplicating
+			{
+				Editor.AppendMesh(&ComponentDMesh, IndexMapping);
+
+				if (UE::Geometry::ComponentTypeSupportsCollision(PrimitiveComponent))
+				{
+					CollisionSettings = UE::Geometry::GetCollisionSettings(PrimitiveComponent);
+					UE::Geometry::AppendSimpleCollision(PrimitiveComponent, &SimpleCollision, UE::Geometry::FTransform3d::Identity());
+				}
+			}
+			else
+			{
+				Editor.AppendMesh(&ComponentDMesh, IndexMapping,
+					[&XF](int Unused, const FVector3d P) { return XF.TransformPosition(P); },
+					[&XF](int Unused, const FVector3d N) { return XF.TransformNormal(N); });
+				if (UE::Geometry::ComponentTypeSupportsCollision(PrimitiveComponent))
+				{
+					UE::Geometry::AppendSimpleCollision(PrimitiveComponent, &SimpleCollision, XF);
+				}
+			}
+
+			FComponentMaterialSet MaterialSet = UE::ToolTarget::GetMaterialSet(Targets[ComponentIdx]);
+			MatIndexBase += MaterialSet.Materials.Num();
+		}
+
+		if (!bNeedColorAttr)
+		{
+			AccumulateDMesh.Attributes()->DisablePrimaryColors();
+		}
+
+#if WITH_EDITOR
+		SlowTask.EnterProgressFrame(1);
+#endif
+
+		if (bDuplicateMode)
+		{
+			// TODO: will need to refactor this when we support duplicating multiple
+			check(Targets.Num() == 1);
+			AccumToWorld = (FTransform)UE::ToolTarget::GetLocalToWorldTransform(Targets[0]);
+		}
+
+		// max len explicitly enforced here, would ideally notify user
+		FString UseBaseName = BasicProperties->OutputName.Left(250);
+		if (UseBaseName.IsEmpty())
+		{
+			UseBaseName = (bDuplicateMode) ? TEXT("Duplicate") : TEXT("Combined");
+		}
+
+		FCreateMeshObjectParams NewMeshObjectParams;
+		NewMeshObjectParams.TargetWorld = TargetWorld;
+		NewMeshObjectParams.Transform = AccumToWorld;
+		NewMeshObjectParams.BaseName = UseBaseName;
+		NewMeshObjectParams.Materials = AllMaterials;
+		NewMeshObjectParams.SetMesh(&AccumulateDMesh);
+		if (OutputTypeProperties->OutputType == UCreateMeshObjectTypeProperties::AutoIdentifier)
+		{
+			UE::ToolTarget::ConfigureCreateMeshObjectParams(Targets[0], NewMeshObjectParams);
+		}
+		else
+		{
+			OutputTypeProperties->ConfigureCreateMeshObjectParams(NewMeshObjectParams);
+		}
+		FCreateMeshObjectResult Result = UE::Modeling::CreateMeshObject(GetToolManager(), MoveTemp(NewMeshObjectParams));
+		if (Result.IsOK() && Result.NewActor != nullptr)
+		{
+			// if any inputs have Simple Collision geometry we will forward it to new mesh.
+			if (UE::Geometry::ComponentTypeSupportsCollision(Result.NewComponent) && SimpleCollision.TotalElementsNum() > 0)
+			{
+				UE::Geometry::SetSimpleCollision(Result.NewComponent, &SimpleCollision, CollisionSettings);
+			}
+
+			// select the new actor
+			ToolSelectionUtil::SetNewActorSelection(GetToolManager(), Result.NewActor);
+		}
+	}
+	
+	TArray<AActor*> Actors;
+	for (int32 Idx = 0; Idx < Targets.Num(); Idx++)
+	{
+		Actors.Add(UE::ToolTarget::GetTargetActor(Targets[Idx]));
+	}
+	HandleSourceProperties->ApplyMethod(Actors, GetToolManager());
+
+
+	GetToolManager()->EndUndoTransaction();
+}
+
+
+
+
+
+
+
+void UCombineMeshesTool::UpdateExistingAsset()
+{
+	// Make sure meshes are available before we open transaction. This is to avoid potential stability issues related 
+	// to creation/load of meshes inside a transaction, for assets that possibly do not have bulk data currently loaded.
+	TArray<FDynamicMesh3> InputMeshes;
+	InputMeshes.Reserve(Targets.Num());
+	for (int32 ComponentIdx = 0; ComponentIdx < Targets.Num(); ComponentIdx++)
+	{
+		InputMeshes.Add(UE::ToolTarget::GetDynamicMeshCopy(Targets[ComponentIdx], true));
+	}
+
+	check(!bDuplicateMode);
+	GetToolManager()->BeginUndoTransaction(LOCTEXT("CombineMeshesToolTransactionName", "Combine Meshes"));
+
+	AActor* SkipActor = nullptr;
+
+	TArray<UMaterialInterface*> AllMaterials;
+	TArray<TArray<int32>> MaterialIDRemaps;
+	BuildCombinedMaterialSet(AllMaterials, MaterialIDRemaps);
+
+	FDynamicMesh3 AccumulateDMesh;
+	AccumulateDMesh.EnableTriangleGroups();
+	AccumulateDMesh.EnableAttributes();
+	AccumulateDMesh.Attributes()->EnableTangents();
+	AccumulateDMesh.Attributes()->EnableMaterialID();
+	AccumulateDMesh.Attributes()->EnablePrimaryColors();
+
+	int32 SkipIndex = (BasicProperties->WriteOutputTo == ECombineTargetType::FirstInputAsset) ? 0 : (Targets.Num() - 1);
+	UPrimitiveComponent* UpdateComponent = UE::ToolTarget::GetTargetComponent(Targets[SkipIndex]);
+	SkipActor = UE::ToolTarget::GetTargetActor(Targets[SkipIndex]);
+
+	UE::Geometry::FTransform3d TargetToWorld = (UE::Geometry::FTransform3d)UE::ToolTarget::GetLocalToWorldTransform(Targets[SkipIndex]);
+	UE::Geometry::FTransform3d WorldToTarget = TargetToWorld.Inverse();
+
+	FSimpleShapeSet3d SimpleCollision;
+	UE::Geometry::FComponentCollisionSettings CollisionSettings;
+	bool bOutputComponentSupportsCollision = UE::Geometry::ComponentTypeSupportsCollision(UpdateComponent);
+	if (bOutputComponentSupportsCollision)
+	{
+		CollisionSettings = UE::Geometry::GetCollisionSettings(UpdateComponent);
+	}
+	TArray<UE::Geometry::FTransform3d> Transforms;
+	Transforms.SetNum(2);
+
+	{
+#if WITH_EDITOR
+		FScopedSlowTask SlowTask(Targets.Num()+1, 
+			bDuplicateMode ? 
+			LOCTEXT("DuplicateMeshBuild", "Building duplicate mesh ...") :
+			LOCTEXT("CombineMeshesBuild", "Building combined mesh ..."));
+		SlowTask.MakeDialog();
+#endif
+		bool bNeedColorAttr = false;
+		for (int32 ComponentIdx = 0; ComponentIdx < Targets.Num(); ComponentIdx++)
+		{
+#if WITH_EDITOR
+			SlowTask.EnterProgressFrame(1);
+#endif
+			UPrimitiveComponent* PrimitiveComponent = UE::ToolTarget::GetTargetComponent(Targets[ComponentIdx]);
+
+			FDynamicMesh3& ComponentDMesh = InputMeshes[ComponentIdx];
+			bNeedColorAttr = bNeedColorAttr || (ComponentDMesh.HasAttributes() && ComponentDMesh.Attributes()->HasPrimaryColors());
+
+			// update material IDs to account for combined material set
+			FDynamicMeshMaterialAttribute* MatAttrib = ComponentDMesh.Attributes()->GetMaterialID();
+			for (int TID : ComponentDMesh.TriangleIndicesItr())
+			{
+				int MatID = MatAttrib->GetValue(TID);
+				MatAttrib->SetValue(TID, MaterialIDRemaps[ComponentIdx][MatID]);
+			}
+
+			if (ComponentIdx != SkipIndex)
+			{
+				UE::Geometry::FTransform3d ComponentToWorld = (UE::Geometry::FTransform3d)UE::ToolTarget::GetLocalToWorldTransform(Targets[ComponentIdx]);
+				MeshTransforms::ApplyTransform(ComponentDMesh, ComponentToWorld);
+				if (ComponentToWorld.GetDeterminant() < 0)
+				{
+					ComponentDMesh.ReverseOrientation(true);
+				}
+				MeshTransforms::ApplyTransform(ComponentDMesh, WorldToTarget);
+				if (WorldToTarget.GetDeterminant() < 0)
+				{
+					ComponentDMesh.ReverseOrientation(true);
+				}
+				Transforms[0] = ComponentToWorld;
+				Transforms[1] = WorldToTarget;
+				if (bOutputComponentSupportsCollision && UE::Geometry::ComponentTypeSupportsCollision(PrimitiveComponent))
+				{
+					UE::Geometry::AppendSimpleCollision(PrimitiveComponent, &SimpleCollision, Transforms);
+				}
+			}
+			else
+			{
+				if (bOutputComponentSupportsCollision && UE::Geometry::ComponentTypeSupportsCollision(PrimitiveComponent))
+				{
+					UE::Geometry::AppendSimpleCollision(PrimitiveComponent, &SimpleCollision, UE::Geometry::FTransform3d::Identity());
+				}
+			}
+
+			FDynamicMeshEditor Editor(&AccumulateDMesh);
+			FMeshIndexMappings IndexMapping;
+			Editor.AppendMesh(&ComponentDMesh, IndexMapping);
+		}
+
+		if (!bNeedColorAttr)
+		{
+			AccumulateDMesh.Attributes()->DisablePrimaryColors();
+		}
+
+#if WITH_EDITOR
+		SlowTask.EnterProgressFrame(1);
+#endif
+
+		FComponentMaterialSet NewMaterialSet;
+		NewMaterialSet.Materials = AllMaterials;
+		UE::ToolTarget::CommitDynamicMeshUpdate(Targets[SkipIndex], AccumulateDMesh, true, FConversionToMeshDescriptionOptions(), &NewMaterialSet);
+
+		if (bOutputComponentSupportsCollision)
+		{
+			UE::Geometry::SetSimpleCollision(UpdateComponent, &SimpleCollision, CollisionSettings);
+		}
+
+		// select the new actor
+		ToolSelectionUtil::SetNewActorSelection(GetToolManager(), SkipActor);
+	}
+
+	
+	TArray<AActor*> Actors;
+	for (int Idx = 0; Idx < Targets.Num(); Idx++)
+	{
+		AActor* Actor = UE::ToolTarget::GetTargetActor(Targets[Idx]);
+		if (Actor != SkipActor)
+		{
+			Actors.Add(Actor);
+		}
+	}
+	HandleSourceProperties->ApplyMethod(Actors, GetToolManager());
+
+	GetToolManager()->EndUndoTransaction();
+}
+
+
+
+
+
+
+
+void UCombineMeshesTool::BuildCombinedMaterialSet(TArray<UMaterialInterface*>& NewMaterialsOut, TArray<TArray<int32>>& MaterialIDRemapsOut)
+{
+	NewMaterialsOut.Reset();
+
+	TMap<UMaterialInterface*, int> KnownMaterials;
+
+	MaterialIDRemapsOut.SetNum(Targets.Num());
+	for (int32 ComponentIdx = 0; ComponentIdx < Targets.Num(); ComponentIdx++)
+	{
+		FComponentMaterialSet MaterialSet = UE::ToolTarget::GetMaterialSet(Targets[ComponentIdx]);
+		int32 NumMaterials = MaterialSet.Materials.Num();
+		for (int MaterialIdx = 0; MaterialIdx < NumMaterials; MaterialIdx++)
+		{
+			UMaterialInterface* Mat = MaterialSet.Materials[MaterialIdx];
+			int32 NewMaterialIdx = 0;
+			if (KnownMaterials.Contains(Mat) == false)
+			{
+				NewMaterialIdx = NewMaterialsOut.Num();
+				KnownMaterials.Add(Mat, NewMaterialIdx);
+				NewMaterialsOut.Add(Mat);
+			}
+			else
+			{
+				NewMaterialIdx = KnownMaterials[Mat];
+			}
+			MaterialIDRemapsOut[ComponentIdx].Add(NewMaterialIdx);
+		}
+	}
+}
+
+
+
+
+
+
+
+
+#undef LOCTEXT_NAMESPACE
diff --git a/Engine/Plugins/Runtime/MeshModelingToolset/Source/MeshModelingTools/Private/DrawAndRevolveTool.cpp b/Engine/Plugins/Runtime/MeshModelingToolset/Source/MeshModelingTools/Private/DrawAndRevolveTool.cpp
new file mode 100644
--- /dev/null
+++ b/Engine/Plugins/Runtime/MeshModelingToolset/Source/MeshModelingTools/Private/DrawAndRevolveTool.cpp
@@ -0,0 +1,355 @@
+// Copyright Epic Games, Inc. All Rights Reserved.
+
+#include "DrawAndRevolveTool.h"
+
+#include "BaseGizmos/GizmoRenderingUtil.h"
+#include "BaseBehaviors/SingleClickBehavior.h"
+#include "BaseBehaviors/MouseHoverBehavior.h"
+#include "CompositionOps/CurveSweepOp.h"
+#include "Generators/SweepGenerator.h"
+#include "InteractiveToolManager.h" // To use SceneState.ToolManager
+#include "Mechanics/ConstructionPlaneMechanic.h"
+#include "Mechanics/CurveControlPointsMechanic.h"
+#include "Properties/MeshMaterialProperties.h"
+#include "ModelingObjectsCreationAPI.h"
+#include "Selection/ToolSelectionUtil.h"
+#include "ToolSceneQueriesUtil.h"
+#include "ToolSetupUtil.h"
+
+#include "ExplicitUseGeometryMathTypes.h"		// using UE::Geometry::(math types)
+using namespace UE::Geometry;
+
+#define LOCTEXT_NAMESPACE "UDrawAndRevolveTool"
+
+const FText InitializationModeMessage = LOCTEXT("CurveInitialization", "Draw a profile curve and revolve it around purple axis. Ctrl+click repositions draw plane, Ctrl+Shift+click repositions without aligning axis. End curve by clicking the end again or connecting to start. Backspace deletes points, shift inverts snapping behavior.");
+const FText EditModeMessage = LOCTEXT("CurveEditing", "Click points to select them, Shift+click to add/remove points to selection. Ctrl+click a segment to add a point, or select an endpoint and Ctrl+click somewhere on the plane to add to the ends. Backspace deletes selected points. Holding Shift toggles snapping to be opposite the EnableSnapping setting.");
+
+
+// Tool builder
+
+bool UDrawAndRevolveToolBuilder::CanBuildTool(const FToolBuilderState& SceneState) const
+{
+	return true;
+}
+
+UInteractiveTool* UDrawAndRevolveToolBuilder::BuildTool(const FToolBuilderState& SceneState) const
+{
+	UDrawAndRevolveTool* NewTool = NewObject<UDrawAndRevolveTool>(SceneState.ToolManager);
+
+	NewTool->SetWorld(SceneState.World);
+	return NewTool;
+}
+
+
+// Operator factory
+
+TUniquePtr<FDynamicMeshOperator> URevolveOperatorFactory::MakeNewOperator()
+{
+	TUniquePtr<FCurveSweepOp> CurveSweepOp = MakeUnique<FCurveSweepOp>();
+
+	// Assemble profile curve
+	CurveSweepOp->ProfileCurve.Reserve(RevolveTool->ControlPointsMechanic->GetNumPoints() + 2); // extra space for top/bottom caps
+	RevolveTool->ControlPointsMechanic->ExtractPointPositions(CurveSweepOp->ProfileCurve);
+	CurveSweepOp->bProfileCurveIsClosed = RevolveTool->ControlPointsMechanic->GetIsLoop();
+
+	// If we are capping the top and bottom, we just add a couple extra vertices and mark the curve as being closed
+	if (!CurveSweepOp->bProfileCurveIsClosed && RevolveTool->Settings->bConnectOpenProfileToAxis)
+	{
+		// Project first and last points onto the revolution axis.
+		FVector3d FirstPoint = CurveSweepOp->ProfileCurve[0];
+		FVector3d LastPoint = CurveSweepOp->ProfileCurve.Last();
+		double DistanceAlongAxis = RevolveTool->RevolutionAxisDirection.Dot(LastPoint - RevolveTool->RevolutionAxisOrigin);
+		FVector3d ProjectedPoint = RevolveTool->RevolutionAxisOrigin + (RevolveTool->RevolutionAxisDirection * DistanceAlongAxis);
+
+		CurveSweepOp->ProfileCurve.Add(ProjectedPoint);
+
+		DistanceAlongAxis = RevolveTool->RevolutionAxisDirection.Dot(FirstPoint - RevolveTool->RevolutionAxisOrigin);
+		ProjectedPoint = RevolveTool->RevolutionAxisOrigin + (RevolveTool->RevolutionAxisDirection * DistanceAlongAxis);
+
+		CurveSweepOp->ProfileCurve.Add(ProjectedPoint);
+		CurveSweepOp->bProfileCurveIsClosed = true;
+	}
+
+	RevolveTool->Settings->ApplyToCurveSweepOp(*RevolveTool->MaterialProperties,
+		RevolveTool->RevolutionAxisOrigin, RevolveTool->RevolutionAxisDirection, *CurveSweepOp);
+
+	return CurveSweepOp;
+}
+
+
+// Tool itself
+
+void UDrawAndRevolveTool::RegisterActions(FInteractiveToolActionSet& ActionSet)
+{
+	ActionSet.RegisterAction(this, (int32)EStandardToolActions::BaseClientDefinedActionID + 1,
+		TEXT("DeletePointBackspaceKey"),
+		LOCTEXT("DeletePointUIName", "Delete Point"),
+		LOCTEXT("DeletePointTooltip", "Delete currently selected point(s)"),
+		EModifierKey::None, EKeys::BackSpace,
+		[this]() { OnPointDeletionKeyPress(); });
+	ActionSet.RegisterAction(this, (int32)EStandardToolActions::BaseClientDefinedActionID + 2,
+		TEXT("DeletePointDeleteKey"),
+		LOCTEXT("DeletePointUIName", "Delete Point"),
+		LOCTEXT("DeletePointTooltip", "Delete currently selected point(s)"),
+		EModifierKey::None, EKeys::Delete,
+		[this]() { OnPointDeletionKeyPress(); });
+}
+
+void UDrawAndRevolveTool::OnPointDeletionKeyPress()
+{
+	if (ControlPointsMechanic)
+	{
+		ControlPointsMechanic->DeleteSelectedPoints();
+	}
+}
+
+bool UDrawAndRevolveTool::CanAccept() const
+{
+	return Preview != nullptr && Preview->HaveValidNonEmptyResult();
+}
+
+void UDrawAndRevolveTool::Setup()
+{
+	UInteractiveTool::Setup();
+
+	SetToolDisplayName(LOCTEXT("ToolName", "Revolve PolyPath"));
+	GetToolManager()->DisplayMessage(InitializationModeMessage, EToolMessageLevel::UserNotification);
+
+	OutputTypeProperties = NewObject<UCreateMeshObjectTypeProperties>(this);
+	OutputTypeProperties->RestoreProperties(this);
+	OutputTypeProperties->InitializeDefault();
+	OutputTypeProperties->WatchProperty(OutputTypeProperties->OutputType, [this](FString) { OutputTypeProperties->UpdatePropertyVisibility(); });
+	AddToolPropertySource(OutputTypeProperties);
+
+	Settings = NewObject<URevolveToolProperties>(this, TEXT("Revolve Tool Settings"));
+	Settings->RestoreProperties(this);
+	Settings->bAllowedToEditDrawPlane = true;
+	AddToolPropertySource(Settings);
+
+	MaterialProperties = NewObject<UNewMeshMaterialProperties>(this);
+	AddToolPropertySource(MaterialProperties);
+	MaterialProperties->RestoreProperties(this);
+
+	ControlPointsMechanic = NewObject<UCurveControlPointsMechanic>(this);
+	ControlPointsMechanic->Setup(this);
+	ControlPointsMechanic->SetWorld(TargetWorld);
+	ControlPointsMechanic->OnPointsChanged.AddLambda([this]() {
+		if (Preview)
+		{
+			Preview->InvalidateResult();
+		}
+		bool bAllowedToEditDrawPlane = (ControlPointsMechanic->GetNumPoints() == 0);
+		if (Settings->bAllowedToEditDrawPlane != bAllowedToEditDrawPlane)
+		{
+			Settings->bAllowedToEditDrawPlane = bAllowedToEditDrawPlane;
+			NotifyOfPropertyChangeByTool(Settings);
+		}
+		});
+	// This gets called when we enter/leave curve initialization mode
+	ControlPointsMechanic->OnModeChanged.AddLambda([this]() {
+		if (ControlPointsMechanic->IsInInteractiveIntialization())
+		{
+			// We're back to initializing so hide the preview
+			if (Preview)
+			{
+				Preview->Cancel();
+				Preview = nullptr;
+			}
+			GetToolManager()->DisplayMessage(InitializationModeMessage, EToolMessageLevel::UserNotification);
+		}
+		else
+		{
+			StartPreview();
+			GetToolManager()->DisplayMessage(EditModeMessage, EToolMessageLevel::UserNotification);
+		}
+		});
+	ControlPointsMechanic->SetSnappingEnabled(Settings->bEnableSnapping);
+
+	UpdateRevolutionAxis();
+
+	// The plane mechanic lets us update the plane in which we draw the profile curve, as long as we haven't
+	// started adding points to it already.
+	FFrame3d ProfileDrawPlane(Settings->DrawPlaneOrigin, Settings->DrawPlaneOrientation.Quaternion());
+	PlaneMechanic = NewObject<UConstructionPlaneMechanic>(this);
+	PlaneMechanic->Setup(this);
+	PlaneMechanic->Initialize(TargetWorld, ProfileDrawPlane);
+	PlaneMechanic->UpdateClickPriority(ControlPointsMechanic->ClickBehavior->GetPriority().MakeHigher());
+	PlaneMechanic->CanUpdatePlaneFunc = [this]() 
+	{ 
+		return ControlPointsMechanic->GetNumPoints() == 0;
+	};
+	PlaneMechanic->OnPlaneChanged.AddLambda([this]() {
+		Settings->DrawPlaneOrigin = (FVector)PlaneMechanic->Plane.Origin;
+		Settings->DrawPlaneOrientation = ((FQuat)PlaneMechanic->Plane.Rotation).Rotator();
+		NotifyOfPropertyChangeByTool(Settings);
+		if (ControlPointsMechanic)
+		{
+			ControlPointsMechanic->SetPlane(PlaneMechanic->Plane);
+		}
+		UpdateRevolutionAxis();
+		});
+
+	ControlPointsMechanic->SetPlane(PlaneMechanic->Plane);
+	ControlPointsMechanic->SetInteractiveInitialization(true);
+	ControlPointsMechanic->SetAutoRevertToInteractiveInitialization(true);
+
+	// For things to behave nicely, we expect to revolve at least a two point
+	// segment if it's not a loop, and a three point segment if it is. Revolving
+	// a two-point loop to make a double sided thing is a pain to support because
+	// it forces us to deal with manifoldness issues that we would really rather
+	// not worry about (we'd have to duplicate vertices to stay manifold)
+	ControlPointsMechanic->SetMinPointsToLeaveInteractiveInitialization(3, 2);
+}
+
+/** Uses the settings currently stored in the properties object to update the revolution axis. */
+void UDrawAndRevolveTool::UpdateRevolutionAxis()
+{
+	RevolutionAxisOrigin = (FVector3d)Settings->DrawPlaneOrigin;
+	RevolutionAxisDirection = (FVector3d)Settings->DrawPlaneOrientation.RotateVector(FVector(1,0,0));
+
+	const int32 AXIS_SNAP_TARGET_ID = 1;
+	ControlPointsMechanic->RemoveSnapLine(AXIS_SNAP_TARGET_ID);
+	ControlPointsMechanic->AddSnapLine(AXIS_SNAP_TARGET_ID, FLine3d(RevolutionAxisOrigin, RevolutionAxisDirection));
+}
+
+void UDrawAndRevolveTool::Shutdown(EToolShutdownType ShutdownType)
+{
+	OutputTypeProperties->SaveProperties(this);
+	Settings->SaveProperties(this);
+	MaterialProperties->SaveProperties(this);
+
+	PlaneMechanic->Shutdown();
+	ControlPointsMechanic->Shutdown();
+
+	if (Preview)
+	{
+		if (ShutdownType == EToolShutdownType::Accept)
+		{
+			GenerateAsset(Preview->Shutdown());
+		}
+		else
+		{
+			Preview->Cancel();
+		}
+	}
+}
+
+void UDrawAndRevolveTool::GenerateAsset(const FDynamicMeshOpResult& OpResult)
+{
+	if (OpResult.Mesh->TriangleCount() <= 0)
+	{
+		return;
+	}
+
+	GetToolManager()->BeginUndoTransaction(LOCTEXT("RevolveToolTransactionName", "Revolve Tool"));
+
+	FCreateMeshObjectParams NewMeshObjectParams;
+	NewMeshObjectParams.TargetWorld = TargetWorld;
+	NewMeshObjectParams.Transform = (FTransform)OpResult.Transform;
+	NewMeshObjectParams.BaseName = TEXT("Revolve");
+	NewMeshObjectParams.Materials.Add(MaterialProperties->Material.Get());
+	NewMeshObjectParams.SetMesh(OpResult.Mesh.Get());
+	OutputTypeProperties->ConfigureCreateMeshObjectParams(NewMeshObjectParams);
+	FCreateMeshObjectResult Result = UE::Modeling::CreateMeshObject(GetToolManager(), MoveTemp(NewMeshObjectParams));
+	if (Result.IsOK() && Result.NewActor != nullptr)
+	{
+		ToolSelectionUtil::SetNewActorSelection(GetToolManager(), Result.NewActor);
+	}
+
+	GetToolManager()->EndUndoTransaction();
+}
+
+void UDrawAndRevolveTool::StartPreview()
+{
+	URevolveOperatorFactory* RevolveOpCreator = NewObject<URevolveOperatorFactory>();
+	RevolveOpCreator->RevolveTool = this;
+
+	// Normally we wouldn't give the object a name, but since we may destroy the preview using undo,
+	// the ability to reuse the non-cleaned up memory is useful. Careful if copy-pasting this!
+	Preview = NewObject<UMeshOpPreviewWithBackgroundCompute>(RevolveOpCreator, "RevolveToolPreview");
+
+	Preview->Setup(TargetWorld, RevolveOpCreator);
+	Preview->PreviewMesh->SetTangentsMode(EDynamicMeshComponentTangentsMode::AutoCalculated);
+
+	Preview->ConfigureMaterials(MaterialProperties->Material.Get(),
+		ToolSetupUtil::GetDefaultWorkingMaterial(GetToolManager()));
+	Preview->PreviewMesh->EnableWireframe(MaterialProperties->bWireframe);
+
+	Preview->OnMeshUpdated.AddLambda(
+		[this](const UMeshOpPreviewWithBackgroundCompute* UpdatedPreview)
+		{
+			UpdateAcceptWarnings(UpdatedPreview->HaveEmptyResult() ? EAcceptWarning::EmptyForbidden : EAcceptWarning::NoWarning);
+		}
+	);
+
+	Preview->SetVisibility(true);
+	Preview->InvalidateResult();
+}
+
+void UDrawAndRevolveTool::OnPropertyModified(UObject* PropertySet, FProperty* Property)
+{
+	FFrame3d ProfileDrawPlane(Settings->DrawPlaneOrigin, Settings->DrawPlaneOrientation.Quaternion());
+	ControlPointsMechanic->SetPlane(ProfileDrawPlane);
+	PlaneMechanic->SetPlaneWithoutBroadcast(ProfileDrawPlane);
+	UpdateRevolutionAxis();
+
+	ControlPointsMechanic->SetSnappingEnabled(Settings->bEnableSnapping);
+
+	if (Preview)
+	{
+		if (Property && (Property->GetFName() == GET_MEMBER_NAME_CHECKED(UNewMeshMaterialProperties, Material)))
+		{
+			Preview->ConfigureMaterials(MaterialProperties->Material.Get(),
+				ToolSetupUtil::GetDefaultWorkingMaterial(GetToolManager()));
+		}
+
+		Preview->PreviewMesh->EnableWireframe(MaterialProperties->bWireframe);
+		Preview->InvalidateResult();
+	}
+}
+
+
+void UDrawAndRevolveTool::OnTick(float DeltaTime)
+{
+	if (PlaneMechanic != nullptr)
+	{
+		PlaneMechanic->Tick(DeltaTime);
+	}
+
+	if (Preview)
+	{
+		Preview->Tick(DeltaTime);
+	}
+}
+
+
+void UDrawAndRevolveTool::Render(IToolsContextRenderAPI* RenderAPI)
+{
+	GetToolManager()->GetContextQueriesAPI()->GetCurrentViewState(CameraState);
+
+	if (PlaneMechanic != nullptr)
+	{
+		PlaneMechanic->Render(RenderAPI);
+
+		// Draw the axis of rotation
+		float PdiScale = CameraState.GetPDIScalingFactor();
+		FPrimitiveDrawInterface* PDI = RenderAPI->GetPrimitiveDrawInterface();
+
+		FColor AxisColor(240, 16, 240);
+		double AxisThickness = 1.0 * PdiScale;
+		double AxisHalfLength = ToolSceneQueriesUtil::CalculateDimensionFromVisualAngleD(CameraState, RevolutionAxisOrigin, 90);
+
+		FVector3d StartPoint = RevolutionAxisOrigin - (RevolutionAxisDirection * (AxisHalfLength * PdiScale));
+		FVector3d EndPoint = RevolutionAxisOrigin + (RevolutionAxisDirection * (AxisHalfLength * PdiScale));
+
+		PDI->DrawLine((FVector)StartPoint, (FVector)EndPoint, AxisColor, SDPG_Foreground, 
+			AxisThickness, 0.0f, true);
+	}
+
+	if (ControlPointsMechanic != nullptr)
+	{
+		ControlPointsMechanic->Render(RenderAPI);
+	}
+}
+
+#undef LOCTEXT_NAMESPACE
diff --git a/Engine/Plugins/Runtime/MeshModelingToolset/Source/MeshModelingTools/Private/DrawPolygonTool.cpp b/Engine/Plugins/Runtime/MeshModelingToolset/Source/MeshModelingTools/Private/DrawPolygonTool.cpp
new file mode 100644
--- /dev/null
+++ b/Engine/Plugins/Runtime/MeshModelingToolset/Source/MeshModelingTools/Private/DrawPolygonTool.cpp
@@ -0,0 +1,1294 @@
+// Copyright Epic Games, Inc. All Rights Reserved.
+
+#include "DrawPolygonTool.h"
+#include "InteractiveToolManager.h"
+#include "InteractiveGizmoManager.h"
+#include "ToolBuilderUtil.h"
+#include "BaseBehaviors/MultiClickSequenceInputBehavior.h"
+#include "BaseBehaviors/KeyAsModifierInputBehavior.h"
+
+#include "Polygon2.h"
+#include "Curve/GeneralPolygon2.h"
+#include "FrameTypes.h"
+#include "MatrixTypes.h"
+#include "DynamicMesh/DynamicMeshAttributeSet.h"
+
+#include "Generators/FlatTriangulationMeshGenerator.h"
+#include "Generators/DiscMeshGenerator.h"
+#include "Generators/RectangleMeshGenerator.h"
+#include "Operations/ExtrudeMesh.h"
+#include "Distance/DistLine3Ray3.h"
+#include "Intersection/IntrSegment2Segment2.h"
+#include "MeshQueries.h"
+#include "ToolSceneQueriesUtil.h"
+#include "ConstrainedDelaunay2.h"
+#include "Arrangement2d.h"
+
+#include "DynamicMeshEditor.h"
+
+#include "BaseGizmos/GizmoComponents.h"
+#include "BaseGizmos/TransformGizmoUtil.h"
+#include "Drawing/MeshDebugDrawing.h"
+
+#include "Selection/SelectClickedAction.h"
+#include "Selection/ToolSelectionUtil.h"
+#include "ModelingObjectsCreationAPI.h"
+
+#include "ExplicitUseGeometryMathTypes.h"		// using UE::Geometry::(math types)
+using namespace UE::Geometry;
+
+#define LOCTEXT_NAMESPACE "UDrawPolygonTool"
+
+/*
+ * ToolBuilder
+ */
+constexpr int StartPointSnapID = FPointPlanarSnapSolver::BaseExternalPointID + 1;
+constexpr int CurrentSceneSnapID = FPointPlanarSnapSolver::BaseExternalPointID + 2;
+constexpr int CurrentGridSnapID = FPointPlanarSnapSolver::BaseExternalPointID + 3;
+
+bool UDrawPolygonToolBuilder::CanBuildTool(const FToolBuilderState& SceneState) const
+{
+	return true;
+}
+
+UInteractiveTool* UDrawPolygonToolBuilder::BuildTool(const FToolBuilderState& SceneState) const
+{
+	UDrawPolygonTool* NewTool = NewObject<UDrawPolygonTool>(SceneState.ToolManager);
+	NewTool->SetWorld(SceneState.World);
+	return NewTool;
+}
+
+/*
+ * Properties
+ */
+UDrawPolygonToolStandardProperties::UDrawPolygonToolStandardProperties()
+{
+}
+
+/*
+ * Tool
+ */
+UDrawPolygonTool::UDrawPolygonTool()
+{
+	DrawPlaneOrigin = FVector3d::Zero();
+	DrawPlaneOrientation = FQuaterniond::Identity();
+	bInInteractiveExtrude = false;
+	SetToolDisplayName(LOCTEXT("ToolName", "Extrude Polygon"));
+}
+
+void UDrawPolygonTool::SetWorld(UWorld* World)
+{
+	this->TargetWorld = World;
+}
+
+void UDrawPolygonTool::Setup()
+{
+	UInteractiveTool::Setup();
+
+	// add default button input behaviors for devices
+	UMultiClickSequenceInputBehavior* MouseBehavior = NewObject<UMultiClickSequenceInputBehavior>();
+	MouseBehavior->Initialize(this);
+	MouseBehavior->Modifiers.RegisterModifier(IgnoreSnappingModifier, FInputDeviceState::IsShiftKeyDown);
+	AddInputBehavior(MouseBehavior);
+
+	// Register a click behavior/action pair, that sets the draw plane to the clicked world position
+	FSelectClickedAction* SetPlaneAction = new FSelectClickedAction();
+	SetPlaneAction->World = this->TargetWorld;
+	SetPlaneAction->OnClickedPositionFunc = [this](const FHitResult& Hit) {
+		SetDrawPlaneFromWorldPos((FVector3d)Hit.ImpactPoint, (FVector3d)Hit.ImpactNormal);
+	};
+	SetPointInWorldConnector = SetPlaneAction;
+
+	USingleClickInputBehavior* ClickToSetPlaneBehavior = NewObject<USingleClickInputBehavior>();
+	ClickToSetPlaneBehavior->ModifierCheckFunc = FInputDeviceState::IsCtrlKeyDown;
+	ClickToSetPlaneBehavior->Initialize(SetPointInWorldConnector);
+	ClickToSetPlaneBehavior->SetDefaultPriority(MouseBehavior->GetPriority().MakeHigher());
+	AddInputBehavior(ClickToSetPlaneBehavior);
+
+	// register modifier key behaviors   (disabled because it is not implemented yet)
+	//UKeyAsModifierInputBehavior* AKeyBehavior = NewObject<UKeyAsModifierInputBehavior>();
+	//AKeyBehavior->Initialize(this, AngleSnapModifier, EKeys::A);
+	//AddInputBehavior(AKeyBehavior);
+
+	OutputTypeProperties = NewObject<UCreateMeshObjectTypeProperties>(this);
+	OutputTypeProperties->RestoreProperties(this);
+	OutputTypeProperties->InitializeDefault();
+	OutputTypeProperties->WatchProperty(OutputTypeProperties->OutputType, [this](FString) { OutputTypeProperties->UpdatePropertyVisibility(); });
+
+	PolygonProperties = NewObject<UDrawPolygonToolStandardProperties>(this);
+	PolygonProperties->RestoreProperties(this);
+	PolygonProperties->WatchProperty(PolygonProperties->bShowGizmo,
+									 [this](bool bNewValue) { this->UpdateShowGizmoState(bNewValue); });
+
+	// Create a new TransformGizmo and associated TransformProxy. The TransformProxy will not be the
+	// parent of any Components in this case, we just use it's transform and change delegate.
+	PlaneTransformProxy = NewObject<UTransformProxy>(this);
+	PlaneTransformProxy->SetTransform(FTransform((FQuat)DrawPlaneOrientation, (FVector)DrawPlaneOrigin));
+	PlaneTransformGizmo = UE::TransformGizmoUtil::CreateCustomTransformGizmo(GetToolManager(),
+		ETransformGizmoSubElements::StandardTranslateRotate, this);
+	PlaneTransformGizmo->SetActiveTarget(PlaneTransformProxy, GetToolManager());
+	// listen for changes to the proxy and update the plane when that happens
+	PlaneTransformProxy->OnTransformChanged.AddUObject(this, &UDrawPolygonTool::PlaneTransformChanged);
+
+	// initialize material properties for new objects
+	MaterialProperties = NewObject<UNewMeshMaterialProperties>(this);
+	MaterialProperties->RestoreProperties(this);
+	MaterialProperties->bShowExtendedOptions = true;
+
+	// create preview mesh object
+	PreviewMesh = NewObject<UPreviewMesh>(this);
+	PreviewMesh->CreateInWorld(this->TargetWorld, FTransform::Identity);
+	PreviewMesh->SetVisible(false);
+	{
+		UMaterialInterface* Material = nullptr;
+		if ( MaterialProperties->Material.IsValid() )
+		{
+			Material = MaterialProperties->Material.Get();
+		}
+		PreviewMesh->SetMaterial(Material);
+	}
+	bPreviewUpdatePending = false;
+
+	// initialize snapping engine and properties
+	SnapEngine.SnapMetricTolerance = ToolSceneQueriesUtil::GetDefaultVisualAngleSnapThreshD();
+	SnapEngine.SnapMetricFunc = [this](const FVector3d& Position1, const FVector3d& Position2) {
+		return ToolSceneQueriesUtil::CalculateNormalizedViewVisualAngleD(this->CameraState, Position1, Position2);
+	};
+	SnapEngine.Plane = FFrame3d(DrawPlaneOrigin, DrawPlaneOrientation);
+
+	SnapProperties = NewObject<UDrawPolygonToolSnapProperties>(this);
+	SnapProperties->RestoreProperties(this);
+
+	// register tool properties
+	if (OutputTypeProperties->ShouldShowPropertySet())
+	{
+		AddToolPropertySource(OutputTypeProperties);
+	}
+	AddToolPropertySource(PolygonProperties);
+	AddToolPropertySource(SnapProperties);
+	AddToolPropertySource(MaterialProperties);
+
+	ShowStartupMessage();
+}
+
+
+void UDrawPolygonTool::Shutdown(EToolShutdownType ShutdownType)
+{
+	PreviewMesh->Disconnect();
+	PreviewMesh = nullptr;
+
+	if (SetPointInWorldConnector != nullptr)
+	{
+		delete SetPointInWorldConnector;
+	}
+
+	GetToolManager()->GetPairedGizmoManager()->DestroyAllGizmosByOwner(this);
+
+	OutputTypeProperties->SaveProperties(this);
+	PolygonProperties->SaveProperties(this);
+	SnapProperties->SaveProperties(this);
+	MaterialProperties->SaveProperties(this);
+}
+
+void UDrawPolygonTool::RegisterActions(FInteractiveToolActionSet& ActionSet)
+{
+	ActionSet.RegisterAction(this, (int32)EStandardToolActions::BaseClientDefinedActionID + 2,
+		TEXT("ToggleGizmo"),
+		LOCTEXT("ToggleGizmo", "Toggle Gizmo"),
+		LOCTEXT("ToggleGizmoTooltip", "Toggle visibility of the transformation Gizmo"),
+		EModifierKey::None, EKeys::A,
+		[this]() { PolygonProperties->bShowGizmo = !PolygonProperties->bShowGizmo; });
+}
+
+
+
+void UDrawPolygonTool::ApplyUndoPoints(const TArray<FVector3d>& ClickPointsIn, const TArray<FVector3d>& PolygonVerticesIn)
+{
+	if (bInInteractiveExtrude || PolygonVertices.Num() == 0)
+	{
+		return;
+	}
+
+	bHaveSelfIntersection = false;
+
+	if (bInFixedPolygonMode == false)
+	{
+		PolygonVertices = PolygonVerticesIn;
+		if ( PolygonVertices.Num() == 0 )
+		{
+			bAbortActivePolygonDraw = true;
+			CurrentCurveTimestamp++;
+		}
+	}
+	else
+	{
+		FixedPolygonClickPoints = ClickPointsIn;
+		if ( FixedPolygonClickPoints.Num() == 0 )
+		{
+			bAbortActivePolygonDraw = true;
+			CurrentCurveTimestamp++;
+		}
+	}
+}
+
+
+
+
+void UDrawPolygonTool::OnTick(float DeltaTime)
+{
+}
+
+
+
+void DrawEdgeTicks(FPrimitiveDrawInterface* PDI, 
+	const FSegment3d& Segment, float Height,
+	const FVector3d& PlaneNormal, 
+	const FLinearColor& Color, uint8 DepthPriorityGroup, float LineThickness, bool bIsScreenSpace)
+{
+	FVector3d Center = Segment.Center;
+	FVector3d X = Segment.Direction;
+	FVector3d Y = X.Cross(PlaneNormal);
+	UE::Geometry::Normalize(Y);
+	FVector3d A = Center - Height * 0.25*X - Height * Y;
+	FVector3d B = Center + Height * 0.25*X + Height * Y;
+	PDI->DrawLine((FVector)A, (FVector)B, Color, DepthPriorityGroup, LineThickness, 0.0f, bIsScreenSpace);
+	A += Height * 0.5*X;
+	B += Height * 0.5*X;
+	PDI->DrawLine((FVector)A, (FVector)B, Color, DepthPriorityGroup, LineThickness, 0.0f, bIsScreenSpace);
+}
+
+void UDrawPolygonTool::Render(IToolsContextRenderAPI* RenderAPI)
+{
+	FPrimitiveDrawInterface* PDI = RenderAPI->GetPrimitiveDrawInterface();
+	// Cache here for usage during interaction, should probably happen in ::Tick() or elsewhere
+	GetToolManager()->GetContextQueriesAPI()->GetCurrentViewState(CameraState);
+	
+	FViewCameraState RenderCameraState = RenderAPI->GetCameraState();
+	float PDIScale = RenderCameraState.GetPDIScalingFactor();
+
+	if (bPreviewUpdatePending)
+	{
+		UpdateLivePreview();
+		bPreviewUpdatePending = false;
+	}
+
+	double CurViewSizeFactor = ToolSceneQueriesUtil::CalculateDimensionFromVisualAngleD(RenderCameraState, PreviewVertex, 1.0);
+
+	FColor PreviewColor = FColor::Green;
+	FColor OpenPolygonColor = FColor::Orange;
+	FColor ClosedPolygonColor = FColor::Yellow;
+	FColor ErrorColor = FColor::Magenta;
+	float HiddenLineThickness = 1.0f*PDIScale;
+	float LineThickness = 4.0f*PDIScale;
+	float SelfIntersectThickness = 8.0f*PDIScale;
+	FColor GridColor(128, 128, 128, 32);
+	float GridThickness = 0.5f*PDIScale;
+	float GridLineSpacing = 25.0f*PDIScale;   // @todo should be relative to view
+	int NumGridLines = 21;
+	FColor SnapLineColor = FColor::Yellow;
+	FColor SnapHighlightColor = SnapLineColor;
+	float ElementSize = CurViewSizeFactor;
+
+	bool bIsClosed = bInInteractiveExtrude 
+		|| (SnapEngine.HaveActiveSnap() && SnapEngine.GetActiveSnapTargetID() == StartPointSnapID);
+
+	//
+	// Draw the grid
+	//
+	if (bInInteractiveExtrude == false)
+	{
+		FFrame3f DrawFrame((FVector3f)DrawPlaneOrigin, (FQuaternionf)DrawPlaneOrientation);
+		MeshDebugDraw::DrawSimpleFixedScreenAreaGrid(RenderCameraState, DrawFrame, NumGridLines, 45.0, GridThickness, GridColor, false, PDI, FTransform::Identity);
+	}
+
+	//
+	// Generate the fixed polygon contour
+	//
+	if ((bInFixedPolygonMode) && (FixedPolygonClickPoints.Num() > 0))
+	{
+		TArray<FVector3d> PreviewClickPoints = FixedPolygonClickPoints;
+		if ( !bInInteractiveExtrude ){ PreviewClickPoints.Add(PreviewVertex); }
+		GenerateFixedPolygon(PreviewClickPoints, PolygonVertices, PolygonHolesVertices);
+	}
+	bIsClosed |= bInFixedPolygonMode;
+
+	int NumVerts = PolygonVertices.Num();
+
+	//
+	// Render snap indicators
+	//
+	if (SnapEngine.HaveActiveSnap())
+	{
+		PDI->DrawPoint((FVector)SnapEngine.GetActiveSnapToPoint(), ClosedPolygonColor, 10.0f*PDIScale, SDPG_Foreground);
+
+		PDI->DrawPoint((FVector)SnapEngine.GetActiveSnapFromPoint(), SnapHighlightColor, 15.0f*PDIScale, SDPG_Foreground);
+		PDI->DrawLine((FVector)SnapEngine.GetActiveSnapToPoint(), (FVector)SnapEngine.GetActiveSnapFromPoint(),
+			ClosedPolygonColor, SDPG_Foreground, 0.5f*PDIScale, 0.0f, true);
+		if (SnapEngine.GetActiveSnapTargetID() == CurrentSceneSnapID)
+		{
+			if (LastSnapGeometry.PointCount == 1) {
+				DrawCircle(PDI, (FVector)LastSnapGeometry.Points[0], RenderCameraState.Right(), RenderCameraState.Up(),
+					SnapHighlightColor, ElementSize, 32, SDPG_Foreground, 1.0f*PDIScale, 0.0f, true);
+			}
+			else
+			{
+				PDI->DrawLine((FVector)LastSnapGeometry.Points[0], (FVector)LastSnapGeometry.Points[1],
+					SnapHighlightColor, SDPG_Foreground, 1.0f*PDIScale, 0.0f, true);
+			}
+		}
+		else if (SnapEngine.GetActiveSnapTargetID() == CurrentGridSnapID)
+		{
+			DrawCircle(PDI, (FVector)LastGridSnapPoint, RenderCameraState.Right(), RenderCameraState.Up(),
+				SnapHighlightColor, ElementSize, 4, SDPG_Foreground, 1.0f*PDIScale, 0.0f, true);
+		}
+
+		if (SnapEngine.HaveActiveSnapLine())
+		{
+			// clip this line to the view plane because if it goes through the view plane the pixel-line-thickness
+			// calculation appears to fail
+			FLine3d DrawSnapLine = SnapEngine.GetActiveSnapLine();
+			FVector3d P0(DrawSnapLine.PointAt(-99999.0)), P1(DrawSnapLine.PointAt(99999.0));		// should be smarter here...
+			if (RenderCameraState.bIsOrthographic == false)
+			{
+				UE::Geometry::FPlane3d CameraPlane((FVector3d)RenderCameraState.Forward(), (FVector3d)RenderCameraState.Position + 1.0*(FVector3d)RenderCameraState.Forward());
+				CameraPlane.ClipSegment(P0, P1);
+			}
+			PDI->DrawLine((FVector)P0, (FVector)P1, SnapLineColor, SDPG_Foreground, 0.5 * PDIScale, 0.0f, true);
+
+			if (SnapEngine.HaveActiveSnapDistance())
+			{
+				int iSegment = SnapEngine.GetActiveSnapDistanceID();
+				TArray<FVector3d>& HistoryPoints = (bInFixedPolygonMode) ? FixedPolygonClickPoints : PolygonVertices;
+				FVector3d UseNormal = DrawPlaneOrientation.AxisZ();
+				DrawEdgeTicks(PDI, FSegment3d(HistoryPoints[iSegment], HistoryPoints[iSegment+1]),
+					0.75f*ElementSize, UseNormal, SnapHighlightColor, SDPG_Foreground, 1.0f*PDIScale, true);
+				DrawEdgeTicks(PDI, FSegment3d(HistoryPoints[HistoryPoints.Num()-1], PreviewVertex),
+					0.75f*ElementSize, UseNormal, SnapHighlightColor, SDPG_Foreground, 1.0f*PDIScale, true);
+				// Drawing a highlight
+				PDI->DrawLine((FVector)HistoryPoints[iSegment], (FVector)HistoryPoints[iSegment + 1],
+					SnapHighlightColor, SDPG_Foreground, 2.0f*PDIScale, 1.0f, true);
+			}
+		}
+	}
+
+	//
+	// Draw Surface Hit Indicator
+	//
+	if (bHaveSurfaceHit)
+	{
+		PDI->DrawPoint((FVector)SurfaceHitPoint, ClosedPolygonColor, 10*PDIScale, SDPG_Foreground);
+		if (SnapProperties->HitNormalOffset != 0)
+		{
+			PDI->DrawPoint((FVector)SurfaceOffsetPoint, OpenPolygonColor, 15*PDIScale, SDPG_Foreground);
+			PDI->DrawLine((FVector)SurfaceOffsetPoint, (FVector)SurfaceHitPoint,
+				ClosedPolygonColor, SDPG_Foreground, 0.5f*PDIScale, 0.0f, true);
+		}
+		PDI->DrawLine((FVector)SurfaceOffsetPoint, (FVector)PreviewVertex,
+			ClosedPolygonColor, SDPG_Foreground, 0.5f*PDIScale, 0.0f, true);
+	}
+
+
+	//
+	// Draw the polygon contour preview
+	//
+	if (PolygonVertices.Num() > 0)
+	{
+		FColor UseColor = bIsClosed ? ClosedPolygonColor 
+			: bHaveSelfIntersection ? ErrorColor : OpenPolygonColor;
+		FColor LastSegmentColor = bIsClosed ? ClosedPolygonColor 
+			: bHaveSelfIntersection ? ErrorColor : PreviewColor;
+		FVector3d UseLastVertex = bIsClosed ? PolygonVertices[0] : PreviewVertex;
+		float UseThickness = bHaveSelfIntersection ? SelfIntersectThickness : LineThickness;
+
+		auto DrawVertices = [&PDI, &UseColor](const TArray<FVector3d>& Vertices, ESceneDepthPriorityGroup Group, float Thickness)
+		{
+			for (int lasti = Vertices.Num() - 1, i = 0, NumVertices = Vertices.Num(); i < NumVertices; lasti = i++)
+
+			{
+				PDI->DrawLine((FVector)Vertices[lasti], (FVector)Vertices[i], UseColor, Group, Thickness, 0.0f, true);
+			}
+		};
+
+		// draw thin no-depth (x-ray draw)
+		//DrawVertices(PolygonVertices, SDPG_Foreground, HiddenLineThickness);
+		for (int i = 0; i < NumVerts - 1; ++i)
+		{
+			PDI->DrawLine((FVector)PolygonVertices[i], (FVector)PolygonVertices[i + 1],
+				UseColor, SDPG_Foreground, HiddenLineThickness, 0.0f, true);
+		}
+		PDI->DrawLine((FVector)PolygonVertices[NumVerts - 1], (FVector)UseLastVertex,
+			LastSegmentColor, SDPG_Foreground, HiddenLineThickness, 0.0f, true);
+		for (int HoleIdx = 0; HoleIdx < PolygonHolesVertices.Num(); HoleIdx++)
+		{
+			DrawVertices(PolygonHolesVertices[HoleIdx], SDPG_Foreground, HiddenLineThickness);
+		}
+
+		// draw thick depth-tested
+		//DrawVertices(PolygonVertices, SDPG_World, LineThickness);
+		for (int i = 0; i < NumVerts - 1; ++i)
+		{
+			PDI->DrawLine((FVector)PolygonVertices[i], (FVector)PolygonVertices[i + 1],
+				UseColor, SDPG_World, LineThickness, 0.0f, true);
+		}
+		PDI->DrawLine((FVector)PolygonVertices[NumVerts - 1], (FVector)UseLastVertex,
+			LastSegmentColor, SDPG_World, LineThickness, 0.0f, true);
+		for (int HoleIdx = 0; HoleIdx < PolygonHolesVertices.Num(); HoleIdx++)
+		{
+			DrawVertices(PolygonHolesVertices[HoleIdx], SDPG_World, LineThickness);
+		}
+
+		// Intersection point
+		if (bHaveSelfIntersection && !bInInteractiveExtrude)
+		{
+			PDI->DrawPoint((FVector)SelfIntersectionPoint, SnapHighlightColor, 12*PDIScale, SDPG_Foreground);
+		}
+	}
+
+	// draw preview vertex
+	if (!bInInteractiveExtrude)
+	{
+		PDI->DrawPoint((FVector)PreviewVertex, PreviewColor, 10 * PDIScale, SDPG_Foreground);
+	}
+
+	// draw height preview stuff
+	if (bInInteractiveExtrude)
+	{
+		HeightMechanic->Render(RenderAPI);
+	}
+}
+
+
+void UDrawPolygonTool::ResetPolygon()
+{
+	PolygonVertices.Reset();
+	PolygonHolesVertices.Reset();
+	SnapEngine.Reset();
+	bHaveSurfaceHit = false;
+	bInFixedPolygonMode = false;
+	bHaveSelfIntersection = false;
+	CurrentCurveTimestamp++;
+}
+
+void UDrawPolygonTool::UpdatePreviewVertex(const FVector3d& PreviewVertexIn)
+{
+	PreviewVertex = PreviewVertexIn;
+
+	// update length and angle
+	if (PolygonVertices.Num() > 0)
+	{
+		FVector3d LastVertex = PolygonVertices[PolygonVertices.Num() - 1];
+		SnapProperties->SegmentLength = Distance(LastVertex, PreviewVertex);
+	}
+}
+
+void UDrawPolygonTool::AppendVertex(const FVector3d& Vertex)
+{
+	PolygonVertices.Add(Vertex);
+}
+
+bool UDrawPolygonTool::FindDrawPlaneHitPoint(const FInputDeviceRay& ClickPos, FVector3d& HitPosOut)
+{
+	bHaveSurfaceHit = false;
+
+	FFrame3d Frame(DrawPlaneOrigin, DrawPlaneOrientation);
+	FVector3d HitPos;
+	bool bHit = Frame.RayPlaneIntersection((FVector3d)ClickPos.WorldRay.Origin, (FVector3d)ClickPos.WorldRay.Direction, 2, HitPos);
+	if (bHit == false)
+	{
+		return false;
+	}
+
+	// if we found a scene snap point, add to snap set
+	if (bIgnoreSnappingToggle || SnapProperties->bEnableSnapping == false)
+	{
+		SnapEngine.Reset();
+	}
+	else 
+	{
+		if (SnapProperties->bSnapToWorldGrid)
+		{
+			FVector3d WorldGridSnapPos;
+			if (ToolSceneQueriesUtil::FindWorldGridSnapPoint(this, HitPos, WorldGridSnapPos))
+			{
+				WorldGridSnapPos = Frame.ToPlane(WorldGridSnapPos, 2);
+				SnapEngine.AddPointTarget(WorldGridSnapPos, CurrentGridSnapID, 
+					FBasePositionSnapSolver3::FCustomMetric::Replace(999), SnapEngine.MinInternalPriority() - 5);
+				LastGridSnapPoint = WorldGridSnapPos;
+			}
+		}
+
+		if (SnapProperties->bSnapToVertices || SnapProperties->bSnapToEdges)
+		{
+			FVector3d SceneSnapPos;
+			if (ToolSceneQueriesUtil::FindSceneSnapPoint(this, HitPos, SceneSnapPos, SnapProperties->bSnapToVertices, SnapProperties->bSnapToEdges, 0, &LastSnapGeometry))
+			{
+				SnapEngine.AddPointTarget(SceneSnapPos, CurrentSceneSnapID, SnapEngine.MinInternalPriority() - 10);
+			}
+		}
+
+		TArray<FVector3d>& HistoryPoints = (bInFixedPolygonMode) ? FixedPolygonClickPoints : PolygonVertices;
+		SnapEngine.UpdatePointHistory(HistoryPoints);
+		if (SnapProperties->bSnapToAngles)
+		{
+			SnapEngine.RegenerateTargetLines(true, true);
+		}
+		SnapEngine.bEnableSnapToKnownLengths = SnapProperties->bSnapToLengths;
+	}
+
+	SnapEngine.UpdateSnappedPoint(HitPos);
+
+	// remove scene snap point
+	SnapEngine.RemovePointTargetsByID(CurrentSceneSnapID);
+	SnapEngine.RemovePointTargetsByID(CurrentGridSnapID);
+
+	if (SnapEngine.HaveActiveSnap())
+	{
+		HitPosOut = SnapEngine.GetActiveSnapToPoint();
+		return true;
+	}
+
+
+	// if not snap and we want to hit objects, do that
+	if (SnapProperties->bHitSceneObjects)
+	{
+		FHitResult Result;
+		bool bWorldHit = ToolSceneQueriesUtil::FindNearestVisibleObjectHit(TargetWorld, Result, ClickPos.WorldRay);
+		if (bWorldHit)
+		{
+			bHaveSurfaceHit = true;
+			SurfaceHitPoint = (FVector3d)Result.ImpactPoint;
+			FVector3d UseHitPos = (FVector3d)Result.ImpactPoint + (double)SnapProperties->HitNormalOffset*(FVector3d)Result.Normal;
+			HitPos = Frame.ToPlane(UseHitPos, 2);
+			SurfaceOffsetPoint = UseHitPos;
+		}
+	}
+
+	HitPosOut = HitPos;
+	return true;
+}
+
+void UDrawPolygonTool::OnBeginSequencePreview(const FInputDeviceRay& DevicePos)
+{
+	// just update snapped point preview
+	FVector3d HitPos;
+	if (FindDrawPlaneHitPoint(DevicePos, HitPos))
+	{
+		PreviewVertex = HitPos;
+	}
+	
+}
+
+bool UDrawPolygonTool::CanBeginClickSequence(const FInputDeviceRay& ClickPos)
+{
+	return true;
+}
+
+void UDrawPolygonTool::OnBeginClickSequence(const FInputDeviceRay& ClickPos)
+{
+	ResetPolygon();
+	
+	FVector3d HitPos;
+	bool bHit = FindDrawPlaneHitPoint(ClickPos, HitPos);
+	if (bHit == false)
+	{
+		bAbortActivePolygonDraw = true;
+		return;
+	}
+	if (ToolSceneQueriesUtil::IsPointVisible(CameraState, HitPos) == false)
+	{
+		bAbortActivePolygonDraw = true;
+		return;		// cannot start a poly an a point that is not visible, this is almost certainly an error due to draw plane
+	}
+
+	UpdatePreviewVertex(HitPos);
+
+	bInFixedPolygonMode = (PolygonProperties->PolygonType != EDrawPolygonDrawMode::Freehand);
+	FixedPolygonClickPoints.Reset();
+
+	// Actually process the click.
+	// TODO: This slightly awkward organization is a reflection of an earlier time when
+	// MultiClickSequenceInputBehavior issued a duplicate OnNextSequenceClick() call 
+	// immediately after OnBeginClickSequence(). The code could be cleaned up.
+	OnNextSequenceClick(ClickPos);
+}
+
+void UDrawPolygonTool::OnNextSequencePreview(const FInputDeviceRay& ClickPos)
+{
+	if (bInInteractiveExtrude)
+	{
+		HeightMechanic->UpdateCurrentDistance(ClickPos.WorldRay);
+		PolygonProperties->ExtrudeHeight = HeightMechanic->CurrentHeight;
+		bPreviewUpdatePending = true;
+		return;
+	}
+
+	FVector3d HitPos;
+	bool bHit = FindDrawPlaneHitPoint(ClickPos, HitPos);
+	if (bHit == false)
+	{
+		return;
+	}
+
+	if (bInFixedPolygonMode)
+	{
+		UpdatePreviewVertex(HitPos);
+		bPreviewUpdatePending = true;
+		return;
+	}
+
+	UpdatePreviewVertex(HitPos);
+	UpdateSelfIntersection();
+	if (PolygonVertices.Num() > 2)
+	{
+		bPreviewUpdatePending = true;
+	}
+}
+
+bool UDrawPolygonTool::OnNextSequenceClick(const FInputDeviceRay& ClickPos)
+{
+	if (bInInteractiveExtrude)
+	{
+		EndInteractiveExtrude();
+		return false;
+	}
+
+	FVector3d HitPos;
+	bool bHit = FindDrawPlaneHitPoint(ClickPos, HitPos);
+	if (bHit == false)
+	{
+		return true;  // ignore click but continue accepting clicks
+	}
+
+	// Construct the change now for the undo queue so it reflects the current state.  We might not do anything, in which
+	// case we will not emit the change
+	TUniquePtr<FDrawPolygonStateChange> Change =
+		MakeUnique<FDrawPolygonStateChange>(CurrentCurveTimestamp, FixedPolygonClickPoints, PolygonVertices);
+
+	bool bDonePolygon = false;
+	if (bInFixedPolygonMode)
+	{
+		// ignore very close click points
+		if (FixedPolygonClickPoints.Num() > 0 && ToolSceneQueriesUtil::PointSnapQuery(this, FixedPolygonClickPoints[FixedPolygonClickPoints.Num()-1], HitPos) )
+		{
+			return true;
+		}
+
+		FixedPolygonClickPoints.Add(HitPos);
+		int NumTargetPoints = (PolygonProperties->PolygonType == EDrawPolygonDrawMode::Rectangle || PolygonProperties->PolygonType == EDrawPolygonDrawMode::RoundedRectangle) ? 3 : 2;
+		bDonePolygon = (FixedPolygonClickPoints.Num() == NumTargetPoints);
+		if (bDonePolygon)
+		{
+			GenerateFixedPolygon(FixedPolygonClickPoints, PolygonVertices, PolygonHolesVertices);
+		}
+	}
+	else
+	{
+		// ignore very close click points
+		if (PolygonVertices.Num() > 0 && ToolSceneQueriesUtil::PointSnapQuery(this, PolygonVertices[PolygonVertices.Num()-1], HitPos))
+		{
+			return true;
+		}
+
+		// close polygon if we clicked on start point
+		bDonePolygon = SnapEngine.HaveActiveSnap() && SnapEngine.GetActiveSnapTargetID() == StartPointSnapID;
+
+		if (bHaveSelfIntersection)
+		{
+			// discard vertex in segments before intersection (this is redundant if idx is 0)
+			for (int j = SelfIntersectSegmentIdx; j < PolygonVertices.Num(); ++j)
+			{
+				PolygonVertices[j-SelfIntersectSegmentIdx] = PolygonVertices[j];
+			}
+			PolygonVertices.SetNum(PolygonVertices.Num() - SelfIntersectSegmentIdx);
+			PolygonVertices[0] = PreviewVertex = (FVector3d)SelfIntersectionPoint;
+			bDonePolygon = true;
+		}
+	}
+
+	// emit change event
+	GetToolManager()->EmitObjectChange(this, MoveTemp(Change), LOCTEXT("DrawPolyAddPoint", "Add Point"));
+
+	if (bDonePolygon)
+	{
+		//SnapEngine.Reset();
+		bHaveSurfaceHit = false;
+		if (PolygonProperties->OutputMode == EDrawPolygonOutputMode::ExtrudedInteractive)
+		{
+			BeginInteractiveExtrude();
+
+			PreviewMesh->ClearPreview();
+			PreviewMesh->SetVisible(true);
+
+			return true;
+		}
+		else 
+		{
+			EmitCurrentPolygon();
+
+			PreviewMesh->ClearPreview();
+			PreviewMesh->SetVisible(false);
+
+			return false;
+		}
+	}
+
+	AppendVertex(HitPos);
+
+	// if we are starting a freehand poly, add start point as snap target, but then ignore it until we get 3 verts
+	if (bInFixedPolygonMode == false && PolygonVertices.Num() == 1)
+	{
+		SnapEngine.AddPointTarget(PolygonVertices[0], StartPointSnapID, 1);
+		SnapEngine.AddIgnoreTarget(StartPointSnapID);
+	}
+	if (PolygonVertices.Num() > 2)
+	{
+		SnapEngine.RemoveIgnoreTarget(StartPointSnapID);
+	}
+
+	UpdatePreviewVertex(HitPos);
+	return true;
+}
+
+void UDrawPolygonTool::OnTerminateClickSequence()
+{
+	ResetPolygon();
+}
+
+bool UDrawPolygonTool::RequestAbortClickSequence()
+{
+	if (bAbortActivePolygonDraw)
+	{
+		bAbortActivePolygonDraw = false;
+		return true;
+	}
+	return false;
+}
+
+void UDrawPolygonTool::OnUpdateModifierState(int ModifierID, bool bIsOn)
+{
+	if (ModifierID == IgnoreSnappingModifier)
+	{
+		bIgnoreSnappingToggle = bIsOn;
+	}
+	else if (ModifierID == AngleSnapModifier)
+	{
+
+	}
+}
+
+bool UDrawPolygonTool::UpdateSelfIntersection()
+{
+	bHaveSelfIntersection = false;
+	if (bInFixedPolygonMode || PolygonProperties->bAllowSelfIntersections == true)
+	{
+		return false;
+	}
+
+	int NumVertices = PolygonVertices.Num();
+	if (NumVertices < 3)
+	{
+		return false;
+	}
+
+	FFrame3d DrawFrame(DrawPlaneOrigin, DrawPlaneOrientation);
+	FSegment2d PreviewSegment(DrawFrame.ToPlaneUV(PolygonVertices[NumVertices - 1],2), DrawFrame.ToPlaneUV(PreviewVertex,2));
+
+	double BestIntersectionParameter = FMathd::MaxReal;
+	for (int k = 0; k < NumVertices - 2; ++k) 
+	{
+		FSegment2d Segment(DrawFrame.ToPlaneUV(PolygonVertices[k],2), DrawFrame.ToPlaneUV(PolygonVertices[k + 1],2));
+		FIntrSegment2Segment2d Intersection(PreviewSegment, Segment);
+		if (Intersection.Find()) 
+		{
+			bHaveSelfIntersection = true;
+			if (Intersection.Parameter0 < BestIntersectionParameter)
+			{
+				BestIntersectionParameter = Intersection.Parameter0;
+				SelfIntersectSegmentIdx = k;
+				SelfIntersectionPoint = DrawFrame.FromPlaneUV(Intersection.Point0, 2);
+			}
+		}
+	}
+	return bHaveSelfIntersection;
+}
+
+void UDrawPolygonTool::GetPolygonParametersFromFixedPoints(const TArray<FVector3d>& FixedPoints, FVector2d& FirstReferencePt, FVector2d& BoxSize, double& YSign, double& AngleRad)
+{
+	if (FixedPoints.Num() < 2)
+	{
+		return;
+	}
+
+	FFrame3d DrawFrame(DrawPlaneOrigin, DrawPlaneOrientation);
+	FirstReferencePt = DrawFrame.ToPlaneUV(FixedPoints[0], 2);
+
+	FVector2d EdgePt = DrawFrame.ToPlaneUV(FixedPoints[1], 2);
+	FVector2d Delta = EdgePt - FirstReferencePt;
+	AngleRad = FMathd::Atan2(Delta.Y, Delta.X);
+
+	double Radius = Delta.Length();
+	FVector2d AxisX = Delta / Radius;
+	FVector2d AxisY = -UE::Geometry::PerpCW(AxisX);
+	FVector2d HeightPt = DrawFrame.ToPlaneUV((FixedPoints.Num() == 3) ? FixedPoints[2] : FixedPoints[1], 2);
+	FVector2d HeightDelta = HeightPt - FirstReferencePt;
+	YSign = FMathd::Sign(HeightDelta.Dot(AxisY));
+	BoxSize.X = Radius;
+	BoxSize.Y = FMathd::Abs(HeightDelta.Dot(AxisY));
+}
+
+void UDrawPolygonTool::GenerateFixedPolygon(const TArray<FVector3d>& FixedPoints, TArray<FVector3d>& VerticesOut, TArray<TArray<FVector3d>>& HolesVerticesOut)
+{
+	FVector2d FirstReferencePt, BoxSize;
+	double YSign, AngleRad;
+	GetPolygonParametersFromFixedPoints(FixedPoints, FirstReferencePt, BoxSize, YSign, AngleRad);
+	double Width = BoxSize.X, Height = BoxSize.Y;
+	FMatrix2d RotationMat = FMatrix2d::RotationRad(AngleRad);
+
+	FPolygon2d Polygon;
+	TArray<FPolygon2d> PolygonHoles;
+	if (PolygonProperties->PolygonType == EDrawPolygonDrawMode::Square)
+	{
+		Polygon = FPolygon2d::MakeRectangle(FVector2d::Zero(), 2*Width, 2*Width);
+	}
+	else if (PolygonProperties->PolygonType == EDrawPolygonDrawMode::Rectangle || PolygonProperties->PolygonType == EDrawPolygonDrawMode::RoundedRectangle)
+	{
+		if (PolygonProperties->PolygonType == EDrawPolygonDrawMode::Rectangle)
+		{
+			Polygon = FPolygon2d::MakeRectangle(FVector2d(Width / 2, YSign*Height / 2), Width, Height);
+		}
+		else // PolygonProperties->PolygonType == EDrawPolygonDrawMode::RoundedRectangle
+		{
+			Polygon = FPolygon2d::MakeRoundedRectangle(FVector2d(Width / 2, YSign*Height / 2), Width, Height, FMathd::Min(Width,Height) * FMathd::Clamp(PolygonProperties->FeatureSizeRatio, .01, .99) * .5, PolygonProperties->Steps);
+		}
+	}
+	else // Circle or HoleyCircle
+	{
+		Polygon = FPolygon2d::MakeCircle(Width, PolygonProperties->Steps, 0);
+		if (PolygonProperties->PolygonType == EDrawPolygonDrawMode::HoleyCircle)
+		{
+			PolygonHoles.Add(FPolygon2d::MakeCircle(Width * FMathd::Clamp(PolygonProperties->FeatureSizeRatio, .01, .99), PolygonProperties->Steps, 0));
+		}
+	}
+	Polygon.Transform([RotationMat](const FVector2d& Pt) { return RotationMat * Pt; });
+	for (FPolygon2d& Hole : PolygonHoles)
+	{
+		Hole.Transform([RotationMat](const FVector2d& Pt) { return RotationMat * Pt; });
+	}
+
+	FFrame3d DrawFrame(DrawPlaneOrigin, DrawPlaneOrientation);
+	VerticesOut.SetNum(Polygon.VertexCount());
+	for (int k = 0; k < Polygon.VertexCount(); ++k)
+	{
+		FVector2d NewPt = FirstReferencePt + Polygon[k];
+		VerticesOut[k] = DrawFrame.FromPlaneUV(NewPt, 2);
+	}
+
+	HolesVerticesOut.SetNum(PolygonHoles.Num());
+	for (int HoleIdx = 0; HoleIdx < PolygonHoles.Num(); HoleIdx++)
+	{
+		int NumHoleVerts = PolygonHoles[HoleIdx].VertexCount();
+		HolesVerticesOut[HoleIdx].SetNum(NumHoleVerts);
+		for (int k = 0; k < NumHoleVerts; ++k)
+		{
+			FVector2d NewPt = FirstReferencePt + PolygonHoles[HoleIdx][k];
+			HolesVerticesOut[HoleIdx][k] = DrawFrame.FromPlaneUV(NewPt, 2);
+		}
+	}
+}
+
+
+void UDrawPolygonTool::BeginInteractiveExtrude()
+{
+	bInInteractiveExtrude = true;
+	SnapEngine.ResetActiveSnap();
+
+	HeightMechanic = NewObject<UPlaneDistanceFromHitMechanic>(this);
+	HeightMechanic->Setup(this);
+
+	HeightMechanic->WorldHitQueryFunc = [this](const FRay& WorldRay, FHitResult& HitResult)
+	{
+		if (this->bIgnoreSnappingToggle == false)
+		{
+			return ToolSceneQueriesUtil::FindNearestVisibleObjectHit(TargetWorld, HitResult, WorldRay);
+		}
+		return false;
+	};
+	HeightMechanic->WorldPointSnapFunc = [this](const FVector3d& WorldPos, FVector3d& SnapPos)
+	{
+		if (bIgnoreSnappingToggle == false && SnapProperties->bEnableSnapping && SnapProperties->bSnapToWorldGrid)
+		{
+			return ToolSceneQueriesUtil::FindWorldGridSnapPoint(this, WorldPos, SnapPos);
+		}
+		return false;
+	};
+	HeightMechanic->CurrentHeight = 1.0f;  // initialize to something non-zero...prob should be based on polygon bounds maybe?
+
+	FDynamicMesh3 HeightMesh;
+	FFrame3d WorldMeshFrame;
+	GeneratePolygonMesh(PolygonVertices, PolygonHolesVertices, &HeightMesh, WorldMeshFrame, false, 99999, true);
+	HeightMechanic->Initialize( MoveTemp(HeightMesh), WorldMeshFrame, false);
+
+	ShowExtrudeMessage();
+}
+
+void UDrawPolygonTool::EndInteractiveExtrude()
+{
+	EmitCurrentPolygon();
+
+	PreviewMesh->ClearPreview();
+	PreviewMesh->SetVisible(false);
+
+	bInInteractiveExtrude = false;
+	HeightMechanic = nullptr;
+
+	ShowStartupMessage();
+}
+
+
+
+void UDrawPolygonTool::SetDrawPlaneFromWorldPos(const FVector3d& Position, const FVector3d& Normal)
+{
+	DrawPlaneOrigin = Position;
+
+	FFrame3d DrawPlane(Position, DrawPlaneOrientation);
+	if (bIgnoreSnappingToggle == false)
+	{
+		DrawPlane.AlignAxis(2, Normal);
+		DrawPlane.ConstrainedAlignPerpAxes();
+		DrawPlaneOrientation = DrawPlane.Rotation;
+	}
+
+	SnapEngine.Plane = FFrame3d(DrawPlane.Origin, DrawPlane.Rotation);
+
+	if (PlaneTransformGizmo != nullptr)
+	{
+		PlaneTransformGizmo->SetNewGizmoTransform(FTransform((FQuat)DrawPlaneOrientation, (FVector)DrawPlaneOrigin));
+	}
+}
+
+
+void UDrawPolygonTool::PlaneTransformChanged(UTransformProxy* Proxy, FTransform Transform)
+{
+	DrawPlaneOrientation = (FQuaterniond)Transform.GetRotation();
+	DrawPlaneOrigin = (FVector3d)Transform.GetLocation();
+	SnapEngine.Plane = FFrame3d(DrawPlaneOrigin, DrawPlaneOrientation);
+}
+
+void UDrawPolygonTool::UpdateShowGizmoState(bool bNewVisibility)
+{
+	if (bNewVisibility == false)
+	{
+		GetToolManager()->GetPairedGizmoManager()->DestroyAllGizmosByOwner(this);
+		PlaneTransformGizmo = nullptr;
+	}
+	else
+	{
+		if (!PlaneTransformGizmo) {
+			PlaneTransformGizmo = UE::TransformGizmoUtil::CreateCustomTransformGizmo(GetToolManager(),
+				ETransformGizmoSubElements::StandardTranslateRotate, this);
+		}
+		PlaneTransformGizmo->SetActiveTarget(PlaneTransformProxy, GetToolManager());
+		PlaneTransformGizmo->ReinitializeGizmoTransform(FTransform((FQuat)DrawPlaneOrientation, (FVector)DrawPlaneOrigin));
+	}
+}
+
+
+void UDrawPolygonTool::EmitCurrentPolygon()
+{
+	FString BaseName = (PolygonProperties->OutputMode == EDrawPolygonOutputMode::MeshedPolygon) ?
+		TEXT("Polygon") : TEXT("Extrude");
+
+	// generate new mesh
+	FFrame3d PlaneFrameOut;
+	FDynamicMesh3 Mesh;
+	double ExtrudeDist = (PolygonProperties->OutputMode == EDrawPolygonOutputMode::MeshedPolygon) ?
+		0 : PolygonProperties->ExtrudeHeight;
+	bool bSucceeded = GeneratePolygonMesh(PolygonVertices, PolygonHolesVertices, &Mesh, PlaneFrameOut, false, ExtrudeDist, false);
+	if (!bSucceeded) // somehow made a polygon with no valid triangulation; just throw it away ...
+	{
+		ResetPolygon();
+		return;
+	}
+
+	GetToolManager()->BeginUndoTransaction(LOCTEXT("CreatePolygon", "Create Polygon"));
+
+	FCreateMeshObjectParams NewMeshObjectParams;
+	NewMeshObjectParams.TargetWorld = TargetWorld;
+	NewMeshObjectParams.Transform = PlaneFrameOut.ToFTransform();
+	NewMeshObjectParams.BaseName = BaseName;
+	NewMeshObjectParams.Materials.Add(MaterialProperties->Material.Get());
+	NewMeshObjectParams.SetMesh(&Mesh);
+	OutputTypeProperties->ConfigureCreateMeshObjectParams(NewMeshObjectParams);
+	FCreateMeshObjectResult Result = UE::Modeling::CreateMeshObject(GetToolManager(), MoveTemp(NewMeshObjectParams));
+	if (Result.IsOK() && Result.NewActor != nullptr)
+	{
+		ToolSelectionUtil::SetNewActorSelection(GetToolManager(), Result.NewActor);
+	}
+
+	GetToolManager()->EndUndoTransaction();
+
+	ResetPolygon();
+}
+
+void UDrawPolygonTool::UpdateLivePreview()
+{
+	int NumVerts = PolygonVertices.Num();
+	if (NumVerts < 2 || PreviewMesh == nullptr || PreviewMesh->IsVisible() == false)
+	{
+		return;
+	}
+
+	FFrame3d PlaneFrame;
+	FDynamicMesh3 Mesh;
+	double ExtrudeDist = (PolygonProperties->OutputMode == EDrawPolygonOutputMode::MeshedPolygon) ?
+		0 : PolygonProperties->ExtrudeHeight;
+	if (GeneratePolygonMesh(PolygonVertices, PolygonHolesVertices, &Mesh, PlaneFrame, false, ExtrudeDist, false))
+	{
+		PreviewMesh->SetTransform(PlaneFrame.ToFTransform());
+		PreviewMesh->SetMaterial(MaterialProperties->Material.Get());
+		PreviewMesh->EnableWireframe(MaterialProperties->bWireframe);
+		PreviewMesh->UpdatePreview(&Mesh);
+	}
+}
+
+bool UDrawPolygonTool::GeneratePolygonMesh(const TArray<FVector3d>& Polygon, const TArray<TArray<FVector3d>>& PolygonHoles, FDynamicMesh3* ResultMeshOut, FFrame3d& WorldFrameOut, bool bIncludePreviewVtx, double ExtrudeDistance, bool bExtrudeSymmetric)
+{
+	// construct centered frame for polygon
+	WorldFrameOut = FFrame3d(DrawPlaneOrigin, DrawPlaneOrientation);
+
+	int NumVerts = Polygon.Num();
+	FVector3d Centroid3d(0, 0, 0);
+	for (int k = 0; k < NumVerts; ++k)
+	{
+		Centroid3d += Polygon[k];
+	}
+	Centroid3d /= (double)NumVerts;
+	FVector2d CentroidInDrawPlane = WorldFrameOut.ToPlaneUV(Centroid3d);
+	WorldFrameOut.Origin = Centroid3d;
+
+	// Compute outer polygon & bounds
+	auto VertexArrayToPolygon = [&WorldFrameOut](const TArray<FVector3d>& Vertices)
+	{
+		FPolygon2d OutPolygon;
+		for (int k = 0, N = Vertices.Num(); k < N; ++k)
+		{
+			OutPolygon.AppendVertex(WorldFrameOut.ToPlaneUV(Vertices[k], 2));
+		}
+		return OutPolygon;
+	};
+	FPolygon2d OuterPolygon = VertexArrayToPolygon(Polygon);
+	// add preview vertex
+	if (bIncludePreviewVtx)
+	{
+		if (Distance(PreviewVertex, Polygon[NumVerts-1]) > 0.1)
+		{
+			OuterPolygon.AppendVertex(WorldFrameOut.ToPlaneUV(PreviewVertex, 2));
+		}
+	}
+	FAxisAlignedBox2d Bounds(OuterPolygon.Bounds());
+
+	// special case paths
+	if (PolygonProperties->PolygonType == EDrawPolygonDrawMode::HoleyCircle || PolygonProperties->PolygonType == EDrawPolygonDrawMode::Circle || PolygonProperties->PolygonType == EDrawPolygonDrawMode::RoundedRectangle)
+	{
+		// get polygon parameters
+		FVector2d FirstReferencePt, BoxSize;
+		double YSign, AngleRad;
+		GetPolygonParametersFromFixedPoints(FixedPolygonClickPoints, FirstReferencePt, BoxSize, YSign, AngleRad);
+		FirstReferencePt -= CentroidInDrawPlane;
+		FMatrix2d RotationMat = FMatrix2d::RotationRad(AngleRad);
+
+		// translate general polygon parameters to specific mesh generator parameters, and generate mesh
+		if (PolygonProperties->PolygonType == EDrawPolygonDrawMode::HoleyCircle)
+		{
+			FPuncturedDiscMeshGenerator HCGen;
+			HCGen.AngleSamples = PolygonProperties->Steps;
+			HCGen.RadialSamples = 1;
+			HCGen.Radius = BoxSize.X;
+			HCGen.HoleRadius = BoxSize.X * FMathd::Clamp(PolygonProperties->FeatureSizeRatio, .01f, .99f);
+			ResultMeshOut->Copy(&HCGen.Generate());
+		}
+		else if (PolygonProperties->PolygonType == EDrawPolygonDrawMode::Circle)
+		{
+			FDiscMeshGenerator CGen;
+			CGen.AngleSamples = PolygonProperties->Steps;
+			CGen.RadialSamples = 1;
+			CGen.Radius = BoxSize.X;
+			ResultMeshOut->Copy(&CGen.Generate());
+		}
+		else if (PolygonProperties->PolygonType == EDrawPolygonDrawMode::RoundedRectangle)
+		{
+			FRoundedRectangleMeshGenerator RRGen;
+			FirstReferencePt += RotationMat * (FVector2d(BoxSize.X, BoxSize.Y * YSign)*.5f);
+			RRGen.AngleSamples = PolygonProperties->Steps;
+			RRGen.Radius = .5 * FMathd::Min(BoxSize.X, BoxSize.Y) * FMathd::Clamp(PolygonProperties->FeatureSizeRatio, .01f, .99f);
+			RRGen.Height = BoxSize.Y - RRGen.Radius * 2.;
+			RRGen.Width = BoxSize.X - RRGen.Radius * 2.;
+			RRGen.WidthVertexCount = 1;
+			RRGen.HeightVertexCount = 1;
+			ResultMeshOut->Copy(&RRGen.Generate());
+		}
+
+		// transform generated mesh
+		for (int VertIdx : ResultMeshOut->VertexIndicesItr())
+		{
+			FVector3d V = ResultMeshOut->GetVertex(VertIdx);
+			FVector2d VTransformed = RotationMat * FVector2d(V.X, V.Y) + FirstReferencePt;
+			ResultMeshOut->SetVertex(VertIdx, FVector3d(VTransformed.X, VTransformed.Y, 0));
+		}
+	}
+	else // generic path: triangulate using polygon vertices
+	{
+		// triangulate polygon into the MeshDescription
+		FGeneralPolygon2d GeneralPolygon;
+		FFlatTriangulationMeshGenerator TriangulationMeshGen;
+
+		if (OuterPolygon.IsClockwise() == false)
+		{
+			OuterPolygon.Reverse();
+		}
+
+		GeneralPolygon.SetOuter(OuterPolygon);
+
+		for (int HoleIdx = 0; HoleIdx < PolygonHoles.Num(); HoleIdx++)
+		{
+			// attempt to add holes (skipping if safety checks fail)
+			GeneralPolygon.AddHole(VertexArrayToPolygon(PolygonHoles[HoleIdx]), true, false /*currently don't care about hole orientation; we'll just set the triangulation algo not to care*/);
+		}
+
+		FConstrainedDelaunay2d Triangulator;
+		if (PolygonProperties->bAllowSelfIntersections)
+		{
+			FArrangement2d Arrangement(OuterPolygon.Bounds());
+			// arrangement2d builds a general 2d graph that discards orientation info ...
+			Triangulator.FillRule = FConstrainedDelaunay2d::EFillRule::Odd;
+			Triangulator.bOrientedEdges = false;
+			Triangulator.bSplitBowties = true;
+			for (FSegment2d Seg : GeneralPolygon.GetOuter().Segments())
+			{
+				Arrangement.Insert(Seg);
+			}
+			Triangulator.Add(Arrangement.Graph);
+			for (const FPolygon2d& Hole : GeneralPolygon.GetHoles())
+			{
+				Triangulator.Add(Hole, true);
+			}
+		}
+		else
+		{
+			Triangulator.Add(GeneralPolygon);
+		}
+
+
+		bool bTriangulationSuccess = Triangulator.Triangulate([&GeneralPolygon](const TArray<FVector2d>& Vertices, FIndex3i Tri)
+		{
+			// keep triangles based on the input polygon's winding
+			return GeneralPolygon.Contains((Vertices[Tri.A] + Vertices[Tri.B] + Vertices[Tri.C]) / 3.0);
+		});
+		// only truly fail if we got zero triangles back from the triangulator; if it just returned false it may still have managed to partially generate something
+		if (Triangulator.Triangles.Num() == 0)
+		{
+			return false;
+		}
+
+		TriangulationMeshGen.Vertices2D = Triangulator.Vertices;
+		TriangulationMeshGen.Triangles2D = Triangulator.Triangles;
+
+		ResultMeshOut->Copy(&TriangulationMeshGen.Generate());
+	}
+
+	// for symmetric extrude we translate the first poly by -dist along axis
+	if (bExtrudeSymmetric)
+	{
+		FVector3d ShiftNormal = FVector3d::UnitZ();
+		for (int vid : ResultMeshOut->VertexIndicesItr())
+		{
+			FVector3d Pos = ResultMeshOut->GetVertex(vid);
+			ResultMeshOut->SetVertex(vid, Pos - ExtrudeDistance * ShiftNormal);
+		}
+		// double extrude dist
+		ExtrudeDistance *= 2.0;
+	}
+
+	if (ExtrudeDistance != 0)
+	{
+		FExtrudeMesh Extruder(ResultMeshOut);
+		Extruder.DefaultExtrudeDistance = ExtrudeDistance;
+		
+		Extruder.UVScaleFactor = 1.0 / Bounds.MaxDim();
+		if (ExtrudeDistance < 0)
+		{
+			Extruder.IsPositiveOffset = false;
+		}
+
+		FVector3d ExtrudeNormal = FVector3d::UnitZ();
+		Extruder.ExtrudedPositionFunc = [&ExtrudeDistance, &ExtrudeNormal](const FVector3d& Position, const FVector3f& Normal, int VertexID)
+		{
+			return Position + ExtrudeDistance * ExtrudeNormal;
+		};
+
+		Extruder.Apply();
+	}
+
+	FDynamicMeshEditor Editor(ResultMeshOut);
+	float InitialUVScale = 1.0 / Bounds.MaxDim(); // this is the UV scale used by both the polymeshgen and the extruder above
+	// default global rescale -- initial scale doesn't factor in extrude distance; rescale so UVScale of 1.0 fits in the unit square texture
+	float GlobalUVRescale = MaterialProperties->UVScale / FMathf::Max(1.0f, FMathd::Abs(ExtrudeDistance) * InitialUVScale);
+	if (MaterialProperties->bWorldSpaceUVScale)
+	{
+		// since we know the initial uv scale, directly compute the global scale (relative to 1 meter as a standard scale)
+		GlobalUVRescale = MaterialProperties->UVScale * .01 / InitialUVScale;
+	}
+	Editor.RescaleAttributeUVs(GlobalUVRescale, false);
+
+	return true;
+}
+
+
+
+
+void UDrawPolygonTool::ShowStartupMessage()
+{
+	GetToolManager()->DisplayMessage(
+		LOCTEXT("OnStartDraw", "Use this Tool to draw a polygon on the Drawing Plane, and Extrude it. Left-click to place points. Ctrl-click on the scene to reposition the Plane (Shift+Ctrl-click to ignore Normal). [A] toggles Gizmo. Hold Shift to ignore Snapping."),
+		EToolMessageLevel::UserNotification);
+}
+
+void UDrawPolygonTool::ShowExtrudeMessage()
+{
+	GetToolManager()->DisplayMessage(
+		LOCTEXT("OnStartExtrude", "Set the height of the Extrusion by positioning the mouse over the extrusion volume, or over the scene to snap to relative heights. Hold Shift to ignore Snapping."),
+		EToolMessageLevel::UserNotification);
+}
+
+
+
+
+void UDrawPolygonTool::UndoCurrentOperation(const TArray<FVector3d>& ClickPointsIn, const TArray<FVector3d>& PolygonVerticesIn)
+{
+	if (bInInteractiveExtrude)
+	{
+		PreviewMesh->ClearPreview();
+		PreviewMesh->SetVisible(false);
+		bInInteractiveExtrude = false;
+	}
+	ApplyUndoPoints(ClickPointsIn, PolygonVerticesIn);
+}
+
+
+void FDrawPolygonStateChange::Revert(UObject* Object)
+{
+	Cast<UDrawPolygonTool>(Object)->UndoCurrentOperation( FixedVertexPoints, PolyPoints );
+	bHaveDoneUndo = true;
+}
+bool FDrawPolygonStateChange::HasExpired(UObject* Object) const
+{
+	return bHaveDoneUndo || (Cast<UDrawPolygonTool>(Object)->CheckInCurve(CurveTimestamp) == false);
+}
+FString FDrawPolygonStateChange::ToString() const
+{
+	return TEXT("FDrawPolygonStateChange");
+}
+
+
+
+#undef LOCTEXT_NAMESPACE
diff --git a/Engine/Plugins/Runtime/MeshModelingToolset/Source/MeshModelingTools/Public/CombineMeshesTool.h b/Engine/Plugins/Runtime/MeshModelingToolset/Source/MeshModelingTools/Public/CombineMeshesTool.h
new file mode 100644
--- /dev/null
+++ b/Engine/Plugins/Runtime/MeshModelingToolset/Source/MeshModelingTools/Public/CombineMeshesTool.h
@@ -0,0 +1,116 @@
+// Copyright Epic Games, Inc. All Rights Reserved.
+
+#pragma once
+
+#include "CoreMinimal.h"
+#include "MultiSelectionTool.h"
+#include "InteractiveToolBuilder.h"
+#include "MeshOpPreviewHelpers.h"
+#include "DynamicMesh/DynamicMesh3.h"
+#include "PropertySets/OnAcceptProperties.h"
+#include "PropertySets/CreateMeshObjectTypeProperties.h"
+#include "CombineMeshesTool.generated.h"
+
+// predeclarations
+struct FMeshDescription;
+
+/**
+ *
+ */
+UCLASS()
+class MESHMODELINGTOOLS_API UCombineMeshesToolBuilder : public UInteractiveToolBuilder
+{
+	GENERATED_BODY()
+
+public:
+	bool bIsDuplicateTool = false;
+
+	virtual bool CanBuildTool(const FToolBuilderState& SceneState) const override;
+	virtual UInteractiveTool* BuildTool(const FToolBuilderState& SceneState) const override;
+
+protected:
+	virtual const FToolTargetTypeRequirements& GetTargetRequirements() const override;
+};
+
+
+
+
+UENUM()
+enum class ECombineTargetType
+{
+	NewAsset,
+	FirstInputAsset,
+	LastInputAsset
+};
+
+
+/**
+ * Standard properties
+ */
+UCLASS()
+class MESHMODELINGTOOLS_API UCombineMeshesToolProperties : public UInteractiveToolPropertySet
+{
+	GENERATED_BODY()
+
+public:
+	UPROPERTY(meta = (TransientToolProperty))
+	bool bIsDuplicateMode = false;
+
+	UPROPERTY(EditAnywhere, Category = AssetOptions, meta = (EditCondition = "bIsDuplicateMode == false", EditConditionHides, HideEditConditionToggle))
+	ECombineTargetType WriteOutputTo = ECombineTargetType::NewAsset;
+
+	/** Base name for newly-generated asset */
+	UPROPERTY(EditAnywhere, Category = AssetOptions, meta = (TransientToolProperty, EditCondition = "bIsDuplicateMode || WriteOutputTo == ECombineTargetType::NewAsset", EditConditionHides))
+	FString OutputName;
+
+	/** Name of asset that will be updated */
+	UPROPERTY(VisibleAnywhere, Category = AssetOptions, meta = (TransientToolProperty, EditCondition = "bIsDuplicateMode == false && WriteOutputTo != ECombineTargetType::NewAsset", EditConditionHides))
+	FString OutputAsset;
+};
+
+
+
+
+
+
+/**
+ * Simple tool to combine multiple meshes into a single mesh asset
+ */
+UCLASS()
+class MESHMODELINGTOOLS_API UCombineMeshesTool : public UMultiSelectionTool
+{
+	GENERATED_BODY()
+
+public:
+	virtual void SetDuplicateMode(bool bDuplicateMode);
+
+	virtual void Setup() override;
+	virtual void Shutdown(EToolShutdownType ShutdownType) override;
+
+	virtual void SetWorld(UWorld* World);
+
+	virtual bool HasCancel() const override { return true; }
+	virtual bool HasAccept() const override { return true; }
+
+protected:
+
+	UPROPERTY()
+	TObjectPtr<UCombineMeshesToolProperties> BasicProperties;
+
+	UPROPERTY()
+	TObjectPtr<UCreateMeshObjectTypeProperties> OutputTypeProperties;
+
+	UPROPERTY()
+	TObjectPtr<UOnAcceptHandleSourcesProperties> HandleSourceProperties;
+
+protected:
+	UWorld* TargetWorld;
+
+	bool bDuplicateMode;
+
+	void CreateNewAsset();
+	void UpdateExistingAsset();
+
+
+	void BuildCombinedMaterialSet(TArray<UMaterialInterface*>& NewMaterialsOut, TArray<TArray<int32>>& MaterialIDRemapsOut);
+};
diff --git a/Engine/Plugins/Runtime/MeshModelingToolset/Source/MeshModelingTools/Public/DrawAndRevolveTool.h b/Engine/Plugins/Runtime/MeshModelingToolset/Source/MeshModelingTools/Public/DrawAndRevolveTool.h
new file mode 100644
--- /dev/null
+++ b/Engine/Plugins/Runtime/MeshModelingToolset/Source/MeshModelingTools/Public/DrawAndRevolveTool.h
@@ -0,0 +1,142 @@
+// Copyright Epic Games, Inc. All Rights Reserved.
+
+#pragma once
+
+#include "CoreMinimal.h"
+
+#include "BaseBehaviors/BehaviorTargetInterfaces.h"
+#include "ModelingOperators.h" //IDynamicMeshOperatorFactory
+#include "InteractiveTool.h" //UInteractiveToolPropertySet
+#include "InteractiveToolBuilder.h" //UInteractiveToolBuilder
+#include "InteractiveToolChange.h" //FToolCommandChange
+#include "MeshOpPreviewHelpers.h" //FDynamicMeshOpResult
+#include "Properties/MeshMaterialProperties.h"
+#include "PropertySets/CreateMeshObjectTypeProperties.h"
+#include "Properties/RevolveProperties.h"
+
+#include "DrawAndRevolveTool.generated.h"
+
+class UCollectSurfacePathMechanic;
+class UConstructionPlaneMechanic;
+class UCurveControlPointsMechanic;
+
+UCLASS()
+class MESHMODELINGTOOLS_API UDrawAndRevolveToolBuilder : public UInteractiveToolBuilder
+{
+	GENERATED_BODY()
+
+public:
+	virtual bool CanBuildTool(const FToolBuilderState& SceneState) const override;
+	virtual UInteractiveTool* BuildTool(const FToolBuilderState& SceneState) const override;
+};
+
+
+UCLASS()
+class MESHMODELINGTOOLS_API URevolveToolProperties : public URevolveProperties
+{
+	GENERATED_BODY()
+
+public:
+
+	/** Connect the ends of an open profile to the axis to close the top and bottom of the revolved result. Not relevant if profile curve is closed. */
+	UPROPERTY(EditAnywhere, Category = RevolveSettings, AdvancedDisplay)
+	bool bConnectOpenProfileToAxis = true;
+	
+	/** Sets the draw plane origin. The revolution axis is the X axis in the plane. */
+	UPROPERTY(EditAnywhere, Category = DrawPlane, meta = (
+		EditCondition = "bAllowedToEditDrawPlane", HideEditConditionToggle))
+	FVector DrawPlaneOrigin = FVector(0, 0, 0);
+
+	/** Sets the draw plane orientation. The revolution axis is the X axis in the plane. */
+	UPROPERTY(EditAnywhere, Category = DrawPlane, meta = (
+		EditCondition = "bAllowedToEditDrawPlane", HideEditConditionToggle, 
+		UIMin = -180, UIMax = 180, ClampMin = -180000, ClampMax = 18000))
+	FRotator DrawPlaneOrientation = FRotator(90, 0, 0);
+
+	/** Enables/disables snapping while editing the profile curve. Also toggled with Shift. */
+	UPROPERTY(EditAnywhere, Category = ProfileCurve)
+	bool bEnableSnapping = true;
+
+	// Not user visible- used to disallow draw plane modification.
+	UPROPERTY(meta = (TransientToolProperty))
+	bool bAllowedToEditDrawPlane = true;
+};
+
+UCLASS()
+class MESHMODELINGTOOLS_API URevolveOperatorFactory : public UObject, public UE::Geometry::IDynamicMeshOperatorFactory
+{
+	GENERATED_BODY()
+
+public:
+	// IDynamicMeshOperatorFactory API
+	virtual TUniquePtr<UE::Geometry::FDynamicMeshOperator> MakeNewOperator() override;
+
+	UPROPERTY()
+	TObjectPtr<UDrawAndRevolveTool> RevolveTool;
+};
+
+
+/** Draws a profile curve and revolves it around an axis. */
+UCLASS()
+class MESHMODELINGTOOLS_API UDrawAndRevolveTool : public UInteractiveTool
+{
+	GENERATED_BODY()
+
+public:
+	virtual void SetWorld(UWorld* World) { TargetWorld = World; }
+
+	virtual void RegisterActions(FInteractiveToolActionSet& ActionSet) override;
+	virtual void OnPointDeletionKeyPress();
+
+	virtual bool HasCancel() const override { return true; }
+	virtual bool HasAccept() const override { return true; }
+	virtual bool CanAccept() const override;
+
+	virtual void Setup() override;
+	virtual void Shutdown(EToolShutdownType ShutdownType) override;
+
+	virtual void OnTick(float DeltaTime) override;
+	virtual void Render(IToolsContextRenderAPI* RenderAPI) override;
+
+	virtual void OnPropertyModified(UObject* PropertySet, FProperty* Property) override;
+
+protected:
+
+	UWorld* TargetWorld;
+
+	FViewCameraState CameraState;
+
+	// This information is replicated in the user-editable transform in the settings and in the PlaneMechanic
+	// plane, but the tool turned out to be much easier to write and edit with this decoupling.
+	FVector3d RevolutionAxisOrigin;
+	FVector3d RevolutionAxisDirection;
+
+	bool bProfileCurveComplete = false;
+
+	void UpdateRevolutionAxis();
+
+	UPROPERTY()
+	TObjectPtr<UCurveControlPointsMechanic> ControlPointsMechanic = nullptr;
+
+	UPROPERTY()
+	TObjectPtr<UConstructionPlaneMechanic> PlaneMechanic = nullptr;
+
+	/** Property set for type of output object (StaticMesh, Volume, etc) */
+	UPROPERTY()
+	TObjectPtr<UCreateMeshObjectTypeProperties> OutputTypeProperties;
+
+	UPROPERTY()
+	TObjectPtr<URevolveToolProperties> Settings = nullptr;
+
+	UPROPERTY()
+	TObjectPtr<UNewMeshMaterialProperties> MaterialProperties;
+
+	UPROPERTY()
+	TObjectPtr<UMeshOpPreviewWithBackgroundCompute> Preview = nullptr;
+
+	void StartPreview();
+
+	void GenerateAsset(const FDynamicMeshOpResult& Result);
+
+	friend class URevolveOperatorFactory;
+};
diff --git a/Engine/Plugins/Runtime/MeshModelingToolset/Source/MeshModelingTools/Public/DrawPolygonTool.h b/Engine/Plugins/Runtime/MeshModelingToolset/Source/MeshModelingTools/Public/DrawPolygonTool.h
new file mode 100644
--- /dev/null
+++ b/Engine/Plugins/Runtime/MeshModelingToolset/Source/MeshModelingTools/Public/DrawPolygonTool.h
@@ -0,0 +1,352 @@
+// Copyright Epic Games, Inc. All Rights Reserved.
+
+#pragma once
+
+#include "CoreMinimal.h"
+#include "UObject/NoExportTypes.h"
+#include "InteractiveGizmo.h"
+#include "InteractiveToolBuilder.h"
+#include "BaseTools/SingleClickTool.h"
+#include "PreviewMesh.h"
+#include "DynamicMesh/DynamicMeshAABBTree3.h"
+#include "Snapping/PointPlanarSnapSolver.h"
+#include "ToolSceneQueriesUtil.h"
+#include "Properties/MeshMaterialProperties.h"
+#include "PropertySets/CreateMeshObjectTypeProperties.h"
+#include "Mechanics/PlaneDistanceFromHitMechanic.h"
+#include "DrawPolygonTool.generated.h"
+
+
+class UCombinedTransformGizmo;
+class UTransformProxy;
+PREDECLARE_USE_GEOMETRY_CLASS(FDynamicMesh3);
+
+/**
+ *
+ */
+UCLASS()
+class MESHMODELINGTOOLS_API UDrawPolygonToolBuilder : public UInteractiveToolBuilder
+{
+	GENERATED_BODY()
+
+public:
+	virtual bool CanBuildTool(const FToolBuilderState& SceneState) const override;
+	virtual UInteractiveTool* BuildTool(const FToolBuilderState& SceneState) const override;
+};
+
+
+
+/** Polygon Tool Draw Type */
+UENUM()
+enum class EDrawPolygonDrawMode : uint8
+{
+	/** Freehand Polygon Drawing */
+	Freehand UMETA(DisplayName = "Freehand"),
+
+	/** Circle */
+	Circle UMETA(DisplayName = "Circle"),
+
+	/** Square */
+	Square UMETA(DisplayName = "Square"),
+
+	/** Rectangle */
+	Rectangle UMETA(DisplayName = "Rectangle"),
+
+	/** Rounded Rectangle */
+	RoundedRectangle UMETA(DisplayName = "Rounded Rectangle"),
+
+	/** Circle w/ Hole */
+	HoleyCircle UMETA(DisplayName = "Circle w/ Hole")
+
+};
+
+
+/** Output of Draw Polygon Tool */
+UENUM()
+enum class EDrawPolygonOutputMode : uint8
+{
+	/** Generate a meshed planar polygon */
+	MeshedPolygon UMETA(DisplayName = "Flat Mesh"),
+
+	/** Extrude closed polygon to constant height determined by Extrude Height Property */
+	ExtrudedConstant UMETA(DisplayName = "Extrude To Height"),
+
+	/** Extrusion height is set via additional mouse input after closing polygon */
+	ExtrudedInteractive UMETA(DisplayName = "Interactive Extrude"),
+};
+
+
+
+
+
+UCLASS()
+class MESHMODELINGTOOLS_API UDrawPolygonToolStandardProperties : public UInteractiveToolPropertySet
+{
+	GENERATED_BODY()
+
+public:
+	UDrawPolygonToolStandardProperties();
+
+	UPROPERTY(EditAnywhere, NonTransactional, Category = Polygon)
+	EDrawPolygonDrawMode PolygonType = EDrawPolygonDrawMode::Freehand;
+
+	/** Feature size as fraction of overall shape size, for shapes with secondary features like the rounded corners of a Rounded Rectangle */
+	UPROPERTY(EditAnywhere, NonTransactional, Category = Polygon, meta = (UIMin = "0.01", UIMax = "0.99", ClampMin = "0.01", ClampMax = "0.99",
+									EditCondition = "PolygonType == EDrawPolygonDrawMode::RoundedRectangle || PolygonType == EDrawPolygonDrawMode::HoleyCircle"))
+	float FeatureSizeRatio = .25;
+
+	/** Number of vertices in round features */
+	UPROPERTY(EditAnywhere, NonTransactional, Category = Polygon, meta = (UIMin = "3", UIMax = "100", ClampMin = "3", ClampMax = "10000",
+				  EditCondition = "PolygonType == EDrawPolygonDrawMode::Circle || PolygonType == EDrawPolygonDrawMode::RoundedRectangle || PolygonType == EDrawPolygonDrawMode::HoleyCircle"))
+	int Steps = 16;
+
+	UPROPERTY(EditAnywhere, NonTransactional, Category = Polygon, meta = (EditCondition = "PolygonType == EDrawPolygonDrawMode::Freehand"))
+	bool bAllowSelfIntersections = false;
+
+	UPROPERTY(EditAnywhere, NonTransactional, Category = Polygon)
+	EDrawPolygonOutputMode OutputMode = EDrawPolygonOutputMode::ExtrudedInteractive;
+
+	/** Extrusion Distance in non-interactive mode */
+	UPROPERTY(EditAnywhere, NonTransactional, Category = Polygon, meta = (UIMin = "-1000", UIMax = "1000", ClampMin = "-10000", ClampMax = "10000",
+									EditCondition = "OutputMode == EDrawPolygonOutputMode::ExtrudedConstant"))
+	float ExtrudeHeight = 100.0f;
+
+	UPROPERTY(EditAnywhere, NonTransactional, Category = Polygon)
+	bool bShowGizmo = true;
+};
+
+UCLASS()
+class MESHMODELINGTOOLS_API UDrawPolygonToolSnapProperties : public UInteractiveToolPropertySet
+{
+	GENERATED_BODY()
+
+public:
+	UPROPERTY(EditAnywhere, NonTransactional, Category = Snapping)
+	bool bEnableSnapping = true;
+
+	/** When true, allows snapping to world grid according to editor settings. */
+	UPROPERTY(EditAnywhere, NonTransactional, Category = Snapping, Meta = (EditCondition = "bEnableSnapping", EditConditionHides))
+	bool bSnapToWorldGrid = false;
+
+	UPROPERTY(EditAnywhere, NonTransactional, Category = Snapping, Meta = (EditCondition = "bEnableSnapping", EditConditionHides))
+	bool bSnapToVertices = true;
+
+	UPROPERTY(EditAnywhere, NonTransactional, Category = Snapping, Meta = (EditCondition = "bEnableSnapping", EditConditionHides))
+	bool bSnapToEdges = false;
+
+	UPROPERTY(EditAnywhere, NonTransactional, Category = Snapping, Meta = (EditCondition = "bEnableSnapping", EditConditionHides))
+	bool bSnapToAngles = true;
+
+	/** When snapping to an angle, also prefer to snap to an existing segment length */
+	UPROPERTY(EditAnywhere, NonTransactional, Category = Snapping, Meta = (EditCondition = "bEnableSnapping && bSnapToAngles", EditConditionHides))
+	bool bSnapToLengths = true;
+
+	UPROPERTY(VisibleAnywhere, NonTransactional, Category = Snapping, meta = (TransientToolProperty))
+	float SegmentLength = 0.0f;
+
+	UPROPERTY(EditAnywhere, NonTransactional, Category = Snapping)
+	bool bHitSceneObjects = false;
+
+	UPROPERTY(EditAnywhere, NonTransactional, Category = Snapping, Meta = (EditCondition = "bHitSceneObjects", EditConditionHides))
+	float HitNormalOffset = 0.0f;
+};
+
+
+
+
+
+/**
+ * This tool allows the user to draw and extrude 2D polygons
+ */
+UCLASS()
+class MESHMODELINGTOOLS_API UDrawPolygonTool : public UInteractiveTool, public IClickSequenceBehaviorTarget
+{
+	GENERATED_BODY()
+public:
+	UDrawPolygonTool();
+
+	virtual void RegisterActions(FInteractiveToolActionSet& ActionSet) override;
+
+	virtual void SetWorld(UWorld* World);
+
+	virtual void Setup() override;
+	virtual void Shutdown(EToolShutdownType ShutdownType) override;
+
+	virtual void OnTick(float DeltaTime) override;
+	virtual void Render(IToolsContextRenderAPI* RenderAPI) override;
+
+	virtual bool HasCancel() const override { return false; }
+	virtual bool HasAccept() const override { return false; }
+	virtual bool CanAccept() const override { return false; }
+
+	// IClickSequenceBehaviorTarget implementation
+
+	virtual void OnBeginSequencePreview(const FInputDeviceRay& ClickPos) override;
+	virtual bool CanBeginClickSequence(const FInputDeviceRay& ClickPos) override;
+	virtual void OnBeginClickSequence(const FInputDeviceRay& ClickPos) override;
+	virtual void OnNextSequencePreview(const FInputDeviceRay& ClickPos) override;
+	virtual bool OnNextSequenceClick(const FInputDeviceRay& ClickPos) override;
+	virtual void OnTerminateClickSequence() override;
+	virtual bool RequestAbortClickSequence() override;
+	virtual void OnUpdateModifierState(int ModifierID, bool bIsOn) override;
+
+	// polygon drawing functions
+	virtual void ResetPolygon();
+	virtual void UpdatePreviewVertex(const FVector3d& PreviewVertex);
+	virtual void AppendVertex(const FVector3d& Vertex);
+	virtual bool FindDrawPlaneHitPoint(const FInputDeviceRay& ClickPos, FVector3d& HitPosOut);
+	virtual void EmitCurrentPolygon();
+
+	virtual void BeginInteractiveExtrude();
+	virtual void EndInteractiveExtrude();
+
+
+public:
+	virtual void ApplyUndoPoints(const TArray<FVector3d>& ClickPointsIn, const TArray<FVector3d>& PolygonVerticesIn);
+
+
+protected:
+	// flags used to identify modifier keys/buttons
+	static const int IgnoreSnappingModifier = 1;
+	static const int AngleSnapModifier = 2;
+
+protected:
+
+	/** Property set for type of output object (StaticMesh, Volume, etc) */
+	UPROPERTY()
+	TObjectPtr<UCreateMeshObjectTypeProperties> OutputTypeProperties;
+
+	/** Properties that control polygon generation exposed to user via detailsview */
+	UPROPERTY()
+	TObjectPtr<UDrawPolygonToolStandardProperties> PolygonProperties;
+
+	UPROPERTY()
+	TObjectPtr<UDrawPolygonToolSnapProperties> SnapProperties;
+
+	UPROPERTY()
+	TObjectPtr<UNewMeshMaterialProperties> MaterialProperties;
+	
+
+	/** Origin of plane we will draw polygon on */
+	FVector3d DrawPlaneOrigin;
+
+	/** Orientation of plane we will draw polygon on */
+	UE::Geometry::FQuaterniond DrawPlaneOrientation;
+	
+	/** Vertices of current preview polygon */
+	TArray<FVector3d> PolygonVertices;
+
+	/** Vertices of holes in current preview polygon */
+	TArray<TArray<FVector3d>> PolygonHolesVertices;
+
+	/** last vertex of polygon that is actively being updated as input device is moved */
+	FVector3d PreviewVertex;
+
+protected:
+	UWorld* TargetWorld;
+
+	FViewCameraState CameraState;
+
+	UPROPERTY()
+	TObjectPtr<UPreviewMesh> PreviewMesh;
+
+	
+	// drawing plane gizmo
+
+	UPROPERTY()
+	TObjectPtr<UCombinedTransformGizmo> PlaneTransformGizmo;
+
+	UPROPERTY()
+	TObjectPtr<UTransformProxy> PlaneTransformProxy;
+
+	// called on PlaneTransformProxy.OnTransformChanged
+	void PlaneTransformChanged(UTransformProxy* Proxy, FTransform Transform);
+
+	// calls SetDrawPlaneFromWorldPos when user ctrl+clicks on scene
+	IClickBehaviorTarget* SetPointInWorldConnector = nullptr;
+
+	// updates plane and gizmo position
+	virtual void SetDrawPlaneFromWorldPos(const FVector3d& Position, const FVector3d& Normal);
+
+	void UpdateShowGizmoState(bool bNewVisibility);
+
+	// polygon drawing
+
+	bool bAbortActivePolygonDraw;
+
+	bool bInFixedPolygonMode = false;
+	TArray<FVector3d> FixedPolygonClickPoints;
+
+	// can close poly if current segment intersects existing segment
+	bool UpdateSelfIntersection();
+	bool bHaveSelfIntersection;
+	int SelfIntersectSegmentIdx;
+	FVector3d SelfIntersectionPoint;
+
+	// only used when SnapSettings.bHitSceneObjects = true
+	bool bHaveSurfaceHit;
+	FVector3d SurfaceHitPoint;
+	FVector3d SurfaceOffsetPoint;
+
+	bool bIgnoreSnappingToggle = false;		// toggled by hotkey (shift)
+	UE::Geometry::FPointPlanarSnapSolver SnapEngine;
+	ToolSceneQueriesUtil::FSnapGeometry LastSnapGeometry;
+	FVector3d LastGridSnapPoint;
+
+	void GetPolygonParametersFromFixedPoints(const TArray<FVector3d>& FixedPoints, FVector2d& FirstReferencePt, FVector2d& BoxSize, double& YSign, double& AngleRad);
+	void GenerateFixedPolygon(const TArray<FVector3d>& FixedPoints, TArray<FVector3d>& VerticesOut, TArray<TArray<FVector3d>>& HolesVerticesOut);
+
+
+	// extrusion control
+
+	bool bInInteractiveExtrude = false;
+
+	void UpdateLivePreview();
+	bool bPreviewUpdatePending;
+
+	UPROPERTY()
+	TObjectPtr<UPlaneDistanceFromHitMechanic> HeightMechanic;
+
+
+	UE::Geometry::FFrame3f HitPosFrameWorld;
+
+	/** Generate extruded meshes.  Returns true on success. */
+	bool GeneratePolygonMesh(const TArray<FVector3d>& Polygon, const TArray<TArray<FVector3d>>& PolygonHoles, FDynamicMesh3* ResultMeshOut, UE::Geometry::FFrame3d& WorldFrameOut, bool bIncludePreviewVtx, double ExtrudeDistance, bool bExtrudeSymmetric);
+
+
+	// user feedback messages
+	void ShowStartupMessage();
+	void ShowExtrudeMessage();
+
+
+	friend class FDrawPolygonStateChange;
+	int32 CurrentCurveTimestamp = 1;
+	void UndoCurrentOperation(const TArray<FVector3d>& ClickPointsIn, const TArray<FVector3d>& PolygonVerticesIn);
+	bool CheckInCurve(int32 Timestamp) const { return CurrentCurveTimestamp == Timestamp; }
+};
+
+
+
+// Change event used by DrawPolygonTool to undo draw state.
+// Currently does not redo.
+class MESHMODELINGTOOLS_API FDrawPolygonStateChange : public FToolCommandChange
+{
+public:
+	bool bHaveDoneUndo = false;
+	int32 CurveTimestamp = 0;
+	const TArray<FVector3d> FixedVertexPoints;
+	const TArray<FVector3d> PolyPoints;
+
+	FDrawPolygonStateChange(int32 CurveTimestampIn,
+							const TArray<FVector3d>& FixedVertexPointsIn,
+							const TArray<FVector3d>& PolyPointsIn)
+		: CurveTimestamp(CurveTimestampIn),
+		FixedVertexPoints(FixedVertexPointsIn),
+		PolyPoints(PolyPointsIn)
+	{
+	}
+	virtual void Apply(UObject* Object) override {}
+	virtual void Revert(UObject* Object) override;
+	virtual bool HasExpired(UObject* Object) const override;
+	virtual FString ToString() const override;
+};
