ModelingTools: Adjusting UI min/max defaults and descriptions for Stairs primitive tool.

#rnx

diff --git a/Engine/Plugins/Runtime/MeshModelingToolset/Source/MeshModelingTools/Public/AddPrimitiveTool.h b/Engine/Plugins/Runtime/MeshModelingToolset/Source/MeshModelingTools/Public/AddPrimitiveTool.h
--- a/Engine/Plugins/Runtime/MeshModelingToolset/Source/MeshModelingTools/Public/AddPrimitiveTool.h
+++ b/Engine/Plugins/Runtime/MeshModelingToolset/Source/MeshModelingTools/Public/AddPrimitiveTool.h
@@ -387,7 +387,7 @@ public:
 		EProceduralStairsType StairsType = EProceduralStairsType::Linear;
 
 	/** Number of Steps of Shape */
-	UPROPERTY(EditAnywhere, Category = "ShapeSettings", meta = (DisplayName = "Number of Steps", UIMin = "1", UIMax = "100", ClampMin = "1", ClampMax = "1000", ProceduralShapeSetting))
+	UPROPERTY(EditAnywhere, Category = "ShapeSettings", meta = (DisplayName = "Number of Steps", UIMin = "1", UIMax = "100", ClampMin = "1", ClampMax = "1000000", ProceduralShapeSetting))
 	int NumSteps = 8;
 
 	/** Width of each step */
@@ -402,12 +402,12 @@ public:
 	UPROPERTY(EditAnywhere, Category = "ShapeSettings", meta = (DisplayName = "Step Depth", UIMin = "1.0", UIMax = "1000.0", ClampMin = "0.0001", ClampMax = "1000000.0", ProceduralShapeSetting, EditCondition = "StairsType == EProceduralStairsType::Linear || StairsType == EProceduralStairsType::Floating", EditConditionHides))
 	float StepDepth = 30.0f;
 
-	/** Angular length of curved stair shapes */
+	/** Angular length of curved stair shapes. Positive values are clockwise, negative are counter-clockwise. */
 	UPROPERTY(EditAnywhere, Category = "ShapeSettings", meta = (DisplayName = "Curve Angle", UIMin = "-360.0", UIMax = "360.0", ClampMin = "-360.0", ClampMax = "360.0", ProceduralShapeSetting, EditCondition = "StairsType == EProceduralStairsType::Curved", EditConditionHides))
 	float CurveAngle = 90.0f;
 
-	/** Angular length of spiral stair shapes */
-	UPROPERTY(EditAnywhere, Category = "ShapeSettings", meta = (DisplayName = "Spiral Angle", UIMin = "-3600.0", UIMax = "3600.0", ClampMin = "-360000.0", ClampMax = "360000.0", ProceduralShapeSetting, EditCondition = "StairsType == EProceduralStairsType::Spiral", EditConditionHides))
+	/** Angular length of spiral stair shapes. Positive values are clockwise, negative are counter-clockwise. */
+	UPROPERTY(EditAnywhere, Category = "ShapeSettings", meta = (DisplayName = "Spiral Angle", UIMin = "-720.0", UIMax = "720.0", ClampMin = "-360000.0", ClampMax = "360000.0", ProceduralShapeSetting, EditCondition = "StairsType == EProceduralStairsType::Spiral", EditConditionHides))
 	float SpiralAngle = 90.0f;
 
 	/** Inner radius for curved/spiral stair shapes */
