MeshModelingTools: Cleanup Shapes tools

diff --git a/Engine/Plugins/Experimental/MeshModelingToolsetExp/Source/MeshModelingToolsExp/Private/AddPatchTool.cpp b/Engine/Plugins/Experimental/MeshModelingToolsetExp/Source/MeshModelingToolsExp/Private/AddPatchTool.cpp
--- a/Engine/Plugins/Experimental/MeshModelingToolsetExp/Source/MeshModelingToolsExp/Private/AddPatchTool.cpp
+++ b/Engine/Plugins/Experimental/MeshModelingToolsetExp/Source/MeshModelingToolsExp/Private/AddPatchTool.cpp
@@ -125,7 +125,7 @@ void UAddPatchTool::Render(IToolsContextRenderAPI* RenderAPI)
 
 void UAddPatchTool::OnPropertyModified(UObject* PropertySet, FProperty* Property)
 {
-	PreviewMesh->EnableWireframe(MaterialProperties->bWireframe);
+	PreviewMesh->EnableWireframe(MaterialProperties->bShowWireframe);
 	PreviewMesh->SetMaterial(MaterialProperties->Material.Get());
 	GeneratePreviewBaseMesh();
 }
diff --git a/Engine/Plugins/Experimental/MeshModelingToolsetExp/Source/MeshModelingToolsExp/Private/RevolveBoundaryTool.cpp b/Engine/Plugins/Experimental/MeshModelingToolsetExp/Source/MeshModelingToolsExp/Private/RevolveBoundaryTool.cpp
--- a/Engine/Plugins/Experimental/MeshModelingToolsetExp/Source/MeshModelingToolsExp/Private/RevolveBoundaryTool.cpp
+++ b/Engine/Plugins/Experimental/MeshModelingToolsetExp/Source/MeshModelingToolsExp/Private/RevolveBoundaryTool.cpp
@@ -234,7 +234,7 @@ void URevolveBoundaryTool::StartPreview()
 
 	Preview->ConfigureMaterials(MaterialProperties->Material.Get(),
 		ToolSetupUtil::GetDefaultWorkingMaterial(GetToolManager()));
-	Preview->PreviewMesh->EnableWireframe(MaterialProperties->bWireframe);
+	Preview->PreviewMesh->EnableWireframe(MaterialProperties->bShowWireframe);
 
 	Preview->OnMeshUpdated.AddLambda(
 		[this](const UMeshOpPreviewWithBackgroundCompute* UpdatedPreview)
@@ -351,7 +351,7 @@ void URevolveBoundaryTool::OnPropertyModified(UObject* PropertySet, FProperty* P
 				ToolSetupUtil::GetDefaultWorkingMaterial(GetToolManager()));
 		}
 
-		Preview->PreviewMesh->EnableWireframe(MaterialProperties->bWireframe);
+		Preview->PreviewMesh->EnableWireframe(MaterialProperties->bShowWireframe);
 		Preview->InvalidateResult();
 	}
 }
diff --git a/Engine/Plugins/Experimental/MeshModelingToolsetExp/Source/MeshModelingToolsExp/Public/Properties/MeshMaterialProperties.h b/Engine/Plugins/Experimental/MeshModelingToolsetExp/Source/MeshModelingToolsExp/Public/Properties/MeshMaterialProperties.h
--- a/Engine/Plugins/Experimental/MeshModelingToolsetExp/Source/MeshModelingToolsExp/Public/Properties/MeshMaterialProperties.h
+++ b/Engine/Plugins/Experimental/MeshModelingToolsetExp/Source/MeshModelingToolsExp/Public/Properties/MeshMaterialProperties.h
@@ -21,23 +21,23 @@ class MESHMODELINGTOOLSEXP_API UNewMeshMaterialProperties : public UInteractiveT
 public:
 	UNewMeshMaterialProperties();
 
-	/** Material for new mesh*/
-	UPROPERTY(EditAnywhere, NonTransactional, Category = Material, meta = (DisplayName = "Material"))
+	/** Material for new mesh */
+	UPROPERTY(EditAnywhere, NonTransactional, Category = Material)
 	TWeakObjectPtr<UMaterialInterface> Material;
 
 	/** Scale factor for generated UVs */
 	UPROPERTY(EditAnywhere, Category = Material, meta = (DisplayName = "UV Scale", HideEditConditionToggle, EditConditionHides, EditCondition = "bShowExtendedOptions"))
 	float UVScale = 1.0;
 
-	/** If set, UV scales will be relative to world space so different objects created with the same UV scale should have the same average texel size */
-	UPROPERTY(EditAnywhere, Category = Material, meta = (DisplayName = "UV Scale Relative to World Space", HideEditConditionToggle, EditConditionHides, EditCondition = "bShowExtendedOptions"))
+	/** If true, UV scale will be relative to world space. This means objects of different sizes created with the same UV scale have the same average texel size. */
+	UPROPERTY(EditAnywhere, Category = Material, meta = (DisplayName = "World Space UV Scale", HideEditConditionToggle, EditConditionHides, EditCondition = "bShowExtendedOptions"))
 	bool bWorldSpaceUVScale = false;
 
-	/** Overlay wireframe on preview */
-	UPROPERTY(EditAnywhere, Category = Material, meta = (DisplayName = "Show Wireframe", HideEditConditionToggle, EditConditionHides, EditCondition = "bShowExtendedOptions"))
-	bool bWireframe = false;
+	/** If true, overlays preview with wireframe */
+	UPROPERTY(EditAnywhere, Category = Material, meta = (HideEditConditionToggle, EditConditionHides, EditCondition = "bShowExtendedOptions"))
+	bool bShowWireframe = false;
 
-	// controls visibility of UV/etc properties
+	/** If true, extended options are available */ 
 	UPROPERTY()
 	bool bShowExtendedOptions = true;
 };
diff --git a/Engine/Plugins/Runtime/MeshModelingToolset/Source/MeshModelingTools/Private/AddPrimitiveTool.cpp b/Engine/Plugins/Runtime/MeshModelingToolset/Source/MeshModelingTools/Private/AddPrimitiveTool.cpp
--- a/Engine/Plugins/Runtime/MeshModelingToolset/Source/MeshModelingTools/Private/AddPrimitiveTool.cpp
+++ b/Engine/Plugins/Runtime/MeshModelingToolset/Source/MeshModelingTools/Private/AddPrimitiveTool.cpp
@@ -10,7 +10,6 @@
 #include "ToolSetupUtil.h"
 
 #include "Components/StaticMeshComponent.h"
-#include "Engine/StaticMesh.h"
 #include "Engine/StaticMeshActor.h"
 
 #include "Generators/SweepGenerator.h"
@@ -23,8 +22,6 @@
 #include "DynamicMesh/DynamicMeshAttributeSet.h"
 #include "FaceGroupUtil.h"
 
-#include "Drawing/MeshDebugDrawing.h"
-
 #include "DynamicMeshEditor.h"
 #include "UObject/PropertyIterator.h"
 #include "UObject/UnrealType.h"
@@ -90,7 +87,7 @@ UProceduralShapeToolProperties::IsEquivalent( const UProceduralShapeToolProperti
 	{
 		return false;
 	}
-	for ( FProperty* Prop : TFieldRange<FProperty>(Class) )
+	for (const FProperty* Prop : TFieldRange<FProperty>(Class))
 	{
 		if (Prop->HasMetaData(TEXT("ProceduralShapeSetting")) &&
 			(!Prop->Identical_InContainer(this, Other)))
@@ -143,7 +140,7 @@ void UAddPrimitiveTool::Setup()
 	ToolSetupUtil::ApplyRenderingConfigurationToPreview(PreviewMesh, nullptr); 
 	PreviewMesh->SetVisible(false);
 	PreviewMesh->SetMaterial(MaterialProperties->Material.Get());
-	PreviewMesh->EnableWireframe(MaterialProperties->bWireframe);
+	PreviewMesh->EnableWireframe(MaterialProperties->bShowWireframe);
 
 	UpdatePreviewMesh();
 
@@ -178,7 +175,7 @@ void UAddPrimitiveTool::OnPropertyModified(UObject* PropertySet, FProperty* Prop
 	// So we need to ensure that we handle this case.
 	if (PreviewMesh)
 	{
-		PreviewMesh->EnableWireframe(MaterialProperties->bWireframe);
+		PreviewMesh->EnableWireframe(MaterialProperties->bShowWireframe);
 		PreviewMesh->SetMaterial(MaterialProperties->Material.Get());
 		UpdatePreviewMesh();
 	}
@@ -218,9 +215,9 @@ void UAddPrimitiveTool::UpdatePreviewPosition(const FInputDeviceRay& DeviceClick
 	bool bHit = false;
 
 	FPlane DrawPlane(FVector::ZeroVector, FVector(0, 0, 1));
-	if (ShapeSettings->PlaceMode == EMakeMeshPlacementType::GroundPlane)
+	if (ShapeSettings->TargetSurface == EMakeMeshPlacementType::GroundPlane)
 	{
-		FVector3f DrawPlanePos = (FVector3f)FMath::RayPlaneIntersection(ClickPosWorldRay.Origin, ClickPosWorldRay.Direction, DrawPlane);
+		FVector3f DrawPlanePos = FMath::RayPlaneIntersection(ClickPosWorldRay.Origin, ClickPosWorldRay.Direction, DrawPlane);
 		bHit = true;
 		ShapeFrame = FFrame3f(DrawPlanePos);
 	}
@@ -231,18 +228,18 @@ void UAddPrimitiveTool::UpdatePreviewPosition(const FInputDeviceRay& DeviceClick
 		bHit = ToolSceneQueriesUtil::FindNearestVisibleObjectHit(TargetWorld, Result, ClickPosWorldRay);
 		if (bHit)
 		{
-			FVector3f Normal = (FVector3f)Result.ImpactNormal;
-			if (!ShapeSettings->bAlignShapeToPlacementSurface)
+			FVector3f Normal = Result.ImpactNormal;
+			if (!ShapeSettings->bAlignToNormal)
 			{
 				Normal = FVector3f::UnitZ();
 			}
-			ShapeFrame = FFrame3f((FVector3f)Result.ImpactPoint, Normal);
+			ShapeFrame = FFrame3f(Result.ImpactPoint, Normal);
 			ShapeFrame.ConstrainedAlignPerpAxes();
 		}
 		else
 		{
 			// fall back to ground plane if we don't have a scene hit
-			FVector3f DrawPlanePos = (FVector3f)FMath::RayPlaneIntersection(ClickPosWorldRay.Origin, ClickPosWorldRay.Direction, DrawPlane);
+			FVector3f DrawPlanePos = FMath::RayPlaneIntersection(ClickPosWorldRay.Origin, ClickPosWorldRay.Direction, DrawPlane);
 			bHit = true;
 			ShapeFrame = FFrame3f(DrawPlanePos);
 		}
@@ -254,11 +251,11 @@ void UAddPrimitiveTool::UpdatePreviewPosition(const FInputDeviceRay& DeviceClick
 		FSceneSnapQueryRequest Request;
 		Request.RequestType = ESceneSnapQueryType::Position;
 		Request.TargetTypes = ESceneSnapQueryTargetType::Grid;
-		Request.Position = (FVector)ShapeFrame.Origin;
+		Request.Position = ShapeFrame.Origin;
 		TArray<FSceneSnapQueryResult> Results;
 		if (GetToolManager()->GetContextQueriesAPI()->ExecuteSceneSnapQuery(Request, Results))
 		{
-			ShapeFrame.Origin = (FVector3f)Results[0].Position;
+			ShapeFrame.Origin = Results[0].Position;
 		}
 	}
 
@@ -278,7 +275,7 @@ void UAddPrimitiveTool::UpdatePreviewPosition(const FInputDeviceRay& DeviceClick
 	}
 }
 
-void UAddPrimitiveTool::UpdatePreviewMesh()
+void UAddPrimitiveTool::UpdatePreviewMesh() const
 {
 	FDynamicMesh3 NewMesh;
 	GenerateMesh( &NewMesh );
@@ -291,12 +288,12 @@ void UAddPrimitiveTool::UpdatePreviewMesh()
 	if (MaterialProperties->UVScale != 1.0 || MaterialProperties->bWorldSpaceUVScale)
 	{
 		FDynamicMeshEditor Editor(&NewMesh);
-		float WorldUnitsInMetersFactor = MaterialProperties->bWorldSpaceUVScale ? .01f : 1.0f;
+		const float WorldUnitsInMetersFactor = MaterialProperties->bWorldSpaceUVScale ? .01f : 1.0f;
 		Editor.RescaleAttributeUVs(MaterialProperties->UVScale * WorldUnitsInMetersFactor, MaterialProperties->bWorldSpaceUVScale);
 	}
 
 	// set mesh position
-	FAxisAlignedBox3d Bounds = NewMesh.GetBounds(true);
+	const FAxisAlignedBox3d Bounds = NewMesh.GetBounds(true);
 	FVector3d TargetOrigin = Bounds.Center();
 	if (ShapeSettings->PivotLocation == EMakeMeshPivotLocation::Base)
 	{
@@ -306,11 +303,11 @@ void UAddPrimitiveTool::UpdatePreviewMesh()
 	{
 		TargetOrigin.Z = Bounds.Max.Z;
 	}
-	for (int vid : NewMesh.VertexIndicesItr())
+	for (const int Vid : NewMesh.VertexIndicesItr())
 	{
-		FVector3d Pos = NewMesh.GetVertex(vid);
+		FVector3d Pos = NewMesh.GetVertex(Vid);
 		Pos -= TargetOrigin;
-		NewMesh.SetVertex(vid, Pos);
+		NewMesh.SetVertex(Vid, Pos);
 	}
 
 	PreviewMesh->UpdatePreview(&NewMesh);
@@ -320,7 +317,7 @@ void UAddPrimitiveTool::UpdatePreviewMesh()
 	checkSlow(CalculateTangentsSuccessful);
 }
 
-void UAddPrimitiveTool::OnClicked(const FInputDeviceRay& DeviceClickPos)
+void UAddPrimitiveTool::OnClicked(const FInputDeviceRay& ClickPos)
 {
 	UMaterialInterface* Material = PreviewMesh->GetMaterial();
 
@@ -329,7 +326,6 @@ void UAddPrimitiveTool::OnClicked(const FInputDeviceRay& DeviceClickPos)
 		GetToolManager()->BeginUndoTransaction(LOCTEXT("AddPrimitiveToolTransactionName", "Add Primitive Mesh"));
 		FActorSpawnParameters SpawnParameters;
 		SpawnParameters.Template = LastGenerated->Actor;
-		UE::Geometry::FTransform3d CurTransform(PreviewMesh->GetTransform());
 		FRotator Rotation(0.0f, 0.0f, 0.0f);
 		AStaticMeshActor* CloneActor = TargetWorld->SpawnActor<AStaticMeshActor>(FVector::ZeroVector, Rotation, SpawnParameters);
 		// some properties must be manually set on the component because they will not persist reliably through the spawn template (especially if the actor creation was undone)
@@ -348,13 +344,12 @@ void UAddPrimitiveTool::OnClicked(const FInputDeviceRay& DeviceClickPos)
 	LastGenerated = nullptr;
 
 	const FDynamicMesh3* CurMesh = PreviewMesh->GetPreviewDynamicMesh();
-	UE::Geometry::FTransform3d CurTransform(PreviewMesh->GetTransform());
 
 	GetToolManager()->BeginUndoTransaction(LOCTEXT("AddPrimitiveToolTransactionName", "Add Primitive Mesh"));
 
 	FCreateMeshObjectParams NewMeshObjectParams;
 	NewMeshObjectParams.TargetWorld = TargetWorld;
-	NewMeshObjectParams.Transform = (FTransform)CurTransform;
+	NewMeshObjectParams.Transform = PreviewMesh->GetTransform();
 	NewMeshObjectParams.BaseName = AssetName;
 	NewMeshObjectParams.Materials.Add(Material);
 	NewMeshObjectParams.SetMesh(CurMesh);
@@ -389,21 +384,18 @@ UAddBoxPrimitiveTool::UAddBoxPrimitiveTool(const FObjectInitializer& ObjectIniti
 	: Super(ObjectInitializer.SetDefaultSubobjectClass<UProceduralBoxToolProperties>(TEXT("ShapeSettings")))
 {
 	AssetName = TEXT("Box");
-	SetToolDisplayName(LOCTEXT("BoxToolName", "Create Boxes"));
+	UInteractiveTool::SetToolDisplayName(LOCTEXT("BoxToolName", "Create Box"));
 }
 
 void UAddBoxPrimitiveTool::GenerateMesh(FDynamicMesh3* OutMesh) const
 {
 	FGridBoxMeshGenerator BoxGen;
-	auto* BoxSettings = Cast<UProceduralBoxToolProperties>(ShapeSettings);
+	const UProceduralBoxToolProperties* BoxSettings = Cast<UProceduralBoxToolProperties>(ShapeSettings);
 	BoxGen.Box = UE::Geometry::FOrientedBox3d(FVector3d::Zero(), 0.5*FVector3d(BoxSettings->Depth, BoxSettings->Width, BoxSettings->Height));
 	BoxGen.EdgeVertices = FIndex3i(BoxSettings->DepthSubdivisions + 1,
 								   BoxSettings->WidthSubdivisions + 1,
 								   BoxSettings->HeightSubdivisions + 1);
-	if (ShapeSettings->PolygroupMode == EMakeMeshPolygroupMode::PerQuad)
-	{
-		BoxGen.bPolygroupPerQuad = true;
-	}
+	BoxGen.bPolygroupPerQuad = ShapeSettings->PolygroupMode == EMakeMeshPolygroupMode::PerQuad;
 	BoxGen.Generate();
 	OutMesh->Copy(&BoxGen);
 }
@@ -414,13 +406,13 @@ UAddRectanglePrimitiveTool::UAddRectanglePrimitiveTool(const FObjectInitializer&
 	: Super(ObjectInitializer.SetDefaultSubobjectClass<UProceduralRectangleToolProperties>(TEXT("ShapeSettings")))
 {
 	AssetName = TEXT("Rectangle");
-	SetToolDisplayName(LOCTEXT("RectToolName", "Create Rectangles"));
+	UInteractiveTool::SetToolDisplayName(LOCTEXT("RectToolName", "Create Rectangle"));
 }
 
 void UAddRectanglePrimitiveTool::GenerateMesh(FDynamicMesh3* OutMesh) const
 {
 	auto* RectangleSettings = Cast<UProceduralRectangleToolProperties>(ShapeSettings);
-	switch (RectangleSettings->RectType)
+	switch (RectangleSettings->RectangleType)
 	{
 	case EProceduralRectType::Rectangle:
 	{
@@ -456,12 +448,12 @@ UAddDiscPrimitiveTool::UAddDiscPrimitiveTool(const FObjectInitializer& ObjectIni
 	: Super(ObjectInitializer.SetDefaultSubobjectClass<UProceduralDiscToolProperties>(TEXT("ShapeSettings")))
 {
 	AssetName = TEXT("Disc");
-	SetToolDisplayName(LOCTEXT("DiscToolName", "Create Discs"));
+	UInteractiveTool::SetToolDisplayName(LOCTEXT("DiscToolName", "Create Disc"));
 }
 
 void UAddDiscPrimitiveTool::GenerateMesh(FDynamicMesh3* OutMesh) const
 {
-	auto* DiscSettings = Cast<UProceduralDiscToolProperties>(ShapeSettings);
+	const UProceduralDiscToolProperties* DiscSettings = Cast<UProceduralDiscToolProperties>(ShapeSettings);
 	switch (DiscSettings->DiscType)
 	{
 	case EProceduralDiscType::Disc:
@@ -495,25 +487,22 @@ UAddTorusPrimitiveTool::UAddTorusPrimitiveTool(const FObjectInitializer& ObjectI
 	: Super(ObjectInitializer.SetDefaultSubobjectClass<UProceduralTorusToolProperties>(TEXT("ShapeSettings")))
 {
 	AssetName = TEXT("Torus");
-	SetToolDisplayName(LOCTEXT("TorusToolName", "Create Torii"));
+	UInteractiveTool::SetToolDisplayName(LOCTEXT("TorusToolName", "Create Torus"));
 }
 
 void UAddTorusPrimitiveTool::GenerateMesh(FDynamicMesh3* OutMesh) const
 {
 	FGeneralizedCylinderGenerator Gen;
-	auto* TorusSettings = Cast<UProceduralTorusToolProperties>(ShapeSettings);
-	Gen.CrossSection = FPolygon2d::MakeCircle(TorusSettings->MinorRadius, TorusSettings->CrossSectionSlices);
-	FPolygon2d PathCircle = FPolygon2d::MakeCircle(TorusSettings->MajorRadius, TorusSettings->TubeSlices);
+	const UProceduralTorusToolProperties* TorusSettings = Cast<UProceduralTorusToolProperties>(ShapeSettings);
+	Gen.CrossSection = FPolygon2d::MakeCircle(TorusSettings->MinorRadius, TorusSettings->MinorSlices);
+	FPolygon2d PathCircle = FPolygon2d::MakeCircle(TorusSettings->MajorRadius, TorusSettings->MajorSlices);
 	for (int Idx = 0; Idx < PathCircle.VertexCount(); Idx++)
 	{
 		Gen.Path.Add( FVector3d(PathCircle[Idx].X, PathCircle[Idx].Y, 0) );
 	}
 	Gen.bLoop = true;
 	Gen.bCapped = false;
-	if (ShapeSettings->PolygroupMode == EMakeMeshPolygroupMode::PerQuad)
-	{
-		Gen.bPolygroupPerQuad = true;
-	}
+	Gen.bPolygroupPerQuad = ShapeSettings->PolygroupMode == EMakeMeshPolygroupMode::PerQuad;
 	Gen.InitialFrame = FFrame3d(Gen.Path[0]);
 	Gen.Generate();
 	OutMesh->Copy(&Gen);
@@ -525,22 +514,19 @@ UAddCylinderPrimitiveTool::UAddCylinderPrimitiveTool(const FObjectInitializer& O
 	: Super(ObjectInitializer.SetDefaultSubobjectClass<UProceduralCylinderToolProperties>(TEXT("ShapeSettings")))
 {
 	AssetName = TEXT("Cylinder");
-	SetToolDisplayName(LOCTEXT("CylinderToolName", "Create Cylinders"));
+	UInteractiveTool::SetToolDisplayName(LOCTEXT("CylinderToolName", "Create Cylinder"));
 }
 
 void UAddCylinderPrimitiveTool::GenerateMesh(FDynamicMesh3* OutMesh) const
 {
 	FCylinderGenerator CylGen;
-	auto* CylinderSettings = Cast<UProceduralCylinderToolProperties>(ShapeSettings);
+	const UProceduralCylinderToolProperties* CylinderSettings = Cast<UProceduralCylinderToolProperties>(ShapeSettings);
 	CylGen.Radius[1] = CylGen.Radius[0] = CylinderSettings->Radius;
 	CylGen.Height = CylinderSettings->Height;
 	CylGen.AngleSamples = CylinderSettings->RadialSlices;
 	CylGen.LengthSamples = CylinderSettings->HeightSubdivisions - 1;
 	CylGen.bCapped = true;
-	if (ShapeSettings->PolygroupMode == EMakeMeshPolygroupMode::PerQuad)
-	{
-		CylGen.bPolygroupPerQuad = true;
-	}
+	CylGen.bPolygroupPerQuad = ShapeSettings->PolygroupMode == EMakeMeshPolygroupMode::PerQuad;
 	CylGen.Generate();
 	OutMesh->Copy(&CylGen);
 }
@@ -551,24 +537,21 @@ UAddConePrimitiveTool::UAddConePrimitiveTool(const FObjectInitializer& ObjectIni
 	: Super(ObjectInitializer.SetDefaultSubobjectClass<UProceduralConeToolProperties>(TEXT("ShapeSettings")))
 {
 	AssetName = TEXT("Cone");
-	SetToolDisplayName(LOCTEXT("ConeToolName", "Create Cones"));
+	UInteractiveTool::SetToolDisplayName(LOCTEXT("ConeToolName", "Create Cone"));
 }
 
 void UAddConePrimitiveTool::GenerateMesh(FDynamicMesh3* OutMesh) const
 {
 	// Unreal's standard cone is just a cylinder with a very small top
 	FCylinderGenerator CylGen;
-	auto* ConeSettings = Cast<UProceduralConeToolProperties>(ShapeSettings);
+	const UProceduralConeToolProperties* ConeSettings = Cast<UProceduralConeToolProperties>(ShapeSettings);
 	CylGen.Radius[0] = ConeSettings->Radius;
 	CylGen.Radius[1] = .01;
 	CylGen.Height = ConeSettings->Height;
 	CylGen.AngleSamples = ConeSettings->RadialSlices;
 	CylGen.LengthSamples = ConeSettings->HeightSubdivisions - 1;
 	CylGen.bCapped = true;
-	if (ShapeSettings->PolygroupMode == EMakeMeshPolygroupMode::PerQuad)
-	{
-		CylGen.bPolygroupPerQuad = true;
-	}
+	CylGen.bPolygroupPerQuad = ShapeSettings->PolygroupMode == EMakeMeshPolygroupMode::PerQuad;
 	CylGen.Generate();
 	OutMesh->Copy(&CylGen);
 }
@@ -578,13 +561,13 @@ UAddArrowPrimitiveTool::UAddArrowPrimitiveTool(const FObjectInitializer& ObjectI
 	: Super(ObjectInitializer.SetDefaultSubobjectClass<UProceduralArrowToolProperties>(TEXT("ShapeSettings")))
 {
 	AssetName = TEXT("Arrow");
-	SetToolDisplayName(LOCTEXT("ArrowToolName", "Create Arrows"));
+	UInteractiveTool::SetToolDisplayName(LOCTEXT("ArrowToolName", "Create Arrow"));
 }
 
 void UAddArrowPrimitiveTool::GenerateMesh(FDynamicMesh3* OutMesh) const
 {
 	FArrowGenerator ArrowGen;
-	auto* ArrowSettings = Cast<UProceduralArrowToolProperties>(ShapeSettings);
+	const UProceduralArrowToolProperties* ArrowSettings = Cast<UProceduralArrowToolProperties>(ShapeSettings);
 	ArrowGen.StickRadius = ArrowSettings->ShaftRadius;
 	ArrowGen.StickLength = ArrowSettings->ShaftHeight;
 	ArrowGen.HeadBaseRadius = ArrowSettings->HeadRadius;
@@ -592,11 +575,14 @@ void UAddArrowPrimitiveTool::GenerateMesh(FDynamicMesh3* OutMesh) const
 	ArrowGen.HeadLength = ArrowSettings->HeadHeight;
 	ArrowGen.AngleSamples = ArrowSettings->RadialSlices;
 	ArrowGen.bCapped = true;
-	if (ShapeSettings->PolygroupMode == EMakeMeshPolygroupMode::PerQuad)
+	ArrowGen.bPolygroupPerQuad = ShapeSettings->PolygroupMode == EMakeMeshPolygroupMode::PerQuad;
+	if (ArrowSettings->HeightSubdivisions > 1)
 	{
-		ArrowGen.bPolygroupPerQuad = true;
+		const int AdditionalLengthsSamples = ArrowSettings->HeightSubdivisions - 1;
+		ArrowGen.AdditionalLengthSamples[0] = AdditionalLengthsSamples;
+		ArrowGen.AdditionalLengthSamples[1] = AdditionalLengthsSamples;
+		ArrowGen.AdditionalLengthSamples[2] = AdditionalLengthsSamples;
 	}
-	ArrowGen.DistributeAdditionalLengthSamples(ArrowSettings->TotalSubdivisions);
 	ArrowGen.Generate();
 	OutMesh->Copy(&ArrowGen);
 }
@@ -607,20 +593,20 @@ UAddSpherePrimitiveTool::UAddSpherePrimitiveTool(const FObjectInitializer& Objec
 	: Super(ObjectInitializer.SetDefaultSubobjectClass<UProceduralSphereToolProperties>(TEXT("ShapeSettings")))
 {
 	AssetName = TEXT("Sphere");
-	SetToolDisplayName(LOCTEXT("SphereToolName", "Create Spheres"));
+	UInteractiveTool::SetToolDisplayName(LOCTEXT("SphereToolName", "Create Sphere"));
 }
 
 void UAddSpherePrimitiveTool::GenerateMesh(FDynamicMesh3* OutMesh) const
 {
-	auto* SphereSettings = Cast<UProceduralSphereToolProperties>(ShapeSettings);
+	const UProceduralSphereToolProperties* SphereSettings = Cast<UProceduralSphereToolProperties>(ShapeSettings);
 	switch (SphereSettings->SphereType)
 	{
 	case EProceduralSphereType::LatLong:
 	{
 		FSphereGenerator SphereGen;
 		SphereGen.Radius = SphereSettings->Radius;
-		SphereGen.NumTheta = SphereSettings->LongitudeSlices + 1;
-		SphereGen.NumPhi = SphereSettings->LatitudeSlices + 1;
+		SphereGen.NumTheta = SphereSettings->VerticalSlices + 1;
+		SphereGen.NumPhi = SphereSettings->HorizontalSlices + 1;
 		SphereGen.bPolygroupPerQuad = (ShapeSettings->PolygroupMode == EMakeMeshPolygroupMode::PerQuad);
 		SphereGen.Generate();
 		OutMesh->Copy(&SphereGen);
@@ -630,7 +616,7 @@ void UAddSpherePrimitiveTool::GenerateMesh(FDynamicMesh3* OutMesh) const
 	{
 		FBoxSphereGenerator SphereGen;
 		SphereGen.Radius = SphereSettings->Radius;
-		SphereGen.Box = UE::Geometry::FOrientedBox3d(FVector3d::Zero(),
+		SphereGen.Box = FOrientedBox3d(FVector3d::Zero(),
 			0.5 * FVector3d(SphereSettings->Subdivisions + 1,
 				SphereSettings->Subdivisions + 1,
 				SphereSettings->Subdivisions + 1));
@@ -650,12 +636,12 @@ UAddStairsPrimitiveTool::UAddStairsPrimitiveTool(const FObjectInitializer& Objec
 	: Super(ObjectInitializer.SetDefaultSubobjectClass<UProceduralStairsToolProperties>(TEXT("ShapeSettings")))
 {
 	AssetName = TEXT("Stairs");
-	SetToolDisplayName(LOCTEXT("StairsToolName", "Create Stairs"));
+	UInteractiveTool::SetToolDisplayName(LOCTEXT("StairsToolName", "Create Stairs"));
 }
 
 void UAddStairsPrimitiveTool::GenerateMesh(FDynamicMesh3* OutMesh) const
 {
-	auto* StairSettings = Cast<UProceduralStairsToolProperties>(ShapeSettings);
+	const UProceduralStairsToolProperties* StairSettings = Cast<UProceduralStairsToolProperties>(ShapeSettings);
 	switch (StairSettings->StairsType)
 	{
 	case EProceduralStairsType::Linear:
diff --git a/Engine/Plugins/Runtime/MeshModelingToolset/Source/MeshModelingTools/Private/DrawAndRevolveTool.cpp b/Engine/Plugins/Runtime/MeshModelingToolset/Source/MeshModelingTools/Private/DrawAndRevolveTool.cpp
--- a/Engine/Plugins/Runtime/MeshModelingToolset/Source/MeshModelingTools/Private/DrawAndRevolveTool.cpp
+++ b/Engine/Plugins/Runtime/MeshModelingToolset/Source/MeshModelingTools/Private/DrawAndRevolveTool.cpp
@@ -274,7 +274,7 @@ void UDrawAndRevolveTool::StartPreview()
 
 	Preview->ConfigureMaterials(MaterialProperties->Material.Get(),
 		ToolSetupUtil::GetDefaultWorkingMaterial(GetToolManager()));
-	Preview->PreviewMesh->EnableWireframe(MaterialProperties->bWireframe);
+	Preview->PreviewMesh->EnableWireframe(MaterialProperties->bShowWireframe);
 
 	Preview->OnMeshUpdated.AddLambda(
 		[this](const UMeshOpPreviewWithBackgroundCompute* UpdatedPreview)
@@ -304,7 +304,7 @@ void UDrawAndRevolveTool::OnPropertyModified(UObject* PropertySet, FProperty* Pr
 				ToolSetupUtil::GetDefaultWorkingMaterial(GetToolManager()));
 		}
 
-		Preview->PreviewMesh->EnableWireframe(MaterialProperties->bWireframe);
+		Preview->PreviewMesh->EnableWireframe(MaterialProperties->bShowWireframe);
 		Preview->InvalidateResult();
 	}
 }
diff --git a/Engine/Plugins/Runtime/MeshModelingToolset/Source/MeshModelingTools/Private/DrawPolygonTool.cpp b/Engine/Plugins/Runtime/MeshModelingToolset/Source/MeshModelingTools/Private/DrawPolygonTool.cpp
--- a/Engine/Plugins/Runtime/MeshModelingToolset/Source/MeshModelingTools/Private/DrawPolygonTool.cpp
+++ b/Engine/Plugins/Runtime/MeshModelingToolset/Source/MeshModelingTools/Private/DrawPolygonTool.cpp
@@ -1048,7 +1048,7 @@ void UDrawPolygonTool::UpdateLivePreview()
 	{
 		PreviewMesh->SetTransform(PlaneFrame.ToFTransform());
 		PreviewMesh->SetMaterial(MaterialProperties->Material.Get());
-		PreviewMesh->EnableWireframe(MaterialProperties->bWireframe);
+		PreviewMesh->EnableWireframe(MaterialProperties->bShowWireframe);
 		PreviewMesh->UpdatePreview(&Mesh);
 	}
 }
diff --git a/Engine/Plugins/Runtime/MeshModelingToolset/Source/MeshModelingTools/Public/AddPrimitiveTool.h b/Engine/Plugins/Runtime/MeshModelingToolset/Source/MeshModelingTools/Public/AddPrimitiveTool.h
--- a/Engine/Plugins/Runtime/MeshModelingToolset/Source/MeshModelingTools/Public/AddPrimitiveTool.h
+++ b/Engine/Plugins/Runtime/MeshModelingToolset/Source/MeshModelingTools/Public/AddPrimitiveTool.h
@@ -5,12 +5,10 @@
 #include "BaseTools/SingleClickTool.h"
 #include "CoreMinimal.h"
 #include "InteractiveToolBuilder.h"
-#include "MaterialShared.h"
 #include "PreviewMesh.h"
 #include "Properties/MeshMaterialProperties.h"
 #include "PropertySets/CreateMeshObjectTypeProperties.h"
 #include "UObject/NoExportTypes.h"
-#include "UObject/UObjectGlobals.h"
 
 #include "AddPrimitiveTool.generated.h"
 
@@ -61,15 +59,15 @@ enum class EMakeMeshPivotLocation : uint8
 	Top
 };
 
-/** Polygroup mode for primitive */
+/** Polygroup mode for shape */
 UENUM()
 enum class EMakeMeshPolygroupMode : uint8
 {
-	/** One polygroup for entire output mesh */
+	/** One Polygroup for the entire shape */
 	Single,
-	/** One polygroup per geometric face of primitive */
+	/** One Polygroup for each geometric face */
 	PerFace,
-	/** One polygroup per mesh quad/triangle */
+	/** One Polygroup for each quad/triangle */
 	PerQuad
 };
 
@@ -79,43 +77,34 @@ class MESHMODELINGTOOLS_API UProceduralShapeToolProperties : public UInteractive
 	GENERATED_BODY()
 
 public:
-	UProceduralShapeToolProperties() = default;
-
-	/** If the shape settings haven't changed, create instances of the last created asset rather than creating a whole new asset.  If false, all created actors will have separate underlying mesh assets. */
+	/** Create instances of the last created asset rather than creating a whole new asset, provided the shape settings have not changed.
+	  * If false, all created actors will have separate underlying mesh assets. */
 	UPROPERTY(EditAnywhere, Category = AssetSettings)
 	bool bInstanceIfPossible = false;
 
-	/** How should Polygroups be assigned to triangles of Primitive */
+	/** How Polygroups are assigned to shape primitives. */
 	UPROPERTY(EditAnywhere, Category = ShapeSettings, meta = (ProceduralShapeSetting))
 	EMakeMeshPolygroupMode PolygroupMode = EMakeMeshPolygroupMode::PerFace;
 
-	/** How to place Primitive in the Scene */
-	UPROPERTY(EditAnywhere, Category = Positioning, meta = (DisplayName = "Target Surface"))
-	EMakeMeshPlacementType PlaceMode = EMakeMeshPlacementType::OnScene;
+	/** How the shape is placed in the scene. */
+	UPROPERTY(EditAnywhere, Category = Positioning)
+	EMakeMeshPlacementType TargetSurface = EMakeMeshPlacementType::OnScene;
 
-	/** If true, placement location will be snapped to grid. Only relevant when coordinate system is set to World. */
+	/** If true, the shape pivot is snapped to the grid. This is only relevant if the coordinate system is set to world space. */
 	UPROPERTY(EditAnywhere, Category = Positioning)
 	bool bSnapToGrid = true;
 
-	/** Location of Pivot within Primitive */
+	/** Location of pivot within the shape */
 	UPROPERTY(EditAnywhere, Category = Positioning, meta = (ProceduralShapeSetting))
 	EMakeMeshPivotLocation PivotLocation = EMakeMeshPivotLocation::Base;
 
-	/** Rotation of Primitive around up axis */
-	UPROPERTY(EditAnywhere, Category = Positioning, meta = (DisplayName = "Rotation", UIMin = "0.0", UIMax = "360.0"))
+	/** Rotation of the shape around its up axis */
+	UPROPERTY(EditAnywhere, Category = Positioning, meta = (UIMin = "0.0", UIMax = "360.0"))
 	float Rotation = 0.0;
 
-	/** Align shape to Placement Surface */
-	UPROPERTY(EditAnywhere, Category = Positioning, meta = (DisplayName = "Align to Normal", EditCondition = "PlaceMode == EMakeMeshPlacementType::OnScene"))
-	bool bAlignShapeToPlacementSurface = true;
-
-	///** Start Angle of Shape */
-	//UPROPERTY(EditAnywhere, Category = ShapeSettings, meta = (DisplayName = "Start Angle", UIMin = "0.0", UIMax = "360.0", ClampMin = "-10000", ClampMax = "10000.0"))
-	//float StartAngle = 0.f;
-
-	///** End Angle of Shape */
-	//UPROPERTY(EditAnywhere, Category = ShapeSettings, meta = (DisplayName = "End Angle", UIMin = "0.0", UIMax = "360.0", ClampMin = "-10000", ClampMax = "10000.0"))
-	//float EndAngle = 360.f;
+	/** If true, aligns the shape along the normal of the surface it is placed on. */
+	UPROPERTY(EditAnywhere, Category = Positioning, meta = (EditCondition = "TargetSurface == EMakeMeshPlacementType::OnScene"))
+	bool bAlignToNormal = true;
 
 	bool IsEquivalent( const UProceduralShapeToolProperties* ) const;
 };
@@ -124,111 +113,116 @@ UCLASS()
 class MESHMODELINGTOOLS_API UProceduralBoxToolProperties : public UProceduralShapeToolProperties
 {
 	GENERATED_BODY()
-public:
-	UProceduralBoxToolProperties() = default;
 
-	/** Width of Shape */
-	UPROPERTY(EditAnywhere, Category = "ShapeSettings", meta = (DisplayName = "Width", UIMin = "1.0", UIMax = "1000.0", ClampMin = "0.0001", ClampMax = "1000000.0", ProceduralShapeSetting))
+public:
+	/** Width of the box */
+	UPROPERTY(EditAnywhere, Category = ShapeSettings, meta = (UIMin = "1.0", UIMax = "1000.0", ClampMin = "0.0001", ClampMax = "1000000.0", ProceduralShapeSetting))
 	float Width = 100.f;
 
-	/** Depth of Shape */
-	UPROPERTY(EditAnywhere, Category = "ShapeSettings", meta = (DisplayName = "Depth", UIMin = "1.0", UIMax = "1000.0", ClampMin = "0.0001", ClampMax = "1000000.0", ProceduralShapeSetting))
+	/** Depth of the box */
+	UPROPERTY(EditAnywhere, Category = ShapeSettings, meta = (UIMin = "1.0", UIMax = "1000.0", ClampMin = "0.0001", ClampMax = "1000000.0", ProceduralShapeSetting))
 	float Depth = 100.f;
 
-	/** Height of Shape */
-	UPROPERTY(EditAnywhere, Category = "ShapeSettings", meta = (DisplayName = "Height", UIMin = "1.0", UIMax = "1000.0", ClampMin = "0.0001", ClampMax = "1000000.0", ProceduralShapeSetting))
+	/** Height of the box */
+	UPROPERTY(EditAnywhere, Category = ShapeSettings, meta = (UIMin = "1.0", UIMax = "1000.0", ClampMin = "0.0001", ClampMax = "1000000.0", ProceduralShapeSetting))
 	float Height = 100.f;
 
-	/** Number of Subdivisions Along the Width */
-	UPROPERTY(EditAnywhere, Category = "ShapeSettings", meta = (DisplayName = "Width Subdivisions", UIMin = "1", UIMax = "100", ClampMin = "1", ClampMax = "500", ProceduralShapeSetting))
+	/** Number of subdivisions along the width */
+	UPROPERTY(EditAnywhere, Category = ShapeSettings, meta = (UIMin = "1", UIMax = "100", ClampMin = "1", ClampMax = "500", ProceduralShapeSetting))
 	int WidthSubdivisions = 1;
 
-	/** Number of Subdivisions Along the Depth */
-	UPROPERTY(EditAnywhere, Category = "ShapeSettings", meta = (DisplayName = "Depth Subdivisions", UIMin = "1", UIMax = "100", ClampMin = "1", ClampMax = "500", ProceduralShapeSetting))
+	/** Number of subdivisions along the depth */
+	UPROPERTY(EditAnywhere, Category = ShapeSettings, meta = (UIMin = "1", UIMax = "100", ClampMin = "1", ClampMax = "500", ProceduralShapeSetting))
 	int DepthSubdivisions = 1;
 
-	/** Number of Subdivisions Along the Height */
-	UPROPERTY(EditAnywhere, Category = "ShapeSettings", meta = (DisplayName = "Height Subdivisions", UIMin = "1", UIMax = "100", ClampMin = "1", ClampMax = "500", ProceduralShapeSetting))
+	/** Number of subdivisions along the height */
+	UPROPERTY(EditAnywhere, Category = ShapeSettings, meta = (UIMin = "1", UIMax = "100", ClampMin = "1", ClampMax = "500", ProceduralShapeSetting))
 	int HeightSubdivisions = 1;
 };
 
 UENUM()
 enum class EProceduralRectType
 {
-	/** Create a Rectangle */
-	Rectangle UMETA(DisplayName = "Rectangle"),
-	/** Create a Rounded Rectangle */
-	RoundedRectangle UMETA(DisplayName = "Rounded Rectangle")
+	/** Create a rectangle */
+	Rectangle,
+	/** Create a rounded rectangle */
+	RoundedRectangle
 };
 
 UCLASS()
 class MESHMODELINGTOOLS_API UProceduralRectangleToolProperties : public UProceduralShapeToolProperties
 {
 	GENERATED_BODY()
-public:
-	UProceduralRectangleToolProperties() = default;
 
-	/** Type of rectangle to create */
-	UPROPERTY(EditAnywhere, Category = "ShapeSettings")
-	EProceduralRectType RectType = EProceduralRectType::Rectangle;
+public:
+	/** Type of rectangle */
+	UPROPERTY(EditAnywhere, Category = ShapeSettings)
+	EProceduralRectType RectangleType = EProceduralRectType::Rectangle;
 
-	/** Width of Shape */
-	UPROPERTY(EditAnywhere, Category = "ShapeSettings", meta = (DisplayName = "Width", UIMin = "1.0", UIMax = "1000.0", ClampMin = "0.0001", ClampMax = "1000000.0", ProceduralShapeSetting))
+	/** Width of the rectangle */
+	UPROPERTY(EditAnywhere, Category = ShapeSettings, meta = (UIMin = "1.0", UIMax = "1000.0", ClampMin = "0.0001", ClampMax = "1000000.0", ProceduralShapeSetting))
 	float Width = 100.f;
 
-	/** Depth of Shape */
-	UPROPERTY(EditAnywhere, Category = "ShapeSettings", meta = (DisplayName = "Depth", UIMin = "1.0", UIMax = "1000.0", ClampMin = "0.0001", ClampMax = "1000000.0", ProceduralShapeSetting))
+	/** Depth of the rectangle */
+	UPROPERTY(EditAnywhere, Category = ShapeSettings, meta = (UIMin = "1.0", UIMax = "1000.0", ClampMin = "0.0001", ClampMax = "1000000.0", ProceduralShapeSetting))
 	float Depth = 100.f;
 
-	/** Number of Subdivisions Along the Width */
-	UPROPERTY(EditAnywhere, Category = "ShapeSettings", meta = (DisplayName = "Width Subdivisions", UIMin = "1", UIMax = "100", ClampMin = "1", ClampMax = "500", ProceduralShapeSetting))
+	/** Number of subdivisions along the width */
+	UPROPERTY(EditAnywhere, Category = ShapeSettings, meta = (UIMin = "1", UIMax = "100", ClampMin = "1", ClampMax = "500", ProceduralShapeSetting))
 	int WidthSubdivisions = 1;
 
-	/** Number of Subdivisions Along the Depth */
-	UPROPERTY(EditAnywhere, Category = "ShapeSettings", meta = (DisplayName = "Depth Subdivisions", UIMin = "1", UIMax = "100", ClampMin = "1", ClampMax = "500", ProceduralShapeSetting))
+	/** Number of subdivisions along the depth */
+	UPROPERTY(EditAnywhere, Category = ShapeSettings, meta = (UIMin = "1", UIMax = "100", ClampMin = "1", ClampMax = "500", ProceduralShapeSetting))
 	int DepthSubdivisions = 1;
 
-	/** Radius of Rounded Corners */
-	UPROPERTY(EditAnywhere, Category = "ShapeSettings", meta = (DisplayName = "Corner Radius", UIMin = "1.0", UIMax = "1000.0", ClampMin = "0.0001", ClampMax = "1000000.0", ProceduralShapeSetting, EditCondition = "RectType == EProceduralRectType::RoundedRectangle", EditConditionHides))
+	/** Radius of rounded corners */
+	UPROPERTY(EditAnywhere, Category = ShapeSettings,
+		meta = (UIMin = "1.0", UIMax = "1000.0", ClampMin = "0.0001", ClampMax = "1000000.0", ProceduralShapeSetting,
+			EditCondition = "RectangleType == EProceduralRectType::RoundedRectangle", EditConditionHides))
 	float CornerRadius = 25.f;
 
-	/** Number of Angular Slices in Each Rounded Corner */
-	UPROPERTY(EditAnywhere, Category = "ShapeSettings", meta = (DisplayName = "Corner Slices", UIMin = "3", UIMax = "128", ClampMin = "3", ClampMax = "500", ProceduralShapeSetting, EditCondition = "RectType == EProceduralRectType::RoundedRectangle", EditConditionHides))
+	/** Number of radial slices for each rounded corner */
+	UPROPERTY(EditAnywhere, Category = ShapeSettings,
+		meta = (UIMin = "3", UIMax = "128", ClampMin = "3", ClampMax = "500", ProceduralShapeSetting,
+			EditCondition = "RectangleType == EProceduralRectType::RoundedRectangle", EditConditionHides))
 	int CornerSlices = 16;
 };
 
 UENUM()
 enum class EProceduralDiscType
 {
-	/** Create a solid Disc */
-	Disc UMETA(DisplayName = "Disc"),
-	/** Create a Disc with a hole */
-	PuncturedDisc UMETA(DisplayName = "Punctured Disc")
+	/** Create a disc */
+	Disc,
+	/** Create a disc with a hole */
+	PuncturedDisc
 };
 
 UCLASS()
 class MESHMODELINGTOOLS_API UProceduralDiscToolProperties : public UProceduralShapeToolProperties
 {
 	GENERATED_BODY()
+
 public:
-	/** Type of disc to create */
-	UPROPERTY(EditAnywhere, Category = "ShapeSettings")
+	/** Type of disc */
+	UPROPERTY(EditAnywhere, Category = ShapeSettings)
 	EProceduralDiscType DiscType = EProceduralDiscType::Disc;
 
-	/** Radius of Disc */
-	UPROPERTY(EditAnywhere, Category = "ShapeSettings", meta = (DisplayName = "Radius", UIMin = "1.0", UIMax = "1000.0", ClampMin = "0.0001", ClampMax = "1000000.0", ProceduralShapeSetting))
+	/** Radius of the disc */
+	UPROPERTY(EditAnywhere, Category = ShapeSettings, meta = (UIMin = "1.0", UIMax = "1000.0", ClampMin = "0.0001", ClampMax = "1000000.0", ProceduralShapeSetting))
 	float Radius = 50.f;
 
-	/** Number of Angular Slices around the Disc */
-	UPROPERTY(EditAnywhere, Category = "ShapeSettings", meta = (DisplayName = "Radial Slices", UIMin = "3", UIMax = "128", ClampMin = "3", ClampMax = "500", ProceduralShapeSetting))
+	/** Number of radial slices for the disc */
+	UPROPERTY(EditAnywhere, Category = ShapeSettings, meta = (UIMin = "3", UIMax = "128", ClampMin = "3", ClampMax = "500", ProceduralShapeSetting))
 	int RadialSlices = 16;
 
-	/** Number of Radial Subdivisions in the Disc */
-	UPROPERTY(EditAnywhere, Category = "ShapeSettings", meta = (DisplayName = "Radial Subdivisions", UIMin = "1", UIMax = "100", ClampMin = "1", ClampMax = "500", ProceduralShapeSetting))
+	/** Number of radial subdivisions for each radial slice of the disc */
+	UPROPERTY(EditAnywhere, Category = ShapeSettings, meta = (UIMin = "1", UIMax = "100", ClampMin = "1", ClampMax = "500", ProceduralShapeSetting))
 	int RadialSubdivisions = 1;
 
-	/** Radius of the Disc's Hole */
-	UPROPERTY(EditAnywhere, Category = "ShapeSettings", meta = (DisplayName = "Hole Radius", UIMin = "1.0", UIMax = "1000.0", ClampMin = "0.0001", ClampMax = "1000000.0", ProceduralShapeSetting, EditCondition = "DiscType == EProceduralDiscType::PuncturedDisc", EditConditionHides))
+	/** Radius of the hole in the disc */
+	UPROPERTY(EditAnywhere, Category = ShapeSettings,
+		meta = (UIMin = "1.0", UIMax = "1000.0", ClampMin = "0.0001", ClampMax = "1000000.0", ProceduralShapeSetting,
+			EditCondition = "DiscType == EProceduralDiscType::PuncturedDisc", EditConditionHides))
 	float HoleRadius = 25.f;
 };
 
@@ -236,44 +230,46 @@ UCLASS()
 class MESHMODELINGTOOLS_API UProceduralTorusToolProperties : public UProceduralShapeToolProperties
 {
 	GENERATED_BODY()
+
 public:
-	/** Radius from the Torus Center to the Center of the Torus Tube */
-	UPROPERTY(EditAnywhere, Category = "ShapeSettings", meta = (DisplayName = "Major Radius", UIMin = "1.0", UIMax = "1000.0", ClampMin = "0.0001", ClampMax = "1000000.0", ProceduralShapeSetting))
+	/** Major radius of the torus, measured from the torus center to the center of the torus tube */
+	UPROPERTY(EditAnywhere, Category = ShapeSettings, meta = (UIMin = "1.0", UIMax = "1000.0", ClampMin = "0.0001", ClampMax = "1000000.0", ProceduralShapeSetting))
 	float MajorRadius = 50.f;
 
-	/** Radius of the Torus Tube */
-	UPROPERTY(EditAnywhere, Category = "ShapeSettings", meta = (DisplayName = "Minor Radius", UIMin = "1.0", UIMax = "1000.0", ClampMin = "0.0001", ClampMax = "1000000.0", ProceduralShapeSetting))
+	/** Minor radius of the torus, measured from the center ot the torus tube to the tube surface */
+	UPROPERTY(EditAnywhere, Category = ShapeSettings, meta = (UIMin = "1.0", UIMax = "1000.0", ClampMin = "0.0001", ClampMax = "1000000.0", ProceduralShapeSetting))
 	float MinorRadius = 25.f;
 
-	/** Number of Angular Slices Along the Torus Tube */
-	UPROPERTY(EditAnywhere, Category = "ShapeSettings", meta = (DisplayName = "Major Slices", UIMin = "3", UIMax = "128", ClampMin = "3", ClampMax = "500", ProceduralShapeSetting))
-	int TubeSlices = 16;
+	/** Number of radial slices along the torus tube */
+	UPROPERTY(EditAnywhere, Category = ShapeSettings, meta = (UIMin = "3", UIMax = "128", ClampMin = "3", ClampMax = "500", ProceduralShapeSetting))
+	int MajorSlices = 16;
 
-	/** Number of Angular Slices Around the Tube of the Torus */
-	UPROPERTY(EditAnywhere, Category = "ShapeSettings", meta = (DisplayName = "Minor Slices", UIMin = "3", UIMax = "128", ClampMin = "3", ClampMax = "500", ProceduralShapeSetting))
-	int CrossSectionSlices = 16;
+	/** Number of radial slices around the torus tube */
+	UPROPERTY(EditAnywhere, Category = ShapeSettings, meta = (UIMin = "3", UIMax = "128", ClampMin = "3", ClampMax = "500", ProceduralShapeSetting))
+	int MinorSlices = 16;
 };
 
 UCLASS()
 class MESHMODELINGTOOLS_API UProceduralCylinderToolProperties : public UProceduralShapeToolProperties
 {
 	GENERATED_BODY()
+
 public:
 
-	/** Radius of The Cylinder */
-	UPROPERTY(EditAnywhere, Category = "ShapeSettings", meta = (DisplayName = "Radius", UIMin = "1.0", UIMax = "1000.0", ClampMin = "0.0001", ClampMax = "1000000.0", ProceduralShapeSetting))
+	/** Radius of the cylinder */
+	UPROPERTY(EditAnywhere, Category = ShapeSettings, meta = (UIMin = "1.0", UIMax = "1000.0", ClampMin = "0.0001", ClampMax = "1000000.0", ProceduralShapeSetting))
 	float Radius = 50.f;
 
-	/** Height of Cylinder */
-	UPROPERTY(EditAnywhere, Category = "ShapeSettings", meta = (DisplayName = "Height", UIMin = "1.0", UIMax = "1000.0", ClampMin = "0.0001", ClampMax = "1000000.0", ProceduralShapeSetting))
+	/** Height of the cylinder */
+	UPROPERTY(EditAnywhere, Category = ShapeSettings, meta = (UIMin = "1.0", UIMax = "1000.0", ClampMin = "0.0001", ClampMax = "1000000.0", ProceduralShapeSetting))
 	float Height = 200.f;
 
-	/** Number of Slices on the Cylinder Caps */
-	UPROPERTY(EditAnywhere, Category = "ShapeSettings", meta = (DisplayName = "Radial Slices", UIMin = "3", UIMax = "128", ClampMin = "3", ClampMax = "500", ProceduralShapeSetting))
+	/** Number of radial slices for the cylinder */
+	UPROPERTY(EditAnywhere, Category = ShapeSettings, meta = (UIMin = "3", UIMax = "128", ClampMin = "3", ClampMax = "500", ProceduralShapeSetting))
 	int RadialSlices = 16;
 
-	/** Number of Vertical Subdivisions Along the Height of the Cylinder */
-	UPROPERTY(EditAnywhere, Category = "ShapeSettings", meta = (DisplayName = "Height Subdivisions", UIMin = "1", UIMax = "100", ClampMin = "1", ClampMax = "500", ProceduralShapeSetting))
+	/** Number of subdivisions along the height of the cylinder */
+	UPROPERTY(EditAnywhere, Category = ShapeSettings, meta = (UIMin = "1", UIMax = "100", ClampMin = "1", ClampMax = "500", ProceduralShapeSetting))
 	int HeightSubdivisions = 1;
 };
 
@@ -281,21 +277,22 @@ UCLASS()
 class MESHMODELINGTOOLS_API UProceduralConeToolProperties : public UProceduralShapeToolProperties
 {
 	GENERATED_BODY()
+
 public:
-	/** Radius of the Cone */
-	UPROPERTY(EditAnywhere, Category = "ShapeSettings", meta = (DisplayName = "Radius", UIMin = "1.0", UIMax = "1000.0", ClampMin = "0.0001", ClampMax = "1000000.0", ProceduralShapeSetting))
+	/** Radius of the cone */
+	UPROPERTY(EditAnywhere, Category = ShapeSettings, meta = (UIMin = "1.0", UIMax = "1000.0", ClampMin = "0.0001", ClampMax = "1000000.0", ProceduralShapeSetting))
 	float Radius = 50.f;
 
-	/** Height of Cone */
-	UPROPERTY(EditAnywhere, Category = "ShapeSettings", meta = (DisplayName = "Height", UIMin = "1.0", UIMax = "1000.0", ClampMin = "0.0001", ClampMax = "1000000.0", ProceduralShapeSetting))
+	/** Height of the cone */
+	UPROPERTY(EditAnywhere, Category = ShapeSettings, meta = (UIMin = "1.0", UIMax = "1000.0", ClampMin = "0.0001", ClampMax = "1000000.0", ProceduralShapeSetting))
 	float Height = 200.f;
 
-	/** Number of Slices on the Cone Base */
-	UPROPERTY(EditAnywhere, Category = "ShapeSettings", meta = (DisplayName = "Radial Slices", UIMin = "3", UIMax = "128", ClampMin = "3", ClampMax = "500", ProceduralShapeSetting))
+	/** Number of radial slices for the cylinder */
+	UPROPERTY(EditAnywhere, Category = ShapeSettings, meta = (UIMin = "3", UIMax = "128", ClampMin = "3", ClampMax = "500", ProceduralShapeSetting))
 	int RadialSlices = 16;
 
-	/** Number of Vertical Subdivisions Along the Hight of the Cone */
-	UPROPERTY(EditAnywhere, Category = "ShapeSettings", meta = (DisplayName = "Height Subdivisions", UIMin = "1", UIMax = "100", ClampMin = "1", ClampMax = "500", ProceduralShapeSetting))
+	/** Number of subdivisions along the height of the cone */
+	UPROPERTY(EditAnywhere, Category = ShapeSettings, meta = (UIMin = "1", UIMax = "100", ClampMin = "1", ClampMax = "500", ProceduralShapeSetting))
 	int HeightSubdivisions = 1;
 };
 
@@ -303,30 +300,31 @@ UCLASS()
 class MESHMODELINGTOOLS_API UProceduralArrowToolProperties : public UProceduralShapeToolProperties
 {
 	GENERATED_BODY()
+
 public:
-	/** Radius of the Arrow Shaft */
-	UPROPERTY(EditAnywhere, Category = "ShapeSettings", meta = (DisplayName = "Shaft Radius", UIMin = "1.0", UIMax = "1000.0", ClampMin = "0.0001", ClampMax = "1000000.0", ProceduralShapeSetting))
+	/** Radius of the arrow shaft */
+	UPROPERTY(EditAnywhere, Category = ShapeSettings, meta = (UIMin = "1.0", UIMax = "1000.0", ClampMin = "0.0001", ClampMax = "1000000.0", ProceduralShapeSetting))
 	float ShaftRadius = 20.f;
 
-	/** Height of Arrow Shaft */
-	UPROPERTY(EditAnywhere, Category = "ShapeSettings", meta = (DisplayName = "Shaft Height", UIMin = "1.0", UIMax = "1000.0", ClampMin = "0.0001", ClampMax = "1000000.0", ProceduralShapeSetting))
+	/** Height of arrow shaft */
+	UPROPERTY(EditAnywhere, Category = ShapeSettings, meta = (UIMin = "1.0", UIMax = "1000.0", ClampMin = "0.0001", ClampMax = "1000000.0", ProceduralShapeSetting))
 	float ShaftHeight = 200.f;
 
-	/** Radius of the Arrow Head */
-	UPROPERTY(EditAnywhere, Category = "ShapeSettings", meta = (DisplayName = "Head Radius", UIMin = "1.0", UIMax = "1000.0", ClampMin = "0.0001", ClampMax = "1000000.0", ProceduralShapeSetting))
+	/** Radius of the arrow head base */
+	UPROPERTY(EditAnywhere, Category = ShapeSettings, meta = (UIMin = "1.0", UIMax = "1000.0", ClampMin = "0.0001", ClampMax = "1000000.0", ProceduralShapeSetting))
 	float HeadRadius = 60.f;
 
-	/** Height of Arrow's Head */
-	UPROPERTY(EditAnywhere, Category = "ShapeSettings", meta = (DisplayName = "Head Height", UIMin = "1.0", UIMax = "1000.0", ClampMin = "0.0001", ClampMax = "1000000.0", ProceduralShapeSetting))
+	/** Height of arrow head */
+	UPROPERTY(EditAnywhere, Category = ShapeSettings, meta = (UIMin = "1.0", UIMax = "1000.0", ClampMin = "0.0001", ClampMax = "1000000.0", ProceduralShapeSetting))
 	float HeadHeight = 120.f;
 
-	/** Number of Angular Slices Around the Arrow */
-	UPROPERTY(EditAnywhere, Category = "ShapeSettings", meta = (DisplayName = "Radial Slices", UIMin = "3", UIMax = "100", ClampMin = "3", ClampMax = "500", ProceduralShapeSetting))
+	/** Number of radial slices for the arrow */
+	UPROPERTY(EditAnywhere, Category = ShapeSettings, meta = (UIMin = "3", UIMax = "100", ClampMin = "3", ClampMax = "500", ProceduralShapeSetting))
 	int RadialSlices = 16;
 
-	/** Number of Vertical Subdivisions Along in the Arrow */
-	UPROPERTY(EditAnywhere, Category = "ShapeSettings", meta = (DisplayName = "Total Subdivisions", UIMin = "1", UIMax = "100", ClampMin = "1", ClampMax = "500", ProceduralShapeSetting))
-	int TotalSubdivisions = 1;
+	/** Number of subdivisions along each part of the arrow, i.e. shaft, head base, head cone */
+	UPROPERTY(EditAnywhere, Category = ShapeSettings, meta = (UIMin = "1", UIMax = "100", ClampMin = "1", ClampMax = "500", ProceduralShapeSetting))
+	int HeightSubdivisions = 1;
 };
 
 UENUM()
@@ -342,25 +340,29 @@ UCLASS()
 class MESHMODELINGTOOLS_API UProceduralSphereToolProperties : public UProceduralShapeToolProperties
 {
 	GENERATED_BODY()
+
 public:
-	/** Type of Sphere to create */
-	UPROPERTY(EditAnywhere, Category = "ShapeSettings")
+	/** Type of sphere */
+	UPROPERTY(EditAnywhere, Category = ShapeSettings)
 	EProceduralSphereType SphereType = EProceduralSphereType::Box;
 
-	/** Radius of the Sphere */
-	UPROPERTY(EditAnywhere, Category = "ShapeSettings", meta = (DisplayName = "Radius", UIMin = "1.0", UIMax = "1000.0", ClampMin = "0.0001", ClampMax = "1000000.0", ProceduralShapeSetting))
+	/** Radius of the sphere */
+	UPROPERTY(EditAnywhere, Category = ShapeSettings, meta = (UIMin = "1.0", UIMax = "1000.0", ClampMin = "0.0001", ClampMax = "1000000.0", ProceduralShapeSetting))
 	float Radius = 50.f;
 
-	/** Number of Latitudinal Slices of the Sphere */
-	UPROPERTY(EditAnywhere, Category = "ShapeSettings", meta = (DisplayName = "Latitude Slices", UIMin = "3", UIMax = "100", ClampMin = "4", ClampMax = "500", ProceduralShapeSetting, EditCondition = "SphereType == EProceduralSphereType::LatLong", EditConditionHides))
-	int LatitudeSlices = 16;
+	/** Number of horizontal slices of the sphere */
+	UPROPERTY(EditAnywhere, Category = ShapeSettings, meta = (UIMin = "3", UIMax = "100", ClampMin = "4", ClampMax = "500", ProceduralShapeSetting,
+		EditCondition = "SphereType == EProceduralSphereType::LatLong", EditConditionHides))
+	int HorizontalSlices = 16;
 
-	/** Number of Longitudinal Slices around the Sphere */
-	UPROPERTY(EditAnywhere, Category = "ShapeSettings", meta = (DisplayName = "Longitude Slices", UIMin = "3", UIMax = "100", ClampMin = "4", ClampMax = "500", ProceduralShapeSetting, EditCondition = "SphereType == EProceduralSphereType::LatLong", EditConditionHides))
-	int LongitudeSlices = 16;
+	/** Number of vertical slices of the sphere */
+	UPROPERTY(EditAnywhere, Category = ShapeSettings, meta = (UIMin = "3", UIMax = "100", ClampMin = "4", ClampMax = "500", ProceduralShapeSetting,
+		EditCondition = "SphereType == EProceduralSphereType::LatLong", EditConditionHides))
+	int VerticalSlices = 16;
 
-	/** Number of Subdivisions of each Side of the Sphere */
-	UPROPERTY(EditAnywhere, Category = "ShapeSettings", meta = (DisplayName = "Side Subdivisions", UIMin = "1", UIMax = "100", ClampMin = "1", ClampMax = "500", ProceduralShapeSetting, EditCondition = "SphereType == EProceduralSphereType::Box", EditConditionHides))
+	/** Number of subdivisions for each side of the sphere */
+	UPROPERTY(EditAnywhere, Category = ShapeSettings, meta = (UIMin = "1", UIMax = "100", ClampMin = "1", ClampMax = "500", ProceduralShapeSetting,
+		EditCondition = "SphereType == EProceduralSphereType::Box", EditConditionHides))
 	int Subdivisions = 16;
 };
 
@@ -368,50 +370,59 @@ UENUM()
 enum class EProceduralStairsType
 {
 	/** Create a linear staircase */
-	Linear UMETA(DisplayName = "Linear"),
+	Linear,
 	/** Create a floating staircase */
-	Floating UMETA(DisplayName = "Floating"),
+	Floating,
 	/** Create a curved staircase */
-	Curved UMETA(DisplayName = "Curved"),
+	Curved,
 	/** Create a spiral staircase */
-	Spiral UMETA(DisplayName = "Spiral")
+	Spiral
 };
 
 UCLASS()
 class MESHMODELINGTOOLS_API UProceduralStairsToolProperties : public UProceduralShapeToolProperties
 {
 	GENERATED_BODY()
+
 public:
-	/** Type of staircase to create */
-	UPROPERTY(EditAnywhere, Category = "ShapeSettings")
-		EProceduralStairsType StairsType = EProceduralStairsType::Linear;
+	/** Type of staircase */
+	UPROPERTY(EditAnywhere, Category = ShapeSettings)
+	EProceduralStairsType StairsType = EProceduralStairsType::Linear;
 
-	/** Number of Steps of Shape */
-	UPROPERTY(EditAnywhere, Category = "ShapeSettings", meta = (DisplayName = "Number of Steps", UIMin = "1", UIMax = "100", ClampMin = "1", ClampMax = "1000000", ProceduralShapeSetting))
+	/** Number of steps */
+	UPROPERTY(EditAnywhere, Category = ShapeSettings, meta = (DisplayName = "Number of Steps", UIMin = "1", UIMax = "100", ClampMin = "1", ClampMax = "1000000", ProceduralShapeSetting))
 	int NumSteps = 8;
 
 	/** Width of each step */
-	UPROPERTY(EditAnywhere, Category = "ShapeSettings", meta = (DisplayName = "Step Width", UIMin = "1.0", UIMax = "1000.0", ClampMin = "0.0001", ClampMax = "1000000.0", ProceduralShapeSetting))
+	UPROPERTY(EditAnywhere, Category = ShapeSettings, meta = (UIMin = "1.0", UIMax = "1000.0", ClampMin = "0.0001", ClampMax = "1000000.0", ProceduralShapeSetting))
 	float StepWidth = 150.0f;
 
 	/** Height of each step */
-	UPROPERTY(EditAnywhere, Category = "ShapeSettings", meta = (DisplayName = "Step Height", UIMin = "1.0", UIMax = "1000.0", ClampMin = "0.0001", ClampMax = "1000000.0", ProceduralShapeSetting))
+	UPROPERTY(EditAnywhere, Category = ShapeSettings, meta = (UIMin = "1.0", UIMax = "1000.0", ClampMin = "0.0001", ClampMax = "1000000.0", ProceduralShapeSetting))
 	float StepHeight = 20.0f;
 
-	/** Depth of each step for linear stair shapes */
-	UPROPERTY(EditAnywhere, Category = "ShapeSettings", meta = (DisplayName = "Step Depth", UIMin = "1.0", UIMax = "1000.0", ClampMin = "0.0001", ClampMax = "1000000.0", ProceduralShapeSetting, EditCondition = "StairsType == EProceduralStairsType::Linear || StairsType == EProceduralStairsType::Floating", EditConditionHides))
+	/** Depth of each step for linear stairs */
+	UPROPERTY(EditAnywhere, Category = ShapeSettings,
+		meta = (UIMin = "1.0", UIMax = "1000.0", ClampMin = "0.0001", ClampMax = "1000000.0", ProceduralShapeSetting,
+			EditCondition =	"StairsType == EProceduralStairsType::Linear || StairsType == EProceduralStairsType::Floating", EditConditionHides))
 	float StepDepth = 30.0f;
 
-	/** Angular length of curved stair shapes. Positive values are clockwise, negative are counter-clockwise. */
-	UPROPERTY(EditAnywhere, Category = "ShapeSettings", meta = (DisplayName = "Curve Angle", UIMin = "-360.0", UIMax = "360.0", ClampMin = "-360.0", ClampMax = "360.0", ProceduralShapeSetting, EditCondition = "StairsType == EProceduralStairsType::Curved", EditConditionHides))
+	/** Angular length of curved stairs; positive values are for clockwise and negative values are for counterclockwise. */
+	UPROPERTY(EditAnywhere, Category = ShapeSettings,
+		meta = (UIMin = "-360.0", UIMax = "360.0", ClampMin = "-360.0", ClampMax = "360.0", ProceduralShapeSetting,
+			EditCondition =	"StairsType == EProceduralStairsType::Curved", EditConditionHides))
 	float CurveAngle = 90.0f;
 
-	/** Angular length of spiral stair shapes. Positive values are clockwise, negative are counter-clockwise. */
-	UPROPERTY(EditAnywhere, Category = "ShapeSettings", meta = (DisplayName = "Spiral Angle", UIMin = "-720.0", UIMax = "720.0", ClampMin = "-360000.0", ClampMax = "360000.0", ProceduralShapeSetting, EditCondition = "StairsType == EProceduralStairsType::Spiral", EditConditionHides))
+	/** Angular length of spiral stairs; positive values are for clockwise and negative values are for counterclockwise. */
+	UPROPERTY(EditAnywhere, Category = ShapeSettings,
+		meta = (UIMin = "-720.0", UIMax = "720.0", ClampMin = "-360000.0", ClampMax = "360000.0", ProceduralShapeSetting,
+			EditCondition =	"StairsType == EProceduralStairsType::Spiral", EditConditionHides))
 	float SpiralAngle = 90.0f;
 
-	/** Inner radius for curved/spiral stair shapes */
-	UPROPERTY(EditAnywhere, Category = "ShapeSettings", meta = (DisplayName = "Inner Radius", UIMin = "1.0", UIMax = "1000.0", ClampMin = "0.0001", ClampMax = "1000000.0", ProceduralShapeSetting, EditCondition = "StairsType == EProceduralStairsType::Curved || StairsType == EProceduralStairsType::Spiral", EditConditionHides))
+	/** Inner radius for curved and spiral stairs */
+	UPROPERTY(EditAnywhere, Category = ShapeSettings,
+		meta = (UIMin = "1.0", UIMax = "1000.0", ClampMin = "0.0001", ClampMax = "1000000.0", ProceduralShapeSetting,
+			EditCondition = "StairsType == EProceduralStairsType::Curved || StairsType == EProceduralStairsType::Spiral", EditConditionHides))
 	float InnerRadius = 150.0f;
 };
 
@@ -436,7 +447,7 @@ public:
 	UPROPERTY()
 	TObjectPtr<UNewMeshMaterialProperties> MaterialProperties;
 
-	bool IsInvalid()
+	bool IsInvalid() const
 	{
 		return Actor == nullptr || StaticMesh == nullptr || ShapeSettings == nullptr || MaterialProperties == nullptr;
 	}
@@ -451,7 +462,7 @@ class MESHMODELINGTOOLS_API UAddPrimitiveTool : public USingleClickTool, public
 	GENERATED_BODY()
 
 public:
-	UAddPrimitiveTool(const FObjectInitializer&);
+	explicit UAddPrimitiveTool(const FObjectInitializer&);
 
 	virtual void SetWorld(UWorld* World);
 
@@ -494,7 +505,7 @@ protected:
 	/**
 	 * Checks if the passed-in settings would create the same asset as the current settings
 	 */
-	bool IsEquivalentLastGeneratedAsset()
+	bool IsEquivalentLastGeneratedAsset() const
 	{
 		if (LastGenerated == nullptr || LastGenerated->IsInvalid())
 		{
@@ -515,13 +526,13 @@ protected:
 	UPROPERTY()
 	FString AssetName = TEXT("GeneratedAsset");
 
-protected:
+	UPROPERTY()
 	UWorld* TargetWorld;
 
 	void UpdatePreviewPosition(const FInputDeviceRay& ClickPos);
 	UE::Geometry::FFrame3f ShapeFrame;
 
-	void UpdatePreviewMesh();
+	void UpdatePreviewMesh() const;
 };
 
 
@@ -530,9 +541,9 @@ class UAddBoxPrimitiveTool : public UAddPrimitiveTool
 {
 	GENERATED_BODY()
 public:
-	UAddBoxPrimitiveTool(const FObjectInitializer& ObjectInitializer);
+	explicit UAddBoxPrimitiveTool(const FObjectInitializer& ObjectInitializer);
 protected:
-	void GenerateMesh(FDynamicMesh3* OutMesh) const override;
+	virtual void GenerateMesh(FDynamicMesh3* OutMesh) const override;
 };
 
 UCLASS()
@@ -540,9 +551,9 @@ class UAddCylinderPrimitiveTool : public UAddPrimitiveTool
 {
 	GENERATED_BODY()
 public:
-	UAddCylinderPrimitiveTool(const FObjectInitializer& ObjectInitializer);
+	explicit UAddCylinderPrimitiveTool(const FObjectInitializer& ObjectInitializer);
 protected:
-	void GenerateMesh(FDynamicMesh3* OutMesh) const override;
+	virtual void GenerateMesh(FDynamicMesh3* OutMesh) const override;
 };
 
 UCLASS()
@@ -550,9 +561,9 @@ class UAddConePrimitiveTool : public UAddPrimitiveTool
 {
 	GENERATED_BODY()
 public:
-	UAddConePrimitiveTool(const FObjectInitializer& ObjectInitializer);
+	explicit UAddConePrimitiveTool(const FObjectInitializer& ObjectInitializer);
 protected:
-	void GenerateMesh(FDynamicMesh3* OutMesh) const override;
+	virtual void GenerateMesh(FDynamicMesh3* OutMesh) const override;
 };
 
 UCLASS()
@@ -560,9 +571,9 @@ class UAddRectanglePrimitiveTool : public UAddPrimitiveTool
 {
 	GENERATED_BODY()
 public:
-	UAddRectanglePrimitiveTool(const FObjectInitializer& ObjectInitializer);
+	explicit UAddRectanglePrimitiveTool(const FObjectInitializer& ObjectInitializer);
 protected:
-	void GenerateMesh(FDynamicMesh3* OutMesh) const override;
+	virtual void GenerateMesh(FDynamicMesh3* OutMesh) const override;
 };
 
 UCLASS()
@@ -570,9 +581,9 @@ class UAddDiscPrimitiveTool : public UAddPrimitiveTool
 {
 	GENERATED_BODY()
 public:
-	UAddDiscPrimitiveTool(const FObjectInitializer& ObjectInitializer);
+	explicit UAddDiscPrimitiveTool(const FObjectInitializer& ObjectInitializer);
 protected:
-	void GenerateMesh(FDynamicMesh3* OutMesh) const override;
+	virtual void GenerateMesh(FDynamicMesh3* OutMesh) const override;
 };
 
 UCLASS()
@@ -580,9 +591,9 @@ class UAddTorusPrimitiveTool : public UAddPrimitiveTool
 {
 	GENERATED_BODY()
 public:
-	UAddTorusPrimitiveTool(const FObjectInitializer& ObjectInitializer);
+	explicit UAddTorusPrimitiveTool(const FObjectInitializer& ObjectInitializer);
 protected:
-	void GenerateMesh(FDynamicMesh3* OutMesh) const override;
+	virtual void GenerateMesh(FDynamicMesh3* OutMesh) const override;
 };
 
 UCLASS()
@@ -590,9 +601,9 @@ class UAddArrowPrimitiveTool : public UAddPrimitiveTool
 {
 	GENERATED_BODY()
 public:
-	UAddArrowPrimitiveTool(const FObjectInitializer& ObjectInitializer);
+	explicit UAddArrowPrimitiveTool(const FObjectInitializer& ObjectInitializer);
 protected:
-	void GenerateMesh(FDynamicMesh3* OutMesh) const override;
+	virtual void GenerateMesh(FDynamicMesh3* OutMesh) const override;
 };
 
 UCLASS()
@@ -600,9 +611,9 @@ class UAddSpherePrimitiveTool : public UAddPrimitiveTool
 {
 	GENERATED_BODY()
 public:
-	UAddSpherePrimitiveTool(const FObjectInitializer& ObjectInitializer);
+	explicit UAddSpherePrimitiveTool(const FObjectInitializer& ObjectInitializer);
 protected:
-	void GenerateMesh(FDynamicMesh3* OutMesh) const override;
+	virtual void GenerateMesh(FDynamicMesh3* OutMesh) const override;
 };
 
 UCLASS()
@@ -610,9 +621,9 @@ class UAddStairsPrimitiveTool : public UAddPrimitiveTool
 {
 	GENERATED_BODY()
 public:
-	UAddStairsPrimitiveTool(const FObjectInitializer& ObjectInitializer);
+	explicit UAddStairsPrimitiveTool(const FObjectInitializer& ObjectInitializer);
 protected:
-	void GenerateMesh(FDynamicMesh3* OutMesh) const override;
+	virtual void GenerateMesh(FDynamicMesh3* OutMesh) const override;
 };
 
 
diff --git a/Engine/Source/Runtime/GeometryCore/Public/Generators/SweepGenerator.h b/Engine/Source/Runtime/GeometryCore/Public/Generators/SweepGenerator.h
--- a/Engine/Source/Runtime/GeometryCore/Public/Generators/SweepGenerator.h
+++ b/Engine/Source/Runtime/GeometryCore/Public/Generators/SweepGenerator.h
@@ -441,7 +441,6 @@ public:
 			int SharpNormalIdx = 0;
 			for (int XIdx = 0, NormalXIdx = 0; XIdx < NumX; ++XIdx, ++NormalXIdx)
 			{
-				double Along = AlongPercents[XIdx];
 				double AlongRadius = Radii[XIdx];
 				Vertices[SubIdx + XIdx * AngleSamples] =
 					FVector3d(XVerts[SubIdx].X * AlongRadius, XVerts[SubIdx].Y * AlongRadius, Heights[XIdx]);
@@ -597,28 +596,41 @@ public:
 	/** Generate the mesh */
 	virtual FMeshShapeGenerator& Generate() override
 	{
+		const float SrcRadii[]{StickRadius, StickRadius, HeadBaseRadius, HeadTipRadius};
+		const float SrcHeights[]{0, StickLength, StickLength, StickLength + HeadLength};
+
 		TArray<float> Radii, Heights;
-		TArray<int> SharpNormalsAlongLength;
-		float SrcRadii[] { StickRadius, StickRadius, HeadBaseRadius, HeadTipRadius };
-		float SrcHeights[] { 0, StickLength, StickLength, StickLength + HeadLength };
+		const int NumVerts = 4 + AdditionalLengthSamples[0] + AdditionalLengthSamples[1] + AdditionalLengthSamples[2];
+		Radii.SetNumUninitialized(NumVerts);
+		Heights.SetNumUninitialized(NumVerts);
+
+		int VertIdx = 0;
+		auto SetVert = [&Radii, &Heights, &VertIdx](float Radius, float Height)
+		{
+			Radii[VertIdx] = Radius;
+			Heights[VertIdx] = Height;
+			++VertIdx;
+		};
 
-		int SegIdx = 0;
-		for (; SegIdx < 3; SegIdx++)
+		for (int SegIdx = 0; true; ++SegIdx)
 		{
-			Radii.Add(SrcRadii[SegIdx]);
-			Heights.Add(SrcHeights[SegIdx]);
-			for (int ExtraIdx = 0; ExtraIdx < AdditionalLengthSamples[SegIdx]; ExtraIdx++)
+			SetVert(SrcRadii[SegIdx], SrcHeights[SegIdx]);
+
+			if (SegIdx == 3)
+			{
+				break;
+			}
+
+			for (float ExtraSeg = 1, NumExtraSegs = AdditionalLengthSamples[SegIdx] + 1; ExtraSeg < NumExtraSegs; ++ExtraSeg)
 			{
-				float Along = float(ExtraIdx + 1) / float(AdditionalLengthSamples[SegIdx] + 1);
-				Radii.Add(FMath::Lerp(SrcRadii[SegIdx], SrcRadii[SegIdx + 1], Along));
-				Heights.Add(FMath::Lerp(SrcHeights[SegIdx], SrcHeights[SegIdx + 1], Along));
+				const float Along = ExtraSeg / NumExtraSegs;
+				SetVert(FMath::Lerp(SrcRadii[SegIdx], SrcRadii[SegIdx + 1], Along),
+						FMath::Lerp(SrcHeights[SegIdx], SrcHeights[SegIdx + 1], Along));
 			}
 		}
-		SharpNormalsAlongLength.Add(1 + AdditionalLengthSamples[0]);
-		SharpNormalsAlongLength.Add(SharpNormalsAlongLength.Last() + 1 + AdditionalLengthSamples[1]);
-		Radii.Add(SrcRadii[SegIdx]);
-		Heights.Add(SrcHeights[SegIdx]);
-		
+
+		TArray<int> SharpNormalsAlongLength{1 + AdditionalLengthSamples[0], 2 + AdditionalLengthSamples[0] + AdditionalLengthSamples[1]};
+
 		GenerateVerticalCircleSweep(Radii, Heights, SharpNormalsAlongLength);
 
 		return *this;
