MeshModelingTools: split extrude and offset out in PolyEdit.

#rnx

diff --git a/Engine/Plugins/Experimental/MeshModelingToolsetExp/Source/ModelingOperators/Private/DeformationOps/ExtrudeOp.cpp b/Engine/Plugins/Experimental/MeshModelingToolsetExp/Source/ModelingOperators/Private/DeformationOps/ExtrudeOp.cpp
--- a/Engine/Plugins/Experimental/MeshModelingToolsetExp/Source/ModelingOperators/Private/DeformationOps/ExtrudeOp.cpp
+++ b/Engine/Plugins/Experimental/MeshModelingToolsetExp/Source/ModelingOperators/Private/DeformationOps/ExtrudeOp.cpp
@@ -57,7 +57,7 @@ void FExtrudeOp::CalculateResult(FProgressCancel* Progress)
 	// Reset result
 	ResultMesh->Copy(*OriginalMesh, true, true, true, true);
 
-	if (ExtrudeMode == EPolyEditExtrudeMode::Boolean)
+	if (ExtrudeMode == EExtrudeMode::Boolean)
 	{
 		BooleanExtrude(Progress);
 	}
@@ -107,7 +107,7 @@ bool FExtrudeOp::BooleanExtrude(FProgressCancel* Progress)
 	};
 
 	Extruder.UVScaleFactor = UVScaleFactor;
-	if (DirectionMode == EPolyEditExtrudeDirectionMode::SingleDirection)
+	if (DirectionMode == EDirectionMode::SingleDirection)
 	{
 		Extruder.OffsetPositionFunc = [this](const FVector3d& Pos, const FVector3d& VertexVector, int32 VertexID)
 		{
@@ -122,9 +122,9 @@ bool FExtrudeOp::BooleanExtrude(FProgressCancel* Progress)
 	}
 
 	Extruder.ExtrusionVectorType =
-		DirectionMode == EPolyEditExtrudeDirectionMode::VertexNormals ? FOffsetMeshRegion::EVertexExtrusionVectorType::VertexNormal
-		: DirectionMode == EPolyEditExtrudeDirectionMode::SelectedTriangleNormals ? FOffsetMeshRegion::EVertexExtrusionVectorType::SelectionTriNormalsAngleWeightedAverage
-		: DirectionMode == EPolyEditExtrudeDirectionMode::SelectedTriangleNormalsEven ? FOffsetMeshRegion::EVertexExtrusionVectorType::SelectionTriNormalsAngleWeightedAdjusted
+		DirectionMode == EDirectionMode::VertexNormals ? FOffsetMeshRegion::EVertexExtrusionVectorType::VertexNormal
+		: DirectionMode == EDirectionMode::SelectedTriangleNormals ? FOffsetMeshRegion::EVertexExtrusionVectorType::SelectionTriNormalsAngleWeightedAverage
+		: DirectionMode == EDirectionMode::SelectedTriangleNormalsEven ? FOffsetMeshRegion::EVertexExtrusionVectorType::SelectionTriNormalsAngleWeightedAdjusted
 		: FOffsetMeshRegion::EVertexExtrusionVectorType::Zero;
 
 	Extruder.bIsPositiveOffset = (ExtrudeDistance > 0);
@@ -216,7 +216,7 @@ bool FExtrudeOp::MoveAndStitchExtrude(FProgressCancel* Progress)
 	};
 
 	Extruder.UVScaleFactor = UVScaleFactor;
-	if (DirectionMode == EPolyEditExtrudeDirectionMode::SingleDirection)
+	if (DirectionMode == EDirectionMode::SingleDirection)
 	{
 		Extruder.OffsetPositionFunc = [this](const FVector3d& Pos, const FVector3d& VertexVector, int32 VertexID)
 		{
@@ -231,9 +231,9 @@ bool FExtrudeOp::MoveAndStitchExtrude(FProgressCancel* Progress)
 	}
 
 	Extruder.ExtrusionVectorType =
-		DirectionMode == EPolyEditExtrudeDirectionMode::VertexNormals ? FOffsetMeshRegion::EVertexExtrusionVectorType::VertexNormal
-		: DirectionMode == EPolyEditExtrudeDirectionMode::SelectedTriangleNormals ? FOffsetMeshRegion::EVertexExtrusionVectorType::SelectionTriNormalsAngleWeightedAverage
-		: DirectionMode == EPolyEditExtrudeDirectionMode::SelectedTriangleNormalsEven ? FOffsetMeshRegion::EVertexExtrusionVectorType::SelectionTriNormalsAngleWeightedAdjusted
+		DirectionMode == EDirectionMode::VertexNormals ? FOffsetMeshRegion::EVertexExtrusionVectorType::VertexNormal
+		: DirectionMode == EDirectionMode::SelectedTriangleNormals ? FOffsetMeshRegion::EVertexExtrusionVectorType::SelectionTriNormalsAngleWeightedAverage
+		: DirectionMode == EDirectionMode::SelectedTriangleNormalsEven ? FOffsetMeshRegion::EVertexExtrusionVectorType::SelectionTriNormalsAngleWeightedAdjusted
 		: FOffsetMeshRegion::EVertexExtrusionVectorType::Zero;
 
 	Extruder.bIsPositiveOffset = (ExtrudeDistance > 0);
diff --git a/Engine/Plugins/Experimental/MeshModelingToolsetExp/Source/ModelingOperators/Public/DeformationOps/ExtrudeOp.h b/Engine/Plugins/Experimental/MeshModelingToolsetExp/Source/ModelingOperators/Public/DeformationOps/ExtrudeOp.h
--- a/Engine/Plugins/Experimental/MeshModelingToolsetExp/Source/ModelingOperators/Public/DeformationOps/ExtrudeOp.h
+++ b/Engine/Plugins/Experimental/MeshModelingToolsetExp/Source/ModelingOperators/Public/DeformationOps/ExtrudeOp.h
@@ -6,37 +6,6 @@
 #include "UObject/ObjectMacros.h"
 #include "ModelingOperators.h"
 
-UENUM()
-enum class EPolyEditExtrudeDirectionMode
-{
-	// Take the angle-weighed average of the selected triangles around each
-	// extruded vertex to determine vertex movement direction.
-	SelectedTriangleNormals,
-	
-	// Like Selected Triangle Normals, but also adjusts the distances moved in
-	// an attempt to keep triangles parallel to their original facing.
-	SelectedTriangleNormalsEven,
-
-	// Vertex normals, regardless of selection.
-	VertexNormals,
-
-	// Extrude all triangles in the same direction regardless of their facing.
-	SingleDirection,
-};
-
-UENUM()
-enum class EPolyEditExtrudeMode
-{
-	// Performs extrusion by shifting the selected faces as stitching them
-	// to the border left behind.
-	MoveAndStitch,
-	
-	// Performs the extrusion by extruding selected faces into a closed mesh
-	// first, and then performing a boolean with the original mesh. This allows
-	// the extrusion to cut holes through the mesh or to bridge sections.
-	Boolean
-};
-
 namespace UE {
 namespace Geometry {
 
@@ -46,14 +15,44 @@ class FDynamicMeshChangeTracker;
 class MODELINGOPERATORS_API FExtrudeOp : public FDynamicMeshOperator
 {
 public:
+	enum class EExtrudeMode
+	{
+		// Performs extrusion by shifting the selected faces as stitching them
+		// to the border left behind.
+		MoveAndStitch,
+
+		// Performs the extrusion by extruding selected faces into a closed mesh
+		// first, and then performing a boolean with the original mesh. This allows
+		// the extrusion to cut holes through the mesh or to bridge sections.
+		Boolean
+	};
+
+	enum class EDirectionMode
+	{
+		// Take the angle-weighed average of the selected triangles around each
+		// extruded vertex to determine vertex movement direction.
+		SelectedTriangleNormals = 0,
+
+		// Like Selected Triangle Normals, but also adjusts the distances moved in
+		// an attempt to keep triangles parallel to their original facing.
+		SelectedTriangleNormalsEven = 1,
+
+		// Vertex normals, regardless of selection.
+		VertexNormals = 2,
+
+		// Extrude all triangles in the same direction regardless of their facing.
+		SingleDirection = 3,
+	};
+
+
 	virtual ~FExtrudeOp() {}
 
 	// Inputs:
 	TSharedPtr<const FDynamicMesh3, ESPMode::ThreadSafe> OriginalMesh;
 	TArray<int32> TriangleSelection;
 	double ExtrudeDistance = 0;
-	EPolyEditExtrudeMode ExtrudeMode = EPolyEditExtrudeMode::MoveAndStitch;
-	EPolyEditExtrudeDirectionMode DirectionMode = EPolyEditExtrudeDirectionMode::SelectedTriangleNormalsEven;
+	EExtrudeMode ExtrudeMode = EExtrudeMode::MoveAndStitch;
+	EDirectionMode DirectionMode = EDirectionMode::SelectedTriangleNormalsEven;
 	// Only used if DirectionMode is SingleDirection
 	FVector3d MeshSpaceExtrudeDirection;
 	// Used when setting groups for the sides when the extrusion includes a mesh border. When true,
diff --git a/Engine/Plugins/Runtime/MeshModelingToolset/Source/MeshModelingTools/Private/EditMeshPolygonsTool.cpp b/Engine/Plugins/Runtime/MeshModelingToolset/Source/MeshModelingTools/Private/EditMeshPolygonsTool.cpp
--- a/Engine/Plugins/Runtime/MeshModelingToolset/Source/MeshModelingTools/Private/EditMeshPolygonsTool.cpp
+++ b/Engine/Plugins/Runtime/MeshModelingToolset/Source/MeshModelingTools/Private/EditMeshPolygonsTool.cpp
@@ -838,13 +838,22 @@ void UEditMeshPolygonsTool::OnTick(float DeltaTime)
 		//Interactive operations:
 		case EEditMeshPolygonsToolActions::Extrude:
 		{
-			ExtrudeActivity->ExtrudeProperties->ExtrudeMode = EPolyEditExtrudeMode::MoveAndStitch;
+			ExtrudeActivity->ExtrudeMode = FExtrudeOp::EExtrudeMode::MoveAndStitch;
+			ExtrudeActivity->PropertySetToUse = UPolyEditExtrudeActivity::EPropertySetToUse::Extrude;
 			StartActivity(ExtrudeActivity);
 			break;
 		}
 		case EEditMeshPolygonsToolActions::PushPull:
 		{
-			ExtrudeActivity->ExtrudeProperties->ExtrudeMode = EPolyEditExtrudeMode::Boolean;
+			ExtrudeActivity->ExtrudeMode = FExtrudeOp::EExtrudeMode::Boolean;
+			ExtrudeActivity->PropertySetToUse = UPolyEditExtrudeActivity::EPropertySetToUse::PushPull;
+			StartActivity(ExtrudeActivity);
+			break;
+		}
+		case EEditMeshPolygonsToolActions::Offset:
+		{
+			ExtrudeActivity->ExtrudeMode = FExtrudeOp::EExtrudeMode::MoveAndStitch;
+			ExtrudeActivity->PropertySetToUse = UPolyEditExtrudeActivity::EPropertySetToUse::Offset;
 			StartActivity(ExtrudeActivity);
 			break;
 		}
diff --git a/Engine/Plugins/Runtime/MeshModelingToolset/Source/MeshModelingTools/Private/ToolActivities/PolyEditExtrudeActivity.cpp b/Engine/Plugins/Runtime/MeshModelingToolset/Source/MeshModelingTools/Private/ToolActivities/PolyEditExtrudeActivity.cpp
--- a/Engine/Plugins/Runtime/MeshModelingToolset/Source/MeshModelingTools/Private/ToolActivities/PolyEditExtrudeActivity.cpp
+++ b/Engine/Plugins/Runtime/MeshModelingToolset/Source/MeshModelingTools/Private/ToolActivities/PolyEditExtrudeActivity.cpp
@@ -5,7 +5,6 @@
 #include "BaseBehaviors/SingleClickBehavior.h"
 #include "BaseBehaviors/MouseHoverBehavior.h"
 #include "ContextObjectStore.h"
-#include "DeformationOps/ExtrudeOp.h"
 #include "Drawing/PolyEditPreviewMesh.h"
 #include "DynamicMesh/MeshNormals.h"
 #include "DynamicMesh/MeshTransforms.h"
@@ -33,16 +32,41 @@ namespace PolyEditExtrudeActivityLocals
 			(int32)(EMeshComponents::FaceGroups) | (int32)(EMeshComponents::VertexNormals), true);
 		return MakeShared<FDynamicMesh3>(MoveTemp(Submesh.GetSubmesh()));
 	}
+
+	template <typename EnumType>
+	FExtrudeOp::EDirectionMode ToOpDirectionMode(EnumType Value)
+	{
+		return static_cast<FExtrudeOp::EDirectionMode>(static_cast<int>(Value));
+	}
 }
 
 TUniquePtr<FDynamicMeshOperator> UPolyEditExtrudeActivity::MakeNewOperator()
 {
+	using namespace PolyEditExtrudeActivityLocals;
 	using UE::Geometry::FTransform3d;
 
 	TUniquePtr<FExtrudeOp> Op = MakeUnique<FExtrudeOp>();
 	Op->OriginalMesh = ActivityContext->CurrentMesh;
-	Op->ExtrudeMode = ExtrudeProperties->ExtrudeMode;
-	Op->DirectionMode = ExtrudeProperties->DirectionMode;
+	Op->ExtrudeMode = ExtrudeMode;
+
+	switch (PropertySetToUse)
+	{
+	case EPropertySetToUse::Extrude:
+		Op->DirectionMode = ToOpDirectionMode(ExtrudeProperties->DirectionMode);
+		Op->bShellsToSolids = ExtrudeProperties->bShellsToSolids;
+		Op->bUseColinearityForSettingBorderGroups = ExtrudeProperties->bUseColinearityForSettingBorderGroups;
+		break;
+	case EPropertySetToUse::Offset:
+		Op->DirectionMode = ToOpDirectionMode(OffsetProperties->DirectionMode);
+		Op->bShellsToSolids = OffsetProperties->bShellsToSolids;
+		Op->bUseColinearityForSettingBorderGroups = OffsetProperties->bUseColinearityForSettingBorderGroups;
+		break;
+	case EPropertySetToUse::PushPull:
+		Op->DirectionMode = ToOpDirectionMode(PushPullProperties->DirectionMode);
+		Op->bShellsToSolids = PushPullProperties->bShellsToSolids;
+		Op->bUseColinearityForSettingBorderGroups = PushPullProperties->bUseColinearityForSettingBorderGroups;
+		break;
+	}
 
 	ActivityContext->CurrentTopology->GetSelectedTriangles(ActiveSelection, Op->TriangleSelection);
 
@@ -53,8 +77,6 @@ TUniquePtr<FDynamicMeshOperator> UPolyEditExtrudeActivity::MakeNewOperator()
 	Op->SetResultTransform(WorldTransform);
 	
 	Op->MeshSpaceExtrudeDirection = WorldTransform.InverseTransformVector(GetExtrudeDirection());
-	Op->bShellsToSolids = ExtrudeProperties->bShellsToSolids;
-	Op->bUseColinearityForSettingBorderGroups = ExtrudeProperties->bUseColinearityForSettingBorderGroups;
 
 	return Op;
 }
@@ -67,29 +89,72 @@ void UPolyEditExtrudeActivity::Setup(UInteractiveTool* ParentToolIn)
 	ExtrudeProperties->RestoreProperties(ParentTool.Get());
 	AddToolPropertySource(ExtrudeProperties);
 	SetToolPropertySourceEnabled(ExtrudeProperties, false);
+	ExtrudeProperties->WatchProperty(ExtrudeProperties->Direction,
+	[this](EPolyEditExtrudeDirection) { 
+		if (bIsRunning && ExtrudeProperties->DirectionMode == EPolyEditExtrudeModeOptions::SingleDirection)
+		{
+			ReinitializeExtrudeHeightMechanic();
+			ActivityContext->Preview->InvalidateResult();
+		}
+	});
 	ExtrudeProperties->WatchProperty(ExtrudeProperties->MeasureDirection,
-		[this](EPolyEditExtrudeDirection) { 
-			if (bIsRunning)
-			{
-				ReinitializeExtrudeHeightMechanic();
-				ActivityContext->Preview->InvalidateResult();
-			}
-		});
+	[this](EPolyEditExtrudeDirection) {
+		if (bIsRunning && ExtrudeProperties->DirectionMode != EPolyEditExtrudeModeOptions::SingleDirection)
+		{
+			ReinitializeExtrudeHeightMechanic();
+			ActivityContext->Preview->InvalidateResult();
+		}
+	});
 	ExtrudeProperties->WatchProperty(ExtrudeProperties->DirectionMode,
-		[this](EPolyEditExtrudeDirectionMode) {
-			if (bIsRunning)
-			{
-				ReinitializeExtrudeHeightMechanic();
-				ActivityContext->Preview->InvalidateResult();
-			}
-		});
-	ExtrudeProperties->WatchProperty(ExtrudeProperties->ExtrudeMode,
-		[this](EPolyEditExtrudeMode) {
-			if (bIsRunning)
-			{
-				ActivityContext->Preview->InvalidateResult();
-			}
-		});
+	[this](EPolyEditExtrudeModeOptions) {
+		if (bIsRunning)
+		{
+			ReinitializeExtrudeHeightMechanic();
+			ActivityContext->Preview->InvalidateResult();
+		}
+	});
+
+	OffsetProperties = NewObject<UPolyEditOffsetProperties>();
+	OffsetProperties->RestoreProperties(ParentTool.Get());
+	AddToolPropertySource(OffsetProperties);
+	SetToolPropertySourceEnabled(OffsetProperties, false);
+	OffsetProperties->WatchProperty(OffsetProperties->MeasureDirection,
+	[this](EPolyEditExtrudeDirection) {
+		if (bIsRunning)
+		{
+			ReinitializeExtrudeHeightMechanic();
+			ActivityContext->Preview->InvalidateResult();
+		}
+	});
+	OffsetProperties->WatchProperty(OffsetProperties->DirectionMode,
+	[this](EPolyEditOffsetModeOptions) {
+		if (bIsRunning)
+		{
+			ReinitializeExtrudeHeightMechanic();
+			ActivityContext->Preview->InvalidateResult();
+		}
+	});
+
+	PushPullProperties = NewObject<UPolyEditPushPullProperties>();
+	PushPullProperties->RestoreProperties(ParentTool.Get());
+	AddToolPropertySource(PushPullProperties);
+	SetToolPropertySourceEnabled(PushPullProperties, false);
+	PushPullProperties->WatchProperty(PushPullProperties->MeasureDirection,
+	[this](EPolyEditExtrudeDirection) {
+		if (bIsRunning)
+		{
+			ReinitializeExtrudeHeightMechanic();
+			ActivityContext->Preview->InvalidateResult();
+		}
+	});
+	PushPullProperties->WatchProperty(PushPullProperties->DirectionMode,
+	[this](EPolyEditPushPullModeOptions) {
+		if (bIsRunning)
+		{
+			ReinitializeExtrudeHeightMechanic();
+			ActivityContext->Preview->InvalidateResult();
+		}
+	});
 
 	// Register ourselves to receive clicks and hover
 	USingleClickInputBehavior* ClickBehavior = NewObject<USingleClickInputBehavior>();
@@ -110,8 +175,12 @@ void UPolyEditExtrudeActivity::Shutdown(EToolShutdownType ShutdownType)
 	}
 
 	ExtrudeProperties->SaveProperties(ParentTool.Get());
+	OffsetProperties->SaveProperties(ParentTool.Get());
+	PushPullProperties->SaveProperties(ParentTool.Get());
 
 	ExtrudeProperties = nullptr;
+	OffsetProperties = nullptr;
+	PushPullProperties = nullptr;
 	ParentTool = nullptr;
 	ActivityContext = nullptr;
 }
@@ -131,7 +200,7 @@ EToolActivityStartResult UPolyEditExtrudeActivity::Start()
 	if (!CanStart())
 	{
 		ParentTool->GetToolManager()->DisplayMessage(
-			LOCTEXT("OnExtrudeFailedMesssage", "Cannot extrude without face selection."),
+			LOCTEXT("OnExtrudeFailedMesssage", "Action requires face selection."),
 			EToolMessageLevel::UserWarning);
 		return EToolActivityStartResult::FailedStart;
 	}
@@ -144,7 +213,18 @@ EToolActivityStartResult UPolyEditExtrudeActivity::Start()
 			NewSelectionGids = Op->ExtrudedFaceNewGids;
 		});
 
-	SetToolPropertySourceEnabled(ExtrudeProperties, true);
+	switch (PropertySetToUse)
+	{
+	case EPropertySetToUse::Extrude:
+		SetToolPropertySourceEnabled(ExtrudeProperties, true); 
+		break;
+	case EPropertySetToUse::Offset:
+		SetToolPropertySourceEnabled(OffsetProperties, true);
+		break;
+	case EPropertySetToUse::PushPull:
+		SetToolPropertySourceEnabled(PushPullProperties, true);
+		break;
+	}
 
 	// Set up the basics of the extrude height mechanic.
 	// TODO: Allow option to use a gizmo instead
@@ -295,6 +375,8 @@ void UPolyEditExtrudeActivity::EndInternal()
 	ExtrudeHeightMechanic->Shutdown();
 	ExtrudeHeightMechanic = nullptr;
 	SetToolPropertySourceEnabled(ExtrudeProperties, false);
+	SetToolPropertySourceEnabled(OffsetProperties, false);
+	SetToolPropertySourceEnabled(PushPullProperties, false);
 	bIsRunning = false;
 }
 
@@ -323,7 +405,7 @@ void UPolyEditExtrudeActivity::ApplyExtrude()
 	// doesn't currently track changed triangles.
 	FDynamicMeshChangeTracker ChangeTracker(ActivityContext->CurrentMesh.Get());
 	ChangeTracker.BeginChange();
-	if (ExtrudeProperties->ExtrudeMode == EPolyEditExtrudeMode::MoveAndStitch)
+	if (ExtrudeMode == FExtrudeOp::EExtrudeMode::MoveAndStitch)
 	{
 		TArray<int32> ActiveTriangleSelection;
 		ActivityContext->CurrentTopology->GetSelectedTriangles(
@@ -382,7 +464,22 @@ FVector3d UPolyEditExtrudeActivity::GetExtrudeDirection() const
 {
 	using UE::Geometry::FTransform3d;
 
-	switch (ExtrudeProperties->MeasureDirection)
+	EPolyEditExtrudeDirection DirectionToUse = EPolyEditExtrudeDirection::SelectionNormal;
+	switch (PropertySetToUse)
+	{
+	case EPropertySetToUse::Extrude:
+		DirectionToUse = ExtrudeProperties->DirectionMode == EPolyEditExtrudeModeOptions::SingleDirection ?
+			ExtrudeProperties->Direction : ExtrudeProperties->MeasureDirection;
+		break;
+	case EPropertySetToUse::Offset:
+		DirectionToUse = OffsetProperties->MeasureDirection;
+		break;
+	case EPropertySetToUse::PushPull:
+		DirectionToUse = PushPullProperties->MeasureDirection;
+		break;
+	}
+
+	switch (DirectionToUse)
 	{
 	default:
 	case EPolyEditExtrudeDirection::SelectionNormal:
diff --git a/Engine/Plugins/Runtime/MeshModelingToolset/Source/MeshModelingTools/Public/EditMeshPolygonsTool.h b/Engine/Plugins/Runtime/MeshModelingToolset/Source/MeshModelingTools/Public/EditMeshPolygonsTool.h
--- a/Engine/Plugins/Runtime/MeshModelingToolset/Source/MeshModelingTools/Public/EditMeshPolygonsTool.h
+++ b/Engine/Plugins/Runtime/MeshModelingToolset/Source/MeshModelingTools/Public/EditMeshPolygonsTool.h
@@ -95,6 +95,7 @@ enum class EEditMeshPolygonsToolActions
 	CancelCurrent,
 	Extrude,
 	PushPull,
+	Offset,
 	Inset,
 	Outset,
 	InsertEdge,
@@ -185,10 +186,14 @@ public:
 	UFUNCTION(CallInEditor, Category = FaceEdits, meta = (DisplayName = "Extrude", DisplayPriority = 1))
 	void Extrude() { PostAction(EEditMeshPolygonsToolActions::Extrude); }
 
-	/** Like Extrude, but performed in a boolean way, meaning that the faces can cut away the mesh or bridge mesh parts. */
+	/** Like Extrude/Offset, but performed in a boolean way, meaning that the faces can cut away the mesh or bridge mesh parts. */
 	UFUNCTION(CallInEditor, Category = FaceEdits, meta = (DisplayName = "Push/Pull", DisplayPriority = 1))
 	void PushPull() { PostAction(EEditMeshPolygonsToolActions::PushPull); }
 
+	/** Like Extrude, but defaults to moving verts along vertex normals instead of a single direction. */
+	UFUNCTION(CallInEditor, Category = FaceEdits, meta = (DisplayName = "Offset", DisplayPriority = 1))
+	void Offset() { PostAction(EEditMeshPolygonsToolActions::Offset); }
+
 	/** Inset the current set of selected faces. Click in viewport to confirm inset distance. */
 	UFUNCTION(CallInEditor, Category = FaceEdits, meta = (DisplayName = "Inset", DisplayPriority = 2))
 	void Inset() { PostAction(EEditMeshPolygonsToolActions::Inset);	}
@@ -197,7 +202,7 @@ public:
 	UFUNCTION(CallInEditor, Category = FaceEdits, meta = (DisplayName = "Outset", DisplayPriority = 3))
 	void Outset() { PostAction(EEditMeshPolygonsToolActions::Outset);	}
 
-	//~ TODO: Make the Merge and Delete comments visible as tooltips. Currently we can't due to a bug that
+	//~ TODO: Make the Merge, Delete, and Flip comments visible as tooltips. Currently we can't due to a bug that
 	//~ limits our total tooltip text allotment: UE-124608
 
 	//~ Merge the current set of selected faces into a single face.
@@ -216,7 +221,7 @@ public:
 	UFUNCTION(CallInEditor, Category = FaceEdits, meta = (DisplayName = "RecalcNormals", DisplayPriority = 6))
 	void RecalcNormals() { PostAction(EEditMeshPolygonsToolActions::RecalculateNormals); }
 
-	/** Flip normals and face orientation for the current set of selected faces */
+	//~ Flip normals and face orientation for the current set of selected faces
 	UFUNCTION(CallInEditor, Category = FaceEdits, meta = (DisplayName = "Flip", DisplayPriority = 7))
 	void Flip() { PostAction(EEditMeshPolygonsToolActions::FlipNormals); }
 
diff --git a/Engine/Plugins/Runtime/MeshModelingToolset/Source/MeshModelingTools/Public/ToolActivities/PolyEditExtrudeActivity.h b/Engine/Plugins/Runtime/MeshModelingToolset/Source/MeshModelingTools/Public/ToolActivities/PolyEditExtrudeActivity.h
--- a/Engine/Plugins/Runtime/MeshModelingToolset/Source/MeshModelingTools/Public/ToolActivities/PolyEditExtrudeActivity.h
+++ b/Engine/Plugins/Runtime/MeshModelingToolset/Source/MeshModelingTools/Public/ToolActivities/PolyEditExtrudeActivity.h
@@ -30,27 +30,92 @@ enum class EPolyEditExtrudeDirection
 	LocalZ
 };
 
+// There is a lot of overlap in the options for Extrude, Offset, and Push/Pull, and they map to
+// the same op behind the scenes. However, we want to keep them as separate buttons to keep some
+// amount of shallowness in the UI, to make it more likely that new users will find the setting
+// they are looking for.
+// A couple of settings are entirely replicated: namely, doing an offset or "extrude" with SelectedTriangleNormals
+// or SelectedTriangleNormalsEven as the movement direction is actually equivalent. Properly speaking, these
+// two should only be options under Offset, not Extrude, but we keep them as (non-default) options 
+// in both because an "extrude along local normals" is a common operation that some users are likely
+// to look for under extrude, regardless of it not lining up with the physical meaning of extrusion.
+
+// We use different property set objects so that we can customize category names, etc, as well as
+// have different defaults and saved settings.
+
+UENUM()
+enum class EPolyEditExtrudeModeOptions
+{
+	// Extrude all triangles in the same direction regardless of their facing.
+	SingleDirection = static_cast<int>(UE::Geometry::FExtrudeOp::EDirectionMode::SingleDirection),
+
+	// Take the angle-weighed average of the selected triangles around each
+	// extruded vertex to determine vertex movement direction.
+	SelectedTriangleNormals = static_cast<int>(UE::Geometry::FExtrudeOp::EDirectionMode::SelectedTriangleNormals),
+
+	// Like Selected Triangle Normals, but also adjusts the distances moved in
+	// an attempt to keep triangles parallel to their original facing.
+	SelectedTriangleNormalsEven = static_cast<int>(UE::Geometry::FExtrudeOp::EDirectionMode::SelectedTriangleNormalsEven),
+};
+
+UENUM()
+enum class EPolyEditOffsetModeOptions
+{
+	// Vertex normals, regardless of selection.
+	VertexNormals = static_cast<int>(UE::Geometry::FExtrudeOp::EDirectionMode::VertexNormals),
+
+	// Take the angle-weighed average of the selected triangles around
+	// offset vertex to determine vertex movement direction.
+	SelectedTriangleNormals = static_cast<int>(UE::Geometry::FExtrudeOp::EDirectionMode::SelectedTriangleNormals),
+
+	// Like Selected Triangle Normals, but also adjusts the distances moved in
+	// an attempt to keep triangles parallel to their original facing.
+	SelectedTriangleNormalsEven = static_cast<int>(UE::Geometry::FExtrudeOp::EDirectionMode::SelectedTriangleNormalsEven),
+};
+
+UENUM()
+enum class EPolyEditPushPullModeOptions
+{
+	// Take the angle-weighed average of the selected triangles around
+	// offset vertex to determine vertex movement direction.
+	SelectedTriangleNormals = static_cast<int>(EPolyEditOffsetModeOptions::SelectedTriangleNormals),
+
+	// Like Selected Triangle Normals, but also adjusts the distances moved in
+	// an attempt to keep triangles parallel to their original facing.
+	SelectedTriangleNormalsEven = static_cast<int>(UE::Geometry::FExtrudeOp::EDirectionMode::SelectedTriangleNormalsEven),
+
+	// Move all triangles in the same direction regardless of their facing.
+	SingleDirection = static_cast<int>(UE::Geometry::FExtrudeOp::EDirectionMode::SingleDirection),
+
+	// Vertex normals, regardless of selection.
+	VertexNormals = static_cast<int>(UE::Geometry::FExtrudeOp::EDirectionMode::VertexNormals),
+};
+
 UCLASS()
 class MESHMODELINGTOOLS_API UPolyEditExtrudeProperties : public UInteractiveToolPropertySet
 {
 	GENERATED_BODY()
 
 public:
-	// Not visible in details panel because we decided that we would use a separate "Push/Pull" button.
-	EPolyEditExtrudeMode ExtrudeMode = EPolyEditExtrudeMode::MoveAndStitch;
 
-	/** Which way to move vertices during the extrude */
-	UPROPERTY(EditAnywhere, Category = ExtrusionOptions)
-	EPolyEditExtrudeDirectionMode DirectionMode = EPolyEditExtrudeDirectionMode::SelectedTriangleNormalsEven;
+	/** Direction in which to extrude. */
+	UPROPERTY(EditAnywhere, Category = Extrude,
+		meta = (EditConditionHides, EditCondition = "DirectionMode == EPolyEditExtrudeModeOptions::SingleDirection"))
+	EPolyEditExtrudeDirection Direction = EPolyEditExtrudeDirection::SelectionNormal;
 
-	/** What axis to measure the extrusion distance along. When the direction mode is Single Direction, also controls the direction. */
-	UPROPERTY(EditAnywhere, Category = ExtrusionOptions, AdvancedDisplay)
+	/** What axis to measure the extrusion distance along. */
+	UPROPERTY(EditAnywhere, Category = Extrude, AdvancedDisplay,
+		meta = (EditConditionHides, EditCondition = "DirectionMode != EPolyEditExtrudeModeOptions::SingleDirection"))
 	EPolyEditExtrudeDirection MeasureDirection = EPolyEditExtrudeDirection::SelectionNormal;
 
 	/** Controls whether extruding an entire open-border patch should create a solid or an open shell */
-	UPROPERTY(EditAnywhere, Category = ExtrusionOptions)
+	UPROPERTY(EditAnywhere, Category = Extrude)
 	bool bShellsToSolids = true;
 
+	/** How to move the vertices during the extrude */
+	UPROPERTY(EditAnywhere, Category = Extrude)
+	EPolyEditExtrudeModeOptions DirectionMode = EPolyEditExtrudeModeOptions::SingleDirection;
+
 	/** 
 	 * When extruding regions that touch the mesh border, assign the side groups (groups on the 
 	 * stitched side of the extrude) in a way that considers edge colinearity. For instance, when
@@ -61,6 +126,62 @@ public:
 	bool bUseColinearityForSettingBorderGroups = true;
 };
 
+UCLASS()
+class MESHMODELINGTOOLS_API UPolyEditOffsetProperties : public UInteractiveToolPropertySet
+{
+	GENERATED_BODY()
+
+public:
+	/** Which way to move vertices during the offset */
+	UPROPERTY(EditAnywhere, Category = Offset)
+	EPolyEditOffsetModeOptions DirectionMode = EPolyEditOffsetModeOptions::VertexNormals;
+
+	/** Controls whether offseting an entire open-border patch should create a solid or an open shell */
+	UPROPERTY(EditAnywhere, Category = Offset)
+	bool bShellsToSolids = true;
+
+	/** What axis to measure the extrusion distance along. When the direction mode is Single Direction, also controls the direction. */
+	UPROPERTY(EditAnywhere, Category = Offset, AdvancedDisplay)
+	EPolyEditExtrudeDirection MeasureDirection = EPolyEditExtrudeDirection::SelectionNormal;
+
+	/**
+	 * When offseting regions that touch the mesh border, assign the side groups (groups on the
+	 * stitched side of the extrude) in a way that considers edge colinearity. For instance, when
+	 * true, extruding a flat rectangle will give four different groups on its sides rather than
+	 * one connected group.
+	 */
+	UPROPERTY(EditAnywhere, Category = Offset, AdvancedDisplay)
+	bool bUseColinearityForSettingBorderGroups = true;
+};
+
+UCLASS()
+class MESHMODELINGTOOLS_API UPolyEditPushPullProperties : public UInteractiveToolPropertySet
+{
+	GENERATED_BODY()
+
+public:
+	/** Which way to move vertices during the offset */
+	UPROPERTY(EditAnywhere, Category = ExtrusionOptions)
+	EPolyEditPushPullModeOptions DirectionMode = EPolyEditPushPullModeOptions::SelectedTriangleNormals;
+
+	/** Controls whether offseting an entire open-border patch should create a solid or an open shell */
+	UPROPERTY(EditAnywhere, Category = ExtrusionOptions)
+	bool bShellsToSolids = true;
+
+	/** What axis to measure the extrusion distance along. When the direction mode is Single Direction, also controls the direction. */
+	UPROPERTY(EditAnywhere, Category = ExtrusionOptions, AdvancedDisplay)
+	EPolyEditExtrudeDirection MeasureDirection = EPolyEditExtrudeDirection::SelectionNormal;
+
+	/**
+	 * When offseting regions that touch the mesh border, assign the side groups (groups on the
+	 * stitched side of the extrude) in a way that considers edge colinearity. For instance, when
+	 * true, extruding a flat rectangle will give four different groups on its sides rather than
+	 * one connected group.
+	 */
+	UPROPERTY(EditAnywhere, Category = ExtrusionOptions, AdvancedDisplay)
+	bool bUseColinearityForSettingBorderGroups = true;
+};
+
 /**
  * 
  */
@@ -73,6 +194,20 @@ class MESHMODELINGTOOLS_API UPolyEditExtrudeActivity : public UInteractiveToolAc
 	GENERATED_BODY()
 
 public:
+	using FExtrudeOp = UE::Geometry::FExtrudeOp;
+
+	enum class EPropertySetToUse
+	{
+		Extrude,
+		Offset,
+		PushPull
+	};
+
+	// Set to different values depending on whether we're using this activity on behalf of
+	// extrude, offset, or push/pull
+	FExtrudeOp::EExtrudeMode ExtrudeMode = FExtrudeOp::EExtrudeMode::MoveAndStitch;
+	EPropertySetToUse PropertySetToUse = EPropertySetToUse::Extrude;
+
 	// IInteractiveToolActivity
 	virtual void Setup(UInteractiveTool* ParentTool) override;
 	virtual void Shutdown(EToolShutdownType ShutdownType) override;
@@ -98,9 +233,16 @@ public:
 	virtual void OnEndHover() override {}
 
 	UPROPERTY()
-	TObjectPtr<UPolyEditExtrudeProperties> ExtrudeProperties;
+	TObjectPtr<UPolyEditExtrudeProperties> ExtrudeProperties = nullptr;
+
+	UPROPERTY()
+	TObjectPtr<UPolyEditOffsetProperties> OffsetProperties = nullptr;
+
+	UPROPERTY()
+	TObjectPtr<UPolyEditPushPullProperties> PushPullProperties = nullptr;
 
 protected:
+
 	FVector3d GetExtrudeDirection() const;
 	virtual void BeginExtrude();
 	virtual void ApplyExtrude();
