MeshModelingToolset: Cleanup Create tools UI

diff --git a/Engine/Plugins/Editor/ModelingToolsEditorMode/Source/ModelingToolsEditorMode/Private/ModelingToolsManagerActions.cpp b/Engine/Plugins/Editor/ModelingToolsEditorMode/Source/ModelingToolsEditorMode/Private/ModelingToolsManagerActions.cpp
--- a/Engine/Plugins/Editor/ModelingToolsEditorMode/Source/ModelingToolsEditorMode/Private/ModelingToolsManagerActions.cpp
+++ b/Engine/Plugins/Editor/ModelingToolsEditorMode/Source/ModelingToolsEditorMode/Private/ModelingToolsManagerActions.cpp
@@ -44,21 +44,26 @@ void FModelingToolsManagerCommands::RegisterCommands()
 		UI_COMMAND(ToolCommandInfo, ToolName, ToolTip, EUserInterfaceActionType::ToggleButton, FInputChord()); \
 		RegisteredTools.Add(FStartToolCommand{ ToolName, ToolCommandInfo });
 
-	REGISTER_MODELING_TOOL_COMMAND(BeginAddBoxPrimitiveTool, "Box", "Create new Box StaticMesh Assets");
-	REGISTER_MODELING_TOOL_COMMAND(BeginAddCylinderPrimitiveTool, "Cyl", "Create new Cylinder StaticMesh Assets");
-	REGISTER_MODELING_TOOL_COMMAND(BeginAddConePrimitiveTool, "Cone", "Create new Cone StaticMesh Assets");
-	REGISTER_MODELING_TOOL_COMMAND(BeginAddArrowPrimitiveTool, "Arrow", "Create new Arrow StaticMesh Assets");
-	REGISTER_MODELING_TOOL_COMMAND(BeginAddRectanglePrimitiveTool, "Rect", "Create new Rectangle StaticMesh Assets");
-	REGISTER_MODELING_TOOL_COMMAND(BeginAddDiscPrimitiveTool, "Disc", "Create new Disc StaticMesh Assets");
-	REGISTER_MODELING_TOOL_COMMAND(BeginAddTorusPrimitiveTool, "Torus", "Create new Torus StaticMesh Assets");
-	REGISTER_MODELING_TOOL_COMMAND(BeginAddSpherePrimitiveTool, "Sphere", "Create new Sphere StaticMesh Assets");
-	REGISTER_MODELING_TOOL_COMMAND(BeginAddStairsPrimitiveTool, "Stair", "Create new Stairs StaticMesh Assets");
+	// Shapes
+	REGISTER_MODELING_TOOL_COMMAND(BeginAddBoxPrimitiveTool, "Box", "Create new box objects");
+	REGISTER_MODELING_TOOL_COMMAND(BeginAddSpherePrimitiveTool, "Sphere", "Create new sphere objects");
+	REGISTER_MODELING_TOOL_COMMAND(BeginAddCylinderPrimitiveTool, "Cyl", "Create new cylinder objects");
+	REGISTER_MODELING_TOOL_COMMAND(BeginAddConePrimitiveTool, "Cone", "Create new cone objects");
+	REGISTER_MODELING_TOOL_COMMAND(BeginAddTorusPrimitiveTool, "Torus", "Create new torus objects");
+	REGISTER_MODELING_TOOL_COMMAND(BeginAddArrowPrimitiveTool, "Arrow", "Create new arrow objects");
+	REGISTER_MODELING_TOOL_COMMAND(BeginAddRectanglePrimitiveTool, "Rect", "Create new rectangle objects");
+	REGISTER_MODELING_TOOL_COMMAND(BeginAddDiscPrimitiveTool, "Disc", "Create new disc objects");
+	REGISTER_MODELING_TOOL_COMMAND(BeginAddStairsPrimitiveTool, "Stairs", "Create new stairs objects");
+
+	// Create
+	REGISTER_MODELING_TOOL_COMMAND(BeginDrawPolygonTool, "PolyExt", "Draw and extrude polygons to create new objects");
+	REGISTER_MODELING_TOOL_COMMAND(BeginDrawPolyPathTool, "PathExt", "Draw and extrude PolyPaths to create new objects");
+	REGISTER_MODELING_TOOL_COMMAND(BeginDrawAndRevolveTool, "PathRev", "Draw and revolve a PolyPaths to create new objects");
+	REGISTER_MODELING_TOOL_COMMAND(BeginRevolveBoundaryTool, "BdryRev", "Revolve mesh boundary loops into new objects");
+	REGISTER_MODELING_TOOL_COMMAND(BeginCombineMeshesTool, "Append", "Combine selections into new objects");
+	REGISTER_MODELING_TOOL_COMMAND(BeginDuplicateMeshesTool, "Dupe", "Duplicate selections into new objects");
 
 	REGISTER_MODELING_TOOL_COMMAND(BeginAddPatchTool, "Patch", "Add Patch");
-	REGISTER_MODELING_TOOL_COMMAND(BeginRevolveBoundaryTool, "BdryRev", "Revolve Mesh Boundary Loop into New Asset");
-	REGISTER_MODELING_TOOL_COMMAND(BeginDrawPolygonTool, "PolyExt", "Draw/Extrude Polygons to create new StaticMesh Assets");
-	REGISTER_MODELING_TOOL_COMMAND(BeginDrawPolyPathTool, "PathExt", "Draw/Extrude PolyPaths to create new StaticMesh Assets");
-	REGISTER_MODELING_TOOL_COMMAND(BeginDrawAndRevolveTool, "PathRev", "Draw and Revolve a PolyPath to create a new StaticMesh Assets");
 	REGISTER_MODELING_TOOL_COMMAND(BeginShapeSprayTool, "Spray", "Shape Spray");
 
 	REGISTER_MODELING_TOOL_COMMAND(BeginSculptMeshTool, "VSclpt", "Vertex Sculpting");
@@ -74,8 +79,6 @@ void FModelingToolsManagerCommands::RegisterCommands()
 	REGISTER_MODELING_TOOL_COMMAND(BeginEditPivotTool, "Pivot", "Edit Mesh Pivots");
 	REGISTER_MODELING_TOOL_COMMAND(BeginAddPivotActorTool, "PivotAct", "Add actor to act as a pivot for child component(s)");
 	REGISTER_MODELING_TOOL_COMMAND(BeginBakeTransformTool, "BakeRS", "Bake Scale/Rotation into Mesh Asset");
-	REGISTER_MODELING_TOOL_COMMAND(BeginCombineMeshesTool, "Append", "Combine Selection into new StaticMesh Asset");
-	REGISTER_MODELING_TOOL_COMMAND(BeginDuplicateMeshesTool, "Dupe", "Duplicate Selection into new StaticMesh Asset");
 	REGISTER_MODELING_TOOL_COMMAND(BeginAlignObjectsTool, "Align", "Align Objects");
 	REGISTER_MODELING_TOOL_COMMAND(BeginTransferMeshTool, "Transfer", "Transfer Mesh");
 	REGISTER_MODELING_TOOL_COMMAND(BeginConvertMeshesTool, "Convert", "Convert Meshes");
diff --git a/Engine/Plugins/Experimental/MeshModelingToolsetExp/Source/ModelingComponents/Private/Snapping/PointPlanarSnapSolver.cpp b/Engine/Plugins/Experimental/MeshModelingToolsetExp/Source/ModelingComponents/Private/Snapping/PointPlanarSnapSolver.cpp
--- a/Engine/Plugins/Experimental/MeshModelingToolsetExp/Source/ModelingComponents/Private/Snapping/PointPlanarSnapSolver.cpp
+++ b/Engine/Plugins/Experimental/MeshModelingToolsetExp/Source/ModelingComponents/Private/Snapping/PointPlanarSnapSolver.cpp
@@ -97,6 +97,7 @@ void FPointPlanarSnapSolver::RegenerateTargetLines(bool bCardinalAxes, bool bLas
 		SegLine.TargetID = LastSegmentTargetID;
 		SegLine.Priority = LastSegmentPriority;
 		SegLine.Line = FLine3d::FromPoints(LastPt, PointHistory[NumHistoryPts-2]);
+		GeneratedLines.Add(SegLine);
 		SegLine.Line.Direction = FQuaterniond(Plane.Z(), 90, true) * SegLine.Line.Direction;
 		GeneratedLines.Add(SegLine);
 	}
diff --git a/Engine/Plugins/Experimental/MeshModelingToolsetExp/Source/ModelingComponents/Public/PropertySets/CreateMeshObjectTypeProperties.h b/Engine/Plugins/Experimental/MeshModelingToolsetExp/Source/ModelingComponents/Public/PropertySets/CreateMeshObjectTypeProperties.h
--- a/Engine/Plugins/Experimental/MeshModelingToolsetExp/Source/ModelingComponents/Public/PropertySets/CreateMeshObjectTypeProperties.h
+++ b/Engine/Plugins/Experimental/MeshModelingToolsetExp/Source/ModelingComponents/Public/PropertySets/CreateMeshObjectTypeProperties.h
@@ -24,11 +24,11 @@ public:
 	void InitializeDefaultWithAuto();
 	void Initialize(bool bEnableStaticMeshes = true, bool bEnableVolumes = false, bool bEnableDynamicMeshActor = false);
 
-	/** Type of Object to create */
+	/** Type of object to create */
 	UPROPERTY(EditAnywhere, Category = OutputType, meta = (DisplayName = "Output Type", GetOptions = GetOutputTypeNamesFunc, NoResetToDefault))
 	FString OutputType;
 
-	/** Type of new Volume to create */
+	/** Type of volume to create */
 	UPROPERTY(EditAnywhere, Category = OutputType, meta = (EditCondition = "bShowVolumeList == true", EditConditionHides, HideEditConditionToggle) )
 	TSubclassOf<class AVolume> VolumeType = ABlockingVolume::StaticClass();
 
diff --git a/Engine/Plugins/Runtime/MeshModelingToolset/Source/MeshModelingTools/Private/DrawPolygonTool.cpp b/Engine/Plugins/Runtime/MeshModelingToolset/Source/MeshModelingTools/Private/DrawPolygonTool.cpp
--- a/Engine/Plugins/Runtime/MeshModelingToolset/Source/MeshModelingTools/Private/DrawPolygonTool.cpp
+++ b/Engine/Plugins/Runtime/MeshModelingToolset/Source/MeshModelingTools/Private/DrawPolygonTool.cpp
@@ -12,7 +12,6 @@
 #include "Curve/GeneralPolygon2.h"
 #include "FrameTypes.h"
 #include "MatrixTypes.h"
-#include "DynamicMesh/DynamicMeshAttributeSet.h"
 
 #include "Generators/FlatTriangulationMeshGenerator.h"
 #include "Generators/DiscMeshGenerator.h"
@@ -20,7 +19,6 @@
 #include "Operations/ExtrudeMesh.h"
 #include "Distance/DistLine3Ray3.h"
 #include "Intersection/IntrSegment2Segment2.h"
-#include "MeshQueries.h"
 #include "ToolSceneQueriesUtil.h"
 #include "ConstrainedDelaunay2.h"
 #include "Arrangement2d.h"
@@ -74,7 +72,7 @@ UDrawPolygonTool::UDrawPolygonTool()
 	DrawPlaneOrigin = FVector3d::Zero();
 	DrawPlaneOrientation = FQuaterniond::Identity();
 	bInInteractiveExtrude = false;
-	SetToolDisplayName(LOCTEXT("ToolName", "Extrude Polygon"));
+	UInteractiveTool::SetToolDisplayName(LOCTEXT("ToolName", "Extrude Polygon"));
 }
 
 void UDrawPolygonTool::SetWorld(UWorld* World)
@@ -118,8 +116,8 @@ void UDrawPolygonTool::Setup()
 
 	PolygonProperties = NewObject<UDrawPolygonToolStandardProperties>(this);
 	PolygonProperties->RestoreProperties(this);
-	PolygonProperties->WatchProperty(PolygonProperties->bShowGizmo,
-									 [this](bool bNewValue) { this->UpdateShowGizmoState(bNewValue); });
+	PolygonProperties->WatchProperty(PolygonProperties->bShowGridGizmo,
+	                                 [this](bool bNewValue) { this->UpdateShowGizmoState(bNewValue); });
 
 	// Create a new TransformGizmo and associated TransformProxy. The TransformProxy will not be the
 	// parent of any Components in this case, we just use it's transform and change delegate.
@@ -199,7 +197,7 @@ void UDrawPolygonTool::RegisterActions(FInteractiveToolActionSet& ActionSet)
 		LOCTEXT("ToggleGizmo", "Toggle Gizmo"),
 		LOCTEXT("ToggleGizmoTooltip", "Toggle visibility of the transformation Gizmo"),
 		EModifierKey::None, EKeys::A,
-		[this]() { PolygonProperties->bShowGizmo = !PolygonProperties->bShowGizmo; });
+		[this]() { PolygonProperties->bShowGridGizmo = !PolygonProperties->bShowGridGizmo; });
 }
 
 
@@ -282,10 +280,8 @@ void UDrawPolygonTool::Render(IToolsContextRenderAPI* RenderAPI)
 	FColor ErrorColor = FColor::Magenta;
 	float HiddenLineThickness = 1.0f*PDIScale;
 	float LineThickness = 4.0f*PDIScale;
-	float SelfIntersectThickness = 8.0f*PDIScale;
 	FColor GridColor(128, 128, 128, 32);
 	float GridThickness = 0.5f*PDIScale;
-	float GridLineSpacing = 25.0f*PDIScale;   // @todo should be relative to view
 	int NumGridLines = 21;
 	FColor SnapLineColor = FColor::Yellow;
 	FColor SnapHighlightColor = SnapLineColor;
@@ -379,7 +375,7 @@ void UDrawPolygonTool::Render(IToolsContextRenderAPI* RenderAPI)
 	if (bHaveSurfaceHit)
 	{
 		PDI->DrawPoint((FVector)SurfaceHitPoint, ClosedPolygonColor, 10*PDIScale, SDPG_Foreground);
-		if (SnapProperties->HitNormalOffset != 0)
+		if (SnapProperties->SnapToSurfacesOffset != 0)
 		{
 			PDI->DrawPoint((FVector)SurfaceOffsetPoint, OpenPolygonColor, 15*PDIScale, SDPG_Foreground);
 			PDI->DrawLine((FVector)SurfaceOffsetPoint, (FVector)SurfaceHitPoint,
@@ -400,7 +396,6 @@ void UDrawPolygonTool::Render(IToolsContextRenderAPI* RenderAPI)
 		FColor LastSegmentColor = bIsClosed ? ClosedPolygonColor 
 			: bHaveSelfIntersection ? ErrorColor : PreviewColor;
 		FVector3d UseLastVertex = bIsClosed ? PolygonVertices[0] : PreviewVertex;
-		float UseThickness = bHaveSelfIntersection ? SelfIntersectThickness : LineThickness;
 
 		auto DrawVertices = [&PDI, &UseColor](const TArray<FVector3d>& Vertices, ESceneDepthPriorityGroup Group, float Thickness)
 		{
@@ -478,8 +473,8 @@ void UDrawPolygonTool::UpdatePreviewVertex(const FVector3d& PreviewVertexIn)
 	// update length and angle
 	if (PolygonVertices.Num() > 0)
 	{
-		FVector3d LastVertex = PolygonVertices[PolygonVertices.Num() - 1];
-		SnapProperties->SegmentLength = Distance(LastVertex, PreviewVertex);
+		const FVector3d LastVertex = PolygonVertices[PolygonVertices.Num() - 1];
+		PolygonProperties->Distance = Distance(LastVertex, PreviewVertex);
 	}
 }
 
@@ -528,9 +523,9 @@ bool UDrawPolygonTool::FindDrawPlaneHitPoint(const FInputDeviceRay& ClickPos, FV
 			}
 		}
 
-		TArray<FVector3d>& HistoryPoints = (bInFixedPolygonMode) ? FixedPolygonClickPoints : PolygonVertices;
+		const TArray<FVector3d>& HistoryPoints = (bInFixedPolygonMode) ? FixedPolygonClickPoints : PolygonVertices;
 		SnapEngine.UpdatePointHistory(HistoryPoints);
-		if (SnapProperties->bSnapToAngles)
+		if (SnapProperties->bSnapToAxes)
 		{
 			SnapEngine.RegenerateTargetLines(true, true);
 		}
@@ -551,7 +546,7 @@ bool UDrawPolygonTool::FindDrawPlaneHitPoint(const FInputDeviceRay& ClickPos, FV
 
 
 	// if not snap and we want to hit objects, do that
-	if (SnapProperties->bHitSceneObjects)
+	if (SnapProperties->bSnapToSurfaces)
 	{
 		FHitResult Result;
 		bool bWorldHit = ToolSceneQueriesUtil::FindNearestVisibleObjectHit(TargetWorld, Result, ClickPos.WorldRay);
@@ -559,7 +554,7 @@ bool UDrawPolygonTool::FindDrawPlaneHitPoint(const FInputDeviceRay& ClickPos, FV
 		{
 			bHaveSurfaceHit = true;
 			SurfaceHitPoint = (FVector3d)Result.ImpactPoint;
-			FVector3d UseHitPos = (FVector3d)Result.ImpactPoint + (double)SnapProperties->HitNormalOffset*(FVector3d)Result.Normal;
+			const FVector3d UseHitPos = Result.ImpactPoint + static_cast<double>(SnapProperties->SnapToSurfacesOffset) * Result.Normal;
 			HitPos = Frame.ToPlane(UseHitPos, 2);
 			SurfaceOffsetPoint = UseHitPos;
 		}
@@ -604,7 +599,7 @@ void UDrawPolygonTool::OnBeginClickSequence(const FInputDeviceRay& ClickPos)
 
 	UpdatePreviewVertex(HitPos);
 
-	bInFixedPolygonMode = (PolygonProperties->PolygonType != EDrawPolygonDrawMode::Freehand);
+	bInFixedPolygonMode = (PolygonProperties->PolygonDrawMode != EDrawPolygonDrawMode::Freehand);
 	FixedPolygonClickPoints.Reset();
 
 	// Actually process the click.
@@ -666,7 +661,7 @@ bool UDrawPolygonTool::OnNextSequenceClick(const FInputDeviceRay& ClickPos)
 	TUniquePtr<FDrawPolygonStateChange> Change =
 		MakeUnique<FDrawPolygonStateChange>(CurrentCurveTimestamp, FixedPolygonClickPoints, PolygonVertices);
 
-	bool bDonePolygon = false;
+	bool bDonePolygon;
 	if (bInFixedPolygonMode)
 	{
 		// ignore very close click points
@@ -676,7 +671,7 @@ bool UDrawPolygonTool::OnNextSequenceClick(const FInputDeviceRay& ClickPos)
 		}
 
 		FixedPolygonClickPoints.Add(HitPos);
-		int NumTargetPoints = (PolygonProperties->PolygonType == EDrawPolygonDrawMode::Rectangle || PolygonProperties->PolygonType == EDrawPolygonDrawMode::RoundedRectangle) ? 3 : 2;
+		int NumTargetPoints = (PolygonProperties->PolygonDrawMode == EDrawPolygonDrawMode::Rectangle || PolygonProperties->PolygonDrawMode == EDrawPolygonDrawMode::RoundedRectangle) ? 3 : 2;
 		bDonePolygon = (FixedPolygonClickPoints.Num() == NumTargetPoints);
 		if (bDonePolygon)
 		{
@@ -691,9 +686,6 @@ bool UDrawPolygonTool::OnNextSequenceClick(const FInputDeviceRay& ClickPos)
 			return true;
 		}
 
-		// close polygon if we clicked on start point
-		bDonePolygon = SnapEngine.HaveActiveSnap() && SnapEngine.GetActiveSnapTargetID() == StartPointSnapID;
-
 		if (bHaveSelfIntersection)
 		{
 			// discard vertex in segments before intersection (this is redundant if idx is 0)
@@ -702,9 +694,14 @@ bool UDrawPolygonTool::OnNextSequenceClick(const FInputDeviceRay& ClickPos)
 				PolygonVertices[j-SelfIntersectSegmentIdx] = PolygonVertices[j];
 			}
 			PolygonVertices.SetNum(PolygonVertices.Num() - SelfIntersectSegmentIdx);
-			PolygonVertices[0] = PreviewVertex = (FVector3d)SelfIntersectionPoint;
+			PolygonVertices[0] = PreviewVertex = SelfIntersectionPoint;
 			bDonePolygon = true;
 		}
+		else
+		{
+			// close polygon if we clicked on start point
+			bDonePolygon = SnapEngine.HaveActiveSnap() && SnapEngine.GetActiveSnapTargetID() == StartPointSnapID;
+		}
 	}
 
 	// emit change event
@@ -714,7 +711,7 @@ bool UDrawPolygonTool::OnNextSequenceClick(const FInputDeviceRay& ClickPos)
 	{
 		//SnapEngine.Reset();
 		bHaveSurfaceHit = false;
-		if (PolygonProperties->OutputMode == EDrawPolygonOutputMode::ExtrudedInteractive)
+		if (PolygonProperties->ExtrudeMode == EDrawPolygonOutputMode::ExtrudedInteractive)
 		{
 			BeginInteractiveExtrude();
 
@@ -848,27 +845,27 @@ void UDrawPolygonTool::GenerateFixedPolygon(const TArray<FVector3d>& FixedPoints
 
 	FPolygon2d Polygon;
 	TArray<FPolygon2d> PolygonHoles;
-	if (PolygonProperties->PolygonType == EDrawPolygonDrawMode::Square)
+	if (PolygonProperties->PolygonDrawMode == EDrawPolygonDrawMode::Square)
 	{
 		Polygon = FPolygon2d::MakeRectangle(FVector2d::Zero(), 2*Width, 2*Width);
 	}
-	else if (PolygonProperties->PolygonType == EDrawPolygonDrawMode::Rectangle || PolygonProperties->PolygonType == EDrawPolygonDrawMode::RoundedRectangle)
+	else if (PolygonProperties->PolygonDrawMode == EDrawPolygonDrawMode::Rectangle || PolygonProperties->PolygonDrawMode == EDrawPolygonDrawMode::RoundedRectangle)
 	{
-		if (PolygonProperties->PolygonType == EDrawPolygonDrawMode::Rectangle)
+		if (PolygonProperties->PolygonDrawMode == EDrawPolygonDrawMode::Rectangle)
 		{
 			Polygon = FPolygon2d::MakeRectangle(FVector2d(Width / 2, YSign*Height / 2), Width, Height);
 		}
-		else // PolygonProperties->PolygonType == EDrawPolygonDrawMode::RoundedRectangle
+		else // PolygonProperties->PolygonDrawMode == EDrawPolygonDrawMode::RoundedRectangle
 		{
-			Polygon = FPolygon2d::MakeRoundedRectangle(FVector2d(Width / 2, YSign*Height / 2), Width, Height, FMathd::Min(Width,Height) * FMathd::Clamp(PolygonProperties->FeatureSizeRatio, .01, .99) * .5, PolygonProperties->Steps);
+			Polygon = FPolygon2d::MakeRoundedRectangle(FVector2d(Width / 2, YSign*Height / 2), Width, Height, FMathd::Min(Width,Height) * FMathd::Clamp(PolygonProperties->FeatureSizeRatio, .01, .99) * .5, PolygonProperties->RadialSlices);
 		}
 	}
-	else // Circle or HoleyCircle
+	else // Circle or Ring
 	{
-		Polygon = FPolygon2d::MakeCircle(Width, PolygonProperties->Steps, 0);
-		if (PolygonProperties->PolygonType == EDrawPolygonDrawMode::HoleyCircle)
+		Polygon = FPolygon2d::MakeCircle(Width, PolygonProperties->RadialSlices, 0);
+		if (PolygonProperties->PolygonDrawMode == EDrawPolygonDrawMode::Ring)
 		{
-			PolygonHoles.Add(FPolygon2d::MakeCircle(Width * FMathd::Clamp(PolygonProperties->FeatureSizeRatio, .01, .99), PolygonProperties->Steps, 0));
+			PolygonHoles.Add(FPolygon2d::MakeCircle(Width * FMathd::Clamp(PolygonProperties->FeatureSizeRatio, .01, .99), PolygonProperties->RadialSlices, 0));
 		}
 	}
 	Polygon.Transform([RotationMat](const FVector2d& Pt) { return RotationMat * Pt; });
@@ -997,13 +994,13 @@ void UDrawPolygonTool::UpdateShowGizmoState(bool bNewVisibility)
 
 void UDrawPolygonTool::EmitCurrentPolygon()
 {
-	FString BaseName = (PolygonProperties->OutputMode == EDrawPolygonOutputMode::MeshedPolygon) ?
+	FString BaseName = (PolygonProperties->ExtrudeMode == EDrawPolygonOutputMode::MeshedPolygon) ?
 		TEXT("Polygon") : TEXT("Extrude");
 
 	// generate new mesh
 	FFrame3d PlaneFrameOut;
 	FDynamicMesh3 Mesh;
-	double ExtrudeDist = (PolygonProperties->OutputMode == EDrawPolygonOutputMode::MeshedPolygon) ?
+	const double ExtrudeDist = (PolygonProperties->ExtrudeMode == EDrawPolygonOutputMode::MeshedPolygon) ?
 		0 : PolygonProperties->ExtrudeHeight;
 	bool bSucceeded = GeneratePolygonMesh(PolygonVertices, PolygonHolesVertices, &Mesh, PlaneFrameOut, false, ExtrudeDist, false);
 	if (!bSucceeded) // somehow made a polygon with no valid triangulation; just throw it away ...
@@ -1042,7 +1039,7 @@ void UDrawPolygonTool::UpdateLivePreview()
 
 	FFrame3d PlaneFrame;
 	FDynamicMesh3 Mesh;
-	double ExtrudeDist = (PolygonProperties->OutputMode == EDrawPolygonOutputMode::MeshedPolygon) ?
+	const double ExtrudeDist = (PolygonProperties->ExtrudeMode == EDrawPolygonOutputMode::MeshedPolygon) ?
 		0 : PolygonProperties->ExtrudeHeight;
 	if (GeneratePolygonMesh(PolygonVertices, PolygonHolesVertices, &Mesh, PlaneFrame, false, ExtrudeDist, false))
 	{
@@ -1090,7 +1087,7 @@ bool UDrawPolygonTool::GeneratePolygonMesh(const TArray<FVector3d>& Polygon, con
 	FAxisAlignedBox2d Bounds(OuterPolygon.Bounds());
 
 	// special case paths
-	if (PolygonProperties->PolygonType == EDrawPolygonDrawMode::HoleyCircle || PolygonProperties->PolygonType == EDrawPolygonDrawMode::Circle || PolygonProperties->PolygonType == EDrawPolygonDrawMode::RoundedRectangle)
+	if (PolygonProperties->PolygonDrawMode == EDrawPolygonDrawMode::Ring || PolygonProperties->PolygonDrawMode == EDrawPolygonDrawMode::Circle || PolygonProperties->PolygonDrawMode == EDrawPolygonDrawMode::RoundedRectangle)
 	{
 		// get polygon parameters
 		FVector2d FirstReferencePt, BoxSize;
@@ -1100,28 +1097,28 @@ bool UDrawPolygonTool::GeneratePolygonMesh(const TArray<FVector3d>& Polygon, con
 		FMatrix2d RotationMat = FMatrix2d::RotationRad(AngleRad);
 
 		// translate general polygon parameters to specific mesh generator parameters, and generate mesh
-		if (PolygonProperties->PolygonType == EDrawPolygonDrawMode::HoleyCircle)
+		if (PolygonProperties->PolygonDrawMode == EDrawPolygonDrawMode::Ring)
 		{
 			FPuncturedDiscMeshGenerator HCGen;
-			HCGen.AngleSamples = PolygonProperties->Steps;
+			HCGen.AngleSamples = PolygonProperties->RadialSlices;
 			HCGen.RadialSamples = 1;
 			HCGen.Radius = BoxSize.X;
 			HCGen.HoleRadius = BoxSize.X * FMathd::Clamp(PolygonProperties->FeatureSizeRatio, .01f, .99f);
 			ResultMeshOut->Copy(&HCGen.Generate());
 		}
-		else if (PolygonProperties->PolygonType == EDrawPolygonDrawMode::Circle)
+		else if (PolygonProperties->PolygonDrawMode == EDrawPolygonDrawMode::Circle)
 		{
 			FDiscMeshGenerator CGen;
-			CGen.AngleSamples = PolygonProperties->Steps;
+			CGen.AngleSamples = PolygonProperties->RadialSlices;
 			CGen.RadialSamples = 1;
 			CGen.Radius = BoxSize.X;
 			ResultMeshOut->Copy(&CGen.Generate());
 		}
-		else if (PolygonProperties->PolygonType == EDrawPolygonDrawMode::RoundedRectangle)
+		else if (PolygonProperties->PolygonDrawMode == EDrawPolygonDrawMode::RoundedRectangle)
 		{
 			FRoundedRectangleMeshGenerator RRGen;
 			FirstReferencePt += RotationMat * (FVector2d(BoxSize.X, BoxSize.Y * YSign)*.5f);
-			RRGen.AngleSamples = PolygonProperties->Steps;
+			RRGen.AngleSamples = PolygonProperties->RadialSlices;
 			RRGen.Radius = .5 * FMathd::Min(BoxSize.X, BoxSize.Y) * FMathd::Clamp(PolygonProperties->FeatureSizeRatio, .01f, .99f);
 			RRGen.Height = BoxSize.Y - RRGen.Radius * 2.;
 			RRGen.Width = BoxSize.X - RRGen.Radius * 2.;
@@ -1181,7 +1178,7 @@ bool UDrawPolygonTool::GeneratePolygonMesh(const TArray<FVector3d>& Polygon, con
 		}
 
 
-		bool bTriangulationSuccess = Triangulator.Triangulate([&GeneralPolygon](const TArray<FVector2d>& Vertices, FIndex3i Tri)
+		Triangulator.Triangulate([&GeneralPolygon](const TArray<FVector2d>& Vertices, FIndex3i Tri)
 		{
 			// keep triangles based on the input polygon's winding
 			return GeneralPolygon.Contains((Vertices[Tri.A] + Vertices[Tri.B] + Vertices[Tri.C]) / 3.0);
diff --git a/Engine/Plugins/Runtime/MeshModelingToolset/Source/MeshModelingTools/Public/DrawPolygonTool.h b/Engine/Plugins/Runtime/MeshModelingToolset/Source/MeshModelingTools/Public/DrawPolygonTool.h
--- a/Engine/Plugins/Runtime/MeshModelingToolset/Source/MeshModelingTools/Public/DrawPolygonTool.h
+++ b/Engine/Plugins/Runtime/MeshModelingToolset/Source/MeshModelingTools/Public/DrawPolygonTool.h
@@ -36,42 +36,41 @@ public:
 
 
 
-/** Polygon Tool Draw Type */
+/** Polygon tool draw type */
 UENUM()
 enum class EDrawPolygonDrawMode : uint8
 {
-	/** Freehand Polygon Drawing */
-	Freehand UMETA(DisplayName = "Freehand"),
+	/** Draw a freehand polygon */
+	Freehand,
 
-	/** Circle */
-	Circle UMETA(DisplayName = "Circle"),
+	/** Draw a circle */
+	Circle,
 
-	/** Square */
-	Square UMETA(DisplayName = "Square"),
+	/** Draw a square */
+	Square,
 
-	/** Rectangle */
-	Rectangle UMETA(DisplayName = "Rectangle"),
+	/** Draw a rectangle */
+	Rectangle,
 
-	/** Rounded Rectangle */
-	RoundedRectangle UMETA(DisplayName = "Rounded Rectangle"),
-
-	/** Circle w/ Hole */
-	HoleyCircle UMETA(DisplayName = "Circle w/ Hole")
+	/** Draw a rounded rectangle */
+	RoundedRectangle,
 
+	/** Draw a circle with a hole in the center */
+	Ring
 };
 
 
-/** Output of Draw Polygon Tool */
+/** Output of draw polygon tool */
 UENUM()
 enum class EDrawPolygonOutputMode : uint8
 {
-	/** Generate a meshed planar polygon */
+	/** Generate a flat polygon mesh */
 	MeshedPolygon UMETA(DisplayName = "Flat Mesh"),
 
-	/** Extrude closed polygon to constant height determined by Extrude Height Property */
-	ExtrudedConstant UMETA(DisplayName = "Extrude To Height"),
+	/** Extrude drawn polygon to constant height determined by the Extrude Height property */
+	ExtrudedConstant UMETA(DisplayName = "Extrude to Height"),
 
-	/** Extrusion height is set via additional mouse input after closing polygon */
+	/** Extrude drawn polygon to height set via additional mouse input after closing the polygon */
 	ExtrudedInteractive UMETA(DisplayName = "Interactive Extrude"),
 };
 
@@ -87,32 +86,42 @@ class MESHMODELINGTOOLS_API UDrawPolygonToolStandardProperties : public UInterac
 public:
 	UDrawPolygonToolStandardProperties();
 
-	UPROPERTY(EditAnywhere, NonTransactional, Category = Polygon)
-	EDrawPolygonDrawMode PolygonType = EDrawPolygonDrawMode::Freehand;
+	/** Type of polygon to draw in the viewport */
+	UPROPERTY(EditAnywhere, NonTransactional, Category = Polygon, meta = (DisplayName = "Draw Mode"))
+	EDrawPolygonDrawMode PolygonDrawMode = EDrawPolygonDrawMode::Freehand;
 
-	/** Feature size as fraction of overall shape size, for shapes with secondary features like the rounded corners of a Rounded Rectangle */
+	/** Size of secondary features, e.g. the rounded corners of a rounded rectangle, as fraction of the overall shape size */
 	UPROPERTY(EditAnywhere, NonTransactional, Category = Polygon, meta = (UIMin = "0.01", UIMax = "0.99", ClampMin = "0.01", ClampMax = "0.99",
-									EditCondition = "PolygonType == EDrawPolygonDrawMode::RoundedRectangle || PolygonType == EDrawPolygonDrawMode::HoleyCircle"))
+		EditCondition = "PolygonDrawMode == EDrawPolygonDrawMode::RoundedRectangle || PolygonDrawMode == EDrawPolygonDrawMode::Ring", EditConditionHides))
 	float FeatureSizeRatio = .25;
 
-	/** Number of vertices in round features */
+	/** Allow freehand drawn polygons to self-intersect */
+	UPROPERTY(EditAnywhere, NonTransactional, Category = Polygon,
+		meta = (DisplayName ="Self-Intersections", EditCondition = "PolygonDrawMode == EDrawPolygonDrawMode::Freehand", EditConditionHides))
+	bool bAllowSelfIntersections = false;
+
+	/** Number of radial subdivisions in round features, e.g. circles or rounded corners */
 	UPROPERTY(EditAnywhere, NonTransactional, Category = Polygon, meta = (UIMin = "3", UIMax = "100", ClampMin = "3", ClampMax = "10000",
-				  EditCondition = "PolygonType == EDrawPolygonDrawMode::Circle || PolygonType == EDrawPolygonDrawMode::RoundedRectangle || PolygonType == EDrawPolygonDrawMode::HoleyCircle"))
-	int Steps = 16;
+		EditCondition =	"PolygonDrawMode == EDrawPolygonDrawMode::Circle || PolygonDrawMode == EDrawPolygonDrawMode::RoundedRectangle || PolygonDrawMode == EDrawPolygonDrawMode::Ring",
+		EditConditionHides))
+	int RadialSlices = 16;
 
-	UPROPERTY(EditAnywhere, NonTransactional, Category = Polygon, meta = (EditCondition = "PolygonType == EDrawPolygonDrawMode::Freehand"))
-	bool bAllowSelfIntersections = false;
+	/** Distance between the last clicked point and the current point  */
+	UPROPERTY(VisibleAnywhere, NonTransactional, Category = Polygon, meta = (TransientToolProperty))
+	float Distance = 0.0f;
 
+	/** If and how the drawn polygon gets extruded */
 	UPROPERTY(EditAnywhere, NonTransactional, Category = Polygon)
-	EDrawPolygonOutputMode OutputMode = EDrawPolygonOutputMode::ExtrudedInteractive;
+	EDrawPolygonOutputMode ExtrudeMode = EDrawPolygonOutputMode::ExtrudedInteractive;
 
-	/** Extrusion Distance in non-interactive mode */
+	/** Extrusion distance when using the non-interactive extrude mode */
 	UPROPERTY(EditAnywhere, NonTransactional, Category = Polygon, meta = (UIMin = "-1000", UIMax = "1000", ClampMin = "-10000", ClampMax = "10000",
-									EditCondition = "OutputMode == EDrawPolygonOutputMode::ExtrudedConstant"))
+									EditCondition = "ExtrudeMode == EDrawPolygonOutputMode::ExtrudedConstant"))
 	float ExtrudeHeight = 100.0f;
 
+	/** If true, shows a gizmo to manipulate the additional grid used to draw the polygon on */
 	UPROPERTY(EditAnywhere, NonTransactional, Category = Polygon)
-	bool bShowGizmo = true;
+	bool bShowGridGizmo = true;
 };
 
 UCLASS()
@@ -121,34 +130,37 @@ class MESHMODELINGTOOLS_API UDrawPolygonToolSnapProperties : public UInteractive
 	GENERATED_BODY()
 
 public:
+	/** If true, enables additional snapping controls */
 	UPROPERTY(EditAnywhere, NonTransactional, Category = Snapping)
 	bool bEnableSnapping = true;
 
-	/** When true, allows snapping to world grid according to editor settings. */
-	UPROPERTY(EditAnywhere, NonTransactional, Category = Snapping, Meta = (EditCondition = "bEnableSnapping", EditConditionHides))
+	/** If true, allows snapping to world grid according to editor settings */
+	UPROPERTY(EditAnywhere, NonTransactional, Category = Snapping, meta = (EditCondition = "bEnableSnapping", EditConditionHides))
 	bool bSnapToWorldGrid = false;
 
-	UPROPERTY(EditAnywhere, NonTransactional, Category = Snapping, Meta = (EditCondition = "bEnableSnapping", EditConditionHides))
+	/** If true, snap to vertices in the current polygon or in other meshes */
+	UPROPERTY(EditAnywhere, NonTransactional, Category = Snapping, meta = (EditCondition = "bEnableSnapping", EditConditionHides))
 	bool bSnapToVertices = true;
 
-	UPROPERTY(EditAnywhere, NonTransactional, Category = Snapping, Meta = (EditCondition = "bEnableSnapping", EditConditionHides))
+	/** If true, snap to edges in the current polygon or in other meshes */
+	UPROPERTY(EditAnywhere, NonTransactional, Category = Snapping, meta = (EditCondition = "bEnableSnapping", EditConditionHides))
 	bool bSnapToEdges = false;
 
-	UPROPERTY(EditAnywhere, NonTransactional, Category = Snapping, Meta = (EditCondition = "bEnableSnapping", EditConditionHides))
-	bool bSnapToAngles = true;
+	/** If true, snap to axes of the additional drawing grid as well as the axes relative to the last segment */
+	UPROPERTY(EditAnywhere, NonTransactional, Category = Snapping, meta = (EditCondition = "bEnableSnapping", EditConditionHides))
+	bool bSnapToAxes = true;
 
-	/** When snapping to an angle, also prefer to snap to an existing segment length */
-	UPROPERTY(EditAnywhere, NonTransactional, Category = Snapping, Meta = (EditCondition = "bEnableSnapping && bSnapToAngles", EditConditionHides))
+	/** If true, when snapping to axes, also try to snap to the length of an existing segment in the polygon */
+	UPROPERTY(EditAnywhere, NonTransactional, Category = Snapping, meta = (EditCondition = "bEnableSnapping && bSnapToAxes", EditConditionHides))
 	bool bSnapToLengths = true;
 
-	UPROPERTY(VisibleAnywhere, NonTransactional, Category = Snapping, meta = (TransientToolProperty))
-	float SegmentLength = 0.0f;
-
+	/** If true, snap to surfaces of existing objects */
 	UPROPERTY(EditAnywhere, NonTransactional, Category = Snapping)
-	bool bHitSceneObjects = false;
+	bool bSnapToSurfaces = false;
 
-	UPROPERTY(EditAnywhere, NonTransactional, Category = Snapping, Meta = (EditCondition = "bHitSceneObjects", EditConditionHides))
-	float HitNormalOffset = 0.0f;
+	/** Offset for snap point on the surface of an existing object in the direction of the surface normal */
+	UPROPERTY(EditAnywhere, NonTransactional, Category = Snapping, meta = (DisplayName = "Surface Offset", EditCondition = "bSnapToSurfaces", EditConditionHides))
+	float SnapToSurfacesOffset = 0.0f;
 };
 
 
@@ -200,23 +212,19 @@ public:
 	virtual void BeginInteractiveExtrude();
 	virtual void EndInteractiveExtrude();
 
-
-public:
 	virtual void ApplyUndoPoints(const TArray<FVector3d>& ClickPointsIn, const TArray<FVector3d>& PolygonVerticesIn);
 
 
 protected:
 	// flags used to identify modifier keys/buttons
-	static const int IgnoreSnappingModifier = 1;
-	static const int AngleSnapModifier = 2;
-
-protected:
+	static constexpr int IgnoreSnappingModifier = 1;
+	static constexpr int AngleSnapModifier = 2;
 
 	/** Property set for type of output object (StaticMesh, Volume, etc) */
 	UPROPERTY()
 	TObjectPtr<UCreateMeshObjectTypeProperties> OutputTypeProperties;
 
-	/** Properties that control polygon generation exposed to user via detailsview */
+	/** Properties that control polygon generation exposed to user via details view */
 	UPROPERTY()
 	TObjectPtr<UDrawPolygonToolStandardProperties> PolygonProperties;
 
@@ -242,7 +250,7 @@ protected:
 	/** last vertex of polygon that is actively being updated as input device is moved */
 	FVector3d PreviewVertex;
 
-protected:
+	UPROPERTY()
 	UWorld* TargetWorld;
 
 	FViewCameraState CameraState;
@@ -283,7 +291,7 @@ protected:
 	int SelfIntersectSegmentIdx;
 	FVector3d SelfIntersectionPoint;
 
-	// only used when SnapSettings.bHitSceneObjects = true
+	// only used when SnapSettings.bSnapToSurfaces = true
 	bool bHaveSurfaceHit;
 	FVector3d SurfaceHitPoint;
 	FVector3d SurfaceOffsetPoint;
