MeshModelingTools: Fix thread issues in edge loop insertion, edge insertion, and cube grid tool activities.

#rnx

diff --git a/Engine/Plugins/Experimental/MeshModelingToolsetExp/Source/MeshModelingToolsExp/Private/CubeGridTool.cpp b/Engine/Plugins/Experimental/MeshModelingToolsetExp/Source/MeshModelingToolsExp/Private/CubeGridTool.cpp
--- a/Engine/Plugins/Experimental/MeshModelingToolsetExp/Source/MeshModelingToolsExp/Private/CubeGridTool.cpp
+++ b/Engine/Plugins/Experimental/MeshModelingToolsetExp/Source/MeshModelingToolsExp/Private/CubeGridTool.cpp
@@ -540,7 +540,7 @@ TUniquePtr<FDynamicMeshOperator> UCubeGridTool::MakeNewOperator()
 
 	// Make the op.
 	TUniquePtr<FCubeGridBooleanOp> Op = MakeUnique<FCubeGridBooleanOp>();
-	Op->InputMesh = CurrentMesh;
+	Op->InputMesh = ComputeStartMesh;
 	Op->InputTransform = CurrentMeshTransform;
 	Op->bKeepInputTransform = true;
 	Op->WorldBox = WorldBox;
@@ -679,9 +679,13 @@ void UCubeGridTool::Setup()
 		Preview->ConfigureMaterials(MaterialSet.Materials, ToolSetupUtil::GetDefaultWorkingMaterial(GetToolManager()));
 	}
 	Preview->PreviewMesh->SetTransform((FTransform)CurrentMeshTransform);
-	Preview->OnOpCompleted.AddWeakLambda(this, [this](const FDynamicMeshOperator* Op) {
-		LastOpChangedTids = static_cast<const FCubeGridBooleanOp*>(Op)->ChangedTids;
-		});
+	Preview->OnOpCompleted.AddWeakLambda(this, [this](const FDynamicMeshOperator* UncastOp) {
+		const FCubeGridBooleanOp* Op = static_cast<const FCubeGridBooleanOp*>(UncastOp);
+		if (Op->InputMesh == ComputeStartMesh)
+		{
+			LastOpChangedTids = static_cast<const FCubeGridBooleanOp*>(Op)->ChangedTids;
+		}
+	});
 
 	CubeGrid = MakeShared<FCubeGrid>();
 	CubeGrid->SetGridFrame(FFrame3d(Settings->GridFrameOrigin, Settings->GridFrameOrientation.Quaternion()));
@@ -830,6 +834,13 @@ void UCubeGridTool::Setup()
 		});
 
 	Settings->SilentUpdateWatched();
+
+	UpdateComputeInputs();
+}
+
+void UCubeGridTool::UpdateComputeInputs()
+{
+	ComputeStartMesh = MakeShared<FDynamicMesh3>(*CurrentMesh);
 }
 
 void UCubeGridTool::Shutdown(EToolShutdownType ShutdownType)
@@ -922,6 +933,8 @@ void UCubeGridTool::ApplyPreview()
 	CurrentMesh->Copy(*Preview->PreviewMesh->GetMesh());
 	MeshSpatial->Build();
 
+	UpdateComputeInputs();
+
 	const FText TransactionText = LOCTEXT("CubeGridToolTransactionName", "Block Tool Change");
 	GetToolManager()->BeginUndoTransaction(TransactionText);
 
@@ -2022,6 +2035,7 @@ void UCubeGridTool::UpdateUsingMeshChange(const FDynamicMeshChange& MeshChange,
 {
 	MeshChange.Apply(CurrentMesh.Get(), bRevert);
 	MeshSpatial->Build();
+	UpdateComputeInputs();
 	CurrentExtrudeAmount = 0;
 	bPreviewMayDiffer = true;
 	InvalidatePreview();
diff --git a/Engine/Plugins/Experimental/MeshModelingToolsetExp/Source/MeshModelingToolsExp/Public/CubeGridTool.h b/Engine/Plugins/Experimental/MeshModelingToolsetExp/Source/MeshModelingToolsExp/Public/CubeGridTool.h
--- a/Engine/Plugins/Experimental/MeshModelingToolsetExp/Source/MeshModelingToolsExp/Public/CubeGridTool.h
+++ b/Engine/Plugins/Experimental/MeshModelingToolsetExp/Source/MeshModelingToolsExp/Public/CubeGridTool.h
@@ -463,6 +463,10 @@ protected:
 	UE::Geometry::FTransform3d CurrentMeshTransform = UE::Geometry::FTransform3d::Identity();
 	TSharedPtr<TArray<int32>, ESPMode::ThreadSafe> LastOpChangedTids;
 
+	// Safe inputs for the background compute to use, untouched by undo/redo/other CurrentMesh updates.
+	TSharedPtr<const UE::Geometry::FDynamicMesh3, ESPMode::ThreadSafe> ComputeStartMesh;
+	void UpdateComputeInputs();
+
 	static const int32 ShiftModifierID = 1;
 	static const int32 CtrlModifierID = 2;
 
diff --git a/Engine/Plugins/Runtime/MeshModelingToolset/Source/MeshModelingTools/Private/EditMeshPolygonsTool.cpp b/Engine/Plugins/Runtime/MeshModelingToolset/Source/MeshModelingTools/Private/EditMeshPolygonsTool.cpp
--- a/Engine/Plugins/Runtime/MeshModelingToolset/Source/MeshModelingTools/Private/EditMeshPolygonsTool.cpp
+++ b/Engine/Plugins/Runtime/MeshModelingToolset/Source/MeshModelingTools/Private/EditMeshPolygonsTool.cpp
@@ -319,7 +319,7 @@ void UEditMeshPolygonsTool::Setup()
 	}
 
 	// Wrap the data structures into a context that we can give to the activities
-	UPolyEditActivityContext* ActivityContext = NewObject<UPolyEditActivityContext>();
+	ActivityContext = NewObject<UPolyEditActivityContext>();
 	ActivityContext->bTriangleMode = bTriangleMode;
 	ActivityContext->CommonProperties = CommonProps;
 	ActivityContext->CurrentMesh = CurrentMesh;
@@ -441,6 +441,7 @@ void UEditMeshPolygonsTool::Shutdown(EToolShutdownType ShutdownType)
 	CommonProps->SaveProperties(this);
 
 	GetToolManager()->GetContextObjectStore()->RemoveContextObjectsOfType(UPolyEditActivityContext::StaticClass());
+	ActivityContext = nullptr;
 
 	ExtrudeActivity->Shutdown(ShutdownType);
 	InsetOutsetActivity->Shutdown(ShutdownType);
@@ -1943,6 +1944,7 @@ void FEditMeshPolygonsToolMeshChange::Apply(UObject* Object)
 	MeshChange->Apply(Tool->CurrentMesh.Get(), false);
 	Tool->UpdateFromCurrentMesh(bGroupTopologyModified);
 	Tool->ModifiedTopologyCounter += bGroupTopologyModified;
+	Tool->ActivityContext->OnUndoRedo.Broadcast(bGroupTopologyModified);
 }
 
 void FEditMeshPolygonsToolMeshChange::Revert(UObject* Object)
@@ -1952,6 +1954,7 @@ void FEditMeshPolygonsToolMeshChange::Revert(UObject* Object)
 	MeshChange->Apply(Tool->CurrentMesh.Get(), true);
 	Tool->UpdateFromCurrentMesh(bGroupTopologyModified);
 	Tool->ModifiedTopologyCounter -= bGroupTopologyModified;
+	Tool->ActivityContext->OnUndoRedo.Broadcast(bGroupTopologyModified);
 }
 
 FString FEditMeshPolygonsToolMeshChange::ToString() const
diff --git a/Engine/Plugins/Runtime/MeshModelingToolset/Source/MeshModelingTools/Private/ToolActivities/PolyEditInsertEdgeActivity.cpp b/Engine/Plugins/Runtime/MeshModelingToolset/Source/MeshModelingTools/Private/ToolActivities/PolyEditInsertEdgeActivity.cpp
--- a/Engine/Plugins/Runtime/MeshModelingToolset/Source/MeshModelingTools/Private/ToolActivities/PolyEditInsertEdgeActivity.cpp
+++ b/Engine/Plugins/Runtime/MeshModelingToolset/Source/MeshModelingTools/Private/ToolActivities/PolyEditInsertEdgeActivity.cpp
@@ -31,8 +31,8 @@ TUniquePtr<FDynamicMeshOperator> UPolyEditInsertEdgeActivity::MakeNewOperator()
 {
 	TUniquePtr<FGroupEdgeInsertionOp> Op = MakeUnique<FGroupEdgeInsertionOp>();
 
-	Op->OriginalMesh = ActivityContext->CurrentMesh;
-	Op->OriginalTopology = ActivityContext->CurrentTopology;
+	Op->OriginalMesh = ComputeStartMesh;
+	Op->OriginalTopology = ComputeStartTopology;
 	Op->SetTransform(TargetTransform);
 
 	if (Settings->InsertionMode == EGroupEdgeInsertionMode::PlaneCut)
@@ -90,8 +90,25 @@ void UPolyEditInsertEdgeActivity::Setup(UInteractiveTool* ParentToolIn)
 	ActivityContext = ParentTool->GetToolManager()->GetContextObjectStore()->FindContext<UPolyEditActivityContext>();
 
 	TopologySelector = ActivityContext->SelectionMechanic->GetTopologySelector();
+
+	ActivityContext->OnUndoRedo.AddWeakLambda(this, [this](bool bGroupTopologyModified)
+	{
+		UpdateComputeInputs();
+		ClearPreview();
+	});
 }
 
+void UPolyEditInsertEdgeActivity::UpdateComputeInputs()
+{
+	ComputeStartMesh = MakeShared<FDynamicMesh3>(*ActivityContext->CurrentMesh);
+
+	TSharedPtr<FGroupTopology> NonConstTopology = MakeShared<FGroupTopology>(*ActivityContext->CurrentTopology);
+	NonConstTopology->RetargetOnClonedMesh(ComputeStartMesh.Get());
+
+	ComputeStartTopology = NonConstTopology;
+}
+
+
 void UPolyEditInsertEdgeActivity::Shutdown(EToolShutdownType ShutdownType)
 {
 	if (bIsRunning)
@@ -104,6 +121,7 @@ void UPolyEditInsertEdgeActivity::Shutdown(EToolShutdownType ShutdownType)
 
 	TopologySelector.Reset();
 	Settings = nullptr;
+	ActivityContext->OnUndoRedo.RemoveAll(this);
 	ActivityContext = nullptr;
 
 	Super::Shutdown(ShutdownType);
@@ -129,6 +147,7 @@ EToolActivityStartResult UPolyEditInsertEdgeActivity::Start()
 
 	TargetTransform = ActivityContext->Preview->PreviewMesh->GetTransform();
 
+	UpdateComputeInputs();
 	SetupPreview();
 
 	ToolState = EState::GettingStart;
@@ -174,11 +193,20 @@ void UPolyEditInsertEdgeActivity::SetupPreview()
 	ActivityContext->Preview->OnOpCompleted.AddWeakLambda(this, [this](const FDynamicMeshOperator* UncastOp) {
 		const FGroupEdgeInsertionOp* Op = static_cast<const FGroupEdgeInsertionOp*>(UncastOp);
 
-		bLastComputeSucceeded = Op->bSucceeded;
 		LatestOpTopologyResult.Reset();
 		LatestOpChangedTids.Reset();
-
 		PreviewEdges.Reset();
+
+		// See if this compute is actually outdated, i.e. we changed the mesh
+		// out from under it.
+		if (Op->OriginalMesh != ComputeStartMesh)
+		{
+			bLastComputeSucceeded = false;
+			return;
+		}
+
+		bLastComputeSucceeded = Op->bSucceeded;
+
 		if (bLastComputeSucceeded)
 		{
 			Op->GetEdgeLocations(PreviewEdges);
@@ -221,6 +249,8 @@ void UPolyEditInsertEdgeActivity::Tick(float DeltaTime)
 			*ActivityContext->CurrentTopology = *LatestOpTopologyResult;
 			ActivityContext->CurrentTopology->RetargetOnClonedMesh(ActivityContext->CurrentMesh.Get());
 			
+			UpdateComputeInputs();
+
 			// Emit transaction
 			FGroupTopologySelection EmptySelection;
 			ActivityContext->EmitCurrentMeshChangeAndUpdate(LOCTEXT("EdgeInsertionTransactionName", "Edge Insertion"),
diff --git a/Engine/Plugins/Runtime/MeshModelingToolset/Source/MeshModelingTools/Private/ToolActivities/PolyEditInsertEdgeLoopActivity.cpp b/Engine/Plugins/Runtime/MeshModelingToolset/Source/MeshModelingTools/Private/ToolActivities/PolyEditInsertEdgeLoopActivity.cpp
--- a/Engine/Plugins/Runtime/MeshModelingToolset/Source/MeshModelingTools/Private/ToolActivities/PolyEditInsertEdgeLoopActivity.cpp
+++ b/Engine/Plugins/Runtime/MeshModelingToolset/Source/MeshModelingTools/Private/ToolActivities/PolyEditInsertEdgeLoopActivity.cpp
@@ -23,8 +23,8 @@ TUniquePtr<FDynamicMeshOperator> UPolyEditInsertEdgeLoopActivity::MakeNewOperato
 {
 	TUniquePtr<FEdgeLoopInsertionOp> Op = MakeUnique<FEdgeLoopInsertionOp>();
 
-	Op->OriginalMesh = ActivityContext->CurrentMesh;
-	Op->OriginalTopology = ActivityContext->CurrentTopology;
+	Op->OriginalMesh = ComputeStartMesh;
+	Op->OriginalTopology = ComputeStartTopology;
 	Op->SetTransform(TargetTransform);
 
 	if (Settings->InsertionMode == EEdgeLoopInsertionMode::PlaneCut)
@@ -111,6 +111,22 @@ void UPolyEditInsertEdgeLoopActivity::Setup(UInteractiveTool* ParentToolIn)
 	ActivityContext = ParentTool->GetToolManager()->GetContextObjectStore()->FindContext<UPolyEditActivityContext>();
 
 	TopologySelector = ActivityContext->SelectionMechanic->GetTopologySelector();
+
+	ActivityContext->OnUndoRedo.AddWeakLambda(this, [this](bool bGroupTopologyModified)
+	{
+		UpdateComputeInputs();
+		ClearPreview();
+	});
+}
+
+void UPolyEditInsertEdgeLoopActivity::UpdateComputeInputs()
+{
+	ComputeStartMesh = MakeShared<FDynamicMesh3>(*ActivityContext->CurrentMesh);
+
+	TSharedPtr<FGroupTopology> NonConstTopology = MakeShared<FGroupTopology>(*ActivityContext->CurrentTopology);
+	NonConstTopology->RetargetOnClonedMesh(ComputeStartMesh.Get());
+
+	ComputeStartTopology = NonConstTopology;
 }
 
 void UPolyEditInsertEdgeLoopActivity::Shutdown(EToolShutdownType ShutdownType)
@@ -125,6 +141,8 @@ void UPolyEditInsertEdgeLoopActivity::Shutdown(EToolShutdownType ShutdownType)
 
 	TopologySelector.Reset();
 	Settings = nullptr;
+
+	ActivityContext->OnUndoRedo.RemoveAll(this);
 	ActivityContext = nullptr;
 
 	Super::Shutdown(ShutdownType);
@@ -150,6 +168,7 @@ EToolActivityStartResult UPolyEditInsertEdgeLoopActivity::Start()
 
 	TargetTransform = ActivityContext->Preview->PreviewMesh->GetTransform();
 
+	UpdateComputeInputs();
 	SetupPreview();
 
 	bWaitingForInsertionCompletion = false;
@@ -194,11 +213,20 @@ void UPolyEditInsertEdgeLoopActivity::SetupPreview()
 	ActivityContext->Preview->OnOpCompleted.AddWeakLambda(this, [this](const FDynamicMeshOperator* UncastOp) {
 		const FEdgeLoopInsertionOp* Op = static_cast<const FEdgeLoopInsertionOp*>(UncastOp);
 
-		bLastComputeSucceeded = Op->bSucceeded;
 		LatestOpTopologyResult.Reset();
 		PreviewEdges.Reset();
 		LatestOpChangedTids.Reset();
 
+		// See if this compute is actually outdated, i.e. we changed the mesh
+		// out from under it.
+		if (Op->OriginalMesh != ComputeStartMesh)
+		{
+			bLastComputeSucceeded = false;
+			return;
+		}
+
+		bLastComputeSucceeded = Op->bSucceeded;
+
 		if (bLastComputeSucceeded)
 		{
 			Op->GetLoopEdgeLocations(PreviewEdges);
@@ -259,11 +287,12 @@ void UPolyEditInsertEdgeLoopActivity::Tick(float DeltaTime)
 			// Update current mesh
 			ActivityContext->CurrentMesh->Copy(*ActivityContext->Preview->PreviewMesh->GetMesh(), 
 				true, true, true, true);
-			
 
 			*ActivityContext->CurrentTopology = *LatestOpTopologyResult;
 			ActivityContext->CurrentTopology->RetargetOnClonedMesh(ActivityContext->CurrentMesh.Get());
 
+			UpdateComputeInputs();
+
 			// Emit transaction
 			FGroupTopologySelection EmptySelection;
 			ActivityContext->EmitCurrentMeshChangeAndUpdate(LOCTEXT("EdgeLoopInsertionTransactionName", "Edge Loop Insertion"),
@@ -315,7 +344,6 @@ bool UPolyEditInsertEdgeLoopActivity::CanAccept() const
 
 void UPolyEditInsertEdgeLoopActivity::OnPropertyModified(UObject* PropertySet, FProperty* Property)
 {
-	PreviewEdges.Reset();
 	ClearPreview();
 }
 
diff --git a/Engine/Plugins/Runtime/MeshModelingToolset/Source/MeshModelingTools/Public/EditMeshPolygonsTool.h b/Engine/Plugins/Runtime/MeshModelingToolset/Source/MeshModelingTools/Public/EditMeshPolygonsTool.h
--- a/Engine/Plugins/Runtime/MeshModelingToolset/Source/MeshModelingTools/Public/EditMeshPolygonsTool.h
+++ b/Engine/Plugins/Runtime/MeshModelingToolset/Source/MeshModelingTools/Public/EditMeshPolygonsTool.h
@@ -20,10 +20,12 @@
 PREDECLARE_GEOMETRY(class FGroupTopology);
 PREDECLARE_GEOMETRY(struct FGroupTopologySelection);
 
+class UCombinedTransformGizmo;
 class UDragAlignmentMechanic;
-class UPersistentMeshSelection;
 class UMeshOpPreviewWithBackgroundCompute; 
 class FMeshVertexChangeBuilder;
+class UPersistentMeshSelection;
+class UPolyEditActivityContext;
 class UPolyEditInsertEdgeActivity;
 class UPolyEditInsertEdgeLoopActivity;
 class UPolyEditExtrudeActivity;
@@ -31,7 +33,6 @@ class UPolyEditInsetOutsetActivity;
 class UPolyEditCutFacesActivity;
 class UPolyEditPlanarProjectionUVActivity;
 class UPolygonSelectionMechanic;
-class UCombinedTransformGizmo;	
 class UTransformProxy;
 
 
@@ -492,7 +493,10 @@ protected:
 	TSharedPtr<UE::Geometry::FDynamicMeshAABBTree3> MeshSpatial;
 
 	UPROPERTY()
-	TObjectPtr<UPolygonSelectionMechanic> SelectionMechanic;
+	TObjectPtr<UPolyEditActivityContext> ActivityContext = nullptr;
+
+	UPROPERTY()
+	TObjectPtr<UPolygonSelectionMechanic> SelectionMechanic = nullptr;
 
 	UPROPERTY()
 	TObjectPtr<UDragAlignmentMechanic> DragAlignmentMechanic = nullptr;
diff --git a/Engine/Plugins/Runtime/MeshModelingToolset/Source/MeshModelingTools/Public/ToolActivities/PolyEditActivityContext.h b/Engine/Plugins/Runtime/MeshModelingToolset/Source/MeshModelingTools/Public/ToolActivities/PolyEditActivityContext.h
--- a/Engine/Plugins/Runtime/MeshModelingToolset/Source/MeshModelingTools/Public/ToolActivities/PolyEditActivityContext.h
+++ b/Engine/Plugins/Runtime/MeshModelingToolset/Source/MeshModelingTools/Public/ToolActivities/PolyEditActivityContext.h
@@ -75,4 +75,13 @@ public:
 	 * appropriate undoable transaction with the given transaction label.
 	 */
 	TUniqueFunction<void(const FText& TransactionLabel)> EmitActivityStart;
+
+	/**
+	 * Gets broadcast when the CurrentMesh is modified by an undo/redo transaction emitted via
+	 * EmitCurrentMeshChangeAndUpdate. Only activities that issue multiple transactions via 
+	 * EmitCurrentMeshChangeAndUpdate during the same invocation need to use this, since a 
+	 * transaction that ends immediately after the call will not have to deal with the undo.
+	 */
+	DECLARE_MULTICAST_DELEGATE_OneParam(FOnUndoRedo, bool bGroupTopologyChanged);
+	FOnUndoRedo OnUndoRedo;
 };
diff --git a/Engine/Plugins/Runtime/MeshModelingToolset/Source/MeshModelingTools/Public/ToolActivities/PolyEditInsertEdgeActivity.h b/Engine/Plugins/Runtime/MeshModelingToolset/Source/MeshModelingTools/Public/ToolActivities/PolyEditInsertEdgeActivity.h
--- a/Engine/Plugins/Runtime/MeshModelingToolset/Source/MeshModelingTools/Public/ToolActivities/PolyEditInsertEdgeActivity.h
+++ b/Engine/Plugins/Runtime/MeshModelingToolset/Source/MeshModelingTools/Public/ToolActivities/PolyEditInsertEdgeActivity.h
@@ -143,6 +143,11 @@ protected:
 
 	int32 CurrentChangeStamp = 0;
 
+	// Safe inputs for the background compute to use, untouched by undo/redo/other CurrentMesh updates.
+	TSharedPtr<const UE::Geometry::FDynamicMesh3, ESPMode::ThreadSafe> ComputeStartMesh;
+	TSharedPtr<const UE::Geometry::FGroupTopology, ESPMode::ThreadSafe> ComputeStartTopology;
+	void UpdateComputeInputs();
+
 	void SetupPreview();
 
 	bool TopologyHitTest(const FRay& WorldRay, FVector3d& RayPositionOut, FRay3d* LocalRayOut = nullptr);
diff --git a/Engine/Plugins/Runtime/MeshModelingToolset/Source/MeshModelingTools/Public/ToolActivities/PolyEditInsertEdgeLoopActivity.h b/Engine/Plugins/Runtime/MeshModelingToolset/Source/MeshModelingTools/Public/ToolActivities/PolyEditInsertEdgeLoopActivity.h
--- a/Engine/Plugins/Runtime/MeshModelingToolset/Source/MeshModelingTools/Public/ToolActivities/PolyEditInsertEdgeLoopActivity.h
+++ b/Engine/Plugins/Runtime/MeshModelingToolset/Source/MeshModelingTools/Public/ToolActivities/PolyEditInsertEdgeLoopActivity.h
@@ -150,7 +150,6 @@ protected:
 
 	FViewCameraState CameraState;
 
-	//FToolDataVisualizer ExistingEdgesRenderer;
 	FToolDataVisualizer PreviewEdgeRenderer;
 	FToolDataVisualizer ProblemTopologyRenderer;
 	FGroupTopologySelector::FSelectionSettings TopologySelectorSettings;
@@ -161,6 +160,11 @@ protected:
 	FInputRayHit HitTest(const FRay& WorldRay);
 	bool UpdateHoveredItem(const FRay& WorldRay);
 
+	// Safe inputs for the background compute to use, untouched by undo/redo/other CurrentMesh updates.
+	TSharedPtr<const UE::Geometry::FDynamicMesh3, ESPMode::ThreadSafe> ComputeStartMesh;
+	TSharedPtr<const UE::Geometry::FGroupTopology, ESPMode::ThreadSafe> ComputeStartTopology;
+	void UpdateComputeInputs();
+
 	void ConditionallyUpdatePreview(int32 NewGroupID, double* NewInputLength = nullptr);
 	void ClearPreview();
 
