MeshModelingToolset: Cleanup PathExt UI; remove EditConditionHides from previously cleaned up tools

diff --git a/Engine/Plugins/Experimental/MeshModelingToolsetExp/Source/MeshModelingToolsExp/Private/DrawPolyPathTool.cpp b/Engine/Plugins/Experimental/MeshModelingToolsetExp/Source/MeshModelingToolsExp/Private/DrawPolyPathTool.cpp
--- a/Engine/Plugins/Experimental/MeshModelingToolsetExp/Source/MeshModelingToolsExp/Private/DrawPolyPathTool.cpp
+++ b/Engine/Plugins/Experimental/MeshModelingToolsetExp/Source/MeshModelingToolsExp/Private/DrawPolyPathTool.cpp
@@ -8,7 +8,6 @@
 #include "BaseBehaviors/MouseHoverBehavior.h"
 
 #include "ToolSceneQueriesUtil.h"
-#include "Intersection/IntersectionUtil.h"
 #include "Util/ColorConstants.h"
 #include "ToolSetupUtil.h"
 #include "DynamicMesh/MeshIndexUtil.h"
@@ -220,8 +219,7 @@ void UDrawPolyPathTool::Setup()
 	// begin path draw
 	InitializeNewSurfacePath();
 
-	SetToolDisplayName(LOCTEXT("ToolName", "Extrude PolyPath"));
-	ShowStartupMessage();
+	SetToolDisplayName(LOCTEXT("ToolName", "Path Extrude"));
 }
 
 
@@ -309,24 +307,23 @@ void UDrawPolyPathTool::OnClicked(const FInputDeviceRay& ClickPos)
 {
 	if (SurfacePathMechanic != nullptr)
 	{
-		bool bIsFirstPoint = SurfacePathMechanic->HitPath.Num() == 0;
 		if (SurfacePathMechanic->TryAddPointFromRay(ClickPos.WorldRay))
 		{
 			if (SurfacePathMechanic->IsDone())
 			{
 				bPathIsClosed = SurfacePathMechanic->LoopWasClosed();
-				GetToolManager()->EmitObjectChange(this, MakeUnique<FDrawPolyPathStateChange>(CurrentCurveTimestamp), LOCTEXT("DrawPolyPathBeginOffset", "Set Offset"));
+				GetToolManager()->EmitObjectChange(this, MakeUnique<FDrawPolyPathStateChange>(CurrentCurveTimestamp), LOCTEXT("DrawPolyPathBeginOffset", "Set path width"));
 				OnCompleteSurfacePath();
 			}
 			else
 			{
-				GetToolManager()->EmitObjectChange(this, MakeUnique<FDrawPolyPathStateChange>(CurrentCurveTimestamp), LOCTEXT("DrawPolyPathBeginPath", "Begin Path"));
+				GetToolManager()->EmitObjectChange(this, MakeUnique<FDrawPolyPathStateChange>(CurrentCurveTimestamp), LOCTEXT("DrawPolyPathBeginPath", "Begin path"));
 			}
 		}
 	}
 	else if (CurveDistMechanic != nullptr)
 	{
-		GetToolManager()->EmitObjectChange(this, MakeUnique<FDrawPolyPathStateChange>(CurrentCurveTimestamp), LOCTEXT("DrawPolyPathBeginHeight", "Set Height"));
+		GetToolManager()->EmitObjectChange(this, MakeUnique<FDrawPolyPathStateChange>(CurrentCurveTimestamp), LOCTEXT("DrawPolyPathBeginHeight", "Set extrude height"));
 		OnCompleteOffsetDistance();
 	}
 	else if (ExtrudeHeightMechanic != nullptr)
@@ -360,7 +357,7 @@ bool UDrawPolyPathTool::OnUpdateHover(const FInputDeviceRay& DevicePos)
 	else if (CurveDistMechanic != nullptr)
 	{
 		CurveDistMechanic->UpdateCurrentDistance(DevicePos.WorldRay);
-		TransformProps->Width = CurveDistMechanic->CurrentDistance;
+		TransformProps->Width = CurveDistMechanic->CurrentDistance * 2.0;
 		CurOffsetDistance = CurveDistMechanic->CurrentDistance;
 		UpdatePathPreview();
 	}
@@ -368,7 +365,7 @@ bool UDrawPolyPathTool::OnUpdateHover(const FInputDeviceRay& DevicePos)
 	{
 		ExtrudeHeightMechanic->UpdateCurrentDistance(DevicePos.WorldRay);
 		CurHeight = ExtrudeHeightMechanic->CurrentHeight;
-		TransformProps->Height = ExtrudeHeightMechanic->CurrentHeight;
+		TransformProps->ExtrudeHeight = ExtrudeHeightMechanic->CurrentHeight;
 		UpdateExtrudePreview();
 	}
 	return true;
@@ -448,7 +445,7 @@ void UDrawPolyPathTool::InitializeNewSurfacePath()
 	double SnapTol = ToolSceneQueriesUtil::GetDefaultVisualAngleSnapThreshD();
 	SurfacePathMechanic->SpatialSnapPointsFunc = [this, SnapTol](FVector3d Position1, FVector3d Position2)
 	{
-		return true && ToolSceneQueriesUtil::PointSnapQuery(this->CameraState, Position1, Position2, SnapTol);
+		return ToolSceneQueriesUtil::PointSnapQuery(this->CameraState, Position1, Position2, SnapTol);
 	};
 	SurfacePathMechanic->SetDoubleClickOrCloseLoopMode();
 	UpdateSurfacePathPlane();
@@ -525,7 +522,7 @@ void UDrawPolyPathTool::OnCompleteSurfacePath()
 	}
 
 	SurfacePathMechanic = nullptr;
-	if (TransformProps->WidthMode == EDrawPolyPathWidthMode::Constant)
+	if (TransformProps->WidthMode == EDrawPolyPathWidthMode::Fixed)
 	{
 		BeginConstantOffsetDistance();
 	}
@@ -552,7 +549,7 @@ void UDrawPolyPathTool::BeginInteractiveOffsetDistance()
 void UDrawPolyPathTool::BeginConstantOffsetDistance()
 {
 	InitializePreviewMesh();
-	CurOffsetDistance = TransformProps->Width;
+	CurOffsetDistance = TransformProps->Width * 0.5;
 	UpdatePathPreview();
 	OnCompleteOffsetDistance();
 }
@@ -563,13 +560,13 @@ void UDrawPolyPathTool::OnCompleteOffsetDistance()
 {
 	CurveDistMechanic = nullptr;
 
-	if (TransformProps->OutputType == EDrawPolyPathOutputMode::Ribbon)
+	if (TransformProps->ExtrudeMode == EDrawPolyPathExtrudeMode::Flat)
 	{
 		OnCompleteExtrudeHeight();
 	}
-	else if (TransformProps->HeightMode == EDrawPolyPathHeightMode::Constant)
+	else if (TransformProps->ExtrudeMode == EDrawPolyPathExtrudeMode::Fixed || TransformProps->ExtrudeMode == EDrawPolyPathExtrudeMode::RampFixed)
 	{
-		CurHeight = TransformProps->Height;
+		CurHeight = TransformProps->ExtrudeHeight;
 		OnCompleteExtrudeHeight();
 	}
 	else
@@ -581,12 +578,12 @@ void UDrawPolyPathTool::OnCompleteOffsetDistance()
 
 void UDrawPolyPathTool::OnCompleteExtrudeHeight()
 {
-	CurHeight = TransformProps->Height;
+	CurHeight = TransformProps->ExtrudeHeight;
 	ExtrudeHeightMechanic = nullptr;
 
 	ClearPreview();
 
-	EmitNewObject(TransformProps->OutputType);
+	EmitNewObject(TransformProps->ExtrudeMode);
 
 	InitializeNewSurfacePath();
 	CurrentCurveTimestamp++;
@@ -639,8 +636,6 @@ void UDrawPolyPathTool::GeneratePathMesh(FDynamicMesh3& Mesh)
 			Mesh.Copy(&MeshGen);
 			Mesh.EnableVertexUVs(FVector2f::Zero());		// we will store arc length for each vtx in VertexUV
 
-			FAxisAlignedBox3d PathBounds = Mesh.GetBounds();
-
 			double ShiftX = 0;
 			double DeltaX = CurPathLength / (double)(NumPoints - 1);
 			for (int32 k = 0; k < NumPoints; ++k)
@@ -712,7 +707,7 @@ void UDrawPolyPathTool::BeginInteractiveExtrudeHeight()
 
 void UDrawPolyPathTool::UpdateExtrudePreview()
 {
-	if (TransformProps->OutputType == EDrawPolyPathOutputMode::Ramp)
+	if (TransformProps->ExtrudeMode == EDrawPolyPathExtrudeMode::RampFixed || TransformProps->ExtrudeMode == EDrawPolyPathExtrudeMode::RampInteractive)
 	{
 		EditPreview->UpdateExtrudeType( [&](FDynamicMesh3& Mesh) { GenerateRampMesh(Mesh); }, true);
 	}
@@ -779,7 +774,6 @@ void UDrawPolyPathTool::GenerateRampMesh(FDynamicMesh3& PathMesh)
 	FExtrudeMesh Extruder(&PathMesh);
 	FVector3d ExtrudeDir = DrawPlaneWorld.Z();
 
-	int NumPoints = CurPathPoints.Num();
 	double RampStartRatio = TransformProps->RampStartRatio;
 	double StartHeight = FMathd::Max(0.1, RampStartRatio * FMathd::Abs(CurHeight)) * FMathd::Sign(CurHeight);
 	double EndHeight = CurHeight;
@@ -798,16 +792,16 @@ void UDrawPolyPathTool::GenerateRampMesh(FDynamicMesh3& PathMesh)
 
 
 
-void UDrawPolyPathTool::EmitNewObject(EDrawPolyPathOutputMode OutputMode)
+void UDrawPolyPathTool::EmitNewObject(EDrawPolyPathExtrudeMode ExtrudeMode)
 {
 	FDynamicMesh3 PathMesh;
 	GeneratePathMesh(PathMesh);
 
-	if (OutputMode == EDrawPolyPathOutputMode::Extrusion)
+	if (ExtrudeMode == EDrawPolyPathExtrudeMode::Fixed || ExtrudeMode == EDrawPolyPathExtrudeMode::Interactive)
 	{
 		GenerateExtrudeMesh(PathMesh);
 	}
-	else if (OutputMode == EDrawPolyPathOutputMode::Ramp)
+	else if (ExtrudeMode == EDrawPolyPathExtrudeMode::RampFixed || ExtrudeMode == EDrawPolyPathExtrudeMode::RampInteractive)
 	{
 		GenerateRampMesh(PathMesh);
 	}
@@ -844,21 +838,21 @@ void UDrawPolyPathTool::EmitNewObject(EDrawPolyPathOutputMode OutputMode)
 void UDrawPolyPathTool::ShowStartupMessage()
 {
 	GetToolManager()->DisplayMessage(
-		LOCTEXT("OnStartDraw", "Use this Tool to draw a path on the Drawing Plane, and then Extrude. Left-click to place points, Double-click (or close) to complete path. Ctrl-click on the scene to reposition the Plane (Shift+Ctrl-click to ignore Normal)."),
+		LOCTEXT("OnStartDraw", "Draw a path on the drawing plane, set its width, and extrude it. Left-click to place path vertices. Click on the last or first vertex to complete the path. Hold Shift to ignore snapping."),
 		EToolMessageLevel::UserNotification);
 }
 
 void UDrawPolyPathTool::ShowOffsetMessage()
 {
 	GetToolManager()->DisplayMessage(
-		LOCTEXT("OnStartOffset", "Set the Width of the Extrusion by clicking on the Drawing Plane."),
+		LOCTEXT("OnStartOffset", "Set the width of the path by clicking on the drawing plane."),
 		EToolMessageLevel::UserNotification);
 }
 
 void UDrawPolyPathTool::ShowExtrudeMessage()
 {
 	GetToolManager()->DisplayMessage(
-		LOCTEXT("OnStartExtrude", "Set the Height of the Extrusion by positioning the mouse over the extrusion volume, or over the scene to snap to relative heights."),
+		LOCTEXT("OnStartExtrude", "Set the height of the extrusion by positioning the mouse over the extrusion volume, or over objects to snap to their heights. Hold Shift to ignore snapping."),
 		EToolMessageLevel::UserNotification);
 }
 
diff --git a/Engine/Plugins/Experimental/MeshModelingToolsetExp/Source/MeshModelingToolsExp/Public/DrawPolyPathTool.h b/Engine/Plugins/Experimental/MeshModelingToolsetExp/Source/MeshModelingToolsExp/Public/DrawPolyPathTool.h
--- a/Engine/Plugins/Experimental/MeshModelingToolsetExp/Source/MeshModelingToolsExp/Public/DrawPolyPathTool.h
+++ b/Engine/Plugins/Experimental/MeshModelingToolsetExp/Source/MeshModelingToolsExp/Public/DrawPolyPathTool.h
@@ -31,30 +31,34 @@ public:
 };
 
 
-UENUM()
-enum class EDrawPolyPathOutputMode
-{
-	Ribbon,
-	Extrusion,
-	Ramp
-};
-
-
-
 UENUM()
 enum class EDrawPolyPathWidthMode
 {
-	Interactive,
-	Constant
+	/** Fixed width along the drawn path determined by the Width property */
+	Fixed,
+
+	/** Extrude drawn path to height set via additional mouse input after finishing the path */
+	Interactive
 };
 
 UENUM()
-enum class EDrawPolyPathHeightMode
+enum class EDrawPolyPathExtrudeMode
 {
+	/** Flat path without extrusion */
+	Flat,
+
+	/** Extrude drawn path to a fixed height determined by the Extrude Height property */
+	Fixed,
+
+	/** Extrude drawn path to height set via additional mouse input after finishing the path */
 	Interactive,
-	Constant
-};
 
+	/** Extrude with increasing height along the drawn path. The height at the start and the end of the ramp is determined by the Extrude Height and Ramp Start Ratio properties .*/
+	RampFixed,
+
+	/** Extrude with increasing height along the drawn path. The height is set via additional mouse input after finishing the path. */
+	RampInteractive
+};
 
 
 UCLASS()
@@ -63,30 +67,31 @@ class MESHMODELINGTOOLSEXP_API UDrawPolyPathProperties : public UInteractiveTool
 	GENERATED_BODY()
 
 public:
-	UPROPERTY(EditAnywhere, Category = Shape)
-	EDrawPolyPathOutputMode OutputType = EDrawPolyPathOutputMode::Extrusion;
-
-	UPROPERTY(EditAnywhere, Category = Shape)
+	/** How the drawn path width gets set */
+	UPROPERTY(EditAnywhere, Category = Path)
 	EDrawPolyPathWidthMode WidthMode = EDrawPolyPathWidthMode::Interactive;
 
-	UPROPERTY(EditAnywhere, Category = Shape, meta = (UIMin = "0.0001", UIMax = "1000", ClampMin = "0", ClampMax = "999999"))
+	/** Width of the drawn path when using Fixed width mode; also shows the width in Interactive width mode */
+	UPROPERTY(EditAnywhere, Category = Path, meta = (EditCondition = "WidthMode == EDrawPolyPathWidthMode::Fixed", UIMin = "0.0001", UIMax = "1000", ClampMin = "0", ClampMax = "999999"))
 	float Width = 10.0f;
 
-	UPROPERTY(EditAnywhere, Category = Shape, meta = (
-		EditCondition = "OutputType != EDrawPolyPathOutputMode::Ribbon", EditConditionHides))
-	EDrawPolyPathHeightMode HeightMode = EDrawPolyPathHeightMode::Interactive;
-
-	UPROPERTY(EditAnywhere, Category = Shape, meta = (
-		EditCondition = "OutputType != EDrawPolyPathOutputMode::Ribbon", 
-		EditConditionHides, UIMin = "-1000", UIMax = "1000", ClampMin = "-10000", ClampMax = "10000"))
-	float Height = 10.0f;
-
-	UPROPERTY(EditAnywhere, Category = Shape, meta = (
-		EditCondition = "OutputType == EDrawPolyPathOutputMode::Ramp", EditConditionHides,
-		UIMin = "0.01", UIMax = "1.0", ClampMin = "0", ClampMax = "100.0"))
+	/** If and how the drawn path gets extruded */
+	UPROPERTY(EditAnywhere, Category = Extrude)
+	EDrawPolyPathExtrudeMode ExtrudeMode = EDrawPolyPathExtrudeMode::Interactive;
+
+	/** Extrusion distance when using the Fixed extrude modes; also shows the distance in Interactive extrude modes */
+	UPROPERTY(EditAnywhere, Category = Extrude, meta = (
+		EditCondition = "ExtrudeMode == EDrawPolyPathExtrudeMode::Fixed || ExtrudeMode == EDrawPolyPathExtrudeMode::RampFixed",
+		UIMin = "-1000", UIMax = "1000", ClampMin = "-10000", ClampMax = "10000"))
+	float ExtrudeHeight = 10.0f;
+
+	/** Height of the start of the ramp as a fraction of the Extrude Height property */
+	UPROPERTY(EditAnywhere, Category = Extrude, meta = (
+		EditCondition = "ExtrudeMode == EDrawPolyPathExtrudeMode::RampFixed || ExtrudeMode == EDrawPolyPathExtrudeMode::RampInteractive",
+		UIMin = "0.0", UIMax = "1.0", ClampMin = "0.0", ClampMax = "100.0"))
 	float RampStartRatio = 0.05f;
 
-
+	/** If true, allows snapping to world grid according to editor settings */
 	UPROPERTY(EditAnywhere, Category = Snapping)
 	bool bSnapToWorldGrid = true;
 };
@@ -225,7 +230,7 @@ protected:
 	void GeneratePathMesh(UE::Geometry::FDynamicMesh3& Mesh);
 	void GenerateExtrudeMesh(UE::Geometry::FDynamicMesh3& PathMesh);
 	void GenerateRampMesh(FDynamicMesh3& PathMesh);
-	void EmitNewObject(EDrawPolyPathOutputMode OutputMode);
+	void EmitNewObject(EDrawPolyPathExtrudeMode ExtrudeMode);
 
 	// user feedback messages
 	void ShowStartupMessage();
diff --git a/Engine/Plugins/Runtime/MeshModelingToolset/Source/MeshModelingTools/Private/AddPrimitiveTool.cpp b/Engine/Plugins/Runtime/MeshModelingToolset/Source/MeshModelingTools/Private/AddPrimitiveTool.cpp
--- a/Engine/Plugins/Runtime/MeshModelingToolset/Source/MeshModelingTools/Private/AddPrimitiveTool.cpp
+++ b/Engine/Plugins/Runtime/MeshModelingToolset/Source/MeshModelingTools/Private/AddPrimitiveTool.cpp
@@ -146,7 +146,7 @@ void UAddPrimitiveTool::Setup()
 	UpdatePreviewMesh();
 
 	GetToolManager()->DisplayMessage(
-		LOCTEXT("OnStartAddPrimitiveTool", "This Tool creates new Primitive mesh assets. Position the Primitive by moving the mouse over the scene. Drop a new Asset or Instance by left-clicking (depending on Asset settings)."),
+		LOCTEXT("OnStartAddPrimitiveTool", "This Tool creates new shapes. Configure the shape via its settings, position it by moving the mouse in the scene, and drop it as a new object or instance by left-clicking."),
 		EToolMessageLevel::UserNotification);
 }
 
@@ -285,7 +285,7 @@ void UAddPrimitiveTool::UpdatePreviewMesh() const
 	FDynamicMesh3 NewMesh;
 	GenerateMesh( &NewMesh );
 
-	if (ShapeSettings->PolygroupMode == EMakeMeshPolygroupMode::Single)
+	if (ShapeSettings->PolygroupMode == EMakeMeshPolygroupMode::PerShape)
 	{
 		FaceGroupUtil::SetGroupID(NewMesh, 0);
 	}
@@ -328,7 +328,7 @@ void UAddPrimitiveTool::OnClicked(const FInputDeviceRay& ClickPos)
 
 	if (ShapeSettings->bInstanceIfPossible && LastGenerated != nullptr && IsEquivalentLastGeneratedAsset())
 	{
-		GetToolManager()->BeginUndoTransaction(LOCTEXT("AddPrimitiveToolTransactionName", "Add Primitive Mesh"));
+		GetToolManager()->BeginUndoTransaction(LOCTEXT("AddPrimitiveToolTransactionName", "Add Shape Mesh"));
 		FActorSpawnParameters SpawnParameters;
 		SpawnParameters.Template = LastGenerated->Actor;
 		FRotator Rotation(0.0f, 0.0f, 0.0f);
@@ -350,7 +350,7 @@ void UAddPrimitiveTool::OnClicked(const FInputDeviceRay& ClickPos)
 
 	const FDynamicMesh3* CurMesh = PreviewMesh->GetPreviewDynamicMesh();
 
-	GetToolManager()->BeginUndoTransaction(LOCTEXT("AddPrimitiveToolTransactionName", "Add Primitive Mesh"));
+	GetToolManager()->BeginUndoTransaction(LOCTEXT("AddPrimitiveToolTransactionName", "Add Shape Mesh"));
 
 	FCreateMeshObjectParams NewMeshObjectParams;
 	NewMeshObjectParams.TargetWorld = TargetWorld;
@@ -604,7 +604,7 @@ UAddSpherePrimitiveTool::UAddSpherePrimitiveTool(const FObjectInitializer& Objec
 void UAddSpherePrimitiveTool::GenerateMesh(FDynamicMesh3* OutMesh) const
 {
 	const UProceduralSphereToolProperties* SphereSettings = Cast<UProceduralSphereToolProperties>(ShapeSettings);
-	switch (SphereSettings->SphereType)
+	switch (SphereSettings->SubdivisionType)
 	{
 	case EProceduralSphereType::LatLong:
 	{
diff --git a/Engine/Plugins/Runtime/MeshModelingToolset/Source/MeshModelingTools/Private/DrawPolygonTool.cpp b/Engine/Plugins/Runtime/MeshModelingToolset/Source/MeshModelingTools/Private/DrawPolygonTool.cpp
--- a/Engine/Plugins/Runtime/MeshModelingToolset/Source/MeshModelingTools/Private/DrawPolygonTool.cpp
+++ b/Engine/Plugins/Runtime/MeshModelingToolset/Source/MeshModelingTools/Private/DrawPolygonTool.cpp
@@ -72,7 +72,7 @@ UDrawPolygonTool::UDrawPolygonTool()
 	DrawPlaneOrigin = FVector3d::Zero();
 	DrawPlaneOrientation = FQuaterniond::Identity();
 	bInInteractiveExtrude = false;
-	UInteractiveTool::SetToolDisplayName(LOCTEXT("ToolName", "Extrude Polygon"));
+	UInteractiveTool::SetToolDisplayName(LOCTEXT("ToolName", "Polygon Extrude"));
 }
 
 void UDrawPolygonTool::SetWorld(UWorld* World)
@@ -195,7 +195,7 @@ void UDrawPolygonTool::RegisterActions(FInteractiveToolActionSet& ActionSet)
 	ActionSet.RegisterAction(this, (int32)EStandardToolActions::BaseClientDefinedActionID + 2,
 		TEXT("ToggleGizmo"),
 		LOCTEXT("ToggleGizmo", "Toggle Gizmo"),
-		LOCTEXT("ToggleGizmoTooltip", "Toggle visibility of the transformation Gizmo"),
+		LOCTEXT("ToggleGizmoTooltip", "Toggle visibility of the transformation gizmo"),
 		EModifierKey::None, EKeys::A,
 		[this]() { PolygonProperties->bShowGridGizmo = !PolygonProperties->bShowGridGizmo; });
 }
@@ -711,7 +711,7 @@ bool UDrawPolygonTool::OnNextSequenceClick(const FInputDeviceRay& ClickPos)
 	{
 		//SnapEngine.Reset();
 		bHaveSurfaceHit = false;
-		if (PolygonProperties->ExtrudeMode == EDrawPolygonOutputMode::ExtrudedInteractive)
+		if (PolygonProperties->ExtrudeMode == EDrawPolygonExtrudeMode::Interactive)
 		{
 			BeginInteractiveExtrude();
 
@@ -994,13 +994,13 @@ void UDrawPolygonTool::UpdateShowGizmoState(bool bNewVisibility)
 
 void UDrawPolygonTool::EmitCurrentPolygon()
 {
-	FString BaseName = (PolygonProperties->ExtrudeMode == EDrawPolygonOutputMode::MeshedPolygon) ?
+	FString BaseName = (PolygonProperties->ExtrudeMode == EDrawPolygonExtrudeMode::Flat) ?
 		TEXT("Polygon") : TEXT("Extrude");
 
 	// generate new mesh
 	FFrame3d PlaneFrameOut;
 	FDynamicMesh3 Mesh;
-	const double ExtrudeDist = (PolygonProperties->ExtrudeMode == EDrawPolygonOutputMode::MeshedPolygon) ?
+	const double ExtrudeDist = (PolygonProperties->ExtrudeMode == EDrawPolygonExtrudeMode::Flat) ?
 		0 : PolygonProperties->ExtrudeHeight;
 	bool bSucceeded = GeneratePolygonMesh(PolygonVertices, PolygonHolesVertices, &Mesh, PlaneFrameOut, false, ExtrudeDist, false);
 	if (!bSucceeded) // somehow made a polygon with no valid triangulation; just throw it away ...
@@ -1039,7 +1039,7 @@ void UDrawPolygonTool::UpdateLivePreview()
 
 	FFrame3d PlaneFrame;
 	FDynamicMesh3 Mesh;
-	const double ExtrudeDist = (PolygonProperties->ExtrudeMode == EDrawPolygonOutputMode::MeshedPolygon) ?
+	const double ExtrudeDist = (PolygonProperties->ExtrudeMode == EDrawPolygonExtrudeMode::Flat) ?
 		0 : PolygonProperties->ExtrudeHeight;
 	if (GeneratePolygonMesh(PolygonVertices, PolygonHolesVertices, &Mesh, PlaneFrame, false, ExtrudeDist, false))
 	{
@@ -1248,14 +1248,14 @@ bool UDrawPolygonTool::GeneratePolygonMesh(const TArray<FVector3d>& Polygon, con
 void UDrawPolygonTool::ShowStartupMessage()
 {
 	GetToolManager()->DisplayMessage(
-		LOCTEXT("OnStartDraw", "Use this Tool to draw a polygon on the Drawing Plane, and Extrude it. Left-click to place points. Ctrl-click on the scene to reposition the Plane (Shift+Ctrl-click to ignore Normal). [A] toggles Gizmo. Hold Shift to ignore Snapping."),
+		LOCTEXT("OnStartDraw", "Draw a polygon on the drawing plane, and extrude it. Left-click to place polygon vertices. Hold Shift to ignore snapping."),
 		EToolMessageLevel::UserNotification);
 }
 
 void UDrawPolygonTool::ShowExtrudeMessage()
 {
 	GetToolManager()->DisplayMessage(
-		LOCTEXT("OnStartExtrude", "Set the height of the Extrusion by positioning the mouse over the extrusion volume, or over the scene to snap to relative heights. Hold Shift to ignore Snapping."),
+		LOCTEXT("OnStartExtrude", "Set the height of the extrusion by positioning the mouse over the extrusion volume, or over objects to snap to their heights. Hold Shift to ignore snapping."),
 		EToolMessageLevel::UserNotification);
 }
 
diff --git a/Engine/Plugins/Runtime/MeshModelingToolset/Source/MeshModelingTools/Public/AddPrimitiveTool.h b/Engine/Plugins/Runtime/MeshModelingToolset/Source/MeshModelingTools/Public/AddPrimitiveTool.h
--- a/Engine/Plugins/Runtime/MeshModelingToolset/Source/MeshModelingTools/Public/AddPrimitiveTool.h
+++ b/Engine/Plugins/Runtime/MeshModelingToolset/Source/MeshModelingTools/Public/AddPrimitiveTool.h
@@ -64,7 +64,7 @@ UENUM()
 enum class EMakeMeshPolygroupMode : uint8
 {
 	/** One Polygroup for the entire shape */
-	Single,
+	PerShape,
 	/** One Polygroup for each geometric face */
 	PerFace,
 	/** One Polygroup for each quad/triangle */
@@ -77,13 +77,13 @@ class MESHMODELINGTOOLS_API UProceduralShapeToolProperties : public UInteractive
 	GENERATED_BODY()
 
 public:
-	/** Create instances of the last created asset rather than creating a whole new asset, provided the shape settings have not changed.
+	/** Create instances of the last created asset rather than creating a whole new asset, provided the Shape settings have not changed.
 	  * If false, all created actors will have separate underlying mesh assets. */
-	UPROPERTY(EditAnywhere, Category = AssetSettings)
+	UPROPERTY(EditAnywhere, Category = Asset)
 	bool bInstanceIfPossible = false;
 
 	/** How Polygroups are assigned to shape primitives. */
-	UPROPERTY(EditAnywhere, Category = ShapeSettings, meta = (ProceduralShapeSetting))
+	UPROPERTY(EditAnywhere, Category = Shape, meta = (ProceduralShapeSetting))
 	EMakeMeshPolygroupMode PolygroupMode = EMakeMeshPolygroupMode::PerFace;
 
 	/** How the shape is placed in the scene. */
@@ -116,27 +116,27 @@ class MESHMODELINGTOOLS_API UProceduralBoxToolProperties : public UProceduralSha
 
 public:
 	/** Width of the box */
-	UPROPERTY(EditAnywhere, Category = ShapeSettings, meta = (UIMin = "1.0", UIMax = "1000.0", ClampMin = "0.0001", ClampMax = "1000000.0", ProceduralShapeSetting))
+	UPROPERTY(EditAnywhere, Category = Shape, meta = (UIMin = "1.0", UIMax = "1000.0", ClampMin = "0.0001", ClampMax = "1000000.0", ProceduralShapeSetting))
 	float Width = 100.f;
 
 	/** Depth of the box */
-	UPROPERTY(EditAnywhere, Category = ShapeSettings, meta = (UIMin = "1.0", UIMax = "1000.0", ClampMin = "0.0001", ClampMax = "1000000.0", ProceduralShapeSetting))
+	UPROPERTY(EditAnywhere, Category = Shape, meta = (UIMin = "1.0", UIMax = "1000.0", ClampMin = "0.0001", ClampMax = "1000000.0", ProceduralShapeSetting))
 	float Depth = 100.f;
 
 	/** Height of the box */
-	UPROPERTY(EditAnywhere, Category = ShapeSettings, meta = (UIMin = "1.0", UIMax = "1000.0", ClampMin = "0.0001", ClampMax = "1000000.0", ProceduralShapeSetting))
+	UPROPERTY(EditAnywhere, Category = Shape, meta = (UIMin = "1.0", UIMax = "1000.0", ClampMin = "0.0001", ClampMax = "1000000.0", ProceduralShapeSetting))
 	float Height = 100.f;
 
 	/** Number of subdivisions along the width */
-	UPROPERTY(EditAnywhere, Category = ShapeSettings, meta = (UIMin = "1", UIMax = "100", ClampMin = "1", ClampMax = "500", ProceduralShapeSetting))
+	UPROPERTY(EditAnywhere, Category = Shape, meta = (UIMin = "1", UIMax = "100", ClampMin = "1", ClampMax = "500", ProceduralShapeSetting))
 	int WidthSubdivisions = 1;
 
 	/** Number of subdivisions along the depth */
-	UPROPERTY(EditAnywhere, Category = ShapeSettings, meta = (UIMin = "1", UIMax = "100", ClampMin = "1", ClampMax = "500", ProceduralShapeSetting))
+	UPROPERTY(EditAnywhere, Category = Shape, meta = (UIMin = "1", UIMax = "100", ClampMin = "1", ClampMax = "500", ProceduralShapeSetting))
 	int DepthSubdivisions = 1;
 
 	/** Number of subdivisions along the height */
-	UPROPERTY(EditAnywhere, Category = ShapeSettings, meta = (UIMin = "1", UIMax = "100", ClampMin = "1", ClampMax = "500", ProceduralShapeSetting))
+	UPROPERTY(EditAnywhere, Category = Shape, meta = (UIMin = "1", UIMax = "100", ClampMin = "1", ClampMax = "500", ProceduralShapeSetting))
 	int HeightSubdivisions = 1;
 };
 
@@ -156,35 +156,35 @@ class MESHMODELINGTOOLS_API UProceduralRectangleToolProperties : public UProcedu
 
 public:
 	/** Type of rectangle */
-	UPROPERTY(EditAnywhere, Category = ShapeSettings)
+	UPROPERTY(EditAnywhere, Category = Shape)
 	EProceduralRectType RectangleType = EProceduralRectType::Rectangle;
 
 	/** Width of the rectangle */
-	UPROPERTY(EditAnywhere, Category = ShapeSettings, meta = (UIMin = "1.0", UIMax = "1000.0", ClampMin = "0.0001", ClampMax = "1000000.0", ProceduralShapeSetting))
+	UPROPERTY(EditAnywhere, Category = Shape, meta = (UIMin = "1.0", UIMax = "1000.0", ClampMin = "0.0001", ClampMax = "1000000.0", ProceduralShapeSetting))
 	float Width = 100.f;
 
 	/** Depth of the rectangle */
-	UPROPERTY(EditAnywhere, Category = ShapeSettings, meta = (UIMin = "1.0", UIMax = "1000.0", ClampMin = "0.0001", ClampMax = "1000000.0", ProceduralShapeSetting))
+	UPROPERTY(EditAnywhere, Category = Shape, meta = (UIMin = "1.0", UIMax = "1000.0", ClampMin = "0.0001", ClampMax = "1000000.0", ProceduralShapeSetting))
 	float Depth = 100.f;
 
 	/** Number of subdivisions along the width */
-	UPROPERTY(EditAnywhere, Category = ShapeSettings, meta = (UIMin = "1", UIMax = "100", ClampMin = "1", ClampMax = "500", ProceduralShapeSetting))
+	UPROPERTY(EditAnywhere, Category = Shape, meta = (UIMin = "1", UIMax = "100", ClampMin = "1", ClampMax = "500", ProceduralShapeSetting))
 	int WidthSubdivisions = 1;
 
 	/** Number of subdivisions along the depth */
-	UPROPERTY(EditAnywhere, Category = ShapeSettings, meta = (UIMin = "1", UIMax = "100", ClampMin = "1", ClampMax = "500", ProceduralShapeSetting))
+	UPROPERTY(EditAnywhere, Category = Shape, meta = (UIMin = "1", UIMax = "100", ClampMin = "1", ClampMax = "500", ProceduralShapeSetting))
 	int DepthSubdivisions = 1;
 
-	/** Radius of rounded corners */
-	UPROPERTY(EditAnywhere, Category = ShapeSettings,
+	/** Radius of rounded corners. This is only available for Rounded Rectangles. */
+	UPROPERTY(EditAnywhere, Category = Shape,
 		meta = (UIMin = "1.0", UIMax = "1000.0", ClampMin = "0.0001", ClampMax = "1000000.0", ProceduralShapeSetting,
-			EditCondition = "RectangleType == EProceduralRectType::RoundedRectangle", EditConditionHides))
+			EditCondition = "RectangleType == EProceduralRectType::RoundedRectangle"))
 	float CornerRadius = 25.f;
 
-	/** Number of radial slices for each rounded corner */
-	UPROPERTY(EditAnywhere, Category = ShapeSettings,
+	/** Number of radial slices for each rounded corner. This is only available for Rounded Rectangles. */
+	UPROPERTY(EditAnywhere, Category = Shape,
 		meta = (UIMin = "3", UIMax = "128", ClampMin = "3", ClampMax = "500", ProceduralShapeSetting,
-			EditCondition = "RectangleType == EProceduralRectType::RoundedRectangle", EditConditionHides))
+			EditCondition = "RectangleType == EProceduralRectType::RoundedRectangle"))
 	int CornerSlices = 16;
 };
 
@@ -204,25 +204,25 @@ class MESHMODELINGTOOLS_API UProceduralDiscToolProperties : public UProceduralSh
 
 public:
 	/** Type of disc */
-	UPROPERTY(EditAnywhere, Category = ShapeSettings)
+	UPROPERTY(EditAnywhere, Category = Shape)
 	EProceduralDiscType DiscType = EProceduralDiscType::Disc;
 
 	/** Radius of the disc */
-	UPROPERTY(EditAnywhere, Category = ShapeSettings, meta = (UIMin = "1.0", UIMax = "1000.0", ClampMin = "0.0001", ClampMax = "1000000.0", ProceduralShapeSetting))
+	UPROPERTY(EditAnywhere, Category = Shape, meta = (UIMin = "1.0", UIMax = "1000.0", ClampMin = "0.0001", ClampMax = "1000000.0", ProceduralShapeSetting))
 	float Radius = 50.f;
 
 	/** Number of radial slices for the disc */
-	UPROPERTY(EditAnywhere, Category = ShapeSettings, meta = (UIMin = "3", UIMax = "128", ClampMin = "3", ClampMax = "500", ProceduralShapeSetting))
+	UPROPERTY(EditAnywhere, Category = Shape, meta = (UIMin = "3", UIMax = "128", ClampMin = "3", ClampMax = "500", ProceduralShapeSetting))
 	int RadialSlices = 16;
 
 	/** Number of radial subdivisions for each radial slice of the disc */
-	UPROPERTY(EditAnywhere, Category = ShapeSettings, meta = (UIMin = "1", UIMax = "100", ClampMin = "1", ClampMax = "500", ProceduralShapeSetting))
+	UPROPERTY(EditAnywhere, Category = Shape, meta = (UIMin = "1", UIMax = "100", ClampMin = "1", ClampMax = "500", ProceduralShapeSetting))
 	int RadialSubdivisions = 1;
 
-	/** Radius of the hole in the disc */
-	UPROPERTY(EditAnywhere, Category = ShapeSettings,
+	/** Radius of the hole in the disc. This is only available for Punctured Discs. */
+	UPROPERTY(EditAnywhere, Category = Shape,
 		meta = (UIMin = "1.0", UIMax = "1000.0", ClampMin = "0.0001", ClampMax = "1000000.0", ProceduralShapeSetting,
-			EditCondition = "DiscType == EProceduralDiscType::PuncturedDisc", EditConditionHides))
+			EditCondition = "DiscType == EProceduralDiscType::PuncturedDisc"))
 	float HoleRadius = 25.f;
 };
 
@@ -233,19 +233,19 @@ class MESHMODELINGTOOLS_API UProceduralTorusToolProperties : public UProceduralS
 
 public:
 	/** Major radius of the torus, measured from the torus center to the center of the torus tube */
-	UPROPERTY(EditAnywhere, Category = ShapeSettings, meta = (UIMin = "1.0", UIMax = "1000.0", ClampMin = "0.0001", ClampMax = "1000000.0", ProceduralShapeSetting))
+	UPROPERTY(EditAnywhere, Category = Shape, meta = (UIMin = "1.0", UIMax = "1000.0", ClampMin = "0.0001", ClampMax = "1000000.0", ProceduralShapeSetting))
 	float MajorRadius = 50.f;
 
 	/** Minor radius of the torus, measured from the center ot the torus tube to the tube surface */
-	UPROPERTY(EditAnywhere, Category = ShapeSettings, meta = (UIMin = "1.0", UIMax = "1000.0", ClampMin = "0.0001", ClampMax = "1000000.0", ProceduralShapeSetting))
+	UPROPERTY(EditAnywhere, Category = Shape, meta = (UIMin = "1.0", UIMax = "1000.0", ClampMin = "0.0001", ClampMax = "1000000.0", ProceduralShapeSetting))
 	float MinorRadius = 25.f;
 
 	/** Number of radial slices along the torus tube */
-	UPROPERTY(EditAnywhere, Category = ShapeSettings, meta = (UIMin = "3", UIMax = "128", ClampMin = "3", ClampMax = "500", ProceduralShapeSetting))
+	UPROPERTY(EditAnywhere, Category = Shape, meta = (UIMin = "3", UIMax = "128", ClampMin = "3", ClampMax = "500", ProceduralShapeSetting))
 	int MajorSlices = 16;
 
 	/** Number of radial slices around the torus tube */
-	UPROPERTY(EditAnywhere, Category = ShapeSettings, meta = (UIMin = "3", UIMax = "128", ClampMin = "3", ClampMax = "500", ProceduralShapeSetting))
+	UPROPERTY(EditAnywhere, Category = Shape, meta = (UIMin = "3", UIMax = "128", ClampMin = "3", ClampMax = "500", ProceduralShapeSetting))
 	int MinorSlices = 16;
 };
 
@@ -257,19 +257,19 @@ class MESHMODELINGTOOLS_API UProceduralCylinderToolProperties : public UProcedur
 public:
 
 	/** Radius of the cylinder */
-	UPROPERTY(EditAnywhere, Category = ShapeSettings, meta = (UIMin = "1.0", UIMax = "1000.0", ClampMin = "0.0001", ClampMax = "1000000.0", ProceduralShapeSetting))
+	UPROPERTY(EditAnywhere, Category = Shape, meta = (UIMin = "1.0", UIMax = "1000.0", ClampMin = "0.0001", ClampMax = "1000000.0", ProceduralShapeSetting))
 	float Radius = 50.f;
 
 	/** Height of the cylinder */
-	UPROPERTY(EditAnywhere, Category = ShapeSettings, meta = (UIMin = "1.0", UIMax = "1000.0", ClampMin = "0.0001", ClampMax = "1000000.0", ProceduralShapeSetting))
+	UPROPERTY(EditAnywhere, Category = Shape, meta = (UIMin = "1.0", UIMax = "1000.0", ClampMin = "0.0001", ClampMax = "1000000.0", ProceduralShapeSetting))
 	float Height = 200.f;
 
 	/** Number of radial slices for the cylinder */
-	UPROPERTY(EditAnywhere, Category = ShapeSettings, meta = (UIMin = "3", UIMax = "128", ClampMin = "3", ClampMax = "500", ProceduralShapeSetting))
+	UPROPERTY(EditAnywhere, Category = Shape, meta = (UIMin = "3", UIMax = "128", ClampMin = "3", ClampMax = "500", ProceduralShapeSetting))
 	int RadialSlices = 16;
 
 	/** Number of subdivisions along the height of the cylinder */
-	UPROPERTY(EditAnywhere, Category = ShapeSettings, meta = (UIMin = "1", UIMax = "100", ClampMin = "1", ClampMax = "500", ProceduralShapeSetting))
+	UPROPERTY(EditAnywhere, Category = Shape, meta = (UIMin = "1", UIMax = "100", ClampMin = "1", ClampMax = "500", ProceduralShapeSetting))
 	int HeightSubdivisions = 1;
 };
 
@@ -280,19 +280,19 @@ class MESHMODELINGTOOLS_API UProceduralConeToolProperties : public UProceduralSh
 
 public:
 	/** Radius of the cone */
-	UPROPERTY(EditAnywhere, Category = ShapeSettings, meta = (UIMin = "1.0", UIMax = "1000.0", ClampMin = "0.0001", ClampMax = "1000000.0", ProceduralShapeSetting))
+	UPROPERTY(EditAnywhere, Category = Shape, meta = (UIMin = "1.0", UIMax = "1000.0", ClampMin = "0.0001", ClampMax = "1000000.0", ProceduralShapeSetting))
 	float Radius = 50.f;
 
 	/** Height of the cone */
-	UPROPERTY(EditAnywhere, Category = ShapeSettings, meta = (UIMin = "1.0", UIMax = "1000.0", ClampMin = "0.0001", ClampMax = "1000000.0", ProceduralShapeSetting))
+	UPROPERTY(EditAnywhere, Category = Shape, meta = (UIMin = "1.0", UIMax = "1000.0", ClampMin = "0.0001", ClampMax = "1000000.0", ProceduralShapeSetting))
 	float Height = 200.f;
 
 	/** Number of radial slices for the cylinder */
-	UPROPERTY(EditAnywhere, Category = ShapeSettings, meta = (UIMin = "3", UIMax = "128", ClampMin = "3", ClampMax = "500", ProceduralShapeSetting))
+	UPROPERTY(EditAnywhere, Category = Shape, meta = (UIMin = "3", UIMax = "128", ClampMin = "3", ClampMax = "500", ProceduralShapeSetting))
 	int RadialSlices = 16;
 
 	/** Number of subdivisions along the height of the cone */
-	UPROPERTY(EditAnywhere, Category = ShapeSettings, meta = (UIMin = "1", UIMax = "100", ClampMin = "1", ClampMax = "500", ProceduralShapeSetting))
+	UPROPERTY(EditAnywhere, Category = Shape, meta = (UIMin = "1", UIMax = "100", ClampMin = "1", ClampMax = "500", ProceduralShapeSetting))
 	int HeightSubdivisions = 1;
 };
 
@@ -303,27 +303,27 @@ class MESHMODELINGTOOLS_API UProceduralArrowToolProperties : public UProceduralS
 
 public:
 	/** Radius of the arrow shaft */
-	UPROPERTY(EditAnywhere, Category = ShapeSettings, meta = (UIMin = "1.0", UIMax = "1000.0", ClampMin = "0.0001", ClampMax = "1000000.0", ProceduralShapeSetting))
+	UPROPERTY(EditAnywhere, Category = Shape, meta = (UIMin = "1.0", UIMax = "1000.0", ClampMin = "0.0001", ClampMax = "1000000.0", ProceduralShapeSetting))
 	float ShaftRadius = 20.f;
 
 	/** Height of arrow shaft */
-	UPROPERTY(EditAnywhere, Category = ShapeSettings, meta = (UIMin = "1.0", UIMax = "1000.0", ClampMin = "0.0001", ClampMax = "1000000.0", ProceduralShapeSetting))
+	UPROPERTY(EditAnywhere, Category = Shape, meta = (UIMin = "1.0", UIMax = "1000.0", ClampMin = "0.0001", ClampMax = "1000000.0", ProceduralShapeSetting))
 	float ShaftHeight = 200.f;
 
 	/** Radius of the arrow head base */
-	UPROPERTY(EditAnywhere, Category = ShapeSettings, meta = (UIMin = "1.0", UIMax = "1000.0", ClampMin = "0.0001", ClampMax = "1000000.0", ProceduralShapeSetting))
+	UPROPERTY(EditAnywhere, Category = Shape, meta = (UIMin = "1.0", UIMax = "1000.0", ClampMin = "0.0001", ClampMax = "1000000.0", ProceduralShapeSetting))
 	float HeadRadius = 60.f;
 
 	/** Height of arrow head */
-	UPROPERTY(EditAnywhere, Category = ShapeSettings, meta = (UIMin = "1.0", UIMax = "1000.0", ClampMin = "0.0001", ClampMax = "1000000.0", ProceduralShapeSetting))
+	UPROPERTY(EditAnywhere, Category = Shape, meta = (UIMin = "1.0", UIMax = "1000.0", ClampMin = "0.0001", ClampMax = "1000000.0", ProceduralShapeSetting))
 	float HeadHeight = 120.f;
 
 	/** Number of radial slices for the arrow */
-	UPROPERTY(EditAnywhere, Category = ShapeSettings, meta = (UIMin = "3", UIMax = "100", ClampMin = "3", ClampMax = "500", ProceduralShapeSetting))
+	UPROPERTY(EditAnywhere, Category = Shape, meta = (UIMin = "3", UIMax = "100", ClampMin = "3", ClampMax = "500", ProceduralShapeSetting))
 	int RadialSlices = 16;
 
 	/** Number of subdivisions along each part of the arrow, i.e. shaft, head base, head cone */
-	UPROPERTY(EditAnywhere, Category = ShapeSettings, meta = (UIMin = "1", UIMax = "100", ClampMin = "1", ClampMax = "500", ProceduralShapeSetting))
+	UPROPERTY(EditAnywhere, Category = Shape, meta = (UIMin = "1", UIMax = "100", ClampMin = "1", ClampMax = "500", ProceduralShapeSetting))
 	int HeightSubdivisions = 1;
 };
 
@@ -331,9 +331,9 @@ UENUM()
 enum class EProceduralSphereType
 {
 	/** Create a Sphere with Lat Long parameterization */
-	LatLong UMETA(DisplayName = "Lat Long"),
+	LatLong,
 	/** Create a Sphere with Box parameterization */
-	Box UMETA(DisplayName = "Box")
+	Box
 };
 
 UCLASS()
@@ -342,28 +342,28 @@ class MESHMODELINGTOOLS_API UProceduralSphereToolProperties : public UProcedural
 	GENERATED_BODY()
 
 public:
-	/** Type of sphere */
-	UPROPERTY(EditAnywhere, Category = ShapeSettings)
-	EProceduralSphereType SphereType = EProceduralSphereType::Box;
-
 	/** Radius of the sphere */
-	UPROPERTY(EditAnywhere, Category = ShapeSettings, meta = (UIMin = "1.0", UIMax = "1000.0", ClampMin = "0.0001", ClampMax = "1000000.0", ProceduralShapeSetting))
+	UPROPERTY(EditAnywhere, Category = Shape, meta = (UIMin = "1.0", UIMax = "1000.0", ClampMin = "0.0001", ClampMax = "1000000.0", ProceduralShapeSetting))
 	float Radius = 50.f;
 
-	/** Number of horizontal slices of the sphere */
-	UPROPERTY(EditAnywhere, Category = ShapeSettings, meta = (UIMin = "3", UIMax = "100", ClampMin = "4", ClampMax = "500", ProceduralShapeSetting,
-		EditCondition = "SphereType == EProceduralSphereType::LatLong", EditConditionHides))
+	/** Type of subdivision for the sphere */
+	UPROPERTY(EditAnywhere, Category = Shape)
+	EProceduralSphereType SubdivisionType = EProceduralSphereType::Box;
+
+	/** Number of subdivisions for each side of the sphere. This is only available for spheres with Box subdivision. */
+	UPROPERTY(EditAnywhere, Category = Shape, meta = (UIMin = "1", UIMax = "100", ClampMin = "1", ClampMax = "500", ProceduralShapeSetting,
+		EditCondition = "SubdivisionType == EProceduralSphereType::Box"))
+	int Subdivisions = 16;
+
+	/** Number of horizontal slices of the sphere. This is only available for spheres with Lat Long subdivision. */
+	UPROPERTY(EditAnywhere, Category = Shape, meta = (UIMin = "3", UIMax = "100", ClampMin = "4", ClampMax = "500", ProceduralShapeSetting,
+		EditCondition = "SubdivisionType == EProceduralSphereType::LatLong"))
 	int HorizontalSlices = 16;
 
-	/** Number of vertical slices of the sphere */
-	UPROPERTY(EditAnywhere, Category = ShapeSettings, meta = (UIMin = "3", UIMax = "100", ClampMin = "4", ClampMax = "500", ProceduralShapeSetting,
-		EditCondition = "SphereType == EProceduralSphereType::LatLong", EditConditionHides))
+	/** Number of vertical slices of the sphere. This is only available for spheres with Lat Long subdivision. */
+	UPROPERTY(EditAnywhere, Category = Shape, meta = (UIMin = "3", UIMax = "100", ClampMin = "4", ClampMax = "500", ProceduralShapeSetting,
+		EditCondition = "SubdivisionType == EProceduralSphereType::LatLong"))
 	int VerticalSlices = 16;
-
-	/** Number of subdivisions for each side of the sphere */
-	UPROPERTY(EditAnywhere, Category = ShapeSettings, meta = (UIMin = "1", UIMax = "100", ClampMin = "1", ClampMax = "500", ProceduralShapeSetting,
-		EditCondition = "SphereType == EProceduralSphereType::Box", EditConditionHides))
-	int Subdivisions = 16;
 };
 
 UENUM()
@@ -386,43 +386,43 @@ class MESHMODELINGTOOLS_API UProceduralStairsToolProperties : public UProcedural
 
 public:
 	/** Type of staircase */
-	UPROPERTY(EditAnywhere, Category = ShapeSettings)
+	UPROPERTY(EditAnywhere, Category = Shape)
 	EProceduralStairsType StairsType = EProceduralStairsType::Linear;
 
 	/** Number of steps */
-	UPROPERTY(EditAnywhere, Category = ShapeSettings, meta = (DisplayName = "Number of Steps", UIMin = "2", UIMax = "100", ClampMin = "2", ClampMax = "1000000", ProceduralShapeSetting))
+	UPROPERTY(EditAnywhere, Category = Shape, meta = (DisplayName = "Number of Steps", UIMin = "2", UIMax = "100", ClampMin = "2", ClampMax = "1000000", ProceduralShapeSetting))
 	int NumSteps = 8;
 
 	/** Width of each step */
-	UPROPERTY(EditAnywhere, Category = ShapeSettings, meta = (UIMin = "1.0", UIMax = "1000.0", ClampMin = "0.0001", ClampMax = "1000000.0", ProceduralShapeSetting))
+	UPROPERTY(EditAnywhere, Category = Shape, meta = (UIMin = "1.0", UIMax = "1000.0", ClampMin = "0.0001", ClampMax = "1000000.0", ProceduralShapeSetting))
 	float StepWidth = 150.0f;
 
 	/** Height of each step */
-	UPROPERTY(EditAnywhere, Category = ShapeSettings, meta = (UIMin = "1.0", UIMax = "1000.0", ClampMin = "0.0001", ClampMax = "1000000.0", ProceduralShapeSetting))
+	UPROPERTY(EditAnywhere, Category = Shape, meta = (UIMin = "1.0", UIMax = "1000.0", ClampMin = "0.0001", ClampMax = "1000000.0", ProceduralShapeSetting))
 	float StepHeight = 20.0f;
 
-	/** Depth of each step for linear stairs */
-	UPROPERTY(EditAnywhere, Category = ShapeSettings,
+	/** Depth of each step of linear stairs */
+	UPROPERTY(EditAnywhere, Category = Shape,
 		meta = (UIMin = "1.0", UIMax = "1000.0", ClampMin = "0.0001", ClampMax = "1000000.0", ProceduralShapeSetting,
-			EditCondition =	"StairsType == EProceduralStairsType::Linear || StairsType == EProceduralStairsType::Floating", EditConditionHides))
+			EditCondition =	"StairsType == EProceduralStairsType::Linear || StairsType == EProceduralStairsType::Floating"))
 	float StepDepth = 30.0f;
 
-	/** Angular length of curved stairs; positive values are for clockwise and negative values are for counterclockwise. */
-	UPROPERTY(EditAnywhere, Category = ShapeSettings,
+	/** Angular length of curved stairs. Positive values are for clockwise and negative values are for counterclockwise. */
+	UPROPERTY(EditAnywhere, Category = Shape,
 		meta = (UIMin = "-360.0", UIMax = "360.0", ClampMin = "-360.0", ClampMax = "360.0", ProceduralShapeSetting,
-			EditCondition =	"StairsType == EProceduralStairsType::Curved", EditConditionHides))
+			EditCondition =	"StairsType == EProceduralStairsType::Curved"))
 	float CurveAngle = 90.0f;
 
-	/** Angular length of spiral stairs; positive values are for clockwise and negative values are for counterclockwise. */
-	UPROPERTY(EditAnywhere, Category = ShapeSettings,
+	/** Angular length of spiral stairs. Positive values are for clockwise and negative values are for counterclockwise. */
+	UPROPERTY(EditAnywhere, Category = Shape,
 		meta = (UIMin = "-720.0", UIMax = "720.0", ClampMin = "-360000.0", ClampMax = "360000.0", ProceduralShapeSetting,
-			EditCondition =	"StairsType == EProceduralStairsType::Spiral", EditConditionHides))
+			EditCondition =	"StairsType == EProceduralStairsType::Spiral"))
 	float SpiralAngle = 90.0f;
 
-	/** Inner radius for curved and spiral stairs */
-	UPROPERTY(EditAnywhere, Category = ShapeSettings,
+	/** Inner radius of curved and spiral stairs */
+	UPROPERTY(EditAnywhere, Category = Shape,
 		meta = (UIMin = "1.0", UIMax = "1000.0", ClampMin = "0.0001", ClampMax = "1000000.0", ProceduralShapeSetting,
-			EditCondition = "StairsType == EProceduralStairsType::Curved || StairsType == EProceduralStairsType::Spiral", EditConditionHides))
+			EditCondition = "StairsType == EProceduralStairsType::Curved || StairsType == EProceduralStairsType::Spiral"))
 	float InnerRadius = 150.0f;
 };
 
diff --git a/Engine/Plugins/Runtime/MeshModelingToolset/Source/MeshModelingTools/Public/DrawPolygonTool.h b/Engine/Plugins/Runtime/MeshModelingToolset/Source/MeshModelingTools/Public/DrawPolygonTool.h
--- a/Engine/Plugins/Runtime/MeshModelingToolset/Source/MeshModelingTools/Public/DrawPolygonTool.h
+++ b/Engine/Plugins/Runtime/MeshModelingToolset/Source/MeshModelingTools/Public/DrawPolygonTool.h
@@ -60,18 +60,18 @@ enum class EDrawPolygonDrawMode : uint8
 };
 
 
-/** Output of draw polygon tool */
+/** How the drawn polygon gets extruded */
 UENUM()
-enum class EDrawPolygonOutputMode : uint8
+enum class EDrawPolygonExtrudeMode : uint8
 {
-	/** Generate a flat polygon mesh */
-	MeshedPolygon UMETA(DisplayName = "Flat Mesh"),
+	/** Flat polygon without extrusion */
+	Flat,
 
-	/** Extrude drawn polygon to constant height determined by the Extrude Height property */
-	ExtrudedConstant UMETA(DisplayName = "Extrude to Height"),
+	/** Extrude drawn polygon to fixed height determined by the Extrude Height property */
+	Fixed,
 
 	/** Extrude drawn polygon to height set via additional mouse input after closing the polygon */
-	ExtrudedInteractive UMETA(DisplayName = "Interactive Extrude"),
+	Interactive,
 };
 
 
@@ -90,38 +90,37 @@ public:
 	UPROPERTY(EditAnywhere, NonTransactional, Category = Polygon, meta = (DisplayName = "Draw Mode"))
 	EDrawPolygonDrawMode PolygonDrawMode = EDrawPolygonDrawMode::Freehand;
 
-	/** Size of secondary features, e.g. the rounded corners of a rounded rectangle, as fraction of the overall shape size */
-	UPROPERTY(EditAnywhere, NonTransactional, Category = Polygon, meta = (UIMin = "0.01", UIMax = "0.99", ClampMin = "0.01", ClampMax = "0.99",
-		EditCondition = "PolygonDrawMode == EDrawPolygonDrawMode::RoundedRectangle || PolygonDrawMode == EDrawPolygonDrawMode::Ring", EditConditionHides))
-	float FeatureSizeRatio = .25;
-
 	/** Allow freehand drawn polygons to self-intersect */
 	UPROPERTY(EditAnywhere, NonTransactional, Category = Polygon,
-		meta = (DisplayName ="Self-Intersections", EditCondition = "PolygonDrawMode == EDrawPolygonDrawMode::Freehand", EditConditionHides))
+		meta = (DisplayName ="Self-Intersections", EditCondition = "PolygonDrawMode == EDrawPolygonDrawMode::Freehand"))
 	bool bAllowSelfIntersections = false;
 
+	/** Size of secondary features, e.g. the rounded corners of a rounded rectangle, as fraction of the overall shape size */
+	UPROPERTY(EditAnywhere, NonTransactional, Category = Polygon, meta = (UIMin = "0.01", UIMax = "0.99", ClampMin = "0.01", ClampMax = "0.99",
+		EditCondition = "PolygonDrawMode == EDrawPolygonDrawMode::RoundedRectangle || PolygonDrawMode == EDrawPolygonDrawMode::Ring"))
+	float FeatureSizeRatio = .25;
+
 	/** Number of radial subdivisions in round features, e.g. circles or rounded corners */
 	UPROPERTY(EditAnywhere, NonTransactional, Category = Polygon, meta = (UIMin = "3", UIMax = "100", ClampMin = "3", ClampMax = "10000",
-		EditCondition =	"PolygonDrawMode == EDrawPolygonDrawMode::Circle || PolygonDrawMode == EDrawPolygonDrawMode::RoundedRectangle || PolygonDrawMode == EDrawPolygonDrawMode::Ring",
-		EditConditionHides))
+		EditCondition =	"PolygonDrawMode == EDrawPolygonDrawMode::Circle || PolygonDrawMode == EDrawPolygonDrawMode::RoundedRectangle || PolygonDrawMode == EDrawPolygonDrawMode::Ring"))
 	int RadialSlices = 16;
 
 	/** Distance between the last clicked point and the current point  */
 	UPROPERTY(VisibleAnywhere, NonTransactional, Category = Polygon, meta = (TransientToolProperty))
 	float Distance = 0.0f;
 
-	/** If and how the drawn polygon gets extruded */
-	UPROPERTY(EditAnywhere, NonTransactional, Category = Polygon)
-	EDrawPolygonOutputMode ExtrudeMode = EDrawPolygonOutputMode::ExtrudedInteractive;
-
-	/** Extrusion distance when using the non-interactive extrude mode */
-	UPROPERTY(EditAnywhere, NonTransactional, Category = Polygon, meta = (UIMin = "-1000", UIMax = "1000", ClampMin = "-10000", ClampMax = "10000",
-									EditCondition = "ExtrudeMode == EDrawPolygonOutputMode::ExtrudedConstant"))
-	float ExtrudeHeight = 100.0f;
-
 	/** If true, shows a gizmo to manipulate the additional grid used to draw the polygon on */
 	UPROPERTY(EditAnywhere, NonTransactional, Category = Polygon)
 	bool bShowGridGizmo = true;
+
+	/** If and how the drawn polygon gets extruded */
+	UPROPERTY(EditAnywhere, NonTransactional, Category = Extrude)
+	EDrawPolygonExtrudeMode ExtrudeMode = EDrawPolygonExtrudeMode::Interactive;
+
+	/** Extrude distance when using the Fixed extrude mode */
+	UPROPERTY(EditAnywhere, NonTransactional, Category = Extrude, meta = (UIMin = "-1000", UIMax = "1000", ClampMin = "-10000", ClampMax = "10000",
+		EditCondition = "ExtrudeMode == EDrawPolygonExtrudeMode::Fixed"))
+	float ExtrudeHeight = 100.0f;
 };
 
 UCLASS()
@@ -130,36 +129,36 @@ class MESHMODELINGTOOLS_API UDrawPolygonToolSnapProperties : public UInteractive
 	GENERATED_BODY()
 
 public:
-	/** If true, enables additional snapping controls */
+	/** If true, enables additional snapping controls. If false, all snapping is disabled. */
 	UPROPERTY(EditAnywhere, NonTransactional, Category = Snapping)
 	bool bEnableSnapping = true;
 
-	/** If true, allows snapping to world grid according to editor settings */
-	UPROPERTY(EditAnywhere, NonTransactional, Category = Snapping, meta = (EditCondition = "bEnableSnapping", EditConditionHides))
+	/** If true, allows snapping to world grid according to editor settings. Enable Snapping needs to be true for this property to be available. */
+	UPROPERTY(EditAnywhere, NonTransactional, Category = Snapping, meta = (EditCondition = "bEnableSnapping"))
 	bool bSnapToWorldGrid = false;
 
-	/** If true, snap to vertices in the current polygon or in other meshes */
-	UPROPERTY(EditAnywhere, NonTransactional, Category = Snapping, meta = (EditCondition = "bEnableSnapping", EditConditionHides))
+	/** If true, snap to vertices in the current polygon or in other meshes. Enable Snapping needs to be true for this property to be available. */
+	UPROPERTY(EditAnywhere, NonTransactional, Category = Snapping, meta = (EditCondition = "bEnableSnapping"))
 	bool bSnapToVertices = true;
 
-	/** If true, snap to edges in the current polygon or in other meshes */
-	UPROPERTY(EditAnywhere, NonTransactional, Category = Snapping, meta = (EditCondition = "bEnableSnapping", EditConditionHides))
+	/** If true, snap to edges in the current polygon or in other meshes. Enable Snapping needs to be true for this property to be available. */
+	UPROPERTY(EditAnywhere, NonTransactional, Category = Snapping, meta = (EditCondition = "bEnableSnapping"))
 	bool bSnapToEdges = false;
 
-	/** If true, snap to axes of the additional drawing grid as well as the axes relative to the last segment */
-	UPROPERTY(EditAnywhere, NonTransactional, Category = Snapping, meta = (EditCondition = "bEnableSnapping", EditConditionHides))
+	/** If true, snap to axes of the additional drawing grid as well as the axes relative to the last segment. Enable Snapping needs to be true for this property to be available. */
+	UPROPERTY(EditAnywhere, NonTransactional, Category = Snapping, meta = (EditCondition = "bEnableSnapping"))
 	bool bSnapToAxes = true;
 
-	/** If true, when snapping to axes, also try to snap to the length of an existing segment in the polygon */
-	UPROPERTY(EditAnywhere, NonTransactional, Category = Snapping, meta = (EditCondition = "bEnableSnapping && bSnapToAxes", EditConditionHides))
+	/** If true, when snapping to axes, also try to snap to the length of an existing segment in the polygon. Enable Snapping and Snap to Axes needs to both be true for this property to be available. */
+	UPROPERTY(EditAnywhere, NonTransactional, Category = Snapping, meta = (EditCondition = "bEnableSnapping && bSnapToAxes"))
 	bool bSnapToLengths = true;
 
-	/** If true, snap to surfaces of existing objects */
-	UPROPERTY(EditAnywhere, NonTransactional, Category = Snapping)
+	/** If true, snap to surfaces of existing objects. Enable Snapping needs to be true for this property to be available. */
+	UPROPERTY(EditAnywhere, NonTransactional, Category = Snapping, meta = (EditCondition = "bEnableSnapping"))
 	bool bSnapToSurfaces = false;
 
-	/** Offset for snap point on the surface of an existing object in the direction of the surface normal */
-	UPROPERTY(EditAnywhere, NonTransactional, Category = Snapping, meta = (DisplayName = "Surface Offset", EditCondition = "bSnapToSurfaces", EditConditionHides))
+	/** Offset for snap point on the surface of an existing object in the direction of the surface normal. Enable Snapping and Snap to Surfaces needs to be true for this property to be available. */
+	UPROPERTY(EditAnywhere, NonTransactional, Category = Snapping, meta = (DisplayName = "Surface Offset", EditCondition = "bSnapToSurfaces"))
 	float SnapToSurfacesOffset = 0.0f;
 };
 
