clamp offsets/extrusions in polyedit in the tri-normal-adjust mode, that could otherwise offset infinitely at sharp corners

diff --git a/Engine/Plugins/Experimental/MeshModelingToolsetExp/Source/ModelingOperators/Private/DeformationOps/ExtrudeOp.cpp b/Engine/Plugins/Experimental/MeshModelingToolsetExp/Source/ModelingOperators/Private/DeformationOps/ExtrudeOp.cpp
--- a/Engine/Plugins/Experimental/MeshModelingToolsetExp/Source/ModelingOperators/Private/DeformationOps/ExtrudeOp.cpp
+++ b/Engine/Plugins/Experimental/MeshModelingToolsetExp/Source/ModelingOperators/Private/DeformationOps/ExtrudeOp.cpp
@@ -127,6 +127,8 @@ bool FExtrudeOp::BooleanExtrude(FProgressCancel* Progress)
 		: DirectionMode == EDirectionMode::SelectedTriangleNormalsEven ? FOffsetMeshRegion::EVertexExtrusionVectorType::SelectionTriNormalsAngleWeightedAdjusted
 		: FOffsetMeshRegion::EVertexExtrusionVectorType::Zero;
 
+	Extruder.MaxScaleForAdjustingTriNormalsOffset = MaxScaleForAdjustingTriNormalsOffset;
+
 	Extruder.bIsPositiveOffset = (ExtrudeDistance > 0);
 	Extruder.bOffsetFullComponentsAsSolids = true;
 	Extruder.Apply();
@@ -236,6 +238,8 @@ bool FExtrudeOp::MoveAndStitchExtrude(FProgressCancel* Progress)
 		: DirectionMode == EDirectionMode::SelectedTriangleNormalsEven ? FOffsetMeshRegion::EVertexExtrusionVectorType::SelectionTriNormalsAngleWeightedAdjusted
 		: FOffsetMeshRegion::EVertexExtrusionVectorType::Zero;
 
+	Extruder.MaxScaleForAdjustingTriNormalsOffset = MaxScaleForAdjustingTriNormalsOffset;
+
 	Extruder.bIsPositiveOffset = (ExtrudeDistance > 0);
 	Extruder.bOffsetFullComponentsAsSolids = bShellsToSolids;
 	Extruder.Apply();
diff --git a/Engine/Plugins/Experimental/MeshModelingToolsetExp/Source/ModelingOperators/Public/DeformationOps/ExtrudeOp.h b/Engine/Plugins/Experimental/MeshModelingToolsetExp/Source/ModelingOperators/Public/DeformationOps/ExtrudeOp.h
--- a/Engine/Plugins/Experimental/MeshModelingToolsetExp/Source/ModelingOperators/Public/DeformationOps/ExtrudeOp.h
+++ b/Engine/Plugins/Experimental/MeshModelingToolsetExp/Source/ModelingOperators/Public/DeformationOps/ExtrudeOp.h
@@ -53,6 +53,7 @@ public:
 	double ExtrudeDistance = 0;
 	EExtrudeMode ExtrudeMode = EExtrudeMode::MoveAndStitch;
 	EDirectionMode DirectionMode = EDirectionMode::SelectedTriangleNormalsEven;
+	double MaxScaleForAdjustingTriNormalsOffset = 4.0;
 	// Only used if DirectionMode is SingleDirection
 	FVector3d MeshSpaceExtrudeDirection;
 	// Used when setting groups for the sides when the extrusion includes a mesh border. When true,
diff --git a/Engine/Plugins/Runtime/GeometryProcessing/Source/DynamicMesh/Private/Operations/OffsetMeshRegion.cpp b/Engine/Plugins/Runtime/GeometryProcessing/Source/DynamicMesh/Private/Operations/OffsetMeshRegion.cpp
--- a/Engine/Plugins/Runtime/GeometryProcessing/Source/DynamicMesh/Private/Operations/OffsetMeshRegion.cpp
+++ b/Engine/Plugins/Runtime/GeometryProcessing/Source/DynamicMesh/Private/Operations/OffsetMeshRegion.cpp
@@ -239,6 +239,12 @@ bool FOffsetMeshRegion::ApplyOffset(FOffsetInfo& Region, FMeshNormals* UseNormal
 					{
 						double Angle = Mesh->GetTriInternalAngleR(Tid, Mesh->GetTriangle(Tid).IndexOf(Vid));
 						double CosTheta = Mesh->GetTriNormal(Tid).Dot(ExtrusionVector);
+
+						double InvertedMaxScaleFactor = FMath::Max(FMathd::ZeroTolerance, 1.0 / MaxScaleForAdjustingTriNormalsOffset);
+						if (CosTheta <= InvertedMaxScaleFactor)
+						{
+							CosTheta = InvertedMaxScaleFactor;
+						}
 						Adjustment += Angle / CosTheta;
 
 						// For the average at the end
diff --git a/Engine/Plugins/Runtime/GeometryProcessing/Source/DynamicMesh/Public/Operations/OffsetMeshRegion.h b/Engine/Plugins/Runtime/GeometryProcessing/Source/DynamicMesh/Public/Operations/OffsetMeshRegion.h
--- a/Engine/Plugins/Runtime/GeometryProcessing/Source/DynamicMesh/Public/Operations/OffsetMeshRegion.h
+++ b/Engine/Plugins/Runtime/GeometryProcessing/Source/DynamicMesh/Public/Operations/OffsetMeshRegion.h
@@ -90,6 +90,9 @@ public:
 	 */
 	bool bIsPositiveOffset = true;
 
+	/** When using SelectionTriNormalsAngleWeightedAdjusted to keep triangles parallel, clamp the offset scale to at most this factor */
+	double MaxScaleForAdjustingTriNormalsOffset = 4.0;
+
 	//
 	// Outputs
 	//
diff --git a/Engine/Plugins/Runtime/MeshModelingToolset/Source/MeshModelingTools/Private/ToolActivities/PolyEditExtrudeActivity.cpp b/Engine/Plugins/Runtime/MeshModelingToolset/Source/MeshModelingTools/Private/ToolActivities/PolyEditExtrudeActivity.cpp
--- a/Engine/Plugins/Runtime/MeshModelingToolset/Source/MeshModelingTools/Private/ToolActivities/PolyEditExtrudeActivity.cpp
+++ b/Engine/Plugins/Runtime/MeshModelingToolset/Source/MeshModelingTools/Private/ToolActivities/PolyEditExtrudeActivity.cpp
@@ -55,16 +55,19 @@ TUniquePtr<FDynamicMeshOperator> UPolyEditExtrudeActivity::MakeNewOperator()
 		Op->DirectionMode = ToOpDirectionMode(ExtrudeProperties->DirectionMode);
 		Op->bShellsToSolids = ExtrudeProperties->bShellsToSolids;
 		Op->bUseColinearityForSettingBorderGroups = ExtrudeProperties->bUseColinearityForSettingBorderGroups;
+		Op->MaxScaleForAdjustingTriNormalsOffset = ExtrudeProperties->MaxDistanceScaleFactor;
 		break;
 	case EPropertySetToUse::Offset:
 		Op->DirectionMode = ToOpDirectionMode(OffsetProperties->DirectionMode);
 		Op->bShellsToSolids = OffsetProperties->bShellsToSolids;
 		Op->bUseColinearityForSettingBorderGroups = OffsetProperties->bUseColinearityForSettingBorderGroups;
+		Op->MaxScaleForAdjustingTriNormalsOffset = OffsetProperties->MaxDistanceScaleFactor;
 		break;
 	case EPropertySetToUse::PushPull:
 		Op->DirectionMode = ToOpDirectionMode(PushPullProperties->DirectionMode);
 		Op->bShellsToSolids = PushPullProperties->bShellsToSolids;
 		Op->bUseColinearityForSettingBorderGroups = PushPullProperties->bUseColinearityForSettingBorderGroups;
+		Op->MaxScaleForAdjustingTriNormalsOffset = PushPullProperties->MaxDistanceScaleFactor;
 		break;
 	}
 
diff --git a/Engine/Plugins/Runtime/MeshModelingToolset/Source/MeshModelingTools/Public/ToolActivities/PolyEditExtrudeActivity.h b/Engine/Plugins/Runtime/MeshModelingToolset/Source/MeshModelingTools/Public/ToolActivities/PolyEditExtrudeActivity.h
--- a/Engine/Plugins/Runtime/MeshModelingToolset/Source/MeshModelingTools/Public/ToolActivities/PolyEditExtrudeActivity.h
+++ b/Engine/Plugins/Runtime/MeshModelingToolset/Source/MeshModelingTools/Public/ToolActivities/PolyEditExtrudeActivity.h
@@ -116,6 +116,11 @@ public:
 	UPROPERTY(EditAnywhere, Category = Extrude)
 	EPolyEditExtrudeModeOptions DirectionMode = EPolyEditExtrudeModeOptions::SingleDirection;
 
+	/** Controls the maximum distance vertices can move from the target distance in order to stay parallel with their source triangles. */
+	UPROPERTY(EditAnywhere, Category = Offset,
+		meta = (ClampMin = "1", EditConditionHides, EditCondition = "DirectionMode == EPolyEditExtrudeModeOptions::SelectedTriangleNormalsEven"))
+	double MaxDistanceScaleFactor = 4.0;
+
 	/** 
 	 * When extruding regions that touch the mesh border, assign the side groups (groups on the 
 	 * stitched side of the extrude) in a way that considers edge colinearity. For instance, when
@@ -136,6 +141,11 @@ public:
 	UPROPERTY(EditAnywhere, Category = Offset)
 	EPolyEditOffsetModeOptions DirectionMode = EPolyEditOffsetModeOptions::VertexNormals;
 
+	/** Controls the maximum distance vertices can move from the target distance in order to stay parallel with their source triangles. */
+	UPROPERTY(EditAnywhere, Category = Offset,
+		meta = (ClampMin = "1", EditConditionHides, EditCondition = "DirectionMode == EPolyEditOffsetModeOptions::SelectedTriangleNormalsEven"))
+	double MaxDistanceScaleFactor = 4.0;
+
 	/** Controls whether offseting an entire open-border patch should create a solid or an open shell */
 	UPROPERTY(EditAnywhere, Category = Offset)
 	bool bShellsToSolids = true;
@@ -164,6 +174,11 @@ public:
 	UPROPERTY(EditAnywhere, Category = ExtrusionOptions)
 	EPolyEditPushPullModeOptions DirectionMode = EPolyEditPushPullModeOptions::SelectedTriangleNormals;
 
+	/** Controls the maximum distance vertices can move from the target distance in order to stay parallel with their source triangles. */
+	UPROPERTY(EditAnywhere, Category = Offset,
+		meta = (ClampMin = "1", EditConditionHides, EditCondition = "DirectionMode == EPolyEditPushPullModeOptions::SelectedTriangleNormalsEven"))
+	double MaxDistanceScaleFactor = 4.0;
+
 	/** Controls whether offseting an entire open-border patch should create a solid or an open shell */
 	UPROPERTY(EditAnywhere, Category = ExtrusionOptions)
 	bool bShellsToSolids = true;
