Fix poly edit tool failing to save updated topology of mesh after face 'flip' operation

diff --git a/Engine/Plugins/Runtime/MeshModelingToolset/Source/MeshModelingTools/Private/EditMeshPolygonsTool.cpp b/Engine/Plugins/Runtime/MeshModelingToolset/Source/MeshModelingTools/Private/EditMeshPolygonsTool.cpp
--- a/Engine/Plugins/Runtime/MeshModelingToolset/Source/MeshModelingTools/Private/EditMeshPolygonsTool.cpp
+++ b/Engine/Plugins/Runtime/MeshModelingToolset/Source/MeshModelingTools/Private/EditMeshPolygonsTool.cpp
@@ -478,10 +478,12 @@ void UEditMeshPolygonsTool::Shutdown(EToolShutdownType ShutdownType)
 				OutputSelection->SetSelection(NewSelection);
 			}
 
-			bool bModifiedTopology = (ModifiedTopologyCounter > 0);
+			// Note: When not in triangle mode, ModifiedTopologyCounter refers to polygroup topology, so does not tell us
+			// about the triangle topology.  In this case, we just assume the triangle topology may have been modified.
+			bool bModifiedTriangleTopology = bTriangleMode ? ModifiedTopologyCounter > 0 : true;
 
 			// may need to compact the mesh if we did undo on a mesh edit, then vertices will be dense but compact checks will fail...
-			if (bModifiedTopology)
+			if (bModifiedTriangleTopology)
 			{
 				// Store the compact maps if we have a selection that we need to update
 				CurrentMesh->CompactInPlace(OutputSelection ? &CompactMaps : nullptr);
@@ -490,12 +492,12 @@ void UEditMeshPolygonsTool::Shutdown(EToolShutdownType ShutdownType)
 			// Finish prepping the stored selection
 			if (OutputSelection)
 			{
-				SelectionMechanic->GetSelection(*OutputSelection, bModifiedTopology ? &CompactMaps : nullptr);
+				SelectionMechanic->GetSelection(*OutputSelection, bModifiedTriangleTopology ? &CompactMaps : nullptr);
 			}
 
 			// Bake CurrentMesh back to target inside an undo transaction
 			GetToolManager()->BeginUndoTransaction(LOCTEXT("EditMeshPolygonsToolTransactionName", "Deform Mesh"));
-			UE::ToolTarget::CommitDynamicMeshUpdate(Target, *CurrentMesh, bModifiedTopology);
+			UE::ToolTarget::CommitDynamicMeshUpdate(Target, *CurrentMesh, bModifiedTriangleTopology);
 
 			UE::Geometry::SetToolOutputSelection(this, OutputSelection);
 		
@@ -1219,11 +1221,9 @@ void UEditMeshPolygonsTool::ApplyFlipNormals()
 		}
 	}
 
-	// We actually don't even need any of the wrapper around this change since we're not altering
-	// positions or topology (so no other structures need updating), but we go ahead and go the
-	// same route as everything else.
+	// Note the topology can change in that the ordering of edge elements can reverse
 	EmitCurrentMeshChangeAndUpdate(LOCTEXT("PolyMeshFlipNormalsChange", "Flip Normals"), 
-		ChangeTracker.EndChange(), ActiveSelection, false);
+		ChangeTracker.EndChange(), ActiveSelection, true);
 }
 
 
