MeshModelingTools: Fix "reproject" being in outset  instead of inset due to using the wrong edit condition.

#rnx

diff --git a/Engine/Plugins/Runtime/MeshModelingToolset/Source/MeshModelingTools/Public/ToolActivities/PolyEditInsetOutsetActivity.h b/Engine/Plugins/Runtime/MeshModelingToolset/Source/MeshModelingTools/Public/ToolActivities/PolyEditInsetOutsetActivity.h
--- a/Engine/Plugins/Runtime/MeshModelingToolset/Source/MeshModelingTools/Public/ToolActivities/PolyEditInsetOutsetActivity.h
+++ b/Engine/Plugins/Runtime/MeshModelingToolset/Source/MeshModelingTools/Public/ToolActivities/PolyEditInsetOutsetActivity.h
@@ -34,7 +34,7 @@ public:
 	float AreaScale = true;
 
 	/** When insetting, determines whether vertices in inset region should be projected back onto input surface */
-	UPROPERTY(EditAnywhere, Category = "Inset/Outset", Meta = (EditCondition = "bOutset", HideEditConditionToggle))
+	UPROPERTY(EditAnywhere, Category = "Inset/Outset", Meta = (EditCondition = "!bOutset", HideEditConditionToggle, EditConditionHides))
 	bool bReproject = true;
 
 	//~ This is not user editable- it gets set by PolyEdit depending on whether the user clicks
