MeshModelingToolset: Cleanup Append and Dupe tools

diff --git a/Engine/Plugins/Editor/ModelingToolsEditorMode/Source/ModelingToolsEditorMode/Private/ModelingToolsManagerActions.cpp b/Engine/Plugins/Editor/ModelingToolsEditorMode/Source/ModelingToolsEditorMode/Private/ModelingToolsManagerActions.cpp
--- a/Engine/Plugins/Editor/ModelingToolsEditorMode/Source/ModelingToolsEditorMode/Private/ModelingToolsManagerActions.cpp
+++ b/Engine/Plugins/Editor/ModelingToolsEditorMode/Source/ModelingToolsEditorMode/Private/ModelingToolsManagerActions.cpp
@@ -1,7 +1,6 @@
 // Copyright Epic Games, Inc. All Rights Reserved.
 
 #include "ModelingToolsManagerActions.h"
-#include "EditorStyleSet.h"
 #include "InputCoreTypes.h"
 #include "ModelingToolsEditorModeStyle.h"
 
@@ -60,8 +59,8 @@ void FModelingToolsManagerCommands::RegisterCommands()
 	REGISTER_MODELING_TOOL_COMMAND(BeginDrawPolyPathTool, "PathExt", "Draw and extrude PolyPaths to create new objects");
 	REGISTER_MODELING_TOOL_COMMAND(BeginDrawAndRevolveTool, "PathRev", "Draw and revolve PolyPaths to create new objects");
 	REGISTER_MODELING_TOOL_COMMAND(BeginRevolveBoundaryTool, "BdryRev", "Revolve mesh boundary loops to create new objects");
-	REGISTER_MODELING_TOOL_COMMAND(BeginCombineMeshesTool, "Append", "Combine selections into new objects");
-	REGISTER_MODELING_TOOL_COMMAND(BeginDuplicateMeshesTool, "Dupe", "Duplicate selections into new objects");
+	REGISTER_MODELING_TOOL_COMMAND(BeginCombineMeshesTool, "MshMrg", "Merge multiple meshes to create new objects");
+	REGISTER_MODELING_TOOL_COMMAND(BeginDuplicateMeshesTool, "MshDup", "Duplicate single meshes to create new objects");
 
 	REGISTER_MODELING_TOOL_COMMAND(BeginAddPatchTool, "Patch", "Add Patch");
 	REGISTER_MODELING_TOOL_COMMAND(BeginShapeSprayTool, "Spray", "Shape Spray");
diff --git a/Engine/Plugins/Experimental/MeshModelingToolsetExp/Source/MeshModelingToolsExp/Public/Properties/RevolveProperties.h b/Engine/Plugins/Experimental/MeshModelingToolsetExp/Source/MeshModelingToolsExp/Public/Properties/RevolveProperties.h
--- a/Engine/Plugins/Experimental/MeshModelingToolsetExp/Source/MeshModelingToolsExp/Public/Properties/RevolveProperties.h
+++ b/Engine/Plugins/Experimental/MeshModelingToolsetExp/Source/MeshModelingToolsExp/Public/Properties/RevolveProperties.h
@@ -12,11 +12,17 @@
 class UNewMeshMaterialProperties;
 PREDECLARE_GEOMETRY(class FCurveSweepOp);
 
+UENUM()
 enum class ERevolvePropertiesCapFillMode : uint8
 {
+	/** No caps will be generated. */
 	None,
+	/** Caps are triangulated by placing a vertex in the center and creating a fan to the boundary. This works well if the path is convex,
+	  * but can create invalid geometry if it is concave. */
 	CenterFan,
+	/** Caps are triangulated to maximize the minimal angle in the triangles using Delaunay triangulation. */
 	Delaunay,
+	/** Caps are triangulated using a standard ear clipping approach. This could result in some very thin triangles. */
 	EarClipping
 };
 
diff --git a/Engine/Plugins/Experimental/MeshModelingToolsetExp/Source/MeshModelingToolsExp/Public/RevolveBoundaryTool.h b/Engine/Plugins/Experimental/MeshModelingToolsetExp/Source/MeshModelingToolsExp/Public/RevolveBoundaryTool.h
--- a/Engine/Plugins/Experimental/MeshModelingToolsetExp/Source/MeshModelingToolsExp/Public/RevolveBoundaryTool.h
+++ b/Engine/Plugins/Experimental/MeshModelingToolsetExp/Source/MeshModelingToolsExp/Public/RevolveBoundaryTool.h
@@ -40,20 +40,6 @@ public:
 	TObjectPtr<URevolveBoundaryTool> RevolveBoundaryTool;
 };
 
-UENUM()
-enum class ERevolvePropertiesBoundaryCapFillMode : uint8
-{
-	/** No caps will be generated. */
-	None = static_cast<uint8>(ERevolvePropertiesCapFillMode::None),
-	/** Caps are triangulated by placing a vertex in the center and creating a fan to the boundary. This works well if the path is convex,
-	  * but can create invalid geometry if it is concave. */
-	CenterFan = static_cast<uint8>(ERevolvePropertiesCapFillMode::CenterFan),
-	/** Caps are triangulated to maximize the minimal angle in the triangles using Delaunay triangulation. */
-	Delaunay = static_cast<uint8>(ERevolvePropertiesCapFillMode::Delaunay),
-	/** Caps are triangulated using a standard ear clipping approach. This could result in some very thin triangles. */
-	EarClipping = static_cast<uint8>(ERevolvePropertiesCapFillMode::EarClipping)
-};
-
 UCLASS()
 class MESHMODELINGTOOLSEXP_API URevolveBoundaryToolProperties : public URevolveProperties
 {
@@ -64,7 +50,7 @@ public:
 	/** Determines how end caps are created. This is not relevant if the end caps are not visible or if the path is not closed. */
 	UPROPERTY(EditAnywhere, Category = Revolve, AdvancedDisplay, meta = (DisplayAfter = "QuadSplitMode",
 		EditCondition = "HeightOffsetPerDegree != 0 || RevolveDegrees != 360"))
-	ERevolvePropertiesBoundaryCapFillMode CapFillMode = ERevolvePropertiesBoundaryCapFillMode::Delaunay;
+	ERevolvePropertiesCapFillMode CapFillMode = ERevolvePropertiesCapFillMode::Delaunay;
 
 	/** If true, displays the original mesh in addition to the revolved boundary. */
 	UPROPERTY(EditAnywhere, Category = Revolve, AdvancedDisplay)
@@ -81,7 +67,7 @@ public:
 protected:
 	virtual ERevolvePropertiesCapFillMode GetCapFillMode() const override
 	{
-		return static_cast<ERevolvePropertiesCapFillMode>(CapFillMode);
+		return CapFillMode;
 	}
 };
 
diff --git a/Engine/Plugins/Experimental/MeshModelingToolsetExp/Source/ModelingComponents/Private/BaseTools/BaseCreateFromSelectedTool.cpp b/Engine/Plugins/Experimental/MeshModelingToolsetExp/Source/ModelingComponents/Private/BaseTools/BaseCreateFromSelectedTool.cpp
--- a/Engine/Plugins/Experimental/MeshModelingToolsetExp/Source/ModelingComponents/Private/BaseTools/BaseCreateFromSelectedTool.cpp
+++ b/Engine/Plugins/Experimental/MeshModelingToolsetExp/Source/ModelingComponents/Private/BaseTools/BaseCreateFromSelectedTool.cpp
@@ -4,14 +4,11 @@
 #include "InteractiveToolManager.h"
 #include "ToolBuilderUtil.h"
 #include "ToolSetupUtil.h"
-#include "Components/DynamicMeshComponent.h"
 #include "BaseGizmos/TransformGizmoUtil.h"
-#include "Async/Async.h"
 
 #include "DynamicMesh/MeshNormals.h"
 #include "DynamicMesh/MeshTransforms.h"
 #include "DynamicMeshToMeshDescription.h"
-#include "MeshDescriptionToDynamicMesh.h"
 
 #include "ModelingObjectsCreationAPI.h"
 #include "Selection/ToolSelectionUtil.h"
@@ -113,19 +110,19 @@ void UBaseCreateFromSelectedTool::Setup()
 	ConvertInputsAndSetPreviewMaterials(true);
 
 	// output name fields
-	HandleSourcesProperties->OutputName = PrefixWithSourceNameIfSingleSelection(GetCreatedAssetName());
-	HandleSourcesProperties->WatchProperty(HandleSourcesProperties->WriteOutputTo, [&](EBaseCreateFromSelectedTargetType NewType)
+	HandleSourcesProperties->OutputNewName = PrefixWithSourceNameIfSingleSelection(GetCreatedAssetName());
+	HandleSourcesProperties->WatchProperty(HandleSourcesProperties->OutputWriteTo, [&](EBaseCreateFromSelectedTargetType NewType)
 	{
-		SetToolPropertySourceEnabled(OutputTypeProperties, (NewType == EBaseCreateFromSelectedTargetType::NewAsset));
-		if (NewType == EBaseCreateFromSelectedTargetType::NewAsset)
+		SetToolPropertySourceEnabled(OutputTypeProperties, (NewType == EBaseCreateFromSelectedTargetType::NewObject));
+		if (NewType == EBaseCreateFromSelectedTargetType::NewObject)
 		{
-			HandleSourcesProperties->OutputAsset = TEXT("");
+			HandleSourcesProperties->OutputExistingName = TEXT("");
 			UpdateGizmoVisibility();
 		}
 		else
 		{
-			int32 Index = (HandleSourcesProperties->WriteOutputTo == EBaseCreateFromSelectedTargetType::FirstInputAsset) ? 0 : Targets.Num() - 1;
-			HandleSourcesProperties->OutputAsset = UE::Modeling::GetComponentAssetBaseName(TargetComponentInterface(Index)->GetOwnerComponent(), false);
+			int32 Index = (HandleSourcesProperties->OutputWriteTo == EBaseCreateFromSelectedTargetType::FirstInputObject) ? 0 : Targets.Num() - 1;
+			HandleSourcesProperties->OutputExistingName = UE::Modeling::GetComponentAssetBaseName(TargetComponentInterface(Index)->GetOwnerComponent(), false);
 
 			// Reset the hidden gizmo to its initial position
 			FTransform ComponentTransform = TargetComponentInterface(Index)->GetWorldTransform();
@@ -140,13 +137,13 @@ void UBaseCreateFromSelectedTool::Setup()
 
 int32 UBaseCreateFromSelectedTool::GetHiddenGizmoIndex() const
 {
-	if (HandleSourcesProperties->WriteOutputTo == EBaseCreateFromSelectedTargetType::NewAsset)
+	if (HandleSourcesProperties->OutputWriteTo == EBaseCreateFromSelectedTargetType::NewObject)
 	{
 		return -1;
 	}
 	else
 	{
-		return (HandleSourcesProperties->WriteOutputTo == EBaseCreateFromSelectedTargetType::FirstInputAsset) ? 0 : Targets.Num() - 1;
+		return (HandleSourcesProperties->OutputWriteTo == EBaseCreateFromSelectedTargetType::FirstInputObject) ? 0 : Targets.Num() - 1;
 	}
 }
 
@@ -233,7 +230,7 @@ void UBaseCreateFromSelectedTool::GenerateAsset(const FDynamicMeshOpResult& OpRe
 	}
 
 	// max len explicitly enforced here, would ideally notify user
-	FString UseBaseName = HandleSourcesProperties->OutputName.Left(250);
+	FString UseBaseName = HandleSourcesProperties->OutputNewName.Left(250);
 	if (UseBaseName.IsEmpty())
 	{
 		UseBaseName = PrefixWithSourceNameIfSingleSelection(GetCreatedAssetName());
@@ -341,13 +338,13 @@ void UBaseCreateFromSelectedTool::Shutdown(EToolShutdownType ShutdownType)
 
 		// Generate the result
 		AActor* KeepActor = nullptr;
-		if (HandleSourcesProperties->WriteOutputTo == EBaseCreateFromSelectedTargetType::NewAsset)
+		if (HandleSourcesProperties->OutputWriteTo == EBaseCreateFromSelectedTargetType::NewObject)
 		{
 			GenerateAsset(Result);
 		}
 		else
 		{
-			int32 TargetIndex = (HandleSourcesProperties->WriteOutputTo == EBaseCreateFromSelectedTargetType::FirstInputAsset) ? 0 : (Targets.Num() - 1);
+			int32 TargetIndex = (HandleSourcesProperties->OutputWriteTo == EBaseCreateFromSelectedTargetType::FirstInputObject) ? 0 : (Targets.Num() - 1);
 			KeepActor = TargetComponentInterface(TargetIndex)->GetOwnerActor();
 
 			UpdateAsset(Result, Targets[TargetIndex]);
diff --git a/Engine/Plugins/Experimental/MeshModelingToolsetExp/Source/ModelingComponents/Private/PropertySets/OnAcceptProperties.cpp b/Engine/Plugins/Experimental/MeshModelingToolsetExp/Source/ModelingComponents/Private/PropertySets/OnAcceptProperties.cpp
--- a/Engine/Plugins/Experimental/MeshModelingToolsetExp/Source/ModelingComponents/Private/PropertySets/OnAcceptProperties.cpp
+++ b/Engine/Plugins/Experimental/MeshModelingToolsetExp/Source/ModelingComponents/Private/PropertySets/OnAcceptProperties.cpp
@@ -9,47 +9,38 @@
 #define LOCTEXT_NAMESPACE "UOnAcceptHandleSourcesProperties"
 
 
-void UOnAcceptHandleSourcesProperties::ApplyMethod(const TArray<AActor*>& Actors, UInteractiveToolManager* ToolManager, const AActor* MustKeepActor)
+void UOnAcceptHandleSourcesPropertiesBase::ApplyMethod(const TArray<AActor*>& Actors, UInteractiveToolManager* ToolManager, const AActor* MustKeepActor)
 {
+	const EHandleSourcesMethod HandleInputs = GetHandleInputs();
+
 	// Hide or destroy the sources
-	bool bKeepSources = OnToolAccept == EHandleSourcesMethod::KeepSources;
-	if (Actors.Num() == 1 && (OnToolAccept == EHandleSourcesMethod::KeepFirstSource || OnToolAccept == EHandleSourcesMethod::KeepLastSource))
+	bool bKeepSources = HandleInputs == EHandleSourcesMethod::KeepSources;
+	if (Actors.Num() == 1 && (HandleInputs == EHandleSourcesMethod::KeepFirstSource || HandleInputs == EHandleSourcesMethod::KeepLastSource))
 	{
 		// if there's only one actor, keeping any source == keeping all sources
 		bKeepSources = true;
 	}
 	if (!bKeepSources)
 	{
-		bool bDelete = OnToolAccept == EHandleSourcesMethod::DeleteSources
-					|| OnToolAccept == EHandleSourcesMethod::KeepFirstSource
-					|| OnToolAccept == EHandleSourcesMethod::KeepLastSource;
+		bool bDelete = HandleInputs == EHandleSourcesMethod::DeleteSources
+					|| HandleInputs == EHandleSourcesMethod::KeepFirstSource
+					|| HandleInputs == EHandleSourcesMethod::KeepLastSource;
 		if (bDelete)
 		{
-			ToolManager->BeginUndoTransaction(LOCTEXT("RemoveSources", "Remove Sources"));
+			ToolManager->BeginUndoTransaction(LOCTEXT("RemoveSources", "Remove Inputs"));
 		}
 		else
 		{
 #if WITH_EDITOR
-			ToolManager->BeginUndoTransaction(LOCTEXT("HideSources", "Hide Sources"));
+			ToolManager->BeginUndoTransaction(LOCTEXT("HideSources", "Hide Inputs"));
 #endif
 		}
 
-		int32 SkipIdx = -1;
-		if (OnToolAccept == EHandleSourcesMethod::KeepFirstSource)
-		{
-			SkipIdx = 0;
-		}
-		else if (OnToolAccept == EHandleSourcesMethod::KeepLastSource)
-		{
-			SkipIdx = Actors.Num() - 1;
-		}
-		for (int32 ActorIdx = 0; ActorIdx < Actors.Num(); ActorIdx++)
-		{
-			if (ActorIdx == SkipIdx)
-			{
-				continue;
-			}
+		const int32 ActorIdxBegin = HandleInputs == EHandleSourcesMethod::KeepFirstSource ? 1 : 0;
+		const int32 ActorIdxEnd = HandleInputs == EHandleSourcesMethod::KeepLastSource ? Actors.Num() - 1 : Actors.Num();
 
+		for (int32 ActorIdx = ActorIdxBegin; ActorIdx < ActorIdxEnd; ActorIdx++)
+		{
 			AActor* Actor = Actors[ActorIdx];
 			if (Actor == MustKeepActor)
 			{
diff --git a/Engine/Plugins/Experimental/MeshModelingToolsetExp/Source/ModelingComponents/Public/BaseTools/BaseCreateFromSelectedTool.h b/Engine/Plugins/Experimental/MeshModelingToolsetExp/Source/ModelingComponents/Public/BaseTools/BaseCreateFromSelectedTool.h
--- a/Engine/Plugins/Experimental/MeshModelingToolsetExp/Source/ModelingComponents/Public/BaseTools/BaseCreateFromSelectedTool.h
+++ b/Engine/Plugins/Experimental/MeshModelingToolsetExp/Source/ModelingComponents/Public/BaseTools/BaseCreateFromSelectedTool.h
@@ -40,12 +40,14 @@ protected:
 UENUM()
 enum class EBaseCreateFromSelectedTargetType
 {
-	/** Create a new asset containing the result mesh */
-	NewAsset,
-	/** Store the result mesh in the first selected input asset */
-	FirstInputAsset,
-	/** Store the result mesh in the last selected input asset */
-	LastInputAsset
+	/** Create and write to a new object with a given name. */
+	NewObject,
+
+	/** Write to the first object in the input selection. */
+	FirstInputObject,
+
+	/** Write to the last object in the input selection. */
+	LastInputObject
 };
 
 
@@ -54,17 +56,19 @@ class MODELINGCOMPONENTS_API UBaseCreateFromSelectedHandleSourceProperties : pub
 {
 	GENERATED_BODY()
 public:
-	/** Where should the output mesh produced by this operation be stored */
-	UPROPERTY(EditAnywhere, Category = ToolOutputOptions)
-	EBaseCreateFromSelectedTargetType WriteOutputTo = EBaseCreateFromSelectedTargetType::NewAsset;
-
-	/** Base name for newly-generated asset */
-	UPROPERTY(EditAnywhere, Category = ToolOutputOptions, meta = (TransientToolProperty, EditCondition = "WriteOutputTo == EBaseCreateFromSelectedTargetType::NewAsset", EditConditionHides))
-	FString OutputName;
-
-	/** Name of asset that will be updated */
-	UPROPERTY(VisibleAnywhere, Category = ToolOutputOptions, meta = (TransientToolProperty, EditCondition = "WriteOutputTo != EBaseCreateFromSelectedTargetType::NewAsset", EditConditionHides))
-	FString OutputAsset;
+	/** Defines the object the tool output is written to. */
+	UPROPERTY(EditAnywhere, Category = OutputObject, meta = (DisplayName = "Write To"))
+	EBaseCreateFromSelectedTargetType OutputWriteTo = EBaseCreateFromSelectedTargetType::NewObject;
+
+	/** Base name of the newly generated object to which the output is written to. */
+	UPROPERTY(EditAnywhere, Category = OutputObject, meta = (TransientToolProperty, DisplayName = "Name",
+		EditCondition = "OutputWriteTo == EBaseCreateFromSelectedTargetType::NewObject", EditConditionHides, NoResetToDefault))
+	FString OutputNewName;
+
+	/** Name of the existing object to which the output is written to. */
+	UPROPERTY(VisibleAnywhere, Category = OutputObject, meta = (TransientToolProperty, DisplayName = "Name",
+		EditCondition = "OutputWriteTo != EBaseCreateFromSelectedTargetType::NewObject", EditConditionHides))
+	FString OutputExistingName;
 };
 
 
diff --git a/Engine/Plugins/Experimental/MeshModelingToolsetExp/Source/ModelingComponents/Public/PropertySets/OnAcceptProperties.h b/Engine/Plugins/Experimental/MeshModelingToolsetExp/Source/ModelingComponents/Public/PropertySets/OnAcceptProperties.h
--- a/Engine/Plugins/Experimental/MeshModelingToolsetExp/Source/ModelingComponents/Public/PropertySets/OnAcceptProperties.h
+++ b/Engine/Plugins/Experimental/MeshModelingToolsetExp/Source/ModelingComponents/Public/PropertySets/OnAcceptProperties.h
@@ -16,32 +16,70 @@ class AActor;
 UENUM()
 enum class EHandleSourcesMethod : uint8
 {
-	DeleteSources = 0	UMETA(DisplayName = "Delete Sources"),
+	/** Delete all input objects */
+	DeleteSources UMETA(DisplayName = "Delete Inputs"),
 
-	HideSources = 1		UMETA(DisplayName = "Hide Sources"),
+	/** Hide all input objects */
+	HideSources UMETA(DisplayName = "Hide Inputs"),
 
-	KeepSources = 2		UMETA(DisplayName = "Keep Sources"),
+	/** Keep all input objects */
+	KeepSources UMETA(DisplayName = "Keep Inputs"),
 
-	/** Delete all but the first source */
-	KeepFirstSource = 3	UMETA(DisplayName = "Keep First Source"),
+	/** Keep only the first input object and delete all other input objects */
+	KeepFirstSource UMETA(DisplayName = "Keep First Input"),
 
-	/** Delete all but the last source */
-	KeepLastSource = 4	UMETA(DisplayName = "Keep Last Source")
+	/** Keep only the last input object and delete all other input objects */
+	KeepLastSource UMETA(DisplayName = "Keep Last Input")
 };
 
+// Base class for property settings for tools that create a new actor and need to decide what to do with the input objects.
+UCLASS()
+class MODELINGCOMPONENTS_API UOnAcceptHandleSourcesPropertiesBase : public UInteractiveToolPropertySet
+{
+	GENERATED_BODY()
+public:
 
+	void ApplyMethod(const TArray<AActor*>& Actors, UInteractiveToolManager* ToolManager, const AActor* MustKeepActor = nullptr);
 
+protected:
+	virtual EHandleSourcesMethod GetHandleInputs() const
+	{
+		return EHandleSourcesMethod::KeepSources;
+	}
+};
 
-// Standard property settings for tools that create a new actor and need to decide what to do with the input (source) actor(s)
+// Specialization for property settings for tools that create a new actor and need to decide what to do with multiple input objects.
 UCLASS()
-class MODELINGCOMPONENTS_API UOnAcceptHandleSourcesProperties : public UInteractiveToolPropertySet
+class MODELINGCOMPONENTS_API UOnAcceptHandleSourcesProperties : public UOnAcceptHandleSourcesPropertiesBase
 {
 	GENERATED_BODY()
 public:
 
-	/** What to do with the source Actors/Components when accepting results of tool.*/
-	UPROPERTY(EditAnywhere, Category = ToolOutputOptions)
-	EHandleSourcesMethod OnToolAccept;
+	/** Defines what to do with the input objects when accepting the tool results. */
+	UPROPERTY(EditAnywhere, Category = OnToolAccept)
+	EHandleSourcesMethod HandleInputs;
 
-	void ApplyMethod(const TArray<AActor*>& Actors, UInteractiveToolManager* ToolManager, const AActor* MustKeepActor = nullptr);
+protected:
+	virtual EHandleSourcesMethod GetHandleInputs() const override
+	{
+		return HandleInputs;
+	}
+};
+
+// Specialization for property settings for tools that create a new actor and need to decide what to do with a single input object.
+UCLASS()
+class MODELINGCOMPONENTS_API UOnAcceptHandleSourcesPropertiesSingle : public UOnAcceptHandleSourcesPropertiesBase
+{
+	GENERATED_BODY()
+public:
+
+	/** Defines what to do with the input object when accepting the tool results. */
+	UPROPERTY(EditAnywhere, Category = OnToolAccept, meta = (ValidEnumValues = "DeleteSources, HideSources, KeepSources"))
+	EHandleSourcesMethod HandleInputs;
+
+protected:
+	virtual EHandleSourcesMethod GetHandleInputs() const override
+	{
+		return HandleInputs;
+	}
 };
diff --git a/Engine/Plugins/Runtime/MeshModelingToolset/Source/MeshModelingTools/Private/CombineMeshesTool.cpp b/Engine/Plugins/Runtime/MeshModelingToolset/Source/MeshModelingTools/Private/CombineMeshesTool.cpp
--- a/Engine/Plugins/Runtime/MeshModelingToolset/Source/MeshModelingTools/Private/CombineMeshesTool.cpp
+++ b/Engine/Plugins/Runtime/MeshModelingToolset/Source/MeshModelingTools/Private/CombineMeshesTool.cpp
@@ -3,7 +3,6 @@
 #include "CombineMeshesTool.h"
 #include "InteractiveToolManager.h"
 #include "ToolBuilderUtil.h"
-#include "ToolSetupUtil.h"
 
 #include "DynamicMesh/DynamicMesh3.h"
 #include "DynamicMeshEditor.h"
@@ -14,9 +13,6 @@
 #include "Physics/ComponentCollisionUtil.h"
 #include "ShapeApproximation/SimpleShapeSet3.h"
 
-#include "Components/StaticMeshComponent.h"
-#include "Engine/StaticMesh.h"
-
 #include "TargetInterfaces/MaterialProvider.h"
 #include "TargetInterfaces/MeshDescriptionCommitter.h"
 #include "TargetInterfaces/MeshDescriptionProvider.h"
@@ -28,7 +24,6 @@
 #include "Misc/ScopedSlowTask.h"
 #endif
 
-#include "ExplicitUseGeometryMathTypes.h"		// using UE::Geometry::(math types)
 using namespace UE::Geometry;
 
 #define LOCTEXT_NAMESPACE "UCombineMeshesTool"
@@ -69,8 +64,6 @@ UInteractiveTool* UCombineMeshesToolBuilder::BuildTool(const FToolBuilderState&
 }
 
 
-
-
 /*
  * Tool
  */
@@ -102,56 +95,55 @@ void UCombineMeshesTool::Setup()
 	OutputTypeProperties->WatchProperty(OutputTypeProperties->OutputType, [this](FString) { OutputTypeProperties->UpdatePropertyVisibility(); });
 	AddToolPropertySource(OutputTypeProperties);
 
-	BasicProperties->WatchProperty(BasicProperties->WriteOutputTo, [&](ECombineTargetType NewType)
+	BasicProperties->WatchProperty(BasicProperties->OutputWriteTo, [&](EBaseCreateFromSelectedTargetType NewType)
 	{
-		if (NewType == ECombineTargetType::NewAsset)
+		if (NewType == EBaseCreateFromSelectedTargetType::NewObject)
 		{
-			BasicProperties->OutputAsset = TEXT("");
+			BasicProperties->OutputExistingName = TEXT("");
 			SetToolPropertySourceEnabled(OutputTypeProperties, true);
 		}
 		else
 		{
-			int32 Index = (BasicProperties->WriteOutputTo == ECombineTargetType::FirstInputAsset) ? 0 : Targets.Num() - 1;
-			BasicProperties->OutputAsset = UE::Modeling::GetComponentAssetBaseName(UE::ToolTarget::GetTargetComponent(Targets[Index]), false);
+			int32 Index = (BasicProperties->OutputWriteTo == EBaseCreateFromSelectedTargetType::FirstInputObject) ? 0 : Targets.Num() - 1;
+			BasicProperties->OutputExistingName = UE::Modeling::GetComponentAssetBaseName(UE::ToolTarget::GetTargetComponent(Targets[Index]), false);
 			SetToolPropertySourceEnabled(OutputTypeProperties, false);
 		}
 	});
 
-	SetToolPropertySourceEnabled(OutputTypeProperties, BasicProperties->WriteOutputTo == ECombineTargetType::NewAsset);
+	SetToolPropertySourceEnabled(OutputTypeProperties, BasicProperties->OutputWriteTo == EBaseCreateFromSelectedTargetType::NewObject);
 
 	if (bDuplicateMode)
 	{
 		SetToolDisplayName(LOCTEXT("DuplicateMeshesToolName", "Duplicate"));
-		BasicProperties->OutputName = UE::Modeling::GetComponentAssetBaseName(UE::ToolTarget::GetTargetComponent(Targets[0]));
+		BasicProperties->OutputNewName = UE::Modeling::GetComponentAssetBaseName(UE::ToolTarget::GetTargetComponent(Targets[0]));
 	}
 	else
 	{
 		SetToolDisplayName(LOCTEXT("CombineMeshesToolName", "Append"));
-		BasicProperties->OutputName = FString("Combined");
+		BasicProperties->OutputNewName = FString("Combined");
 	}
 
-
-	HandleSourceProperties = NewObject<UOnAcceptHandleSourcesProperties>(this);
+	HandleSourceProperties = bDuplicateMode
+		                         ? static_cast<UOnAcceptHandleSourcesPropertiesBase*>(NewObject<UOnAcceptHandleSourcesPropertiesSingle>(this))
+		                         : static_cast<UOnAcceptHandleSourcesPropertiesBase*>(NewObject<UOnAcceptHandleSourcesProperties>(this));
 	AddToolPropertySource(HandleSourceProperties);
 	HandleSourceProperties->RestoreProperties(this);
 
-
 	if (bDuplicateMode)
 	{
 		GetToolManager()->DisplayMessage(
-			LOCTEXT("OnStartToolDuplicate", "This Tool duplicates input Asset into a new Asset, and optionally replaces the input Actor with a new Actor containing the new Asset."),
+			LOCTEXT("OnStartToolDuplicate", "This tool duplicates a single input object to create new objects, and optionally replaces the input object."),
 			EToolMessageLevel::UserNotification);
 	}
 	else
 	{
 		GetToolManager()->DisplayMessage(
-			LOCTEXT("OnStartToolCombine", "This Tool appends the meshes from the input Assets into a new Asset, and optionally replaces the source Actors with a new Actor containing the new Asset."),
+			LOCTEXT("OnStartToolCombine", "This tool appends multiple input object to create new objects, and optionally replaces the one of the input objects."),
 			EToolMessageLevel::UserNotification);
 	}
 }
 
 
-
 void UCombineMeshesTool::Shutdown(EToolShutdownType ShutdownType)
 {
 	BasicProperties->SaveProperties(this);
@@ -160,7 +152,7 @@ void UCombineMeshesTool::Shutdown(EToolShutdownType ShutdownType)
 
 	if (ShutdownType == EToolShutdownType::Accept)
 	{
-		if (bDuplicateMode || BasicProperties->WriteOutputTo == ECombineTargetType::NewAsset)
+		if (bDuplicateMode || BasicProperties->OutputWriteTo == EBaseCreateFromSelectedTargetType::NewObject)
 		{
 			CreateNewAsset();
 		}
@@ -171,6 +163,7 @@ void UCombineMeshesTool::Shutdown(EToolShutdownType ShutdownType)
 	}
 }
 
+
 void UCombineMeshesTool::CreateNewAsset()
 {
 	// Make sure meshes are available before we open transaction. This is to avoid potential stability issues related 
@@ -184,7 +177,7 @@ void UCombineMeshesTool::CreateNewAsset()
 
 	GetToolManager()->BeginUndoTransaction( bDuplicateMode ? 
 		LOCTEXT("DuplicateMeshToolTransactionName", "Duplicate Mesh") :
-		LOCTEXT("CombineMeshesToolTransactionName", "Combine Meshes"));
+		LOCTEXT("CombineMeshesToolTransactionName", "Merge Meshes"));
 
 	FBox Box(ForceInit);
 	for (int32 ComponentIdx = 0; ComponentIdx < Targets.Num(); ComponentIdx++)
@@ -213,7 +206,7 @@ void UCombineMeshesTool::CreateNewAsset()
 		FScopedSlowTask SlowTask(Targets.Num()+1, 
 			bDuplicateMode ? 
 			LOCTEXT("DuplicateMeshBuild", "Building duplicate mesh ...") :
-			LOCTEXT("CombineMeshesBuild", "Building combined mesh ..."));
+			LOCTEXT("CombineMeshesBuild", "Building merged mesh ..."));
 		SlowTask.MakeDialog();
 #endif
 		bool bNeedColorAttr = false;
@@ -291,10 +284,10 @@ void UCombineMeshesTool::CreateNewAsset()
 		}
 
 		// max len explicitly enforced here, would ideally notify user
-		FString UseBaseName = BasicProperties->OutputName.Left(250);
+		FString UseBaseName = BasicProperties->OutputNewName.Left(250);
 		if (UseBaseName.IsEmpty())
 		{
-			UseBaseName = (bDuplicateMode) ? TEXT("Duplicate") : TEXT("Combined");
+			UseBaseName = (bDuplicateMode) ? TEXT("Duplicate") : TEXT("Merge");
 		}
 
 		FCreateMeshObjectParams NewMeshObjectParams;
@@ -337,11 +330,6 @@ void UCombineMeshesTool::CreateNewAsset()
 }
 
 
-
-
-
-
-
 void UCombineMeshesTool::UpdateExistingAsset()
 {
 	// Make sure meshes are available before we open transaction. This is to avoid potential stability issues related 
@@ -354,7 +342,7 @@ void UCombineMeshesTool::UpdateExistingAsset()
 	}
 
 	check(!bDuplicateMode);
-	GetToolManager()->BeginUndoTransaction(LOCTEXT("CombineMeshesToolTransactionName", "Combine Meshes"));
+	GetToolManager()->BeginUndoTransaction(LOCTEXT("CombineMeshesToolTransactionName", "Merge Meshes"));
 
 	AActor* SkipActor = nullptr;
 
@@ -369,7 +357,7 @@ void UCombineMeshesTool::UpdateExistingAsset()
 	AccumulateDMesh.Attributes()->EnableMaterialID();
 	AccumulateDMesh.Attributes()->EnablePrimaryColors();
 
-	int32 SkipIndex = (BasicProperties->WriteOutputTo == ECombineTargetType::FirstInputAsset) ? 0 : (Targets.Num() - 1);
+	int32 SkipIndex = (BasicProperties->OutputWriteTo == EBaseCreateFromSelectedTargetType::FirstInputObject) ? 0 : (Targets.Num() - 1);
 	UPrimitiveComponent* UpdateComponent = UE::ToolTarget::GetTargetComponent(Targets[SkipIndex]);
 	SkipActor = UE::ToolTarget::GetTargetActor(Targets[SkipIndex]);
 
@@ -391,7 +379,7 @@ void UCombineMeshesTool::UpdateExistingAsset()
 		FScopedSlowTask SlowTask(Targets.Num()+1, 
 			bDuplicateMode ? 
 			LOCTEXT("DuplicateMeshBuild", "Building duplicate mesh ...") :
-			LOCTEXT("CombineMeshesBuild", "Building combined mesh ..."));
+			LOCTEXT("CombineMeshesBuild", "Building merged mesh ..."));
 		SlowTask.MakeDialog();
 #endif
 		bool bNeedColorAttr = false;
@@ -481,11 +469,6 @@ void UCombineMeshesTool::UpdateExistingAsset()
 }
 
 
-
-
-
-
-
 void UCombineMeshesTool::BuildCombinedMaterialSet(TArray<UMaterialInterface*>& NewMaterialsOut, TArray<TArray<int32>>& MaterialIDRemapsOut)
 {
 	NewMaterialsOut.Reset();
@@ -517,10 +500,4 @@ void UCombineMeshesTool::BuildCombinedMaterialSet(TArray<UMaterialInterface*>& N
 }
 
 
-
-
-
-
-
-
 #undef LOCTEXT_NAMESPACE
diff --git a/Engine/Plugins/Runtime/MeshModelingToolset/Source/MeshModelingTools/Public/CombineMeshesTool.h b/Engine/Plugins/Runtime/MeshModelingToolset/Source/MeshModelingTools/Public/CombineMeshesTool.h
--- a/Engine/Plugins/Runtime/MeshModelingToolset/Source/MeshModelingTools/Public/CombineMeshesTool.h
+++ b/Engine/Plugins/Runtime/MeshModelingToolset/Source/MeshModelingTools/Public/CombineMeshesTool.h
@@ -7,11 +7,12 @@
 #include "InteractiveToolBuilder.h"
 #include "MeshOpPreviewHelpers.h"
 #include "DynamicMesh/DynamicMesh3.h"
+#include "BaseTools/BaseCreateFromSelectedTool.h"
 #include "PropertySets/OnAcceptProperties.h"
 #include "PropertySets/CreateMeshObjectTypeProperties.h"
 #include "CombineMeshesTool.generated.h"
 
-// predeclarations
+// Forward declarations
 struct FMeshDescription;
 
 /**
@@ -33,19 +34,8 @@ protected:
 };
 
 
-
-
-UENUM()
-enum class ECombineTargetType
-{
-	NewAsset,
-	FirstInputAsset,
-	LastInputAsset
-};
-
-
 /**
- * Standard properties
+ * Common properties
  */
 UCLASS()
 class MESHMODELINGTOOLS_API UCombineMeshesToolProperties : public UInteractiveToolPropertySet
@@ -56,23 +46,23 @@ public:
 	UPROPERTY(meta = (TransientToolProperty))
 	bool bIsDuplicateMode = false;
 
-	UPROPERTY(EditAnywhere, Category = AssetOptions, meta = (EditCondition = "bIsDuplicateMode == false", EditConditionHides, HideEditConditionToggle))
-	ECombineTargetType WriteOutputTo = ECombineTargetType::NewAsset;
+	/** Defines the object the tool output is written to. */
+	UPROPERTY(EditAnywhere, Category = OutputObject, meta = (DisplayName = "Write To",
+		EditCondition = "bIsDuplicateMode == false", EditConditionHides, HideEditConditionToggle))
+	EBaseCreateFromSelectedTargetType OutputWriteTo = EBaseCreateFromSelectedTargetType::NewObject;
 
-	/** Base name for newly-generated asset */
-	UPROPERTY(EditAnywhere, Category = AssetOptions, meta = (TransientToolProperty, EditCondition = "bIsDuplicateMode || WriteOutputTo == ECombineTargetType::NewAsset", EditConditionHides))
-	FString OutputName;
+	/** Base name of the newly generated object to which the output is written to. */
+	UPROPERTY(EditAnywhere, Category = OutputObject, meta = (TransientToolProperty, DisplayName = "Name",
+		EditCondition = "bIsDuplicateMode || OutputWriteTo == EBaseCreateFromSelectedTargetType::NewObject", EditConditionHides, NoResetToDefault))
+	FString OutputNewName;
 
-	/** Name of asset that will be updated */
-	UPROPERTY(VisibleAnywhere, Category = AssetOptions, meta = (TransientToolProperty, EditCondition = "bIsDuplicateMode == false && WriteOutputTo != ECombineTargetType::NewAsset", EditConditionHides))
-	FString OutputAsset;
+	/** Name of the existing object to which the output is written to. */
+	UPROPERTY(VisibleAnywhere, Category = OutputObject, meta = (TransientToolProperty, DisplayName = "Name",
+		EditCondition = "bIsDuplicateMode == false && OutputWriteTo != EBaseCreateFromSelectedTargetType::NewObject", EditConditionHides))
+	FString OutputExistingName;
 };
 
 
-
-
-
-
 /**
  * Simple tool to combine multiple meshes into a single mesh asset
  */
@@ -101,9 +91,9 @@ protected:
 	TObjectPtr<UCreateMeshObjectTypeProperties> OutputTypeProperties;
 
 	UPROPERTY()
-	TObjectPtr<UOnAcceptHandleSourcesProperties> HandleSourceProperties;
+	TObjectPtr<UOnAcceptHandleSourcesPropertiesBase> HandleSourceProperties;
 
-protected:
+	UPROPERTY()
 	UWorld* TargetWorld;
 
 	bool bDuplicateMode;
@@ -111,6 +101,5 @@ protected:
 	void CreateNewAsset();
 	void UpdateExistingAsset();
 
-
 	void BuildCombinedMaterialSet(TArray<UMaterialInterface*>& NewMaterialsOut, TArray<TArray<int32>>& MaterialIDRemapsOut);
 };
diff --git a/Engine/Plugins/Runtime/MeshModelingToolset/Source/MeshModelingTools/Public/DrawAndRevolveTool.h b/Engine/Plugins/Runtime/MeshModelingToolset/Source/MeshModelingTools/Public/DrawAndRevolveTool.h
--- a/Engine/Plugins/Runtime/MeshModelingToolset/Source/MeshModelingTools/Public/DrawAndRevolveTool.h
+++ b/Engine/Plugins/Runtime/MeshModelingToolset/Source/MeshModelingTools/Public/DrawAndRevolveTool.h
@@ -28,18 +28,6 @@ public:
 	virtual UInteractiveTool* BuildTool(const FToolBuilderState& SceneState) const override;
 };
 
-UENUM()
-enum class ERevolvePropertiesPathCapFillMode : uint8
-{
-	/** No caps will be generated. */
-	None = static_cast<uint8>(ERevolvePropertiesCapFillMode::None),
-	/** Caps are triangulated by placing a vertex in the center and creating a fan to the boundary. This works well if the path is convex,
-	  * but can create invalid geometry if it is concave. */
-	CenterFan = static_cast<uint8>(ERevolvePropertiesCapFillMode::CenterFan),
-	/** Caps are triangulated to maximize the minimal angle in the triangles using Delaunay triangulation. */
-	Delaunay = static_cast<uint8>(ERevolvePropertiesCapFillMode::Delaunay)
-};
-
 UCLASS()
 class MESHMODELINGTOOLS_API URevolveToolProperties : public URevolveProperties
 {
@@ -49,8 +37,8 @@ public:
 
 	/** Determines how end caps are created. This is not relevant if the end caps are not visible or if the path is not closed. */
 	UPROPERTY(EditAnywhere, Category = Revolve, AdvancedDisplay, meta = (DisplayAfter = "QuadSplitMode",
-		EditCondition = "HeightOffsetPerDegree != 0 || RevolveDegrees != 360"))
-	ERevolvePropertiesPathCapFillMode CapFillMode = ERevolvePropertiesPathCapFillMode::Delaunay;
+		EditCondition = "HeightOffsetPerDegree != 0 || RevolveDegrees != 360", ValidEnumValues = "None, CenterFan, Delaunay"))
+	ERevolvePropertiesCapFillMode CapFillMode = ERevolvePropertiesCapFillMode::Delaunay;
 
 	/** Connect the ends of an open path to the axis to add caps to the top and bottom of the revolved result.
 	  * This is not relevant for paths that are already closed. */
@@ -77,7 +65,7 @@ public:
 protected:
 	virtual ERevolvePropertiesCapFillMode GetCapFillMode() const override
 	{
-		return static_cast<ERevolvePropertiesCapFillMode>(CapFillMode);
+		return CapFillMode;
 	}
 };
 
