MshModelingToolset: Fix build error due to mative pointer usage

diff --git a/Engine/Plugins/Runtime/MeshModelingToolset/Source/MeshModelingTools/Public/AddPrimitiveTool.h b/Engine/Plugins/Runtime/MeshModelingToolset/Source/MeshModelingTools/Public/AddPrimitiveTool.h
--- a/Engine/Plugins/Runtime/MeshModelingToolset/Source/MeshModelingTools/Public/AddPrimitiveTool.h
+++ b/Engine/Plugins/Runtime/MeshModelingToolset/Source/MeshModelingTools/Public/AddPrimitiveTool.h
@@ -526,7 +526,6 @@ protected:
 	UPROPERTY()
 	FString AssetName = TEXT("GeneratedAsset");
 
-	UPROPERTY()
 	UWorld* TargetWorld;
 
 	void UpdatePreviewPosition(const FInputDeviceRay& ClickPos);
diff --git a/Engine/Plugins/Runtime/MeshModelingToolset/Source/MeshModelingTools/Public/CombineMeshesTool.h b/Engine/Plugins/Runtime/MeshModelingToolset/Source/MeshModelingTools/Public/CombineMeshesTool.h
--- a/Engine/Plugins/Runtime/MeshModelingToolset/Source/MeshModelingTools/Public/CombineMeshesTool.h
+++ b/Engine/Plugins/Runtime/MeshModelingToolset/Source/MeshModelingTools/Public/CombineMeshesTool.h
@@ -93,7 +93,6 @@ protected:
 	UPROPERTY()
 	TObjectPtr<UOnAcceptHandleSourcesPropertiesBase> HandleSourceProperties;
 
-	UPROPERTY()
 	UWorld* TargetWorld;
 
 	bool bDuplicateMode;
diff --git a/Engine/Plugins/Runtime/MeshModelingToolset/Source/MeshModelingTools/Public/DrawPolygonTool.h b/Engine/Plugins/Runtime/MeshModelingToolset/Source/MeshModelingTools/Public/DrawPolygonTool.h
--- a/Engine/Plugins/Runtime/MeshModelingToolset/Source/MeshModelingTools/Public/DrawPolygonTool.h
+++ b/Engine/Plugins/Runtime/MeshModelingToolset/Source/MeshModelingTools/Public/DrawPolygonTool.h
@@ -249,7 +249,6 @@ protected:
 	/** last vertex of polygon that is actively being updated as input device is moved */
 	FVector3d PreviewVertex;
 
-	UPROPERTY()
 	UWorld* TargetWorld;
 
 	FViewCameraState CameraState;
