MeshModelingToolset: Move ModelingOperators and ModelingOperatorsEditorOnly modules out of experimental plugin

diff --git a/Engine/Plugins/Runtime/MeshModelingToolset/Source/ModelingOperatorsEditorOnly/Private/ParameterizationOps/CalculateTangentsOp.cpp b/Engine/Plugins/Runtime/MeshModelingToolset/Source/ModelingOperatorsEditorOnly/Private/ParameterizationOps/CalculateTangentsOp.cpp
new file mode 100644
--- /dev/null
+++ b/Engine/Plugins/Runtime/MeshModelingToolset/Source/ModelingOperatorsEditorOnly/Private/ParameterizationOps/CalculateTangentsOp.cpp
@@ -0,0 +1,211 @@
+// Copyright Epic Games, Inc. All Rights Reserved.
+
+#include "ParameterizationOps/CalculateTangentsOp.h"
+#include "VectorUtil.h"
+
+#define WITH_MIKKTSPACE 1
+
+#if WITH_MIKKTSPACE
+#include "mikktspace.h"
+#endif //WITH_MIKKTSPACE
+
+#include "ExplicitUseGeometryMathTypes.h"		// using UE::Geometry::(math types)
+using namespace UE::Geometry;
+
+void FCalculateTangentsOp::CalculateResult(FProgressCancel* Progress)
+{
+	if (Progress && Progress->Cancelled())
+	{
+		return;
+	}
+
+	TUniquePtr<FMeshTangentsd> Tangents = MakeUnique<FMeshTangentsd>(SourceMesh.Get());
+
+	if (SourceMesh->HasAttributes() == false
+		|| SourceMesh->Attributes()->NumUVLayers() < 1
+		|| SourceMesh->Attributes()->NumNormalLayers() == 0)
+	{
+		bNoAttributesError = true;
+		Tangents->InitializeTriVertexTangents(true);
+		SetResult(MoveTemp(Tangents));
+		return;
+	}
+
+
+	switch (CalculationMethod)
+	{
+
+	default:
+	case EMeshTangentsType::FastMikkTSpace:
+		CalculateStandard(Progress, Tangents);
+		break;
+
+	case EMeshTangentsType::PerTriangle:
+		CalculateStandard(Progress, Tangents);
+		break;
+
+	case EMeshTangentsType::MikkTSpace:
+#if WITH_MIKKTSPACE
+		CalculateMikkT(Progress, Tangents);
+#else
+		CalculateStandard(Progress, Tangents);
+#endif
+		break;
+
+
+	case EMeshTangentsType::CopyExisting:
+		CopyFromSource(Progress, Tangents);
+		break;
+	}
+
+	SetResult(MoveTemp(Tangents));
+}
+
+
+void FCalculateTangentsOp::CalculateStandard(FProgressCancel* Progress, TUniquePtr<FMeshTangentsd>& Tangents)
+{
+	const FDynamicMeshNormalOverlay* NormalOverlay = SourceMesh->Attributes()->PrimaryNormals();
+	const FDynamicMeshUVOverlay* UVOverlay = SourceMesh->Attributes()->GetUVLayer(0);
+
+	FComputeTangentsOptions Options;
+	Options.bAveraged = (CalculationMethod != EMeshTangentsType::PerTriangle);
+
+	Tangents->ComputeTriVertexTangents(NormalOverlay, UVOverlay, Options);
+}
+
+
+void FCalculateTangentsOp::CopyFromSource(FProgressCancel* Progress, TUniquePtr<FMeshTangentsd>& Tangents)
+{
+	int32 NumTangents = SourceTangents->GetTangents().Num();
+	if ( ! ensure(NumTangents == SourceMesh->MaxTriangleID() * 3) )
+	{
+		bNoAttributesError = true;
+		Tangents->InitializeTriVertexTangents(true);
+		return;
+	}
+
+	Tangents->CopyTriVertexTangents(*SourceTangents);
+}
+
+
+
+#if WITH_MIKKTSPACE
+namespace DynamicMeshMikkTInterface
+{
+	struct FDynamicMeshInfo
+	{
+		const FDynamicMesh3* Mesh;
+		const FDynamicMeshNormalOverlay* NormalOverlay;
+		const FDynamicMeshUVOverlay* UVOverlay;
+		FMeshTangentsd* TangentsOut;
+	};
+
+
+	int MikkGetNumFaces(const SMikkTSpaceContext* Context)
+	{
+		const FDynamicMeshInfo* MeshInfo = (const FDynamicMeshInfo*)(Context->m_pUserData);
+		return MeshInfo->Mesh->TriangleCount();
+	}
+
+	int MikkGetNumVertsOfFace(const SMikkTSpaceContext* Context, const int FaceIdx)
+	{
+		const FDynamicMeshInfo* MeshInfo = (const FDynamicMeshInfo*)(Context->m_pUserData);
+		return MeshInfo->Mesh->IsTriangle(FaceIdx) ? 3 : 0;
+	}
+
+	void MikkGetPosition(const SMikkTSpaceContext* Context, float Position[3], const int FaceIdx, const int VertIdx)
+	{
+		const FDynamicMeshInfo* MeshInfo = (const FDynamicMeshInfo*)(Context->m_pUserData);
+		FIndex3i Triangle = MeshInfo->Mesh->GetTriangle(FaceIdx);
+		FVector3d VertexPos = MeshInfo->Mesh->GetVertex(Triangle[VertIdx]);
+		Position[0] = (float)VertexPos.X;
+		Position[1] = (float)VertexPos.Y;
+		Position[2] = (float)VertexPos.Z;
+	}
+
+	void MikkGetNormal(const SMikkTSpaceContext* Context, float Normal[3], const int FaceIdx, const int VertIdx)
+	{
+		const FDynamicMeshInfo* MeshInfo = (const FDynamicMeshInfo*)(Context->m_pUserData);
+		FIndex3i NormTriangle = MeshInfo->NormalOverlay->GetTriangle(FaceIdx);
+		FVector3f Normalf;
+		MeshInfo->NormalOverlay->GetElement(NormTriangle[VertIdx], Normalf);
+		Normal[0] = Normalf.X;
+		Normal[1] = Normalf.Y;
+		Normal[2] = Normalf.Z;
+	}
+
+	void MikkGetTexCoord(const SMikkTSpaceContext* Context, float UV[2], const int FaceIdx, const int VertIdx)
+	{
+		const FDynamicMeshInfo* MeshInfo = (const FDynamicMeshInfo*)(Context->m_pUserData);
+		FIndex3i UVTriangle = MeshInfo->UVOverlay->GetTriangle(FaceIdx);
+		FVector2f UVf;
+		MeshInfo->UVOverlay->GetElement(UVTriangle[VertIdx], UVf);
+		UV[0] = UVf.X;
+		UV[1] = UVf.Y;
+	}
+
+	void MikkSetTSpaceBasic(const SMikkTSpaceContext* Context, const float Tangent[3], const float BitangentSign, const int FaceIdx, const int VertIdx)
+	{
+		FDynamicMeshInfo* MeshInfo = (FDynamicMeshInfo*)(Context->m_pUserData);
+		FVector3d Tangentd(Tangent[0], Tangent[1], Tangent[2]);
+
+		FIndex3i NormTriangle = MeshInfo->NormalOverlay->GetTriangle(FaceIdx);
+		FVector3d Normald = (FVector3d)MeshInfo->NormalOverlay->GetElement(NormTriangle[VertIdx]);
+		FVector3d Bitangentd = VectorUtil::Bitangent(Normald, Tangentd, -(double)BitangentSign);
+
+		MeshInfo->TangentsOut->SetPerTriangleTangent(FaceIdx, VertIdx, Tangentd, Bitangentd);
+	}
+
+
+	void MikkSetTSpac(const SMikkTSpaceContext* Context, const float Tangent[3], const float BiTangent[3],
+		const float MagS, const float MagT, const tbool bIsOrientationPreserving,
+		const int FaceIdx, const int VertIdx)
+	{
+		FDynamicMeshInfo* MeshInfo = (FDynamicMeshInfo*)(Context->m_pUserData);
+		FVector3d Tangentd(Tangent[0], Tangent[1], Tangent[2]);
+		FVector3d Bitangentd(BiTangent[0], BiTangent[1], BiTangent[2]);
+		MeshInfo->TangentsOut->SetPerTriangleTangent(FaceIdx, VertIdx, Tangentd, Bitangentd);
+	}
+
+}
+
+
+
+void FCalculateTangentsOp::CalculateMikkT(FProgressCancel* Progress, TUniquePtr<FMeshTangentsd>& Tangents)
+{
+	Tangents->InitializeTriVertexTangents(true);
+
+	DynamicMeshMikkTInterface::FDynamicMeshInfo MeshInfo;
+	MeshInfo.Mesh = SourceMesh.Get();
+	MeshInfo.NormalOverlay = SourceMesh->Attributes()->PrimaryNormals();
+	MeshInfo.UVOverlay = SourceMesh->Attributes()->GetUVLayer(0);;
+	MeshInfo.TangentsOut = Tangents.Get();
+
+	SMikkTSpaceInterface MikkTInterface;
+	MikkTInterface.m_getNormal = DynamicMeshMikkTInterface::MikkGetNormal;
+	MikkTInterface.m_getNumFaces = DynamicMeshMikkTInterface::MikkGetNumFaces;
+	MikkTInterface.m_getNumVerticesOfFace = DynamicMeshMikkTInterface::MikkGetNumVertsOfFace;
+	MikkTInterface.m_getPosition = DynamicMeshMikkTInterface::MikkGetPosition;
+	MikkTInterface.m_getTexCoord = DynamicMeshMikkTInterface::MikkGetTexCoord;
+	MikkTInterface.m_setTSpaceBasic = DynamicMeshMikkTInterface::MikkSetTSpaceBasic;
+	MikkTInterface.m_setTSpace = nullptr;
+
+	SMikkTSpaceContext MikkTContext;
+	MikkTContext.m_pInterface = &MikkTInterface;
+	MikkTContext.m_pUserData = (void*)(&MeshInfo);
+	MikkTContext.m_bIgnoreDegenerates = false;
+	//MikkTContext.m_bIgnoreDegenerates = true;
+	//MikkTContext.m_bIgnoreDegenerates = bIgnoreDegenerateTriangles;
+
+	// execute mikkt
+	genTangSpaceDefault(&MikkTContext);
+}
+
+#else
+
+void FCalculateTangentsOp::CalculateMikkT(FProgressCancel* Progress, TUniquePtr<FMeshTangentsd>& Tangents)
+{
+	check(false);
+}
+
+#endif
\ No newline at end of file
diff --git a/Engine/Plugins/Runtime/MeshModelingToolset/Source/ModelingOperatorsEditorOnly/Public/CompositionOps/VoxelBooleanMeshesOp.h b/Engine/Plugins/Runtime/MeshModelingToolset/Source/ModelingOperatorsEditorOnly/Public/CompositionOps/VoxelBooleanMeshesOp.h
new file mode 100644
--- /dev/null
+++ b/Engine/Plugins/Runtime/MeshModelingToolset/Source/ModelingOperatorsEditorOnly/Public/CompositionOps/VoxelBooleanMeshesOp.h
@@ -0,0 +1,60 @@
+// Copyright Epic Games, Inc. All Rights Reserved.
+
+#pragma once
+
+#include "CoreMinimal.h"
+#include "Util/ProgressCancel.h"
+#include "ModelingOperators.h"
+
+struct FMeshDescription;
+
+namespace UE
+{
+namespace Geometry
+{
+
+
+class MODELINGOPERATORSEDITORONLY_API FVoxelBooleanMeshesOp : public FDynamicMeshOperator
+{
+public:
+	virtual ~FVoxelBooleanMeshesOp() {}
+
+	enum class EBooleanOperation
+	{
+		DifferenceAB = 0,
+		DifferenceBA = 1,
+		Intersect = 2,
+		Union = 3
+	};
+
+	struct FInputMesh
+	{
+		const FMeshDescription* Mesh;
+		FTransform Transform;
+	};
+
+	// inputs
+	TArray<FInputMesh> InputMeshArray;
+	int32 VoxelCount = 128;
+	double VoxelSizeD = 1.0;
+	double AdaptivityD = 0;
+	double IsoSurfaceD = 0;
+	bool   bAutoSimplify = false;
+	EBooleanOperation Operation = EBooleanOperation::DifferenceBA;
+
+	//
+	// FDynamicMeshOperator implementation
+	// 
+
+	virtual void CalculateResult(FProgressCancel* Progress) override;
+
+private:
+
+	// compute the voxel size based on the voxel count and the input geometry
+	float ComputeVoxelSize() const;
+
+};
+
+} // end namespace UE::Geometry
+} // end namespace UE
+
diff --git a/Engine/Plugins/Runtime/MeshModelingToolset/Source/ModelingOperatorsEditorOnly/Public/CompositionOps/VoxelMergeMeshesOp.h b/Engine/Plugins/Runtime/MeshModelingToolset/Source/ModelingOperatorsEditorOnly/Public/CompositionOps/VoxelMergeMeshesOp.h
new file mode 100644
--- /dev/null
+++ b/Engine/Plugins/Runtime/MeshModelingToolset/Source/ModelingOperatorsEditorOnly/Public/CompositionOps/VoxelMergeMeshesOp.h
@@ -0,0 +1,53 @@
+// Copyright Epic Games, Inc. All Rights Reserved.
+
+#pragma once
+
+#include "CoreMinimal.h"
+#include "Util/ProgressCancel.h"
+#include "ModelingOperators.h"
+
+struct FMeshDescription;
+
+namespace UE
+{
+namespace Geometry
+{
+
+
+class MODELINGOPERATORSEDITORONLY_API FVoxelMergeMeshesOp : public FDynamicMeshOperator
+{
+public:
+	virtual ~FVoxelMergeMeshesOp() {}
+
+	struct FInputMesh
+	{
+		const FMeshDescription* Mesh;
+		FTransform Transform;
+	};
+
+	// inputs
+	TArray<FInputMesh> InputMeshArray;
+	int32 VoxelCount = 128;
+	double VoxelSizeD = 1.0;
+	double AdaptivityD = 0;
+	double IsoSurfaceD = 0;
+	bool bAutoSimplify = false;
+
+	double FastCollapseVoxelSizeMultipler = 1.2;
+	int NumFastCollapsePasses = 10;
+
+	//
+	// FDynamicMeshOperator implementation
+	// 
+
+	virtual void CalculateResult(FProgressCancel* Progress) override;
+
+private:
+
+	// compute the voxel size based on the voxel count and the input geometry
+	float ComputeVoxelSize() const;
+};
+
+
+} // end namespace UE::Geometry
+} // end namespace UE
diff --git a/Engine/Plugins/Runtime/MeshModelingToolset/Source/ModelingOperatorsEditorOnly/Public/ParameterizationOps/CalculateTangentsOp.h b/Engine/Plugins/Runtime/MeshModelingToolset/Source/ModelingOperatorsEditorOnly/Public/ParameterizationOps/CalculateTangentsOp.h
new file mode 100644
--- /dev/null
+++ b/Engine/Plugins/Runtime/MeshModelingToolset/Source/ModelingOperatorsEditorOnly/Public/ParameterizationOps/CalculateTangentsOp.h
@@ -0,0 +1,61 @@
+// Copyright Epic Games, Inc. All Rights Reserved.
+
+#pragma once
+
+#include "CoreMinimal.h"
+#include "Util/ProgressCancel.h"
+#include "ModelingOperators.h"
+#include "DynamicMesh/MeshTangents.h"
+#include "CalculateTangentsOp.generated.h"
+
+
+
+UENUM()
+enum class EMeshTangentsType : uint8
+{
+	/** Standard MikkTSpace Tangent Calculation */
+	MikkTSpace = 0,
+	/** MikkT-like Blended Per-Triangle Tangents, with Blending based on existing Mesh/UV/Normal Topology */
+	FastMikkTSpace = 1,
+	/** Project per-triangle Tangents projected onto Normals */
+	PerTriangle = 2,
+	/** Use Source Mesh Tangents */
+	CopyExisting = 3
+};
+
+namespace UE
+{
+namespace Geometry
+{
+
+
+class MODELINGOPERATORSEDITORONLY_API FCalculateTangentsOp : public TGenericDataOperator<UE::Geometry::FMeshTangentsd>
+{
+public:
+	virtual ~FCalculateTangentsOp() {}
+
+	// inputs
+	TSharedPtr<FDynamicMesh3, ESPMode::ThreadSafe> SourceMesh;
+	TSharedPtr<UE::Geometry::FMeshTangentsf, ESPMode::ThreadSafe> SourceTangents;
+
+	// parameters
+	EMeshTangentsType CalculationMethod;
+
+	// error flags
+	bool bNoAttributesError = false;
+
+	//
+	// TGenericDataOperator implementation
+	// 
+
+	virtual void CalculateResult(FProgressCancel* Progress) override;
+
+
+	virtual void CalculateStandard(FProgressCancel* Progress, TUniquePtr<UE::Geometry::FMeshTangentsd>& Tangents);
+	virtual void CalculateMikkT(FProgressCancel* Progress, TUniquePtr<UE::Geometry::FMeshTangentsd>& Tangents);
+	virtual void CopyFromSource(FProgressCancel* Progress, TUniquePtr<UE::Geometry::FMeshTangentsd>& Tangents);
+};
+
+} // end namespace UE::Geometry
+} // end namespace UE
+
