MeshModelingToolset: Cleanup Boolean tool

diff --git a/Engine/Plugins/Editor/ModelingToolsEditorMode/Source/ModelingToolsEditorMode/Private/ModelingToolsManagerActions.cpp b/Engine/Plugins/Editor/ModelingToolsEditorMode/Source/ModelingToolsEditorMode/Private/ModelingToolsManagerActions.cpp
--- a/Engine/Plugins/Editor/ModelingToolsEditorMode/Source/ModelingToolsEditorMode/Private/ModelingToolsManagerActions.cpp
+++ b/Engine/Plugins/Editor/ModelingToolsEditorMode/Source/ModelingToolsEditorMode/Private/ModelingToolsManagerActions.cpp
@@ -62,14 +62,19 @@ void FModelingToolsManagerCommands::RegisterCommands()
 	REGISTER_MODELING_TOOL_COMMAND(BeginCombineMeshesTool, "MshMrg", "Merge multiple meshes to create new objects");
 	REGISTER_MODELING_TOOL_COMMAND(BeginDuplicateMeshesTool, "MshDup", "Duplicate single meshes to create new objects");
 
+	// PolyModel
+	REGISTER_MODELING_TOOL_COMMAND(BeginPolyEditTool, "PolyEd", "Edit meshes via PolyGroups");
+	REGISTER_MODELING_TOOL_COMMAND(BeginPolyDeformTool, "PolyDef", "Deform meshes via PolyGroups");
+	REGISTER_MODELING_TOOL_COMMAND(BeginCubeGridTool, "CubeGr", "Create block out meshes using a repositionable grid");
+	REGISTER_MODELING_TOOL_COMMAND(BeginMeshBooleanTool, "MshBool", "Apply Boolean operations to mesh pairs");
+	REGISTER_MODELING_TOOL_COMMAND(BeginCutMeshWithMeshTool, "MshCut", "Split one mesh into parts using a second mesh");
+	REGISTER_MODELING_TOOL_COMMAND(BeginSubdividePolyTool, "SubDiv", "Subdivide Mesh via PolyGroups or Triangles");
+	
 	REGISTER_MODELING_TOOL_COMMAND(BeginAddPatchTool, "Patch", "Add Patch");
 	REGISTER_MODELING_TOOL_COMMAND(BeginShapeSprayTool, "Spray", "Shape Spray");
 
 	REGISTER_MODELING_TOOL_COMMAND(BeginSculptMeshTool, "VSclpt", "Vertex Sculpting");
-	REGISTER_MODELING_TOOL_COMMAND(BeginCubeGridTool, "CubeGr", "Create block-like geometry from a repositionable grid");
-	REGISTER_MODELING_TOOL_COMMAND(BeginPolyEditTool, "PolyEd", "Edit Mesh via PolyGroups");
 	REGISTER_MODELING_TOOL_COMMAND(BeginTriEditTool, "TriEd", "Edit Mesh via Triangles");
-	REGISTER_MODELING_TOOL_COMMAND(BeginPolyDeformTool, "PolyDef", "Deform Mesh via PolyGroups");
 	REGISTER_MODELING_TOOL_COMMAND(BeginSmoothMeshTool, "Smooth", "Smooth Mesh surface");
 	REGISTER_MODELING_TOOL_COMMAND(BeginOffsetMeshTool, "Offset", "Offset Mesh surface");
 	REGISTER_MODELING_TOOL_COMMAND(BeginDisplaceMeshTool, "Displce", "Displace Mesh surface with optional subdivision");
@@ -98,9 +103,7 @@ void FModelingToolsManagerCommands::RegisterCommands()
 	REGISTER_MODELING_TOOL_COMMAND(BeginVoxelBlendTool, "VoxBlnd", "Blend Selected Meshes (Voxel-Based)");
 	REGISTER_MODELING_TOOL_COMMAND(BeginVoxelMorphologyTool, "VoxMrph", "Offset/Inset Selected Meshes (Voxel-Based)");
 	REGISTER_MODELING_TOOL_COMMAND(BeginSelfUnionTool, "Merge", "Self-Union Selected Meshes to resolve Self-Intersections");
-	REGISTER_MODELING_TOOL_COMMAND(BeginMeshBooleanTool, "Boolean", "Apply Boolean/CSG operation to selected Mesh pair");
 	REGISTER_MODELING_TOOL_COMMAND(BeginMeshTrimTool, "Trim", "Trim/Cut selected mesh with second mesh");
-	REGISTER_MODELING_TOOL_COMMAND(BeginCutMeshWithMeshTool, "MshCut", "Split selected mesh into parts using second mesh");
 	REGISTER_MODELING_TOOL_COMMAND(BeginBspConversionTool, "BSPConv", "Convert BSP to StaticMesh Asset");
 	REGISTER_MODELING_TOOL_COMMAND(BeginMeshToVolumeTool, "Msh2Vol", "Convert Mesh to Volume");
 	REGISTER_MODELING_TOOL_COMMAND(BeginVolumeToMeshTool, "Vol2Msh", "Convert Volume to new Mesh Asset");
@@ -133,8 +136,6 @@ void FModelingToolsManagerCommands::RegisterCommands()
 	// why are these ::Button ?
 	UI_COMMAND(BeginSkinWeightsPaintTool, "SkinWts", "Start the Paint Skin Weights Tool", EUserInterfaceActionType::Button, FInputChord());
 	UI_COMMAND(BeginSkinWeightsBindingTool, "SkinBind", "Start the Skin Weights Binding Tool", EUserInterfaceActionType::Button, FInputChord());
-
-	REGISTER_MODELING_TOOL_COMMAND(BeginSubdividePolyTool, "SubD", "Subdivide Mesh Polygroups or Triangles");
 	
 	REGISTER_MODELING_TOOL_COMMAND(BeginLODManagerTool, "LODMgr", "Static Mesh Asset LOD Manager");
 	REGISTER_MODELING_TOOL_COMMAND(BeginGenerateStaticMeshLODAssetTool, "AutoLOD", "Generate Static Mesh LOD Asset");
diff --git a/Engine/Plugins/Runtime/MeshModelingToolset/Source/MeshModelingTools/Private/CSGMeshesTool.cpp b/Engine/Plugins/Runtime/MeshModelingToolset/Source/MeshModelingTools/Private/CSGMeshesTool.cpp
--- a/Engine/Plugins/Runtime/MeshModelingToolset/Source/MeshModelingTools/Private/CSGMeshesTool.cpp
+++ b/Engine/Plugins/Runtime/MeshModelingToolset/Source/MeshModelingTools/Private/CSGMeshesTool.cpp
@@ -1,20 +1,17 @@
 // Copyright Epic Games, Inc. All Rights Reserved.
 
 #include "CSGMeshesTool.h"
-#include "CompositionOps/BooleanMeshesOp.h"
-
+#include "MeshDescriptionToDynamicMesh.h"
 #include "ToolSetupUtil.h"
 #include "BaseGizmos/TransformProxy.h"
-
+#include "CompositionOps/BooleanMeshesOp.h"
 #include "DynamicMesh/DynamicMesh3.h"
 
-#include "MeshDescriptionToDynamicMesh.h"
-
-#include "ExplicitUseGeometryMathTypes.h"		// using UE::Geometry::(math types)
 using namespace UE::Geometry;
 
 #define LOCTEXT_NAMESPACE "UCSGMeshesTool"
 
+
 void UCSGMeshesTool::EnableTrimMode()
 {
 	check(OriginalDynamicMeshes.Num() == 0);		// must not have been initialized!
@@ -69,18 +66,18 @@ void UCSGMeshesTool::SetupProperties()
 		{
 			UpdatePreviewsVisibility();
 		});
-		CSGProperties->WatchProperty(CSGProperties->ColorOfSubtractedMesh, [this](FLinearColor)
+		CSGProperties->WatchProperty(CSGProperties->SubtractedMeshColor, [this](FLinearColor)
 		{
 			UpdatePreviewsMaterial();
 		});
-		CSGProperties->WatchProperty(CSGProperties->OpacityOfSubtractedMesh, [this](float)
+		CSGProperties->WatchProperty(CSGProperties->SubtractedMeshOpacity, [this](float)
 		{
 			UpdatePreviewsMaterial();
 		});
 
-		SetToolDisplayName(LOCTEXT("CSGMeshesToolName", "Boolean"));
+		SetToolDisplayName(LOCTEXT("CSGMeshesToolName", "Mesh Boolean"));
 		GetToolManager()->DisplayMessage(
-			LOCTEXT("OnStartTool", "Compute CSG Booleans on the input meshes. Use the transform gizmos to tweak the positions of the input objects (can help to resolve errors/failures)"),
+			LOCTEXT("OnStartTool", "Perform Boolean operations on the input meshes; any interior faces will be removed. Use the transform gizmos to modify the position and orientation of the input objects."),
 			EToolMessageLevel::UserNotification);
 	}
 }
@@ -100,8 +97,8 @@ void UCSGMeshesTool::UpdatePreviewsMaterial()
 	}
 	else
 	{
-		Color = CSGProperties->ColorOfSubtractedMesh;
-		Opacity = CSGProperties->OpacityOfSubtractedMesh;
+		Color = CSGProperties->SubtractedMeshColor;
+		Opacity = CSGProperties->SubtractedMeshOpacity;
 	}
 	PreviewsGhostMaterial->SetVectorParameterValue(TEXT("Color"), Color);
 	PreviewsGhostMaterial->SetScalarParameterValue(TEXT("Opacity"), Opacity);
@@ -177,7 +174,7 @@ void UCSGMeshesTool::ConvertInputsAndSetPreviewMaterials(bool bSetPreviewMesh)
 	TMap<UMaterialInterface*, int> KnownMaterials;
 	TArray<TArray<int>> MaterialRemap; MaterialRemap.SetNum(Targets.Num());
 
-	if (bTrimMode || !CSGProperties->bOnlyUseFirstMeshMaterials)
+	if (bTrimMode || !CSGProperties->bUseFirstMeshMaterials)
 	{
 		for (int ComponentIdx = 0; ComponentIdx < Targets.Num(); ComponentIdx++)
 		{
@@ -278,7 +275,7 @@ void UCSGMeshesTool::UpdateVisualization()
 	FVector3d A, B;
 
 	DrawnLineSet->Clear();
-	if (!bTrimMode && CSGProperties->bShowNewBoundaryEdges)
+	if (!bTrimMode && CSGProperties->bShowNewBoundaries)
 	{
 		for (int EID : CreatedBoundaryEdges)
 		{
@@ -304,8 +301,8 @@ TUniquePtr<FDynamicMeshOperator> UCSGMeshesTool::MakeNewOperator()
 	else
 	{
 		BooleanOp->CSGOperation = CSGProperties->Operation;
-		BooleanOp->bAttemptFixHoles = CSGProperties->bAttemptFixHoles;
-		BooleanOp->bTryCollapseExtraEdges = CSGProperties->bCollapseExtraEdges;
+		BooleanOp->bAttemptFixHoles = CSGProperties->bTryFixHoles;
+		BooleanOp->bTryCollapseExtraEdges = CSGProperties->bTryCollapseEdges;
 	}
 
 	check(OriginalDynamicMeshes.Num() == 2);
@@ -325,7 +322,7 @@ TUniquePtr<FDynamicMeshOperator> UCSGMeshesTool::MakeNewOperator()
 
 void UCSGMeshesTool::OnPropertyModified(UObject* PropertySet, FProperty* Property)
 {
-	if (Property && (Property->GetFName() == GET_MEMBER_NAME_CHECKED(UCSGMeshesToolProperties, bOnlyUseFirstMeshMaterials)))
+	if (Property && (Property->GetFName() == GET_MEMBER_NAME_CHECKED(UCSGMeshesToolProperties, bUseFirstMeshMaterials)))
 	{
 		if (!AreAllTargetsValid())
 		{
@@ -335,7 +332,7 @@ void UCSGMeshesTool::OnPropertyModified(UObject* PropertySet, FProperty* Propert
 		ConvertInputsAndSetPreviewMaterials(false);
 		Preview->InvalidateResult();
 	}
-	else if (Property && (Property->GetFName() == GET_MEMBER_NAME_CHECKED(UCSGMeshesToolProperties, bShowNewBoundaryEdges)))
+	else if (Property && (Property->GetFName() == GET_MEMBER_NAME_CHECKED(UCSGMeshesToolProperties, bShowNewBoundaries)))
 	{
 		GetToolManager()->PostInvalidation();
 		UpdateVisualization();
diff --git a/Engine/Plugins/Runtime/MeshModelingToolset/Source/MeshModelingTools/Public/CSGMeshesTool.h b/Engine/Plugins/Runtime/MeshModelingToolset/Source/MeshModelingTools/Public/CSGMeshesTool.h
--- a/Engine/Plugins/Runtime/MeshModelingToolset/Source/MeshModelingTools/Public/CSGMeshesTool.h
+++ b/Engine/Plugins/Runtime/MeshModelingToolset/Source/MeshModelingTools/Public/CSGMeshesTool.h
@@ -3,16 +3,14 @@
 #pragma once
 
 #include "CoreMinimal.h"
-#include "UObject/NoExportTypes.h"
-#include "InteractiveToolBuilder.h"
-#include "Drawing/LineSetComponent.h"
 #include "BaseTools/BaseCreateFromSelectedTool.h"
-
 #include "CompositionOps/BooleanMeshesOp.h"
+#include "Drawing/LineSetComponent.h"
+#include "UObject/NoExportTypes.h"
 
 #include "CSGMeshesTool.generated.h"
 
-// predeclarations
+// Forward declarations
 PREDECLARE_USE_GEOMETRY_CLASS(FDynamicMesh3);
 
 
@@ -24,39 +22,39 @@ class MESHMODELINGTOOLS_API UCSGMeshesToolProperties : public UInteractiveToolPr
 {
 	GENERATED_BODY()
 public:
-	/** The type of operation */
-	UPROPERTY(EditAnywhere, Category = Operation)
+	/** Type of Boolean operation */
+	UPROPERTY(EditAnywhere, Category = Boolean)
 	ECSGOperation Operation = ECSGOperation::DifferenceAB;
 
-	/** Show boundary edges created by the CSG operation (often due to numerical error) */
-	UPROPERTY(EditAnywhere, Category = Operation)
-	bool bShowNewBoundaryEdges = true;
-
-	/** Automatically attempt to fill any holes left by CSG (e.g. due to numerical errors) */
-	UPROPERTY(EditAnywhere, Category = Operation)
-	bool bAttemptFixHoles = false;
+	/** Try to fill holes created by the Boolean operation, e.g. due to numerical errors */
+	UPROPERTY(EditAnywhere, Category = Boolean, AdvancedDisplay)
+	bool bTryFixHoles = false;
 
 	/** Try to collapse extra edges created by the Boolean operation */
-	UPROPERTY(EditAnywhere, Category = Operation)
-	bool bCollapseExtraEdges = true;
+	UPROPERTY(EditAnywhere, Category = Boolean, AdvancedDisplay)
+	bool bTryCollapseEdges = true;
 
-	/** Whether to show a translucent version of the subtracted mesh, to help visualize what is being removed */
-	UPROPERTY(EditAnywhere, Category = Display, meta = (EditCondition = "Operation == ECSGOperation::DifferenceAB || Operation == ECSGOperation::DifferenceBA", EditConditionHides))
+	/** Show boundary edges created by the Boolean operation, which might happen due to numerical errors */
+	UPROPERTY(EditAnywhere, Category = Display)
+	bool bShowNewBoundaries = true;
+
+	/** Show a translucent version of the subtracted mesh, to help visualize geometry that is being removed */
+	UPROPERTY(EditAnywhere, Category = Display, meta = (EditCondition = "Operation == ECSGOperation::DifferenceAB || Operation == ECSGOperation::DifferenceBA"))
 	bool bShowSubtractedMesh = true;
 	
-	/** Opacity of translucent version of the selected mesh */
-	UPROPERTY(EditAnywhere, Category = Display, meta = (ClampMin = "0", ClampMax = "1", 
-		EditCondition = "bShowSubtractedMesh && Operation == ECSGOperation::DifferenceAB || bShowSubtractedMesh && Operation == ECSGOperation::DifferenceBA", EditConditionHides))
-	float OpacityOfSubtractedMesh = .2;
+	/** Opacity of the translucent subtracted mesh */
+	UPROPERTY(EditAnywhere, Category = Display, meta = (DisplayName = "Opacity Subtracted Mesh", ClampMin = "0", ClampMax = "1",
+		EditCondition = "bShowSubtractedMesh && Operation == ECSGOperation::DifferenceAB || bShowSubtractedMesh && Operation == ECSGOperation::DifferenceBA"))
+	float SubtractedMeshOpacity = 0.2f;
 	
-	/** Color of translucent version of the subtracted mesh */
-	UPROPERTY(EditAnywhere, Category = Display, meta = (
-		EditCondition = "bShowSubtractedMesh && Operation == ECSGOperation::DifferenceAB || bShowSubtractedMesh && Operation == ECSGOperation::DifferenceBA", EditConditionHides), AdvancedDisplay)
-	FLinearColor ColorOfSubtractedMesh = FLinearColor::Black;
+	/** Color of the translucent subtracted mesh */
+	UPROPERTY(EditAnywhere, Category = Display, meta = (DisplayName = "Color Subtracted Mesh", HideAlphaChannel, ClampMin = "0", ClampMax = "1",
+		EditCondition = "bShowSubtractedMesh && Operation == ECSGOperation::DifferenceAB || bShowSubtractedMesh && Operation == ECSGOperation::DifferenceBA"))
+	FLinearColor SubtractedMeshColor = FLinearColor::Black;
 
-	/** If true, only the first mesh will keep its materials assignments; all other triangles will be assigned material 0 */
+	/** If true, only the first mesh will keep its material assignments, and all other faces will have the first material assigned */
 	UPROPERTY(EditAnywhere, Category = Materials)
-	bool bOnlyUseFirstMeshMaterials = false;
+	bool bUseFirstMeshMaterials = false;
 };
 
 
@@ -91,7 +89,6 @@ public:
 };
 
 
-
 /**
  * Simple Mesh Plane Cutting Tool
  */
@@ -124,8 +121,6 @@ protected:
 	// IDynamicMeshOperatorFactory API
 	virtual TUniquePtr<UE::Geometry::FDynamicMeshOperator> MakeNewOperator() override;
 
-protected:
-
 	void UpdateVisualization();
 
 	UPROPERTY()
@@ -161,8 +156,6 @@ protected:
 };
 
 
-
-
 UCLASS()
 class MESHMODELINGTOOLS_API UCSGMeshesToolBuilder : public UBaseCreateFromSelectedToolBuilder
 {
@@ -192,6 +185,3 @@ public:
 		return Tool;
 	}
 };
-
-
-
diff --git a/Engine/Plugins/Runtime/MeshModelingToolset/Source/ModelingComponents/Private/BaseTools/BaseCreateFromSelectedTool.cpp b/Engine/Plugins/Runtime/MeshModelingToolset/Source/ModelingComponents/Private/BaseTools/BaseCreateFromSelectedTool.cpp
--- a/Engine/Plugins/Runtime/MeshModelingToolset/Source/ModelingComponents/Private/BaseTools/BaseCreateFromSelectedTool.cpp
+++ b/Engine/Plugins/Runtime/MeshModelingToolset/Source/ModelingComponents/Private/BaseTools/BaseCreateFromSelectedTool.cpp
@@ -166,7 +166,7 @@ void UBaseCreateFromSelectedTool::UpdateGizmoVisibility()
 	for (int32 GizmoIndex = 0; GizmoIndex < TransformGizmos.Num(); GizmoIndex++)
 	{
 		UCombinedTransformGizmo* Gizmo = TransformGizmos[GizmoIndex];
-		Gizmo->SetVisibility(TransformProperties->bShowTransformUI && GizmoIndex != GetHiddenGizmoIndex());
+		Gizmo->SetVisibility(TransformProperties->bShowTransformGizmo && GizmoIndex != GetHiddenGizmoIndex());
 	}
 }
 
@@ -301,7 +301,7 @@ FString UBaseCreateFromSelectedTool::PrefixWithSourceNameIfSingleSelection(const
 
 void UBaseCreateFromSelectedTool::OnPropertyModified(UObject* PropertySet, FProperty* Property)
 {
-	if (Property && (Property->GetFName() == GET_MEMBER_NAME_CHECKED(UTransformInputsToolProperties, bShowTransformUI)))
+	if (Property && (Property->GetFName() == GET_MEMBER_NAME_CHECKED(UTransformInputsToolProperties, bShowTransformGizmo)))
 	{
 		UpdateGizmoVisibility();
 	}
diff --git a/Engine/Plugins/Runtime/MeshModelingToolset/Source/ModelingComponents/Public/BaseTools/BaseCreateFromSelectedTool.h b/Engine/Plugins/Runtime/MeshModelingToolset/Source/ModelingComponents/Public/BaseTools/BaseCreateFromSelectedTool.h
--- a/Engine/Plugins/Runtime/MeshModelingToolset/Source/ModelingComponents/Public/BaseTools/BaseCreateFromSelectedTool.h
+++ b/Engine/Plugins/Runtime/MeshModelingToolset/Source/ModelingComponents/Public/BaseTools/BaseCreateFromSelectedTool.h
@@ -84,9 +84,9 @@ class MODELINGCOMPONENTS_API UTransformInputsToolProperties : public UInteractiv
 	GENERATED_BODY()
 
 public:
-	/** Show UI to allow changing translation, rotation and scale of input meshes */
-	UPROPERTY(EditAnywhere, Category = Transform)
-	bool bShowTransformUI = true;
+	/** Show transform gizmo in the viewport to allow changing translation, rotation and scale of input meshes. */
+	UPROPERTY(EditAnywhere, Category = Transform, meta = (DisplayName = "Show Gizmo"))
+	bool bShowTransformGizmo = true;
 };
 
 
diff --git a/Engine/Plugins/Runtime/MeshModelingToolset/Source/ModelingOperators/Public/CompositionOps/BooleanMeshesOp.h b/Engine/Plugins/Runtime/MeshModelingToolset/Source/ModelingOperators/Public/CompositionOps/BooleanMeshesOp.h
--- a/Engine/Plugins/Runtime/MeshModelingToolset/Source/ModelingOperators/Public/CompositionOps/BooleanMeshesOp.h
+++ b/Engine/Plugins/Runtime/MeshModelingToolset/Source/ModelingOperators/Public/CompositionOps/BooleanMeshesOp.h
@@ -14,26 +14,26 @@
 UENUM()
 enum class ECSGOperation : uint8
 {
-	/** Subtracts the first object from the second */
-	DifferenceAB = 0 UMETA(DisplayName = "A - B"),
+	/** Subtract the second object from the first object */
+	DifferenceAB = 0 UMETA(DisplayName = "Difference A - B"),
 
-	/** Subtracts the second object from the first */
-	DifferenceBA = 1 UMETA(DisplayName = "B - A"),
+	/** Subtract the first object from the second object */
+	DifferenceBA = 1 UMETA(DisplayName = "Difference B - A"),
 
-	/** intersection of two objects */
-	Intersect = 2 UMETA(DisplayName = "Intersect"),
+	/** Intersection of the two objects, i.e. where both objects overlap */
+	Intersect = 2 UMETA(DisplayName = "Intersection"),
 
-	/** union of two objects */
+	/** Union of the two objects, i.e. merger of both objects  */
 	Union = 3 UMETA(DisplayName = "Union"),
 };
 
 UENUM()
 enum class ETrimOperation : uint8
 {
-	/** Remove geometry from the first object using the second */
+	/** Remove geometry from the first object using the second object */
 	TrimA = 0 UMETA(DisplayName = "Trim A"),
 
-	/** Remove geometry from the second object using the first */
+	/** Remove geometry from the second object using the first object */
 	TrimB = 1 UMETA(DisplayName = "Trim B"),
 };
 
@@ -54,7 +54,7 @@ namespace Geometry
 class MODELINGOPERATORS_API FBooleanMeshesOp : public FDynamicMeshOperator
 {
 public:
-	virtual ~FBooleanMeshesOp() {}
+	virtual ~FBooleanMeshesOp() override = default;
 
 	// inputs
 	ECSGOperation CSGOperation;
