UV Editor: Make selection mode change convert rather than clear the selection

#rnx

diff --git a/Engine/Plugins/Experimental/UVEditor/Source/UVEditorTools/Private/UVSelectTool.cpp b/Engine/Plugins/Experimental/UVEditor/Source/UVEditorTools/Private/UVSelectTool.cpp
--- a/Engine/Plugins/Experimental/UVEditor/Source/UVEditorTools/Private/UVSelectTool.cpp
+++ b/Engine/Plugins/Experimental/UVEditor/Source/UVEditorTools/Private/UVSelectTool.cpp
@@ -532,27 +532,29 @@ void UUVSelectTool::UpdateGizmo()
 
 void UUVSelectTool::ConfigureSelectionModeFromControls()
 {
+	EMeshSelectionMechanicMode TargetMode;
 	switch (Settings->SelectionMode)
 	{
 	case EUVSelectToolSelectionMode::Island:
-		SelectionMechanic->SelectionMode = EMeshSelectionMechanicMode::Component;
+		TargetMode = EMeshSelectionMechanicMode::Component;
 		break;
 	case EUVSelectToolSelectionMode::Edge:
-		SelectionMechanic->SelectionMode = EMeshSelectionMechanicMode::Edge;
+		TargetMode = EMeshSelectionMechanicMode::Edge;
 		break;
 	case EUVSelectToolSelectionMode::Vertex:
-		SelectionMechanic->SelectionMode = EMeshSelectionMechanicMode::Vertex;
+		TargetMode = EMeshSelectionMechanicMode::Vertex;
 		break;
 	case EUVSelectToolSelectionMode::Triangle:
-		SelectionMechanic->SelectionMode = EMeshSelectionMechanicMode::Triangle;
+		TargetMode = EMeshSelectionMechanicMode::Triangle;
 		break;
 	case EUVSelectToolSelectionMode::Mesh:
-		SelectionMechanic->SelectionMode = EMeshSelectionMechanicMode::Mesh;
+		TargetMode = EMeshSelectionMechanicMode::Mesh;
 		break;
 	default:
 		ensure(false);
 		break;
 	}
+	SelectionMechanic->ChangeSelectionMode(TargetMode);
 }
 
 void UUVSelectTool::OnSelectionChanged()
diff --git a/Engine/Plugins/Runtime/MeshModelingToolset/Source/ModelingComponents/Private/Selection/MeshSelectionMechanic.cpp b/Engine/Plugins/Runtime/MeshModelingToolset/Source/ModelingComponents/Private/Selection/MeshSelectionMechanic.cpp
--- a/Engine/Plugins/Runtime/MeshModelingToolset/Source/ModelingComponents/Private/Selection/MeshSelectionMechanic.cpp
+++ b/Engine/Plugins/Runtime/MeshModelingToolset/Source/ModelingComponents/Private/Selection/MeshSelectionMechanic.cpp
@@ -415,6 +415,8 @@ void UMeshSelectionMechanic::SetSelection(const FDynamicMeshSelection& Selection
 
 void UMeshSelectionMechanic::RebuildDrawnElements(const FTransform& StartTransform)
 {
+	TRACE_CPUPROFILER_EVENT_SCOPE(MeshSelectionMechanic_RebuildDrawnElements);
+
 	LineSet->Clear();
 	PointSet->Clear();
 	PreviewGeometryActor->SetActorTransform(StartTransform);
@@ -434,7 +436,7 @@ void UMeshSelectionMechanic::RebuildDrawnElements(const FTransform& StartTransfo
 
 	if (CurrentSelection.Type == FDynamicMeshSelection::EType::Triangle)
 	{
-		TRACE_CPUPROFILER_EVENT_SCOPE(RebuildDrawnElements_Triangle);
+		TRACE_CPUPROFILER_EVENT_SCOPE(Triangle);
 		
 		LineSet->ReserveLines(CurrentSelection.SelectedIDs.Num() * 3);
 		for (int32 Tid : CurrentSelection.SelectedIDs)
@@ -455,7 +457,7 @@ void UMeshSelectionMechanic::RebuildDrawnElements(const FTransform& StartTransfo
 	}
 	else if (CurrentSelection.Type == FDynamicMeshSelection::EType::Edge)
 	{
-		TRACE_CPUPROFILER_EVENT_SCOPE(RebuildDrawnElements_Edge);
+		TRACE_CPUPROFILER_EVENT_SCOPE(Elements_Edge);
 
 		LineSet->ReserveLines(CurrentSelection.SelectedIDs.Num());
 		for (int32 Eid : CurrentSelection.SelectedIDs)
@@ -469,7 +471,7 @@ void UMeshSelectionMechanic::RebuildDrawnElements(const FTransform& StartTransfo
 	}
 	else if (CurrentSelection.Type == FDynamicMeshSelection::EType::Vertex)
 	{
-		TRACE_CPUPROFILER_EVENT_SCOPE(RebuildDrawnElements_Vertex);
+		TRACE_CPUPROFILER_EVENT_SCOPE(Vertex);
 
 		PointSet->ReservePoints(CurrentSelection.SelectedIDs.Num());
 		for (int32 Vid : CurrentSelection.SelectedIDs)
@@ -572,6 +574,242 @@ void UMeshSelectionMechanic::ClearCurrentSelection()
 	CurrentSelectionIndex = IndexConstants::InvalidID;
 }
 
+
+void UMeshSelectionMechanic::ChangeSelectionMode(const EMeshSelectionMechanicMode& TargetMode)
+{
+	TRACE_CPUPROFILER_EVENT_SCOPE(MeshSelectionMechanic_ChangeSelectionMode);
+	
+	if (CurrentSelection.SelectedIDs.IsEmpty())
+	{
+		SelectionMode = TargetMode;
+		return;
+	}
+	
+	const FDynamicMeshSelection OriginalSelection = CurrentSelection;
+	
+	auto VerticesToEdges = [this, &OriginalSelection]()
+	{
+		TRACE_CPUPROFILER_EVENT_SCOPE(VerticesToEdges);
+		
+		CurrentSelection.SelectedIDs.Empty();
+		CurrentSelection.Type = FDynamicMeshSelection::EType::Edge;
+
+		for (int32 Vid : OriginalSelection.SelectedIDs)
+		{
+			for (int32 Eid : OriginalSelection.Mesh->VtxEdgesItr(Vid))
+			{
+				if (!CurrentSelection.SelectedIDs.Contains(Eid))
+				{
+					FIndex2i Verts = OriginalSelection.Mesh->GetEdgeV(Eid);
+					if (OriginalSelection.SelectedIDs.Contains(Verts.A) &&
+						OriginalSelection.SelectedIDs.Contains(Verts.B))
+					{
+						CurrentSelection.SelectedIDs.Add(Eid);	
+					}
+				}
+			}
+		}
+	};
+
+	auto VerticesToTriangles = [this, &OriginalSelection]()
+	{
+		TRACE_CPUPROFILER_EVENT_SCOPE(VerticesToTriangles);
+		
+		CurrentSelection.SelectedIDs.Empty();
+		CurrentSelection.Type = FDynamicMeshSelection::EType::Triangle;
+		
+		for (int32 Vid : OriginalSelection.SelectedIDs)
+		{
+			for (int32 Tid : OriginalSelection.Mesh->VtxTrianglesItr(Vid))
+			{
+				if (!CurrentSelection.SelectedIDs.Contains(Tid))
+				{
+					FIndex3i Verts = OriginalSelection.Mesh->GetTriangle(Tid);
+					if (OriginalSelection.SelectedIDs.Contains(Verts.A) &&
+						OriginalSelection.SelectedIDs.Contains(Verts.B) &&
+						OriginalSelection.SelectedIDs.Contains(Verts.C))
+					{
+						CurrentSelection.SelectedIDs.Add(Tid);	
+					}
+				}
+			}
+		}
+	};
+
+	auto EdgesToVertices = [this, &OriginalSelection]()
+	{
+		TRACE_CPUPROFILER_EVENT_SCOPE(EdgesToVertices);
+		
+		CurrentSelection.SelectedIDs.Empty();
+		CurrentSelection.Type = FDynamicMeshSelection::EType::Vertex;
+		
+		for (int32 Eid : OriginalSelection.SelectedIDs)
+		{
+			FIndex2i Verts = OriginalSelection.Mesh->GetEdgeV(Eid);
+			CurrentSelection.SelectedIDs.Add(Verts.A);
+			CurrentSelection.SelectedIDs.Add(Verts.B);
+		}
+	};
+
+	// Triangles with two selected edges will be selected
+	auto EdgesToTriangles = [this, &OriginalSelection]()
+	{
+		TRACE_CPUPROFILER_EVENT_SCOPE(EdgesToTriangles);
+		
+		CurrentSelection.SelectedIDs.Empty();
+		CurrentSelection.Type = FDynamicMeshSelection::EType::Triangle;
+		
+		TArray<int32> FoundTriangles;
+		for (int32 Eid : OriginalSelection.SelectedIDs)
+		{
+			FIndex2i Tris = OriginalSelection.Mesh->GetEdgeT(Eid);
+			FoundTriangles.Add(Tris.A);
+			if (Tris.B != IndexConstants::InvalidID)
+			{
+				FoundTriangles.Add(Tris.B);
+			}
+		}
+
+		if (FoundTriangles.Num() < 2)
+		{
+			return;
+		}
+
+		Algo::Sort(FoundTriangles);
+
+		for (int I = 0; I < FoundTriangles.Num() - 1; I++)
+		{
+			if (FoundTriangles[I] == FoundTriangles[I + 1])
+			{
+				CurrentSelection.SelectedIDs.Add(FoundTriangles[I]);
+				I++;
+			}	
+		}
+	};
+
+	auto TrianglesToVertices = [this, &OriginalSelection]()
+	{
+		TRACE_CPUPROFILER_EVENT_SCOPE(TrianglesToVertices);
+		
+		CurrentSelection.SelectedIDs.Empty();
+		CurrentSelection.Type = FDynamicMeshSelection::EType::Vertex;
+
+		for (int32 Tid : OriginalSelection.SelectedIDs)
+		{
+			FIndex3i Verts = OriginalSelection.Mesh->GetTriangle(Tid);
+			CurrentSelection.SelectedIDs.Add(Verts.A);
+			CurrentSelection.SelectedIDs.Add(Verts.B);
+			CurrentSelection.SelectedIDs.Add(Verts.C);
+		}
+	};
+
+	auto TrianglesToEdges = [this, &OriginalSelection]()
+	{
+		TRACE_CPUPROFILER_EVENT_SCOPE(TrianglesToEdges);
+			
+		CurrentSelection.SelectedIDs.Empty();
+		CurrentSelection.Type = FDynamicMeshSelection::EType::Edge;
+	
+		for (int32 Tid : OriginalSelection.SelectedIDs)
+		{
+			FIndex3i Edges = OriginalSelection.Mesh->GetTriEdgesRef(Tid);
+			CurrentSelection.SelectedIDs.Add(Edges.A);
+			CurrentSelection.SelectedIDs.Add(Edges.B);
+			CurrentSelection.SelectedIDs.Add(Edges.C);
+		}
+	};
+
+	switch (SelectionMode)
+	{
+	case EMeshSelectionMechanicMode::Vertex:
+		switch (TargetMode)
+		{
+		case EMeshSelectionMechanicMode::Vertex:
+			// Do nothing
+			break;
+
+		case EMeshSelectionMechanicMode::Edge:
+			VerticesToEdges();
+			break;
+			
+		case EMeshSelectionMechanicMode::Triangle:
+		case EMeshSelectionMechanicMode::Component:
+		case EMeshSelectionMechanicMode::Mesh:
+			VerticesToTriangles();
+			break;
+
+		default:
+			checkNoEntry();
+			break;
+		}
+		break; // SelectionMode
+	
+	case EMeshSelectionMechanicMode::Edge:
+		switch (TargetMode)
+		{
+		case EMeshSelectionMechanicMode::Vertex:
+			EdgesToVertices();
+			break;
+
+		case EMeshSelectionMechanicMode::Edge:
+			// Do nothing
+			break;
+
+		case EMeshSelectionMechanicMode::Mesh:
+		case EMeshSelectionMechanicMode::Triangle:
+		case EMeshSelectionMechanicMode::Component:
+			EdgesToTriangles();
+			break;
+
+		default:
+			checkNoEntry();
+			break;
+		}
+		break; // SelectionMode
+	
+	case EMeshSelectionMechanicMode::Mesh:
+	case EMeshSelectionMechanicMode::Triangle:
+	case EMeshSelectionMechanicMode::Component:
+		switch (TargetMode)
+		{
+		case EMeshSelectionMechanicMode::Vertex:
+			TrianglesToVertices();
+			break;
+
+		case EMeshSelectionMechanicMode::Edge:
+			TrianglesToEdges();
+			break;
+
+		case EMeshSelectionMechanicMode::Mesh:
+		case EMeshSelectionMechanicMode::Triangle:
+		case EMeshSelectionMechanicMode::Component:
+			// Do nothing
+			break;
+
+		default:
+			checkNoEntry();
+			break;
+		}
+		break; // SelectionMode
+	
+	default:
+		checkNoEntry();
+		break; // SelectionMode
+	}
+
+	if (CurrentSelection != OriginalSelection)
+	{
+		UpdateCentroid();
+		RebuildDrawnElements(FTransform(GetCurrentSelectionCentroid()));
+		EmitSelectionChange(OriginalSelection, CurrentSelection, true);
+		OnSelectionChanged.Broadcast();
+	}
+
+	SelectionMode = TargetMode;
+	return;
+}
+
+
 void UMeshSelectionMechanic::UpdateCurrentSelection(const TSet<int32>& NewSelection, bool CalledFromOnDragRectangleChanged)
 {
 	TRACE_CPUPROFILER_EVENT_SCOPE(MeshSelectionMechanic_UpdateCurrentSelection);
@@ -618,6 +856,8 @@ void UMeshSelectionMechanic::UpdateCurrentSelection(const TSet<int32>& NewSelect
 
 void UMeshSelectionMechanic::OnClicked(const FInputDeviceRay& ClickPos)
 {
+	TRACE_CPUPROFILER_EVENT_SCOPE(MeshSelectionMechanic_OnClicked);
+	
 	FDynamicMeshSelection OriginalSelection = CurrentSelection;
 	
 	if (CurrentSelection.Type != MeshSelectionMechanicLocals::ToCompatibleDynamicMeshSelectionType(SelectionMode))
@@ -648,6 +888,8 @@ void UMeshSelectionMechanic::OnClicked(const FInputDeviceRay& ClickPos)
 		{
 			if (SelectionMode == EMeshSelectionMechanicMode::Component)
 			{
+				TRACE_CPUPROFILER_EVENT_SCOPE(Component);
+
 				FMeshConnectedComponents MeshSelectedComponent(CurrentSelection.Mesh);
 				TArray<int32> SeedTriangles;
 				SeedTriangles.Add(HitTid);
@@ -657,6 +899,7 @@ void UMeshSelectionMechanic::OnClicked(const FInputDeviceRay& ClickPos)
 			}
 			else if (SelectionMode == EMeshSelectionMechanicMode::Edge)
 			{
+				TRACE_CPUPROFILER_EVENT_SCOPE(Edge);
 				// TODO: We'll need the ability to hit occluded triangles to see if there is a better edge to snap to.
 
 				// Try to snap to one of the edges.
@@ -681,6 +924,7 @@ void UMeshSelectionMechanic::OnClicked(const FInputDeviceRay& ClickPos)
 			}
 			else if (SelectionMode == EMeshSelectionMechanicMode::Vertex)
 			{
+				TRACE_CPUPROFILER_EVENT_SCOPE(Vertex);
 				// TODO: Improve this to handle super narrow, sliver triangles better, where testing near vertices can be difficult.
 
 				// Try to snap to one of the vertices
@@ -703,10 +947,14 @@ void UMeshSelectionMechanic::OnClicked(const FInputDeviceRay& ClickPos)
 			}
 			else if (SelectionMode == EMeshSelectionMechanicMode::Triangle)
 			{
+				TRACE_CPUPROFILER_EVENT_SCOPE(Triangle);
+
 				ClickSelectedIDs = TSet<int32>{HitTid};
 			}
 			else if (SelectionMode == EMeshSelectionMechanicMode::Mesh)
 			{
+				TRACE_CPUPROFILER_EVENT_SCOPE(Mesh);
+
 				for (int32 Tid : CurrentSelection.Mesh->TriangleIndicesItr())
 				{
 					ClickSelectedIDs.Add(Tid);
diff --git a/Engine/Plugins/Runtime/MeshModelingToolset/Source/ModelingComponents/Public/Selection/MeshSelectionMechanic.h b/Engine/Plugins/Runtime/MeshModelingToolset/Source/ModelingComponents/Public/Selection/MeshSelectionMechanic.h
--- a/Engine/Plugins/Runtime/MeshModelingToolset/Source/ModelingComponents/Public/Selection/MeshSelectionMechanic.h
+++ b/Engine/Plugins/Runtime/MeshModelingToolset/Source/ModelingComponents/Public/Selection/MeshSelectionMechanic.h
@@ -77,6 +77,7 @@ public:
 	void SetDrawnElementsTransform(const FTransform& Transform);
 
 	virtual const FDynamicMeshSelection& GetCurrentSelection() const;
+	void ChangeSelectionMode(const EMeshSelectionMechanicMode& TargetMode);
 	virtual void SetSelection(const FDynamicMeshSelection& Selection, bool bBroadcast = false, bool bEmitChange = false);
 
 	virtual void Render(IToolsContextRenderAPI* RenderAPI) override;
