UV Editor: Disable marquee live selection update for a less surprising UX.

Also add namespaces to internal trace macros and prefer possible broken behavior selection conversion to editor crashes

#rnx

diff --git a/Engine/Plugins/Runtime/MeshModelingToolset/Source/ModelingComponents/Private/Selection/MeshSelectionMechanic.cpp b/Engine/Plugins/Runtime/MeshModelingToolset/Source/ModelingComponents/Private/Selection/MeshSelectionMechanic.cpp
--- a/Engine/Plugins/Runtime/MeshModelingToolset/Source/ModelingComponents/Private/Selection/MeshSelectionMechanic.cpp
+++ b/Engine/Plugins/Runtime/MeshModelingToolset/Source/ModelingComponents/Private/Selection/MeshSelectionMechanic.cpp
@@ -279,10 +279,11 @@ void UMeshSelectionMechanic::Setup(UInteractiveTool* ParentToolIn)
 	// This will be the target for the click drag behavior below
 	MarqueeMechanic = NewObject<URectangleMarqueeMechanic>();
 	MarqueeMechanic->bUseExternalClickDragBehavior = true;
-	MarqueeMechanic->OnDragRectangleChangedDeferredThreshold = 1./60;
 	MarqueeMechanic->Setup(ParentToolIn);
 	MarqueeMechanic->OnDragRectangleStarted.AddUObject(this, &UMeshSelectionMechanic::OnDragRectangleStarted);
-	MarqueeMechanic->OnDragRectangleChanged.AddUObject(this, &UMeshSelectionMechanic::OnDragRectangleChanged);
+	// TODO(Performance) :DynamicMarqueeSelection It would be cool to have the marquee selection update dynamically as
+	//  the rectangle gets changed, right now this isn't interactive for large meshes so we disabled it
+	// MarqueeMechanic->OnDragRectangleChanged.AddUObject(this, &UMeshSelectionMechanic::OnDragRectangleChanged);
 	MarqueeMechanic->OnDragRectangleFinished.AddUObject(this, &UMeshSelectionMechanic::OnDragRectangleFinished);
 
 	USingleClickOrDragInputBehavior* ClickOrDragBehavior = NewObject<USingleClickOrDragInputBehavior>();
@@ -436,7 +437,7 @@ void UMeshSelectionMechanic::RebuildDrawnElements(const FTransform& StartTransfo
 
 	if (CurrentSelection.Type == FDynamicMeshSelection::EType::Triangle)
 	{
-		TRACE_CPUPROFILER_EVENT_SCOPE(Triangle);
+		TRACE_CPUPROFILER_EVENT_SCOPE(MeshSelectionMechanic_RebuildDrawnElements_Triangle);
 		
 		LineSet->ReserveLines(CurrentSelection.SelectedIDs.Num() * 3);
 		for (int32 Tid : CurrentSelection.SelectedIDs)
@@ -457,7 +458,7 @@ void UMeshSelectionMechanic::RebuildDrawnElements(const FTransform& StartTransfo
 	}
 	else if (CurrentSelection.Type == FDynamicMeshSelection::EType::Edge)
 	{
-		TRACE_CPUPROFILER_EVENT_SCOPE(Elements_Edge);
+		TRACE_CPUPROFILER_EVENT_SCOPE(MeshSelectionMechanic_RebuildDrawnElements_Edge);
 
 		LineSet->ReserveLines(CurrentSelection.SelectedIDs.Num());
 		for (int32 Eid : CurrentSelection.SelectedIDs)
@@ -471,7 +472,7 @@ void UMeshSelectionMechanic::RebuildDrawnElements(const FTransform& StartTransfo
 	}
 	else if (CurrentSelection.Type == FDynamicMeshSelection::EType::Vertex)
 	{
-		TRACE_CPUPROFILER_EVENT_SCOPE(Vertex);
+		TRACE_CPUPROFILER_EVENT_SCOPE(MeshSelectionMechanic_RebuildDrawnElements_Vertex);
 
 		PointSet->ReservePoints(CurrentSelection.SelectedIDs.Num());
 		for (int32 Vid : CurrentSelection.SelectedIDs)
@@ -589,7 +590,7 @@ void UMeshSelectionMechanic::ChangeSelectionMode(const EMeshSelectionMechanicMod
 	
 	auto VerticesToEdges = [this, &OriginalSelection]()
 	{
-		TRACE_CPUPROFILER_EVENT_SCOPE(VerticesToEdges);
+		TRACE_CPUPROFILER_EVENT_SCOPE(MeshSelectionMechanic_ChangeSelectionMode_VerticesToEdges);
 		
 		CurrentSelection.SelectedIDs.Empty();
 		CurrentSelection.Type = FDynamicMeshSelection::EType::Edge;
@@ -613,7 +614,7 @@ void UMeshSelectionMechanic::ChangeSelectionMode(const EMeshSelectionMechanicMod
 
 	auto VerticesToTriangles = [this, &OriginalSelection]()
 	{
-		TRACE_CPUPROFILER_EVENT_SCOPE(VerticesToTriangles);
+		TRACE_CPUPROFILER_EVENT_SCOPE(MeshSelectionMechanic_ChangeSelectionMode_VerticesToTriangles);
 		
 		CurrentSelection.SelectedIDs.Empty();
 		CurrentSelection.Type = FDynamicMeshSelection::EType::Triangle;
@@ -638,7 +639,7 @@ void UMeshSelectionMechanic::ChangeSelectionMode(const EMeshSelectionMechanicMod
 
 	auto EdgesToVertices = [this, &OriginalSelection]()
 	{
-		TRACE_CPUPROFILER_EVENT_SCOPE(EdgesToVertices);
+		TRACE_CPUPROFILER_EVENT_SCOPE(MeshSelectionMechanic_ChangeSelectionMode_EdgesToVertices);
 		
 		CurrentSelection.SelectedIDs.Empty();
 		CurrentSelection.Type = FDynamicMeshSelection::EType::Vertex;
@@ -654,7 +655,7 @@ void UMeshSelectionMechanic::ChangeSelectionMode(const EMeshSelectionMechanicMod
 	// Triangles with two selected edges will be selected
 	auto EdgesToTriangles = [this, &OriginalSelection]()
 	{
-		TRACE_CPUPROFILER_EVENT_SCOPE(EdgesToTriangles);
+		TRACE_CPUPROFILER_EVENT_SCOPE(MeshSelectionMechanic_ChangeSelectionMode_EdgesToTriangles);
 		
 		CurrentSelection.SelectedIDs.Empty();
 		CurrentSelection.Type = FDynamicMeshSelection::EType::Triangle;
@@ -689,7 +690,7 @@ void UMeshSelectionMechanic::ChangeSelectionMode(const EMeshSelectionMechanicMod
 
 	auto TrianglesToVertices = [this, &OriginalSelection]()
 	{
-		TRACE_CPUPROFILER_EVENT_SCOPE(TrianglesToVertices);
+		TRACE_CPUPROFILER_EVENT_SCOPE(MeshSelectionMechanic_ChangeSelectionMode_TrianglesToVertices);
 		
 		CurrentSelection.SelectedIDs.Empty();
 		CurrentSelection.Type = FDynamicMeshSelection::EType::Vertex;
@@ -705,7 +706,7 @@ void UMeshSelectionMechanic::ChangeSelectionMode(const EMeshSelectionMechanicMod
 
 	auto TrianglesToEdges = [this, &OriginalSelection]()
 	{
-		TRACE_CPUPROFILER_EVENT_SCOPE(TrianglesToEdges);
+		TRACE_CPUPROFILER_EVENT_SCOPE(MeshSelectionMechanic_ChangeSelectionMode_TrianglesToEdges);
 			
 		CurrentSelection.SelectedIDs.Empty();
 		CurrentSelection.Type = FDynamicMeshSelection::EType::Edge;
@@ -737,10 +738,6 @@ void UMeshSelectionMechanic::ChangeSelectionMode(const EMeshSelectionMechanicMod
 		case EMeshSelectionMechanicMode::Mesh:
 			VerticesToTriangles();
 			break;
-
-		default:
-			checkNoEntry();
-			break;
 		}
 		break; // SelectionMode
 	
@@ -760,10 +757,6 @@ void UMeshSelectionMechanic::ChangeSelectionMode(const EMeshSelectionMechanicMod
 		case EMeshSelectionMechanicMode::Component:
 			EdgesToTriangles();
 			break;
-
-		default:
-			checkNoEntry();
-			break;
 		}
 		break; // SelectionMode
 	
@@ -785,16 +778,8 @@ void UMeshSelectionMechanic::ChangeSelectionMode(const EMeshSelectionMechanicMod
 		case EMeshSelectionMechanicMode::Component:
 			// Do nothing
 			break;
-
-		default:
-			checkNoEntry();
-			break;
 		}
 		break; // SelectionMode
-	
-	default:
-		checkNoEntry();
-		break; // SelectionMode
 	}
 
 	if (CurrentSelection != OriginalSelection)
@@ -888,7 +873,7 @@ void UMeshSelectionMechanic::OnClicked(const FInputDeviceRay& ClickPos)
 		{
 			if (SelectionMode == EMeshSelectionMechanicMode::Component)
 			{
-				TRACE_CPUPROFILER_EVENT_SCOPE(Component);
+				TRACE_CPUPROFILER_EVENT_SCOPE(MeshSelectionMechanic_OnClicked_Component);
 
 				FMeshConnectedComponents MeshSelectedComponent(CurrentSelection.Mesh);
 				TArray<int32> SeedTriangles;
@@ -899,7 +884,7 @@ void UMeshSelectionMechanic::OnClicked(const FInputDeviceRay& ClickPos)
 			}
 			else if (SelectionMode == EMeshSelectionMechanicMode::Edge)
 			{
-				TRACE_CPUPROFILER_EVENT_SCOPE(Edge);
+				TRACE_CPUPROFILER_EVENT_SCOPE(MeshSelectionMechanic_OnClicked_Edge);
 				// TODO: We'll need the ability to hit occluded triangles to see if there is a better edge to snap to.
 
 				// Try to snap to one of the edges.
@@ -924,7 +909,7 @@ void UMeshSelectionMechanic::OnClicked(const FInputDeviceRay& ClickPos)
 			}
 			else if (SelectionMode == EMeshSelectionMechanicMode::Vertex)
 			{
-				TRACE_CPUPROFILER_EVENT_SCOPE(Vertex);
+				TRACE_CPUPROFILER_EVENT_SCOPE(MeshSelectionMechanic_OnClicked_Vertex);
 				// TODO: Improve this to handle super narrow, sliver triangles better, where testing near vertices can be difficult.
 
 				// Try to snap to one of the vertices
@@ -947,13 +932,13 @@ void UMeshSelectionMechanic::OnClicked(const FInputDeviceRay& ClickPos)
 			}
 			else if (SelectionMode == EMeshSelectionMechanicMode::Triangle)
 			{
-				TRACE_CPUPROFILER_EVENT_SCOPE(Triangle);
+				TRACE_CPUPROFILER_EVENT_SCOPE(MeshSelectionMechanic_OnClicked_Triangle);
 
 				ClickSelectedIDs = TSet<int32>{HitTid};
 			}
 			else if (SelectionMode == EMeshSelectionMechanicMode::Mesh)
 			{
-				TRACE_CPUPROFILER_EVENT_SCOPE(Mesh);
+				TRACE_CPUPROFILER_EVENT_SCOPE(MeshSelectionMechanic_OnClicked_Mesh);
 
 				for (int32 Tid : CurrentSelection.Mesh->TriangleIndicesItr())
 				{
@@ -1036,7 +1021,7 @@ void UMeshSelectionMechanic::OnDragRectangleChanged(const FCameraRectangle& Curr
 	TSet<int32> RectangleSelectedIDs;
 	if (SelectionMode == EMeshSelectionMechanicMode::Vertex)
 	{
-		TRACE_CPUPROFILER_EVENT_SCOPE(Vertex);
+		TRACE_CPUPROFILER_EVENT_SCOPE(MeshSelectionMechanic_OnDragRectangleChanged_Vertex);
 	
 		auto AddVertexIDsToSelection = [this, &RectangleXY](const FDynamicMeshAABBTree3& Tree, TSet<int32>& OutSelectedIDs)
 		{
@@ -1049,7 +1034,7 @@ void UMeshSelectionMechanic::OnDragRectangleChanged(const FCameraRectangle& Curr
 	}
 	else if (SelectionMode == EMeshSelectionMechanicMode::Edge)
 	{
-		TRACE_CPUPROFILER_EVENT_SCOPE(Edge);
+		TRACE_CPUPROFILER_EVENT_SCOPE(MeshSelectionMechanic_OnDragRectangleChanged_Edge);
 		
 		auto AddEdgeIDsToSelection = [this, &RectangleXY](const FDynamicMeshAABBTree3& Tree, TSet<int32>& OutSelectedIDs)
 		{
@@ -1062,7 +1047,7 @@ void UMeshSelectionMechanic::OnDragRectangleChanged(const FCameraRectangle& Curr
 	}
 	else if (SelectionMode == EMeshSelectionMechanicMode::Triangle)
 	{
-		TRACE_CPUPROFILER_EVENT_SCOPE(Triangle);
+		TRACE_CPUPROFILER_EVENT_SCOPE(MeshSelectionMechanic_OnDragRectangleChanged_Triangle);
 		
 		auto AddTriangleIDsToSelection = [this, &RectangleXY] (const FDynamicMeshAABBTree3& Tree, TSet<int32>& OutSelectedIDs)
 		{
@@ -1075,14 +1060,16 @@ void UMeshSelectionMechanic::OnDragRectangleChanged(const FCameraRectangle& Curr
 	}
 	else if (SelectionMode == EMeshSelectionMechanicMode::Component)
 	{
-		TRACE_CPUPROFILER_EVENT_SCOPE(Component);
+		TRACE_CPUPROFILER_EVENT_SCOPE(MeshSelectionMechanic_OnDragRectangleChanged_Component);
 		
 		auto AddTriangleIDsToSelection = [this, &RectangleXY](const FDynamicMeshAABBTree3& Tree, TSet<int32>& OutSelectedIDs)
 		{
 			TArray<int32> SeedTriangles =
 				FindAllIntersectionsAxisAlignedBox2(Tree, RectangleXY, AppendTriangleID, AppendTriangleIDIfIntersected);
 
-			// TODO :NanitePerformance The following code is MUCH slower than AABB traversal
+			// TODO(Performance) For large meshes and selections following code is MUCH slower than AABB traversal,
+			//  consider precomputing the connected components an only updating them when the mesh topology changes
+			//  rather than every time the selection changes.
 			FMeshConnectedComponents MeshSelectedComponent(Tree.GetMesh());
 			MeshSelectedComponent.FindTrianglesConnectedToSeeds(SeedTriangles);
 			for (int ComponentIndex = 0; ComponentIndex < MeshSelectedComponent.Components.Num(); ComponentIndex++)
@@ -1095,7 +1082,7 @@ void UMeshSelectionMechanic::OnDragRectangleChanged(const FCameraRectangle& Curr
 	}
 	else if (SelectionMode == EMeshSelectionMechanicMode::Mesh)
 	{
-		TRACE_CPUPROFILER_EVENT_SCOPE(Mesh);
+		TRACE_CPUPROFILER_EVENT_SCOPE(MeshSelectionMechanic_OnDragRectangleChanged_Mesh);
 		
 		auto AddAllMeshTriangleIDsToSelection = [this, &RectangleXY] (const FDynamicMeshAABBTree3& Tree, TSet<int32>& OutSelectedIDs)
 		{
@@ -1119,17 +1106,21 @@ void UMeshSelectionMechanic::OnDragRectangleChanged(const FCameraRectangle& Curr
 
 	UpdateCurrentSelection(RectangleSelectedIDs, true);
 	
-	// TODO :NanitePerformance With large meshes and selections this call is much slower than AABB traversal (4x)
+	// TODO(Performance) With large meshes and selections this call is much slower than AABB traversal (4x)
 	UpdateCentroid();
 	
-	// TODO :NanitePerformance With large meshes and selections this call is MUCH slower than AABB traversal (60x)
+	// TODO(Performance) With large meshes and selections this call is MUCH slower than AABB traversal (60x)
 	RebuildDrawnElements(FTransform(GetCurrentSelectionCentroid()));
 }
 
-void UMeshSelectionMechanic::OnDragRectangleFinished(const FCameraRectangle&, bool bCancelled)
+void UMeshSelectionMechanic::OnDragRectangleFinished(const FCameraRectangle& CurrentRectangle, bool bCancelled)
 {
 	TRACE_CPUPROFILER_EVENT_SCOPE(MeshSelectionMechanic_OnDragRectangleFinished); // Mark end of drag sequence
 
+	// TODO(Performance) :DynamicMarqueeSelection Remove this call when marquee selection is fast enough to update
+	//  dynamically for large meshes
+	OnDragRectangleChanged(CurrentRectangle);
+
 	if (!bCancelled && (PreDragSelection != CurrentSelection))
 	{
 		UpdateCentroid();
@@ -1155,3 +1146,4 @@ void UMeshSelectionMechanic::OnUpdateModifierState(int ModifierID, bool bIsOn)
 }
 
 #undef LOCTEXT_NAMESPACE
+
