UVEditor: Adjust handling of selection colors to reflect the target specific colors in the 2D and 3D viewports. Adds an additional API to convey colors from the UVEditorMode to the Selection tool. Adds a triangle component to the MeshSelectionMechanic to allow for interior shading of selection areas.

#rnx

diff --git a/Engine/Plugins/Experimental/UVEditor/Source/UVEditor/Private/UVEditorMode.cpp b/Engine/Plugins/Experimental/UVEditor/Source/UVEditor/Private/UVEditorMode.cpp
--- a/Engine/Plugins/Experimental/UVEditor/Source/UVEditor/Private/UVEditorMode.cpp
+++ b/Engine/Plugins/Experimental/UVEditor/Source/UVEditor/Private/UVEditorMode.cpp
@@ -340,7 +340,7 @@ void UUVEditorMode::Exit()
 	Super::Exit();
 }
 
-void UUVEditorMode::InitializeTargets(const TArray<TObjectPtr<UObject>>& AssetsIn, 
+void UUVEditorMode::InitializeTargets(const TArray<TObjectPtr<UObject>>& AssetsIn,
 	const TArray<FTransform>& TransformsIn)
 {
 	using namespace UVEditorModeLocals;
@@ -377,7 +377,7 @@ void UUVEditorMode::InitializeTargets(const TArray<TObjectPtr<UObject>>& AssetsI
 		AppliedPreview->PreviewMesh->UpdatePreview(AppliedCanonical.Get());
 
 		FComponentMaterialSet MaterialSet = UE::ToolTarget::GetMaterialSet(Target);
-		AppliedPreview->ConfigureMaterials(MaterialSet.Materials, 
+		AppliedPreview->ConfigureMaterials(MaterialSet.Materials,
 			ToolSetupUtil::GetDefaultWorkingMaterial(GetToolManager()));
 		AppliedPreviews.Add(AppliedPreview);
 	}
@@ -447,7 +447,7 @@ void UUVEditorMode::InitializeTargets(const TArray<TObjectPtr<UObject>>& AssetsI
 		// The wireframe will track the unwrap preview mesh
 		WireframeDisplay->SetMeshAccessFunction([ToolInputObject](UMeshElementsVisualizer::ProcessDynamicMeshFunc ProcessFunc) {
 			ToolInputObject->UnwrapPreview->ProcessCurrentMesh(ProcessFunc);
-		});
+			});
 
 		// The settings object and wireframe are not part of a tool, so they won't get ticked like they
 		// are supposed to (to enable property watching), unless we add this here.
@@ -460,15 +460,15 @@ void UUVEditorMode::InitializeTargets(const TArray<TObjectPtr<UObject>>& AssetsI
 		// Bind to delegate so that we can detect changes
 		ToolInputObject->OnCanonicalModified.AddWeakLambda(this, [this]
 		(UUVEditorToolMeshInput* InputObject, const UUVEditorToolMeshInput::FCanonicalModifiedInfo&) {
-			ModifiedAssetIDs.Add(InputObject->AssetID);
-		});
+				ModifiedAssetIDs.Add(InputObject->AssetID);
+			});
 
 		ToolInputObjects.Add(ToolInputObject);
 	}
 
 	// Prep things for layer/channel selection
 	InitializeAssetNames(ToolTargets, AssetNames);
-	PendingUVLayerIndex.SetNumZeroed(ToolTargets.Num());		
+	PendingUVLayerIndex.SetNumZeroed(ToolTargets.Num());
 
 	// Prep channel/layer change API
 	UUVToolAssetAndChannelAPI* AssetAndLayerAPI = ContextStore->FindContext<UUVToolAssetAndChannelAPI>();
@@ -493,6 +493,15 @@ void UUVEditorMode::InitializeTargets(const TArray<TObjectPtr<UObject>>& AssetsI
 			return VisibleLayers;
 		};
 	}
+
+	// Prep VisualStyle API
+	UUVVisualStyleAPI* VisualStyleAPI = ContextStore->FindContext<UUVVisualStyleAPI>();
+	if (VisualStyleAPI)
+	{
+		VisualStyleAPI->GetSelectionColorForAssetFunc = [this](int32 AssetID) {
+			return GetSelectionColorByTargetIndex(AssetID);
+		};
+	}
 }
 
 FLinearColor UUVEditorMode::GetTriangleColorByTargetIndex(int32 TargetIndex) const
@@ -526,6 +535,11 @@ FLinearColor UUVEditorMode::GetBoundaryColorByTargetIndex(int32 TargetIndex) con
 	return BoundaryColorHSV.HSVToLinearRGB();
 }
 
+FLinearColor UUVEditorMode::GetSelectionColorByTargetIndex(int32 TargetIndex) const
+{
+	return FColor::FromHex("EBA30DFF");	
+}
+
 void UUVEditorMode::EmitToolIndependentObjectChange(UObject* TargetObject, TUniquePtr<FToolCommandChange> Change, const FText& Description)
 {
 	GetInteractiveToolsContext()->GetTransactionAPI()->AppendChange(TargetObject, MoveTemp(Change), Description);
diff --git a/Engine/Plugins/Experimental/UVEditor/Source/UVEditor/Private/UVEditorToolkit.cpp b/Engine/Plugins/Experimental/UVEditor/Source/UVEditor/Private/UVEditorToolkit.cpp
--- a/Engine/Plugins/Experimental/UVEditor/Source/UVEditor/Private/UVEditorToolkit.cpp
+++ b/Engine/Plugins/Experimental/UVEditor/Source/UVEditor/Private/UVEditorToolkit.cpp
@@ -409,6 +409,9 @@ void FUVEditorToolkit::PostInitAssetEditor()
 	UUVToolAssetAndChannelAPI* AssetAndLayerAPI = NewObject<UUVToolAssetAndChannelAPI>();
 	ContextStore->AddContextObject(AssetAndLayerAPI);
 
+	UUVVisualStyleAPI* VisualStyleAPI = NewObject<UUVVisualStyleAPI>();
+	ContextStore->AddContextObject(VisualStyleAPI);
+
 	// Initialize mode state.
 	TArray<TObjectPtr<UObject>> ObjectsToEdit;
 	OwningAssetEditor->GetObjectsToEdit(ObjectsToEdit);
diff --git a/Engine/Plugins/Experimental/UVEditor/Source/UVEditor/Public/UVEditorMode.h b/Engine/Plugins/Experimental/UVEditor/Source/UVEditor/Public/UVEditorMode.h
--- a/Engine/Plugins/Experimental/UVEditor/Source/UVEditor/Public/UVEditorMode.h
+++ b/Engine/Plugins/Experimental/UVEditor/Source/UVEditor/Public/UVEditorMode.h
@@ -194,7 +194,7 @@ protected:
 	const float WireframeDepthOffset = 0.6;
 	float UnwrapBaseColorHueStart = 200;
 	float UnwrapBaseColorSaturation = 0.25;
-	float UnwrapBaseColorValue = 0.75;
+	float UnwrapBaseColorValue = 0.60;
 	float UnwrapWireframeHueShift = 15;
 	float UnwrapWireframeSaturation = 0.40;
 	float UnwrapWireframeValue = 0.40;
@@ -205,6 +205,7 @@ protected:
 	FLinearColor GetTriangleColorByTargetIndex(int32 TargetIndex) const;
 	FLinearColor GetWireframeColorByTargetIndex(int32 TargetIndex) const;
 	FLinearColor GetBoundaryColorByTargetIndex(int32 TargetIndex) const;
+	FLinearColor GetSelectionColorByTargetIndex(int32 TargetIndex) const;
 
 	// Here largely for convenience to avoid having to pass it around functions.
 	UPROPERTY()
diff --git a/Engine/Plugins/Experimental/UVEditor/Source/UVEditorTools/Private/UVSelectTool.cpp b/Engine/Plugins/Experimental/UVEditor/Source/UVEditorTools/Private/UVSelectTool.cpp
--- a/Engine/Plugins/Experimental/UVEditor/Source/UVEditorTools/Private/UVSelectTool.cpp
+++ b/Engine/Plugins/Experimental/UVEditor/Source/UVEditorTools/Private/UVSelectTool.cpp
@@ -329,6 +329,7 @@ void UUVSelectTool::Setup()
 	SelectionMechanic->SetWorld(Targets[0]->UnwrapPreview->GetWorld());
 	SelectionMechanic->OnSelectionChanged.AddUObject(this, &UUVSelectTool::OnSelectionChanged);
 
+
 	// Make it so that our selection mechanic creates undo/redo transactions that go to a selection
 	// change router, which we use to route to the current selection mechanic on each tool invocation.
 	ChangeRouter = ContextStore->FindContext<UUVSelectToolChangeRouter>();
@@ -587,6 +588,14 @@ void UUVSelectTool::OnSelectionChanged()
 		}
 		check(SelectionTargetIndex >= 0);
 
+		UContextObjectStore* ContextStore = GetToolManager()->GetContextObjectStore();
+		UUVVisualStyleAPI* VisualStyleAPI = ContextStore->FindContext<UUVVisualStyleAPI>();
+		if (VisualStyleAPI)
+		{
+			FColor SelectionColor = VisualStyleAPI->GetSelectionColorForAsset(SelectionTargetIndex).ToFColor(true);
+			SelectionMechanic->ChangeSelectionColor(SelectionColor, 0.3f, SelectionColor, SelectionColor);
+		}
+
 		// Note the selected vids
 		TSet<int32> VidSet;
 		TSet<int32> TidSet;
@@ -706,6 +715,14 @@ void UUVSelectTool::UpdateLivePreviewLines()
 		FTransform MeshTransform = Targets[SelectionTargetIndex]->AppliedPreview->PreviewMesh->GetTransform();
 		const FDynamicMesh3* LivePreviewMesh = Targets[SelectionTargetIndex]->AppliedCanonical.Get();
 
+		UContextObjectStore* ContextStore = GetToolManager()->GetContextObjectStore();
+		UUVVisualStyleAPI* VisualStyleAPI = ContextStore->FindContext<UUVVisualStyleAPI>();
+		FColor SelectionColor = FColor::Yellow;
+		if (VisualStyleAPI)
+		{
+			SelectionColor = VisualStyleAPI->GetSelectionColorForAsset(SelectionTargetIndex).ToFColor(true);
+		}
+
 		for (int32 Eid : LivePreviewBoundaryEids)
 		{
 			FVector3d Vert1, Vert2;
@@ -714,7 +731,7 @@ void UUVSelectTool::UpdateLivePreviewLines()
 			LivePreviewLineSet->AddLine(
 				MeshTransform.TransformPosition(Vert1), 
 				MeshTransform.TransformPosition(Vert2), 
-				FColor::Yellow, LivePreviewHighlightThickness, LivePreviewHighlightDepthOffset);
+				SelectionColor, LivePreviewHighlightThickness, LivePreviewHighlightDepthOffset);
 		}	
 	}
 }
diff --git a/Engine/Plugins/Experimental/UVEditor/Source/UVEditorTools/Public/UVToolContextObjects.h b/Engine/Plugins/Experimental/UVEditor/Source/UVEditorTools/Public/UVToolContextObjects.h
--- a/Engine/Plugins/Experimental/UVEditor/Source/UVEditorTools/Public/UVToolContextObjects.h
+++ b/Engine/Plugins/Experimental/UVEditor/Source/UVEditorTools/Public/UVToolContextObjects.h
@@ -241,6 +241,28 @@ public:
 
 };
 
+
+/**
+  Allows tools and mode to access color and styling information per asset
+*/
+UCLASS()
+class UVEDITORTOOLS_API UUVVisualStyleAPI : public UUVToolContextObject
+{
+	GENERATED_BODY()
+public:
+	FLinearColor GetSelectionColorForAsset(int32 AssetID)
+	{
+		if (GetSelectionColorForAssetFunc)
+		{
+			return GetSelectionColorForAssetFunc(AssetID);
+		}
+		return FLinearColor();
+	}
+
+	TUniqueFunction<FLinearColor(int32)> GetSelectionColorForAssetFunc;
+};
+
+
 /** 
  * Stores AABB trees for UV input object unwrap canonical or applied canonical meshes.
  * Binds to the input objects' OnCanonicalModified delegate to automatically update 
diff --git a/Engine/Plugins/Runtime/MeshModelingToolset/Source/ModelingComponents/Private/Selection/MeshSelectionMechanic.cpp b/Engine/Plugins/Runtime/MeshModelingToolset/Source/ModelingComponents/Private/Selection/MeshSelectionMechanic.cpp
--- a/Engine/Plugins/Runtime/MeshModelingToolset/Source/ModelingComponents/Private/Selection/MeshSelectionMechanic.cpp
+++ b/Engine/Plugins/Runtime/MeshModelingToolset/Source/ModelingComponents/Private/Selection/MeshSelectionMechanic.cpp
@@ -4,6 +4,7 @@
 
 #include "BaseBehaviors/BehaviorTargetInterfaces.h"
 #include "BaseBehaviors/SingleClickOrDragBehavior.h"
+#include "Drawing/TriangleSetComponent.h"
 #include "Drawing/LineSetComponent.h"
 #include "Drawing/PointSetComponent.h"
 #include "Drawing/PreviewGeometryActor.h"
@@ -294,6 +295,9 @@ void UMeshSelectionMechanic::Setup(UInteractiveTool* ParentToolIn)
 
 	ClearCurrentSelection();
 
+	TriangleSet = NewObject<UTriangleSetComponent>();
+	TriangleSetMaterial = ToolSetupUtil::GetCustomTwoSidedDepthOffsetMaterial(GetParentTool()->GetToolManager(), TriangleColor, TriangleDepthBias, TriangleOpacity);
+	
 	LineSet = NewObject<ULineSetComponent>();
 	LineSet->SetLineMaterial(ToolSetupUtil::GetDefaultLineComponentMaterial(
 		GetParentTool()->GetToolManager(), /*bDepthTested*/ true));
@@ -327,9 +331,21 @@ void UMeshSelectionMechanic::SetWorld(UWorld* World)
 	FActorSpawnParameters SpawnInfo;
 	PreviewGeometryActor = World->SpawnActor<APreviewGeometryActor>(FVector::ZeroVector, Rotation, SpawnInfo);
 
-	// Attach the rendering component to the actor
+
+	// Attach the rendering component to the actor		
+	TriangleSet->Rename(nullptr, PreviewGeometryActor); // Changes the "outer"
+	PreviewGeometryActor->SetRootComponent(TriangleSet);
+	if (TriangleSet->IsRegistered())
+	{
+		TriangleSet->ReregisterComponent();
+	}
+	else
+	{
+		TriangleSet->RegisterComponent();
+	}
+
 	LineSet->Rename(nullptr, PreviewGeometryActor); // Changes the "outer"
-	PreviewGeometryActor->SetRootComponent(LineSet);
+	LineSet->AttachToComponent(TriangleSet, FAttachmentTransformRules::KeepWorldTransform);
 	if (LineSet->IsRegistered())
 	{
 		LineSet->ReregisterComponent();
@@ -340,7 +356,7 @@ void UMeshSelectionMechanic::SetWorld(UWorld* World)
 	}
 
 	PointSet->Rename(nullptr, PreviewGeometryActor); // Changes the "outer"
-	PointSet->AttachToComponent(LineSet, FAttachmentTransformRules::KeepWorldTransform);
+	PointSet->AttachToComponent(TriangleSet, FAttachmentTransformRules::KeepWorldTransform);
 	if (PointSet->IsRegistered())
 	{
 		PointSet->ReregisterComponent();
@@ -371,6 +387,19 @@ const FDynamicMeshSelection& UMeshSelectionMechanic::GetCurrentSelection() const
 	return CurrentSelection;
 }
 
+void UMeshSelectionMechanic::ChangeSelectionColor(const FColor& TriangleColorIn, float TriangleOpacityIn, const FColor& LineColorIn, const FColor& PointColorIn)
+{
+	TriangleColor = TriangleColorIn;
+	TriangleOpacity = TriangleOpacityIn;
+	LineColor = LineColorIn;
+	PointColor = PointColorIn;
+
+	if (TriangleSetMaterial) {
+		TriangleSetMaterial->SetVectorParameterValue(TEXT("Color"), TriangleColorIn);
+		TriangleSetMaterial->SetScalarParameterValue(TEXT("Opacity"), TriangleOpacityIn);
+	}
+}
+
 void UMeshSelectionMechanic::SetSelection(const FDynamicMeshSelection& Selection, bool bBroadcast, bool bEmitChange)
 {
 	FDynamicMeshSelection OriginalSelection = CurrentSelection;
@@ -402,7 +431,6 @@ void UMeshSelectionMechanic::SetSelection(const FDynamicMeshSelection& Selection
 	}
 
 	UpdateCentroid();
-	RebuildDrawnElements(FTransform(CurrentSelectionCentroid));
 
 	if (bEmitChange && OriginalSelection != CurrentSelection)
 	{
@@ -412,12 +440,16 @@ void UMeshSelectionMechanic::SetSelection(const FDynamicMeshSelection& Selection
 	{
 		OnSelectionChanged.Broadcast();
 	}
+
+	// Rebuild after broadcast in case the outside world wants to adjust things like color...
+	RebuildDrawnElements(FTransform(CurrentSelectionCentroid));
 }
 
 void UMeshSelectionMechanic::RebuildDrawnElements(const FTransform& StartTransform)
 {
 	TRACE_CPUPROFILER_EVENT_SCOPE(MeshSelectionMechanic_RebuildDrawnElements);
 
+	TriangleSet->Clear();
 	LineSet->Clear();
 	PointSet->Clear();
 	PreviewGeometryActor->SetActorTransform(StartTransform);
@@ -439,6 +471,7 @@ void UMeshSelectionMechanic::RebuildDrawnElements(const FTransform& StartTransfo
 	{
 		TRACE_CPUPROFILER_EVENT_SCOPE(MeshSelectionMechanic_RebuildDrawnElements_Triangle);
 		
+		TriangleSet->ReserveTriangles(CurrentSelection.SelectedIDs.Num());
 		LineSet->ReserveLines(CurrentSelection.SelectedIDs.Num() * 3);
 		for (int32 Tid : CurrentSelection.SelectedIDs)
 		{
@@ -448,6 +481,7 @@ void UMeshSelectionMechanic::RebuildDrawnElements(const FTransform& StartTransfo
 			{
 				Points[i] = TransformToApply(CurrentSelection.Mesh->GetVertex(Vids[i]));
 			}
+			TriangleSet->AddTriangle(Points[0], Points[1], Points[2], FVector(0, 0, 1), TriangleColor, TriangleSetMaterial);
 			for (int i = 0; i < 3; ++i)
 			{
 				int NextIndex = (i + 1) % 3;
@@ -958,9 +992,11 @@ void UMeshSelectionMechanic::OnClicked(const FInputDeviceRay& ClickPos)
 	if (OriginalSelection != CurrentSelection)
 	{
 		UpdateCentroid();
-		RebuildDrawnElements(FTransform(GetCurrentSelectionCentroid()));
 		EmitSelectionChange(OriginalSelection, CurrentSelection, true);
 		OnSelectionChanged.Broadcast();
+
+		// Rebuild after broadcast in case the outside world wants to adjust things like color...
+		RebuildDrawnElements(FTransform(GetCurrentSelectionCentroid()));
 	}
 }
 
@@ -1123,10 +1159,12 @@ void UMeshSelectionMechanic::OnDragRectangleFinished(const FCameraRectangle& Cur
 
 	if (!bCancelled && (PreDragSelection != CurrentSelection))
 	{
-		UpdateCentroid();
-		RebuildDrawnElements(FTransform(GetCurrentSelectionCentroid()));
+		UpdateCentroid();		
 		EmitSelectionChange(PreDragSelection, CurrentSelection, true);
 		OnSelectionChanged.Broadcast();
+
+		// Rebuild after broadcast in case the outside world wants to adjust things like color...
+		RebuildDrawnElements(FTransform(GetCurrentSelectionCentroid()));
 	}
 }
 
diff --git a/Engine/Plugins/Runtime/MeshModelingToolset/Source/ModelingComponents/Public/Selection/MeshSelectionMechanic.h b/Engine/Plugins/Runtime/MeshModelingToolset/Source/ModelingComponents/Public/Selection/MeshSelectionMechanic.h
--- a/Engine/Plugins/Runtime/MeshModelingToolset/Source/ModelingComponents/Public/Selection/MeshSelectionMechanic.h
+++ b/Engine/Plugins/Runtime/MeshModelingToolset/Source/ModelingComponents/Public/Selection/MeshSelectionMechanic.h
@@ -15,6 +15,7 @@
 
 #include "MeshSelectionMechanic.generated.h"
 
+class UTriangleSetComponent;
 class ULineSetComponent;
 class UPointSetComponent;
 class APreviewGeometryActor;
@@ -78,6 +79,7 @@ public:
 
 	virtual const FDynamicMeshSelection& GetCurrentSelection() const;
 	void ChangeSelectionMode(const EMeshSelectionMechanicMode& TargetMode);
+	void ChangeSelectionColor(const FColor& TriangleColorIn, float TriangleOpacityIn, const FColor& LineColorIn, const FColor& PointColorIn);
 	virtual void SetSelection(const FDynamicMeshSelection& Selection, bool bBroadcast = false, bool bEmitChange = false);
 
 	virtual void Render(IToolsContextRenderAPI* RenderAPI) override;
@@ -122,6 +124,13 @@ protected:
 	UPROPERTY()
 	TObjectPtr<APreviewGeometryActor> PreviewGeometryActor = nullptr;
 
+	/** The material being displayed for selected triangles */
+	UPROPERTY()
+	TObjectPtr<UMaterialInstanceDynamic> TriangleSetMaterial;
+
+	UPROPERTY()
+	TObjectPtr<UTriangleSetComponent> TriangleSet = nullptr;
+
 	UPROPERTY()
 	TObjectPtr<ULineSetComponent> LineSet = nullptr;
 
@@ -136,8 +145,10 @@ protected:
 	int32 CurrentSelectionIndex = IndexConstants::InvalidID;
 	FViewCameraState CameraState;
 
+	FColor TriangleColor = FColor::Yellow;
 	FColor LineColor = FColor::Yellow;
 	FColor PointColor = FColor::Yellow;
+	float TriangleOpacity = 0.3;
 
 	bool bShiftToggle = false;
 	bool bCtrlToggle = false;
@@ -147,6 +158,7 @@ protected:
 	float LineThickness = 3;
 	float PointThickness = 6;
 	float DepthBias = 0.7;
+	float TriangleDepthBias = -0.1;
 
 	FVector3d CurrentSelectionCentroid;
 	int32 CentroidTimestamp = -1;
