UVEditor: Fix issue with MeshSelectionMechanic preventing Non-unity builds from succeeding by incorrectly not declaring UMaterialInstanceDynamic

#rnx

diff --git a/Engine/Plugins/Runtime/MeshModelingToolset/Source/ModelingComponents/Public/Selection/MeshSelectionMechanic.h b/Engine/Plugins/Runtime/MeshModelingToolset/Source/ModelingComponents/Public/Selection/MeshSelectionMechanic.h
--- a/Engine/Plugins/Runtime/MeshModelingToolset/Source/ModelingComponents/Public/Selection/MeshSelectionMechanic.h
+++ b/Engine/Plugins/Runtime/MeshModelingToolset/Source/ModelingComponents/Public/Selection/MeshSelectionMechanic.h
@@ -19,6 +19,7 @@ class UTriangleSetComponent;
 class ULineSetComponent;
 class UPointSetComponent;
 class APreviewGeometryActor;
+class UMaterialInstanceDynamic;
 struct FCameraRectangle;
 
 enum class EMeshSelectionMechanicMode
