MeshModelingToolset: Cleanup MeshCut tool

diff --git a/Engine/Plugins/Runtime/MeshModelingToolset/Source/MeshModelingTools/Private/CutMeshWithMeshTool.cpp b/Engine/Plugins/Runtime/MeshModelingToolset/Source/MeshModelingTools/Private/CutMeshWithMeshTool.cpp
--- a/Engine/Plugins/Runtime/MeshModelingToolset/Source/MeshModelingTools/Private/CutMeshWithMeshTool.cpp
+++ b/Engine/Plugins/Runtime/MeshModelingToolset/Source/MeshModelingTools/Private/CutMeshWithMeshTool.cpp
@@ -46,7 +46,7 @@ void UCutMeshWithMeshTool::SetupProperties()
 
 	SetToolDisplayName(LOCTEXT("CutMeshWithMeshToolName", "Cut With Mesh"));
 	GetToolManager()->DisplayMessage(
-		LOCTEXT("OnStartTool", "Cut first mesh with second. Use the transform gizmos to tweak the positions of the input objects (can help to resolve errors/failures)"),
+		LOCTEXT("OnStartTool", "Cut the first input mesh with the second input mesh. Use the transform gizmos to modify the position and orientation of the input objects."),
 		EToolMessageLevel::UserNotification);
 
 
@@ -55,8 +55,6 @@ void UCutMeshWithMeshTool::SetupProperties()
 	IntersectPreviewMesh->CreateInWorld(TargetWorld, FTransform::Identity);
 	ToolSetupUtil::ApplyRenderingConfigurationToPreview(IntersectPreviewMesh, nullptr); 
 	IntersectPreviewMesh->SetVisible(true);
-	//IntersectPreviewMesh->SetMaterial(MaterialProperties->Material.Get());
-	//IntersectPreviewMesh->SetMaterial(ToolSetupUtil::GetDefaultMaterial(nullptr));
 	IntersectPreviewMesh->SetMaterial(ToolSetupUtil::GetDefaultBrushVolumeMaterial(GetToolManager()));
 }
 
@@ -77,11 +75,12 @@ void UCutMeshWithMeshTool::ConvertInputsAndSetPreviewMaterials(bool bSetPreviewM
 
 
 	FComponentMaterialSet AllMaterialSet;
-	TMap<UMaterialInterface*, int> KnownMaterials;
 	TArray<TArray<int>> MaterialRemap; MaterialRemap.SetNum(Targets.Num());
 
-	if (!CutProperties->bOnlyUseFirstMeshMaterials)
+	if (!CutProperties->bUseFirstMeshMaterials)
 	{
+		
+		TMap<UMaterialInterface*, int> KnownMaterials;
 		for (int ComponentIdx = 0; ComponentIdx < Targets.Num(); ComponentIdx++)
 		{
 			FComponentMaterialSet ComponentMaterialSet;
@@ -146,7 +145,7 @@ void UCutMeshWithMeshTool::ConvertInputsAndSetPreviewMaterials(bool bSetPreviewM
 	if (Cast<IAssetBackedTarget>(Targets[0]) != nullptr && Cast<IAssetBackedTarget>(Targets[0])->HasSameSourceData(Targets[1]))
 	{
 		GetToolManager()->DisplayMessage(
-			LOCTEXT("SameSourceError", "WARNING: Target Mesh has same Asset as Cutting Mesh, both inputs will be affected"),
+			LOCTEXT("SameSourceError", "WARNING: Both input meshes have the same Asset; both inputs will be affected."),
 			EToolMessageLevel::UserWarning);
 	}
 }
@@ -168,7 +167,7 @@ public:
 	bool bCollapseExtraEdges = true;
 	double WindingThreshold = 0.5;
 
-	virtual void CalculateResult(FProgressCancel* Progress)
+	virtual void CalculateResult(FProgressCancel* Progress) override
 	{
 		TUniquePtr<FBooleanMeshesOp> SubtractOp = MakeUnique<FBooleanMeshesOp>();
 		SubtractOp->CSGOperation = ECSGOperation::DifferenceAB;
@@ -250,12 +249,12 @@ void UCutMeshWithMeshTool::SetPreviewCallbacks()
 
 void UCutMeshWithMeshTool::UpdateVisualization()
 {
-	FColor BoundaryEdgeColor(240, 15, 15);
-	float BoundaryEdgeThickness = 2.0;
-	float BoundaryEdgeDepthBias = 2.0f;
+	constexpr FColor BoundaryEdgeColor(240, 15, 15);
+	constexpr float BoundaryEdgeThickness = 2.0;
+	constexpr float BoundaryEdgeDepthBias = 2.0f;
 
 	DrawnLineSet->Clear();
-	if (CutProperties->bShowNewBoundaryEdges)
+	if (CutProperties->bShowNewBoundaries)
 	{
 		const FDynamicMesh3* TargetMesh = Preview->PreviewMesh->GetPreviewDynamicMesh();
 		FVector3d A, B;
@@ -285,8 +284,8 @@ TUniquePtr<FDynamicMeshOperator> UCutMeshWithMeshTool::MakeNewOperator()
 	CuttingOp->CuttingMesh = OriginalCuttingMesh;
 	CuttingOp->CuttingMeshTransform = TransformProxies[1]->GetTransform();
 
-	CuttingOp->bAttemptToFixHoles = CutProperties->bAttemptFixHoles;
-	CuttingOp->bCollapseExtraEdges = CutProperties->bCollapseExtraEdges;
+	CuttingOp->bAttemptToFixHoles = CutProperties->bTryFixHoles;
+	CuttingOp->bCollapseExtraEdges = CutProperties->bTryCollapseEdges;
 	CuttingOp->WindingThreshold = CutProperties->WindingThreshold;
 
 	return CuttingOp;
@@ -296,17 +295,17 @@ TUniquePtr<FDynamicMeshOperator> UCutMeshWithMeshTool::MakeNewOperator()
 
 void UCutMeshWithMeshTool::OnPropertyModified(UObject* PropertySet, FProperty* Property)
 {
-	if (Property && (Property->GetFName() == GET_MEMBER_NAME_CHECKED(UCutMeshWithMeshToolProperties, bOnlyUseFirstMeshMaterials)))
+	if (Property && (Property->GetFName() == GET_MEMBER_NAME_CHECKED(UCutMeshWithMeshToolProperties, bUseFirstMeshMaterials)))
 	{
 		if (!AreAllTargetsValid())
 		{
-			GetToolManager()->DisplayMessage(LOCTEXT("InvalidTargets", "Target meshes are no longer valid"), EToolMessageLevel::UserWarning);
+			GetToolManager()->DisplayMessage(LOCTEXT("InvalidTargets", "Input meshes are no longer valid"), EToolMessageLevel::UserWarning);
 			return;
 		}
 		ConvertInputsAndSetPreviewMaterials(false);
 		Preview->InvalidateResult();
 	}
-	else if (Property && (Property->GetFName() == GET_MEMBER_NAME_CHECKED(UCutMeshWithMeshToolProperties, bShowNewBoundaryEdges)))
+	else if (Property && (Property->GetFName() == GET_MEMBER_NAME_CHECKED(UCutMeshWithMeshToolProperties, bShowNewBoundaries)))
 	{
 		GetToolManager()->PostInvalidation();
 		UpdateVisualization();
@@ -326,7 +325,7 @@ FString UCutMeshWithMeshTool::GetCreatedAssetName() const
 
 FText UCutMeshWithMeshTool::GetActionName() const
 {
-	return LOCTEXT("CutMeshWithMeshActionName", "Boolean Meshes");
+	return LOCTEXT("CutMeshWithMeshActionName", "Cut Mesh");
 }
 
 
diff --git a/Engine/Plugins/Runtime/MeshModelingToolset/Source/MeshModelingTools/Public/CutMeshWithMeshTool.h b/Engine/Plugins/Runtime/MeshModelingToolset/Source/MeshModelingTools/Public/CutMeshWithMeshTool.h
--- a/Engine/Plugins/Runtime/MeshModelingToolset/Source/MeshModelingTools/Public/CutMeshWithMeshTool.h
+++ b/Engine/Plugins/Runtime/MeshModelingToolset/Source/MeshModelingTools/Public/CutMeshWithMeshTool.h
@@ -10,7 +10,7 @@
 
 #include "CutMeshWithMeshTool.generated.h"
 
-// predeclarations
+// Forward declarations
 class UPreviewMesh;
 PREDECLARE_USE_GEOMETRY_CLASS(FDynamicMesh3);
 
@@ -23,25 +23,25 @@ class MESHMODELINGTOOLS_API UCutMeshWithMeshToolProperties : public UInteractive
 {
 	GENERATED_BODY()
 public:
-	/** Automatically attempt to fill any holes left by Booleans (e.g. due to numerical errors) */
-	UPROPERTY(EditAnywhere, Category = Options, meta = (UIMin = "0", UIMax = "1"))
-	float WindingThreshold = 0.5;
+	/** Try to fill holes created by the Boolean operation, e.g. due to numerical errors */
+	UPROPERTY(EditAnywhere, Category = Boolean, AdvancedDisplay)
+	bool bTryFixHoles = false;
 
-	/** Show boundary edges created by the Boolean operations (often due to numerical error) */
-	UPROPERTY(EditAnywhere, Category = Options)
-	bool bShowNewBoundaryEdges = true;
+	/** Try to collapse extra edges created by the Boolean operation */
+	UPROPERTY(EditAnywhere, Category = Boolean, AdvancedDisplay)
+	bool bTryCollapseEdges = true;
 
-	/** Automatically attempt to fill any holes left by Booleans (e.g. due to numerical errors) */
-	UPROPERTY(EditAnywhere, Category = Options)
-	bool bAttemptFixHoles = false;
+	/** Threshold to determine whether a triangle in one mesh is inside or outside of the other */
+	UPROPERTY(EditAnywhere, Category = Boolean, AdvancedDisplay, meta = (UIMin = "0", UIMax = "1"))
+	float WindingThreshold = 0.5;
 
-	/** Try to collapse extra edges created by the Boolean operation */
-	UPROPERTY(EditAnywhere, Category = Options)
-	bool bCollapseExtraEdges = true;
+	/** Show boundary edges created by the Boolean operation, which might happen due to numerical errors */
+	UPROPERTY(EditAnywhere, Category = Display)
+	bool bShowNewBoundaries = true;
 
-	/** If true, only the first mesh will keep its materials assignments; all other triangles will be assigned material 0 */
+	/** If true, only the first mesh will keep its material assignments, and all other faces will have the first material assigned */
 	UPROPERTY(EditAnywhere, Category = Materials)
-	bool bOnlyUseFirstMeshMaterials = false;
+	bool bUseFirstMeshMaterials = false;
 };
 
 
@@ -70,15 +70,14 @@ protected:
 	virtual void SaveProperties() override;
 	virtual void SetPreviewCallbacks() override;
 
-	virtual FString GetCreatedAssetName() const;
-	virtual FText GetActionName() const;
+	virtual FString GetCreatedAssetName() const override;
+	virtual FText GetActionName() const override;
 
 	// IDynamicMeshOperatorFactory API
 	virtual TUniquePtr<UE::Geometry::FDynamicMeshOperator> MakeNewOperator() override;
 
 	virtual void Shutdown(EToolShutdownType ShutdownType) override;
 
-protected:
 	void UpdateVisualization();
 
 	UPROPERTY()
